Populated initial PCG plugin files

diff --git a/Engine/Plugins/Experimental/PCG/PCG.uplugin b/Engine/Plugins/Experimental/PCG/PCG.uplugin
new file mode 100644
--- /dev/null
+++ b/Engine/Plugins/Experimental/PCG/PCG.uplugin
@@ -0,0 +1,33 @@
+{
+	"FileVersion": 3,
+	"Version": 8,
+	"VersionName": "0.1",
+	"EngineVersion": "5.1.0",
+	"FriendlyName": "PCG",
+	"Description": "Procedural generation framework for UE",
+	"Category": "Editor",
+	"CreatedBy": "Epic Games Inc",
+	"CreatedByURL": "http://epicgames.com",
+	"DocsURL": "",
+	"MarketplaceURL": "",
+	"SupportURL": "",
+	"EnabledByDefault": false,
+	"CanContainContent": false,
+	"IsBetaVersion": false,
+	"Installed": false,
+	"Modules": [
+		{
+			"Name": "PCG",
+			"Type": "Runtime",
+			"LoadingPhase": "Default"
+		},
+		{
+			"Name": "PCGEditor",
+			"Type": "Editor",
+			"LoadingPhase": "Default"			
+		}
+
+	],
+	"Plugins": [
+	]
+}
diff --git a/Engine/Plugins/Experimental/PCG/Source/PCG/PCG.Build.cs b/Engine/Plugins/Experimental/PCG/Source/PCG/PCG.Build.cs
new file mode 100644
--- /dev/null
+++ b/Engine/Plugins/Experimental/PCG/Source/PCG/PCG.Build.cs
@@ -0,0 +1,45 @@
+// Copyright 1998-2017 Epic Games, Inc. All Rights Reserved.
+
+namespace UnrealBuildTool.Rules
+{
+	public class PCG : ModuleRules
+	{
+		public PCG(ReadOnlyTargetRules Target) : base(Target)
+		{
+			PCHUsage = PCHUsageMode.UseExplicitOrSharedPCHs;
+
+			PublicDependencyModuleNames.AddRange(
+				new string[] {
+					"Core",
+					"CoreUObject",
+					"Engine",
+					"Landscape",
+					"Foliage",
+					"Projects",
+					"RenderCore",
+					"RHI",
+				});
+
+			PrivateDependencyModuleNames.AddRange(
+				new string[]
+				{					
+				}
+			);
+
+			if (Target.bBuildEditor == true)
+			{
+				PrivateDependencyModuleNames.AddRange(
+					new string[]
+					{
+						"UnrealEd",
+						"SourceControl",
+					}
+				);
+			}
+
+			PrivateIncludePaths.AddRange(
+				new string[] {
+				});
+		}
+	}
+}
diff --git a/Engine/Plugins/Experimental/PCG/Source/PCG/Private/Data/PCGDifferenceData.cpp b/Engine/Plugins/Experimental/PCG/Source/PCG/Private/Data/PCGDifferenceData.cpp
new file mode 100644
--- /dev/null
+++ b/Engine/Plugins/Experimental/PCG/Source/PCG/Private/Data/PCGDifferenceData.cpp
@@ -0,0 +1,101 @@
+// Copyright Epic Games, Inc. All Rights Reserved.
+
+#include "Data/PCGDifferenceData.h"
+#include "Data/PCGPointData.h"
+#include "Data/PCGUnionData.h"
+
+void UPCGDifferenceData::Initialize(const UPCGSpatialData* InData)
+{
+	check(InData);
+	Source = InData;
+	TargetActor = InData->TargetActor;
+}
+
+void UPCGDifferenceData::AddDifference(const UPCGSpatialData* InDifference)
+{
+	check(InDifference && InDifference->GetDimension() == Source->GetDimension());
+
+	if (!Difference)
+	{
+		Difference = InDifference;
+	}
+	else
+	{
+		if (!DifferencesUnion)
+		{
+			DifferencesUnion = NewObject<UPCGUnionData>(this);
+			DifferencesUnion->AddData(Difference);
+			Difference = DifferencesUnion;
+		}
+
+		check(Difference == DifferencesUnion);
+		DifferencesUnion->AddData(InDifference);
+	}
+}
+
+int UPCGDifferenceData::GetDimension() const
+{
+	return Source->GetDimension();
+}
+
+FBox UPCGDifferenceData::GetBounds() const
+{
+	return Source->GetBounds();
+}
+
+FBox UPCGDifferenceData::GetStrictBounds() const
+{
+	return FBox(EForceInit::ForceInit);
+}
+
+float UPCGDifferenceData::GetDensityAtPosition(const FVector& InPosition) const
+{
+	if (!Source->GetBounds().IsInside(InPosition) ||
+		(Difference && Difference->GetStrictBounds().IsInside(InPosition)))
+	{
+		return 0;
+	}
+	else if(Difference)
+	{
+		const float DensityInSource = Source->GetDensityAtPosition(InPosition);
+		if (DensityInSource <= 0)
+		{
+			return 0;
+		}
+
+		const float DensityInDifference = Difference->GetDensityAtPosition(InPosition);
+		return (DensityInDifference > 0 ? 0 : DensityInSource);
+	}
+	else
+	{
+		return Source->GetDensityAtPosition(InPosition);
+	}
+}
+
+const UPCGPointData* UPCGDifferenceData::CreatePointData() const
+{
+	TRACE_CPUPROFILER_EVENT_SCOPE(UPCGDifferenceData::CreatePointData);
+	
+	// This is similar to what we are doing in UPCGUnionData::CreatePointData
+	const UPCGPointData* SourcePointData = Source->ToPointData();
+
+	if (!Difference)
+	{
+		return SourcePointData;
+	}
+
+	UPCGPointData* Data = NewObject<UPCGPointData>(const_cast<UPCGDifferenceData*>(this));
+	Data->TargetActor = TargetActor;
+
+	TArray<FPCGPoint>& TargetPoints = Data->GetMutablePoints();
+
+	for (const FPCGPoint& Point : SourcePointData->GetPoints())
+	{
+		if (Difference->GetDensityAtPosition(Point.Transform.GetLocation()) == 0)
+		{
+			TargetPoints.Add(Point);
+		}
+	}
+
+	return Data;
+}
\ No newline at end of file
diff --git a/Engine/Plugins/Experimental/PCG/Source/PCG/Private/Data/PCGIntersectionData.cpp b/Engine/Plugins/Experimental/PCG/Source/PCG/Private/Data/PCGIntersectionData.cpp
new file mode 100644
--- /dev/null
+++ b/Engine/Plugins/Experimental/PCG/Source/PCG/Private/Data/PCGIntersectionData.cpp
@@ -0,0 +1,100 @@
+// Copyright Epic Games, Inc. All Rights Reserved.
+
+#include "Data/PCGIntersectionData.h"
+#include "Data/PCGPointData.h"
+
+void UPCGIntersectionData::Initialize(const UPCGSpatialData* InA, const UPCGSpatialData* InB)
+{
+	check(InA && InB);
+	A = InA;
+	B = InB;
+	TargetActor = A->TargetActor;
+
+	CachedBounds = A->GetBounds().Overlap(B->GetBounds());
+	CachedStrictBounds = A->GetStrictBounds().Overlap(B->GetStrictBounds());
+}
+
+int UPCGIntersectionData::GetDimension() const
+{
+	check(A && B);
+	return FMath::Min(A->GetDimension(), B->GetDimension());
+}
+
+FBox UPCGIntersectionData::GetBounds() const
+{
+	check(A && B);
+	return CachedBounds;
+}
+
+FBox UPCGIntersectionData::GetStrictBounds() const
+{
+	check(A && B);
+	return CachedStrictBounds;
+}
+
+float UPCGIntersectionData::GetDensityAtPosition(const FVector& InPosition) const
+{
+	check(A && B);
+	if (!CachedBounds.IsInside(InPosition))
+	{
+		return 0;
+	}
+	else if (CachedStrictBounds.IsInside(InPosition))
+	{
+		return 1.0f;
+	}
+	else
+	{
+		return A->GetDensityAtPosition(InPosition) * B->GetDensityAtPosition(InPosition);
+	}
+}
+
+const UPCGPointData* UPCGIntersectionData::CreatePointData() const
+{
+	check(A && B);
+	// TODO: this is a placeholder;
+	// Here we will get the point data from the lower-dimensionality data
+	// and then cull out any of the points that are outside the bounds of the other
+	if (A->GetDimension() <= B->GetDimension())
+	{
+		return CreateAndFilterPointData(A, B);
+	}
+	else
+	{
+		return CreateAndFilterPointData(B, A);
+	}
+}
+
+UPCGPointData* UPCGIntersectionData::CreateAndFilterPointData(const UPCGSpatialData* X, const UPCGSpatialData* Y) const
+{
+	TRACE_CPUPROFILER_EVENT_SCOPE(UPCGIntersectionData::CreateAndFilterPointData);
+	check(X && Y);
+	check(X->GetDimension() <= Y->GetDimension());
+
+	const UPCGPointData* SourcePointData = X->ToPointData();
+
+	if (!SourcePointData)
+	{
+		return nullptr;
+	}
+
+	const TArray<FPCGPoint>& SourcePoints = SourcePointData->GetPoints();
+
+	UPCGPointData* Data = NewObject<UPCGPointData>(const_cast<UPCGIntersectionData*>(this));
+	Data->TargetActor = TargetActor;
+	TArray<FPCGPoint>& TargetPoints = Data->GetMutablePoints();
+
+	// TODO: go through only the points in the bounding box, otherwise we know beforehand
+	// that the density will be zero. Requires that points are stored in a better data structure.
+	for (const FPCGPoint& Point : SourcePoints)
+	{
+		const float YDensity = Y->GetDensityAtPosition(Point.Transform.GetLocation());
+		if (YDensity > 0)
+		{
+			FPCGPoint& NewPoint = TargetPoints.Add_GetRef(Point);
+			NewPoint.Density *= YDensity;
+		}
+	}
+
+	return Data;
+}
\ No newline at end of file
diff --git a/Engine/Plugins/Experimental/PCG/Source/PCG/Private/Data/PCGLandscapeData.cpp b/Engine/Plugins/Experimental/PCG/Source/PCG/Private/Data/PCGLandscapeData.cpp
new file mode 100644
--- /dev/null
+++ b/Engine/Plugins/Experimental/PCG/Source/PCG/Private/Data/PCGLandscapeData.cpp
@@ -0,0 +1,120 @@
+// Copyright Epic Games, Inc. All Rights Reserved.
+
+#include "Data/PCGLandscapeData.h"
+
+#include "PCGHelpers.h"
+#include "Data/PCGPointData.h"
+#include "Data/PCGPolyLineData.h"
+#include "Data/PCGSpatialData.h"
+
+#include "Landscape.h"
+#include "LandscapeEdit.h"
+
+FPCGLandscapeDataPoint::FPCGLandscapeDataPoint(int InX, int InY, float InHeight)
+	: X(InX), Y(InY), Height(InHeight)
+{
+}
+
+void UPCGLandscapeData::Initialize(ALandscapeProxy* InLandscape, const FBox& InBounds)
+{
+	check(InLandscape);
+	Landscape = InLandscape;
+	TargetActor = InLandscape;
+	Bounds = InBounds;
+
+	Transform = Landscape->GetActorTransform();
+}
+
+FBox UPCGLandscapeData::GetBounds() const
+{
+	return Bounds;
+}
+
+FBox UPCGLandscapeData::GetStrictBounds() const
+{
+	// TODO: if the landscape contains holes, then the strict bounds
+	// should be empty
+	return Bounds;
+}
+
+FVector UPCGLandscapeData::TransformPosition(const FVector& InPosition) const
+{
+	TOptional<float> HeightAtVertex = Landscape->GetHeightAtLocation(InPosition);
+	if (HeightAtVertex.IsSet())
+	{
+		return FVector(InPosition.X, InPosition.Y, HeightAtVertex.GetValue());
+	}
+	else
+	{
+		return InPosition; // not on landscape
+	}
+}
+
+FPCGPoint UPCGLandscapeData::TransformPoint(const FPCGPoint& InPoint) const
+{
+	// TODO: change orientation, ... based on landscape data
+	FPCGPoint Point = InPoint;
+	FVector PointLocation = InPoint.Transform.GetLocation();
+
+	TOptional<float> HeightAtVertex = Landscape->GetHeightAtLocation(PointLocation);
+	if (HeightAtVertex.IsSet())
+	{
+		PointLocation.Z = HeightAtVertex.GetValue();
+		Point.Transform.SetLocation(PointLocation);
+	}
+	else
+	{
+		Point.Density = 0;
+	}
+	
+	return Point;
+}
+
+float UPCGLandscapeData::GetDensityAtPosition(const FVector& InPosition) const
+{
+	TOptional<float> HeightAtVertex = Landscape->GetHeightAtLocation(InPosition);
+	return HeightAtVertex.IsSet() ? 1.0f : 0;
+}
+
+const UPCGPointData* UPCGLandscapeData::CreatePointData() const
+{
+	TRACE_CPUPROFILER_EVENT_SCOPE(UPCGLandscapeData::CreatePointData);
+
+	UPCGPointData* Data = NewObject<UPCGPointData>(const_cast<UPCGLandscapeData*>(this));
+	Data->TargetActor = TargetActor;
+	TArray<FPCGPoint>& Points = Data->GetMutablePoints();
+
+	// TODO: add offset to nearest edge, will have an impact if the grid size doesn't match the landscape size
+	const FVector MinPt = Transform.InverseTransformPosition(Bounds.Min);
+	const FVector MaxPt = Transform.InverseTransformPosition(Bounds.Max);
+
+	const int32 MinX = FMath::FloorToInt(MinPt.X);
+	const int32 MaxX = FMath::FloorToInt(MaxPt.X);
+	const int32 MinY = FMath::FloorToInt(MinPt.Y);
+	const int32 MaxY = FMath::FloorToInt(MaxPt.Y);
+
+	Points.Reserve((MaxX - MinX) * (MaxY - MinY));
+
+	bool bPlaneCase = Bounds.Min.Z == Bounds.Max.Z;
+
+	for (int X = MinX; X < MaxX; ++X)
+	{
+		for (int Y = MinY; Y < MaxY; ++Y)
+		{
+			FVector VertexLocation = Transform.TransformPosition(FVector(X, Y, 0));
+			TOptional<float> HeightAtVertex = Landscape->GetHeightAtLocation(VertexLocation);
+			if (HeightAtVertex.IsSet() &&
+				HeightAtVertex.GetValue() >= Bounds.Min.Z && 
+				(bPlaneCase ? HeightAtVertex.GetValue() <= Bounds.Max.Z : HeightAtVertex.GetValue() < Bounds.Max.Z))
+			{
+				VertexLocation.Z = HeightAtVertex.GetValue();
+
+				Points.Emplace(FTransform(VertexLocation),
+					1.0f,
+					PCGHelpers::ComputeSeed(X, Y, (int)HeightAtVertex.GetValue()));
+			}
+		}
+	}
+
+	return Data;
+}
\ No newline at end of file
diff --git a/Engine/Plugins/Experimental/PCG/Source/PCG/Private/Data/PCGLandscapeSplineData.cpp b/Engine/Plugins/Experimental/PCG/Source/PCG/Private/Data/PCGLandscapeSplineData.cpp
new file mode 100644
--- /dev/null
+++ b/Engine/Plugins/Experimental/PCG/Source/PCG/Private/Data/PCGLandscapeSplineData.cpp
@@ -0,0 +1,256 @@
+// Copyright Epic Games, Inc. All Rights Reserved.
+
+#include "Data/PCGLandscapeSplineData.h"
+
+#include "PCGHelpers.h"
+#include "Data/PCGPointData.h"
+
+#include "Landscape.h"
+#include "LandscapeSplinesComponent.h"
+#include "LandscapeSplineSegment.h"
+#include "LandscapeSplineControlPoint.h"
+
+namespace PCGLandscapeDataHelpers
+{
+	// This function assumes that the A-B segment has a "1" density, while the C-D segment has a "0" density
+	FVector::FReal GetDensityInQuad(const FVector& A, const FVector& B, const FVector& C, const FVector& D, const FVector& P)
+	{
+		FVector BaryABC = FMath::ComputeBaryCentric2D(P, A, B, C);
+
+		if (BaryABC.X >= 0 && BaryABC.Y >= 0 && BaryABC.Z >= 0)
+		{
+			return 1.0f - BaryABC.Z;
+		}
+
+		FVector BaryACD = FMath::ComputeBaryCentric2D(P, A, C, D);
+
+		if (BaryACD.X >= 0 && BaryACD.Y >= 0 && BaryACD.Z >= 0)
+		{
+			return BaryACD.X;
+		}
+
+		return FVector::FReal(-1);
+	}
+}
+
+void UPCGLandscapeSplineData::Initialize(ULandscapeSplinesComponent* InSplineComponent)
+{
+	check(InSplineComponent);
+	Spline = InSplineComponent;
+}
+
+int UPCGLandscapeSplineData::GetNumSegments() const
+{
+	check(Spline);
+	return Spline->GetSegments().Num();
+}
+
+float UPCGLandscapeSplineData::GetSegmentLength(int SegmentIndex) const
+{
+	check(Spline);
+	check(SegmentIndex >= 0 && SegmentIndex < Spline->GetSegments().Num());
+	
+	const ULandscapeSplineSegment* Segment = Spline->GetSegments()[SegmentIndex];
+	const TArray<FLandscapeSplineInterpPoint>& InterpPoints = Segment->GetPoints();
+	float Length = 0;
+
+	for (int PointIndex = 1; PointIndex < InterpPoints.Num(); ++PointIndex)
+	{
+		Length += (InterpPoints[PointIndex].Center - InterpPoints[PointIndex - 1].Center).Length();
+	}
+
+	return Length;
+}
+
+FVector UPCGLandscapeSplineData::GetLocationAtDistance(int SegmentIndex, float Distance) const
+{
+	check(Spline);
+	check(SegmentIndex >= 0 && SegmentIndex < Spline->GetSegments().Num());
+
+	const ULandscapeSplineSegment* Segment = Spline->GetSegments()[SegmentIndex];
+	const TArray<FLandscapeSplineInterpPoint>& InterpPoints = Segment->GetPoints();
+	float Length = Distance;
+
+	for (int PointIndex = 1; PointIndex < InterpPoints.Num(); ++PointIndex)
+	{
+		const float SegmentLength = (InterpPoints[PointIndex].Center - InterpPoints[PointIndex - 1].Center).Length();
+		if (SegmentLength > Length)
+		{
+			return InterpPoints[PointIndex - 1].Center + (InterpPoints[PointIndex].Center - InterpPoints[PointIndex - 1].Center) * (Length / SegmentLength);
+		}
+		else
+		{
+			Length -= SegmentLength;
+		}
+	}
+	
+	check(0);
+	return FVector::Zero();
+}
+
+const UPCGPointData* UPCGLandscapeSplineData::CreatePointData() const
+{
+	check(Spline);
+	TRACE_CPUPROFILER_EVENT_SCOPE(UPCGLandscapeSplineData::CreatePointData);
+
+	UPCGPointData* Data = NewObject<UPCGPointData>(const_cast<UPCGLandscapeSplineData*>(this));
+	Data->TargetActor = TargetActor;
+	TArray<FPCGPoint>& Points = Data->GetMutablePoints();
+
+	// TODO: replace all the logic with sampling settings; currently it uses the landscape scaling as a basis
+	const FTransform SplineTransform = Spline->GetComponentTransform();
+	const FTransform LandscapeTransform = Spline->GetSplineOwner()->LandscapeActorToWorld();
+	const FTransform SplineToLandscape = SplineTransform.GetRelativeTransform(LandscapeTransform);
+	ULandscapeInfo* LandscapeInfo = Spline->GetSplineOwner()->GetLandscapeInfo();
+
+	auto AddPoints = [&Points, &SplineTransform, &LandscapeTransform, &SplineToLandscape, LandscapeInfo](const FVector& A, const FVector& B, const FVector& C, const FVector& D, bool bComputeDensity, bool bAddA, bool bAddB) {
+
+		auto AddPoint = [&Points](const FVector& P, float Density) {
+			FPCGPoint& Point = Points.Emplace_GetRef();
+			Point.Transform = FTransform(P); // pass in segment vector as orientation?
+			Point.Seed = PCGHelpers::ComputeSeed((int)P.X, (int)P.Y, (int)P.Z);
+			Point.Density = Density;
+		};
+
+		if (bAddA)
+		{
+			AddPoint(SplineTransform.TransformPosition(A), 1.0f);
+		}
+
+		// Interpolate points in the ABCD quad based on the given scales, computing density based on ratio from the A->B segment if the bComputeDensity flag is on, and 1.0f otherwise
+		if (LandscapeInfo && LandscapeInfo->LandscapeActor.IsValid())
+		{
+			FBox QuadBoxOnLandscape(EForceInit::ForceInit);
+			QuadBoxOnLandscape += SplineToLandscape.TransformPosition(A);
+			QuadBoxOnLandscape += SplineToLandscape.TransformPosition(B);
+			QuadBoxOnLandscape += SplineToLandscape.TransformPosition(C);
+			QuadBoxOnLandscape += SplineToLandscape.TransformPosition(D);
+
+			int32 MinX = FMath::CeilToInt(QuadBoxOnLandscape.Min.X);
+			int32 MinY = FMath::CeilToInt(QuadBoxOnLandscape.Min.Y);
+			int32 MaxX = FMath::FloorToInt(QuadBoxOnLandscape.Max.X);
+			int32 MaxY = FMath::FloorToInt(QuadBoxOnLandscape.Max.Y);
+
+			const FVector QuadNormal = (C - B).Cross(B - A).GetSafeNormal();
+
+			for (int X = MinX; X <= MaxX; ++X)
+			{
+				for (int Y = MinY; Y <= MaxY; ++Y)
+				{
+					FVector TentativeLocation = LandscapeTransform.TransformPosition(FVector(X, Y, 0));
+					const FVector TentativeLocationInSplineSpace = SplineTransform.InverseTransformPosition(TentativeLocation);
+
+					const FVector::FReal ComputedDensity = PCGLandscapeDataHelpers::GetDensityInQuad(A, B, C, D, TentativeLocationInSplineSpace);
+
+					// Check if the point would be in the quad
+					if (ComputedDensity >= 0)
+					{
+						// TODO: note that we can't call GetHeight on the ULandscapeHEightfieldCollisionComponent, as it is not exported
+						// so we have to resort to calling the method from the landscape actor
+						TOptional<float> HeightAtVertex = LandscapeInfo->LandscapeActor->GetHeightAtLocation(TentativeLocation);
+						if (HeightAtVertex.IsSet())
+						{
+							TentativeLocation.Z = HeightAtVertex.GetValue();
+							AddPoint(TentativeLocation, bComputeDensity ? ComputedDensity : 1.0f);
+						}
+					}
+				}
+			}
+		}
+
+		if (bAddB)
+		{
+			AddPoint(SplineTransform.TransformPosition(B), 1.0f);
+		}
+	};
+	
+	const TArray<TObjectPtr<ULandscapeSplineSegment>>& Segments = Spline->GetSegments();
+
+	// For each segment on the spline,
+	for (int SegmentIndex = 0; SegmentIndex < Segments.Num(); ++SegmentIndex)
+	{
+		// Sample points on and between points
+		const TArray<FLandscapeSplineInterpPoint>& InterpPoints = Segments[SegmentIndex]->GetPoints();
+		bool bIsLastSegment = (SegmentIndex == Segments.Num() - 1);
+
+		for (int PointIndex = 1; PointIndex < InterpPoints.Num(); ++PointIndex)
+		{
+			bool bIsLastPointInSpline = (bIsLastSegment && PointIndex == InterpPoints.Num() - 1);
+
+			const FLandscapeSplineInterpPoint& Start = InterpPoints[PointIndex - 1];
+			const FLandscapeSplineInterpPoint& End = InterpPoints[PointIndex];
+
+			AddPoints(Start.Center, End.Center, End.Left, Start.Left, /*bComputeDensity=*/false, /*bAddA=*/true, /*bAddB=*/false);
+			AddPoints(Start.Left, End.Left, End.FalloffLeft, Start.FalloffLeft, /*bComputeDensity=*/true, /*bAddA=*/true, /*bAddB=*/bIsLastPointInSpline);
+			AddPoints(End.Center, Start.Center, Start.Right, End.Right, /*bComputeDensity=*/false, /*bAddA=*/bIsLastPointInSpline, /*bAddB=*/false);
+			AddPoints(End.Right, Start.Right, Start.FalloffRight, End.FalloffRight, /*bComputeDensity=*/true, /*bAddA=*/bIsLastPointInSpline, /*bAddB=*/true);
+		}
+	}
+
+	return Data;
+}
+
+FBox UPCGLandscapeSplineData::GetBounds() const
+{
+	check(Spline);
+
+	FBox Bounds(EForceInit::ForceInit);
+	for (const TObjectPtr<ULandscapeSplineSegment>& Segment : Spline->GetSegments())
+	{
+		Bounds += Segment->GetBounds();
+	}
+
+	return Bounds;
+}
+
+float UPCGLandscapeSplineData::GetDensityAtPosition(const FVector& InPosition) const
+{
+	check(Spline);
+
+	const FVector Position = Spline->GetComponentTransform().InverseTransformPosition(InPosition);
+
+	float PointDensity = 0.0f;
+
+	for(const TObjectPtr<ULandscapeSplineSegment> Segment : Spline->GetSegments())
+	{
+		// Considering the landscape spline always exists on the landscape,
+		// we'll ignore the Z component of the input here for the bounds check.
+		if (!Segment->GetBounds().IsInsideXY(Position))
+		{
+			continue;
+		}
+		
+		const TArray<FLandscapeSplineInterpPoint>& InterpPoints = Segment->GetPoints();
+
+		float SegmentDensity = 0.0f;
+
+		for (int PointIndex = 1; PointIndex < InterpPoints.Num(); ++PointIndex)
+		{
+			const FLandscapeSplineInterpPoint& Start = InterpPoints[PointIndex - 1];
+			const FLandscapeSplineInterpPoint& End = InterpPoints[PointIndex];
+
+			float Density = 0.0f;
+
+			// Note: these checks have no prior information on the structure of the data, except that they form quads.
+			// Considering that the points on a given control point are probably aligned, we could do an early check
+			// in the original quad (start left falloff -> start right falloff -> end right falloff -> end left falloff)
+			// TODO: this sequence here can be optimized knowing that some checks will be redundant.
+			Density = FMath::Max(Density, PCGLandscapeDataHelpers::GetDensityInQuad(Start.Center, End.Center, End.Left, Start.Left, Position) >= 0 ? 1.0f : 0.0f);
+			Density = FMath::Max(Density, PCGLandscapeDataHelpers::GetDensityInQuad(Start.Left, End.Left, End.FalloffLeft, Start.FalloffLeft, Position));
+			Density = FMath::Max(Density, PCGLandscapeDataHelpers::GetDensityInQuad(Start.Center, Start.Right, End.Right, End.Center, Position) >= 0 ? 1.0f : 0.0f);
+			Density = FMath::Max(Density, PCGLandscapeDataHelpers::GetDensityInQuad(Start.Right, Start.FalloffRight, End.FalloffRight, End.Right, Position));
+
+			if (Density > SegmentDensity)
+			{
+				SegmentDensity = Density;
+			}
+		}
+
+		if (SegmentDensity > PointDensity)
+		{
+			PointDensity = SegmentDensity;
+		}
+	}
+
+	return PointDensity;
+}
\ No newline at end of file
diff --git a/Engine/Plugins/Experimental/PCG/Source/PCG/Private/Data/PCGPointData.cpp b/Engine/Plugins/Experimental/PCG/Source/PCG/Private/Data/PCGPointData.cpp
new file mode 100644
--- /dev/null
+++ b/Engine/Plugins/Experimental/PCG/Source/PCG/Private/Data/PCGPointData.cpp
@@ -0,0 +1,147 @@
+// Copyright Epic Games, Inc. All Rights Reserved.
+
+#include "Data/PCGPointData.h"
+
+#include "PCGHelpers.h"
+
+#include "GameFramework/Actor.h"
+#include "Misc/ScopeLock.h"
+
+FPCGPointRef::FPCGPointRef(const FPCGPoint& InPoint)
+{
+	Point = &InPoint;
+	Bounds = InPoint.GetBounds();
+}
+
+FPCGPointRef::FPCGPointRef(const FPCGPointRef& InPointRef)
+{
+	Point = InPointRef.Point;
+	Bounds = InPointRef.Bounds;
+}
+
+TArray<FPCGPoint>& UPCGPointData::GetMutablePoints()
+{
+	bOctreeIsDirty = true;
+	bBoundsAreDirty = true;
+	return Points;
+}
+
+FBox UPCGPointData::GetBounds() const
+{
+	if (bBoundsAreDirty)
+	{
+		RecomputeBounds();
+	}
+
+	return Bounds;
+}
+
+void UPCGPointData::RecomputeBounds() const
+{
+	FScopeLock Lock(&CachedDataLock);
+
+	if (!bBoundsAreDirty)
+	{
+		return;
+	}
+
+	FBox NewBounds(EForceInit::ForceInit);
+	for (const FPCGPoint& Point : Points)
+	{
+		FBoxSphereBounds PointBounds = Point.GetBounds();
+		NewBounds += FBox::BuildAABB(PointBounds.Origin, PointBounds.BoxExtent);
+	}
+
+	Bounds = NewBounds;
+	bBoundsAreDirty = false;
+}
+
+void UPCGPointData::CopyPointsFrom(const UPCGPointData* InData, const TArray<int>& InDataIndices)
+{
+	TRACE_CPUPROFILER_EVENT_SCOPE(UPCGPointData::CopyPointsFrom);
+	check(InData);
+	Points.SetNum(InDataIndices.Num());
+
+	// TODO: parallel-for this?
+	for (int PointIndex = 0; PointIndex < InDataIndices.Num(); ++PointIndex)
+	{
+		Points[PointIndex] = InData->Points[InDataIndices[PointIndex]];
+	}
+
+	bBoundsAreDirty = true;
+	bOctreeIsDirty = true;
+}
+
+void UPCGPointData::SetPoints(const TArray<FPCGPoint>& InPoints)
+{
+	GetMutablePoints() = InPoints;
+}
+
+void UPCGPointData::Initialize(AActor* InActor)
+{
+	check(InActor);
+
+	Points.SetNum(1);
+	Points[0].Transform = InActor->GetActorTransform();
+
+	const FVector& Position = Points[0].Transform.GetLocation();
+	Points[0].Seed = PCGHelpers::ComputeSeed((int)Position.X, (int)Position.Y, (int)Position.Z);
+
+	TargetActor = InActor;
+}
+
+const FPCGPoint* UPCGPointData::GetPointAtPosition(const FVector& InPosition) const
+{
+	if (bOctreeIsDirty)
+	{
+		RebuildOctree();
+	}
+
+	const FPCGPoint* BestPoint = nullptr;
+
+	Octree.FindElementsWithBoundsTest(FBoxCenterAndExtent(InPosition, FVector::Zero()), [&BestPoint](const FPCGPointRef& InPointRef) {
+		if (!BestPoint || BestPoint->Density < InPointRef.Point->Density)
+		{
+			BestPoint = InPointRef.Point;
+		}
+	});
+
+	return BestPoint;
+}
+
+float UPCGPointData::GetDensityAtPosition(const FVector& InPosition) const
+{
+	if (const FPCGPoint* BestPoint = GetPointAtPosition(InPosition))
+	{
+		// It is going to be hard to do a linear fall-off here and maybe not really required.
+		return BestPoint->Density;
+	}
+	else
+	{
+		return 0;
+	}
+}
+
+void UPCGPointData::RebuildOctree() const
+{
+	FScopeLock Lock(&CachedDataLock);
+
+	if (!bOctreeIsDirty)
+	{
+		return;
+	}
+
+	TRACE_CPUPROFILER_EVENT_SCOPE(UPCGPointData::RebuildOctree)
+	check(bOctreeIsDirty);
+
+	FBox PointBounds = GetBounds();
+	TOctree2<FPCGPointRef, FPCGPointRefSemantics> NewOctree(PointBounds.GetCenter(), PointBounds.GetExtent().Length());
+
+	for (const FPCGPoint& Point : Points)
+	{
+		NewOctree.AddElement(FPCGPointRef(Point));
+	}
+
+	Octree = NewOctree;
+	bOctreeIsDirty = false;
+}
\ No newline at end of file
diff --git a/Engine/Plugins/Experimental/PCG/Source/PCG/Private/Data/PCGPolyLineData.cpp b/Engine/Plugins/Experimental/PCG/Source/PCG/Private/Data/PCGPolyLineData.cpp
new file mode 100644
--- /dev/null
+++ b/Engine/Plugins/Experimental/PCG/Source/PCG/Private/Data/PCGPolyLineData.cpp
@@ -0,0 +1,18 @@
+// Copyright Epic Games, Inc. All Rights Reserved.
+
+#include "Data/PCGPolyLineData.h"
+
+FBox UPCGPolyLineData::GetBounds() const
+{
+	FBox Bounds(EForceInit::ForceInit);
+
+	const int NumSegments = GetNumSegments();
+
+	for (int SegmentIndex = 0; SegmentIndex < NumSegments; ++SegmentIndex)
+	{
+		Bounds += GetLocationAtDistance(SegmentIndex, 0);
+		Bounds += GetLocationAtDistance(SegmentIndex, GetSegmentLength(SegmentIndex));
+	}
+	
+	return Bounds;
+}
\ No newline at end of file
diff --git a/Engine/Plugins/Experimental/PCG/Source/PCG/Private/Data/PCGProjectionData.cpp b/Engine/Plugins/Experimental/PCG/Source/PCG/Private/Data/PCGProjectionData.cpp
new file mode 100644
--- /dev/null
+++ b/Engine/Plugins/Experimental/PCG/Source/PCG/Private/Data/PCGProjectionData.cpp
@@ -0,0 +1,107 @@
+// Copyright Epic Games, Inc. All Rights Reserved.
+
+#include "Data/PCGProjectionData.h"
+
+void UPCGProjectionData::Initialize(const UPCGSpatialData* InSource, const UPCGSpatialData* InTarget)
+{
+	check(InSource && InTarget);
+	// TODO: improve support for higher-dimension projection.
+	// The problem is that there isn't a valid 1:1 mapping otherwise
+	check(InSource->GetDimension() <= InTarget->GetDimension());
+	Source = InSource;
+	Target = InTarget;
+	TargetActor = InSource->TargetActor;
+
+	CachedBounds = ProjectBounds(Source->GetBounds());
+	CachedStrictBounds = ProjectBounds(Source->GetStrictBounds());
+}
+
+int UPCGProjectionData::GetDimension() const
+{
+	check(Source && Target);
+	return FMath::Min(Source->GetDimension(), Target->GetDimension());
+}
+
+FBox UPCGProjectionData::GetBounds() const
+{
+	check(Source && Target);
+	return CachedBounds;
+}
+
+FBox UPCGProjectionData::GetStrictBounds() const
+{
+	check(Source && Target);
+	return CachedStrictBounds;
+}
+
+FVector UPCGProjectionData::GetNormal() const
+{
+	check(Source && Target);
+	if (Source->GetDimension() > Target->GetDimension())
+	{
+		return Source->GetNormal();
+	}
+	else
+	{
+		return Target->GetNormal();
+	}
+}
+
+FBox UPCGProjectionData::ProjectBounds(const FBox& InBounds) const
+{
+	FBox Bounds(EForceInit::ForceInit);
+
+	for (int Corner = 0; Corner < 8; ++Corner)
+	{
+		Bounds += ProjectPosition(
+			FVector(
+				(Corner / 4) ? InBounds.Max.X : InBounds.Min.X,
+				((Corner / 2) % 2) ? InBounds.Max.Y : InBounds.Min.Y,
+				(Corner % 2) ? InBounds.Max.Z : InBounds.Min.Z));
+	}
+
+	return Bounds;
+}
+
+float UPCGProjectionData::GetDensityAtPosition(const FVector& InPosition) const
+{
+	// TODO: improve projection/unprojection mechanism
+	return Source->GetDensityAtPosition(InPosition);
+}
+
+FVector UPCGProjectionData::ProjectPosition(const FVector& InPosition) const
+{
+	// TODO: improve projection/unprojection mechanism
+	return Target->TransformPosition(InPosition);
+}
+
+FPCGPoint UPCGProjectionData::ProjectPoint(const FPCGPoint& InPoint) const
+{
+	// TODO: improve projection/unprojection mechanism
+	return Target->TransformPoint(InPoint);
+}
+
+const UPCGPointData* UPCGProjectionData::CreatePointData() const
+{
+	TRACE_CPUPROFILER_EVENT_SCOPE(UPCGProjectionData::CreatePointData);
+	// TODO: add mechanism in the ToPointData so we can pass in a transform
+	// so we can forego creating the points twice if they're not used.
+	const UPCGPointData* SourcePointData = Source->ToPointData();
+	const TArray<FPCGPoint>& SourcePoints = SourcePointData->GetPoints();
+
+	UPCGPointData* PointData = NewObject<UPCGPointData>(const_cast<UPCGProjectionData*>(this));
+	PointData->TargetActor = TargetActor;
+	TArray<FPCGPoint>& Points = PointData->GetMutablePoints();
+
+	Points.Reserve(SourcePoints.Num());
+	for (const FPCGPoint& SourcePoint : SourcePoints)
+	{
+		FPCGPoint ProjectedPoint = ProjectPoint(SourcePoint);
+		if (ProjectedPoint.Density > 0)
+		{
+			Points.Add(MoveTemp(ProjectedPoint));
+		}
+	}
+
+	return PointData;
+}
\ No newline at end of file
diff --git a/Engine/Plugins/Experimental/PCG/Source/PCG/Private/Data/PCGRenderTargetData.cpp b/Engine/Plugins/Experimental/PCG/Source/PCG/Private/Data/PCGRenderTargetData.cpp
new file mode 100644
--- /dev/null
+++ b/Engine/Plugins/Experimental/PCG/Source/PCG/Private/Data/PCGRenderTargetData.cpp
@@ -0,0 +1,34 @@
+// Copyright Epic Games, Inc. All Rights Reserved.
+
+#include "Data/PCGRenderTargetData.h"
+#include "Kismet/KismetRenderingLibrary.h"
+
+void UPCGRenderTargetData::Initialize(UTextureRenderTarget2D* InRenderTarget, const FTransform& InTransform)
+{
+	RenderTarget = InRenderTarget;
+	Transform = InTransform;
+
+	ColorData.Reset();
+
+	if (RenderTarget)
+	{
+		if (FTextureRenderTargetResource* RTResource = RenderTarget->GameThread_GetRenderTargetResource())
+		{
+			ETextureRenderTargetFormat Format = RenderTarget->RenderTargetFormat;
+			if ((Format == RTF_RGBA16f) || (Format == RTF_RGBA32f) || (Format == RTF_RGBA8) || (Format == RTF_R8) || (Format == RTF_RGB10A2))
+			{
+				const FIntRect Rect = FIntRect(0, 0, RenderTarget->SizeX, RenderTarget->SizeY);
+				const FReadSurfaceDataFlags ReadPixelFlags(RCM_MinMax);
+				RTResource->ReadLinearColorPixels(ColorData, ReadPixelFlags, Rect);
+			}
+		}
+
+		Width = RenderTarget->SizeX;
+		Height = RenderTarget->SizeY;
+	}
+
+	Bounds = FBox(EForceInit::ForceInit);
+	Bounds += FVector(-1.0f, -1.0f, 0.0f);
+	Bounds += FVector(1.0f, 1.0f, 0.0f);
+	Bounds = Bounds.TransformBy(Transform);
+}
\ No newline at end of file
diff --git a/Engine/Plugins/Experimental/PCG/Source/PCG/Private/Data/PCGSpatialData.cpp b/Engine/Plugins/Experimental/PCG/Source/PCG/Private/Data/PCGSpatialData.cpp
new file mode 100644
--- /dev/null
+++ b/Engine/Plugins/Experimental/PCG/Source/PCG/Private/Data/PCGSpatialData.cpp
@@ -0,0 +1,38 @@
+// Copyright Epic Games, Inc. All Rights Reserved.
+
+#include "Data/PCGSpatialData.h"
+#include "Data/PCGIntersectionData.h"
+#include "Data/PCGProjectionData.h"
+
+const UPCGPointData* UPCGSpatialDataWithPointCache::ToPointData() const
+{
+	if (!CachedPointData)
+	{
+		CacheLock.Lock();
+
+		if (!CachedPointData)
+		{
+			CachedPointData = CreatePointData();
+		}
+
+		CacheLock.Unlock();
+	}
+
+	return CachedPointData;
+}
+
+UPCGSpatialData* UPCGSpatialData::IntersectWith(const UPCGSpatialData* InOther) const
+{
+	UPCGIntersectionData* IntersectionData = NewObject<UPCGIntersectionData>(const_cast<UPCGSpatialData*>(this));
+	IntersectionData->Initialize(this, InOther);
+
+	return IntersectionData;
+}
+
+UPCGSpatialData* UPCGSpatialData::ProjectOn(const UPCGSpatialData* InOther) const
+{
+	UPCGProjectionData* ProjectionData = NewObject<UPCGProjectionData>(const_cast<UPCGSpatialData*>(this));
+	ProjectionData->Initialize(this, InOther);
+
+	return ProjectionData;
+}
\ No newline at end of file
diff --git a/Engine/Plugins/Experimental/PCG/Source/PCG/Private/Data/PCGSplineData.cpp b/Engine/Plugins/Experimental/PCG/Source/PCG/Private/Data/PCGSplineData.cpp
new file mode 100644
--- /dev/null
+++ b/Engine/Plugins/Experimental/PCG/Source/PCG/Private/Data/PCGSplineData.cpp
@@ -0,0 +1,192 @@
+// Copyright Epic Games, Inc. All Rights Reserved.
+
+#include "Data/PCGSplineData.h"
+
+#include "PCGHelpers.h"
+#include "Data/PCGPointData.h"
+#include "Data/PCGSurfaceData.h"
+
+#include "Components/SplineComponent.h"
+
+void UPCGSplineData::Initialize(USplineComponent* InSpline)
+{
+	check(InSpline);
+	Spline = InSpline;
+	TargetActor = InSpline->GetOwner();
+}
+
+int UPCGSplineData::GetNumSegments() const
+{
+	return 1;
+}
+
+float UPCGSplineData::GetSegmentLength(int SegmentIndex) const
+{
+	return Spline->GetSplineLength();
+}
+
+FVector UPCGSplineData::GetLocationAtDistance(int SegmentIndex, float Distance) const
+{
+	return Spline->GetLocationAtDistanceAlongSpline(Distance, ESplineCoordinateSpace::World);
+}
+
+const UPCGPointData* UPCGSplineData::CreatePointData() const
+{
+	TRACE_CPUPROFILER_EVENT_SCOPE(UPCGSplineData::CreatePointData);
+	UPCGPointData* Data = NewObject<UPCGPointData>(const_cast<UPCGSplineData*>(this));
+	Data->TargetActor = TargetActor;
+
+	// TODO: introduce settings to sample, 
+	// should be part of the settings passed in param.
+	float SplineLength = Spline->GetSplineLength();
+	float Offset = (SplineLength - FMath::Floor(SplineLength)) / 2.0f;
+
+	TArray<FPCGPoint>& Points = Data->GetMutablePoints();
+	int NumPoints = int(SplineLength) + 1;
+	Points.SetNum(NumPoints);
+
+	int Index = 0;
+	float Distance = Offset;
+	while (Distance < SplineLength)
+	{
+		Points[Index].Transform = Spline->GetTransformAtDistanceAlongSpline(Distance, ESplineCoordinateSpace::World);
+		// Exception: considering the spline holds size data in its scale, we'll reset it to identity here
+		Points[Index].Transform.SetScale3D(FVector::One());
+		Points[Index].Seed = PCGHelpers::ComputeSeed((int)Distance);
+		Points[Index].Density = 1.0f;
+		++Index;
+		Distance += 1.0f;
+	}
+
+	return Data;
+}
+
+FBox UPCGSplineData::GetBounds() const
+{
+	return PCGHelpers::GetActorBounds(Spline->GetOwner());
+}
+
+float UPCGSplineData::GetDensityAtPosition(const FVector& InPosition) const
+{
+	// Find nearest point on spline
+	float NearestPointKey = Spline->FindInputKeyClosestToWorldLocation(InPosition);
+
+	FTransform NearestTransform = Spline->GetTransformAtSplineInputKey(NearestPointKey, ESplineCoordinateSpace::World, true);
+
+	FVector LocalPoint = NearestTransform.InverseTransformPosition(InPosition);
+
+	// Linear fall off based on the distance to the nearest point
+	// TODO: should be based on explicit settings
+	float Distance = LocalPoint.Length();
+	return FMath::Max(0, 1.0f - Distance);
+}
+
+UPCGSpatialData* UPCGSplineData::ProjectOn(const UPCGSpatialData* InOther) const
+{
+	if(InOther->GetDimension() == 2)
+	{
+		UPCGSplineProjectionData* SplineProjectionData = NewObject<UPCGSplineProjectionData>(const_cast<UPCGSplineData*>(this));
+		SplineProjectionData->Initialize(this, InOther);
+		return SplineProjectionData;
+	}
+	else
+	{
+		return Super::ProjectOn(InOther);
+	}
+}
+
+FVector2D UPCGSplineProjectionData::Project(const FVector& InVector) const
+{
+	const FVector& SurfaceNormal = GetSurface()->GetNormal();
+	FVector Projection = InVector - InVector.ProjectOnToNormal(SurfaceNormal);
+
+	// Ignore smallest absolute coordinate value.
+	// Normally, one should be zero, but because of numerical precision,
+	// We'll make sure we're doing it right
+	FVector::FReal SmallestCoordinate = TNumericLimits<FVector::FReal>::Max();
+	int SmallestCoordinateAxis = -1;
+
+	for (int Axis = 0; Axis < 3; ++Axis)
+	{
+		FVector::FReal AbsoluteCoordinateValue = FMath::Abs(Projection[Axis]);
+		if (AbsoluteCoordinateValue < SmallestCoordinate)
+		{
+			SmallestCoordinate = AbsoluteCoordinateValue;
+			SmallestCoordinateAxis = Axis;
+		}
+	}
+
+	FVector2D Projection2D;
+	int AxisIndex = 0;
+	for (int Axis = 0; Axis < 3; ++Axis)
+	{
+		if (Axis != SmallestCoordinateAxis)
+		{
+			Projection2D[AxisIndex++] = Projection[Axis];
+		}
+	}
+
+	return Projection2D;
+}
+
+void UPCGSplineProjectionData::Initialize(const UPCGSplineData* InSourceSpline, const UPCGSpatialData* InTargetSurface)
+{
+	Super::Initialize(InSourceSpline, InTargetSurface);
+
+	const USplineComponent* Spline = GetSpline()->Spline.Get();
+	const FVector& SurfaceNormal = GetSurface()->GetNormal();
+
+	if (Spline)
+	{
+		const FInterpCurveVector& SplinePosition = Spline->GetSplinePointsPosition();
+
+		// Build projected spline data
+		ProjectedPosition.bIsLooped = SplinePosition.bIsLooped;
+		ProjectedPosition.LoopKeyOffset = SplinePosition.LoopKeyOffset;
+
+		ProjectedPosition.Points.Reserve(SplinePosition.Points.Num());
+
+		for (const FInterpCurvePoint<FVector>& SplinePoint : SplinePosition.Points)
+		{
+			FInterpCurvePoint<FVector2D>& ProjectedPoint = ProjectedPosition.Points.Emplace_GetRef();
+
+			ProjectedPoint.InVal = SplinePoint.InVal;
+			ProjectedPoint.OutVal = Project(SplinePoint.OutVal);
+			// TODO: correct tangent if it becomes null
+			ProjectedPoint.ArriveTangent = Project(SplinePoint.ArriveTangent).GetSafeNormal();
+			ProjectedPoint.LeaveTangent = Project(SplinePoint.LeaveTangent).GetSafeNormal();
+			ProjectedPoint.InterpMode = SplinePoint.InterpMode;
+		}
+	}
+}
+
+float UPCGSplineProjectionData::GetDensityAtPosition(const FVector& InPosition) const
+{
+	// Find nearest point on projected spline
+	const USplineComponent* Spline = GetSpline()->Spline.Get();
+	const FVector& SurfaceNormal = GetSurface()->GetNormal();
+
+	// Project input point to 2D space
+	FVector LocalPosition = Spline->GetComponentTransform().InverseTransformPosition(InPosition);
+	FVector2D LocalPosition2D = Project(LocalPosition);
+	float Dummy;
+	// Find nearest key on 2D spline	
+	float NearestInputKey = ProjectedPosition.InaccurateFindNearest(LocalPosition2D, Dummy);
+	// TODO: if we didn't want to hand off density computation to the spline and do it here instead, we could do it in 2D space.
+	// Find point on original spline using the previously found key
+	// Note: this is an approximation that might not hold true since we are changing the curve length
+	const FVector NearestPointOnSpline = Spline->GetLocationAtSplineInputKey(NearestInputKey, ESplineCoordinateSpace::World);
+	const FVector PointOnLine = FMath::ClosestPointOnInfiniteLine(InPosition, InPosition + SurfaceNormal, NearestPointOnSpline);
+
+	return GetSpline()->GetDensityAtPosition(PointOnLine);
+}
+
+const UPCGSplineData* UPCGSplineProjectionData::GetSpline() const
+{
+	return Cast<const UPCGSplineData>(Source);
+}
+
+const UPCGSpatialData* UPCGSplineProjectionData::GetSurface() const
+{
+	return Target;
+}
\ No newline at end of file
diff --git a/Engine/Plugins/Experimental/PCG/Source/PCG/Private/Data/PCGSurfaceData.cpp b/Engine/Plugins/Experimental/PCG/Source/PCG/Private/Data/PCGSurfaceData.cpp
new file mode 100644
--- /dev/null
+++ b/Engine/Plugins/Experimental/PCG/Source/PCG/Private/Data/PCGSurfaceData.cpp
@@ -0,0 +1,23 @@
+// Copyright Epic Games, Inc. All Rights Reserved.
+
+#include "Data/PCGSurfaceData.h"
+
+FVector UPCGSurfaceData::TransformPosition(const FVector& InPosition) const
+{
+	return Transform.InverseTransformPosition(InPosition);
+}
+
+FPCGPoint UPCGSurfaceData::TransformPoint(const FPCGPoint& InPoint) const
+{
+	FPCGPoint Point = InPoint;
+
+	// Update point location: put it on the surface plane
+	FVector PointPositionInLocalSpace = TransformPosition(InPoint.Transform.GetLocation());
+	PointPositionInLocalSpace.Z = 0;
+	Point.Transform.SetLocation(Transform.TransformPosition(PointPositionInLocalSpace));
+
+	// Update density
+	Point.Density *= GetDensityAtPosition(InPoint.Transform.GetLocation());
+
+	return Point;
+}
\ No newline at end of file
diff --git a/Engine/Plugins/Experimental/PCG/Source/PCG/Private/Data/PCGTextureData.cpp b/Engine/Plugins/Experimental/PCG/Source/PCG/Private/Data/PCGTextureData.cpp
new file mode 100644
--- /dev/null
+++ b/Engine/Plugins/Experimental/PCG/Source/PCG/Private/Data/PCGTextureData.cpp
@@ -0,0 +1,160 @@
+// Copyright Epic Games, Inc. All Rights Reserved.
+
+#include "Data/PCGTextureData.h"
+
+#include "PCGHelpers.h"
+#include "Data/PCGPointData.h"
+
+namespace PCGTextureSampling
+{
+	template<typename ValueType>
+	ValueType Sample(const FVector2D& InPosition, const FBox2D& InSurface, int32 Width, int32 Height, TFunctionRef<ValueType(int32 Index)> SamplingFunction)
+	{
+		// TODO: There seems to be a bias issue here, as the bounds size are not in the same space as the texels.
+		// Implementation note: Supports only stretch fit
+		FVector2D Pos = (InPosition - InSurface.Min) * FVector2D(Width, Height) / InSurface.GetSize();
+
+		// TODO: this isn't super robust, if that becomes an issue
+		int32 X0 = FMath::FloorToInt(Pos.X);
+		if (X0 >= Width)
+		{
+			X0 = 0;
+		}
+
+		int32 X1 = FMath::CeilToInt(Pos.X);
+		if (X1 >= Width)
+		{
+			X1 = 0;
+		}
+
+		int32 Y0 = FMath::FloorToInt(Pos.Y);
+		if (Y0 >= Height)
+		{
+			Y0 = 0;
+		}
+
+		int32 Y1 = FMath::CeilToInt(Pos.Y);
+		if (Y1 >= Height)
+		{
+			Y1 = 0;
+		}
+
+		ValueType SampleX0Y0 = SamplingFunction(X0 + Y0 * Width);
+		ValueType SampleX1Y0 = SamplingFunction(X1 + Y0 * Width);
+		ValueType SampleX0Y1 = SamplingFunction(X0 + Y1 * Width);
+		ValueType SampleX1Y1 = SamplingFunction(X1 + Y1 * Width);
+
+		ValueType BilinearInterpolation = FMath::BiLerp(SampleX0Y0, SampleX1Y0, SampleX0Y1, SampleX1Y1, Pos.X - X0, Pos.Y - Y0);
+		return BilinearInterpolation;
+	}
+
+	float SampleFloatChannel(const FLinearColor& InColor, EPCGTextureColorChannel ColorChannel)
+	{
+		switch (ColorChannel)
+		{
+		case EPCGTextureColorChannel::Red:
+			return InColor.R;
+		case EPCGTextureColorChannel::Green:
+			return InColor.G;
+		case EPCGTextureColorChannel::Blue:
+			return InColor.B;
+		case EPCGTextureColorChannel::Alpha:
+		default:
+			return InColor.A;
+		}
+	}
+}
+
+FBox UPCGBaseTextureData::GetBounds() const
+{
+	return Bounds;
+}
+
+FBox UPCGBaseTextureData::GetStrictBounds() const
+{
+	return Bounds;
+}
+
+float UPCGBaseTextureData::GetDensityAtPosition(const FVector& InPosition) const
+{
+	FVector LocalPosition = Transform.InverseTransformPosition(InPosition);
+	FVector2D Position2D(LocalPosition.X, LocalPosition.Y);
+
+	FBox2D Surface(FVector2D(-1.0f, -1.0f), FVector2D(1.0f, 1.0f));
+
+	return PCGTextureSampling::Sample<float>(Position2D, Surface, Width, Height, [this](int32 Index) { return PCGTextureSampling::SampleFloatChannel(ColorData[Index], ColorChannel); });
+}
+
+const UPCGPointData* UPCGBaseTextureData::CreatePointData() const
+{
+	TRACE_CPUPROFILER_EVENT_SCOPE(UPCGBaseTextureData::CreatePointData);
+	// TODO: this is a trivial implementation
+	// A better sampler would allow to sample a fixed number of points in either direction
+	// or based on a given texel size
+	FBox2D LocalSurfaceBounds(FVector2D(-1.0f, -1.0f), FVector2D(1.0f, 1.0f));
+
+	UPCGPointData* Data = NewObject<UPCGPointData>(const_cast<UPCGBaseTextureData*>(this));
+	Data->TargetActor = TargetActor;
+	TArray<FPCGPoint>& Points = Data->GetMutablePoints();
+
+	// TODO: There's a bias issue here where we should correct by a 0.5 unit...
+	Points.Reserve(Width * Height);
+	const FVector::FReal XScale = LocalSurfaceBounds.GetSize().X / Width;
+	const FVector::FReal YScale = LocalSurfaceBounds.GetSize().Y / Height;
+	const FVector2D Bias = LocalSurfaceBounds.Min;
+
+	for (int32 X = 0; X < Width; ++X)
+	{
+		for (int32 Y = 0; Y < Height; ++Y)
+		{
+			const float Density = PCGTextureSampling::SampleFloatChannel(ColorData[X + Y * Width], ColorChannel);
+
+			if (Density > 0)
+			{
+				FVector LocalPosition(X * XScale + Bias.X, Y * YScale + Bias.Y, 0);
+				Points.Emplace(FTransform(Transform.TransformPosition(LocalPosition)),
+					Density,
+					PCGHelpers::ComputeSeed(X, Y));
+			}
+		}
+	}
+
+	return Data;
+}
+
+void UPCGTextureData::Initialize(UTexture2D* InTexture, const FTransform& InTransform)
+{
+	Texture = InTexture;
+	Transform = InTransform;
+
+	if (Texture)
+	{
+#if WITH_EDITORONLY_DATA
+		if (Texture->GetPlatformData()->Mips.Num() > 0)
+		{
+			Width = Texture->GetSizeX();
+			Height = Texture->GetSizeY();
+
+			UTexture2D* TempTexture2D = DuplicateObject<UTexture2D>(InTexture, GetTransientPackage());
+			TempTexture2D->CompressionSettings = TextureCompressionSettings::TC_VectorDisplacementmap;
+			TempTexture2D->MipGenSettings = TextureMipGenSettings::TMGS_NoMipmaps;
+			TempTexture2D->SRGB = false;
+			TempTexture2D->UpdateResource();
+			const FColor* FormattedImageData = static_cast<const FColor*>(TempTexture2D->GetPlatformData()->Mips[0].BulkData.LockReadOnly());
+
+			ColorData.SetNum(Width * Height);
+			for (int32 D = 0; D < Width * Height; ++D)
+			{
+				ColorData[D] = FormattedImageData[D].ReinterpretAsLinear();
+			}
+
+			TempTexture2D->GetPlatformData()->Mips[0].BulkData.Unlock();
+		}
+#endif
+	}
+
+	Bounds = FBox(EForceInit::ForceInit);
+	Bounds += FVector(-1.0f, -1.0f, 0.0f);
+	Bounds += FVector(1.0f, 1.0f, 0.0f);
+	Bounds = Bounds.TransformBy(Transform);
+}
\ No newline at end of file
diff --git a/Engine/Plugins/Experimental/PCG/Source/PCG/Private/Data/PCGUnionData.cpp b/Engine/Plugins/Experimental/PCG/Source/PCG/Private/Data/PCGUnionData.cpp
new file mode 100644
--- /dev/null
+++ b/Engine/Plugins/Experimental/PCG/Source/PCG/Private/Data/PCGUnionData.cpp
@@ -0,0 +1,185 @@
+// Copyright Epic Games, Inc. All Rights Reserved.
+
+#include "Data/PCGUnionData.h"
+#include "Data/PCGPointData.h"
+
+void UPCGUnionData::Initialize(const UPCGSpatialData* InA, const UPCGSpatialData* InB)
+{
+	check(InA && InB);
+	check(InA->GetDimension() == InB->GetDimension());
+	AddData(InA);
+	AddData(InB);
+}
+
+void UPCGUnionData::AddData(const UPCGSpatialData* InC)
+{
+	check(InC);
+	check(!A || InC->GetDimension() == A->GetDimension());
+
+	// Early tests: if we don't have a full union yet, just push in the data.
+	if (!A)
+	{
+		check(!B);
+		A = InC;
+		TargetActor = InC->TargetActor;
+		CachedBounds = InC->GetBounds();
+		CachedStrictBounds = InC->GetStrictBounds();
+	}
+	else if (!B)
+	{
+		check(A);
+		B = InC;
+		CachedBounds += InC->GetBounds();
+		CachedStrictBounds = A->GetStrictBounds().Overlap(InC->GetStrictBounds());
+	}
+	else
+	{
+		// TODO: Bad usage of this call could lead to sub-unions having inefficient bounds
+		// We can revisit and keep the list of members and bundle them appropriately on demand 
+		const bool bAIsUnion = Cast<UPCGUnionData>(A) != nullptr;
+		const bool bBIsUnion = Cast<UPCGUnionData>(B) != nullptr;
+
+		UPCGUnionData* Union = NewObject<UPCGUnionData>(this);
+		Union->UnionType = UnionType;
+
+		if (!bAIsUnion || bBIsUnion)
+		{
+			Union->Initialize(A, B);
+			A = Union;
+			B = InC;
+		}
+		else
+		{
+			Union->Initialize(B, InC);
+			B = Union;
+		}
+
+		// Update bounds
+		CachedBounds += InC->GetBounds();
+		CachedStrictBounds = CachedStrictBounds.Overlap(InC->GetStrictBounds());
+	}
+}
+
+int UPCGUnionData::GetDimension() const
+{
+	check(A && B && A->GetDimension() == B->GetDimension());
+	return A->GetDimension();
+}
+
+FBox UPCGUnionData::GetBounds() const
+{
+	check(A && B);
+	return CachedBounds;
+}
+
+FBox UPCGUnionData::GetStrictBounds() const
+{
+	check(A && B);
+	return CachedStrictBounds;
+}
+
+float UPCGUnionData::GetDensityAtPosition(const FVector& InPosition) const
+{
+	check(A && B);
+	if (!CachedBounds.IsInside(InPosition))
+	{
+		return 0;
+	}
+	else if (CachedStrictBounds.IsInside(InPosition) ||
+		A->GetStrictBounds().IsInside(InPosition) ||
+		B->GetStrictBounds().IsInside(InPosition))
+	{
+		return 1.0f;
+	}
+	else
+	{
+		return FMath::Max(A->GetDensityAtPosition(InPosition), B->GetDensityAtPosition(InPosition));
+	}
+}
+
+const UPCGPointData* UPCGUnionData::CreatePointData() const
+{
+	TRACE_CPUPROFILER_EVENT_SCOPE(UPCGUnionData::CreatePointData);
+	check(A && B);
+
+	const UPCGPointData* APointData = A->ToPointData();
+	const UPCGPointData* BPointData = B->ToPointData();
+
+	if (!APointData || !BPointData)
+	{
+		return APointData ? APointData : BPointData;
+	}
+
+	UPCGPointData* Data = NewObject<UPCGPointData>(const_cast<UPCGUnionData*>(this));
+	Data->TargetActor = TargetActor;
+
+	TArray<FPCGPoint>& TargetPoints = Data->GetMutablePoints();
+
+	auto AddNonExcludedPoints = [&TargetPoints](const UPCGPointData* PointsToAdd, const UPCGSpatialData* Exclusion) {
+		check(PointsToAdd && Exclusion);
+		for (const FPCGPoint& Point : PointsToAdd->GetPoints())
+		{
+			if (Exclusion->GetDensityAtPosition(Point.Transform.GetLocation()) == 0)
+			{
+				TargetPoints.Add(Point);
+			}
+		}
+	};
+
+	switch (UnionType)
+	{
+	case EPCGUnionType::LeftToRightPriority:
+	case EPCGUnionType::RightToLeftPriority:
+	default:
+		{
+			const UPCGSpatialData* PrioritizedData = (UnionType != EPCGUnionType::RightToLeftPriority ? A : B);
+			const UPCGPointData* PrioritizedPointData = (UnionType != EPCGUnionType::RightToLeftPriority ? APointData : BPointData);
+			const UPCGSpatialData* SecondaryData = (UnionType != EPCGUnionType::RightToLeftPriority ? B : A);
+			const UPCGPointData* SecondaryPointData = (UnionType != EPCGUnionType::RightToLeftPriority ? BPointData : APointData);
+
+			// First: add all points from the prioritized data and update its density
+			TargetPoints.Append(PrioritizedPointData->GetPoints());
+
+			for (FPCGPoint& Point : TargetPoints)
+			{
+				Point.Density = FMath::Max(Point.Density, SecondaryData->GetDensityAtPosition(Point.Transform.GetLocation()));
+			}
+
+			// Second: all points from the secondary data set only if they have a null density in the prioritized data set
+			AddNonExcludedPoints(SecondaryPointData, PrioritizedData);
+		}
+
+		break;
+
+	case EPCGUnionType::KeepAB:
+		{
+			TargetPoints.Append(APointData->GetPoints());
+			TargetPoints.Append(BPointData->GetPoints());
+		}
+
+		break;
+
+	case EPCGUnionType::ABMinusIntersection:
+		{
+			// First, build the intersection of A & B
+			if (CachedStrictBounds.IsValid)
+			{
+				UPCGSpatialData* IntersectionData = A->IntersectWith(B);
+				const UPCGPointData* IntersectionPointData = IntersectionData ? IntersectionData->ToPointData() : nullptr;
+
+				if (IntersectionPointData)
+				{
+					TargetPoints.Append(IntersectionPointData->GetPoints());
+				}
+			}
+
+			// Second: all all points from A & B that have no match in their counterpart.
+			AddNonExcludedPoints(APointData, B);
+			AddNonExcludedPoints(BPointData, A);
+		}
+
+		break;
+	}
+
+	return Data;
+}
diff --git a/Engine/Plugins/Experimental/PCG/Source/PCG/Private/Data/PCGVolumeData.cpp b/Engine/Plugins/Experimental/PCG/Source/PCG/Private/Data/PCGVolumeData.cpp
new file mode 100644
--- /dev/null
+++ b/Engine/Plugins/Experimental/PCG/Source/PCG/Private/Data/PCGVolumeData.cpp
@@ -0,0 +1,69 @@
+// Copyright Epic Games, Inc. All Rights Reserved.
+
+#include "Data/PCGVolumeData.h"
+
+#include "Components/BrushComponent.h"
+#include "GameFramework/Volume.h"
+
+void UPCGVolumeData::Initialize(AVolume* InVolume, AActor* InTargetActor)
+{
+	check(InVolume);
+	Volume = InVolume;
+	TargetActor = InTargetActor ? InTargetActor : InVolume;
+	
+	FBoxSphereBounds BoxSphereBounds = Volume->GetBounds();
+	Bounds = FBox::BuildAABB(BoxSphereBounds.Origin, BoxSphereBounds.BoxExtent);
+
+	// TODO: Compute the strict bounds, we must find a FBox inscribed into the oriented box.
+	// Currently, we'll leave the strict bounds empty and fall back to checking against the local box
+	if (UBrushComponent* BrushComponent = Volume->GetBrushComponent())
+	{
+		VolumeTransform = BrushComponent->GetComponentTransform();
+		FBoxSphereBounds LocalBoxSphereBounds = BrushComponent->CalcBounds(FTransform::Identity);
+		VolumeLocalBounds = FBox::BuildAABB(LocalBoxSphereBounds.Origin, LocalBoxSphereBounds.BoxExtent);
+	}
+}
+
+void UPCGVolumeData::Initialize(const FBox& InBounds, AActor* InTargetActor)
+{
+	Bounds = InBounds;
+	StrictBounds = InBounds;
+	TargetActor = InTargetActor;
+}
+
+FBox UPCGVolumeData::GetBounds() const
+{
+	return Bounds;
+}
+
+FBox UPCGVolumeData::GetStrictBounds() const
+{
+	return StrictBounds;
+}
+
+const UPCGPointData* UPCGVolumeData::CreatePointData() const
+{
+	// TODO
+	return nullptr;
+}
+
+float UPCGVolumeData::GetDensityAtPosition(const FVector& InPosition) const
+{
+	// TODO: support fall-off between 0-1 density as needed
+	// could use FBox::GetClosestPointTo
+	if (GetBounds().IsInside(InPosition))
+	{
+		if(!Volume || GetStrictBounds().IsInside(InPosition))
+		{
+			return 1.0f;
+		}
+		else
+		{
+			return VolumeLocalBounds.IsInside(VolumeTransform.InverseTransformPosition(InPosition)) ? 1.0f : 0.0f;
+		}
+	}
+	else
+	{
+		return 0.0f;
+	}
+}
\ No newline at end of file
diff --git a/Engine/Plugins/Experimental/PCG/Source/PCG/Private/Elements/PCGExecuteBlueprint.cpp b/Engine/Plugins/Experimental/PCG/Source/PCG/Private/Elements/PCGExecuteBlueprint.cpp
new file mode 100644
--- /dev/null
+++ b/Engine/Plugins/Experimental/PCG/Source/PCG/Private/Elements/PCGExecuteBlueprint.cpp
@@ -0,0 +1,82 @@
+// Copyright Epic Games, Inc. All Rights Reserved.
+
+#include "Elements/PCGExecuteBlueprint.h"
+
+#if WITH_EDITOR
+#include "Engine/World.h"
+#endif
+
+UWorld* UPCGBlueprintElement::GetWorld() const
+{
+#if WITH_EDITOR
+	return GWorld;
+#else
+	return nullptr;
+#endif
+}
+
+#if WITH_EDITOR
+void UPCGBlueprintSettings::PostLoad()
+{
+	Super::PostLoad();
+
+	if (BlueprintElement_DEPRECATED && !BlueprintElementType)
+	{
+		BlueprintElementType = BlueprintElement_DEPRECATED;
+		BlueprintElement_DEPRECATED = nullptr;
+	}
+
+	if (!BlueprintElementInstance)
+	{
+		RefreshBlueprintElement();
+	}
+}
+
+void UPCGBlueprintSettings::PostEditChangeProperty(FPropertyChangedEvent& PropertyChangedEvent)
+{
+	if (!BlueprintElementInstance || BlueprintElementInstance->GetClass() != BlueprintElementType)
+	{
+		RefreshBlueprintElement();
+	}
+}
+#endif
+
+void UPCGBlueprintSettings::SetElementType(TSubclassOf<UPCGBlueprintElement> InElementType)
+{
+	if (!BlueprintElementInstance || InElementType != BlueprintElementType)
+	{
+		BlueprintElementType = InElementType;
+		RefreshBlueprintElement();
+	}
+}
+
+void UPCGBlueprintSettings::RefreshBlueprintElement()
+{
+	if (BlueprintElementType)
+	{
+		BlueprintElementInstance = NewObject<UPCGBlueprintElement>(this, BlueprintElementType);
+	}
+	else
+	{
+		BlueprintElementInstance = nullptr;
+	}	
+}
+
+FPCGElementPtr UPCGBlueprintSettings::CreateElement() const
+{
+	return MakeShared<FPCGExecuteBlueprintElement>();
+}
+
+bool FPCGExecuteBlueprintElement::Execute(FPCGContextPtr Context) const
+{
+	const UPCGBlueprintSettings* Settings = Context->GetInputSettings<UPCGBlueprintSettings>();
+	check(Settings);
+	check(Settings->BlueprintElementInstance);
+
+	if (Settings && Settings->BlueprintElementInstance)
+	{
+		Settings->BlueprintElementInstance->Execute(Context->InputData, Context->OutputData);
+	}
+	
+	return true;
+}
\ No newline at end of file
diff --git a/Engine/Plugins/Experimental/PCG/Source/PCG/Private/Elements/PCGPointSampler.cpp b/Engine/Plugins/Experimental/PCG/Source/PCG/Private/Elements/PCGPointSampler.cpp
new file mode 100644
--- /dev/null
+++ b/Engine/Plugins/Experimental/PCG/Source/PCG/Private/Elements/PCGPointSampler.cpp
@@ -0,0 +1,118 @@
+// Copyright Epic Games, Inc. All Rights Reserved.
+
+#include "Elements/PCGPointSampler.h"
+#include "PCGHelpers.h"
+#include "Data/PCGPointData.h"
+#include "Math/RandomStream.h"
+
+FPCGElementPtr UPCGPointSamplerSettings::CreateElement() const
+{
+	return MakeShared<FPCGPointSamplerElement>();
+}
+
+bool FPCGPointSamplerElement::Execute(FPCGContextPtr Context) const
+{
+	TRACE_CPUPROFILER_EVENT_SCOPE(FPCGPointSamplerElement::Execute);
+	// TODO: make time-sliced implementation
+	const UPCGPointSamplerSettings* Settings = Context->GetInputSettings<UPCGPointSamplerSettings>();
+	check(Settings);
+
+	TArray<FPCGTaggedData> Inputs = Context->InputData.GetInputs();
+	TArray<FPCGTaggedData> Exclusions = Context->InputData.GetExclusions();
+	TArray<FPCGTaggedData>& Outputs = Context->OutputData.TaggedData;
+
+	// TODO: embarassingly parallel loop
+	for (const FPCGTaggedData& Input : Inputs)
+	{
+		FPCGTaggedData& Output = Outputs.Emplace_GetRef();
+		Output = Input;
+
+		if (!Input.Data || Cast<UPCGSpatialData>(Input.Data) == nullptr)
+		{
+			continue;
+		}
+
+		const UPCGPointData* OriginalData = Cast<UPCGSpatialData>(Input.Data)->ToPointData();
+
+		if (!OriginalData)
+		{
+			continue;
+		}
+
+		const TArray<FPCGPoint>& Points = OriginalData->GetPoints();
+		const int OriginalPointCount = Points.Num();
+
+		if(Settings->Ratio >= 1.0f && Exclusions.IsEmpty())
+		{
+			Output.Data = OriginalData;
+		}
+		else
+		{
+			UPCGPointData* SampledData = NewObject<UPCGPointData>();
+			SampledData->TargetActor = OriginalData->TargetActor;
+
+			Output.Data = SampledData;
+
+			// TODO: randomize on the fractional number of points
+			int TargetNumPoints = OriginalPointCount * Settings->Ratio;
+
+			// Early out
+			if (TargetNumPoints == 0)
+			{
+				return true;
+			}
+
+			TArray<int> SelectedIndices;
+			// Build indices in a deterministic manner
+			{
+				TRACE_CPUPROFILER_EVENT_SCOPE(UPCGComponent::CreatePCGData::SelectIndices)
+
+				// Approximate upper bound
+				SelectedIndices.Reserve(4 * TargetNumPoints / 3);
+
+				for (int Index = 0; Index < OriginalPointCount; ++Index)
+				{
+					const FPCGPoint& Point = Points[Index];
+
+					// Apply a high-pass filter based on selected ratio
+					FRandomStream RandomSource(PCGHelpers::ComputeSeed(Settings->Seed, Point.Seed));
+					float Chance = RandomSource.FRand();
+
+					if (Chance < Settings->Ratio)
+					{
+						// Remap density to a value between 0.5 and 1.0
+						float Density = Point.Density * (0.5f + 0.5f * (Settings->Ratio - Chance));
+
+						for (const FPCGTaggedData& Exclusion : Exclusions)
+						{
+							if (!Exclusion.Data)
+							{
+								continue;
+							}
+
+							Density *= (1.0f - Cast<UPCGSpatialData>(Exclusion.Data)->GetDensityAtPosition(Point.Transform.GetLocation()));
+
+							if (Density == 0)
+							{
+								break;
+							}
+						}
+
+						// Implicit threshold at 0.5f based on cumulated density
+						if(Density >= 0.5f)
+						{
+							SelectedIndices.Add(Index);
+						}
+					}
+				}
+			}
+
+			// Copy points using the indices
+			// TODO: pass thread info for parallel for
+			SampledData->CopyPointsFrom(OriginalData, SelectedIndices);
+		}
+	}
+
+	return true;
+}
+
diff --git a/Engine/Plugins/Experimental/PCG/Source/PCG/Private/Elements/PCGStaticMeshSpawner.cpp b/Engine/Plugins/Experimental/PCG/Source/PCG/Private/Elements/PCGStaticMeshSpawner.cpp
new file mode 100644
--- /dev/null
+++ b/Engine/Plugins/Experimental/PCG/Source/PCG/Private/Elements/PCGStaticMeshSpawner.cpp
@@ -0,0 +1,167 @@
+// Copyright Epic Games, Inc. All Rights Reserved.
+
+#include "Elements/PCGStaticMeshSpawner.h"
+
+#include "PCGComponent.h"
+#include "PCGHelpers.h"
+#include "Data/PCGPointData.h"
+#include "Data/PCGSpatialData.h"
+
+#include "Components/InstancedStaticMeshComponent.h"
+#include "Components/HierarchicalInstancedStaticMeshComponent.h"
+#include "Engine/StaticMesh.h"
+#include "Math/RandomStream.h"
+
+FPCGElementPtr UPCGStaticMeshSpawnerSettings::CreateElement() const
+{
+	return MakeShared<FPCGStaticMeshSpawnerElement>();
+}
+
+struct FWeightedMeshAndInstances
+{
+	int Weight = 0;
+	UStaticMesh* Mesh = nullptr;
+	TArray<FTransform> Instances;
+};
+
+bool FPCGStaticMeshSpawnerElement::Execute(FPCGContextPtr Context) const
+{
+	TRACE_CPUPROFILER_EVENT_SCOPE(FPCGStaticMeshSpawnerElement::Execute);
+	// TODO : time-sliced implementation
+	const UPCGStaticMeshSpawnerSettings* Settings = Context->GetInputSettings<UPCGStaticMeshSpawnerSettings>();
+	check(Settings);
+
+	// TODO : see if we could do async mesh loading here
+	TArray<FWeightedMeshAndInstances> WeightedEntries;
+	int TotalWeight = 0;
+	for (const FPCGStaticMeshSpawnerEntry& Entry : Settings->Meshes)
+	{
+		if (Entry.Weight <= 0)
+			continue;
+
+		TotalWeight += Entry.Weight;
+
+		FWeightedMeshAndInstances& WeightedEntry = WeightedEntries.Emplace_GetRef();
+		WeightedEntry.Weight = TotalWeight;
+		WeightedEntry.Mesh = Entry.Mesh.LoadSynchronous();
+	}
+
+	if (TotalWeight <= 0)
+	{
+		return true;
+	}
+
+	TArray<FPCGTaggedData> Inputs = Context->InputData.GetInputs();
+
+	for (const FPCGTaggedData& Input : Inputs)
+	{
+		const UPCGSpatialData* SpatialData = Cast<UPCGSpatialData>(Input.Data);
+
+		if (!SpatialData)
+		{
+			// Data type mismatch
+			continue;
+		}
+
+		AActor* TargetActor = SpatialData->TargetActor;
+
+		if (!TargetActor)
+		{
+			// No target actor
+			continue;
+		}
+
+		// First, create target instance transforms
+		const UPCGPointData* PointData = SpatialData->ToPointData();
+
+		if (!PointData)
+		{
+			continue;
+		}
+
+		// Assign points to entries
+		{
+			TRACE_CPUPROFILER_EVENT_SCOPE(FPCGStaticMeshSpawnerElement::Execute::SelectEntries);
+			const TArray<FPCGPoint>& Points = PointData->GetPoints();
+
+			for (const FPCGPoint& Point : Points)
+			{
+				FRandomStream RandomSource(PCGHelpers::ComputeSeed(Point.Seed, Settings->Seed));
+				int RandomWeightedPick = RandomSource.RandRange(0, TotalWeight - 1);
+
+				int RandomPick = 0;
+				while (RandomPick < WeightedEntries.Num() && WeightedEntries[RandomPick].Weight <= RandomWeightedPick)
+				{
+					++RandomPick;
+				}
+
+				if (RandomPick < WeightedEntries.Num())
+				{
+					WeightedEntries[RandomPick].Instances.Emplace(Point.Transform);
+				}
+			}
+		}
+
+		// Populate the (H)ISM from the previously prepared entries
+		{
+			TRACE_CPUPROFILER_EVENT_SCOPE(FPCGStaticMeshSpawnerElement::Execute::PopulateISMs);
+
+			// Second, add the instances
+			for (FWeightedMeshAndInstances& Instances : WeightedEntries)
+			{
+				if (!Instances.Mesh || Instances.Instances.Num() == 0)
+				{
+					continue;
+				}
+
+				UInstancedStaticMeshComponent* ISMC = GetOrCreateISMC(TargetActor, Context->SourceComponent, Instances.Mesh);
+
+				// TODO: add scaling
+				// TODO: document these arguments
+				ISMC->AddInstances(Instances.Instances, false, true);
+				ISMC->UpdateBounds();
+			}
+		}
+	}
+
+	return true;
+}
+
+UInstancedStaticMeshComponent* FPCGStaticMeshSpawnerElement::GetOrCreateISMC(AActor* InTargetActor, const UPCGComponent* InSourceComponent, UStaticMesh* InMesh) const
+{
+	check(InTargetActor != nullptr && InMesh != nullptr);
+
+	TArray<UInstancedStaticMeshComponent*> ISMCs;
+	InTargetActor->GetComponents<UInstancedStaticMeshComponent>(ISMCs);
+
+	for (UInstancedStaticMeshComponent* ISMC : ISMCs)
+	{
+		if (ISMC->GetStaticMesh() == InMesh && (!InSourceComponent || ISMC->ComponentTags.Contains(InSourceComponent->GetFName())))
+		{
+			return ISMC;
+		}
+	}
+
+	InTargetActor->Modify();
+
+	// Otherwise, create a new component
+	// TODO: use static mesh component if there's only one instance
+	// TODO: add hism/ism switch or better yet, use a template component
+	UInstancedStaticMeshComponent* ISMC = NewObject<UHierarchicalInstancedStaticMeshComponent>(InTargetActor);
+	ISMC->SetStaticMesh(InMesh);
+	// TODO: add way to do material overrides, maybe on a finer basis as well
+	ISMC->RegisterComponent();
+	InTargetActor->AddInstanceComponent(ISMC);
+	ISMC->SetMobility(EComponentMobility::Static);
+	// TODO: add option for collision, or use a template
+	ISMC->SetCollisionEnabled(ECollisionEnabled::NoCollision);
+	ISMC->AttachToComponent(InTargetActor->GetRootComponent(), FAttachmentTransformRules::KeepWorldTransform);
+
+	if (InSourceComponent)
+	{
+		ISMC->ComponentTags.Add(InSourceComponent->GetFName());
+		ISMC->ComponentTags.Add(PCGHelpers::DefaultPCGTag);
+	}
+
+	return ISMC;
+}
diff --git a/Engine/Plugins/Experimental/PCG/Source/PCG/Private/Elements/PCGSurfaceSampler.cpp b/Engine/Plugins/Experimental/PCG/Source/PCG/Private/Elements/PCGSurfaceSampler.cpp
new file mode 100644
--- /dev/null
+++ b/Engine/Plugins/Experimental/PCG/Source/PCG/Private/Elements/PCGSurfaceSampler.cpp
@@ -0,0 +1,13 @@
+// Copyright Epic Games, Inc. All Rights Reserved.
+
+#include "Elements/PCGSurfaceSampler.h"
+
+FPCGElementPtr UPCGSurfaceSamplerSettings::CreateElement() const
+{
+	return MakeShared<FPCGSurfaceSamplerElement>();
+}
+
+bool FPCGSurfaceSamplerElement::Execute(FPCGContextPtr Context) const
+{
+	return true;
+}
\ No newline at end of file
diff --git a/Engine/Plugins/Experimental/PCG/Source/PCG/Private/Graph/PCGGraphCompiler.cpp b/Engine/Plugins/Experimental/PCG/Source/PCG/Private/Graph/PCGGraphCompiler.cpp
new file mode 100644
--- /dev/null
+++ b/Engine/Plugins/Experimental/PCG/Source/PCG/Private/Graph/PCGGraphCompiler.cpp
@@ -0,0 +1,295 @@
+// Copyright Epic Games, Inc. All Rights Reserved.
+
+#include "PCGGraphCompiler.h"
+#include "PCGGraph.h"
+#include "PCGSubgraph.h"
+#include "PCGSettings.h" // Needed for trivial element
+
+TArray<FPCGGraphTask> FPCGGraphCompiler::CompileGraph(UPCGGraph* InGraph, FPCGTaskId& NextId)
+{
+	check(InGraph);
+	TArray<FPCGGraphTask> CompiledTasks;
+	TMap<const UPCGNode*, FPCGTaskId> IdMapping;
+	TArray<const UPCGNode*> NodeQueue;
+
+	// Prime the node queue with all nodes that have no inbound edges
+	for (const UPCGNode* Node : InGraph->GetNodes())
+	{
+		if (Node->InboundNodes.IsEmpty())
+		{
+			NodeQueue.Add(Node);
+		}
+	}
+
+	// By definition, the input node has no inbound edge.
+	// Put it last in the queue so it gets picked up first - order is important for hooking up the fetch input element
+	NodeQueue.Add(InGraph->GetInputNode());
+
+	while (NodeQueue.Num() > 0)
+	{
+		const UPCGNode* Node = NodeQueue.Pop();
+
+		const UPCGSubgraphNode* SubgraphNode = Cast<const UPCGSubgraphNode>(Node);
+
+		if (SubgraphNode != nullptr && !SubgraphNode->bAllowDynamicGraph)
+		{
+			const FPCGTaskId PreId = NextId++;
+
+			// 1. Compile the subgraph making sure we don't reuse the same ids
+			// Note that we will not consume the pre or post-execute tasks, ergo bIsTopGraph=false
+			TArray<FPCGGraphTask> Subtasks = GetCompiledTasks(SubgraphNode->GetGraph(), /*bIsTopGraph=*/false);
+
+#if WITH_EDITOR
+			GraphDependenciesLock.Lock();
+			GraphDependencies.AddUnique(SubgraphNode->GetGraph(), InGraph);
+			GraphDependenciesLock.Unlock();
+#endif // WITH_EDITOR
+
+			OffsetNodeIds(Subtasks, NextId);
+			NextId += Subtasks.Num();
+
+			const UPCGNode* SubgraphInputNode = SubgraphNode->GetGraph()->GetInputNode();
+			const UPCGNode* SubgraphOutputNode = SubgraphNode->GetGraph()->GetOutputNode();
+
+			// 2. Update the "input" and "output" node tasks so we can add the proper dependencies
+			FPCGGraphTask* InputNodeTask = Subtasks.FindByPredicate([SubgraphInputNode](const FPCGGraphTask& Subtask) {
+				return Subtask.Node == SubgraphInputNode;
+				});
+			check(InputNodeTask);
+
+			FPCGGraphTask* OutputNodeTask = Subtasks.FindByPredicate([SubgraphOutputNode](const FPCGGraphTask& Subtask) {
+				return Subtask.Node == SubgraphOutputNode;
+				});
+			check(OutputNodeTask);
+
+			// Build pre-task
+			FPCGGraphTask& PreTask = CompiledTasks.Emplace_GetRef();
+			PreTask.Node = Node;
+			PreTask.NodeId = PreId;
+
+			for (const UPCGNode* Inbound : Node->InboundNodes)
+			{
+				PreTask.Inputs.Add(IdMapping[Inbound]);
+			}
+
+			// Add pre-task as input to subgraph input node task
+			InputNodeTask->Inputs.Add(PreId);
+
+			// Merge subgraph tasks into current tasks.
+			CompiledTasks.Append(Subtasks);
+
+			// Build post-task
+			const FPCGTaskId PostId = NextId++;
+			FPCGGraphTask& PostTask = CompiledTasks.Emplace_GetRef();
+			PostTask.Node = Node;
+			PostTask.NodeId = PostId;
+
+			// Add subgraph output node task as input to the post-task
+			PostTask.Inputs.Add(OutputNodeTask->NodeId);
+
+			IdMapping.Add(Node, PostId);
+		}
+		else
+		{
+			const FPCGTaskId NodeId = NextId++;
+			FPCGGraphTask& Task = CompiledTasks.Emplace_GetRef();
+			Task.Node = Node;
+			Task.NodeId = NodeId;
+
+			for (const UPCGNode* Inbound : Node->InboundNodes)
+			{
+				Task.Inputs.Add(IdMapping[Inbound]);
+			}
+
+			IdMapping.Add(Node, NodeId);
+		}
+
+		// Push next ready nodes on the queue
+		for (const UPCGNode* Outbound : Node->OutboundNodes)
+		{
+			bool bAllPrerequisitesMet = true;
+			for (const UPCGNode* Inbound : Outbound->InboundNodes)
+			{
+				bAllPrerequisitesMet &= IdMapping.Contains(Inbound);
+			}
+
+			if (bAllPrerequisitesMet)
+			{
+				NodeQueue.Add(Outbound);
+			}
+		}
+	}
+
+	return CompiledTasks;
+}
+
+void FPCGGraphCompiler::Compile(UPCGGraph* InGraph)
+{
+	GraphToTaskMapLock.ReadLock();
+	bool bAlreadyCached = GraphToTaskMap.Contains(InGraph);
+	GraphToTaskMapLock.ReadUnlock();
+
+	if (bAlreadyCached)
+	{
+		return;
+	}
+
+	// Otherwise, do the compilation; note that we always start at zero since
+	// the caller will offset the ids as needed
+	FPCGTaskId FirstId = 0;
+	TArray<FPCGGraphTask> CompiledTasks = CompileGraph(InGraph, FirstId);
+
+	// TODO: optimize no-ops, etc.
+
+	// Store back the results in the cache
+	GraphToTaskMapLock.WriteLock();
+	if (!GraphToTaskMap.Contains(InGraph))
+	{
+		GraphToTaskMap.Add(InGraph, MoveTemp(CompiledTasks));
+	}
+	GraphToTaskMapLock.WriteUnlock();
+}
+
+TArray<FPCGGraphTask> FPCGGraphCompiler::GetCompiledTasks(UPCGGraph* InGraph, bool bIsTopGraph)
+{
+	TArray<FPCGGraphTask> CompiledTasks;
+
+#if WITH_EDITOR
+	if (bIsTopGraph)
+	{
+		// Always try to compile
+		CompileTopGraph(InGraph);
+
+		// Get compiled tasks in a threadsafe way
+		GraphToTaskMapLock.ReadLock();
+		check(TopGraphToTaskMap.Contains(InGraph));
+		CompiledTasks = TopGraphToTaskMap[InGraph];
+		GraphToTaskMapLock.ReadUnlock();
+	}
+	else
+#endif
+	{
+		// Always try to compile
+		Compile(InGraph);
+
+		// Get compiled tasks in a threadsafe way
+		GraphToTaskMapLock.ReadLock();
+		check(GraphToTaskMap.Contains(InGraph));
+		CompiledTasks = GraphToTaskMap[InGraph];
+		GraphToTaskMapLock.ReadUnlock();
+	}
+
+	return CompiledTasks;
+}
+
+void FPCGGraphCompiler::OffsetNodeIds(TArray<FPCGGraphTask>& Tasks, FPCGTaskId Offset)
+{
+	for (FPCGGraphTask& Task : Tasks)
+	{
+		Task.NodeId += Offset;
+
+		for (FPCGTaskId& Input : Task.Inputs)
+		{
+			Input += Offset;
+		}
+	}
+}
+
+#if WITH_EDITOR
+void FPCGGraphCompiler::CompileTopGraph(UPCGGraph* InGraph)
+{
+	GraphToTaskMapLock.ReadLock();
+	bool bAlreadyCached = TopGraphToTaskMap.Contains(InGraph);
+	GraphToTaskMapLock.ReadUnlock();
+
+	if (bAlreadyCached)
+	{
+		return;
+	}
+
+	// Build from non-top tasks
+	TArray<FPCGGraphTask> CompiledTasks = GetCompiledTasks(InGraph, /*bIsTopGraph=*/false);
+
+	const int TaskNum = CompiledTasks.Num();
+	const FPCGTaskId PreExecuteTaskId = FPCGTaskId(TaskNum);
+	const FPCGTaskId PostExecuteTaskId = PreExecuteTaskId + 1;
+
+	FPCGGraphTask& PreExecuteTask = CompiledTasks.Emplace_GetRef();
+	PreExecuteTask.Element = MakeShared<FPCGTrivialElement>();
+	PreExecuteTask.NodeId = PreExecuteTaskId;
+
+	for (int TaskIndex = 0; TaskIndex < TaskNum; ++TaskIndex)
+	{
+		FPCGGraphTask& Task = CompiledTasks[TaskIndex];
+		if (Task.Inputs.IsEmpty())
+		{
+			Task.Inputs.Add(PreExecuteTaskId);
+		}
+	}
+
+	FPCGGraphTask& PostExecuteTask = CompiledTasks.Emplace_GetRef();
+	PostExecuteTask.Element = MakeShared<FPCGTrivialElement>();
+	PostExecuteTask.NodeId = PostExecuteTaskId;
+
+	// Find end nodes, e.g. all nodes that have no successors.
+	// In our representation we don't have this, so find it out by going backwards
+	// Note: this works because there is a weak ordering on the tasks such that
+	// a successor task is always after if predecessor
+	TSet<FPCGTaskId> TasksWithSuccessors;
+	TArray<FPCGTaskId> EndNodes;
+
+	for (int TaskIndex = TaskNum - 1; TaskIndex >= 0; --TaskIndex)
+	{
+		const FPCGGraphTask& Task = CompiledTasks[TaskIndex];
+		if (!TasksWithSuccessors.Contains(Task.NodeId))
+		{
+			PostExecuteTask.Inputs.Add(Task.NodeId);
+		}
+
+		TasksWithSuccessors.Append(Task.Inputs);
+	}
+
+	// Store back the results in the cache
+	GraphToTaskMapLock.WriteLock();
+	if (!TopGraphToTaskMap.Contains(InGraph))
+	{
+		TopGraphToTaskMap.Add(InGraph, MoveTemp(CompiledTasks));
+	}
+	GraphToTaskMapLock.WriteUnlock();
+}
+
+void FPCGGraphCompiler::NotifyGraphChanged(UPCGGraph* InGraph)
+{
+	if (InGraph)
+	{
+		RemoveFromCache(InGraph);
+	}
+}
+
+void FPCGGraphCompiler::RemoveFromCache(UPCGGraph* InGraph)
+{
+	check(InGraph);
+	RemoveFromCacheRecursive(InGraph);
+}
+
+void FPCGGraphCompiler::RemoveFromCacheRecursive(UPCGGraph* InGraph)
+{
+	GraphToTaskMapLock.WriteLock();
+	GraphToTaskMap.Remove(InGraph);
+#if WITH_EDITOR
+	TopGraphToTaskMap.Remove(InGraph);
+#endif
+	GraphToTaskMapLock.WriteUnlock();
+
+	GraphDependenciesLock.Lock();
+	TArray<UPCGGraph*> ParentGraphs;
+	GraphDependencies.MultiFind(InGraph, ParentGraphs);
+	GraphDependencies.Remove(InGraph);
+	GraphDependenciesLock.Unlock();
+
+	for (UPCGGraph* Graph : ParentGraphs)
+	{
+		RemoveFromCacheRecursive(Graph);
+	}
+}
+#endif // WITH_EDITOR
diff --git a/Engine/Plugins/Experimental/PCG/Source/PCG/Private/Graph/PCGGraphCompiler.h b/Engine/Plugins/Experimental/PCG/Source/PCG/Private/Graph/PCGGraphCompiler.h
new file mode 100644
--- /dev/null
+++ b/Engine/Plugins/Experimental/PCG/Source/PCG/Private/Graph/PCGGraphCompiler.h
@@ -0,0 +1,40 @@
+// Copyright Epic Games, Inc. All Rights Reserved.
+
+#pragma once
+
+#include "CoreMinimal.h"
+#include "PCGGraphExecutor.h"
+
+/** 
+* FPCGGraphCompiler
+* This class compiles a graph into tasks and keeps an internal cache
+*/
+class FPCGGraphCompiler
+{
+public:
+	void Compile(UPCGGraph* InGraph);
+	TArray<FPCGGraphTask> GetCompiledTasks(UPCGGraph* InGraph, bool bIsTopGraph = true);
+
+	static void OffsetNodeIds(TArray<FPCGGraphTask>& Tasks, FPCGTaskId Offset);
+
+#if WITH_EDITOR
+	void NotifyGraphChanged(UPCGGraph* InGraph);
+#endif
+
+private:
+	TArray<FPCGGraphTask> CompileGraph(UPCGGraph* InGraph, FPCGTaskId& NextId);
+
+	FRWLock GraphToTaskMapLock;
+	TMap<UPCGGraph*, TArray<FPCGGraphTask>> GraphToTaskMap;
+
+#if WITH_EDITOR
+	void CompileTopGraph(UPCGGraph* InGraph);
+	TMap<UPCGGraph*, TArray<FPCGGraphTask>> TopGraphToTaskMap;	
+
+	void RemoveFromCache(UPCGGraph* InGraph);
+	void RemoveFromCacheRecursive(UPCGGraph* InGraph);
+
+	FCriticalSection GraphDependenciesLock;
+	TMultiMap<UPCGGraph*, UPCGGraph*> GraphDependencies;
+#endif // WITH_EDITOR
+};
diff --git a/Engine/Plugins/Experimental/PCG/Source/PCG/Private/Graph/PCGGraphExecutor.cpp b/Engine/Plugins/Experimental/PCG/Source/PCG/Private/Graph/PCGGraphExecutor.cpp
new file mode 100644
--- /dev/null
+++ b/Engine/Plugins/Experimental/PCG/Source/PCG/Private/Graph/PCGGraphExecutor.cpp
@@ -0,0 +1,378 @@
+// Copyright Epic Games, Inc. All Rights Reserved.
+
+#include "PCGGraphExecutor.h"
+#include "PCGGraph.h"
+#include "PCGData.h"
+#include "PCGSubgraph.h"
+#include "PCGGraphCompiler.h"
+#include "PCGComponent.h"
+
+// World partition support for in-editor workflows needs these includes
+#if WITH_EDITOR
+#include "WorldPartition/WorldPartition.h"
+#include "FileHelpers.h"
+#endif
+
+FPCGGraphExecutor::FPCGGraphExecutor()
+	: GraphCompiler(MakeUnique<FPCGGraphCompiler>())
+{
+}
+
+FPCGGraphExecutor::~FPCGGraphExecutor() = default;
+
+void FPCGGraphExecutor::Compile(UPCGGraph* Graph)
+{
+	GraphCompiler->Compile(Graph);
+}
+
+FPCGTaskId FPCGGraphExecutor::Schedule(UPCGComponent* Component, const TArray<FPCGTaskId>& ExternalDependencies)
+{
+	check(Component);
+	UPCGGraph* Graph = Component->GetGraph();
+	FPCGTaskId ScheduledId = InvalidTaskId;
+
+	// Get compiled tasks from compiler
+	TArray<FPCGGraphTask> CompiledTasks = GraphCompiler->GetCompiledTasks(Graph);
+
+	// Assign this component to the tasks
+	for (FPCGGraphTask& Task : CompiledTasks)
+	{
+		Task.SourceComponent = Component;
+	}
+
+	// Prepare scheduled task that will be promoted in the next Execute call.
+	if (CompiledTasks.Num() > 0)
+	{
+		check(CompiledTasks[0].Node == Graph->GetInputNode());
+
+		// Setup fetch task on input node
+		CompiledTasks[0].Element = MakeShared<FPCGFetchInputElement>(Component);
+
+		ScheduleLock.Lock();
+
+		FPCGGraphScheduleTask& ScheduledTask = ScheduledTasks.Emplace_GetRef();
+		ScheduledTask.Tasks = MoveTemp(CompiledTasks);
+
+		// Offset task node ids
+		FPCGGraphCompiler::OffsetNodeIds(ScheduledTask.Tasks, NextTaskId);
+		NextTaskId += ScheduledTask.Tasks.Num();
+		ScheduledId = NextTaskId - 1; // This is true because the last task is from the output node or is the post-execute task
+
+#if WITH_EDITOR
+		// Push task (not data) dependencies on the pre-execute task
+		// Note must be done after the offset ids, otherwise we'll break the dependencies
+		check(ScheduledTask.Tasks.Num() >= 2 && ScheduledTask.Tasks[ScheduledTask.Tasks.Num() - 2].Node == nullptr);
+		ScheduledTask.Tasks[ScheduledTask.Tasks.Num() - 2].Inputs.Append(ExternalDependencies);
+#endif
+
+		ScheduleLock.Unlock();
+	}
+
+	return ScheduledId;
+}
+
+#if WITH_EDITOR
+FPCGTaskId FPCGGraphExecutor::ScheduleGeneric(TFunction<bool()> InOperation, const TArray<FPCGTaskId>& TaskDependencies)
+{
+	// Build task & element to hold the operation to perform
+	FPCGGraphTask Task;
+	Task.Inputs.Append(TaskDependencies);
+	Task.Element = MakeShared<FPCGGenericElement>(InOperation);
+
+	ScheduleLock.Lock();
+
+	// Assign task id
+	Task.NodeId = NextTaskId++;
+
+	FPCGGraphScheduleTask& ScheduledTask = ScheduledTasks.Emplace_GetRef();
+	ScheduledTask.Tasks.Add(Task);
+
+	ScheduleLock.Unlock();
+
+	return Task.NodeId;
+}
+#endif
+
+void FPCGGraphExecutor::Execute()
+{
+	// TODO: change this so we can have time-slicing, priority and thread affinities
+
+	// Process any newly scheduled graphs to execute
+	ScheduleLock.Lock();
+
+	for (FPCGGraphScheduleTask& ScheduledTask : ScheduledTasks)
+	{
+		check(ScheduledTask.Tasks.Num() > 0);
+
+		// Finally, push the tasks to the master list
+		Tasks.Append(ScheduledTask.Tasks);
+	}
+
+	if (ScheduledTasks.Num() > 0)
+	{
+		ScheduledTasks.Reset();
+		// Kick off any of the newly added ready tasks
+		QueueReadyTasks();
+	}
+	
+	ScheduleLock.Unlock();
+
+	// TODO: add optimization phase if we've added new graph(s)/tasks to be executed
+
+	// This is a safeguard to check if we're in a stuck state
+	check(Tasks.Num() == 0 || ReadyTasks.Num() > 0 || ActiveTasks.Num() > 0);
+
+	// TODO: change this when we support time-slicing
+	while (ReadyTasks.Num() > 0 || ActiveTasks.Num() > 0)
+	{
+		// First: if we have free resources, move ready tasks to the active tasks
+		bool bCanStartNewTasks = true; // TODO implement this
+
+		if (bCanStartNewTasks && ReadyTasks.Num() > 0)
+		{
+			// Sort by priority
+			// TODO
+
+			// Create active tasks according to resources & priority
+			while (ReadyTasks.Num() > 0 && bCanStartNewTasks)
+			{
+				// TODO: peek at task properties to determine if we could start it (esp. thread affinity)
+				FPCGGraphTask Task = ReadyTasks.Pop();
+
+				// Build input
+				FPCGDataCollection TaskInput;
+				BuildTaskInput(Task, TaskInput);
+
+				FPCGElementPtr Element = Task.Element;
+
+				// If the task didn't have a fixed element, acquire one from the node/settings
+				if (!Element)
+				{
+					// Get appropriate settings
+					const UPCGSettings* Settings = TaskInput.GetSettings(Task.Node->DefaultSettings);
+
+					if (!Settings)
+					{
+						// Can't execute; abort
+						continue;
+					}
+
+					// Create element
+					Element = Settings->GetElement();
+				}
+
+				if (!Element)
+				{
+					// Can't execute; abort
+					continue;
+				}
+				
+				FPCGContextPtr Context = Element->Initialize(TaskInput, Task.SourceComponent);
+				Context->Node = Task.Node; // still needed?
+
+				FPCGGraphActiveTask& ActiveTask = ActiveTasks.Emplace_GetRef();
+				ActiveTask.Element = Element;
+				ActiveTask.Context = Context;
+				ActiveTask.NodeId = Task.NodeId;
+				
+				// TODO update bCanStartNewTasks
+			}
+		}
+
+		// Perform execution
+		// TODO: this is a naive single-threaded implementation, should be replaced by
+		// the task graph or futures w/proper thread assingment
+		for (int ExecutionIndex = 0; ExecutionIndex < ActiveTasks.Num(); ++ExecutionIndex)
+		{
+			FPCGGraphActiveTask& ActiveTask = ActiveTasks[ExecutionIndex];
+
+			if (ActiveTask.Element->Execute(ActiveTask.Context))
+			{
+				// Store output data in map
+				StoreResults(ActiveTask.NodeId, ActiveTask.Context->OutputData);
+
+				// Move successor nodes to the ready queue
+				QueueReadyTasks(ActiveTask.NodeId);
+
+				// Remove current active task from list
+				ActiveTasks.RemoveAt(ExecutionIndex);
+				--ExecutionIndex;
+			}
+		}
+
+		// Nothing left to do; we'll release everything here.
+		// TODO: this is fine and will make sure any intermediate data is properly
+		// garbage collected, however, this goes a bit against our goals if we want to
+		// keep a cache of intermediate results.
+		if (ReadyTasks.Num() == 0 && ActiveTasks.Num() == 0 && Tasks.Num() == 0)
+		{
+			ClearResults();
+		}
+
+#if WITH_EDITOR
+		// Save & release resources when running in-editor
+		SaveDirtyActors();
+		ReleaseUnusedActors();
+#endif
+	}
+}
+
+void FPCGGraphExecutor::QueueReadyTasks(FPCGTaskId FinishedTaskHint)
+{
+	for (int TaskIndex = Tasks.Num() - 1; TaskIndex >= 0; --TaskIndex)
+	{
+		bool bAllPrerequisitesMet = true;
+		for (const FPCGTaskId& Id : Tasks[TaskIndex].Inputs)
+		{
+			bAllPrerequisitesMet &= OutputData.Contains(Id);
+		}
+
+		if (bAllPrerequisitesMet)
+		{
+			ReadyTasks.Emplace(MoveTemp(Tasks[TaskIndex]));
+			Tasks.RemoveAtSwap(TaskIndex);
+		}
+	}
+}
+
+void FPCGGraphExecutor::BuildTaskInput(const FPCGGraphTask& Task, FPCGDataCollection& TaskInput)
+{
+	for (const FPCGTaskId& Id : Task.Inputs)
+	{
+		check(OutputData.Contains(Id));
+		TaskInput.TaggedData.Append(OutputData[Id].TaggedData);
+	}
+}
+
+void FPCGGraphExecutor::StoreResults(FPCGTaskId InTaskId, const FPCGDataCollection& InTaskOutput)
+{
+	// Store output in map
+	OutputData.Add(InTaskId, InTaskOutput);
+
+	// Root any non-rooted results, otherwise they'll get garbage-collected
+	for (const FPCGTaggedData& TaggedData : InTaskOutput.TaggedData)
+	{
+		if (TaggedData.Data != nullptr && !TaggedData.Data->IsRooted())
+		{
+			// This is technically a const_cast
+			UObject* DataToRoot = Cast<UObject>(TaggedData.Data);
+			DataToRoot->AddToRoot();
+			RootedData.Add(DataToRoot);
+		}
+	}
+}
+
+void FPCGGraphExecutor::ClearResults()
+{
+	ScheduleLock.Lock();
+	
+	NextTaskId = 0;
+	OutputData.Reset();
+
+	for (UObject* Data : RootedData)
+	{
+		Data->RemoveFromRoot();
+	}
+	RootedData.Reset();
+
+	ScheduleLock.Unlock();
+}
+
+#if WITH_EDITOR
+void FPCGGraphExecutor::AddToDirtyActors(AActor* Actor)
+{
+	ActorsListLock.Lock();
+	ActorsToSave.Add(Actor);
+	ActorsListLock.Unlock();
+}
+
+void FPCGGraphExecutor::AddToUnusedActors(const TSet<FWorldPartitionReference>& UnusedActors)
+{
+	ActorsListLock.Lock();
+	ActorsToRelease.Append(UnusedActors);
+	ActorsListLock.Unlock();
+}
+
+void FPCGGraphExecutor::SaveDirtyActors()
+{
+	ActorsListLock.Lock();
+	TSet<AActor*> ToSave = MoveTemp(ActorsToSave);
+	ActorsToSave.Reset();
+	ActorsListLock.Unlock();
+
+	TSet<UPackage*> PackagesToSave;
+	for (AActor* Actor : ToSave)
+	{
+		PackagesToSave.Add(Actor->GetExternalPackage());
+	}
+
+	if (PackagesToSave.Num() > 0)
+	{
+		UEditorLoadingAndSavingUtils::SavePackages(PackagesToSave.Array(), true);
+	}
+}
+
+void FPCGGraphExecutor::ReleaseUnusedActors()
+{
+	ActorsListLock.Lock();
+	bool bRunGC = ActorsToRelease.Num() > 0;
+	ActorsToRelease.Reset();
+	ActorsListLock.Unlock();
+
+	// TODO: maybe consider other checks
+	if (bRunGC)
+	{
+		CollectGarbage(RF_NoFlags, true);
+	}
+}
+
+void FPCGGraphExecutor::NotifyGraphChanged(UPCGGraph* InGraph)
+{
+	if (GraphCompiler)
+	{
+		GraphCompiler->NotifyGraphChanged(InGraph);
+	}
+}
+#endif
+
+FPCGFetchInputElement::FPCGFetchInputElement(UPCGComponent* InComponent)
+	: Component(InComponent)
+{
+}
+
+bool FPCGFetchInputElement::Execute(FPCGContextPtr Context) const
+{
+	// First: any input can be passed through to the output trivially
+	Context->OutputData = Context->InputData;
+
+	// Second: fetch the inputs provided by the component
+	check(Component);
+	if(UPCGData* PCGData = Component->GetPCGData())
+	{
+		FPCGTaggedData& TaggedData = Context->OutputData.TaggedData.Emplace_GetRef();
+		TaggedData.Data = PCGData;
+	}
+
+	// Third: fetch the exclusions provided by the component
+	TArray<UPCGData*> Exclusions = Component->GetPCGExclusionData();
+	for (const UPCGData* Exclusion : Exclusions)
+	{
+		FPCGTaggedData& TaggedData = Context->OutputData.TaggedData.Emplace_GetRef();
+		TaggedData.Data = Exclusion;
+		TaggedData.Usage = EPCGDataUsage::Exclusion;
+	}
+
+	return true;
+}
+
+#if WITH_EDITOR
+FPCGGenericElement::FPCGGenericElement(TFunction<bool()> InOperation)
+	: Operation(InOperation)
+{
+}
+
+bool FPCGGenericElement::Execute(FPCGContextPtr Context) const
+{
+	return Operation();
+}
+#endif // WITH_EDITOR
diff --git a/Engine/Plugins/Experimental/PCG/Source/PCG/Private/Graph/PCGGraphExecutor.h b/Engine/Plugins/Experimental/PCG/Source/PCG/Private/Graph/PCGGraphExecutor.h
new file mode 100644
--- /dev/null
+++ b/Engine/Plugins/Experimental/PCG/Source/PCG/Private/Graph/PCGGraphExecutor.h
@@ -0,0 +1,125 @@
+// Copyright Epic Games, Inc. All Rights Reserved.
+
+#pragma once
+
+#include "CoreMinimal.h"
+#include "PCGData.h"
+#include "PCGContext.h"
+#include "PCGElement.h"
+#include "PCGSubsystem.h"
+
+#if WITH_EDITOR
+#include "WorldPartition/WorldPartitionHandle.h" // Needed for FWorldPartitonReference
+#endif
+
+class UPCGGraph;
+class UPCGNode;
+class UPCGComponent;
+class FPCGGraphCompiler;
+
+struct FPCGGraphTask
+{
+	TArray<FPCGTaskId> Inputs;
+	//TArray<DataId> Outputs;
+	const UPCGNode* Node = nullptr;
+	const UPCGComponent* SourceComponent = nullptr; // TODO: can technically be editor only
+	FPCGElementPtr Element; // Added to have tasks that aren't node-bound
+	FPCGTaskId NodeId = InvalidTaskId;
+};
+
+struct FPCGGraphScheduleTask
+{
+	TArray<FPCGGraphTask> Tasks;
+};
+
+struct FPCGGraphActiveTask
+{
+	FPCGElementPtr Element;
+	FPCGContextPtr Context;
+	FPCGTaskId NodeId = InvalidTaskId;
+};
+
+class FPCGGraphExecutor
+{
+public:
+	FPCGGraphExecutor();
+	~FPCGGraphExecutor();
+
+	/** Compile (and cache) a graph for later use. This call is threadsafe */
+	void Compile(UPCGGraph* InGraph);
+
+	/** Schedules the execution of a given graph with specified inputs. This call is threadsafe */
+	FPCGTaskId Schedule(UPCGComponent* InComponent, const TArray<FPCGTaskId>& TaskDependency = TArray<FPCGTaskId>());
+
+#if WITH_EDITOR
+	/** General job scheduling, used to control loading/unloading */
+	FPCGTaskId ScheduleGeneric(TFunction<bool()> InOperation, const TArray<FPCGTaskId>& TaskDependencies);
+
+	void AddToDirtyActors(AActor* Actor);
+	void AddToUnusedActors(const TSet<FWorldPartitionReference>& UnusedActors);
+
+	/** Notify compiler that graph has changed so it'll be removed from the cache */
+	void NotifyGraphChanged(UPCGGraph* InGraph);
+#endif
+
+	/** "Tick" of the graph executor. This call is NOT THREADSAFE */
+	void Execute();
+
+private:
+	void QueueReadyTasks(FPCGTaskId FinishedTaskHint = InvalidTaskId);
+	void BuildTaskInput(const FPCGGraphTask& Task, FPCGDataCollection& TaskInput);
+	void StoreResults(FPCGTaskId InTaskId, const FPCGDataCollection& InTaskOutput);
+	void ClearResults();
+
+#if WITH_EDITOR
+	void SaveDirtyActors();
+	void ReleaseUnusedActors();
+#endif
+
+	/** Graph compiler that turns a graph into tasks */
+	TUniquePtr<FPCGGraphCompiler> GraphCompiler;
+
+	FCriticalSection ScheduleLock;
+	TArray<FPCGGraphScheduleTask> ScheduledTasks;
+
+	TArray<FPCGGraphTask> Tasks;
+	TArray<FPCGGraphTask> ReadyTasks;
+	TArray<FPCGGraphActiveTask> ActiveTasks;
+	TArray<UObject*> RootedData;
+	/** Map of node instances to their output, could be cleared once execution is done */
+	/** Note: this should at some point unload based on loaded/unloaded proxies, otherwise memory cost will be unbounded */
+	TMap<FPCGTaskId, FPCGDataCollection> OutputData;
+	/** Monotonically increasing id. Should be reset once all tasks are executed, should be protected by the ScheduleLock */
+	FPCGTaskId NextTaskId = 0;
+
+#if WITH_EDITOR
+	FCriticalSection ActorsListLock;
+	TSet<AActor*> ActorsToSave;
+	TSet<FWorldPartitionReference> ActorsToRelease;
+#endif
+};
+
+class FPCGFetchInputElement : public FSimplePCGElement
+{
+public:
+	FPCGFetchInputElement(UPCGComponent* InComponent);
+
+	virtual bool Execute(FPCGContextPtr Context) const override;
+private:
+	UPCGComponent* Component = nullptr;
+};
+
+#if WITH_EDITOR
+class FPCGGenericElement : public FSimplePCGElement
+{
+public:
+	FPCGGenericElement(TFunction<bool()> InOperation);
+
+	// Important note: generic elements must always be run on the main thread
+	// as most of these will impact the editor in some way (loading, unloading, saving)
+	virtual bool Execute(FPCGContextPtr Context) const override;
+
+private:
+	TFunction<bool()> Operation;
+};
+#endif // WITH_EDITOR
diff --git a/Engine/Plugins/Experimental/PCG/Source/PCG/Private/Grid/PCGPartitionActor.cpp b/Engine/Plugins/Experimental/PCG/Source/PCG/Private/Grid/PCGPartitionActor.cpp
new file mode 100644
--- /dev/null
+++ b/Engine/Plugins/Experimental/PCG/Source/PCG/Private/Grid/PCGPartitionActor.cpp
@@ -0,0 +1,172 @@
+// Copyright Epic Games, Inc. All Rights Reserved.
+
+#include "PCGPartitionActor.h"
+#include "PCGComponent.h"
+#include "PCGHelpers.h"
+#include "PCGWorldActor.h"
+
+#include "Landscape.h"
+
+void APCGPartitionActor::BeginPlay()
+{
+	Super::BeginPlay();
+
+	// Register cell to the PCG grid
+	// TODO
+}
+
+void APCGPartitionActor::EndPlay(const EEndPlayReason::Type EndPlayReason)
+{
+	// Unregister each cell to the PCG grid
+	// TODO
+
+	Super::EndPlay(EndPlayReason);
+}
+
+#if WITH_EDITOR
+uint32 APCGPartitionActor::GetDefaultGridSize(UWorld* InWorld) const
+{
+	// TODO: promote to world settings
+	return 25600;
+}
+#endif
+
+FBox APCGPartitionActor::GetFixedBounds() const
+{
+	const FVector Center = GetActorLocation();
+	// TODO: keep this in the actor instead of hardcoding.
+	const uint32 HalfGridSize = 25600 / 2;
+	return FBox(Center - HalfGridSize, Center + HalfGridSize);
+}
+
+void APCGPartitionActor::GetActorBounds(bool bOnlyCollidingComponents, FVector& Origin, FVector& BoxExtent, bool bIncludeFromChildActors) const
+{
+	Super::GetActorBounds(bOnlyCollidingComponents, Origin, BoxExtent, bIncludeFromChildActors);
+
+	// To keep consistency with the other GetBounds functions, transform our result into an origin / extent formatting
+	FBox Bounds(Origin - BoxExtent, Origin + BoxExtent);
+	Bounds += GetFixedBounds();
+	Bounds.GetCenterAndExtents(Origin, BoxExtent);
+}
+
+UPCGComponent* APCGPartitionActor::GetLocalComponent(const UPCGComponent* OriginalComponent) const
+{
+	const TObjectPtr<UPCGComponent>* LocalComponent = OriginalToLocalMap.Find(OriginalComponent);
+	return LocalComponent ? *LocalComponent : nullptr;
+}
+
+UPCGComponent* APCGPartitionActor::GetOriginalComponent(const UPCGComponent* LocalComponent) const
+{
+	const TObjectPtr<UPCGComponent>* OriginalComponent = LocalToOriginalMap.Find(LocalComponent);
+	return OriginalComponent ? (*OriginalComponent).Get() : nullptr;
+}
+
+#if WITH_EDITOR
+FBox APCGPartitionActor::GetStreamingBounds() const
+{
+	return Super::GetStreamingBounds() + GetFixedBounds();
+}
+
+AActor* APCGPartitionActor::GetSceneOutlinerParent() const
+{
+	if (APCGWorldActor* PCGActor = PCGHelpers::GetPCGWorldActor(GetWorld()))
+	{
+		return PCGActor;
+	}
+	else
+	{
+		return Super::GetSceneOutlinerParent();
+	}	
+}
+
+void APCGPartitionActor::AddGraphInstance(UPCGComponent* OriginalComponent)
+{
+	if (!OriginalComponent)
+	{
+		return;
+	}
+
+	// Make sure we don't have that graph twice;
+	// Here we'll check if there has been some changes worth propagating or not
+	UPCGComponent* LocalComponent = GetLocalComponent(OriginalComponent);
+
+	if (LocalComponent)
+	{
+		// Update properties as needed and early out
+		LocalComponent->SetPropertiesFromOriginal(OriginalComponent);
+		return;
+	}
+
+	Modify();
+
+	// Create a new local component
+	LocalComponent = NewObject<UPCGComponent>(this);
+	LocalComponent->SetPropertiesFromOriginal(OriginalComponent);
+
+	LocalComponent->RegisterComponent();
+	// TODO: check if we should use a non-instanced component?
+	AddInstanceComponent(LocalComponent);
+
+	OriginalToLocalMap.Add(OriginalComponent, LocalComponent);
+	LocalToOriginalMap.Add(LocalComponent, OriginalComponent);
+}
+
+bool APCGPartitionActor::RemoveGraphInstance(UPCGComponent* OriginalComponent)
+{
+	UPCGComponent* LocalComponent = GetLocalComponent(OriginalComponent);
+
+	if (!LocalComponent)
+	{
+		return false;
+	}
+
+	Modify();
+
+	OriginalToLocalMap.Remove(OriginalComponent);
+	LocalToOriginalMap.Remove(LocalComponent);
+
+	// TODO Add option to not cleanup?
+	LocalComponent->Cleanup(/*bRemoveComponents=*/true);
+	LocalComponent->DestroyComponent();
+
+	return OriginalToLocalMap.IsEmpty();
+}
+
+bool APCGPartitionActor::CleanupDeadGraphInstances()
+{
+	TSet<TObjectPtr<UPCGComponent>> DeadLocalInstances;
+
+	// Note: since we might end up with multiple nulls in the original to local map
+	// it might not be very stable to use it; we'll use the 
+	for (const auto& LocalToOriginalItem : LocalToOriginalMap)
+	{
+		if (!LocalToOriginalItem.Value)
+		{
+			DeadLocalInstances.Add(LocalToOriginalItem.Key);
+		}
+	}
+
+	if (DeadLocalInstances.Num() == 0)
+	{
+		return OriginalToLocalMap.IsEmpty();
+	}
+
+	Modify();
+
+	for (const TObjectPtr<UPCGComponent>& DeadInstance : DeadLocalInstances)
+	{
+		LocalToOriginalMap.Remove(DeadInstance);
+
+		if (DeadInstance)
+		{
+			DeadInstance->Cleanup(/*bRemoveComponents=*/true);
+			DeadInstance->DestroyComponent();
+		}
+	}
+
+	// Remove all dead entries
+	OriginalToLocalMap.Remove(nullptr);
+
+	return OriginalToLocalMap.IsEmpty();
+}
+#endif // WITH_EDITOR
\ No newline at end of file
diff --git a/Engine/Plugins/Experimental/PCG/Source/PCG/Private/Grid/PCGPartitionActor.h b/Engine/Plugins/Experimental/PCG/Source/PCG/Private/Grid/PCGPartitionActor.h
new file mode 100644
--- /dev/null
+++ b/Engine/Plugins/Experimental/PCG/Source/PCG/Private/Grid/PCGPartitionActor.h
@@ -0,0 +1,61 @@
+// Copyright Epic Games, Inc. All Rights Reserved.
+
+#pragma once
+
+#include "CoreMinimal.h"
+#include "ActorPartition/PartitionActor.h"
+
+#include "PCGPartitionActor.generated.h"
+
+class UPCGComponent;
+
+/** 
+* The APCGPartitionActor actor is used to store grid cell data
+* and its size will be a multiple of the grid size.
+*/
+UCLASS(MinimalAPI, NotBlueprintable, NotPlaceable)
+class APCGPartitionActor : public APartitionActor
+{
+	GENERATED_BODY()
+
+public:
+	//~Begin AActor Interface
+	virtual void BeginPlay();
+	virtual void EndPlay(const EEndPlayReason::Type EndPlayReason) override;
+	virtual void GetActorBounds(bool bOnlyCollidingComponents, FVector& Origin, FVector& BoxExtent, bool bIncludeFromChildActors) const override;
+#if WITH_EDITOR
+	virtual FBox GetStreamingBounds() const override;
+	virtual AActor* GetSceneOutlinerParent() const override;
+#endif
+	//~End AActor Interface
+
+#if WITH_EDITOR
+	//~Begin APartitionActor Interface
+	virtual uint32 GetDefaultGridSize(UWorld* InWorld) const override;
+	virtual FGuid GetGridGuid() const override { return PCGGuid; }
+	//~End APartitionActor Interface
+#endif
+
+	FBox GetFixedBounds() const;
+#if WITH_EDITOR
+	void AddGraphInstance(UPCGComponent* OriginalComponent);
+	bool RemoveGraphInstance(UPCGComponent* OriginalComponent);
+	bool CleanupDeadGraphInstances();
+#endif
+
+	// TODO: Make this in-editor only; during runtime, we should keep a map of component to bounds/volume only
+	// and preferably precompute the intersection, so this would make it easier/possible to not have the original actor in game version.
+	UPCGComponent* GetLocalComponent(const UPCGComponent* OriginalComponent) const;
+	UPCGComponent* GetOriginalComponent(const UPCGComponent* LocalComponent) const;
+
+	UPROPERTY()
+	FGuid PCGGuid;
+
+private:
+	// TODO: Make these properties editor only (see comment before).
+	UPROPERTY()
+	TMap<TObjectPtr<UPCGComponent>, TObjectPtr<UPCGComponent>> OriginalToLocalMap;
+
+	UPROPERTY()
+	TMap<TObjectPtr<UPCGComponent>, TObjectPtr<UPCGComponent>> LocalToOriginalMap;
+};
diff --git a/Engine/Plugins/Experimental/PCG/Source/PCG/Private/PCGComponent.cpp b/Engine/Plugins/Experimental/PCG/Source/PCG/Private/PCGComponent.cpp
new file mode 100644
--- /dev/null
+++ b/Engine/Plugins/Experimental/PCG/Source/PCG/Private/PCGComponent.cpp
@@ -0,0 +1,1031 @@
+// Copyright Epic Games, Inc. All Rights Reserved.
+
+#include "PCGComponent.h"
+#include "PCGGraph.h"
+#include "PCGHelpers.h"
+#include "PCGSubsystem.h"
+#include "PCGVolume.h"
+#include "Data/PCGLandscapeData.h"
+#include "Data/PCGLandscapeSplineData.h"
+#include "Data/PCGPointData.h"
+#include "Data/PCGSplineData.h"
+#include "Data/PCGUnionData.h"
+#include "Data/PCGVolumeData.h"
+#include "Graph/PCGGraphExecutor.h"
+#include "Grid/PCGPartitionActor.h"
+
+#include "ActorPartition/ActorPartitionSubsystem.h"
+#include "Components/InstancedStaticMeshComponent.h"
+#include "Components/SplineComponent.h"
+#include "GameFramework/Volume.h"
+#include "Kismet/GameplayStatics.h"
+#include "Landscape.h"
+#include "LandscapeInfo.h"
+#include "LandscapeSplineActor.h"
+#include "LandscapeSplinesComponent.h"
+#include "WorldPartition/WorldPartition.h"
+
+#if WITH_EDITOR
+#include "ScopedTransaction.h"
+#endif
+
+#define LOCTEXT_NAMESPACE "UPCGComponent"
+
+bool UPCGComponent::CanPartition() const
+{
+	// Support/Force partitioning on non-PCG partition actors in WP worlds.
+	return GetOwner() && GetOwner()->GetWorld() && GetOwner()->GetWorld()->GetWorldPartition() != nullptr && Cast<APCGPartitionActor>(GetOwner()) == nullptr;
+}
+
+bool UPCGComponent::IsPartitioned() const
+{
+	return bIsPartitioned && CanPartition();
+}
+
+bool UPCGComponent::HasSplineComponent() const
+{
+	if (!GetOwner())
+	{
+		return false;
+	}
+
+	TArray<USplineComponent*> Splines;
+	GetOwner()->GetComponents<USplineComponent>(Splines);
+
+	return !Splines.IsEmpty();
+}
+
+void UPCGComponent::SetGraph(UPCGGraph* InGraph)
+{
+	if(Graph == InGraph)
+	{
+		return;
+	}
+
+#if WITH_EDITOR
+	if (Graph)
+	{
+		Graph->OnGraphChangedDelegate.RemoveAll(this);
+	}
+#endif
+
+	Graph = InGraph;
+
+#if WITH_EDITOR
+	if (InGraph)
+	{
+		Graph->OnGraphChangedDelegate.AddUObject(this, &UPCGComponent::OnGraphChanged);
+	}
+
+	DirtyGenerated();
+#endif	
+}
+
+bool UPCGComponent::ShouldGenerate(bool bForce) const
+{
+	if (!bActivated || !Graph || !GetSubsystem())
+	{
+		return false;
+	}
+
+	if (IsPartitioned())
+	{
+#if WITH_EDITOR
+		return !GIsPlayInEditorWorld;
+#else
+		return false;
+#endif
+	}
+	else
+	{
+		return (!bGenerated || 
+#if WITH_EDITOR
+			bDirtyGenerated || 
+#endif
+			bForce);
+	}
+}
+
+void UPCGComponent::SetPropertiesFromOriginal(const UPCGComponent* Original)
+{
+	check(Original);
+
+	EPCGComponentInput NewInputType = Original->InputType;
+
+	// If we're inheriting properties from another component that would have targeted a "special" actor
+	// then we must make sure we update the InputType appropriately
+	if (NewInputType == EPCGComponentInput::Actor)
+	{
+		if(Cast<ALandscapeProxy>(Original->GetOwner()) != nullptr && Cast<ALandscapeProxy>(GetOwner()) == nullptr)
+		{
+			NewInputType = EPCGComponentInput::Landscape;
+		}
+	}
+
+#if WITH_EDITOR
+	const bool bHasDirtyInput = InputType != NewInputType;
+	const bool bHasDirtyExclusions = ExclusionTags != Original->ExclusionTags;	
+	const bool bIsDirty = bHasDirtyInput || bHasDirtyExclusions || Seed != Original->Seed || Graph != Original->Graph;
+
+	// Note that while we dirty here, we won't trigger a refresh since we don't have the required context
+	if (bIsDirty)
+	{
+		Modify();
+		DirtyGenerated(bHasDirtyInput, bHasDirtyExclusions);
+	}
+#endif
+
+	InputType = NewInputType;
+	ExclusionTags = Original->ExclusionTags;
+	Seed = Original->Seed;
+	SetGraph(Original->Graph);
+}
+
+void UPCGComponent::Generate()
+{
+#if WITH_EDITOR
+	if (bIsGenerating)
+	{
+		return;
+	}
+
+	FScopedTransaction Transaction(LOCTEXT("PCGGenerate", "Execute generation on PCG component"));
+#endif
+
+	Generate(/*bForce=*/true);
+}
+
+void UPCGComponent::Generate(bool bForce)
+{
+#if WITH_EDITOR
+	if (bIsGenerating)
+	{
+		return;
+	}
+#endif
+
+	FPCGTaskId TaskId = GenerateInternal(bForce, {});
+
+#if WITH_EDITOR
+	if (TaskId != InvalidTaskId)
+	{
+		bIsGenerating = true;
+	}
+#endif
+}
+
+FPCGTaskId UPCGComponent::GenerateInternal(bool bForce, const TArray<FPCGTaskId>& TaskDependencies)
+{
+	FPCGTaskId TaskId = InvalidTaskId;
+
+	if (!ShouldGenerate(bForce))
+	{
+		return InvalidTaskId;
+	}
+
+#if WITH_EDITOR
+	if (bForce && bGenerated && !bDirtyGenerated)
+	{
+		// TODO: generate new seed
+		++Seed;
+	}
+#endif
+
+	Modify();
+
+	if (IsPartitioned())
+	{
+#if WITH_EDITOR
+		TaskId = GetSubsystem()->DelayGenerateGraph(this, /*bSave=*/bForce);
+#endif
+	}
+	else
+	{
+		// Immediate operation: cleanup beforehand
+		if (bGenerated)
+		{
+			CleanupInternal(/*bRemoveComponents=*/false);
+		}
+
+		const FBox NewBounds = GetGridBounds();
+		if (NewBounds.IsValid)
+		{
+			TaskId = GetSubsystem()->ScheduleComponent(this, TaskDependencies);
+		}
+	}
+
+	return TaskId;
+}
+
+void UPCGComponent::PostProcessGraph(const FBox& InNewBounds, bool bInGenerated)
+{
+	LastGeneratedBounds = InNewBounds;
+
+	if (bInGenerated)
+	{
+		bGenerated = true;
+
+#if WITH_EDITOR
+		bDirtyGenerated = false;
+		bIsGenerating = false;
+		OnPCGGraphGeneratedDelegate.Broadcast(this);
+#endif
+	}
+}
+
+void UPCGComponent::Cleanup()
+{
+	if (!bGenerated || !GetSubsystem())
+	{
+		return;
+	}
+
+#if WITH_EDITOR
+	FScopedTransaction Transaction(LOCTEXT("PCGCleanup", "Clean up PCG component"));
+#endif
+
+	Cleanup(/*bRemoveComponents=*/true, /*bSave=*/true);
+}
+
+void UPCGComponent::Cleanup(bool bRemoveComponents, bool bSave)
+{
+	if (!bGenerated || !GetSubsystem())
+	{
+		return;
+	}
+
+	if (IsPartitioned())
+	{
+#if WITH_EDITOR
+		if (!GIsPlayInEditorWorld)
+		{
+			Modify();
+			GetSubsystem()->CleanupGraph(this, LastGeneratedBounds, bRemoveComponents, bSave);
+			bGenerated = false;
+		}
+#endif
+	}
+	else
+	{
+		CleanupInternal(bRemoveComponents);
+	}
+
+#if WITH_EDITOR
+	OnPCGGraphCleanedDelegate.Broadcast(this);
+#endif
+}
+
+void UPCGComponent::CleanupInternal(bool bRemoveComponents)
+{
+	if (!bGenerated || IsPartitioned())
+	{
+		return;
+	}
+
+	Modify();
+	bGenerated = false;
+
+	// TODO: remove any ISM, HISM on this actor that were created by the PCG execution
+	// TODO: call recursively on actors generated through this + delete them.
+	TArray<UInstancedStaticMeshComponent*> ISMCs;
+	GetOwner()->GetComponents<UInstancedStaticMeshComponent>(ISMCs);
+
+	for (UInstancedStaticMeshComponent* ISMC : ISMCs)
+	{
+		if (ISMC->ComponentTags.Contains(GetFName()))
+		{
+			if (bRemoveComponents)
+			{
+				ISMC->DestroyComponent();
+			}
+			else
+			{
+				ISMC->ClearInstances();
+				ISMC->UpdateBounds();
+			}
+		}
+	}
+}
+
+void UPCGComponent::BeginPlay()
+{
+	Super::BeginPlay();
+
+	if(bActivated && !bGenerated && !IsPartitioned())
+	{
+		Generate(/*bForce=*/false);
+		bRuntimeGenerated = true;
+	}
+}
+
+void UPCGComponent::OnComponentCreated()
+{
+	Super::OnComponentCreated();
+
+#if WITH_EDITOR
+	SetupActorCallbacks();
+#endif
+}
+
+void UPCGComponent::OnComponentDestroyed(bool bDestroyingHierarchy)
+{
+#if WITH_EDITOR
+	// This is inspired by UChildActorComponent::DestroyChildActor()
+	// In the case of level change or exit, the subsystem will be null
+	if (UPCGSubsystem* Subsystem = GetSubsystem())
+	{
+		// The RF_BeginDestroyed flag is set when the object is being unloaded, but not in the editor-destroy context we're interested in.
+		if (!HasAnyFlags(RF_BeginDestroyed) && !IsUnreachable() && IsPartitioned() && !GetOwner()->GetWorld()->IsGameWorld())
+		{
+			Subsystem->CleanupPartitionActors(LastGeneratedBounds);
+		}
+	}
+#endif
+
+	Super::OnComponentDestroyed(bDestroyingHierarchy);
+}
+
+void UPCGComponent::PostLoad()
+{
+	Super::PostLoad();
+
+#if WITH_EDITOR
+	SetupActorCallbacks();
+	SetupExclusionCallbacks();
+
+	if (Graph)
+	{
+		Graph->OnGraphChangedDelegate.AddUObject(this, &UPCGComponent::OnGraphChanged);
+	}
+#endif
+}
+
+void UPCGComponent::BeginDestroy()
+{
+#if WITH_EDITOR
+	if (Graph)
+	{
+		Graph->OnGraphChangedDelegate.RemoveAll(this);
+	}
+
+	TeardownExclusionCallbacks();
+	TeardownActorCallbacks();
+#endif
+
+	Super::BeginDestroy();
+}
+
+#if WITH_EDITOR
+void UPCGComponent::PreEditChange(FProperty* PropertyAboutToChange)
+{
+	if (PropertyAboutToChange)
+	{
+		const FName PropName = PropertyAboutToChange->GetFName();
+
+		if (PropName == GET_MEMBER_NAME_CHECKED(UPCGComponent, Graph) && Graph)
+		{
+			Graph->OnGraphChangedDelegate.RemoveAll(this);
+		}
+		else if (PropName == GET_MEMBER_NAME_CHECKED(UPCGComponent, ExclusionTags))
+		{
+			TeardownExclusionCallbacks();
+		}
+	}
+
+	Super::PreEditChange(PropertyAboutToChange);
+}
+
+void UPCGComponent::PostEditChangeProperty(FPropertyChangedEvent& PropertyChangedEvent)
+{
+	Super::PostEditChangeProperty(PropertyChangedEvent);
+
+	if (!PropertyChangedEvent.Property)
+	{
+		return;
+	}
+
+	const FName PropName = PropertyChangedEvent.Property->GetFName();
+
+	// Important note: all property changes already go through the OnObjectPropertyChanged, 
+	// So there is no need to add cases that do simple Refresh() calls
+	if (PropName == GET_MEMBER_NAME_CHECKED(UPCGComponent, bIsPartitioned))
+	{
+		if (CanPartition())
+		{
+			if (bActivated)
+			{
+				bool bIsNowPartitioned = bIsPartitioned;
+				bIsPartitioned = !bIsPartitioned;
+				// First, we'll cleanup
+				bActivated = false;
+				Refresh();
+				// Then do a normal refresh
+				bActivated = true;
+				bIsPartitioned = bIsNowPartitioned;
+				DirtyGenerated();
+				Refresh();
+			}
+		}
+		else
+		{
+			// Just ignore the change
+			bIsPartitioned = false;
+		}
+	}
+	else if (PropName == GET_MEMBER_NAME_CHECKED(UPCGComponent, Graph))
+	{
+		if (Graph)
+		{
+			Graph->OnGraphChangedDelegate.AddUObject(this, &UPCGComponent::OnGraphChanged);
+		}
+
+		DirtyGenerated();
+		Refresh();
+	}
+	else if (PropName == GET_MEMBER_NAME_CHECKED(UPCGComponent, InputType))
+	{
+		DirtyGenerated(/*bDirtyCachedInput=*/true);
+		Refresh();
+	}
+	// General properties that don't affect behavior
+	else if (PropName == GET_MEMBER_NAME_CHECKED(UPCGComponent, Seed))
+	{
+		DirtyGenerated();
+		Refresh();
+	}
+	else if (PropName == GET_MEMBER_NAME_CHECKED(UPCGComponent, ExclusionTags))
+	{
+		SetupExclusionCallbacks();
+
+		bool bHadExclusionData = !CachedExclusionData.IsEmpty();
+		bool bHasNonTrivialTags = false;
+
+		for (const FName& ExclusionTag : ExclusionTags)
+		{
+			bHasNonTrivialTags |= ExclusionTag != NAME_None;
+		}
+
+		bool bNeedToDirtyExclusionData = (bHadExclusionData || bHasNonTrivialTags);
+
+		if (bNeedToDirtyExclusionData)
+		{
+			DirtyGenerated(/*bDirtyCachedInput=*/false, bNeedToDirtyExclusionData);
+			Refresh();
+		}
+	}
+}
+
+void UPCGComponent::PreEditUndo()
+{
+	// We don't know what is changing so remove all callbacks
+	if (Graph)
+	{
+		Graph->OnGraphChangedDelegate.RemoveAll(this);
+	}
+
+	TeardownExclusionCallbacks();
+	// Here we will keep a copy of flags that we require to keep through the undo
+	// so we can have a consistent state
+	bWasGeneratedPriorToUndo = bGenerated;
+	LastGeneratedBoundsPriorToUndo = LastGeneratedBounds;
+}
+
+void UPCGComponent::PostEditUndo()
+{
+	LastGeneratedBounds = LastGeneratedBoundsPriorToUndo;
+
+	if (Graph)
+	{
+		Graph->OnGraphChangedDelegate.AddUObject(this, &UPCGComponent::OnGraphChanged);
+	}
+
+	SetupExclusionCallbacks();
+	DirtyGenerated(/*bDirtyCachedInput=*/true, /*bDirtyCachedExclusions=*/true);
+	
+	if (bWasGeneratedPriorToUndo && !bGenerated)
+	{
+		// Need to reset the generated flag to go through the Cleanup
+		bGenerated = true;
+		Cleanup(/*bRemoveComponents=*/true, /*bSave=*/true);
+	}
+	else
+	{
+		Refresh();
+	}
+}
+
+void UPCGComponent::SetupActorCallbacks()
+{
+	GEngine->OnActorMoved().AddUObject(this, &UPCGComponent::OnActorMoved);
+	FCoreUObjectDelegates::OnObjectPropertyChanged.AddUObject(this, &UPCGComponent::OnObjectPropertyChanged);
+}
+
+void UPCGComponent::TeardownActorCallbacks()
+{
+	FCoreUObjectDelegates::OnObjectPropertyChanged.RemoveAll(this);
+	GEngine->OnActorMoved().RemoveAll(this);
+}
+
+void UPCGComponent::SetupExclusionCallbacks()
+{
+	bool bHasNonTrivialTags = false;
+
+	for (const FName& ExclusionTag : ExclusionTags)
+	{
+		bHasNonTrivialTags |= ExclusionTag != NAME_None;
+	}
+
+	if (bHasNonTrivialTags)
+	{
+		GEngine->OnLevelActorAdded().AddUObject(this, &UPCGComponent::OnActorAdded);
+		GEngine->OnLevelActorDeleted().AddUObject(this, &UPCGComponent::OnActorDeleted);
+	}
+}
+
+void UPCGComponent::TeardownExclusionCallbacks()
+{
+	bool bHasNonTrivialTags = false;
+
+	for (const FName& ExclusionTag : ExclusionTags)
+	{
+		bHasNonTrivialTags |= ExclusionTag != NAME_None;
+	}
+
+	if (bHasNonTrivialTags)
+	{
+		GEngine->OnLevelActorAdded().RemoveAll(this);
+		GEngine->OnLevelActorDeleted().RemoveAll(this);
+	}
+}
+
+bool UPCGComponent::ActorHasExcludedTag(AActor* InActor) const
+{
+	if (!InActor)
+	{
+		return false;
+	}
+
+	bool bHasExcludedTag = false;
+
+	for (const FName& Tag : InActor->Tags)
+	{
+		if (ExclusionTags.Contains(Tag))
+		{
+			bHasExcludedTag = true;
+			break;
+		}
+	}
+
+	return bHasExcludedTag;
+}
+
+void UPCGComponent::OnActorAdded(AActor* InActor)
+{
+	if (ActorHasExcludedTag(InActor))
+	{
+		UpdateExcludedActor(InActor, EExcludedActorChange::Added);
+	}
+}
+
+void UPCGComponent::OnActorDeleted(AActor* InActor)
+{
+	if (ActorHasExcludedTag(InActor))
+	{
+		UpdateExcludedActor(InActor, EExcludedActorChange::Deleted);
+	}
+}
+
+void UPCGComponent::OnActorMoved(AActor* InActor)
+{
+	if (InActor == GetOwner())
+	{
+		// Generalize the cases where we want the position to dirty the inputs
+		if (InputType == EPCGComponentInput::Landscape ||
+			(InputType == EPCGComponentInput::Actor && HasSplineComponent()))
+		{
+			DirtyGenerated(true);
+			Refresh();
+		}
+	}
+	else if (ActorHasExcludedTag(InActor))
+	{
+		UpdateExcludedActor(InActor, EExcludedActorChange::Changed);
+	}
+}
+
+void UPCGComponent::OnObjectPropertyChanged(UObject* InObject, FPropertyChangedEvent& InEvent)
+{
+	bool bValueSet = (InEvent.ChangeType == EPropertyChangeType::ValueSet);
+	// Special exception for actor tags, as we can't track otherwise an actor "losing" a tag
+	bool bActorTagChange = (InEvent.Property && InEvent.Property->GetFName() == GET_MEMBER_NAME_CHECKED(AActor, Tags));
+
+	if(!bValueSet && !bActorTagChange)
+	{
+		return;
+	}
+
+	// First, check if it's an actor
+	AActor* Actor = Cast<AActor>(InObject);
+
+	// Otherwise, if it's an actor component, track it as well
+	if (!Actor)
+	{
+		if (UActorComponent* ActorComponent = Cast<UActorComponent>(InObject))
+		{
+			Actor = ActorComponent->GetOwner();
+		}
+	}
+
+	if (Actor == GetOwner())
+	{
+		// Something has changed on the owner (including properties of this component)
+		// In the case of splines, this is where we'd get notified if some component properties (incl. spline vertices) have changed
+		if(HasSplineComponent())
+		{
+			DirtyGenerated(true);
+		}
+
+		Refresh();
+	}
+	else if (ActorHasExcludedTag(Actor) || (bActorTagChange && Actor == InObject))
+	{
+		// TODO: be more subtle about this, could be an addition, a change or deletion
+		UpdateExcludedActor(Actor, EExcludedActorChange::Changed);
+	}
+}
+
+void UPCGComponent::UpdateExcludedActor(AActor* InActor, EExcludedActorChange Change)
+{
+	// TODO: manage properly changes
+	DirtyGenerated(/*bInDirtyCachedInput=*/false, /*bInDirtyCachedExclusions=*/true);
+	Refresh();
+}
+
+void UPCGComponent::OnGraphChanged(UPCGGraph* InGraph, bool bIsStructural)
+{
+	if (InGraph == Graph)
+	{
+		DirtyGenerated();
+		Refresh();
+	}
+}
+
+void UPCGComponent::DirtyGenerated(bool bInDirtyCachedInput, bool bInDirtyCachedExclusions)
+{
+	bDirtyGenerated = true;
+
+	if (bInDirtyCachedInput)
+	{
+		CachedInputData = nullptr;
+		CachedActorData = nullptr;
+	}
+
+	if (bInDirtyCachedExclusions)
+	{
+		CachedExclusionData.Reset();
+	}
+
+	// For partitioned graph, we must forward the call to the partition actor
+	// Note that we do not need to forward "normal" dirty as these will be picked up by the local PCG components
+	// However, input changes / moves of the partitioned object will not be caught
+	// It would be possible for partitioned actors to add callbacks to their original component, but that inverses the processing flow
+	if ((bInDirtyCachedInput || bInDirtyCachedExclusions) && bActivated && IsPartitioned())
+	{
+		GetSubsystem()->DirtyGraph(this, LastGeneratedBounds, bInDirtyCachedInput, bInDirtyCachedExclusions);
+	}
+}
+
+void UPCGComponent::Refresh()
+{
+	// Following a change in some properties or in some spatial information related to this component,
+	// We need to regenerate the graph, depending of the state in the editor.
+	// In the case of a non-partitioned graph, we need to generate the graph only if it was previously generated & tagged for regeneration
+	// In the partitioned graph case, however, we need to do a bit more:
+	// 1. Regenerate the graph if it was previously generated & tagged for regeneration;
+	//  notice that the associated partition actors will not (and should not) have the regenerate flag on.
+	// 2. Otherwise, we need to update the partitioning if the spatial data has changed.
+	if (!bActivated)
+	{
+		if (IsPartitioned())
+		{
+			GetSubsystem()->DelayUnpartitionGraph(this);
+		}
+		else
+		{
+			bool bWasGenerated = bGenerated;
+			Cleanup(/*bRemoveComponents=*/false);
+			bGenerated = bWasGenerated;
+		}
+	}
+	else
+	{
+		if (bGenerated && bRegenerateInEditor)
+		{
+			Generate(/*bForce=*/false);
+		}
+		else if (IsPartitioned())
+		{
+			GetSubsystem()->DelayPartitionGraph(this);
+		}
+	}
+}
+#endif
+
+UPCGData* UPCGComponent::GetPCGData()
+{
+	if (!CachedInputData)
+	{
+		CachedInputData = CreatePCGData();
+	}
+
+	return CachedInputData;
+}
+
+UPCGData* UPCGComponent::GetActorPCGData()
+{
+	if (!CachedActorData)
+	{
+		CachedActorData = CreateActorPCGData();
+	}
+
+	return CachedActorData;
+}
+
+TArray<UPCGData*> UPCGComponent::GetPCGExclusionData()
+{
+	// TODO: replace with a boolean, unify.
+	if (CachedExclusionData.IsEmpty())
+	{
+		CachedExclusionData = CreatePCGExclusionData();
+	}
+
+	return CachedExclusionData;
+}
+
+TArray<UPCGData*> UPCGComponent::CreatePCGExclusionData()
+{
+	TArray<UPCGData*> ExclusionData;
+
+	// Find all actors with the exclusion tags
+	TSet<AActor*> ExclusionActors;
+	TArray<AActor*> PerTagExclusionActors;
+
+	for (const FName& ExclusionTag : ExclusionTags)
+	{
+		UGameplayStatics::GetAllActorsWithTag(GetWorld(), ExclusionTag, PerTagExclusionActors);
+		ExclusionActors.Append(PerTagExclusionActors);
+		PerTagExclusionActors.SetNum(0);
+	}
+	
+	// Build their PCG data
+	// Now, what we need to do depends on the dimension of the input data
+	const UPCGData* InputData = GetPCGData();
+	const UPCGSpatialData* InputSpatialData = Cast<const UPCGSpatialData>(InputData);
+
+	// Similar in logic to the input data mechanism.
+	for (AActor* ExclusionActor : ExclusionActors)
+	{
+		UPCGData* ActorData = CreateActorPCGData(ExclusionActor);
+		UPCGSpatialData* ActorSpatialData = Cast<UPCGSpatialData>(ActorData);
+
+		if (InputSpatialData && ActorSpatialData)
+		{
+			// Create intersection or projection depending on the dimension
+			// TODO: there's an ambiguity here when it's the same dimension.
+			// For volumes, we'd expect an intersection, for surfaces we'd expect a projection
+			if (ActorSpatialData->GetDimension() > InputSpatialData->GetDimension())
+			{
+				ExclusionData.Add(ActorSpatialData->IntersectWith(InputSpatialData));
+			}
+			else
+			{
+				ExclusionData.Add(ActorSpatialData->ProjectOn(InputSpatialData));
+			}
+		}
+	}
+
+	return ExclusionData;
+}
+
+UPCGData* UPCGComponent::CreateActorPCGData()
+{
+	return CreateActorPCGData(GetOwner());
+}
+
+UPCGData* UPCGComponent::CreateActorPCGData(AActor* Actor)
+{
+	TRACE_CPUPROFILER_EVENT_SCOPE(UPCGComponent::CreateActorPCGData)
+	// In this case, we'll build the data type that's closest to known actor types
+	// TODO: add factory for extensibility
+	if (APCGPartitionActor* PartitionActor = Cast<APCGPartitionActor>(Actor))
+	{
+		check(GetOwner() == Actor); // Invalid processing otherwise because of the this usage
+		if (UPCGComponent* OriginalComponent = PartitionActor->GetOriginalComponent(this))
+		{
+			check(OriginalComponent->IsPartitioned());
+			// TODO: cache/share the original component's actor pcg data
+			if (const UPCGSpatialData* OriginalComponentSpatialData = Cast<const UPCGSpatialData>(OriginalComponent->GetActorPCGData()))
+			{
+				UPCGVolumeData* Data = NewObject<UPCGVolumeData>(this);
+				Data->Initialize(PartitionActor->GetFixedBounds(), PartitionActor);
+
+				return Data->IntersectWith(OriginalComponentSpatialData);
+			}
+		}
+
+		// TODO: review this once we support non-spatial data?
+		return nullptr;
+	}
+	else if (ALandscapeProxy* Landscape = Cast<ALandscapeProxy>(Actor))
+	{
+		UPCGLandscapeData* Data = NewObject<UPCGLandscapeData>(this);
+		Data->Initialize(Landscape, GetGridBounds(Actor));
+
+		return Data;
+	}
+	else if (AVolume* Volume = Cast<AVolume>(Actor))
+	{
+		UPCGVolumeData* Data = NewObject<UPCGVolumeData>(this);
+		Data->Initialize(Volume);
+
+		return Data;
+	}
+	else if (ALandscapeSplineActor* LandscapeSpline = Cast<ALandscapeSplineActor>(Actor))
+	{
+		TArray<ULandscapeSplinesComponent*> LandscapeSplines;
+		Actor->GetComponents<ULandscapeSplinesComponent>(LandscapeSplines);
+
+		if (LandscapeSplines.Num() > 1)
+		{
+			UPCGUnionData* Data = NewObject<UPCGUnionData>(this);
+			for (ULandscapeSplinesComponent* SplineComponent : LandscapeSplines)
+			{
+				UPCGLandscapeSplineData* SplineData = NewObject<UPCGLandscapeSplineData>(this);
+				SplineData->Initialize(SplineComponent);
+				Data->AddData(SplineData);
+			}
+
+			return Data;
+		}
+		else if (LandscapeSplines.Num() == 1)
+		{
+			UPCGLandscapeSplineData* Data = NewObject<UPCGLandscapeSplineData>(this);
+			Data->Initialize(LandscapeSplines[0]);
+			return Data;
+		}
+		else
+		{
+			return nullptr;
+		}
+	}
+	else // TODO how about surfaces, splines?
+	{
+		TArray<USplineComponent*> Splines;
+		Actor->GetComponents<USplineComponent>(Splines);
+
+		if (Splines.Num() > 1)
+		{
+			UPCGUnionData* Data = NewObject<UPCGUnionData>(this);
+			for (USplineComponent* SplineComponent : Splines)
+			{
+				UPCGSplineData* SplineData = NewObject<UPCGSplineData>(this);
+				SplineData->Initialize(SplineComponent);
+				Data->AddData(SplineData);
+			}
+
+			return Data;
+		}
+		else if (Splines.Num() == 1)
+		{
+			UPCGSplineData* Data = NewObject<UPCGSplineData>(this);
+			Data->Initialize(Splines[0]);
+			return Data;
+		}
+		else // Default behavior on unknown actors
+		{
+			UPCGPointData* Data = NewObject<UPCGPointData>(this);
+			Data->Initialize(Actor);
+			return Data;
+		}
+	}
+}
+
+UPCGData* UPCGComponent::CreatePCGData()
+{
+	TRACE_CPUPROFILER_EVENT_SCOPE(UPCGComponent::CreatePCGData)
+	AActor* Actor = GetOwner();
+	check(Actor);
+
+	// Construct proper input based on input type
+	if (InputType == EPCGComponentInput::Actor)
+	{
+		return GetActorPCGData();
+	}
+	else if (InputType == EPCGComponentInput::Landscape)
+	{
+		UPCGData* ActorData = GetActorPCGData();
+
+		if (ALandscapeProxy* Landscape = Cast<ALandscapeProxy>(this))
+		{
+			// Current actor is a landscape, we don't need to do an additional projection
+			return ActorData;
+		}
+
+		const UPCGSpatialData* ActorSpatialData = Cast<const UPCGSpatialData>(ActorData);
+
+		if (!ActorSpatialData)
+		{
+			// TODO ? support non-spatial data on landscape?
+			return nullptr;
+		}
+
+		const FBox ActorDataBounds = ActorSpatialData->GetBounds();
+		ALandscapeProxy* Landscape = PCGHelpers::GetLandscape(Actor->GetWorld(), ActorDataBounds);
+
+		if (!Landscape)
+		{
+			// No landscape found
+			return nullptr;
+		}
+
+		const FBox LandscapeBounds = GetGridBounds(Landscape);
+		check(LandscapeBounds.IsValid);
+
+		const FBox OverlappedBounds = LandscapeBounds.Overlap(ActorDataBounds);
+
+		if (!OverlappedBounds.IsValid)
+		{
+			return nullptr;
+		}
+
+		UPCGLandscapeData* LandscapeData = NewObject<UPCGLandscapeData>(this);
+		// Since we're not sharing this landscape data with any other users, we can limit the bounds earlier
+		LandscapeData->Initialize(Landscape, OverlappedBounds);
+
+		// Need to override target actor for this one, not the landscape
+		LandscapeData->TargetActor = Actor;
+
+		// Decide whether to intersect or project
+		// Currently, it makes sense to intersect only for volumes;
+		// Note that we don't currently check for a volume object but only on dimension
+		// so intersections (such as volume X partition actor) get picked up properly
+		if (ActorSpatialData->GetDimension() >= 3)
+		{
+			return LandscapeData->IntersectWith(ActorSpatialData);
+		}
+		else
+		{
+			return ActorSpatialData->ProjectOn(LandscapeData);
+		}
+	}
+	else
+	{
+		// In this case, the input data will be provided in some other form,
+		// Most likely to be stored in the PCG data grid.
+		return nullptr;
+	}
+}
+
+FBox UPCGComponent::GetGridBounds() const
+{
+	return GetGridBounds(GetOwner());
+}
+
+FBox UPCGComponent::GetGridBounds(AActor* Actor) const
+{
+	check(Actor);
+
+	FBox Bounds(EForceInit::ForceInit);
+
+	if (APCGPartitionActor* PartitionActor = Cast<APCGPartitionActor>(Actor))
+	{
+		// First, get the bounds from the partition actor
+		Bounds = PartitionActor->GetFixedBounds();
+
+		// Then intersect with the original component's bounds.
+		if (const UPCGComponent* OriginalComponent = PartitionActor->GetOriginalComponent(this))
+		{
+			Bounds = Bounds.Overlap(OriginalComponent->GetGridBounds());
+		}
+	}
+	// TODO: this might not work in non-editor builds
+#if WITH_EDITOR
+	else if (ALandscapeProxy* LandscapeActor = Cast<ALandscape>(Actor))
+	{
+		Bounds = PCGHelpers::GetLandscapeBounds(LandscapeActor);
+	}
+#endif
+	else
+	{
+		Bounds = PCGHelpers::GetActorBounds(Actor);
+	}
+
+	return Bounds;
+}
+
+UPCGSubsystem* UPCGComponent::GetSubsystem() const
+{
+	return GetOwner() && GetOwner()->GetWorld() ? GetOwner()->GetWorld()->GetSubsystem<UPCGSubsystem>() : nullptr;
+}
+
+#undef LOCTEXT_NAMESPACE
\ No newline at end of file
diff --git a/Engine/Plugins/Experimental/PCG/Source/PCG/Private/PCGData.cpp b/Engine/Plugins/Experimental/PCG/Source/PCG/Private/PCGData.cpp
new file mode 100644
--- /dev/null
+++ b/Engine/Plugins/Experimental/PCG/Source/PCG/Private/PCGData.cpp
@@ -0,0 +1,46 @@
+// Copyright Epic Games, Inc. All Rights Reserved.
+
+#include "PCGData.h"
+#include "PCGSettings.h"
+
+TArray<FPCGTaggedData> FPCGDataCollection::GetInputs() const
+{
+	return TaggedData.FilterByPredicate([](const FPCGTaggedData& Data) {
+		return Data.Usage == EPCGDataUsage::Input;
+		});
+}
+
+TArray<FPCGTaggedData> FPCGDataCollection::GetExclusions() const
+{
+	return TaggedData.FilterByPredicate([](const FPCGTaggedData& Data) {
+		return Data.Usage == EPCGDataUsage::Exclusion;
+		});
+}
+
+const UPCGSettings* FPCGDataCollection::GetSettings(const UPCGSettings* InDefaultSettings) const
+{
+	if (!InDefaultSettings)
+	{
+		return GetSettings<UPCGSettings>();
+	}
+	else
+	{
+		const FPCGTaggedData* MatchingData = TaggedData.FindByPredicate([InDefaultSettings](const FPCGTaggedData& Data) {
+			return Data.Usage == EPCGDataUsage::Settings &&
+				(Data.Data->GetClass() == InDefaultSettings->GetClass() ||
+					Data.Data->GetClass()->IsChildOf(InDefaultSettings->GetClass()));
+			});
+
+		return MatchingData ? Cast<const UPCGSettings>(MatchingData->Data) : InDefaultSettings;
+	}
+}
+
+TArray<FPCGTaggedData> UPCGDataFunctionLibrary::GetInputs(const FPCGDataCollection& InCollection)
+{
+	return InCollection.GetInputs();
+}
+
+TArray<FPCGTaggedData> UPCGDataFunctionLibrary::GetExclusions(const FPCGDataCollection& InCollection)
+{
+	return InCollection.GetExclusions();
+}
\ No newline at end of file
diff --git a/Engine/Plugins/Experimental/PCG/Source/PCG/Private/PCGElement.cpp b/Engine/Plugins/Experimental/PCG/Source/PCG/Private/PCGElement.cpp
new file mode 100644
--- /dev/null
+++ b/Engine/Plugins/Experimental/PCG/Source/PCG/Private/PCGElement.cpp
@@ -0,0 +1,13 @@
+// Copyright Epic Games, Inc. All Rights Reserved.
+
+#include "PCGElement.h"
+#include "PCGContext.h"
+
+FPCGContextPtr FSimplePCGElement::Initialize(const FPCGDataCollection& InputData, const UPCGComponent* SourceComponent)
+{
+	FPCGContextPtr Context = MakeShared<FPCGContext>();
+	Context->InputData = InputData;
+	Context->SourceComponent = SourceComponent;
+
+	return Context;
+}
\ No newline at end of file
diff --git a/Engine/Plugins/Experimental/PCG/Source/PCG/Private/PCGGraph.cpp b/Engine/Plugins/Experimental/PCG/Source/PCG/Private/PCGGraph.cpp
new file mode 100644
--- /dev/null
+++ b/Engine/Plugins/Experimental/PCG/Source/PCG/Private/PCGGraph.cpp
@@ -0,0 +1,203 @@
+// Copyright Epic Games, Inc. All Rights Reserved.
+
+#include "PCGGraph.h"
+#include "PCGSubsystem.h"
+
+#if WITH_EDITOR
+#include "Editor.h"
+#endif
+
+UPCGGraph::UPCGGraph(const FObjectInitializer& ObjectInitializer)
+	: Super(ObjectInitializer)
+{
+	InputNode = ObjectInitializer.CreateDefaultSubobject<UPCGNode>(this, TEXT("DefaultInputNode"));
+	OutputNode = ObjectInitializer.CreateDefaultSubobject<UPCGNode>(this, TEXT("DefaultOutputNode"));
+
+	// Note: default connection from input to output
+	// should be added when creating from scratch,
+	// but not when using a blueprint construct script.
+	//InputNode->ConnectTo(OutputNode);
+	//OutputNode->ConnectFrom(InputNode);
+}
+
+void UPCGGraph::PostLoad()
+{
+	Super::PostLoad();
+
+#if WITH_EDITOR
+	for (UPCGNode* Node : Nodes)
+	{
+		if (Node)
+		{
+			Node->OnNodeSettingsChangedDelegate.AddUObject(this, &UPCGGraph::OnSettingsChanged);
+		}
+	}
+#endif
+}
+
+void UPCGGraph::BeginDestroy()
+{
+#if WITH_EDITOR
+	for (UPCGNode* Node : Nodes)
+	{
+		if (Node)
+		{
+			Node->OnNodeSettingsChangedDelegate.RemoveAll(this);
+		}
+	}
+
+	// Notify the compiler to remove this graph from its cache
+	UWorld* World = GEditor->GetEditorWorldContext().World();
+	UPCGSubsystem* PCGSubsystem = World ? World->GetSubsystem<UPCGSubsystem>() : nullptr;
+	if (PCGSubsystem)
+	{
+		PCGSubsystem->NotifyGraphChanged(this);
+	}
+#endif
+
+	Super::BeginDestroy();
+}
+
+#if WITH_EDITOR
+void UPCGGraph::InitializeFromTemplate()
+{
+	Modify();
+
+	InputNode->OutboundNodes.Reset();
+	OutputNode->InboundNodes.Reset();
+
+	for (UPCGNode* Node : Nodes)
+	{
+		if (Node)
+		{
+			Node->OnNodeSettingsChangedDelegate.RemoveAll(this);
+		}
+	}
+
+	Nodes.Reset();
+
+	if (GraphTemplate)
+	{
+		bEnableGraphChangeNotifications = false;
+		Cast<UPCGGraphSetupBP>(GraphTemplate->GetDefaultObject())->Setup(this);
+		bEnableGraphChangeNotifications = true;
+	}
+
+	NotifyGraphChanged(/*bIsStructural=*/true);
+}
+#endif
+
+UPCGNode* UPCGGraph::AddNodeOfType(TSubclassOf<class UPCGSettings> InSettingsClass)
+{
+	UPCGSettings* Settings = NewObject<UPCGSettings>(GetTransientPackage(), InSettingsClass);
+
+	if (!Settings)
+	{
+		return nullptr;
+	}
+
+	UPCGNode* Node = Settings->CreateNode();
+
+	if (Node)
+	{
+		Modify();
+
+		// Assign settings to node
+		Node->DefaultSettings = Settings;
+		Settings->Rename(nullptr, Node);
+		Settings->SetFlags(RF_Transactional);
+
+		// Reparent node to this graph
+		Node->Rename(nullptr, this);
+		Nodes.Add(Node);
+
+#if WITH_EDITOR
+		Node->OnNodeSettingsChangedDelegate.AddUObject(this, &UPCGGraph::OnSettingsChanged);
+		NotifyGraphChanged(/*bIsStructural=*/true);
+#endif
+	}
+
+	return Node;
+}
+
+UPCGNode* UPCGGraph::AddNode(UPCGSettings* InSettings)
+{
+	if (!InSettings)
+	{
+		return nullptr;
+	}
+
+	UPCGNode* Node = InSettings->CreateNode();
+
+	if (Node)
+	{
+		Modify();
+
+		// Assign settings to node & reparent
+		Node->DefaultSettings = InSettings;
+
+		// Reparent node to this graph
+		Node->Rename(nullptr, this);
+		Nodes.Add(Node);
+
+#if WITH_EDITOR
+		Node->OnNodeSettingsChangedDelegate.AddUObject(this, &UPCGGraph::OnSettingsChanged);
+		NotifyGraphChanged(/*bIsStructural=*/true);
+#endif
+	}
+
+	return Node;
+}
+
+UPCGNode* UPCGGraph::AddEdge(UPCGNode* From, UPCGNode* To)
+{
+	if (!From || !To)
+	{
+		// TODO: log error
+		return To;
+	}
+
+	Modify();
+
+	From->ConnectTo(To);
+	To->ConnectFrom(From);
+
+#if WITH_EDITOR
+	NotifyGraphChanged(/*bIsStructural=*/true);
+#endif
+
+	return To;
+}
+
+bool UPCGGraph::Contains(UPCGNode* Node) const
+{
+	return Node == InputNode || Node == OutputNode || Nodes.Contains(Node);
+}
+
+#if WITH_EDITOR
+void UPCGGraph::NotifyGraphChanged(bool bIsStructural)
+{
+	if (!bEnableGraphChangeNotifications)
+	{
+		return;
+	}
+
+	// Notify the subsystem/compiler cache before so it gets recompiled properly
+	if (bIsStructural)
+	{
+		UWorld* World = GEditor->GetEditorWorldContext().World();
+		UPCGSubsystem* PCGSubsystem = World ? World->GetSubsystem<UPCGSubsystem>() : nullptr;
+		if (PCGSubsystem)
+		{
+			PCGSubsystem->NotifyGraphChanged(this);
+		}
+	}
+
+	OnGraphChangedDelegate.Broadcast(this, /*bIsStructural=*/bIsStructural);
+}
+
+void UPCGGraph::OnSettingsChanged(UPCGNode* InNode)
+{
+	NotifyGraphChanged(/*bIsStructural=*/false);
+}
+#endif // WITH_EDITOR
diff --git a/Engine/Plugins/Experimental/PCG/Source/PCG/Private/PCGGraphSetupBP.cpp b/Engine/Plugins/Experimental/PCG/Source/PCG/Private/PCGGraphSetupBP.cpp
new file mode 100644
--- /dev/null
+++ b/Engine/Plugins/Experimental/PCG/Source/PCG/Private/PCGGraphSetupBP.cpp
@@ -0,0 +1,8 @@
+// Copyright Epic Games, Inc. All Rights Reserved.
+
+#include "PCGGraphSetupBP.h"
+
+void UPCGGraphSetupBP::Setup_Implementation(UPCGGraph* Graph)
+{
+	// Do nothing in this base-class
+}
\ No newline at end of file
diff --git a/Engine/Plugins/Experimental/PCG/Source/PCG/Private/PCGHelpers.cpp b/Engine/Plugins/Experimental/PCG/Source/PCG/Private/PCGHelpers.cpp
new file mode 100644
--- /dev/null
+++ b/Engine/Plugins/Experimental/PCG/Source/PCG/Private/PCGHelpers.cpp
@@ -0,0 +1,102 @@
+// Copyright Epic Games, Inc. All Rights Reserved.
+
+#include "PCGHelpers.h"
+
+#include "PCGSubsystem.h"
+#include "PCGWorldActor.h"
+
+#include "GameFramework/Actor.h"
+#include "Landscape.h"
+#include "UObject/UObjectIterator.h"
+#include "WorldPartition/WorldPartition.h"
+
+namespace PCGHelpers
+{
+	int ComputeSeed(int A)
+	{
+		return (A * 196314165U) + 907633515U;
+	}
+
+	int ComputeSeed(int A, int B)
+	{
+		return ((A * 196314165U) + 907633515U) ^ ((B * 73148459U) + 453816763U);
+	}
+
+	int ComputeSeed(int A, int B, int C)
+	{
+		return ((A * 196314165U) + 907633515U) ^ ((B * 73148459U) + 453816763U) ^ ((C * 34731343U) + 453816743U);
+	}
+
+	FBox GetActorBounds(AActor* InActor)
+	{
+		// Specialized version of GetComponentsBoundingBox that skips over PCG generated components
+		// This is to ensure stable bounds and no timing issues (cleared ISMs, etc.)
+		FBox Box(EForceInit::ForceInit);
+
+		const bool bNonColliding = true;
+		const bool bIncludeFromChildActors = true;
+
+		InActor->ForEachComponent<UPrimitiveComponent>(bIncludeFromChildActors, [&](const UPrimitiveComponent* InPrimComp)
+			{
+				// Note: we omit the IsRegistered check here (e.g. InPrimComp->IsRegistered() )
+				// since this can be called in a scope where the components are temporarily unregistered
+				if ((bNonColliding || InPrimComp->IsCollisionEnabled()) &&
+					!InPrimComp->ComponentTags.Contains(DefaultPCGTag))
+				{
+					Box += InPrimComp->Bounds.GetBox();
+				}
+			});
+
+		return Box;
+	}
+
+	FBox GetLandscapeBounds(ALandscapeProxy* InLandscape)
+	{
+		check(InLandscape);
+
+		if (ALandscape* Landscape = Cast<ALandscape>(InLandscape))
+		{
+#if WITH_EDITOR
+			return Landscape->GetCompleteBounds();
+#else
+			return Landscape->GetLoadedBounds();
+#endif
+		}
+		else
+		{
+			return GetActorBounds(InLandscape);
+		}
+	}
+
+	ALandscape* GetLandscape(UWorld* InWorld, const FBox& InBounds)
+	{
+		ALandscape* Landscape = nullptr;
+
+		if (!InBounds.IsValid)
+		{
+			return Landscape;
+		}
+
+		for (TObjectIterator<ALandscape> It; It; ++It)
+		{
+			if (It->GetWorld() == InWorld)
+			{
+				const FBox LandscapeBounds = GetLandscapeBounds(*It);
+				if (LandscapeBounds.IsValid && LandscapeBounds.Intersect(InBounds))
+				{
+					Landscape = (*It);
+					break;
+				}
+			}
+		}
+
+		return Landscape;
+	}
+
+#if WITH_EDITOR
+	APCGWorldActor* GetPCGWorldActor(UWorld* InWorld)
+	{
+		return InWorld && InWorld->GetSubsystem<UPCGSubsystem>() ? InWorld->GetSubsystem<UPCGSubsystem>()->GetPCGWorldActor() : nullptr;
+	}
+#endif
+}
diff --git a/Engine/Plugins/Experimental/PCG/Source/PCG/Private/PCGHelpers.h b/Engine/Plugins/Experimental/PCG/Source/PCG/Private/PCGHelpers.h
new file mode 100644
--- /dev/null
+++ b/Engine/Plugins/Experimental/PCG/Source/PCG/Private/PCGHelpers.h
@@ -0,0 +1,30 @@
+// Copyright Epic Games, Inc. All Rights Reserved.
+
+#pragma once
+
+#include "CoreMinimal.h"
+
+class AActor;
+class APCGWorldActor;
+class ALandscape;
+class ALandscapeProxy;
+class UWorld;
+
+namespace PCGHelpers
+{
+	/** Tag that will be added on every component generated through the PCG system */
+	const FName DefaultPCGTag = TEXT("PCG Generated Component");
+
+	int ComputeSeed(int A);
+	int ComputeSeed(int A, int B);
+	int ComputeSeed(int A, int B, int C);
+
+	FBox GetActorBounds(AActor* InActor);
+	FBox GetLandscapeBounds(ALandscapeProxy* InLandscape);
+
+	ALandscape* GetLandscape(UWorld* InWorld, const FBox& InActorBounds);
+
+#if WITH_EDITOR
+	APCGWorldActor* GetPCGWorldActor(UWorld* InWorld);
+#endif
+};
\ No newline at end of file
diff --git a/Engine/Plugins/Experimental/PCG/Source/PCG/Private/PCGModule.cpp b/Engine/Plugins/Experimental/PCG/Source/PCG/Private/PCGModule.cpp
new file mode 100644
--- /dev/null
+++ b/Engine/Plugins/Experimental/PCG/Source/PCG/Private/PCGModule.cpp
@@ -0,0 +1,29 @@
+// Copyright Epic Games, Inc. All Rights Reserved.
+
+#include "PCGModule.h"
+#include "Modules/ModuleInterface.h"
+
+class FPCGModule : public IModuleInterface
+{
+public:
+	//~ IModuleInterface implementation
+	virtual void StartupModule() override
+	{
+
+	}
+
+	virtual void ShutdownModule() override
+	{
+
+	}
+
+	virtual bool SupportsDynamicReloading() override
+	{
+		return true;
+	}
+	//~ End IModuleInterface implementation
+};
+
+IMPLEMENT_MODULE(FPCGModule, PCG);
+
+DEFINE_LOG_CATEGORY(LogPCG);
diff --git a/Engine/Plugins/Experimental/PCG/Source/PCG/Private/PCGNode.cpp b/Engine/Plugins/Experimental/PCG/Source/PCG/Private/PCGNode.cpp
new file mode 100644
--- /dev/null
+++ b/Engine/Plugins/Experimental/PCG/Source/PCG/Private/PCGNode.cpp
@@ -0,0 +1,106 @@
+// Copyright Epic Games, Inc. All Rights Reserved.
+
+#include "PCGNode.h"
+#include "PCGGraph.h"
+#include "PCGSettings.h"
+
+UPCGNode::UPCGNode(const FObjectInitializer& ObjectInitializer)
+	: Super(ObjectInitializer)
+{
+	DefaultSettings = ObjectInitializer.CreateDefaultSubobject<UPCGTrivialSettings>(this, TEXT("DefaultNodeSettings"));
+}
+
+void UPCGNode::PostLoad()
+{
+	Super::PostLoad();
+
+#if WITH_EDITOR
+	if (DefaultSettings)
+	{
+		DefaultSettings->OnSettingsChangedDelegate.AddUObject(this, &UPCGNode::OnSettingsChanged);
+	}
+#endif
+}
+
+void UPCGNode::BeginDestroy()
+{
+#if WITH_EDITOR
+	if (DefaultSettings)
+	{
+		DefaultSettings->OnSettingsChangedDelegate.RemoveAll(this);
+	}
+#endif
+
+	Super::BeginDestroy();
+}
+
+UPCGGraph* UPCGNode::GetGraph() const
+{
+	return Cast<UPCGGraph>(GetOuter());
+}
+
+UPCGNode* UPCGNode::AddEdgeTo(UPCGNode* To)
+{
+	check(GetGraph());
+	if (UPCGGraph* Graph = GetGraph())
+	{
+		return Graph->AddEdge(this, To);
+	}
+	else
+	{
+		return nullptr;
+	}
+}
+
+void UPCGNode::ConnectTo(UPCGNode* InSuccessor)
+{
+	OutboundNodes.AddUnique(InSuccessor);
+}
+
+void UPCGNode::ConnectFrom(UPCGNode* InPredecessor)
+{
+	InboundNodes.AddUnique(InPredecessor);
+}
+
+#if WITH_EDITOR
+
+void UPCGNode::PreEditChange(FProperty* PropertyAboutToChange)
+{
+	if (PropertyAboutToChange && PropertyAboutToChange->GetFName() == GET_MEMBER_NAME_CHECKED(UPCGNode, DefaultSettings))
+	{
+		if (DefaultSettings)
+		{
+			DefaultSettings->OnSettingsChangedDelegate.RemoveAll(this);
+		}
+	}
+
+	Super::PreEditChange(PropertyAboutToChange);
+}
+
+void UPCGNode::PostEditChangeProperty(struct FPropertyChangedEvent& PropertyChangedEvent)
+{
+	Super::PostEditChangeProperty(PropertyChangedEvent);
+
+	if (PropertyChangedEvent.Property && PropertyChangedEvent.Property->GetFName() == GET_MEMBER_NAME_CHECKED(UPCGNode, DefaultSettings))
+	{
+		if (DefaultSettings)
+		{
+			DefaultSettings->OnSettingsChangedDelegate.AddUObject(this, &UPCGNode::OnSettingsChanged);
+		}
+	}
+
+	if (DefaultSettings)
+	{
+		OnSettingsChanged(DefaultSettings);
+	}
+}
+
+void UPCGNode::OnSettingsChanged(UPCGSettings* InSettings)
+{
+	if (InSettings == DefaultSettings)
+	{
+		OnNodeSettingsChangedDelegate.Broadcast(this);
+	}
+}
+
+#endif // WITH_EDITOR
diff --git a/Engine/Plugins/Experimental/PCG/Source/PCG/Private/PCGPoint.cpp b/Engine/Plugins/Experimental/PCG/Source/PCG/Private/PCGPoint.cpp
new file mode 100644
--- /dev/null
+++ b/Engine/Plugins/Experimental/PCG/Source/PCG/Private/PCGPoint.cpp
@@ -0,0 +1,15 @@
+// Copyright Epic Games, Inc. All Rights Reserved.
+
+#include "PCGPoint.h"
+
+FPCGPoint::FPCGPoint(const FTransform& InTransform, float InDensity, int32 InSeed)
+	: Transform(InTransform), Density(InDensity), Seed(InSeed)
+{
+	// TODO: this is a placeholder
+	Extents = FVector::One();
+}
+
+FBoxSphereBounds FPCGPoint::GetBounds() const
+{
+	return FBoxSphereBounds(FBox::BuildAABB(FVector::Zero(), Extents).TransformBy(Transform));
+}
\ No newline at end of file
diff --git a/Engine/Plugins/Experimental/PCG/Source/PCG/Private/PCGSettings.cpp b/Engine/Plugins/Experimental/PCG/Source/PCG/Private/PCGSettings.cpp
new file mode 100644
--- /dev/null
+++ b/Engine/Plugins/Experimental/PCG/Source/PCG/Private/PCGSettings.cpp
@@ -0,0 +1,48 @@
+// Copyright Epic Games, Inc. All Rights Reserved.
+
+#include "PCGSettings.h"
+#include "PCGNode.h"
+
+FPCGElementPtr UPCGSettings::GetElement() const
+{
+	if (!CachedElement)
+	{
+		CacheLock.Lock();
+
+		if (!CachedElement)
+		{
+			CachedElement = CreateElement();
+		}
+
+		CacheLock.Unlock();
+	}
+
+	return CachedElement;
+}
+
+UPCGNode* UPCGSettings::CreateNode() const
+{
+	return NewObject<UPCGNode>();
+}
+
+FPCGElementPtr UPCGTrivialSettings::CreateElement() const
+{
+	return MakeShared<FPCGTrivialElement>();
+}
+
+bool FPCGTrivialElement::Execute(FPCGContextPtr Context) const
+{
+	// Pass-through
+	Context->OutputData = Context->InputData;
+	return true;
+}
+
+#if WITH_EDITOR
+
+void UPCGSettings::PostEditChangeProperty(FPropertyChangedEvent& PropertyChangedEvent)
+{
+	Super::PostEditChangeProperty(PropertyChangedEvent);
+	OnSettingsChangedDelegate.Broadcast(this);
+}
+
+#endif // WITH_EDITOR
\ No newline at end of file
diff --git a/Engine/Plugins/Experimental/PCG/Source/PCG/Private/PCGSubgraph.cpp b/Engine/Plugins/Experimental/PCG/Source/PCG/Private/PCGSubgraph.cpp
new file mode 100644
--- /dev/null
+++ b/Engine/Plugins/Experimental/PCG/Source/PCG/Private/PCGSubgraph.cpp
@@ -0,0 +1,82 @@
+// Copyright Epic Games, Inc. All Rights Reserved.
+
+#include "PCGSubgraph.h"
+#include "PCGGraph.h"
+
+void UPCGSubgraphSettings::PostLoad()
+{
+	Super::PostLoad();
+
+#if WITH_EDITOR
+	if (Subgraph)
+	{
+		Subgraph->OnGraphChangedDelegate.AddUObject(this, &UPCGSubgraphSettings::OnSubgraphChanged);
+	}
+#endif
+}
+
+void UPCGSubgraphSettings::BeginDestroy()
+{
+#if WITH_EDITOR
+	if (Subgraph)
+	{
+		Subgraph->OnGraphChangedDelegate.RemoveAll(this);
+	}
+#endif
+
+	Super::BeginDestroy();
+}
+
+UPCGNode* UPCGSubgraphSettings::CreateNode() const
+{
+	return NewObject<UPCGSubgraphNode>();
+}
+
+FPCGElementPtr UPCGSubgraphSettings::CreateElement() const
+{
+	return MakeShared<FPCGTrivialElement>();
+}
+
+TObjectPtr<UPCGGraph> UPCGSubgraphNode::GetGraph() const
+{
+	TObjectPtr<UPCGSubgraphSettings> Settings = Cast<UPCGSubgraphSettings>(DefaultSettings);
+	return Settings ? Settings->Subgraph : nullptr;
+}
+
+#if WITH_EDITOR
+
+void UPCGSubgraphSettings::PreEditChange(FProperty* PropertyAboutToChange)
+{
+	if (PropertyAboutToChange && PropertyAboutToChange->GetFName() == GET_MEMBER_NAME_CHECKED(UPCGSubgraphSettings, Subgraph))
+	{
+		if (Subgraph)
+		{
+			Subgraph->OnGraphChangedDelegate.RemoveAll(this);
+		}
+	}
+
+	Super::PreEditChange(PropertyAboutToChange);
+}
+
+void UPCGSubgraphSettings::PostEditChangeProperty(struct FPropertyChangedEvent& PropertyChangedEvent)
+{
+	Super::PostEditChangeProperty(PropertyChangedEvent);
+
+	if (PropertyChangedEvent.Property->GetFName() == GET_MEMBER_NAME_CHECKED(UPCGSubgraphSettings, Subgraph))
+	{
+		if (Subgraph)
+		{
+			Subgraph->OnGraphChangedDelegate.AddUObject(this, &UPCGSubgraphSettings::OnSubgraphChanged);
+		}
+	}
+}
+
+void UPCGSubgraphSettings::OnSubgraphChanged(UPCGGraph* InGraph, bool bIsStructural)
+{
+	if (InGraph == Subgraph)
+	{
+		OnSettingsChangedDelegate.Broadcast(this);
+	}
+}
+
+#endif // WITH_EDITOR
diff --git a/Engine/Plugins/Experimental/PCG/Source/PCG/Private/PCGSubsystem.cpp b/Engine/Plugins/Experimental/PCG/Source/PCG/Private/PCGSubsystem.cpp
new file mode 100644
--- /dev/null
+++ b/Engine/Plugins/Experimental/PCG/Source/PCG/Private/PCGSubsystem.cpp
@@ -0,0 +1,512 @@
+// Copyright Epic Games, Inc. All Rights Reserved.
+
+#include "PCGSubsystem.h"
+#include "PCGComponent.h"
+#include "PCGWorldActor.h"
+#include "Graph/PCGGraphExecutor.h"
+#include "Grid/PCGPartitionActor.h"
+
+#include "Engine/World.h"
+
+#if WITH_EDITOR
+#include "Editor.h"
+#include "Engine/Engine.h"
+#include "ActorPartition/ActorPartitionSubsystem.h"
+#include "WorldPartition/WorldPartition.h"
+#include "WorldPartition/WorldPartitionHelpers.h"
+#include "WorldPartition/ActorPartition/PartitionActorDesc.h"
+#include "ObjectTools.h"
+#endif
+
+void UPCGSubsystem::Deinitialize()
+{
+	// Cancel all tasks
+	// TODO
+	delete GraphExecutor;
+	GraphExecutor = nullptr;
+
+	Super::Deinitialize();
+}
+
+void UPCGSubsystem::PostInitialize()
+{
+	Super::PostInitialize();
+
+	// Initialize graph executor
+	check(!GraphExecutor);
+	GraphExecutor = new FPCGGraphExecutor();
+	// gather things.. ?
+	// TODO	
+}
+
+void UPCGSubsystem::Tick(float DeltaSeconds)
+{
+	Super::Tick(DeltaSeconds);
+
+	// If we have any tasks to execute, schedule some
+	GraphExecutor->Execute();
+}
+
+APCGWorldActor* UPCGSubsystem::GetPCGWorldActor()
+{
+#if WITH_EDITOR
+	if (!PCGWorldActor)
+	{
+		PCGWorldActor = APCGWorldActor::CreatePCGWorldActor(GetWorld());
+	}
+#endif
+
+	return PCGWorldActor;
+}
+
+void UPCGSubsystem::RegisterPCGWorldActor(APCGWorldActor* InActor)
+{
+	check(!PCGWorldActor || PCGWorldActor == InActor);
+	PCGWorldActor = InActor;
+}
+
+void UPCGSubsystem::UnregisterPCGWorldActor(APCGWorldActor* InActor)
+{
+	check(PCGWorldActor == InActor);
+	if (PCGWorldActor == InActor)
+	{
+		PCGWorldActor = nullptr;
+	}
+}
+
+FPCGTaskId UPCGSubsystem::ScheduleComponent(UPCGComponent* PCGComponent, const TArray<FPCGTaskId>& Dependencies)
+{
+	check(GraphExecutor);
+	FPCGTaskId ProcessTaskId = GraphExecutor->Schedule(PCGComponent, Dependencies);
+
+	if (ProcessTaskId != InvalidTaskId)
+	{
+#if WITH_EDITOR
+		return GraphExecutor->ScheduleGeneric([PCGComponent]() {
+			const FBox NewBounds = PCGComponent->GetGridBounds();
+			PCGComponent->PostProcessGraph(NewBounds, /*bGenerate=*/true);
+			return true;
+			}, { ProcessTaskId });
+#else
+		return ProcessTaskId;
+#endif
+	}
+	else
+	{
+		return InvalidTaskId;
+	}
+}
+
+ETickableTickType UPCGSubsystem::GetTickableTickType() const
+{
+	return IsTemplate() ? ETickableTickType::Never : ETickableTickType::Always;
+}
+
+TStatId UPCGSubsystem::GetStatId() const
+{
+	RETURN_QUICK_DECLARE_CYCLE_STAT(UPCGSubsystem, STATGROUP_Tickables);
+}
+
+#if WITH_EDITOR
+
+namespace PCGSubsystem
+{
+	FPCGTaskId ForEachIntersectingCell(FPCGGraphExecutor* GraphExecutor, UWorld* World, const FBox& InBounds, bool bCreateActor, bool bLoadCell, bool bSaveActors, TFunctionRef<FPCGTaskId(APCGPartitionActor*, const FBox&, const TArray<FPCGTaskId>&)> InOperation)
+	{
+		if (!GraphExecutor || !World)
+		{
+			return InvalidTaskId;
+		}
+
+		TArray<FPCGTaskId> CellTasks;
+
+		auto CellLambda = [&CellTasks, GraphExecutor, World, bCreateActor, bLoadCell, bSaveActors, InBounds, &InOperation](const UActorPartitionSubsystem::FCellCoord& CellCoord, const FBox& CellBounds)
+		{
+			UActorPartitionSubsystem* PartitionSubsystem = UWorld::GetSubsystem<UActorPartitionSubsystem>(World);
+			FBox IntersectedBounds = InBounds.Overlap(CellBounds);
+
+			if (IntersectedBounds.IsValid)
+			{
+				TSharedPtr<TSet<FWorldPartitionReference>> ActorReferences = MakeShared<TSet<FWorldPartitionReference>>();
+				APCGPartitionActor* PCGActor = Cast<APCGPartitionActor>(PartitionSubsystem->GetActor(APCGPartitionActor::StaticClass(), CellCoord, bCreateActor));
+
+				// At this point, if bCreateActor was true, then it exists, but it is not currently loaded; make sure it is loaded
+				// Otherwise, we still need to load it if it exists
+				// TODO: Revisit after API review on the WP side, we shouldn't have to load here or get the actor desc directly
+				if (!PCGActor)
+				{
+					const FWorldPartitionActorDesc* PCGActorDesc = nullptr;
+					auto FindFirst = [&CellCoord, &PCGActorDesc](const FWorldPartitionActorDesc* ActorDesc) {
+						FPartitionActorDesc* PartitionActorDesc = (FPartitionActorDesc*)ActorDesc;
+
+						if (PartitionActorDesc &&
+							PartitionActorDesc->GridIndexX == CellCoord.X &&
+							PartitionActorDesc->GridIndexY == CellCoord.Y &&
+							PartitionActorDesc->GridIndexZ == CellCoord.Z)
+						{
+							PCGActorDesc = ActorDesc;
+							return false;
+						}
+						else
+						{
+							return true;
+						}
+					};
+
+					FWorldPartitionHelpers::ForEachIntersectingActorDesc<APCGPartitionActor>(World->GetWorldPartition(), CellBounds, FindFirst);
+
+					check(!bCreateActor || PCGActorDesc);
+					if (PCGActorDesc)
+					{
+						ActorReferences->Add(FWorldPartitionReference(World->GetWorldPartition(), PCGActorDesc->GetGuid()));
+						PCGActor = Cast<APCGPartitionActor>(PCGActorDesc->GetActor());
+					}
+				}
+				else
+				{
+					// We still need to keep a reference on the PCG actor - note that newly created PCG actors will not have a reference here, but won't be unloaded
+					ActorReferences->Add(FWorldPartitionReference(World->GetWorldPartition(), PCGActor->GetActorGuid()));
+				}
+
+				if (!PCGActor)
+				{
+					return true;
+				}
+
+				TArray<FPCGTaskId> PreviousTasks;
+
+				auto SetPreviousTaskIfValid = [&PreviousTasks](FPCGTaskId TaskId){
+					if (TaskId != InvalidTaskId)
+					{
+						PreviousTasks.Reset();
+						PreviousTasks.Add(TaskId);
+					}
+				};
+
+				// We'll need to make sure actors in the bounds are loaded only if we need them.
+				if (bLoadCell)
+				{
+					auto WorldPartitionLoadActorsInBounds = [World, ActorReferences](const FWorldPartitionActorDesc* ActorDesc) {
+						check(ActorDesc);
+						ActorReferences->Add(FWorldPartitionReference(World->GetWorldPartition(), ActorDesc->GetGuid()));
+						// Load actor if not already loaded
+						ActorDesc->GetActor();
+						return true;
+					};
+
+					auto LoadActorsTask = [World, IntersectedBounds, WorldPartitionLoadActorsInBounds]() {
+						FWorldPartitionHelpers::ForEachIntersectingActorDesc(World->GetWorldPartition(), IntersectedBounds, WorldPartitionLoadActorsInBounds);
+						return true;
+					};
+
+					FPCGTaskId LoadTaskId = GraphExecutor->ScheduleGeneric(LoadActorsTask, {});
+					SetPreviousTaskIfValid(LoadTaskId);
+				}
+
+				// Execute
+				FPCGTaskId ExecuteTaskId = InOperation(PCGActor, IntersectedBounds, PreviousTasks);
+				SetPreviousTaskIfValid(ExecuteTaskId);
+
+				// Save changes; note that there's no need to save if the operation was cancelled
+				if (bSaveActors && ExecuteTaskId != InvalidTaskId)
+				{
+					auto SaveActorTask = [PCGActor, GraphExecutor]() {
+						GraphExecutor->AddToDirtyActors(PCGActor);
+						return true;
+					};
+
+					FPCGTaskId SaveTaskId = GraphExecutor->ScheduleGeneric(SaveActorTask, PreviousTasks);
+					SetPreviousTaskIfValid(SaveTaskId);
+				}
+
+				// Unload actors from cell (or the pcg actor refered here)
+				auto UnloadActorsTask = [GraphExecutor, ActorReferences]() {
+					GraphExecutor->AddToUnusedActors(*ActorReferences);
+					return true;
+				};
+
+				// Schedule after the save (if valid), then the execute so we can queue this after the load.
+				FPCGTaskId UnloadTaskId = GraphExecutor->ScheduleGeneric(UnloadActorsTask, PreviousTasks);
+				SetPreviousTaskIfValid(UnloadTaskId);
+
+				// Finally, mark "last" valid task in the cell tasks.
+				CellTasks.Append(PreviousTasks);
+			}
+
+			return true;
+		};
+
+		// TODO: accumulate last phases of every cell + run a GC at the end?
+		//  or add some mechanism on the graph executor side to run GC every now and then
+		FActorPartitionGridHelper::ForEachIntersectingCell(APCGPartitionActor::StaticClass(), InBounds, World->PersistentLevel, CellLambda);
+
+		// Finally, create a dummy generic task to wait on all cells
+		if (!CellTasks.IsEmpty())
+		{
+			return GraphExecutor->ScheduleGeneric([]() { return true; }, CellTasks);
+		}
+		else
+		{
+			return InvalidTaskId;
+		}
+	}
+
+} // end namepsace PCGSubsystem
+
+FPCGTaskId UPCGSubsystem::ProcessGraph(UPCGComponent* Component, const FBox& InPreviousBounds, const FBox& InNewBounds, bool bGenerate, bool bSave)
+{
+	check(Component && (InPreviousBounds.IsValid || InNewBounds.IsValid));
+
+	// TODO: optimal implementation would find the difference between the previous bounds and the new bounds
+	// and process these only. This is esp. important because of the CreateActor parameter.
+	auto ScheduleTask = [this, Component, InPreviousBounds, InNewBounds, bGenerate](APCGPartitionActor* PCGActor, const FBox& InBounds, const TArray<FPCGTaskId>& TaskDependencies){
+		auto UnpartitionTask = [Component, PCGActor]() {
+			// TODO: PCG actors that become empty could be deleted, but we also need to keep track
+			// of packages that would need to be deleted from SCC.
+			PCGActor->RemoveGraphInstance(Component);
+			return true;
+		};
+
+		auto PartitionTask = [Component, PCGActor]() {
+			PCGActor->AddGraphInstance(Component);
+			return true;
+		};
+
+		auto ScheduleGraph = [this, PCGActor, Component, &TaskDependencies]() {
+			// Ensure that the PCG actor has a matching local component.
+			// This is done immediately, but technically we could add it as a task
+			PCGActor->AddGraphInstance(Component);
+
+			UPCGComponent* LocalComponent = PCGActor->GetLocalComponent(Component);
+
+			if (!LocalComponent)
+			{
+				return InvalidTaskId;
+			}
+
+			return LocalComponent->GenerateInternal(/*bForce=*/false, TaskDependencies);
+		};
+
+		const bool bIsInOldBounds = InPreviousBounds.IsValid && InBounds.Intersect(InPreviousBounds);
+		const bool bIsInNewBounds = InNewBounds.IsValid && InBounds.Intersect(InNewBounds);
+
+		// 4 cases here:
+		// old & new -> generate only
+		// old & !new -> unpartition
+		// !old & new -> partition & generate
+		// !old & !new -> nothing to do
+		if (bIsInOldBounds && !bIsInNewBounds)
+		{
+			return GraphExecutor->ScheduleGeneric(UnpartitionTask, TaskDependencies);
+		}
+		else if (!bIsInOldBounds && bIsInNewBounds)
+		{
+			// Here, we should try to schedule the local component, but it doesn't exist yet.
+			// what to do? in practice, it's not a problem to not execute the Partition in a separate process,
+			// since there won't be any dependencies there in this case; but, we should execute it in a task
+			// if it's the only thing we do, otherwise we'll have issues with the other operations (save, etc.)
+			if (bGenerate)
+			{
+				return ScheduleGraph();
+			}
+			else
+			{
+				return GraphExecutor->ScheduleGeneric(PartitionTask, TaskDependencies);
+			}
+		}
+		else if (bIsInOldBounds && bIsInNewBounds && bGenerate)
+		{
+			return ScheduleGraph();
+		}
+		else
+		{
+			return InvalidTaskId;
+		}
+	};
+
+	FBox UnionBounds = InPreviousBounds + InNewBounds;
+	const bool bCreateActors = (!UnionBounds.Equals(InPreviousBounds));
+	const bool bLoadCell = (bGenerate); // TODO: review this
+	const bool bSaveActors = (bSave);
+
+	FPCGTaskId ProcessAllCellsTaskId = InvalidTaskId;
+
+	if (UnionBounds.IsValid)
+	{
+		ProcessAllCellsTaskId = PCGSubsystem::ForEachIntersectingCell(GraphExecutor, Component->GetWorld(), UnionBounds, bCreateActors, bLoadCell, bSaveActors, ScheduleTask);
+	}
+
+	// Finally, call PostProcessGraph if something happened
+	if (ProcessAllCellsTaskId != InvalidTaskId)
+	{
+		auto PostProcessGraph = [Component, InNewBounds, bGenerate]() {
+			Component->PostProcessGraph(InNewBounds, bGenerate);
+			return true;
+		};
+
+		return GraphExecutor->ScheduleGeneric(PostProcessGraph, { ProcessAllCellsTaskId });
+	}
+	else
+	{
+		return InvalidTaskId;
+	}
+}
+
+void UPCGSubsystem::CleanupGraph(UPCGComponent* Component, const FBox& InBounds, bool bRemoveComponents, bool bSave)
+{
+	auto ScheduleTask = [this, Component, bRemoveComponents](APCGPartitionActor* PCGActor, const FBox& InIntersectedBounds, const TArray<FPCGTaskId>& TaskDependencies) {
+		auto CleanupTask = [Component, PCGActor, bRemoveComponents]() {
+			check(Component != nullptr && PCGActor != nullptr);
+
+			if (!PCGActor || !Component)
+			{
+				return true;
+			}
+
+			if (UPCGComponent* LocalComponent = PCGActor->GetLocalComponent(Component))
+			{
+				LocalComponent->CleanupInternal(bRemoveComponents);
+			}
+
+			return true;
+		};
+
+		return GraphExecutor->ScheduleGeneric(CleanupTask, TaskDependencies);
+	};
+
+	PCGSubsystem::ForEachIntersectingCell(GraphExecutor, Component->GetWorld(), InBounds, /*bCreateActor=*/false, /*bLoadCell=*/false, bSave, ScheduleTask);
+}
+
+void UPCGSubsystem::DirtyGraph(UPCGComponent* Component, const FBox& InBounds, bool bDirtyInputs, bool bDirtyExclusions)
+{
+	auto ScheduleTask = [this, Component, bDirtyInputs, bDirtyExclusions](APCGPartitionActor* PCGActor, const FBox& InIntersectedBounds, const TArray<FPCGTaskId>& TaskDependencies) {
+		// In the specific case of the dirty, we want to bypass the execution queue, esp. since there's nothing happening here
+		// so we will run the command now, and not delay it.
+		check(Component != nullptr && PCGActor != nullptr);
+
+		if (!PCGActor || !Component)
+		{
+			return InvalidTaskId;
+		}
+
+		if (UPCGComponent* LocalComponent = PCGActor->GetLocalComponent(Component))
+		{
+			LocalComponent->DirtyGenerated(bDirtyInputs, bDirtyExclusions);
+		}
+
+		return InvalidTaskId;
+	};
+
+	// TODO: Could implement a "ForEachLoadedIntersectingCellImmediate" method to clarify this use case
+	PCGSubsystem::ForEachIntersectingCell(GraphExecutor, Component->GetWorld(), InBounds, /*bCreateActor=*/false, /*bLoadCell=*/false, /*bSave=*/false, ScheduleTask);
+}
+
+void UPCGSubsystem::CleanupPartitionActors(const FBox& InBounds)
+{
+	auto ScheduleTask = [this](APCGPartitionActor* PCGActor, const FBox& InIntersectedBounds, const TArray<FPCGTaskId>& TaskDependencies) {
+		auto CleanupTask = [PCGActor]() {
+			check(PCGActor);
+			PCGActor->CleanupDeadGraphInstances();
+
+			return true;
+		};
+
+		return GraphExecutor->ScheduleGeneric(CleanupTask, TaskDependencies);
+	};
+
+	PCGSubsystem::ForEachIntersectingCell(GraphExecutor, GetWorld(), InBounds, /*bCreateActor=*/false, /*bLoadCell=*/false, /*bSave=*/true, ScheduleTask);
+}
+
+void UPCGSubsystem::DeletePartitionActors()
+{
+	TSet<UPackage*> PackagesToCleanup;
+	TSet<FString> PackagesToDeleteFromSCC;
+	UWorld* World = GetWorld();
+
+	if (!World || !World->GetWorldPartition())
+	{
+		return;
+	}
+
+	auto GatherAndDestroyActors = [&PackagesToCleanup, &PackagesToDeleteFromSCC, World](const FWorldPartitionActorDesc* ActorDesc) {
+		if (ActorDesc->GetActor())
+		{
+			PackagesToCleanup.Add(ActorDesc->GetActor()->GetExternalPackage());
+			World->DestroyActor(ActorDesc->GetActor());
+		}
+		else
+		{
+			PackagesToDeleteFromSCC.Add(ActorDesc->GetActorPackage().ToString());
+			World->GetWorldPartition()->RemoveActor(ActorDesc->GetGuid());
+		}
+
+		return true;
+	};
+
+	// First, clear selection otherwise it might crash
+	GEditor->SelectNone(true, true, false);
+
+	FWorldPartitionHelpers::ForEachActorDesc<APCGPartitionActor>(World->GetWorldPartition(), GatherAndDestroyActors);
+
+	if (PackagesToCleanup.Num() > 0)
+	{
+		ObjectTools::CleanupAfterSuccessfulDelete(PackagesToCleanup.Array(), /*bPerformanceReferenceCheck=*/true);
+	}
+
+	if (PackagesToDeleteFromSCC.Num() > 0)
+	{
+		FPackageSourceControlHelper PackageHelper;
+		if (!PackageHelper.Delete(PackagesToDeleteFromSCC.Array()))
+		{
+			// Log error...
+		}
+	}
+}
+
+FPCGTaskId UPCGSubsystem::ScheduleGeneric(TFunction<bool()> InOperation, const TArray<FPCGTaskId>& TaskDependencies)
+{
+	check(GraphExecutor);
+	return GraphExecutor->ScheduleGeneric(InOperation, TaskDependencies);
+}
+
+void UPCGSubsystem::NotifyGraphChanged(UPCGGraph* InGraph)
+{
+	if (GraphExecutor)
+	{
+		GraphExecutor->NotifyGraphChanged(InGraph);
+	}
+}
+
+///////////////////////////////////////////////////////////////////////////////////////////
+void UPCGSubsystem::DelayPartitionGraph(UPCGComponent* Component)
+{
+	DelayProcessGraph(Component, /*bGenerate=*/false, /*bSave=*/false, /*bUseEmptyNewBounds=*/false);
+}
+
+void UPCGSubsystem::DelayUnpartitionGraph(UPCGComponent* Component)
+{
+	DelayProcessGraph(Component, /*bGenerate=*/false, /*bSave=*/false, /*bUseEmptyNewBounds=*/true);
+}
+
+FPCGTaskId UPCGSubsystem::DelayGenerateGraph(UPCGComponent* Component, bool bSave)
+{
+	return DelayProcessGraph(Component, /*bGenerate=*/true, bSave, /*bUseEmptyNewBounds=*/false);
+}
+
+FPCGTaskId UPCGSubsystem::DelayProcessGraph(UPCGComponent* Component, bool bGenerate, bool bSave, bool bUseEmptyNewBounds)
+{
+	check(Component && Component->IsPartitioned());
+
+	auto ExecuteProcessGraph = [this, Component, bGenerate, bSave, bUseEmptyNewBounds]() {
+		const FBox PreviousBounds = Component->LastGeneratedBounds;
+		FBox NewBounds = bUseEmptyNewBounds ? FBox(EForceInit::ForceInit) : Component->GetGridBounds();
+		/*FPCGTaskId GraphTaskId = */ ProcessGraph(Component, PreviousBounds, NewBounds, bGenerate, bSave);
+
+		return true;
+	};
+
+	// Delayed graph scheduling
+	return GraphExecutor->ScheduleGeneric(ExecuteProcessGraph, {});
+}
+
+#endif // WITH_EDITOR
diff --git a/Engine/Plugins/Experimental/PCG/Source/PCG/Private/PCGVolume.cpp b/Engine/Plugins/Experimental/PCG/Source/PCG/Private/PCGVolume.cpp
new file mode 100644
--- /dev/null
+++ b/Engine/Plugins/Experimental/PCG/Source/PCG/Private/PCGVolume.cpp
@@ -0,0 +1,11 @@
+// Copyright Epic Games, Inc. All Rights Reserved.
+
+#include "PCGVolume.h"
+#include "PCGComponent.h"
+#include "Engine/World.h"
+
+APCGVolume::APCGVolume(const FObjectInitializer& ObjectInitializer)
+	: Super(ObjectInitializer)
+{
+	PCGComponent = ObjectInitializer.CreateDefaultSubobject<UPCGComponent>(this, TEXT("PCG Component"));
+}
\ No newline at end of file
diff --git a/Engine/Plugins/Experimental/PCG/Source/PCG/Private/PCGWorldActor.cpp b/Engine/Plugins/Experimental/PCG/Source/PCG/Private/PCGWorldActor.cpp
new file mode 100644
--- /dev/null
+++ b/Engine/Plugins/Experimental/PCG/Source/PCG/Private/PCGWorldActor.cpp
@@ -0,0 +1,53 @@
+// Copyright Epic Games, Inc. All Rights Reserved.
+
+#include "PCGWorldActor.h"
+
+#include "PCGSubsystem.h"
+
+#include "Engine/World.h"
+
+void APCGWorldActor::PostLoad()
+{
+	Super::PostLoad();
+	RegisterToSubsystem();
+}
+
+void APCGWorldActor::BeginDestroy()
+{
+	UnregisterFromSubsystem();
+	Super::BeginDestroy();
+}
+
+#if WITH_EDITOR
+APCGWorldActor* APCGWorldActor::CreatePCGWorldActor(UWorld* InWorld)
+{
+	check(InWorld);
+	APCGWorldActor* PCGActor = InWorld->SpawnActor<APCGWorldActor>();
+
+	// This actor should be editor only and not spatially loaded
+	PCGActor->bIsEditorOnlyActor = true;
+	PCGActor->bIsSpatiallyLoaded = false;
+
+	PCGActor->RegisterToSubsystem();
+
+	return PCGActor;
+}
+#endif
+
+void APCGWorldActor::RegisterToSubsystem()
+{
+	UPCGSubsystem* PCGSubsystem = GetWorld() ? GetWorld()->GetSubsystem<UPCGSubsystem>() : nullptr;
+	if (PCGSubsystem)
+	{
+		PCGSubsystem->RegisterPCGWorldActor(this);
+	}
+}
+
+void APCGWorldActor::UnregisterFromSubsystem()
+{
+	UPCGSubsystem* PCGSubsystem = GetWorld() ? GetWorld()->GetSubsystem<UPCGSubsystem>() : nullptr;
+	if (PCGSubsystem)
+	{
+		PCGSubsystem->UnregisterPCGWorldActor(this);
+	}
+}
\ No newline at end of file
diff --git a/Engine/Plugins/Experimental/PCG/Source/PCG/Public/Data/PCGDifferenceData.h b/Engine/Plugins/Experimental/PCG/Source/PCG/Public/Data/PCGDifferenceData.h
new file mode 100644
--- /dev/null
+++ b/Engine/Plugins/Experimental/PCG/Source/PCG/Public/Data/PCGDifferenceData.h
@@ -0,0 +1,42 @@
+// Copyright Epic Games, Inc. All Rights Reserved.
+
+#pragma once
+
+#include "CoreMinimal.h"
+#include "PCGSpatialData.h"
+
+#include "PCGDifferenceData.generated.h"
+
+class UPCGUnionData;
+
+UCLASS(BlueprintType, ClassGroup = (Procedural))
+class UPCGDifferenceData : public UPCGSpatialDataWithPointCache
+{
+	GENERATED_BODY()
+public:
+	UFUNCTION(BlueprintCallable, Category = SpatialData)
+	void Initialize(const UPCGSpatialData* InData);
+
+	UFUNCTION(BlueprintCallable, Category = SpatialData)
+	void AddDifference(const UPCGSpatialData* InDifference);
+
+	//~Begin UPCGSpatialData interface
+	virtual int GetDimension() const override;
+	virtual FBox GetBounds() const override;
+	virtual FBox GetStrictBounds() const override;
+	virtual float GetDensityAtPosition(const FVector& InPosition) const override;
+	//~End UPCGSpatialData interface
+
+	//~Begin UPCGSpatialDataWithPointCache interface
+	virtual const UPCGPointData* CreatePointData() const override;
+	//~End UPCGSpatialDataWithPointCache interface
+protected:
+	UPROPERTY(BlueprintReadOnly, VisibleAnywhere, Category = SpatialData)
+	TObjectPtr<const UPCGSpatialData> Source = nullptr;
+
+	UPROPERTY(BlueprintReadOnly, VisibleAnywhere, Category = SpatialData)
+	TObjectPtr<const UPCGSpatialData> Difference = nullptr;
+
+	UPROPERTY()
+	TObjectPtr<UPCGUnionData> DifferencesUnion = nullptr;
+};
\ No newline at end of file
diff --git a/Engine/Plugins/Experimental/PCG/Source/PCG/Public/Data/PCGIntersectionData.h b/Engine/Plugins/Experimental/PCG/Source/PCG/Public/Data/PCGIntersectionData.h
new file mode 100644
--- /dev/null
+++ b/Engine/Plugins/Experimental/PCG/Source/PCG/Public/Data/PCGIntersectionData.h
@@ -0,0 +1,42 @@
+// Copyright Epic Games, Inc. All Rights Reserved.
+
+#pragma once
+
+#include "PCGSpatialData.h"
+
+#include "PCGIntersectionData.generated.h"
+
+/**
+* Generic intersection class that delays operations as long as possible.
+*/
+UCLASS(BlueprintType, ClassGroup=(Procedural))
+class UPCGIntersectionData : public UPCGSpatialDataWithPointCache
+{
+	GENERATED_BODY()
+public:
+	UFUNCTION(BlueprintCallable, Category = SpatialData)
+	void Initialize(const UPCGSpatialData* InA, const UPCGSpatialData* InB);
+
+	//~Begin UPCGSpatialData interface
+	virtual int GetDimension() const override;
+	virtual FBox GetBounds() const override;
+	virtual FBox GetStrictBounds() const override;
+	virtual float GetDensityAtPosition(const FVector& InPosition) const override;
+	//~End UPCGSpatialData interface
+
+	//~Begin UPCGSpatialDataWithPointCache interface
+	virtual const UPCGPointData* CreatePointData() const override;
+	//~End UPCGSpatialDataWithPointCache interface
+
+protected:
+	UPCGPointData* CreateAndFilterPointData(const UPCGSpatialData* X, const UPCGSpatialData* Y) const;
+
+	UPROPERTY(BlueprintReadOnly, VisibleAnywhere, Category = SpatialData)
+	TObjectPtr<const UPCGSpatialData> A = nullptr;
+
+	UPROPERTY(BlueprintReadOnly, VisibleAnywhere, Category = SpatialData)
+	TObjectPtr<const UPCGSpatialData> B = nullptr;
+
+	FBox CachedBounds = FBox(EForceInit::ForceInit);
+	FBox CachedStrictBounds = FBox(EForceInit::ForceInit);
+};
\ No newline at end of file
diff --git a/Engine/Plugins/Experimental/PCG/Source/PCG/Public/Data/PCGLandscapeData.h b/Engine/Plugins/Experimental/PCG/Source/PCG/Public/Data/PCGLandscapeData.h
new file mode 100644
--- /dev/null
+++ b/Engine/Plugins/Experimental/PCG/Source/PCG/Public/Data/PCGLandscapeData.h
@@ -0,0 +1,65 @@
+// Copyright Epic Games, Inc. All Rights Reserved.
+
+#pragma once
+
+#include "CoreMinimal.h"
+#include "PCGSurfaceData.h"
+#include "Elements/PCGSurfaceSampler.h"
+
+#include "PCGLandscapeData.generated.h"
+
+class ALandscapeProxy;
+class UPCGPolyLineData;
+
+USTRUCT()
+struct FPCGLandscapeDataPoint
+{
+	GENERATED_BODY()
+
+	FPCGLandscapeDataPoint() = default;
+	FPCGLandscapeDataPoint(int InX, int InY, float InHeight);
+
+	UPROPERTY()
+	int X;
+
+	UPROPERTY()
+	int Y;
+
+	UPROPERTY()
+	float Height;
+};
+
+UCLASS(BlueprintType, ClassGroup = (Procedural))
+class UPCGLandscapeData : public UPCGSurfaceData
+{
+	GENERATED_BODY()
+
+public:
+	void Initialize(ALandscapeProxy* InLandscape, const FBox& InBounds);
+
+	// ~Begin UPGCSpatialData interface
+	virtual FBox GetBounds() const override;
+	virtual FBox GetStrictBounds() const override;
+	virtual FVector TransformPosition(const FVector& InPosition) const override;
+	virtual FPCGPoint TransformPoint(const FPCGPoint& InPoint) const override;
+	virtual float GetDensityAtPosition(const FVector& InPosition) const override;
+	// ~End UPGCConcreteData interface
+
+	// ~Begin UPCGSpatialDataWithPointCache interface
+	virtual const UPCGPointData* CreatePointData() const override;
+	// ~End UPCGConcreteDataWithPointCache interface
+
+	// TODO: add on property changed to clear cached data
+	UPROPERTY(BlueprintReadWrite, EditAnywhere, Category = SourceData)
+	TSoftObjectPtr<ALandscapeProxy> Landscape;
+
+	// TODO: add on property changed to clear cached data
+	//UPROPERTY(BlueprintReadWrite, EditAnywhere)
+	//TObjectPtr<UPCGSurfaceSamplerSettings> ImplicitSamplerSettings;
+protected:
+	UPROPERTY()
+	TArray<FPCGLandscapeDataPoint> LandscapePoints;
+
+	UPROPERTY()
+	FBox Bounds = FBox(EForceInit::ForceInit);
+};
\ No newline at end of file
diff --git a/Engine/Plugins/Experimental/PCG/Source/PCG/Public/Data/PCGLandscapeSplineData.h b/Engine/Plugins/Experimental/PCG/Source/PCG/Public/Data/PCGLandscapeSplineData.h
new file mode 100644
--- /dev/null
+++ b/Engine/Plugins/Experimental/PCG/Source/PCG/Public/Data/PCGLandscapeSplineData.h
@@ -0,0 +1,37 @@
+// Copyright Epic Games, Inc. All Rights Reserved.
+
+#pragma once
+
+#include "CoreMinimal.h"
+#include "PCGPolyLineData.h"
+
+#include "PCGLandscapeSplineData.generated.h"
+
+class ULandscapeSplinesComponent;
+
+UCLASS(BlueprintType, ClassGroup=(Procedural))
+class UPCGLandscapeSplineData : public UPCGPolyLineData
+{
+	GENERATED_BODY()
+public:
+	void Initialize(ULandscapeSplinesComponent* InSplineComponent);
+
+	//~Begin UPCGPolyLineData interface
+	virtual int GetNumSegments() const override;
+	virtual float GetSegmentLength(int SegmentIndex) const override;
+	virtual FVector GetLocationAtDistance(int SegmentIndex, float Distance) const override;
+	//~End UPCGPolyLineData interface
+
+	//~Begin UPCGSpatialDataWithPointCache interface
+	virtual const UPCGPointData* CreatePointData() const override;
+	//~End UPCGSpatialDataWithPointCache interface
+
+	//~Begin UPCGSpatialData interface
+	virtual FBox GetBounds() const override;
+	virtual float GetDensityAtPosition(const FVector& InPosition) const override;
+	//~End
+
+protected:
+	UPROPERTY(BlueprintReadWrite, EditAnywhere, Category = SourceData)
+	TSoftObjectPtr<ULandscapeSplinesComponent> Spline;
+};
diff --git a/Engine/Plugins/Experimental/PCG/Source/PCG/Public/Data/PCGPointData.h b/Engine/Plugins/Experimental/PCG/Source/PCG/Public/Data/PCGPointData.h
new file mode 100644
--- /dev/null
+++ b/Engine/Plugins/Experimental/PCG/Source/PCG/Public/Data/PCGPointData.h
@@ -0,0 +1,94 @@
+// Copyright Epic Games, Inc. All Rights Reserved.
+
+#pragma once
+
+#include "CoreMinimal.h"
+#include "PCGSpatialData.h"
+
+#include "Math/GenericOctreePublic.h"
+#include "Math/GenericOctree.h"
+
+#include "PCGPointData.generated.h"
+
+class AActor;
+
+struct FPCGPointRef
+{
+	FPCGPointRef(const FPCGPoint& InPoint);
+	FPCGPointRef(const FPCGPointRef& InPointRef);
+
+	const FPCGPoint* Point;
+	FBoxSphereBounds Bounds;
+};
+
+struct FPCGPointRefSemantics
+{
+	enum { MaxElementsPerLeaf = 16 };
+	enum { MinInclusiveElementsPerNode = 7 };
+	enum { MaxNodeDepth = 12 };
+
+	typedef TInlineAllocator<MaxElementsPerLeaf> ElementAllocator;
+
+	FORCEINLINE static const FBoxSphereBounds& GetBoundingBox(const FPCGPointRef& InPoint)
+	{
+		return InPoint.Bounds;
+	}
+
+	FORCEINLINE static const bool AreElementsEqual(const FPCGPointRef& A, const FPCGPointRef& B)
+	{
+		// TODO: verify if that's sufficient
+		return A.Point == B.Point;
+	}
+
+	FORCEINLINE static void ApplyOffset(FPCGPointRef& InPoint)
+	{
+		ensureMsgf(false, TEXT("Not implemented"));
+	}
+
+	FORCEINLINE static void SetElementId(const FPCGPointRef& Element, FOctreeElementId2 OctreeElementID)
+	{
+	}
+};
+
+// TODO: Split this in "concrete" vs "api" class (needed for views)
+UCLASS(BlueprintType, ClassGroup = (Procedural))
+class UPCGPointData : public UPCGSpatialData
+{
+	GENERATED_BODY()
+
+public:
+	void Initialize(AActor* InActor);
+
+	// ~Begin UPCGSpatialData interface
+	virtual int GetDimension() const override { return 0; }
+	virtual FBox GetBounds() const override;
+	virtual float GetDensityAtPosition(const FVector& InPosition) const override;
+	virtual const UPCGPointData* ToPointData() const { return this; }
+	// ~End UPCGSpatialData interface
+
+	UFUNCTION(BlueprintCallable, Category = SpatialData)
+	const TArray<FPCGPoint>& GetPoints() const { return Points; }
+
+	UFUNCTION(BlueprintCallable, Category = SpatialData)
+	void SetPoints(const TArray<FPCGPoint>& InPoints);
+	
+	UFUNCTION(BlueprintCallable, Category = SpatialData)
+	void CopyPointsFrom(const UPCGPointData* InData, const TArray<int>& InDataIndices);
+
+	const FPCGPoint* GetPointAtPosition(const FVector& InPosition) const;
+
+	TArray<FPCGPoint>& GetMutablePoints();
+
+protected:
+	void RebuildOctree() const;
+	void RecomputeBounds() const;
+
+	UPROPERTY()
+	TArray<FPCGPoint> Points;
+
+	mutable FCriticalSection CachedDataLock;
+	mutable TOctree2<FPCGPointRef, FPCGPointRefSemantics> Octree;
+	mutable FBox Bounds; // TODO: review if this needs to be threadsafe
+	mutable bool bBoundsAreDirty = true;
+	mutable bool bOctreeIsDirty = true;
+};
\ No newline at end of file
diff --git a/Engine/Plugins/Experimental/PCG/Source/PCG/Public/Data/PCGPolyLineData.h b/Engine/Plugins/Experimental/PCG/Source/PCG/Public/Data/PCGPolyLineData.h
new file mode 100644
--- /dev/null
+++ b/Engine/Plugins/Experimental/PCG/Source/PCG/Public/Data/PCGPolyLineData.h
@@ -0,0 +1,24 @@
+// Copyright Epic Games, Inc. All Rights Reserved.
+
+#pragma once
+
+#include "CoreMinimal.h"
+#include "PCGSpatialData.h"
+
+#include "PCGPolyLineData.generated.h"
+
+UCLASS(BlueprintType, ClassGroup = (Procedural))
+class UPCGPolyLineData : public UPCGSpatialDataWithPointCache
+{
+	GENERATED_BODY()
+
+public:
+	//~Begin UPCGSpatialData interface
+	virtual int GetDimension() const override { return 1; }
+	virtual FBox GetBounds() const override;
+	//~End UPCGSpatialData interface
+
+	virtual int GetNumSegments() const PURE_VIRTUAL(UPCGPolyLineData::GetNumSegments, return 0;);
+	virtual float GetSegmentLength(int SegmentIndex) const PURE_VIRTUAL(UPCGPolyLineData::GetSegmentLength, return 0.0f;);
+	virtual FVector GetLocationAtDistance(int SegmentIndex, float Distance) const PURE_VIRTUAL(UPCGPolyLine::GetLocationAtDistance, return FVector(););
+};
\ No newline at end of file
diff --git a/Engine/Plugins/Experimental/PCG/Source/PCG/Public/Data/PCGProjectionData.h b/Engine/Plugins/Experimental/PCG/Source/PCG/Public/Data/PCGProjectionData.h
new file mode 100644
--- /dev/null
+++ b/Engine/Plugins/Experimental/PCG/Source/PCG/Public/Data/PCGProjectionData.h
@@ -0,0 +1,47 @@
+// Copyright Epic Games, Inc. All Rights Reserved.
+
+#pragma once
+
+#include "PCGSpatialData.h"
+#include "PCGPointData.h"
+
+#include "PCGProjectionData.generated.h"
+
+/**
+* Generic projection class (A projected onto B) that intercepts spatial queries
+*/
+UCLASS(BlueprintType, ClassGroup=(Procedural))
+class UPCGProjectionData : public UPCGSpatialDataWithPointCache
+{
+	GENERATED_BODY()
+public:
+	void Initialize(const UPCGSpatialData* InSource, const UPCGSpatialData* InTarget);
+
+	//~Begin UPCGSpatialData interface
+	virtual int GetDimension() const override;
+	virtual FBox GetBounds() const override;
+	virtual FBox GetStrictBounds() const override;
+	virtual float GetDensityAtPosition(const FVector& InPosition) const override;
+	virtual FVector GetNormal() const override;
+	//~End UPCGSpatialData interface
+
+	//~Begin UPCGSpatialDataWithPointCache interface
+	virtual const UPCGPointData* CreatePointData() const override;
+	//~End UPCGSpatialDataWithPointCache interface
+
+protected:
+	virtual FVector ProjectPosition(const FVector& InPosition) const;
+	virtual FPCGPoint ProjectPoint(const FPCGPoint& InPoint) const;
+
+	FBox ProjectBounds(const FBox& InBounds) const;
+
+	UPROPERTY(BlueprintReadOnly, VisibleAnywhere, Category = SpatialData)
+	TObjectPtr<const UPCGSpatialData> Source = nullptr;
+
+	UPROPERTY(BlueprintReadOnly, VisibleAnywhere, Category = SpatialData)
+	TObjectPtr<const UPCGSpatialData> Target = nullptr;
+
+	FBox CachedBounds = FBox(EForceInit::ForceInit);
+	FBox CachedStrictBounds = FBox(EForceInit::ForceInit);
+};
+
diff --git a/Engine/Plugins/Experimental/PCG/Source/PCG/Public/Data/PCGRenderTargetData.h b/Engine/Plugins/Experimental/PCG/Source/PCG/Public/Data/PCGRenderTargetData.h
new file mode 100644
--- /dev/null
+++ b/Engine/Plugins/Experimental/PCG/Source/PCG/Public/Data/PCGRenderTargetData.h
@@ -0,0 +1,23 @@
+// Copyright Epic Games, Inc. All Rights Reserved.
+
+#pragma once
+
+#include "Data/PCGTextureData.h"
+#include "Engine/TextureRenderTarget2D.h"
+
+#include "PCGRenderTargetData.generated.h"
+
+//TODO: It's possible that caching the result in this class is not as efficient as it could be
+// if we expect to sample in different ways (e.g. channel) in the same render target
+UCLASS(BlueprintType, ClassGroup = (Procedural))
+class UPCGRenderTargetData : public UPCGBaseTextureData
+{
+	GENERATED_BODY()
+
+public:
+	UFUNCTION(BlueprintCallable, Category = RenderTarget)
+	void Initialize(UTextureRenderTarget2D* InRenderTarget, const FTransform& InTransform);
+
+	UPROPERTY(BlueprintReadOnly, VisibleAnywhere, Category = Properties)
+	UTextureRenderTarget2D* RenderTarget = nullptr;
+};
\ No newline at end of file
diff --git a/Engine/Plugins/Experimental/PCG/Source/PCG/Public/Data/PCGSpatialData.h b/Engine/Plugins/Experimental/PCG/Source/PCG/Public/Data/PCGSpatialData.h
new file mode 100644
--- /dev/null
+++ b/Engine/Plugins/Experimental/PCG/Source/PCG/Public/Data/PCGSpatialData.h
@@ -0,0 +1,88 @@
+// Copyright Epic Games, Inc. All Rights Reserved.
+
+#pragma once
+
+#include "PCGData.h"
+#include "PCGPoint.h"
+
+#include "PCGSpatialData.generated.h"
+
+class UPCGPointData;
+class AActor;
+
+/**
+* "Concrete" data base class for PCG generation
+* This will be the base class for data classes that actually represent
+* concrete evidence of spatial data - points, surfaces, splines, etc.
+* In opposition to settings/control type of data.
+* 
+* Conceptually, any concrete data can be decayed into points (potentially through transformations)
+* which hold metadata and a transform, and this is the basic currency of the PCG framework.
+*/
+UCLASS(Abstract, BlueprintType, ClassGroup = (Procedural))
+class UPCGSpatialData : public UPCGData
+{
+	GENERATED_BODY()
+
+public:
+	/** Returns the dimension of the data type, which has nothing to do with the dimension of its points */
+	UFUNCTION(BlueprintCallable, Category = SpatialData)
+	virtual int GetDimension() const PURE_VIRTUAL(UPCGSpatialData::GetDimension, return 0;);
+
+	/** Returns the full bounds (including density fall-off) of the data */
+	UFUNCTION(BlueprintCallable, Category = SpatialData)
+	virtual FBox GetBounds() const PURE_VIRTUAL(UPCGSpatialData::GetBounds, return FBox(EForceInit::ForceInit););
+
+	/** Returns the bounds in which the density is always 1 */
+	UFUNCTION(BlueprintCallable, Category = SpatialData)
+	virtual FBox GetStrictBounds() const { return GetBounds(); }
+
+	/** Returns the expected data normal (for surfaces) or eventual projection axis (for volumes) */
+	UFUNCTION(BlueprintCallable, Category = SpatialData)
+	virtual FVector GetNormal() const { return FVector::UnitZ(); }
+
+	/** Computes the density at a given location */
+	UFUNCTION(BlueprintCallable, Category = Distribution)
+	virtual float GetDensityAtPosition(const FVector& InPosition) const PURE_VIRTUAL(UPCGSpatialData::GetDensityAtPosition, return 0;);
+
+	/** Discretizes the data into points */
+	UFUNCTION(BlueprintCallable, Category = Data)
+	virtual const UPCGPointData* ToPointData() const PURE_VIRTUAL(UPCGSpatialData::ToPointData, return nullptr;);
+
+	/** Transform a world-space position to a world-space position in relation to the current data. (ex: projection on surface) */
+	UFUNCTION(BlueprintCallable, Category = SpatialData)
+	virtual FVector TransformPosition(const FVector& InPosition) const { return InPosition; }
+
+	/** Transform a full point similar to TransformPosition */
+	UFUNCTION(BlueprintCallable, Category = SpatialData)
+	virtual FPCGPoint TransformPoint(const FPCGPoint& InPoint) const { return InPoint; }
+
+	/** Returns a specialized data to intersect with another data */
+	virtual UPCGSpatialData* IntersectWith(const UPCGSpatialData* InOther) const;
+
+	/** Returns a specialized data to project this on another data of equal or higher dimension */
+	virtual UPCGSpatialData* ProjectOn(const UPCGSpatialData* InOther) const;
+
+	UPROPERTY(BlueprintReadWrite, EditAnywhere, Category = Data)
+	AActor* TargetActor = nullptr;
+};
+
+UCLASS(Abstract, ClassGroup = (Procedural))
+class UPCGSpatialDataWithPointCache : public UPCGSpatialData
+{
+	GENERATED_BODY()
+
+public:
+	// ~UPCGSpatialData implementation
+	virtual const UPCGPointData* ToPointData() const override;
+	// ~End UPCGSpatialData implementation
+
+protected:
+	virtual const UPCGPointData* CreatePointData() const PURE_VIRTUAL(UPCGSpatialData::CreatePointData, return nullptr;);
+
+private:
+	UPROPERTY(Transient)
+	mutable TObjectPtr<const UPCGPointData> CachedPointData;
+
+	mutable FCriticalSection CacheLock;
+};
\ No newline at end of file
diff --git a/Engine/Plugins/Experimental/PCG/Source/PCG/Public/Data/PCGSplineData.h b/Engine/Plugins/Experimental/PCG/Source/PCG/Public/Data/PCGSplineData.h
new file mode 100644
--- /dev/null
+++ b/Engine/Plugins/Experimental/PCG/Source/PCG/Public/Data/PCGSplineData.h
@@ -0,0 +1,60 @@
+// Copyright Epic Games, Inc. All Rights Reserved.
+
+#pragma once
+
+#include "CoreMinimal.h"
+#include "PCGPolyLineData.h"
+#include "PCGProjectionData.h"
+#include "PCGSplineData.generated.h"
+
+class USplineComponent;
+class UPCGSurfaceData;
+
+UCLASS(BlueprintType, ClassGroup = (Procedural))
+class UPCGSplineData : public UPCGPolyLineData
+{
+	GENERATED_BODY()
+
+public:
+	void Initialize(USplineComponent* InSpline);
+
+	//~Begin UPCGPolyLineData interface
+	virtual int GetNumSegments() const override;
+	virtual float GetSegmentLength(int SegmentIndex) const override;
+	virtual FVector GetLocationAtDistance(int SegmentIndex, float Distance) const override;
+	//~End UPCGPolyLineData interface
+
+	//~Begin UPCGSpatialDataWithPointCache interface
+	virtual const UPCGPointData* CreatePointData() const override;
+	//~End UPCGSpatialDataWithPointCache interface
+
+	//~Begin UPCGSpatialData interface
+	virtual FBox GetBounds() const override;
+	virtual float GetDensityAtPosition(const FVector& InPosition) const override;
+	virtual UPCGSpatialData* ProjectOn(const UPCGSpatialData* InOther) const override;
+	//~End 
+
+	UPROPERTY(BlueprintReadWrite, EditAnywhere, Category = SourceData)
+	TSoftObjectPtr<USplineComponent> Spline;
+};
+
+UCLASS(BlueprintType, ClassGroup=(Procedural))
+class UPCGSplineProjectionData : public UPCGProjectionData
+{
+	GENERATED_BODY()
+public:
+	void Initialize(const UPCGSplineData* InSourceSpline, const UPCGSpatialData* InTargetSurface);
+
+	//~Begin UPCGSpatialData interface
+	virtual float GetDensityAtPosition(const FVector& InPosition) const override;
+	//~End UPCGSpatialData interface
+
+	const UPCGSplineData* GetSpline() const;
+	const UPCGSpatialData* GetSurface() const;
+
+	UPROPERTY(BlueprintReadWrite, EditAnywhere, Category = SpatialData)
+	FInterpCurveVector2D ProjectedPosition;
+
+protected:
+	FVector2D Project(const FVector& InVector) const;
+};
\ No newline at end of file
diff --git a/Engine/Plugins/Experimental/PCG/Source/PCG/Public/Data/PCGSurfaceData.h b/Engine/Plugins/Experimental/PCG/Source/PCG/Public/Data/PCGSurfaceData.h
new file mode 100644
--- /dev/null
+++ b/Engine/Plugins/Experimental/PCG/Source/PCG/Public/Data/PCGSurfaceData.h
@@ -0,0 +1,25 @@
+// Copyright Epic Games, Inc. All Rights Reserved.
+
+#pragma once
+
+#include "CoreMinimal.h"
+#include "PCGSpatialData.h"
+
+#include "PCGSurfaceData.generated.h"
+
+UCLASS(BlueprintType, ClassGroup = (Procedural))
+class UPCGSurfaceData : public UPCGSpatialDataWithPointCache
+{
+	GENERATED_BODY()
+
+public:
+	//~Begin UPCGSpatialData interface
+	virtual int GetDimension() const override { return 2; }
+	virtual FVector TransformPosition(const FVector& InPosition) const override;
+	virtual FPCGPoint TransformPoint(const FPCGPoint& InPoint) const override;
+	//~End UPCGSpatialData interface
+
+protected:
+	UPROPERTY(BlueprintReadOnly, VisibleAnywhere, Category = SpatialData)
+	FTransform Transform;
+};
\ No newline at end of file
diff --git a/Engine/Plugins/Experimental/PCG/Source/PCG/Public/Data/PCGTextureData.h b/Engine/Plugins/Experimental/PCG/Source/PCG/Public/Data/PCGTextureData.h
new file mode 100644
--- /dev/null
+++ b/Engine/Plugins/Experimental/PCG/Source/PCG/Public/Data/PCGTextureData.h
@@ -0,0 +1,63 @@
+// Copyright Epic Games, Inc. All Rights Reserved.
+
+#pragma once
+
+#include "Data/PCGSurfaceData.h"
+#include "Engine/Texture2D.h"
+
+#include "PCGTextureData.generated.h"
+
+UENUM(BlueprintType)
+enum class EPCGTextureColorChannel : uint8
+{
+	Red,
+	Green,
+	Blue,
+	Alpha
+};
+
+UCLASS(Abstract)
+class UPCGBaseTextureData : public UPCGSurfaceData
+{
+	GENERATED_BODY()
+
+public:
+	//~Begin UPCGSpatialData interface
+	virtual FBox GetBounds() const override;
+	virtual FBox GetStrictBounds() const override;
+	virtual float GetDensityAtPosition(const FVector& InPosition) const override;
+	//~End UPCGSpatialData interface
+
+	//~Begin UPCGSpatialDataWithPointCache interface
+	virtual const UPCGPointData* CreatePointData() const override;
+	//~End UPCGSpatialDataWithPointCache interface
+
+protected:
+	UPROPERTY()
+	TArray<FLinearColor> ColorData;
+
+	UPROPERTY(BlueprintReadOnly, VisibleAnywhere, Category = SpatialData)
+	FBox Bounds = FBox(EForceInit::ForceInit);
+
+	UPROPERTY(BlueprintReadWrite, EditAnywhere, Category = Settings)
+	EPCGTextureColorChannel ColorChannel = EPCGTextureColorChannel::Alpha;
+
+	UPROPERTY(BlueprintReadOnly, VisibleAnywhere, Category = SpatialData)
+	int32 Height = 0;
+
+	UPROPERTY(BlueprintReadOnly, VisibleAnywhere, Category = SpatialData)
+	int32 Width = 0;
+};
+
+UCLASS(BlueprintType, ClassGroup=(Procedural))
+class UPCGTextureData : public UPCGBaseTextureData
+{
+	GENERATED_BODY()
+
+public:
+	UFUNCTION(BlueprintCallable, Category = Texture)
+	void Initialize(UTexture2D* InTexture, const FTransform& InTransform);
+
+	UPROPERTY(BlueprintReadOnly, VisibleAnywhere, Category = Properties)
+	UTexture2D* Texture = nullptr;
+};
\ No newline at end of file
diff --git a/Engine/Plugins/Experimental/PCG/Source/PCG/Public/Data/PCGUnionData.h b/Engine/Plugins/Experimental/PCG/Source/PCG/Public/Data/PCGUnionData.h
new file mode 100644
--- /dev/null
+++ b/Engine/Plugins/Experimental/PCG/Source/PCG/Public/Data/PCGUnionData.h
@@ -0,0 +1,56 @@
+// Copyright Epic Games, Inc. All Rights Reserved.
+
+#pragma once
+
+#include "CoreMinimal.h"
+#include "PCGSpatialData.h"
+
+#include "PCGUnionData.generated.h"
+
+UENUM()
+enum class EPCGUnionType : uint8
+{
+	LeftToRightPriority,
+	RightToLeftPriority,
+	KeepAB,
+	ABMinusIntersection
+};
+
+UCLASS(BlueprintType, ClassGroup=(Procedural))
+class UPCGUnionData : public UPCGSpatialDataWithPointCache
+{
+	GENERATED_BODY()
+public:
+	UFUNCTION(BlueprintCallable, Category = SpatialData)
+	void Initialize(const UPCGSpatialData* InA, const UPCGSpatialData* InB);
+
+	UFUNCTION(BlueprintCallable, Category = SpatialData)
+	void AddData(const UPCGSpatialData* InC);
+
+	//~Begin UPCGSpatialData interface
+	virtual int GetDimension() const override;
+	virtual FBox GetBounds() const override;
+	virtual FBox GetStrictBounds() const override;
+	virtual float GetDensityAtPosition(const FVector& InPosition) const override;
+	//~End UPCGSpatialData interface
+
+	//~Begin UPCGSpatialDataWithPointCache interface
+	virtual const UPCGPointData* CreatePointData() const override;
+	//~End UPCGSpatialDataWithPointCache interface
+
+protected:
+	UPROPERTY(BlueprintReadOnly, VisibleAnywhere, Category = SpatialData)
+	TObjectPtr<const UPCGSpatialData> A = nullptr;
+
+	UPROPERTY(BlueprintReadOnly, VisibleAnywhere, Category = SpatialData)
+	TObjectPtr<const UPCGSpatialData> B = nullptr;
+
+	UPROPERTY(BlueprintReadWrite, EditAnywhere, Category = Settings)
+	EPCGUnionType UnionType = EPCGUnionType::LeftToRightPriority;
+
+	UPROPERTY()
+	FBox CachedBounds = FBox(EForceInit::ForceInit);
+
+	UPROPERTY()
+	FBox CachedStrictBounds = FBox(EForceInit::ForceInit);
+};
\ No newline at end of file
diff --git a/Engine/Plugins/Experimental/PCG/Source/PCG/Public/Data/PCGVolumeData.h b/Engine/Plugins/Experimental/PCG/Source/PCG/Public/Data/PCGVolumeData.h
new file mode 100644
--- /dev/null
+++ b/Engine/Plugins/Experimental/PCG/Source/PCG/Public/Data/PCGVolumeData.h
@@ -0,0 +1,41 @@
+// Copyright Epic Games, Inc. All Rights Reserved.
+
+#pragma once
+
+#include "CoreMinimal.h"
+#include "PCGSpatialData.h"
+
+#include "PCGVolumeData.generated.h"
+
+class AVolume;
+
+UCLASS(BlueprintType, ClassGroup = (Procedural))
+class UPCGVolumeData : public UPCGSpatialDataWithPointCache
+{
+	GENERATED_BODY()
+
+public:
+	void Initialize(AVolume* InVolume, AActor* InTargetActor = nullptr);
+	void Initialize(const FBox& InBounds, AActor* InTargetActor);
+
+	// ~Begin UPGCSpatialData interface
+	virtual int GetDimension() const override { return 3; }
+	virtual FBox GetBounds() const override;
+	virtual FBox GetStrictBounds() const override;
+	virtual float GetDensityAtPosition(const FVector& InPosition) const override;
+	// ~End UPGCSpatialData interface
+
+	// ~Begin UPCGSpatialDataWithPointCache interface
+	virtual const UPCGPointData* CreatePointData() const override;
+	// ~End UPCGSpatialDataWithPointCache interface
+
+protected:
+	UPROPERTY(BlueprintReadWrite, EditAnywhere, Category = SourceData)
+	TObjectPtr<AVolume> Volume = nullptr;
+
+	FBox Bounds = FBox(EForceInit::ForceInit);
+	FBox StrictBounds = FBox(EForceInit::ForceInit);
+
+	FTransform VolumeTransform = FTransform::Identity;
+	FBox VolumeLocalBounds = FBox(EForceInit::ForceInit);
+};
\ No newline at end of file
diff --git a/Engine/Plugins/Experimental/PCG/Source/PCG/Public/Elements/PCGExecuteBlueprint.h b/Engine/Plugins/Experimental/PCG/Source/PCG/Public/Elements/PCGExecuteBlueprint.h
new file mode 100644
--- /dev/null
+++ b/Engine/Plugins/Experimental/PCG/Source/PCG/Public/Elements/PCGExecuteBlueprint.h
@@ -0,0 +1,62 @@
+// Copyright Epic Games, Inc. All Rights Reserved.
+
+#pragma once
+
+#include "PCGElement.h"
+#include "PCGSettings.h"
+
+#include "Templates/SubclassOf.h"
+
+#include "PCGExecuteBlueprint.generated.h"
+
+class UWorld;
+
+UCLASS(Abstract, BlueprintType, Blueprintable, hidecategories = (Object))
+class UPCGBlueprintElement : public UObject
+{
+	GENERATED_BODY()
+
+public:
+	UFUNCTION(BlueprintImplementableEvent, BlueprintCallable, Category = Execution)
+	void Execute(const FPCGDataCollection& Input, FPCGDataCollection& Output) const;
+
+	/** Needed to be able to call certain blueprint functions */
+	virtual UWorld* GetWorld() const override;
+};
+
+UCLASS(BlueprintType, ClassGroup = (Procedural))
+class UPCGBlueprintSettings : public UPCGSettings
+{
+	GENERATED_BODY()
+
+public:
+	friend class FPCGExecuteBlueprintElement;
+
+#if WITH_EDITOR
+	virtual void PostLoad() override;
+	virtual void PostEditChangeProperty(FPropertyChangedEvent& PropertyChangedEvent) override;
+#endif
+
+	UFUNCTION(BlueprintCallable, Category = Settings)
+	void SetElementType(TSubclassOf<UPCGBlueprintElement> InElementType);
+
+protected:
+	UPROPERTY()
+	TSubclassOf<UPCGBlueprintElement> BlueprintElement_DEPRECATED;
+
+	UPROPERTY(BlueprintSetter = SetElementType, EditAnywhere, Category = Template)
+	TSubclassOf<UPCGBlueprintElement> BlueprintElementType;
+
+	UPROPERTY(BlueprintReadOnly, VisibleAnywhere, Instanced, Category = Settings, meta = (ShowOnlyInnerProperties))
+	TObjectPtr<UPCGBlueprintElement> BlueprintElementInstance;
+
+protected:
+	void RefreshBlueprintElement();
+	virtual FPCGElementPtr CreateElement() const override;
+};
+
+class FPCGExecuteBlueprintElement : public FSimpleTypedPCGElement<UPCGBlueprintSettings>
+{
+public:
+	virtual bool Execute(FPCGContextPtr Context) const override;
+};
\ No newline at end of file
diff --git a/Engine/Plugins/Experimental/PCG/Source/PCG/Public/Elements/PCGPointSampler.h b/Engine/Plugins/Experimental/PCG/Source/PCG/Public/Elements/PCGPointSampler.h
new file mode 100644
--- /dev/null
+++ b/Engine/Plugins/Experimental/PCG/Source/PCG/Public/Elements/PCGPointSampler.h
@@ -0,0 +1,32 @@
+// Copyright Epic Games, Inc. All Rights Reserved.
+
+#pragma once
+
+#include "CoreMinimal.h"
+#include "PCGSettings.h"
+#include "PCGElement.h"
+#include "PCGNode.h"
+
+#include "PCGPointSampler.generated.h"
+
+class FPCGPointSamplerElement;
+
+UCLASS(BlueprintType, ClassGroup = (Procedural))
+class UPCGPointSamplerSettings : public UPCGSettings
+{
+	GENERATED_BODY()
+
+public:
+
+	UPROPERTY(BlueprintReadWrite, EditAnywhere, Category = Settings, meta=(ClampMin="0", ClampMax="1"))
+	float Ratio = 0.1f;
+
+protected:
+	virtual FPCGElementPtr CreateElement() const override;
+};
+
+class FPCGPointSamplerElement : public FSimpleTypedPCGElement<UPCGPointSamplerSettings>
+{
+public:
+	virtual bool Execute(FPCGContextPtr Context) const override;
+};
\ No newline at end of file
diff --git a/Engine/Plugins/Experimental/PCG/Source/PCG/Public/Elements/PCGStaticMeshSpawner.h b/Engine/Plugins/Experimental/PCG/Source/PCG/Public/Elements/PCGStaticMeshSpawner.h
new file mode 100644
--- /dev/null
+++ b/Engine/Plugins/Experimental/PCG/Source/PCG/Public/Elements/PCGStaticMeshSpawner.h
@@ -0,0 +1,47 @@
+// Copyright Epic Games, Inc. All Rights Reserved.
+
+#pragma once
+
+#include "CoreMinimal.h"
+#include "PCGSettings.h"
+
+#include "PCGStaticMeshSpawner.generated.h"
+
+class AActor;
+class UInstancedStaticMeshComponent;
+class UPCGComponent;
+class UStaticMesh;
+
+USTRUCT(BlueprintType)
+struct FPCGStaticMeshSpawnerEntry
+{
+	GENERATED_BODY()
+
+	UPROPERTY(BlueprintReadWrite, EditAnywhere, Category = Settings, meta = (ClampMin = "0"))
+	int Weight = 1;
+
+	UPROPERTY(BlueprintReadWrite, EditAnywhere, Category = Settings)
+	TSoftObjectPtr<UStaticMesh> Mesh;
+};
+
+UCLASS(BlueprintType, ClassGroup = (Procedural))
+class UPCGStaticMeshSpawnerSettings : public UPCGSettings
+{
+	GENERATED_BODY()
+
+public:
+	UPROPERTY(BlueprintReadWrite, EditAnywhere, Category = Settings)
+	TArray<FPCGStaticMeshSpawnerEntry> Meshes;
+
+protected:
+	virtual FPCGElementPtr CreateElement() const override;
+};
+
+class FPCGStaticMeshSpawnerElement : public FSimpleTypedPCGElement<UPCGStaticMeshSpawnerSettings>
+{
+public:
+	virtual bool Execute(FPCGContextPtr Context) const override;
+
+private:
+	UInstancedStaticMeshComponent* GetOrCreateISMC(AActor* InActor, const UPCGComponent* SourceComponent, UStaticMesh* InMesh) const;
+};
\ No newline at end of file
diff --git a/Engine/Plugins/Experimental/PCG/Source/PCG/Public/Elements/PCGSurfaceSampler.h b/Engine/Plugins/Experimental/PCG/Source/PCG/Public/Elements/PCGSurfaceSampler.h
new file mode 100644
--- /dev/null
+++ b/Engine/Plugins/Experimental/PCG/Source/PCG/Public/Elements/PCGSurfaceSampler.h
@@ -0,0 +1,30 @@
+// Copyright Epic Games, Inc. All Rights Reserved.
+
+#pragma once
+
+#include "CoreMinimal.h"
+#include "PCGNode.h"
+#include "PCGSettings.h"
+#include "PCGElement.h"
+
+#include "PCGSurfaceSampler.generated.h"
+
+class FPCGSurfaceSamplerElement;
+
+UCLASS(BlueprintType, ClassGroup = (Procedural))
+class UPCGSurfaceSamplerSettings : public UPCGSettings
+{
+	GENERATED_BODY()
+
+public:
+	// TODO
+
+protected:
+	virtual FPCGElementPtr CreateElement() const override;
+};
+
+class FPCGSurfaceSamplerElement : public FSimpleTypedPCGElement<UPCGSurfaceSamplerSettings>
+{
+public:
+	virtual bool Execute(FPCGContextPtr Context) const override;
+};
\ No newline at end of file
diff --git a/Engine/Plugins/Experimental/PCG/Source/PCG/Public/PCGComponent.h b/Engine/Plugins/Experimental/PCG/Source/PCG/Public/PCGComponent.h
new file mode 100644
--- /dev/null
+++ b/Engine/Plugins/Experimental/PCG/Source/PCG/Public/PCGComponent.h
@@ -0,0 +1,178 @@
+// Copyright Epic Games, Inc. All Rights Reserved.
+
+#pragma once
+
+#include "CoreMinimal.h"
+#include "Components/ActorComponent.h"
+#include "PCGSubsystem.h"
+
+#include "PCGComponent.generated.h"
+
+class UPCGComponent;
+class UPCGGraph;
+class UPCGData;
+class ALandscape;
+
+#if WITH_EDITOR
+	DECLARE_MULTICAST_DELEGATE_OneParam(FOnPCGGraphGenerated, UPCGComponent*);
+	DECLARE_MULTICAST_DELEGATE_OneParam(FOnPCGGraphCleaned, UPCGComponent*);
+#endif
+
+UENUM(Blueprintable)
+enum class EPCGComponentInput : uint8
+{
+	Actor, /** Generates based on owning actor */
+	Landscape,
+	Other,
+	// More?
+	EPCGComponentInput_MAX
+};
+
+UCLASS(BlueprintType, ClassGroup = (Procedural), meta = (BlueprintSpawnableComponent))
+class PCG_API UPCGComponent : public UActorComponent
+{
+	GENERATED_BODY()
+
+	friend class UPCGSubsystem;
+
+public:
+	/** ~Begin UObject interface */
+	virtual void PostLoad() override;
+	virtual void BeginDestroy() override;
+	/** ~End UObject interface */
+
+	//~Begin UActorComponent Interface
+	virtual void BeginPlay() override;
+	virtual void OnComponentCreated() override;
+	virtual void OnComponentDestroyed(bool bDestroyingHierarchy) override;
+	//~End UActorComponent Interface
+
+	UPCGData* GetPCGData();
+	TArray<UPCGData*> GetPCGExclusionData();
+
+	bool CanPartition() const;
+	bool HasSplineComponent() const;
+
+	UPCGGraph* GetGraph() const { return Graph; }
+	void SetGraph(UPCGGraph* InGraph);
+
+	/** Transactionable methods to be called from details UI */
+	void Generate();
+	void Cleanup();
+
+	UFUNCTION(BlueprintCallable, Category = PCG)
+	void Generate(bool bForce);
+
+	UFUNCTION(BlueprintCallable, Category = PCG)
+	void Cleanup(bool bRemoveComponents, bool bSave = false);
+
+	UPROPERTY(BlueprintReadWrite, EditAnywhere, Category = Settings)
+	EPCGComponentInput InputType = EPCGComponentInput::Actor;
+
+	UPROPERTY(BlueprintReadWrite, EditAnywhere, Category = Settings)
+	int Seed = 42;
+
+	UPROPERTY(BlueprintReadWrite, EditAnywhere, Category = Settings)
+	TArray<FName> ExclusionTags;
+
+	UPROPERTY(BlueprintReadWrite, EditAnywhere, Category = Properties)
+	bool bActivated = true;
+
+	UPROPERTY(BlueprintReadWrite, EditAnywhere, AdvancedDisplay, Category = Properties)
+	bool bIsPartitioned = true;
+
+	/** Flag to indicate whether this component has run in the editor. Note that for partitionable actors, this will always be false. */
+	UPROPERTY(BlueprintReadOnly, VisibleAnywhere, AdvancedDisplay, Category = Properties)
+	bool bGenerated = false;
+
+	UPROPERTY(NonPIEDuplicateTransient)
+	bool bRuntimeGenerated = false;
+
+#if WITH_EDITORONLY_DATA
+	UPROPERTY(BlueprintReadWrite, EditAnywhere, AdvancedDisplay, Category = Properties, meta = (DisplayName = "Regenerate PCG volume in editor"))
+	bool bRegenerateInEditor = true;
+
+	UPROPERTY(BlueprintReadOnly, VisibleAnywhere, AdvancedDisplay, Category = Properties)
+	bool bDirtyGenerated = false;
+
+	FOnPCGGraphGenerated OnPCGGraphGeneratedDelegate;
+	FOnPCGGraphCleaned OnPCGGraphCleanedDelegate;
+#endif
+
+#if WITH_EDITOR
+	void Refresh();
+	void DirtyGenerated(bool bInDirtyCachedInput = false, bool bInDirtyCachedExclusions = false);
+#endif
+
+	bool IsPartitioned() const;
+
+	/** Updates internal properties from other component, dirties as required but does not trigger Refresh */
+	void SetPropertiesFromOriginal(const UPCGComponent* Original);
+
+protected:
+	UPROPERTY(BlueprintReadWrite, EditAnywhere, Category = Properties)
+	TObjectPtr<UPCGGraph> Graph;
+
+private:
+	UPCGData* CreatePCGData();
+	UPCGData* GetActorPCGData();
+	UPCGData* CreateActorPCGData();
+	UPCGData* CreateActorPCGData(AActor* Actor);
+	TArray<UPCGData*> CreatePCGExclusionData();
+
+	UPCGSubsystem* GetSubsystem() const;
+	bool ShouldGenerate(bool bForce = false) const;
+	FPCGTaskId GenerateInternal(bool bForce, const TArray<FPCGTaskId>& Dependencies);
+	void CleanupInternal(bool bRemoveComponents);
+	void PostProcessGraph(const FBox& InNewBounds, bool bInGenerated);
+
+#if WITH_EDITOR
+	virtual void PreEditChange(FProperty* PropertyAboutToChange) override;
+	virtual void PostEditChangeProperty(FPropertyChangedEvent& PropertyChangedEvent) override;
+	virtual void PreEditUndo() override;
+	virtual void PostEditUndo() override;
+
+	void OnGraphChanged(UPCGGraph* InGraph, bool bIsStructural);
+
+	void SetupActorCallbacks();
+	void TeardownActorCallbacks();
+	void SetupExclusionCallbacks();
+	void TeardownExclusionCallbacks();	
+
+	void OnActorAdded(AActor* InActor);
+	void OnActorDeleted(AActor* InActor);
+	void OnActorMoved(AActor* InActor);
+	void OnObjectPropertyChanged(UObject* InObject, FPropertyChangedEvent& InEvent);
+	bool ActorHasExcludedTag(AActor* InActor) const;
+
+	enum class EExcludedActorChange
+	{
+		Added = 0,
+		Deleted,
+		Changed
+	};
+
+	void UpdateExcludedActor(AActor* InActor, EExcludedActorChange Change);
+#endif
+
+	FBox GetGridBounds() const;
+	FBox GetGridBounds(AActor* InActor) const;
+
+	UPROPERTY(Transient, NonPIEDuplicateTransient)
+	UPCGData* CachedInputData = nullptr;
+
+	UPROPERTY(Transient, NonPIEDuplicateTransient)
+	UPCGData* CachedActorData = nullptr;
+
+	UPROPERTY(Transient, NonPIEDuplicateTransient)
+	TArray<UPCGData*> CachedExclusionData;
+
+	UPROPERTY()
+	FBox LastGeneratedBounds = FBox(EForceInit::ForceInit);
+
+#if WITH_EDITOR
+	bool bIsGenerating = false;
+	bool bWasGeneratedPriorToUndo = false;
+	FBox LastGeneratedBoundsPriorToUndo = FBox(EForceInit::ForceInit);
+#endif
+};
\ No newline at end of file
diff --git a/Engine/Plugins/Experimental/PCG/Source/PCG/Public/PCGContext.h b/Engine/Plugins/Experimental/PCG/Source/PCG/Public/PCGContext.h
new file mode 100644
--- /dev/null
+++ b/Engine/Plugins/Experimental/PCG/Source/PCG/Public/PCGContext.h
@@ -0,0 +1,36 @@
+// Copyright Epic Games, Inc. All Rights Reserved.
+
+#pragma once
+
+#include "Templates/SharedPointer.h"
+#include "PCGData.h"
+#include "PCGNode.h"
+
+struct FPCGContext;
+class UPCGComponent;
+
+typedef TSharedPtr<FPCGContext, ESPMode::ThreadSafe> FPCGContextPtr;
+
+struct FPCGContext
+{
+	FPCGDataCollection InputData;
+	FPCGDataCollection OutputData;
+	const UPCGComponent* SourceComponent = nullptr;
+	bool bExecutionSuccessful = false;
+	// TODO: add RNG source
+	// TODO: replace this by a better identification mechanism
+	const UPCGNode* Node = nullptr;
+
+	template<typename SettingsType>
+	const SettingsType* GetInputSettings()
+	{
+		if (Node && Node->DefaultSettings)
+		{
+			return Cast<SettingsType>(InputData.GetSettings(Node->DefaultSettings));
+		}
+		else
+		{
+			return InputData.GetSettings<SettingsType>();
+		}
+	}
+};
diff --git a/Engine/Plugins/Experimental/PCG/Source/PCG/Public/PCGData.h b/Engine/Plugins/Experimental/PCG/Source/PCG/Public/PCGData.h
new file mode 100644
--- /dev/null
+++ b/Engine/Plugins/Experimental/PCG/Source/PCG/Public/PCGData.h
@@ -0,0 +1,87 @@
+// Copyright Epic Games, Inc. All Rights Reserved.
+
+#pragma once
+
+#include "CoreMinimal.h"
+#include "Kismet/BlueprintFunctionLibrary.h"
+
+#include "PCGData.generated.h"
+
+class UPCGSettings;
+
+/**
+* Base class for any "data" class in the PCG framework.
+* This is an intentionally vague base class so we can have the required
+* flexibility to pass in various concrete data types, settings, and more.
+*/
+UCLASS(BlueprintType, ClassGroup = (Procedural))
+class UPCGData : public UObject
+{
+	GENERATED_BODY()
+};
+
+UENUM(BlueprintType)
+enum class EPCGDataUsage : uint8
+{
+	Input, // must be castable to concrete data
+	Exclusion,
+	Settings,
+	Override,
+	// Manual? Blocker?
+};
+
+USTRUCT(BlueprintType)
+struct FPCGTaggedData
+{
+	GENERATED_BODY()
+
+	UPROPERTY(BlueprintReadWrite, EditAnywhere, Category = Data)
+	TObjectPtr<const UPCGData> Data;
+
+	UPROPERTY(BlueprintReadWrite, EditAnywhere, Category = Data)
+	EPCGDataUsage Usage = EPCGDataUsage::Input;
+
+	UPROPERTY(BlueprintReadWrite, EditAnywhere, Category = Data)
+	TSet<FString> Tags;
+};
+
+USTRUCT(BlueprintType)
+struct FPCGDataCollection
+{
+	GENERATED_BODY()
+
+	TArray<FPCGTaggedData> GetInputs() const;
+	TArray<FPCGTaggedData> GetExclusions() const;
+
+	template<typename SettingsType>
+	const SettingsType* GetSettings() const;
+
+	const UPCGSettings* GetSettings(const UPCGSettings* InDefaultSettings) const;
+
+	UPROPERTY(BlueprintReadWrite, EditAnywhere, Category = Data)
+	TArray<FPCGTaggedData> TaggedData;
+};
+
+template<typename SettingsType>
+inline const SettingsType* FPCGDataCollection::GetSettings() const
+{
+	const FPCGTaggedData* MatchingData = TaggedData.FindByPredicate([](const FPCGTaggedData& Data) {
+		return Data.Usage == EPCGDataUsage::Settings && Cast<const SettingsType>(Data.Data) != nullptr;
+		});
+
+	return MatchingData ? Cast<const SettingsType>(MatchingData->Data) : nullptr;
+}
+
+UCLASS()
+class UPCGDataFunctionLibrary : public UBlueprintFunctionLibrary
+{
+	GENERATED_BODY()
+
+public:
+	// Blueprint methods to support interaction with FPCGDataCollection
+	UFUNCTION(BlueprintCallable, Category = Data)
+	static TArray<FPCGTaggedData> GetInputs(const FPCGDataCollection& InCollection);
+
+	UFUNCTION(BlueprintCallable, Category = Data)
+	static TArray<FPCGTaggedData> GetExclusions(const FPCGDataCollection& InCollection);
+};
diff --git a/Engine/Plugins/Experimental/PCG/Source/PCG/Public/PCGElement.h b/Engine/Plugins/Experimental/PCG/Source/PCG/Public/PCGElement.h
new file mode 100644
--- /dev/null
+++ b/Engine/Plugins/Experimental/PCG/Source/PCG/Public/PCGElement.h
@@ -0,0 +1,42 @@
+// Copyright Epic Games, Inc. All Rights Reserved.
+
+#pragma once
+
+#include "Templates/SharedPointer.h"
+#include "PCGContext.h"
+#include "PCGData.h"
+
+class IPCGElement;
+class UPCGComponent;
+
+typedef TSharedPtr<IPCGElement, ESPMode::ThreadSafe> FPCGElementPtr;
+
+/**
+* Base class for the processing bit of a PCG node/settings
+*/
+class IPCGElement
+{
+public:
+	virtual ~IPCGElement() = default;
+	virtual FPCGContextPtr Initialize(const FPCGDataCollection& InputData, const UPCGComponent* SourceComponent) = 0;
+	virtual bool Execute(FPCGContextPtr Context) const = 0;
+};
+
+/**
+* Basic PCG element class for elements that do not store any intermediate data in the context
+*/
+class FSimplePCGElement : public IPCGElement
+{
+public:
+	virtual FPCGContextPtr Initialize(const FPCGDataCollection& InputData, const UPCGComponent* SourceComponent) override;
+};
+
+/**
+* CRTP PCG element class to facilitate settings retrieval
+*/
+template<typename SettingsClass>
+class FSimpleTypedPCGElement : public FSimplePCGElement
+{
+public:
+	const SettingsClass* GetSettings();
+};
\ No newline at end of file
diff --git a/Engine/Plugins/Experimental/PCG/Source/PCG/Public/PCGGraph.h b/Engine/Plugins/Experimental/PCG/Source/PCG/Public/PCGGraph.h
new file mode 100644
--- /dev/null
+++ b/Engine/Plugins/Experimental/PCG/Source/PCG/Public/PCGGraph.h
@@ -0,0 +1,83 @@
+// Copyright Epic Games, Inc. All Rights Reserved.
+
+#pragma once
+
+#include "CoreMinimal.h"
+#include "Templates/SubclassOf.h"
+#include "PCGNode.h"
+#include "PCGSettings.h"
+#include "PCGGraphSetupBP.h"
+
+#include "PCGGraph.generated.h"
+
+#if WITH_EDITOR
+	DECLARE_MULTICAST_DELEGATE_TwoParams(FOnPCGGraphChanged, UPCGGraph* /*Graph*/, bool /*bIsStructural*/);
+#endif // WITH_EDITOR
+
+UCLASS(BlueprintType, ClassGroup = (Procedural), hidecategories=(Object))
+class PCG_API UPCGGraph : public UObject
+{
+	GENERATED_BODY()
+
+public:
+	UPCGGraph(const FObjectInitializer& ObjectInitializer);
+	/** ~Begin UObject interface */
+	virtual void PostLoad() override;
+	virtual void BeginDestroy() override;
+	/** ~End UObject interface */
+
+	UPROPERTY(BlueprintReadWrite, EditAnywhere, Category = Template)
+	TSubclassOf<UPCGGraphSetupBP> GraphTemplate;
+
+#if WITH_EDITOR
+	/** Resets & initializes the graph from the Graph Template parameter. */
+	UFUNCTION(BlueprintCallable, Category = Template)
+	void InitializeFromTemplate();
+#endif
+
+	/** Creates a default node based on the settings class wanted. Returns the newly created node. */
+	UFUNCTION(BlueprintCallable, Category = Graph)
+	UPCGNode* AddNodeOfType(TSubclassOf<class UPCGSettings> InSettingsClass);
+
+	/** Creates a node and assigns it in the input settings. Returns the created node. */
+	UFUNCTION(BlueprintCallable, Category = Graph)
+	UPCGNode* AddNode(UPCGSettings* InSettings);
+
+	/** Adds a directed edge in the graph. Returns the "To" node for easy chaining */
+	UFUNCTION(BlueprintCallable, Category = Graph)
+	UPCGNode* AddEdge(UPCGNode* From, UPCGNode* To);
+
+	/** Returns the graph input node */
+	UFUNCTION(BlueprintCallable, Category = Graph)
+	UPCGNode* GetInputNode() const { return InputNode; }
+
+	/** Returns the graph output node */
+	UFUNCTION(BlueprintCallable, Category = Graph)
+	UPCGNode* GetOutputNode() const { return OutputNode; }
+
+	bool Contains(UPCGNode* Node) const;
+	const TArray<UPCGNode*>& GetNodes() const { return Nodes; }
+
+#if WITH_EDITOR
+	FOnPCGGraphChanged OnGraphChangedDelegate;
+#endif // WITH_EDITOR
+
+protected:
+	UPROPERTY(BlueprintReadOnly, VisibleAnywhere, Category = Graph)
+	TArray<TObjectPtr<UPCGNode>> Nodes;
+
+	// Add input/output nodes
+	UPROPERTY(BlueprintReadOnly, VisibleAnywhere, Category = Graph)
+	TObjectPtr<UPCGNode> InputNode;
+
+	UPROPERTY(BlueprintReadOnly, VisibleAnywhere, Category = Graph)
+	TObjectPtr<UPCGNode> OutputNode;
+
+#if WITH_EDITOR
+private:
+	void NotifyGraphChanged(bool bIsStructural);
+	void OnSettingsChanged(UPCGNode* InNode);
+
+	bool bEnableGraphChangeNotifications = true;
+#endif // WITH_EDITOR
+};
diff --git a/Engine/Plugins/Experimental/PCG/Source/PCG/Public/PCGGraphSetupBP.h b/Engine/Plugins/Experimental/PCG/Source/PCG/Public/PCGGraphSetupBP.h
new file mode 100644
--- /dev/null
+++ b/Engine/Plugins/Experimental/PCG/Source/PCG/Public/PCGGraphSetupBP.h
@@ -0,0 +1,19 @@
+// Copyright Epic Games, Inc. All Rights Reserved.
+
+#pragma once
+
+#include "CoreMinimal.h"
+
+#include "PCGGraphSetupBP.generated.h"
+
+class UPCGGraph;
+
+UCLASS(Abstract, BlueprintType, Blueprintable, hidecategories=(Object))
+class UPCGGraphSetupBP : public UObject
+{
+	GENERATED_BODY()
+
+public:
+	UFUNCTION(BlueprintNativeEvent, BlueprintCallable, Category = Graph)
+	void Setup(UPCGGraph* Graph);
+};
\ No newline at end of file
diff --git a/Engine/Plugins/Experimental/PCG/Source/PCG/Public/PCGModule.h b/Engine/Plugins/Experimental/PCG/Source/PCG/Public/PCGModule.h
new file mode 100644
--- /dev/null
+++ b/Engine/Plugins/Experimental/PCG/Source/PCG/Public/PCGModule.h
@@ -0,0 +1,15 @@
+// Copyright Epic Games, Inc. All Rights Reserved.
+
+#pragma once
+
+#include "CoreMinimal.h"
+#include "Modules/ModuleManager.h"
+#include "Stats/Stats.h"
+
+// Logs
+DECLARE_LOG_CATEGORY_EXTERN(LogPCG, Log, All);
+
+// Stats
+DECLARE_STATS_GROUP(TEXT("PCG"), STATGROUP_PCG, STATCAT_Advanced)
+
+// CVars
diff --git a/Engine/Plugins/Experimental/PCG/Source/PCG/Public/PCGNode.h b/Engine/Plugins/Experimental/PCG/Source/PCG/Public/PCGNode.h
new file mode 100644
--- /dev/null
+++ b/Engine/Plugins/Experimental/PCG/Source/PCG/Public/PCGNode.h
@@ -0,0 +1,69 @@
+// Copyright Epic Games, Inc. All Rights Reserved.
+
+#pragma once
+
+#include "CoreMinimal.h"
+
+#include "PCGNode.generated.h"
+
+class UPCGSettings;
+class UPCGGraph;
+class IPCGElement;
+
+#if WITH_EDITOR
+	DECLARE_MULTICAST_DELEGATE_OneParam(FOnPCGNodeSettingsChanged, UPCGNode*);
+#endif
+
+UCLASS(ClassGroup = (Procedural))
+class UPCGNode : public UObject
+{
+	GENERATED_BODY()
+
+	friend class UPCGGraph;
+	friend class FPCGGraphCompiler;
+
+public:
+	UPCGNode(const FObjectInitializer& ObjectInitializer);
+	
+	/** ~Begin UObject interface */
+	virtual void PostLoad() override;
+	virtual void BeginDestroy() override;
+	/** ~End UObject interface */
+
+	/** Returns the owning graph */
+	UFUNCTION(BlueprintCallable, Category = Node)
+	UPCGGraph* GetGraph() const;
+
+	/** Adds an edge in the owning graph to the given "To" node. */
+	UFUNCTION(BlueprintCallable, Category = Node)
+	UPCGNode* AddEdgeTo(UPCGNode* To);
+
+	UPROPERTY(BlueprintReadWrite, EditAnywhere, Category = Node, meta=(EditInline))
+	TObjectPtr<UPCGSettings> DefaultSettings;
+
+#if WITH_EDITOR
+	FOnPCGNodeSettingsChanged OnNodeSettingsChangedDelegate;
+#endif
+
+protected:
+	void ConnectTo(UPCGNode* InSuccessor);
+	void ConnectFrom(UPCGNode* InPredecessor);
+
+#if WITH_EDITOR
+	virtual void PreEditChange(FProperty* PropertyAboutToChange) override;
+	virtual void PostEditChangeProperty(struct FPropertyChangedEvent& PropertyChangedEvent) override;
+	void OnSettingsChanged(UPCGSettings* InSettings);
+#endif
+
+	UPROPERTY(BlueprintReadOnly, VisibleAnywhere, Category = Node)
+	TArray<TObjectPtr<UPCGNode>> InboundNodes;
+
+	UPROPERTY(BlueprintReadOnly, VisibleAnywhere, Category = Node)
+	TArray<TObjectPtr<UPCGNode>> OutboundNodes;
+
+	// TODO: add this information:
+	// - Ability to run on non-game threads (here or element)
+	// - Ability to be multithreaded (here or element)
+	// - Generates artifacts (here or element)
+	// - Priority
+};
\ No newline at end of file
diff --git a/Engine/Plugins/Experimental/PCG/Source/PCG/Public/PCGPoint.h b/Engine/Plugins/Experimental/PCG/Source/PCG/Public/PCGPoint.h
new file mode 100644
--- /dev/null
+++ b/Engine/Plugins/Experimental/PCG/Source/PCG/Public/PCGPoint.h
@@ -0,0 +1,31 @@
+// Copyright Epic Games, Inc. All Rights Reserved.
+
+#pragma once
+
+#include "CoreMinimal.h"
+
+#include "PCGPoint.generated.h"
+
+USTRUCT(BlueprintType)
+struct FPCGPoint
+{
+	GENERATED_BODY()
+public:
+	FPCGPoint() = default;
+
+	FPCGPoint(const FTransform& InTransform, float InDensity, int32 InSeed);
+
+	FBoxSphereBounds GetBounds() const;
+
+	UPROPERTY(BlueprintReadWrite, EditAnywhere, Category = Properties)
+	FTransform Transform;
+
+	UPROPERTY(BlueprintReadWrite, EditAnywhere, Category = Properties)
+	FVector Extents;
+
+	UPROPERTY(BlueprintReadWrite, EditAnywhere, Category = Properties)
+	float Density;
+
+	UPROPERTY(BlueprintReadWrite, EditAnywhere, Category = Properties)
+	int32 Seed;
+};
\ No newline at end of file
diff --git a/Engine/Plugins/Experimental/PCG/Source/PCG/Public/PCGSettings.h b/Engine/Plugins/Experimental/PCG/Source/PCG/Public/PCGSettings.h
new file mode 100644
--- /dev/null
+++ b/Engine/Plugins/Experimental/PCG/Source/PCG/Public/PCGSettings.h
@@ -0,0 +1,63 @@
+// Copyright Epic Games, Inc. All Rights Reserved.
+
+#pragma once
+
+#include "PCGData.h"
+#include "PCGElement.h"
+
+#include "PCGSettings.generated.h"
+
+class UPCGNode;
+
+#if WITH_EDITOR
+	DECLARE_MULTICAST_DELEGATE_OneParam(FOnPCGSettingsChanged, UPCGSettings*);
+#endif
+
+/**
+* Base class for settings-as-data in the PCG framework
+*/
+UCLASS(Abstract, BlueprintType, ClassGroup = (Procedural))
+class PCG_API UPCGSettings : public UPCGData
+{
+	GENERATED_BODY()
+
+public:
+	// TODO: check if we need this to be virtual, we don't really need if we're always caching
+	/*virtual*/ FPCGElementPtr GetElement() const;
+	virtual UPCGNode* CreateNode() const;
+
+	UPROPERTY(EditAnywhere, BlueprintReadWrite, Category = Settings)
+	int Seed = 0xC35A9631; // random prime number
+
+#if WITH_EDITOR
+	FOnPCGSettingsChanged OnSettingsChangedDelegate;
+#endif
+
+protected:
+	virtual FPCGElementPtr CreateElement() const PURE_VIRTUAL(UPCGSettings::CreateElement, return nullptr;);
+
+#if WITH_EDITOR
+	virtual void PostEditChangeProperty(FPropertyChangedEvent& PropertyChangedEvent) override;
+#endif
+
+private:
+	mutable FPCGElementPtr CachedElement;
+	mutable FCriticalSection CacheLock;
+};
+
+/** Trivial / Pass-through settings used for input/output nodes */
+UCLASS(BlueprintType, ClassGroup = (Procedural))
+class UPCGTrivialSettings : public UPCGSettings
+{
+	GENERATED_BODY()
+
+protected:
+	//~UPCGSettings implementation
+	virtual FPCGElementPtr CreateElement() const override;
+};
+
+class FPCGTrivialElement : public FSimplePCGElement
+{
+public:
+	virtual bool Execute(FPCGContextPtr Context) const override;
+};
\ No newline at end of file
diff --git a/Engine/Plugins/Experimental/PCG/Source/PCG/Public/PCGSubgraph.h b/Engine/Plugins/Experimental/PCG/Source/PCG/Public/PCGSubgraph.h
new file mode 100644
--- /dev/null
+++ b/Engine/Plugins/Experimental/PCG/Source/PCG/Public/PCGSubgraph.h
@@ -0,0 +1,51 @@
+// Copyright Epic Games, Inc. All Rights Reserved.
+
+#pragma once
+
+#include "PCGSettings.h"
+
+class UPCGGraph;
+
+#include "PCGSubgraph.generated.h"
+
+UCLASS(BlueprintType, ClassGroup=(Procedural))
+class UPCGSubgraphSettings : public UPCGSettings
+{
+	GENERATED_BODY()
+
+public:
+	//~UPCGSettings interface implementation
+	virtual UPCGNode* CreateNode() const override;
+
+protected:
+	virtual FPCGElementPtr CreateElement() const override;
+	//~End UPCGSettings interface implementation
+
+protected:
+	//~Begin UObject interface implementation
+	virtual void PostLoad() override;
+	virtual void BeginDestroy() override;
+#if WITH_EDITOR
+	virtual void PreEditChange(FProperty* PropertyAboutToChange) override;
+	virtual void PostEditChangeProperty(struct FPropertyChangedEvent& PropertyChangedEvent) override;
+	//~End UObject interface implementation
+
+	void OnSubgraphChanged(UPCGGraph* InGraph, bool bIsStructural);
+#endif
+
+public:
+	UPROPERTY(BlueprintReadWrite, EditAnywhere, Category = Settings)
+	TObjectPtr<UPCGGraph> Subgraph;
+};
+
+UCLASS(ClassGroup = (Procedural))
+class UPCGSubgraphNode : public UPCGNode
+{
+	GENERATED_BODY()
+
+public:
+	UPROPERTY(BlueprintReadWrite, EditAnywhere, Category = Settings)
+	bool bAllowDynamicGraph = false;
+
+	TObjectPtr<UPCGGraph> GetGraph() const;
+};
\ No newline at end of file
diff --git a/Engine/Plugins/Experimental/PCG/Source/PCG/Public/PCGSubsystem.h b/Engine/Plugins/Experimental/PCG/Source/PCG/Public/PCGSubsystem.h
new file mode 100644
--- /dev/null
+++ b/Engine/Plugins/Experimental/PCG/Source/PCG/Public/PCGSubsystem.h
@@ -0,0 +1,81 @@
+// Copyright Epic Games, Inc. All Rights Reserved.
+
+#pragma once
+
+#include "Subsystems/WorldSubsystem.h"
+
+#include "PCGSubsystem.generated.h"
+
+class UPCGComponent;
+class FPCGGraphExecutor;
+class APCGPartitionActor;
+class APCGWorldActor;
+class UPCGGraph;
+
+using FPCGTaskId = uint64;
+static const FPCGTaskId InvalidTaskId = (uint64)-1;
+
+/**
+* UPCGSubsystem
+*/
+UCLASS()
+class PCG_API UPCGSubsystem : public UTickableWorldSubsystem
+{
+	GENERATED_BODY()
+
+public:
+	//~ Begin USubsystem Interface.
+	virtual void Deinitialize() override;
+	//~ End USubsystem Interface
+
+	//~ Begin UWorldSubsystem Interface.
+	virtual void PostInitialize() override;
+	// need UpdateStreamingState? 
+	//~ End UWorldSubsystem Interface.
+
+	//~ Begin FTickableGameObject
+	virtual void Tick(float DeltaSeconds) override;
+	virtual bool IsTickableInEditor() const override { return true; }
+	virtual ETickableTickType GetTickableTickType() const override;
+	virtual TStatId GetStatId() const override;
+	//~ End FTickableGameObject
+
+	APCGWorldActor* GetPCGWorldActor();
+	void RegisterPCGWorldActor(APCGWorldActor* InActor);
+	void UnregisterPCGWorldActor(APCGWorldActor* InActor);
+
+	// Schedule graph (owner -> graph)
+	FPCGTaskId ScheduleComponent(UPCGComponent* PCGComponent, const TArray<FPCGTaskId>& Dependencies);
+
+#if WITH_EDITOR
+public:
+	/** Operations to schedule a later operation, which enables to delay bounds querying */
+	void DelayPartitionGraph(UPCGComponent* Component);
+	void DelayUnpartitionGraph(UPCGComponent* Component);
+	FPCGTaskId DelayGenerateGraph(UPCGComponent* Component, bool bSave);
+
+	/** Schedules an operation to cleanup the graph in the given bounds */
+	void CleanupGraph(UPCGComponent* Component, const FBox& InBounds, bool bRemoveComponents, bool bSave);
+
+	/** Immediately dirties the partition actors in the given bounds */
+	void DirtyGraph(UPCGComponent* Component, const FBox& InBounds, bool bDirtyInputs, bool bDirtyExclusions);
+
+	/** Partition actors methods */
+	void CleanupPartitionActors(const FBox& InBounds);
+	void DeletePartitionActors();
+
+	/** General job scheduling, used to control loading/unloading */
+	FPCGTaskId ScheduleGeneric(TFunction<bool()> InOperation, const TArray<FPCGTaskId>& TaskDependencies);
+
+	/** Propagate to the graph compiler graph changes */
+	void NotifyGraphChanged(UPCGGraph* InGraph);
+
+private:
+	FPCGTaskId DelayProcessGraph(UPCGComponent* Component, bool bGenerate, bool bSave, bool bUseEmptyNewBounds);
+	FPCGTaskId ProcessGraph(UPCGComponent* Component, const FBox& InPreviousBounds, const FBox& InNewBounds, bool bGenerate, bool bSave);
+#endif // WITH_EDITOR
+	
+private:
+	APCGWorldActor* PCGWorldActor = nullptr;
+	FPCGGraphExecutor* GraphExecutor = nullptr;
+};
\ No newline at end of file
diff --git a/Engine/Plugins/Experimental/PCG/Source/PCG/Public/PCGVolume.h b/Engine/Plugins/Experimental/PCG/Source/PCG/Public/PCGVolume.h
new file mode 100644
--- /dev/null
+++ b/Engine/Plugins/Experimental/PCG/Source/PCG/Public/PCGVolume.h
@@ -0,0 +1,23 @@
+// Copyright Epic Games, Inc. All Rights Reserved.
+
+#pragma once
+
+#include "CoreMinimal.h"
+#include "UObject/ObjectMacros.h"
+#include "GameFramework/Volume.h"
+
+#include "PCGVolume.generated.h"
+
+class UPCGComponent;
+
+UCLASS(BlueprintType)
+class PCG_API APCGVolume : public AVolume
+{
+	GENERATED_BODY()
+
+public:
+	APCGVolume(const FObjectInitializer& ObjectInitializer);
+
+	UPROPERTY(BlueprintReadWrite, EditAnywhere, Category = PCG)
+	TObjectPtr<UPCGComponent> PCGComponent;
+};
diff --git a/Engine/Plugins/Experimental/PCG/Source/PCG/Public/PCGWorldActor.h b/Engine/Plugins/Experimental/PCG/Source/PCG/Public/PCGWorldActor.h
new file mode 100644
--- /dev/null
+++ b/Engine/Plugins/Experimental/PCG/Source/PCG/Public/PCGWorldActor.h
@@ -0,0 +1,31 @@
+// Copyright Epic Games, Inc. All Rights Reserved.
+
+#pragma once
+
+#include "CoreMinimal.h"
+#include "GameFramework/Actor.h"
+
+#include "PCGWorldActor.generated.h"
+
+UCLASS(MinimalAPI, NotBlueprintable, NotPlaceable)
+class APCGWorldActor : public AActor
+{
+	GENERATED_BODY()
+
+public:
+	//~Begin AActor Interface
+	virtual void PostLoad() override;
+	virtual void BeginDestroy() override;
+
+#if WITH_EDITOR	
+	virtual bool CanChangeIsSpatiallyLoadedFlag() const { return false; }
+	virtual bool IsUserManaged() const override { return false; }
+	//~End AActor Interface
+
+	static APCGWorldActor* CreatePCGWorldActor(UWorld* InWorld);
+#endif
+
+private:
+	void RegisterToSubsystem();
+	void UnregisterFromSubsystem();
+};
\ No newline at end of file
diff --git a/Engine/Plugins/Experimental/PCG/Source/PCGEditor/PCGEditor.Build.cs b/Engine/Plugins/Experimental/PCG/Source/PCGEditor/PCGEditor.Build.cs
new file mode 100644
--- /dev/null
+++ b/Engine/Plugins/Experimental/PCG/Source/PCGEditor/PCGEditor.Build.cs
@@ -0,0 +1,47 @@
+// Copyright 1998-2017 Epic Games, Inc. All Rights Reserved.
+
+namespace UnrealBuildTool.Rules
+{
+	public class PCGEditor : ModuleRules
+	{
+		public PCGEditor(ReadOnlyTargetRules Target) : base(Target)
+		{
+			PCHUsage = PCHUsageMode.UseExplicitOrSharedPCHs;
+
+			PublicDependencyModuleNames.AddRange(
+				new string[] {
+					"Core",
+					"Projects",
+					"Engine",
+					"CoreUObject",
+				});
+
+			PrivateDependencyModuleNames.AddRange(
+				new string[]
+				{					
+				}
+			);
+
+			PrivateDependencyModuleNames.AddRange(
+				new string[]{
+					"AssetTools",
+					"DesktopWidgets",
+					"EditorStyle",
+					"EditorSubsystem",
+					"InputCore",
+					"Kismet",
+					"PCG",
+					"Slate",
+					"SlateCore",
+					"SourceControl",
+					"ToolMenus",
+					"UnrealEd",
+				}
+			);
+
+			PrivateIncludePaths.AddRange(
+				new string[] {
+				});
+		}
+	}
+}
diff --git a/Engine/Plugins/Experimental/PCG/Source/PCGEditor/Private/AssetTypeActions/PCGCommonAssetTypeActions.cpp b/Engine/Plugins/Experimental/PCG/Source/PCGEditor/Private/AssetTypeActions/PCGCommonAssetTypeActions.cpp
new file mode 100644
--- /dev/null
+++ b/Engine/Plugins/Experimental/PCG/Source/PCGEditor/Private/AssetTypeActions/PCGCommonAssetTypeActions.cpp
@@ -0,0 +1,14 @@
+// Copyright Epic Games, Inc. All Rights Reserved.
+
+#include "PCGCommonAssetTypeActions.h"
+#include "PCGEditorModule.h"
+
+uint32 FPCGCommonAssetTypeActions::GetCategories()
+{
+	return FPCGEditorModule::GetAssetCategory();
+}
+
+FColor FPCGCommonAssetTypeActions::GetTypeColor() const
+{
+	return FColor::Turquoise;
+}
\ No newline at end of file
diff --git a/Engine/Plugins/Experimental/PCG/Source/PCGEditor/Private/AssetTypeActions/PCGCommonAssetTypeActions.h b/Engine/Plugins/Experimental/PCG/Source/PCGEditor/Private/AssetTypeActions/PCGCommonAssetTypeActions.h
new file mode 100644
--- /dev/null
+++ b/Engine/Plugins/Experimental/PCG/Source/PCGEditor/Private/AssetTypeActions/PCGCommonAssetTypeActions.h
@@ -0,0 +1,14 @@
+// Copyright Epic Games, Inc. All Rights Reserved.
+
+#pragma once
+
+#include "CoreMinimal.h"
+#include "AssetTypeActions_Base.h"
+
+class FPCGCommonAssetTypeActions : public FAssetTypeActions_Base
+{
+public:
+	// ~IAssetTypeActions Implementation
+	virtual uint32 GetCategories() override;
+	virtual FColor GetTypeColor() const override;
+};
\ No newline at end of file
diff --git a/Engine/Plugins/Experimental/PCG/Source/PCGEditor/Private/AssetTypeActions/PCGGraphAssetTypeActions.cpp b/Engine/Plugins/Experimental/PCG/Source/PCGEditor/Private/AssetTypeActions/PCGGraphAssetTypeActions.cpp
new file mode 100644
--- /dev/null
+++ b/Engine/Plugins/Experimental/PCG/Source/PCGEditor/Private/AssetTypeActions/PCGGraphAssetTypeActions.cpp
@@ -0,0 +1,14 @@
+// Copyright Epic Games, Inc. All Rights Reserved.
+
+#include "PCGGraphAssetTypeActions.h"
+#include "PCGGraph.h"
+
+FText FPCGGraphAssetTypeActions::GetName() const
+{
+	return NSLOCTEXT("AssetTypeActions", "PCGGraphAssetTypeActions", "PCG Graph");
+}
+
+UClass* FPCGGraphAssetTypeActions::GetSupportedClass() const
+{
+	return UPCGGraph::StaticClass();
+}
\ No newline at end of file
diff --git a/Engine/Plugins/Experimental/PCG/Source/PCGEditor/Private/AssetTypeActions/PCGGraphAssetTypeActions.h b/Engine/Plugins/Experimental/PCG/Source/PCGEditor/Private/AssetTypeActions/PCGGraphAssetTypeActions.h
new file mode 100644
--- /dev/null
+++ b/Engine/Plugins/Experimental/PCG/Source/PCGEditor/Private/AssetTypeActions/PCGGraphAssetTypeActions.h
@@ -0,0 +1,13 @@
+// Copyright Epic Games, Inc. All Rights Reserved.
+
+#pragma once
+
+#include "PCGCommonAssetTypeActions.h"
+
+class FPCGGraphAssetTypeActions : public FPCGCommonAssetTypeActions
+{
+public:
+	// IAssetTypeActions Implementation
+	virtual FText GetName() const override;
+	virtual UClass* GetSupportedClass() const override;
+};
\ No newline at end of file
diff --git a/Engine/Plugins/Experimental/PCG/Source/PCGEditor/Private/AssetTypeActions/PCGSettingsAssetTypeActions.cpp b/Engine/Plugins/Experimental/PCG/Source/PCGEditor/Private/AssetTypeActions/PCGSettingsAssetTypeActions.cpp
new file mode 100644
--- /dev/null
+++ b/Engine/Plugins/Experimental/PCG/Source/PCGEditor/Private/AssetTypeActions/PCGSettingsAssetTypeActions.cpp
@@ -0,0 +1,14 @@
+// Copyright Epic Games, Inc. All Rights Reserved.
+
+#include "PCGSettingsAssetTypeActions.h"
+#include "PCGGraph.h"
+
+FText FPCGSettingsAssetTypeActions::GetName() const
+{
+	return NSLOCTEXT("AssetTypeActions", "PCGSettingsAssetTypeActions", "PCG Settings");
+}
+
+UClass* FPCGSettingsAssetTypeActions::GetSupportedClass() const
+{
+	return UPCGSettings::StaticClass();
+}
\ No newline at end of file
diff --git a/Engine/Plugins/Experimental/PCG/Source/PCGEditor/Private/AssetTypeActions/PCGSettingsAssetTypeActions.h b/Engine/Plugins/Experimental/PCG/Source/PCGEditor/Private/AssetTypeActions/PCGSettingsAssetTypeActions.h
new file mode 100644
--- /dev/null
+++ b/Engine/Plugins/Experimental/PCG/Source/PCGEditor/Private/AssetTypeActions/PCGSettingsAssetTypeActions.h
@@ -0,0 +1,13 @@
+// Copyright Epic Games, Inc. All Rights Reserved.
+
+#pragma once
+
+#include "PCGCommonAssetTypeActions.h"
+
+class FPCGSettingsAssetTypeActions : public FPCGCommonAssetTypeActions
+{
+public:
+	// IAssetTypeActions Implementation
+	virtual FText GetName() const override;
+	virtual UClass* GetSupportedClass() const override;
+};
\ No newline at end of file
diff --git a/Engine/Plugins/Experimental/PCG/Source/PCGEditor/Private/PCGComponentDetails.cpp b/Engine/Plugins/Experimental/PCG/Source/PCGEditor/Private/PCGComponentDetails.cpp
new file mode 100644
--- /dev/null
+++ b/Engine/Plugins/Experimental/PCG/Source/PCGEditor/Private/PCGComponentDetails.cpp
@@ -0,0 +1,145 @@
+// Copyright Epic Games, Inc. All Rights Reserved.
+
+#include "PCGComponentDetails.h"
+#include "PCGComponent.h"
+
+#include "Widgets/DeclarativeSyntaxSupport.h"
+#include "Widgets/Text/STextBlock.h"
+#include "Widgets/Input/SButton.h"
+#include "PropertyHandle.h"
+#include "DetailLayoutBuilder.h"
+#include "DetailWidgetRow.h"
+#include "DetailCategoryBuilder.h"
+#include "Modules/ModuleManager.h"
+#include "LevelEditor.h"
+
+#define LOCTEXT_NAMESPACE "PCGComponentDetails"
+
+TSharedRef<IDetailCustomization> FPCGComponentDetails::MakeInstance()
+{
+	return MakeShareable(new FPCGComponentDetails());
+}
+
+void FPCGComponentDetails::CustomizeDetails(IDetailLayoutBuilder& DetailBuilder)
+{
+	const FName PCGCategoryName("PCG");
+	IDetailCategoryBuilder& PCGCategory = DetailBuilder.EditCategory(PCGCategoryName);
+
+	TArray<TWeakObjectPtr<UObject>> ObjectsBeingCustomized;
+	DetailBuilder.GetObjectsBeingCustomized(ObjectsBeingCustomized);
+
+	for (TWeakObjectPtr<UObject>& Object : ObjectsBeingCustomized)
+	{
+		UPCGComponent* Component = Cast<UPCGComponent>(Object.Get());
+		if (ensure(Component))
+		{
+			SelectedComponents.Add(Component);
+		}
+	}
+
+	TArray<TSharedRef<IPropertyHandle>> AllProperties;
+	bool bSimpleProperties = true;
+	bool bAdvancedProperties = false;
+	// Add all properties in the category in order
+	PCGCategory.GetDefaultProperties(AllProperties, bSimpleProperties, bAdvancedProperties);
+
+	for (auto& Property : AllProperties)
+	{
+		PCGCategory.AddProperty(Property);
+	}
+
+	FDetailWidgetRow& NewRow = PCGCategory.AddCustomRow(FText::GetEmpty());
+
+	NewRow.ValueContent()
+		.MaxDesiredWidth(120.f)
+		[
+			SNew(SHorizontalBox)
+			+ SHorizontalBox::Slot()
+			.AutoWidth()
+			.Padding(4.0f)
+			[
+				SNew(SButton)
+				.OnClicked(this, &FPCGComponentDetails::OnGenerateClicked)
+				[
+					SNew(STextBlock)
+					.Font(IDetailLayoutBuilder::GetDetailFont())
+					.Text(LOCTEXT("GenerateButton", "Generate"))
+				]
+			]
+			+ SHorizontalBox::Slot()
+			.AutoWidth()
+			.Padding(4.0f)
+			[
+				SNew(SButton)
+				.OnClicked(this, &FPCGComponentDetails::OnCleanupClicked)
+				[
+					SNew(STextBlock)
+					.Font(IDetailLayoutBuilder::GetDetailFont())
+					.Text(LOCTEXT("CleanupButton", "Cleanup"))
+				]
+			]
+		];
+
+	// Attach to generated delegate on the selected components
+	for (TWeakObjectPtr<UPCGComponent>& Component : SelectedComponents)
+	{
+		if (Component.IsValid())
+		{
+			Component.Get()->OnPCGGraphGeneratedDelegate.AddRaw(this, &FPCGComponentDetails::OnGraphChanged);
+			Component.Get()->OnPCGGraphCleanedDelegate.AddRaw(this, &FPCGComponentDetails::OnGraphChanged);
+		}
+	}
+}
+
+void FPCGComponentDetails::PendingDelete()
+{
+	// Detach from the generated delegate
+	for (TWeakObjectPtr<UPCGComponent>& Component : SelectedComponents)
+	{
+		if (Component.IsValid())
+		{
+			Component.Get()->OnPCGGraphGeneratedDelegate.RemoveAll(this);
+			Component.Get()->OnPCGGraphCleanedDelegate.RemoveAll(this);
+		}
+	}
+}
+
+FReply FPCGComponentDetails::OnGenerateClicked()
+{
+	for (TWeakObjectPtr<UPCGComponent>& Component : SelectedComponents)
+	{
+		if (Component.IsValid())
+		{
+			Component.Get()->Generate();
+		}
+	}
+
+	return FReply::Handled();
+}
+
+void FPCGComponentDetails::OnGraphChanged(UPCGComponent* InComponent)
+{
+	if (!InComponent || !InComponent->GetOwner())
+	{
+		return;
+	}
+
+	// Notify editor that some components might have changed
+	FLevelEditorModule& LevelEditor = FModuleManager::LoadModuleChecked<FLevelEditorModule>("LevelEditor");
+	LevelEditor.BroadcastComponentsEdited();
+}
+
+FReply FPCGComponentDetails::OnCleanupClicked()
+{
+	for (TWeakObjectPtr<UPCGComponent>& Component : SelectedComponents)
+	{
+		if (Component.IsValid())
+		{
+			Component.Get()->Cleanup();
+		}
+	}
+
+	return FReply::Handled();
+}
+
+#undef LOCTEXT_NAMESPACE
diff --git a/Engine/Plugins/Experimental/PCG/Source/PCGEditor/Private/PCGEditorModule.cpp b/Engine/Plugins/Experimental/PCG/Source/PCGEditor/Private/PCGEditorModule.cpp
new file mode 100644
--- /dev/null
+++ b/Engine/Plugins/Experimental/PCG/Source/PCGEditor/Private/PCGEditorModule.cpp
@@ -0,0 +1,156 @@
+// Copyright Epic Games, Inc. All Rights Reserved.
+
+#include "PCGEditorModule.h"
+
+#include "IAssetTools.h"
+#include "LevelEditor.h"
+#include "Modules/ModuleManager.h"
+#include "PropertyEditorModule.h"
+#include "Toolkits/IToolkit.h"
+#include "Framework/MultiBox/MultiBoxBuilder.h"
+
+#include "PCGComponentDetails.h"
+#include "PCGGraphDetails.h"
+#include "PCGVolumeFactory.h"
+#include "AssetTypeActions/PCGGraphAssetTypeActions.h"
+#include "AssetTypeActions/PCGSettingsAssetTypeActions.h"
+
+#include "PCGSubsystem.h"
+
+#define LOCTEXT_NAMESPACE "FPCGEditorModule"
+
+EAssetTypeCategories::Type FPCGEditorModule::PCGAssetCategory;
+
+void FPCGEditorModule::StartupModule()
+{
+	RegisterDetailsCustomizations();
+	RegisterAssetTypeActions();
+	RegisterMenuExtensions();
+
+	if (GEditor)
+	{
+		GEditor->ActorFactories.Add(NewObject<UPCGVolumeFactory>());
+	}
+}
+
+void FPCGEditorModule::ShutdownModule()
+{
+	UnregisterAssetTypeActions();
+	UnregisterDetailsCustomizations();
+	UnregisterMenuExtensions();
+
+	if (GEditor)
+	{
+		GEditor->ActorFactories.RemoveAll([](const UActorFactory* ActorFactory) { return ActorFactory->IsA<UPCGVolumeFactory>(); });
+	}
+}
+
+bool FPCGEditorModule::SupportsDynamicReloading()
+{
+	return true;
+}
+
+void FPCGEditorModule::RegisterDetailsCustomizations()
+{
+	FPropertyEditorModule& PropertyEditor = FModuleManager::LoadModuleChecked<FPropertyEditorModule>("PropertyEditor");
+	PropertyEditor.RegisterCustomClassLayout("PCGComponent", FOnGetDetailCustomizationInstance::CreateStatic(&FPCGComponentDetails::MakeInstance));
+	PropertyEditor.RegisterCustomClassLayout("PCGGraph", FOnGetDetailCustomizationInstance::CreateStatic(&FPCGGraphDetails::MakeInstance));
+}
+
+void FPCGEditorModule::UnregisterDetailsCustomizations()
+{
+	if (FModuleManager::Get().IsModuleLoaded("PropertyEditor"))
+	{
+		FPropertyEditorModule& PropertyModule = FModuleManager::LoadModuleChecked<FPropertyEditorModule>("PropertyEditor");
+		PropertyModule.UnregisterCustomClassLayout("PCGComponent");
+		PropertyModule.UnregisterCustomClassLayout("PCGGraph");
+		PropertyModule.NotifyCustomizationModuleChanged();
+	}
+}
+
+void FPCGEditorModule::RegisterAssetTypeActions()
+{
+	IAssetTools& AssetTools = FModuleManager::LoadModuleChecked<FAssetToolsModule>("AssetTools").Get();
+	PCGAssetCategory = AssetTools.RegisterAdvancedAssetCategory(FName(TEXT("PCG")), LOCTEXT("PCGAssetCategory", "PCG"));
+
+	RegisteredAssetTypeActions.Emplace(MakeShareable(new FPCGGraphAssetTypeActions()));
+	RegisteredAssetTypeActions.Emplace(MakeShareable(new FPCGSettingsAssetTypeActions()));
+
+	for (auto Action : RegisteredAssetTypeActions)
+	{
+		AssetTools.RegisterAssetTypeActions(Action);
+	}
+}
+
+void FPCGEditorModule::UnregisterAssetTypeActions()
+{
+	FAssetToolsModule* AssetToolsModule = FModuleManager::GetModulePtr<FAssetToolsModule>("AssetTools");
+
+	if (!AssetToolsModule)
+	{
+		return;
+	}
+
+	IAssetTools& AssetTools = AssetToolsModule->Get();
+
+	for (auto Action : RegisteredAssetTypeActions)
+	{
+		AssetTools.UnregisterAssetTypeActions(Action);
+	}
+}
+
+void FPCGEditorModule::RegisterMenuExtensions()
+{
+	MenuExtensibilityManager = MakeShareable(new FExtensibilityManager);
+	ToolBarExtensibilityManager = MakeShareable(new FExtensibilityManager);
+
+	FLevelEditorModule& LevelEditorModule = FModuleManager::LoadModuleChecked<FLevelEditorModule>("LevelEditor");
+
+	{
+		TSharedPtr<FExtender> NewMenuExtender = MakeShareable(new FExtender);
+		NewMenuExtender->AddMenuExtension("LevelEditor",
+			EExtensionHook::After,
+			nullptr,
+			FMenuExtensionDelegate::CreateRaw(this, &FPCGEditorModule::AddMenuEntry));
+
+		LevelEditorModule.GetMenuExtensibilityManager()->AddExtender(NewMenuExtender);
+	}
+}
+
+void FPCGEditorModule::UnregisterMenuExtensions()
+{
+	MenuExtensibilityManager.Reset();
+	ToolBarExtensibilityManager.Reset();
+}
+
+void FPCGEditorModule::AddMenuEntry(FMenuBuilder& MenuBuilder)
+{
+	MenuBuilder.BeginSection("PCGMenu", TAttribute<FText>(FText::FromString("PCG Tools")));
+
+	MenuBuilder.AddSubMenu(
+		LOCTEXT("PCGSubMenu", "PCG Framework"),
+		LOCTEXT("PCGSubMenu_Tooltip", "PCG Framework related functionality"),
+		FNewMenuDelegate::CreateRaw(this, &FPCGEditorModule::PopulateMenuActions));
+	
+	MenuBuilder.EndSection();
+}
+
+void FPCGEditorModule::PopulateMenuActions(FMenuBuilder& MenuBuilder)
+{
+	MenuBuilder.AddMenuEntry(
+		LOCTEXT("DeletePCGActors", "Delete all PCG partition actors"),
+		LOCTEXT("DeletePCGActors_Tooltip", "Deletes all PCG partition actors in the current world"),
+		FSlateIcon(),
+		FUIAction(
+			FExecuteAction::CreateLambda([]() {
+				if (UWorld* World = GEditor->GetEditorWorldContext().World())
+				{
+					World->GetSubsystem<UPCGSubsystem>()->DeletePartitionActors();
+				}
+			})),
+		NAME_None);
+}
+
+IMPLEMENT_MODULE(FPCGEditorModule, PCGEditor);
+
+#undef LOCTEXT_NAMESPACE
\ No newline at end of file
diff --git a/Engine/Plugins/Experimental/PCG/Source/PCGEditor/Private/PCGEditorModule.h b/Engine/Plugins/Experimental/PCG/Source/PCGEditor/Private/PCGEditorModule.h
new file mode 100644
--- /dev/null
+++ b/Engine/Plugins/Experimental/PCG/Source/PCGEditor/Private/PCGEditorModule.h
@@ -0,0 +1,46 @@
+// Copyright Epic Games, Inc. All Rights Reserved.
+
+#pragma once
+
+#include "AssetTypeCategories.h"
+#include "IAssetTypeActions.h"
+#include "Modules/ModuleInterface.h"
+#include "Toolkits/AssetEditorToolkit.h"
+
+class FMenuBuilder;
+
+class FPCGEditorModule : public IModuleInterface, public IHasMenuExtensibility, public IHasToolBarExtensibility
+{
+public:
+	// ~IModuleInterface implementation
+	virtual void StartupModule() override;
+	virtual void ShutdownModule() override;
+	virtual bool SupportsDynamicReloading() override;
+	// ~End IModuleInterface implementation
+
+	//~ IHasMenuExtensibility interface
+	virtual TSharedPtr<FExtensibilityManager> GetMenuExtensibilityManager() override { return MenuExtensibilityManager; }
+	//~ End IHasMenuExtensibility interface
+
+	//~ IHasToolBarExtensibility interface
+	virtual TSharedPtr<FExtensibilityManager> GetToolBarExtensibilityManager() override { return ToolBarExtensibilityManager; }
+	//~ End IHasToolBarExtensibility interface
+
+	static EAssetTypeCategories::Type GetAssetCategory() { return PCGAssetCategory; }
+
+protected:
+	void RegisterDetailsCustomizations();
+	void UnregisterDetailsCustomizations();
+	void RegisterAssetTypeActions();
+	void UnregisterAssetTypeActions();
+	void RegisterMenuExtensions();
+	void UnregisterMenuExtensions();
+	void AddMenuEntry(FMenuBuilder& MenuBuilder);
+	void PopulateMenuActions(FMenuBuilder& MenuBuilder);
+
+	TArray<TSharedRef<IAssetTypeActions>> RegisteredAssetTypeActions;
+	static EAssetTypeCategories::Type PCGAssetCategory;
+
+	TSharedPtr<FExtensibilityManager> MenuExtensibilityManager;
+	TSharedPtr<FExtensibilityManager> ToolBarExtensibilityManager;
+};
diff --git a/Engine/Plugins/Experimental/PCG/Source/PCGEditor/Private/PCGGraphDetails.cpp b/Engine/Plugins/Experimental/PCG/Source/PCGEditor/Private/PCGGraphDetails.cpp
new file mode 100644
--- /dev/null
+++ b/Engine/Plugins/Experimental/PCG/Source/PCGEditor/Private/PCGGraphDetails.cpp
@@ -0,0 +1,83 @@
+// Copyright Epic Games, Inc. All Rights Reserved.
+
+#include "PCGGraphDetails.h"
+#include "PCGGraph.h"
+
+#include "Widgets/DeclarativeSyntaxSupport.h"
+#include "Widgets/Text/STextBlock.h"
+#include "Widgets/Input/SButton.h"
+#include "PropertyHandle.h"
+#include "DetailLayoutBuilder.h"
+#include "DetailWidgetRow.h"
+#include "DetailCategoryBuilder.h"
+
+#define LOCTEXT_NAMESPACE "PCGGraphDetails"
+
+TSharedRef<IDetailCustomization> FPCGGraphDetails::MakeInstance()
+{
+	return MakeShareable(new FPCGGraphDetails());
+}
+
+void FPCGGraphDetails::CustomizeDetails(IDetailLayoutBuilder& DetailBuilder)
+{
+	const FName PCGCategoryName("PCG");
+	IDetailCategoryBuilder& PCGCategory = DetailBuilder.EditCategory(PCGCategoryName);
+
+	TArray<TWeakObjectPtr<UObject>> ObjectsBeingCustomized;
+	DetailBuilder.GetObjectsBeingCustomized(ObjectsBeingCustomized);
+
+	for (TWeakObjectPtr<UObject>& Object : ObjectsBeingCustomized)
+	{
+		UPCGGraph* Graph = Cast<UPCGGraph>(Object.Get());
+		if (ensure(Graph))
+		{
+			SelectedGraphs.Add(Graph);
+		}
+	}
+
+	TArray<TSharedRef<IPropertyHandle>> AllProperties;
+	bool bSimpleProperties = true;
+	bool bAdvancedProperties = false;
+	// Add all properties in the category in order
+	PCGCategory.GetDefaultProperties(AllProperties, bSimpleProperties, bAdvancedProperties);
+
+	for (auto& Property : AllProperties)
+	{
+		PCGCategory.AddProperty(Property);
+	}
+
+	FDetailWidgetRow& NewRow = PCGCategory.AddCustomRow(FText::GetEmpty());
+
+	NewRow.ValueContent()
+		.MaxDesiredWidth(120.f)
+		[
+			SNew(SHorizontalBox)
+			+ SHorizontalBox::Slot()
+			.AutoWidth()
+			.Padding(4.0f)
+			[
+				SNew(SButton)
+				.OnClicked(this, &FPCGGraphDetails::OnInitializeClicked)
+				[
+					SNew(STextBlock)
+					.Font(IDetailLayoutBuilder::GetDetailFont())
+					.Text(LOCTEXT("InitializeButton", "Initialize from template"))
+				]
+			]
+		];
+}
+
+FReply FPCGGraphDetails::OnInitializeClicked()
+{
+	for (TWeakObjectPtr<UPCGGraph>& Graph : SelectedGraphs)
+	{
+		if (Graph.IsValid())
+		{
+			Graph.Get()->InitializeFromTemplate();
+		}
+	}
+
+	return FReply::Handled();
+}
+
+#undef LOCTEXT_NAMESPACE
\ No newline at end of file
diff --git a/Engine/Plugins/Experimental/PCG/Source/PCGEditor/Private/PCGGraphFactory.cpp b/Engine/Plugins/Experimental/PCG/Source/PCGEditor/Private/PCGGraphFactory.cpp
new file mode 100644
--- /dev/null
+++ b/Engine/Plugins/Experimental/PCG/Source/PCGEditor/Private/PCGGraphFactory.cpp
@@ -0,0 +1,22 @@
+// Copyright Epic Games, Inc. All Rights Reserved.
+
+#include "PCGGraphFactory.h"
+#include "PCGGraph.h"
+
+UPCGGraphFactory::UPCGGraphFactory(const FObjectInitializer& ObjectInitializer)
+	: Super(ObjectInitializer)
+{
+	SupportedClass = UPCGGraph::StaticClass();
+	bCreateNew = true;
+	bEditAfterNew = true;
+}
+
+UObject* UPCGGraphFactory::FactoryCreateNew(UClass* InClass, UObject* InParent, FName InName, EObjectFlags Flags, UObject* Context, FFeedbackContext* Warn)
+{
+	return NewObject<UPCGGraph>(InParent, InClass, InName, Flags);
+}
+
+bool UPCGGraphFactory::ShouldShowInNewMenu() const
+{
+	return true;
+}
\ No newline at end of file
diff --git a/Engine/Plugins/Experimental/PCG/Source/PCGEditor/Private/PCGGraphFactory.h b/Engine/Plugins/Experimental/PCG/Source/PCGEditor/Private/PCGGraphFactory.h
new file mode 100644
--- /dev/null
+++ b/Engine/Plugins/Experimental/PCG/Source/PCGEditor/Private/PCGGraphFactory.h
@@ -0,0 +1,21 @@
+// Copyright Epic Games, Inc. All Rights Reserved.
+
+#pragma once
+
+#include "Factories/Factory.h"
+#include "UObject/ObjectMacros.h"
+
+#include "PCGGraphFactory.generated.h"
+
+UCLASS(hidecategories=Object)
+class UPCGGraphFactory : public UFactory
+{
+	GENERATED_BODY()
+
+public:
+	UPCGGraphFactory(const FObjectInitializer& ObjectInitializer);
+
+	//~UFactory interface
+	virtual UObject* FactoryCreateNew(UClass* InClass, UObject* InParent, FName InName, EObjectFlags Flags, UObject* Context, FFeedbackContext* Warn) override;
+	virtual bool ShouldShowInNewMenu() const override;
+};
\ No newline at end of file
diff --git a/Engine/Plugins/Experimental/PCG/Source/PCGEditor/Private/PCGVolumeFactory.cpp b/Engine/Plugins/Experimental/PCG/Source/PCGEditor/Private/PCGVolumeFactory.cpp
new file mode 100644
--- /dev/null
+++ b/Engine/Plugins/Experimental/PCG/Source/PCGEditor/Private/PCGVolumeFactory.cpp
@@ -0,0 +1,16 @@
+// Copyright Epic Games, Inc. All Rights Reserved.
+
+#include "PCGVolumeFactory.h"
+#include "PCGVolume.h"
+
+#define LOCTEXT_NAMESPACE "PCGVolumeFactory"
+
+UPCGVolumeFactory::UPCGVolumeFactory(const FObjectInitializer& ObjectInitializer)
+	: Super(ObjectInitializer)
+{
+	DisplayName = LOCTEXT("PCGVolumeDisplayName", "PCG Volume");
+	NewActorClass = APCGVolume::StaticClass();
+	bUseSurfaceOrientation = true;
+}
+
+#undef LOCTEXT_NAMESPACE
\ No newline at end of file
diff --git a/Engine/Plugins/Experimental/PCG/Source/PCGEditor/Private/PCGVolumeFactory.h b/Engine/Plugins/Experimental/PCG/Source/PCGEditor/Private/PCGVolumeFactory.h
new file mode 100644
--- /dev/null
+++ b/Engine/Plugins/Experimental/PCG/Source/PCGEditor/Private/PCGVolumeFactory.h
@@ -0,0 +1,18 @@
+// Copyright Epic Games, Inc. All Rights Reserved.
+
+#pragma once
+
+#include "CoreMinimal.h"
+#include "UObject/ObjectMacros.h"
+#include "ActorFactories/ActorFactoryBoxVolume.h"
+
+#include "PCGVolumeFactory.generated.h"
+
+UCLASS(MinimalAPI, config=Editor)
+class UPCGVolumeFactory : public UActorFactoryBoxVolume
+{
+	GENERATED_BODY()
+
+public:
+	UPCGVolumeFactory(const FObjectInitializer& ObjectInitializer);
+};
\ No newline at end of file
diff --git a/Engine/Plugins/Experimental/PCG/Source/PCGEditor/Public/PCGComponentDetails.h b/Engine/Plugins/Experimental/PCG/Source/PCGEditor/Public/PCGComponentDetails.h
new file mode 100644
--- /dev/null
+++ b/Engine/Plugins/Experimental/PCG/Source/PCGEditor/Public/PCGComponentDetails.h
@@ -0,0 +1,27 @@
+// Copyright Epic Games, Inc. All Rights Reserved.
+
+#pragma once
+
+#include "UObject/WeakObjectPtr.h"
+#include "Input/Reply.h"
+#include "IDetailCustomization.h"
+
+class UPCGComponent;
+
+class FPCGComponentDetails : public IDetailCustomization
+{
+public:
+	static TSharedRef<IDetailCustomization> MakeInstance();
+
+	/** ~Begin IDetailCustomization interface */
+	virtual void CustomizeDetails(IDetailLayoutBuilder& DetailBuilder) override;
+	virtual void PendingDelete() override;
+	/** ~End IDetailCustomization interface */
+
+private:
+	FReply OnGenerateClicked();
+	FReply OnCleanupClicked();
+	void OnGraphChanged(UPCGComponent* InComponent);
+
+	TArray<TWeakObjectPtr<UPCGComponent>> SelectedComponents;
+};
\ No newline at end of file
diff --git a/Engine/Plugins/Experimental/PCG/Source/PCGEditor/Public/PCGGraphDetails.h b/Engine/Plugins/Experimental/PCG/Source/PCGEditor/Public/PCGGraphDetails.h
new file mode 100644
--- /dev/null
+++ b/Engine/Plugins/Experimental/PCG/Source/PCGEditor/Public/PCGGraphDetails.h
@@ -0,0 +1,22 @@
+// Copyright Epic Games, Inc. All Rights Reserved.
+
+#pragma once
+
+#include "UObject/WeakObjectPtr.h"
+#include "Input/Reply.h"
+#include "IDetailCustomization.h"
+
+class UPCGGraph;
+
+class FPCGGraphDetails : public IDetailCustomization
+{
+public:
+	static TSharedRef<IDetailCustomization> MakeInstance();
+
+	virtual void CustomizeDetails(IDetailLayoutBuilder& DetailBuilder) override;
+
+private:
+	FReply OnInitializeClicked();
+
+	TArray<TWeakObjectPtr<UPCGGraph>> SelectedGraphs;
+};
\ No newline at end of file
