Fixed noPCG compilation error due to missing Abstract tag

diff --git a/Engine/Plugins/Experimental/PCG/Source/PCG/Public/Data/PCGSurfaceData.h b/Engine/Plugins/Experimental/PCG/Source/PCG/Public/Data/PCGSurfaceData.h
--- a/Engine/Plugins/Experimental/PCG/Source/PCG/Public/Data/PCGSurfaceData.h
+++ b/Engine/Plugins/Experimental/PCG/Source/PCG/Public/Data/PCGSurfaceData.h
@@ -7,7 +7,7 @@
 
 #include "PCGSurfaceData.generated.h"
 
-UCLASS(BlueprintType, ClassGroup = (Procedural))
+UCLASS(Abstract, BlueprintType, ClassGroup = (Procedural))
 class UPCGSurfaceData : public UPCGSpatialDataWithPointCache
 {
 	GENERATED_BODY()
