Fixed static analysis warnings + compilation error in non PCH builds.

diff --git a/Engine/Plugins/Experimental/PCG/Source/PCG/Private/Elements/PCGExecuteBlueprint.cpp b/Engine/Plugins/Experimental/PCG/Source/PCG/Private/Elements/PCGExecuteBlueprint.cpp
--- a/Engine/Plugins/Experimental/PCG/Source/PCG/Private/Elements/PCGExecuteBlueprint.cpp
+++ b/Engine/Plugins/Experimental/PCG/Source/PCG/Private/Elements/PCGExecuteBlueprint.cpp
@@ -70,8 +70,7 @@ FPCGElementPtr UPCGBlueprintSettings::CreateElement() const
 bool FPCGExecuteBlueprintElement::Execute(FPCGContextPtr Context) const
 {
 	const UPCGBlueprintSettings* Settings = Context->GetInputSettings<UPCGBlueprintSettings>();
-	check(Settings);
-	check(Settings->BlueprintElementInstance);
+	check(Settings && Settings->BlueprintElementInstance);
 
 	if (Settings && Settings->BlueprintElementInstance)
 	{
diff --git a/Engine/Plugins/Experimental/PCG/Source/PCG/Public/Data/PCGPolyLineData.h b/Engine/Plugins/Experimental/PCG/Source/PCG/Public/Data/PCGPolyLineData.h
--- a/Engine/Plugins/Experimental/PCG/Source/PCG/Public/Data/PCGPolyLineData.h
+++ b/Engine/Plugins/Experimental/PCG/Source/PCG/Public/Data/PCGPolyLineData.h
@@ -7,7 +7,7 @@
 
 #include "PCGPolyLineData.generated.h"
 
-UCLASS(BlueprintType, ClassGroup = (Procedural))
+UCLASS(Abstract, BlueprintType, ClassGroup = (Procedural))
 class UPCGPolyLineData : public UPCGSpatialDataWithPointCache
 {
 	GENERATED_BODY()
