Add density scheme to union & difference
Remove dimension requirements in union

diff --git a/Engine/Plugins/Experimental/PCG/Source/PCG/Private/Data/PCGDifferenceData.cpp b/Engine/Plugins/Experimental/PCG/Source/PCG/Private/Data/PCGDifferenceData.cpp
--- a/Engine/Plugins/Experimental/PCG/Source/PCG/Private/Data/PCGDifferenceData.cpp
+++ b/Engine/Plugins/Experimental/PCG/Source/PCG/Private/Data/PCGDifferenceData.cpp
@@ -4,6 +4,14 @@
 #include "Data/PCGPointData.h"
 #include "Data/PCGUnionData.h"
 
+namespace PCGDifferenceDataUtils
+{
+	EPCGUnionDensityFunction ToUnionDensityFunction(EPCGDifferenceDensityFunction InDensityFunction)
+	{
+		return (InDensityFunction == EPCGDifferenceDensityFunction::ClampedSubstraction ? EPCGUnionDensityFunction::ClampedAddition : EPCGUnionDensityFunction::Maximum);
+	}
+}
+
 void UPCGDifferenceData::Initialize(const UPCGSpatialData* InData)
 {
 	check(InData);
@@ -25,6 +33,7 @@ void UPCGDifferenceData::AddDifference(const UPCGSpatialData* InDifference)
 		{
 			DifferencesUnion = NewObject<UPCGUnionData>(this);
 			DifferencesUnion->AddData(Difference);
+			DifferencesUnion->SetDensityFunction(PCGDifferenceDataUtils::ToUnionDensityFunction(DensityFunction));
 			Difference = DifferencesUnion;
 		}
 
@@ -33,6 +42,28 @@ void UPCGDifferenceData::AddDifference(const UPCGSpatialData* InDifference)
 	}
 }
 
+void UPCGDifferenceData::SetDensityFunction(EPCGDifferenceDensityFunction InDensityFunction)
+{
+	DensityFunction = InDensityFunction;
+
+	if (DifferencesUnion)
+	{
+		DifferencesUnion->SetDensityFunction(PCGDifferenceDataUtils::ToUnionDensityFunction(DensityFunction));
+	}
+}
+
+#if WITH_EDITOR
+void UPCGDifferenceData::PostEditChangeProperty(FPropertyChangedEvent& PropertyChangedEvent)
+{
+	if (PropertyChangedEvent.Property && PropertyChangedEvent.Property->GetFName() == GET_MEMBER_NAME_CHECKED(UPCGDifferenceData, DensityFunction))
+	{
+		SetDensityFunction(DensityFunction);
+	}
+
+	Super::PostEditChangeProperty(PropertyChangedEvent);
+}
+#endif
+
 int UPCGDifferenceData::GetDimension() const
 {
 	return Source->GetDimension();
@@ -58,7 +89,7 @@ float UPCGDifferenceData::GetDensityAtPosition(const FVector& InPosition) const
 	else if(Difference)
 	{
 		const float DensityInSource = Source->GetDensityAtPosition(InPosition);
-		if (DensityInSource <= 0)
+		if (DensityInSource == 0)
 		{
 			return 0;
 		}
@@ -91,9 +122,11 @@ const UPCGPointData* UPCGDifferenceData::CreatePointData() const
 
 	for (const FPCGPoint& Point : SourcePointData->GetPoints())
 	{
-		if (Difference->GetDensityAtPosition(Point.Transform.GetLocation()) == 0)
+		const float DensityInDifference = Difference->GetDensityAtPosition(Point.Transform.GetLocation());
+		if (DensityInDifference < Point.Density)
 		{
-			TargetPoints.Add(Point);
+			FPCGPoint& TargetPoint = TargetPoints.Add_GetRef(Point);
+			TargetPoint.Density -= DensityInDifference;
 		}
 	}
 
diff --git a/Engine/Plugins/Experimental/PCG/Source/PCG/Private/Data/PCGIntersectionData.cpp b/Engine/Plugins/Experimental/PCG/Source/PCG/Private/Data/PCGIntersectionData.cpp
--- a/Engine/Plugins/Experimental/PCG/Source/PCG/Private/Data/PCGIntersectionData.cpp
+++ b/Engine/Plugins/Experimental/PCG/Source/PCG/Private/Data/PCGIntersectionData.cpp
@@ -3,6 +3,21 @@
 #include "Data/PCGIntersectionData.h"
 #include "Data/PCGPointData.h"
 
+namespace PCGIntersectionDataMaths
+{
+	float ComputeDensity(float InDensityA, float InDensityB, EPCGIntersectionDensityFunction InDensityFunction)
+	{
+		if (InDensityFunction == EPCGIntersectionDensityFunction::Minimum)
+		{
+			return FMath::Min(InDensityA, InDensityB);
+		}
+		else // default: Multiply
+		{
+			return InDensityA * InDensityB;
+		}
+	}
+}
+
 void UPCGIntersectionData::Initialize(const UPCGSpatialData* InA, const UPCGSpatialData* InB)
 {
 	check(InA && InB);
@@ -45,7 +60,13 @@ float UPCGIntersectionData::GetDensityAtPosition(const FVector& InPosition) cons
 	}
 	else
 	{
-		return A->GetDensityAtPosition(InPosition) * B->GetDensityAtPosition(InPosition);
+		float Density = A->GetDensityAtPosition(InPosition);
+		if (Density > 0)
+		{
+			Density = PCGIntersectionDataMaths::ComputeDensity(Density, B->GetDensityAtPosition(InPosition), DensityFunction);
+		}
+
+		return Density;
 	}
 }
 
@@ -92,7 +113,7 @@ UPCGPointData* UPCGIntersectionData::CreateAndFilterPointData(const UPCGSpatialD
 		if (YDensity > 0)
 		{
 			FPCGPoint& NewPoint = TargetPoints.Add_GetRef(Point);
-			NewPoint.Density *= YDensity;
+			NewPoint.Density = PCGIntersectionDataMaths::ComputeDensity(Point.Density, YDensity, DensityFunction);
 		}
 	}
 
diff --git a/Engine/Plugins/Experimental/PCG/Source/PCG/Private/Data/PCGUnionData.cpp b/Engine/Plugins/Experimental/PCG/Source/PCG/Private/Data/PCGUnionData.cpp
--- a/Engine/Plugins/Experimental/PCG/Source/PCG/Private/Data/PCGUnionData.cpp
+++ b/Engine/Plugins/Experimental/PCG/Source/PCG/Private/Data/PCGUnionData.cpp
@@ -3,183 +3,193 @@
 #include "Data/PCGUnionData.h"
 #include "Data/PCGPointData.h"
 
+namespace PCGUnionDataMaths
+{
+	float ComputeDensity(float InDensityToUpdate, float InOtherDensity, EPCGUnionDensityFunction DensityFunction)
+	{
+		if (DensityFunction == EPCGUnionDensityFunction::ClampedAddition)
+		{
+			return FMath::Min(InDensityToUpdate + InOtherDensity, 1.0f);
+		}
+		else
+		{
+			return FMath::Max(InDensityToUpdate, InOtherDensity);
+		}
+	}
+
+	float UpdateDensity(float& InDensityToUpdate, float InOtherDensity, EPCGUnionDensityFunction DensityFunction)
+	{
+		InDensityToUpdate = ComputeDensity(InDensityToUpdate, InOtherDensity, DensityFunction);
+		return InDensityToUpdate;
+	}
+}
+
 void UPCGUnionData::Initialize(const UPCGSpatialData* InA, const UPCGSpatialData* InB)
 {
 	check(InA && InB);
-	check(InA->GetDimension() == InB->GetDimension());
 	AddData(InA);
 	AddData(InB);
 }
 
-void UPCGUnionData::AddData(const UPCGSpatialData* InC)
+void UPCGUnionData::AddData(const UPCGSpatialData* InData)
 {
-	check(InC);
-	check(!A || InC->GetDimension() == A->GetDimension());
+	check(InData);
 
-	// Early tests: if we don't have a full union yet, just push in the data.
-	if (!A)
-	{
-		check(!B);
-		A = InC;
-		TargetActor = InC->TargetActor;
-		CachedBounds = InC->GetBounds();
-		CachedStrictBounds = InC->GetStrictBounds();
-	}
-	else if (!B)
+	Data.Add(InData);
+
+	if (Data.Num() == 1)
 	{
-		check(A);
-		B = InC;
-		CachedBounds += InC->GetBounds();
-		CachedStrictBounds = A->GetStrictBounds().Overlap(InC->GetStrictBounds());
+		TargetActor = InData->TargetActor;
+		CachedBounds = InData->GetBounds();
+		CachedStrictBounds = InData->GetStrictBounds();
+		CachedDimension = InData->GetDimension();
 	}
 	else
 	{
-		// TODO: Bad usage of this call could lead to sub-unions having inefficient bounds
-		// We can revisit and keep the list of members and bundle them appropriately on demand 
-		const bool bAIsUnion = Cast<UPCGUnionData>(A) != nullptr;
-		const bool bBIsUnion = Cast<UPCGUnionData>(B) != nullptr;
-
-		UPCGUnionData* Union = NewObject<UPCGUnionData>(this);
-		Union->UnionType = UnionType;
-
-		if (!bAIsUnion || bBIsUnion)
-		{
-			Union->Initialize(A, B);
-			A = Union;
-			B = InC;
-		}
-		else
-		{
-			Union->Initialize(B, InC);
-			B = Union;
-		}
-
-		// Update bounds
-		CachedBounds += InC->GetBounds();
-		CachedStrictBounds = CachedStrictBounds.Overlap(InC->GetStrictBounds());
+		CachedBounds += InData->GetBounds();
+		CachedStrictBounds = CachedStrictBounds.Overlap(InData->GetStrictBounds());
+		CachedDimension = FMath::Max(CachedDimension, InData->GetDimension());
 	}
 }
 
 int UPCGUnionData::GetDimension() const
 {
-	check(A && B && A->GetDimension() == B->GetDimension());
-	return A->GetDimension();
+	return CachedDimension;
 }
 
 FBox UPCGUnionData::GetBounds() const
 {
-	check(A && B);
 	return CachedBounds;
 }
 
 FBox UPCGUnionData::GetStrictBounds() const
 {
-	check(A && B);
 	return CachedStrictBounds;
 }
 
 float UPCGUnionData::GetDensityAtPosition(const FVector& InPosition) const
 {
-	check(A && B);
+	// Early exits
 	if (!CachedBounds.IsInside(InPosition))
 	{
 		return 0;
 	}
-	else if (CachedStrictBounds.IsInside(InPosition) ||
-		A->GetStrictBounds().IsInside(InPosition) ||
-		B->GetStrictBounds().IsInside(InPosition))
+	else if (CachedStrictBounds.IsInside(InPosition))
 	{
 		return 1.0f;
 	}
-	else
+
+	// Check for presence in any strict bounds of the data.
+	// Note that it can be superfluous in some instances as we will might end up testing
+	// the strict bounds twice per data, but it will perform better in the worst case.
+	for (int32 DataIndex = 0; DataIndex < Data.Num(); ++DataIndex)
+	{
+		if (Data[DataIndex]->GetStrictBounds().IsInside(InPosition))
+		{
+			return 1.0f;
+		}
+	}
+
+	float Density = 0.0f;
+
+	for (TObjectPtr<const UPCGSpatialData> Datum : Data)
 	{
-		return FMath::Max(A->GetDensityAtPosition(InPosition), B->GetDensityAtPosition(InPosition));
+		if (PCGUnionDataMaths::UpdateDensity(Density, Datum->GetDensityAtPosition(InPosition), DensityFunction) == 1.0f)
+		{
+			break;
+		}
 	}
+
+	return Density;
 }
 
 const UPCGPointData* UPCGUnionData::CreatePointData() const
 {
 	TRACE_CPUPROFILER_EVENT_SCOPE(UPCGUnionData::CreatePointData);
-	check(A && B);
 
-	const UPCGPointData* APointData = A->ToPointData();
-	const UPCGPointData* BPointData = B->ToPointData();
-
-	if (!APointData || !BPointData)
+	// Trivial results
+	if (Data.Num() == 0)
 	{
-		return APointData ? APointData : BPointData;
+		return nullptr;
+	}
+	else if (Data.Num() == 1)
+	{
+		return Data[0]->ToPointData();
 	}
 
-	UPCGPointData* Data = NewObject<UPCGPointData>(const_cast<UPCGUnionData*>(this));
-	Data->TargetActor = TargetActor;
-
-	TArray<FPCGPoint>& TargetPoints = Data->GetMutablePoints();
-
-	auto AddNonExcludedPoints = [&TargetPoints](const UPCGPointData* PointsToAdd, const UPCGSpatialData* Exclusion) {
-		check(PointsToAdd && Exclusion);
-		for (const FPCGPoint& Point : PointsToAdd->GetPoints())
-		{
-			if (Exclusion->GetDensityAtPosition(Point.Transform.GetLocation()) == 0)
-			{
-				TargetPoints.Add(Point);
-			}
-		}
-	};
+	UPCGPointData* PointData = NewObject<UPCGPointData>(const_cast<UPCGUnionData*>(this));
+	PointData->TargetActor = TargetActor;
 
 	switch (UnionType)
 	{
 	case EPCGUnionType::LeftToRightPriority:
-	case EPCGUnionType::RightToLeftPriority:
 	default:
-		{
-			const UPCGSpatialData* PrioritizedData = (UnionType != EPCGUnionType::RightToLeftPriority ? A : B);
-			const UPCGPointData* PrioritizedPointData = (UnionType != EPCGUnionType::RightToLeftPriority ? APointData : BPointData);
-			const UPCGSpatialData* SecondaryData = (UnionType != EPCGUnionType::RightToLeftPriority ? B : A);
-			const UPCGPointData* SecondaryPointData = (UnionType != EPCGUnionType::RightToLeftPriority ? BPointData : APointData);
+		CreateSequentialPointData(PointData, /*bLeftToRight=*/true);
+		break;
 
-			// First: add all points from the prioritized data and update its density
-			TargetPoints.Append(PrioritizedPointData->GetPoints());
+	case EPCGUnionType::RightToLeftPriority:
+		CreateSequentialPointData(PointData, /*bLeftToRight=*/false);
+		break;
 
-			for (FPCGPoint& Point : TargetPoints)
+	case EPCGUnionType::KeepAll:
+		{
+			TArray<FPCGPoint>& TargetPoints = PointData->GetMutablePoints();
+			for (TObjectPtr<const UPCGSpatialData> Datum : Data)
 			{
-				Point.Density = FMath::Max(Point.Density, SecondaryData->GetDensityAtPosition(Point.Transform.GetLocation()));
+				TargetPoints.Append(Datum->ToPointData()->GetPoints());
 			}
-
-			// Second: all points from the secondary data set only if they have a null density in the prioritized data set
-			AddNonExcludedPoints(SecondaryPointData, PrioritizedData);
 		}
-
 		break;
+	}
 
-	case EPCGUnionType::KeepAB:
-		{
-			TargetPoints.Append(APointData->GetPoints());
-			TargetPoints.Append(BPointData->GetPoints());
-		}
+	return PointData;
+}
 
-		break;
+void UPCGUnionData::CreateSequentialPointData(UPCGPointData* PointData, bool bLeftToRight) const
+{
+	check(PointData);
+
+	TArray<FPCGPoint>& TargetPoints = PointData->GetMutablePoints();
 
-	case EPCGUnionType::ABMinusIntersection:
+	int32 FirstDataIndex = (bLeftToRight ? 0 : Data.Num() - 1);
+	int32 LastDataIndex = (bLeftToRight ? Data.Num() : -1);
+	int32 DataIndexIncrement = (bLeftToRight ? 1 : -1);
+
+	// Note: this is a O(N^2) implementation. 
+	// TODO: It is easy to implement a kind of divide & conquer algorithm here, but it will require some temporary storage.
+	for (int32 DataIndex = FirstDataIndex; DataIndex != LastDataIndex; DataIndex += DataIndexIncrement)
+	{
+		// For each point, if it is not already "processed" by previous data,
+		// add it & compute its final density
+		const TArray<FPCGPoint>& Points = Data[DataIndex]->ToPointData()->GetPoints();
+		for (const FPCGPoint& Point : Points)
 		{
-			// First, build the intersection of A & B
-			if (CachedStrictBounds.IsValid)
+			// Discard point if it is already covered by a previous data
+			bool bPointToExclude = false;
+			for (int32 PreviousDataIndex = FirstDataIndex; PreviousDataIndex != DataIndex; PreviousDataIndex += DataIndexIncrement)
 			{
-				UPCGSpatialData* IntersectionData = A->IntersectWith(B);
-				const UPCGPointData* IntersectionPointData = IntersectionData ? IntersectionData->ToPointData() : nullptr;
-
-				if (IntersectionPointData)
+				if (Data[PreviousDataIndex]->GetDensityAtPosition(Point.Transform.GetLocation()) != 0)
 				{
-					TargetPoints.Append(IntersectionPointData->GetPoints());
+					bPointToExclude = true;
+					break;
 				}
 			}
 
-			// Second: all all points from A & B that have no match in their counterpart.
-			AddNonExcludedPoints(APointData, B);
-			AddNonExcludedPoints(BPointData, A);
-		}
+			if (bPointToExclude)
+			{
+				continue;
+			}
 
-		break;
-	}
+			FPCGPoint& TargetPoint = TargetPoints.Add_GetRef(Point);
 
-	return Data;
-}
+			// Compute final density based on current & following data
+			for (int32 FollowingDataIndex = DataIndex + DataIndexIncrement; FollowingDataIndex != LastDataIndex; FollowingDataIndex += DataIndexIncrement)
+			{
+				if (PCGUnionDataMaths::UpdateDensity(TargetPoint.Density, Data[FollowingDataIndex]->GetDensityAtPosition(TargetPoint.Transform.GetLocation()), DensityFunction) == 1.0f)
+				{
+					break;
+				}
+			}
+		}
+	}
+}
\ No newline at end of file
diff --git a/Engine/Plugins/Experimental/PCG/Source/PCG/Private/Elements/PCGExecuteBlueprint.cpp b/Engine/Plugins/Experimental/PCG/Source/PCG/Private/Elements/PCGExecuteBlueprint.cpp
--- a/Engine/Plugins/Experimental/PCG/Source/PCG/Private/Elements/PCGExecuteBlueprint.cpp
+++ b/Engine/Plugins/Experimental/PCG/Source/PCG/Private/Elements/PCGExecuteBlueprint.cpp
@@ -38,6 +38,8 @@ void UPCGBlueprintSettings::PostEditChangeProperty(FPropertyChangedEvent& Proper
 	{
 		RefreshBlueprintElement();
 	}
+
+	Super::PostEditChangeProperty(PropertyChangedEvent);
 }
 #endif
 
diff --git a/Engine/Plugins/Experimental/PCG/Source/PCG/Public/Data/PCGDifferenceData.h b/Engine/Plugins/Experimental/PCG/Source/PCG/Public/Data/PCGDifferenceData.h
--- a/Engine/Plugins/Experimental/PCG/Source/PCG/Public/Data/PCGDifferenceData.h
+++ b/Engine/Plugins/Experimental/PCG/Source/PCG/Public/Data/PCGDifferenceData.h
@@ -9,6 +9,13 @@
 
 class UPCGUnionData;
 
+UENUM()
+enum class EPCGDifferenceDensityFunction : uint8
+{
+	Minimum,
+	ClampedSubstraction
+};
+
 UCLASS(BlueprintType, ClassGroup = (Procedural))
 class UPCGDifferenceData : public UPCGSpatialDataWithPointCache
 {
@@ -20,6 +27,13 @@ public:
 	UFUNCTION(BlueprintCallable, Category = SpatialData)
 	void AddDifference(const UPCGSpatialData* InDifference);
 
+	UFUNCTION(BlueprintCallable, Category = Settings)
+	void SetDensityFunction(EPCGDifferenceDensityFunction InDensityFunction);
+
+#if WITH_EDITOR
+	virtual void PostEditChangeProperty(FPropertyChangedEvent& PropertyChangedEvent) override;
+#endif
+
 	//~Begin UPCGSpatialData interface
 	virtual int GetDimension() const override;
 	virtual FBox GetBounds() const override;
@@ -39,4 +53,7 @@ protected:
 
 	UPROPERTY()
 	TObjectPtr<UPCGUnionData> DifferencesUnion = nullptr;
+
+	UPROPERTY(BlueprintSetter = SetDensityFunction, EditAnywhere, Category = Settings)
+	EPCGDifferenceDensityFunction DensityFunction = EPCGDifferenceDensityFunction::Minimum;
 };
\ No newline at end of file
diff --git a/Engine/Plugins/Experimental/PCG/Source/PCG/Public/Data/PCGIntersectionData.h b/Engine/Plugins/Experimental/PCG/Source/PCG/Public/Data/PCGIntersectionData.h
--- a/Engine/Plugins/Experimental/PCG/Source/PCG/Public/Data/PCGIntersectionData.h
+++ b/Engine/Plugins/Experimental/PCG/Source/PCG/Public/Data/PCGIntersectionData.h
@@ -6,6 +6,13 @@
 
 #include "PCGIntersectionData.generated.h"
 
+UENUM()
+enum class EPCGIntersectionDensityFunction : uint8
+{
+	Multiply,
+	Minimum
+};
+
 /**
 * Generic intersection class that delays operations as long as possible.
 */
@@ -37,6 +44,12 @@ protected:
 	UPROPERTY(BlueprintReadOnly, VisibleAnywhere, Category = SpatialData)
 	TObjectPtr<const UPCGSpatialData> B = nullptr;
 
+	UPROPERTY(BlueprintReadWrite, EditAnywhere, Category = Settings)
+	EPCGIntersectionDensityFunction DensityFunction = EPCGIntersectionDensityFunction::Multiply;
+
+	UPROPERTY()
 	FBox CachedBounds = FBox(EForceInit::ForceInit);
+
+	UPROPERTY()
 	FBox CachedStrictBounds = FBox(EForceInit::ForceInit);
 };
\ No newline at end of file
diff --git a/Engine/Plugins/Experimental/PCG/Source/PCG/Public/Data/PCGUnionData.h b/Engine/Plugins/Experimental/PCG/Source/PCG/Public/Data/PCGUnionData.h
--- a/Engine/Plugins/Experimental/PCG/Source/PCG/Public/Data/PCGUnionData.h
+++ b/Engine/Plugins/Experimental/PCG/Source/PCG/Public/Data/PCGUnionData.h
@@ -12,8 +12,14 @@ enum class EPCGUnionType : uint8
 {
 	LeftToRightPriority,
 	RightToLeftPriority,
-	KeepAB,
-	ABMinusIntersection
+	KeepAll,
+};
+
+UENUM()
+enum class EPCGUnionDensityFunction : uint8
+{
+	Maximum,
+	ClampedAddition
 };
 
 UCLASS(BlueprintType, ClassGroup=(Procedural))
@@ -25,7 +31,10 @@ public:
 	void Initialize(const UPCGSpatialData* InA, const UPCGSpatialData* InB);
 
 	UFUNCTION(BlueprintCallable, Category = SpatialData)
-	void AddData(const UPCGSpatialData* InC);
+	void AddData(const UPCGSpatialData* InData);
+
+	void SetType(EPCGUnionType InUnionType) { UnionType = InUnionType; }
+	void SetDensityFunction(EPCGUnionDensityFunction InDensityFunction) { DensityFunction = InDensityFunction; }
 
 	//~Begin UPCGSpatialData interface
 	virtual int GetDimension() const override;
@@ -40,17 +49,23 @@ public:
 
 protected:
 	UPROPERTY(BlueprintReadOnly, VisibleAnywhere, Category = SpatialData)
-	TObjectPtr<const UPCGSpatialData> A = nullptr;
-
-	UPROPERTY(BlueprintReadOnly, VisibleAnywhere, Category = SpatialData)
-	TObjectPtr<const UPCGSpatialData> B = nullptr;
+	TArray<TObjectPtr<const UPCGSpatialData>> Data;
 
 	UPROPERTY(BlueprintReadWrite, EditAnywhere, Category = Settings)
 	EPCGUnionType UnionType = EPCGUnionType::LeftToRightPriority;
 
+	UPROPERTY(BlueprintReadWrite, EditAnywhere, Category = Settings)
+	EPCGUnionDensityFunction DensityFunction = EPCGUnionDensityFunction::Maximum;
+
 	UPROPERTY()
 	FBox CachedBounds = FBox(EForceInit::ForceInit);
 
 	UPROPERTY()
 	FBox CachedStrictBounds = FBox(EForceInit::ForceInit);
+
+	UPROPERTY()
+	int CachedDimension = 0;
+
+private:
+	void CreateSequentialPointData(UPCGPointData* PointData, bool bLeftToRight) const;
 };
\ No newline at end of file
