Exposed union/diff/intersection/projection operations in blueprint instead of relying on manual construction

diff --git a/Engine/Plugins/Experimental/PCG/Source/PCG/Private/Data/PCGSpatialData.cpp b/Engine/Plugins/Experimental/PCG/Source/PCG/Private/Data/PCGSpatialData.cpp
--- a/Engine/Plugins/Experimental/PCG/Source/PCG/Private/Data/PCGSpatialData.cpp
+++ b/Engine/Plugins/Experimental/PCG/Source/PCG/Private/Data/PCGSpatialData.cpp
@@ -1,8 +1,10 @@
 // Copyright Epic Games, Inc. All Rights Reserved.
 
 #include "Data/PCGSpatialData.h"
+#include "Data/PCGDifferenceData.h"
 #include "Data/PCGIntersectionData.h"
 #include "Data/PCGProjectionData.h"
+#include "Data/PCGUnionData.h"
 
 const UPCGPointData* UPCGSpatialDataWithPointCache::ToPointData() const
 {
@@ -21,7 +23,7 @@ const UPCGPointData* UPCGSpatialDataWithPointCache::ToPointData() const
 	return CachedPointData;
 }
 
-UPCGSpatialData* UPCGSpatialData::IntersectWith(const UPCGSpatialData* InOther) const
+UPCGIntersectionData* UPCGSpatialData::IntersectWith(const UPCGSpatialData* InOther) const
 {
 	UPCGIntersectionData* IntersectionData = NewObject<UPCGIntersectionData>(const_cast<UPCGSpatialData*>(this));
 	IntersectionData->Initialize(this, InOther);
@@ -29,10 +31,27 @@ UPCGSpatialData* UPCGSpatialData::IntersectWith(const UPCGSpatialData* InOther)
 	return IntersectionData;
 }
 
-UPCGSpatialData* UPCGSpatialData::ProjectOn(const UPCGSpatialData* InOther) const
+UPCGProjectionData* UPCGSpatialData::ProjectOn(const UPCGSpatialData* InOther) const
 {
 	UPCGProjectionData* ProjectionData = NewObject<UPCGProjectionData>(const_cast<UPCGSpatialData*>(this));
 	ProjectionData->Initialize(this, InOther);
 
 	return ProjectionData;
+}
+
+UPCGUnionData* UPCGSpatialData::UnionWith(const UPCGSpatialData* InOther) const
+{
+	UPCGUnionData* UnionData = NewObject<UPCGUnionData>(const_cast<UPCGSpatialData*>(this));
+	UnionData->Initialize(this, InOther);
+
+	return UnionData;
+}
+
+UPCGDifferenceData* UPCGSpatialData::Subtract(const UPCGSpatialData* InOther) const
+{
+	UPCGDifferenceData* DifferenceData = NewObject<UPCGDifferenceData>(const_cast<UPCGSpatialData*>(this));
+	DifferenceData->Initialize(this);
+	DifferenceData->AddDifference(InOther);
+
+	return DifferenceData;
 }
\ No newline at end of file
diff --git a/Engine/Plugins/Experimental/PCG/Source/PCG/Private/Data/PCGSplineData.cpp b/Engine/Plugins/Experimental/PCG/Source/PCG/Private/Data/PCGSplineData.cpp
--- a/Engine/Plugins/Experimental/PCG/Source/PCG/Private/Data/PCGSplineData.cpp
+++ b/Engine/Plugins/Experimental/PCG/Source/PCG/Private/Data/PCGSplineData.cpp
@@ -81,7 +81,7 @@ float UPCGSplineData::GetDensityAtPosition(const FVector& InPosition) const
 	return FMath::Max(0, 1.0f - Distance);
 }
 
-UPCGSpatialData* UPCGSplineData::ProjectOn(const UPCGSpatialData* InOther) const
+UPCGProjectionData* UPCGSplineData::ProjectOn(const UPCGSpatialData* InOther) const
 {
 	if(InOther->GetDimension() == 2)
 	{
diff --git a/Engine/Plugins/Experimental/PCG/Source/PCG/Private/PCGComponent.cpp b/Engine/Plugins/Experimental/PCG/Source/PCG/Private/PCGComponent.cpp
--- a/Engine/Plugins/Experimental/PCG/Source/PCG/Private/PCGComponent.cpp
+++ b/Engine/Plugins/Experimental/PCG/Source/PCG/Private/PCGComponent.cpp
@@ -5,6 +5,7 @@
 #include "PCGHelpers.h"
 #include "PCGSubsystem.h"
 #include "PCGVolume.h"
+#include "Data/PCGIntersectionData.h"
 #include "Data/PCGLandscapeData.h"
 #include "Data/PCGLandscapeSplineData.h"
 #include "Data/PCGPointData.h"
diff --git a/Engine/Plugins/Experimental/PCG/Source/PCG/Public/Data/PCGSpatialData.h b/Engine/Plugins/Experimental/PCG/Source/PCG/Public/Data/PCGSpatialData.h
--- a/Engine/Plugins/Experimental/PCG/Source/PCG/Public/Data/PCGSpatialData.h
+++ b/Engine/Plugins/Experimental/PCG/Source/PCG/Public/Data/PCGSpatialData.h
@@ -8,6 +8,10 @@
 #include "PCGSpatialData.generated.h"
 
 class UPCGPointData;
+class UPCGIntersectionData;
+class UPCGUnionData;
+class UPCGDifferenceData;
+class UPCGProjectionData;
 class AActor;
 
 /**
@@ -58,10 +62,19 @@ public:
 	virtual FPCGPoint TransformPoint(const FPCGPoint& InPoint) const { return InPoint; }
 
 	/** Returns a specialized data to intersect with another data */
-	virtual UPCGSpatialData* IntersectWith(const UPCGSpatialData* InOther) const;
+	UFUNCTION(BlueprintCallable, Category = SpatialData)
+	virtual UPCGIntersectionData* IntersectWith(const UPCGSpatialData* InOther) const;
 
 	/** Returns a specialized data to project this on another data of equal or higher dimension */
-	virtual UPCGSpatialData* ProjectOn(const UPCGSpatialData* InOther) const;
+	UFUNCTION(BlueprintCallable, Category = SpatialData)
+	virtual UPCGProjectionData* ProjectOn(const UPCGSpatialData* InOther) const;
+
+	/** Returns a specialized data to union this with another data */
+	UFUNCTION(BlueprintCallable, Category = SpatialData)
+	virtual UPCGUnionData* UnionWith(const UPCGSpatialData* InOther) const;
+
+	UFUNCTION(BlueprintCallable, Category = SpatialData)
+	virtual UPCGDifferenceData* Subtract(const UPCGSpatialData* InOther) const;
 
 	UPROPERTY(BlueprintReadWrite, EditAnywhere, Category = Data)
 	AActor* TargetActor = nullptr;
diff --git a/Engine/Plugins/Experimental/PCG/Source/PCG/Public/Data/PCGSplineData.h b/Engine/Plugins/Experimental/PCG/Source/PCG/Public/Data/PCGSplineData.h
--- a/Engine/Plugins/Experimental/PCG/Source/PCG/Public/Data/PCGSplineData.h
+++ b/Engine/Plugins/Experimental/PCG/Source/PCG/Public/Data/PCGSplineData.h
@@ -31,7 +31,7 @@ public:
 	//~Begin UPCGSpatialData interface
 	virtual FBox GetBounds() const override;
 	virtual float GetDensityAtPosition(const FVector& InPosition) const override;
-	virtual UPCGSpatialData* ProjectOn(const UPCGSpatialData* InOther) const override;
+	virtual UPCGProjectionData* ProjectOn(const UPCGSpatialData* InOther) const override;
 	//~End 
 
 	UPROPERTY(BlueprintReadWrite, EditAnywhere, Category = SourceData)
