Added automatic node labelling for ease of use

diff --git a/Engine/Plugins/Experimental/PCG/Source/PCG/Private/PCGGraph.cpp b/Engine/Plugins/Experimental/PCG/Source/PCG/Private/PCGGraph.cpp
--- a/Engine/Plugins/Experimental/PCG/Source/PCG/Private/PCGGraph.cpp
+++ b/Engine/Plugins/Experimental/PCG/Source/PCG/Private/PCGGraph.cpp
@@ -109,6 +109,16 @@ UPCGNode* UPCGGraph::AddNodeOfType(TSubclassOf<class UPCGSettings> InSettingsCla
 
 		// Reparent node to this graph
 		Node->Rename(nullptr, this);
+
+#if WITH_EDITOR
+		const FName DefaultNodeName = Settings->GetDefaultNodeName();
+		if (DefaultNodeName != NAME_None)
+		{
+			FName NodeName = MakeUniqueObjectName(GetOuter(), UPCGNode::StaticClass(), DefaultNodeName);
+			Node->Rename(*NodeName.ToString());
+		}
+#endif
+
 		Nodes.Add(Node);
 
 #if WITH_EDITOR
@@ -138,6 +148,16 @@ UPCGNode* UPCGGraph::AddNode(UPCGSettings* InSettings)
 
 		// Reparent node to this graph
 		Node->Rename(nullptr, this);
+
+#if WITH_EDITOR
+		const FName DefaultNodeName = InSettings->GetDefaultNodeName();
+		if (DefaultNodeName != NAME_None)
+		{
+			FName NodeName = MakeUniqueObjectName(GetOuter(), UPCGNode::StaticClass(), DefaultNodeName);
+			Node->Rename(*NodeName.ToString());
+		}
+#endif
+
 		Nodes.Add(Node);
 
 #if WITH_EDITOR
diff --git a/Engine/Plugins/Experimental/PCG/Source/PCG/Public/Elements/PCGExecuteBlueprint.h b/Engine/Plugins/Experimental/PCG/Source/PCG/Public/Elements/PCGExecuteBlueprint.h
--- a/Engine/Plugins/Experimental/PCG/Source/PCG/Public/Elements/PCGExecuteBlueprint.h
+++ b/Engine/Plugins/Experimental/PCG/Source/PCG/Public/Elements/PCGExecuteBlueprint.h
@@ -32,9 +32,21 @@ class UPCGBlueprintSettings : public UPCGSettings
 public:
 	friend class FPCGExecuteBlueprintElement;
 
+	// ~Begin UPCGSettings interface
 #if WITH_EDITOR
+	virtual FName GetDefaultNodeName() const override { return FName(TEXT("BlueprintNode")); }
+#endif
+
+protected:
+	virtual FPCGElementPtr CreateElement() const override;
+	// ~End UPCGSettings interface
+
+public:
+#if WITH_EDITOR
+	// ~Begin UObject interface
 	virtual void PostLoad() override;
 	virtual void PostEditChangeProperty(FPropertyChangedEvent& PropertyChangedEvent) override;
+	// ~End UObject interface
 #endif
 
 	UFUNCTION(BlueprintCallable, Category = Settings)
@@ -52,7 +64,6 @@ protected:
 
 protected:
 	void RefreshBlueprintElement();
-	virtual FPCGElementPtr CreateElement() const override;
 };
 
 class FPCGExecuteBlueprintElement : public FSimpleTypedPCGElement<UPCGBlueprintSettings>
diff --git a/Engine/Plugins/Experimental/PCG/Source/PCG/Public/Elements/PCGPointSampler.h b/Engine/Plugins/Experimental/PCG/Source/PCG/Public/Elements/PCGPointSampler.h
--- a/Engine/Plugins/Experimental/PCG/Source/PCG/Public/Elements/PCGPointSampler.h
+++ b/Engine/Plugins/Experimental/PCG/Source/PCG/Public/Elements/PCGPointSampler.h
@@ -17,12 +17,18 @@ class UPCGPointSamplerSettings : public UPCGSettings
 	GENERATED_BODY()
 
 public:
-
-	UPROPERTY(BlueprintReadWrite, EditAnywhere, Category = Settings, meta=(ClampMin="0", ClampMax="1"))
-	float Ratio = 0.1f;
+	//~Begin UPCGSettings interface
+#if WITH_EDITOR
+	virtual FName GetDefaultNodeName() const override { return FName(TEXT("PointSamplerNode")); }
+#endif
 
 protected:
 	virtual FPCGElementPtr CreateElement() const override;
+	//~End UPCGSettings interface
+
+public:
+	UPROPERTY(BlueprintReadWrite, EditAnywhere, Category = Settings, meta=(ClampMin="0", ClampMax="1"))
+	float Ratio = 0.1f;
 };
 
 class FPCGPointSamplerElement : public FSimpleTypedPCGElement<UPCGPointSamplerSettings>
diff --git a/Engine/Plugins/Experimental/PCG/Source/PCG/Public/Elements/PCGStaticMeshSpawner.h b/Engine/Plugins/Experimental/PCG/Source/PCG/Public/Elements/PCGStaticMeshSpawner.h
--- a/Engine/Plugins/Experimental/PCG/Source/PCG/Public/Elements/PCGStaticMeshSpawner.h
+++ b/Engine/Plugins/Experimental/PCG/Source/PCG/Public/Elements/PCGStaticMeshSpawner.h
@@ -30,11 +30,18 @@ class UPCGStaticMeshSpawnerSettings : public UPCGSettings
 	GENERATED_BODY()
 
 public:
-	UPROPERTY(BlueprintReadWrite, EditAnywhere, Category = Settings)
-	TArray<FPCGStaticMeshSpawnerEntry> Meshes;
+#if WITH_EDITOR
+	// ~Begin UPCGSettings interface
+	virtual FName GetDefaultNodeName() const override { return FName(TEXT("StaticMeshSpawnerNode")); }
+#endif
 
 protected:
 	virtual FPCGElementPtr CreateElement() const override;
+	// ~End UPCGSettings interface
+
+public:
+	UPROPERTY(BlueprintReadWrite, EditAnywhere, Category = Settings)
+	TArray<FPCGStaticMeshSpawnerEntry> Meshes;
 };
 
 class FPCGStaticMeshSpawnerElement : public FSimpleTypedPCGElement<UPCGStaticMeshSpawnerSettings>
diff --git a/Engine/Plugins/Experimental/PCG/Source/PCG/Public/PCGSettings.h b/Engine/Plugins/Experimental/PCG/Source/PCG/Public/PCGSettings.h
--- a/Engine/Plugins/Experimental/PCG/Source/PCG/Public/PCGSettings.h
+++ b/Engine/Plugins/Experimental/PCG/Source/PCG/Public/PCGSettings.h
@@ -26,6 +26,10 @@ public:
 	/*virtual*/ FPCGElementPtr GetElement() const;
 	virtual UPCGNode* CreateNode() const;
 
+#if WITH_EDITOR
+	virtual FName GetDefaultNodeName() const { return NAME_None; }
+#endif
+
 	UPROPERTY(EditAnywhere, BlueprintReadWrite, Category = Settings)
 	int Seed = 0xC35A9631; // random prime number
 
diff --git a/Engine/Plugins/Experimental/PCG/Source/PCG/Public/PCGSubgraph.h b/Engine/Plugins/Experimental/PCG/Source/PCG/Public/PCGSubgraph.h
--- a/Engine/Plugins/Experimental/PCG/Source/PCG/Public/PCGSubgraph.h
+++ b/Engine/Plugins/Experimental/PCG/Source/PCG/Public/PCGSubgraph.h
@@ -17,6 +17,10 @@ public:
 	//~UPCGSettings interface implementation
 	virtual UPCGNode* CreateNode() const override;
 
+#if WITH_EDITOR
+	virtual FName GetDefaultNodeName() const override { return FName(TEXT("SubgraphNode")); }
+#endif
+
 protected:
 	virtual FPCGElementPtr CreateElement() const override;
 	//~End UPCGSettings interface implementation
