Fixed outstanding issues in PCG projection

diff --git a/Engine/Plugins/Experimental/PCG/Source/PCG/Private/Data/PCGDifferenceData.cpp b/Engine/Plugins/Experimental/PCG/Source/PCG/Private/Data/PCGDifferenceData.cpp
--- a/Engine/Plugins/Experimental/PCG/Source/PCG/Private/Data/PCGDifferenceData.cpp
+++ b/Engine/Plugins/Experimental/PCG/Source/PCG/Private/Data/PCGDifferenceData.cpp
@@ -103,6 +103,24 @@ float UPCGDifferenceData::GetDensityAtPosition(const FVector& InPosition) const
 	}
 }
 
+FVector UPCGDifferenceData::TransformPosition(const FVector& InPosition) const
+{
+	check(Source);
+	return Source->TransformPosition(InPosition);
+}
+
+FPCGPoint UPCGDifferenceData::TransformPoint(const FPCGPoint& InPoint) const
+{
+	check(Source);
+	return Source->TransformPoint(InPoint);
+}
+
+bool UPCGDifferenceData::HasNonTrivialTransform() const
+{
+	check(Source);
+	return Source->HasNonTrivialTransform();
+}
+
 const UPCGPointData* UPCGDifferenceData::CreatePointData() const
 {
 	TRACE_CPUPROFILER_EVENT_SCOPE(UPCGDifferenceData::CreatePointData);
diff --git a/Engine/Plugins/Experimental/PCG/Source/PCG/Private/Data/PCGIntersectionData.cpp b/Engine/Plugins/Experimental/PCG/Source/PCG/Private/Data/PCGIntersectionData.cpp
--- a/Engine/Plugins/Experimental/PCG/Source/PCG/Private/Data/PCGIntersectionData.cpp
+++ b/Engine/Plugins/Experimental/PCG/Source/PCG/Private/Data/PCGIntersectionData.cpp
@@ -70,6 +70,24 @@ float UPCGIntersectionData::GetDensityAtPosition(const FVector& InPosition) cons
 	}
 }
 
+FVector UPCGIntersectionData::TransformPosition(const FVector& InPosition) const
+{
+	check(A && B);
+	return A->HasNonTrivialTransform() ? A->TransformPosition(InPosition) : B->TransformPosition(InPosition);
+}
+
+FPCGPoint UPCGIntersectionData::TransformPoint(const FPCGPoint& InPoint) const
+{
+	check(A && B);
+	return A->HasNonTrivialTransform() ? A->TransformPoint(InPoint) : B->TransformPoint(InPoint);
+}
+
+bool UPCGIntersectionData::HasNonTrivialTransform() const
+{
+	check(A && B);
+	return A->HasNonTrivialTransform() || B->HasNonTrivialTransform();
+}
+
 const UPCGPointData* UPCGIntersectionData::CreatePointData() const
 {
 	check(A && B);
diff --git a/Engine/Plugins/Experimental/PCG/Source/PCG/Private/Data/PCGProjectionData.cpp b/Engine/Plugins/Experimental/PCG/Source/PCG/Private/Data/PCGProjectionData.cpp
--- a/Engine/Plugins/Experimental/PCG/Source/PCG/Private/Data/PCGProjectionData.cpp
+++ b/Engine/Plugins/Experimental/PCG/Source/PCG/Private/Data/PCGProjectionData.cpp
@@ -53,7 +53,7 @@ FBox UPCGProjectionData::ProjectBounds(const FBox& InBounds) const
 
 	for (int Corner = 0; Corner < 8; ++Corner)
 	{
-		Bounds += ProjectPosition(
+		Bounds += TransformPosition(
 			FVector(
 				(Corner / 4) ? InBounds.Max.X : InBounds.Min.X,
 				((Corner / 2) % 2) ? InBounds.Max.Y : InBounds.Min.Y,
@@ -69,18 +69,23 @@ float UPCGProjectionData::GetDensityAtPosition(const FVector& InPosition) const
 	return Source->GetDensityAtPosition(InPosition);
 }
 
-FVector UPCGProjectionData::ProjectPosition(const FVector& InPosition) const
+FVector UPCGProjectionData::TransformPosition(const FVector& InPosition) const
 {
 	// TODO: improve projection/unprojection mechanism
 	return Target->TransformPosition(InPosition);
 }
 
-FPCGPoint UPCGProjectionData::ProjectPoint(const FPCGPoint& InPoint) const
+FPCGPoint UPCGProjectionData::TransformPoint(const FPCGPoint& InPoint) const
 {
 	// TODO: improve projection/unprojection mechanism
 	return Target->TransformPoint(InPoint);
 }
 
+bool UPCGProjectionData::HasNonTrivialTransform() const
+{
+	return Target->HasNonTrivialTransform();
+}
+
 const UPCGPointData* UPCGProjectionData::CreatePointData() const
 {
 	TRACE_CPUPROFILER_EVENT_SCOPE(UPCGProjectionData::CreatePointData);
@@ -96,7 +101,7 @@ const UPCGPointData* UPCGProjectionData::CreatePointData() const
 	Points.Reserve(SourcePoints.Num());
 	for (const FPCGPoint& SourcePoint : SourcePoints)
 	{
-		FPCGPoint ProjectedPoint = ProjectPoint(SourcePoint);
+		FPCGPoint ProjectedPoint = TransformPoint(SourcePoint);
 		if (ProjectedPoint.Density > 0)
 		{
 			Points.Add(MoveTemp(ProjectedPoint));
diff --git a/Engine/Plugins/Experimental/PCG/Source/PCG/Private/Data/PCGUnionData.cpp b/Engine/Plugins/Experimental/PCG/Source/PCG/Private/Data/PCGUnionData.cpp
--- a/Engine/Plugins/Experimental/PCG/Source/PCG/Private/Data/PCGUnionData.cpp
+++ b/Engine/Plugins/Experimental/PCG/Source/PCG/Private/Data/PCGUnionData.cpp
@@ -103,6 +103,45 @@ float UPCGUnionData::GetDensityAtPosition(const FVector& InPosition) const
 	return Density;
 }
 
+FVector UPCGUnionData::TransformPosition(const FVector& InPosition) const
+{
+	for (TObjectPtr<const UPCGSpatialData> Datum : Data)
+	{
+		if (Datum->HasNonTrivialTransform())
+		{
+			return Datum->TransformPosition(InPosition);
+		}
+	}
+
+	return Super::TransformPosition(InPosition);
+}
+
+FPCGPoint UPCGUnionData::TransformPoint(const FPCGPoint& InPoint) const
+{
+	for (TObjectPtr<const UPCGSpatialData> Datum : Data)
+	{
+		if (Datum->HasNonTrivialTransform())
+		{
+			return Datum->TransformPoint(InPoint);
+		}
+	}
+
+	return Super::TransformPoint(InPoint);
+}
+
+bool UPCGUnionData::HasNonTrivialTransform() const
+{
+	for (TObjectPtr<const UPCGSpatialData> Datum : Data)
+	{
+		if (Datum->HasNonTrivialTransform())
+		{
+			return true;
+		}
+	}
+
+	return Super::HasNonTrivialTransform();
+}
+
 const UPCGPointData* UPCGUnionData::CreatePointData() const
 {
 	TRACE_CPUPROFILER_EVENT_SCOPE(UPCGUnionData::CreatePointData);
diff --git a/Engine/Plugins/Experimental/PCG/Source/PCG/Public/Data/PCGDifferenceData.h b/Engine/Plugins/Experimental/PCG/Source/PCG/Public/Data/PCGDifferenceData.h
--- a/Engine/Plugins/Experimental/PCG/Source/PCG/Public/Data/PCGDifferenceData.h
+++ b/Engine/Plugins/Experimental/PCG/Source/PCG/Public/Data/PCGDifferenceData.h
@@ -39,6 +39,9 @@ public:
 	virtual FBox GetBounds() const override;
 	virtual FBox GetStrictBounds() const override;
 	virtual float GetDensityAtPosition(const FVector& InPosition) const override;
+	virtual FVector TransformPosition(const FVector& InPosition) const override;
+	virtual FPCGPoint TransformPoint(const FPCGPoint& InPoint) const override;
+	virtual bool HasNonTrivialTransform() const override;
 	//~End UPCGSpatialData interface
 
 	//~Begin UPCGSpatialDataWithPointCache interface
diff --git a/Engine/Plugins/Experimental/PCG/Source/PCG/Public/Data/PCGIntersectionData.h b/Engine/Plugins/Experimental/PCG/Source/PCG/Public/Data/PCGIntersectionData.h
--- a/Engine/Plugins/Experimental/PCG/Source/PCG/Public/Data/PCGIntersectionData.h
+++ b/Engine/Plugins/Experimental/PCG/Source/PCG/Public/Data/PCGIntersectionData.h
@@ -29,6 +29,9 @@ public:
 	virtual FBox GetBounds() const override;
 	virtual FBox GetStrictBounds() const override;
 	virtual float GetDensityAtPosition(const FVector& InPosition) const override;
+	virtual FVector TransformPosition(const FVector& InPosition) const override;
+	virtual FPCGPoint TransformPoint(const FPCGPoint& InPoint) const override;
+	virtual bool HasNonTrivialTransform() const override;
 	//~End UPCGSpatialData interface
 
 	//~Begin UPCGSpatialDataWithPointCache interface
diff --git a/Engine/Plugins/Experimental/PCG/Source/PCG/Public/Data/PCGLandscapeData.h b/Engine/Plugins/Experimental/PCG/Source/PCG/Public/Data/PCGLandscapeData.h
--- a/Engine/Plugins/Experimental/PCG/Source/PCG/Public/Data/PCGLandscapeData.h
+++ b/Engine/Plugins/Experimental/PCG/Source/PCG/Public/Data/PCGLandscapeData.h
@@ -40,9 +40,10 @@ public:
 	// ~Begin UPGCSpatialData interface
 	virtual FBox GetBounds() const override;
 	virtual FBox GetStrictBounds() const override;
+	virtual float GetDensityAtPosition(const FVector& InPosition) const override;
 	virtual FVector TransformPosition(const FVector& InPosition) const override;
 	virtual FPCGPoint TransformPoint(const FPCGPoint& InPoint) const override;
-	virtual float GetDensityAtPosition(const FVector& InPosition) const override;
+	virtual bool HasNonTrivialTransform() const override { return true; }
 	// ~End UPGCConcreteData interface
 
 	// ~Begin UPCGSpatialDataWithPointCache interface
diff --git a/Engine/Plugins/Experimental/PCG/Source/PCG/Public/Data/PCGProjectionData.h b/Engine/Plugins/Experimental/PCG/Source/PCG/Public/Data/PCGProjectionData.h
--- a/Engine/Plugins/Experimental/PCG/Source/PCG/Public/Data/PCGProjectionData.h
+++ b/Engine/Plugins/Experimental/PCG/Source/PCG/Public/Data/PCGProjectionData.h
@@ -21,8 +21,11 @@ public:
 	virtual int GetDimension() const override;
 	virtual FBox GetBounds() const override;
 	virtual FBox GetStrictBounds() const override;
-	virtual float GetDensityAtPosition(const FVector& InPosition) const override;
 	virtual FVector GetNormal() const override;
+	virtual float GetDensityAtPosition(const FVector& InPosition) const override;
+	virtual FVector TransformPosition(const FVector& InPosition) const override;
+	virtual FPCGPoint TransformPoint(const FPCGPoint& InPoint) const override;
+	virtual bool HasNonTrivialTransform() const override;
 	//~End UPCGSpatialData interface
 
 	//~Begin UPCGSpatialDataWithPointCache interface
@@ -30,9 +33,6 @@ public:
 	//~End UPCGSpatialDataWithPointCache interface
 
 protected:
-	virtual FVector ProjectPosition(const FVector& InPosition) const;
-	virtual FPCGPoint ProjectPoint(const FPCGPoint& InPoint) const;
-
 	FBox ProjectBounds(const FBox& InBounds) const;
 
 	UPROPERTY(BlueprintReadOnly, VisibleAnywhere, Category = SpatialData)
diff --git a/Engine/Plugins/Experimental/PCG/Source/PCG/Public/Data/PCGSpatialData.h b/Engine/Plugins/Experimental/PCG/Source/PCG/Public/Data/PCGSpatialData.h
--- a/Engine/Plugins/Experimental/PCG/Source/PCG/Public/Data/PCGSpatialData.h
+++ b/Engine/Plugins/Experimental/PCG/Source/PCG/Public/Data/PCGSpatialData.h
@@ -61,6 +61,10 @@ public:
 	UFUNCTION(BlueprintCallable, Category = SpatialData)
 	virtual FPCGPoint TransformPoint(const FPCGPoint& InPoint) const { return InPoint; }
 
+	/** Returns true if the data has a non-trivial transform */
+	UFUNCTION(BlueprintCallable, Category = SpatialData)
+	virtual bool HasNonTrivialTransform() const { return false; }
+
 	/** Returns a specialized data to intersect with another data */
 	UFUNCTION(BlueprintCallable, Category = SpatialData)
 	virtual UPCGIntersectionData* IntersectWith(const UPCGSpatialData* InOther) const;
diff --git a/Engine/Plugins/Experimental/PCG/Source/PCG/Public/Data/PCGSurfaceData.h b/Engine/Plugins/Experimental/PCG/Source/PCG/Public/Data/PCGSurfaceData.h
--- a/Engine/Plugins/Experimental/PCG/Source/PCG/Public/Data/PCGSurfaceData.h
+++ b/Engine/Plugins/Experimental/PCG/Source/PCG/Public/Data/PCGSurfaceData.h
@@ -17,6 +17,7 @@ public:
 	virtual int GetDimension() const override { return 2; }
 	virtual FVector TransformPosition(const FVector& InPosition) const override;
 	virtual FPCGPoint TransformPoint(const FPCGPoint& InPoint) const override;
+	virtual bool HasNonTrivialTransform() const override { return true; }
 	//~End UPCGSpatialData interface
 
 protected:
diff --git a/Engine/Plugins/Experimental/PCG/Source/PCG/Public/Data/PCGUnionData.h b/Engine/Plugins/Experimental/PCG/Source/PCG/Public/Data/PCGUnionData.h
--- a/Engine/Plugins/Experimental/PCG/Source/PCG/Public/Data/PCGUnionData.h
+++ b/Engine/Plugins/Experimental/PCG/Source/PCG/Public/Data/PCGUnionData.h
@@ -41,6 +41,9 @@ public:
 	virtual FBox GetBounds() const override;
 	virtual FBox GetStrictBounds() const override;
 	virtual float GetDensityAtPosition(const FVector& InPosition) const override;
+	virtual FVector TransformPosition(const FVector& InPosition) const override;
+	virtual FPCGPoint TransformPoint(const FPCGPoint& InPoint) const override;
+	virtual bool HasNonTrivialTransform() const override;
 	//~End UPCGSpatialData interface
 
 	//~Begin UPCGSpatialDataWithPointCache interface
