Implemented grid-based surface sampler

diff --git a/Engine/Plugins/Experimental/PCG/Source/PCG/Private/Elements/PCGSurfaceSampler.cpp b/Engine/Plugins/Experimental/PCG/Source/PCG/Private/Elements/PCGSurfaceSampler.cpp
--- a/Engine/Plugins/Experimental/PCG/Source/PCG/Private/Elements/PCGSurfaceSampler.cpp
+++ b/Engine/Plugins/Experimental/PCG/Source/PCG/Private/Elements/PCGSurfaceSampler.cpp
@@ -1,6 +1,10 @@
 // Copyright Epic Games, Inc. All Rights Reserved.
 
 #include "Elements/PCGSurfaceSampler.h"
+#include "Data/PCGPointData.h"
+#include "Data/PCGSpatialData.h"
+#include "PCGHelpers.h"
+#include "Math/RandomStream.h"
 
 FPCGElementPtr UPCGSurfaceSamplerSettings::CreateElement() const
 {
@@ -9,5 +13,125 @@ FPCGElementPtr UPCGSurfaceSamplerSettings::CreateElement() const
 
 bool FPCGSurfaceSamplerElement::Execute(FPCGContextPtr Context) const
 {
+	TRACE_CPUPROFILER_EVENT_SCOPE(FPCGSurfaceSamplerElement::Execute);
+	// TODO: time-sliced implementation
+	const UPCGSurfaceSamplerSettings* Settings = Context->GetInputSettings<UPCGSurfaceSamplerSettings>();
+	check(Settings);
+
+	// Early out on invalid settings
+	// TODO: we could compute an approximate radius based on the points per squared meters if that's useful
+	if (Settings->PointRadius <= 0)
+	{
+		return true;
+	}
+
+	TArray<FPCGTaggedData>& Outputs = Context->OutputData.TaggedData;
+
+	// TODO: embarassingly parallel loop
+	TArray<FPCGTaggedData> Inputs = Context->InputData.GetInputs();
+	for (const FPCGTaggedData& Input : Inputs)
+	{
+		const UPCGSpatialData* SpatialInput = Cast<UPCGSpatialData>(Input.Data);
+
+		if (!SpatialInput || !SpatialInput->TargetActor)
+		{
+			// Data type mismatch / no output actor
+			continue;
+		}
+
+		const FBox InputBounds = SpatialInput->GetBounds();
+
+		if (!InputBounds.IsValid)
+		{
+			// With invalid bounds, we can't do anything
+			continue;
+		}
+
+		// Conceptually, we will break down the surface bounds in a N x M grid
+		const bool bApplyDensity = Settings->bApplyDensityToPoints;
+		const FVector::FReal InterstitialDistance = Settings->PointRadius;
+		const FVector::FReal InnerCellSize = 2 * Settings->Looseness * Settings->PointRadius;
+		const FVector::FReal CellSize = InterstitialDistance + InnerCellSize;
+		
+		check(CellSize > 0);
+
+		// By using scaled indices in the world, we can easily make this process deterministic
+		const int32 CellMinX = FMath::CeilToInt((InputBounds.Min.X) / CellSize);
+		const int32 CellMaxX = FMath::FloorToInt((InputBounds.Max.X) / CellSize);
+		const int32 CellMinY = FMath::CeilToInt((InputBounds.Min.Y) / CellSize);
+		const int32 CellMaxY = FMath::FloorToInt((InputBounds.Max.Y) / CellSize);
+
+		if (CellMinX > CellMaxX || CellMinY > CellMaxY)
+		{
+			continue;
+		}
+
+		const int32 CellCount = (1 + CellMaxX - CellMinX) * (1 + CellMaxY - CellMinY);
+		check(CellCount > 0);
+
+		const FVector::FReal InvSquaredMeterUnits = 1.0 / (100.0 * 100.0);
+		int32 TargetPointCount = (InputBounds.Max.X - InputBounds.Min.X) * (InputBounds.Max.Y - InputBounds.Min.Y) * Settings->PointsPerSquaredMeter * InvSquaredMeterUnits;
+
+		if (TargetPointCount == 0)
+		{
+			continue;
+		}
+		else if (TargetPointCount > CellCount)
+		{
+			TargetPointCount = CellCount;
+		}
+
+		const FVector::FReal Ratio = TargetPointCount / (FVector::FReal)CellCount;
+
+		// Finally, create data
+		FPCGTaggedData& Output = Outputs.Emplace_GetRef();
+		Output = Input;
+
+		UPCGPointData* SampledData = NewObject<UPCGPointData>();
+		SampledData->TargetActor = SpatialInput->TargetActor;
+		Output.Data = SampledData;
+
+		TArray<FPCGPoint>& SampledPoints = SampledData->GetMutablePoints();
+
+		const FVector::FReal CellStartX = CellMinX * CellSize;
+		const FVector::FReal CellStartY = CellMinY * CellSize;
+
+		FVector::FReal CurrentX = CellStartX;
+		
+		for(int32 CellX = CellMinX; CellX <= CellMaxX; ++CellX)
+		{
+			FVector::FReal CurrentY = CellStartY;
+
+			for(int32 CellY = CellMinY; CellY <= CellMaxY; ++CellY)
+			{
+				FRandomStream RandomSource(PCGHelpers::ComputeSeed(Settings->Seed, CellX, CellY));
+				float Chance = RandomSource.FRand();
+
+				if (Chance < Ratio)
+				{
+					const float RandX = RandomSource.FRand();
+					const float RandY = RandomSource.FRand();
+
+					FPCGPoint Point;
+					Point.Transform = FTransform(FVector(CurrentX + RandX * InnerCellSize, CurrentY + RandY * InnerCellSize, 0));
+					Point.Extents = FVector(InterstitialDistance);
+					Point.Density = bApplyDensity ? Chance / Ratio : 1.0f;
+					Point.Seed = RandomSource.GetCurrentSeed();
+
+					Point = SpatialInput->TransformPoint(Point);
+
+					if (Point.Density > 0)
+					{
+						SampledPoints.Add(Point);
+					}
+				}
+
+				CurrentY += CellSize;
+			}
+
+			CurrentX += CellSize;
+		}
+	}
+
 	return true;
 }
\ No newline at end of file
diff --git a/Engine/Plugins/Experimental/PCG/Source/PCG/Public/Elements/PCGSurfaceSampler.h b/Engine/Plugins/Experimental/PCG/Source/PCG/Public/Elements/PCGSurfaceSampler.h
--- a/Engine/Plugins/Experimental/PCG/Source/PCG/Public/Elements/PCGSurfaceSampler.h
+++ b/Engine/Plugins/Experimental/PCG/Source/PCG/Public/Elements/PCGSurfaceSampler.h
@@ -17,10 +17,26 @@ class UPCGSurfaceSamplerSettings : public UPCGSettings
 	GENERATED_BODY()
 
 public:
-	// TODO
+	UPROPERTY(BlueprintReadWrite, EditAnywhere, Category = Settings, meta=(ClampMin="0"))
+	float PointsPerSquaredMeter = 0.1f;
+
+	UPROPERTY(BlueprintReadWrite, EditAnywhere, Category = Settings, meta=(ClampMin="0"))
+	float PointRadius = 100.0f;
+
+	UPROPERTY(BlueprintReadWrite, EditAnywhere, Category = Settings, meta=(ClampMin="0"))
+	float Looseness = 1.0f;
+
+	UPROPERTY(BlueprintReadWrite, EditAnywhere, Category = Settings)
+	bool bApplyDensityToPoints = true;
+
+#if WITH_EDITOR
+	//~Begin UPCGSettings interface
+	virtual FName GetDefaultNodeName() const override { return FName(TEXT("SurfaceSamplerNode")); }
+#endif
 
 protected:
 	virtual FPCGElementPtr CreateElement() const override;
+	// ~End UPCGSettings interface
 };
 
 class FPCGSurfaceSamplerElement : public FSimpleTypedPCGElement<UPCGSurfaceSamplerSettings>
