Made the surface sampler forward any exclusions

diff --git a/Engine/Plugins/Experimental/PCG/Source/PCG/Private/Elements/PCGSurfaceSampler.cpp b/Engine/Plugins/Experimental/PCG/Source/PCG/Private/Elements/PCGSurfaceSampler.cpp
--- a/Engine/Plugins/Experimental/PCG/Source/PCG/Private/Elements/PCGSurfaceSampler.cpp
+++ b/Engine/Plugins/Experimental/PCG/Source/PCG/Private/Elements/PCGSurfaceSampler.cpp
@@ -133,5 +133,9 @@ bool FPCGSurfaceSamplerElement::Execute(FPCGContextPtr Context) const
 		}
 	}
 
+	// Finally, forward any exclusions/settings
+	Outputs.Append(Context->InputData.GetExclusions());
+	Outputs.Append(Context->InputData.GetAllSettings());
+
 	return true;
 }
\ No newline at end of file
diff --git a/Engine/Plugins/Experimental/PCG/Source/PCG/Private/PCGData.cpp b/Engine/Plugins/Experimental/PCG/Source/PCG/Private/PCGData.cpp
--- a/Engine/Plugins/Experimental/PCG/Source/PCG/Private/PCGData.cpp
+++ b/Engine/Plugins/Experimental/PCG/Source/PCG/Private/PCGData.cpp
@@ -17,6 +17,13 @@ TArray<FPCGTaggedData> FPCGDataCollection::GetExclusions() const
 		});
 }
 
+TArray<FPCGTaggedData> FPCGDataCollection::GetAllSettings() const
+{
+	return TaggedData.FilterByPredicate([](const FPCGTaggedData& Data) {
+		return Data.Usage == EPCGDataUsage::Settings;
+		});
+}
+
 const UPCGSettings* FPCGDataCollection::GetSettings(const UPCGSettings* InDefaultSettings) const
 {
 	if (!InDefaultSettings)
@@ -43,4 +50,9 @@ TArray<FPCGTaggedData> UPCGDataFunctionLibrary::GetInputs(const FPCGDataCollecti
 TArray<FPCGTaggedData> UPCGDataFunctionLibrary::GetExclusions(const FPCGDataCollection& InCollection)
 {
 	return InCollection.GetExclusions();
+}
+
+TArray<FPCGTaggedData> UPCGDataFunctionLibrary::GetAllSettings(const FPCGDataCollection& InCollection)
+{
+	return InCollection.GetAllSettings();
 }
\ No newline at end of file
diff --git a/Engine/Plugins/Experimental/PCG/Source/PCG/Public/PCGData.h b/Engine/Plugins/Experimental/PCG/Source/PCG/Public/PCGData.h
--- a/Engine/Plugins/Experimental/PCG/Source/PCG/Public/PCGData.h
+++ b/Engine/Plugins/Experimental/PCG/Source/PCG/Public/PCGData.h
@@ -52,6 +52,7 @@ struct FPCGDataCollection
 
 	TArray<FPCGTaggedData> GetInputs() const;
 	TArray<FPCGTaggedData> GetExclusions() const;
+	TArray<FPCGTaggedData> GetAllSettings() const;
 
 	template<typename SettingsType>
 	const SettingsType* GetSettings() const;
@@ -84,4 +85,7 @@ public:
 
 	UFUNCTION(BlueprintCallable, Category = Data)
 	static TArray<FPCGTaggedData> GetExclusions(const FPCGDataCollection& InCollection);
+
+	UFUNCTION(BlueprintCallable, Category = Data)
+	static TArray<FPCGTaggedData> GetAllSettings(const FPCGDataCollection& InCollection);
 };
