Added PCG first pass self-pruning element.
Fixed an error with point extents set in the surface sampler
Improved default behavior on blueprint node that wouldn't be setup properly
Fixed a category for the ToPointData call

diff --git a/Engine/Plugins/Experimental/PCG/Source/PCG/Private/Data/PCGPointData.cpp b/Engine/Plugins/Experimental/PCG/Source/PCG/Private/Data/PCGPointData.cpp
--- a/Engine/Plugins/Experimental/PCG/Source/PCG/Private/Data/PCGPointData.cpp
+++ b/Engine/Plugins/Experimental/PCG/Source/PCG/Private/Data/PCGPointData.cpp
@@ -26,6 +26,16 @@ TArray<FPCGPoint>& UPCGPointData::GetMutablePoints()
 	return Points;
 }
 
+const UPCGPointData::PointOctree& UPCGPointData::GetOctree() const
+{
+	if (bOctreeIsDirty)
+	{
+		RebuildOctree();
+	}
+
+	return Octree;
+}
+
 FBox UPCGPointData::GetBounds() const
 {
 	if (bBoundsAreDirty)
diff --git a/Engine/Plugins/Experimental/PCG/Source/PCG/Private/Elements/PCGExecuteBlueprint.cpp b/Engine/Plugins/Experimental/PCG/Source/PCG/Private/Elements/PCGExecuteBlueprint.cpp
--- a/Engine/Plugins/Experimental/PCG/Source/PCG/Private/Elements/PCGExecuteBlueprint.cpp
+++ b/Engine/Plugins/Experimental/PCG/Source/PCG/Private/Elements/PCGExecuteBlueprint.cpp
@@ -72,12 +72,16 @@ FPCGElementPtr UPCGBlueprintSettings::CreateElement() const
 bool FPCGExecuteBlueprintElement::Execute(FPCGContextPtr Context) const
 {
 	const UPCGBlueprintSettings* Settings = Context->GetInputSettings<UPCGBlueprintSettings>();
-	check(Settings && Settings->BlueprintElementInstance);
 
 	if (Settings && Settings->BlueprintElementInstance)
 	{
 		Settings->BlueprintElementInstance->Execute(Context->InputData, Context->OutputData);
 	}
+	else
+	{
+		// Nothing to do but forward data
+		Context->OutputData = Context->InputData;
+	}
 	
 	return true;
 }
\ No newline at end of file
diff --git a/Engine/Plugins/Experimental/PCG/Source/PCG/Private/Elements/PCGSelfPruning.cpp b/Engine/Plugins/Experimental/PCG/Source/PCG/Private/Elements/PCGSelfPruning.cpp
new file mode 100644
--- /dev/null
+++ b/Engine/Plugins/Experimental/PCG/Source/PCG/Private/Elements/PCGSelfPruning.cpp
@@ -0,0 +1,174 @@
+// Copyright Epic Games, Inc. All Rights Reserved.
+
+#include "Elements/PCGSelfPruning.h"
+#include "Data/PCGPointData.h"
+#include "Data/PCGSpatialData.h"
+#include "PCGHelpers.h"
+#include "Math/RandomStream.h"
+
+namespace PCGSelfPruningAlgorithms
+{
+	bool RandomSort(const FPCGPoint* A, const FPCGPoint* B)
+	{
+		return A->Seed < B->Seed;
+	}
+
+	bool SortSmallToLargeNoRandom(const FPCGPoint* A, const FPCGPoint* B, FVector::FReal SquaredRadiusEquality)
+	{
+		return A->Extents.SquaredLength() * SquaredRadiusEquality < B->Extents.SquaredLength();
+	}
+
+	bool SortSmallToLargeWithRandom(const FPCGPoint* A, const FPCGPoint* B, FVector::FReal SquaredRadiusEquality)
+	{
+		const FVector::FReal SqrLenA = A->Extents.SquaredLength();
+		const FVector::FReal SqrLenB = B->Extents.SquaredLength();
+		if (SqrLenA * SquaredRadiusEquality < SqrLenB)
+		{
+			return true;
+		}
+		else if (SqrLenB * SquaredRadiusEquality < SqrLenA)
+		{
+			return false;
+		}
+		else
+		{
+			return RandomSort(A, B);
+		}
+	}
+}
+
+FPCGElementPtr UPCGSelfPruningSettings::CreateElement() const
+{
+	return MakeShared<FPCGSelfPruningElement>();
+}
+
+bool FPCGSelfPruningElement::Execute(FPCGContextPtr Context) const
+{
+	TRACE_CPUPROFILER_EVENT_SCOPE(FPCGSelfPruningElement::Execute);
+	// TODO: time-sliced implementation
+	const UPCGSelfPruningSettings* Settings = Context->GetInputSettings<UPCGSelfPruningSettings>();
+	check(Settings);
+
+	// Early out: if pruning is disabled
+	if (Settings->PruningType == EPCGSelfPruningType::None)
+	{
+		Context->OutputData = Context->InputData;
+		return true;
+	}
+
+	const FVector::FReal RadiusEquality = 1.0f + Settings->RadiusSimilarityFactor;
+	const FVector::FReal SquaredRadiusEquality = FMath::Square(RadiusEquality);
+
+	TArray<FPCGTaggedData>& Outputs = Context->OutputData.TaggedData;
+
+	// TODO: embarassingly parallel loop
+	TArray<FPCGTaggedData> Inputs = Context->InputData.GetInputs();
+	for (const FPCGTaggedData& Input : Inputs)
+	{
+		const UPCGSpatialData* SpatialInput = Cast<UPCGSpatialData>(Input.Data);
+
+		if (!SpatialInput || !SpatialInput->TargetActor)
+		{
+			// Data type mismatch / no output actor
+			continue;
+		}
+
+		const UPCGPointData* InputPointData = SpatialInput->ToPointData();
+		const TArray<FPCGPoint>& Points = InputPointData->GetPoints();
+		const UPCGPointData::PointOctree& Octree = InputPointData->GetOctree();
+
+		// Self-pruning will be done as follows:
+		// For each point:
+		//  if in its vicinity, there is >=1 non-rejected point with a radius significantly larger
+		//  or in its range + has a randomly assigned index -> we'll look at its seed
+		//  then remove this point
+		TArray<const FPCGPoint*> SortedPoints;
+		SortedPoints.Reserve(Points.Num());
+		for (const FPCGPoint& Point : Points)
+		{
+			SortedPoints.Add(&Point);
+		}
+
+		// Apply proper sort algorithm
+		if (Settings->PruningType == EPCGSelfPruningType::LargeToSmall)
+		{
+			if (Settings->bRandomizedPruning)
+			{
+				Algo::Sort(SortedPoints, [SquaredRadiusEquality](const FPCGPoint* A, const FPCGPoint* B) { return !PCGSelfPruningAlgorithms::SortSmallToLargeWithRandom(A, B, SquaredRadiusEquality); });
+			}
+			else
+			{
+				Algo::Sort(SortedPoints, [SquaredRadiusEquality](const FPCGPoint* A, const FPCGPoint* B) { return !PCGSelfPruningAlgorithms::SortSmallToLargeNoRandom(A, B, SquaredRadiusEquality); });
+			}
+		}
+		else if (Settings->PruningType == EPCGSelfPruningType::SmallToLarge)
+		{
+			if (Settings->bRandomizedPruning)
+			{
+				Algo::Sort(SortedPoints, [SquaredRadiusEquality](const FPCGPoint* A, const FPCGPoint* B) { return PCGSelfPruningAlgorithms::SortSmallToLargeWithRandom(A, B, SquaredRadiusEquality); });
+			}
+			else
+			{
+				Algo::Sort(SortedPoints, [SquaredRadiusEquality](const FPCGPoint* A, const FPCGPoint* B) { return PCGSelfPruningAlgorithms::SortSmallToLargeNoRandom(A, B, SquaredRadiusEquality); });
+			}
+		}
+		else
+		{
+			if (Settings->bRandomizedPruning)
+			{
+				Algo::Sort(SortedPoints, PCGSelfPruningAlgorithms::RandomSort);
+			}
+		}
+
+		TSet<const FPCGPoint*> ExclusionPoints;
+		ExclusionPoints.Reserve(Points.Num());
+
+		TSet<const FPCGPoint*> ExcludedPoints;
+		ExcludedPoints.Reserve(Points.Num());
+
+		for (int PointIndex = 0; PointIndex < SortedPoints.Num(); ++PointIndex)
+		{
+			// Find all points in the vicinity of the current point and reject them, if they aren't part of the exclusions yet
+			// If any of these have been seen before (e.g. in the ExclusionPoints)
+			// then reject this point
+			const FPCGPoint* CurrentPoint = SortedPoints[PointIndex];
+
+			if (ExcludedPoints.Contains(CurrentPoint))
+			{
+				continue;
+			}
+
+			ExclusionPoints.Add(CurrentPoint);
+
+			Octree.FindElementsWithBoundsTest(FBoxCenterAndExtent(CurrentPoint->Transform.GetLocation(), CurrentPoint->Extents), [&ExclusionPoints, &ExcludedPoints](const FPCGPointRef& InPointRef) {
+				if (!ExclusionPoints.Contains(InPointRef.Point))
+				{
+					ExcludedPoints.Add(InPointRef.Point);
+				}
+			});
+		}
+
+		// Finally, output all points that are present in the ExclusionPoints.
+		FPCGTaggedData& Output = Outputs.Emplace_GetRef();
+		Output = Input;
+
+		UPCGPointData* PrunedData = NewObject<UPCGPointData>();
+		PrunedData->TargetActor = SpatialInput->TargetActor;
+
+		Output.Data = PrunedData;
+
+		TArray<FPCGPoint>& OutputPoints = PrunedData->GetMutablePoints();
+		OutputPoints.Reserve(ExclusionPoints.Num());
+
+		for (const FPCGPoint* Point : ExclusionPoints)
+		{
+			OutputPoints.Add(*Point);
+		}
+	}
+
+	// Finally, forward any exclusions/settings
+	Outputs.Append(Context->InputData.GetExclusions());
+	Outputs.Append(Context->InputData.GetAllSettings());
+
+	return true;
+}
\ No newline at end of file
diff --git a/Engine/Plugins/Experimental/PCG/Source/PCG/Private/Elements/PCGSurfaceSampler.cpp b/Engine/Plugins/Experimental/PCG/Source/PCG/Private/Elements/PCGSurfaceSampler.cpp
--- a/Engine/Plugins/Experimental/PCG/Source/PCG/Private/Elements/PCGSurfaceSampler.cpp
+++ b/Engine/Plugins/Experimental/PCG/Source/PCG/Private/Elements/PCGSurfaceSampler.cpp
@@ -49,8 +49,9 @@ bool FPCGSurfaceSamplerElement::Execute(FPCGContextPtr Context) const
 
 		// Conceptually, we will break down the surface bounds in a N x M grid
 		const bool bApplyDensity = Settings->bApplyDensityToPoints;
-		const FVector::FReal InterstitialDistance = Settings->PointRadius;
-		const FVector::FReal InnerCellSize = 2 * Settings->Looseness * Settings->PointRadius;
+		const FVector::FReal PointRadius = Settings->PointRadius;
+		const FVector::FReal InterstitialDistance = 2 * PointRadius;
+		const FVector::FReal InnerCellSize = 2 * Settings->Looseness * PointRadius;
 		const FVector::FReal CellSize = InterstitialDistance + InnerCellSize;
 		
 		check(CellSize > 0);
@@ -114,8 +115,8 @@ bool FPCGSurfaceSamplerElement::Execute(FPCGContextPtr Context) const
 
 					FPCGPoint Point;
 					Point.Transform = FTransform(FVector(CurrentX + RandX * InnerCellSize, CurrentY + RandY * InnerCellSize, 0));
-					Point.Extents = FVector(InterstitialDistance);
-					Point.Density = bApplyDensity ? Chance / Ratio : 1.0f;
+					Point.Extents = FVector(PointRadius);
+					Point.Density = bApplyDensity ? ((Ratio - Chance) / Ratio) : 1.0f;
 					Point.Seed = RandomSource.GetCurrentSeed();
 
 					Point = SpatialInput->TransformPoint(Point);
diff --git a/Engine/Plugins/Experimental/PCG/Source/PCG/Private/PCGGraph.cpp b/Engine/Plugins/Experimental/PCG/Source/PCG/Private/PCGGraph.cpp
--- a/Engine/Plugins/Experimental/PCG/Source/PCG/Private/PCGGraph.cpp
+++ b/Engine/Plugins/Experimental/PCG/Source/PCG/Private/PCGGraph.cpp
@@ -114,7 +114,7 @@ UPCGNode* UPCGGraph::AddNodeOfType(TSubclassOf<class UPCGSettings> InSettingsCla
 		const FName DefaultNodeName = Settings->GetDefaultNodeName();
 		if (DefaultNodeName != NAME_None)
 		{
-			FName NodeName = MakeUniqueObjectName(GetOuter(), UPCGNode::StaticClass(), DefaultNodeName);
+			FName NodeName = MakeUniqueObjectName(this, UPCGNode::StaticClass(), DefaultNodeName);
 			Node->Rename(*NodeName.ToString());
 		}
 #endif
diff --git a/Engine/Plugins/Experimental/PCG/Source/PCG/Public/Data/PCGPointData.h b/Engine/Plugins/Experimental/PCG/Source/PCG/Public/Data/PCGPointData.h
--- a/Engine/Plugins/Experimental/PCG/Source/PCG/Public/Data/PCGPointData.h
+++ b/Engine/Plugins/Experimental/PCG/Source/PCG/Public/Data/PCGPointData.h
@@ -57,6 +57,8 @@ class UPCGPointData : public UPCGSpatialData
 	GENERATED_BODY()
 
 public:
+	typedef TOctree2<FPCGPointRef, FPCGPointRefSemantics> PointOctree;
+
 	void Initialize(AActor* InActor);
 
 	// ~Begin UPCGSpatialData interface
@@ -79,6 +81,8 @@ public:
 
 	TArray<FPCGPoint>& GetMutablePoints();
 
+	const PointOctree& GetOctree() const;
+
 protected:
 	void RebuildOctree() const;
 	void RecomputeBounds() const;
@@ -87,7 +91,7 @@ protected:
 	TArray<FPCGPoint> Points;
 
 	mutable FCriticalSection CachedDataLock;
-	mutable TOctree2<FPCGPointRef, FPCGPointRefSemantics> Octree;
+	mutable PointOctree Octree;
 	mutable FBox Bounds; // TODO: review if this needs to be threadsafe
 	mutable bool bBoundsAreDirty = true;
 	mutable bool bOctreeIsDirty = true;
diff --git a/Engine/Plugins/Experimental/PCG/Source/PCG/Public/Data/PCGSpatialData.h b/Engine/Plugins/Experimental/PCG/Source/PCG/Public/Data/PCGSpatialData.h
--- a/Engine/Plugins/Experimental/PCG/Source/PCG/Public/Data/PCGSpatialData.h
+++ b/Engine/Plugins/Experimental/PCG/Source/PCG/Public/Data/PCGSpatialData.h
@@ -50,7 +50,7 @@ public:
 	virtual float GetDensityAtPosition(const FVector& InPosition) const PURE_VIRTUAL(UPCGSpatialData::GetDensityAtPosition, return 0;);
 
 	/** Discretizes the data into points */
-	UFUNCTION(BlueprintCallable, Category = Data)
+	UFUNCTION(BlueprintCallable, Category = SpatialData)
 	virtual const UPCGPointData* ToPointData() const PURE_VIRTUAL(UPCGSpatialData::ToPointData, return nullptr;);
 
 	/** Transform a world-space position to a world-space position in relation to the current data. (ex: projection on surface) */
diff --git a/Engine/Plugins/Experimental/PCG/Source/PCG/Public/Elements/PCGSelfPruning.h b/Engine/Plugins/Experimental/PCG/Source/PCG/Public/Elements/PCGSelfPruning.h
new file mode 100644
--- /dev/null
+++ b/Engine/Plugins/Experimental/PCG/Source/PCG/Public/Elements/PCGSelfPruning.h
@@ -0,0 +1,50 @@
+// Copyright Epic Games, Inc. All Rights Reserved.
+
+#pragma once
+
+#include "CoreMinimal.h"
+#include "PCGNode.h"
+#include "PCGSettings.h"
+#include "PCGElement.h"
+
+#include "PCGSelfPruning.generated.h"
+
+UENUM()
+enum class EPCGSelfPruningType : uint8
+{
+	LargeToSmall,
+	SmallToLarge,
+	AllEqual,
+	None
+};
+
+UCLASS(BlueprintType, ClassGroup=(Procedural))
+class UPCGSelfPruningSettings : public UPCGSettings
+{
+	GENERATED_BODY()
+
+public:
+	UPROPERTY(BlueprintReadWrite, EditAnywhere, Category = Settings)
+	EPCGSelfPruningType PruningType = EPCGSelfPruningType::LargeToSmall;
+
+	UPROPERTY(BlueprintReadWrite, EditAnywhere, Category = Settings, meta=(ClampMin=0.0f, EditCondition="PruningType == EPCGSelfPruningType::LargeToSmall || PruningType == PruningType == EPCGSelfPruningType::SmallToLarge"))
+	float RadiusSimilarityFactor = 0.25f;
+
+	UPROPERTY(BlueprintReadWrite, EditAnywhere, Category = Settings)
+	bool bRandomizedPruning = true;
+
+#if WITH_EDITOR
+	//~Begin UPCGSettings interface
+	virtual FName GetDefaultNodeName() const override { return FName(TEXT("SelfPruningNode")); }
+#endif
+
+protected:
+	virtual FPCGElementPtr CreateElement() const override;
+	// ~End UPCGSettings interface
+};
+
+class FPCGSelfPruningElement : public FSimpleTypedPCGElement<UPCGSelfPruningSettings>
+{
+public:
+	virtual bool Execute(FPCGContextPtr Context) const override;
+};
\ No newline at end of file
