Improved volume/brush support
Fixed outstanding issue in how we compute final density in intersection/union/difference

diff --git a/Engine/Plugins/Experimental/PCG/Source/PCG/Private/Data/PCGDifferenceData.cpp b/Engine/Plugins/Experimental/PCG/Source/PCG/Private/Data/PCGDifferenceData.cpp
--- a/Engine/Plugins/Experimental/PCG/Source/PCG/Private/Data/PCGDifferenceData.cpp
+++ b/Engine/Plugins/Experimental/PCG/Source/PCG/Private/Data/PCGDifferenceData.cpp
@@ -21,7 +21,7 @@ void UPCGDifferenceData::Initialize(const UPCGSpatialData* InData)
 
 void UPCGDifferenceData::AddDifference(const UPCGSpatialData* InDifference)
 {
-	check(InDifference && InDifference->GetDimension() == Source->GetDimension());
+	check(InDifference);
 
 	if (!Difference)
 	{
@@ -112,7 +112,14 @@ FVector UPCGDifferenceData::TransformPosition(const FVector& InPosition) const
 FPCGPoint UPCGDifferenceData::TransformPoint(const FPCGPoint& InPoint) const
 {
 	check(Source);
-	return Source->TransformPoint(InPoint);
+	FPCGPoint TransformedPoint = Source->TransformPoint(InPoint);
+
+	if (Difference && TransformedPoint.Density > 0)
+	{
+		TransformedPoint.Density = FMath::Max(0, TransformedPoint.Density - Difference->GetDensityAtPosition(TransformedPoint.Transform.GetLocation()));
+	}
+
+	return TransformedPoint;
 }
 
 bool UPCGDifferenceData::HasNonTrivialTransform() const
diff --git a/Engine/Plugins/Experimental/PCG/Source/PCG/Private/Data/PCGIntersectionData.cpp b/Engine/Plugins/Experimental/PCG/Source/PCG/Private/Data/PCGIntersectionData.cpp
--- a/Engine/Plugins/Experimental/PCG/Source/PCG/Private/Data/PCGIntersectionData.cpp
+++ b/Engine/Plugins/Experimental/PCG/Source/PCG/Private/Data/PCGIntersectionData.cpp
@@ -77,9 +77,18 @@ FVector UPCGIntersectionData::TransformPosition(const FVector& InPosition) const
 }
 
 FPCGPoint UPCGIntersectionData::TransformPoint(const FPCGPoint& InPoint) const
-{
+{ 
 	check(A && B);
-	return A->HasNonTrivialTransform() ? A->TransformPoint(InPoint) : B->TransformPoint(InPoint);
+	const UPCGSpatialData* X = A->HasNonTrivialTransform() ? A : B;
+	const UPCGSpatialData* Y = X == A ? B : A;
+
+	FPCGPoint TransformedPoint = X->TransformPoint(InPoint);
+	if (TransformedPoint.Density > 0)
+	{
+		TransformedPoint.Density = PCGIntersectionDataMaths::ComputeDensity(TransformedPoint.Density, Y->GetDensityAtPosition(TransformedPoint.Transform.GetLocation()), DensityFunction);
+	}	
+
+	return TransformedPoint;
 }
 
 bool UPCGIntersectionData::HasNonTrivialTransform() const
diff --git a/Engine/Plugins/Experimental/PCG/Source/PCG/Private/Data/PCGSpatialData.cpp b/Engine/Plugins/Experimental/PCG/Source/PCG/Private/Data/PCGSpatialData.cpp
--- a/Engine/Plugins/Experimental/PCG/Source/PCG/Private/Data/PCGSpatialData.cpp
+++ b/Engine/Plugins/Experimental/PCG/Source/PCG/Private/Data/PCGSpatialData.cpp
@@ -23,6 +23,14 @@ const UPCGPointData* UPCGSpatialDataWithPointCache::ToPointData() const
 	return CachedPointData;
 }
 
+FPCGPoint UPCGSpatialData::TransformPoint(const FPCGPoint& InPoint) const
+{
+	FPCGPoint TransformedPoint = InPoint;
+	TransformedPoint.Transform.SetTranslation(TransformPosition(InPoint.Transform.GetLocation()));
+	TransformedPoint.Density = GetDensityAtPosition(InPoint.Transform.GetLocation());
+	return TransformedPoint;
+}
+
 UPCGIntersectionData* UPCGSpatialData::IntersectWith(const UPCGSpatialData* InOther) const
 {
 	UPCGIntersectionData* IntersectionData = NewObject<UPCGIntersectionData>(const_cast<UPCGSpatialData*>(this));
diff --git a/Engine/Plugins/Experimental/PCG/Source/PCG/Private/Data/PCGUnionData.cpp b/Engine/Plugins/Experimental/PCG/Source/PCG/Private/Data/PCGUnionData.cpp
--- a/Engine/Plugins/Experimental/PCG/Source/PCG/Private/Data/PCGUnionData.cpp
+++ b/Engine/Plugins/Experimental/PCG/Source/PCG/Private/Data/PCGUnionData.cpp
@@ -50,6 +50,11 @@ void UPCGUnionData::AddData(const UPCGSpatialData* InData)
 		CachedStrictBounds = CachedStrictBounds.Overlap(InData->GetStrictBounds());
 		CachedDimension = FMath::Max(CachedDimension, InData->GetDimension());
 	}
+
+	if (!FirstNonTrivialTransformData && InData->HasNonTrivialTransform())
+	{
+		FirstNonTrivialTransformData = InData;
+	}
 }
 
 int UPCGUnionData::GetDimension() const
@@ -105,41 +110,42 @@ float UPCGUnionData::GetDensityAtPosition(const FVector& InPosition) const
 
 FVector UPCGUnionData::TransformPosition(const FVector& InPosition) const
 {
-	for (TObjectPtr<const UPCGSpatialData> Datum : Data)
+	if (FirstNonTrivialTransformData)
 	{
-		if (Datum->HasNonTrivialTransform())
-		{
-			return Datum->TransformPosition(InPosition);
-		}
+		return FirstNonTrivialTransformData->TransformPosition(InPosition);
+	}
+	else
+	{
+		return Super::TransformPosition(InPosition);
 	}
-
-	return Super::TransformPosition(InPosition);
 }
 
 FPCGPoint UPCGUnionData::TransformPoint(const FPCGPoint& InPoint) const
 {
-	for (TObjectPtr<const UPCGSpatialData> Datum : Data)
+	if (FirstNonTrivialTransformData)
 	{
-		if (Datum->HasNonTrivialTransform())
+		FPCGPoint TransformedPoint = FirstNonTrivialTransformData->TransformPoint(InPoint);
+
+		const int32 DataCount = Data.Num();
+		for(int32 DataIndex = 0; DataIndex < DataCount && TransformedPoint.Density < 1.0f; ++DataIndex)
 		{
-			return Datum->TransformPoint(InPoint);
+			if (Data[DataIndex] != FirstNonTrivialTransformData)
+			{
+				PCGUnionDataMaths::UpdateDensity(TransformedPoint.Density, Data[DataIndex]->GetDensityAtPosition(TransformedPoint.Transform.GetLocation()), DensityFunction);
+			}
 		}
-	}
 
-	return Super::TransformPoint(InPoint);
+		return TransformedPoint;
+	}
+	else
+	{
+		return Super::TransformPoint(InPoint);
+	}
 }
 
 bool UPCGUnionData::HasNonTrivialTransform() const
 {
-	for (TObjectPtr<const UPCGSpatialData> Datum : Data)
-	{
-		if (Datum->HasNonTrivialTransform())
-		{
-			return true;
-		}
-	}
-
-	return Super::HasNonTrivialTransform();
+	return (FirstNonTrivialTransformData != nullptr || Super::HasNonTrivialTransform());
 }
 
 const UPCGPointData* UPCGUnionData::CreatePointData() const
diff --git a/Engine/Plugins/Experimental/PCG/Source/PCG/Private/Data/PCGVolumeData.cpp b/Engine/Plugins/Experimental/PCG/Source/PCG/Private/Data/PCGVolumeData.cpp
--- a/Engine/Plugins/Experimental/PCG/Source/PCG/Private/Data/PCGVolumeData.cpp
+++ b/Engine/Plugins/Experimental/PCG/Source/PCG/Private/Data/PCGVolumeData.cpp
@@ -16,12 +16,6 @@ void UPCGVolumeData::Initialize(AVolume* InVolume, AActor* InTargetActor)
 
 	// TODO: Compute the strict bounds, we must find a FBox inscribed into the oriented box.
 	// Currently, we'll leave the strict bounds empty and fall back to checking against the local box
-	if (UBrushComponent* BrushComponent = Volume->GetBrushComponent())
-	{
-		VolumeTransform = BrushComponent->GetComponentTransform();
-		FBoxSphereBounds LocalBoxSphereBounds = BrushComponent->CalcBounds(FTransform::Identity);
-		VolumeLocalBounds = FBox::BuildAABB(LocalBoxSphereBounds.Origin, LocalBoxSphereBounds.BoxExtent);
-	}
 }
 
 void UPCGVolumeData::Initialize(const FBox& InBounds, AActor* InTargetActor)
@@ -59,7 +53,7 @@ float UPCGVolumeData::GetDensityAtPosition(const FVector& InPosition) const
 		}
 		else
 		{
-			return VolumeLocalBounds.IsInside(VolumeTransform.InverseTransformPosition(InPosition)) ? 1.0f : 0.0f;
+			return Volume->EncompassesPoint(InPosition) ? 1.0f : 0.0f;
 		}
 	}
 	else
diff --git a/Engine/Plugins/Experimental/PCG/Source/PCG/Private/PCGComponent.cpp b/Engine/Plugins/Experimental/PCG/Source/PCG/Private/PCGComponent.cpp
--- a/Engine/Plugins/Experimental/PCG/Source/PCG/Private/PCGComponent.cpp
+++ b/Engine/Plugins/Experimental/PCG/Source/PCG/Private/PCGComponent.cpp
@@ -641,7 +641,7 @@ void UPCGComponent::OnObjectPropertyChanged(UObject* InObject, FPropertyChangedE
 	{
 		// Something has changed on the owner (including properties of this component)
 		// In the case of splines, this is where we'd get notified if some component properties (incl. spline vertices) have changed
-		if(HasSplineComponent())
+		if(HasSplineComponent() || Cast<AVolume>(Actor) != nullptr)
 		{
 			DirtyGenerated(true);
 		}
diff --git a/Engine/Plugins/Experimental/PCG/Source/PCG/Public/Data/PCGSpatialData.h b/Engine/Plugins/Experimental/PCG/Source/PCG/Public/Data/PCGSpatialData.h
--- a/Engine/Plugins/Experimental/PCG/Source/PCG/Public/Data/PCGSpatialData.h
+++ b/Engine/Plugins/Experimental/PCG/Source/PCG/Public/Data/PCGSpatialData.h
@@ -39,7 +39,7 @@ public:
 
 	/** Returns the bounds in which the density is always 1 */
 	UFUNCTION(BlueprintCallable, Category = SpatialData)
-	virtual FBox GetStrictBounds() const { return GetBounds(); }
+	virtual FBox GetStrictBounds() const { return FBox(EForceInit::ForceInit); }
 
 	/** Returns the expected data normal (for surfaces) or eventual projection axis (for volumes) */
 	UFUNCTION(BlueprintCallable, Category = SpatialData)
@@ -59,7 +59,7 @@ public:
 
 	/** Transform a full point similar to TransformPosition */
 	UFUNCTION(BlueprintCallable, Category = SpatialData)
-	virtual FPCGPoint TransformPoint(const FPCGPoint& InPoint) const { return InPoint; }
+	virtual FPCGPoint TransformPoint(const FPCGPoint& InPoint) const;
 
 	/** Returns true if the data has a non-trivial transform */
 	UFUNCTION(BlueprintCallable, Category = SpatialData)
diff --git a/Engine/Plugins/Experimental/PCG/Source/PCG/Public/Data/PCGUnionData.h b/Engine/Plugins/Experimental/PCG/Source/PCG/Public/Data/PCGUnionData.h
--- a/Engine/Plugins/Experimental/PCG/Source/PCG/Public/Data/PCGUnionData.h
+++ b/Engine/Plugins/Experimental/PCG/Source/PCG/Public/Data/PCGUnionData.h
@@ -54,6 +54,9 @@ protected:
 	UPROPERTY(BlueprintReadOnly, VisibleAnywhere, Category = SpatialData)
 	TArray<TObjectPtr<const UPCGSpatialData>> Data;
 
+	UPROPERTY()
+	TObjectPtr<const UPCGSpatialData> FirstNonTrivialTransformData = nullptr;
+
 	UPROPERTY(BlueprintReadWrite, EditAnywhere, Category = Settings)
 	EPCGUnionType UnionType = EPCGUnionType::LeftToRightPriority;
 
diff --git a/Engine/Plugins/Experimental/PCG/Source/PCG/Public/Data/PCGVolumeData.h b/Engine/Plugins/Experimental/PCG/Source/PCG/Public/Data/PCGVolumeData.h
--- a/Engine/Plugins/Experimental/PCG/Source/PCG/Public/Data/PCGVolumeData.h
+++ b/Engine/Plugins/Experimental/PCG/Source/PCG/Public/Data/PCGVolumeData.h
@@ -33,9 +33,9 @@ protected:
 	UPROPERTY(BlueprintReadWrite, EditAnywhere, Category = SourceData)
 	TObjectPtr<AVolume> Volume = nullptr;
 
+	UPROPERTY()
 	FBox Bounds = FBox(EForceInit::ForceInit);
-	FBox StrictBounds = FBox(EForceInit::ForceInit);
 
-	FTransform VolumeTransform = FTransform::Identity;
-	FBox VolumeLocalBounds = FBox(EForceInit::ForceInit);
+	UPROPERTY()
+	FBox StrictBounds = FBox(EForceInit::ForceInit);
 };
\ No newline at end of file
