Fixed issue with pcg spline data bounds (too tight; wasn't taking the radius into account)

diff --git a/Engine/Plugins/Experimental/PCG/Source/PCG/Private/Data/PCGSplineData.cpp b/Engine/Plugins/Experimental/PCG/Source/PCG/Private/Data/PCGSplineData.cpp
--- a/Engine/Plugins/Experimental/PCG/Source/PCG/Private/Data/PCGSplineData.cpp
+++ b/Engine/Plugins/Experimental/PCG/Source/PCG/Private/Data/PCGSplineData.cpp
@@ -13,6 +13,19 @@ void UPCGSplineData::Initialize(USplineComponent* InSpline)
 	check(InSpline);
 	Spline = InSpline;
 	TargetActor = InSpline->GetOwner();
+
+	CachedBounds = PCGHelpers::GetActorBounds(Spline->GetOwner());
+
+	// Expand bounds by the radius of points, otherwise sections of the curve that are close
+	// to the bounds will report an invalid density.
+	FVector SplinePointsRadius = FVector::ZeroVector;
+	const FInterpCurveVector& SplineScales = Spline->GetSplinePointsScale();
+	for (const FInterpCurvePoint<FVector>& SplineScale : SplineScales.Points)
+	{
+		SplinePointsRadius = FVector::Max(SplinePointsRadius, SplineScale.OutVal.GetAbs());
+	}
+
+	CachedBounds = CachedBounds.ExpandBy(SplinePointsRadius, SplinePointsRadius);
 }
 
 int UPCGSplineData::GetNumSegments() const
@@ -63,7 +76,7 @@ const UPCGPointData* UPCGSplineData::CreatePointData() const
 
 FBox UPCGSplineData::GetBounds() const
 {
-	return PCGHelpers::GetActorBounds(Spline->GetOwner());
+	return CachedBounds;
 }
 
 float UPCGSplineData::GetDensityAtPosition(const FVector& InPosition) const
diff --git a/Engine/Plugins/Experimental/PCG/Source/PCG/Public/Data/PCGSplineData.h b/Engine/Plugins/Experimental/PCG/Source/PCG/Public/Data/PCGSplineData.h
--- a/Engine/Plugins/Experimental/PCG/Source/PCG/Public/Data/PCGSplineData.h
+++ b/Engine/Plugins/Experimental/PCG/Source/PCG/Public/Data/PCGSplineData.h
@@ -36,6 +36,10 @@ public:
 
 	UPROPERTY(BlueprintReadWrite, EditAnywhere, Category = SourceData)
 	TSoftObjectPtr<USplineComponent> Spline;
+
+protected:
+	UPROPERTY()
+	FBox CachedBounds = FBox(EForceInit::ForceInit);
 };
 
 UCLASS(BlueprintType, ClassGroup=(Procedural))
