Added graph regeneration on blueprint (element) compilation.
Fixed an issue where compiling a PCG blueprint element would make its parameter changes not trigger regeneration anymore.

diff --git a/Engine/Plugins/Experimental/PCG/Source/PCG/Private/Elements/PCGExecuteBlueprint.cpp b/Engine/Plugins/Experimental/PCG/Source/PCG/Private/Elements/PCGExecuteBlueprint.cpp
--- a/Engine/Plugins/Experimental/PCG/Source/PCG/Private/Elements/PCGExecuteBlueprint.cpp
+++ b/Engine/Plugins/Experimental/PCG/Source/PCG/Private/Elements/PCGExecuteBlueprint.cpp
@@ -16,6 +16,59 @@ UWorld* UPCGBlueprintElement::GetWorld() const
 }
 
 #if WITH_EDITOR
+void UPCGBlueprintElement::PostEditChangeProperty(FPropertyChangedEvent& PropertyChangedEvent)
+{
+	Super::PostEditChangeProperty(PropertyChangedEvent);
+	OnBlueprintChangedDelegate.Broadcast(this);
+}
+#endif // WITH_EDITOR
+
+void UPCGBlueprintSettings::SetupBlueprintEvent()
+{
+#if WITH_EDITOR
+	if (BlueprintElementType)
+	{
+		if (UBlueprint* Blueprint = Cast<UBlueprint>(BlueprintElementType->ClassGeneratedBy))
+		{
+			Blueprint->OnChanged().AddUObject(this, &UPCGBlueprintSettings::OnBlueprintChanged);
+		}
+	}
+#endif
+}
+
+void UPCGBlueprintSettings::TeardownBlueprintEvent()
+{
+#if WITH_EDITOR
+	if (BlueprintElementType)
+	{
+		if (UBlueprint* Blueprint = Cast<UBlueprint>(BlueprintElementType->ClassGeneratedBy))
+		{
+			Blueprint->OnChanged().RemoveAll(this);
+		}
+	}
+#endif
+}
+
+void UPCGBlueprintSettings::SetupBlueprintElementEvent()
+{
+#if WITH_EDITOR
+	if (BlueprintElementInstance)
+	{
+		BlueprintElementInstance->OnBlueprintChangedDelegate.AddUObject(this, &UPCGBlueprintSettings::OnBlueprintElementChanged);
+	}
+#endif
+}
+
+void UPCGBlueprintSettings::TeardownBlueprintElementEvent()
+{
+#if WITH_EDITOR
+	if (BlueprintElementInstance)
+	{
+		BlueprintElementInstance->OnBlueprintChangedDelegate.RemoveAll(this);
+	}
+#endif
+}
+
 void UPCGBlueprintSettings::PostLoad()
 {
 	Super::PostLoad();
@@ -26,12 +79,25 @@ void UPCGBlueprintSettings::PostLoad()
 		BlueprintElement_DEPRECATED = nullptr;
 	}
 
+	SetupBlueprintEvent();
+
 	if (!BlueprintElementInstance)
 	{
 		RefreshBlueprintElement();
 	}
+	else
+	{
+		SetupBlueprintElementEvent();
+	}
 }
 
+void UPCGBlueprintSettings::BeginDestroy()
+{
+	TeardownBlueprintElementEvent();
+	TeardownBlueprintEvent();
+}
+
+#if WITH_EDITOR
 void UPCGBlueprintSettings::PostEditChangeProperty(FPropertyChangedEvent& PropertyChangedEvent)
 {
 	if (!BlueprintElementInstance || BlueprintElementInstance->GetClass() != BlueprintElementType)
@@ -41,22 +107,45 @@ void UPCGBlueprintSettings::PostEditChangeProperty(FPropertyChangedEvent& Proper
 
 	Super::PostEditChangeProperty(PropertyChangedEvent);
 }
+
+void UPCGBlueprintSettings::OnBlueprintChanged(UBlueprint* InBlueprint)
+{
+	// When the blueprint changes, the element gets recreated, so we must rewire it here.
+	TeardownBlueprintElementEvent();
+	SetupBlueprintElementEvent();
+
+	OnSettingsChangedDelegate.Broadcast(this);
+}
+
+void UPCGBlueprintSettings::OnBlueprintElementChanged(UPCGBlueprintElement* InElement)
+{
+	OnSettingsChangedDelegate.Broadcast(this);
+}
 #endif
 
 void UPCGBlueprintSettings::SetElementType(TSubclassOf<UPCGBlueprintElement> InElementType)
 {
 	if (!BlueprintElementInstance || InElementType != BlueprintElementType)
 	{
-		BlueprintElementType = InElementType;
+		if (InElementType != BlueprintElementType)
+		{
+			TeardownBlueprintEvent();
+			BlueprintElementType = InElementType;
+			SetupBlueprintEvent();
+		}
+		
 		RefreshBlueprintElement();
 	}
 }
 
 void UPCGBlueprintSettings::RefreshBlueprintElement()
 {
+	TeardownBlueprintElementEvent();
+
 	if (BlueprintElementType)
 	{
 		BlueprintElementInstance = NewObject<UPCGBlueprintElement>(this, BlueprintElementType);
+		SetupBlueprintElementEvent();
 	}
 	else
 	{
diff --git a/Engine/Plugins/Experimental/PCG/Source/PCG/Public/Elements/PCGExecuteBlueprint.h b/Engine/Plugins/Experimental/PCG/Source/PCG/Public/Elements/PCGExecuteBlueprint.h
--- a/Engine/Plugins/Experimental/PCG/Source/PCG/Public/Elements/PCGExecuteBlueprint.h
+++ b/Engine/Plugins/Experimental/PCG/Source/PCG/Public/Elements/PCGExecuteBlueprint.h
@@ -11,6 +11,10 @@
 
 class UWorld;
 
+#if WITH_EDITOR
+DECLARE_MULTICAST_DELEGATE_OneParam(FOnPCGBlueprintChanged, UPCGBlueprintElement*);
+#endif
+
 UCLASS(Abstract, BlueprintType, Blueprintable, hidecategories = (Object))
 class UPCGBlueprintElement : public UObject
 {
@@ -20,8 +24,18 @@ public:
 	UFUNCTION(BlueprintImplementableEvent, BlueprintCallable, Category = Execution)
 	void Execute(const FPCGDataCollection& Input, FPCGDataCollection& Output) const;
 
+#if WITH_EDITOR
+	// ~Begin UObject interface
+	virtual void PostEditChangeProperty(FPropertyChangedEvent& PropertyChangedEvent) override;
+	// ~End UObject interface
+#endif
+
 	/** Needed to be able to call certain blueprint functions */
 	virtual UWorld* GetWorld() const override;
+
+#if WITH_EDITOR
+	FOnPCGBlueprintChanged OnBlueprintChangedDelegate;
+#endif
 };
 
 UCLASS(BlueprintType, ClassGroup = (Procedural))
@@ -42,9 +56,10 @@ protected:
 	// ~End UPCGSettings interface
 
 public:
-#if WITH_EDITOR
 	// ~Begin UObject interface
 	virtual void PostLoad() override;
+	virtual void BeginDestroy() override;
+#if WITH_EDITOR
 	virtual void PostEditChangeProperty(FPropertyChangedEvent& PropertyChangedEvent) override;
 	// ~End UObject interface
 #endif
@@ -63,7 +78,16 @@ protected:
 	TObjectPtr<UPCGBlueprintElement> BlueprintElementInstance;
 
 protected:
+#if WITH_EDITOR
+	void OnBlueprintChanged(UBlueprint* InBlueprint);
+	void OnBlueprintElementChanged(UPCGBlueprintElement* InElement);
+#endif
+
 	void RefreshBlueprintElement();
+	void SetupBlueprintEvent();
+	void TeardownBlueprintEvent();
+	void SetupBlueprintElementEvent();
+	void TeardownBlueprintElementEvent();
 };
 
 class FPCGExecuteBlueprintElement : public FSimpleTypedPCGElement<UPCGBlueprintSettings>
