Fixed issue in PCG blueprint settings begin destroy

diff --git a/Engine/Plugins/Experimental/PCG/Source/PCG/Private/Elements/PCGExecuteBlueprint.cpp b/Engine/Plugins/Experimental/PCG/Source/PCG/Private/Elements/PCGExecuteBlueprint.cpp
--- a/Engine/Plugins/Experimental/PCG/Source/PCG/Private/Elements/PCGExecuteBlueprint.cpp
+++ b/Engine/Plugins/Experimental/PCG/Source/PCG/Private/Elements/PCGExecuteBlueprint.cpp
@@ -95,6 +95,8 @@ void UPCGBlueprintSettings::BeginDestroy()
 {
 	TeardownBlueprintElementEvent();
 	TeardownBlueprintEvent();
+
+	Super::BeginDestroy();
 }
 
 #if WITH_EDITOR
