Added dependency traversal in PCG blueprint elements so we can regenerate when something changes

diff --git a/Engine/Plugins/Experimental/PCG/Source/PCG/Private/Elements/PCGExecuteBlueprint.cpp b/Engine/Plugins/Experimental/PCG/Source/PCG/Private/Elements/PCGExecuteBlueprint.cpp
--- a/Engine/Plugins/Experimental/PCG/Source/PCG/Private/Elements/PCGExecuteBlueprint.cpp
+++ b/Engine/Plugins/Experimental/PCG/Source/PCG/Private/Elements/PCGExecuteBlueprint.cpp
@@ -4,7 +4,113 @@
 
 #if WITH_EDITOR
 #include "Engine/World.h"
-#endif
+
+namespace PCGBlueprintHelper
+{
+	void GatherDependencies(UObject* Object, TSet<TObjectPtr<UObject>>& OutDependencies)
+	{
+		UClass* ObjectClass = Object ? Object->GetClass() : nullptr;
+
+		if (!ObjectClass)
+		{
+			return;
+		}
+
+		for (FProperty* Property = ObjectClass->PropertyLink; Property != nullptr; Property = Property->PropertyLinkNext)
+		{
+			GatherDependencies(Property, Object, OutDependencies);
+		}
+	}
+
+	// Inspired by IteratePropertiesRecursive in ObjectPropertyTrace.cpp
+	void GatherDependencies(FProperty* Property, const void* InContainer, TSet<TObjectPtr<UObject>>& OutDependencies)
+	{
+		auto AddToDependenciesAndGatherRecursively = [&OutDependencies](UObject* Object) {
+			if (Object && !OutDependencies.Contains(Object))
+			{
+				OutDependencies.Add(Object);
+				GatherDependencies(Object, OutDependencies);
+			}
+		};
+
+		if (FObjectProperty* ObjectProperty = CastField<FObjectProperty>(Property))
+		{
+			UObject* Object = ObjectProperty->GetPropertyValue_InContainer(InContainer);
+			AddToDependenciesAndGatherRecursively(Object);
+		}
+		else if (FWeakObjectProperty* WeakObjectProperty = CastField<FWeakObjectProperty>(Property))
+		{
+			FWeakObjectPtr WeakObject = WeakObjectProperty->GetPropertyValue_InContainer(InContainer);
+			AddToDependenciesAndGatherRecursively(WeakObject.Get());
+		}
+		else if (FSoftObjectProperty* SoftObjectProperty = CastField<FSoftObjectProperty>(Property))
+		{
+			FSoftObjectPtr SoftObject = SoftObjectProperty->GetPropertyValue_InContainer(InContainer);
+			AddToDependenciesAndGatherRecursively(SoftObject.Get());
+		}
+		else if (FStructProperty* StructProperty = CastField<FStructProperty>(Property))
+		{
+			const void* StructContainer = StructProperty->ContainerPtrToValuePtr<const void>(InContainer);
+			for (TFieldIterator<FProperty> It(StructProperty->Struct); It; ++It)
+			{
+				GatherDependencies(*It, StructContainer, OutDependencies);
+			}
+		}
+		else if (FArrayProperty* ArrayProperty = CastField<FArrayProperty>(Property))
+		{
+			FScriptArrayHelper_InContainer Helper(ArrayProperty, InContainer);
+			for (int32 DynamicIndex = 0; DynamicIndex < Helper.Num(); ++DynamicIndex)
+			{
+				const void* ValuePtr = Helper.GetRawPtr(DynamicIndex);
+				GatherDependencies(ArrayProperty->Inner, ValuePtr, OutDependencies);
+			}
+		}
+		else if (FMapProperty* MapProperty = CastField<FMapProperty>(Property))
+		{
+			FScriptMapHelper_InContainer Helper(MapProperty, InContainer);
+			int32 Num = Helper.Num();
+			for (int32 DynamicIndex = 0; Num; ++DynamicIndex)
+			{
+				if (Helper.IsValidIndex(DynamicIndex))
+				{
+					const void* KeyPtr = Helper.GetKeyPtr(DynamicIndex);
+					GatherDependencies(MapProperty->KeyProp, KeyPtr, OutDependencies);
+
+					const void* ValuePtr = Helper.GetValuePtr(DynamicIndex);
+					GatherDependencies(MapProperty->ValueProp, ValuePtr, OutDependencies);
+
+					--Num;
+				}
+			}
+		}
+		else if (FSetProperty* SetProperty = CastField<FSetProperty>(Property))
+		{
+			FScriptSetHelper_InContainer Helper(SetProperty, InContainer);
+			int32 Num = Helper.Num();
+			for (int32 DynamicIndex = 0; Num; ++DynamicIndex)
+			{
+				if (Helper.IsValidIndex(DynamicIndex))
+				{
+					const void* ValuePtr = Helper.GetElementPtr(DynamicIndex);
+					GatherDependencies(SetProperty->ElementProp, ValuePtr, OutDependencies);
+
+					--Num;
+				}
+			}
+		}
+	}
+
+	TSet<TObjectPtr<UObject>> GetDataDependencies(UPCGBlueprintElement* InElement)
+	{
+		check(InElement && InElement->GetClass());
+		UClass* BPClass = InElement->GetClass();
+
+		TSet<TObjectPtr<UObject>> Dependencies;
+		GatherDependencies(InElement, Dependencies);
+		return Dependencies;
+	}
+}
+#endif // WITH_EDITOR
 
 UWorld* UPCGBlueprintElement::GetWorld() const
 {
@@ -54,6 +160,7 @@ void UPCGBlueprintSettings::SetupBlueprintElementEvent()
 #if WITH_EDITOR
 	if (BlueprintElementInstance)
 	{
+		DataDependencies = PCGBlueprintHelper::GetDataDependencies(BlueprintElementInstance);
 		BlueprintElementInstance->OnBlueprintChangedDelegate.AddUObject(this, &UPCGBlueprintSettings::OnBlueprintElementChanged);
 	}
 #endif
@@ -65,14 +172,31 @@ void UPCGBlueprintSettings::TeardownBlueprintElementEvent()
 	if (BlueprintElementInstance)
 	{
 		BlueprintElementInstance->OnBlueprintChangedDelegate.RemoveAll(this);
+		DataDependencies.Reset();
 	}
 #endif
 }
 
+#if WITH_EDITOR
+void UPCGBlueprintSettings::OnDependencyChanged(UObject* Object, FPropertyChangedEvent& PropertyChangedEvent)
+{
+	if (!BlueprintElementInstance || PropertyChangedEvent.ChangeType == EPropertyChangeType::Interactive || !DataDependencies.Contains(Object))
+	{
+		return;
+	}
+
+	BlueprintElementInstance->OnBlueprintChangedDelegate.Broadcast(BlueprintElementInstance);
+}
+#endif
+
 void UPCGBlueprintSettings::PostLoad()
 {
 	Super::PostLoad();
 
+#if WITH_EDITOR
+	FCoreUObjectDelegates::OnObjectPropertyChanged.AddUObject(this, &UPCGBlueprintSettings::OnDependencyChanged);
+#endif
+
 	if (BlueprintElement_DEPRECATED && !BlueprintElementType)
 	{
 		BlueprintElementType = BlueprintElement_DEPRECATED;
@@ -96,6 +220,10 @@ void UPCGBlueprintSettings::BeginDestroy()
 	TeardownBlueprintElementEvent();
 	TeardownBlueprintEvent();
 
+#if WITH_EDITOR
+	FCoreUObjectDelegates::OnObjectPropertyChanged.RemoveAll(this);
+#endif
+
 	Super::BeginDestroy();
 }
 
diff --git a/Engine/Plugins/Experimental/PCG/Source/PCG/Public/Elements/PCGExecuteBlueprint.h b/Engine/Plugins/Experimental/PCG/Source/PCG/Public/Elements/PCGExecuteBlueprint.h
--- a/Engine/Plugins/Experimental/PCG/Source/PCG/Public/Elements/PCGExecuteBlueprint.h
+++ b/Engine/Plugins/Experimental/PCG/Source/PCG/Public/Elements/PCGExecuteBlueprint.h
@@ -13,7 +13,14 @@ class UWorld;
 
 #if WITH_EDITOR
 DECLARE_MULTICAST_DELEGATE_OneParam(FOnPCGBlueprintChanged, UPCGBlueprintElement*);
-#endif
+
+namespace PCGBlueprintHelper
+{
+	void GatherDependencies(UObject* Object, TSet<TObjectPtr<UObject>>& OutDependencies);
+	void GatherDependencies(FProperty* Property, const void* InContainer, TSet<TObjectPtr<UObject>>& OutDependencies);
+	TSet<TObjectPtr<UObject>> GetDataDependencies(UPCGBlueprintElement* InElement);
+}
+#endif // WITH_EDITOR
 
 UCLASS(Abstract, BlueprintType, Blueprintable, hidecategories = (Object))
 class UPCGBlueprintElement : public UObject
@@ -80,6 +87,7 @@ protected:
 protected:
 #if WITH_EDITOR
 	void OnBlueprintChanged(UBlueprint* InBlueprint);
+	void OnDependencyChanged(UObject* Object, FPropertyChangedEvent& PropertyChangedEvent);
 	void OnBlueprintElementChanged(UPCGBlueprintElement* InElement);
 #endif
 
@@ -88,6 +96,10 @@ protected:
 	void TeardownBlueprintEvent();
 	void SetupBlueprintElementEvent();
 	void TeardownBlueprintElementEvent();
+
+#if WITH_EDITORONLY_DATA
+	TSet<TObjectPtr<UObject>> DataDependencies;
+#endif
 };
 
 class FPCGExecuteBlueprintElement : public FSimpleTypedPCGElement<UPCGBlueprintSettings>
