UX cleanup - made the PCG world actor collapsed by default

diff --git a/Engine/Plugins/Experimental/PCG/Source/PCG/Private/PCGWorldActor.cpp b/Engine/Plugins/Experimental/PCG/Source/PCG/Private/PCGWorldActor.cpp
--- a/Engine/Plugins/Experimental/PCG/Source/PCG/Private/PCGWorldActor.cpp
+++ b/Engine/Plugins/Experimental/PCG/Source/PCG/Private/PCGWorldActor.cpp
@@ -6,6 +6,16 @@
 
 #include "Engine/World.h"
 
+APCGWorldActor::APCGWorldActor(const FObjectInitializer& ObjectInitializer)
+	: Super(ObjectInitializer)
+{
+	bIsEditorOnlyActor = true;
+#if WITH_EDITORONLY_DATA
+	bIsSpatiallyLoaded = false;
+	bDefaultOutlinerExpansionState = false;
+#endif
+}
+
 void APCGWorldActor::PostLoad()
 {
 	Super::PostLoad();
@@ -23,11 +33,6 @@ APCGWorldActor* APCGWorldActor::CreatePCGWorldActor(UWorld* InWorld)
 {
 	check(InWorld);
 	APCGWorldActor* PCGActor = InWorld->SpawnActor<APCGWorldActor>();
-
-	// This actor should be editor only and not spatially loaded
-	PCGActor->bIsEditorOnlyActor = true;
-	PCGActor->bIsSpatiallyLoaded = false;
-
 	PCGActor->RegisterToSubsystem();
 
 	return PCGActor;
diff --git a/Engine/Plugins/Experimental/PCG/Source/PCG/Public/PCGWorldActor.h b/Engine/Plugins/Experimental/PCG/Source/PCG/Public/PCGWorldActor.h
--- a/Engine/Plugins/Experimental/PCG/Source/PCG/Public/PCGWorldActor.h
+++ b/Engine/Plugins/Experimental/PCG/Source/PCG/Public/PCGWorldActor.h
@@ -13,6 +13,8 @@ class APCGWorldActor : public AActor
 	GENERATED_BODY()
 
 public:
+	APCGWorldActor(const FObjectInitializer& ObjectInitializer);
+
 	//~Begin AActor Interface
 	virtual void PostLoad() override;
 	virtual void BeginDestroy() override;
