Fixed issue where changing a property in the blueprint element instance wouldn't update the callbacks for regen properly

diff --git a/Engine/Plugins/Experimental/PCG/Source/PCG/Private/Elements/PCGExecuteBlueprint.cpp b/Engine/Plugins/Experimental/PCG/Source/PCG/Private/Elements/PCGExecuteBlueprint.cpp
--- a/Engine/Plugins/Experimental/PCG/Source/PCG/Private/Elements/PCGExecuteBlueprint.cpp
+++ b/Engine/Plugins/Experimental/PCG/Source/PCG/Private/Elements/PCGExecuteBlueprint.cpp
@@ -249,7 +249,13 @@ void UPCGBlueprintSettings::OnBlueprintChanged(UBlueprint* InBlueprint)
 
 void UPCGBlueprintSettings::OnBlueprintElementChanged(UPCGBlueprintElement* InElement)
 {
-	OnSettingsChangedDelegate.Broadcast(this);
+	if (InElement == BlueprintElementInstance)
+	{
+		OnSettingsChangedDelegate.Broadcast(this);
+
+		// Since we don't know what property has changed in the element, we need to rebuild our list of dependencies
+		DataDependencies = PCGBlueprintHelper::GetDataDependencies(BlueprintElementInstance);
+	}
 }
 #endif
 
