Added custom tag tracking for pcg regen

diff --git a/Engine/Plugins/Experimental/PCG/Source/PCG/Private/Elements/PCGExecuteBlueprint.cpp b/Engine/Plugins/Experimental/PCG/Source/PCG/Private/Elements/PCGExecuteBlueprint.cpp
--- a/Engine/Plugins/Experimental/PCG/Source/PCG/Private/Elements/PCGExecuteBlueprint.cpp
+++ b/Engine/Plugins/Experimental/PCG/Source/PCG/Private/Elements/PCGExecuteBlueprint.cpp
@@ -289,6 +289,17 @@ void UPCGBlueprintSettings::RefreshBlueprintElement()
 	}	
 }
 
+#if WITH_EDITOR
+TArray<FName> UPCGBlueprintSettings::GetTrackedActorTags() const
+{
+#if WITH_EDITORONLY_DATA
+	return TrackedActorTags;
+#else
+	return Super::GetTrackedActorTags();
+#endif
+}
+#endif
+
 FPCGElementPtr UPCGBlueprintSettings::CreateElement() const
 {
 	return MakeShared<FPCGExecuteBlueprintElement>();
diff --git a/Engine/Plugins/Experimental/PCG/Source/PCG/Private/PCGComponent.cpp b/Engine/Plugins/Experimental/PCG/Source/PCG/Private/PCGComponent.cpp
--- a/Engine/Plugins/Experimental/PCG/Source/PCG/Private/PCGComponent.cpp
+++ b/Engine/Plugins/Experimental/PCG/Source/PCG/Private/PCGComponent.cpp
@@ -580,12 +580,39 @@ bool UPCGComponent::ActorHasExcludedTag(AActor* InActor) const
 	return bHasExcludedTag;
 }
 
+bool UPCGComponent::ActorIsTracked(AActor* InActor) const
+{
+	if (!InActor || !Graph)
+	{
+		return false;
+	}
+
+	TArray<FName> TrackedTags = Graph->GetTrackedActorTags();
+
+	bool bIsTracked = false;
+	for (const FName& Tag : InActor->Tags)
+	{
+		if (TrackedTags.Contains(Tag))
+		{
+			bIsTracked = true;
+			break;
+		}
+	}
+
+	return bIsTracked;
+}
+
 void UPCGComponent::OnActorAdded(AActor* InActor)
 {
 	if (ActorHasExcludedTag(InActor))
 	{
 		UpdateExcludedActor(InActor, EExcludedActorChange::Added);
 	}
+	else if (ActorIsTracked(InActor))
+	{
+		DirtyGenerated();
+		Refresh();
+	}
 }
 
 void UPCGComponent::OnActorDeleted(AActor* InActor)
@@ -594,6 +621,11 @@ void UPCGComponent::OnActorDeleted(AActor* InActor)
 	{
 		UpdateExcludedActor(InActor, EExcludedActorChange::Deleted);
 	}
+	else if (ActorIsTracked(InActor))
+	{
+		DirtyGenerated();
+		Refresh();
+	}
 }
 
 void UPCGComponent::OnActorMoved(AActor* InActor)
@@ -608,9 +640,17 @@ void UPCGComponent::OnActorMoved(AActor* InActor)
 			Refresh();
 		}
 	}
-	else if (ActorHasExcludedTag(InActor))
+	else
 	{
-		UpdateExcludedActor(InActor, EExcludedActorChange::Changed);
+		if (ActorHasExcludedTag(InActor))
+		{
+			UpdateExcludedActor(InActor, EExcludedActorChange::Changed);
+		}
+		else if (ActorIsTracked(InActor))
+		{
+			DirtyGenerated();
+			Refresh();
+		}
 	}
 }
 
@@ -648,10 +688,18 @@ void UPCGComponent::OnObjectPropertyChanged(UObject* InObject, FPropertyChangedE
 
 		Refresh();
 	}
-	else if (ActorHasExcludedTag(Actor) || (bActorTagChange && Actor == InObject))
+	else
 	{
-		// TODO: be more subtle about this, could be an addition, a change or deletion
-		UpdateExcludedActor(Actor, EExcludedActorChange::Changed);
+		if (ActorHasExcludedTag(Actor) || (bActorTagChange && Actor == InObject))
+		{
+			// TODO: be more subtle about this, could be an addition, a change or deletion
+			UpdateExcludedActor(Actor, EExcludedActorChange::Changed);
+		}
+		else if (ActorIsTracked(Actor))
+		{
+			DirtyGenerated();
+			Refresh();
+		}
 	}
 }
 
diff --git a/Engine/Plugins/Experimental/PCG/Source/PCG/Private/PCGGraph.cpp b/Engine/Plugins/Experimental/PCG/Source/PCG/Private/PCGGraph.cpp
--- a/Engine/Plugins/Experimental/PCG/Source/PCG/Private/PCGGraph.cpp
+++ b/Engine/Plugins/Experimental/PCG/Source/PCG/Private/PCGGraph.cpp
@@ -195,6 +195,21 @@ bool UPCGGraph::Contains(UPCGNode* Node) const
 }
 
 #if WITH_EDITOR
+TArray<FName> UPCGGraph::GetTrackedActorTags() const
+{
+	TArray<FName> TrackedActorTags;
+
+	for (UPCGNode* Node : Nodes)
+	{
+		if (Node && Node->DefaultSettings)
+		{
+			TrackedActorTags.Append(Node->DefaultSettings->GetTrackedActorTags());
+		}
+	}
+
+	return TrackedActorTags;
+}
+
 void UPCGGraph::NotifyGraphChanged(bool bIsStructural)
 {
 	if (!bEnableGraphChangeNotifications)
diff --git a/Engine/Plugins/Experimental/PCG/Source/PCG/Private/PCGSubgraph.cpp b/Engine/Plugins/Experimental/PCG/Source/PCG/Private/PCGSubgraph.cpp
--- a/Engine/Plugins/Experimental/PCG/Source/PCG/Private/PCGSubgraph.cpp
+++ b/Engine/Plugins/Experimental/PCG/Source/PCG/Private/PCGSubgraph.cpp
@@ -32,6 +32,13 @@ UPCGNode* UPCGSubgraphSettings::CreateNode() const
 	return NewObject<UPCGSubgraphNode>();
 }
 
+#if WITH_EDITOR
+TArray<FName> UPCGSubgraphSettings::GetTrackedActorTags() const
+{
+	return Subgraph ? Subgraph->GetTrackedActorTags() : TArray<FName>();
+}
+#endif
+
 FPCGElementPtr UPCGSubgraphSettings::CreateElement() const
 {
 	return MakeShared<FPCGTrivialElement>();
diff --git a/Engine/Plugins/Experimental/PCG/Source/PCG/Public/Elements/PCGExecuteBlueprint.h b/Engine/Plugins/Experimental/PCG/Source/PCG/Public/Elements/PCGExecuteBlueprint.h
--- a/Engine/Plugins/Experimental/PCG/Source/PCG/Public/Elements/PCGExecuteBlueprint.h
+++ b/Engine/Plugins/Experimental/PCG/Source/PCG/Public/Elements/PCGExecuteBlueprint.h
@@ -56,6 +56,7 @@ public:
 	// ~Begin UPCGSettings interface
 #if WITH_EDITOR
 	virtual FName GetDefaultNodeName() const override { return FName(TEXT("BlueprintNode")); }
+	virtual TArray<FName> GetTrackedActorTags() const override;
 #endif
 
 protected:
@@ -84,6 +85,11 @@ protected:
 	UPROPERTY(BlueprintReadOnly, VisibleAnywhere, Instanced, Category = Settings, meta = (ShowOnlyInnerProperties))
 	TObjectPtr<UPCGBlueprintElement> BlueprintElementInstance;
 
+#if WITH_EDITORONLY_DATA
+	UPROPERTY(BlueprintReadWrite, EditAnywhere, Category = Settings)
+	TArray<FName> TrackedActorTags;
+#endif
+
 protected:
 #if WITH_EDITOR
 	void OnBlueprintChanged(UBlueprint* InBlueprint);
diff --git a/Engine/Plugins/Experimental/PCG/Source/PCG/Public/PCGComponent.h b/Engine/Plugins/Experimental/PCG/Source/PCG/Public/PCGComponent.h
--- a/Engine/Plugins/Experimental/PCG/Source/PCG/Public/PCGComponent.h
+++ b/Engine/Plugins/Experimental/PCG/Source/PCG/Public/PCGComponent.h
@@ -144,6 +144,7 @@ private:
 	void OnActorMoved(AActor* InActor);
 	void OnObjectPropertyChanged(UObject* InObject, FPropertyChangedEvent& InEvent);
 	bool ActorHasExcludedTag(AActor* InActor) const;
+	bool ActorIsTracked(AActor* InActor) const;
 
 	enum class EExcludedActorChange
 	{
diff --git a/Engine/Plugins/Experimental/PCG/Source/PCG/Public/PCGGraph.h b/Engine/Plugins/Experimental/PCG/Source/PCG/Public/PCGGraph.h
--- a/Engine/Plugins/Experimental/PCG/Source/PCG/Public/PCGGraph.h
+++ b/Engine/Plugins/Experimental/PCG/Source/PCG/Public/PCGGraph.h
@@ -57,6 +57,9 @@ public:
 
 	bool Contains(UPCGNode* Node) const;
 	const TArray<UPCGNode*>& GetNodes() const { return Nodes; }
+#if WITH_EDITOR
+	TArray<FName> GetTrackedActorTags() const;
+#endif
 
 #if WITH_EDITOR
 	FOnPCGGraphChanged OnGraphChangedDelegate;
diff --git a/Engine/Plugins/Experimental/PCG/Source/PCG/Public/PCGSettings.h b/Engine/Plugins/Experimental/PCG/Source/PCG/Public/PCGSettings.h
--- a/Engine/Plugins/Experimental/PCG/Source/PCG/Public/PCGSettings.h
+++ b/Engine/Plugins/Experimental/PCG/Source/PCG/Public/PCGSettings.h
@@ -28,6 +28,8 @@ public:
 
 #if WITH_EDITOR
 	virtual FName GetDefaultNodeName() const { return NAME_None; }
+	/** Derived classes must implement this to communicate dependencies on external actors */
+	virtual TArray<FName> GetTrackedActorTags() const { return TArray<FName>(); }
 #endif
 
 	UPROPERTY(EditAnywhere, BlueprintReadWrite, Category = Settings)
diff --git a/Engine/Plugins/Experimental/PCG/Source/PCG/Public/PCGSubgraph.h b/Engine/Plugins/Experimental/PCG/Source/PCG/Public/PCGSubgraph.h
--- a/Engine/Plugins/Experimental/PCG/Source/PCG/Public/PCGSubgraph.h
+++ b/Engine/Plugins/Experimental/PCG/Source/PCG/Public/PCGSubgraph.h
@@ -19,6 +19,7 @@ public:
 
 #if WITH_EDITOR
 	virtual FName GetDefaultNodeName() const override { return FName(TEXT("SubgraphNode")); }
+	virtual TArray<FName> GetTrackedActorTags() const override;
 #endif
 
 protected:
