Improved excluded & tracked actors behavior in the PCG component: changing a single actor will now dirty only what's needed.

diff --git a/Engine/Plugins/Experimental/PCG/Source/PCG/Private/PCGComponent.cpp b/Engine/Plugins/Experimental/PCG/Source/PCG/Private/PCGComponent.cpp
--- a/Engine/Plugins/Experimental/PCG/Source/PCG/Private/PCGComponent.cpp
+++ b/Engine/Plugins/Experimental/PCG/Source/PCG/Private/PCGComponent.cpp
@@ -78,7 +78,7 @@ void UPCGComponent::SetGraph(UPCGGraph* InGraph)
 		Graph->OnGraphChangedDelegate.AddUObject(this, &UPCGComponent::OnGraphChanged);
 	}
 
-	DirtyGenerated();
+	OnGraphChanged(Graph, true, false);
 #endif	
 }
 
@@ -125,21 +125,32 @@ void UPCGComponent::SetPropertiesFromOriginal(const UPCGComponent* Original)
 
 #if WITH_EDITOR
 	const bool bHasDirtyInput = InputType != NewInputType;
-	const bool bHasDirtyExclusions = ExclusionTags != Original->ExclusionTags;	
+	const bool bHasDirtyExclusions = !(ExcludedTags.Num() == Original->ExcludedTags.Num() && ExcludedTags.Includes(Original->ExcludedTags));
 	const bool bIsDirty = bHasDirtyInput || bHasDirtyExclusions || Seed != Original->Seed || Graph != Original->Graph;
 
-	// Note that while we dirty here, we won't trigger a refresh since we don't have the required context
-	if (bIsDirty)
+	if (bHasDirtyExclusions)
 	{
-		Modify();
-		DirtyGenerated(bHasDirtyInput, bHasDirtyExclusions);
+		TeardownTrackingCallbacks();
+		ExcludedTags = Original->ExcludedTags;
+		SetupTrackingCallbacks();
+		RefreshTrackingData();
 	}
+#else
+	ExcludedTags = Original->ExcludedTags;
 #endif
 
 	InputType = NewInputType;
-	ExclusionTags = Original->ExclusionTags;
 	Seed = Original->Seed;
 	SetGraph(Original->Graph);
+
+#if WITH_EDITOR
+	// Note that while we dirty here, we won't trigger a refresh since we don't have the required context
+	if (bIsDirty)
+	{
+		Modify();
+		DirtyGenerated(bHasDirtyInput);
+	}
+#endif
 }
 
 void UPCGComponent::Generate()
@@ -218,6 +229,39 @@ FPCGTaskId UPCGComponent::GenerateInternal(bool bForce, const TArray<FPCGTaskId>
 	return TaskId;
 }
 
+bool UPCGComponent::GetActorsFromTags(const TSet<FName>& InTags, TSet<TWeakObjectPtr<AActor>>& OutActors)
+{
+	UWorld* World = GetWorld();
+
+	if (!World)
+	{
+		return false;
+	}
+
+	TArray<AActor*> PerTagActors;
+
+	OutActors.Reset();
+
+	bool bHasValidTag = false;
+	for (const FName& Tag : InTags)
+	{
+		if (Tag != NAME_None)
+		{
+			bHasValidTag = true;
+			UGameplayStatics::GetAllActorsWithTag(World, Tag, PerTagActors);
+
+			for (AActor* Actor : PerTagActors)
+			{
+				OutActors.Emplace(Actor);
+			}
+
+			PerTagActors.Reset();
+		}
+	}
+
+	return bHasValidTag;
+}
+
 void UPCGComponent::PostProcessGraph(const FBox& InNewBounds, bool bInGenerated)
 {
 	LastGeneratedBounds = InNewBounds;
@@ -350,9 +394,17 @@ void UPCGComponent::PostLoad()
 {
 	Super::PostLoad();
 
+#if WITH_EDITORONLY_DATA
+	if (!ExclusionTags_DEPRECATED.IsEmpty() && ExcludedTags.IsEmpty())
+	{
+		ExcludedTags.Append(ExclusionTags_DEPRECATED);
+		ExclusionTags_DEPRECATED.Reset();
+	}
+#endif
+
 #if WITH_EDITOR
 	SetupActorCallbacks();
-	SetupExclusionCallbacks();
+	SetupTrackingCallbacks();
 
 	if (Graph)
 	{
@@ -369,7 +421,7 @@ void UPCGComponent::BeginDestroy()
 		Graph->OnGraphChangedDelegate.RemoveAll(this);
 	}
 
-	TeardownExclusionCallbacks();
+	TeardownTrackingCallbacks();
 	TeardownActorCallbacks();
 #endif
 
@@ -387,9 +439,9 @@ void UPCGComponent::PreEditChange(FProperty* PropertyAboutToChange)
 		{
 			Graph->OnGraphChangedDelegate.RemoveAll(this);
 		}
-		else if (PropName == GET_MEMBER_NAME_CHECKED(UPCGComponent, ExclusionTags))
+		else if (PropName == GET_MEMBER_NAME_CHECKED(UPCGComponent, ExcludedTags))
 		{
-			TeardownExclusionCallbacks();
+			TeardownTrackingCallbacks();
 		}
 	}
 
@@ -440,8 +492,7 @@ void UPCGComponent::PostEditChangeProperty(FPropertyChangedEvent& PropertyChange
 			Graph->OnGraphChangedDelegate.AddUObject(this, &UPCGComponent::OnGraphChanged);
 		}
 
-		DirtyGenerated();
-		Refresh();
+		OnGraphChanged(Graph, /*bIsStructural=*/true, /*bShouldRefresh=*/true);
 	}
 	else if (PropName == GET_MEMBER_NAME_CHECKED(UPCGComponent, InputType))
 	{
@@ -454,23 +505,17 @@ void UPCGComponent::PostEditChangeProperty(FPropertyChangedEvent& PropertyChange
 		DirtyGenerated();
 		Refresh();
 	}
-	else if (PropName == GET_MEMBER_NAME_CHECKED(UPCGComponent, ExclusionTags))
+	else if (PropName == GET_MEMBER_NAME_CHECKED(UPCGComponent, ExcludedTags))
 	{
-		SetupExclusionCallbacks();
+		SetupTrackingCallbacks();
+		RefreshTrackingData();
 
-		bool bHadExclusionData = !CachedExclusionData.IsEmpty();
-		bool bHasNonTrivialTags = false;
+		const bool bHadExclusionData = !CachedExclusionData.IsEmpty();
+		const bool bHasExcludedActors = !CachedExcludedActors.IsEmpty();
 
-		for (const FName& ExclusionTag : ExclusionTags)
+		if(bHadExclusionData || bHasExcludedActors)
 		{
-			bHasNonTrivialTags |= ExclusionTag != NAME_None;
-		}
-
-		bool bNeedToDirtyExclusionData = (bHadExclusionData || bHasNonTrivialTags);
-
-		if (bNeedToDirtyExclusionData)
-		{
-			DirtyGenerated(/*bDirtyCachedInput=*/false, bNeedToDirtyExclusionData);
+			DirtyGenerated();
 			Refresh();
 		}
 	}
@@ -484,7 +529,7 @@ void UPCGComponent::PreEditUndo()
 		Graph->OnGraphChangedDelegate.RemoveAll(this);
 	}
 
-	TeardownExclusionCallbacks();
+	TeardownTrackingCallbacks();
 	// Here we will keep a copy of flags that we require to keep through the undo
 	// so we can have a consistent state
 	bWasGeneratedPriorToUndo = bGenerated;
@@ -500,8 +545,9 @@ void UPCGComponent::PostEditUndo()
 		Graph->OnGraphChangedDelegate.AddUObject(this, &UPCGComponent::OnGraphChanged);
 	}
 
-	SetupExclusionCallbacks();
-	DirtyGenerated(/*bDirtyCachedInput=*/true, /*bDirtyCachedExclusions=*/true);
+	SetupTrackingCallbacks();
+	RefreshTrackingData();
+	DirtyGenerated(/*bDirtyCachedInput=*/true);
 	
 	if (bWasGeneratedPriorToUndo && !bGenerated)
 	{
@@ -527,36 +573,31 @@ void UPCGComponent::TeardownActorCallbacks()
 	GEngine->OnActorMoved().RemoveAll(this);
 }
 
-void UPCGComponent::SetupExclusionCallbacks()
+void UPCGComponent::SetupTrackingCallbacks()
 {
-	bool bHasNonTrivialTags = false;
-
-	for (const FName& ExclusionTag : ExclusionTags)
+	CachedTrackedActorTags.Reset();
+	if (Graph)
 	{
-		bHasNonTrivialTags |= ExclusionTag != NAME_None;
+		CachedTrackedActorTags.Append(Graph->GetTrackedActorTags());
 	}
 
-	if (bHasNonTrivialTags)
+	if(!ExcludedTags.IsEmpty() || !CachedTrackedActorTags.IsEmpty())
 	{
 		GEngine->OnLevelActorAdded().AddUObject(this, &UPCGComponent::OnActorAdded);
 		GEngine->OnLevelActorDeleted().AddUObject(this, &UPCGComponent::OnActorDeleted);
 	}
 }
 
-void UPCGComponent::TeardownExclusionCallbacks()
+void UPCGComponent::RefreshTrackingData()
 {
-	bool bHasNonTrivialTags = false;
-
-	for (const FName& ExclusionTag : ExclusionTags)
-	{
-		bHasNonTrivialTags |= ExclusionTag != NAME_None;
-	}
+	GetActorsFromTags(ExcludedTags, CachedExcludedActors);
+	GetActorsFromTags(CachedTrackedActorTags, CachedTrackedActors);
+}
 
-	if (bHasNonTrivialTags)
-	{
-		GEngine->OnLevelActorAdded().RemoveAll(this);
-		GEngine->OnLevelActorDeleted().RemoveAll(this);
-	}
+void UPCGComponent::TeardownTrackingCallbacks()
+{
+	GEngine->OnLevelActorAdded().RemoveAll(this);
+	GEngine->OnLevelActorDeleted().RemoveAll(this);
 }
 
 bool UPCGComponent::ActorHasExcludedTag(AActor* InActor) const
@@ -570,7 +611,7 @@ bool UPCGComponent::ActorHasExcludedTag(AActor* InActor) const
 
 	for (const FName& Tag : InActor->Tags)
 	{
-		if (ExclusionTags.Contains(Tag))
+		if (ExcludedTags.Contains(Tag))
 		{
 			bHasExcludedTag = true;
 			break;
@@ -580,6 +621,14 @@ bool UPCGComponent::ActorHasExcludedTag(AActor* InActor) const
 	return bHasExcludedTag;
 }
 
+void UPCGComponent::DirtyExclusionData(AActor* InActor)
+{
+	if (UPCGData** ExclusionData = CachedExclusionData.Find(InActor))
+	{
+		*ExclusionData = nullptr;
+	}
+}
+
 bool UPCGComponent::ActorIsTracked(AActor* InActor) const
 {
 	if (!InActor || !Graph)
@@ -587,12 +636,10 @@ bool UPCGComponent::ActorIsTracked(AActor* InActor) const
 		return false;
 	}
 
-	TArray<FName> TrackedTags = Graph->GetTrackedActorTags();
-
 	bool bIsTracked = false;
 	for (const FName& Tag : InActor->Tags)
 	{
-		if (TrackedTags.Contains(Tag))
+		if (CachedTrackedActorTags.Contains(Tag))
 		{
 			bIsTracked = true;
 			break;
@@ -604,11 +651,15 @@ bool UPCGComponent::ActorIsTracked(AActor* InActor) const
 
 void UPCGComponent::OnActorAdded(AActor* InActor)
 {
-	if (ActorHasExcludedTag(InActor))
+	const bool bHasExcludedTag = ActorHasExcludedTag(InActor);
+	const bool bIsTracked = ActorIsTracked(InActor);
+
+	if (bIsTracked)
 	{
-		UpdateExcludedActor(InActor, EExcludedActorChange::Added);
+		CachedTrackedActors.Add(InActor);
 	}
-	else if (ActorIsTracked(InActor))
+
+	if (bHasExcludedTag || bIsTracked)
 	{
 		DirtyGenerated();
 		Refresh();
@@ -617,11 +668,15 @@ void UPCGComponent::OnActorAdded(AActor* InActor)
 
 void UPCGComponent::OnActorDeleted(AActor* InActor)
 {
-	if (ActorHasExcludedTag(InActor))
+	const bool bHasExcludedTag = ActorHasExcludedTag(InActor);
+	const bool bIsTracked = ActorIsTracked(InActor);
+
+	if (bIsTracked)
 	{
-		UpdateExcludedActor(InActor, EExcludedActorChange::Deleted);
+		CachedTrackedActors.Remove(InActor);
 	}
-	else if (ActorIsTracked(InActor))
+
+	if (bHasExcludedTag || bIsTracked)
 	{
 		DirtyGenerated();
 		Refresh();
@@ -642,12 +697,9 @@ void UPCGComponent::OnActorMoved(AActor* InActor)
 	}
 	else
 	{
-		if (ActorHasExcludedTag(InActor))
-		{
-			UpdateExcludedActor(InActor, EExcludedActorChange::Changed);
-		}
-		else if (ActorIsTracked(InActor))
+		if (ActorHasExcludedTag(InActor) || ActorIsTracked(InActor))
 		{
+			DirtyExclusionData(InActor);
 			DirtyGenerated();
 			Refresh();
 		}
@@ -690,36 +742,57 @@ void UPCGComponent::OnObjectPropertyChanged(UObject* InObject, FPropertyChangedE
 	}
 	else
 	{
-		if (ActorHasExcludedTag(Actor) || (bActorTagChange && Actor == InObject))
-		{
-			// TODO: be more subtle about this, could be an addition, a change or deletion
-			UpdateExcludedActor(Actor, EExcludedActorChange::Changed);
-		}
-		else if (ActorIsTracked(Actor))
+		if (ActorHasExcludedTag(Actor) || ActorIsTracked(Actor))
 		{
+			DirtyExclusionData(Actor);
 			DirtyGenerated();
 			Refresh();
 		}
+		else if (bActorTagChange && Actor == InObject)
+		{
+			bool bDirtyAndRefresh = false;
+			if (CachedExcludedActors.Remove(Actor))
+			{
+				bDirtyAndRefresh = true;
+			}
+
+			if (CachedTrackedActors.Remove(Actor))
+			{
+				bDirtyAndRefresh = true;
+			}
+
+			if (bDirtyAndRefresh)
+			{
+				DirtyGenerated();
+				Refresh();
+			}
+		}
 	}
 }
 
-void UPCGComponent::UpdateExcludedActor(AActor* InActor, EExcludedActorChange Change)
+void UPCGComponent::OnGraphChanged(UPCGGraph* InGraph, bool bIsStructural)
 {
-	// TODO: manage properly changes
-	DirtyGenerated(/*bInDirtyCachedInput=*/false, /*bInDirtyCachedExclusions=*/true);
-	Refresh();
+	OnGraphChanged(InGraph, bIsStructural, true);
 }
 
-void UPCGComponent::OnGraphChanged(UPCGGraph* InGraph, bool bIsStructural)
+void UPCGComponent::OnGraphChanged(UPCGGraph* InGraph, bool bIsStructural, bool bShouldRefresh)
 {
 	if (InGraph == Graph)
 	{
+		// Since we've changed the graph, we might have changed the tracked actor tags as well
+		TeardownTrackingCallbacks();
+		SetupTrackingCallbacks();
+		RefreshTrackingData();
+
 		DirtyGenerated();
-		Refresh();
+		if (bShouldRefresh)
+		{
+			Refresh();
+		}
 	}
 }
 
-void UPCGComponent::DirtyGenerated(bool bInDirtyCachedInput, bool bInDirtyCachedExclusions)
+void UPCGComponent::DirtyGenerated(bool bInDirtyCachedInput)
 {
 	bDirtyGenerated = true;
 
@@ -729,18 +802,13 @@ void UPCGComponent::DirtyGenerated(bool bInDirtyCachedInput, bool bInDirtyCached
 		CachedActorData = nullptr;
 	}
 
-	if (bInDirtyCachedExclusions)
-	{
-		CachedExclusionData.Reset();
-	}
-
 	// For partitioned graph, we must forward the call to the partition actor
 	// Note that we do not need to forward "normal" dirty as these will be picked up by the local PCG components
 	// However, input changes / moves of the partitioned object will not be caught
 	// It would be possible for partitioned actors to add callbacks to their original component, but that inverses the processing flow
-	if ((bInDirtyCachedInput || bInDirtyCachedExclusions) && bActivated && IsPartitioned())
+	if (bInDirtyCachedInput && bActivated && IsPartitioned())
 	{
-		GetSubsystem()->DirtyGraph(this, LastGeneratedBounds, bInDirtyCachedInput, bInDirtyCachedExclusions);
+		GetSubsystem()->DirtyGraph(this, LastGeneratedBounds, bInDirtyCachedInput);
 	}
 }
 
@@ -803,57 +871,63 @@ UPCGData* UPCGComponent::GetActorPCGData()
 TArray<UPCGData*> UPCGComponent::GetPCGExclusionData()
 {
 	// TODO: replace with a boolean, unify.
-	if (CachedExclusionData.IsEmpty())
-	{
-		CachedExclusionData = CreatePCGExclusionData();
-	}
+	UpdatePCGExclusionData();
 
-	return CachedExclusionData;
+	TArray<UPCGData*> ExclusionData;
+	CachedExclusionData.GenerateValueArray(ExclusionData);
+	return ExclusionData;
 }
 
-TArray<UPCGData*> UPCGComponent::CreatePCGExclusionData()
+void UPCGComponent::UpdatePCGExclusionData()
 {
-	TArray<UPCGData*> ExclusionData;
-
-	// Find all actors with the exclusion tags
-	TSet<AActor*> ExclusionActors;
-	TArray<AActor*> PerTagExclusionActors;
-
-	for (const FName& ExclusionTag : ExclusionTags)
-	{
-		UGameplayStatics::GetAllActorsWithTag(GetWorld(), ExclusionTag, PerTagExclusionActors);
-		ExclusionActors.Append(PerTagExclusionActors);
-		PerTagExclusionActors.SetNum(0);
-	}
-	
-	// Build their PCG data
-	// Now, what we need to do depends on the dimension of the input data
 	const UPCGData* InputData = GetPCGData();
 	const UPCGSpatialData* InputSpatialData = Cast<const UPCGSpatialData>(InputData);
 
-	// Similar in logic to the input data mechanism.
-	for (AActor* ExclusionActor : ExclusionActors)
+	// Update the list of cached excluded actors here, since we might not have picked up everything on map load (due to WP)
+	GetActorsFromTags(ExcludedTags, CachedExcludedActors);
+
+	// Build exclusion data based on the CachedExcludedActors
+	TMap<AActor*, UPCGData*> ExclusionData;
+
+	for(TWeakObjectPtr<AActor> ExcludedActorWeakPtr : CachedExcludedActors)
 	{
-		UPCGData* ActorData = CreateActorPCGData(ExclusionActor);
-		UPCGSpatialData* ActorSpatialData = Cast<UPCGSpatialData>(ActorData);
+		if(!ExcludedActorWeakPtr.IsValid())
+		{
+			continue;
+		}
+
+		AActor* ExcludedActor = ExcludedActorWeakPtr.Get();
 
-		if (InputSpatialData && ActorSpatialData)
+		UPCGData** PreviousExclusionData = CachedExclusionData.Find(ExcludedActor);
+
+		if (PreviousExclusionData && *PreviousExclusionData)
 		{
-			// Create intersection or projection depending on the dimension
-			// TODO: there's an ambiguity here when it's the same dimension.
-			// For volumes, we'd expect an intersection, for surfaces we'd expect a projection
-			if (ActorSpatialData->GetDimension() > InputSpatialData->GetDimension())
-			{
-				ExclusionData.Add(ActorSpatialData->IntersectWith(InputSpatialData));
-			}
-			else
+			ExclusionData.Add(ExcludedActor, *PreviousExclusionData);
+		}
+		else
+		{
+			// Create the new exclusion data
+			UPCGData* ActorData = CreateActorPCGData(ExcludedActor);
+			UPCGSpatialData* ActorSpatialData = Cast<UPCGSpatialData>(ActorData);
+
+			if (InputSpatialData && ActorSpatialData)
 			{
-				ExclusionData.Add(ActorSpatialData->ProjectOn(InputSpatialData));
+				// Create intersection or projection depending on the dimension
+				// TODO: there's an ambiguity here when it's the same dimension.
+				// For volumes, we'd expect an intersection, for surfaces we'd expect a projection
+				if (ActorSpatialData->GetDimension() > InputSpatialData->GetDimension())
+				{
+					ExclusionData.Add(ExcludedActor, ActorSpatialData->IntersectWith(InputSpatialData));
+				}
+				else
+				{
+					ExclusionData.Add(ExcludedActor, ActorSpatialData->ProjectOn(InputSpatialData));
+				}
 			}
 		}
 	}
 
-	return ExclusionData;
+	CachedExclusionData = ExclusionData;
 }
 
 UPCGData* UPCGComponent::CreateActorPCGData()
diff --git a/Engine/Plugins/Experimental/PCG/Source/PCG/Private/PCGSubsystem.cpp b/Engine/Plugins/Experimental/PCG/Source/PCG/Private/PCGSubsystem.cpp
--- a/Engine/Plugins/Experimental/PCG/Source/PCG/Private/PCGSubsystem.cpp
+++ b/Engine/Plugins/Experimental/PCG/Source/PCG/Private/PCGSubsystem.cpp
@@ -377,9 +377,9 @@ void UPCGSubsystem::CleanupGraph(UPCGComponent* Component, const FBox& InBounds,
 	PCGSubsystem::ForEachIntersectingCell(GraphExecutor, Component->GetWorld(), InBounds, /*bCreateActor=*/false, /*bLoadCell=*/false, bSave, ScheduleTask);
 }
 
-void UPCGSubsystem::DirtyGraph(UPCGComponent* Component, const FBox& InBounds, bool bDirtyInputs, bool bDirtyExclusions)
+void UPCGSubsystem::DirtyGraph(UPCGComponent* Component, const FBox& InBounds, bool bDirtyInputs)
 {
-	auto ScheduleTask = [this, Component, bDirtyInputs, bDirtyExclusions](APCGPartitionActor* PCGActor, const FBox& InIntersectedBounds, const TArray<FPCGTaskId>& TaskDependencies) {
+	auto ScheduleTask = [this, Component, bDirtyInputs](APCGPartitionActor* PCGActor, const FBox& InIntersectedBounds, const TArray<FPCGTaskId>& TaskDependencies) {
 		// In the specific case of the dirty, we want to bypass the execution queue, esp. since there's nothing happening here
 		// so we will run the command now, and not delay it.
 		check(Component != nullptr && PCGActor != nullptr);
@@ -391,7 +391,7 @@ void UPCGSubsystem::DirtyGraph(UPCGComponent* Component, const FBox& InBounds, b
 
 		if (UPCGComponent* LocalComponent = PCGActor->GetLocalComponent(Component))
 		{
-			LocalComponent->DirtyGenerated(bDirtyInputs, bDirtyExclusions);
+			LocalComponent->DirtyGenerated(bDirtyInputs);
 		}
 
 		return InvalidTaskId;
diff --git a/Engine/Plugins/Experimental/PCG/Source/PCG/Public/PCGComponent.h b/Engine/Plugins/Experimental/PCG/Source/PCG/Public/PCGComponent.h
--- a/Engine/Plugins/Experimental/PCG/Source/PCG/Public/PCGComponent.h
+++ b/Engine/Plugins/Experimental/PCG/Source/PCG/Public/PCGComponent.h
@@ -73,7 +73,12 @@ public:
 	int Seed = 42;
 
 	UPROPERTY(BlueprintReadWrite, EditAnywhere, Category = Settings)
-	TArray<FName> ExclusionTags;
+	TSet<FName> ExcludedTags;
+
+#if WITH_EDITORONLY_DATA
+	UPROPERTY()
+	TArray<FName> ExclusionTags_DEPRECATED;
+#endif
 
 	UPROPERTY(BlueprintReadWrite, EditAnywhere, Category = Properties)
 	bool bActivated = true;
@@ -101,7 +106,7 @@ public:
 
 #if WITH_EDITOR
 	void Refresh();
-	void DirtyGenerated(bool bInDirtyCachedInput = false, bool bInDirtyCachedExclusions = false);
+	void DirtyGenerated(bool bInDirtyCachedInput = false);
 #endif
 
 	bool IsPartitioned() const;
@@ -118,7 +123,7 @@ private:
 	UPCGData* GetActorPCGData();
 	UPCGData* CreateActorPCGData();
 	UPCGData* CreateActorPCGData(AActor* Actor);
-	TArray<UPCGData*> CreatePCGExclusionData();
+	void UpdatePCGExclusionData();
 
 	UPCGSubsystem* GetSubsystem() const;
 	bool ShouldGenerate(bool bForce = false) const;
@@ -126,18 +131,22 @@ private:
 	void CleanupInternal(bool bRemoveComponents);
 	void PostProcessGraph(const FBox& InNewBounds, bool bInGenerated);
 
+	bool GetActorsFromTags(const TSet<FName>& InTags, TSet<TWeakObjectPtr<AActor>>& OutActors);
+
 #if WITH_EDITOR
 	virtual void PreEditChange(FProperty* PropertyAboutToChange) override;
 	virtual void PostEditChangeProperty(FPropertyChangedEvent& PropertyChangedEvent) override;
 	virtual void PreEditUndo() override;
 	virtual void PostEditUndo() override;
 
+	void OnGraphChanged(UPCGGraph* InGraph, bool bIsStructural, bool bShouldRefresh);
 	void OnGraphChanged(UPCGGraph* InGraph, bool bIsStructural);
 
 	void SetupActorCallbacks();
 	void TeardownActorCallbacks();
-	void SetupExclusionCallbacks();
-	void TeardownExclusionCallbacks();	
+	void SetupTrackingCallbacks();
+	void TeardownTrackingCallbacks();
+	void RefreshTrackingData();
 
 	void OnActorAdded(AActor* InActor);
 	void OnActorDeleted(AActor* InActor);
@@ -145,15 +154,7 @@ private:
 	void OnObjectPropertyChanged(UObject* InObject, FPropertyChangedEvent& InEvent);
 	bool ActorHasExcludedTag(AActor* InActor) const;
 	bool ActorIsTracked(AActor* InActor) const;
-
-	enum class EExcludedActorChange
-	{
-		Added = 0,
-		Deleted,
-		Changed
-	};
-
-	void UpdateExcludedActor(AActor* InActor, EExcludedActorChange Change);
+	void DirtyExclusionData(AActor* InActor);
 #endif
 
 	FBox GetGridBounds() const;
@@ -166,7 +167,10 @@ private:
 	UPCGData* CachedActorData = nullptr;
 
 	UPROPERTY(Transient, NonPIEDuplicateTransient)
-	TArray<UPCGData*> CachedExclusionData;
+	TMap<AActor*, UPCGData*> CachedExclusionData;
+
+	UPROPERTY()
+	TSet<TWeakObjectPtr<AActor>> CachedExcludedActors; 
 
 	UPROPERTY()
 	FBox LastGeneratedBounds = FBox(EForceInit::ForceInit);
@@ -175,5 +179,11 @@ private:
 	bool bIsGenerating = false;
 	bool bWasGeneratedPriorToUndo = false;
 	FBox LastGeneratedBoundsPriorToUndo = FBox(EForceInit::ForceInit);
+	TSet<FName> CachedTrackedActorTags;
+#endif
+
+#if WITH_EDITORONLY_DATA
+	UPROPERTY()
+	TSet<TWeakObjectPtr<AActor>> CachedTrackedActors;
 #endif
 };
\ No newline at end of file
diff --git a/Engine/Plugins/Experimental/PCG/Source/PCG/Public/PCGSubsystem.h b/Engine/Plugins/Experimental/PCG/Source/PCG/Public/PCGSubsystem.h
--- a/Engine/Plugins/Experimental/PCG/Source/PCG/Public/PCGSubsystem.h
+++ b/Engine/Plugins/Experimental/PCG/Source/PCG/Public/PCGSubsystem.h
@@ -58,7 +58,7 @@ public:
 	void CleanupGraph(UPCGComponent* Component, const FBox& InBounds, bool bRemoveComponents, bool bSave);
 
 	/** Immediately dirties the partition actors in the given bounds */
-	void DirtyGraph(UPCGComponent* Component, const FBox& InBounds, bool bDirtyInputs, bool bDirtyExclusions);
+	void DirtyGraph(UPCGComponent* Component, const FBox& InBounds, bool bDirtyInputs);
 
 	/** Partition actors methods */
 	void CleanupPartitionActors(const FBox& InBounds);
