Added pass-through capability (e.g. enable/disable) on all pcg nodes.

diff --git a/Engine/Plugins/Experimental/PCG/Source/PCG/Private/Elements/PCGExecuteBlueprint.cpp b/Engine/Plugins/Experimental/PCG/Source/PCG/Private/Elements/PCGExecuteBlueprint.cpp
--- a/Engine/Plugins/Experimental/PCG/Source/PCG/Private/Elements/PCGExecuteBlueprint.cpp
+++ b/Engine/Plugins/Experimental/PCG/Source/PCG/Private/Elements/PCGExecuteBlueprint.cpp
@@ -305,7 +305,7 @@ FPCGElementPtr UPCGBlueprintSettings::CreateElement() const
 	return MakeShared<FPCGExecuteBlueprintElement>();
 }
 
-bool FPCGExecuteBlueprintElement::Execute(FPCGContextPtr Context) const
+bool FPCGExecuteBlueprintElement::ExecuteInternal(FPCGContextPtr Context) const
 {
 	const UPCGBlueprintSettings* Settings = Context->GetInputSettings<UPCGBlueprintSettings>();
 
diff --git a/Engine/Plugins/Experimental/PCG/Source/PCG/Private/Elements/PCGPointSampler.cpp b/Engine/Plugins/Experimental/PCG/Source/PCG/Private/Elements/PCGPointSampler.cpp
--- a/Engine/Plugins/Experimental/PCG/Source/PCG/Private/Elements/PCGPointSampler.cpp
+++ b/Engine/Plugins/Experimental/PCG/Source/PCG/Private/Elements/PCGPointSampler.cpp
@@ -10,7 +10,7 @@ FPCGElementPtr UPCGPointSamplerSettings::CreateElement() const
 	return MakeShared<FPCGPointSamplerElement>();
 }
 
-bool FPCGPointSamplerElement::Execute(FPCGContextPtr Context) const
+bool FPCGPointSamplerElement::ExecuteInternal(FPCGContextPtr Context) const
 {
 	TRACE_CPUPROFILER_EVENT_SCOPE(FPCGPointSamplerElement::Execute);
 	// TODO: make time-sliced implementation
diff --git a/Engine/Plugins/Experimental/PCG/Source/PCG/Private/Elements/PCGSelfPruning.cpp b/Engine/Plugins/Experimental/PCG/Source/PCG/Private/Elements/PCGSelfPruning.cpp
--- a/Engine/Plugins/Experimental/PCG/Source/PCG/Private/Elements/PCGSelfPruning.cpp
+++ b/Engine/Plugins/Experimental/PCG/Source/PCG/Private/Elements/PCGSelfPruning.cpp
@@ -42,7 +42,7 @@ FPCGElementPtr UPCGSelfPruningSettings::CreateElement() const
 	return MakeShared<FPCGSelfPruningElement>();
 }
 
-bool FPCGSelfPruningElement::Execute(FPCGContextPtr Context) const
+bool FPCGSelfPruningElement::ExecuteInternal(FPCGContextPtr Context) const
 {
 	TRACE_CPUPROFILER_EVENT_SCOPE(FPCGSelfPruningElement::Execute);
 	// TODO: time-sliced implementation
diff --git a/Engine/Plugins/Experimental/PCG/Source/PCG/Private/Elements/PCGStaticMeshSpawner.cpp b/Engine/Plugins/Experimental/PCG/Source/PCG/Private/Elements/PCGStaticMeshSpawner.cpp
--- a/Engine/Plugins/Experimental/PCG/Source/PCG/Private/Elements/PCGStaticMeshSpawner.cpp
+++ b/Engine/Plugins/Experimental/PCG/Source/PCG/Private/Elements/PCGStaticMeshSpawner.cpp
@@ -24,7 +24,7 @@ struct FWeightedMeshAndInstances
 	TArray<FTransform> Instances;
 };
 
-bool FPCGStaticMeshSpawnerElement::Execute(FPCGContextPtr Context) const
+bool FPCGStaticMeshSpawnerElement::ExecuteInternal(FPCGContextPtr Context) const
 {
 	TRACE_CPUPROFILER_EVENT_SCOPE(FPCGStaticMeshSpawnerElement::Execute);
 	// TODO : time-sliced implementation
diff --git a/Engine/Plugins/Experimental/PCG/Source/PCG/Private/Elements/PCGSurfaceSampler.cpp b/Engine/Plugins/Experimental/PCG/Source/PCG/Private/Elements/PCGSurfaceSampler.cpp
--- a/Engine/Plugins/Experimental/PCG/Source/PCG/Private/Elements/PCGSurfaceSampler.cpp
+++ b/Engine/Plugins/Experimental/PCG/Source/PCG/Private/Elements/PCGSurfaceSampler.cpp
@@ -11,7 +11,7 @@ FPCGElementPtr UPCGSurfaceSamplerSettings::CreateElement() const
 	return MakeShared<FPCGSurfaceSamplerElement>();
 }
 
-bool FPCGSurfaceSamplerElement::Execute(FPCGContextPtr Context) const
+bool FPCGSurfaceSamplerElement::ExecuteInternal(FPCGContextPtr Context) const
 {
 	TRACE_CPUPROFILER_EVENT_SCOPE(FPCGSurfaceSamplerElement::Execute);
 	// TODO: time-sliced implementation
diff --git a/Engine/Plugins/Experimental/PCG/Source/PCG/Private/Graph/PCGGraphExecutor.cpp b/Engine/Plugins/Experimental/PCG/Source/PCG/Private/Graph/PCGGraphExecutor.cpp
--- a/Engine/Plugins/Experimental/PCG/Source/PCG/Private/Graph/PCGGraphExecutor.cpp
+++ b/Engine/Plugins/Experimental/PCG/Source/PCG/Private/Graph/PCGGraphExecutor.cpp
@@ -340,7 +340,7 @@ FPCGFetchInputElement::FPCGFetchInputElement(UPCGComponent* InComponent)
 {
 }
 
-bool FPCGFetchInputElement::Execute(FPCGContextPtr Context) const
+bool FPCGFetchInputElement::ExecuteInternal(FPCGContextPtr Context) const
 {
 	// First: any input can be passed through to the output trivially
 	Context->OutputData = Context->InputData;
@@ -371,7 +371,7 @@ FPCGGenericElement::FPCGGenericElement(TFunction<bool()> InOperation)
 {
 }
 
-bool FPCGGenericElement::Execute(FPCGContextPtr Context) const
+bool FPCGGenericElement::ExecuteInternal(FPCGContextPtr Context) const
 {
 	return Operation();
 }
diff --git a/Engine/Plugins/Experimental/PCG/Source/PCG/Private/Graph/PCGGraphExecutor.h b/Engine/Plugins/Experimental/PCG/Source/PCG/Private/Graph/PCGGraphExecutor.h
--- a/Engine/Plugins/Experimental/PCG/Source/PCG/Private/Graph/PCGGraphExecutor.h
+++ b/Engine/Plugins/Experimental/PCG/Source/PCG/Private/Graph/PCGGraphExecutor.h
@@ -104,7 +104,9 @@ class FPCGFetchInputElement : public FSimplePCGElement
 public:
 	FPCGFetchInputElement(UPCGComponent* InComponent);
 
-	virtual bool Execute(FPCGContextPtr Context) const override;
+protected:
+	virtual bool ExecuteInternal(FPCGContextPtr Context) const override;
+
 private:
 	UPCGComponent* Component = nullptr;
 };
@@ -115,9 +117,10 @@ class FPCGGenericElement : public FSimplePCGElement
 public:
 	FPCGGenericElement(TFunction<bool()> InOperation);
 
+protected:
 	// Important note: generic elements must always be run on the main thread
 	// as most of these will impact the editor in some way (loading, unloading, saving)
-	virtual bool Execute(FPCGContextPtr Context) const override;
+	virtual bool ExecuteInternal(FPCGContextPtr Context) const override;
 
 private:
 	TFunction<bool()> Operation;
diff --git a/Engine/Plugins/Experimental/PCG/Source/PCG/Private/PCGElement.cpp b/Engine/Plugins/Experimental/PCG/Source/PCG/Private/PCGElement.cpp
--- a/Engine/Plugins/Experimental/PCG/Source/PCG/Private/PCGElement.cpp
+++ b/Engine/Plugins/Experimental/PCG/Source/PCG/Private/PCGElement.cpp
@@ -2,6 +2,22 @@
 
 #include "PCGElement.h"
 #include "PCGContext.h"
+#include "PCGSettings.h"
+
+bool IPCGElement::Execute(FPCGContextPtr Context) const
+{
+	const UPCGSettings* Settings = Context->GetInputSettings<UPCGSettings>();
+	if (Settings && !Settings->bEnabled)
+	{
+		//Pass-through
+		Context->OutputData = Context->InputData;
+		return true;
+	}
+	else
+	{
+		return ExecuteInternal(Context);
+	}
+}
 
 FPCGContextPtr FSimplePCGElement::Initialize(const FPCGDataCollection& InputData, const UPCGComponent* SourceComponent)
 {
diff --git a/Engine/Plugins/Experimental/PCG/Source/PCG/Private/PCGSettings.cpp b/Engine/Plugins/Experimental/PCG/Source/PCG/Private/PCGSettings.cpp
--- a/Engine/Plugins/Experimental/PCG/Source/PCG/Private/PCGSettings.cpp
+++ b/Engine/Plugins/Experimental/PCG/Source/PCG/Private/PCGSettings.cpp
@@ -30,7 +30,7 @@ FPCGElementPtr UPCGTrivialSettings::CreateElement() const
 	return MakeShared<FPCGTrivialElement>();
 }
 
-bool FPCGTrivialElement::Execute(FPCGContextPtr Context) const
+bool FPCGTrivialElement::ExecuteInternal(FPCGContextPtr Context) const
 {
 	// Pass-through
 	Context->OutputData = Context->InputData;
diff --git a/Engine/Plugins/Experimental/PCG/Source/PCG/Public/Elements/PCGExecuteBlueprint.h b/Engine/Plugins/Experimental/PCG/Source/PCG/Public/Elements/PCGExecuteBlueprint.h
--- a/Engine/Plugins/Experimental/PCG/Source/PCG/Public/Elements/PCGExecuteBlueprint.h
+++ b/Engine/Plugins/Experimental/PCG/Source/PCG/Public/Elements/PCGExecuteBlueprint.h
@@ -110,6 +110,6 @@ protected:
 
 class FPCGExecuteBlueprintElement : public FSimpleTypedPCGElement<UPCGBlueprintSettings>
 {
-public:
-	virtual bool Execute(FPCGContextPtr Context) const override;
+protected:
+	virtual bool ExecuteInternal(FPCGContextPtr Context) const override;
 };
\ No newline at end of file
diff --git a/Engine/Plugins/Experimental/PCG/Source/PCG/Public/Elements/PCGPointSampler.h b/Engine/Plugins/Experimental/PCG/Source/PCG/Public/Elements/PCGPointSampler.h
--- a/Engine/Plugins/Experimental/PCG/Source/PCG/Public/Elements/PCGPointSampler.h
+++ b/Engine/Plugins/Experimental/PCG/Source/PCG/Public/Elements/PCGPointSampler.h
@@ -33,6 +33,6 @@ public:
 
 class FPCGPointSamplerElement : public FSimpleTypedPCGElement<UPCGPointSamplerSettings>
 {
-public:
-	virtual bool Execute(FPCGContextPtr Context) const override;
+protected:
+	virtual bool ExecuteInternal(FPCGContextPtr Context) const override;
 };
\ No newline at end of file
diff --git a/Engine/Plugins/Experimental/PCG/Source/PCG/Public/Elements/PCGSelfPruning.h b/Engine/Plugins/Experimental/PCG/Source/PCG/Public/Elements/PCGSelfPruning.h
--- a/Engine/Plugins/Experimental/PCG/Source/PCG/Public/Elements/PCGSelfPruning.h
+++ b/Engine/Plugins/Experimental/PCG/Source/PCG/Public/Elements/PCGSelfPruning.h
@@ -45,6 +45,6 @@ protected:
 
 class FPCGSelfPruningElement : public FSimpleTypedPCGElement<UPCGSelfPruningSettings>
 {
-public:
-	virtual bool Execute(FPCGContextPtr Context) const override;
+protected:
+	virtual bool ExecuteInternal(FPCGContextPtr Context) const override;
 };
\ No newline at end of file
diff --git a/Engine/Plugins/Experimental/PCG/Source/PCG/Public/Elements/PCGStaticMeshSpawner.h b/Engine/Plugins/Experimental/PCG/Source/PCG/Public/Elements/PCGStaticMeshSpawner.h
--- a/Engine/Plugins/Experimental/PCG/Source/PCG/Public/Elements/PCGStaticMeshSpawner.h
+++ b/Engine/Plugins/Experimental/PCG/Source/PCG/Public/Elements/PCGStaticMeshSpawner.h
@@ -46,8 +46,8 @@ public:
 
 class FPCGStaticMeshSpawnerElement : public FSimpleTypedPCGElement<UPCGStaticMeshSpawnerSettings>
 {
-public:
-	virtual bool Execute(FPCGContextPtr Context) const override;
+protected:
+	virtual bool ExecuteInternal(FPCGContextPtr Context) const override;
 
 private:
 	UInstancedStaticMeshComponent* GetOrCreateISMC(AActor* InActor, const UPCGComponent* SourceComponent, UStaticMesh* InMesh) const;
diff --git a/Engine/Plugins/Experimental/PCG/Source/PCG/Public/Elements/PCGSurfaceSampler.h b/Engine/Plugins/Experimental/PCG/Source/PCG/Public/Elements/PCGSurfaceSampler.h
--- a/Engine/Plugins/Experimental/PCG/Source/PCG/Public/Elements/PCGSurfaceSampler.h
+++ b/Engine/Plugins/Experimental/PCG/Source/PCG/Public/Elements/PCGSurfaceSampler.h
@@ -41,6 +41,6 @@ protected:
 
 class FPCGSurfaceSamplerElement : public FSimpleTypedPCGElement<UPCGSurfaceSamplerSettings>
 {
-public:
-	virtual bool Execute(FPCGContextPtr Context) const override;
+protected:
+	virtual bool ExecuteInternal(FPCGContextPtr Context) const override;
 };
\ No newline at end of file
diff --git a/Engine/Plugins/Experimental/PCG/Source/PCG/Public/PCGElement.h b/Engine/Plugins/Experimental/PCG/Source/PCG/Public/PCGElement.h
--- a/Engine/Plugins/Experimental/PCG/Source/PCG/Public/PCGElement.h
+++ b/Engine/Plugins/Experimental/PCG/Source/PCG/Public/PCGElement.h
@@ -19,7 +19,11 @@ class IPCGElement
 public:
 	virtual ~IPCGElement() = default;
 	virtual FPCGContextPtr Initialize(const FPCGDataCollection& InputData, const UPCGComponent* SourceComponent) = 0;
-	virtual bool Execute(FPCGContextPtr Context) const = 0;
+
+	bool Execute(FPCGContextPtr Context) const;
+
+protected:
+	virtual bool ExecuteInternal(FPCGContextPtr Context) const = 0;
 };
 
 /**
diff --git a/Engine/Plugins/Experimental/PCG/Source/PCG/Public/PCGSettings.h b/Engine/Plugins/Experimental/PCG/Source/PCG/Public/PCGSettings.h
--- a/Engine/Plugins/Experimental/PCG/Source/PCG/Public/PCGSettings.h
+++ b/Engine/Plugins/Experimental/PCG/Source/PCG/Public/PCGSettings.h
@@ -35,6 +35,9 @@ public:
 	UPROPERTY(EditAnywhere, BlueprintReadWrite, Category = Settings)
 	int Seed = 0xC35A9631; // random prime number
 
+	UPROPERTY(EditAnywhere, BlueprintReadWrite, Category = Settings)
+	bool bEnabled = true;
+
 #if WITH_EDITOR
 	FOnPCGSettingsChanged OnSettingsChangedDelegate;
 #endif
@@ -64,6 +67,6 @@ protected:
 
 class FPCGTrivialElement : public FSimplePCGElement
 {
-public:
-	virtual bool Execute(FPCGContextPtr Context) const override;
+protected:
+	virtual bool ExecuteInternal(FPCGContextPtr Context) const override;
 };
\ No newline at end of file
