Added basic debug info capability in PCG.
Supports to add debug data as a post-process in a given node that can also be seen in isolation.

diff --git a/Engine/Plugins/Experimental/PCG/Source/PCG/Private/Elements/PCGDebugElement.cpp b/Engine/Plugins/Experimental/PCG/Source/PCG/Private/Elements/PCGDebugElement.cpp
new file mode 100644
--- /dev/null
+++ b/Engine/Plugins/Experimental/PCG/Source/PCG/Private/Elements/PCGDebugElement.cpp
@@ -0,0 +1,127 @@
+// Copyright Epic Games, Inc. All Rights Reserved.
+
+#include "Elements/PCGDebugElement.h"
+
+#include "PCGHelpers.h"
+#include "Data/PCGPointData.h"
+#include "Helpers/PCGActorHelpers.h"
+
+#include "Components/InstancedStaticMeshComponent.h"
+#include "Engine/StaticMesh.h"
+#include "Materials/MaterialInterface.h"
+
+namespace PCGDebugElement
+{
+	void ExecuteDebugDisplay(FPCGContextPtr Context)
+	{
+		TRACE_CPUPROFILER_EVENT_SCOPE(PCGDebugElement::ExecuteDebugDisplay);
+#if WITH_EDITORONLY_DATA
+		const UPCGSettings* Settings = Context->GetInputSettings<UPCGSettings>();
+
+		if (!Settings)
+		{
+			return;
+		}
+
+		const FPCGDebugVisualizationSettings& DebugSettings = Settings->DebugSettings;
+
+		UStaticMesh* Mesh = DebugSettings.PointMesh.LoadSynchronous();
+
+		if (!Mesh)
+		{
+			return;
+		}
+
+		UMaterialInterface* Material = DebugSettings.Material.LoadSynchronous();
+
+		TArray<FPCGTaggedData> Outputs = Context->OutputData.GetInputs();
+
+		for (const FPCGTaggedData& Output : Outputs)
+		{
+			const UPCGSpatialData* SpatialData = Cast<UPCGSpatialData>(Output.Data);
+
+			if (!SpatialData)
+			{
+				// Data type mismatch
+				continue;
+			}
+
+			AActor* TargetActor = SpatialData->TargetActor;
+
+			if (!TargetActor)
+			{
+				// No target actor
+				continue;
+			}
+
+			const UPCGPointData* PointData = SpatialData->ToPointData();
+
+			if (!PointData)
+			{
+				continue;
+			}
+
+			const TArray<FPCGPoint>& Points = PointData->GetPoints();
+
+			if (Points.Num() == 0)
+			{
+				continue;
+			}
+
+			const int NumCustomData = 8;
+
+			TArray<FTransform> Instances;
+			TArray<float> InstanceCustomData;
+
+			Instances.Reserve(Points.Num());
+			InstanceCustomData.Reserve(NumCustomData);
+
+			// First, create target instance transforms
+			for (const FPCGPoint& Point : Points)
+			{
+				Instances.Add(Point.Transform);
+			}
+
+			UInstancedStaticMeshComponent* ISMC = UPCGActorHelpers::GetOrCreateISMC(TargetActor, Context->SourceComponent, Mesh, { Material });
+			ISMC->ComponentTags.AddUnique(PCGHelpers::DefaultPCGDebugTag);
+			ISMC->NumCustomDataFloats = NumCustomData;
+			const int32 PreExistingInstanceCount = ISMC->GetInstanceCount();
+			ISMC->AddInstances(Instances, false);
+
+			// Then get & assign custom data
+			for (int32 PointIndex = 0; PointIndex < Points.Num(); ++PointIndex)
+			{
+				const FPCGPoint& Point = Points[PointIndex];
+				InstanceCustomData.Add(Point.Density);
+				InstanceCustomData.Add(Point.Extents[0]);
+				InstanceCustomData.Add(Point.Extents[1]);
+				InstanceCustomData.Add(Point.Extents[2]);
+				InstanceCustomData.Add(Point.Color[0]);
+				InstanceCustomData.Add(Point.Color[1]);
+				InstanceCustomData.Add(Point.Color[2]);
+				InstanceCustomData.Add(Point.Color[3]);
+
+				ISMC->SetCustomData(PointIndex + PreExistingInstanceCount, InstanceCustomData);
+
+				InstanceCustomData.Reset();
+			}
+
+			ISMC->UpdateBounds();
+		}
+#endif
+	}
+}
+
+FPCGElementPtr UPCGDebugSettings::CreateElement() const
+{
+	return MakeShared<FPCGDebugElement>();
+}
+
+bool FPCGDebugElement::ExecuteInternal(FPCGContextPtr Context) const
+{
+	TRACE_CPUPROFILER_EVENT_SCOPE(FPCGDebugElement::Execute);
+	PCGDebugElement::ExecuteDebugDisplay(Context);
+	
+	Context->OutputData = Context->InputData;
+	return true;
+}
\ No newline at end of file
diff --git a/Engine/Plugins/Experimental/PCG/Source/PCG/Private/Elements/PCGStaticMeshSpawner.cpp b/Engine/Plugins/Experimental/PCG/Source/PCG/Private/Elements/PCGStaticMeshSpawner.cpp
--- a/Engine/Plugins/Experimental/PCG/Source/PCG/Private/Elements/PCGStaticMeshSpawner.cpp
+++ b/Engine/Plugins/Experimental/PCG/Source/PCG/Private/Elements/PCGStaticMeshSpawner.cpp
@@ -2,13 +2,11 @@
 
 #include "Elements/PCGStaticMeshSpawner.h"
 
-#include "PCGComponent.h"
 #include "PCGHelpers.h"
+#include "Helpers/PCGActorHelpers.h"
 #include "Data/PCGPointData.h"
-#include "Data/PCGSpatialData.h"
 
 #include "Components/InstancedStaticMeshComponent.h"
-#include "Components/HierarchicalInstancedStaticMeshComponent.h"
 #include "Engine/StaticMesh.h"
 #include "Math/RandomStream.h"
 
@@ -43,6 +41,7 @@ bool FPCGStaticMeshSpawnerElement::ExecuteInternal(FPCGContextPtr Context) const
 
 		FWeightedMeshAndInstances& WeightedEntry = WeightedEntries.Emplace_GetRef();
 		WeightedEntry.Weight = TotalWeight;
+		// Todo: we could likely pre-load these meshes asynchronously in the settings
 		WeightedEntry.Mesh = Entry.Mesh.LoadSynchronous();
 	}
 
@@ -114,10 +113,11 @@ bool FPCGStaticMeshSpawnerElement::ExecuteInternal(FPCGContextPtr Context) const
 					continue;
 				}
 
-				UInstancedStaticMeshComponent* ISMC = GetOrCreateISMC(TargetActor, Context->SourceComponent, Instances.Mesh);
+				UInstancedStaticMeshComponent* ISMC = UPCGActorHelpers::GetOrCreateISMC(TargetActor, Context->SourceComponent, Instances.Mesh);
 
 				// TODO: add scaling
 				// TODO: document these arguments
+				ISMC->NumCustomDataFloats = 0;
 				ISMC->AddInstances(Instances.Instances, false, true);
 				ISMC->UpdateBounds();
 			}
@@ -126,42 +126,3 @@ bool FPCGStaticMeshSpawnerElement::ExecuteInternal(FPCGContextPtr Context) const
 
 	return true;
 }
-
-UInstancedStaticMeshComponent* FPCGStaticMeshSpawnerElement::GetOrCreateISMC(AActor* InTargetActor, const UPCGComponent* InSourceComponent, UStaticMesh* InMesh) const
-{
-	check(InTargetActor != nullptr && InMesh != nullptr);
-
-	TArray<UInstancedStaticMeshComponent*> ISMCs;
-	InTargetActor->GetComponents<UInstancedStaticMeshComponent>(ISMCs);
-
-	for (UInstancedStaticMeshComponent* ISMC : ISMCs)
-	{
-		if (ISMC->GetStaticMesh() == InMesh && (!InSourceComponent || ISMC->ComponentTags.Contains(InSourceComponent->GetFName())))
-		{
-			return ISMC;
-		}
-	}
-
-	InTargetActor->Modify();
-
-	// Otherwise, create a new component
-	// TODO: use static mesh component if there's only one instance
-	// TODO: add hism/ism switch or better yet, use a template component
-	UInstancedStaticMeshComponent* ISMC = NewObject<UHierarchicalInstancedStaticMeshComponent>(InTargetActor);
-	ISMC->SetStaticMesh(InMesh);
-	// TODO: add way to do material overrides, maybe on a finer basis as well
-	ISMC->RegisterComponent();
-	InTargetActor->AddInstanceComponent(ISMC);
-	ISMC->SetMobility(EComponentMobility::Static);
-	// TODO: add option for collision, or use a template
-	ISMC->SetCollisionEnabled(ECollisionEnabled::NoCollision);
-	ISMC->AttachToComponent(InTargetActor->GetRootComponent(), FAttachmentTransformRules::KeepWorldTransform);
-
-	if (InSourceComponent)
-	{
-		ISMC->ComponentTags.Add(InSourceComponent->GetFName());
-		ISMC->ComponentTags.Add(PCGHelpers::DefaultPCGTag);
-	}
-
-	return ISMC;
-}
diff --git a/Engine/Plugins/Experimental/PCG/Source/PCG/Private/Graph/PCGGraphExecutor.cpp b/Engine/Plugins/Experimental/PCG/Source/PCG/Private/Graph/PCGGraphExecutor.cpp
--- a/Engine/Plugins/Experimental/PCG/Source/PCG/Private/Graph/PCGGraphExecutor.cpp
+++ b/Engine/Plugins/Experimental/PCG/Source/PCG/Private/Graph/PCGGraphExecutor.cpp
@@ -241,6 +241,7 @@ void FPCGGraphExecutor::BuildTaskInput(const FPCGGraphTask& Task, FPCGDataCollec
 	{
 		check(OutputData.Contains(Id));
 		TaskInput.TaggedData.Append(OutputData[Id].TaggedData);
+		TaskInput.bCancelExecution |= OutputData[Id].bCancelExecution;
 	}
 }
 
diff --git a/Engine/Plugins/Experimental/PCG/Source/PCG/Private/Graph/PCGGraphExecutor.h b/Engine/Plugins/Experimental/PCG/Source/PCG/Private/Graph/PCGGraphExecutor.h
--- a/Engine/Plugins/Experimental/PCG/Source/PCG/Private/Graph/PCGGraphExecutor.h
+++ b/Engine/Plugins/Experimental/PCG/Source/PCG/Private/Graph/PCGGraphExecutor.h
@@ -121,6 +121,7 @@ protected:
 	// Important note: generic elements must always be run on the main thread
 	// as most of these will impact the editor in some way (loading, unloading, saving)
 	virtual bool ExecuteInternal(FPCGContextPtr Context) const override;
+	virtual bool IsCancellable() const { return false; }
 
 private:
 	TFunction<bool()> Operation;
diff --git a/Engine/Plugins/Experimental/PCG/Source/PCG/Private/Helpers/PCGActorHelpers.cpp b/Engine/Plugins/Experimental/PCG/Source/PCG/Private/Helpers/PCGActorHelpers.cpp
new file mode 100644
--- /dev/null
+++ b/Engine/Plugins/Experimental/PCG/Source/PCG/Private/Helpers/PCGActorHelpers.cpp
@@ -0,0 +1,73 @@
+// Copyright Epic Games, Inc. All Rights Reserved.
+
+#include "Helpers/PCGActorHelpers.h"
+
+#include "PCGHelpers.h"
+
+#include "PCGComponent.h"
+#include "Components/InstancedStaticMeshComponent.h"
+#include "Components/HierarchicalInstancedStaticMeshComponent.h"
+#include "Engine/StaticMesh.h"
+
+UInstancedStaticMeshComponent* UPCGActorHelpers::GetOrCreateISMC(AActor* InTargetActor, const UPCGComponent* InSourceComponent, UStaticMesh* InMesh, const TArray<UMaterialInterface*>& InMaterials)
+{
+	check(InTargetActor != nullptr && InMesh != nullptr);
+
+	TArray<UInstancedStaticMeshComponent*> ISMCs;
+	InTargetActor->GetComponents<UInstancedStaticMeshComponent>(ISMCs);
+
+	for (UInstancedStaticMeshComponent* ISMC : ISMCs)
+	{
+		if (ISMC->GetStaticMesh() == InMesh &&
+			(!InSourceComponent || ISMC->ComponentTags.Contains(InSourceComponent->GetFName())))
+		{
+			// Check materials, if provided; if there's none provided but the current ISMC has overrides, skip this
+			bool bMaterialsMatched = (InMaterials.Num() != 0 == ISMC->GetNumOverrideMaterials() != 0);
+
+			const int32 NumMaterialsToCheck = FMath::Min(InMaterials.Num(), ISMC->GetNumMaterials());
+			for (int32 MaterialIndex = 0; MaterialIndex < NumMaterialsToCheck && bMaterialsMatched; ++MaterialIndex)
+			{
+				if (InMaterials[MaterialIndex] && InMaterials[MaterialIndex] != ISMC->GetMaterial(MaterialIndex))
+				{
+					bMaterialsMatched = false;
+				}
+			}
+
+			if (bMaterialsMatched)
+			{
+				return ISMC;
+			}
+		}
+	}
+
+	InTargetActor->Modify();
+
+	// Otherwise, create a new component
+	// TODO: use static mesh component if there's only one instance
+	// TODO: add hism/ism switch or better yet, use a template component
+	UInstancedStaticMeshComponent* ISMC = NewObject<UHierarchicalInstancedStaticMeshComponent>(InTargetActor);
+	ISMC->SetStaticMesh(InMesh);
+
+	// TODO: improve material override mechanisms
+	const int32 NumMaterials = ISMC->GetNumMaterials();
+	for (int32 MaterialIndex = 0; MaterialIndex < NumMaterials; ++MaterialIndex)
+	{
+		ISMC->SetMaterial(MaterialIndex, MaterialIndex < InMaterials.Num() ? InMaterials[MaterialIndex] : nullptr);
+	}
+
+	ISMC->RegisterComponent();
+	InTargetActor->AddInstanceComponent(ISMC);
+	ISMC->SetMobility(EComponentMobility::Static);
+	// TODO: add option for collision, or use a template
+	ISMC->SetCollisionEnabled(ECollisionEnabled::NoCollision);
+	ISMC->AttachToComponent(InTargetActor->GetRootComponent(), FAttachmentTransformRules::KeepWorldTransform);
+
+	if (InSourceComponent)
+	{
+		ISMC->ComponentTags.Add(InSourceComponent->GetFName());
+	}
+	
+	ISMC->ComponentTags.Add(PCGHelpers::DefaultPCGTag);
+
+	return ISMC;
+}
\ No newline at end of file
diff --git a/Engine/Plugins/Experimental/PCG/Source/PCG/Private/PCGComponent.cpp b/Engine/Plugins/Experimental/PCG/Source/PCG/Private/PCGComponent.cpp
--- a/Engine/Plugins/Experimental/PCG/Source/PCG/Private/PCGComponent.cpp
+++ b/Engine/Plugins/Experimental/PCG/Source/PCG/Private/PCGComponent.cpp
@@ -339,7 +339,7 @@ void UPCGComponent::CleanupInternal(bool bRemoveComponents)
 	{
 		if (ISMC->ComponentTags.Contains(GetFName()))
 		{
-			if (bRemoveComponents)
+			if (bRemoveComponents || ISMC->ComponentTags.Contains(PCGHelpers::DefaultPCGDebugTag))
 			{
 				ISMC->DestroyComponent();
 			}
diff --git a/Engine/Plugins/Experimental/PCG/Source/PCG/Private/PCGDebug.cpp b/Engine/Plugins/Experimental/PCG/Source/PCG/Private/PCGDebug.cpp
new file mode 100644
--- /dev/null
+++ b/Engine/Plugins/Experimental/PCG/Source/PCG/Private/PCGDebug.cpp
@@ -0,0 +1,9 @@
+// Copyright Epic Games, Inc. All Rights Reserved.
+
+#include "PCGDebug.h"
+
+FPCGDebugVisualizationSettings::FPCGDebugVisualizationSettings()
+{
+	PointMesh = FSoftObjectPath(TEXT("/Engine/BasicShapes/Cube.Cube"));
+	Material = FSoftObjectPath(TEXT("Material'/Game/PCG/DebugMaterial.DebugMaterial'"));
+}
\ No newline at end of file
diff --git a/Engine/Plugins/Experimental/PCG/Source/PCG/Private/PCGElement.cpp b/Engine/Plugins/Experimental/PCG/Source/PCG/Private/PCGElement.cpp
--- a/Engine/Plugins/Experimental/PCG/Source/PCG/Private/PCGElement.cpp
+++ b/Engine/Plugins/Experimental/PCG/Source/PCG/Private/PCGElement.cpp
@@ -3,11 +3,23 @@
 #include "PCGElement.h"
 #include "PCGContext.h"
 #include "PCGSettings.h"
+#include "Elements/PCGDebugElement.h"
 
 bool IPCGElement::Execute(FPCGContextPtr Context) const
 {
+	// Early out to stop execution
+	if (Context->InputData.bCancelExecution)
+	{
+		Context->OutputData.bCancelExecution = true;
+
+		if (IsCancellable())
+		{
+			return true;
+		}
+	}
+
 	const UPCGSettings* Settings = Context->GetInputSettings<UPCGSettings>();
-	if (Settings && !Settings->bEnabled)
+	if (Settings && Settings->ExecutionMode == EPCGSettingsExecutionMode::Disabled)
 	{
 		//Pass-through
 		Context->OutputData = Context->InputData;
@@ -15,7 +27,22 @@ bool IPCGElement::Execute(FPCGContextPtr Context) const
 	}
 	else
 	{
-		return ExecuteInternal(Context);
+		bool bDone = ExecuteInternal(Context);
+
+#if WITH_EDITOR
+		if (bDone && Settings && (Settings->ExecutionMode == EPCGSettingsExecutionMode::Debug || Settings->ExecutionMode == EPCGSettingsExecutionMode::Isolated))
+		{
+			PCGDebugElement::ExecuteDebugDisplay(Context);
+
+			// Null out the output if this node is executed in isolation
+			if (Settings->ExecutionMode == EPCGSettingsExecutionMode::Isolated)
+			{
+				Context->OutputData.bCancelExecution = true;
+			}
+		}
+#endif
+
+		return bDone;
 	}
 }
 
diff --git a/Engine/Plugins/Experimental/PCG/Source/PCG/Private/PCGHelpers.h b/Engine/Plugins/Experimental/PCG/Source/PCG/Private/PCGHelpers.h
--- a/Engine/Plugins/Experimental/PCG/Source/PCG/Private/PCGHelpers.h
+++ b/Engine/Plugins/Experimental/PCG/Source/PCG/Private/PCGHelpers.h
@@ -14,6 +14,7 @@ namespace PCGHelpers
 {
 	/** Tag that will be added on every component generated through the PCG system */
 	const FName DefaultPCGTag = TEXT("PCG Generated Component");
+	const FName DefaultPCGDebugTag = TEXT("PCG Generated Debug Component");
 
 	int ComputeSeed(int A);
 	int ComputeSeed(int A, int B);
diff --git a/Engine/Plugins/Experimental/PCG/Source/PCG/Private/PCGPoint.cpp b/Engine/Plugins/Experimental/PCG/Source/PCG/Private/PCGPoint.cpp
--- a/Engine/Plugins/Experimental/PCG/Source/PCG/Private/PCGPoint.cpp
+++ b/Engine/Plugins/Experimental/PCG/Source/PCG/Private/PCGPoint.cpp
@@ -3,10 +3,12 @@
 #include "PCGPoint.h"
 
 FPCGPoint::FPCGPoint(const FTransform& InTransform, float InDensity, int32 InSeed)
-	: Transform(InTransform), Density(InDensity), Seed(InSeed)
+	: Transform(InTransform)
+	, Density(InDensity)
+	, Extents(FVector::One())
+	, Color(FVector4::One())
+	, Seed(InSeed)
 {
-	// TODO: this is a placeholder
-	Extents = FVector::One();
 }
 
 FBoxSphereBounds FPCGPoint::GetBounds() const
diff --git a/Engine/Plugins/Experimental/PCG/Source/PCG/Public/Elements/PCGDebugElement.h b/Engine/Plugins/Experimental/PCG/Source/PCG/Public/Elements/PCGDebugElement.h
new file mode 100644
--- /dev/null
+++ b/Engine/Plugins/Experimental/PCG/Source/PCG/Public/Elements/PCGDebugElement.h
@@ -0,0 +1,35 @@
+// Copyright Epic Games, Inc. All Rights Reserved.
+
+#pragma once
+
+#include "CoreMinimal.h"
+#include "PCGSettings.h"
+
+#include "PCGDebugElement.generated.h"
+
+namespace PCGDebugElement
+{
+	void ExecuteDebugDisplay(FPCGContextPtr Context);
+}
+
+UCLASS(BlueprintType, ClassGroup = (Procedural))
+class UPCGDebugSettings : public UPCGSettings
+{
+	GENERATED_BODY()
+
+public:
+#if WITH_EDITOR
+	// ~Begin UPCGSettings interface
+	virtual FName GetDefaultNodeName() const override { return FName(TEXT("DebugNode")); }
+#endif
+
+protected:
+	virtual FPCGElementPtr CreateElement() const override;
+	// ~End UPCGSettings interface
+};
+
+class FPCGDebugElement : public FSimpleTypedPCGElement<UPCGDebugSettings>
+{
+protected:
+	virtual bool ExecuteInternal(FPCGContextPtr Contexxt) const override;
+};
\ No newline at end of file
diff --git a/Engine/Plugins/Experimental/PCG/Source/PCG/Public/Elements/PCGStaticMeshSpawner.h b/Engine/Plugins/Experimental/PCG/Source/PCG/Public/Elements/PCGStaticMeshSpawner.h
--- a/Engine/Plugins/Experimental/PCG/Source/PCG/Public/Elements/PCGStaticMeshSpawner.h
+++ b/Engine/Plugins/Experimental/PCG/Source/PCG/Public/Elements/PCGStaticMeshSpawner.h
@@ -7,9 +7,6 @@
 
 #include "PCGStaticMeshSpawner.generated.h"
 
-class AActor;
-class UInstancedStaticMeshComponent;
-class UPCGComponent;
 class UStaticMesh;
 
 USTRUCT(BlueprintType)
@@ -48,7 +45,4 @@ class FPCGStaticMeshSpawnerElement : public FSimpleTypedPCGElement<UPCGStaticMes
 {
 protected:
 	virtual bool ExecuteInternal(FPCGContextPtr Context) const override;
-
-private:
-	UInstancedStaticMeshComponent* GetOrCreateISMC(AActor* InActor, const UPCGComponent* SourceComponent, UStaticMesh* InMesh) const;
 };
\ No newline at end of file
diff --git a/Engine/Plugins/Experimental/PCG/Source/PCG/Public/Helpers/PCGActorHelpers.h b/Engine/Plugins/Experimental/PCG/Source/PCG/Public/Helpers/PCGActorHelpers.h
new file mode 100644
--- /dev/null
+++ b/Engine/Plugins/Experimental/PCG/Source/PCG/Public/Helpers/PCGActorHelpers.h
@@ -0,0 +1,23 @@
+// Copyright Epic Games, Inc. All Rights Reserved.
+
+#pragma once
+
+#include "CoreMinimal.h"
+#include "Kismet/BlueprintFunctionLibrary.h"
+
+#include "PCGActorHelpers.generated.h"
+
+class AActor;
+class UInstancedStaticMeshComponent;
+class UStaticMesh;
+class UPCGComponent;
+class UMaterialInterface;
+
+UCLASS(BlueprintType)
+class PCG_API UPCGActorHelpers : public UBlueprintFunctionLibrary
+{
+	GENERATED_BODY()
+
+public:
+	static UInstancedStaticMeshComponent* GetOrCreateISMC(AActor* InActor, const UPCGComponent* SourceComponent, UStaticMesh* InMesh, const TArray<UMaterialInterface*>& InMaterials = TArray<UMaterialInterface*>());
+};
\ No newline at end of file
diff --git a/Engine/Plugins/Experimental/PCG/Source/PCG/Public/PCGContext.h b/Engine/Plugins/Experimental/PCG/Source/PCG/Public/PCGContext.h
--- a/Engine/Plugins/Experimental/PCG/Source/PCG/Public/PCGContext.h
+++ b/Engine/Plugins/Experimental/PCG/Source/PCG/Public/PCGContext.h
@@ -16,7 +16,6 @@ struct FPCGContext
 	FPCGDataCollection InputData;
 	FPCGDataCollection OutputData;
 	const UPCGComponent* SourceComponent = nullptr;
-	bool bExecutionSuccessful = false;
 	// TODO: add RNG source
 	// TODO: replace this by a better identification mechanism
 	const UPCGNode* Node = nullptr;
diff --git a/Engine/Plugins/Experimental/PCG/Source/PCG/Public/PCGData.h b/Engine/Plugins/Experimental/PCG/Source/PCG/Public/PCGData.h
--- a/Engine/Plugins/Experimental/PCG/Source/PCG/Public/PCGData.h
+++ b/Engine/Plugins/Experimental/PCG/Source/PCG/Public/PCGData.h
@@ -61,6 +61,9 @@ struct FPCGDataCollection
 
 	UPROPERTY(BlueprintReadWrite, EditAnywhere, Category = Data)
 	TArray<FPCGTaggedData> TaggedData;
+
+	/** This flag is used to cancel further computation or for the debug/isolate feature */
+	bool bCancelExecution = false;
 };
 
 template<typename SettingsType>
diff --git a/Engine/Plugins/Experimental/PCG/Source/PCG/Public/PCGDebug.h b/Engine/Plugins/Experimental/PCG/Source/PCG/Public/PCGDebug.h
new file mode 100644
--- /dev/null
+++ b/Engine/Plugins/Experimental/PCG/Source/PCG/Public/PCGDebug.h
@@ -0,0 +1,25 @@
+// Copyright Epic Games, Inc. All Rights Reserved.
+
+#pragma once
+
+#include "CoreMinimal.h"
+
+#include "PCGDebug.generated.h"
+
+class UStaticMesh;
+class UMaterialInterface;
+
+USTRUCT(BlueprintType)
+struct FPCGDebugVisualizationSettings
+{
+	GENERATED_BODY()
+
+public:
+	FPCGDebugVisualizationSettings();
+
+	UPROPERTY(BlueprintReadWrite, EditAnywhere, Category = Settings)
+	TSoftObjectPtr<UStaticMesh> PointMesh;
+
+	UPROPERTY(BlueprintReadWrite, EditAnywhere, Category = Settings)
+	TSoftObjectPtr<UMaterialInterface> Material;
+};
\ No newline at end of file
diff --git a/Engine/Plugins/Experimental/PCG/Source/PCG/Public/PCGElement.h b/Engine/Plugins/Experimental/PCG/Source/PCG/Public/PCGElement.h
--- a/Engine/Plugins/Experimental/PCG/Source/PCG/Public/PCGElement.h
+++ b/Engine/Plugins/Experimental/PCG/Source/PCG/Public/PCGElement.h
@@ -24,6 +24,7 @@ public:
 
 protected:
 	virtual bool ExecuteInternal(FPCGContextPtr Context) const = 0;
+	virtual bool IsCancellable() const { return true; }
 };
 
 /**
diff --git a/Engine/Plugins/Experimental/PCG/Source/PCG/Public/PCGPoint.h b/Engine/Plugins/Experimental/PCG/Source/PCG/Public/PCGPoint.h
--- a/Engine/Plugins/Experimental/PCG/Source/PCG/Public/PCGPoint.h
+++ b/Engine/Plugins/Experimental/PCG/Source/PCG/Public/PCGPoint.h
@@ -20,11 +20,14 @@ public:
 	UPROPERTY(BlueprintReadWrite, EditAnywhere, Category = Properties)
 	FTransform Transform;
 
+	UPROPERTY(BlueprintReadWrite, EditAnywhere, Category = Properties)
+	float Density;
+
 	UPROPERTY(BlueprintReadWrite, EditAnywhere, Category = Properties)
 	FVector Extents;
 
 	UPROPERTY(BlueprintReadWrite, EditAnywhere, Category = Properties)
-	float Density;
+	FVector4 Color;
 
 	UPROPERTY(BlueprintReadWrite, EditAnywhere, Category = Properties)
 	int32 Seed;
diff --git a/Engine/Plugins/Experimental/PCG/Source/PCG/Public/PCGSettings.h b/Engine/Plugins/Experimental/PCG/Source/PCG/Public/PCGSettings.h
--- a/Engine/Plugins/Experimental/PCG/Source/PCG/Public/PCGSettings.h
+++ b/Engine/Plugins/Experimental/PCG/Source/PCG/Public/PCGSettings.h
@@ -4,11 +4,21 @@
 
 #include "PCGData.h"
 #include "PCGElement.h"
+#include "PCGDebug.h"
 
 #include "PCGSettings.generated.h"
 
 class UPCGNode;
 
+UENUM()
+enum class EPCGSettingsExecutionMode : uint8
+{
+	Enabled,
+	Debug,
+	Isolated,
+	Disabled
+};
+
 #if WITH_EDITOR
 	DECLARE_MULTICAST_DELEGATE_OneParam(FOnPCGSettingsChanged, UPCGSettings*);
 #endif
@@ -36,7 +46,12 @@ public:
 	int Seed = 0xC35A9631; // random prime number
 
 	UPROPERTY(EditAnywhere, BlueprintReadWrite, Category = Settings)
-	bool bEnabled = true;
+	EPCGSettingsExecutionMode ExecutionMode = EPCGSettingsExecutionMode::Enabled;
+
+#if WITH_EDITORONLY_DATA
+	UPROPERTY(Transient, EditAnywhere, BlueprintReadWrite, Category = Debug, meta = (ShowOnlyInnerProperties))
+	FPCGDebugVisualizationSettings DebugSettings;
+#endif
 
 #if WITH_EDITOR
 	FOnPCGSettingsChanged OnSettingsChangedDelegate;
