Added placeholder tag filtering/post-processing tagging

diff --git a/Engine/Plugins/Experimental/PCG/Source/PCG/Private/PCGData.cpp b/Engine/Plugins/Experimental/PCG/Source/PCG/Private/PCGData.cpp
--- a/Engine/Plugins/Experimental/PCG/Source/PCG/Private/PCGData.cpp
+++ b/Engine/Plugins/Experimental/PCG/Source/PCG/Private/PCGData.cpp
@@ -3,6 +3,14 @@
 #include "PCGData.h"
 #include "PCGSettings.h"
 
+bool FPCGTaggedData::operator==(const FPCGTaggedData& Other) const
+{
+	return Data == Other.Data &&
+		Usage == Other.Usage &&
+		Tags.Num() == Other.Tags.Num() &&
+		Tags.Includes(Other.Tags);
+}
+
 TArray<FPCGTaggedData> FPCGDataCollection::GetInputs() const
 {
 	return TaggedData.FilterByPredicate([](const FPCGTaggedData& Data) {
@@ -10,6 +18,13 @@ TArray<FPCGTaggedData> FPCGDataCollection::GetInputs() const
 		});
 }
 
+TArray<FPCGTaggedData> FPCGDataCollection::GetTaggedInputs(const FString& InTag) const
+{
+	return TaggedData.FilterByPredicate([&InTag](const FPCGTaggedData& Data) {
+		return Data.Usage == EPCGDataUsage::Input && Data.Tags.Contains(InTag);
+		});
+}
+
 TArray<FPCGTaggedData> FPCGDataCollection::GetExclusions() const
 {
 	return TaggedData.FilterByPredicate([](const FPCGTaggedData& Data) {
@@ -47,6 +62,11 @@ TArray<FPCGTaggedData> UPCGDataFunctionLibrary::GetInputs(const FPCGDataCollecti
 	return InCollection.GetInputs();
 }
 
+TArray<FPCGTaggedData> UPCGDataFunctionLibrary::GetTaggedInputs(const FPCGDataCollection& InCollection, const FString& InTag)
+{
+	return InCollection.GetTaggedInputs(InTag);
+}
+
 TArray<FPCGTaggedData> UPCGDataFunctionLibrary::GetExclusions(const FPCGDataCollection& InCollection)
 {
 	return InCollection.GetExclusions();
diff --git a/Engine/Plugins/Experimental/PCG/Source/PCG/Private/PCGElement.cpp b/Engine/Plugins/Experimental/PCG/Source/PCG/Private/PCGElement.cpp
--- a/Engine/Plugins/Experimental/PCG/Source/PCG/Private/PCGElement.cpp
+++ b/Engine/Plugins/Experimental/PCG/Source/PCG/Private/PCGElement.cpp
@@ -27,8 +27,53 @@ bool IPCGElement::Execute(FPCGContextPtr Context) const
 	}
 	else
 	{
+		/** TODO - Placeholder feature */
+		TArray<FPCGTaggedData> BypassedTaggedData;
+		if (Settings && !Settings->FilterOnTags.IsEmpty())
+		{
+			// Move any of the inputs that don't have the tags to the outputs as a pass-through
+			// NOTE: this breaks a bit the ordering of inputs, however, there's no obvious way around it
+			TArray<FPCGTaggedData> FilteredTaggedData;
+			for (FPCGTaggedData& TaggedData : Context->InputData.TaggedData)
+			{
+				if (TaggedData.Usage != EPCGDataUsage::Input)
+				{
+					FilteredTaggedData.Add(TaggedData);
+				}
+				else if (TaggedData.Tags.Intersect(Settings->FilterOnTags).IsEmpty())
+				{
+					if (Settings->bPassThroughFilteredOutInputs)
+					{
+						Context->OutputData.TaggedData.Add(TaggedData);
+						BypassedTaggedData.Add(TaggedData);
+					}
+				}
+				else // input has the required tags
+				{
+					FilteredTaggedData.Add(TaggedData);
+				}
+			}
+
+			Context->InputData.TaggedData = FilteredTaggedData;
+		}
+
 		bool bDone = ExecuteInternal(Context);
 
+		/** TODO - Placeholder feature */
+		if (bDone && Settings && !Settings->TagsAppliedOnOutput.IsEmpty())
+		{
+			for (FPCGTaggedData& TaggedData : Context->OutputData.TaggedData)
+			{
+				if (TaggedData.Usage == EPCGDataUsage::Input)
+				{
+					if (!BypassedTaggedData.Contains(TaggedData))
+					{
+						TaggedData.Tags.Append(Settings->TagsAppliedOnOutput);
+					}
+				}
+			}
+		}
+
 #if WITH_EDITOR
 		if (bDone && Settings && (Settings->ExecutionMode == EPCGSettingsExecutionMode::Debug || Settings->ExecutionMode == EPCGSettingsExecutionMode::Isolated))
 		{
diff --git a/Engine/Plugins/Experimental/PCG/Source/PCG/Public/PCGData.h b/Engine/Plugins/Experimental/PCG/Source/PCG/Public/PCGData.h
--- a/Engine/Plugins/Experimental/PCG/Source/PCG/Public/PCGData.h
+++ b/Engine/Plugins/Experimental/PCG/Source/PCG/Public/PCGData.h
@@ -43,6 +43,8 @@ struct FPCGTaggedData
 
 	UPROPERTY(BlueprintReadWrite, EditAnywhere, Category = Data)
 	TSet<FString> Tags;
+
+	bool operator==(const FPCGTaggedData& Other) const;
 };
 
 USTRUCT(BlueprintType)
@@ -51,6 +53,7 @@ struct FPCGDataCollection
 	GENERATED_BODY()
 
 	TArray<FPCGTaggedData> GetInputs() const;
+	TArray<FPCGTaggedData> GetTaggedInputs(const FString& InTag) const;
 	TArray<FPCGTaggedData> GetExclusions() const;
 	TArray<FPCGTaggedData> GetAllSettings() const;
 
@@ -86,6 +89,9 @@ public:
 	UFUNCTION(BlueprintCallable, Category = Data)
 	static TArray<FPCGTaggedData> GetInputs(const FPCGDataCollection& InCollection);
 
+	UFUNCTION(BlueprintCallable, Category = Data)
+	static TArray<FPCGTaggedData> GetTaggedInputs(const FPCGDataCollection& InCollection, const FString& InTag);
+
 	UFUNCTION(BlueprintCallable, Category = Data)
 	static TArray<FPCGTaggedData> GetExclusions(const FPCGDataCollection& InCollection);
 
diff --git a/Engine/Plugins/Experimental/PCG/Source/PCG/Public/PCGSettings.h b/Engine/Plugins/Experimental/PCG/Source/PCG/Public/PCGSettings.h
--- a/Engine/Plugins/Experimental/PCG/Source/PCG/Public/PCGSettings.h
+++ b/Engine/Plugins/Experimental/PCG/Source/PCG/Public/PCGSettings.h
@@ -45,7 +45,19 @@ public:
 	UPROPERTY(EditAnywhere, BlueprintReadWrite, Category = Settings)
 	int Seed = 0xC35A9631; // random prime number
 
-	UPROPERTY(EditAnywhere, BlueprintReadWrite, Category = Settings)
+	/** TODO: Remove this - Placeholder feature until we have a nodegraph */
+	UPROPERTY(EditAnywhere, BlueprintReadWrite, Category = "Settings|Tags")
+	TSet<FString> FilterOnTags;
+
+	/** TODO: Remove this - Placeholder feature until we have a nodegraph */
+	UPROPERTY(EditAnywhere, BlueprintReadWrite, Category = "Settings|Tags")
+	bool bPassThroughFilteredOutInputs = true;
+
+	/** TODO: Remove this - Placeholder feature until we have a nodegraph */
+	UPROPERTY(EditAnywhere, BlueprintReadWrite, Category = "Settings|Tags")
+	TSet<FString> TagsAppliedOnOutput;
+
+	UPROPERTY(EditAnywhere, BlueprintReadWrite, Category = Debug)
 	EPCGSettingsExecutionMode ExecutionMode = EPCGSettingsExecutionMode::Enabled;
 
 #if WITH_EDITORONLY_DATA
