Added operations for projection, union, difference and intersection.

diff --git a/Engine/Plugins/Experimental/PCG/Source/PCG/Private/Elements/PCGDifferenceElement.cpp b/Engine/Plugins/Experimental/PCG/Source/PCG/Private/Elements/PCGDifferenceElement.cpp
new file mode 100644
--- /dev/null
+++ b/Engine/Plugins/Experimental/PCG/Source/PCG/Private/Elements/PCGDifferenceElement.cpp
@@ -0,0 +1,81 @@
+// Copyright Epic Games, Inc. All Rights Reserved.
+
+#include "Elements/PCGDifferenceElement.h"
+
+FPCGElementPtr UPCGDifferenceSettings::CreateElement() const
+{
+	return MakeShared<FPCGDifferenceElement>();
+}
+
+bool FPCGDifferenceElement::ExecuteInternal(FPCGContextPtr Context) const
+{
+	TRACE_CPUPROFILER_EVENT_SCOPE(FPCGDifferenceElement::Execute);
+
+	const UPCGDifferenceSettings* Settings = Context->GetInputSettings<UPCGDifferenceSettings>();
+	check(Settings);
+
+	TArray<FPCGTaggedData> Inputs = Context->InputData.GetInputs();
+	TArray<FPCGTaggedData> Exclusions = Context->InputData.GetExclusions();
+	TArray<FPCGTaggedData>& Outputs = Context->OutputData.TaggedData;
+
+	const UPCGSpatialData* FirstSpatialData = nullptr;
+	UPCGDifferenceData* DifferenceData = nullptr;
+	int32 DifferenceTaggedDataIndex = -1;
+
+	auto AddToDifference = [&FirstSpatialData, &DifferenceData, &DifferenceTaggedDataIndex, Settings, &Outputs](const UPCGSpatialData* SpatialData) {
+		if (!DifferenceData)
+		{
+			DifferenceData = FirstSpatialData->Subtract(SpatialData);
+			DifferenceData->SetDensityFunction(Settings->DensityFunction);
+
+			FPCGTaggedData& DifferenceTaggedData = Outputs[DifferenceTaggedDataIndex];
+			DifferenceTaggedData.Data = DifferenceData;
+		}
+		else
+		{
+			DifferenceData->AddDifference(SpatialData);
+		}
+	};
+
+	// TODO: it might not make sense to perform the difference if the first
+	// data isn't a spatial data, otherwise, what would it really mean?
+	for (FPCGTaggedData& Input : Inputs)
+	{
+		const UPCGSpatialData* SpatialData = Cast<UPCGSpatialData>(Input.Data);
+
+		// Non-spatial data, we're not going to touch
+		if (!SpatialData)
+		{
+			Outputs.Add(Input);
+			continue;
+		}
+
+		if (!FirstSpatialData)
+		{
+			FirstSpatialData = SpatialData;
+			DifferenceTaggedDataIndex = Outputs.Num();
+			Outputs.Add(Input);
+
+			continue;
+		}
+
+		AddToDifference(SpatialData);
+	}
+
+	// Can't have a difference if we don't have an input
+	if (FirstSpatialData)
+	{
+		for (FPCGTaggedData& Exclusion : Exclusions)
+		{
+			if (const UPCGSpatialData* SpatialData = Cast<UPCGSpatialData>(Exclusion.Data))
+			{
+				AddToDifference(SpatialData);
+			}
+		}
+	}
+
+	// Finally, pass-through settings
+	Outputs.Append(Context->InputData.GetAllSettings());
+
+	return true;
+}
\ No newline at end of file
diff --git a/Engine/Plugins/Experimental/PCG/Source/PCG/Private/Elements/PCGIntersectionElement.cpp b/Engine/Plugins/Experimental/PCG/Source/PCG/Private/Elements/PCGIntersectionElement.cpp
new file mode 100644
--- /dev/null
+++ b/Engine/Plugins/Experimental/PCG/Source/PCG/Private/Elements/PCGIntersectionElement.cpp
@@ -0,0 +1,59 @@
+// Copyright Epic Games, Inc. All Rights Reserved.
+
+#include "Elements/PCGIntersectionElement.h"
+
+FPCGElementPtr UPCGIntersectionSettings::CreateElement() const
+{
+	return MakeShared<FPCGIntersectionElement>();
+}
+
+bool FPCGIntersectionElement::ExecuteInternal(FPCGContextPtr Context) const
+{
+	TRACE_CPUPROFILER_EVENT_SCOPE(FPCGIntersectionElement::Execute);
+
+	const UPCGIntersectionSettings* Settings = Context->GetInputSettings<UPCGIntersectionSettings>();
+	check(Settings);
+
+	TArray<FPCGTaggedData> Inputs = Context->InputData.GetInputs();
+	TArray<FPCGTaggedData>& Outputs = Context->OutputData.TaggedData;
+
+	const UPCGSpatialData* FirstSpatialData = nullptr;
+	UPCGIntersectionData* IntersectionData = nullptr;
+	int32 IntersectionTaggedDataIndex = -1;
+
+	for (FPCGTaggedData& Input : Inputs)
+	{
+		const UPCGSpatialData* SpatialData = Cast<UPCGSpatialData>(Input.Data);
+
+		// Non-spatial data we're not going to touch
+		if (!SpatialData)
+		{
+			Outputs.Add(Input);
+			continue;
+		}
+
+		if (!FirstSpatialData)
+		{
+			FirstSpatialData = SpatialData;
+			IntersectionTaggedDataIndex = Outputs.Num();
+			Outputs.Add(Input);
+			continue;
+		}
+
+		// Create a new intersection
+		IntersectionData = (IntersectionData ? IntersectionData : FirstSpatialData)->IntersectWith(SpatialData);
+		// Propagate settings
+		IntersectionData->DensityFunction = Settings->DensityFunction;
+
+		// Update tagged data
+		FPCGTaggedData& IntersectionTaggedData = Outputs[IntersectionTaggedDataIndex];
+		IntersectionTaggedData.Data = IntersectionData;
+		IntersectionTaggedData.Tags.Append(Input.Tags);
+	}
+
+	// Pass-through settings & exclusions
+	Outputs.Append(Context->InputData.GetExclusions());
+	Outputs.Append(Context->InputData.GetAllSettings());
+
+	return true;
+}
\ No newline at end of file
diff --git a/Engine/Plugins/Experimental/PCG/Source/PCG/Private/Elements/PCGProjectionElement.cpp b/Engine/Plugins/Experimental/PCG/Source/PCG/Private/Elements/PCGProjectionElement.cpp
new file mode 100644
--- /dev/null
+++ b/Engine/Plugins/Experimental/PCG/Source/PCG/Private/Elements/PCGProjectionElement.cpp
@@ -0,0 +1,59 @@
+// Copyright Epic Games, Inc. All Rights Reserved.
+
+#include "Elements/PCGProjectionElement.h"
+#include "Data/PCGProjectionData.h"
+
+FPCGElementPtr UPCGProjectionSettings::CreateElement() const
+{
+	return MakeShared<FPCGProjectionElement>();
+}
+
+bool FPCGProjectionElement::ExecuteInternal(FPCGContextPtr Context) const
+{
+	TRACE_CPUPROFILER_EVENT_SCOPE(FPCGProjectionElement::Execute);
+
+	const UPCGProjectionSettings* Settings = Context->GetInputSettings<UPCGProjectionSettings>();
+	check(Settings);
+
+	TArray<FPCGTaggedData> Inputs = Context->InputData.GetInputs();
+	TArray<FPCGTaggedData>& Outputs = Context->OutputData.TaggedData;
+
+	const UPCGSpatialData* FirstSpatialData = nullptr;
+	UPCGProjectionData* ProjectionData = nullptr;
+	int32 ProjectionTaggedDataIndex = -1;
+
+	// TODO: it might not make sense to perform the projection if the first
+	// data isn't a spatial data, otherwise, what would it really mean?
+	for (FPCGTaggedData& Input : Inputs)
+	{
+		const UPCGSpatialData* SpatialData = Cast<UPCGSpatialData>(Input.Data);
+
+		if (!SpatialData)
+		{
+			Outputs.Add(Input);
+			continue;
+		}
+
+		if (!FirstSpatialData)
+		{
+			FirstSpatialData = SpatialData;
+			ProjectionTaggedDataIndex = Outputs.Num();
+			Outputs.Add(Input);
+			continue;
+		}
+
+		// Create a new projection
+		ProjectionData = (ProjectionData ? ProjectionData : FirstSpatialData)->ProjectOn(SpatialData);
+		
+		// Update the tagged data
+		FPCGTaggedData& ProjectionTaggedData = Outputs[ProjectionTaggedDataIndex];
+		ProjectionTaggedData.Data = ProjectionData;
+		ProjectionTaggedData.Tags.Append(Input.Tags);
+	}
+
+	// Pass-through exclusions/settings
+	Outputs.Append(Context->InputData.GetExclusions());
+	Outputs.Append(Context->InputData.GetAllSettings());
+
+	return true;
+}
\ No newline at end of file
diff --git a/Engine/Plugins/Experimental/PCG/Source/PCG/Private/Elements/PCGUnionElement.cpp b/Engine/Plugins/Experimental/PCG/Source/PCG/Private/Elements/PCGUnionElement.cpp
new file mode 100644
--- /dev/null
+++ b/Engine/Plugins/Experimental/PCG/Source/PCG/Private/Elements/PCGUnionElement.cpp
@@ -0,0 +1,68 @@
+// Copyright Epic Games, Inc. All Rights Reserved.
+
+#include "Elements/PCGUnionElement.h"
+
+FPCGElementPtr UPCGUnionSettings::CreateElement() const
+{
+	return MakeShared<FPCGUnionElement>();
+}
+
+bool FPCGUnionElement::ExecuteInternal(FPCGContextPtr Context) const
+{
+	TRACE_CPUPROFILER_EVENT_SCOPE(FPCGUnionElement::Execute);
+
+	const UPCGUnionSettings* Settings = Context->GetInputSettings<UPCGUnionSettings>();
+	check(Settings);
+
+	TArray<FPCGTaggedData> Inputs = Context->InputData.GetInputs();
+	TArray<FPCGTaggedData>& Outputs = Context->OutputData.TaggedData;
+
+	const UPCGSpatialData* FirstSpatialData = nullptr;
+	UPCGUnionData* UnionData = nullptr;
+	int32 UnionTaggedDataIndex = -1;
+
+	for (FPCGTaggedData& Input : Inputs)
+	{
+		const UPCGSpatialData* SpatialData = Cast<UPCGSpatialData>(Input.Data);
+
+		// Non-spatial data we're not going to touch
+		if (!SpatialData)
+		{
+			Outputs.Add(Input);
+			continue;
+		}
+
+		if (!FirstSpatialData)
+		{
+			FirstSpatialData = SpatialData;
+			UnionTaggedDataIndex = Outputs.Num();
+			Outputs.Add(Input);
+			continue;
+		}
+
+		FPCGTaggedData& UnionTaggedData = Outputs[UnionTaggedDataIndex];
+
+		// Create union or add to it
+		if (!UnionData)
+		{
+			UnionData = FirstSpatialData->UnionWith(SpatialData);
+			UnionData->SetType(Settings->Type);
+			UnionData->SetDensityFunction(Settings->DensityFunction);
+
+			UnionTaggedData.Data = UnionData;
+		}
+		else
+		{
+			UnionData->AddData(SpatialData);
+			UnionTaggedData.Tags.Append(Input.Tags);
+		}
+		
+		UnionTaggedData.Data = UnionData;
+	}
+
+	// Finally, pass-through settings
+	Outputs.Append(Context->InputData.GetExclusions());
+	Outputs.Append(Context->InputData.GetAllSettings());
+
+	return true;
+}
\ No newline at end of file
diff --git a/Engine/Plugins/Experimental/PCG/Source/PCG/Public/Data/PCGIntersectionData.h b/Engine/Plugins/Experimental/PCG/Source/PCG/Public/Data/PCGIntersectionData.h
--- a/Engine/Plugins/Experimental/PCG/Source/PCG/Public/Data/PCGIntersectionData.h
+++ b/Engine/Plugins/Experimental/PCG/Source/PCG/Public/Data/PCGIntersectionData.h
@@ -38,6 +38,10 @@ public:
 	virtual const UPCGPointData* CreatePointData() const override;
 	//~End UPCGSpatialDataWithPointCache interface
 
+public:
+	UPROPERTY(BlueprintReadWrite, EditAnywhere, Category = Settings)
+	EPCGIntersectionDensityFunction DensityFunction = EPCGIntersectionDensityFunction::Multiply;
+
 protected:
 	UPCGPointData* CreateAndFilterPointData(const UPCGSpatialData* X, const UPCGSpatialData* Y) const;
 
@@ -47,9 +51,6 @@ protected:
 	UPROPERTY(BlueprintReadOnly, VisibleAnywhere, Category = SpatialData)
 	TObjectPtr<const UPCGSpatialData> B = nullptr;
 
-	UPROPERTY(BlueprintReadWrite, EditAnywhere, Category = Settings)
-	EPCGIntersectionDensityFunction DensityFunction = EPCGIntersectionDensityFunction::Multiply;
-
 	UPROPERTY()
 	FBox CachedBounds = FBox(EForceInit::ForceInit);
 
diff --git a/Engine/Plugins/Experimental/PCG/Source/PCG/Public/Elements/PCGDifferenceElement.h b/Engine/Plugins/Experimental/PCG/Source/PCG/Public/Elements/PCGDifferenceElement.h
new file mode 100644
--- /dev/null
+++ b/Engine/Plugins/Experimental/PCG/Source/PCG/Public/Elements/PCGDifferenceElement.h
@@ -0,0 +1,36 @@
+// Copyright Epic Games, Inc. All Rights Reserved.
+
+#pragma once
+
+#include "CoreMinimal.h"
+#include "PCGElement.h"
+#include "PCGSettings.h"
+#include "Data/PCGDifferenceData.h"
+
+#include "PCGDifferenceElement.generated.h"
+
+UCLASS(BlueprintType, ClassGroup = (Procedural))
+class UPCGDifferenceSettings : public UPCGSettings
+{
+	GENERATED_BODY()
+
+public:
+	//~Begin UPCGSettings interface
+#if WITH_EDITOR
+	virtual FName GetDefaultNodeName() const override { return FName(TEXT("DifferenceNode")); }
+#endif
+
+protected:
+	virtual FPCGElementPtr CreateElement() const override;
+	//~End UPCGSettings interface
+
+public:
+	UPROPERTY(BlueprintReadWrite, EditAnywhere, Category = Settings)
+	EPCGDifferenceDensityFunction DensityFunction = EPCGDifferenceDensityFunction::Minimum;
+};
+
+class FPCGDifferenceElement : public FSimpleTypedPCGElement<UPCGDifferenceSettings>
+{
+protected:
+	virtual bool ExecuteInternal(FPCGContextPtr Context) const;
+};
\ No newline at end of file
diff --git a/Engine/Plugins/Experimental/PCG/Source/PCG/Public/Elements/PCGIntersectionElement.h b/Engine/Plugins/Experimental/PCG/Source/PCG/Public/Elements/PCGIntersectionElement.h
new file mode 100644
--- /dev/null
+++ b/Engine/Plugins/Experimental/PCG/Source/PCG/Public/Elements/PCGIntersectionElement.h
@@ -0,0 +1,36 @@
+// Copyright Epic Games, Inc. All Rights Reserved.
+
+#pragma once
+
+#include "CoreMinimal.h"
+#include "PCGElement.h"
+#include "PCGSettings.h"
+#include "Data/PCGIntersectionData.h"
+
+#include "PCGIntersectionElement.generated.h"
+
+UCLASS(BlueprintType, ClassGroup = (Procedural))
+class UPCGIntersectionSettings : public UPCGSettings
+{
+	GENERATED_BODY()
+
+public:
+	//~Begin UPCGSettings interface
+#if WITH_EDITOR
+	virtual FName GetDefaultNodeName() const override { return FName(TEXT("IntersectionNode")); }
+#endif
+
+protected:
+	virtual FPCGElementPtr CreateElement() const override;
+	//~End UPCGSettings interface
+
+public:
+	UPROPERTY(BlueprintReadWrite, EditAnywhere, Category = Settings)
+	EPCGIntersectionDensityFunction DensityFunction = EPCGIntersectionDensityFunction::Multiply;
+};
+
+class FPCGIntersectionElement : public FSimpleTypedPCGElement<UPCGIntersectionSettings>
+{
+protected:
+	virtual bool ExecuteInternal(FPCGContextPtr Context) const;
+};
\ No newline at end of file
diff --git a/Engine/Plugins/Experimental/PCG/Source/PCG/Public/Elements/PCGProjectionElement.h b/Engine/Plugins/Experimental/PCG/Source/PCG/Public/Elements/PCGProjectionElement.h
new file mode 100644
--- /dev/null
+++ b/Engine/Plugins/Experimental/PCG/Source/PCG/Public/Elements/PCGProjectionElement.h
@@ -0,0 +1,31 @@
+// Copyright Epic Games, Inc. All Rights Reserved.
+
+#pragma once
+
+#include "CoreMinimal.h"
+#include "PCGElement.h"
+#include "PCGSettings.h"
+
+#include "PCGProjectionElement.generated.h"
+
+UCLASS(BlueprintType, ClassGroup = (Procedural))
+class UPCGProjectionSettings : public UPCGSettings
+{
+	GENERATED_BODY()
+
+public:
+	//~Begin UPCGSettings interface
+#if WITH_EDITOR
+	virtual FName GetDefaultNodeName() const override { return FName(TEXT("ProjectionNode")); }
+#endif
+
+protected:
+	virtual FPCGElementPtr CreateElement() const override;
+	//~End UPCGSettings interface
+};
+
+class FPCGProjectionElement : public FSimpleTypedPCGElement<UPCGProjectionSettings>
+{
+protected:
+	virtual bool ExecuteInternal(FPCGContextPtr Context) const;
+};
\ No newline at end of file
diff --git a/Engine/Plugins/Experimental/PCG/Source/PCG/Public/Elements/PCGUnionElement.h b/Engine/Plugins/Experimental/PCG/Source/PCG/Public/Elements/PCGUnionElement.h
new file mode 100644
--- /dev/null
+++ b/Engine/Plugins/Experimental/PCG/Source/PCG/Public/Elements/PCGUnionElement.h
@@ -0,0 +1,39 @@
+// Copyright Epic Games, Inc. All Rights Reserved.
+
+#pragma once
+
+#include "CoreMinimal.h"
+#include "PCGElement.h"
+#include "PCGSettings.h"
+#include "Data/PCGUnionData.h"
+
+#include "PCGUnionElement.generated.h"
+
+UCLASS(BlueprintType, ClassGroup = (Procedural))
+class UPCGUnionSettings : public UPCGSettings
+{
+	GENERATED_BODY()
+
+public:
+	//~Begin UPCGSettings interface
+#if WITH_EDITOR
+	virtual FName GetDefaultNodeName() const override { return FName(TEXT("UnionNode")); }
+#endif
+
+protected:
+	virtual FPCGElementPtr CreateElement() const override;
+	//~End UPCGSettings interface
+
+public:
+	UPROPERTY(BlueprintReadWrite, EditAnywhere, Category = Settings)
+	EPCGUnionType Type = EPCGUnionType::LeftToRightPriority;
+
+	UPROPERTY(BlueprintReadWrite, EditAnywhere, Category = Settings)
+	EPCGUnionDensityFunction DensityFunction = EPCGUnionDensityFunction::Maximum;
+};
+
+class FPCGUnionElement : public FSimpleTypedPCGElement<UPCGUnionSettings>
+{
+protected:
+	virtual bool ExecuteInternal(FPCGContextPtr Context) const;
+};
\ No newline at end of file
