Allowed subgraphs to have no inputs/outputs

diff --git a/Engine/Plugins/Experimental/PCG/Source/PCG/Private/Graph/PCGGraphCompiler.cpp b/Engine/Plugins/Experimental/PCG/Source/PCG/Private/Graph/PCGGraphCompiler.cpp
--- a/Engine/Plugins/Experimental/PCG/Source/PCG/Private/Graph/PCGGraphCompiler.cpp
+++ b/Engine/Plugins/Experimental/PCG/Source/PCG/Private/Graph/PCGGraphCompiler.cpp
@@ -55,12 +55,10 @@ TArray<FPCGGraphTask> FPCGGraphCompiler::CompileGraph(UPCGGraph* InGraph, FPCGTa
 			FPCGGraphTask* InputNodeTask = Subtasks.FindByPredicate([SubgraphInputNode](const FPCGGraphTask& Subtask) {
 				return Subtask.Node == SubgraphInputNode;
 				});
-			check(InputNodeTask);
 
 			FPCGGraphTask* OutputNodeTask = Subtasks.FindByPredicate([SubgraphOutputNode](const FPCGGraphTask& Subtask) {
 				return Subtask.Node == SubgraphOutputNode;
 				});
-			check(OutputNodeTask);
 
 			// Build pre-task
 			FPCGGraphTask& PreTask = CompiledTasks.Emplace_GetRef();
@@ -73,7 +71,10 @@ TArray<FPCGGraphTask> FPCGGraphCompiler::CompileGraph(UPCGGraph* InGraph, FPCGTa
 			}
 
 			// Add pre-task as input to subgraph input node task
-			InputNodeTask->Inputs.Add(PreId);
+			if (InputNodeTask)
+			{
+				InputNodeTask->Inputs.Add(PreId);
+			}
 
 			// Merge subgraph tasks into current tasks.
 			CompiledTasks.Append(Subtasks);
@@ -85,9 +86,21 @@ TArray<FPCGGraphTask> FPCGGraphCompiler::CompileGraph(UPCGGraph* InGraph, FPCGTa
 			PostTask.NodeId = PostId;
 
 			// Add subgraph output node task as input to the post-task
-			PostTask.Inputs.Add(OutputNodeTask->NodeId);
+			if (OutputNodeTask)
+			{
+				PostTask.Inputs.Add(OutputNodeTask->NodeId);
+			}
 
 			IdMapping.Add(Node, PostId);
+
+			// Prime any input-less nodes from the subgraph
+			for (const UPCGNode* SubNode : SubgraphNode->GetGraph()->GetNodes())
+			{
+				if (SubNode->InboundNodes.IsEmpty())
+				{
+					NodeQueue.Add(SubNode);
+				}
+			}
 		}
 		else
 		{
