Fixed issue where graph wasn't recompiled when changing the subgraph in subgraph node/settings.
Reset settings on the input/output nodes on initialize from template.

diff --git a/Engine/Plugins/Experimental/PCG/Source/PCG/Private/PCGGraph.cpp b/Engine/Plugins/Experimental/PCG/Source/PCG/Private/PCGGraph.cpp
--- a/Engine/Plugins/Experimental/PCG/Source/PCG/Private/PCGGraph.cpp
+++ b/Engine/Plugins/Experimental/PCG/Source/PCG/Private/PCGGraph.cpp
@@ -1,6 +1,7 @@
 // Copyright Epic Games, Inc. All Rights Reserved.
 
 #include "PCGGraph.h"
+#include "PCGSubgraph.h"
 #include "PCGSubsystem.h"
 
 #if WITH_EDITOR
@@ -27,10 +28,7 @@ void UPCGGraph::PostLoad()
 #if WITH_EDITOR
 	for (UPCGNode* Node : Nodes)
 	{
-		if (Node)
-		{
-			Node->OnNodeSettingsChangedDelegate.AddUObject(this, &UPCGGraph::OnSettingsChanged);
-		}
+		OnNodeAdded(Node);
 	}
 #endif
 }
@@ -40,10 +38,7 @@ void UPCGGraph::BeginDestroy()
 #if WITH_EDITOR
 	for (UPCGNode* Node : Nodes)
 	{
-		if (Node)
-		{
-			Node->OnNodeSettingsChangedDelegate.RemoveAll(this);
-		}
+		OnNodeRemoved(Node);
 	}
 
 	// Notify the compiler to remove this graph from its cache
@@ -63,15 +58,27 @@ void UPCGGraph::InitializeFromTemplate()
 {
 	Modify();
 
-	InputNode->OutboundNodes.Reset();
-	OutputNode->InboundNodes.Reset();
+	auto ResetDefaultNode = [](UPCGNode* InNode, bool bIsInput) {
+		check(InNode);
+		if (bIsInput)
+		{
+			InNode->OutboundNodes.Reset();
+		}
+		else
+		{
+			InNode->InboundNodes.Reset();
+		}
+
+		// Reset settings as well
+		InNode->DefaultSettings = NewObject<UPCGTrivialSettings>(InNode);
+	};
+
+	ResetDefaultNode(InputNode, true);
+	ResetDefaultNode(OutputNode, false);
 
 	for (UPCGNode* Node : Nodes)
 	{
-		if (Node)
-		{
-			Node->OnNodeSettingsChangedDelegate.RemoveAll(this);
-		}
+		OnNodeRemoved(Node);
 	}
 
 	Nodes.Reset();
@@ -120,11 +127,7 @@ UPCGNode* UPCGGraph::AddNodeOfType(TSubclassOf<class UPCGSettings> InSettingsCla
 #endif
 
 		Nodes.Add(Node);
-
-#if WITH_EDITOR
-		Node->OnNodeSettingsChangedDelegate.AddUObject(this, &UPCGGraph::OnSettingsChanged);
-		NotifyGraphChanged(/*bIsStructural=*/true);
-#endif
+		OnNodeAdded(Node);
 	}
 
 	return Node;
@@ -159,14 +162,39 @@ UPCGNode* UPCGGraph::AddNode(UPCGSettings* InSettings)
 #endif
 
 		Nodes.Add(Node);
+		OnNodeAdded(Node);
+	}
 
+	return Node;
+}
+
+void UPCGGraph::OnNodeAdded(UPCGNode* InNode)
+{
 #if WITH_EDITOR
-		Node->OnNodeSettingsChangedDelegate.AddUObject(this, &UPCGGraph::OnSettingsChanged);
-		NotifyGraphChanged(/*bIsStructural=*/true);
-#endif
+	InNode->OnNodeSettingsChangedDelegate.AddUObject(this, &UPCGGraph::OnSettingsChanged);
+
+	if (UPCGSubgraphNode* SubgraphNode = Cast<UPCGSubgraphNode>(InNode))
+	{
+		SubgraphNode->OnNodeStructuralSettingsChangedDelegate.AddUObject(this, &UPCGGraph::OnStructuralSettingsChanged);
 	}
 
-	return Node;
+	NotifyGraphChanged(/*bIsStructural=*/true);
+#endif
+}
+
+void UPCGGraph::OnNodeRemoved(UPCGNode* InNode)
+{
+#if WITH_EDITOR
+	if (InNode)
+	{
+		InNode->OnNodeSettingsChangedDelegate.RemoveAll(this);
+		
+		if (UPCGSubgraphNode* SubgraphNode = Cast<UPCGSubgraphNode>(InNode))
+		{
+			SubgraphNode->OnNodeStructuralSettingsChangedDelegate.RemoveAll(this);
+		}
+	}
+#endif
 }
 
 UPCGNode* UPCGGraph::AddEdge(UPCGNode* From, UPCGNode* To)
@@ -235,4 +263,9 @@ void UPCGGraph::OnSettingsChanged(UPCGNode* InNode)
 {
 	NotifyGraphChanged(/*bIsStructural=*/false);
 }
+
+void UPCGGraph::OnStructuralSettingsChanged(UPCGNode* InNode)
+{
+	NotifyGraphChanged(/*bIsStructural=*/true);
+}
 #endif // WITH_EDITOR
diff --git a/Engine/Plugins/Experimental/PCG/Source/PCG/Private/PCGSubgraph.cpp b/Engine/Plugins/Experimental/PCG/Source/PCG/Private/PCGSubgraph.cpp
--- a/Engine/Plugins/Experimental/PCG/Source/PCG/Private/PCGSubgraph.cpp
+++ b/Engine/Plugins/Experimental/PCG/Source/PCG/Private/PCGSubgraph.cpp
@@ -44,12 +44,6 @@ FPCGElementPtr UPCGSubgraphSettings::CreateElement() const
 	return MakeShared<FPCGTrivialElement>();
 }
 
-TObjectPtr<UPCGGraph> UPCGSubgraphNode::GetGraph() const
-{
-	TObjectPtr<UPCGSubgraphSettings> Settings = Cast<UPCGSubgraphSettings>(DefaultSettings);
-	return Settings ? Settings->Subgraph : nullptr;
-}
-
 #if WITH_EDITOR
 
 void UPCGSubgraphSettings::PreEditChange(FProperty* PropertyAboutToChange)
@@ -82,8 +76,87 @@ void UPCGSubgraphSettings::OnSubgraphChanged(UPCGGraph* InGraph, bool bIsStructu
 {
 	if (InGraph == Subgraph)
 	{
-		OnSettingsChangedDelegate.Broadcast(this);
+		if (bIsStructural)
+		{
+			OnStructuralSettingsChangedDelegate.Broadcast(this);
+		}
+		else
+		{
+			OnSettingsChangedDelegate.Broadcast(this);
+		}
 	}
 }
 
 #endif // WITH_EDITOR
+
+TObjectPtr<UPCGGraph> UPCGSubgraphNode::GetGraph() const
+{
+	TObjectPtr<UPCGSubgraphSettings> Settings = Cast<UPCGSubgraphSettings>(DefaultSettings);
+	return Settings ? Settings->Subgraph : nullptr;
+}
+
+void UPCGSubgraphNode::PostLoad()
+{
+	Super::PostLoad();
+
+#if WITH_EDITOR
+	if (UPCGSubgraphSettings* SubgraphSettings = Cast<UPCGSubgraphSettings>(DefaultSettings))
+	{
+		SubgraphSettings->OnStructuralSettingsChangedDelegate.AddUObject(this, &UPCGSubgraphNode::OnStructuralSettingsChanged);
+	}
+#endif
+}
+
+void UPCGSubgraphNode::BeginDestroy()
+{
+#if WITH_EDITOR
+	if (UPCGSubgraphSettings* SubgraphSettings = Cast<UPCGSubgraphSettings>(DefaultSettings))
+	{
+		SubgraphSettings->OnStructuralSettingsChangedDelegate.RemoveAll(this);
+	}
+#endif
+
+	Super::BeginDestroy();
+}
+
+#if WITH_EDITOR
+void UPCGSubgraphNode::PreEditChange(FProperty* PropertyAboutToChange)
+{
+	if (PropertyAboutToChange && PropertyAboutToChange->GetFName() == GET_MEMBER_NAME_CHECKED(UPCGNode, DefaultSettings))
+	{
+		if (UPCGSubgraphSettings* SubgraphSettings = Cast<UPCGSubgraphSettings>(DefaultSettings))
+		{
+			SubgraphSettings->OnStructuralSettingsChangedDelegate.RemoveAll(this);
+		}
+	}
+
+	Super::PreEditChange(PropertyAboutToChange);
+}
+
+void UPCGSubgraphNode::PostEditChangeProperty(struct FPropertyChangedEvent& PropertyChangedEvent)
+{
+	// Implementation note:
+	// We must process structural changes before the parent class' otherwise the graph might be rescheduled with its tasks
+	// before it is appropriately dirtied (which will trigger a recompilation)
+	if (PropertyChangedEvent.Property && PropertyChangedEvent.Property->GetFName() == GET_MEMBER_NAME_CHECKED(UPCGNode, DefaultSettings))
+	{
+		if (UPCGSubgraphSettings* SubgraphSettings = Cast<UPCGSubgraphSettings>(DefaultSettings))
+		{
+			SubgraphSettings->OnStructuralSettingsChangedDelegate.AddUObject(this, &UPCGSubgraphNode::OnStructuralSettingsChanged);
+
+			// Changing the default settings should trigger immediately a structural change
+			OnStructuralSettingsChanged(DefaultSettings);
+		}
+	}
+
+	Super::PostEditChangeProperty(PropertyChangedEvent);
+}
+
+void UPCGSubgraphNode::OnStructuralSettingsChanged(UPCGSettings* InSettings)
+{
+	if (InSettings == DefaultSettings)
+	{
+		OnNodeStructuralSettingsChangedDelegate.Broadcast(this);
+	}
+}
+#endif // WITH_EDITOR
\ No newline at end of file
diff --git a/Engine/Plugins/Experimental/PCG/Source/PCG/Public/PCGGraph.h b/Engine/Plugins/Experimental/PCG/Source/PCG/Public/PCGGraph.h
--- a/Engine/Plugins/Experimental/PCG/Source/PCG/Public/PCGGraph.h
+++ b/Engine/Plugins/Experimental/PCG/Source/PCG/Public/PCGGraph.h
@@ -66,6 +66,9 @@ public:
 #endif // WITH_EDITOR
 
 protected:
+	void OnNodeAdded(UPCGNode* InNode);
+	void OnNodeRemoved(UPCGNode* InNode);
+
 	UPROPERTY(BlueprintReadOnly, VisibleAnywhere, Category = Graph)
 	TArray<TObjectPtr<UPCGNode>> Nodes;
 
@@ -80,6 +83,7 @@ protected:
 private:
 	void NotifyGraphChanged(bool bIsStructural);
 	void OnSettingsChanged(UPCGNode* InNode);
+	void OnStructuralSettingsChanged(UPCGNode* InNode);
 
 	bool bEnableGraphChangeNotifications = true;
 #endif // WITH_EDITOR
diff --git a/Engine/Plugins/Experimental/PCG/Source/PCG/Public/PCGSubgraph.h b/Engine/Plugins/Experimental/PCG/Source/PCG/Public/PCGSubgraph.h
--- a/Engine/Plugins/Experimental/PCG/Source/PCG/Public/PCGSubgraph.h
+++ b/Engine/Plugins/Experimental/PCG/Source/PCG/Public/PCGSubgraph.h
@@ -8,6 +8,11 @@ class UPCGGraph;
 
 #include "PCGSubgraph.generated.h"
 
+#if WITH_EDITOR
+	DECLARE_MULTICAST_DELEGATE_OneParam(FOnPCGStructuralSettingsChanged, UPCGSettings*);
+	DECLARE_MULTICAST_DELEGATE_OneParam(FOnPCGNodeStructuralSettingsChanged, UPCGNode*);
+#endif
+
 UCLASS(BlueprintType, ClassGroup=(Procedural))
 class UPCGSubgraphSettings : public UPCGSettings
 {
@@ -41,6 +46,10 @@ protected:
 public:
 	UPROPERTY(BlueprintReadWrite, EditAnywhere, Category = Settings)
 	TObjectPtr<UPCGGraph> Subgraph;
+
+#if WITH_EDITOR
+	FOnPCGStructuralSettingsChanged OnStructuralSettingsChangedDelegate;
+#endif
 };
 
 UCLASS(ClassGroup = (Procedural))
@@ -52,5 +61,21 @@ public:
 	UPROPERTY(BlueprintReadWrite, EditAnywhere, Category = Settings)
 	bool bAllowDynamicGraph = false;
 
+#if WITH_EDITOR
+	FOnPCGNodeStructuralSettingsChanged OnNodeStructuralSettingsChangedDelegate;
+#endif
+
 	TObjectPtr<UPCGGraph> GetGraph() const;
+
+protected:	
+	/** ~Begin UObject interface */
+	virtual void PostLoad() override;
+	virtual void BeginDestroy() override;
+	/** ~End UObject interface */
+
+#if WITH_EDITOR
+	virtual void PreEditChange(FProperty* PropertyAboutToChange) override;
+	virtual void PostEditChangeProperty(struct FPropertyChangedEvent& PropertyChangedEvent) override;
+	void OnStructuralSettingsChanged(UPCGSettings* InSettings);
+#endif
 };
\ No newline at end of file
