Fixed issue with ISMC reuse in PCG

diff --git a/Engine/Plugins/Experimental/PCG/Source/PCG/Private/Helpers/PCGActorHelpers.cpp b/Engine/Plugins/Experimental/PCG/Source/PCG/Private/Helpers/PCGActorHelpers.cpp
--- a/Engine/Plugins/Experimental/PCG/Source/PCG/Private/Helpers/PCGActorHelpers.cpp
+++ b/Engine/Plugins/Experimental/PCG/Source/PCG/Private/Helpers/PCGActorHelpers.cpp
@@ -21,15 +21,26 @@ UInstancedStaticMeshComponent* UPCGActorHelpers::GetOrCreateISMC(AActor* InTarge
 		if (ISMC->GetStaticMesh() == InMesh &&
 			(!InSourceComponent || ISMC->ComponentTags.Contains(InSourceComponent->GetFName())))
 		{
-			// Check materials, if provided; if there's none provided but the current ISMC has overrides, skip this
-			bool bMaterialsMatched = (InMaterials.Num() != 0 == ISMC->GetNumOverrideMaterials() != 0);
+			// If materials are provided, we'll make sure they match to the already set materials.
+			// If not provided, we'll make sure that the current materials aren't overriden
+			bool bMaterialsMatched = true;
 
-			const int32 NumMaterialsToCheck = FMath::Min(InMaterials.Num(), ISMC->GetNumMaterials());
-			for (int32 MaterialIndex = 0; MaterialIndex < NumMaterialsToCheck && bMaterialsMatched; ++MaterialIndex)
+			for (int32 MaterialIndex = 0; MaterialIndex < ISMC->GetNumMaterials() && bMaterialsMatched; ++MaterialIndex)
 			{
-				if (InMaterials[MaterialIndex] && InMaterials[MaterialIndex] != ISMC->GetMaterial(MaterialIndex))
+				if (MaterialIndex < InMaterials.Num() && InMaterials[MaterialIndex])
 				{
-					bMaterialsMatched = false;
+					if (InMaterials[MaterialIndex] != ISMC->GetMaterial(MaterialIndex))
+					{
+						bMaterialsMatched = false;
+					}
+				}
+				else
+				{
+					// Material is currently overriden
+					if (ISMC->OverrideMaterials.IsValidIndex(MaterialIndex) && ISMC->OverrideMaterials[MaterialIndex])
+					{
+						bMaterialsMatched = false;
+					}
 				}
 			}
 
