Added density falloff on points + density interpolation in point data octree

diff --git a/Engine/Plugins/Experimental/PCG/Source/PCG/Private/Data/PCGPointData.cpp b/Engine/Plugins/Experimental/PCG/Source/PCG/Private/Data/PCGPointData.cpp
--- a/Engine/Plugins/Experimental/PCG/Source/PCG/Private/Data/PCGPointData.cpp
+++ b/Engine/Plugins/Experimental/PCG/Source/PCG/Private/Data/PCGPointData.cpp
@@ -7,10 +7,41 @@
 #include "GameFramework/Actor.h"
 #include "Misc/ScopeLock.h"
 
+namespace PCGPointHelpers
+{
+	float ManhattanDensity(const FPCGPoint& InPoint, const FVector& InPosition)
+	{
+		FVector LocalPosition = InPoint.Transform.InverseTransformPosition(InPosition);
+		LocalPosition /= InPoint.Extents;
+
+		// ]-2+s, 2-s] is the valid range of values
+		const FVector::FReal LowerBound = InPoint.Steepness - 2;
+		const FVector::FReal HigherBound = 2 - InPoint.Steepness;
+		
+		if (LocalPosition.X <= LowerBound || LocalPosition.X > HigherBound ||
+			LocalPosition.Y <= LowerBound || LocalPosition.Y > HigherBound ||
+			LocalPosition.Z <= LowerBound || LocalPosition.Z > HigherBound)
+		{
+			return 0;
+		}
+
+		// [-s, +s] is the range where the density is 1 on that axis
+		const FVector::FReal XDist = FMath::Max(0, FMath::Abs(LocalPosition.X) - InPoint.Steepness);
+		const FVector::FReal YDist = FMath::Max(0, FMath::Abs(LocalPosition.Y) - InPoint.Steepness);
+		const FVector::FReal ZDist = FMath::Max(0, FMath::Abs(LocalPosition.Z) - InPoint.Steepness);
+
+		const FVector::FReal DistanceScale = FMath::Max(2 - 2 * InPoint.Steepness, KINDA_SMALL_NUMBER);
+
+		//Note: for euclidean, we could do 1 - (dist / scale)^2
+		//Note: for maximum norm, we could do density * max(x factor, y factor, z factor)
+		return InPoint.Density * (1 - XDist/DistanceScale) * (1 - YDist/DistanceScale) * (1 - ZDist/DistanceScale);
+	}
+}
+
 FPCGPointRef::FPCGPointRef(const FPCGPoint& InPoint)
 {
 	Point = &InPoint;
-	Bounds = InPoint.GetBounds();
+	Bounds = InPoint.GetDensityBounds();
 }
 
 FPCGPointRef::FPCGPointRef(const FPCGPointRef& InPointRef)
@@ -58,7 +89,7 @@ void UPCGPointData::RecomputeBounds() const
 	FBox NewBounds(EForceInit::ForceInit);
 	for (const FPCGPoint& Point : Points)
 	{
-		FBoxSphereBounds PointBounds = Point.GetBounds();
+		FBoxSphereBounds PointBounds = Point.GetDensityBounds();
 		NewBounds += FBox::BuildAABB(PointBounds.Origin, PointBounds.BoxExtent);
 	}
 
@@ -121,15 +152,18 @@ const FPCGPoint* UPCGPointData::GetPointAtPosition(const FVector& InPosition) co
 
 float UPCGPointData::GetDensityAtPosition(const FVector& InPosition) const
 {
-	if (const FPCGPoint* BestPoint = GetPointAtPosition(InPosition))
-	{
-		// It is going to be hard to do a linear fall-off here and maybe not really required.
-		return BestPoint->Density;
-	}
-	else
+	if (bOctreeIsDirty)
 	{
-		return 0;
+		RebuildOctree();
 	}
+
+	float Density = 0;
+
+	Octree.FindElementsWithBoundsTest(FBoxCenterAndExtent(InPosition, FVector::Zero()), [&InPosition, &Density](const FPCGPointRef& InPointRef) {
+		Density += PCGPointHelpers::ManhattanDensity(*InPointRef.Point, InPosition);
+	});
+
+	return FMath::Min(Density, 1.0f);
 }
 
 void UPCGPointData::RebuildOctree() const
diff --git a/Engine/Plugins/Experimental/PCG/Source/PCG/Private/Data/PCGSpatialData.cpp b/Engine/Plugins/Experimental/PCG/Source/PCG/Private/Data/PCGSpatialData.cpp
--- a/Engine/Plugins/Experimental/PCG/Source/PCG/Private/Data/PCGSpatialData.cpp
+++ b/Engine/Plugins/Experimental/PCG/Source/PCG/Private/Data/PCGSpatialData.cpp
@@ -27,7 +27,7 @@ FPCGPoint UPCGSpatialData::TransformPoint(const FPCGPoint& InPoint) const
 {
 	FPCGPoint TransformedPoint = InPoint;
 	TransformedPoint.Transform.SetTranslation(TransformPosition(InPoint.Transform.GetLocation()));
-	TransformedPoint.Density = GetDensityAtPosition(InPoint.Transform.GetLocation());
+	TransformedPoint.Density *= GetDensityAtPosition(InPoint.Transform.GetLocation());
 	return TransformedPoint;
 }
 
diff --git a/Engine/Plugins/Experimental/PCG/Source/PCG/Private/Elements/PCGSurfaceSampler.cpp b/Engine/Plugins/Experimental/PCG/Source/PCG/Private/Elements/PCGSurfaceSampler.cpp
--- a/Engine/Plugins/Experimental/PCG/Source/PCG/Private/Elements/PCGSurfaceSampler.cpp
+++ b/Engine/Plugins/Experimental/PCG/Source/PCG/Private/Elements/PCGSurfaceSampler.cpp
@@ -48,6 +48,7 @@ bool FPCGSurfaceSamplerElement::ExecuteInternal(FPCGContextPtr Context) const
 		}
 
 		// Conceptually, we will break down the surface bounds in a N x M grid
+		const float PointSteepness = Settings->PointSteepness;
 		const bool bApplyDensity = Settings->bApplyDensityToPoints;
 		const FVector::FReal PointRadius = Settings->PointRadius;
 		const FVector::FReal InterstitialDistance = 2 * PointRadius;
@@ -71,7 +72,7 @@ bool FPCGSurfaceSamplerElement::ExecuteInternal(FPCGContextPtr Context) const
 		check(CellCount > 0);
 
 		const FVector::FReal InvSquaredMeterUnits = 1.0 / (100.0 * 100.0);
-		int32 TargetPointCount = (InputBounds.Max.X - InputBounds.Min.X) * (InputBounds.Max.Y - InputBounds.Min.Y) * Settings->PointsPerSquaredMeter * InvSquaredMeterUnits;
+		int64 TargetPointCount = (InputBounds.Max.X - InputBounds.Min.X) * (InputBounds.Max.Y - InputBounds.Min.Y) * Settings->PointsPerSquaredMeter * InvSquaredMeterUnits;
 
 		if (TargetPointCount == 0)
 		{
@@ -98,7 +99,7 @@ bool FPCGSurfaceSamplerElement::ExecuteInternal(FPCGContextPtr Context) const
 		const FVector::FReal CellStartY = CellMinY * CellSize;
 
 		FVector::FReal CurrentX = CellStartX;
-		
+
 		for(int32 CellX = CellMinX; CellX <= CellMaxX; ++CellX)
 		{
 			FVector::FReal CurrentY = CellStartY;
@@ -117,11 +118,16 @@ bool FPCGSurfaceSamplerElement::ExecuteInternal(FPCGContextPtr Context) const
 					Point.Transform = FTransform(FVector(CurrentX + RandX * InnerCellSize, CurrentY + RandY * InnerCellSize, 0));
 					Point.Extents = FVector(PointRadius);
 					Point.Density = bApplyDensity ? ((Ratio - Chance) / Ratio) : 1.0f;
+					Point.Steepness = PointSteepness;
 					Point.Seed = RandomSource.GetCurrentSeed();
 
 					Point = SpatialInput->TransformPoint(Point);
 
+#if WITH_EDITORONLY_DATA
+					if(Point.Density > 0 || Settings->bKeepZeroDensityPoints)
+#else
 					if (Point.Density > 0)
+#endif
 					{
 						SampledPoints.Add(Point);
 					}
diff --git a/Engine/Plugins/Experimental/PCG/Source/PCG/Private/PCGPoint.cpp b/Engine/Plugins/Experimental/PCG/Source/PCG/Private/PCGPoint.cpp
--- a/Engine/Plugins/Experimental/PCG/Source/PCG/Private/PCGPoint.cpp
+++ b/Engine/Plugins/Experimental/PCG/Source/PCG/Private/PCGPoint.cpp
@@ -5,8 +5,6 @@
 FPCGPoint::FPCGPoint(const FTransform& InTransform, float InDensity, int32 InSeed)
 	: Transform(InTransform)
 	, Density(InDensity)
-	, Extents(FVector::One())
-	, Color(FVector4::One())
 	, Seed(InSeed)
 {
 }
@@ -14,4 +12,9 @@ FPCGPoint::FPCGPoint(const FTransform& InTransform, float InDensity, int32 InSee
 FBoxSphereBounds FPCGPoint::GetBounds() const
 {
 	return FBoxSphereBounds(FBox::BuildAABB(FVector::Zero(), Extents).TransformBy(Transform));
+}
+
+FBoxSphereBounds FPCGPoint::GetDensityBounds() const
+{
+	return FBoxSphereBounds(FBox::BuildAABB(FVector::Zero(), (2 - Steepness) * Extents).TransformBy(Transform));
 }
\ No newline at end of file
diff --git a/Engine/Plugins/Experimental/PCG/Source/PCG/Public/Elements/PCGSurfaceSampler.h b/Engine/Plugins/Experimental/PCG/Source/PCG/Public/Elements/PCGSurfaceSampler.h
--- a/Engine/Plugins/Experimental/PCG/Source/PCG/Public/Elements/PCGSurfaceSampler.h
+++ b/Engine/Plugins/Experimental/PCG/Source/PCG/Public/Elements/PCGSurfaceSampler.h
@@ -26,9 +26,17 @@ public:
 	UPROPERTY(BlueprintReadWrite, EditAnywhere, Category = Settings, meta=(ClampMin="0"))
 	float Looseness = 1.0f;
 
-	UPROPERTY(BlueprintReadWrite, EditAnywhere, Category = Settings)
+	UPROPERTY(BlueprintReadWrite, EditAnywhere, Category = "Settings|Points")
 	bool bApplyDensityToPoints = true;
 
+	UPROPERTY(BlueprintReadWrite, EditAnywhere, Category = "Settings|Points", meta=(ClampMin="0", ClampMax="1"))
+	float PointSteepness = 0.5f;
+
+#if WITH_EDITORONLY_DATA
+	UPROPERTY(Transient, BlueprintReadWrite, EditAnywhere, Category = "Debug")
+	bool bKeepZeroDensityPoints = false;
+#endif
+
 #if WITH_EDITOR
 	//~Begin UPCGSettings interface
 	virtual FName GetDefaultNodeName() const override { return FName(TEXT("SurfaceSamplerNode")); }
diff --git a/Engine/Plugins/Experimental/PCG/Source/PCG/Public/PCGPoint.h b/Engine/Plugins/Experimental/PCG/Source/PCG/Public/PCGPoint.h
--- a/Engine/Plugins/Experimental/PCG/Source/PCG/Public/PCGPoint.h
+++ b/Engine/Plugins/Experimental/PCG/Source/PCG/Public/PCGPoint.h
@@ -12,23 +12,26 @@ struct FPCGPoint
 	GENERATED_BODY()
 public:
 	FPCGPoint() = default;
-
 	FPCGPoint(const FTransform& InTransform, float InDensity, int32 InSeed);
 
 	FBoxSphereBounds GetBounds() const;
+	FBoxSphereBounds GetDensityBounds() const;
 
 	UPROPERTY(BlueprintReadWrite, EditAnywhere, Category = Properties)
 	FTransform Transform;
 
 	UPROPERTY(BlueprintReadWrite, EditAnywhere, Category = Properties)
-	float Density;
+	float Density = 1.0f;
 
 	UPROPERTY(BlueprintReadWrite, EditAnywhere, Category = Properties)
-	FVector Extents;
+	FVector Extents = FVector::One();
 
 	UPROPERTY(BlueprintReadWrite, EditAnywhere, Category = Properties)
-	FVector4 Color;
+	FVector4 Color = FVector4::One();
+
+	UPROPERTY(BlueprintReadWrite, EditAnywhere, Category = Properties, meta = (ClampMin = "0", ClampMax = "1"))
+	float Steepness = 0.5f;
 
 	UPROPERTY(BlueprintReadWrite, EditAnywhere, Category = Properties)
-	int32 Seed;
+	int32 Seed = 0;
 };
\ No newline at end of file
