Added option to keep zero density points for debug purposes

diff --git a/Engine/Plugins/Experimental/PCG/Source/PCG/Private/Data/PCGDifferenceData.cpp b/Engine/Plugins/Experimental/PCG/Source/PCG/Private/Data/PCGDifferenceData.cpp
--- a/Engine/Plugins/Experimental/PCG/Source/PCG/Private/Data/PCGDifferenceData.cpp
+++ b/Engine/Plugins/Experimental/PCG/Source/PCG/Private/Data/PCGDifferenceData.cpp
@@ -135,7 +135,7 @@ const UPCGPointData* UPCGDifferenceData::CreatePointData() const
 	// This is similar to what we are doing in UPCGUnionData::CreatePointData
 	const UPCGPointData* SourcePointData = Source->ToPointData();
 
-	if (!Difference)
+	if (!Difference || !SourcePointData)
 	{
 		return SourcePointData;
 	}
@@ -148,11 +148,19 @@ const UPCGPointData* UPCGDifferenceData::CreatePointData() const
 	for (const FPCGPoint& Point : SourcePointData->GetPoints())
 	{
 		const float DensityInDifference = Difference->GetDensityAtPosition(Point.Transform.GetLocation());
+
 		if (DensityInDifference < Point.Density)
 		{
 			FPCGPoint& TargetPoint = TargetPoints.Add_GetRef(Point);
 			TargetPoint.Density -= DensityInDifference;
 		}
+#if WITH_EDITORONLY_DATA
+		else if (bKeepZeroDensityPoints)
+		{
+			FPCGPoint& TargetPoint = TargetPoints.Add_GetRef(Point);
+			TargetPoint.Density = 0;
+		}
+#endif
 	}
 
 	return Data;
diff --git a/Engine/Plugins/Experimental/PCG/Source/PCG/Private/Data/PCGIntersectionData.cpp b/Engine/Plugins/Experimental/PCG/Source/PCG/Private/Data/PCGIntersectionData.cpp
--- a/Engine/Plugins/Experimental/PCG/Source/PCG/Private/Data/PCGIntersectionData.cpp
+++ b/Engine/Plugins/Experimental/PCG/Source/PCG/Private/Data/PCGIntersectionData.cpp
@@ -137,7 +137,11 @@ UPCGPointData* UPCGIntersectionData::CreateAndFilterPointData(const UPCGSpatialD
 	for (const FPCGPoint& Point : SourcePoints)
 	{
 		const float YDensity = Y->GetDensityAtPosition(Point.Transform.GetLocation());
+#if WITH_EDITORONLY_DATA
+		if (YDensity > 0 || bKeepZeroDensityPoints)
+#else
 		if (YDensity > 0)
+#endif
 		{
 			FPCGPoint& NewPoint = TargetPoints.Add_GetRef(Point);
 			NewPoint.Density = PCGIntersectionDataMaths::ComputeDensity(Point.Density, YDensity, DensityFunction);
diff --git a/Engine/Plugins/Experimental/PCG/Source/PCG/Private/Data/PCGProjectionData.cpp b/Engine/Plugins/Experimental/PCG/Source/PCG/Private/Data/PCGProjectionData.cpp
--- a/Engine/Plugins/Experimental/PCG/Source/PCG/Private/Data/PCGProjectionData.cpp
+++ b/Engine/Plugins/Experimental/PCG/Source/PCG/Private/Data/PCGProjectionData.cpp
@@ -102,7 +102,11 @@ const UPCGPointData* UPCGProjectionData::CreatePointData() const
 	for (const FPCGPoint& SourcePoint : SourcePoints)
 	{
 		FPCGPoint ProjectedPoint = TransformPoint(SourcePoint);
+#if WITH_EDITORONLY_DATA
+		if(ProjectedPoint.Density > 0 || bKeepZeroDensityPoints)
+#else
 		if (ProjectedPoint.Density > 0)
+#endif
 		{
 			Points.Add(MoveTemp(ProjectedPoint));
 		}
diff --git a/Engine/Plugins/Experimental/PCG/Source/PCG/Private/Data/PCGTextureData.cpp b/Engine/Plugins/Experimental/PCG/Source/PCG/Private/Data/PCGTextureData.cpp
--- a/Engine/Plugins/Experimental/PCG/Source/PCG/Private/Data/PCGTextureData.cpp
+++ b/Engine/Plugins/Experimental/PCG/Source/PCG/Private/Data/PCGTextureData.cpp
@@ -16,25 +16,25 @@ namespace PCGTextureSampling
 
 		// TODO: this isn't super robust, if that becomes an issue
 		int32 X0 = FMath::FloorToInt(Pos.X);
-		if (X0 >= Width)
+		if (X0 < 0 || X0 >= Width)
 		{
 			X0 = 0;
 		}
 
 		int32 X1 = FMath::CeilToInt(Pos.X);
-		if (X1 >= Width)
+		if (X1 < 0 || X1 >= Width)
 		{
 			X1 = 0;
 		}
 
 		int32 Y0 = FMath::FloorToInt(Pos.Y);
-		if (Y0 >= Height)
+		if (Y0 < 0 || Y0 >= Height)
 		{
 			Y0 = 0;
 		}
 
 		int32 Y1 = FMath::CeilToInt(Pos.Y);
-		if (Y1 >= Height)
+		if (Y1 < 0 || Y1 >= Height)
 		{
 			Y1 = 0;
 		}
@@ -109,7 +109,11 @@ const UPCGPointData* UPCGBaseTextureData::CreatePointData() const
 		{
 			const float Density = PCGTextureSampling::SampleFloatChannel(ColorData[X + Y * Width], ColorChannel);
 
+#if WITH_EDITORONLY_DATA
+			if(Density > 0 || bKeepZeroDensityPoints)
+#else
 			if (Density > 0)
+#endif
 			{
 				FVector LocalPosition(X * XScale + Bias.X, Y * YScale + Bias.Y, 0);
 				Points.Emplace(FTransform(Transform.TransformPosition(LocalPosition)),
diff --git a/Engine/Plugins/Experimental/PCG/Source/PCG/Private/Elements/PCGDifferenceElement.cpp b/Engine/Plugins/Experimental/PCG/Source/PCG/Private/Elements/PCGDifferenceElement.cpp
--- a/Engine/Plugins/Experimental/PCG/Source/PCG/Private/Elements/PCGDifferenceElement.cpp
+++ b/Engine/Plugins/Experimental/PCG/Source/PCG/Private/Elements/PCGDifferenceElement.cpp
@@ -27,6 +27,9 @@ bool FPCGDifferenceElement::ExecuteInternal(FPCGContextPtr Context) const
 		{
 			DifferenceData = FirstSpatialData->Subtract(SpatialData);
 			DifferenceData->SetDensityFunction(Settings->DensityFunction);
+#if WITH_EDITORONLY_DATA
+			DifferenceData->bKeepZeroDensityPoints = Settings->bKeepZeroDensityPoints;
+#endif
 
 			FPCGTaggedData& DifferenceTaggedData = Outputs[DifferenceTaggedDataIndex];
 			DifferenceTaggedData.Data = DifferenceData;
diff --git a/Engine/Plugins/Experimental/PCG/Source/PCG/Private/Elements/PCGIntersectionElement.cpp b/Engine/Plugins/Experimental/PCG/Source/PCG/Private/Elements/PCGIntersectionElement.cpp
--- a/Engine/Plugins/Experimental/PCG/Source/PCG/Private/Elements/PCGIntersectionElement.cpp
+++ b/Engine/Plugins/Experimental/PCG/Source/PCG/Private/Elements/PCGIntersectionElement.cpp
@@ -44,6 +44,10 @@ bool FPCGIntersectionElement::ExecuteInternal(FPCGContextPtr Context) const
 		IntersectionData = (IntersectionData ? IntersectionData : FirstSpatialData)->IntersectWith(SpatialData);
 		// Propagate settings
 		IntersectionData->DensityFunction = Settings->DensityFunction;
+#if WITH_EDITORONLY_DATA
+		IntersectionData->bKeepZeroDensityPoints = Settings->bKeepZeroDensityPoints;
+#endif
+
 
 		// Update tagged data
 		FPCGTaggedData& IntersectionTaggedData = Outputs[IntersectionTaggedDataIndex];
diff --git a/Engine/Plugins/Experimental/PCG/Source/PCG/Private/Elements/PCGProjectionElement.cpp b/Engine/Plugins/Experimental/PCG/Source/PCG/Private/Elements/PCGProjectionElement.cpp
--- a/Engine/Plugins/Experimental/PCG/Source/PCG/Private/Elements/PCGProjectionElement.cpp
+++ b/Engine/Plugins/Experimental/PCG/Source/PCG/Private/Elements/PCGProjectionElement.cpp
@@ -44,6 +44,10 @@ bool FPCGProjectionElement::ExecuteInternal(FPCGContextPtr Context) const
 
 		// Create a new projection
 		ProjectionData = (ProjectionData ? ProjectionData : FirstSpatialData)->ProjectOn(SpatialData);
+
+#if WITH_EDITORONLY_DATA
+		ProjectionData->bKeepZeroDensityPoints = Settings->bKeepZeroDensityPoints;
+#endif
 		
 		// Update the tagged data
 		FPCGTaggedData& ProjectionTaggedData = Outputs[ProjectionTaggedDataIndex];
diff --git a/Engine/Plugins/Experimental/PCG/Source/PCG/Public/Data/PCGSpatialData.h b/Engine/Plugins/Experimental/PCG/Source/PCG/Public/Data/PCGSpatialData.h
--- a/Engine/Plugins/Experimental/PCG/Source/PCG/Public/Data/PCGSpatialData.h
+++ b/Engine/Plugins/Experimental/PCG/Source/PCG/Public/Data/PCGSpatialData.h
@@ -82,6 +82,11 @@ public:
 
 	UPROPERTY(BlueprintReadWrite, EditAnywhere, Category = Data)
 	AActor* TargetActor = nullptr;
+
+#if WITH_EDITORONLY_DATA
+	UPROPERTY(Transient, BlueprintReadWrite, EditAnywhere, Category = Debug)
+	bool bKeepZeroDensityPoints = false;
+#endif
 };
 
 UCLASS(Abstract, ClassGroup = (Procedural))
diff --git a/Engine/Plugins/Experimental/PCG/Source/PCG/Public/Elements/PCGDifferenceElement.h b/Engine/Plugins/Experimental/PCG/Source/PCG/Public/Elements/PCGDifferenceElement.h
--- a/Engine/Plugins/Experimental/PCG/Source/PCG/Public/Elements/PCGDifferenceElement.h
+++ b/Engine/Plugins/Experimental/PCG/Source/PCG/Public/Elements/PCGDifferenceElement.h
@@ -27,6 +27,11 @@ protected:
 public:
 	UPROPERTY(BlueprintReadWrite, EditAnywhere, Category = Settings)
 	EPCGDifferenceDensityFunction DensityFunction = EPCGDifferenceDensityFunction::Minimum;
+
+#if WITH_EDITORONLY_DATA
+	UPROPERTY(Transient, BlueprintReadWrite, EditAnywhere, Category = Debug)
+	bool bKeepZeroDensityPoints = false;
+#endif
 };
 
 class FPCGDifferenceElement : public FSimpleTypedPCGElement<UPCGDifferenceSettings>
diff --git a/Engine/Plugins/Experimental/PCG/Source/PCG/Public/Elements/PCGIntersectionElement.h b/Engine/Plugins/Experimental/PCG/Source/PCG/Public/Elements/PCGIntersectionElement.h
--- a/Engine/Plugins/Experimental/PCG/Source/PCG/Public/Elements/PCGIntersectionElement.h
+++ b/Engine/Plugins/Experimental/PCG/Source/PCG/Public/Elements/PCGIntersectionElement.h
@@ -27,6 +27,11 @@ protected:
 public:
 	UPROPERTY(BlueprintReadWrite, EditAnywhere, Category = Settings)
 	EPCGIntersectionDensityFunction DensityFunction = EPCGIntersectionDensityFunction::Multiply;
+
+#if WITH_EDITORONLY_DATA
+	UPROPERTY(Transient, BlueprintReadWrite, EditAnywhere, Category = Debug)
+	bool bKeepZeroDensityPoints = false;
+#endif
 };
 
 class FPCGIntersectionElement : public FSimpleTypedPCGElement<UPCGIntersectionSettings>
diff --git a/Engine/Plugins/Experimental/PCG/Source/PCG/Public/Elements/PCGProjectionElement.h b/Engine/Plugins/Experimental/PCG/Source/PCG/Public/Elements/PCGProjectionElement.h
--- a/Engine/Plugins/Experimental/PCG/Source/PCG/Public/Elements/PCGProjectionElement.h
+++ b/Engine/Plugins/Experimental/PCG/Source/PCG/Public/Elements/PCGProjectionElement.h
@@ -19,6 +19,11 @@ public:
 	virtual FName GetDefaultNodeName() const override { return FName(TEXT("ProjectionNode")); }
 #endif
 
+#if WITH_EDITORONLY_DATA
+	UPROPERTY(Transient, BlueprintReadWrite, EditAnywhere, Category = Debug)
+	bool bKeepZeroDensityPoints = false;
+#endif
+
 protected:
 	virtual FPCGElementPtr CreateElement() const override;
 	//~End UPCGSettings interface
