Exposed mechanisms to get matching component in original from pcg partition actors

diff --git a/Engine/Plugins/Experimental/PCG/Source/PCG/Private/Grid/PCGPartitionActor.h b/Engine/Plugins/Experimental/PCG/Source/PCG/Private/Grid/PCGPartitionActor.h
--- a/Engine/Plugins/Experimental/PCG/Source/PCG/Private/Grid/PCGPartitionActor.h
+++ b/Engine/Plugins/Experimental/PCG/Source/PCG/Private/Grid/PCGPartitionActor.h
@@ -45,7 +45,10 @@ public:
 
 	// TODO: Make this in-editor only; during runtime, we should keep a map of component to bounds/volume only
 	// and preferably precompute the intersection, so this would make it easier/possible to not have the original actor in game version.
+	UFUNCTION(BlueprintCallable, Category = "PCG|PartitionActor")
 	UPCGComponent* GetLocalComponent(const UPCGComponent* OriginalComponent) const;
+
+	UFUNCTION(BlueprintCallable, Category = "PCG|PartitionActor")
 	UPCGComponent* GetOriginalComponent(const UPCGComponent* LocalComponent) const;
 
 	UPROPERTY()
