Support for dynamic graph execution

diff --git a/Engine/Plugins/Experimental/PCG/Source/PCG/Private/Graph/PCGGraphCompiler.cpp b/Engine/Plugins/Experimental/PCG/Source/PCG/Private/Graph/PCGGraphCompiler.cpp
--- a/Engine/Plugins/Experimental/PCG/Source/PCG/Private/Graph/PCGGraphCompiler.cpp
+++ b/Engine/Plugins/Experimental/PCG/Source/PCG/Private/Graph/PCGGraphCompiler.cpp
@@ -29,27 +29,28 @@ TArray<FPCGGraphTask> FPCGGraphCompiler::CompileGraph(UPCGGraph* InGraph, FPCGTa
 	{
 		const UPCGNode* Node = NodeQueue.Pop();
 
-		const UPCGSubgraphNode* SubgraphNode = Cast<const UPCGSubgraphNode>(Node);
+		const UPCGBaseSubgraphNode* SubgraphNode = Cast<const UPCGBaseSubgraphNode>(Node);
+		UPCGGraph* Subgraph = SubgraphNode ? SubgraphNode->GetSubgraph() : nullptr;
 
-		if (SubgraphNode != nullptr && !SubgraphNode->bAllowDynamicGraph)
+		if (SubgraphNode != nullptr && !SubgraphNode->bDynamicGraph && Subgraph)
 		{
 			const FPCGTaskId PreId = NextId++;
 
 			// 1. Compile the subgraph making sure we don't reuse the same ids
 			// Note that we will not consume the pre or post-execute tasks, ergo bIsTopGraph=false
-			TArray<FPCGGraphTask> Subtasks = GetCompiledTasks(SubgraphNode->GetGraph(), /*bIsTopGraph=*/false);
+			TArray<FPCGGraphTask> Subtasks = GetCompiledTasks(Subgraph, /*bIsTopGraph=*/false);
 
 #if WITH_EDITOR
 			GraphDependenciesLock.Lock();
-			GraphDependencies.AddUnique(SubgraphNode->GetGraph(), InGraph);
+			GraphDependencies.AddUnique(Subgraph, InGraph);
 			GraphDependenciesLock.Unlock();
 #endif // WITH_EDITOR
 
 			OffsetNodeIds(Subtasks, NextId);
 			NextId += Subtasks.Num();
 
-			const UPCGNode* SubgraphInputNode = SubgraphNode->GetGraph()->GetInputNode();
-			const UPCGNode* SubgraphOutputNode = SubgraphNode->GetGraph()->GetOutputNode();
+			const UPCGNode* SubgraphInputNode = Subgraph->GetInputNode();
+			const UPCGNode* SubgraphOutputNode = Subgraph->GetOutputNode();
 
 			// 2. Update the "input" and "output" node tasks so we can add the proper dependencies
 			FPCGGraphTask* InputNodeTask = Subtasks.FindByPredicate([SubgraphInputNode](const FPCGGraphTask& Subtask) {
@@ -94,7 +95,7 @@ TArray<FPCGGraphTask> FPCGGraphCompiler::CompileGraph(UPCGGraph* InGraph, FPCGTa
 			IdMapping.Add(Node, PostId);
 
 			// Prime any input-less nodes from the subgraph
-			for (const UPCGNode* SubNode : SubgraphNode->GetGraph()->GetNodes())
+			for (const UPCGNode* SubNode : Subgraph->GetNodes())
 			{
 				if (SubNode->InboundNodes.IsEmpty())
 				{
@@ -247,7 +248,7 @@ void FPCGGraphCompiler::CompileTopGraph(UPCGGraph* InGraph)
 	// Find end nodes, e.g. all nodes that have no successors.
 	// In our representation we don't have this, so find it out by going backwards
 	// Note: this works because there is a weak ordering on the tasks such that
-	// a successor task is always after if predecessor
+	// a successor task is always after its predecessors
 	TSet<FPCGTaskId> TasksWithSuccessors;
 	TArray<FPCGTaskId> EndNodes;
 
diff --git a/Engine/Plugins/Experimental/PCG/Source/PCG/Private/Graph/PCGGraphExecutor.cpp b/Engine/Plugins/Experimental/PCG/Source/PCG/Private/Graph/PCGGraphExecutor.cpp
--- a/Engine/Plugins/Experimental/PCG/Source/PCG/Private/Graph/PCGGraphExecutor.cpp
+++ b/Engine/Plugins/Experimental/PCG/Source/PCG/Private/Graph/PCGGraphExecutor.cpp
@@ -29,6 +29,12 @@ FPCGTaskId FPCGGraphExecutor::Schedule(UPCGComponent* Component, const TArray<FP
 {
 	check(Component);
 	UPCGGraph* Graph = Component->GetGraph();
+
+	return Schedule(Graph, Component, MakeShared<FPCGFetchInputElement>(Component), ExternalDependencies);
+}
+
+FPCGTaskId FPCGGraphExecutor::Schedule(UPCGGraph* Graph, const UPCGComponent* SourceComponent, FPCGElementPtr InputElement, const TArray<FPCGTaskId>& ExternalDependencies)
+{
 	FPCGTaskId ScheduledId = InvalidTaskId;
 
 	// Get compiled tasks from compiler
@@ -37,7 +43,7 @@ FPCGTaskId FPCGGraphExecutor::Schedule(UPCGComponent* Component, const TArray<FP
 	// Assign this component to the tasks
 	for (FPCGGraphTask& Task : CompiledTasks)
 	{
-		Task.SourceComponent = Component;
+		Task.SourceComponent = SourceComponent;
 	}
 
 	// Prepare scheduled task that will be promoted in the next Execute call.
@@ -46,7 +52,7 @@ FPCGTaskId FPCGGraphExecutor::Schedule(UPCGComponent* Component, const TArray<FP
 		check(CompiledTasks[0].Node == Graph->GetInputNode());
 
 		// Setup fetch task on input node
-		CompiledTasks[0].Element = MakeShared<FPCGFetchInputElement>(Component);
+		CompiledTasks[0].Element = InputElement;
 
 		ScheduleLock.Lock();
 
@@ -71,7 +77,6 @@ FPCGTaskId FPCGGraphExecutor::Schedule(UPCGComponent* Component, const TArray<FP
 	return ScheduledId;
 }
 
-#if WITH_EDITOR
 FPCGTaskId FPCGGraphExecutor::ScheduleGeneric(TFunction<bool()> InOperation, const TArray<FPCGTaskId>& TaskDependencies)
 {
 	// Build task & element to hold the operation to perform
@@ -91,7 +96,20 @@ FPCGTaskId FPCGGraphExecutor::ScheduleGeneric(TFunction<bool()> InOperation, con
 
 	return Task.NodeId;
 }
-#endif
+
+bool FPCGGraphExecutor::GetOutputData(FPCGTaskId TaskId, FPCGDataCollection& OutData)
+{
+	// TODO: this is not threadsafe - make threadsafe once we multithread execution
+	if (OutputData.Contains(TaskId))
+	{
+		OutData = OutputData[TaskId];
+		return true;
+	}
+	else
+	{
+		return false;
+	}
+}
 
 void FPCGGraphExecutor::Execute()
 {
@@ -123,7 +141,7 @@ void FPCGGraphExecutor::Execute()
 	check(Tasks.Num() == 0 || ReadyTasks.Num() > 0 || ActiveTasks.Num() > 0);
 
 	// TODO: change this when we support time-slicing
-	while (ReadyTasks.Num() > 0 || ActiveTasks.Num() > 0)
+	if(ReadyTasks.Num() > 0 || ActiveTasks.Num() > 0)
 	{
 		// First: if we have free resources, move ready tasks to the active tasks
 		bool bCanStartNewTasks = true; // TODO implement this
@@ -169,6 +187,7 @@ void FPCGGraphExecutor::Execute()
 				
 				FPCGContextPtr Context = Element->Initialize(TaskInput, Task.SourceComponent);
 				Context->Node = Task.Node; // still needed?
+				Context->TaskId = Task.NodeId;
 
 				FPCGGraphActiveTask& ActiveTask = ActiveTasks.Emplace_GetRef();
 				ActiveTask.Element = Element;
@@ -186,7 +205,7 @@ void FPCGGraphExecutor::Execute()
 		{
 			FPCGGraphActiveTask& ActiveTask = ActiveTasks[ExecutionIndex];
 
-			if (ActiveTask.Element->Execute(ActiveTask.Context))
+			if (!ActiveTask.Context->bIsPaused && ActiveTask.Element->Execute(ActiveTask.Context))
 			{
 				// Store output data in map
 				StoreResults(ActiveTask.NodeId, ActiveTask.Context->OutputData);
@@ -366,7 +385,6 @@ bool FPCGFetchInputElement::ExecuteInternal(FPCGContextPtr Context) const
 	return true;
 }
 
-#if WITH_EDITOR
 FPCGGenericElement::FPCGGenericElement(TFunction<bool()> InOperation)
 	: Operation(InOperation)
 {
@@ -376,4 +394,3 @@ bool FPCGGenericElement::ExecuteInternal(FPCGContextPtr Context) const
 {
 	return Operation();
 }
-#endif // WITH_EDITOR
diff --git a/Engine/Plugins/Experimental/PCG/Source/PCG/Private/Graph/PCGGraphExecutor.h b/Engine/Plugins/Experimental/PCG/Source/PCG/Private/Graph/PCGGraphExecutor.h
--- a/Engine/Plugins/Experimental/PCG/Source/PCG/Private/Graph/PCGGraphExecutor.h
+++ b/Engine/Plugins/Experimental/PCG/Source/PCG/Private/Graph/PCGGraphExecutor.h
@@ -50,11 +50,15 @@ public:
 
 	/** Schedules the execution of a given graph with specified inputs. This call is threadsafe */
 	FPCGTaskId Schedule(UPCGComponent* InComponent, const TArray<FPCGTaskId>& TaskDependency = TArray<FPCGTaskId>());
+	FPCGTaskId Schedule(UPCGGraph* Graph, const UPCGComponent* InSourceComponent, FPCGElementPtr InputElement, const TArray<FPCGTaskId>& TaskDependency);
 
-#if WITH_EDITOR
 	/** General job scheduling, used to control loading/unloading */
 	FPCGTaskId ScheduleGeneric(TFunction<bool()> InOperation, const TArray<FPCGTaskId>& TaskDependencies);
 
+	/** Gets data in the output results. Returns false if data is not ready. */
+	bool GetOutputData(FPCGTaskId InTaskId, FPCGDataCollection& OutData);
+
+#if WITH_EDITOR
 	void AddToDirtyActors(AActor* Actor);
 	void AddToUnusedActors(const TSet<FWorldPartitionReference>& UnusedActors);
 
@@ -111,7 +115,6 @@ private:
 	UPCGComponent* Component = nullptr;
 };
 
-#if WITH_EDITOR
 class FPCGGenericElement : public FSimplePCGElement
 {
 public:
@@ -126,4 +129,3 @@ protected:
 private:
 	TFunction<bool()> Operation;
 };
-#endif // WITH_EDITOR
diff --git a/Engine/Plugins/Experimental/PCG/Source/PCG/Private/PCGSubgraph.cpp b/Engine/Plugins/Experimental/PCG/Source/PCG/Private/PCGSubgraph.cpp
--- a/Engine/Plugins/Experimental/PCG/Source/PCG/Private/PCGSubgraph.cpp
+++ b/Engine/Plugins/Experimental/PCG/Source/PCG/Private/PCGSubgraph.cpp
@@ -2,6 +2,7 @@
 
 #include "PCGSubgraph.h"
 #include "PCGGraph.h"
+#include "PCGComponent.h"
 
 void UPCGSubgraphSettings::PostLoad()
 {
@@ -41,7 +42,7 @@ TArray<FName> UPCGSubgraphSettings::GetTrackedActorTags() const
 
 FPCGElementPtr UPCGSubgraphSettings::CreateElement() const
 {
-	return MakeShared<FPCGTrivialElement>();
+	return MakeShared<FPCGSubgraphElement>();
 }
 
 #if WITH_EDITOR
@@ -89,7 +90,7 @@ void UPCGSubgraphSettings::OnSubgraphChanged(UPCGGraph* InGraph, bool bIsStructu
 
 #endif // WITH_EDITOR
 
-TObjectPtr<UPCGGraph> UPCGSubgraphNode::GetGraph() const
+TObjectPtr<UPCGGraph> UPCGSubgraphNode::GetSubgraph() const
 {
 	TObjectPtr<UPCGSubgraphSettings> Settings = Cast<UPCGSubgraphSettings>(DefaultSettings);
 	return Settings ? Settings->Subgraph : nullptr;
@@ -159,4 +160,98 @@ void UPCGSubgraphNode::OnStructuralSettingsChanged(UPCGSettings* InSettings)
 		OnNodeStructuralSettingsChangedDelegate.Broadcast(this);
 	}
 }
-#endif // WITH_EDITOR
\ No newline at end of file
+#endif // WITH_EDITOR
+
+FPCGContextPtr FPCGSubgraphElement::Initialize(const FPCGDataCollection& InputData, const UPCGComponent* SourceComponent)
+{
+	TSharedPtr<FPCGSubgraphContext> Context = MakeShared<FPCGSubgraphContext>();
+	Context->InputData = InputData;
+	Context->SourceComponent = SourceComponent;
+
+	return StaticCastSharedPtr<FPCGContext>(Context);
+}
+
+bool FPCGSubgraphElement::ExecuteInternal(FPCGContextPtr InContext) const
+{
+	TSharedPtr<FPCGSubgraphContext> Context = StaticCastSharedPtr<FPCGSubgraphContext>(InContext);
+
+	const UPCGSubgraphNode* SubgraphNode = Cast<const UPCGSubgraphNode>(Context->Node);
+	if (SubgraphNode && SubgraphNode->bDynamicGraph)
+	{
+		TRACE_CPUPROFILER_EVENT_SCOPE(FPCGSubgraphElement::Execute);
+
+		if (!Context->bScheduledSubgraph)
+		{
+			const UPCGSubgraphSettings* Settings = Context->GetInputSettings<UPCGSubgraphSettings>();
+			check(Settings);
+			UPCGGraph* Subgraph = Settings->Subgraph;
+
+			UPCGSubsystem* Subsystem = Context->SourceComponent->GetSubsystem();
+
+			if (Subsystem && Subgraph)
+			{
+				// Dispatch graph to execute with the given information we have
+				// using this node's task id as additional inputs
+				FPCGTaskId SubgraphTaskId = Subsystem->ScheduleGraph(Subgraph, Context->SourceComponent, MakeShared<FPCGInputForwardingElement>(Context->InputData), {});
+
+				Context->SubgraphTaskId = SubgraphTaskId;
+				Context->bScheduledSubgraph = true;
+				Context->bIsPaused = true;
+
+				// add a trivial task after the output task that wakes up this task
+				Subsystem->ScheduleGeneric([Context]() {
+					// Wake up the current task
+					Context->bIsPaused = false;
+					return true;
+					}, { Context->SubgraphTaskId });
+
+				return false;
+			}
+			else
+			{
+				// Job cannot run; cancel
+				Context->OutputData.bCancelExecution = true;
+				return true;
+			}
+		}
+		else if (Context->bIsPaused)
+		{
+			// Should not happen once we skip it in the graph executor
+			return false;
+		}
+		else
+		{
+			// when woken up, get the output data from the subgraph
+			// and copy it to the current context output data, and finally return true
+			UPCGSubsystem* Subsystem = Context->SourceComponent->GetSubsystem();
+			if (Subsystem)
+			{
+				ensure(Subsystem->GetOutputData(Context->SubgraphTaskId, Context->OutputData));
+			}
+			else
+			{
+				// Job cannot run, cancel
+				Context->OutputData.bCancelExecution = true;
+			}
+
+			return true;
+		}
+	}
+	else
+	{
+		Context->OutputData = Context->InputData;
+		return true;
+	}
+}
+
+FPCGInputForwardingElement::FPCGInputForwardingElement(const FPCGDataCollection& InputToForward)
+	: Input(InputToForward)
+{
+
+}
+
+bool FPCGInputForwardingElement::ExecuteInternal(FPCGContextPtr Context) const
+{
+	Context->OutputData = Input;
+	return true;
+}
\ No newline at end of file
diff --git a/Engine/Plugins/Experimental/PCG/Source/PCG/Private/PCGSubsystem.cpp b/Engine/Plugins/Experimental/PCG/Source/PCG/Private/PCGSubsystem.cpp
--- a/Engine/Plugins/Experimental/PCG/Source/PCG/Private/PCGSubsystem.cpp
+++ b/Engine/Plugins/Experimental/PCG/Source/PCG/Private/PCGSubsystem.cpp
@@ -97,6 +97,11 @@ FPCGTaskId UPCGSubsystem::ScheduleComponent(UPCGComponent* PCGComponent, const T
 	}
 }
 
+FPCGTaskId UPCGSubsystem::ScheduleGraph(UPCGGraph* Graph, const UPCGComponent* SourceComponent, FPCGElementPtr InputElement, const TArray<FPCGTaskId>& Dependencies)
+{
+	return GraphExecutor->Schedule(Graph, SourceComponent, InputElement, Dependencies);
+}
+
 ETickableTickType UPCGSubsystem::GetTickableTickType() const
 {
 	return IsTemplate() ? ETickableTickType::Never : ETickableTickType::Always;
@@ -107,6 +112,18 @@ TStatId UPCGSubsystem::GetStatId() const
 	RETURN_QUICK_DECLARE_CYCLE_STAT(UPCGSubsystem, STATGROUP_Tickables);
 }
 
+FPCGTaskId UPCGSubsystem::ScheduleGeneric(TFunction<bool()> InOperation, const TArray<FPCGTaskId>& TaskDependencies)
+{
+	check(GraphExecutor);
+	return GraphExecutor->ScheduleGeneric(InOperation, TaskDependencies);
+}
+
+bool UPCGSubsystem::GetOutputData(FPCGTaskId TaskId, FPCGDataCollection& OutData)
+{
+	check(GraphExecutor);
+	return GraphExecutor->GetOutputData(TaskId, OutData);
+}
+
 #if WITH_EDITOR
 
 namespace PCGSubsystem
@@ -463,12 +480,6 @@ void UPCGSubsystem::DeletePartitionActors()
 	}
 }
 
-FPCGTaskId UPCGSubsystem::ScheduleGeneric(TFunction<bool()> InOperation, const TArray<FPCGTaskId>& TaskDependencies)
-{
-	check(GraphExecutor);
-	return GraphExecutor->ScheduleGeneric(InOperation, TaskDependencies);
-}
-
 void UPCGSubsystem::NotifyGraphChanged(UPCGGraph* InGraph)
 {
 	if (GraphExecutor)
diff --git a/Engine/Plugins/Experimental/PCG/Source/PCG/Public/PCGContext.h b/Engine/Plugins/Experimental/PCG/Source/PCG/Public/PCGContext.h
--- a/Engine/Plugins/Experimental/PCG/Source/PCG/Public/PCGContext.h
+++ b/Engine/Plugins/Experimental/PCG/Source/PCG/Public/PCGContext.h
@@ -5,6 +5,7 @@
 #include "Templates/SharedPointer.h"
 #include "PCGData.h"
 #include "PCGNode.h"
+#include "PCGSubsystem.h"
 
 struct FPCGContext;
 class UPCGComponent;
@@ -13,12 +14,16 @@ typedef TSharedPtr<FPCGContext, ESPMode::ThreadSafe> FPCGContextPtr;
 
 struct FPCGContext
 {
+	virtual ~FPCGContext() = default;
+
 	FPCGDataCollection InputData;
 	FPCGDataCollection OutputData;
 	const UPCGComponent* SourceComponent = nullptr;
 	// TODO: add RNG source
 	// TODO: replace this by a better identification mechanism
 	const UPCGNode* Node = nullptr;
+	FPCGTaskId TaskId = InvalidTaskId;
+	bool bIsPaused = false;
 
 	template<typename SettingsType>
 	const SettingsType* GetInputSettings()
diff --git a/Engine/Plugins/Experimental/PCG/Source/PCG/Public/PCGSubgraph.h b/Engine/Plugins/Experimental/PCG/Source/PCG/Public/PCGSubgraph.h
--- a/Engine/Plugins/Experimental/PCG/Source/PCG/Public/PCGSubgraph.h
+++ b/Engine/Plugins/Experimental/PCG/Source/PCG/Public/PCGSubgraph.h
@@ -2,6 +2,7 @@
 
 #pragma once
 
+#include "CoreMinimal.h"
 #include "PCGSettings.h"
 
 class UPCGGraph;
@@ -52,20 +53,31 @@ public:
 #endif
 };
 
-UCLASS(ClassGroup = (Procedural))
-class UPCGSubgraphNode : public UPCGNode
+UCLASS(Abstract)
+class UPCGBaseSubgraphNode : public UPCGNode
 {
 	GENERATED_BODY()
 
 public:
 	UPROPERTY(BlueprintReadWrite, EditAnywhere, Category = Settings)
-	bool bAllowDynamicGraph = false;
+	bool bDynamicGraph = false;
+
+	virtual TObjectPtr<UPCGGraph> GetSubgraph() const { return nullptr; }
+};
+
+UCLASS(ClassGroup = (Procedural))
+class UPCGSubgraphNode : public UPCGBaseSubgraphNode
+{
+	GENERATED_BODY()
 
+public:
 #if WITH_EDITOR
 	FOnPCGNodeStructuralSettingsChanged OnNodeStructuralSettingsChangedDelegate;
 #endif
 
-	TObjectPtr<UPCGGraph> GetGraph() const;
+	/** ~Begin UPCGBaseSubgraphNode interface */
+	TObjectPtr<UPCGGraph> GetSubgraph() const override;
+	/** ~End UPCGBaseSubgraphNode interface */
 
 protected:	
 	/** ~Begin UObject interface */
@@ -78,4 +90,29 @@ protected:
 	virtual void PostEditChangeProperty(struct FPropertyChangedEvent& PropertyChangedEvent) override;
 	void OnStructuralSettingsChanged(UPCGSettings* InSettings);
 #endif
+};
+
+struct FPCGSubgraphContext : public FPCGContext
+{
+	FPCGTaskId SubgraphTaskId = InvalidTaskId;
+	bool bScheduledSubgraph = false;
+};
+
+class FPCGSubgraphElement : public IPCGElement
+{
+public:
+	virtual FPCGContextPtr Initialize(const FPCGDataCollection& InputData, const UPCGComponent* SourceComponent) override;
+
+protected:
+	virtual bool ExecuteInternal(FPCGContextPtr Context) const override;
+};
+
+class FPCGInputForwardingElement : public FSimplePCGElement
+{
+public:
+	FPCGInputForwardingElement(const FPCGDataCollection& InputToForward);
+
+protected:
+	virtual bool ExecuteInternal(FPCGContextPtr Context) const override;
+	FPCGDataCollection Input;
 };
\ No newline at end of file
diff --git a/Engine/Plugins/Experimental/PCG/Source/PCG/Public/PCGSubsystem.h b/Engine/Plugins/Experimental/PCG/Source/PCG/Public/PCGSubsystem.h
--- a/Engine/Plugins/Experimental/PCG/Source/PCG/Public/PCGSubsystem.h
+++ b/Engine/Plugins/Experimental/PCG/Source/PCG/Public/PCGSubsystem.h
@@ -11,10 +11,14 @@ class FPCGGraphExecutor;
 class APCGPartitionActor;
 class APCGWorldActor;
 class UPCGGraph;
+struct FPCGDataCollection;
 
 using FPCGTaskId = uint64;
 static const FPCGTaskId InvalidTaskId = (uint64)-1;
 
+class IPCGElement;
+typedef TSharedPtr<IPCGElement, ESPMode::ThreadSafe> FPCGElementPtr;
+
 /**
 * UPCGSubsystem
 */
@@ -47,6 +51,15 @@ public:
 	// Schedule graph (owner -> graph)
 	FPCGTaskId ScheduleComponent(UPCGComponent* PCGComponent, const TArray<FPCGTaskId>& Dependencies);
 
+	// Schedule graph (used internally for dynamic subgraph execution)
+	FPCGTaskId ScheduleGraph(UPCGGraph* Graph, const UPCGComponent* SourceComponent, FPCGElementPtr InputElement, const TArray<FPCGTaskId>& Dependencies);
+
+	/** General job scheduling, used to control loading/unloading */
+	FPCGTaskId ScheduleGeneric(TFunction<bool()> InOperation, const TArray<FPCGTaskId>& TaskDependencies);
+
+	/** Gets the output data for a given task */
+	bool GetOutputData(FPCGTaskId InTaskId, FPCGDataCollection& OutData);
+
 #if WITH_EDITOR
 public:
 	/** Operations to schedule a later operation, which enables to delay bounds querying */
@@ -64,8 +77,6 @@ public:
 	void CleanupPartitionActors(const FBox& InBounds);
 	void DeletePartitionActors();
 
-	/** General job scheduling, used to control loading/unloading */
-	FPCGTaskId ScheduleGeneric(TFunction<bool()> InOperation, const TArray<FPCGTaskId>& TaskDependencies);
 
 	/** Propagate to the graph compiler graph changes */
 	void NotifyGraphChanged(UPCGGraph* InGraph);
