Added template actor spawner incl. PCG awareness & efficiency

diff --git a/Engine/Plugins/Experimental/PCG/Source/PCG/Private/Elements/PCGSpawnActor.cpp b/Engine/Plugins/Experimental/PCG/Source/PCG/Private/Elements/PCGSpawnActor.cpp
new file mode 100644
--- /dev/null
+++ b/Engine/Plugins/Experimental/PCG/Source/PCG/Private/Elements/PCGSpawnActor.cpp
@@ -0,0 +1,208 @@
+// Copyright Epic Games, Inc. All Rights Reserved.
+
+#include "Elements/PCGSpawnActor.h"
+#include "Data/PCGPointData.h"
+#include "Helpers/PCGActorHelpers.h"
+#include "PCGComponent.h"
+#include "PCGHelpers.h"
+
+#include "Components/InstancedStaticMeshComponent.h"
+#include "GameFramework/Actor.h"
+#include "Engine/World.h"
+
+UPCGNode* UPCGSpawnActorSettings::CreateNode() const
+{
+	return NewObject<UPCGSpawnActorNode>();
+}
+
+FPCGElementPtr UPCGSpawnActorSettings::CreateElement() const
+{
+	return MakeShared<FPCGSpawnActorElement>();
+}
+
+UPCGGraph* UPCGSpawnActorSettings::GetSubgraph() const
+{
+	if(!TemplateActorClass || TemplateActorClass->HasAnyClassFlags(CLASS_Abstract))
+	{
+		return nullptr;
+	}
+
+	TArray<UActorComponent*> PCGComponents;
+	UPCGActorHelpers::GetActorClassDefaultComponents(TemplateActorClass, PCGComponents, UPCGComponent::StaticClass());
+
+	if (PCGComponents.IsEmpty())
+	{
+		return nullptr;
+	}
+
+	for (UActorComponent* Component : PCGComponents)
+	{
+		if (UPCGComponent* PCGComponent = Cast<UPCGComponent>(Component))
+		{
+			if (PCGComponent->GetGraph() && PCGComponent->bActivated)
+			{
+				return PCGComponent->GetGraph();
+			}
+		}
+	}
+
+	return nullptr;
+}
+
+#if WITH_EDITOR
+bool UPCGSpawnActorSettings::IsStructuralProperty(const FName& InPropertyName) const
+{
+	return InPropertyName == GET_MEMBER_NAME_CHECKED(UPCGSpawnActorSettings, TemplateActorClass) || 
+		InPropertyName == GET_MEMBER_NAME_CHECKED(UPCGSpawnActorSettings, Option) ||
+		Super::IsStructuralProperty(InPropertyName);
+}
+#endif // WITH_EDITOR
+
+TObjectPtr<UPCGGraph> UPCGSpawnActorNode::GetSubgraph() const
+{
+	TObjectPtr<UPCGSpawnActorSettings> Settings = Cast<UPCGSpawnActorSettings>(DefaultSettings);
+	return (Settings && Settings->Option != EPCGSpawnActorOption::NoMerging) ? Settings->GetSubgraph() : nullptr;
+}
+
+bool FPCGSpawnActorElement::ExecuteInternal(FPCGContextPtr Context) const
+{
+	TRACE_CPUPROFILER_EVENT_SCOPE(FPCSpawnActorElement::Execute);
+
+	const UPCGSpawnActorSettings* Settings = Context->GetInputSettings<UPCGSpawnActorSettings>();
+	check(Settings);
+
+	// Early out
+	if(!Settings->TemplateActorClass)
+	{
+		return true;
+	}
+
+	bool bShouldPassThroughInputs = (Settings->Option != EPCGSpawnActorOption::NoMerging && Settings->GetSubgraph() != nullptr);
+
+	// Pass-through exclusions & settings
+	TArray<FPCGTaggedData>& Outputs = Context->OutputData.TaggedData;
+
+	TArray<FPCGTaggedData> Inputs = Context->InputData.GetInputs();
+	for (const FPCGTaggedData& Input : Inputs)
+	{
+		const UPCGSpatialData* SpatialData = Cast<UPCGSpatialData>(Input.Data);
+
+		if (!SpatialData)
+		{
+			// Data type mismatch
+			continue;
+		}
+
+		AActor* TargetActor = SpatialData->TargetActor;
+
+		if (!TargetActor)
+		{
+			// No target actor
+			continue;
+		}
+
+		// First, create target instance transforms
+		const UPCGPointData* PointData = SpatialData->ToPointData();
+
+		if (!PointData)
+		{
+			continue;
+		}
+
+		const TArray<FPCGPoint>& Points = PointData->GetPoints();
+
+		if (Points.IsEmpty())
+		{
+			continue;
+		}
+
+		// Spawn actors/populate ISM
+		{
+			UInstancedStaticMeshComponent* ISMC = nullptr;
+			TArray<FTransform> Instances;
+
+			// If we are collapsing actors, we need to get the mesh & prep the ISMC
+			if (Settings->Option == EPCGSpawnActorOption::CollapseActors)
+			{
+				TArray<UActorComponent*> Components;
+				UPCGActorHelpers::GetActorClassDefaultComponents(Settings->TemplateActorClass, Components, UStaticMeshComponent::StaticClass());
+				UStaticMesh* Mesh = nullptr;
+
+				for (UActorComponent* Component : Components)
+				{
+					if (UStaticMeshComponent* SMC = Cast<UStaticMeshComponent>(Component))
+					{
+						Mesh = SMC->GetStaticMesh();
+						if (Mesh)
+						{
+							break;
+						}
+					}
+				}
+
+				if (Mesh)
+				{
+					ISMC = UPCGActorHelpers::GetOrCreateISMC(TargetActor, Context->SourceComponent, Mesh);
+				}
+			}
+
+			{
+				TRACE_CPUPROFILER_EVENT_SCOPE(FPCGSpawnActorElement::ExecuteInternal::SpawnActors);
+				if (Settings->Option == EPCGSpawnActorOption::CollapseActors && ISMC)
+				{
+					for (const FPCGPoint& Point : Points)
+					{
+						Instances.Add(Point.Transform);
+					}
+				}
+				else if (Settings->Option != EPCGSpawnActorOption::CollapseActors)
+				{
+					FActorSpawnParameters SpawnParams;
+					SpawnParams.Owner = TargetActor;
+					SpawnParams.SpawnCollisionHandlingOverride = ESpawnActorCollisionHandlingMethod::AlwaysSpawn;
+
+					for (const FPCGPoint& Point : Points)
+					{
+						AActor* GeneratedActor = TargetActor->GetWorld()->SpawnActor(Settings->TemplateActorClass, &Point.Transform, SpawnParams);
+						GeneratedActor->Tags.Add(PCGHelpers::DefaultPCGActorTag);
+						GeneratedActor->AttachToActor(TargetActor, FAttachmentTransformRules::KeepWorldTransform);
+
+						Context->SourceComponent->AddToGeneratedActors(GeneratedActor);
+
+						// If the actor spawned has a PCG component, either generate it or mark it as generated if we pass through its inputs
+						TArray<UPCGComponent*> PCGComponents;
+						GeneratedActor->GetComponents<UPCGComponent>(PCGComponents);
+
+						for (UPCGComponent* PCGComponent : PCGComponents)
+						{
+							if (!bShouldPassThroughInputs)
+							{
+								PCGComponent->Generate();
+							}
+							else
+							{
+								PCGComponent->bActivated = false;
+							}
+						}
+					}
+				}
+			}
+
+			// Finalize
+			if (ISMC && Instances.Num() > 0)
+			{
+				ISMC->NumCustomDataFloats = 0;
+				ISMC->AddInstances(Instances, false, true);
+				ISMC->UpdateBounds();
+			}
+		}
+
+		// Finally, pass through the input if we're merging the PCG here
+		if (bShouldPassThroughInputs)
+		{
+			Outputs.Add(Input);
+		}
+	}
+
+	return true;
+}
\ No newline at end of file
diff --git a/Engine/Plugins/Experimental/PCG/Source/PCG/Private/Graph/PCGGraphExecutor.cpp b/Engine/Plugins/Experimental/PCG/Source/PCG/Private/Graph/PCGGraphExecutor.cpp
--- a/Engine/Plugins/Experimental/PCG/Source/PCG/Private/Graph/PCGGraphExecutor.cpp
+++ b/Engine/Plugins/Experimental/PCG/Source/PCG/Private/Graph/PCGGraphExecutor.cpp
@@ -33,7 +33,7 @@ FPCGTaskId FPCGGraphExecutor::Schedule(UPCGComponent* Component, const TArray<FP
 	return Schedule(Graph, Component, MakeShared<FPCGFetchInputElement>(Component), ExternalDependencies);
 }
 
-FPCGTaskId FPCGGraphExecutor::Schedule(UPCGGraph* Graph, const UPCGComponent* SourceComponent, FPCGElementPtr InputElement, const TArray<FPCGTaskId>& ExternalDependencies)
+FPCGTaskId FPCGGraphExecutor::Schedule(UPCGGraph* Graph, UPCGComponent* SourceComponent, FPCGElementPtr InputElement, const TArray<FPCGTaskId>& ExternalDependencies)
 {
 	FPCGTaskId ScheduledId = InvalidTaskId;
 
diff --git a/Engine/Plugins/Experimental/PCG/Source/PCG/Private/Graph/PCGGraphExecutor.h b/Engine/Plugins/Experimental/PCG/Source/PCG/Private/Graph/PCGGraphExecutor.h
--- a/Engine/Plugins/Experimental/PCG/Source/PCG/Private/Graph/PCGGraphExecutor.h
+++ b/Engine/Plugins/Experimental/PCG/Source/PCG/Private/Graph/PCGGraphExecutor.h
@@ -22,7 +22,7 @@ struct FPCGGraphTask
 	TArray<FPCGTaskId> Inputs;
 	//TArray<DataId> Outputs;
 	const UPCGNode* Node = nullptr;
-	const UPCGComponent* SourceComponent = nullptr; // TODO: can technically be editor only
+	UPCGComponent* SourceComponent = nullptr;
 	FPCGElementPtr Element; // Added to have tasks that aren't node-bound
 	FPCGTaskId NodeId = InvalidTaskId;
 };
@@ -50,7 +50,7 @@ public:
 
 	/** Schedules the execution of a given graph with specified inputs. This call is threadsafe */
 	FPCGTaskId Schedule(UPCGComponent* InComponent, const TArray<FPCGTaskId>& TaskDependency = TArray<FPCGTaskId>());
-	FPCGTaskId Schedule(UPCGGraph* Graph, const UPCGComponent* InSourceComponent, FPCGElementPtr InputElement, const TArray<FPCGTaskId>& TaskDependency);
+	FPCGTaskId Schedule(UPCGGraph* Graph, UPCGComponent* InSourceComponent, FPCGElementPtr InputElement, const TArray<FPCGTaskId>& TaskDependency);
 
 	/** General job scheduling, used to control loading/unloading */
 	FPCGTaskId ScheduleGeneric(TFunction<bool()> InOperation, const TArray<FPCGTaskId>& TaskDependencies);
diff --git a/Engine/Plugins/Experimental/PCG/Source/PCG/Private/Helpers/PCGActorHelpers.cpp b/Engine/Plugins/Experimental/PCG/Source/PCG/Private/Helpers/PCGActorHelpers.cpp
--- a/Engine/Plugins/Experimental/PCG/Source/PCG/Private/Helpers/PCGActorHelpers.cpp
+++ b/Engine/Plugins/Experimental/PCG/Source/PCG/Private/Helpers/PCGActorHelpers.cpp
@@ -7,7 +7,20 @@
 #include "PCGComponent.h"
 #include "Components/InstancedStaticMeshComponent.h"
 #include "Components/HierarchicalInstancedStaticMeshComponent.h"
+#include "Engine/InheritableComponentHandler.h"
+#include "Engine/SCS_Node.h"
+#include "Engine/SimpleConstructionScript.h"
 #include "Engine/StaticMesh.h"
+#include "Engine/World.h"
+#include "WorldPartition/WorldPartition.h"
+
+#if WITH_EDITOR
+#include "Editor.h"
+#include "SourceControlHelpers.h"
+#include "PackageSourceControlHelper.h"
+#include "FileHelpers.h"
+#include "ObjectTools.h"
+#endif
 
 UInstancedStaticMeshComponent* UPCGActorHelpers::GetOrCreateISMC(AActor* InTargetActor, const UPCGComponent* InSourceComponent, UStaticMesh* InMesh, const TArray<UMaterialInterface*>& InMaterials)
 {
@@ -81,4 +94,153 @@ UInstancedStaticMeshComponent* UPCGActorHelpers::GetOrCreateISMC(AActor* InTarge
 	ISMC->ComponentTags.Add(PCGHelpers::DefaultPCGTag);
 
 	return ISMC;
+}
+
+bool UPCGActorHelpers::DeleteActors(UWorld* World, const TArray<TSoftObjectPtr<AActor>>& ActorsToDelete)
+{
+	check(World);
+
+	if (ActorsToDelete.Num() == 0)
+	{
+		return true;
+	}
+
+#if WITH_EDITOR
+	// Remove potential references to to-be deleted objects from the global selection sets.
+	/*if (GIsEditor)
+	{
+		GEditor->ResetAllSelectionSets();
+	}*/
+
+	UWorldPartition* WorldPartition = World ? World->GetWorldPartition() : nullptr;
+
+	if (WorldPartition)
+	{
+		TArray<FString> PackagesToDeleteFromSCC;
+		TSet<UPackage*> PackagesToCleanup;
+
+		for (const TSoftObjectPtr<AActor>& ManagedActor : ActorsToDelete)
+		{
+			// If actor is loaded, just remove from world and keep track of package to cleanup
+			if (AActor* Actor = ManagedActor.Get())
+			{
+				if (UPackage* ActorPackage = Actor->GetExternalPackage())
+				{
+					PackagesToCleanup.Emplace(ActorPackage);
+				}
+				
+				World->DestroyActor(Actor);
+			}
+			// Otherwise, get from World Partition.
+			// Note that it is possible that some actors don't exist anymore, so a null here is not a critical condition
+			else if (const FWorldPartitionActorDesc* ActorDesc = WorldPartition->GetActorDesc(ManagedActor.ToSoftObjectPath()))
+			{
+				PackagesToDeleteFromSCC.Emplace(ActorDesc->GetActorPackage().ToString());
+				WorldPartition->RemoveActor(ActorDesc->GetGuid());
+			}
+		}
+
+		// Save currently loaded packages so they get deleted
+		if (PackagesToCleanup.Num() > 0)
+		{
+			ObjectTools::CleanupAfterSuccessfulDelete(PackagesToCleanup.Array(), /*bPerformReferenceCheck=*/true);
+		}
+
+		// Delete outstanding unloaded packages
+		if (PackagesToDeleteFromSCC.Num() > 0)
+		{
+			FPackageSourceControlHelper PackageHelper;
+			if (!PackageHelper.Delete(PackagesToDeleteFromSCC))
+			{
+				return false;
+			}
+		}
+	}
+	else
+#endif
+	{
+		// Not in editor, really unlikely to happen but might be slow
+		for (const TSoftObjectPtr<AActor>& ManagedActor : ActorsToDelete)
+		{
+			if (ManagedActor.Get())
+			{
+				World->DestroyActor(ManagedActor.Get());
+			}
+		}
+	}
+
+	return true;
+}
+
+// Note: this is copied verbatim from AIHelpers.h/.cpp
+void UPCGActorHelpers::GetActorClassDefaultComponents(const TSubclassOf<AActor>& ActorClass, TArray<UActorComponent*>& OutComponents, const TSubclassOf<UActorComponent>& InComponentClass)
+{
+	if (!ensure(ActorClass.Get()))
+	{
+		return;
+	}
+
+	UClass* ClassPtr = InComponentClass.Get();
+	TArray<UActorComponent*> ResultComponents;
+
+	// Get the components defined on the native class.
+	AActor* CDO = ActorClass->GetDefaultObject<AActor>();
+	check(CDO);
+	if (ClassPtr)
+	{
+		CDO->GetComponents(InComponentClass, ResultComponents);
+	}
+	else
+	{
+		CDO->GetComponents(ResultComponents);
+	}
+
+	// Try to get the components off the BP class.
+	UBlueprintGeneratedClass* BPClass = Cast<UBlueprintGeneratedClass>(*ActorClass);
+	if (BPClass)
+	{
+		// A BlueprintGeneratedClass has a USimpleConstructionScript member. This member has an array of RootNodes
+		// which contains the SCSNode for the root SceneComponent and non-SceneComponents. For the SceneComponent
+		// hierarchy, each SCSNode knows its children SCSNodes. Each SCSNode stores the component template that will
+		// be created when the Actor is spawned.
+		//
+		// WARNING: This may change in future engine versions!
+
+		TArray<UActorComponent*> Unfiltered;
+		// using this semantic to avoid duplicating following loops or adding a filtering check condition inside the loops
+		TArray<UActorComponent*>& TmpComponents = ClassPtr ? Unfiltered : ResultComponents;
+
+		// Check added components.
+		USimpleConstructionScript* ConstructionScript = BPClass->SimpleConstructionScript;
+		if (ConstructionScript)
+		{
+			for (const USCS_Node* Node : ConstructionScript->GetAllNodes())
+			{
+				TmpComponents.Add(Node->ComponentTemplate);
+			}
+		}
+		// Check modified inherited components.
+		UInheritableComponentHandler* InheritableComponentHandler = BPClass->InheritableComponentHandler;
+		if (InheritableComponentHandler)
+		{
+			for (TArray<FComponentOverrideRecord>::TIterator It = InheritableComponentHandler->CreateRecordIterator(); It; ++It)
+			{
+				TmpComponents.Add(It->ComponentTemplate);
+			}
+		}
+
+		// Filter to the ones matching the requested class.
+		if (ClassPtr)
+		{
+			for (UActorComponent* TemplateComponent : Unfiltered)
+			{
+				if (TemplateComponent->IsA(ClassPtr))
+				{
+					ResultComponents.Add(TemplateComponent);
+				}
+			}
+		}
+	}
+
+	OutComponents = MoveTemp(ResultComponents);
 }
\ No newline at end of file
diff --git a/Engine/Plugins/Experimental/PCG/Source/PCG/Private/PCGComponent.cpp b/Engine/Plugins/Experimental/PCG/Source/PCG/Private/PCGComponent.cpp
--- a/Engine/Plugins/Experimental/PCG/Source/PCG/Private/PCGComponent.cpp
+++ b/Engine/Plugins/Experimental/PCG/Source/PCG/Private/PCGComponent.cpp
@@ -14,6 +14,7 @@
 #include "Data/PCGVolumeData.h"
 #include "Graph/PCGGraphExecutor.h"
 #include "Grid/PCGPartitionActor.h"
+#include "Helpers/PCGActorHelpers.h"
 
 #include "ActorPartition/ActorPartitionSubsystem.h"
 #include "Components/InstancedStaticMeshComponent.h"
@@ -82,6 +83,14 @@ void UPCGComponent::SetGraph(UPCGGraph* InGraph)
 #endif	
 }
 
+void UPCGComponent::AddToGeneratedActors(AActor* InActor)
+{
+	if (InActor)
+	{
+		GeneratedActors.Add(InActor);
+	}
+}
+
 bool UPCGComponent::ShouldGenerate(bool bForce) const
 {
 	if (!bActivated || !Graph || !GetSubsystem())
@@ -321,6 +330,13 @@ void UPCGComponent::Cleanup(bool bRemoveComponents, bool bSave)
 }
 
 void UPCGComponent::CleanupInternal(bool bRemoveComponents)
+{
+	TSet<TSoftObjectPtr<AActor>> ActorsToDelete;
+	CleanupInternal(bRemoveComponents, ActorsToDelete);
+	UPCGActorHelpers::DeleteActors(GetWorld(), ActorsToDelete.Array());
+}
+
+void UPCGComponent::CleanupInternal(bool bRemoveComponents, TSet<TSoftObjectPtr<AActor>>& OutActorsToDelete)
 {
 	if (!bGenerated || IsPartitioned())
 	{
@@ -350,6 +366,28 @@ void UPCGComponent::CleanupInternal(bool bRemoveComponents)
 			}
 		}
 	}
+
+	// Cleanup recursively before deleting the actor(s)
+	OutActorsToDelete.Append(GeneratedActors);
+
+	TArray<UPCGComponent*> ComponentsToCleanup;
+
+	for (TSoftObjectPtr<AActor> GeneratedActor : GeneratedActors)
+	{
+		if (GeneratedActor.IsValid())
+		{
+			GeneratedActor.Get()->GetComponents<UPCGComponent>(ComponentsToCleanup);
+
+			for (UPCGComponent* Component : ComponentsToCleanup)
+			{
+				Component->CleanupInternal(/*bRemoveComponents=*/false, OutActorsToDelete);
+			}
+
+			ComponentsToCleanup.Reset();
+		}
+	}
+
+	GeneratedActors.Reset();
 }
 
 void UPCGComponent::BeginPlay()
diff --git a/Engine/Plugins/Experimental/PCG/Source/PCG/Private/PCGElement.cpp b/Engine/Plugins/Experimental/PCG/Source/PCG/Private/PCGElement.cpp
--- a/Engine/Plugins/Experimental/PCG/Source/PCG/Private/PCGElement.cpp
+++ b/Engine/Plugins/Experimental/PCG/Source/PCG/Private/PCGElement.cpp
@@ -91,7 +91,7 @@ bool IPCGElement::Execute(FPCGContextPtr Context) const
 	}
 }
 
-FPCGContextPtr FSimplePCGElement::Initialize(const FPCGDataCollection& InputData, const UPCGComponent* SourceComponent)
+FPCGContextPtr FSimplePCGElement::Initialize(const FPCGDataCollection& InputData, UPCGComponent* SourceComponent)
 {
 	FPCGContextPtr Context = MakeShared<FPCGContext>();
 	Context->InputData = InputData;
diff --git a/Engine/Plugins/Experimental/PCG/Source/PCG/Private/PCGGraph.cpp b/Engine/Plugins/Experimental/PCG/Source/PCG/Private/PCGGraph.cpp
--- a/Engine/Plugins/Experimental/PCG/Source/PCG/Private/PCGGraph.cpp
+++ b/Engine/Plugins/Experimental/PCG/Source/PCG/Private/PCGGraph.cpp
@@ -58,23 +58,17 @@ void UPCGGraph::InitializeFromTemplate()
 {
 	Modify();
 
-	auto ResetDefaultNode = [](UPCGNode* InNode, bool bIsInput) {
+	auto ResetDefaultNode = [](UPCGNode* InNode) {
 		check(InNode);
-		if (bIsInput)
-		{
-			InNode->OutboundNodes.Reset();
-		}
-		else
-		{
-			InNode->InboundNodes.Reset();
-		}
+		InNode->OutboundNodes.Reset();
+		InNode->InboundNodes.Reset();
 
 		// Reset settings as well
 		InNode->DefaultSettings = NewObject<UPCGTrivialSettings>(InNode);
 	};
 
-	ResetDefaultNode(InputNode, true);
-	ResetDefaultNode(OutputNode, false);
+	ResetDefaultNode(InputNode);
+	ResetDefaultNode(OutputNode);
 
 	for (UPCGNode* Node : Nodes)
 	{
diff --git a/Engine/Plugins/Experimental/PCG/Source/PCG/Private/PCGHelpers.h b/Engine/Plugins/Experimental/PCG/Source/PCG/Private/PCGHelpers.h
--- a/Engine/Plugins/Experimental/PCG/Source/PCG/Private/PCGHelpers.h
+++ b/Engine/Plugins/Experimental/PCG/Source/PCG/Private/PCGHelpers.h
@@ -15,6 +15,7 @@ namespace PCGHelpers
 	/** Tag that will be added on every component generated through the PCG system */
 	const FName DefaultPCGTag = TEXT("PCG Generated Component");
 	const FName DefaultPCGDebugTag = TEXT("PCG Generated Debug Component");
+	const FName DefaultPCGActorTag = TEXT("PCG Generated Actor");
 
 	int ComputeSeed(int A);
 	int ComputeSeed(int A, int B);
diff --git a/Engine/Plugins/Experimental/PCG/Source/PCG/Private/PCGSubgraph.cpp b/Engine/Plugins/Experimental/PCG/Source/PCG/Private/PCGSubgraph.cpp
--- a/Engine/Plugins/Experimental/PCG/Source/PCG/Private/PCGSubgraph.cpp
+++ b/Engine/Plugins/Experimental/PCG/Source/PCG/Private/PCGSubgraph.cpp
@@ -4,22 +4,22 @@
 #include "PCGGraph.h"
 #include "PCGComponent.h"
 
-void UPCGSubgraphSettings::PostLoad()
+void UPCGBaseSubgraphSettings::PostLoad()
 {
 	Super::PostLoad();
 
 #if WITH_EDITOR
-	if (Subgraph)
+	if (UPCGGraph* Subgraph = GetSubgraph())
 	{
 		Subgraph->OnGraphChangedDelegate.AddUObject(this, &UPCGSubgraphSettings::OnSubgraphChanged);
 	}
 #endif
 }
 
-void UPCGSubgraphSettings::BeginDestroy()
+void UPCGBaseSubgraphSettings::BeginDestroy()
 {
 #if WITH_EDITOR
-	if (Subgraph)
+	if (UPCGGraph* Subgraph = GetSubgraph())
 	{
 		Subgraph->OnGraphChangedDelegate.RemoveAll(this);
 	}
@@ -28,30 +28,12 @@ void UPCGSubgraphSettings::BeginDestroy()
 	Super::BeginDestroy();
 }
 
-UPCGNode* UPCGSubgraphSettings::CreateNode() const
-{
-	return NewObject<UPCGSubgraphNode>();
-}
-
 #if WITH_EDITOR
-TArray<FName> UPCGSubgraphSettings::GetTrackedActorTags() const
+void UPCGBaseSubgraphSettings::PreEditChange(FProperty* PropertyAboutToChange)
 {
-	return Subgraph ? Subgraph->GetTrackedActorTags() : TArray<FName>();
-}
-#endif
-
-FPCGElementPtr UPCGSubgraphSettings::CreateElement() const
-{
-	return MakeShared<FPCGSubgraphElement>();
-}
-
-#if WITH_EDITOR
-
-void UPCGSubgraphSettings::PreEditChange(FProperty* PropertyAboutToChange)
-{
-	if (PropertyAboutToChange && PropertyAboutToChange->GetFName() == GET_MEMBER_NAME_CHECKED(UPCGSubgraphSettings, Subgraph))
+	if (PropertyAboutToChange && IsStructuralProperty(PropertyAboutToChange->GetFName()))
 	{
-		if (Subgraph)
+		if (UPCGGraph* Subgraph = GetSubgraph())
 		{
 			Subgraph->OnGraphChangedDelegate.RemoveAll(this);
 		}
@@ -60,22 +42,22 @@ void UPCGSubgraphSettings::PreEditChange(FProperty* PropertyAboutToChange)
 	Super::PreEditChange(PropertyAboutToChange);
 }
 
-void UPCGSubgraphSettings::PostEditChangeProperty(struct FPropertyChangedEvent& PropertyChangedEvent)
+void UPCGBaseSubgraphSettings::PostEditChangeProperty(struct FPropertyChangedEvent& PropertyChangedEvent)
 {
 	Super::PostEditChangeProperty(PropertyChangedEvent);
 
-	if (PropertyChangedEvent.Property->GetFName() == GET_MEMBER_NAME_CHECKED(UPCGSubgraphSettings, Subgraph))
+	if (PropertyChangedEvent.Property && IsStructuralProperty(PropertyChangedEvent.Property->GetFName()))
 	{
-		if (Subgraph)
+		if (UPCGGraph* Subgraph = GetSubgraph())
 		{
 			Subgraph->OnGraphChangedDelegate.AddUObject(this, &UPCGSubgraphSettings::OnSubgraphChanged);
 		}
 	}
 }
 
-void UPCGSubgraphSettings::OnSubgraphChanged(UPCGGraph* InGraph, bool bIsStructural)
+void UPCGBaseSubgraphSettings::OnSubgraphChanged(UPCGGraph* InGraph, bool bIsStructural)
 {
-	if (InGraph == Subgraph)
+	if (InGraph == GetSubgraph())
 	{
 		if (bIsStructural)
 		{
@@ -90,6 +72,28 @@ void UPCGSubgraphSettings::OnSubgraphChanged(UPCGGraph* InGraph, bool bIsStructu
 
 #endif // WITH_EDITOR
 
+UPCGNode* UPCGSubgraphSettings::CreateNode() const
+{
+	return NewObject<UPCGSubgraphNode>();
+}
+
+#if WITH_EDITOR
+bool UPCGSubgraphSettings::IsStructuralProperty(const FName& InPropertyName) const
+{
+	return (InPropertyName == GET_MEMBER_NAME_CHECKED(UPCGSubgraphSettings, Subgraph)) || Super::IsStructuralProperty(InPropertyName);
+}
+
+TArray<FName> UPCGSubgraphSettings::GetTrackedActorTags() const
+{
+	return Subgraph ? Subgraph->GetTrackedActorTags() : TArray<FName>();
+}
+#endif
+
+FPCGElementPtr UPCGSubgraphSettings::CreateElement() const
+{
+	return MakeShared<FPCGSubgraphElement>();
+}
+
 TObjectPtr<UPCGGraph> UPCGSubgraphNode::GetSubgraph() const
 {
 	TObjectPtr<UPCGSubgraphSettings> Settings = Cast<UPCGSubgraphSettings>(DefaultSettings);
@@ -162,7 +166,7 @@ void UPCGSubgraphNode::OnStructuralSettingsChanged(UPCGSettings* InSettings)
 }
 #endif // WITH_EDITOR
 
-FPCGContextPtr FPCGSubgraphElement::Initialize(const FPCGDataCollection& InputData, const UPCGComponent* SourceComponent)
+FPCGContextPtr FPCGSubgraphElement::Initialize(const FPCGDataCollection& InputData, UPCGComponent* SourceComponent)
 {
 	TSharedPtr<FPCGSubgraphContext> Context = MakeShared<FPCGSubgraphContext>();
 	Context->InputData = InputData;
diff --git a/Engine/Plugins/Experimental/PCG/Source/PCG/Private/PCGSubsystem.cpp b/Engine/Plugins/Experimental/PCG/Source/PCG/Private/PCGSubsystem.cpp
--- a/Engine/Plugins/Experimental/PCG/Source/PCG/Private/PCGSubsystem.cpp
+++ b/Engine/Plugins/Experimental/PCG/Source/PCG/Private/PCGSubsystem.cpp
@@ -97,7 +97,7 @@ FPCGTaskId UPCGSubsystem::ScheduleComponent(UPCGComponent* PCGComponent, const T
 	}
 }
 
-FPCGTaskId UPCGSubsystem::ScheduleGraph(UPCGGraph* Graph, const UPCGComponent* SourceComponent, FPCGElementPtr InputElement, const TArray<FPCGTaskId>& Dependencies)
+FPCGTaskId UPCGSubsystem::ScheduleGraph(UPCGGraph* Graph, UPCGComponent* SourceComponent, FPCGElementPtr InputElement, const TArray<FPCGTaskId>& Dependencies)
 {
 	return GraphExecutor->Schedule(Graph, SourceComponent, InputElement, Dependencies);
 }
diff --git a/Engine/Plugins/Experimental/PCG/Source/PCG/Public/Elements/PCGSpawnActor.h b/Engine/Plugins/Experimental/PCG/Source/PCG/Public/Elements/PCGSpawnActor.h
new file mode 100644
--- /dev/null
+++ b/Engine/Plugins/Experimental/PCG/Source/PCG/Public/Elements/PCGSpawnActor.h
@@ -0,0 +1,70 @@
+// Copyright Epic Games, Inc. All Rights Reserved.
+
+#pragma once
+
+#include "CoreMinimal.h"
+#include "PCGElement.h"
+#include "PCGSettings.h"
+#include "PCGSubgraph.h"
+#include "Templates/SubclassOf.h"
+
+#include "PCGSpawnActor.generated.h"
+
+class AActor;
+
+UENUM()
+enum class EPCGSpawnActorOption : uint8
+{
+	CollapseActors,
+	MergePCGOnly,
+	NoMerging
+};
+
+UCLASS(BlueprintType, ClassGroup = (Procedural))
+class UPCGSpawnActorSettings : public UPCGBaseSubgraphSettings
+{
+	GENERATED_BODY()
+
+public:
+	UPROPERTY(BlueprintReadWrite, EditAnywhere, Category = Settings)
+	TSubclassOf<AActor> TemplateActorClass = nullptr;
+
+	UPROPERTY(BlueprintReadWrite, EditAnywhere, Category = Settings)
+	EPCGSpawnActorOption Option = EPCGSpawnActorOption::CollapseActors;
+
+	//~Begin UCPGSettings interface
+	virtual UPCGNode* CreateNode() const override;
+
+#if WITH_EDITOR	
+	virtual FName GetDefaultNodeName() const override { return FName(TEXT("SpawnActorNode")); }
+#endif
+protected:
+	virtual FPCGElementPtr CreateElement() const override;
+	//~End UPCGSettings interface
+
+public:
+	//~Begin UPCGBaseSubgraphSettings interface
+	virtual UPCGGraph* GetSubgraph() const override;
+
+protected:
+#if WITH_EDITOR
+	virtual bool IsStructuralProperty(const FName& InPropertyName) const override;
+#endif
+	//~End UPCGBaseSubgraphSettings interface
+};
+
+UCLASS(ClassGroup = (Procedural))
+class UPCGSpawnActorNode : public UPCGBaseSubgraphNode
+{
+	GENERATED_BODY()
+public:
+	/** ~Begin UPCGBaseSubgraphNode interface */
+	virtual TObjectPtr<UPCGGraph> GetSubgraph() const override;
+	/** ~End UPCGBaseSubgraphNode interface */
+};
+
+class FPCGSpawnActorElement : public FSimpleTypedPCGElement<UPCGSpawnActorSettings>
+{
+protected:
+	virtual bool ExecuteInternal(FPCGContextPtr Context) const override;
+};
\ No newline at end of file
diff --git a/Engine/Plugins/Experimental/PCG/Source/PCG/Public/Helpers/PCGActorHelpers.h b/Engine/Plugins/Experimental/PCG/Source/PCG/Public/Helpers/PCGActorHelpers.h
--- a/Engine/Plugins/Experimental/PCG/Source/PCG/Public/Helpers/PCGActorHelpers.h
+++ b/Engine/Plugins/Experimental/PCG/Source/PCG/Public/Helpers/PCGActorHelpers.h
@@ -4,6 +4,7 @@
 
 #include "CoreMinimal.h"
 #include "Kismet/BlueprintFunctionLibrary.h"
+#include "Templates/SubclassOf.h"
 
 #include "PCGActorHelpers.generated.h"
 
@@ -12,6 +13,7 @@ class UInstancedStaticMeshComponent;
 class UStaticMesh;
 class UPCGComponent;
 class UMaterialInterface;
+class UActorComponent;
 
 UCLASS(BlueprintType)
 class PCG_API UPCGActorHelpers : public UBlueprintFunctionLibrary
@@ -20,4 +22,13 @@ class PCG_API UPCGActorHelpers : public UBlueprintFunctionLibrary
 
 public:
 	static UInstancedStaticMeshComponent* GetOrCreateISMC(AActor* InActor, const UPCGComponent* SourceComponent, UStaticMesh* InMesh, const TArray<UMaterialInterface*>& InMaterials = TArray<UMaterialInterface*>());
+	static bool DeleteActors(UWorld* World, const TArray<TSoftObjectPtr<AActor>>& ActorsToDelete);
+
+	/**
+	* Fetches all the components of ActorClass's CDO, including the ones added via the BP editor (which AActor.GetComponents fails to do)
+	* @param ActorClass class of AActor for which we will retrieve all components
+	* @param OutComponents this is where the found components will end up. Note that the preexisting contents of OutComponents will get overridden.
+	* @param InComponentClass if supplied will be used to filter the results
+	*/
+	static void GetActorClassDefaultComponents(const TSubclassOf<AActor>& ActorClass, TArray<UActorComponent*>& OutComponents, const TSubclassOf<UActorComponent>& InComponentClass = TSubclassOf<UActorComponent>());
 };
\ No newline at end of file
diff --git a/Engine/Plugins/Experimental/PCG/Source/PCG/Public/PCGComponent.h b/Engine/Plugins/Experimental/PCG/Source/PCG/Public/PCGComponent.h
--- a/Engine/Plugins/Experimental/PCG/Source/PCG/Public/PCGComponent.h
+++ b/Engine/Plugins/Experimental/PCG/Source/PCG/Public/PCGComponent.h
@@ -56,6 +56,8 @@ public:
 	UPCGGraph* GetGraph() const { return Graph; }
 	void SetGraph(UPCGGraph* InGraph);
 
+	void AddToGeneratedActors(AActor* InActor);
+
 	/** Transactionable methods to be called from details UI */
 	void Generate();
 	void Cleanup();
@@ -130,6 +132,7 @@ private:
 	bool ShouldGenerate(bool bForce = false) const;
 	FPCGTaskId GenerateInternal(bool bForce, const TArray<FPCGTaskId>& Dependencies);
 	void CleanupInternal(bool bRemoveComponents);
+	void CleanupInternal(bool bRemoveComponents, TSet<TSoftObjectPtr<AActor>>& OutActorsToDelete);
 	void PostProcessGraph(const FBox& InNewBounds, bool bInGenerated);
 
 	bool GetActorsFromTags(const TSet<FName>& InTags, TSet<TWeakObjectPtr<AActor>>& OutActors);
@@ -173,6 +176,9 @@ private:
 	UPROPERTY()
 	TSet<TWeakObjectPtr<AActor>> CachedExcludedActors; 
 
+	UPROPERTY()
+	TSet<TSoftObjectPtr<AActor>> GeneratedActors;
+
 	UPROPERTY()
 	FBox LastGeneratedBounds = FBox(EForceInit::ForceInit);
 
diff --git a/Engine/Plugins/Experimental/PCG/Source/PCG/Public/PCGContext.h b/Engine/Plugins/Experimental/PCG/Source/PCG/Public/PCGContext.h
--- a/Engine/Plugins/Experimental/PCG/Source/PCG/Public/PCGContext.h
+++ b/Engine/Plugins/Experimental/PCG/Source/PCG/Public/PCGContext.h
@@ -18,7 +18,7 @@ struct FPCGContext
 
 	FPCGDataCollection InputData;
 	FPCGDataCollection OutputData;
-	const UPCGComponent* SourceComponent = nullptr;
+	UPCGComponent* SourceComponent = nullptr;
 	// TODO: add RNG source
 	// TODO: replace this by a better identification mechanism
 	const UPCGNode* Node = nullptr;
diff --git a/Engine/Plugins/Experimental/PCG/Source/PCG/Public/PCGElement.h b/Engine/Plugins/Experimental/PCG/Source/PCG/Public/PCGElement.h
--- a/Engine/Plugins/Experimental/PCG/Source/PCG/Public/PCGElement.h
+++ b/Engine/Plugins/Experimental/PCG/Source/PCG/Public/PCGElement.h
@@ -18,7 +18,7 @@ class IPCGElement
 {
 public:
 	virtual ~IPCGElement() = default;
-	virtual FPCGContextPtr Initialize(const FPCGDataCollection& InputData, const UPCGComponent* SourceComponent) = 0;
+	virtual FPCGContextPtr Initialize(const FPCGDataCollection& InputData, UPCGComponent* SourceComponent) = 0;
 
 	bool Execute(FPCGContextPtr Context) const;
 
@@ -33,7 +33,7 @@ protected:
 class FSimplePCGElement : public IPCGElement
 {
 public:
-	virtual FPCGContextPtr Initialize(const FPCGDataCollection& InputData, const UPCGComponent* SourceComponent) override;
+	virtual FPCGContextPtr Initialize(const FPCGDataCollection& InputData, UPCGComponent* SourceComponent) override;
 };
 
 /**
diff --git a/Engine/Plugins/Experimental/PCG/Source/PCG/Public/PCGSubgraph.h b/Engine/Plugins/Experimental/PCG/Source/PCG/Public/PCGSubgraph.h
--- a/Engine/Plugins/Experimental/PCG/Source/PCG/Public/PCGSubgraph.h
+++ b/Engine/Plugins/Experimental/PCG/Source/PCG/Public/PCGSubgraph.h
@@ -14,8 +14,35 @@ class UPCGGraph;
 	DECLARE_MULTICAST_DELEGATE_OneParam(FOnPCGNodeStructuralSettingsChanged, UPCGNode*);
 #endif
 
+UCLASS()
+class UPCGBaseSubgraphSettings : public UPCGSettings
+{
+	GENERATED_BODY()
+public:
+	virtual UPCGGraph* GetSubgraph() const { return nullptr; }
+
+protected:
+	//~Begin UObject interface implementation
+	virtual void PostLoad() override;
+	virtual void BeginDestroy() override;
+#if WITH_EDITOR
+	virtual void PreEditChange(FProperty* PropertyAboutToChange) override;
+	virtual void PostEditChangeProperty(struct FPropertyChangedEvent& PropertyChangedEvent) override;
+	//~End UObject interface implementation
+
+	void OnSubgraphChanged(UPCGGraph* InGraph, bool bIsStructural);
+
+	virtual bool IsStructuralProperty(const FName& InPropertyName) const { return false; }
+#endif
+
+public:
+#if WITH_EDITOR
+	FOnPCGStructuralSettingsChanged OnStructuralSettingsChangedDelegate;
+#endif
+};
+
 UCLASS(BlueprintType, ClassGroup=(Procedural))
-class UPCGSubgraphSettings : public UPCGSettings
+class UPCGSubgraphSettings : public UPCGBaseSubgraphSettings
 {
 	GENERATED_BODY()
 
@@ -32,25 +59,18 @@ protected:
 	virtual FPCGElementPtr CreateElement() const override;
 	//~End UPCGSettings interface implementation
 
+	//~Begin UPCGBaseSubgraphSettings interface
+public:
+	virtual UPCGGraph* GetSubgraph() const override { return Subgraph; }
 protected:
-	//~Begin UObject interface implementation
-	virtual void PostLoad() override;
-	virtual void BeginDestroy() override;
 #if WITH_EDITOR
-	virtual void PreEditChange(FProperty* PropertyAboutToChange) override;
-	virtual void PostEditChangeProperty(struct FPropertyChangedEvent& PropertyChangedEvent) override;
-	//~End UObject interface implementation
-
-	void OnSubgraphChanged(UPCGGraph* InGraph, bool bIsStructural);
+	virtual bool IsStructuralProperty(const FName& InPropertyName) const override;
 #endif
+	//~End UPCGBaseSubgraphSettings interface
 
 public:
 	UPROPERTY(BlueprintReadWrite, EditAnywhere, Category = Settings)
 	TObjectPtr<UPCGGraph> Subgraph;
-
-#if WITH_EDITOR
-	FOnPCGStructuralSettingsChanged OnStructuralSettingsChangedDelegate;
-#endif
 };
 
 UCLASS(Abstract)
@@ -101,7 +121,7 @@ struct FPCGSubgraphContext : public FPCGContext
 class FPCGSubgraphElement : public IPCGElement
 {
 public:
-	virtual FPCGContextPtr Initialize(const FPCGDataCollection& InputData, const UPCGComponent* SourceComponent) override;
+	virtual FPCGContextPtr Initialize(const FPCGDataCollection& InputData, UPCGComponent* SourceComponent) override;
 
 protected:
 	virtual bool ExecuteInternal(FPCGContextPtr Context) const override;
diff --git a/Engine/Plugins/Experimental/PCG/Source/PCG/Public/PCGSubsystem.h b/Engine/Plugins/Experimental/PCG/Source/PCG/Public/PCGSubsystem.h
--- a/Engine/Plugins/Experimental/PCG/Source/PCG/Public/PCGSubsystem.h
+++ b/Engine/Plugins/Experimental/PCG/Source/PCG/Public/PCGSubsystem.h
@@ -52,7 +52,7 @@ public:
 	FPCGTaskId ScheduleComponent(UPCGComponent* PCGComponent, const TArray<FPCGTaskId>& Dependencies);
 
 	// Schedule graph (used internally for dynamic subgraph execution)
-	FPCGTaskId ScheduleGraph(UPCGGraph* Graph, const UPCGComponent* SourceComponent, FPCGElementPtr InputElement, const TArray<FPCGTaskId>& Dependencies);
+	FPCGTaskId ScheduleGraph(UPCGGraph* Graph, UPCGComponent* SourceComponent, FPCGElementPtr InputElement, const TArray<FPCGTaskId>& Dependencies);
 
 	/** General job scheduling, used to control loading/unloading */
 	FPCGTaskId ScheduleGeneric(TFunction<bool()> InOperation, const TArray<FPCGTaskId>& TaskDependencies);
