Missing abstract tag on base class

diff --git a/Engine/Plugins/Experimental/PCG/Source/PCG/Public/PCGSubgraph.h b/Engine/Plugins/Experimental/PCG/Source/PCG/Public/PCGSubgraph.h
--- a/Engine/Plugins/Experimental/PCG/Source/PCG/Public/PCGSubgraph.h
+++ b/Engine/Plugins/Experimental/PCG/Source/PCG/Public/PCGSubgraph.h
@@ -14,7 +14,7 @@ class UPCGGraph;
 	DECLARE_MULTICAST_DELEGATE_OneParam(FOnPCGNodeStructuralSettingsChanged, UPCGNode*);
 #endif
 
-UCLASS()
+UCLASS(Abstract)
 class UPCGBaseSubgraphSettings : public UPCGSettings
 {
 	GENERATED_BODY()
