Implemented basic data cache

diff --git a/Engine/Plugins/Experimental/PCG/Source/PCG/Private/Elements/PCGExecuteBlueprint.cpp b/Engine/Plugins/Experimental/PCG/Source/PCG/Private/Elements/PCGExecuteBlueprint.cpp
--- a/Engine/Plugins/Experimental/PCG/Source/PCG/Private/Elements/PCGExecuteBlueprint.cpp
+++ b/Engine/Plugins/Experimental/PCG/Source/PCG/Private/Elements/PCGExecuteBlueprint.cpp
@@ -320,4 +320,16 @@ bool FPCGExecuteBlueprintElement::ExecuteInternal(FPCGContextPtr Context) const
 	}
 	
 	return true;
+}
+
+bool FPCGExecuteBlueprintElement::IsCacheable(const UPCGSettings* InSettings) const
+{
+	if (const UPCGBlueprintSettings* BPSettings = Cast<const UPCGBlueprintSettings>(InSettings))
+	{
+		return !BPSettings->bCreatesArtifacts;
+	}
+	else
+	{
+		return false;
+	}
 }
\ No newline at end of file
diff --git a/Engine/Plugins/Experimental/PCG/Source/PCG/Private/Graph/PCGGraphCache.cpp b/Engine/Plugins/Experimental/PCG/Source/PCG/Private/Graph/PCGGraphCache.cpp
new file mode 100644
--- /dev/null
+++ b/Engine/Plugins/Experimental/PCG/Source/PCG/Private/Graph/PCGGraphCache.cpp
@@ -0,0 +1,140 @@
+// Copyright Epic Games, Inc. All Rights Reserved.
+
+#include "PCGGraphCache.h"
+#include "PCGSettings.h"
+
+#include "Misc/ScopeRWLock.h"
+#include "Algo/AnyOf.h"
+#include "GameFramework/Actor.h"
+
+FPCGGraphCacheEntry::FPCGGraphCacheEntry(const FPCGDataCollection& InInput, const UPCGSettings* InSettings, const FPCGDataCollection& InOutput, TWeakObjectPtr<UObject> InOwner, TSet<UObject*>& OutRootedObjects)
+	: Input(InInput)
+	, Output(InOutput)
+{
+	Settings = InSettings ? Cast<UPCGSettings>(StaticDuplicateObject(InSettings, InOwner.Get())) : nullptr;
+
+	Input.RootUnrootedData(OutRootedObjects);
+	Output.RootUnrootedData(OutRootedObjects);
+
+	if (Settings && !Settings->IsRooted())
+	{
+		Settings->AddToRoot();
+		OutRootedObjects.Add(Settings);
+	}
+}
+
+bool FPCGGraphCacheEntry::Matches(const FPCGDataCollection& InInput, const UPCGSettings* InSettings) const
+{
+	const bool bHasSameSettings = ((InSettings == nullptr && Settings == nullptr) || (InSettings && Settings && *InSettings == *Settings));
+	return bHasSameSettings && (Input == InInput);
+}
+
+FPCGGraphCache::FPCGGraphCache(TWeakObjectPtr<UObject> InOwner)
+	: Owner(InOwner)
+{
+}
+
+FPCGGraphCache::~FPCGGraphCache()
+{
+	ClearCache();
+}
+
+bool FPCGGraphCache::GetFromCache(const IPCGElement* InElement, const FPCGDataCollection& InInput, const UPCGSettings* InSettings, FPCGDataCollection& OutOutput) const
+{
+	if (!Owner.IsValid())
+	{
+		return false;
+	}
+
+	FReadScopeLock ScopedReadLock(CacheLock);
+
+	if (const FPCGGraphCacheEntries* Entries = CacheData.Find(InElement))
+	{
+		for (const FPCGGraphCacheEntry& Entry : *Entries)
+		{
+			if (Entry.Matches(InInput, InSettings))
+			{
+				OutOutput = Entry.Output;
+				return true;
+			}
+		}
+
+		return false;
+	}
+
+	return false;
+}
+
+void FPCGGraphCache::StoreInCache(const IPCGElement* InElement, const FPCGDataCollection& InInput, const UPCGSettings* InSettings, const FPCGDataCollection& InOutput)
+{
+	if (!Owner.IsValid())
+	{
+		return;
+	}
+
+	FWriteScopeLock ScopedWriteLock(CacheLock);
+
+	FPCGGraphCacheEntries* Entries = CacheData.Find(InElement);
+	if(!Entries)
+	{
+		Entries = &(CacheData.Add(InElement));
+	}
+
+	Entries->Emplace(InInput, InSettings, InOutput, Owner, RootedData);
+}
+
+void FPCGGraphCache::ClearCache()
+{
+	FWriteScopeLock ScopedWriteLock(CacheLock);
+
+	// Remove all entries
+	CacheData.Reset();
+
+	// Unroot all previously rooted data
+	for (UObject* Data : RootedData)
+	{
+		Data->RemoveFromRoot();
+	}
+	RootedData.Reset();
+}
+
+#if WITH_EDITOR
+void FPCGGraphCache::CleanFromCache(AActor* InActor)
+{
+	if (!InActor || InActor->Tags.IsEmpty())
+	{
+		return;
+	}
+
+	FWriteScopeLock ScopeWriteLock(CacheLock);
+
+	for (TPair<const IPCGElement*, FPCGGraphCacheEntries>& ElementToEntries : CacheData)
+	{
+		FPCGGraphCacheEntries& Entries = ElementToEntries.Value;
+		for (int32 EntryIndex = Entries.Num() - 1; EntryIndex >= 0; --EntryIndex)
+		{
+			if (!Entries[EntryIndex].Settings)
+			{
+				continue;
+			}
+
+			TArray<FName> TrackedTags = Entries[EntryIndex].Settings->GetTrackedActorTags();
+
+			bool bShouldClean = false;
+			for (const FName& Tag : InActor->Tags)
+			{
+				if (TrackedTags.Contains(Tag))
+				{
+					bShouldClean = true;
+					break;
+				}
+			}
+
+			if (bShouldClean)
+			{
+				Entries.RemoveAtSwap(EntryIndex);
+			}
+		}
+	}
+}
+#endif // WITH_EDITOR
\ No newline at end of file
diff --git a/Engine/Plugins/Experimental/PCG/Source/PCG/Private/Graph/PCGGraphCache.h b/Engine/Plugins/Experimental/PCG/Source/PCG/Private/Graph/PCGGraphCache.h
new file mode 100644
--- /dev/null
+++ b/Engine/Plugins/Experimental/PCG/Source/PCG/Private/Graph/PCGGraphCache.h
@@ -0,0 +1,61 @@
+// Copyright Epic Games, Inc. All Rights Reserved.
+
+#pragma once
+
+#include "CoreMinimal.h"
+#include "PCGData.h"
+#include "PCGElement.h"
+
+class UPCGSettings;
+class AActor;
+
+struct FPCGGraphCacheEntry
+{
+	FPCGGraphCacheEntry() = default;
+	FPCGGraphCacheEntry(const FPCGDataCollection& InInput, const UPCGSettings* InSettings, const FPCGDataCollection& InOutput, TWeakObjectPtr<UObject> InOwner, TSet<UObject*>& OutRootedData);
+
+	bool Matches(const FPCGDataCollection& InInput, const UPCGSettings* InSettings) const;
+
+	FPCGDataCollection Input;
+	TObjectPtr<UPCGSettings> Settings;
+	FPCGDataCollection Output;
+};
+
+// TODO: investigate if we need a more evolved data structure here
+// since we could want to have a lock per entries structure
+typedef TArray<FPCGGraphCacheEntry> FPCGGraphCacheEntries;
+
+/** Core idea is to store cache entries per node, but that will be less efficient
+* In cases where we have some subgraph reuse. Under that premise, we can then
+* instead store by element, as we will never recreate elements (except for arbitrary tasks)
+*/
+struct FPCGGraphCache
+{
+	FPCGGraphCache(TWeakObjectPtr<UObject> InOwner);
+	~FPCGGraphCache();
+
+	/** Returns true if data was found from the cache, in which case the outputs are written in OutOutput */
+	bool GetFromCache(const IPCGElement* InElement, const FPCGDataCollection& InInput, const UPCGSettings* InSettings, FPCGDataCollection& OutOutput) const;
+
+	/** Stores data in the cache for later use */
+	void StoreInCache(const IPCGElement* InElement, const FPCGDataCollection& InInput, const UPCGSettings* InSettings, const FPCGDataCollection& InOutput);
+
+	/** Removes all entries from the cache, unroots data, etc. */
+	void ClearCache();
+
+#if WITH_EDITOR
+	/** Clears any cache entry for the given data */
+	void CleanFromCache(AActor* InObject);
+#endif
+
+private:
+	// Note: we are not going to serialize this as-is, since the pointers will change
+	// We will have to serialize on a node id basis most likely
+	TMap<const IPCGElement*, FPCGGraphCacheEntries> CacheData;
+	TWeakObjectPtr<UObject> Owner = nullptr;
+
+	/** To prevent garbage collection on data in the cache, we'll need to root some data */
+	TSet<UObject*> RootedData;
+
+	mutable FRWLock CacheLock;
+};
\ No newline at end of file
diff --git a/Engine/Plugins/Experimental/PCG/Source/PCG/Private/Graph/PCGGraphExecutor.cpp b/Engine/Plugins/Experimental/PCG/Source/PCG/Private/Graph/PCGGraphExecutor.cpp
--- a/Engine/Plugins/Experimental/PCG/Source/PCG/Private/Graph/PCGGraphExecutor.cpp
+++ b/Engine/Plugins/Experimental/PCG/Source/PCG/Private/Graph/PCGGraphExecutor.cpp
@@ -13,8 +13,9 @@
 #include "FileHelpers.h"
 #endif
 
-FPCGGraphExecutor::FPCGGraphExecutor()
+FPCGGraphExecutor::FPCGGraphExecutor(UObject* InOwner)
 	: GraphCompiler(MakeUnique<FPCGGraphCompiler>())
+	, GraphCache(InOwner)
 {
 }
 
@@ -30,7 +31,7 @@ FPCGTaskId FPCGGraphExecutor::Schedule(UPCGComponent* Component, const TArray<FP
 	check(Component);
 	UPCGGraph* Graph = Component->GetGraph();
 
-	return Schedule(Graph, Component, MakeShared<FPCGFetchInputElement>(Component), ExternalDependencies);
+	return Schedule(Graph, Component, GetFetchInputElement(), ExternalDependencies);
 }
 
 FPCGTaskId FPCGGraphExecutor::Schedule(UPCGGraph* Graph, UPCGComponent* SourceComponent, FPCGElementPtr InputElement, const TArray<FPCGTaskId>& ExternalDependencies)
@@ -188,6 +189,7 @@ void FPCGGraphExecutor::Execute()
 				FPCGContextPtr Context = Element->Initialize(TaskInput, Task.SourceComponent);
 				Context->Node = Task.Node; // still needed?
 				Context->TaskId = Task.NodeId;
+				Context->Cache = &GraphCache;
 
 				FPCGGraphActiveTask& ActiveTask = ActiveTasks.Emplace_GetRef();
 				ActiveTask.Element = Element;
@@ -270,16 +272,7 @@ void FPCGGraphExecutor::StoreResults(FPCGTaskId InTaskId, const FPCGDataCollecti
 	OutputData.Add(InTaskId, InTaskOutput);
 
 	// Root any non-rooted results, otherwise they'll get garbage-collected
-	for (const FPCGTaggedData& TaggedData : InTaskOutput.TaggedData)
-	{
-		if (TaggedData.Data != nullptr && !TaggedData.Data->IsRooted())
-		{
-			// This is technically a const_cast
-			UObject* DataToRoot = Cast<UObject>(TaggedData.Data);
-			DataToRoot->AddToRoot();
-			RootedData.Add(DataToRoot);
-		}
-	}
+	InTaskOutput.RootUnrootedData(RootedData);
 }
 
 void FPCGGraphExecutor::ClearResults()
@@ -298,6 +291,16 @@ void FPCGGraphExecutor::ClearResults()
 	ScheduleLock.Unlock();
 }
 
+FPCGElementPtr FPCGGraphExecutor::GetFetchInputElement()
+{
+	if (!FetchInputElement)
+	{
+		FetchInputElement = MakeShared<FPCGFetchInputElement>();
+	}
+
+	return FetchInputElement;
+}
+
 #if WITH_EDITOR
 void FPCGGraphExecutor::AddToDirtyActors(AActor* Actor)
 {
@@ -355,17 +358,13 @@ void FPCGGraphExecutor::NotifyGraphChanged(UPCGGraph* InGraph)
 }
 #endif
 
-FPCGFetchInputElement::FPCGFetchInputElement(UPCGComponent* InComponent)
-	: Component(InComponent)
-{
-}
-
 bool FPCGFetchInputElement::ExecuteInternal(FPCGContextPtr Context) const
 {
 	// First: any input can be passed through to the output trivially
 	Context->OutputData = Context->InputData;
 
 	// Second: fetch the inputs provided by the component
+	UPCGComponent* Component = Context->SourceComponent;
 	check(Component);
 	if(UPCGData* PCGData = Component->GetPCGData())
 	{
diff --git a/Engine/Plugins/Experimental/PCG/Source/PCG/Private/Graph/PCGGraphExecutor.h b/Engine/Plugins/Experimental/PCG/Source/PCG/Private/Graph/PCGGraphExecutor.h
--- a/Engine/Plugins/Experimental/PCG/Source/PCG/Private/Graph/PCGGraphExecutor.h
+++ b/Engine/Plugins/Experimental/PCG/Source/PCG/Private/Graph/PCGGraphExecutor.h
@@ -7,6 +7,7 @@
 #include "PCGContext.h"
 #include "PCGElement.h"
 #include "PCGSubsystem.h"
+#include "PCGGraphCache.h"
 
 #if WITH_EDITOR
 #include "WorldPartition/WorldPartitionHandle.h" // Needed for FWorldPartitonReference
@@ -42,7 +43,7 @@ struct FPCGGraphActiveTask
 class FPCGGraphExecutor
 {
 public:
-	FPCGGraphExecutor();
+	FPCGGraphExecutor(UObject* InOwner);
 	~FPCGGraphExecutor();
 
 	/** Compile (and cache) a graph for later use. This call is threadsafe */
@@ -69,12 +70,17 @@ public:
 	/** "Tick" of the graph executor. This call is NOT THREADSAFE */
 	void Execute();
 
+	/** Expose cache so it can be dirtied */
+	FPCGGraphCache& GetCache() { return GraphCache; }
+
 private:
 	void QueueReadyTasks(FPCGTaskId FinishedTaskHint = InvalidTaskId);
 	void BuildTaskInput(const FPCGGraphTask& Task, FPCGDataCollection& TaskInput);
 	void StoreResults(FPCGTaskId InTaskId, const FPCGDataCollection& InTaskOutput);
 	void ClearResults();
 
+	FPCGElementPtr GetFetchInputElement();
+
 #if WITH_EDITOR
 	void SaveDirtyActors();
 	void ReleaseUnusedActors();
@@ -83,13 +89,19 @@ private:
 	/** Graph compiler that turns a graph into tasks */
 	TUniquePtr<FPCGGraphCompiler> GraphCompiler;
 
+	/** Graph results cache */
+	FPCGGraphCache GraphCache;
+
+	/** Input fetch element, stored here so we have only one */
+	FPCGElementPtr FetchInputElement;
+
 	FCriticalSection ScheduleLock;
 	TArray<FPCGGraphScheduleTask> ScheduledTasks;
 
 	TArray<FPCGGraphTask> Tasks;
 	TArray<FPCGGraphTask> ReadyTasks;
 	TArray<FPCGGraphActiveTask> ActiveTasks;
-	TArray<UObject*> RootedData;
+	TSet<UObject*> RootedData;
 	/** Map of node instances to their output, could be cleared once execution is done */
 	/** Note: this should at some point unload based on loaded/unloaded proxies, otherwise memory cost will be unbounded */
 	TMap<FPCGTaskId, FPCGDataCollection> OutputData;
@@ -105,14 +117,9 @@ private:
 
 class FPCGFetchInputElement : public FSimplePCGElement
 {
-public:
-	FPCGFetchInputElement(UPCGComponent* InComponent);
-
 protected:
 	virtual bool ExecuteInternal(FPCGContextPtr Context) const override;
-
-private:
-	UPCGComponent* Component = nullptr;
+	virtual bool IsCacheable(const UPCGSettings* InSettings) const override { return false; }
 };
 
 class FPCGGenericElement : public FSimplePCGElement
@@ -125,6 +132,7 @@ protected:
 	// as most of these will impact the editor in some way (loading, unloading, saving)
 	virtual bool ExecuteInternal(FPCGContextPtr Context) const override;
 	virtual bool IsCancellable() const { return false; }
+	virtual bool IsCacheable(const UPCGSettings* InSettings) const override { return false; }
 
 private:
 	TFunction<bool()> Operation;
diff --git a/Engine/Plugins/Experimental/PCG/Source/PCG/Private/PCGComponent.cpp b/Engine/Plugins/Experimental/PCG/Source/PCG/Private/PCGComponent.cpp
--- a/Engine/Plugins/Experimental/PCG/Source/PCG/Private/PCGComponent.cpp
+++ b/Engine/Plugins/Experimental/PCG/Source/PCG/Private/PCGComponent.cpp
@@ -712,6 +712,7 @@ void UPCGComponent::OnActorDeleted(AActor* InActor)
 	if (bIsTracked)
 	{
 		CachedTrackedActors.Remove(InActor);
+		GetSubsystem()->CleanFromCache(InActor);
 	}
 
 	if (bHasExcludedTag || bIsTracked)
@@ -735,9 +736,20 @@ void UPCGComponent::OnActorMoved(AActor* InActor)
 	}
 	else
 	{
-		if (ActorHasExcludedTag(InActor) || ActorIsTracked(InActor))
+		const bool bHasExcludedTag = ActorHasExcludedTag(InActor);
+		const bool bActorIsTracked = ActorIsTracked(InActor);
+		if (bHasExcludedTag || bActorIsTracked)
 		{
-			DirtyExclusionData(InActor);
+			if (bHasExcludedTag)
+			{
+				DirtyExclusionData(InActor);
+			}
+
+			if (bActorIsTracked)
+			{
+				GetSubsystem()->CleanFromCache(InActor);
+			}
+
 			DirtyGenerated();
 			Refresh();
 		}
@@ -780,9 +792,21 @@ void UPCGComponent::OnObjectPropertyChanged(UObject* InObject, FPropertyChangedE
 	}
 	else
 	{
-		if (ActorHasExcludedTag(Actor) || ActorIsTracked(Actor))
+		const bool bActorHasExcludedTag = ActorHasExcludedTag(Actor);
+		const bool bActorIsTracked = ActorIsTracked(Actor);
+
+		if(bActorHasExcludedTag || bActorIsTracked)
 		{
-			DirtyExclusionData(Actor);
+			if (bActorHasExcludedTag)
+			{
+				DirtyExclusionData(Actor);
+			}
+
+			if (bActorIsTracked)
+			{
+				GetSubsystem()->CleanFromCache(Actor);
+			}
+			
 			DirtyGenerated();
 			Refresh();
 		}
@@ -796,6 +820,7 @@ void UPCGComponent::OnObjectPropertyChanged(UObject* InObject, FPropertyChangedE
 
 			if (CachedTrackedActors.Remove(Actor))
 			{
+				GetSubsystem()->CleanFromCache(Actor);
 				bDirtyAndRefresh = true;
 			}
 
diff --git a/Engine/Plugins/Experimental/PCG/Source/PCG/Private/PCGData.cpp b/Engine/Plugins/Experimental/PCG/Source/PCG/Private/PCGData.cpp
--- a/Engine/Plugins/Experimental/PCG/Source/PCG/Private/PCGData.cpp
+++ b/Engine/Plugins/Experimental/PCG/Source/PCG/Private/PCGData.cpp
@@ -11,6 +11,11 @@ bool FPCGTaggedData::operator==(const FPCGTaggedData& Other) const
 		Tags.Includes(Other.Tags);
 }
 
+bool FPCGTaggedData::operator!=(const FPCGTaggedData& Other) const
+{
+	return !operator==(Other);
+}
+
 TArray<FPCGTaggedData> FPCGDataCollection::GetInputs() const
 {
 	return TaggedData.FilterByPredicate([](const FPCGTaggedData& Data) {
@@ -57,6 +62,39 @@ const UPCGSettings* FPCGDataCollection::GetSettings(const UPCGSettings* InDefaul
 	}
 }
 
+bool FPCGDataCollection::operator==(const FPCGDataCollection& Other) const
+{
+	if (bCancelExecution != Other.bCancelExecution || TaggedData.Num() != Other.TaggedData.Num())
+	{
+		return false;
+	}
+
+	// TODO: Once we make the arguments order irrelevant, then this should be updated
+	for (int32 DataIndex = 0; DataIndex < TaggedData.Num(); ++DataIndex)
+	{
+		if (TaggedData[DataIndex] != Other.TaggedData[DataIndex])
+		{
+			return false;
+		}
+	}
+	
+	return true;
+}
+
+void FPCGDataCollection::RootUnrootedData(TSet<UObject*>& OutRootedData) const
+{
+	for (const FPCGTaggedData& Data : TaggedData)
+	{
+		if (Data.Data && !Data.Data->IsRooted())
+		{
+			// This is technically a const_cast
+			UObject* DataToRoot = Cast<UObject>(Data.Data);
+			DataToRoot->AddToRoot();
+			OutRootedData.Add(DataToRoot);
+		}
+	}
+}
+
 TArray<FPCGTaggedData> UPCGDataFunctionLibrary::GetInputs(const FPCGDataCollection& InCollection)
 {
 	return InCollection.GetInputs();
diff --git a/Engine/Plugins/Experimental/PCG/Source/PCG/Private/PCGElement.cpp b/Engine/Plugins/Experimental/PCG/Source/PCG/Private/PCGElement.cpp
--- a/Engine/Plugins/Experimental/PCG/Source/PCG/Private/PCGElement.cpp
+++ b/Engine/Plugins/Experimental/PCG/Source/PCG/Private/PCGElement.cpp
@@ -4,6 +4,7 @@
 #include "PCGContext.h"
 #include "PCGSettings.h"
 #include "Elements/PCGDebugElement.h"
+#include "Graph/PCGGraphCache.h"
 
 bool IPCGElement::Execute(FPCGContextPtr Context) const
 {
@@ -57,7 +58,21 @@ bool IPCGElement::Execute(FPCGContextPtr Context) const
 			Context->InputData.TaggedData = FilteredTaggedData;
 		}
 
-		bool bDone = ExecuteInternal(Context);
+		bool bDone = false;
+
+		if (IsCacheable(Settings) && Context->Cache && Context->Cache->GetFromCache(this, Context->InputData, Settings, Context->OutputData))
+		{
+			bDone = true;
+		}
+		else
+		{
+			bDone = ExecuteInternal(Context);
+
+			if (bDone && IsCacheable(Settings) && Context->Cache)
+			{
+				Context->Cache->StoreInCache(this, Context->InputData, Settings, Context->OutputData);
+			}
+		}
 
 		/** TODO - Placeholder feature */
 		if (bDone && Settings && !Settings->TagsAppliedOnOutput.IsEmpty())
diff --git a/Engine/Plugins/Experimental/PCG/Source/PCG/Private/PCGSettings.cpp b/Engine/Plugins/Experimental/PCG/Source/PCG/Private/PCGSettings.cpp
--- a/Engine/Plugins/Experimental/PCG/Source/PCG/Private/PCGSettings.cpp
+++ b/Engine/Plugins/Experimental/PCG/Source/PCG/Private/PCGSettings.cpp
@@ -2,6 +2,22 @@
 
 #include "PCGSettings.h"
 #include "PCGNode.h"
+#include "Serialization/ArchiveObjectCrc32.h"
+
+bool UPCGSettings::operator==(const UPCGSettings& Other) const
+{
+	if (this == &Other)
+	{
+		return true;
+	}
+	else
+	{
+		FArchiveObjectCrc32 Ar;
+		uint32 ThisCrc = Ar.Crc32(const_cast<UPCGSettings*>(this));
+		uint32 OtherCrc = Ar.Crc32(const_cast<UPCGSettings*>(&Other));
+		return ThisCrc == OtherCrc;
+	}
+}
 
 FPCGElementPtr UPCGSettings::GetElement() const
 {
diff --git a/Engine/Plugins/Experimental/PCG/Source/PCG/Private/PCGSubsystem.cpp b/Engine/Plugins/Experimental/PCG/Source/PCG/Private/PCGSubsystem.cpp
--- a/Engine/Plugins/Experimental/PCG/Source/PCG/Private/PCGSubsystem.cpp
+++ b/Engine/Plugins/Experimental/PCG/Source/PCG/Private/PCGSubsystem.cpp
@@ -34,7 +34,7 @@ void UPCGSubsystem::PostInitialize()
 
 	// Initialize graph executor
 	check(!GraphExecutor);
-	GraphExecutor = new FPCGGraphExecutor();
+	GraphExecutor = new FPCGGraphExecutor(this);
 	// gather things.. ?
 	// TODO	
 }
@@ -488,6 +488,14 @@ void UPCGSubsystem::NotifyGraphChanged(UPCGGraph* InGraph)
 	}
 }
 
+void UPCGSubsystem::CleanFromCache(AActor* InActor)
+{
+	if (GraphExecutor)
+	{
+		GraphExecutor->GetCache().CleanFromCache(InActor);
+	}
+}
+
 ///////////////////////////////////////////////////////////////////////////////////////////
 void UPCGSubsystem::DelayPartitionGraph(UPCGComponent* Component)
 {
diff --git a/Engine/Plugins/Experimental/PCG/Source/PCG/Public/Elements/PCGDebugElement.h b/Engine/Plugins/Experimental/PCG/Source/PCG/Public/Elements/PCGDebugElement.h
--- a/Engine/Plugins/Experimental/PCG/Source/PCG/Public/Elements/PCGDebugElement.h
+++ b/Engine/Plugins/Experimental/PCG/Source/PCG/Public/Elements/PCGDebugElement.h
@@ -32,4 +32,5 @@ class FPCGDebugElement : public FSimpleTypedPCGElement<UPCGDebugSettings>
 {
 protected:
 	virtual bool ExecuteInternal(FPCGContextPtr Contexxt) const override;
+	virtual bool IsCacheable(const UPCGSettings* InSettings) const override { return false; }
 };
\ No newline at end of file
diff --git a/Engine/Plugins/Experimental/PCG/Source/PCG/Public/Elements/PCGExecuteBlueprint.h b/Engine/Plugins/Experimental/PCG/Source/PCG/Public/Elements/PCGExecuteBlueprint.h
--- a/Engine/Plugins/Experimental/PCG/Source/PCG/Public/Elements/PCGExecuteBlueprint.h
+++ b/Engine/Plugins/Experimental/PCG/Source/PCG/Public/Elements/PCGExecuteBlueprint.h
@@ -90,6 +90,9 @@ protected:
 	TArray<FName> TrackedActorTags;
 #endif
 
+	UPROPERTY(BlueprintReadWrite, EditAnywhere, Category = Settings)
+	bool bCreatesArtifacts = false;
+
 protected:
 #if WITH_EDITOR
 	void OnBlueprintChanged(UBlueprint* InBlueprint);
@@ -112,4 +115,5 @@ class FPCGExecuteBlueprintElement : public FSimpleTypedPCGElement<UPCGBlueprintS
 {
 protected:
 	virtual bool ExecuteInternal(FPCGContextPtr Context) const override;
+	virtual bool IsCacheable(const UPCGSettings* InSettings) const override;
 };
\ No newline at end of file
diff --git a/Engine/Plugins/Experimental/PCG/Source/PCG/Public/Elements/PCGSpawnActor.h b/Engine/Plugins/Experimental/PCG/Source/PCG/Public/Elements/PCGSpawnActor.h
--- a/Engine/Plugins/Experimental/PCG/Source/PCG/Public/Elements/PCGSpawnActor.h
+++ b/Engine/Plugins/Experimental/PCG/Source/PCG/Public/Elements/PCGSpawnActor.h
@@ -67,4 +67,5 @@ class FPCGSpawnActorElement : public FSimpleTypedPCGElement<UPCGSpawnActorSettin
 {
 protected:
 	virtual bool ExecuteInternal(FPCGContextPtr Context) const override;
+	virtual bool IsCacheable(const UPCGSettings* InSettings) const override { return false; }
 };
\ No newline at end of file
diff --git a/Engine/Plugins/Experimental/PCG/Source/PCG/Public/Elements/PCGStaticMeshSpawner.h b/Engine/Plugins/Experimental/PCG/Source/PCG/Public/Elements/PCGStaticMeshSpawner.h
--- a/Engine/Plugins/Experimental/PCG/Source/PCG/Public/Elements/PCGStaticMeshSpawner.h
+++ b/Engine/Plugins/Experimental/PCG/Source/PCG/Public/Elements/PCGStaticMeshSpawner.h
@@ -45,4 +45,5 @@ class FPCGStaticMeshSpawnerElement : public FSimpleTypedPCGElement<UPCGStaticMes
 {
 protected:
 	virtual bool ExecuteInternal(FPCGContextPtr Context) const override;
+	virtual bool IsCacheable(const UPCGSettings* InSettings) const override { return false; }
 };
\ No newline at end of file
diff --git a/Engine/Plugins/Experimental/PCG/Source/PCG/Public/PCGContext.h b/Engine/Plugins/Experimental/PCG/Source/PCG/Public/PCGContext.h
--- a/Engine/Plugins/Experimental/PCG/Source/PCG/Public/PCGContext.h
+++ b/Engine/Plugins/Experimental/PCG/Source/PCG/Public/PCGContext.h
@@ -9,6 +9,7 @@
 
 struct FPCGContext;
 class UPCGComponent;
+struct FPCGGraphCache;
 
 typedef TSharedPtr<FPCGContext, ESPMode::ThreadSafe> FPCGContextPtr;
 
@@ -19,6 +20,7 @@ struct FPCGContext
 	FPCGDataCollection InputData;
 	FPCGDataCollection OutputData;
 	UPCGComponent* SourceComponent = nullptr;
+	FPCGGraphCache* Cache = nullptr;
 	// TODO: add RNG source
 	// TODO: replace this by a better identification mechanism
 	const UPCGNode* Node = nullptr;
diff --git a/Engine/Plugins/Experimental/PCG/Source/PCG/Public/PCGData.h b/Engine/Plugins/Experimental/PCG/Source/PCG/Public/PCGData.h
--- a/Engine/Plugins/Experimental/PCG/Source/PCG/Public/PCGData.h
+++ b/Engine/Plugins/Experimental/PCG/Source/PCG/Public/PCGData.h
@@ -45,6 +45,7 @@ struct FPCGTaggedData
 	TSet<FString> Tags;
 
 	bool operator==(const FPCGTaggedData& Other) const;
+	bool operator!=(const FPCGTaggedData& Other) const;
 };
 
 USTRUCT(BlueprintType)
@@ -62,6 +63,9 @@ struct FPCGDataCollection
 
 	const UPCGSettings* GetSettings(const UPCGSettings* InDefaultSettings) const;
 
+	bool operator==(const FPCGDataCollection& Other) const;
+	void RootUnrootedData(TSet<UObject*>& OutRootedData) const;
+
 	UPROPERTY(BlueprintReadWrite, EditAnywhere, Category = Data)
 	TArray<FPCGTaggedData> TaggedData;
 
diff --git a/Engine/Plugins/Experimental/PCG/Source/PCG/Public/PCGElement.h b/Engine/Plugins/Experimental/PCG/Source/PCG/Public/PCGElement.h
--- a/Engine/Plugins/Experimental/PCG/Source/PCG/Public/PCGElement.h
+++ b/Engine/Plugins/Experimental/PCG/Source/PCG/Public/PCGElement.h
@@ -8,6 +8,7 @@
 
 class IPCGElement;
 class UPCGComponent;
+class UPCGSettings;
 
 typedef TSharedPtr<IPCGElement, ESPMode::ThreadSafe> FPCGElementPtr;
 
@@ -25,6 +26,7 @@ public:
 protected:
 	virtual bool ExecuteInternal(FPCGContextPtr Context) const = 0;
 	virtual bool IsCancellable() const { return true; }
+	virtual bool IsCacheable(const UPCGSettings* InSettings) const { return true; }
 };
 
 /**
diff --git a/Engine/Plugins/Experimental/PCG/Source/PCG/Public/PCGSettings.h b/Engine/Plugins/Experimental/PCG/Source/PCG/Public/PCGSettings.h
--- a/Engine/Plugins/Experimental/PCG/Source/PCG/Public/PCGSettings.h
+++ b/Engine/Plugins/Experimental/PCG/Source/PCG/Public/PCGSettings.h
@@ -35,6 +35,8 @@ public:
 	// TODO: check if we need this to be virtual, we don't really need if we're always caching
 	/*virtual*/ FPCGElementPtr GetElement() const;
 	virtual UPCGNode* CreateNode() const;
+	
+	bool operator==(const UPCGSettings& Other) const;
 
 #if WITH_EDITOR
 	virtual FName GetDefaultNodeName() const { return NAME_None; }
@@ -96,4 +98,5 @@ class FPCGTrivialElement : public FSimplePCGElement
 {
 protected:
 	virtual bool ExecuteInternal(FPCGContextPtr Context) const override;
+	virtual bool IsCacheable(const UPCGSettings* InSettings) const override { return false; }
 };
\ No newline at end of file
diff --git a/Engine/Plugins/Experimental/PCG/Source/PCG/Public/PCGSubsystem.h b/Engine/Plugins/Experimental/PCG/Source/PCG/Public/PCGSubsystem.h
--- a/Engine/Plugins/Experimental/PCG/Source/PCG/Public/PCGSubsystem.h
+++ b/Engine/Plugins/Experimental/PCG/Source/PCG/Public/PCGSubsystem.h
@@ -81,6 +81,9 @@ public:
 	/** Propagate to the graph compiler graph changes */
 	void NotifyGraphChanged(UPCGGraph* InGraph);
 
+	/** Cleans up the graph cache when a tracked actor is changed */
+	void CleanFromCache(AActor* InActor);
+
 private:
 	FPCGTaskId DelayProcessGraph(UPCGComponent* Component, bool bGenerate, bool bSave, bool bUseEmptyNewBounds);
 	FPCGTaskId ProcessGraph(UPCGComponent* Component, const FBox& InPreviousBounds, const FBox& InNewBounds, bool bGenerate, bool bSave);
