Exploded the point sampler class in favor of three operations that will have more flexbility (sampler, filter, remapping)
Applied exclusions at the start of the graph for ease-of-use.

diff --git a/Engine/Plugins/Experimental/PCG/Source/PCG/Private/Elements/PCGDensityFilter.cpp b/Engine/Plugins/Experimental/PCG/Source/PCG/Private/Elements/PCGDensityFilter.cpp
new file mode 100644
--- /dev/null
+++ b/Engine/Plugins/Experimental/PCG/Source/PCG/Private/Elements/PCGDensityFilter.cpp
@@ -0,0 +1,93 @@
+// Copyright Epic Games, Inc. All Rights Reserved.
+
+#include "Elements/PCGDensityFilter.h"
+#include "Data/PCGPointData.h"
+
+FPCGElementPtr UPCGDensityFilterSettings::CreateElement() const
+{
+	return MakeShared<FPCGDensityFilterElement>();
+}
+
+bool FPCGDensityFilterElement::ExecuteInternal(FPCGContextPtr Context) const
+{
+	TRACE_CPUPROFILER_EVENT_SCOPE(FPCGDensityFilterElement::Execute);
+
+	const UPCGDensityFilterSettings* Settings = Context->GetInputSettings<UPCGDensityFilterSettings>();
+	check(Settings);
+
+	TArray<FPCGTaggedData> Inputs = Context->InputData.GetInputs();
+	TArray<FPCGTaggedData>& Outputs = Context->OutputData.TaggedData;
+
+	// Forward any non-input data
+	Outputs.Append(Context->InputData.GetExclusions());
+	Outputs.Append(Context->InputData.GetAllSettings());
+
+	const float MinBound = FMath::Min(Settings->LowerBound, Settings->UpperBound);
+	const float MaxBound = FMath::Max(Settings->LowerBound, Settings->UpperBound);
+
+	const bool bNoResults = (MaxBound <= 0 && !Settings->bInvertFilter) || (MinBound == 0 && MaxBound >= 1 && Settings->bInvertFilter);
+	const bool bTrivialFilter = (MinBound <= 0 && MaxBound >= 1.0 && !Settings->bInvertFilter) || (MinBound == 0 && MaxBound == 0 && Settings->bInvertFilter);
+
+#if WITH_EDITORONLY_DATA
+	if(bNoResults && !Settings->bKeepZeroDensityPoints)
+#else
+	if (bNoResults)
+#endif
+	{
+		return true;
+	}
+
+	// TODO: embarassingly parallel loop
+	for (const FPCGTaggedData& Input : Inputs)
+	{
+		TRACE_CPUPROFILER_EVENT_SCOPE(FPCGDensityFilterElement::Execute::InputLoop);
+
+		FPCGTaggedData& Output = Outputs.Add_GetRef(Input);
+
+		if (!Input.Data || Cast<UPCGSpatialData>(Input.Data) == nullptr)
+		{
+			// TODO: log
+			continue;
+		}
+
+		// Skip processing if the transformation is trivial
+		if (bTrivialFilter)
+		{
+			continue;
+		}
+
+		const UPCGPointData* OriginalData = Cast<UPCGSpatialData>(Input.Data)->ToPointData();
+
+		if (!OriginalData)
+		{
+			// Log error
+			continue;
+		}
+
+		const TArray<FPCGPoint>& Points = OriginalData->GetPoints();
+		
+		UPCGPointData* FilteredData = NewObject<UPCGPointData>();
+		FilteredData->TargetActor = OriginalData->TargetActor;
+		TArray<FPCGPoint>& FilteredPoints = FilteredData->GetMutablePoints();
+
+		Output.Data = FilteredData;
+
+		for (const FPCGPoint& Point : Points)
+		{
+			bool bInRange = (Point.Density >= MinBound && Point.Density <= MaxBound);
+			if (bInRange != Settings->bInvertFilter)
+			{
+				FilteredPoints.Add(Point);
+			}
+#if WITH_EDITORONLY_DATA
+			else if (Settings->bKeepZeroDensityPoints)
+			{
+				FPCGPoint& FilteredPoint = FilteredPoints.Add_GetRef(Point);
+				FilteredPoint.Density = 0;
+			}
+#endif
+		}
+	}
+
+	return true;
+}
\ No newline at end of file
diff --git a/Engine/Plugins/Experimental/PCG/Source/PCG/Private/Elements/PCGDensityRemap.cpp b/Engine/Plugins/Experimental/PCG/Source/PCG/Private/Elements/PCGDensityRemap.cpp
new file mode 100644
--- /dev/null
+++ b/Engine/Plugins/Experimental/PCG/Source/PCG/Private/Elements/PCGDensityRemap.cpp
@@ -0,0 +1,86 @@
+// Copyright Epic Games, Inc. All Rights Reserved.
+
+#include "Elements/PCGDensityRemap.h"
+#include "PCGHelpers.h"
+#include "Data/PCGPointData.h"
+#include "Math/RandomStream.h"
+
+FPCGElementPtr UPCGLinearDensityRemapSettings::CreateElement() const
+{
+	return MakeShared<FPCGLinearDensityRemapElement>();
+}
+
+bool FPCGLinearDensityRemapElement::ExecuteInternal(FPCGContextPtr Context) const
+{
+	TRACE_CPUPROFILER_EVENT_SCOPE(FPCGLinearDensityRemapElement::Execute);
+
+	const UPCGLinearDensityRemapSettings* Settings = Context->GetInputSettings<UPCGLinearDensityRemapSettings>();
+	check(Settings);
+
+	TArray<FPCGTaggedData> Inputs = Context->InputData.GetInputs();
+	TArray<FPCGTaggedData>& Outputs = Context->OutputData.TaggedData;
+
+	// Forward any non-input data
+	Outputs.Append(Context->InputData.GetExclusions());
+	Outputs.Append(Context->InputData.GetAllSettings());
+
+	const float RemapMin = FMath::Min(Settings->RemapMin, Settings->RemapMax);
+	const float RemapMax = FMath::Max(Settings->RemapMin, Settings->RemapMax);
+
+	const bool bTrivialRemapping = (RemapMin == 0 && RemapMax == 1 && Settings->bMultiplyDensity);
+
+	// TODO: embarassingly parallel loop
+	for (const FPCGTaggedData& Input : Inputs)
+	{
+		TRACE_CPUPROFILER_EVENT_SCOPE(FPCGLinearDensityRemapElement::Execute::InputLoop);
+		FPCGTaggedData& Output = Outputs.Add_GetRef(Input);
+
+		if (!Input.Data || Cast<UPCGSpatialData>(Input.Data) == nullptr)
+		{
+			// TODO: log error
+			continue;
+		}
+
+		// Skip processing if the remapping is trivial
+		if (bTrivialRemapping)
+		{
+			continue;
+		}
+
+		const UPCGPointData* OriginalData = Cast<UPCGSpatialData>(Input.Data)->ToPointData();
+
+		if (!OriginalData)
+		{
+			// TODO: log error
+			continue;
+		}
+
+		const TArray<FPCGPoint>& Points = OriginalData->GetPoints();
+		const int OriginalPointCount = Points.Num();
+
+		UPCGPointData* SampledData = NewObject<UPCGPointData>();
+		SampledData->TargetActor = OriginalData->TargetActor;
+		TArray<FPCGPoint>& SampledPoints = SampledData->GetMutablePoints();
+		Output.Data = SampledData;
+
+		SampledPoints = Points;
+		if (Settings->bMultiplyDensity)
+		{
+			for (FPCGPoint& Point : SampledPoints)
+			{
+				FRandomStream RandomSource(PCGHelpers::ComputeSeed(Settings->Seed, Point.Seed));
+				Point.Density *= RandomSource.FRandRange(RemapMin, RemapMax);
+			}
+		}
+		else
+		{
+			for (FPCGPoint& Point : SampledPoints)
+			{
+				FRandomStream RandomSource(PCGHelpers::ComputeSeed(Settings->Seed, Point.Seed));
+				Point.Density = RandomSource.FRandRange(RemapMin, RemapMax);
+			}
+		}
+	}
+
+	return true;
+}
\ No newline at end of file
diff --git a/Engine/Plugins/Experimental/PCG/Source/PCG/Private/Elements/PCGPointSampler.cpp b/Engine/Plugins/Experimental/PCG/Source/PCG/Private/Elements/PCGPointSampler.cpp
--- a/Engine/Plugins/Experimental/PCG/Source/PCG/Private/Elements/PCGPointSampler.cpp
+++ b/Engine/Plugins/Experimental/PCG/Source/PCG/Private/Elements/PCGPointSampler.cpp
@@ -18,16 +18,38 @@ bool FPCGPointSamplerElement::ExecuteInternal(FPCGContextPtr Context) const
 	check(Settings);
 
 	TArray<FPCGTaggedData> Inputs = Context->InputData.GetInputs();
-	TArray<FPCGTaggedData> Exclusions = Context->InputData.GetExclusions();
 	TArray<FPCGTaggedData>& Outputs = Context->OutputData.TaggedData;
 
+	// Forward any non-input data
+	Outputs.Append(Context->InputData.GetExclusions());
+	Outputs.Append(Context->InputData.GetAllSettings());
+
+	const bool bNoSampling = (Settings->Ratio <= 0);
+	const bool bTrivialSampling = (Settings->Ratio >= 1);
+
+	// Early exit when nothing will be generated out of this sampler
+#if WITH_EDITORONLY_DATA
+	if (bNoSampling && !Settings->bKeepZeroDensityPoints)
+#else
+	if (bNoSampling)
+#endif
+	{
+		return true;
+	}
+
 	// TODO: embarassingly parallel loop
 	for (const FPCGTaggedData& Input : Inputs)
 	{
-		FPCGTaggedData& Output = Outputs.Emplace_GetRef();
-		Output = Input;
+		FPCGTaggedData& Output = Outputs.Add_GetRef(Input);
 
 		if (!Input.Data || Cast<UPCGSpatialData>(Input.Data) == nullptr)
+		{
+			// TODO: add a log
+			continue;
+		}
+
+		// Skip processing if the transformation would be trivial
+		if (bTrivialSampling)
 		{
 			continue;
 		}
@@ -36,83 +58,61 @@ bool FPCGPointSamplerElement::ExecuteInternal(FPCGContextPtr Context) const
 
 		if (!OriginalData)
 		{
+			// Log error
 			continue;
 		}
 
 		const TArray<FPCGPoint>& Points = OriginalData->GetPoints();
 		const int OriginalPointCount = Points.Num();
 
-		if(Settings->Ratio >= 1.0f && Exclusions.IsEmpty())
+		UPCGPointData* SampledData = NewObject<UPCGPointData>();
+		SampledData->TargetActor = OriginalData->TargetActor;
+		TArray<FPCGPoint>& SampledPoints = SampledData->GetMutablePoints();
+
+		Output.Data = SampledData;
+
+		// TODO: randomize on the fractional number of points
+#if WITH_EDITORONLY_DATA
+		int TargetNumPoints = (Settings->bKeepZeroDensityPoints ? OriginalPointCount : OriginalPointCount * Settings->Ratio);
+#else
+		int TargetNumPoints = OriginalPointCount * Settings->Ratio;
+#endif
+
+		// Early out
+		if (TargetNumPoints == 0)
 		{
-			Output.Data = OriginalData;
+			// TODO log
+			continue;
 		}
 		else
 		{
-			UPCGPointData* SampledData = NewObject<UPCGPointData>();
-			SampledData->TargetActor = OriginalData->TargetActor;
-
-			Output.Data = SampledData;
+			TRACE_CPUPROFILER_EVENT_SCOPE(FPCGPointSamplerElement::Execute::SelectPoints);
 
-			// TODO: randomize on the fractional number of points
-			int TargetNumPoints = OriginalPointCount * Settings->Ratio;
-
-			// Early out
-			if (TargetNumPoints == 0)
-			{
-				return true;
-			}
+			// Approximate upper bound
+			SampledPoints.Reserve(FMath::Min(OriginalPointCount, 4 * TargetNumPoints / 3));
 
-			TArray<int> SelectedIndices;
-			// Build indices in a deterministic manner
+			for (int Index = 0; Index < OriginalPointCount; ++Index)
 			{
-				TRACE_CPUPROFILER_EVENT_SCOPE(UPCGComponent::CreatePCGData::SelectIndices)
+				const FPCGPoint& Point = Points[Index];
 
-				// Approximate upper bound
-				SelectedIndices.Reserve(4 * TargetNumPoints / 3);
+				// Apply a high-pass filter based on selected ratio
+				FRandomStream RandomSource(PCGHelpers::ComputeSeed(Settings->Seed, Point.Seed));
+				float Chance = RandomSource.FRand();
 
-				for (int Index = 0; Index < OriginalPointCount; ++Index)
+				if (Chance < Settings->Ratio)
+				{
+					SampledPoints.Add(Point);
+				}
+#if WITH_EDITORONLY_DATA
+				else if (Settings->bKeepZeroDensityPoints)
 				{
-					const FPCGPoint& Point = Points[Index];
-
-					// Apply a high-pass filter based on selected ratio
-					FRandomStream RandomSource(PCGHelpers::ComputeSeed(Settings->Seed, Point.Seed));
-					float Chance = RandomSource.FRand();
-
-					if (Chance < Settings->Ratio)
-					{
-						// Remap density to a value between 0.5 and 1.0
-						float Density = Point.Density * (0.5f + 0.5f * (Settings->Ratio - Chance));
-
-						for (const FPCGTaggedData& Exclusion : Exclusions)
-						{
-							if (!Exclusion.Data)
-							{
-								continue;
-							}
-
-							Density *= (1.0f - Cast<UPCGSpatialData>(Exclusion.Data)->GetDensityAtPosition(Point.Transform.GetLocation()));
-
-							if (Density == 0)
-							{
-								break;
-							}
-						}
-
-						// Implicit threshold at 0.5f based on cumulated density
-						if(Density >= 0.5f)
-						{
-							SelectedIndices.Add(Index);
-						}
-					}
+					FPCGPoint& SampledPoint = SampledPoints.Add_GetRef(Point);
+					SampledPoint.Density = 0;
 				}
+#endif
 			}
-
-			// Copy points using the indices
-			// TODO: pass thread info for parallel for
-			SampledData->CopyPointsFrom(OriginalData, SelectedIndices);
 		}
 	}
 
 	return true;
 }
-
diff --git a/Engine/Plugins/Experimental/PCG/Source/PCG/Private/Elements/PCGStaticMeshSpawner.cpp b/Engine/Plugins/Experimental/PCG/Source/PCG/Private/Elements/PCGStaticMeshSpawner.cpp
--- a/Engine/Plugins/Experimental/PCG/Source/PCG/Private/Elements/PCGStaticMeshSpawner.cpp
+++ b/Engine/Plugins/Experimental/PCG/Source/PCG/Private/Elements/PCGStaticMeshSpawner.cpp
@@ -85,6 +85,11 @@ bool FPCGStaticMeshSpawnerElement::ExecuteInternal(FPCGContextPtr Context) const
 
 			for (const FPCGPoint& Point : Points)
 			{
+				if (Point.Density <= 0)
+				{
+					continue;
+				}
+
 				FRandomStream RandomSource(PCGHelpers::ComputeSeed(Point.Seed, Settings->Seed));
 				int RandomWeightedPick = RandomSource.RandRange(0, TotalWeight - 1);
 
diff --git a/Engine/Plugins/Experimental/PCG/Source/PCG/Private/Graph/PCGGraphExecutor.cpp b/Engine/Plugins/Experimental/PCG/Source/PCG/Private/Graph/PCGGraphExecutor.cpp
--- a/Engine/Plugins/Experimental/PCG/Source/PCG/Private/Graph/PCGGraphExecutor.cpp
+++ b/Engine/Plugins/Experimental/PCG/Source/PCG/Private/Graph/PCGGraphExecutor.cpp
@@ -372,15 +372,6 @@ bool FPCGFetchInputElement::ExecuteInternal(FPCGContextPtr Context) const
 		TaggedData.Data = PCGData;
 	}
 
-	// Third: fetch the exclusions provided by the component
-	TArray<UPCGData*> Exclusions = Component->GetPCGExclusionData();
-	for (const UPCGData* Exclusion : Exclusions)
-	{
-		FPCGTaggedData& TaggedData = Context->OutputData.TaggedData.Emplace_GetRef();
-		TaggedData.Data = Exclusion;
-		TaggedData.Usage = EPCGDataUsage::Exclusion;
-	}
-
 	return true;
 }
 
diff --git a/Engine/Plugins/Experimental/PCG/Source/PCG/Private/PCGComponent.cpp b/Engine/Plugins/Experimental/PCG/Source/PCG/Private/PCGComponent.cpp
--- a/Engine/Plugins/Experimental/PCG/Source/PCG/Private/PCGComponent.cpp
+++ b/Engine/Plugins/Experimental/PCG/Source/PCG/Private/PCGComponent.cpp
@@ -5,6 +5,7 @@
 #include "PCGHelpers.h"
 #include "PCGSubsystem.h"
 #include "PCGVolume.h"
+#include "Data/PCGDifferenceData.h"
 #include "Data/PCGIntersectionData.h"
 #include "Data/PCGLandscapeData.h"
 #include "Data/PCGLandscapeSplineData.h"
@@ -664,6 +665,7 @@ void UPCGComponent::DirtyExclusionData(AActor* InActor)
 	if (UPCGData** ExclusionData = CachedExclusionData.Find(InActor))
 	{
 		*ExclusionData = nullptr;
+		CachedPCGData = nullptr;
 	}
 }
 
@@ -863,6 +865,7 @@ void UPCGComponent::DirtyGenerated(bool bInDirtyCachedInput)
 	{
 		CachedInputData = nullptr;
 		CachedActorData = nullptr;
+		CachedPCGData = nullptr;
 	}
 
 	// For partitioned graph, we must forward the call to the partition actor
@@ -912,10 +915,20 @@ void UPCGComponent::Refresh()
 #endif
 
 UPCGData* UPCGComponent::GetPCGData()
+{
+	if (!CachedPCGData)
+	{
+		CachedPCGData = CreatePCGData();
+	}
+
+	return CachedPCGData;
+}
+
+UPCGData* UPCGComponent::GetInputPCGData()
 {
 	if (!CachedInputData)
 	{
-		CachedInputData = CreatePCGData();
+		CachedInputData = CreateInputPCGData();
 	}
 
 	return CachedInputData;
@@ -943,7 +956,7 @@ TArray<UPCGData*> UPCGComponent::GetPCGExclusionData()
 
 void UPCGComponent::UpdatePCGExclusionData()
 {
-	const UPCGData* InputData = GetPCGData();
+	const UPCGData* InputData = GetInputPCGData();
 	const UPCGSpatialData* InputSpatialData = Cast<const UPCGSpatialData>(InputData);
 
 	// Update the list of cached excluded actors here, since we might not have picked up everything on map load (due to WP)
@@ -1098,7 +1111,40 @@ UPCGData* UPCGComponent::CreateActorPCGData(AActor* Actor)
 
 UPCGData* UPCGComponent::CreatePCGData()
 {
-	TRACE_CPUPROFILER_EVENT_SCOPE(UPCGComponent::CreatePCGData)
+	TRACE_CPUPROFILER_EVENT_SCOPE(UPCGComponent::CreatePCGData);
+	UPCGData* InputData = GetInputPCGData();
+	UPCGSpatialData* SpatialInput = Cast<UPCGSpatialData>(InputData);
+	
+	// Early out: incompatible data
+	if (!SpatialInput)
+	{
+		return InputData;
+	}
+
+	UPCGDifferenceData* Difference = nullptr;
+	TArray<UPCGData*> ExclusionData = GetPCGExclusionData();
+
+	for (UPCGData* Exclusion : ExclusionData)
+	{
+		if (UPCGSpatialData* SpatialExclusion = Cast<UPCGSpatialData>(Exclusion))
+		{
+			if (!Difference)
+			{
+				Difference = SpatialInput->Subtract(SpatialExclusion);
+			}
+			else
+			{
+				Difference->AddDifference(SpatialExclusion);
+			}
+		}
+	}
+
+	return Difference ? Difference : InputData;
+}
+
+UPCGData* UPCGComponent::CreateInputPCGData()
+{
+	TRACE_CPUPROFILER_EVENT_SCOPE(UPCGComponent::CreateInputPCGData);
 	AActor* Actor = GetOwner();
 	check(Actor);
 
diff --git a/Engine/Plugins/Experimental/PCG/Source/PCG/Public/Elements/PCGDensityFilter.h b/Engine/Plugins/Experimental/PCG/Source/PCG/Public/Elements/PCGDensityFilter.h
new file mode 100644
--- /dev/null
+++ b/Engine/Plugins/Experimental/PCG/Source/PCG/Public/Elements/PCGDensityFilter.h
@@ -0,0 +1,47 @@
+// Copyright Epic Games, Inc. All Rights Reserved.
+
+#pragma once
+
+#include "CoreMinimal.h"
+#include "PCGSettings.h"
+#include "PCGElement.h"
+#include "PCGNode.h"
+
+#include "PCGDensityFilter.generated.h"
+
+UCLASS(BlueprintType, ClassGroup = (Procedural))
+class UPCGDensityFilterSettings : public UPCGSettings
+{
+	GENERATED_BODY()
+
+public:
+	//~Begin UPCGSettings interface
+#if WITH_EDITOR
+	virtual FName GetDefaultNodeName() const override { return FName(TEXT("DensityFilter")); }
+#endif
+
+protected:
+	virtual FPCGElementPtr CreateElement() const override;
+	//~End UPCGSettings interface
+
+public:
+	UPROPERTY(BlueprintReadWrite, EditAnywhere, Category = Settings, meta = (ClampMin = "0", ClampMax = "1"))
+	float LowerBound = 0.5f;
+
+	UPROPERTY(BlueprintReadWrite, EditAnywhere, Category = Settings, meta = (ClampMin = "0", ClampMax = "1"))
+	float UpperBound = 1.0f;
+
+	UPROPERTY(BlueprintReadWrite, EditAnywhere, Category = Settings)
+	bool bInvertFilter = false;
+
+#if WITH_EDITORONLY_DATA
+	UPROPERTY(Transient, BlueprintReadWrite, EditAnywhere, Category = Debug)
+	bool bKeepZeroDensityPoints = false;
+#endif
+};
+
+class FPCGDensityFilterElement : public FSimplePCGElement
+{
+protected:
+	virtual bool ExecuteInternal(FPCGContextPtr Context) const override;
+};
\ No newline at end of file
diff --git a/Engine/Plugins/Experimental/PCG/Source/PCG/Public/Elements/PCGDensityRemap.h b/Engine/Plugins/Experimental/PCG/Source/PCG/Public/Elements/PCGDensityRemap.h
new file mode 100644
--- /dev/null
+++ b/Engine/Plugins/Experimental/PCG/Source/PCG/Public/Elements/PCGDensityRemap.h
@@ -0,0 +1,42 @@
+// Copyright Epic Games, Inc. All Rights Reserved.
+
+#pragma once
+
+#include "CoreMinimal.h"
+#include "PCGSettings.h"
+#include "PCGElement.h"
+#include "PCGNode.h"
+
+#include "PCGDensityRemap.generated.h"
+
+UCLASS(BlueprintType, ClassGroup = (Procedural))
+class UPCGLinearDensityRemapSettings : public UPCGSettings
+{
+	GENERATED_BODY()
+
+public:
+	// ~Begin UPCGSettings interface
+#if WITH_EDITOR
+	virtual FName GetDefaultNodeName() const override { return FName(TEXT("LinearDensityRemap")); }
+#endif
+
+protected:
+	virtual FPCGElementPtr CreateElement() const override;
+	//~End UPCGSettings interface
+
+public:
+	UPROPERTY(BlueprintReadWrite, EditAnywhere, Category = Settings, meta = (ClampMin = "0", ClampMax = "1"))
+	float RemapMin = 0.5f;
+
+	UPROPERTY(BlueprintReadWrite, EditAnywhere, Category = Settings, meta = (ClampMin = "0", ClampMax = "1"))
+	float RemapMax = 1.0f;
+
+	UPROPERTY(BlueprintReadWrite, EditAnywhere, Category = Settings)
+	bool bMultiplyDensity = true;
+};
+
+class FPCGLinearDensityRemapElement : public FSimplePCGElement
+{
+protected:
+	virtual bool ExecuteInternal(FPCGContextPtr Context) const;
+};
\ No newline at end of file
diff --git a/Engine/Plugins/Experimental/PCG/Source/PCG/Public/Elements/PCGPointSampler.h b/Engine/Plugins/Experimental/PCG/Source/PCG/Public/Elements/PCGPointSampler.h
--- a/Engine/Plugins/Experimental/PCG/Source/PCG/Public/Elements/PCGPointSampler.h
+++ b/Engine/Plugins/Experimental/PCG/Source/PCG/Public/Elements/PCGPointSampler.h
@@ -9,8 +9,6 @@
 
 #include "PCGPointSampler.generated.h"
 
-class FPCGPointSamplerElement;
-
 UCLASS(BlueprintType, ClassGroup = (Procedural))
 class UPCGPointSamplerSettings : public UPCGSettings
 {
@@ -29,6 +27,11 @@ protected:
 public:
 	UPROPERTY(BlueprintReadWrite, EditAnywhere, Category = Settings, meta=(ClampMin="0", ClampMax="1"))
 	float Ratio = 0.1f;
+
+#if WITH_EDITORONLY_DATA
+	UPROPERTY(Transient, BlueprintReadWrite, EditAnywhere, Category = Debug)
+	bool bKeepZeroDensityPoints = false;
+#endif
 };
 
 class FPCGPointSamplerElement : public FSimpleTypedPCGElement<UPCGPointSamplerSettings>
diff --git a/Engine/Plugins/Experimental/PCG/Source/PCG/Public/Elements/PCGSelfPruning.h b/Engine/Plugins/Experimental/PCG/Source/PCG/Public/Elements/PCGSelfPruning.h
--- a/Engine/Plugins/Experimental/PCG/Source/PCG/Public/Elements/PCGSelfPruning.h
+++ b/Engine/Plugins/Experimental/PCG/Source/PCG/Public/Elements/PCGSelfPruning.h
@@ -27,7 +27,7 @@ public:
 	UPROPERTY(BlueprintReadWrite, EditAnywhere, Category = Settings)
 	EPCGSelfPruningType PruningType = EPCGSelfPruningType::LargeToSmall;
 
-	UPROPERTY(BlueprintReadWrite, EditAnywhere, Category = Settings, meta=(ClampMin=0.0f, EditCondition="PruningType == EPCGSelfPruningType::LargeToSmall || PruningType == PruningType == EPCGSelfPruningType::SmallToLarge"))
+	UPROPERTY(BlueprintReadWrite, EditAnywhere, Category = Settings, meta=(ClampMin=0.0f, EditCondition="PruningType == EPCGSelfPruningType::LargeToSmall || PruningType == EPCGSelfPruningType::SmallToLarge"))
 	float RadiusSimilarityFactor = 0.25f;
 
 	UPROPERTY(BlueprintReadWrite, EditAnywhere, Category = Settings)
diff --git a/Engine/Plugins/Experimental/PCG/Source/PCG/Public/PCGComponent.h b/Engine/Plugins/Experimental/PCG/Source/PCG/Public/PCGComponent.h
--- a/Engine/Plugins/Experimental/PCG/Source/PCG/Public/PCGComponent.h
+++ b/Engine/Plugins/Experimental/PCG/Source/PCG/Public/PCGComponent.h
@@ -48,7 +48,6 @@ public:
 	//~End UActorComponent Interface
 
 	UPCGData* GetPCGData();
-	TArray<UPCGData*> GetPCGExclusionData();
 
 	bool CanPartition() const;
 	bool HasSplineComponent() const;
@@ -124,10 +123,13 @@ protected:
 
 private:
 	UPCGData* CreatePCGData();
+	UPCGData* CreateInputPCGData();
+	UPCGData* GetInputPCGData();
 	UPCGData* GetActorPCGData();
 	UPCGData* CreateActorPCGData();
 	UPCGData* CreateActorPCGData(AActor* Actor);
 	void UpdatePCGExclusionData();
+	TArray<UPCGData*> GetPCGExclusionData();
 
 	bool ShouldGenerate(bool bForce = false) const;
 	FPCGTaskId GenerateInternal(bool bForce, const TArray<FPCGTaskId>& Dependencies);
@@ -164,6 +166,9 @@ private:
 	FBox GetGridBounds() const;
 	FBox GetGridBounds(AActor* InActor) const;
 
+	UPROPERTY(Transient, NonPIEDuplicateTransient)
+	UPCGData* CachedPCGData = nullptr;
+
 	UPROPERTY(Transient, NonPIEDuplicateTransient)
 	UPCGData* CachedInputData = nullptr;
 
