Added debug points display transform scale

diff --git a/Engine/Plugins/Experimental/PCG/Source/PCG/Private/Elements/PCGDebugElement.cpp b/Engine/Plugins/Experimental/PCG/Source/PCG/Private/Elements/PCGDebugElement.cpp
--- a/Engine/Plugins/Experimental/PCG/Source/PCG/Private/Elements/PCGDebugElement.cpp
+++ b/Engine/Plugins/Experimental/PCG/Source/PCG/Private/Elements/PCGDebugElement.cpp
@@ -77,9 +77,11 @@ namespace PCGDebugElement
 			InstanceCustomData.Reserve(NumCustomData);
 
 			// First, create target instance transforms
+			const float PointScale = DebugSettings.PointScale;
 			for (const FPCGPoint& Point : Points)
 			{
-				Instances.Add(Point.Transform);
+				FTransform& InstanceTransform = Instances.Add_GetRef(Point.Transform);
+				InstanceTransform.SetScale3D(InstanceTransform.GetScale3D() * PointScale);
 			}
 
 			UInstancedStaticMeshComponent* ISMC = UPCGActorHelpers::GetOrCreateISMC(TargetActor, Context->SourceComponent, Mesh, { Material });
diff --git a/Engine/Plugins/Experimental/PCG/Source/PCG/Public/PCGDebug.h b/Engine/Plugins/Experimental/PCG/Source/PCG/Public/PCGDebug.h
--- a/Engine/Plugins/Experimental/PCG/Source/PCG/Public/PCGDebug.h
+++ b/Engine/Plugins/Experimental/PCG/Source/PCG/Public/PCGDebug.h
@@ -17,6 +17,9 @@ struct FPCGDebugVisualizationSettings
 public:
 	FPCGDebugVisualizationSettings();
 
+	UPROPERTY(BlueprintReadWrite, EditAnywhere, Category = Settings, meta=(ClampMin="0"))
+	float PointScale = 1.0f;
+
 	UPROPERTY(BlueprintReadWrite, EditAnywhere, Category = Settings)
 	TSoftObjectPtr<UStaticMesh> PointMesh;
 
