Fixed static analysis warnings

#rnx

diff --git a/Engine/Plugins/Experimental/PCG/Source/PCG/Private/Elements/PCGDensityFilter.cpp b/Engine/Plugins/Experimental/PCG/Source/PCG/Private/Elements/PCGDensityFilter.cpp
--- a/Engine/Plugins/Experimental/PCG/Source/PCG/Private/Elements/PCGDensityFilter.cpp
+++ b/Engine/Plugins/Experimental/PCG/Source/PCG/Private/Elements/PCGDensityFilter.cpp
@@ -25,8 +25,8 @@ bool FPCGDensityFilterElement::ExecuteInternal(FPCGContextPtr Context) const
 	const float MinBound = FMath::Min(Settings->LowerBound, Settings->UpperBound);
 	const float MaxBound = FMath::Max(Settings->LowerBound, Settings->UpperBound);
 
-	const bool bNoResults = (MaxBound <= 0 && !Settings->bInvertFilter) || (MinBound == 0 && MaxBound >= 1 && Settings->bInvertFilter);
-	const bool bTrivialFilter = (MinBound <= 0 && MaxBound >= 1.0 && !Settings->bInvertFilter) || (MinBound == 0 && MaxBound == 0 && Settings->bInvertFilter);
+	const bool bNoResults = (MaxBound <= 0.0f && !Settings->bInvertFilter) || (MinBound == 0.0f && MaxBound >= 1.0f && Settings->bInvertFilter);
+	const bool bTrivialFilter = (MinBound <= 0.0f && MaxBound >= 1.0f && !Settings->bInvertFilter) || (MinBound == 0.0f && MaxBound == 0.0f && Settings->bInvertFilter);
 
 #if WITH_EDITORONLY_DATA
 	if(bNoResults && !Settings->bKeepZeroDensityPoints)
diff --git a/Engine/Plugins/Experimental/PCG/Source/PCG/Private/Elements/PCGDensityRemap.cpp b/Engine/Plugins/Experimental/PCG/Source/PCG/Private/Elements/PCGDensityRemap.cpp
--- a/Engine/Plugins/Experimental/PCG/Source/PCG/Private/Elements/PCGDensityRemap.cpp
+++ b/Engine/Plugins/Experimental/PCG/Source/PCG/Private/Elements/PCGDensityRemap.cpp
@@ -27,7 +27,7 @@ bool FPCGLinearDensityRemapElement::ExecuteInternal(FPCGContextPtr Context) cons
 	const float RemapMin = FMath::Min(Settings->RemapMin, Settings->RemapMax);
 	const float RemapMax = FMath::Max(Settings->RemapMin, Settings->RemapMax);
 
-	const bool bTrivialRemapping = (RemapMin == 0 && RemapMax == 1 && Settings->bMultiplyDensity);
+	const bool bTrivialRemapping = (RemapMin == 0.0f && RemapMax == 1.0f && Settings->bMultiplyDensity);
 
 	// TODO: embarassingly parallel loop
 	for (const FPCGTaggedData& Input : Inputs)
diff --git a/Engine/Plugins/Experimental/PCG/Source/PCG/Private/Elements/PCGPointSampler.cpp b/Engine/Plugins/Experimental/PCG/Source/PCG/Private/Elements/PCGPointSampler.cpp
--- a/Engine/Plugins/Experimental/PCG/Source/PCG/Private/Elements/PCGPointSampler.cpp
+++ b/Engine/Plugins/Experimental/PCG/Source/PCG/Private/Elements/PCGPointSampler.cpp
@@ -24,8 +24,8 @@ bool FPCGPointSamplerElement::ExecuteInternal(FPCGContextPtr Context) const
 	Outputs.Append(Context->InputData.GetExclusions());
 	Outputs.Append(Context->InputData.GetAllSettings());
 
-	const bool bNoSampling = (Settings->Ratio <= 0);
-	const bool bTrivialSampling = (Settings->Ratio >= 1);
+	const bool bNoSampling = (Settings->Ratio <= 0.0f);
+	const bool bTrivialSampling = (Settings->Ratio >= 1.0f);
 
 	// Early exit when nothing will be generated out of this sampler
 #if WITH_EDITORONLY_DATA
diff --git a/Engine/Plugins/Experimental/PCG/Source/PCG/Private/Elements/PCGStaticMeshSpawner.cpp b/Engine/Plugins/Experimental/PCG/Source/PCG/Private/Elements/PCGStaticMeshSpawner.cpp
--- a/Engine/Plugins/Experimental/PCG/Source/PCG/Private/Elements/PCGStaticMeshSpawner.cpp
+++ b/Engine/Plugins/Experimental/PCG/Source/PCG/Private/Elements/PCGStaticMeshSpawner.cpp
@@ -85,7 +85,7 @@ bool FPCGStaticMeshSpawnerElement::ExecuteInternal(FPCGContextPtr Context) const
 
 			for (const FPCGPoint& Point : Points)
 			{
-				if (Point.Density <= 0)
+				if (Point.Density <= 0.0f)
 				{
 					continue;
 				}
