Fixed issue where tracked actors wouldn't get picked up correctly.
Fixed issue where tracked actors in subgraphs in spawned pcg actors wouldn't be picked up.
Improved caching cleanup on tracked actor change.

diff --git a/Engine/Plugins/Experimental/PCG/Source/PCG/Private/Elements/PCGExecuteBlueprint.cpp b/Engine/Plugins/Experimental/PCG/Source/PCG/Private/Elements/PCGExecuteBlueprint.cpp
--- a/Engine/Plugins/Experimental/PCG/Source/PCG/Private/Elements/PCGExecuteBlueprint.cpp
+++ b/Engine/Plugins/Experimental/PCG/Source/PCG/Private/Elements/PCGExecuteBlueprint.cpp
@@ -290,15 +290,16 @@ void UPCGBlueprintSettings::RefreshBlueprintElement()
 }
 
 #if WITH_EDITOR
-TArray<FName> UPCGBlueprintSettings::GetTrackedActorTags() const
+void UPCGBlueprintSettings::GetTrackedActorTags(FPCGTagToSettingsMap& OutTagToSettings) const
 {
 #if WITH_EDITORONLY_DATA
-	return TrackedActorTags;
-#else
-	return Super::GetTrackedActorTags();
-#endif
+	for (const FName& Tag : TrackedActorTags)
+	{
+		OutTagToSettings.FindOrAdd(Tag).Add(this);
+	}
+#endif // WITH_EDITORONLY_DATA
 }
-#endif
+#endif // WITH_EDITOR
 
 FPCGElementPtr UPCGBlueprintSettings::CreateElement() const
 {
diff --git a/Engine/Plugins/Experimental/PCG/Source/PCG/Private/Graph/PCGGraphCache.cpp b/Engine/Plugins/Experimental/PCG/Source/PCG/Private/Graph/PCGGraphCache.cpp
--- a/Engine/Plugins/Experimental/PCG/Source/PCG/Private/Graph/PCGGraphCache.cpp
+++ b/Engine/Plugins/Experimental/PCG/Source/PCG/Private/Graph/PCGGraphCache.cpp
@@ -99,42 +99,14 @@ void FPCGGraphCache::ClearCache()
 }
 
 #if WITH_EDITOR
-void FPCGGraphCache::CleanFromCache(AActor* InActor)
+void FPCGGraphCache::CleanFromCache(const IPCGElement* InElement)
 {
-	if (!InActor || InActor->Tags.IsEmpty())
+	if (!InElement)
 	{
 		return;
 	}
 
 	FWriteScopeLock ScopeWriteLock(CacheLock);
-
-	for (TPair<const IPCGElement*, FPCGGraphCacheEntries>& ElementToEntries : CacheData)
-	{
-		FPCGGraphCacheEntries& Entries = ElementToEntries.Value;
-		for (int32 EntryIndex = Entries.Num() - 1; EntryIndex >= 0; --EntryIndex)
-		{
-			if (!Entries[EntryIndex].Settings)
-			{
-				continue;
-			}
-
-			TArray<FName> TrackedTags = Entries[EntryIndex].Settings->GetTrackedActorTags();
-
-			bool bShouldClean = false;
-			for (const FName& Tag : InActor->Tags)
-			{
-				if (TrackedTags.Contains(Tag))
-				{
-					bShouldClean = true;
-					break;
-				}
-			}
-
-			if (bShouldClean)
-			{
-				Entries.RemoveAtSwap(EntryIndex);
-			}
-		}
-	}
+	CacheData.Remove(InElement);
 }
 #endif // WITH_EDITOR
\ No newline at end of file
diff --git a/Engine/Plugins/Experimental/PCG/Source/PCG/Private/Graph/PCGGraphCache.h b/Engine/Plugins/Experimental/PCG/Source/PCG/Private/Graph/PCGGraphCache.h
--- a/Engine/Plugins/Experimental/PCG/Source/PCG/Private/Graph/PCGGraphCache.h
+++ b/Engine/Plugins/Experimental/PCG/Source/PCG/Private/Graph/PCGGraphCache.h
@@ -45,7 +45,7 @@ struct FPCGGraphCache
 
 #if WITH_EDITOR
 	/** Clears any cache entry for the given data */
-	void CleanFromCache(AActor* InObject);
+	void CleanFromCache(const IPCGElement* InElement);
 #endif
 
 private:
diff --git a/Engine/Plugins/Experimental/PCG/Source/PCG/Private/PCGComponent.cpp b/Engine/Plugins/Experimental/PCG/Source/PCG/Private/PCGComponent.cpp
--- a/Engine/Plugins/Experimental/PCG/Source/PCG/Private/PCGComponent.cpp
+++ b/Engine/Plugins/Experimental/PCG/Source/PCG/Private/PCGComponent.cpp
@@ -3,7 +3,6 @@
 #include "PCGComponent.h"
 #include "PCGGraph.h"
 #include "PCGHelpers.h"
-#include "PCGSubsystem.h"
 #include "PCGVolume.h"
 #include "Data/PCGDifferenceData.h"
 #include "Data/PCGIntersectionData.h"
@@ -587,6 +586,7 @@ void UPCGComponent::PostEditUndo()
 	SetupTrackingCallbacks();
 	RefreshTrackingData();
 	DirtyGenerated(/*bDirtyCachedInput=*/true);
+	DirtyCacheForAllTrackedTags();
 	
 	if (bWasGeneratedPriorToUndo && !bGenerated)
 	{
@@ -614,13 +614,13 @@ void UPCGComponent::TeardownActorCallbacks()
 
 void UPCGComponent::SetupTrackingCallbacks()
 {
-	CachedTrackedActorTags.Reset();
+	CachedTrackedTagsToSettings.Reset();
 	if (Graph)
 	{
-		CachedTrackedActorTags.Append(Graph->GetTrackedActorTags());
+		CachedTrackedTagsToSettings = Graph->GetTrackedTagsToSettings();
 	}
 
-	if(!ExcludedTags.IsEmpty() || !CachedTrackedActorTags.IsEmpty())
+	if(!ExcludedTags.IsEmpty() || !CachedTrackedTagsToSettings.IsEmpty())
 	{
 		GEngine->OnLevelActorAdded().AddUObject(this, &UPCGComponent::OnActorAdded);
 		GEngine->OnLevelActorDeleted().AddUObject(this, &UPCGComponent::OnActorDeleted);
@@ -630,7 +630,11 @@ void UPCGComponent::SetupTrackingCallbacks()
 void UPCGComponent::RefreshTrackingData()
 {
 	GetActorsFromTags(ExcludedTags, CachedExcludedActors);
-	GetActorsFromTags(CachedTrackedActorTags, CachedTrackedActors);
+
+	TSet<FName> TrackedTags;
+	CachedTrackedTagsToSettings.GetKeys(TrackedTags);
+	GetActorsFromTags(TrackedTags, CachedTrackedActors);
+	PopulateTrackedActorToTagsMap(/*bForce=*/true);
 }
 
 void UPCGComponent::TeardownTrackingCallbacks()
@@ -679,7 +683,7 @@ bool UPCGComponent::ActorIsTracked(AActor* InActor) const
 	bool bIsTracked = false;
 	for (const FName& Tag : InActor->Tags)
 	{
-		if (CachedTrackedActorTags.Contains(Tag))
+		if (CachedTrackedTagsToSettings.Contains(Tag))
 		{
 			bIsTracked = true;
 			break;
@@ -692,12 +696,7 @@ bool UPCGComponent::ActorIsTracked(AActor* InActor) const
 void UPCGComponent::OnActorAdded(AActor* InActor)
 {
 	const bool bHasExcludedTag = ActorHasExcludedTag(InActor);
-	const bool bIsTracked = ActorIsTracked(InActor);
-
-	if (bIsTracked)
-	{
-		CachedTrackedActors.Add(InActor);
-	}
+	const bool bIsTracked = AddTrackedActor(InActor);
 
 	if (bHasExcludedTag || bIsTracked)
 	{
@@ -709,15 +708,9 @@ void UPCGComponent::OnActorAdded(AActor* InActor)
 void UPCGComponent::OnActorDeleted(AActor* InActor)
 {
 	const bool bHasExcludedTag = ActorHasExcludedTag(InActor);
-	const bool bIsTracked = ActorIsTracked(InActor);
+	const bool bWasTracked = RemoveTrackedActor(InActor);
 
-	if (bIsTracked)
-	{
-		CachedTrackedActors.Remove(InActor);
-		GetSubsystem()->CleanFromCache(InActor);
-	}
-
-	if (bHasExcludedTag || bIsTracked)
+	if (bHasExcludedTag || bWasTracked)
 	{
 		DirtyGenerated();
 		Refresh();
@@ -738,20 +731,20 @@ void UPCGComponent::OnActorMoved(AActor* InActor)
 	}
 	else
 	{
-		const bool bHasExcludedTag = ActorHasExcludedTag(InActor);
-		const bool bActorIsTracked = ActorIsTracked(InActor);
-		if (bHasExcludedTag || bActorIsTracked)
+		bool bDirtyAndRefresh = false;
+
+		if (ActorHasExcludedTag(InActor))
 		{
-			if (bHasExcludedTag)
-			{
-				DirtyExclusionData(InActor);
-			}
+			bDirtyAndRefresh = true;
+		}
 
-			if (bActorIsTracked)
-			{
-				GetSubsystem()->CleanFromCache(InActor);
-			}
+		if (DirtyTrackedActor(InActor))
+		{
+			bDirtyAndRefresh = true;
+		}
 
+		if (bDirtyAndRefresh)
+		{
 			DirtyGenerated();
 			Refresh();
 		}
@@ -792,45 +785,40 @@ void UPCGComponent::OnObjectPropertyChanged(UObject* InObject, FPropertyChangedE
 
 		Refresh();
 	}
-	else
+	else if(Actor)
 	{
-		const bool bActorHasExcludedTag = ActorHasExcludedTag(Actor);
-		const bool bActorIsTracked = ActorIsTracked(Actor);
+		bool bDirtyAndRefresh = false;
 
-		if(bActorHasExcludedTag || bActorIsTracked)
+		if (bActorTagChange && Actor == InObject)
 		{
-			if (bActorHasExcludedTag)
+			// We don't really need to remove the actor here, but it's a good metric on whether the exclusion is dirty
+			if(CachedExcludedActors.Remove(Actor))
 			{
-				DirtyExclusionData(Actor);
+				bDirtyAndRefresh = true;
 			}
 
-			if (bActorIsTracked)
+			if (UpdateTrackedActor(Actor))
 			{
-				GetSubsystem()->CleanFromCache(Actor);
+				bDirtyAndRefresh = true;
 			}
-			
-			DirtyGenerated();
-			Refresh();
 		}
-		else if (bActorTagChange && Actor == InObject)
+		else
 		{
-			bool bDirtyAndRefresh = false;
-			if (CachedExcludedActors.Remove(Actor))
+			if (ActorHasExcludedTag(Actor))
 			{
-				bDirtyAndRefresh = true;
+				DirtyExclusionData(Actor);
 			}
 
-			if (CachedTrackedActors.Remove(Actor))
+			if (DirtyTrackedActor(Actor))
 			{
-				GetSubsystem()->CleanFromCache(Actor);
 				bDirtyAndRefresh = true;
 			}
+		}
 
-			if (bDirtyAndRefresh)
-			{
-				DirtyGenerated();
-				Refresh();
-			}
+		if (bDirtyAndRefresh)
+		{
+			DirtyGenerated();
+			Refresh();
 		}
 	}
 }
@@ -848,6 +836,7 @@ void UPCGComponent::OnGraphChanged(UPCGGraph* InGraph, bool bIsStructural, bool
 		TeardownTrackingCallbacks();
 		SetupTrackingCallbacks();
 		RefreshTrackingData();
+		DirtyCacheForAllTrackedTags();
 
 		DirtyGenerated();
 		if (bShouldRefresh)
@@ -1260,4 +1249,177 @@ UPCGSubsystem* UPCGComponent::GetSubsystem() const
 	return GetOwner() && GetOwner()->GetWorld() ? GetOwner()->GetWorld()->GetSubsystem<UPCGSubsystem>() : nullptr;
 }
 
+#if WITH_EDITOR
+bool UPCGComponent::PopulateTrackedActorToTagsMap(bool bForce)
+{
+	if (bActorToTagsMapPopulated && !bForce)
+	{
+		return false;
+	}
+
+	CachedTrackedActorToTags.Reset();
+	for (TWeakObjectPtr<AActor> Actor : CachedTrackedActors)
+	{
+		if (Actor.IsValid())
+		{
+			AddTrackedActor(Actor.Get(), /*bForce=*/true);
+		}
+	}
+
+	bActorToTagsMapPopulated = true;
+	return true;
+}
+
+bool UPCGComponent::AddTrackedActor(AActor* InActor, bool bForce)
+{
+	if (!bForce)
+	{
+		PopulateTrackedActorToTagsMap();
+	}	
+
+	check(InActor);
+	bool bAppliedChange = false;
+
+	if (bForce)
+	{
+		CachedTrackedActorToTags.FindOrAdd(InActor);
+	}
+
+	for (const FName& Tag : InActor->Tags)
+	{
+		if (!CachedTrackedTagsToSettings.Contains(Tag))
+		{
+			continue;
+		}
+
+		bAppliedChange = true;
+		CachedTrackedActorToTags.FindOrAdd(InActor).Add(Tag);
+
+		if (!bForce)
+		{
+			DirtyCacheFromTag(Tag);
+		}
+	}
+
+	return bAppliedChange;
+}
+
+bool UPCGComponent::RemoveTrackedActor(AActor* InActor)
+{
+	PopulateTrackedActorToTagsMap();
+
+	check(InActor);
+	bool bAppliedChange = false;
+
+	if(CachedTrackedActorToTags.Contains(InActor))
+	{
+		for (const FName& Tag : CachedTrackedActorToTags[InActor])
+		{
+			DirtyCacheFromTag(Tag);
+		}
+
+		CachedTrackedActorToTags.Remove(InActor);
+		bAppliedChange = true;
+	}
+
+	return bAppliedChange;
+}
+
+bool UPCGComponent::UpdateTrackedActor(AActor* InActor)
+{
+	check(InActor);
+	// If the tracked data wasn't initialized before, then it is not possible to know if we need to update or not - take no chances
+	bool bAppliedChange = PopulateTrackedActorToTagsMap();
+
+	// Update the contents of the tracked actor vs. its current tags, and dirty accordingly
+	if (CachedTrackedActorToTags.Contains(InActor))
+	{
+		// Any tags that aren't on the actor and were in the cached actor to tags -> remove & dirty
+		TSet<FName> CachedTags = CachedTrackedActorToTags[InActor];
+		for (const FName& CachedTag : CachedTags)
+		{
+			if (!InActor->Tags.Contains(CachedTag))
+			{
+				CachedTrackedActorToTags[InActor].Remove(CachedTag);
+				DirtyCacheFromTag(CachedTag);
+				bAppliedChange = true;
+			}
+		}
+	}
+		
+	// Any tags that are new on the actor and not in the cached actor to tags -> add & dirty
+	for (const FName& Tag : InActor->Tags)
+	{
+		if (!CachedTrackedTagsToSettings.Contains(Tag))
+		{
+			continue;
+		}
+
+		if (!CachedTrackedActorToTags.FindOrAdd(InActor).Find(Tag))
+		{
+			CachedTrackedActorToTags[InActor].Add(Tag);
+			DirtyCacheFromTag(Tag);
+			bAppliedChange = true;
+		}
+	}
+
+	// Finally, if the current has no tag anymore, we can remove it from the map
+	if (CachedTrackedActorToTags.Contains(InActor) && CachedTrackedActorToTags[InActor].IsEmpty())
+	{
+		CachedTrackedActorToTags.Remove(InActor);
+	}
+
+	return bAppliedChange;
+}
+
+bool UPCGComponent::DirtyTrackedActor(AActor* InActor)
+{
+	PopulateTrackedActorToTagsMap();
+
+	check(InActor);
+	bool bAppliedChange = false;
+
+	if (CachedTrackedActorToTags.Contains(InActor))
+	{
+		for (const FName& Tag : CachedTrackedActorToTags[InActor])
+		{
+			DirtyCacheFromTag(Tag);
+		}
+
+		bAppliedChange = true;
+	}
+
+	return bAppliedChange;
+}
+
+void UPCGComponent::DirtyCacheFromTag(const FName& InTag)
+{
+	if (CachedTrackedTagsToSettings.Contains(InTag))
+	{
+		for (TWeakObjectPtr<const UPCGSettings> Settings : CachedTrackedTagsToSettings[InTag])
+		{
+			if (Settings.IsValid())
+			{
+				GetSubsystem()->CleanFromCache(Settings->GetElement().Get());
+			}
+		}
+	}
+}
+
+void UPCGComponent::DirtyCacheForAllTrackedTags()
+{
+	for (const auto& TagToSettings : CachedTrackedTagsToSettings)
+	{
+		for (TWeakObjectPtr<const UPCGSettings> Settings : TagToSettings.Value)
+		{
+			if (Settings.IsValid())
+			{
+				GetSubsystem()->CleanFromCache(Settings->GetElement().Get());
+			}
+		}
+	}
+}
+
+#endif // WITH_EDITOR
+
 #undef LOCTEXT_NAMESPACE
\ No newline at end of file
diff --git a/Engine/Plugins/Experimental/PCG/Source/PCG/Private/PCGGraph.cpp b/Engine/Plugins/Experimental/PCG/Source/PCG/Private/PCGGraph.cpp
--- a/Engine/Plugins/Experimental/PCG/Source/PCG/Private/PCGGraph.cpp
+++ b/Engine/Plugins/Experimental/PCG/Source/PCG/Private/PCGGraph.cpp
@@ -217,19 +217,18 @@ bool UPCGGraph::Contains(UPCGNode* Node) const
 }
 
 #if WITH_EDITOR
-TArray<FName> UPCGGraph::GetTrackedActorTags() const
+FPCGTagToSettingsMap UPCGGraph::GetTrackedTagsToSettings() const
 {
-	TArray<FName> TrackedActorTags;
-
+	FPCGTagToSettingsMap TagsToSettings;
 	for (UPCGNode* Node : Nodes)
 	{
 		if (Node && Node->DefaultSettings)
 		{
-			TrackedActorTags.Append(Node->DefaultSettings->GetTrackedActorTags());
+			Node->DefaultSettings->GetTrackedActorTags(TagsToSettings);
 		}
 	}
 
-	return TrackedActorTags;
+	return TagsToSettings;
 }
 
 void UPCGGraph::NotifyGraphChanged(bool bIsStructural)
diff --git a/Engine/Plugins/Experimental/PCG/Source/PCG/Private/PCGSubgraph.cpp b/Engine/Plugins/Experimental/PCG/Source/PCG/Private/PCGSubgraph.cpp
--- a/Engine/Plugins/Experimental/PCG/Source/PCG/Private/PCGSubgraph.cpp
+++ b/Engine/Plugins/Experimental/PCG/Source/PCG/Private/PCGSubgraph.cpp
@@ -55,6 +55,19 @@ void UPCGBaseSubgraphSettings::PostEditChangeProperty(struct FPropertyChangedEve
 	}
 }
 
+void UPCGBaseSubgraphSettings::GetTrackedActorTags(FPCGTagToSettingsMap& OutTagToSettings) const
+{
+	if (UPCGGraph* Subgraph = GetSubgraph())
+	{
+		FPCGTagToSettingsMap SubgraphMap = Subgraph->GetTrackedTagsToSettings();
+
+		for (const auto& TagToSettings : SubgraphMap)
+		{
+			OutTagToSettings.FindOrAdd(TagToSettings.Key).Append(TagToSettings.Value);
+		}
+	}
+}
+
 void UPCGBaseSubgraphSettings::OnSubgraphChanged(UPCGGraph* InGraph, bool bIsStructural)
 {
 	if (InGraph == GetSubgraph())
@@ -69,7 +82,6 @@ void UPCGBaseSubgraphSettings::OnSubgraphChanged(UPCGGraph* InGraph, bool bIsStr
 		}
 	}
 }
-
 #endif // WITH_EDITOR
 
 UPCGNode* UPCGSubgraphSettings::CreateNode() const
@@ -82,11 +94,6 @@ bool UPCGSubgraphSettings::IsStructuralProperty(const FName& InPropertyName) con
 {
 	return (InPropertyName == GET_MEMBER_NAME_CHECKED(UPCGSubgraphSettings, Subgraph)) || Super::IsStructuralProperty(InPropertyName);
 }
-
-TArray<FName> UPCGSubgraphSettings::GetTrackedActorTags() const
-{
-	return Subgraph ? Subgraph->GetTrackedActorTags() : TArray<FName>();
-}
 #endif
 
 FPCGElementPtr UPCGSubgraphSettings::CreateElement() const
diff --git a/Engine/Plugins/Experimental/PCG/Source/PCG/Private/PCGSubsystem.cpp b/Engine/Plugins/Experimental/PCG/Source/PCG/Private/PCGSubsystem.cpp
--- a/Engine/Plugins/Experimental/PCG/Source/PCG/Private/PCGSubsystem.cpp
+++ b/Engine/Plugins/Experimental/PCG/Source/PCG/Private/PCGSubsystem.cpp
@@ -488,11 +488,11 @@ void UPCGSubsystem::NotifyGraphChanged(UPCGGraph* InGraph)
 	}
 }
 
-void UPCGSubsystem::CleanFromCache(AActor* InActor)
+void UPCGSubsystem::CleanFromCache(const IPCGElement* InElement)
 {
 	if (GraphExecutor)
 	{
-		GraphExecutor->GetCache().CleanFromCache(InActor);
+		GraphExecutor->GetCache().CleanFromCache(InElement);
 	}
 }
 
diff --git a/Engine/Plugins/Experimental/PCG/Source/PCG/Public/Elements/PCGExecuteBlueprint.h b/Engine/Plugins/Experimental/PCG/Source/PCG/Public/Elements/PCGExecuteBlueprint.h
--- a/Engine/Plugins/Experimental/PCG/Source/PCG/Public/Elements/PCGExecuteBlueprint.h
+++ b/Engine/Plugins/Experimental/PCG/Source/PCG/Public/Elements/PCGExecuteBlueprint.h
@@ -56,7 +56,7 @@ public:
 	// ~Begin UPCGSettings interface
 #if WITH_EDITOR
 	virtual FName GetDefaultNodeName() const override { return FName(TEXT("BlueprintNode")); }
-	virtual TArray<FName> GetTrackedActorTags() const override;
+	virtual void GetTrackedActorTags(FPCGTagToSettingsMap& OutTagToSettings) const override;
 #endif
 
 protected:
diff --git a/Engine/Plugins/Experimental/PCG/Source/PCG/Public/PCGComponent.h b/Engine/Plugins/Experimental/PCG/Source/PCG/Public/PCGComponent.h
--- a/Engine/Plugins/Experimental/PCG/Source/PCG/Public/PCGComponent.h
+++ b/Engine/Plugins/Experimental/PCG/Source/PCG/Public/PCGComponent.h
@@ -4,6 +4,7 @@
 
 #include "CoreMinimal.h"
 #include "Components/ActorComponent.h"
+#include "PCGSettings.h"
 #include "PCGSubsystem.h"
 
 #include "PCGComponent.generated.h"
@@ -161,6 +162,14 @@ private:
 	bool ActorHasExcludedTag(AActor* InActor) const;
 	bool ActorIsTracked(AActor* InActor) const;
 	void DirtyExclusionData(AActor* InActor);
+
+	bool PopulateTrackedActorToTagsMap(bool bForce = false);
+	bool AddTrackedActor(AActor* InActor, bool bForce = false);
+	bool RemoveTrackedActor(AActor* InActor);
+	bool UpdateTrackedActor(AActor* InActor);
+	bool DirtyTrackedActor(AActor* InActor);
+	void DirtyCacheFromTag(const FName& InTag);
+	void DirtyCacheForAllTrackedTags();
 #endif
 
 	FBox GetGridBounds() const;
@@ -178,6 +187,7 @@ private:
 	UPROPERTY(Transient, NonPIEDuplicateTransient)
 	TMap<AActor*, UPCGData*> CachedExclusionData;
 
+	// Cached excluded actors list is serialized because we can't get it at postload time
 	UPROPERTY()
 	TSet<TWeakObjectPtr<AActor>> CachedExcludedActors; 
 
@@ -191,11 +201,15 @@ private:
 	bool bIsGenerating = false;
 	bool bWasGeneratedPriorToUndo = false;
 	FBox LastGeneratedBoundsPriorToUndo = FBox(EForceInit::ForceInit);
-	TSet<FName> CachedTrackedActorTags;
+	FPCGTagToSettingsMap CachedTrackedTagsToSettings;
 #endif
 
 #if WITH_EDITORONLY_DATA
+	// Cached tracked actors list is serialized because we can't get it at postload time
 	UPROPERTY()
 	TSet<TWeakObjectPtr<AActor>> CachedTrackedActors;
+
+	TMap<TWeakObjectPtr<AActor>, TSet<FName>> CachedTrackedActorToTags;
+	bool bActorToTagsMapPopulated = false;
 #endif
 };
\ No newline at end of file
diff --git a/Engine/Plugins/Experimental/PCG/Source/PCG/Public/PCGGraph.h b/Engine/Plugins/Experimental/PCG/Source/PCG/Public/PCGGraph.h
--- a/Engine/Plugins/Experimental/PCG/Source/PCG/Public/PCGGraph.h
+++ b/Engine/Plugins/Experimental/PCG/Source/PCG/Public/PCGGraph.h
@@ -58,7 +58,7 @@ public:
 	bool Contains(UPCGNode* Node) const;
 	const TArray<UPCGNode*>& GetNodes() const { return Nodes; }
 #if WITH_EDITOR
-	TArray<FName> GetTrackedActorTags() const;
+	FPCGTagToSettingsMap GetTrackedTagsToSettings() const;
 #endif
 
 #if WITH_EDITOR
diff --git a/Engine/Plugins/Experimental/PCG/Source/PCG/Public/PCGSettings.h b/Engine/Plugins/Experimental/PCG/Source/PCG/Public/PCGSettings.h
--- a/Engine/Plugins/Experimental/PCG/Source/PCG/Public/PCGSettings.h
+++ b/Engine/Plugins/Experimental/PCG/Source/PCG/Public/PCGSettings.h
@@ -9,6 +9,9 @@
 #include "PCGSettings.generated.h"
 
 class UPCGNode;
+class UPCGSettings;
+
+typedef TMap<FName, TSet<TWeakObjectPtr<const UPCGSettings>>> FPCGTagToSettingsMap;
 
 UENUM()
 enum class EPCGSettingsExecutionMode : uint8
@@ -41,7 +44,7 @@ public:
 #if WITH_EDITOR
 	virtual FName GetDefaultNodeName() const { return NAME_None; }
 	/** Derived classes must implement this to communicate dependencies on external actors */
-	virtual TArray<FName> GetTrackedActorTags() const { return TArray<FName>(); }
+	virtual void GetTrackedActorTags(FPCGTagToSettingsMap& OutTagToSettings) const {}
 #endif
 
 	UPROPERTY(EditAnywhere, BlueprintReadWrite, Category = Settings)
diff --git a/Engine/Plugins/Experimental/PCG/Source/PCG/Public/PCGSubgraph.h b/Engine/Plugins/Experimental/PCG/Source/PCG/Public/PCGSubgraph.h
--- a/Engine/Plugins/Experimental/PCG/Source/PCG/Public/PCGSubgraph.h
+++ b/Engine/Plugins/Experimental/PCG/Source/PCG/Public/PCGSubgraph.h
@@ -30,6 +30,12 @@ protected:
 	virtual void PostEditChangeProperty(struct FPropertyChangedEvent& PropertyChangedEvent) override;
 	//~End UObject interface implementation
 
+	//~Begin UPCGSettings interface
+#if WITH_EDITOR
+	virtual void GetTrackedActorTags(FPCGTagToSettingsMap& OutTagToSettings) const override;
+#endif
+	//~End UPCGSettings interface
+
 	void OnSubgraphChanged(UPCGGraph* InGraph, bool bIsStructural);
 
 	virtual bool IsStructuralProperty(const FName& InPropertyName) const { return false; }
@@ -52,7 +58,6 @@ public:
 
 #if WITH_EDITOR
 	virtual FName GetDefaultNodeName() const override { return FName(TEXT("SubgraphNode")); }
-	virtual TArray<FName> GetTrackedActorTags() const override;
 #endif
 
 protected:
diff --git a/Engine/Plugins/Experimental/PCG/Source/PCG/Public/PCGSubsystem.h b/Engine/Plugins/Experimental/PCG/Source/PCG/Public/PCGSubsystem.h
--- a/Engine/Plugins/Experimental/PCG/Source/PCG/Public/PCGSubsystem.h
+++ b/Engine/Plugins/Experimental/PCG/Source/PCG/Public/PCGSubsystem.h
@@ -77,12 +77,11 @@ public:
 	void CleanupPartitionActors(const FBox& InBounds);
 	void DeletePartitionActors();
 
-
 	/** Propagate to the graph compiler graph changes */
 	void NotifyGraphChanged(UPCGGraph* InGraph);
 
-	/** Cleans up the graph cache when a tracked actor is changed */
-	void CleanFromCache(AActor* InActor);
+	/** Cleans up the graph cache on an element basis */
+	void CleanFromCache(const IPCGElement* InElement);
 
 private:
 	FPCGTaskId DelayProcessGraph(UPCGComponent* Component, bool bGenerate, bool bSave, bool bUseEmptyNewBounds);
