Removed useless condition in actor tracking

#rnx

diff --git a/Engine/Plugins/Experimental/PCG/Source/PCG/Private/PCGComponent.cpp b/Engine/Plugins/Experimental/PCG/Source/PCG/Private/PCGComponent.cpp
--- a/Engine/Plugins/Experimental/PCG/Source/PCG/Private/PCGComponent.cpp
+++ b/Engine/Plugins/Experimental/PCG/Source/PCG/Private/PCGComponent.cpp
@@ -1280,11 +1280,6 @@ bool UPCGComponent::AddTrackedActor(AActor* InActor, bool bForce)
 	check(InActor);
 	bool bAppliedChange = false;
 
-	if (bForce)
-	{
-		CachedTrackedActorToTags.FindOrAdd(InActor);
-	}
-
 	for (const FName& Tag : InActor->Tags)
 	{
 		if (!CachedTrackedTagsToSettings.Contains(Tag))
