Made debug settings non-transient
Repurposed the material field to be an override
Added absolute/relative scale control

diff --git a/Engine/Plugins/Experimental/PCG/Source/PCG/Private/Elements/PCGDebugElement.cpp b/Engine/Plugins/Experimental/PCG/Source/PCG/Private/Elements/PCGDebugElement.cpp
--- a/Engine/Plugins/Experimental/PCG/Source/PCG/Private/Elements/PCGDebugElement.cpp
+++ b/Engine/Plugins/Experimental/PCG/Source/PCG/Private/Elements/PCGDebugElement.cpp
@@ -32,7 +32,13 @@ namespace PCGDebugElement
 			return;
 		}
 
-		UMaterialInterface* Material = DebugSettings.Material.LoadSynchronous();
+		UMaterialInterface* Material = DebugSettings.GetMaterial().LoadSynchronous();
+
+		TArray<UMaterialInterface*> Materials;
+		if (Material)
+		{
+			Materials.Add(Material);
+		}
 
 		TArray<FPCGTaggedData> Outputs = Context->OutputData.GetInputs();
 
@@ -78,13 +84,16 @@ namespace PCGDebugElement
 
 			// First, create target instance transforms
 			const float PointScale = DebugSettings.PointScale;
+			const bool bIsRelative = DebugSettings.ScaleMethod == EPCGDebugVisScaleMethod::Relative;
+
 			for (const FPCGPoint& Point : Points)
 			{
 				FTransform& InstanceTransform = Instances.Add_GetRef(Point.Transform);
-				InstanceTransform.SetScale3D(InstanceTransform.GetScale3D() * PointScale);
+				InstanceTransform.SetScale3D(bIsRelative ? InstanceTransform.GetScale3D() * PointScale : FVector(PointScale));
 			}
 
-			UInstancedStaticMeshComponent* ISMC = UPCGActorHelpers::GetOrCreateISMC(TargetActor, Context->SourceComponent, Mesh, { Material });
+			UInstancedStaticMeshComponent* ISMC = UPCGActorHelpers::GetOrCreateISMC(TargetActor, Context->SourceComponent, Mesh, Materials);
+			
 			ISMC->ComponentTags.AddUnique(PCGHelpers::DefaultPCGDebugTag);
 			ISMC->NumCustomDataFloats = NumCustomData;
 			const int32 PreExistingInstanceCount = ISMC->GetInstanceCount();
diff --git a/Engine/Plugins/Experimental/PCG/Source/PCG/Private/PCGDebug.cpp b/Engine/Plugins/Experimental/PCG/Source/PCG/Private/PCGDebug.cpp
--- a/Engine/Plugins/Experimental/PCG/Source/PCG/Private/PCGDebug.cpp
+++ b/Engine/Plugins/Experimental/PCG/Source/PCG/Private/PCGDebug.cpp
@@ -1,9 +1,27 @@
 // Copyright Epic Games, Inc. All Rights Reserved.
 
 #include "PCGDebug.h"
+#include "Materials/MaterialInterface.h"
+
+namespace PCGDebugVisConstants
+{
+	const FSoftObjectPath DefaultPointMesh = FSoftObjectPath(TEXT("/Engine/BasicShapes/Cube.Cube"));
+	const FSoftObjectPath MaterialForDefaultPointMesh = FSoftObjectPath(TEXT("Material'/Game/PCG/DebugMaterial.DebugMaterial'"));
+}
 
 FPCGDebugVisualizationSettings::FPCGDebugVisualizationSettings()
 {
-	PointMesh = FSoftObjectPath(TEXT("/Engine/BasicShapes/Cube.Cube"));
-	Material = FSoftObjectPath(TEXT("Material'/Game/PCG/DebugMaterial.DebugMaterial'"));
+	PointMesh = PCGDebugVisConstants::DefaultPointMesh;
+}
+
+TSoftObjectPtr<UMaterialInterface> FPCGDebugVisualizationSettings::GetMaterial() const
+{
+	if (!MaterialOverride.IsValid() && PointMesh.ToSoftObjectPath() == PCGDebugVisConstants::DefaultPointMesh)
+	{
+		return TSoftObjectPtr<UMaterialInterface>(PCGDebugVisConstants::MaterialForDefaultPointMesh);
+	}
+	else
+	{
+		return MaterialOverride;
+	}
 }
\ No newline at end of file
diff --git a/Engine/Plugins/Experimental/PCG/Source/PCG/Public/PCGDebug.h b/Engine/Plugins/Experimental/PCG/Source/PCG/Public/PCGDebug.h
--- a/Engine/Plugins/Experimental/PCG/Source/PCG/Public/PCGDebug.h
+++ b/Engine/Plugins/Experimental/PCG/Source/PCG/Public/PCGDebug.h
@@ -9,6 +9,13 @@
 class UStaticMesh;
 class UMaterialInterface;
 
+UENUM()
+enum class EPCGDebugVisScaleMethod : uint8
+{
+	Relative,
+	Absolute
+};
+
 USTRUCT(BlueprintType)
 struct FPCGDebugVisualizationSettings
 {
@@ -20,9 +27,14 @@ public:
 	UPROPERTY(BlueprintReadWrite, EditAnywhere, Category = Settings, meta=(ClampMin="0"))
 	float PointScale = 1.0f;
 
+	UPROPERTY(BlueprintReadWrite, EditAnywhere, Category = Settings)
+	EPCGDebugVisScaleMethod ScaleMethod = EPCGDebugVisScaleMethod::Relative;
+
 	UPROPERTY(BlueprintReadWrite, EditAnywhere, Category = Settings)
 	TSoftObjectPtr<UStaticMesh> PointMesh;
 
 	UPROPERTY(BlueprintReadWrite, EditAnywhere, Category = Settings)
-	TSoftObjectPtr<UMaterialInterface> Material;
+	TSoftObjectPtr<UMaterialInterface> MaterialOverride;
+
+	TSoftObjectPtr<UMaterialInterface> GetMaterial() const;
 };
\ No newline at end of file
diff --git a/Engine/Plugins/Experimental/PCG/Source/PCG/Public/PCGSettings.h b/Engine/Plugins/Experimental/PCG/Source/PCG/Public/PCGSettings.h
--- a/Engine/Plugins/Experimental/PCG/Source/PCG/Public/PCGSettings.h
+++ b/Engine/Plugins/Experimental/PCG/Source/PCG/Public/PCGSettings.h
@@ -66,7 +66,7 @@ public:
 	EPCGSettingsExecutionMode ExecutionMode = EPCGSettingsExecutionMode::Enabled;
 
 #if WITH_EDITORONLY_DATA
-	UPROPERTY(Transient, EditAnywhere, BlueprintReadWrite, Category = Debug, meta = (ShowOnlyInnerProperties))
+	UPROPERTY(EditAnywhere, BlueprintReadWrite, Category = Debug, meta = (ShowOnlyInnerProperties))
 	FPCGDebugVisualizationSettings DebugSettings;
 #endif
 
