Made PCG settings operator== ignore debug settings (since they're now non-transient)

#rnx

diff --git a/Engine/Plugins/Experimental/PCG/Source/PCG/Private/PCGSettings.cpp b/Engine/Plugins/Experimental/PCG/Source/PCG/Private/PCGSettings.cpp
--- a/Engine/Plugins/Experimental/PCG/Source/PCG/Private/PCGSettings.cpp
+++ b/Engine/Plugins/Experimental/PCG/Source/PCG/Private/PCGSettings.cpp
@@ -4,6 +4,20 @@
 #include "PCGNode.h"
 #include "Serialization/ArchiveObjectCrc32.h"
 
+/** In order to reuse the cache when only debug settings change, we must make sure to ignore these from the CRC check */
+#if WITH_EDITORONLY_DATA
+class FPCGSettingsObjectCrc32 : public FArchiveObjectCrc32
+{
+public:
+	virtual bool ShouldSkipProperty(const FProperty* InProperty) const override
+	{
+		return !(InProperty && InProperty->GetFName() == GET_MEMBER_NAME_CHECKED(UPCGSettings, DebugSettings));
+	}
+};
+#else
+typedef FArchiveObjectCrc32 FPCGSettingsObjectCrc32;
+#endif
+
 bool UPCGSettings::operator==(const UPCGSettings& Other) const
 {
 	if (this == &Other)
@@ -12,7 +26,7 @@ bool UPCGSettings::operator==(const UPCGSettings& Other) const
 	}
 	else
 	{
-		FArchiveObjectCrc32 Ar;
+		FPCGSettingsObjectCrc32 Ar;
 		uint32 ThisCrc = Ar.Crc32(const_cast<UPCGSettings*>(this));
 		uint32 OtherCrc = Ar.Crc32(const_cast<UPCGSettings*>(&Other));
 		return ThisCrc == OtherCrc;
