Fix trivial logic error in settings comparison

#rnx

diff --git a/Engine/Plugins/Experimental/PCG/Source/PCG/Private/PCGSettings.cpp b/Engine/Plugins/Experimental/PCG/Source/PCG/Private/PCGSettings.cpp
--- a/Engine/Plugins/Experimental/PCG/Source/PCG/Private/PCGSettings.cpp
+++ b/Engine/Plugins/Experimental/PCG/Source/PCG/Private/PCGSettings.cpp
@@ -11,7 +11,7 @@ class FPCGSettingsObjectCrc32 : public FArchiveObjectCrc32
 public:
 	virtual bool ShouldSkipProperty(const FProperty* InProperty) const override
 	{
-		return !(InProperty && InProperty->GetFName() == GET_MEMBER_NAME_CHECKED(UPCGSettings, DebugSettings));
+		return InProperty && InProperty->GetFName() == GET_MEMBER_NAME_CHECKED(UPCGSettings, DebugSettings);
 	}
 };
 #else
