Added basic logging (errors, timing, additional info)

diff --git a/Engine/Plugins/Experimental/PCG/Source/PCG/Private/Data/PCGDifferenceData.cpp b/Engine/Plugins/Experimental/PCG/Source/PCG/Private/Data/PCGDifferenceData.cpp
--- a/Engine/Plugins/Experimental/PCG/Source/PCG/Private/Data/PCGDifferenceData.cpp
+++ b/Engine/Plugins/Experimental/PCG/Source/PCG/Private/Data/PCGDifferenceData.cpp
@@ -135,8 +135,15 @@ const UPCGPointData* UPCGDifferenceData::CreatePointData() const
 	// This is similar to what we are doing in UPCGUnionData::CreatePointData
 	const UPCGPointData* SourcePointData = Source->ToPointData();
 
-	if (!Difference || !SourcePointData)
+	if (!SourcePointData)
 	{
+		UE_LOG(LogPCG, Error, TEXT("Difference unable to get source points"));
+		return SourcePointData;
+	}
+
+	if (!Difference)
+	{
+		UE_LOG(LogPCG, Verbose, TEXT("Difference is trivial"));
 		return SourcePointData;
 	}
 
@@ -163,5 +170,7 @@ const UPCGPointData* UPCGDifferenceData::CreatePointData() const
 #endif
 	}
 
+	UE_LOG(LogPCG, Verbose, TEXT("Difference generated %d points from %d source points"), TargetPoints.Num(), SourcePointData->GetPoints().Num());
+
 	return Data;
 }
\ No newline at end of file
diff --git a/Engine/Plugins/Experimental/PCG/Source/PCG/Private/Data/PCGIntersectionData.cpp b/Engine/Plugins/Experimental/PCG/Source/PCG/Private/Data/PCGIntersectionData.cpp
--- a/Engine/Plugins/Experimental/PCG/Source/PCG/Private/Data/PCGIntersectionData.cpp
+++ b/Engine/Plugins/Experimental/PCG/Source/PCG/Private/Data/PCGIntersectionData.cpp
@@ -123,6 +123,7 @@ UPCGPointData* UPCGIntersectionData::CreateAndFilterPointData(const UPCGSpatialD
 
 	if (!SourcePointData)
 	{
+		UE_LOG(LogPCG, Error, TEXT("Intersection unable to get source points"));
 		return nullptr;
 	}
 
@@ -148,5 +149,7 @@ UPCGPointData* UPCGIntersectionData::CreateAndFilterPointData(const UPCGSpatialD
 		}
 	}
 
+	UE_LOG(LogPCG, Verbose, TEXT("Intersection generated %d points from %d source points"), TargetPoints.Num(), SourcePoints.Num());
+
 	return Data;
 }
\ No newline at end of file
diff --git a/Engine/Plugins/Experimental/PCG/Source/PCG/Private/Data/PCGLandscapeData.cpp b/Engine/Plugins/Experimental/PCG/Source/PCG/Private/Data/PCGLandscapeData.cpp
--- a/Engine/Plugins/Experimental/PCG/Source/PCG/Private/Data/PCGLandscapeData.cpp
+++ b/Engine/Plugins/Experimental/PCG/Source/PCG/Private/Data/PCGLandscapeData.cpp
@@ -116,5 +116,7 @@ const UPCGPointData* UPCGLandscapeData::CreatePointData() const
 		}
 	}
 
+	UE_LOG(LogPCG, Verbose, TEXT("Landscape %s extracted %d of %d potential points"), *Landscape->GetFName().ToString(), Points.Num(), (MaxX - MinX) * (MaxY - MinY));
+
 	return Data;
 }
\ No newline at end of file
diff --git a/Engine/Plugins/Experimental/PCG/Source/PCG/Private/Data/PCGLandscapeSplineData.cpp b/Engine/Plugins/Experimental/PCG/Source/PCG/Private/Data/PCGLandscapeSplineData.cpp
--- a/Engine/Plugins/Experimental/PCG/Source/PCG/Private/Data/PCGLandscapeSplineData.cpp
+++ b/Engine/Plugins/Experimental/PCG/Source/PCG/Private/Data/PCGLandscapeSplineData.cpp
@@ -187,6 +187,8 @@ const UPCGPointData* UPCGLandscapeSplineData::CreatePointData() const
 		}
 	}
 
+	UE_LOG(LogPCG, Verbose, TEXT("Landscape spline %s generated %d points on %d segments"), *Spline->GetFName().ToString(), Points.Num(), Segments.Num());
+
 	return Data;
 }
 
diff --git a/Engine/Plugins/Experimental/PCG/Source/PCG/Private/Data/PCGProjectionData.cpp b/Engine/Plugins/Experimental/PCG/Source/PCG/Private/Data/PCGProjectionData.cpp
--- a/Engine/Plugins/Experimental/PCG/Source/PCG/Private/Data/PCGProjectionData.cpp
+++ b/Engine/Plugins/Experimental/PCG/Source/PCG/Private/Data/PCGProjectionData.cpp
@@ -112,5 +112,7 @@ const UPCGPointData* UPCGProjectionData::CreatePointData() const
 		}
 	}
 
+	UE_LOG(LogPCG, Verbose, TEXT("Projection generated %d points from %d source points"), Points.Num(), SourcePoints.Num());
+
 	return PointData;
 }
\ No newline at end of file
diff --git a/Engine/Plugins/Experimental/PCG/Source/PCG/Private/Data/PCGRenderTargetData.cpp b/Engine/Plugins/Experimental/PCG/Source/PCG/Private/Data/PCGRenderTargetData.cpp
--- a/Engine/Plugins/Experimental/PCG/Source/PCG/Private/Data/PCGRenderTargetData.cpp
+++ b/Engine/Plugins/Experimental/PCG/Source/PCG/Private/Data/PCGRenderTargetData.cpp
@@ -26,6 +26,11 @@ void UPCGRenderTargetData::Initialize(UTextureRenderTarget2D* InRenderTarget, co
 		Width = RenderTarget->SizeX;
 		Height = RenderTarget->SizeY;
 	}
+	else
+	{
+		Width = 0;
+		Height = 0;
+	}
 
 	Bounds = FBox(EForceInit::ForceInit);
 	Bounds += FVector(-1.0f, -1.0f, 0.0f);
diff --git a/Engine/Plugins/Experimental/PCG/Source/PCG/Private/Data/PCGSplineData.cpp b/Engine/Plugins/Experimental/PCG/Source/PCG/Private/Data/PCGSplineData.cpp
--- a/Engine/Plugins/Experimental/PCG/Source/PCG/Private/Data/PCGSplineData.cpp
+++ b/Engine/Plugins/Experimental/PCG/Source/PCG/Private/Data/PCGSplineData.cpp
@@ -71,6 +71,8 @@ const UPCGPointData* UPCGSplineData::CreatePointData() const
 		Distance += 1.0f;
 	}
 
+	UE_LOG(LogPCG, Verbose, TEXT("Spline %s generated %d points"), *Spline->GetFName().ToString(), Points.Num());
+
 	return Data;
 }
 
diff --git a/Engine/Plugins/Experimental/PCG/Source/PCG/Private/Data/PCGTextureData.cpp b/Engine/Plugins/Experimental/PCG/Source/PCG/Private/Data/PCGTextureData.cpp
--- a/Engine/Plugins/Experimental/PCG/Source/PCG/Private/Data/PCGTextureData.cpp
+++ b/Engine/Plugins/Experimental/PCG/Source/PCG/Private/Data/PCGTextureData.cpp
@@ -130,6 +130,8 @@ void UPCGTextureData::Initialize(UTexture2D* InTexture, const FTransform& InTran
 {
 	Texture = InTexture;
 	Transform = InTransform;
+	Width = 0;
+	Height = 0;
 
 	if (Texture)
 	{
diff --git a/Engine/Plugins/Experimental/PCG/Source/PCG/Private/Data/PCGUnionData.cpp b/Engine/Plugins/Experimental/PCG/Source/PCG/Private/Data/PCGUnionData.cpp
--- a/Engine/Plugins/Experimental/PCG/Source/PCG/Private/Data/PCGUnionData.cpp
+++ b/Engine/Plugins/Experimental/PCG/Source/PCG/Private/Data/PCGUnionData.cpp
@@ -155,10 +155,12 @@ const UPCGPointData* UPCGUnionData::CreatePointData() const
 	// Trivial results
 	if (Data.Num() == 0)
 	{
+		UE_LOG(LogPCG, Error, TEXT("Invalid union"));
 		return nullptr;
 	}
 	else if (Data.Num() == 1)
 	{
+		UE_LOG(LogPCG, Verbose, TEXT("Union is trivial"));
 		return Data[0]->ToPointData();
 	}
 
@@ -187,6 +189,8 @@ const UPCGPointData* UPCGUnionData::CreatePointData() const
 		break;
 	}
 
+	UE_LOG(LogPCG, Verbose, TEXT("Union generated %d points out of %d data sources"), PointData->GetPoints().Num(), Data.Num());
+
 	return PointData;
 }
 
diff --git a/Engine/Plugins/Experimental/PCG/Source/PCG/Private/Data/PCGVolumeData.cpp b/Engine/Plugins/Experimental/PCG/Source/PCG/Private/Data/PCGVolumeData.cpp
--- a/Engine/Plugins/Experimental/PCG/Source/PCG/Private/Data/PCGVolumeData.cpp
+++ b/Engine/Plugins/Experimental/PCG/Source/PCG/Private/Data/PCGVolumeData.cpp
@@ -37,7 +37,7 @@ FBox UPCGVolumeData::GetStrictBounds() const
 
 const UPCGPointData* UPCGVolumeData::CreatePointData() const
 {
-	// TODO
+	UE_LOG(LogPCG, Error, TEXT("Volume data has no default point sampling"));
 	return nullptr;
 }
 
diff --git a/Engine/Plugins/Experimental/PCG/Source/PCG/Private/Elements/PCGDensityFilter.cpp b/Engine/Plugins/Experimental/PCG/Source/PCG/Private/Elements/PCGDensityFilter.cpp
--- a/Engine/Plugins/Experimental/PCG/Source/PCG/Private/Elements/PCGDensityFilter.cpp
+++ b/Engine/Plugins/Experimental/PCG/Source/PCG/Private/Elements/PCGDensityFilter.cpp
@@ -34,6 +34,7 @@ bool FPCGDensityFilterElement::ExecuteInternal(FPCGContextPtr Context) const
 	if (bNoResults)
 #endif
 	{
+		PCGE_LOG(Verbose, "Skipped - all inputs rejected");
 		return true;
 	}
 
@@ -46,13 +47,14 @@ bool FPCGDensityFilterElement::ExecuteInternal(FPCGContextPtr Context) const
 
 		if (!Input.Data || Cast<UPCGSpatialData>(Input.Data) == nullptr)
 		{
-			// TODO: log
+			PCGE_LOG(Error, "Invalid input data");
 			continue;
 		}
 
 		// Skip processing if the transformation is trivial
 		if (bTrivialFilter)
 		{
+			PCGE_LOG(Verbose, "Skipped - trivial filter");
 			continue;
 		}
 
@@ -60,7 +62,7 @@ bool FPCGDensityFilterElement::ExecuteInternal(FPCGContextPtr Context) const
 
 		if (!OriginalData)
 		{
-			// Log error
+			PCGE_LOG(Error, "Unable to get point data from input");
 			continue;
 		}
 
@@ -87,6 +89,8 @@ bool FPCGDensityFilterElement::ExecuteInternal(FPCGContextPtr Context) const
 			}
 #endif
 		}
+
+		PCGE_LOG(Verbose, "Generated %d points out of %d source points", FilteredPoints.Num(), Points.Num());
 	}
 
 	return true;
diff --git a/Engine/Plugins/Experimental/PCG/Source/PCG/Private/Elements/PCGDensityRemap.cpp b/Engine/Plugins/Experimental/PCG/Source/PCG/Private/Elements/PCGDensityRemap.cpp
--- a/Engine/Plugins/Experimental/PCG/Source/PCG/Private/Elements/PCGDensityRemap.cpp
+++ b/Engine/Plugins/Experimental/PCG/Source/PCG/Private/Elements/PCGDensityRemap.cpp
@@ -37,13 +37,14 @@ bool FPCGLinearDensityRemapElement::ExecuteInternal(FPCGContextPtr Context) cons
 
 		if (!Input.Data || Cast<UPCGSpatialData>(Input.Data) == nullptr)
 		{
-			// TODO: log error
+			PCGE_LOG(Error, "Invalid input data");
 			continue;
 		}
 
 		// Skip processing if the remapping is trivial
 		if (bTrivialRemapping)
 		{
+			PCGE_LOG(Verbose, "Skipped - trivial remapping");
 			continue;
 		}
 
@@ -51,7 +52,7 @@ bool FPCGLinearDensityRemapElement::ExecuteInternal(FPCGContextPtr Context) cons
 
 		if (!OriginalData)
 		{
-			// TODO: log error
+			PCGE_LOG(Error, "Unable to get points from input");
 			continue;
 		}
 
diff --git a/Engine/Plugins/Experimental/PCG/Source/PCG/Private/Elements/PCGPointSampler.cpp b/Engine/Plugins/Experimental/PCG/Source/PCG/Private/Elements/PCGPointSampler.cpp
--- a/Engine/Plugins/Experimental/PCG/Source/PCG/Private/Elements/PCGPointSampler.cpp
+++ b/Engine/Plugins/Experimental/PCG/Source/PCG/Private/Elements/PCGPointSampler.cpp
@@ -34,6 +34,7 @@ bool FPCGPointSamplerElement::ExecuteInternal(FPCGContextPtr Context) const
 	if (bNoSampling)
 #endif
 	{
+		PCGE_LOG(Verbose, "Skipped - all inputs rejected");
 		return true;
 	}
 
@@ -44,13 +45,14 @@ bool FPCGPointSamplerElement::ExecuteInternal(FPCGContextPtr Context) const
 
 		if (!Input.Data || Cast<UPCGSpatialData>(Input.Data) == nullptr)
 		{
-			// TODO: add a log
+			PCGE_LOG(Error, "Invalid input data");
 			continue;
 		}
 
 		// Skip processing if the transformation would be trivial
 		if (bTrivialSampling)
 		{
+			PCGE_LOG(Verbose, "Skipped - trivial sampling");
 			continue;
 		}
 
@@ -58,7 +60,7 @@ bool FPCGPointSamplerElement::ExecuteInternal(FPCGContextPtr Context) const
 
 		if (!OriginalData)
 		{
-			// Log error
+			PCGE_LOG(Error, "Unable to get point data from input");
 			continue;
 		}
 
@@ -81,7 +83,7 @@ bool FPCGPointSamplerElement::ExecuteInternal(FPCGContextPtr Context) const
 		// Early out
 		if (TargetNumPoints == 0)
 		{
-			// TODO log
+			PCGE_LOG(Verbose, "Skipped - all points rejected");
 			continue;
 		}
 		else
@@ -111,6 +113,8 @@ bool FPCGPointSamplerElement::ExecuteInternal(FPCGContextPtr Context) const
 				}
 #endif
 			}
+
+			PCGE_LOG(Verbose, "Generated %d points from %d source points", SampledPoints.Num(), OriginalPointCount);
 		}
 	}
 
diff --git a/Engine/Plugins/Experimental/PCG/Source/PCG/Private/Elements/PCGSelfPruning.cpp b/Engine/Plugins/Experimental/PCG/Source/PCG/Private/Elements/PCGSelfPruning.cpp
--- a/Engine/Plugins/Experimental/PCG/Source/PCG/Private/Elements/PCGSelfPruning.cpp
+++ b/Engine/Plugins/Experimental/PCG/Source/PCG/Private/Elements/PCGSelfPruning.cpp
@@ -53,6 +53,7 @@ bool FPCGSelfPruningElement::ExecuteInternal(FPCGContextPtr Context) const
 	if (Settings->PruningType == EPCGSelfPruningType::None)
 	{
 		Context->OutputData = Context->InputData;
+		PCGE_LOG(Verbose, "Skipped - Type is none");
 		return true;
 	}
 
@@ -67,9 +68,9 @@ bool FPCGSelfPruningElement::ExecuteInternal(FPCGContextPtr Context) const
 	{
 		const UPCGSpatialData* SpatialInput = Cast<UPCGSpatialData>(Input.Data);
 
-		if (!SpatialInput || !SpatialInput->TargetActor)
+		if (!SpatialInput)
 		{
-			// Data type mismatch / no output actor
+			PCGE_LOG(Error, "Invalid input data");
 			continue;
 		}
 
@@ -164,6 +165,8 @@ bool FPCGSelfPruningElement::ExecuteInternal(FPCGContextPtr Context) const
 		{
 			OutputPoints.Add(*Point);
 		}
+
+		PCGE_LOG(Verbose, "Generated %d points from %d source points", OutputPoints.Num(), Points.Num());
 	}
 
 	// Finally, forward any exclusions/settings
diff --git a/Engine/Plugins/Experimental/PCG/Source/PCG/Private/Elements/PCGSpawnActor.cpp b/Engine/Plugins/Experimental/PCG/Source/PCG/Private/Elements/PCGSpawnActor.cpp
--- a/Engine/Plugins/Experimental/PCG/Source/PCG/Private/Elements/PCGSpawnActor.cpp
+++ b/Engine/Plugins/Experimental/PCG/Source/PCG/Private/Elements/PCGSpawnActor.cpp
@@ -72,8 +72,9 @@ bool FPCGSpawnActorElement::ExecuteInternal(FPCGContextPtr Context) const
 	check(Settings);
 
 	// Early out
-	if(!Settings->TemplateActorClass)
+	if(!Settings->TemplateActorClass || Settings->TemplateActorClass->HasAnyClassFlags(CLASS_Abstract))
 	{
+		PCGE_LOG(Error, "Invalid template actor class (%s)", Settings->TemplateActorClass ? *Settings->TemplateActorClass->GetFName().ToString() : TEXT("None"));
 		return true;
 	}
 
@@ -89,7 +90,7 @@ bool FPCGSpawnActorElement::ExecuteInternal(FPCGContextPtr Context) const
 
 		if (!SpatialData)
 		{
-			// Data type mismatch
+			PCGE_LOG(Error, "Invalid input data");
 			continue;
 		}
 
@@ -97,7 +98,7 @@ bool FPCGSpawnActorElement::ExecuteInternal(FPCGContextPtr Context) const
 
 		if (!TargetActor)
 		{
-			// No target actor
+			PCGE_LOG(Error, "Invalid target actor");
 			continue;
 		}
 
@@ -106,6 +107,7 @@ bool FPCGSpawnActorElement::ExecuteInternal(FPCGContextPtr Context) const
 
 		if (!PointData)
 		{
+			PCGE_LOG(Error, "Unable to get point data from input");
 			continue;
 		}
 
@@ -113,6 +115,7 @@ bool FPCGSpawnActorElement::ExecuteInternal(FPCGContextPtr Context) const
 
 		if (Points.IsEmpty())
 		{
+			PCGE_LOG(Verbose, "Skipped - no points");
 			continue;
 		}
 
@@ -144,6 +147,10 @@ bool FPCGSpawnActorElement::ExecuteInternal(FPCGContextPtr Context) const
 				{
 					ISMC = UPCGActorHelpers::GetOrCreateISMC(TargetActor, Context->SourceComponent, Mesh);
 				}
+				else
+				{
+					PCGE_LOG(Error, "No supported mesh found");
+				}
 			}
 
 			{
@@ -185,6 +192,8 @@ bool FPCGSpawnActorElement::ExecuteInternal(FPCGContextPtr Context) const
 							}
 						}
 					}
+
+					PCGE_LOG(Verbose, "Generated %d actors", Points.Num());
 				}
 			}
 
@@ -194,6 +203,8 @@ bool FPCGSpawnActorElement::ExecuteInternal(FPCGContextPtr Context) const
 				ISMC->NumCustomDataFloats = 0;
 				ISMC->AddInstances(Instances, false, true);
 				ISMC->UpdateBounds();
+
+				PCGE_LOG(Verbose, "Added %d ISM instances", Instances.Num());
 			}
 		}
 
diff --git a/Engine/Plugins/Experimental/PCG/Source/PCG/Private/Elements/PCGStaticMeshSpawner.cpp b/Engine/Plugins/Experimental/PCG/Source/PCG/Private/Elements/PCGStaticMeshSpawner.cpp
--- a/Engine/Plugins/Experimental/PCG/Source/PCG/Private/Elements/PCGStaticMeshSpawner.cpp
+++ b/Engine/Plugins/Experimental/PCG/Source/PCG/Private/Elements/PCGStaticMeshSpawner.cpp
@@ -58,7 +58,7 @@ bool FPCGStaticMeshSpawnerElement::ExecuteInternal(FPCGContextPtr Context) const
 
 		if (!SpatialData)
 		{
-			// Data type mismatch
+			PCGE_LOG(Error, "Invalid input data");
 			continue;
 		}
 
@@ -66,7 +66,7 @@ bool FPCGStaticMeshSpawnerElement::ExecuteInternal(FPCGContextPtr Context) const
 
 		if (!TargetActor)
 		{
-			// No target actor
+			PCGE_LOG(Error, "Invalid target actor");
 			continue;
 		}
 
@@ -75,6 +75,7 @@ bool FPCGStaticMeshSpawnerElement::ExecuteInternal(FPCGContextPtr Context) const
 
 		if (!PointData)
 		{
+			PCGE_LOG(Error, "Unable to get point data from input");
 			continue;
 		}
 
@@ -125,6 +126,8 @@ bool FPCGStaticMeshSpawnerElement::ExecuteInternal(FPCGContextPtr Context) const
 				ISMC->NumCustomDataFloats = 0;
 				ISMC->AddInstances(Instances.Instances, false, true);
 				ISMC->UpdateBounds();
+
+				PCGE_LOG(Verbose, "Added %d instances of %s on actor %s", Instances.Instances.Num(), *Instances.Mesh->GetFName().ToString(), *TargetActor->GetFName().ToString());
 			}
 		}
 	}
diff --git a/Engine/Plugins/Experimental/PCG/Source/PCG/Private/Elements/PCGSurfaceSampler.cpp b/Engine/Plugins/Experimental/PCG/Source/PCG/Private/Elements/PCGSurfaceSampler.cpp
--- a/Engine/Plugins/Experimental/PCG/Source/PCG/Private/Elements/PCGSurfaceSampler.cpp
+++ b/Engine/Plugins/Experimental/PCG/Source/PCG/Private/Elements/PCGSurfaceSampler.cpp
@@ -22,6 +22,7 @@ bool FPCGSurfaceSamplerElement::ExecuteInternal(FPCGContextPtr Context) const
 	// TODO: we could compute an approximate radius based on the points per squared meters if that's useful
 	if (Settings->PointRadius <= 0)
 	{
+		PCGE_LOG(Warning, "Skipped - Invalid point radius");
 		return true;
 	}
 
@@ -35,7 +36,7 @@ bool FPCGSurfaceSamplerElement::ExecuteInternal(FPCGContextPtr Context) const
 
 		if (!SpatialInput || !SpatialInput->TargetActor)
 		{
-			// Data type mismatch / no output actor
+			PCGE_LOG(Error, "Invalid input data");
 			continue;
 		}
 
@@ -43,7 +44,7 @@ bool FPCGSurfaceSamplerElement::ExecuteInternal(FPCGContextPtr Context) const
 
 		if (!InputBounds.IsValid)
 		{
-			// With invalid bounds, we can't do anything
+			PCGE_LOG(Warning, "Input data has invalid bounds");
 			continue;
 		}
 
@@ -65,6 +66,7 @@ bool FPCGSurfaceSamplerElement::ExecuteInternal(FPCGContextPtr Context) const
 
 		if (CellMinX > CellMaxX || CellMinY > CellMaxY)
 		{
+			PCGE_LOG(Verbose, "Skipped - invalid cell bounds");
 			continue;
 		}
 
@@ -76,6 +78,7 @@ bool FPCGSurfaceSamplerElement::ExecuteInternal(FPCGContextPtr Context) const
 
 		if (TargetPointCount == 0)
 		{
+			PCGE_LOG(Verbose, "Skipped - density yields no points");
 			continue;
 		}
 		else if (TargetPointCount > CellCount)
@@ -138,6 +141,8 @@ bool FPCGSurfaceSamplerElement::ExecuteInternal(FPCGContextPtr Context) const
 
 			CurrentX += CellSize;
 		}
+
+		PCGE_LOG(Verbose, "Generated %d points in %d cells", SampledPoints.Num(), CellCount);
 	}
 
 	// Finally, forward any exclusions/settings
diff --git a/Engine/Plugins/Experimental/PCG/Source/PCG/Private/PCGContext.cpp b/Engine/Plugins/Experimental/PCG/Source/PCG/Private/PCGContext.cpp
new file mode 100644
--- /dev/null
+++ b/Engine/Plugins/Experimental/PCG/Source/PCG/Private/PCGContext.cpp
@@ -0,0 +1,15 @@
+// Copyright Epic Games, Inc. All Rights Reserved.
+
+#include "PCGContext.h"
+#include "PCGComponent.h"
+#include "GameFramework/Actor.h"
+
+FName FPCGContext::GetTaskName() const
+{
+	return Node ? Node->GetFName() : TEXT("Anonymous task");
+}
+
+FName FPCGContext::GetComponentName() const
+{
+	return SourceComponent ? SourceComponent->GetOwner()->GetFName() : TEXT("Non-PCG Component");
+}
\ No newline at end of file
diff --git a/Engine/Plugins/Experimental/PCG/Source/PCG/Private/PCGElement.cpp b/Engine/Plugins/Experimental/PCG/Source/PCG/Private/PCGElement.cpp
--- a/Engine/Plugins/Experimental/PCG/Source/PCG/Private/PCGElement.cpp
+++ b/Engine/Plugins/Experimental/PCG/Source/PCG/Private/PCGElement.cpp
@@ -2,6 +2,7 @@
 
 #include "PCGElement.h"
 #include "PCGContext.h"
+#include "PCGModule.h"
 #include "PCGSettings.h"
 #include "Elements/PCGDebugElement.h"
 #include "Graph/PCGGraphCache.h"
@@ -60,17 +61,41 @@ bool IPCGElement::Execute(FPCGContextPtr Context) const
 
 		bool bDone = false;
 
+#if WITH_EDITOR
+		UE_LOG(LogPCG, Verbose, TEXT("---------------------------------------"));
+#endif
+
 		if (IsCacheable(Settings) && Context->Cache && Context->Cache->GetFromCache(this, Context->InputData, Settings, Context->OutputData))
 		{
+			PCGE_LOG(Verbose, "Used cached results");
 			bDone = true;
 		}
 		else
 		{
+#if WITH_EDITOR
+			const double StartTime = FPlatformTime::Seconds();
+#endif
+
 			bDone = ExecuteInternal(Context);
 
-			if (bDone && IsCacheable(Settings) && Context->Cache)
+#if WITH_EDITOR
+			const double EndTime = FPlatformTime::Seconds();
+			Context->ElapsedTime += (EndTime - StartTime);
+			Context->ExecutionCount++;
+#endif
+
+			if (bDone)
 			{
-				Context->Cache->StoreInCache(this, Context->InputData, Settings, Context->OutputData);
+				if (IsCacheable(Settings) && Context->Cache)
+				{
+					Context->Cache->StoreInCache(this, Context->InputData, Settings, Context->OutputData);
+				}
+
+#if WITH_EDITOR
+				PCGE_LOG(Log, "Executed in (%f)s and (%d) call(s)", Context->ElapsedTime, Context->ExecutionCount);
+#else
+				PCGE_LOG(Log, "Executed");
+#endif
 			}
 		}
 
diff --git a/Engine/Plugins/Experimental/PCG/Source/PCG/Public/PCGContext.h b/Engine/Plugins/Experimental/PCG/Source/PCG/Public/PCGContext.h
--- a/Engine/Plugins/Experimental/PCG/Source/PCG/Public/PCGContext.h
+++ b/Engine/Plugins/Experimental/PCG/Source/PCG/Public/PCGContext.h
@@ -27,6 +27,11 @@ struct FPCGContext
 	FPCGTaskId TaskId = InvalidTaskId;
 	bool bIsPaused = false;
 
+#if WITH_EDITOR
+	double ElapsedTime = 0.0;
+	int32 ExecutionCount = 0;
+#endif
+
 	template<typename SettingsType>
 	const SettingsType* GetInputSettings()
 	{
@@ -39,4 +44,7 @@ struct FPCGContext
 			return InputData.GetSettings<SettingsType>();
 		}
 	}
+
+	FName GetTaskName() const;
+	FName GetComponentName() const;
 };
diff --git a/Engine/Plugins/Experimental/PCG/Source/PCG/Public/PCGData.h b/Engine/Plugins/Experimental/PCG/Source/PCG/Public/PCGData.h
--- a/Engine/Plugins/Experimental/PCG/Source/PCG/Public/PCGData.h
+++ b/Engine/Plugins/Experimental/PCG/Source/PCG/Public/PCGData.h
@@ -5,6 +5,8 @@
 #include "CoreMinimal.h"
 #include "Kismet/BlueprintFunctionLibrary.h"
 
+#include "PCGModule.h"
+
 #include "PCGData.generated.h"
 
 class UPCGSettings;
diff --git a/Engine/Plugins/Experimental/PCG/Source/PCG/Public/PCGElement.h b/Engine/Plugins/Experimental/PCG/Source/PCG/Public/PCGElement.h
--- a/Engine/Plugins/Experimental/PCG/Source/PCG/Public/PCGElement.h
+++ b/Engine/Plugins/Experimental/PCG/Source/PCG/Public/PCGElement.h
@@ -12,6 +12,16 @@ class UPCGSettings;
 
 typedef TSharedPtr<IPCGElement, ESPMode::ThreadSafe> FPCGElementPtr;
 
+#define PCGE_LOG_C(Verbosity, CustomContext, Format, ...) \
+	UE_LOG(LogPCG, \
+		Verbosity, \
+		TEXT("[%s - %s]: " Format), \
+		*CustomContext->GetComponentName().ToString(), \
+		*CustomContext->GetTaskName().ToString(), \
+		##__VA_ARGS__)
+
+#define PCGE_LOG(Verbosity, Format, ...) PCGE_LOG_C(Verbosity, Context, Format, ##__VA_ARGS__)
+
 /**
 * Base class for the processing bit of a PCG node/settings
 */
