Fixed an issue where duplicated actors wouldn't be picked up by the tracked actors
Fixed an issue where data dependencies in blueprints wouldn't trigger a cache dirty on change.

#rnx

diff --git a/Engine/Plugins/Experimental/PCG/Source/PCG/Private/Elements/PCGExecuteBlueprint.cpp b/Engine/Plugins/Experimental/PCG/Source/PCG/Private/Elements/PCGExecuteBlueprint.cpp
--- a/Engine/Plugins/Experimental/PCG/Source/PCG/Private/Elements/PCGExecuteBlueprint.cpp
+++ b/Engine/Plugins/Experimental/PCG/Source/PCG/Private/Elements/PCGExecuteBlueprint.cpp
@@ -185,6 +185,9 @@ void UPCGBlueprintSettings::OnDependencyChanged(UObject* Object, FPropertyChange
 		return;
 	}
 
+	// When a data dependency is changed, this means we have to dirty the cache, otherwise it will not register as a change.
+	DirtyCache();
+
 	BlueprintElementInstance->OnBlueprintChangedDelegate.Broadcast(BlueprintElementInstance);
 }
 #endif
diff --git a/Engine/Plugins/Experimental/PCG/Source/PCG/Private/PCGComponent.cpp b/Engine/Plugins/Experimental/PCG/Source/PCG/Private/PCGComponent.cpp
--- a/Engine/Plugins/Experimental/PCG/Source/PCG/Private/PCGComponent.cpp
+++ b/Engine/Plugins/Experimental/PCG/Source/PCG/Private/PCGComponent.cpp
@@ -1383,6 +1383,10 @@ bool UPCGComponent::DirtyTrackedActor(AActor* InActor)
 
 		bAppliedChange = true;
 	}
+	else if (AddTrackedActor(InActor))
+	{
+		bAppliedChange = true;
+	}
 
 	return bAppliedChange;
 }
diff --git a/Engine/Plugins/Experimental/PCG/Source/PCG/Private/PCGSettings.cpp b/Engine/Plugins/Experimental/PCG/Source/PCG/Private/PCGSettings.cpp
--- a/Engine/Plugins/Experimental/PCG/Source/PCG/Private/PCGSettings.cpp
+++ b/Engine/Plugins/Experimental/PCG/Source/PCG/Private/PCGSettings.cpp
@@ -4,6 +4,10 @@
 #include "PCGNode.h"
 #include "Serialization/ArchiveObjectCrc32.h"
 
+#if WITH_EDITOR
+#include "Editor.h"
+#endif
+
 /** In order to reuse the cache when only debug settings change, we must make sure to ignore these from the CRC check */
 #if WITH_EDITORONLY_DATA
 class FPCGSettingsObjectCrc32 : public FArchiveObjectCrc32
@@ -55,6 +59,24 @@ UPCGNode* UPCGSettings::CreateNode() const
 	return NewObject<UPCGNode>();
 }
 
+#if WITH_EDITOR
+void UPCGSettings::PostEditChangeProperty(FPropertyChangedEvent& PropertyChangedEvent)
+{
+	Super::PostEditChangeProperty(PropertyChangedEvent);
+	OnSettingsChangedDelegate.Broadcast(this);
+}
+
+void UPCGSettings::DirtyCache()
+{
+	UWorld* World = GEditor->GetEditorWorldContext().World();
+	UPCGSubsystem* PCGSubsystem = World ? World->GetSubsystem<UPCGSubsystem>() : nullptr;
+	if (PCGSubsystem)
+	{
+		PCGSubsystem->CleanFromCache(GetElement().Get());
+	}
+}
+#endif // WITH_EDITOR
+
 FPCGElementPtr UPCGTrivialSettings::CreateElement() const
 {
 	return MakeShared<FPCGTrivialElement>();
@@ -65,14 +87,4 @@ bool FPCGTrivialElement::ExecuteInternal(FPCGContextPtr Context) const
 	// Pass-through
 	Context->OutputData = Context->InputData;
 	return true;
-}
-
-#if WITH_EDITOR
-
-void UPCGSettings::PostEditChangeProperty(FPropertyChangedEvent& PropertyChangedEvent)
-{
-	Super::PostEditChangeProperty(PropertyChangedEvent);
-	OnSettingsChangedDelegate.Broadcast(this);
-}
-
-#endif // WITH_EDITOR
\ No newline at end of file
+}
\ No newline at end of file
diff --git a/Engine/Plugins/Experimental/PCG/Source/PCG/Public/PCGSettings.h b/Engine/Plugins/Experimental/PCG/Source/PCG/Public/PCGSettings.h
--- a/Engine/Plugins/Experimental/PCG/Source/PCG/Public/PCGSettings.h
+++ b/Engine/Plugins/Experimental/PCG/Source/PCG/Public/PCGSettings.h
@@ -79,7 +79,10 @@ protected:
 
 #if WITH_EDITOR
 	virtual void PostEditChangeProperty(FPropertyChangedEvent& PropertyChangedEvent) override;
-#endif
+
+	/** Method that can be called to dirty the cache data from this settings objects if the operator== does not allow to detect changes */
+	void DirtyCache();
+#endif 
 
 private:
 	mutable FPCGElementPtr CachedElement;
