Added a few profiling tags in PCG to have complete coverage

#rnx

diff --git a/Engine/Plugins/Experimental/PCG/Source/PCG/Private/Data/PCGRenderTargetData.cpp b/Engine/Plugins/Experimental/PCG/Source/PCG/Private/Data/PCGRenderTargetData.cpp
--- a/Engine/Plugins/Experimental/PCG/Source/PCG/Private/Data/PCGRenderTargetData.cpp
+++ b/Engine/Plugins/Experimental/PCG/Source/PCG/Private/Data/PCGRenderTargetData.cpp
@@ -14,6 +14,7 @@ void UPCGRenderTargetData::Initialize(UTextureRenderTarget2D* InRenderTarget, co
 	{
 		if (FTextureRenderTargetResource* RTResource = RenderTarget->GameThread_GetRenderTargetResource())
 		{
+			TRACE_CPUPROFILER_EVENT_SCOPE(UPCGRenderTargetData::Initialize::ReadData);
 			ETextureRenderTargetFormat Format = RenderTarget->RenderTargetFormat;
 			if ((Format == RTF_RGBA16f) || (Format == RTF_RGBA32f) || (Format == RTF_RGBA8) || (Format == RTF_R8) || (Format == RTF_RGB10A2))
 			{
diff --git a/Engine/Plugins/Experimental/PCG/Source/PCG/Private/Data/PCGTextureData.cpp b/Engine/Plugins/Experimental/PCG/Source/PCG/Private/Data/PCGTextureData.cpp
--- a/Engine/Plugins/Experimental/PCG/Source/PCG/Private/Data/PCGTextureData.cpp
+++ b/Engine/Plugins/Experimental/PCG/Source/PCG/Private/Data/PCGTextureData.cpp
@@ -138,6 +138,7 @@ void UPCGTextureData::Initialize(UTexture2D* InTexture, const FTransform& InTran
 #if WITH_EDITORONLY_DATA
 		if (Texture->GetPlatformData()->Mips.Num() > 0)
 		{
+			TRACE_CPUPROFILER_EVENT_SCOPE(UPCGTextureData::Initialize::ReadData);
 			Width = Texture->GetSizeX();
 			Height = Texture->GetSizeY();
 
diff --git a/Engine/Plugins/Experimental/PCG/Source/PCG/Private/Elements/PCGExecuteBlueprint.cpp b/Engine/Plugins/Experimental/PCG/Source/PCG/Private/Elements/PCGExecuteBlueprint.cpp
--- a/Engine/Plugins/Experimental/PCG/Source/PCG/Private/Elements/PCGExecuteBlueprint.cpp
+++ b/Engine/Plugins/Experimental/PCG/Source/PCG/Private/Elements/PCGExecuteBlueprint.cpp
@@ -311,6 +311,7 @@ FPCGElementPtr UPCGBlueprintSettings::CreateElement() const
 
 bool FPCGExecuteBlueprintElement::ExecuteInternal(FPCGContextPtr Context) const
 {
+	TRACE_CPUPROFILER_EVENT_SCOPE(FPCGExecuteBlueprintElement::Execute);
 	const UPCGBlueprintSettings* Settings = Context->GetInputSettings<UPCGBlueprintSettings>();
 
 	if (Settings && Settings->BlueprintElementInstance)
