Refactored dependency management in PCG blueprints to be cleaner;
Fixed an issue related to initializing a graph, where the appropriate callbacks woudn't be setup properly

#rnx

diff --git a/Engine/Plugins/Experimental/PCG/Source/PCG/Private/Elements/PCGExecuteBlueprint.cpp b/Engine/Plugins/Experimental/PCG/Source/PCG/Private/Elements/PCGExecuteBlueprint.cpp
--- a/Engine/Plugins/Experimental/PCG/Source/PCG/Private/Elements/PCGExecuteBlueprint.cpp
+++ b/Engine/Plugins/Experimental/PCG/Source/PCG/Private/Elements/PCGExecuteBlueprint.cpp
@@ -121,12 +121,55 @@ UWorld* UPCGBlueprintElement::GetWorld() const
 #endif
 }
 
+void UPCGBlueprintElement::PostLoad()
+{
+	Super::PostLoad();
+	Initialize();
+}
+
+void UPCGBlueprintElement::BeginDestroy()
+{
+#if WITH_EDITOR
+	FCoreUObjectDelegates::OnObjectPropertyChanged.RemoveAll(this);
+#endif
+
+	Super::BeginDestroy();
+}
+
+void UPCGBlueprintElement::Initialize()
+{
+#if WITH_EDITOR
+	FCoreUObjectDelegates::OnObjectPropertyChanged.AddUObject(this, &UPCGBlueprintElement::OnDependencyChanged);
+	DataDependencies = PCGBlueprintHelper::GetDataDependencies(this);
+#endif
+}
+
 #if WITH_EDITOR
 void UPCGBlueprintElement::PostEditChangeProperty(FPropertyChangedEvent& PropertyChangedEvent)
 {
 	Super::PostEditChangeProperty(PropertyChangedEvent);
+
+	// Since we don't really know what changed, let's just rebuild our data dependencies
+	DataDependencies = PCGBlueprintHelper::GetDataDependencies(this);
+
 	OnBlueprintChangedDelegate.Broadcast(this);
 }
+
+void UPCGBlueprintElement::OnDependencyChanged(UObject* Object, FPropertyChangedEvent& PropertyChangedEvent)
+{
+	if (PropertyChangedEvent.ChangeType == EPropertyChangeType::Interactive)
+	{
+		return;
+	}
+
+	if (!DataDependencies.Contains(Object))
+	{
+		return;
+	}
+
+	OnBlueprintChangedDelegate.Broadcast(this);
+}
+
 #endif // WITH_EDITOR
 
 void UPCGBlueprintSettings::SetupBlueprintEvent()
@@ -160,7 +203,6 @@ void UPCGBlueprintSettings::SetupBlueprintElementEvent()
 #if WITH_EDITOR
 	if (BlueprintElementInstance)
 	{
-		DataDependencies = PCGBlueprintHelper::GetDataDependencies(BlueprintElementInstance);
 		BlueprintElementInstance->OnBlueprintChangedDelegate.AddUObject(this, &UPCGBlueprintSettings::OnBlueprintElementChanged);
 	}
 #endif
@@ -172,34 +214,14 @@ void UPCGBlueprintSettings::TeardownBlueprintElementEvent()
 	if (BlueprintElementInstance)
 	{
 		BlueprintElementInstance->OnBlueprintChangedDelegate.RemoveAll(this);
-		DataDependencies.Reset();
 	}
 #endif
 }
 
-#if WITH_EDITOR
-void UPCGBlueprintSettings::OnDependencyChanged(UObject* Object, FPropertyChangedEvent& PropertyChangedEvent)
-{
-	if (!BlueprintElementInstance || PropertyChangedEvent.ChangeType == EPropertyChangeType::Interactive || !DataDependencies.Contains(Object))
-	{
-		return;
-	}
-
-	// When a data dependency is changed, this means we have to dirty the cache, otherwise it will not register as a change.
-	DirtyCache();
-
-	BlueprintElementInstance->OnBlueprintChangedDelegate.Broadcast(BlueprintElementInstance);
-}
-#endif
-
 void UPCGBlueprintSettings::PostLoad()
 {
 	Super::PostLoad();
 
-#if WITH_EDITOR
-	FCoreUObjectDelegates::OnObjectPropertyChanged.AddUObject(this, &UPCGBlueprintSettings::OnDependencyChanged);
-#endif
-
 	if (BlueprintElement_DEPRECATED && !BlueprintElementType)
 	{
 		BlueprintElementType = BlueprintElement_DEPRECATED;
@@ -223,10 +245,6 @@ void UPCGBlueprintSettings::BeginDestroy()
 	TeardownBlueprintElementEvent();
 	TeardownBlueprintEvent();
 
-#if WITH_EDITOR
-	FCoreUObjectDelegates::OnObjectPropertyChanged.RemoveAll(this);
-#endif
-
 	Super::BeginDestroy();
 }
 
@@ -254,10 +272,10 @@ void UPCGBlueprintSettings::OnBlueprintElementChanged(UPCGBlueprintElement* InEl
 {
 	if (InElement == BlueprintElementInstance)
 	{
-		OnSettingsChangedDelegate.Broadcast(this);
+		// When a data dependency is changed, this means we have to dirty the cache, otherwise it will not register as a change.
+		DirtyCache();
 
-		// Since we don't know what property has changed in the element, we need to rebuild our list of dependencies
-		DataDependencies = PCGBlueprintHelper::GetDataDependencies(BlueprintElementInstance);
+		OnSettingsChangedDelegate.Broadcast(this);
 	}
 }
 #endif
@@ -284,6 +302,7 @@ void UPCGBlueprintSettings::RefreshBlueprintElement()
 	if (BlueprintElementType)
 	{
 		BlueprintElementInstance = NewObject<UPCGBlueprintElement>(this, BlueprintElementType);
+		BlueprintElementInstance->Initialize();
 		SetupBlueprintElementEvent();
 	}
 	else
diff --git a/Engine/Plugins/Experimental/PCG/Source/PCG/Private/PCGGraph.cpp b/Engine/Plugins/Experimental/PCG/Source/PCG/Private/PCGGraph.cpp
--- a/Engine/Plugins/Experimental/PCG/Source/PCG/Private/PCGGraph.cpp
+++ b/Engine/Plugins/Experimental/PCG/Source/PCG/Private/PCGGraph.cpp
@@ -64,7 +64,7 @@ void UPCGGraph::InitializeFromTemplate()
 		InNode->InboundNodes.Reset();
 
 		// Reset settings as well
-		InNode->DefaultSettings = NewObject<UPCGTrivialSettings>(InNode);
+		InNode->SetDefaultSettings(NewObject<UPCGTrivialSettings>(InNode));
 	};
 
 	ResetDefaultNode(InputNode);
@@ -104,7 +104,7 @@ UPCGNode* UPCGGraph::AddNodeOfType(TSubclassOf<class UPCGSettings> InSettingsCla
 		Modify();
 
 		// Assign settings to node
-		Node->DefaultSettings = Settings;
+		Node->SetDefaultSettings(Settings);
 		Settings->Rename(nullptr, Node);
 		Settings->SetFlags(RF_Transactional);
 
@@ -141,7 +141,7 @@ UPCGNode* UPCGGraph::AddNode(UPCGSettings* InSettings)
 		Modify();
 
 		// Assign settings to node & reparent
-		Node->DefaultSettings = InSettings;
+		Node->SetDefaultSettings(InSettings);
 
 		// Reparent node to this graph
 		Node->Rename(nullptr, this);
diff --git a/Engine/Plugins/Experimental/PCG/Source/PCG/Private/PCGNode.cpp b/Engine/Plugins/Experimental/PCG/Source/PCG/Private/PCGNode.cpp
--- a/Engine/Plugins/Experimental/PCG/Source/PCG/Private/PCGNode.cpp
+++ b/Engine/Plugins/Experimental/PCG/Source/PCG/Private/PCGNode.cpp
@@ -52,6 +52,26 @@ UPCGNode* UPCGNode::AddEdgeTo(UPCGNode* To)
 	}
 }
 
+void UPCGNode::SetDefaultSettings(TObjectPtr<UPCGSettings> InSettings)
+{
+#if WITH_EDITOR
+	const bool bDifferentSettings = (DefaultSettings != InSettings);
+	if (bDifferentSettings && DefaultSettings)
+	{
+		DefaultSettings->OnSettingsChangedDelegate.RemoveAll(this);
+	}
+#endif
+
+	DefaultSettings = InSettings;
+
+#if WITH_EDITOR
+	if (bDifferentSettings && DefaultSettings)
+	{
+		DefaultSettings->OnSettingsChangedDelegate.AddUObject(this, &UPCGNode::OnSettingsChanged);
+	}
+#endif
+}
+
 void UPCGNode::ConnectTo(UPCGNode* InSuccessor)
 {
 	OutboundNodes.AddUnique(InSuccessor);
diff --git a/Engine/Plugins/Experimental/PCG/Source/PCG/Public/Elements/PCGExecuteBlueprint.h b/Engine/Plugins/Experimental/PCG/Source/PCG/Public/Elements/PCGExecuteBlueprint.h
--- a/Engine/Plugins/Experimental/PCG/Source/PCG/Public/Elements/PCGExecuteBlueprint.h
+++ b/Engine/Plugins/Experimental/PCG/Source/PCG/Public/Elements/PCGExecuteBlueprint.h
@@ -28,9 +28,17 @@ class UPCGBlueprintElement : public UObject
 	GENERATED_BODY()
 
 public:
+	// ~Begin UObject interface
+	virtual void PostLoad() override;
+	virtual void BeginDestroy() override;
+	// ~End UObject interface
+
 	UFUNCTION(BlueprintImplementableEvent, BlueprintCallable, Category = Execution)
 	void Execute(const FPCGDataCollection& Input, FPCGDataCollection& Output) const;
 
+	/** Called after object creation to setup the object callbacks */
+	void Initialize();
+
 #if WITH_EDITOR
 	// ~Begin UObject interface
 	virtual void PostEditChangeProperty(FPropertyChangedEvent& PropertyChangedEvent) override;
@@ -43,6 +51,12 @@ public:
 #if WITH_EDITOR
 	FOnPCGBlueprintChanged OnBlueprintChangedDelegate;
 #endif
+
+protected:
+#if WITH_EDITOR
+	void OnDependencyChanged(UObject* Object, FPropertyChangedEvent& PropertyChangedEvent);
+	TSet<TObjectPtr<UObject>> DataDependencies;
+#endif
 };
 
 UCLASS(BlueprintType, ClassGroup = (Procedural))
@@ -96,7 +110,6 @@ protected:
 protected:
 #if WITH_EDITOR
 	void OnBlueprintChanged(UBlueprint* InBlueprint);
-	void OnDependencyChanged(UObject* Object, FPropertyChangedEvent& PropertyChangedEvent);
 	void OnBlueprintElementChanged(UPCGBlueprintElement* InElement);
 #endif
 
@@ -105,10 +118,6 @@ protected:
 	void TeardownBlueprintEvent();
 	void SetupBlueprintElementEvent();
 	void TeardownBlueprintElementEvent();
-
-#if WITH_EDITORONLY_DATA
-	TSet<TObjectPtr<UObject>> DataDependencies;
-#endif
 };
 
 class FPCGExecuteBlueprintElement : public FSimpleTypedPCGElement<UPCGBlueprintSettings>
diff --git a/Engine/Plugins/Experimental/PCG/Source/PCG/Public/PCGNode.h b/Engine/Plugins/Experimental/PCG/Source/PCG/Public/PCGNode.h
--- a/Engine/Plugins/Experimental/PCG/Source/PCG/Public/PCGNode.h
+++ b/Engine/Plugins/Experimental/PCG/Source/PCG/Public/PCGNode.h
@@ -38,6 +38,10 @@ public:
 	UFUNCTION(BlueprintCallable, Category = Node)
 	UPCGNode* AddEdgeTo(UPCGNode* To);
 
+	/** Changes the default settings in the node */
+	void SetDefaultSettings(TObjectPtr<UPCGSettings> InSettings);
+
+	/** Note: do not set this property directly from code, use SetDefaultSettings instead */
 	UPROPERTY(BlueprintReadWrite, EditAnywhere, Category = Node, meta=(EditInline))
 	TObjectPtr<UPCGSettings> DefaultSettings;
 
