Fixed static analysis warning in PCG settings

#rnx

diff --git a/Engine/Plugins/Experimental/PCG/Source/PCG/Private/PCGSettings.cpp b/Engine/Plugins/Experimental/PCG/Source/PCG/Private/PCGSettings.cpp
--- a/Engine/Plugins/Experimental/PCG/Source/PCG/Private/PCGSettings.cpp
+++ b/Engine/Plugins/Experimental/PCG/Source/PCG/Private/PCGSettings.cpp
@@ -9,18 +9,18 @@
 #endif
 
 /** In order to reuse the cache when only debug settings change, we must make sure to ignore these from the CRC check */
-#if WITH_EDITORONLY_DATA
 class FPCGSettingsObjectCrc32 : public FArchiveObjectCrc32
 {
 public:
 	virtual bool ShouldSkipProperty(const FProperty* InProperty) const override
 	{
+#if WITH_EDITORONLY_DATA
 		return InProperty && InProperty->GetFName() == GET_MEMBER_NAME_CHECKED(UPCGSettings, DebugSettings);
+#else
+		return FArchiveObjectCrc32::ShouldSkipProperty(InProperty);
+#endif // WITH_EDITORONLY_DATA
 	}
 };
-#else
-typedef FArchiveObjectCrc32 FPCGSettingsObjectCrc32;
-#endif
 
 bool UPCGSettings::operator==(const UPCGSettings& Other) const
 {
