Updated self-pruning following extents/steepness changes

#rnx

diff --git a/Engine/Plugins/Experimental/PCG/Source/PCG/Private/Elements/PCGSelfPruning.cpp b/Engine/Plugins/Experimental/PCG/Source/PCG/Private/Elements/PCGSelfPruning.cpp
--- a/Engine/Plugins/Experimental/PCG/Source/PCG/Private/Elements/PCGSelfPruning.cpp
+++ b/Engine/Plugins/Experimental/PCG/Source/PCG/Private/Elements/PCGSelfPruning.cpp
@@ -15,13 +15,13 @@ namespace PCGSelfPruningAlgorithms
 
 	bool SortSmallToLargeNoRandom(const FPCGPoint* A, const FPCGPoint* B, FVector::FReal SquaredRadiusEquality)
 	{
-		return A->Extents.SquaredLength() * SquaredRadiusEquality < B->Extents.SquaredLength();
+		return A->GetDensityBounds().BoxExtent.SquaredLength() * SquaredRadiusEquality < B->GetDensityBounds().BoxExtent.SquaredLength();
 	}
 
 	bool SortSmallToLargeWithRandom(const FPCGPoint* A, const FPCGPoint* B, FVector::FReal SquaredRadiusEquality)
 	{
-		const FVector::FReal SqrLenA = A->Extents.SquaredLength();
-		const FVector::FReal SqrLenB = B->Extents.SquaredLength();
+		const FVector::FReal SqrLenA = A->GetDensityBounds().BoxExtent.SquaredLength();
+		const FVector::FReal SqrLenB = B->GetDensityBounds().BoxExtent.SquaredLength();
 		if (SqrLenA * SquaredRadiusEquality < SqrLenB)
 		{
 			return true;
@@ -141,7 +141,8 @@ bool FPCGSelfPruningElement::ExecuteInternal(FPCGContextPtr Context) const
 
 			ExclusionPoints.Add(CurrentPoint);
 
-			Octree.FindElementsWithBoundsTest(FBoxCenterAndExtent(CurrentPoint->Transform.GetLocation(), CurrentPoint->Extents), [&ExclusionPoints, &ExcludedPoints](const FPCGPointRef& InPointRef) {
+			const FBoxSphereBounds CurrentPointBounds = CurrentPoint->GetDensityBounds();
+			Octree.FindElementsWithBoundsTest(FBoxCenterAndExtent(CurrentPointBounds.Origin, CurrentPointBounds.BoxExtent), [&ExclusionPoints, &ExcludedPoints](const FPCGPointRef& InPointRef) {
 				if (!ExclusionPoints.Contains(InPointRef.Point))
 				{
 					ExcludedPoints.Add(InPointRef.Point);
