Fixed issue with PCG excluded tag actors

#rnx

diff --git a/Engine/Plugins/Experimental/PCG/Source/PCG/Private/PCGComponent.cpp b/Engine/Plugins/Experimental/PCG/Source/PCG/Private/PCGComponent.cpp
--- a/Engine/Plugins/Experimental/PCG/Source/PCG/Private/PCGComponent.cpp
+++ b/Engine/Plugins/Experimental/PCG/Source/PCG/Private/PCGComponent.cpp
@@ -735,6 +735,7 @@ void UPCGComponent::OnActorMoved(AActor* InActor)
 
 		if (ActorHasExcludedTag(InActor))
 		{
+			DirtyExclusionData(InActor);
 			bDirtyAndRefresh = true;
 		}
 
@@ -807,6 +808,7 @@ void UPCGComponent::OnObjectPropertyChanged(UObject* InObject, FPropertyChangedE
 			if (ActorHasExcludedTag(Actor))
 			{
 				DirtyExclusionData(Actor);
+				bDirtyAndRefresh = true;
 			}
 
 			if (DirtyTrackedActor(Actor))
