Improved PCG projection behavior
Added color to PCG texture data point outputs/transforms

#rnx

diff --git a/Engine/Plugins/Experimental/PCG/Source/PCG/Private/Data/PCGPointData.cpp b/Engine/Plugins/Experimental/PCG/Source/PCG/Private/Data/PCGPointData.cpp
--- a/Engine/Plugins/Experimental/PCG/Source/PCG/Private/Data/PCGPointData.cpp
+++ b/Engine/Plugins/Experimental/PCG/Source/PCG/Private/Data/PCGPointData.cpp
@@ -9,7 +9,7 @@
 
 namespace PCGPointHelpers
 {
-	float ManhattanDensity(const FPCGPoint& InPoint, const FVector& InPosition)
+	bool GetDistanceRatios(const FPCGPoint& InPoint, const FVector& InPosition, FVector& OutRatios)
 	{
 		FVector LocalPosition = InPoint.Transform.InverseTransformPosition(InPosition);
 		LocalPosition /= InPoint.Extents;
@@ -17,12 +17,12 @@ namespace PCGPointHelpers
 		// ]-2+s, 2-s] is the valid range of values
 		const FVector::FReal LowerBound = InPoint.Steepness - 2;
 		const FVector::FReal HigherBound = 2 - InPoint.Steepness;
-		
+
 		if (LocalPosition.X <= LowerBound || LocalPosition.X > HigherBound ||
 			LocalPosition.Y <= LowerBound || LocalPosition.Y > HigherBound ||
 			LocalPosition.Z <= LowerBound || LocalPosition.Z > HigherBound)
 		{
-			return 0;
+			return false;
 		}
 
 		// [-s, +s] is the range where the density is 1 on that axis
@@ -32,9 +32,36 @@ namespace PCGPointHelpers
 
 		const FVector::FReal DistanceScale = FMath::Max(2 - 2 * InPoint.Steepness, KINDA_SMALL_NUMBER);
 
-		//Note: for euclidean, we could do 1 - (dist / scale)^2
-		//Note: for maximum norm, we could do density * max(x factor, y factor, z factor)
-		return InPoint.Density * (1 - XDist/DistanceScale) * (1 - YDist/DistanceScale) * (1 - ZDist/DistanceScale);
+		OutRatios.X = XDist / DistanceScale;
+		OutRatios.Y = YDist / DistanceScale;
+		OutRatios.Z = ZDist / DistanceScale;
+		return true;
+	}
+
+	float ManhattanDensity(const FPCGPoint& InPoint, const FVector& InPosition)
+	{
+		FVector Ratios;
+		if (GetDistanceRatios(InPoint, InPosition, Ratios))
+		{
+			return InPoint.Density * (1 - Ratios.X) * (1 - Ratios.Y) * (1 - Ratios.Z);
+		}
+		else
+		{
+			return 0;
+		}
+	}
+
+	float InverseEuclidianDistance(const FPCGPoint& InPoint, const FVector& InPosition)
+	{
+		FVector Ratios;
+		if (GetDistanceRatios(InPoint, InPosition, Ratios))
+		{
+			return 1 - Ratios.Length();
+		}
+		else
+		{
+			return 0;
+		}
 	}
 }
 
@@ -150,6 +177,65 @@ const FPCGPoint* UPCGPointData::GetPointAtPosition(const FVector& InPosition) co
 	return BestPoint;
 }
 
+FPCGPoint UPCGPointData::TransformPoint(const FPCGPoint& InPoint) const
+{
+	if (bOctreeIsDirty)
+	{
+		RebuildOctree();
+	}
+
+	FPCGPoint Point = InPoint;
+
+	TArray<TPair<const FPCGPoint*, float>> Contributions;
+	const FVector PointPosition = InPoint.Transform.GetLocation();
+
+	Octree.FindElementsWithBoundsTest(FBoxCenterAndExtent(PointPosition, FVector::Zero()), [&PointPosition, &Contributions](const FPCGPointRef& InPointRef) {
+		Contributions.Emplace(InPointRef.Point, PCGPointHelpers::InverseEuclidianDistance(*InPointRef.Point, PointPosition));
+	});
+
+	float SumContrib = 0;
+	for (const auto& Contribution : Contributions)
+	{
+		SumContrib += Contribution.Value;
+	}
+
+	if (SumContrib <= 0)
+	{
+		return InPoint;
+	}
+
+	FVector WeightedScale = FVector::Zero();
+	FRotator WeightedRotator(0);
+	float WeightedDensity = 0;
+	FVector WeightedExtents = FVector::Zero();
+	FVector WeightedColor = FVector::Zero();
+	float WeightedSteepness = 0;
+
+	for (const auto& Contribution : Contributions)
+	{
+		const FPCGPoint& SourcePoint = *Contribution.Key;
+		float Weight = Contribution.Value / SumContrib;
+
+		WeightedRotator += (SourcePoint.Transform.Rotator() * Weight); // TODO: This is wonky
+		WeightedScale += (SourcePoint.Transform.GetScale3D() * Weight);
+		WeightedDensity += PCGPointHelpers::ManhattanDensity(SourcePoint, PointPosition);
+		WeightedColor += SourcePoint.Color * Weight;
+		WeightedSteepness += SourcePoint.Steepness * Weight;
+	}
+
+	// Finally, apply changes to point
+	FQuat PointRotation = (Point.Transform.Rotator() + WeightedRotator).Quaternion();
+
+	Point.Transform.SetRotation(PointRotation);
+	Point.Transform.NormalizeRotation();
+	Point.Transform.SetScale3D(Point.Transform.GetScale3D() * WeightedScale);	
+	Point.Density *= WeightedDensity;
+	Point.Color *= WeightedColor;
+	Point.Steepness *= WeightedSteepness;
+
+	return Point;
+}
+
 float UPCGPointData::GetDensityAtPosition(const FVector& InPosition) const
 {
 	if (bOctreeIsDirty)
diff --git a/Engine/Plugins/Experimental/PCG/Source/PCG/Private/Data/PCGProjectionData.cpp b/Engine/Plugins/Experimental/PCG/Source/PCG/Private/Data/PCGProjectionData.cpp
--- a/Engine/Plugins/Experimental/PCG/Source/PCG/Private/Data/PCGProjectionData.cpp
+++ b/Engine/Plugins/Experimental/PCG/Source/PCG/Private/Data/PCGProjectionData.cpp
@@ -53,7 +53,7 @@ FBox UPCGProjectionData::ProjectBounds(const FBox& InBounds) const
 
 	for (int Corner = 0; Corner < 8; ++Corner)
 	{
-		Bounds += TransformPosition(
+		Bounds += Target->TransformPosition(
 			FVector(
 				(Corner / 4) ? InBounds.Max.X : InBounds.Min.X,
 				((Corner / 2) % 2) ? InBounds.Max.Y : InBounds.Min.Y,
@@ -72,13 +72,13 @@ float UPCGProjectionData::GetDensityAtPosition(const FVector& InPosition) const
 FVector UPCGProjectionData::TransformPosition(const FVector& InPosition) const
 {
 	// TODO: improve projection/unprojection mechanism
-	return Target->TransformPosition(InPosition);
+	return Target->TransformPosition(Source->TransformPosition(InPosition));
 }
 
 FPCGPoint UPCGProjectionData::TransformPoint(const FPCGPoint& InPoint) const
 {
 	// TODO: improve projection/unprojection mechanism
-	return Target->TransformPoint(InPoint);
+	return Target->TransformPoint(Source->TransformPoint(InPoint));
 }
 
 bool UPCGProjectionData::HasNonTrivialTransform() const
@@ -101,7 +101,7 @@ const UPCGPointData* UPCGProjectionData::CreatePointData() const
 	Points.Reserve(SourcePoints.Num());
 	for (const FPCGPoint& SourcePoint : SourcePoints)
 	{
-		FPCGPoint ProjectedPoint = TransformPoint(SourcePoint);
+		FPCGPoint ProjectedPoint = Target->TransformPoint(SourcePoint);
 #if WITH_EDITORONLY_DATA
 		if(ProjectedPoint.Density > 0 || bKeepZeroDensityPoints)
 #else
diff --git a/Engine/Plugins/Experimental/PCG/Source/PCG/Private/Data/PCGTextureData.cpp b/Engine/Plugins/Experimental/PCG/Source/PCG/Private/Data/PCGTextureData.cpp
--- a/Engine/Plugins/Experimental/PCG/Source/PCG/Private/Data/PCGTextureData.cpp
+++ b/Engine/Plugins/Experimental/PCG/Source/PCG/Private/Data/PCGTextureData.cpp
@@ -85,6 +85,27 @@ float UPCGBaseTextureData::GetDensityAtPosition(const FVector& InPosition) const
 	return PCGTextureSampling::Sample<float>(Position2D, Surface, Width, Height, [this](int32 Index) { return PCGTextureSampling::SampleFloatChannel(ColorData[Index], ColorChannel); });
 }
 
+FPCGPoint UPCGBaseTextureData::TransformPoint(const FPCGPoint& InPoint) const
+{
+	FPCGPoint Point = InPoint;
+
+	// Update point location: put it on the surface plane
+	FVector PointPositionInLocalSpace = TransformPosition(InPoint.Transform.GetLocation());
+	PointPositionInLocalSpace.Z = 0;
+	Point.Transform.SetLocation(Transform.TransformPosition(PointPositionInLocalSpace));
+
+	// Set/Update density & color
+	FVector2D Position2D(PointPositionInLocalSpace.X, PointPositionInLocalSpace.Y);
+	FBox2D Surface(FVector2D(-1.0f, -1.0f), FVector2D(1.0f, 1.0f));
+
+	FLinearColor Color = PCGTextureSampling::Sample<FLinearColor>(Position2D, Surface, Width, Height, [this](int32 Index) { return ColorData[Index]; });
+
+	Point.Color = Color;
+	Point.Density *= PCGTextureSampling::SampleFloatChannel(Color, ColorChannel);
+
+	return Point;
+}
+
 const UPCGPointData* UPCGBaseTextureData::CreatePointData() const
 {
 	TRACE_CPUPROFILER_EVENT_SCOPE(UPCGBaseTextureData::CreatePointData);
@@ -116,9 +137,10 @@ const UPCGPointData* UPCGBaseTextureData::CreatePointData() const
 #endif
 			{
 				FVector LocalPosition(X * XScale + Bias.X, Y * YScale + Bias.Y, 0);
-				Points.Emplace(FTransform(Transform.TransformPosition(LocalPosition)),
+				FPCGPoint& Point = Points.Emplace_GetRef(FTransform(Transform.TransformPosition(LocalPosition)),
 					Density,
 					PCGHelpers::ComputeSeed(X, Y));
+				Point.Color = ColorData[X + Y * Width];
 			}
 		}
 	}
diff --git a/Engine/Plugins/Experimental/PCG/Source/PCG/Public/Data/PCGPointData.h b/Engine/Plugins/Experimental/PCG/Source/PCG/Public/Data/PCGPointData.h
--- a/Engine/Plugins/Experimental/PCG/Source/PCG/Public/Data/PCGPointData.h
+++ b/Engine/Plugins/Experimental/PCG/Source/PCG/Public/Data/PCGPointData.h
@@ -66,6 +66,7 @@ public:
 	virtual FBox GetBounds() const override;
 	virtual float GetDensityAtPosition(const FVector& InPosition) const override;
 	virtual const UPCGPointData* ToPointData() const { return this; }
+	virtual FPCGPoint TransformPoint(const FPCGPoint& InPoint) const;
 	// ~End UPCGSpatialData interface
 
 	UFUNCTION(BlueprintCallable, Category = SpatialData)
diff --git a/Engine/Plugins/Experimental/PCG/Source/PCG/Public/Data/PCGTextureData.h b/Engine/Plugins/Experimental/PCG/Source/PCG/Public/Data/PCGTextureData.h
--- a/Engine/Plugins/Experimental/PCG/Source/PCG/Public/Data/PCGTextureData.h
+++ b/Engine/Plugins/Experimental/PCG/Source/PCG/Public/Data/PCGTextureData.h
@@ -26,6 +26,7 @@ public:
 	virtual FBox GetBounds() const override;
 	virtual FBox GetStrictBounds() const override;
 	virtual float GetDensityAtPosition(const FVector& InPosition) const override;
+	virtual FPCGPoint TransformPoint(const FPCGPoint& InPoint) const;
 	//~End UPCGSpatialData interface
 
 	//~Begin UPCGSpatialDataWithPointCache interface
