Added basic multithreading across elements and data.
Note: self-pruning, landscape spline, spline will be done at a later time

diff --git a/Engine/Plugins/Experimental/PCG/Source/PCG/Private/Data/PCGDifferenceData.cpp b/Engine/Plugins/Experimental/PCG/Source/PCG/Private/Data/PCGDifferenceData.cpp
--- a/Engine/Plugins/Experimental/PCG/Source/PCG/Private/Data/PCGDifferenceData.cpp
+++ b/Engine/Plugins/Experimental/PCG/Source/PCG/Private/Data/PCGDifferenceData.cpp
@@ -3,6 +3,7 @@
 #include "Data/PCGDifferenceData.h"
 #include "Data/PCGPointData.h"
 #include "Data/PCGUnionData.h"
+#include "Helpers/PCGAsync.h"
 
 namespace PCGDifferenceDataUtils
 {
@@ -128,12 +129,12 @@ bool UPCGDifferenceData::HasNonTrivialTransform() const
 	return Source->HasNonTrivialTransform();
 }
 
-const UPCGPointData* UPCGDifferenceData::CreatePointData() const
+const UPCGPointData* UPCGDifferenceData::CreatePointData(FPCGContextPtr Context) const
 {
 	TRACE_CPUPROFILER_EVENT_SCOPE(UPCGDifferenceData::CreatePointData);
 	
 	// This is similar to what we are doing in UPCGUnionData::CreatePointData
-	const UPCGPointData* SourcePointData = Source->ToPointData();
+	const UPCGPointData* SourcePointData = Source->ToPointData(Context);
 
 	if (!SourcePointData)
 	{
@@ -150,25 +151,33 @@ const UPCGPointData* UPCGDifferenceData::CreatePointData() const
 	UPCGPointData* Data = NewObject<UPCGPointData>(const_cast<UPCGDifferenceData*>(this));
 	Data->TargetActor = TargetActor;
 
+	const TArray<FPCGPoint>& SourcePoints = SourcePointData->GetPoints();
 	TArray<FPCGPoint>& TargetPoints = Data->GetMutablePoints();
 
-	for (const FPCGPoint& Point : SourcePointData->GetPoints())
+	FPCGAsync::AsyncPointProcessing(Context, SourcePoints.Num(), TargetPoints, [this, &SourcePoints](int32 Index, FPCGPoint& OutPoint)
 	{
+		const FPCGPoint& Point = SourcePoints[Index];
 		const float DensityInDifference = Difference->GetDensityAtPosition(Point.Transform.GetLocation());
 
 		if (DensityInDifference < Point.Density)
 		{
-			FPCGPoint& TargetPoint = TargetPoints.Add_GetRef(Point);
-			TargetPoint.Density -= DensityInDifference;
+			OutPoint = Point;
+			OutPoint.Density -= DensityInDifference;
+			return true;
 		}
 #if WITH_EDITORONLY_DATA
 		else if (bKeepZeroDensityPoints)
 		{
-			FPCGPoint& TargetPoint = TargetPoints.Add_GetRef(Point);
-			TargetPoint.Density = 0;
+			OutPoint = Point;
+			OutPoint.Density = 0;
+			return true;
 		}
 #endif
-	}
+		else
+		{
+			return false;
+		}
+	});
 
 	UE_LOG(LogPCG, Verbose, TEXT("Difference generated %d points from %d source points"), TargetPoints.Num(), SourcePointData->GetPoints().Num());
 
diff --git a/Engine/Plugins/Experimental/PCG/Source/PCG/Private/Data/PCGIntersectionData.cpp b/Engine/Plugins/Experimental/PCG/Source/PCG/Private/Data/PCGIntersectionData.cpp
--- a/Engine/Plugins/Experimental/PCG/Source/PCG/Private/Data/PCGIntersectionData.cpp
+++ b/Engine/Plugins/Experimental/PCG/Source/PCG/Private/Data/PCGIntersectionData.cpp
@@ -2,6 +2,7 @@
 
 #include "Data/PCGIntersectionData.h"
 #include "Data/PCGPointData.h"
+#include "Helpers/PCGAsync.h"
 
 namespace PCGIntersectionDataMaths
 {
@@ -97,7 +98,7 @@ bool UPCGIntersectionData::HasNonTrivialTransform() const
 	return A->HasNonTrivialTransform() || B->HasNonTrivialTransform();
 }
 
-const UPCGPointData* UPCGIntersectionData::CreatePointData() const
+const UPCGPointData* UPCGIntersectionData::CreatePointData(FPCGContextPtr Context) const
 {
 	check(A && B);
 	// TODO: this is a placeholder;
@@ -105,21 +106,21 @@ const UPCGPointData* UPCGIntersectionData::CreatePointData() const
 	// and then cull out any of the points that are outside the bounds of the other
 	if (A->GetDimension() <= B->GetDimension())
 	{
-		return CreateAndFilterPointData(A, B);
+		return CreateAndFilterPointData(Context, A, B);
 	}
 	else
 	{
-		return CreateAndFilterPointData(B, A);
+		return CreateAndFilterPointData(Context, B, A);
 	}
 }
 
-UPCGPointData* UPCGIntersectionData::CreateAndFilterPointData(const UPCGSpatialData* X, const UPCGSpatialData* Y) const
+UPCGPointData* UPCGIntersectionData::CreateAndFilterPointData(FPCGContextPtr Context, const UPCGSpatialData* X, const UPCGSpatialData* Y) const
 {
 	TRACE_CPUPROFILER_EVENT_SCOPE(UPCGIntersectionData::CreateAndFilterPointData);
 	check(X && Y);
 	check(X->GetDimension() <= Y->GetDimension());
 
-	const UPCGPointData* SourcePointData = X->ToPointData();
+	const UPCGPointData* SourcePointData = X->ToPointData(Context);
 
 	if (!SourcePointData)
 	{
@@ -133,21 +134,26 @@ UPCGPointData* UPCGIntersectionData::CreateAndFilterPointData(const UPCGSpatialD
 	Data->TargetActor = TargetActor;
 	TArray<FPCGPoint>& TargetPoints = Data->GetMutablePoints();
 
-	// TODO: go through only the points in the bounding box, otherwise we know beforehand
-	// that the density will be zero. Requires that points are stored in a better data structure.
-	for (const FPCGPoint& Point : SourcePoints)
+	FPCGAsync::AsyncPointProcessing(Context, SourcePoints.Num(), TargetPoints, [this, &SourcePoints, Y](int32 Index, FPCGPoint& OutPoint)
 	{
+		const FPCGPoint& Point = SourcePoints[Index];
 		const float YDensity = Y->GetDensityAtPosition(Point.Transform.GetLocation());
+
 #if WITH_EDITORONLY_DATA
 		if (YDensity > 0 || bKeepZeroDensityPoints)
 #else
 		if (YDensity > 0)
 #endif
 		{
-			FPCGPoint& NewPoint = TargetPoints.Add_GetRef(Point);
-			NewPoint.Density = PCGIntersectionDataMaths::ComputeDensity(Point.Density, YDensity, DensityFunction);
+			OutPoint = Point;
+			OutPoint.Density = PCGIntersectionDataMaths::ComputeDensity(Point.Density, YDensity, DensityFunction);
+			return true;
 		}
-	}
+		else
+		{
+			return false;
+		}
+	});
 
 	UE_LOG(LogPCG, Verbose, TEXT("Intersection generated %d points from %d source points"), TargetPoints.Num(), SourcePoints.Num());
 
diff --git a/Engine/Plugins/Experimental/PCG/Source/PCG/Private/Data/PCGLandscapeData.cpp b/Engine/Plugins/Experimental/PCG/Source/PCG/Private/Data/PCGLandscapeData.cpp
--- a/Engine/Plugins/Experimental/PCG/Source/PCG/Private/Data/PCGLandscapeData.cpp
+++ b/Engine/Plugins/Experimental/PCG/Source/PCG/Private/Data/PCGLandscapeData.cpp
@@ -6,6 +6,7 @@
 #include "Data/PCGPointData.h"
 #include "Data/PCGPolyLineData.h"
 #include "Data/PCGSpatialData.h"
+#include "Helpers/PCGAsync.h"
 
 #include "Landscape.h"
 #include "LandscapeEdit.h"
@@ -76,7 +77,7 @@ float UPCGLandscapeData::GetDensityAtPosition(const FVector& InPosition) const
 	return HeightAtVertex.IsSet() ? 1.0f : 0;
 }
 
-const UPCGPointData* UPCGLandscapeData::CreatePointData() const
+const UPCGPointData* UPCGLandscapeData::CreatePointData(FPCGContextPtr Context) const
 {
 	TRACE_CPUPROFILER_EVENT_SCOPE(UPCGLandscapeData::CreatePointData);
 
@@ -93,28 +94,33 @@ const UPCGPointData* UPCGLandscapeData::CreatePointData() const
 	const int32 MinY = FMath::FloorToInt(MinPt.Y);
 	const int32 MaxY = FMath::FloorToInt(MaxPt.Y);
 
-	Points.Reserve((MaxX - MinX) * (MaxY - MinY));
+	int32 NumIterations = (MaxX - MinX) * (MaxY - MinY);
 
-	bool bPlaneCase = Bounds.Min.Z == Bounds.Max.Z;
-
-	for (int X = MinX; X < MaxX; ++X)
+	FPCGAsync::AsyncPointProcessing(Context, NumIterations, Points, [this, MinX, MaxX, MinY](int32 Index, FPCGPoint& OutPoint)
 	{
-		for (int Y = MinY; Y < MaxY; ++Y)
+		const int X = MinX + (Index % (MaxX - MinX));
+		const int Y = MinY + (Index / (MaxX - MinX));
+		const bool bPlaneCase = Bounds.Min.Z == Bounds.Max.Z;
+
+		FVector VertexLocation = Transform.TransformPosition(FVector(X, Y, 0));
+		TOptional<float> HeightAtVertex = Landscape->GetHeightAtLocation(VertexLocation);
+		if (HeightAtVertex.IsSet() &&
+			HeightAtVertex.GetValue() >= Bounds.Min.Z &&
+			(bPlaneCase ? HeightAtVertex.GetValue() <= Bounds.Max.Z : HeightAtVertex.GetValue() < Bounds.Max.Z))
+		{
+			VertexLocation.Z = HeightAtVertex.GetValue();
+
+			OutPoint = FPCGPoint(FTransform(VertexLocation),
+				1.0f,
+				PCGHelpers::ComputeSeed(X, Y, (int)HeightAtVertex.GetValue()));
+
+			return true;
+		}
+		else
 		{
-			FVector VertexLocation = Transform.TransformPosition(FVector(X, Y, 0));
-			TOptional<float> HeightAtVertex = Landscape->GetHeightAtLocation(VertexLocation);
-			if (HeightAtVertex.IsSet() &&
-				HeightAtVertex.GetValue() >= Bounds.Min.Z && 
-				(bPlaneCase ? HeightAtVertex.GetValue() <= Bounds.Max.Z : HeightAtVertex.GetValue() < Bounds.Max.Z))
-			{
-				VertexLocation.Z = HeightAtVertex.GetValue();
-
-				Points.Emplace(FTransform(VertexLocation),
-					1.0f,
-					PCGHelpers::ComputeSeed(X, Y, (int)HeightAtVertex.GetValue()));
-			}
+			return false;
 		}
-	}
+	});
 
 	UE_LOG(LogPCG, Verbose, TEXT("Landscape %s extracted %d of %d potential points"), *Landscape->GetFName().ToString(), Points.Num(), (MaxX - MinX) * (MaxY - MinY));
 
diff --git a/Engine/Plugins/Experimental/PCG/Source/PCG/Private/Data/PCGLandscapeSplineData.cpp b/Engine/Plugins/Experimental/PCG/Source/PCG/Private/Data/PCGLandscapeSplineData.cpp
--- a/Engine/Plugins/Experimental/PCG/Source/PCG/Private/Data/PCGLandscapeSplineData.cpp
+++ b/Engine/Plugins/Experimental/PCG/Source/PCG/Private/Data/PCGLandscapeSplineData.cpp
@@ -88,7 +88,7 @@ FVector UPCGLandscapeSplineData::GetLocationAtDistance(int SegmentIndex, float D
 	return FVector::Zero();
 }
 
-const UPCGPointData* UPCGLandscapeSplineData::CreatePointData() const
+const UPCGPointData* UPCGLandscapeSplineData::CreatePointData(FPCGContextPtr Context) const
 {
 	check(Spline);
 	TRACE_CPUPROFILER_EVENT_SCOPE(UPCGLandscapeSplineData::CreatePointData);
diff --git a/Engine/Plugins/Experimental/PCG/Source/PCG/Private/Data/PCGProjectionData.cpp b/Engine/Plugins/Experimental/PCG/Source/PCG/Private/Data/PCGProjectionData.cpp
--- a/Engine/Plugins/Experimental/PCG/Source/PCG/Private/Data/PCGProjectionData.cpp
+++ b/Engine/Plugins/Experimental/PCG/Source/PCG/Private/Data/PCGProjectionData.cpp
@@ -1,6 +1,7 @@
 // Copyright Epic Games, Inc. All Rights Reserved.
 
 #include "Data/PCGProjectionData.h"
+#include "Helpers/PCGAsync.h"
 
 void UPCGProjectionData::Initialize(const UPCGSpatialData* InSource, const UPCGSpatialData* InTarget)
 {
@@ -86,31 +87,29 @@ bool UPCGProjectionData::HasNonTrivialTransform() const
 	return Target->HasNonTrivialTransform();
 }
 
-const UPCGPointData* UPCGProjectionData::CreatePointData() const
+const UPCGPointData* UPCGProjectionData::CreatePointData(FPCGContextPtr Context) const
 {
 	TRACE_CPUPROFILER_EVENT_SCOPE(UPCGProjectionData::CreatePointData);
 	// TODO: add mechanism in the ToPointData so we can pass in a transform
 	// so we can forego creating the points twice if they're not used.
-	const UPCGPointData* SourcePointData = Source->ToPointData();
+	const UPCGPointData* SourcePointData = Source->ToPointData(Context);
 	const TArray<FPCGPoint>& SourcePoints = SourcePointData->GetPoints();
 
 	UPCGPointData* PointData = NewObject<UPCGPointData>(const_cast<UPCGProjectionData*>(this));
 	PointData->TargetActor = TargetActor;
 	TArray<FPCGPoint>& Points = PointData->GetMutablePoints();
 
-	Points.Reserve(SourcePoints.Num());
-	for (const FPCGPoint& SourcePoint : SourcePoints)
+	FPCGAsync::AsyncPointProcessing(Context, SourcePoints.Num(), Points, [this, &SourcePoints](int32 Index, FPCGPoint& OutPoint)
 	{
-		FPCGPoint ProjectedPoint = Target->TransformPoint(SourcePoint);
+		const FPCGPoint& SourcePoint = SourcePoints[Index];
+		OutPoint = Target->TransformPoint(SourcePoint);
+
 #if WITH_EDITORONLY_DATA
-		if(ProjectedPoint.Density > 0 || bKeepZeroDensityPoints)
+		return OutPoint.Density > 0 || bKeepZeroDensityPoints;
 #else
-		if (ProjectedPoint.Density > 0)
+		return OutPoint.Density > 0;
 #endif
-		{
-			Points.Add(MoveTemp(ProjectedPoint));
-		}
-	}
+	});
 
 	UE_LOG(LogPCG, Verbose, TEXT("Projection generated %d points from %d source points"), Points.Num(), SourcePoints.Num());
 
diff --git a/Engine/Plugins/Experimental/PCG/Source/PCG/Private/Data/PCGSpatialData.cpp b/Engine/Plugins/Experimental/PCG/Source/PCG/Private/Data/PCGSpatialData.cpp
--- a/Engine/Plugins/Experimental/PCG/Source/PCG/Private/Data/PCGSpatialData.cpp
+++ b/Engine/Plugins/Experimental/PCG/Source/PCG/Private/Data/PCGSpatialData.cpp
@@ -6,7 +6,7 @@
 #include "Data/PCGProjectionData.h"
 #include "Data/PCGUnionData.h"
 
-const UPCGPointData* UPCGSpatialDataWithPointCache::ToPointData() const
+const UPCGPointData* UPCGSpatialDataWithPointCache::ToPointData(FPCGContextPtr Context) const
 {
 	if (!CachedPointData)
 	{
@@ -14,7 +14,7 @@ const UPCGPointData* UPCGSpatialDataWithPointCache::ToPointData() const
 
 		if (!CachedPointData)
 		{
-			CachedPointData = CreatePointData();
+			CachedPointData = CreatePointData(Context);
 		}
 
 		CacheLock.Unlock();
diff --git a/Engine/Plugins/Experimental/PCG/Source/PCG/Private/Data/PCGSplineData.cpp b/Engine/Plugins/Experimental/PCG/Source/PCG/Private/Data/PCGSplineData.cpp
--- a/Engine/Plugins/Experimental/PCG/Source/PCG/Private/Data/PCGSplineData.cpp
+++ b/Engine/Plugins/Experimental/PCG/Source/PCG/Private/Data/PCGSplineData.cpp
@@ -43,7 +43,7 @@ FVector UPCGSplineData::GetLocationAtDistance(int SegmentIndex, float Distance)
 	return Spline->GetLocationAtDistanceAlongSpline(Distance, ESplineCoordinateSpace::World);
 }
 
-const UPCGPointData* UPCGSplineData::CreatePointData() const
+const UPCGPointData* UPCGSplineData::CreatePointData(FPCGContextPtr Context) const
 {
 	TRACE_CPUPROFILER_EVENT_SCOPE(UPCGSplineData::CreatePointData);
 	UPCGPointData* Data = NewObject<UPCGPointData>(const_cast<UPCGSplineData*>(this));
diff --git a/Engine/Plugins/Experimental/PCG/Source/PCG/Private/Data/PCGTextureData.cpp b/Engine/Plugins/Experimental/PCG/Source/PCG/Private/Data/PCGTextureData.cpp
--- a/Engine/Plugins/Experimental/PCG/Source/PCG/Private/Data/PCGTextureData.cpp
+++ b/Engine/Plugins/Experimental/PCG/Source/PCG/Private/Data/PCGTextureData.cpp
@@ -4,6 +4,7 @@
 
 #include "PCGHelpers.h"
 #include "Data/PCGPointData.h"
+#include "Helpers/PCGAsync.h"
 
 namespace PCGTextureSampling
 {
@@ -106,7 +107,7 @@ FPCGPoint UPCGBaseTextureData::TransformPoint(const FPCGPoint& InPoint) const
 	return Point;
 }
 
-const UPCGPointData* UPCGBaseTextureData::CreatePointData() const
+const UPCGPointData* UPCGBaseTextureData::CreatePointData(FPCGContextPtr Context) const
 {
 	TRACE_CPUPROFILER_EVENT_SCOPE(UPCGBaseTextureData::CreatePointData);
 	// TODO: this is a trivial implementation
@@ -119,31 +120,36 @@ const UPCGPointData* UPCGBaseTextureData::CreatePointData() const
 	TArray<FPCGPoint>& Points = Data->GetMutablePoints();
 
 	// TODO: There's a bias issue here where we should correct by a 0.5 unit...
-	Points.Reserve(Width * Height);
 	const FVector::FReal XScale = LocalSurfaceBounds.GetSize().X / Width;
 	const FVector::FReal YScale = LocalSurfaceBounds.GetSize().Y / Height;
 	const FVector2D Bias = LocalSurfaceBounds.Min;
 
-	for (int32 X = 0; X < Width; ++X)
+	FPCGAsync::AsyncPointProcessing(Context, Width * Height, Points, [this, XScale, YScale, Bias](int32 Index, FPCGPoint& OutPoint)
 	{
-		for (int32 Y = 0; Y < Height; ++Y)
-		{
-			const float Density = PCGTextureSampling::SampleFloatChannel(ColorData[X + Y * Width], ColorChannel);
+		const int X = Index % Width;
+		const int Y = Index / Width;
+
+		const float Density = PCGTextureSampling::SampleFloatChannel(ColorData[X + Y * Width], ColorChannel);
 
 #if WITH_EDITORONLY_DATA
-			if(Density > 0 || bKeepZeroDensityPoints)
+		if (Density > 0 || bKeepZeroDensityPoints)
 #else
-			if (Density > 0)
+		if (Density > 0)
 #endif
-			{
-				FVector LocalPosition(X * XScale + Bias.X, Y * YScale + Bias.Y, 0);
-				FPCGPoint& Point = Points.Emplace_GetRef(FTransform(Transform.TransformPosition(LocalPosition)),
-					Density,
-					PCGHelpers::ComputeSeed(X, Y));
-				Point.Color = ColorData[X + Y * Width];
-			}
+		{
+			FVector LocalPosition(X * XScale + Bias.X, Y * YScale + Bias.Y, 0);
+			OutPoint = FPCGPoint(FTransform(Transform.TransformPosition(LocalPosition)),
+				Density,
+				PCGHelpers::ComputeSeed(X, Y));
+			OutPoint.Color = ColorData[X + Y * Width];
+
+			return true;
 		}
-	}
+		else
+		{
+			return false;
+		}
+	});
 
 	return Data;
 }
diff --git a/Engine/Plugins/Experimental/PCG/Source/PCG/Private/Data/PCGUnionData.cpp b/Engine/Plugins/Experimental/PCG/Source/PCG/Private/Data/PCGUnionData.cpp
--- a/Engine/Plugins/Experimental/PCG/Source/PCG/Private/Data/PCGUnionData.cpp
+++ b/Engine/Plugins/Experimental/PCG/Source/PCG/Private/Data/PCGUnionData.cpp
@@ -2,6 +2,7 @@
 
 #include "Data/PCGUnionData.h"
 #include "Data/PCGPointData.h"
+#include "Helpers/PCGAsync.h"
 
 namespace PCGUnionDataMaths
 {
@@ -148,7 +149,7 @@ bool UPCGUnionData::HasNonTrivialTransform() const
 	return (FirstNonTrivialTransformData != nullptr || Super::HasNonTrivialTransform());
 }
 
-const UPCGPointData* UPCGUnionData::CreatePointData() const
+const UPCGPointData* UPCGUnionData::CreatePointData(FPCGContextPtr Context) const
 {
 	TRACE_CPUPROFILER_EVENT_SCOPE(UPCGUnionData::CreatePointData);
 
@@ -161,7 +162,7 @@ const UPCGPointData* UPCGUnionData::CreatePointData() const
 	else if (Data.Num() == 1)
 	{
 		UE_LOG(LogPCG, Verbose, TEXT("Union is trivial"));
-		return Data[0]->ToPointData();
+		return Data[0]->ToPointData(Context);
 	}
 
 	UPCGPointData* PointData = NewObject<UPCGPointData>(const_cast<UPCGUnionData*>(this));
@@ -171,11 +172,11 @@ const UPCGPointData* UPCGUnionData::CreatePointData() const
 	{
 	case EPCGUnionType::LeftToRightPriority:
 	default:
-		CreateSequentialPointData(PointData, /*bLeftToRight=*/true);
+		CreateSequentialPointData(Context, PointData, /*bLeftToRight=*/true);
 		break;
 
 	case EPCGUnionType::RightToLeftPriority:
-		CreateSequentialPointData(PointData, /*bLeftToRight=*/false);
+		CreateSequentialPointData(Context, PointData, /*bLeftToRight=*/false);
 		break;
 
 	case EPCGUnionType::KeepAll:
@@ -183,7 +184,7 @@ const UPCGPointData* UPCGUnionData::CreatePointData() const
 			TArray<FPCGPoint>& TargetPoints = PointData->GetMutablePoints();
 			for (TObjectPtr<const UPCGSpatialData> Datum : Data)
 			{
-				TargetPoints.Append(Datum->ToPointData()->GetPoints());
+				TargetPoints.Append(Datum->ToPointData(Context)->GetPoints());
 			}
 		}
 		break;
@@ -194,11 +195,12 @@ const UPCGPointData* UPCGUnionData::CreatePointData() const
 	return PointData;
 }
 
-void UPCGUnionData::CreateSequentialPointData(UPCGPointData* PointData, bool bLeftToRight) const
+void UPCGUnionData::CreateSequentialPointData(FPCGContextPtr Context, UPCGPointData* PointData, bool bLeftToRight) const
 {
 	check(PointData);
 
 	TArray<FPCGPoint>& TargetPoints = PointData->GetMutablePoints();
+	TArray<FPCGPoint> SelectedDataPoints;
 
 	int32 FirstDataIndex = (bLeftToRight ? 0 : Data.Num() - 1);
 	int32 LastDataIndex = (bLeftToRight ? Data.Num() : -1);
@@ -210,9 +212,12 @@ void UPCGUnionData::CreateSequentialPointData(UPCGPointData* PointData, bool bLe
 	{
 		// For each point, if it is not already "processed" by previous data,
 		// add it & compute its final density
-		const TArray<FPCGPoint>& Points = Data[DataIndex]->ToPointData()->GetPoints();
-		for (const FPCGPoint& Point : Points)
+		const TArray<FPCGPoint>& Points = Data[DataIndex]->ToPointData(Context)->GetPoints();
+
+		FPCGAsync::AsyncPointProcessing(Context, Points.Num(), SelectedDataPoints, [this, &Points, DataIndex, FirstDataIndex, LastDataIndex, DataIndexIncrement](int32 Index, FPCGPoint& OutPoint)
 		{
+			const FPCGPoint& Point = Points[Index];
+
 			// Discard point if it is already covered by a previous data
 			bool bPointToExclude = false;
 			for (int32 PreviousDataIndex = FirstDataIndex; PreviousDataIndex != DataIndex; PreviousDataIndex += DataIndexIncrement)
@@ -226,19 +231,25 @@ void UPCGUnionData::CreateSequentialPointData(UPCGPointData* PointData, bool bLe
 
 			if (bPointToExclude)
 			{
-				continue;
+				return false;
 			}
 
-			FPCGPoint& TargetPoint = TargetPoints.Add_GetRef(Point);
+			OutPoint = Point;
 
 			// Compute final density based on current & following data
 			for (int32 FollowingDataIndex = DataIndex + DataIndexIncrement; FollowingDataIndex != LastDataIndex; FollowingDataIndex += DataIndexIncrement)
 			{
-				if (PCGUnionDataMaths::UpdateDensity(TargetPoint.Density, Data[FollowingDataIndex]->GetDensityAtPosition(TargetPoint.Transform.GetLocation()), DensityFunction) == 1.0f)
+				if (PCGUnionDataMaths::UpdateDensity(OutPoint.Density, Data[FollowingDataIndex]->GetDensityAtPosition(OutPoint.Transform.GetLocation()), DensityFunction) == 1.0f)
 				{
 					break;
 				}
 			}
-		}
+
+			return true;
+		});
+
+		// Append current iteration results to target points
+		TargetPoints += SelectedDataPoints;
+		SelectedDataPoints.Reset();
 	}
 }
\ No newline at end of file
diff --git a/Engine/Plugins/Experimental/PCG/Source/PCG/Private/Data/PCGVolumeData.cpp b/Engine/Plugins/Experimental/PCG/Source/PCG/Private/Data/PCGVolumeData.cpp
--- a/Engine/Plugins/Experimental/PCG/Source/PCG/Private/Data/PCGVolumeData.cpp
+++ b/Engine/Plugins/Experimental/PCG/Source/PCG/Private/Data/PCGVolumeData.cpp
@@ -35,7 +35,7 @@ FBox UPCGVolumeData::GetStrictBounds() const
 	return StrictBounds;
 }
 
-const UPCGPointData* UPCGVolumeData::CreatePointData() const
+const UPCGPointData* UPCGVolumeData::CreatePointData(FPCGContextPtr Context) const
 {
 	UE_LOG(LogPCG, Error, TEXT("Volume data has no default point sampling"));
 	return nullptr;
diff --git a/Engine/Plugins/Experimental/PCG/Source/PCG/Private/Elements/PCGDebugElement.cpp b/Engine/Plugins/Experimental/PCG/Source/PCG/Private/Elements/PCGDebugElement.cpp
--- a/Engine/Plugins/Experimental/PCG/Source/PCG/Private/Elements/PCGDebugElement.cpp
+++ b/Engine/Plugins/Experimental/PCG/Source/PCG/Private/Elements/PCGDebugElement.cpp
@@ -60,7 +60,7 @@ namespace PCGDebugElement
 				continue;
 			}
 
-			const UPCGPointData* PointData = SpatialData->ToPointData();
+			const UPCGPointData* PointData = SpatialData->ToPointData(Context);
 
 			if (!PointData)
 			{
diff --git a/Engine/Plugins/Experimental/PCG/Source/PCG/Private/Elements/PCGDensityFilter.cpp b/Engine/Plugins/Experimental/PCG/Source/PCG/Private/Elements/PCGDensityFilter.cpp
--- a/Engine/Plugins/Experimental/PCG/Source/PCG/Private/Elements/PCGDensityFilter.cpp
+++ b/Engine/Plugins/Experimental/PCG/Source/PCG/Private/Elements/PCGDensityFilter.cpp
@@ -2,6 +2,7 @@
 
 #include "Elements/PCGDensityFilter.h"
 #include "Data/PCGPointData.h"
+#include "Helpers/PCGAsync.h"
 
 FPCGElementPtr UPCGDensityFilterSettings::CreateElement() const
 {
@@ -58,7 +59,7 @@ bool FPCGDensityFilterElement::ExecuteInternal(FPCGContextPtr Context) const
 			continue;
 		}
 
-		const UPCGPointData* OriginalData = Cast<UPCGSpatialData>(Input.Data)->ToPointData();
+		const UPCGPointData* OriginalData = Cast<UPCGSpatialData>(Input.Data)->ToPointData(Context);
 
 		if (!OriginalData)
 		{
@@ -74,21 +75,29 @@ bool FPCGDensityFilterElement::ExecuteInternal(FPCGContextPtr Context) const
 
 		Output.Data = FilteredData;
 
-		for (const FPCGPoint& Point : Points)
+		FPCGAsync::AsyncPointProcessing(Context, Points.Num(), FilteredPoints, [&Points, Settings, MinBound, MaxBound](int32 Index, FPCGPoint& OutPoint)
 		{
+			const FPCGPoint& Point = Points[Index];
+
 			bool bInRange = (Point.Density >= MinBound && Point.Density <= MaxBound);
 			if (bInRange != Settings->bInvertFilter)
 			{
-				FilteredPoints.Add(Point);
+				OutPoint = Point;
+				return true;
 			}
 #if WITH_EDITORONLY_DATA
 			else if (Settings->bKeepZeroDensityPoints)
 			{
-				FPCGPoint& FilteredPoint = FilteredPoints.Add_GetRef(Point);
-				FilteredPoint.Density = 0;
+				OutPoint = Point;
+				OutPoint.Density = 0;
+				return true;
 			}
 #endif
-		}
+			else
+			{
+				return false;
+			}
+		});
 
 		PCGE_LOG(Verbose, "Generated %d points out of %d source points", FilteredPoints.Num(), Points.Num());
 	}
diff --git a/Engine/Plugins/Experimental/PCG/Source/PCG/Private/Elements/PCGDensityRemap.cpp b/Engine/Plugins/Experimental/PCG/Source/PCG/Private/Elements/PCGDensityRemap.cpp
--- a/Engine/Plugins/Experimental/PCG/Source/PCG/Private/Elements/PCGDensityRemap.cpp
+++ b/Engine/Plugins/Experimental/PCG/Source/PCG/Private/Elements/PCGDensityRemap.cpp
@@ -3,6 +3,7 @@
 #include "Elements/PCGDensityRemap.h"
 #include "PCGHelpers.h"
 #include "Data/PCGPointData.h"
+#include "Helpers/PCGAsync.h"
 #include "Math/RandomStream.h"
 
 FPCGElementPtr UPCGLinearDensityRemapSettings::CreateElement() const
@@ -48,7 +49,7 @@ bool FPCGLinearDensityRemapElement::ExecuteInternal(FPCGContextPtr Context) cons
 			continue;
 		}
 
-		const UPCGPointData* OriginalData = Cast<UPCGSpatialData>(Input.Data)->ToPointData();
+		const UPCGPointData* OriginalData = Cast<UPCGSpatialData>(Input.Data)->ToPointData(Context);
 
 		if (!OriginalData)
 		{
@@ -64,22 +65,25 @@ bool FPCGLinearDensityRemapElement::ExecuteInternal(FPCGContextPtr Context) cons
 		TArray<FPCGPoint>& SampledPoints = SampledData->GetMutablePoints();
 		Output.Data = SampledData;
 
-		SampledPoints = Points;
 		if (Settings->bMultiplyDensity)
 		{
-			for (FPCGPoint& Point : SampledPoints)
+			FPCGAsync::AsyncPointProcessing(Context, OriginalPointCount, SampledPoints, [&Points, Settings, RemapMin, RemapMax](int32 Index, FPCGPoint& OutPoint)
 			{
-				FRandomStream RandomSource(PCGHelpers::ComputeSeed(Settings->Seed, Point.Seed));
-				Point.Density *= RandomSource.FRandRange(RemapMin, RemapMax);
-			}
+				OutPoint = Points[Index];
+				FRandomStream RandomSource(PCGHelpers::ComputeSeed(Settings->Seed, OutPoint.Seed));
+				OutPoint.Density *= RandomSource.FRandRange(RemapMin, RemapMax);
+				return true;
+			});
 		}
 		else
 		{
-			for (FPCGPoint& Point : SampledPoints)
+			FPCGAsync::AsyncPointProcessing(Context, OriginalPointCount, SampledPoints, [&Points, Settings, RemapMin, RemapMax](int32 Index, FPCGPoint& OutPoint)
 			{
-				FRandomStream RandomSource(PCGHelpers::ComputeSeed(Settings->Seed, Point.Seed));
-				Point.Density = RandomSource.FRandRange(RemapMin, RemapMax);
-			}
+				OutPoint = Points[Index];
+				FRandomStream RandomSource(PCGHelpers::ComputeSeed(Settings->Seed, OutPoint.Seed));
+				OutPoint.Density  = RandomSource.FRandRange(RemapMin, RemapMax);
+				return true;
+			});
 		}
 	}
 
diff --git a/Engine/Plugins/Experimental/PCG/Source/PCG/Private/Elements/PCGPointSampler.cpp b/Engine/Plugins/Experimental/PCG/Source/PCG/Private/Elements/PCGPointSampler.cpp
--- a/Engine/Plugins/Experimental/PCG/Source/PCG/Private/Elements/PCGPointSampler.cpp
+++ b/Engine/Plugins/Experimental/PCG/Source/PCG/Private/Elements/PCGPointSampler.cpp
@@ -3,6 +3,7 @@
 #include "Elements/PCGPointSampler.h"
 #include "PCGHelpers.h"
 #include "Data/PCGPointData.h"
+#include "Helpers/PCGAsync.h"
 #include "Math/RandomStream.h"
 
 FPCGElementPtr UPCGPointSamplerSettings::CreateElement() const
@@ -56,7 +57,7 @@ bool FPCGPointSamplerElement::ExecuteInternal(FPCGContextPtr Context) const
 			continue;
 		}
 
-		const UPCGPointData* OriginalData = Cast<UPCGSpatialData>(Input.Data)->ToPointData();
+		const UPCGPointData* OriginalData = Cast<UPCGSpatialData>(Input.Data)->ToPointData(Context);
 
 		if (!OriginalData)
 		{
@@ -90,10 +91,7 @@ bool FPCGPointSamplerElement::ExecuteInternal(FPCGContextPtr Context) const
 		{
 			TRACE_CPUPROFILER_EVENT_SCOPE(FPCGPointSamplerElement::Execute::SelectPoints);
 
-			// Approximate upper bound
-			SampledPoints.Reserve(FMath::Min(OriginalPointCount, 4 * TargetNumPoints / 3));
-
-			for (int Index = 0; Index < OriginalPointCount; ++Index)
+			FPCGAsync::AsyncPointProcessing(Context, OriginalPointCount, SampledPoints, [&Points, Settings](int32 Index, FPCGPoint& OutPoint)
 			{
 				const FPCGPoint& Point = Points[Index];
 
@@ -103,16 +101,22 @@ bool FPCGPointSamplerElement::ExecuteInternal(FPCGContextPtr Context) const
 
 				if (Chance < Settings->Ratio)
 				{
-					SampledPoints.Add(Point);
+					OutPoint = Point;
+					return true;
 				}
 #if WITH_EDITORONLY_DATA
 				else if (Settings->bKeepZeroDensityPoints)
 				{
-					FPCGPoint& SampledPoint = SampledPoints.Add_GetRef(Point);
-					SampledPoint.Density = 0;
+					OutPoint = Point;
+					OutPoint.Density = 0;
+					return true;
 				}
 #endif
-			}
+				else
+				{
+					return false;
+				}
+			});
 
 			PCGE_LOG(Verbose, "Generated %d points from %d source points", SampledPoints.Num(), OriginalPointCount);
 		}
diff --git a/Engine/Plugins/Experimental/PCG/Source/PCG/Private/Elements/PCGSelfPruning.cpp b/Engine/Plugins/Experimental/PCG/Source/PCG/Private/Elements/PCGSelfPruning.cpp
--- a/Engine/Plugins/Experimental/PCG/Source/PCG/Private/Elements/PCGSelfPruning.cpp
+++ b/Engine/Plugins/Experimental/PCG/Source/PCG/Private/Elements/PCGSelfPruning.cpp
@@ -74,7 +74,7 @@ bool FPCGSelfPruningElement::ExecuteInternal(FPCGContextPtr Context) const
 			continue;
 		}
 
-		const UPCGPointData* InputPointData = SpatialInput->ToPointData();
+		const UPCGPointData* InputPointData = SpatialInput->ToPointData(Context);
 		const TArray<FPCGPoint>& Points = InputPointData->GetPoints();
 		const UPCGPointData::PointOctree& Octree = InputPointData->GetOctree();
 
diff --git a/Engine/Plugins/Experimental/PCG/Source/PCG/Private/Elements/PCGSpawnActor.cpp b/Engine/Plugins/Experimental/PCG/Source/PCG/Private/Elements/PCGSpawnActor.cpp
--- a/Engine/Plugins/Experimental/PCG/Source/PCG/Private/Elements/PCGSpawnActor.cpp
+++ b/Engine/Plugins/Experimental/PCG/Source/PCG/Private/Elements/PCGSpawnActor.cpp
@@ -103,7 +103,7 @@ bool FPCGSpawnActorElement::ExecuteInternal(FPCGContextPtr Context) const
 		}
 
 		// First, create target instance transforms
-		const UPCGPointData* PointData = SpatialData->ToPointData();
+		const UPCGPointData* PointData = SpatialData->ToPointData(Context);
 
 		if (!PointData)
 		{
diff --git a/Engine/Plugins/Experimental/PCG/Source/PCG/Private/Elements/PCGStaticMeshSpawner.cpp b/Engine/Plugins/Experimental/PCG/Source/PCG/Private/Elements/PCGStaticMeshSpawner.cpp
--- a/Engine/Plugins/Experimental/PCG/Source/PCG/Private/Elements/PCGStaticMeshSpawner.cpp
+++ b/Engine/Plugins/Experimental/PCG/Source/PCG/Private/Elements/PCGStaticMeshSpawner.cpp
@@ -71,7 +71,7 @@ bool FPCGStaticMeshSpawnerElement::ExecuteInternal(FPCGContextPtr Context) const
 		}
 
 		// First, create target instance transforms
-		const UPCGPointData* PointData = SpatialData->ToPointData();
+		const UPCGPointData* PointData = SpatialData->ToPointData(Context);
 
 		if (!PointData)
 		{
diff --git a/Engine/Plugins/Experimental/PCG/Source/PCG/Private/Elements/PCGSurfaceSampler.cpp b/Engine/Plugins/Experimental/PCG/Source/PCG/Private/Elements/PCGSurfaceSampler.cpp
--- a/Engine/Plugins/Experimental/PCG/Source/PCG/Private/Elements/PCGSurfaceSampler.cpp
+++ b/Engine/Plugins/Experimental/PCG/Source/PCG/Private/Elements/PCGSurfaceSampler.cpp
@@ -3,9 +3,84 @@
 #include "Elements/PCGSurfaceSampler.h"
 #include "Data/PCGPointData.h"
 #include "Data/PCGSpatialData.h"
+#include "Helpers/PCGAsync.h"
 #include "PCGHelpers.h"
 #include "Math/RandomStream.h"
 
+struct FPCGSurfaceSamplerLoopData
+{
+	const UPCGSurfaceSamplerSettings* Settings;
+	
+	FVector::FReal PointRadius;
+	FVector::FReal InterstitialDistance;
+	FVector::FReal InnerCellSize;
+	FVector::FReal CellSize;
+
+	int32 CellMinX;
+	int32 CellMaxX;
+	int32 CellMinY;
+	int32 CellMaxY;
+
+	int32 CellCount;
+
+	int64 TargetPointCount;
+
+	float Ratio;
+
+	bool Initialize(const UPCGSurfaceSamplerSettings* InSettings, FPCGContextPtr Context, const FBox& InputBounds)
+	{
+		Settings = InSettings;
+
+		// Conceptually, we will break down the surface bounds in a N x M grid
+		InterstitialDistance = 2 * Settings->PointRadius;
+		InnerCellSize = 2 * Settings->Looseness * Settings->PointRadius;
+		CellSize = InterstitialDistance + InnerCellSize;
+
+		check(CellSize > 0);
+
+		// By using scaled indices in the world, we can easily make this process deterministic
+		CellMinX = FMath::CeilToInt((InputBounds.Min.X) / CellSize);
+		CellMaxX = FMath::FloorToInt((InputBounds.Max.X) / CellSize);
+		CellMinY = FMath::CeilToInt((InputBounds.Min.Y) / CellSize);
+		CellMaxY = FMath::FloorToInt((InputBounds.Max.Y) / CellSize);
+
+		if (CellMinX > CellMaxX || CellMinY > CellMaxY)
+		{
+			PCGE_LOG_C(Verbose, Context, "Skipped - invalid cell bounds");
+			return false;
+		}
+
+		CellCount = (1 + CellMaxX - CellMinX) * (1 + CellMaxY - CellMinY);
+		check(CellCount > 0);
+
+		const FVector::FReal InvSquaredMeterUnits = 1.0 / (100.0 * 100.0);
+		TargetPointCount = (InputBounds.Max.X - InputBounds.Min.X) * (InputBounds.Max.Y - InputBounds.Min.Y) * Settings->PointsPerSquaredMeter * InvSquaredMeterUnits;
+
+		if (TargetPointCount == 0)
+		{
+			PCGE_LOG_C(Verbose, Context, "Skipped - density yields no points");
+			return false;
+		}
+		else if (TargetPointCount > CellCount)
+		{
+			TargetPointCount = CellCount;
+		}
+
+		Ratio = TargetPointCount / (FVector::FReal)CellCount;
+
+		return true;
+	}
+
+	void ComputeCellIndices(int32 Index, int32& CellX, int32& CellY)
+	{
+		check(Index >= 0 && Index < CellCount);
+		const int32 CellCountX = CellMaxX - CellMinX;
+
+		CellX = CellMinX + (Index % CellCountX);
+		CellY = CellMinY + (Index / CellCountX);
+	}
+};
+
 FPCGElementPtr UPCGSurfaceSamplerSettings::CreateElement() const
 {
 	return MakeShared<FPCGSurfaceSamplerElement>();
@@ -48,46 +123,12 @@ bool FPCGSurfaceSamplerElement::ExecuteInternal(FPCGContextPtr Context) const
 			continue;
 		}
 
-		// Conceptually, we will break down the surface bounds in a N x M grid
-		const float PointSteepness = Settings->PointSteepness;
-		const bool bApplyDensity = Settings->bApplyDensityToPoints;
-		const FVector::FReal PointRadius = Settings->PointRadius;
-		const FVector::FReal InterstitialDistance = 2 * PointRadius;
-		const FVector::FReal InnerCellSize = 2 * Settings->Looseness * PointRadius;
-		const FVector::FReal CellSize = InterstitialDistance + InnerCellSize;
-		
-		check(CellSize > 0);
-
-		// By using scaled indices in the world, we can easily make this process deterministic
-		const int32 CellMinX = FMath::CeilToInt((InputBounds.Min.X) / CellSize);
-		const int32 CellMaxX = FMath::FloorToInt((InputBounds.Max.X) / CellSize);
-		const int32 CellMinY = FMath::CeilToInt((InputBounds.Min.Y) / CellSize);
-		const int32 CellMaxY = FMath::FloorToInt((InputBounds.Max.Y) / CellSize);
-
-		if (CellMinX > CellMaxX || CellMinY > CellMaxY)
+		FPCGSurfaceSamplerLoopData LoopData;
+		if (!LoopData.Initialize(Settings, Context, InputBounds))
 		{
-			PCGE_LOG(Verbose, "Skipped - invalid cell bounds");
 			continue;
 		}
 
-		const int32 CellCount = (1 + CellMaxX - CellMinX) * (1 + CellMaxY - CellMinY);
-		check(CellCount > 0);
-
-		const FVector::FReal InvSquaredMeterUnits = 1.0 / (100.0 * 100.0);
-		int64 TargetPointCount = (InputBounds.Max.X - InputBounds.Min.X) * (InputBounds.Max.Y - InputBounds.Min.Y) * Settings->PointsPerSquaredMeter * InvSquaredMeterUnits;
-
-		if (TargetPointCount == 0)
-		{
-			PCGE_LOG(Verbose, "Skipped - density yields no points");
-			continue;
-		}
-		else if (TargetPointCount > CellCount)
-		{
-			TargetPointCount = CellCount;
-		}
-
-		const FVector::FReal Ratio = TargetPointCount / (FVector::FReal)CellCount;
-
 		// Finally, create data
 		FPCGTaggedData& Output = Outputs.Emplace_GetRef();
 		Output = Input;
@@ -98,51 +139,47 @@ bool FPCGSurfaceSamplerElement::ExecuteInternal(FPCGContextPtr Context) const
 
 		TArray<FPCGPoint>& SampledPoints = SampledData->GetMutablePoints();
 
-		const FVector::FReal CellStartX = CellMinX * CellSize;
-		const FVector::FReal CellStartY = CellMinY * CellSize;
+		FPCGAsync::AsyncPointProcessing(Context, LoopData.CellCount, SampledPoints, [&LoopData, SpatialInput](int32 Index, FPCGPoint& OutPoint)
+		{
+			int32 CellX;
+			int32 CellY;
+			LoopData.ComputeCellIndices(Index, CellX, CellY);
+
+			const FVector::FReal CurrentX = CellX * LoopData.CellSize;
+			const FVector::FReal CurrentY = CellY * LoopData.CellSize;
+			const FVector::FReal InnerCellSize = LoopData.InnerCellSize;
 
-		FVector::FReal CurrentX = CellStartX;
+			FRandomStream RandomSource(PCGHelpers::ComputeSeed(LoopData.Settings->Seed, CellX, CellY));
+			float Chance = RandomSource.FRand();
 
-		for(int32 CellX = CellMinX; CellX <= CellMaxX; ++CellX)
-		{
-			FVector::FReal CurrentY = CellStartY;
+			const float Ratio = LoopData.Ratio;
 
-			for(int32 CellY = CellMinY; CellY <= CellMaxY; ++CellY)
+			if (Chance < Ratio)
 			{
-				FRandomStream RandomSource(PCGHelpers::ComputeSeed(Settings->Seed, CellX, CellY));
-				float Chance = RandomSource.FRand();
+				const float RandX = RandomSource.FRand();
+				const float RandY = RandomSource.FRand();
 
-				if (Chance < Ratio)
-				{
-					const float RandX = RandomSource.FRand();
-					const float RandY = RandomSource.FRand();
+				OutPoint.Transform = FTransform(FVector(CurrentX + RandX * InnerCellSize, CurrentY + RandY * InnerCellSize, 0));
+				OutPoint.Extents = FVector(LoopData.Settings->PointRadius);
+				OutPoint.Density = LoopData.Settings->bApplyDensityToPoints ? ((Ratio - Chance) / Ratio) : 1.0f;
+				OutPoint.Steepness = LoopData.Settings->PointSteepness;
+				OutPoint.Seed = RandomSource.GetCurrentSeed();
 
-					FPCGPoint Point;
-					Point.Transform = FTransform(FVector(CurrentX + RandX * InnerCellSize, CurrentY + RandY * InnerCellSize, 0));
-					Point.Extents = FVector(PointRadius);
-					Point.Density = bApplyDensity ? ((Ratio - Chance) / Ratio) : 1.0f;
-					Point.Steepness = PointSteepness;
-					Point.Seed = RandomSource.GetCurrentSeed();
-
-					Point = SpatialInput->TransformPoint(Point);
+				OutPoint = SpatialInput->TransformPoint(OutPoint);
 
 #if WITH_EDITORONLY_DATA
-					if(Point.Density > 0 || Settings->bKeepZeroDensityPoints)
+				return OutPoint.Density > 0 || LoopData.Settings->bKeepZeroDensityPoints;
 #else
-					if (Point.Density > 0)
+				return OutPoint.Density > 0;
 #endif
-					{
-						SampledPoints.Add(Point);
-					}
-				}
-
-				CurrentY += CellSize;
 			}
+			else
+			{
+				return false;
+			}
+		});
 
-			CurrentX += CellSize;
-		}
-
-		PCGE_LOG(Verbose, "Generated %d points in %d cells", SampledPoints.Num(), CellCount);
+		PCGE_LOG(Verbose, "Generated %d points in %d cells", SampledPoints.Num(), LoopData.CellCount);
 	}
 
 	// Finally, forward any exclusions/settings
diff --git a/Engine/Plugins/Experimental/PCG/Source/PCG/Private/Graph/PCGGraphExecutor.cpp b/Engine/Plugins/Experimental/PCG/Source/PCG/Private/Graph/PCGGraphExecutor.cpp
--- a/Engine/Plugins/Experimental/PCG/Source/PCG/Private/Graph/PCGGraphExecutor.cpp
+++ b/Engine/Plugins/Experimental/PCG/Source/PCG/Private/Graph/PCGGraphExecutor.cpp
@@ -207,6 +207,10 @@ void FPCGGraphExecutor::Execute()
 		{
 			FPCGGraphActiveTask& ActiveTask = ActiveTasks[ExecutionIndex];
 
+			// Reassign resources to the context
+			// TODO: change this when we support MT in the graph executor
+			ActiveTask.Context->NumAvailableTasks = FMath::Max(FPlatformMisc::NumberOfCoresIncludingHyperthreads() - 2, 2);
+
 			if (!ActiveTask.Context->bIsPaused && ActiveTask.Element->Execute(ActiveTask.Context))
 			{
 				// Store output data in map
diff --git a/Engine/Plugins/Experimental/PCG/Source/PCG/Private/Helpers/PCGAsync.cpp b/Engine/Plugins/Experimental/PCG/Source/PCG/Private/Helpers/PCGAsync.cpp
new file mode 100644
--- /dev/null
+++ b/Engine/Plugins/Experimental/PCG/Source/PCG/Private/Helpers/PCGAsync.cpp
@@ -0,0 +1,84 @@
+// Copyright Epic Games, Inc. All Rights Reserved.
+
+#include "Helpers/PCGAsync.h"
+#include "PCGPoint.h"
+#include "Async/Async.h"
+
+namespace FPCGAsync
+{
+	void AsyncPointProcessing(FPCGContextPtr Context, int32 NumIterations, TArray<FPCGPoint>& OutPoints, const TFunction<bool(int32, FPCGPoint&)>& PointFunc)
+	{
+		const int32 MinIterationsPerTask = 256;
+		AsyncPointProcessing(Context.IsValid() ? Context->NumAvailableTasks : 1, MinIterationsPerTask, NumIterations, OutPoints, PointFunc);
+	}
+
+	void AsyncPointProcessing(int32 NumAvailableTasks, int32 MinIterationsPerTask, int32 NumIterations, TArray<FPCGPoint>& OutPoints, const TFunction<bool(int32, FPCGPoint&)>& PointFunc)
+	{
+		TRACE_CPUPROFILER_EVENT_SCOPE(IPCGElement::AsyncPointProcessing);
+		check(NumAvailableTasks > 0 && MinIterationsPerTask > 0 && NumIterations >= 0);
+		// Get number of available threads from the context
+		const int32 NumTasks = FMath::Min(NumAvailableTasks, FMath::Max(1, NumIterations / MinIterationsPerTask));
+		const int32 IterationsPerTask = NumIterations / NumTasks;
+
+		{
+			TRACE_CPUPROFILER_EVENT_SCOPE(IPCGElement::AsyncPointProcessing::AllocatingArray);
+			// Pre-reserve the out points array
+			OutPoints.SetNum(NumIterations);
+		}
+
+		// Setup [current, last, nb points] data per dispatch
+		// Execute
+		TArray<TFuture<int32>> AsyncTasks;
+		AsyncTasks.Reserve(NumTasks);
+
+		// Launch the async tasks
+		for (int32 TaskIndex = 0; TaskIndex < NumTasks; ++TaskIndex)
+		{
+			const int32 StartIndex = TaskIndex * IterationsPerTask;
+			const int32 EndIndex = (TaskIndex == NumTasks - 1 ? NumIterations : StartIndex + IterationsPerTask);
+
+			AsyncTasks.Emplace(Async(EAsyncExecution::ThreadPool, [&PointFunc, StartIndex, EndIndex, &OutPoints]()
+			{
+				TRACE_CPUPROFILER_EVENT_SCOPE(IPCGElement::AsyncPointProcessing::InnerLoop);
+				int32 NumPointsWritten = 0;
+
+				for (int32 Index = StartIndex; Index < EndIndex; ++Index)
+				{
+					if (PointFunc(Index, OutPoints[StartIndex + NumPointsWritten]))
+					{
+						++NumPointsWritten;
+					}
+				}
+
+				return NumPointsWritten;
+			}));
+		}
+
+		// Wait/Gather results & collapse points
+		{
+			TRACE_CPUPROFILER_EVENT_SCOPE(IPCGElement::AsyncPointProcessing::WaitAndCollapseArray);
+			int RangeIndex = 0;
+			for (int32 AsyncIndex = 0; AsyncIndex < AsyncTasks.Num(); ++AsyncIndex)
+			{
+				const int32 StartPointsIndex = AsyncIndex * IterationsPerTask;
+
+				TFuture<int32>& AsyncTask = AsyncTasks[AsyncIndex];
+				AsyncTask.Wait();
+				int32 NumberOfPointsAdded = AsyncTask.Get();
+
+				// Move points from [StartsPointIndex, StartsPointIndex + NumberPointsAdded] to [RangeIndex, RangeIndex + NumberPointsAdded]
+				if (StartPointsIndex != RangeIndex)
+				{
+					for (int32 MoveIndex = 0; MoveIndex < NumberOfPointsAdded; ++MoveIndex)
+					{
+						OutPoints[RangeIndex + MoveIndex] = MoveTemp(OutPoints[StartPointsIndex + MoveIndex]);
+					}
+				}
+
+				RangeIndex += NumberOfPointsAdded;
+			}
+
+			OutPoints.SetNum(RangeIndex);
+		}
+	}
+}
\ No newline at end of file
diff --git a/Engine/Plugins/Experimental/PCG/Source/PCG/Public/Data/PCGDifferenceData.h b/Engine/Plugins/Experimental/PCG/Source/PCG/Public/Data/PCGDifferenceData.h
--- a/Engine/Plugins/Experimental/PCG/Source/PCG/Public/Data/PCGDifferenceData.h
+++ b/Engine/Plugins/Experimental/PCG/Source/PCG/Public/Data/PCGDifferenceData.h
@@ -45,7 +45,7 @@ public:
 	//~End UPCGSpatialData interface
 
 	//~Begin UPCGSpatialDataWithPointCache interface
-	virtual const UPCGPointData* CreatePointData() const override;
+	virtual const UPCGPointData* CreatePointData(FPCGContextPtr Context) const override;
 	//~End UPCGSpatialDataWithPointCache interface
 protected:
 	UPROPERTY(BlueprintReadOnly, VisibleAnywhere, Category = SpatialData)
diff --git a/Engine/Plugins/Experimental/PCG/Source/PCG/Public/Data/PCGIntersectionData.h b/Engine/Plugins/Experimental/PCG/Source/PCG/Public/Data/PCGIntersectionData.h
--- a/Engine/Plugins/Experimental/PCG/Source/PCG/Public/Data/PCGIntersectionData.h
+++ b/Engine/Plugins/Experimental/PCG/Source/PCG/Public/Data/PCGIntersectionData.h
@@ -35,7 +35,7 @@ public:
 	//~End UPCGSpatialData interface
 
 	//~Begin UPCGSpatialDataWithPointCache interface
-	virtual const UPCGPointData* CreatePointData() const override;
+	virtual const UPCGPointData* CreatePointData(FPCGContextPtr Context) const override;
 	//~End UPCGSpatialDataWithPointCache interface
 
 public:
@@ -43,7 +43,7 @@ public:
 	EPCGIntersectionDensityFunction DensityFunction = EPCGIntersectionDensityFunction::Multiply;
 
 protected:
-	UPCGPointData* CreateAndFilterPointData(const UPCGSpatialData* X, const UPCGSpatialData* Y) const;
+	UPCGPointData* CreateAndFilterPointData(FPCGContextPtr Context, const UPCGSpatialData* X, const UPCGSpatialData* Y) const;
 
 	UPROPERTY(BlueprintReadOnly, VisibleAnywhere, Category = SpatialData)
 	TObjectPtr<const UPCGSpatialData> A = nullptr;
diff --git a/Engine/Plugins/Experimental/PCG/Source/PCG/Public/Data/PCGLandscapeData.h b/Engine/Plugins/Experimental/PCG/Source/PCG/Public/Data/PCGLandscapeData.h
--- a/Engine/Plugins/Experimental/PCG/Source/PCG/Public/Data/PCGLandscapeData.h
+++ b/Engine/Plugins/Experimental/PCG/Source/PCG/Public/Data/PCGLandscapeData.h
@@ -47,7 +47,7 @@ public:
 	// ~End UPGCConcreteData interface
 
 	// ~Begin UPCGSpatialDataWithPointCache interface
-	virtual const UPCGPointData* CreatePointData() const override;
+	virtual const UPCGPointData* CreatePointData(FPCGContextPtr Context) const override;
 	// ~End UPCGConcreteDataWithPointCache interface
 
 	// TODO: add on property changed to clear cached data
diff --git a/Engine/Plugins/Experimental/PCG/Source/PCG/Public/Data/PCGLandscapeSplineData.h b/Engine/Plugins/Experimental/PCG/Source/PCG/Public/Data/PCGLandscapeSplineData.h
--- a/Engine/Plugins/Experimental/PCG/Source/PCG/Public/Data/PCGLandscapeSplineData.h
+++ b/Engine/Plugins/Experimental/PCG/Source/PCG/Public/Data/PCGLandscapeSplineData.h
@@ -23,7 +23,7 @@ public:
 	//~End UPCGPolyLineData interface
 
 	//~Begin UPCGSpatialDataWithPointCache interface
-	virtual const UPCGPointData* CreatePointData() const override;
+	virtual const UPCGPointData* CreatePointData(FPCGContextPtr Context) const override;
 	//~End UPCGSpatialDataWithPointCache interface
 
 	//~Begin UPCGSpatialData interface
diff --git a/Engine/Plugins/Experimental/PCG/Source/PCG/Public/Data/PCGPointData.h b/Engine/Plugins/Experimental/PCG/Source/PCG/Public/Data/PCGPointData.h
--- a/Engine/Plugins/Experimental/PCG/Source/PCG/Public/Data/PCGPointData.h
+++ b/Engine/Plugins/Experimental/PCG/Source/PCG/Public/Data/PCGPointData.h
@@ -65,7 +65,7 @@ public:
 	virtual int GetDimension() const override { return 0; }
 	virtual FBox GetBounds() const override;
 	virtual float GetDensityAtPosition(const FVector& InPosition) const override;
-	virtual const UPCGPointData* ToPointData() const { return this; }
+	virtual const UPCGPointData* ToPointData(FPCGContextPtr Context) const { return this; }
 	virtual FPCGPoint TransformPoint(const FPCGPoint& InPoint) const;
 	// ~End UPCGSpatialData interface
 
diff --git a/Engine/Plugins/Experimental/PCG/Source/PCG/Public/Data/PCGProjectionData.h b/Engine/Plugins/Experimental/PCG/Source/PCG/Public/Data/PCGProjectionData.h
--- a/Engine/Plugins/Experimental/PCG/Source/PCG/Public/Data/PCGProjectionData.h
+++ b/Engine/Plugins/Experimental/PCG/Source/PCG/Public/Data/PCGProjectionData.h
@@ -29,7 +29,7 @@ public:
 	//~End UPCGSpatialData interface
 
 	//~Begin UPCGSpatialDataWithPointCache interface
-	virtual const UPCGPointData* CreatePointData() const override;
+	virtual const UPCGPointData* CreatePointData(FPCGContextPtr Context) const override;
 	//~End UPCGSpatialDataWithPointCache interface
 
 protected:
diff --git a/Engine/Plugins/Experimental/PCG/Source/PCG/Public/Data/PCGSpatialData.h b/Engine/Plugins/Experimental/PCG/Source/PCG/Public/Data/PCGSpatialData.h
--- a/Engine/Plugins/Experimental/PCG/Source/PCG/Public/Data/PCGSpatialData.h
+++ b/Engine/Plugins/Experimental/PCG/Source/PCG/Public/Data/PCGSpatialData.h
@@ -14,6 +14,9 @@ class UPCGDifferenceData;
 class UPCGProjectionData;
 class AActor;
 
+struct FPCGContext;
+typedef TSharedPtr<FPCGContext, ESPMode::ThreadSafe> FPCGContextPtr;
+
 /**
 * "Concrete" data base class for PCG generation
 * This will be the base class for data classes that actually represent
@@ -51,7 +54,9 @@ public:
 
 	/** Discretizes the data into points */
 	UFUNCTION(BlueprintCallable, Category = SpatialData)
-	virtual const UPCGPointData* ToPointData() const PURE_VIRTUAL(UPCGSpatialData::ToPointData, return nullptr;);
+	virtual const UPCGPointData* ToPointData() const { return ToPointData(nullptr); }
+
+	virtual const UPCGPointData* ToPointData(FPCGContextPtr Context) const PURE_VIRTUAL(UPCGSpatialData::ToPointData, return nullptr;);
 
 	/** Transform a world-space position to a world-space position in relation to the current data. (ex: projection on surface) */
 	UFUNCTION(BlueprintCallable, Category = SpatialData)
@@ -96,11 +101,11 @@ class UPCGSpatialDataWithPointCache : public UPCGSpatialData
 
 public:
 	// ~UPCGSpatialData implementation
-	virtual const UPCGPointData* ToPointData() const override;
+	virtual const UPCGPointData* ToPointData(FPCGContextPtr Context) const override;
 	// ~End UPCGSpatialData implementation
 
 protected:
-	virtual const UPCGPointData* CreatePointData() const PURE_VIRTUAL(UPCGSpatialData::CreatePointData, return nullptr;);
+	virtual const UPCGPointData* CreatePointData(FPCGContextPtr Context) const PURE_VIRTUAL(UPCGSpatialData::CreatePointData, return nullptr;);
 
 private:
 	UPROPERTY(Transient)
diff --git a/Engine/Plugins/Experimental/PCG/Source/PCG/Public/Data/PCGSplineData.h b/Engine/Plugins/Experimental/PCG/Source/PCG/Public/Data/PCGSplineData.h
--- a/Engine/Plugins/Experimental/PCG/Source/PCG/Public/Data/PCGSplineData.h
+++ b/Engine/Plugins/Experimental/PCG/Source/PCG/Public/Data/PCGSplineData.h
@@ -25,7 +25,7 @@ public:
 	//~End UPCGPolyLineData interface
 
 	//~Begin UPCGSpatialDataWithPointCache interface
-	virtual const UPCGPointData* CreatePointData() const override;
+	virtual const UPCGPointData* CreatePointData(FPCGContextPtr Context) const override;
 	//~End UPCGSpatialDataWithPointCache interface
 
 	//~Begin UPCGSpatialData interface
diff --git a/Engine/Plugins/Experimental/PCG/Source/PCG/Public/Data/PCGTextureData.h b/Engine/Plugins/Experimental/PCG/Source/PCG/Public/Data/PCGTextureData.h
--- a/Engine/Plugins/Experimental/PCG/Source/PCG/Public/Data/PCGTextureData.h
+++ b/Engine/Plugins/Experimental/PCG/Source/PCG/Public/Data/PCGTextureData.h
@@ -30,7 +30,7 @@ public:
 	//~End UPCGSpatialData interface
 
 	//~Begin UPCGSpatialDataWithPointCache interface
-	virtual const UPCGPointData* CreatePointData() const override;
+	virtual const UPCGPointData* CreatePointData(FPCGContextPtr Context) const override;
 	//~End UPCGSpatialDataWithPointCache interface
 
 protected:
diff --git a/Engine/Plugins/Experimental/PCG/Source/PCG/Public/Data/PCGUnionData.h b/Engine/Plugins/Experimental/PCG/Source/PCG/Public/Data/PCGUnionData.h
--- a/Engine/Plugins/Experimental/PCG/Source/PCG/Public/Data/PCGUnionData.h
+++ b/Engine/Plugins/Experimental/PCG/Source/PCG/Public/Data/PCGUnionData.h
@@ -47,7 +47,7 @@ public:
 	//~End UPCGSpatialData interface
 
 	//~Begin UPCGSpatialDataWithPointCache interface
-	virtual const UPCGPointData* CreatePointData() const override;
+	virtual const UPCGPointData* CreatePointData(FPCGContextPtr Context) const override;
 	//~End UPCGSpatialDataWithPointCache interface
 
 protected:
@@ -73,5 +73,5 @@ protected:
 	int CachedDimension = 0;
 
 private:
-	void CreateSequentialPointData(UPCGPointData* PointData, bool bLeftToRight) const;
+	void CreateSequentialPointData(FPCGContextPtr Context, UPCGPointData* PointData, bool bLeftToRight) const;
 };
\ No newline at end of file
diff --git a/Engine/Plugins/Experimental/PCG/Source/PCG/Public/Data/PCGVolumeData.h b/Engine/Plugins/Experimental/PCG/Source/PCG/Public/Data/PCGVolumeData.h
--- a/Engine/Plugins/Experimental/PCG/Source/PCG/Public/Data/PCGVolumeData.h
+++ b/Engine/Plugins/Experimental/PCG/Source/PCG/Public/Data/PCGVolumeData.h
@@ -26,7 +26,7 @@ public:
 	// ~End UPGCSpatialData interface
 
 	// ~Begin UPCGSpatialDataWithPointCache interface
-	virtual const UPCGPointData* CreatePointData() const override;
+	virtual const UPCGPointData* CreatePointData(FPCGContextPtr Context) const override;
 	// ~End UPCGSpatialDataWithPointCache interface
 
 protected:
diff --git a/Engine/Plugins/Experimental/PCG/Source/PCG/Public/Helpers/PCGAsync.h b/Engine/Plugins/Experimental/PCG/Source/PCG/Public/Helpers/PCGAsync.h
new file mode 100644
--- /dev/null
+++ b/Engine/Plugins/Experimental/PCG/Source/PCG/Public/Helpers/PCGAsync.h
@@ -0,0 +1,30 @@
+// Copyright Epic Games, Inc. All Rights Reserved.
+
+#pragma once
+
+#include "PCGContext.h"
+#include "Async/Async.h"
+
+struct FPCGPoint;
+
+namespace FPCGAsync
+{
+	/** 
+	* Helper to do simple point processing loops
+	* @param Context - The context containing the information about how many tasks to launch
+	* @param NumIterations - The number of calls that will be done to the provided function, also an upper bound on the number of points generated
+	* @param OutPoints - The array in which the results will be written to. Note that the array will be cleared before execution
+	* @param PointFunc - A function that has the index [0; NumIterations] and has to write to the point & return true when the current call generates a point
+	*/
+	void AsyncPointProcessing(FPCGContextPtr Context, int32 NumIterations, TArray<FPCGPoint>& OutPoints, const TFunction<bool(int32, FPCGPoint&)>& PointFunc);
+
+	/** 
+	* Helper to do more general 1:1 point processing loops
+	* @param NumAvailableTasks - The upper bound on the number of async tasks we'll start
+	* @param MinIterationsPerTask - The lower bound on the number of iterations per task we'll dispatch
+	* @param NumIterations - The number of calls that will be done to the provided function, also an upper bound on the number of points generated
+ 	* @param OutPoints - The array in which the results will be written to. Note that the array will be cleared before execution
+	* @param PointFunc - A function that has the index [0; NumIterations] and has to write to the point & return true when the current call generates a point
+	*/
+	void AsyncPointProcessing(int32 NumAvailableTasks, int32 MinIterationsPerTask, int32 NumIterations, TArray<FPCGPoint>& OutPoints, const TFunction<bool(int32, FPCGPoint&)>& PointFunc);
+}
\ No newline at end of file
diff --git a/Engine/Plugins/Experimental/PCG/Source/PCG/Public/PCGContext.h b/Engine/Plugins/Experimental/PCG/Source/PCG/Public/PCGContext.h
--- a/Engine/Plugins/Experimental/PCG/Source/PCG/Public/PCGContext.h
+++ b/Engine/Plugins/Experimental/PCG/Source/PCG/Public/PCGContext.h
@@ -21,6 +21,8 @@ struct FPCGContext
 	FPCGDataCollection OutputData;
 	UPCGComponent* SourceComponent = nullptr;
 	FPCGGraphCache* Cache = nullptr;
+	int32 NumAvailableTasks = 0;
+
 	// TODO: add RNG source
 	// TODO: replace this by a better identification mechanism
 	const UPCGNode* Node = nullptr;
