Missed a field in point data transform point

#rnx

diff --git a/Engine/Plugins/Experimental/PCG/Source/PCG/Private/Data/PCGPointData.cpp b/Engine/Plugins/Experimental/PCG/Source/PCG/Private/Data/PCGPointData.cpp
--- a/Engine/Plugins/Experimental/PCG/Source/PCG/Private/Data/PCGPointData.cpp
+++ b/Engine/Plugins/Experimental/PCG/Source/PCG/Private/Data/PCGPointData.cpp
@@ -204,8 +204,8 @@ FPCGPoint UPCGPointData::TransformPoint(const FPCGPoint& InPoint) const
 		return InPoint;
 	}
 
-	FVector WeightedScale = FVector::Zero();
 	FRotator WeightedRotator(0);
+	FVector WeightedScale = FVector::Zero();
 	float WeightedDensity = 0;
 	FVector WeightedExtents = FVector::Zero();
 	FVector WeightedColor = FVector::Zero();
@@ -219,6 +219,7 @@ FPCGPoint UPCGPointData::TransformPoint(const FPCGPoint& InPoint) const
 		WeightedRotator += (SourcePoint.Transform.Rotator() * Weight); // TODO: This is wonky
 		WeightedScale += (SourcePoint.Transform.GetScale3D() * Weight);
 		WeightedDensity += PCGPointHelpers::ManhattanDensity(SourcePoint, PointPosition);
+		WeightedExtents += SourcePoint.Extents * Weight;
 		WeightedColor += SourcePoint.Color * Weight;
 		WeightedSteepness += SourcePoint.Steepness * Weight;
 	}
@@ -230,6 +231,7 @@ FPCGPoint UPCGPointData::TransformPoint(const FPCGPoint& InPoint) const
 	Point.Transform.NormalizeRotation();
 	Point.Transform.SetScale3D(Point.Transform.GetScale3D() * WeightedScale);	
 	Point.Density *= WeightedDensity;
+	Point.Extents *= WeightedExtents;
 	Point.Color *= WeightedColor;
 	Point.Steepness *= WeightedSteepness;
 
