Added a cvar to flush the pcg results cache

diff --git a/Engine/Plugins/Experimental/PCG/Source/PCG/Private/PCGSubsystem.cpp b/Engine/Plugins/Experimental/PCG/Source/PCG/Private/PCGSubsystem.cpp
--- a/Engine/Plugins/Experimental/PCG/Source/PCG/Private/PCGSubsystem.cpp
+++ b/Engine/Plugins/Experimental/PCG/Source/PCG/Private/PCGSubsystem.cpp
@@ -7,6 +7,7 @@
 #include "Grid/PCGPartitionActor.h"
 
 #include "Engine/World.h"
+#include "HAL/IConsoleManager.h"
 
 #if WITH_EDITOR
 #include "Editor.h"
@@ -18,6 +19,22 @@
 #include "ObjectTools.h"
 #endif
 
+#if WITH_EDITOR
+namespace PCGSubsystemConsole
+{
+	static FAutoConsoleCommand CommandFlushCache(
+		TEXT("pcg.FlushCache"),
+		TEXT("Clears the PCG results cache."),
+		FConsoleCommandDelegate::CreateLambda([]()
+			{
+				if (UWorld* World = GEditor->GetEditorWorldContext().World())
+				{
+					World->GetSubsystem<UPCGSubsystem>()->FlushCache();
+				}
+			}));
+}
+#endif
+
 void UPCGSubsystem::Deinitialize()
 {
 	// Cancel all tasks
@@ -496,6 +513,14 @@ void UPCGSubsystem::CleanFromCache(const IPCGElement* InElement)
 	}
 }
 
+void UPCGSubsystem::FlushCache()
+{
+	if (GraphExecutor)
+	{
+		GraphExecutor->GetCache().ClearCache();
+	}
+}
+
 ///////////////////////////////////////////////////////////////////////////////////////////
 void UPCGSubsystem::DelayPartitionGraph(UPCGComponent* Component)
 {
diff --git a/Engine/Plugins/Experimental/PCG/Source/PCG/Public/PCGSubsystem.h b/Engine/Plugins/Experimental/PCG/Source/PCG/Public/PCGSubsystem.h
--- a/Engine/Plugins/Experimental/PCG/Source/PCG/Public/PCGSubsystem.h
+++ b/Engine/Plugins/Experimental/PCG/Source/PCG/Public/PCGSubsystem.h
@@ -83,6 +83,9 @@ public:
 	/** Cleans up the graph cache on an element basis */
 	void CleanFromCache(const IPCGElement* InElement);
 
+	/** Flushes the graph cache completely, use only for debugging */
+	void FlushCache();
+
 private:
 	FPCGTaskId DelayProcessGraph(UPCGComponent* Component, bool bGenerate, bool bSave, bool bUseEmptyNewBounds);
 	FPCGTaskId ProcessGraph(UPCGComponent* Component, const FBox& InPreviousBounds, const FBox& InNewBounds, bool bGenerate, bool bSave);
