Added cvar to limit number of tasks in PCG processing

diff --git a/Engine/Plugins/Experimental/PCG/Source/PCG/Private/Graph/PCGGraphExecutor.cpp b/Engine/Plugins/Experimental/PCG/Source/PCG/Private/Graph/PCGGraphExecutor.cpp
--- a/Engine/Plugins/Experimental/PCG/Source/PCG/Private/Graph/PCGGraphExecutor.cpp
+++ b/Engine/Plugins/Experimental/PCG/Source/PCG/Private/Graph/PCGGraphExecutor.cpp
@@ -13,6 +13,11 @@
 #include "FileHelpers.h"
 #endif
 
+static TAutoConsoleVariable<int32> CVarMaxNumTasks(
+	TEXT("pcg.MaxConcurrentTasks"),
+	4096,
+	TEXT("Maximum number of concurrent tasks for PCG processing"));
+
 FPCGGraphExecutor::FPCGGraphExecutor(UObject* InOwner)
 	: GraphCompiler(MakeUnique<FPCGGraphCompiler>())
 	, GraphCache(InOwner)
@@ -209,7 +214,7 @@ void FPCGGraphExecutor::Execute()
 
 			// Reassign resources to the context
 			// TODO: change this when we support MT in the graph executor
-			ActiveTask.Context->NumAvailableTasks = FMath::Max(FPlatformMisc::NumberOfCoresIncludingHyperthreads() - 2, 2);
+			ActiveTask.Context->NumAvailableTasks = FMath::Max(1, FMath::Min(FPlatformMisc::NumberOfCoresIncludingHyperthreads() - 2, CVarMaxNumTasks.GetValueOnAnyThread()));
 
 			if (!ActiveTask.Context->bIsPaused && ActiveTask.Element->Execute(ActiveTask.Context))
 			{
