Small improvements around the blueprint elements (logging, point caching)

diff --git a/Engine/Plugins/Experimental/PCG/Source/PCG/Private/Elements/PCGExecuteBlueprint.cpp b/Engine/Plugins/Experimental/PCG/Source/PCG/Private/Elements/PCGExecuteBlueprint.cpp
--- a/Engine/Plugins/Experimental/PCG/Source/PCG/Private/Elements/PCGExecuteBlueprint.cpp
+++ b/Engine/Plugins/Experimental/PCG/Source/PCG/Private/Elements/PCGExecuteBlueprint.cpp
@@ -1,6 +1,8 @@
 // Copyright Epic Games, Inc. All Rights Reserved.
 
 #include "Elements/PCGExecuteBlueprint.h"
+#include "Data/PCGPointData.h"
+#include "Data/PCGSpatialData.h"
 
 #if WITH_EDITOR
 #include "Engine/World.h"
@@ -323,6 +325,11 @@ void UPCGBlueprintSettings::GetTrackedActorTags(FPCGTagToSettingsMap& OutTagToSe
 }
 #endif // WITH_EDITOR
 
+FName UPCGBlueprintSettings::AdditionalTaskName() const
+{
+	return BlueprintElementType ? BlueprintElementType->GetFName() : Super::AdditionalTaskName();
+}
+
 FPCGElementPtr UPCGBlueprintSettings::CreateElement() const
 {
 	return MakeShared<FPCGExecuteBlueprintElement>();
@@ -335,7 +342,39 @@ bool FPCGExecuteBlueprintElement::ExecuteInternal(FPCGContextPtr Context) const
 
 	if (Settings && Settings->BlueprintElementInstance)
 	{
+		// Precache points if needed
+		if (Settings->bPrecachePoints)
+		{
+			for (const FPCGTaggedData& Input : Context->InputData.TaggedData)
+			{
+				if (const UPCGSpatialDataWithPointCache* SpatialInput = Cast<UPCGSpatialDataWithPointCache>(Input.Data))
+				{
+					SpatialInput->ToPointData(Context);
+				}
+			}
+		}
+
+		// Log info on inputs
+		for(int32 InputIndex = 0; InputIndex < Context->InputData.TaggedData.Num(); ++InputIndex)
+		{
+			const FPCGTaggedData& Input = Context->InputData.TaggedData[InputIndex];
+			if (const UPCGPointData* PointData = Cast<UPCGPointData>(Input.Data))
+			{
+				PCGE_LOG(Verbose, "Input %d has %d points", InputIndex, PointData->GetPoints().Num());
+			}
+		}
+
 		Settings->BlueprintElementInstance->Execute(Context->InputData, Context->OutputData);
+
+		// Log info on outputs
+		for (int32 OutputIndex = 0; OutputIndex < Context->OutputData.TaggedData.Num(); ++OutputIndex)
+		{
+			const FPCGTaggedData& Output = Context->OutputData.TaggedData[OutputIndex];
+			if (const UPCGPointData* PointData = Cast<UPCGPointData>(Output.Data))
+			{
+				PCGE_LOG(Verbose, "Output %d has %d points", OutputIndex, PointData->GetPoints().Num());
+			}
+		}
 	}
 	else
 	{
diff --git a/Engine/Plugins/Experimental/PCG/Source/PCG/Private/PCGContext.cpp b/Engine/Plugins/Experimental/PCG/Source/PCG/Private/PCGContext.cpp
--- a/Engine/Plugins/Experimental/PCG/Source/PCG/Private/PCGContext.cpp
+++ b/Engine/Plugins/Experimental/PCG/Source/PCG/Private/PCGContext.cpp
@@ -4,12 +4,31 @@
 #include "PCGComponent.h"
 #include "GameFramework/Actor.h"
 
-FName FPCGContext::GetTaskName() const
+FString FPCGContext::GetTaskName() const
 {
-	return Node ? Node->GetFName() : TEXT("Anonymous task");
+	if (Node)
+	{
+		const FName NodeName = Node->GetFName();
+
+		const UPCGSettings* Settings = GetInputSettings<UPCGSettings>();
+		const FName NodeAdditionalName = Settings ? Settings->AdditionalTaskName() : NAME_None;
+
+		if (NodeAdditionalName == NAME_None)
+		{
+			return NodeName.ToString();
+		}
+		else
+		{
+			return FString::Printf(TEXT("%s (%s)"), *NodeName.ToString(), *NodeAdditionalName.ToString());
+		}
+	}
+	else
+	{
+		return TEXT("Anonymous task");
+	}
 }
 
-FName FPCGContext::GetComponentName() const
+FString FPCGContext::GetComponentName() const
 {
-	return SourceComponent ? SourceComponent->GetOwner()->GetFName() : TEXT("Non-PCG Component");
+	return SourceComponent ? SourceComponent->GetOwner()->GetFName().ToString() : TEXT("Non-PCG Component");
 }
\ No newline at end of file
diff --git a/Engine/Plugins/Experimental/PCG/Source/PCG/Public/Elements/PCGExecuteBlueprint.h b/Engine/Plugins/Experimental/PCG/Source/PCG/Public/Elements/PCGExecuteBlueprint.h
--- a/Engine/Plugins/Experimental/PCG/Source/PCG/Public/Elements/PCGExecuteBlueprint.h
+++ b/Engine/Plugins/Experimental/PCG/Source/PCG/Public/Elements/PCGExecuteBlueprint.h
@@ -73,6 +73,8 @@ public:
 	virtual void GetTrackedActorTags(FPCGTagToSettingsMap& OutTagToSettings) const override;
 #endif
 
+	virtual FName AdditionalTaskName() const override;
+
 protected:
 	virtual FPCGElementPtr CreateElement() const override;
 	// ~End UPCGSettings interface
@@ -107,6 +109,9 @@ protected:
 	UPROPERTY(BlueprintReadWrite, EditAnywhere, Category = Settings)
 	bool bCreatesArtifacts = false;
 
+	UPROPERTY(BlueprintReadWrite, EditAnywhere, Category = Settings)
+	bool bPrecachePoints = false;
+
 protected:
 #if WITH_EDITOR
 	void OnBlueprintChanged(UBlueprint* InBlueprint);
diff --git a/Engine/Plugins/Experimental/PCG/Source/PCG/Public/PCGContext.h b/Engine/Plugins/Experimental/PCG/Source/PCG/Public/PCGContext.h
--- a/Engine/Plugins/Experimental/PCG/Source/PCG/Public/PCGContext.h
+++ b/Engine/Plugins/Experimental/PCG/Source/PCG/Public/PCGContext.h
@@ -35,7 +35,7 @@ struct FPCGContext
 #endif
 
 	template<typename SettingsType>
-	const SettingsType* GetInputSettings()
+	const SettingsType* GetInputSettings() const
 	{
 		if (Node && Node->DefaultSettings)
 		{
@@ -47,6 +47,6 @@ struct FPCGContext
 		}
 	}
 
-	FName GetTaskName() const;
-	FName GetComponentName() const;
+	FString GetTaskName() const;
+	FString GetComponentName() const;
 };
diff --git a/Engine/Plugins/Experimental/PCG/Source/PCG/Public/PCGElement.h b/Engine/Plugins/Experimental/PCG/Source/PCG/Public/PCGElement.h
--- a/Engine/Plugins/Experimental/PCG/Source/PCG/Public/PCGElement.h
+++ b/Engine/Plugins/Experimental/PCG/Source/PCG/Public/PCGElement.h
@@ -16,8 +16,8 @@ typedef TSharedPtr<IPCGElement, ESPMode::ThreadSafe> FPCGElementPtr;
 	UE_LOG(LogPCG, \
 		Verbosity, \
 		TEXT("[%s - %s]: " Format), \
-		*CustomContext->GetComponentName().ToString(), \
-		*CustomContext->GetTaskName().ToString(), \
+		*CustomContext->GetComponentName(), \
+		*CustomContext->GetTaskName(), \
 		##__VA_ARGS__)
 
 #define PCGE_LOG(Verbosity, Format, ...) PCGE_LOG_C(Verbosity, Context, Format, ##__VA_ARGS__)
diff --git a/Engine/Plugins/Experimental/PCG/Source/PCG/Public/PCGSettings.h b/Engine/Plugins/Experimental/PCG/Source/PCG/Public/PCGSettings.h
--- a/Engine/Plugins/Experimental/PCG/Source/PCG/Public/PCGSettings.h
+++ b/Engine/Plugins/Experimental/PCG/Source/PCG/Public/PCGSettings.h
@@ -47,6 +47,9 @@ public:
 	virtual void GetTrackedActorTags(FPCGTagToSettingsMap& OutTagToSettings) const {}
 #endif
 
+	/** Derived classes can implement this to expose additional name information in the logs */
+	virtual FName AdditionalTaskName() const { return NAME_None; }
+
 	UPROPERTY(EditAnywhere, BlueprintReadWrite, Category = Settings)
 	int Seed = 0xC35A9631; // random prime number
 
