PCG: changed default debug mesh and materials assets.

diff --git a/Engine/Plugins/Experimental/PCG/Source/PCG/Private/PCGDebug.cpp b/Engine/Plugins/Experimental/PCG/Source/PCG/Private/PCGDebug.cpp
--- a/Engine/Plugins/Experimental/PCG/Source/PCG/Private/PCGDebug.cpp
+++ b/Engine/Plugins/Experimental/PCG/Source/PCG/Private/PCGDebug.cpp
@@ -5,8 +5,8 @@
 
 namespace PCGDebugVisConstants
 {
-	const FSoftObjectPath DefaultPointMesh = FSoftObjectPath(TEXT("/Engine/BasicShapes/Cube.Cube"));
-	const FSoftObjectPath MaterialForDefaultPointMesh = FSoftObjectPath(TEXT("Material'/Game/PCG/DebugMaterial.DebugMaterial'"));
+	const FSoftObjectPath DefaultPointMesh = FSoftObjectPath(TEXT("/Game/PCG/Content/DebugObjects/PCG_Cube.PCG_Cube"));
+	const FSoftObjectPath MaterialForDefaultPointMesh = FSoftObjectPath(TEXT("Material'/Game/PCG/Content/DebugObjects/PCG_DebugMaterial.PCG_DebugMaterial'"));
 }
 
 FPCGDebugVisualizationSettings::FPCGDebugVisualizationSettings()
