Better API for graph setup BP & blueprint node settings

diff --git a/Engine/Plugins/Experimental/PCG/Source/PCG/Private/Elements/PCGExecuteBlueprint.cpp b/Engine/Plugins/Experimental/PCG/Source/PCG/Private/Elements/PCGExecuteBlueprint.cpp
--- a/Engine/Plugins/Experimental/PCG/Source/PCG/Private/Elements/PCGExecuteBlueprint.cpp
+++ b/Engine/Plugins/Experimental/PCG/Source/PCG/Private/Elements/PCGExecuteBlueprint.cpp
@@ -282,7 +282,13 @@ void UPCGBlueprintSettings::OnBlueprintElementChanged(UPCGBlueprintElement* InEl
 }
 #endif
 
-void UPCGBlueprintSettings::SetElementType(TSubclassOf<UPCGBlueprintElement> InElementType)
+void UPCGBlueprintSettings::SetElementTypeInternal(TSubclassOf<UPCGBlueprintElement> InElementType)
+{
+	UPCGBlueprintElement* ElementInstance = nullptr;
+	SetElementType(InElementType, ElementInstance);
+}
+
+void UPCGBlueprintSettings::SetElementType(TSubclassOf<UPCGBlueprintElement> InElementType, UPCGBlueprintElement*& ElementInstance)
 {
 	if (!BlueprintElementInstance || InElementType != BlueprintElementType)
 	{
@@ -295,6 +301,8 @@ void UPCGBlueprintSettings::SetElementType(TSubclassOf<UPCGBlueprintElement> InE
 		
 		RefreshBlueprintElement();
 	}
+
+	ElementInstance = BlueprintElementInstance;
 }
 
 void UPCGBlueprintSettings::RefreshBlueprintElement()
diff --git a/Engine/Plugins/Experimental/PCG/Source/PCG/Private/PCGGraph.cpp b/Engine/Plugins/Experimental/PCG/Source/PCG/Private/PCGGraph.cpp
--- a/Engine/Plugins/Experimental/PCG/Source/PCG/Private/PCGGraph.cpp
+++ b/Engine/Plugins/Experimental/PCG/Source/PCG/Private/PCGGraph.cpp
@@ -88,7 +88,7 @@ void UPCGGraph::InitializeFromTemplate()
 }
 #endif
 
-UPCGNode* UPCGGraph::AddNodeOfType(TSubclassOf<class UPCGSettings> InSettingsClass)
+UPCGNode* UPCGGraph::AddNodeOfType(TSubclassOf<class UPCGSettings> InSettingsClass, UPCGSettings*& OutDefaultNodeSettings)
 {
 	UPCGSettings* Settings = NewObject<UPCGSettings>(GetTransientPackage(), InSettingsClass);
 
@@ -124,6 +124,7 @@ UPCGNode* UPCGGraph::AddNodeOfType(TSubclassOf<class UPCGSettings> InSettingsCla
 		OnNodeAdded(Node);
 	}
 
+	OutDefaultNodeSettings = Settings;
 	return Node;
 }
 
diff --git a/Engine/Plugins/Experimental/PCG/Source/PCG/Public/Elements/PCGExecuteBlueprint.h b/Engine/Plugins/Experimental/PCG/Source/PCG/Public/Elements/PCGExecuteBlueprint.h
--- a/Engine/Plugins/Experimental/PCG/Source/PCG/Public/Elements/PCGExecuteBlueprint.h
+++ b/Engine/Plugins/Experimental/PCG/Source/PCG/Public/Elements/PCGExecuteBlueprint.h
@@ -88,14 +88,18 @@ public:
 	// ~End UObject interface
 #endif
 
-	UFUNCTION(BlueprintCallable, Category = Settings)
-	void SetElementType(TSubclassOf<UPCGBlueprintElement> InElementType);
+	/** To facilitate transition, we will keep this version around even though it's a lot less useful */
+	UFUNCTION(BlueprintCallable, BlueprintInternalUseOnly)
+	void SetElementTypeInternal(TSubclassOf<UPCGBlueprintElement> InElementType);
+
+	UFUNCTION(BlueprintCallable, Category = Settings, meta=(DeterminesOutputType="InElementType", DynamicOutputParam = "ElementInstance"))
+	void SetElementType(TSubclassOf<UPCGBlueprintElement> InElementType, UPCGBlueprintElement*& ElementInstance);
 
 protected:
 	UPROPERTY()
 	TSubclassOf<UPCGBlueprintElement> BlueprintElement_DEPRECATED;
 
-	UPROPERTY(BlueprintSetter = SetElementType, EditAnywhere, Category = Template)
+	UPROPERTY(BlueprintSetter = SetElementTypeInternal, EditAnywhere, Category = Template)
 	TSubclassOf<UPCGBlueprintElement> BlueprintElementType;
 
 	UPROPERTY(BlueprintReadOnly, VisibleAnywhere, Instanced, Category = Settings, meta = (ShowOnlyInnerProperties))
diff --git a/Engine/Plugins/Experimental/PCG/Source/PCG/Public/PCGGraph.h b/Engine/Plugins/Experimental/PCG/Source/PCG/Public/PCGGraph.h
--- a/Engine/Plugins/Experimental/PCG/Source/PCG/Public/PCGGraph.h
+++ b/Engine/Plugins/Experimental/PCG/Source/PCG/Public/PCGGraph.h
@@ -36,8 +36,8 @@ public:
 #endif
 
 	/** Creates a default node based on the settings class wanted. Returns the newly created node. */
-	UFUNCTION(BlueprintCallable, Category = Graph)
-	UPCGNode* AddNodeOfType(TSubclassOf<class UPCGSettings> InSettingsClass);
+	UFUNCTION(BlueprintCallable, Category = Graph, meta=(DeterminesOutputType = "InSettingsClass", DynamicOutputParam = "DefaultNodeSettings"))
+	UPCGNode* AddNodeOfType(TSubclassOf<class UPCGSettings> InSettingsClass, UPCGSettings*& DefaultNodeSettings);
 
 	/** Creates a node and assigns it in the input settings. Returns the created node. */
 	UFUNCTION(BlueprintCallable, Category = Graph)
