Clarified radius/extent/sizes in surface sampler, landscape, textures

diff --git a/Engine/Plugins/Experimental/PCG/Source/PCG/Private/Data/PCGLandscapeData.cpp b/Engine/Plugins/Experimental/PCG/Source/PCG/Private/Data/PCGLandscapeData.cpp
--- a/Engine/Plugins/Experimental/PCG/Source/PCG/Private/Data/PCGLandscapeData.cpp
+++ b/Engine/Plugins/Experimental/PCG/Source/PCG/Private/Data/PCGLandscapeData.cpp
@@ -113,6 +113,7 @@ const UPCGPointData* UPCGLandscapeData::CreatePointData(FPCGContextPtr Context)
 			OutPoint = FPCGPoint(FTransform(VertexLocation),
 				1.0f,
 				PCGHelpers::ComputeSeed(X, Y, (int)HeightAtVertex.GetValue()));
+			OutPoint.Extents = Transform.GetScale3D() / 2.0;
 
 			return true;
 		}
diff --git a/Engine/Plugins/Experimental/PCG/Source/PCG/Private/Data/PCGTextureData.cpp b/Engine/Plugins/Experimental/PCG/Source/PCG/Private/Data/PCGTextureData.cpp
--- a/Engine/Plugins/Experimental/PCG/Source/PCG/Private/Data/PCGTextureData.cpp
+++ b/Engine/Plugins/Experimental/PCG/Source/PCG/Private/Data/PCGTextureData.cpp
@@ -101,7 +101,7 @@ FPCGPoint UPCGBaseTextureData::TransformPoint(const FPCGPoint& InPoint) const
 
 	FLinearColor Color = PCGTextureSampling::Sample<FLinearColor>(Position2D, Surface, Width, Height, [this](int32 Index) { return ColorData[Index]; });
 
-	Point.Color = Color;
+	Point.Color *= Color;
 	Point.Density *= PCGTextureSampling::SampleFloatChannel(Color, ColorChannel);
 
 	return Point;
@@ -141,6 +141,10 @@ const UPCGPointData* UPCGBaseTextureData::CreatePointData(FPCGContextPtr Context
 			OutPoint = FPCGPoint(FTransform(Transform.TransformPosition(LocalPosition)),
 				Density,
 				PCGHelpers::ComputeSeed(X, Y));
+
+			const FVector TransformScale = Transform.GetScale3D();
+			// Note: divided by 4 here because the scale is doubled before, and the extents represent half a pixel
+			OutPoint.Extents = FVector(TransformScale.X * XScale / 4.0, TransformScale.Y * YScale / 4.0, 1.0);
 			OutPoint.Color = ColorData[X + Y * Width];
 
 			return true;
diff --git a/Engine/Plugins/Experimental/PCG/Source/PCG/Private/Elements/PCGSurfaceSampler.cpp b/Engine/Plugins/Experimental/PCG/Source/PCG/Private/Elements/PCGSurfaceSampler.cpp
--- a/Engine/Plugins/Experimental/PCG/Source/PCG/Private/Elements/PCGSurfaceSampler.cpp
+++ b/Engine/Plugins/Experimental/PCG/Source/PCG/Private/Elements/PCGSurfaceSampler.cpp
@@ -10,11 +10,11 @@
 struct FPCGSurfaceSamplerLoopData
 {
 	const UPCGSurfaceSamplerSettings* Settings;
-	
-	FVector::FReal PointRadius;
-	FVector::FReal InterstitialDistance;
-	FVector::FReal InnerCellSize;
-	FVector::FReal CellSize;
+
+	FVector PointExtents;
+	FVector InterstitialDistance;
+	FVector InnerCellSize;
+	FVector CellSize;
 
 	int32 CellMinX;
 	int32 CellMaxX;
@@ -32,17 +32,16 @@ struct FPCGSurfaceSamplerLoopData
 		Settings = InSettings;
 
 		// Conceptually, we will break down the surface bounds in a N x M grid
-		InterstitialDistance = 2 * Settings->PointRadius;
-		InnerCellSize = 2 * Settings->Looseness * Settings->PointRadius;
+		InterstitialDistance = Settings->PointExtents * 2;
+		InnerCellSize = InterstitialDistance * Settings->Looseness;
 		CellSize = InterstitialDistance + InnerCellSize;
-
-		check(CellSize > 0);
+		check(CellSize.X > 0 && CellSize.Y > 0);
 
 		// By using scaled indices in the world, we can easily make this process deterministic
-		CellMinX = FMath::CeilToInt((InputBounds.Min.X) / CellSize);
-		CellMaxX = FMath::FloorToInt((InputBounds.Max.X) / CellSize);
-		CellMinY = FMath::CeilToInt((InputBounds.Min.Y) / CellSize);
-		CellMaxY = FMath::FloorToInt((InputBounds.Max.Y) / CellSize);
+		CellMinX = FMath::CeilToInt((InputBounds.Min.X) / CellSize.X);
+		CellMaxX = FMath::FloorToInt((InputBounds.Max.X) / CellSize.X);
+		CellMinY = FMath::CeilToInt((InputBounds.Min.Y) / CellSize.Y);
+		CellMaxY = FMath::FloorToInt((InputBounds.Max.Y) / CellSize.Y);
 
 		if (CellMinX > CellMaxX || CellMinY > CellMaxY)
 		{
@@ -81,6 +80,17 @@ struct FPCGSurfaceSamplerLoopData
 	}
 };
 
+void UPCGSurfaceSamplerSettings::PostLoad()
+{
+	Super::PostLoad();
+
+	if (PointRadius_DEPRECATED != 0)
+	{
+		PointExtents = FVector(PointRadius_DEPRECATED);
+		PointRadius_DEPRECATED = 0;
+	}
+}
+
 FPCGElementPtr UPCGSurfaceSamplerSettings::CreateElement() const
 {
 	return MakeShared<FPCGSurfaceSamplerElement>();
@@ -95,9 +105,9 @@ bool FPCGSurfaceSamplerElement::ExecuteInternal(FPCGContextPtr Context) const
 
 	// Early out on invalid settings
 	// TODO: we could compute an approximate radius based on the points per squared meters if that's useful
-	if (Settings->PointRadius <= 0)
+	if(Settings->PointExtents.X <= 0 || Settings->PointExtents.Y <= 0)
 	{
-		PCGE_LOG(Warning, "Skipped - Invalid point radius");
+		PCGE_LOG(Warning, "Skipped - Invalid point extents");
 		return true;
 	}
 
@@ -145,9 +155,9 @@ bool FPCGSurfaceSamplerElement::ExecuteInternal(FPCGContextPtr Context) const
 			int32 CellY;
 			LoopData.ComputeCellIndices(Index, CellX, CellY);
 
-			const FVector::FReal CurrentX = CellX * LoopData.CellSize;
-			const FVector::FReal CurrentY = CellY * LoopData.CellSize;
-			const FVector::FReal InnerCellSize = LoopData.InnerCellSize;
+			const FVector::FReal CurrentX = CellX * LoopData.CellSize.X;
+			const FVector::FReal CurrentY = CellY * LoopData.CellSize.Y;
+			const FVector InnerCellSize = LoopData.InnerCellSize;
 
 			FRandomStream RandomSource(PCGHelpers::ComputeSeed(LoopData.Settings->Seed, CellX, CellY));
 			float Chance = RandomSource.FRand();
@@ -159,8 +169,8 @@ bool FPCGSurfaceSamplerElement::ExecuteInternal(FPCGContextPtr Context) const
 				const float RandX = RandomSource.FRand();
 				const float RandY = RandomSource.FRand();
 
-				OutPoint.Transform = FTransform(FVector(CurrentX + RandX * InnerCellSize, CurrentY + RandY * InnerCellSize, 0));
-				OutPoint.Extents = FVector(LoopData.Settings->PointRadius);
+				OutPoint.Transform = FTransform(FVector(CurrentX + RandX * InnerCellSize.X, CurrentY + RandY * InnerCellSize.Y, 0));
+				OutPoint.Extents = FVector(LoopData.Settings->PointExtents);
 				OutPoint.Density = LoopData.Settings->bApplyDensityToPoints ? ((Ratio - Chance) / Ratio) : 1.0f;
 				OutPoint.Steepness = LoopData.Settings->PointSteepness;
 				OutPoint.Seed = RandomSource.GetCurrentSeed();
diff --git a/Engine/Plugins/Experimental/PCG/Source/PCG/Public/Elements/PCGSurfaceSampler.h b/Engine/Plugins/Experimental/PCG/Source/PCG/Public/Elements/PCGSurfaceSampler.h
--- a/Engine/Plugins/Experimental/PCG/Source/PCG/Public/Elements/PCGSurfaceSampler.h
+++ b/Engine/Plugins/Experimental/PCG/Source/PCG/Public/Elements/PCGSurfaceSampler.h
@@ -20,8 +20,11 @@ public:
 	UPROPERTY(BlueprintReadWrite, EditAnywhere, Category = Settings, meta=(ClampMin="0"))
 	float PointsPerSquaredMeter = 0.1f;
 
-	UPROPERTY(BlueprintReadWrite, EditAnywhere, Category = Settings, meta=(ClampMin="0"))
-	float PointRadius = 100.0f;
+	UPROPERTY()
+	float PointRadius_DEPRECATED = 100.0f;
+
+	UPROPERTY(BlueprintReadWrite, EditAnywhere, Category = Settings)
+	FVector PointExtents = FVector(100.0f);
 
 	UPROPERTY(BlueprintReadWrite, EditAnywhere, Category = Settings, meta=(ClampMin="0"))
 	float Looseness = 1.0f;
@@ -37,6 +40,10 @@ public:
 	bool bKeepZeroDensityPoints = false;
 #endif
 
+	//~Begin UObject interface
+	virtual void PostLoad() override;
+	//~End UObject interface
+
 #if WITH_EDITOR
 	//~Begin UPCGSettings interface
 	virtual FName GetDefaultNodeName() const override { return FName(TEXT("SurfaceSamplerNode")); }
