Made it so nanite meshes don't use HISM
Copy settings from pre-existing static mesh component in Spawn Actor

diff --git a/Engine/Plugins/Experimental/PCG/Source/PCG/Private/Elements/PCGSpawnActor.cpp b/Engine/Plugins/Experimental/PCG/Source/PCG/Private/Elements/PCGSpawnActor.cpp
--- a/Engine/Plugins/Experimental/PCG/Source/PCG/Private/Elements/PCGSpawnActor.cpp
+++ b/Engine/Plugins/Experimental/PCG/Source/PCG/Private/Elements/PCGSpawnActor.cpp
@@ -129,12 +129,14 @@ bool FPCGSpawnActorElement::ExecuteInternal(FPCGContextPtr Context) const
 			{
 				TArray<UActorComponent*> Components;
 				UPCGActorHelpers::GetActorClassDefaultComponents(Settings->TemplateActorClass, Components, UStaticMeshComponent::StaticClass());
+				UStaticMeshComponent* FirstSMC = nullptr;
 				UStaticMesh* Mesh = nullptr;
 
 				for (UActorComponent* Component : Components)
 				{
 					if (UStaticMeshComponent* SMC = Cast<UStaticMeshComponent>(Component))
 					{
+						FirstSMC = SMC;
 						Mesh = SMC->GetStaticMesh();
 						if (Mesh)
 						{
@@ -146,6 +148,7 @@ bool FPCGSpawnActorElement::ExecuteInternal(FPCGContextPtr Context) const
 				if (Mesh)
 				{
 					ISMC = UPCGActorHelpers::GetOrCreateISMC(TargetActor, Context->SourceComponent, Mesh);
+					UEngine::CopyPropertiesForUnrelatedObjects(FirstSMC, ISMC);
 				}
 				else
 				{
diff --git a/Engine/Plugins/Experimental/PCG/Source/PCG/Private/Helpers/PCGActorHelpers.cpp b/Engine/Plugins/Experimental/PCG/Source/PCG/Private/Helpers/PCGActorHelpers.cpp
--- a/Engine/Plugins/Experimental/PCG/Source/PCG/Private/Helpers/PCGActorHelpers.cpp
+++ b/Engine/Plugins/Experimental/PCG/Source/PCG/Private/Helpers/PCGActorHelpers.cpp
@@ -69,7 +69,17 @@ UInstancedStaticMeshComponent* UPCGActorHelpers::GetOrCreateISMC(AActor* InTarge
 	// Otherwise, create a new component
 	// TODO: use static mesh component if there's only one instance
 	// TODO: add hism/ism switch or better yet, use a template component
-	UInstancedStaticMeshComponent* ISMC = NewObject<UHierarchicalInstancedStaticMeshComponent>(InTargetActor);
+	UInstancedStaticMeshComponent* ISMC = nullptr;
+
+	if (InMesh->HasValidNaniteData())
+	{
+		ISMC = NewObject<UInstancedStaticMeshComponent>(InTargetActor);
+	}
+	else
+	{
+		ISMC = NewObject<UHierarchicalInstancedStaticMeshComponent>(InTargetActor);
+	}
+
 	ISMC->SetStaticMesh(InMesh);
 
 	// TODO: improve material override mechanisms
