PCG:
*Initial implementation of PCG graph editor, this is early WIP.

diff --git a/Engine/Plugins/Experimental/PCG/Source/PCG/Public/PCGNode.h b/Engine/Plugins/Experimental/PCG/Source/PCG/Public/PCGNode.h
--- a/Engine/Plugins/Experimental/PCG/Source/PCG/Public/PCGNode.h
+++ b/Engine/Plugins/Experimental/PCG/Source/PCG/Public/PCGNode.h
@@ -41,6 +41,8 @@ public:
 	/** Changes the default settings in the node */
 	void SetDefaultSettings(TObjectPtr<UPCGSettings> InSettings);
 
+	const TArray<TObjectPtr<UPCGNode>>& GetOutboundNodes() const { return OutboundNodes; }
+
 	/** Note: do not set this property directly from code, use SetDefaultSettings instead */
 	UPROPERTY(BlueprintReadWrite, EditAnywhere, Category = Node, meta=(EditInline))
 	TObjectPtr<UPCGSettings> DefaultSettings;
@@ -49,6 +51,14 @@ public:
 	FOnPCGNodeSettingsChanged OnNodeSettingsChangedDelegate;
 #endif
 
+#if WITH_EDITORONLY_DATA
+	UPROPERTY()
+	int32 PositionX;
+
+	UPROPERTY()
+	int32 PositionY;
+#endif // WITH_EDITORONLY_DATA
+
 protected:
 	void ConnectTo(UPCGNode* InSuccessor);
 	void ConnectFrom(UPCGNode* InPredecessor);
diff --git a/Engine/Plugins/Experimental/PCG/Source/PCGEditor/PCGEditor.Build.cs b/Engine/Plugins/Experimental/PCG/Source/PCGEditor/PCGEditor.Build.cs
--- a/Engine/Plugins/Experimental/PCG/Source/PCGEditor/PCGEditor.Build.cs
+++ b/Engine/Plugins/Experimental/PCG/Source/PCGEditor/PCGEditor.Build.cs
@@ -28,6 +28,7 @@ namespace UnrealBuildTool.Rules
 					"DesktopWidgets",
 					"EditorStyle",
 					"EditorSubsystem",
+					"GraphEditor",
 					"InputCore",
 					"Kismet",
 					"PCG",
diff --git a/Engine/Plugins/Experimental/PCG/Source/PCGEditor/Private/AssetTypeActions/PCGGraphAssetTypeActions.cpp b/Engine/Plugins/Experimental/PCG/Source/PCGEditor/Private/AssetTypeActions/PCGGraphAssetTypeActions.cpp
--- a/Engine/Plugins/Experimental/PCG/Source/PCGEditor/Private/AssetTypeActions/PCGGraphAssetTypeActions.cpp
+++ b/Engine/Plugins/Experimental/PCG/Source/PCGEditor/Private/AssetTypeActions/PCGGraphAssetTypeActions.cpp
@@ -2,6 +2,11 @@
 
 #include "PCGGraphAssetTypeActions.h"
 #include "PCGGraph.h"
+#include "PCGEditor.h"
+
+#include "HAL/IConsoleManager.h"
+
+static TAutoConsoleVariable<bool> CVarPCGUseGraphEditor(TEXT("pcg.UseGraphEditor"), false, TEXT("Wheter to use use the new graph editor or not."));
 
 FText FPCGGraphAssetTypeActions::GetName() const
 {
@@ -11,4 +16,23 @@ FText FPCGGraphAssetTypeActions::GetName() const
 UClass* FPCGGraphAssetTypeActions::GetSupportedClass() const
 {
 	return UPCGGraph::StaticClass();
+}
+
+void FPCGGraphAssetTypeActions::OpenAssetEditor(const TArray<UObject*>& InObjects, TSharedPtr<class IToolkitHost> EditWithinLevelEditor)
+{
+	if (CVarPCGUseGraphEditor.GetValueOnAnyThread())
+	{
+		for (UObject* Object : InObjects)
+		{
+			if (UPCGGraph* PCGGraph = Cast<UPCGGraph>(Object))
+			{
+				TSharedRef<FPCGEditor> PCGEditor(new FPCGEditor());
+				PCGEditor->Initialize(EToolkitMode::Standalone, EditWithinLevelEditor, PCGGraph);
+			}
+		}
+	}
+	else
+	{
+		FPCGCommonAssetTypeActions::OpenAssetEditor(InObjects, EditWithinLevelEditor);
+	}
 }
\ No newline at end of file
diff --git a/Engine/Plugins/Experimental/PCG/Source/PCGEditor/Private/AssetTypeActions/PCGGraphAssetTypeActions.h b/Engine/Plugins/Experimental/PCG/Source/PCGEditor/Private/AssetTypeActions/PCGGraphAssetTypeActions.h
--- a/Engine/Plugins/Experimental/PCG/Source/PCGEditor/Private/AssetTypeActions/PCGGraphAssetTypeActions.h
+++ b/Engine/Plugins/Experimental/PCG/Source/PCGEditor/Private/AssetTypeActions/PCGGraphAssetTypeActions.h
@@ -10,4 +10,5 @@ public:
 	// IAssetTypeActions Implementation
 	virtual FText GetName() const override;
 	virtual UClass* GetSupportedClass() const override;
+	virtual void OpenAssetEditor(const TArray<UObject*>& InObjects, TSharedPtr<class IToolkitHost> EditWithinLevelEditor = TSharedPtr<IToolkitHost>()) override;
 };
\ No newline at end of file
diff --git a/Engine/Plugins/Experimental/PCG/Source/PCGEditor/Private/PCGEditor.cpp b/Engine/Plugins/Experimental/PCG/Source/PCGEditor/Private/PCGEditor.cpp
new file mode 100644
--- /dev/null
+++ b/Engine/Plugins/Experimental/PCG/Source/PCGEditor/Private/PCGEditor.cpp
@@ -0,0 +1,384 @@
+// Copyright Epic Games, Inc. All Rights Reserved.
+
+#include "PCGEditor.h"
+
+#include "PCGEditorGraph.h"
+#include "PCGEditorGraphNode.h"
+#include "PCGEditorGraphSchema.h"
+#include "PCGGraph.h"
+
+#include "Editor.h"
+#include "Framework/Commands/GenericCommands.h"
+#include "Framework/Commands/UIAction.h"
+#include "Framework/Commands/UICommandList.h"
+#include "Framework/Docking/TabManager.h"
+#include "Framework/MultiBox/MultiBoxBuilder.h"
+#include "GraphEditor.h"
+#include "GraphEditorActions.h"
+#include "IDetailsView.h"
+#include "Modules/ModuleManager.h"
+#include "PropertyEditorModule.h"
+#include "Widgets/Docking/SDockTab.h"
+
+#include "ToolMenus.h"
+
+#define LOCTEXT_NAMESPACE "PCGGraphEditor"
+
+namespace FPCGEditor_private
+{
+	const FName GraphEditorID = FName(TEXT("GraphEditor"));
+	const FName PropertyDetailsID = FName(TEXT("PropertyDetails"));
+	const FName LibraryID = FName(TEXT("Library"));
+	const FName AttributesID = FName(TEXT("Attributes"));
+	const FName ViewportID = FName(TEXT("Viewport"));
+}
+
+void FPCGEditor::Initialize(const EToolkitMode::Type InMode, const TSharedPtr<class IToolkitHost>& InToolkitHost, UPCGGraph* InPCGGraph)
+{
+	PCGGraphBeingEdited = InPCGGraph;
+
+	PCGEditorGraph = NewObject<UPCGEditorGraph>(PCGGraphBeingEdited, UPCGEditorGraph::StaticClass(), NAME_None, RF_Transactional | RF_Transient);
+	PCGEditorGraph->Schema = UPCGEditorGraphSchema::StaticClass();
+	PCGEditorGraph->InitFromNodeGraph(InPCGGraph);
+
+	FPropertyEditorModule& PropertyEditorModule = FModuleManager::GetModuleChecked<FPropertyEditorModule>("PropertyEditor");
+
+	FDetailsViewArgs DetailsViewArgs;
+	DetailsViewArgs.NameAreaSettings = FDetailsViewArgs::HideNameArea;
+	DetailsViewArgs.bHideSelectionTip = true;
+	PropertyDetailsWidget = PropertyEditorModule.CreateDetailView(DetailsViewArgs);
+
+	GraphEditorWidget = CreateGraphEditorWidget();
+
+	const TSharedRef<FTabManager::FLayout> StandaloneDefaultLayout = FTabManager::NewLayout("Standalone_PCGGraphEditor_Layout_v0.3")
+		->AddArea
+		(
+			FTabManager::NewPrimaryArea()->SetOrientation(Orient_Horizontal)
+			->Split
+			(
+				
+				FTabManager::NewSplitter()->SetOrientation(Orient_Vertical)
+				->SetSizeCoefficient(0.10f)
+				->Split
+				(
+					FTabManager::NewStack()
+					->SetSizeCoefficient(0.16)
+					->SetHideTabWell(true)
+					->AddTab(FPCGEditor_private::ViewportID, ETabState::OpenedTab)
+				)
+				->Split
+				(
+					FTabManager::NewStack()
+					->SetSizeCoefficient(0.84)
+					->SetHideTabWell(true)
+					->AddTab(FPCGEditor_private::LibraryID, ETabState::OpenedTab)
+				)
+			)
+			->Split
+			(
+				FTabManager::NewSplitter()->SetOrientation(Orient_Vertical)
+				->SetSizeCoefficient(0.70f)
+				->Split
+				(
+					FTabManager::NewStack()
+					->SetSizeCoefficient(0.72)
+					->SetHideTabWell(true)
+					->AddTab(FPCGEditor_private::GraphEditorID, ETabState::OpenedTab)
+				)
+				->Split
+				(
+					FTabManager::NewStack()
+					->SetSizeCoefficient(0.28)
+					->SetHideTabWell(true)
+					->AddTab(FPCGEditor_private::AttributesID, ETabState::OpenedTab)
+				)
+			)
+			->Split
+			(
+				FTabManager::NewStack()
+				->SetSizeCoefficient(0.20f)
+				->SetHideTabWell(true)
+				->AddTab(FPCGEditor_private::PropertyDetailsID, ETabState::OpenedTab)
+			)
+		);
+
+	const FName PCGGraphEditorAppName = FName(TEXT("PCGEditorApp"));
+
+	InitAssetEditor(InMode, InToolkitHost, PCGGraphEditorAppName, StandaloneDefaultLayout, /*bCreateDefaultStandaloneMenu=*/ true, /*bCreateDefaultToolbar=*/ true, InPCGGraph);
+}
+
+void FPCGEditor::RegisterTabSpawners(const TSharedRef<class FTabManager>& InTabManager)
+{
+	WorkspaceMenuCategory = InTabManager->AddLocalWorkspaceMenuCategory(LOCTEXT("WorkspaceMenu_PCGEditor", "PCG Editor"));
+	auto WorkspaceMenuCategoryRef = WorkspaceMenuCategory.ToSharedRef();
+
+	FAssetEditorToolkit::RegisterTabSpawners(InTabManager);
+
+	//TODO: Add Icons
+	InTabManager->RegisterTabSpawner(FPCGEditor_private::GraphEditorID, FOnSpawnTab::CreateSP(this, &FPCGEditor::SpawnTab_GraphEditor))
+		.SetDisplayName(LOCTEXT("GraphTab", "Graph"))
+		.SetGroup(WorkspaceMenuCategoryRef);
+
+	InTabManager->RegisterTabSpawner(FPCGEditor_private::PropertyDetailsID, FOnSpawnTab::CreateSP(this, &FPCGEditor::SpawnTab_PropertyDetails))
+		.SetDisplayName(LOCTEXT("DetailsTab", "Details"))
+		.SetGroup(WorkspaceMenuCategoryRef);
+
+	InTabManager->RegisterTabSpawner(FPCGEditor_private::LibraryID, FOnSpawnTab::CreateSP(this, &FPCGEditor::SpawnTab_Library))
+		.SetDisplayName(LOCTEXT("LibraryTab", "Library"))
+		.SetGroup(WorkspaceMenuCategoryRef);
+
+	InTabManager->RegisterTabSpawner(FPCGEditor_private::AttributesID, FOnSpawnTab::CreateSP(this, &FPCGEditor::SpawnTab_Attributes))
+		.SetDisplayName(LOCTEXT("AttributesTab", "Attributes"))
+		.SetGroup(WorkspaceMenuCategoryRef);
+
+	InTabManager->RegisterTabSpawner(FPCGEditor_private::ViewportID, FOnSpawnTab::CreateSP(this, &FPCGEditor::SpawnTab_Viewport))
+		.SetDisplayName(LOCTEXT("ViewportTab", "Viewport"))
+		.SetGroup(WorkspaceMenuCategoryRef);
+}
+
+void FPCGEditor::UnregisterTabSpawners(const TSharedRef<class FTabManager>& InTabManager)
+{
+	InTabManager->UnregisterTabSpawner(FPCGEditor_private::GraphEditorID);
+	InTabManager->UnregisterTabSpawner(FPCGEditor_private::PropertyDetailsID);
+	InTabManager->UnregisterTabSpawner(FPCGEditor_private::LibraryID);
+	InTabManager->UnregisterTabSpawner(FPCGEditor_private::AttributesID);
+	InTabManager->UnregisterTabSpawner(FPCGEditor_private::ViewportID);
+
+	FAssetEditorToolkit::UnregisterTabSpawners(InTabManager);
+}
+
+FName FPCGEditor::GetToolkitFName() const
+{
+	return FName(TEXT("PCGEditor"));
+}
+
+FText FPCGEditor::GetBaseToolkitName() const
+{
+	return LOCTEXT("AppLabel", "PCG Editor");
+}
+
+FLinearColor FPCGEditor::GetWorldCentricTabColorScale() const
+{
+	return FLinearColor::White;
+}
+
+FString FPCGEditor::GetWorldCentricTabPrefix() const
+{
+	return LOCTEXT("WorldCentricTabPrefix", "PCG ").ToString();
+}
+
+void FPCGEditor::SelectAllNodes()
+{
+	if (GraphEditorWidget.IsValid())
+	{
+		GraphEditorWidget->SelectAllNodes();
+	}
+}
+
+bool FPCGEditor::CanSelectAllNodes()
+{
+	return GraphEditorWidget.IsValid();
+}
+
+void FPCGEditor::OnAlignTop()
+{
+	if (GraphEditorWidget.IsValid())
+	{
+		GraphEditorWidget->OnAlignTop();
+	}
+}
+
+void FPCGEditor::OnAlignMiddle()
+{
+	if (GraphEditorWidget.IsValid())
+	{
+		GraphEditorWidget->OnAlignMiddle();
+	}
+}
+
+void FPCGEditor::OnAlignBottom()
+{
+	if (GraphEditorWidget.IsValid())
+	{
+		GraphEditorWidget->OnAlignBottom();
+	}
+}
+
+void FPCGEditor::OnAlignLeft()
+{
+	if (GraphEditorWidget.IsValid())
+	{
+		GraphEditorWidget->OnAlignLeft();
+	}
+}
+
+void FPCGEditor::OnAlignCenter()
+{
+	if (GraphEditorWidget.IsValid())
+	{
+		GraphEditorWidget->OnAlignCenter();
+	}
+}
+
+void FPCGEditor::OnAlignRight()
+{
+	if (GraphEditorWidget.IsValid())
+	{
+		GraphEditorWidget->OnAlignRight();
+	}
+}
+
+void FPCGEditor::OnStraightenConnections()
+{
+	if (GraphEditorWidget.IsValid())
+	{
+		GraphEditorWidget->OnStraightenConnections();
+	}
+}
+
+void FPCGEditor::OnDistributeNodesH()
+{
+	if (GraphEditorWidget.IsValid())
+	{
+		GraphEditorWidget->OnDistributeNodesH();
+	}
+}
+
+void FPCGEditor::OnDistributeNodesV()
+{
+	if (GraphEditorWidget.IsValid())
+	{
+		GraphEditorWidget->OnDistributeNodesV();
+	}
+}
+
+TSharedRef<SGraphEditor> FPCGEditor::CreateGraphEditorWidget()
+{
+	GraphEditorCommands = MakeShareable(new FUICommandList);
+
+	// Editing commands
+	GraphEditorCommands->MapAction(FGenericCommands::Get().SelectAll,
+		FExecuteAction::CreateSP(this, &FPCGEditor::SelectAllNodes),
+		FCanExecuteAction::CreateSP(this, &FPCGEditor::CanSelectAllNodes));
+
+	// Alignment Commands
+	GraphEditorCommands->MapAction(FGraphEditorCommands::Get().AlignNodesTop,
+		FExecuteAction::CreateSP(this, &FPCGEditor::OnAlignTop)
+	);
+
+	GraphEditorCommands->MapAction(FGraphEditorCommands::Get().AlignNodesMiddle,
+		FExecuteAction::CreateSP(this, &FPCGEditor::OnAlignMiddle)
+	);
+
+	GraphEditorCommands->MapAction(FGraphEditorCommands::Get().AlignNodesBottom,
+		FExecuteAction::CreateSP(this, &FPCGEditor::OnAlignBottom)
+	);
+
+	GraphEditorCommands->MapAction(FGraphEditorCommands::Get().AlignNodesLeft,
+		FExecuteAction::CreateSP(this, &FPCGEditor::OnAlignLeft)
+	);
+
+	GraphEditorCommands->MapAction(FGraphEditorCommands::Get().AlignNodesCenter,
+		FExecuteAction::CreateSP(this, &FPCGEditor::OnAlignCenter)
+	);
+
+	GraphEditorCommands->MapAction(FGraphEditorCommands::Get().AlignNodesRight,
+		FExecuteAction::CreateSP(this, &FPCGEditor::OnAlignRight)
+	);
+
+	GraphEditorCommands->MapAction(FGraphEditorCommands::Get().StraightenConnections,
+		FExecuteAction::CreateSP(this, &FPCGEditor::OnStraightenConnections)
+	);
+
+	// Distribution Commands
+	GraphEditorCommands->MapAction(FGraphEditorCommands::Get().DistributeNodesHorizontally,
+		FExecuteAction::CreateSP(this, &FPCGEditor::OnDistributeNodesH)
+	);
+
+	GraphEditorCommands->MapAction(FGraphEditorCommands::Get().DistributeNodesVertically,
+		FExecuteAction::CreateSP(this, &FPCGEditor::OnDistributeNodesV)
+	);
+
+	FGraphAppearanceInfo AppearanceInfo;
+	AppearanceInfo.CornerText = LOCTEXT("PCGGraphEditorCornerText", "Procedural Graph");
+
+	SGraphEditor::FGraphEditorEvents InEvents;
+	InEvents.OnSelectionChanged = SGraphEditor::FOnSelectionChanged::CreateSP(this, &FPCGEditor::OnSelectedNodesChanged);
+
+	return SNew(SGraphEditor)
+		.AdditionalCommands(GraphEditorCommands)
+		.IsEditable(true)
+		.Appearance(AppearanceInfo)
+		.GraphToEdit(PCGEditorGraph)
+		.GraphEvents(InEvents)
+		.ShowGraphStateOverlay(false);
+}
+
+void FPCGEditor::OnSelectedNodesChanged(const TSet<UObject*>& NewSelection)
+{
+	TArray<TWeakObjectPtr<UObject>> SelectedObjects;
+
+	for (UObject* Object : NewSelection)
+	{
+		if (UPCGEditorGraphNode* GraphNode = Cast<UPCGEditorGraphNode>(Object))
+		{
+			SelectedObjects.Add(GraphNode->GetPCGNode()->DefaultSettings);
+		}
+	}
+
+	PropertyDetailsWidget->SetObjects(SelectedObjects, /*bForceRefresh=*/true);
+
+	GetTabManager()->TryInvokeTab(FPCGEditor_private::PropertyDetailsID);
+}
+
+TSharedRef<SDockTab> FPCGEditor::SpawnTab_GraphEditor(const FSpawnTabArgs& Args)
+{
+	return SNew(SDockTab)
+		.Label(LOCTEXT("PCGGraphTitle", "Graph"))
+		.TabColorScale(GetTabColorScale())
+		[
+			GraphEditorWidget.ToSharedRef()
+		];
+}
+
+TSharedRef<SDockTab> FPCGEditor::SpawnTab_PropertyDetails(const FSpawnTabArgs& Args)
+{
+	return SNew(SDockTab)
+		.Label(LOCTEXT("PCGDetailsTitle", "Details"))
+		.TabColorScale(GetTabColorScale())
+		[
+			PropertyDetailsWidget.ToSharedRef()
+		];
+}
+
+TSharedRef<SDockTab> FPCGEditor::SpawnTab_Library(const FSpawnTabArgs& Args)
+{
+	return SNew(SDockTab)
+		.Label(LOCTEXT("PCGLibraryTitle", "Library"))
+		.TabColorScale(GetTabColorScale())
+		[
+			SNullWidget::NullWidget
+		];
+}
+
+TSharedRef<SDockTab> FPCGEditor::SpawnTab_Attributes(const FSpawnTabArgs& Args)
+{
+	return SNew(SDockTab)
+		.Label(LOCTEXT("PCGAttributesTitle", "Attributes"))
+		.TabColorScale(GetTabColorScale())
+		[
+			SNullWidget::NullWidget
+		];
+}
+
+TSharedRef<SDockTab> FPCGEditor::SpawnTab_Viewport(const FSpawnTabArgs& Args)
+{
+	return SNew(SDockTab)
+		.Label(LOCTEXT("PCGViewportTitle", "Viewport"))
+		.TabColorScale(GetTabColorScale())
+		[
+			SNullWidget::NullWidget
+		];
+}
+
+#undef LOCTEXT_NAMESPACE
diff --git a/Engine/Plugins/Experimental/PCG/Source/PCGEditor/Private/PCGEditor.h b/Engine/Plugins/Experimental/PCG/Source/PCGEditor/Private/PCGEditor.h
new file mode 100644
--- /dev/null
+++ b/Engine/Plugins/Experimental/PCG/Source/PCGEditor/Private/PCGEditor.h
@@ -0,0 +1,66 @@
+// Copyright Epic Games, Inc. All Rights Reserved.
+
+#pragma once
+
+#include "Toolkits/AssetEditorToolkit.h"
+
+class UPCGGraph;
+class UPCGEditorGraph;
+class SGraphEditor;
+class IDetailsView;
+class FUICommandList;
+
+class FPCGEditor : public FAssetEditorToolkit
+{
+public:
+	/** Edits the specified PCGGraph */
+	void Initialize(const EToolkitMode::Type InMode, const TSharedPtr<class IToolkitHost>& InToolkitHost, UPCGGraph* InPCGGraph);
+
+	// ~Begin IToolkit interface
+	virtual void RegisterTabSpawners(const TSharedRef<class FTabManager>& InTabManager) override;
+	virtual void UnregisterTabSpawners(const TSharedRef<class FTabManager>& InTabManager) override;
+	// ~End IToolkit interface
+
+	// ~Begin FAssetEditorToolkit interface
+	virtual FName GetToolkitFName() const override;
+	virtual FText GetBaseToolkitName() const override;
+	virtual FLinearColor GetWorldCentricTabColorScale() const override;
+	virtual FString GetWorldCentricTabPrefix() const override;
+	// ~End FAssetEditorToolkit interface
+
+private:
+	/** Select every node in the graph */
+	void SelectAllNodes();
+	/** Whether we can select every node */
+	bool CanSelectAllNodes();
+
+	void OnAlignTop();
+	void OnAlignMiddle();
+	void OnAlignBottom();
+	void OnAlignLeft();
+	void OnAlignCenter();
+	void OnAlignRight();
+	void OnStraightenConnections();
+	void OnDistributeNodesH();
+	void OnDistributeNodesV();
+
+	/** Create new graph editor widget */
+	TSharedRef<SGraphEditor> CreateGraphEditorWidget();
+
+	/** Called when the selection changes in the GraphEditor */
+	void OnSelectedNodesChanged(const TSet<UObject*>& NewSelection);
+
+	TSharedRef<SDockTab> SpawnTab_GraphEditor(const FSpawnTabArgs& Args);
+	TSharedRef<SDockTab> SpawnTab_PropertyDetails(const FSpawnTabArgs& Args);
+	TSharedRef<SDockTab> SpawnTab_Library(const FSpawnTabArgs& Args);
+	TSharedRef<SDockTab> SpawnTab_Attributes(const FSpawnTabArgs& Args);
+	TSharedRef<SDockTab> SpawnTab_Viewport(const FSpawnTabArgs& Args);
+
+	TSharedPtr<SGraphEditor> GraphEditorWidget;
+	TSharedPtr<IDetailsView> PropertyDetailsWidget;
+
+	TSharedPtr<FUICommandList> GraphEditorCommands;
+
+	UPCGGraph* PCGGraphBeingEdited = nullptr;
+	UPCGEditorGraph* PCGEditorGraph = nullptr;
+};
\ No newline at end of file
diff --git a/Engine/Plugins/Experimental/PCG/Source/PCGEditor/Private/PCGEditorGraph.cpp b/Engine/Plugins/Experimental/PCG/Source/PCGEditor/Private/PCGEditorGraph.cpp
new file mode 100644
--- /dev/null
+++ b/Engine/Plugins/Experimental/PCG/Source/PCGEditor/Private/PCGEditorGraph.cpp
@@ -0,0 +1,62 @@
+// Copyright Epic Games, Inc. All Rights Reserved.
+
+#include "PCGEditorGraph.h"
+
+#include "PCGGraph.h"
+#include "PCGEditorGraphNode.h"
+
+#include "EdGraph/EdGraphPin.h"
+
+void UPCGEditorGraph::InitFromNodeGraph(UPCGGraph* InPCGGraph)
+{
+	PCGGraph = InPCGGraph;
+
+	TMap<UPCGNode*, UPCGEditorGraphNode*> NodeLookup;
+	const bool bSelectNewNode = false;
+
+	UPCGNode* InputNode = PCGGraph->GetInputNode();
+	FGraphNodeCreator<UPCGEditorGraphNode> InputNodeCreator(*this);
+	UPCGEditorGraphNode* InputGraphNode = InputNodeCreator.CreateNode(bSelectNewNode);
+	InputGraphNode->Construct(InputNode, EPCGEditorGraphNodeType::Input);
+	InputNodeCreator.Finalize();
+	NodeLookup.Add(InputNode, InputGraphNode);
+
+	UPCGNode* OutputNode = PCGGraph->GetOutputNode();
+	FGraphNodeCreator<UPCGEditorGraphNode> OutputNodeCreator(*this);
+	UPCGEditorGraphNode* OutputGraphNode = OutputNodeCreator.CreateNode(bSelectNewNode);
+	OutputGraphNode->Construct(OutputNode, EPCGEditorGraphNodeType::Output);
+	OutputNodeCreator.Finalize();
+	NodeLookup.Add(OutputNode, OutputGraphNode);
+
+	for (UPCGNode* PCGNode : PCGGraph->GetNodes())
+	{
+		FGraphNodeCreator<UPCGEditorGraphNode> NodeCreator(*this);
+		UPCGEditorGraphNode* GraphNode = NodeCreator.CreateNode(bSelectNewNode);
+		GraphNode->Construct(PCGNode, EPCGEditorGraphNodeType::Settings);
+		NodeCreator.Finalize();
+		NodeLookup.Add(PCGNode, GraphNode);
+	}
+
+	for (const auto& NodeLookupIt : NodeLookup)
+	{
+		UPCGNode* PCGNode = NodeLookupIt.Key;
+		UPCGEditorGraphNode* GraphNode = NodeLookupIt.Value;
+
+		UEdGraphPin* OutPin = GraphNode->FindPin(TEXT("Out"));
+		if (!OutPin)
+		{
+			continue;
+		}
+
+		for (UPCGNode* OutboundNode : PCGNode->GetOutboundNodes())
+		{
+			if (UPCGEditorGraphNode** ConnectedGraphNode = NodeLookup.Find(OutboundNode))
+			{
+				if (UEdGraphPin* InPin = (*ConnectedGraphNode)->FindPin(TEXT("In")))
+				{
+					OutPin->MakeLinkTo(InPin);
+				}
+			}
+		}
+	}
+}
diff --git a/Engine/Plugins/Experimental/PCG/Source/PCGEditor/Private/PCGEditorGraph.h b/Engine/Plugins/Experimental/PCG/Source/PCGEditor/Private/PCGEditorGraph.h
new file mode 100644
--- /dev/null
+++ b/Engine/Plugins/Experimental/PCG/Source/PCGEditor/Private/PCGEditorGraph.h
@@ -0,0 +1,27 @@
+// Copyright Epic Games, Inc. All Rights Reserved.
+
+#pragma once
+
+#include "CoreMinimal.h"
+#include "UObject/ObjectMacros.h"
+#include "EdGraph/EdGraph.h"
+
+#include "PCGEditorGraph.generated.h"
+
+class UPCGGraph;
+
+UCLASS()
+class UPCGEditorGraph : public UEdGraph
+{
+	GENERATED_BODY()
+
+public:
+	/** Initialize the editor graph from a PCGGraph */
+	void InitFromNodeGraph(UPCGGraph* InPCGGraph);
+
+	UPCGGraph* GetPCGGraph() { return PCGGraph; }
+
+private:
+	UPROPERTY()
+	TObjectPtr<UPCGGraph> PCGGraph = nullptr;
+};
diff --git a/Engine/Plugins/Experimental/PCG/Source/PCGEditor/Private/PCGEditorGraphNode.cpp b/Engine/Plugins/Experimental/PCG/Source/PCGEditor/Private/PCGEditorGraphNode.cpp
new file mode 100644
--- /dev/null
+++ b/Engine/Plugins/Experimental/PCG/Source/PCGEditor/Private/PCGEditorGraphNode.cpp
@@ -0,0 +1,75 @@
+// Copyright Epic Games, Inc. All Rights Reserved.
+
+#include "PCGEditorGraphNode.h"
+
+#include "EdGraph/EdGraphPin.h"
+#include "GraphEditorActions.h"
+#include "ToolMenu.h"
+#include "ToolMenuSection.h"
+
+#include "PCGNode.h"
+#include "PCGSettings.h"
+
+#define LOCTEXT_NAMESPACE "PCGEditorGraphNode"
+
+void UPCGEditorGraphNode::Construct(UPCGNode* InPCGNode, EPCGEditorGraphNodeType InNodeType)
+{
+	check(InPCGNode);
+	PCGNode = InPCGNode;
+
+	NodePosX = InPCGNode->PositionX;
+	NodePosY = InPCGNode->PositionY;
+	
+	NodeType = InNodeType;
+}
+
+FText UPCGEditorGraphNode::GetNodeTitle(ENodeTitleType::Type TitleType) const
+{
+	check(PCGNode && PCGNode->DefaultSettings);
+	return FText::FromName(PCGNode->DefaultSettings->GetDefaultNodeName());
+}
+
+void UPCGEditorGraphNode::GetNodeContextMenuActions(UToolMenu* Menu, class UGraphNodeContextMenuContext* Context) const
+{
+	if (!Context->Node)
+	{
+		return;
+	}
+
+	FToolMenuSection& Section = Menu->AddSection("EdGraphSchemaOrganization", LOCTEXT("OrganizationHeader", "Organization"));
+
+	Section.AddSubMenu("Alignment", LOCTEXT("AlignmentHeader", "Alignment"), FText(), FNewToolMenuDelegate::CreateLambda([](UToolMenu* AlignmentMenu)
+		{
+			{
+				FToolMenuSection& InSection = AlignmentMenu->AddSection("EdGraphSchemaAlignment", LOCTEXT("AlignHeader", "Align"));
+				InSection.AddMenuEntry(FGraphEditorCommands::Get().AlignNodesTop);
+				InSection.AddMenuEntry(FGraphEditorCommands::Get().AlignNodesMiddle);
+				InSection.AddMenuEntry(FGraphEditorCommands::Get().AlignNodesBottom);
+				InSection.AddMenuEntry(FGraphEditorCommands::Get().AlignNodesLeft);
+				InSection.AddMenuEntry(FGraphEditorCommands::Get().AlignNodesCenter);
+				InSection.AddMenuEntry(FGraphEditorCommands::Get().AlignNodesRight);
+				InSection.AddMenuEntry(FGraphEditorCommands::Get().StraightenConnections);
+			}
+
+			{
+				FToolMenuSection& InSection = AlignmentMenu->AddSection("EdGraphSchemaDistribution", LOCTEXT("DistributionHeader", "Distribution"));
+				InSection.AddMenuEntry(FGraphEditorCommands::Get().DistributeNodesHorizontally);
+				InSection.AddMenuEntry(FGraphEditorCommands::Get().DistributeNodesVertically);
+			}
+		}));
+}
+
+void UPCGEditorGraphNode::AllocateDefaultPins()
+{
+	if (NodeType == EPCGEditorGraphNodeType::Input || NodeType == EPCGEditorGraphNodeType::Settings)
+	{
+		CreatePin(EEdGraphPinDirection::EGPD_Output, NAME_None, FName(TEXT("Out")));
+	}
+
+	if (NodeType == EPCGEditorGraphNodeType::Output || NodeType == EPCGEditorGraphNodeType::Settings)
+	{
+		CreatePin(EEdGraphPinDirection::EGPD_Input, NAME_None, FName(TEXT("In")));
+	}
+}
+
+#undef LOCTEXT_NAMESPACE
diff --git a/Engine/Plugins/Experimental/PCG/Source/PCGEditor/Private/PCGEditorGraphNode.h b/Engine/Plugins/Experimental/PCG/Source/PCGEditor/Private/PCGEditorGraphNode.h
new file mode 100644
--- /dev/null
+++ b/Engine/Plugins/Experimental/PCG/Source/PCGEditor/Private/PCGEditorGraphNode.h
@@ -0,0 +1,44 @@
+// Copyright Epic Games, Inc. All Rights Reserved.
+
+#pragma once
+
+#include "CoreMinimal.h"
+#include "UObject/ObjectMacros.h"
+#include "EdGraph/EdGraphNode.h"
+
+#include "PCGEditorGraphNode.generated.h"
+
+class UPCGNode;
+class UToolMenu;
+
+UENUM()
+enum class EPCGEditorGraphNodeType : uint8
+{
+	Input,
+	Output,
+	Settings
+};
+
+UCLASS()
+class UPCGEditorGraphNode : public UEdGraphNode
+{
+	GENERATED_BODY()
+
+public:
+	void Construct(UPCGNode* InPCGNode, EPCGEditorGraphNodeType InNodeType);
+
+	// ~Begin UEdGraphNode interface
+	virtual FText GetNodeTitle(ENodeTitleType::Type TitleType) const override;
+	virtual void GetNodeContextMenuActions(UToolMenu* Menu, class UGraphNodeContextMenuContext* Context) const override;
+	virtual void AllocateDefaultPins() override;
+	// ~End UEdGraphNode interface
+
+	UPCGNode* GetPCGNode() { return PCGNode; }
+
+private:
+	UPROPERTY()
+	TObjectPtr<UPCGNode> PCGNode = nullptr;
+
+ 	UPROPERTY()
+ 	EPCGEditorGraphNodeType NodeType = EPCGEditorGraphNodeType::Settings;
+};
diff --git a/Engine/Plugins/Experimental/PCG/Source/PCGEditor/Private/PCGEditorGraphNodeFactory.cpp b/Engine/Plugins/Experimental/PCG/Source/PCGEditor/Private/PCGEditorGraphNodeFactory.cpp
new file mode 100644
--- /dev/null
+++ b/Engine/Plugins/Experimental/PCG/Source/PCGEditor/Private/PCGEditorGraphNodeFactory.cpp
@@ -0,0 +1,21 @@
+// Copyright Epic Games, Inc. All Rights Reserved.
+
+#include "PCGEditorGraphNodeFactory.h"
+
+#include "PCGEditorGraphNode.h"
+#include "SPCGEditorGraphNode.h"
+
+TSharedPtr<SGraphNode> FPCGEditorGraphNodeFactory::CreateNode(UEdGraphNode* InNode) const
+{
+	if (UPCGEditorGraphNode* GraphNode = Cast<UPCGEditorGraphNode>(InNode))
+	{
+		TSharedRef<SGraphNode> VisualNode =
+			SNew(SPCGEditorGraphNode, GraphNode);
+
+		VisualNode->SlatePrepass();
+
+		return VisualNode;
+	}
+
+	return nullptr;
+}
diff --git a/Engine/Plugins/Experimental/PCG/Source/PCGEditor/Private/PCGEditorGraphNodeFactory.h b/Engine/Plugins/Experimental/PCG/Source/PCGEditor/Private/PCGEditorGraphNodeFactory.h
new file mode 100644
--- /dev/null
+++ b/Engine/Plugins/Experimental/PCG/Source/PCGEditor/Private/PCGEditorGraphNodeFactory.h
@@ -0,0 +1,13 @@
+// Copyright Epic Games, Inc. All Rights Reserved.
+
+#pragma once
+
+#include "EdGraphUtilities.h"
+
+class SGraphNode;
+class UEdGraphNode;
+
+class FPCGEditorGraphNodeFactory : public FGraphPanelNodeFactory
+{
+	virtual TSharedPtr<SGraphNode> CreateNode(UEdGraphNode* InNode) const override;
+};
diff --git a/Engine/Plugins/Experimental/PCG/Source/PCGEditor/Private/PCGEditorGraphSchema.cpp b/Engine/Plugins/Experimental/PCG/Source/PCGEditor/Private/PCGEditorGraphSchema.cpp
new file mode 100644
--- /dev/null
+++ b/Engine/Plugins/Experimental/PCG/Source/PCGEditor/Private/PCGEditorGraphSchema.cpp
@@ -0,0 +1,47 @@
+// Copyright Epic Games, Inc. All Rights Reserved.
+
+#include "PCGEditorGraphSchema.h"
+
+#include "PCGEditorGraph.h"
+#include "PCGEditorGraphNode.h"
+#include "PCGEditorGraphSchemaActions.h"
+#include "PCGSettings.h"
+
+#include "UObject/UObjectIterator.h"
+
+#define LOCTEXT_NAMESPACE "PCGEditorGraphSchema"
+
+void UPCGEditorGraphSchema::GetGraphContextActions(FGraphContextMenuBuilder& ContextMenuBuilder) const
+{
+	Super::GetGraphContextActions(ContextMenuBuilder);
+
+	const UPCGEditorGraph* PCGEditorGraph = CastChecked<UPCGEditorGraph>(ContextMenuBuilder.CurrentGraph);
+
+	TArray<UClass*> SettingsClasses;
+	for (TObjectIterator<UClass> It; It; ++It)
+	{
+		UClass* Class = *It;
+		
+		if (Class->IsChildOf(UPCGSettings::StaticClass()) &&
+			!Class->HasAnyClassFlags(CLASS_Abstract | CLASS_Deprecated | CLASS_Hidden))
+		{
+			SettingsClasses.Add(Class);
+		}
+	}
+
+	for (UClass* SettingsClass : SettingsClasses)
+	{
+		const FText MenuDesc = FText::FromName(SettingsClass->GetFName());
+
+		TSharedPtr<FPCGEditorGraphSchemaAction_NewNode> NewAction(new FPCGEditorGraphSchemaAction_NewNode(FText::GetEmpty(), MenuDesc, FText::GetEmpty(), 0));
+		NewAction->SettingsClass = SettingsClass;
+		ContextMenuBuilder.AddAction(NewAction);
+	}
+}
+
+FLinearColor UPCGEditorGraphSchema::GetPinTypeColor(const FEdGraphPinType& PinType) const
+{
+	return FLinearColor(1.0f, 0.0f, 1.0f, 1.0f);
+}
+
+#undef LOCTEXT_NAMESPACE
diff --git a/Engine/Plugins/Experimental/PCG/Source/PCGEditor/Private/PCGEditorGraphSchema.h b/Engine/Plugins/Experimental/PCG/Source/PCGEditor/Private/PCGEditorGraphSchema.h
new file mode 100644
--- /dev/null
+++ b/Engine/Plugins/Experimental/PCG/Source/PCGEditor/Private/PCGEditorGraphSchema.h
@@ -0,0 +1,19 @@
+// Copyright Epic Games, Inc. All Rights Reserved.
+
+#pragma once
+
+#include "EdGraph/EdGraphSchema.h"
+
+#include "PCGEditorGraphSchema.generated.h"
+
+UCLASS()
+class UPCGEditorGraphSchema : public UEdGraphSchema
+{
+	GENERATED_BODY()
+
+public:
+	//~ Begin EdGraphSchema Interface
+	virtual void GetGraphContextActions(FGraphContextMenuBuilder& ContextMenuBuilder) const;
+	virtual FLinearColor GetPinTypeColor(const FEdGraphPinType& PinType) const override;
+	//~ End EdGraphSchema Interface
+};
diff --git a/Engine/Plugins/Experimental/PCG/Source/PCGEditor/Private/PCGEditorGraphSchemaActions.cpp b/Engine/Plugins/Experimental/PCG/Source/PCGEditor/Private/PCGEditorGraphSchemaActions.cpp
new file mode 100644
--- /dev/null
+++ b/Engine/Plugins/Experimental/PCG/Source/PCGEditor/Private/PCGEditorGraphSchemaActions.cpp
@@ -0,0 +1,37 @@
+// Copyright Epic Games, Inc. All Rights Reserved.
+
+#include "PCGEditorGraphSchemaActions.h"
+
+#include "PCGGraph.h"
+#include "PCGEditorGraph.h"
+#include "PCGEditorGraphNode.h"
+#include "PCGSettings.h"
+
+#include "EdGraph/EdGraph.h"
+
+UEdGraphNode* FPCGEditorGraphSchemaAction_NewNode::PerformAction(UEdGraph* ParentGraph, UEdGraphPin* FromPin, const FVector2D Location, bool bSelectNewNode)
+{
+	UPCGEditorGraph* EditorGraph = Cast<UPCGEditorGraph>(ParentGraph);
+	if (!EditorGraph)
+	{
+		return nullptr;
+	}
+
+	UPCGGraph* PCGGraph = EditorGraph->GetPCGGraph();
+	if (!PCGGraph)
+	{
+		return nullptr;
+	}
+
+	UPCGSettings* DefaultNodeSettings = nullptr;
+	UPCGNode* NewPCGNode = PCGGraph->AddNodeOfType(SettingsClass, DefaultNodeSettings);
+
+	FGraphNodeCreator<UPCGEditorGraphNode> NodeCreator(*EditorGraph);
+	UPCGEditorGraphNode* NewNode = NodeCreator.CreateNode(bSelectNewNode);
+	NewNode->Construct(NewPCGNode, EPCGEditorGraphNodeType::Settings);
+	NewNode->NodePosX = Location.X;
+	NewNode->NodePosY = Location.Y;
+	NodeCreator.Finalize();
+
+	return NewNode;
+}
diff --git a/Engine/Plugins/Experimental/PCG/Source/PCGEditor/Private/PCGEditorGraphSchemaActions.h b/Engine/Plugins/Experimental/PCG/Source/PCGEditor/Private/PCGEditorGraphSchemaActions.h
new file mode 100644
--- /dev/null
+++ b/Engine/Plugins/Experimental/PCG/Source/PCGEditor/Private/PCGEditorGraphSchemaActions.h
@@ -0,0 +1,27 @@
+// Copyright Epic Games, Inc. All Rights Reserved.
+
+#pragma once
+
+#include "EdGraph/EdGraphSchema.h"
+#include "Templates/SubclassOf.h"
+
+#include "PCGEditorGraphSchemaActions.generated.h"
+
+class UEdGraph;
+class UEdGraphPin;
+class UPCGEditorGraphNode;
+class UPCGSettings;
+
+USTRUCT()
+struct FPCGEditorGraphSchemaAction_NewNode : public FEdGraphSchemaAction
+{
+	GENERATED_BODY()
+
+	// Inherit the base class's constructors
+	using FEdGraphSchemaAction::FEdGraphSchemaAction;
+
+	UPROPERTY()
+	TSubclassOf<UPCGSettings> SettingsClass;
+
+	virtual UEdGraphNode* PerformAction(UEdGraph* ParentGraph, UEdGraphPin* FromPin, const FVector2D Location, bool bSelectNewNode = true) override;
+};
diff --git a/Engine/Plugins/Experimental/PCG/Source/PCGEditor/Private/PCGEditorModule.cpp b/Engine/Plugins/Experimental/PCG/Source/PCGEditor/Private/PCGEditorModule.cpp
--- a/Engine/Plugins/Experimental/PCG/Source/PCGEditor/Private/PCGEditorModule.cpp
+++ b/Engine/Plugins/Experimental/PCG/Source/PCGEditor/Private/PCGEditorModule.cpp
@@ -2,18 +2,20 @@
 
 #include "PCGEditorModule.h"
 
+#include "EdGraphUtilities.h"
+#include "Framework/MultiBox/MultiBoxBuilder.h"
 #include "IAssetTools.h"
 #include "LevelEditor.h"
 #include "Modules/ModuleManager.h"
 #include "PropertyEditorModule.h"
 #include "Toolkits/IToolkit.h"
-#include "Framework/MultiBox/MultiBoxBuilder.h"
 
+#include "AssetTypeActions/PCGGraphAssetTypeActions.h"
+#include "AssetTypeActions/PCGSettingsAssetTypeActions.h"
 #include "PCGComponentDetails.h"
+#include "PCGEditorGraphNodeFactory.h"
 #include "PCGGraphDetails.h"
 #include "PCGVolumeFactory.h"
-#include "AssetTypeActions/PCGGraphAssetTypeActions.h"
-#include "AssetTypeActions/PCGSettingsAssetTypeActions.h"
 
 #include "PCGSubsystem.h"
 
@@ -27,6 +29,9 @@ void FPCGEditorModule::StartupModule()
 	RegisterAssetTypeActions();
 	RegisterMenuExtensions();
 
+	GraphNodeFactory = MakeShareable(new FPCGEditorGraphNodeFactory());
+	FEdGraphUtilities::RegisterVisualNodeFactory(GraphNodeFactory);
+
 	if (GEditor)
 	{
 		GEditor->ActorFactories.Add(NewObject<UPCGVolumeFactory>());
@@ -39,6 +44,8 @@ void FPCGEditorModule::ShutdownModule()
 	UnregisterDetailsCustomizations();
 	UnregisterMenuExtensions();
 
+	FEdGraphUtilities::UnregisterVisualNodeFactory(GraphNodeFactory);
+
 	if (GEditor)
 	{
 		GEditor->ActorFactories.RemoveAll([](const UActorFactory* ActorFactory) { return ActorFactory->IsA<UPCGVolumeFactory>(); });
diff --git a/Engine/Plugins/Experimental/PCG/Source/PCGEditor/Private/PCGEditorModule.h b/Engine/Plugins/Experimental/PCG/Source/PCGEditor/Private/PCGEditorModule.h
--- a/Engine/Plugins/Experimental/PCG/Source/PCGEditor/Private/PCGEditorModule.h
+++ b/Engine/Plugins/Experimental/PCG/Source/PCGEditor/Private/PCGEditorModule.h
@@ -8,6 +8,7 @@
 #include "Toolkits/AssetEditorToolkit.h"
 
 class FMenuBuilder;
+class FPCGEditorGraphNodeFactory;
 
 class FPCGEditorModule : public IModuleInterface, public IHasMenuExtensibility, public IHasToolBarExtensibility
 {
@@ -43,4 +44,6 @@ protected:
 
 	TSharedPtr<FExtensibilityManager> MenuExtensibilityManager;
 	TSharedPtr<FExtensibilityManager> ToolBarExtensibilityManager;
+
+	TSharedPtr<FPCGEditorGraphNodeFactory> GraphNodeFactory;
 };
diff --git a/Engine/Plugins/Experimental/PCG/Source/PCGEditor/Private/SPCGEditorGraphNode.cpp b/Engine/Plugins/Experimental/PCG/Source/PCGEditor/Private/SPCGEditorGraphNode.cpp
new file mode 100644
--- /dev/null
+++ b/Engine/Plugins/Experimental/PCG/Source/PCGEditor/Private/SPCGEditorGraphNode.cpp
@@ -0,0 +1,26 @@
+// Copyright Epic Games, Inc. All Rights Reserved.
+
+#include "SPCGEditorGraphNode.h"
+
+#include "PCGEditorGraphNode.h"
+#include "PCGNode.h"
+
+void SPCGEditorGraphNode::Construct(const FArguments& InArgs, UPCGEditorGraphNode* InNode)
+{
+	GraphNode = InNode;
+	PCGEditorGraphNode = InNode;
+
+	UpdateGraphNode();
+}
+
+void SPCGEditorGraphNode::MoveTo(const FVector2D& NewPosition, FNodeSet& NodeFilter, bool bMarkDirty /*= true*/)
+{
+	SGraphNode::MoveTo(NewPosition, NodeFilter, bMarkDirty);
+
+	if (UPCGNode* PCGNode = PCGEditorGraphNode->GetPCGNode())
+	{
+		PCGNode->PositionX = PCGEditorGraphNode->NodePosX;
+		PCGNode->PositionY = PCGEditorGraphNode->NodePosY;
+		PCGNode->MarkPackageDirty();
+	}
+}
diff --git a/Engine/Plugins/Experimental/PCG/Source/PCGEditor/Private/SPCGEditorGraphNode.h b/Engine/Plugins/Experimental/PCG/Source/PCGEditor/Private/SPCGEditorGraphNode.h
new file mode 100644
--- /dev/null
+++ b/Engine/Plugins/Experimental/PCG/Source/PCGEditor/Private/SPCGEditorGraphNode.h
@@ -0,0 +1,24 @@
+// Copyright Epic Games, Inc. All Rights Reserved.
+
+#pragma once
+
+#include "SGraphNode.h"
+#include "Widgets/DeclarativeSyntaxSupport.h"
+
+class UPCGEditorGraphNode;
+
+class SPCGEditorGraphNode : public SGraphNode
+{
+public:
+	SLATE_BEGIN_ARGS(SPCGEditorGraphNode){}
+	SLATE_END_ARGS()
+
+	void Construct(const FArguments& InArgs, UPCGEditorGraphNode* InNode);
+
+	//~ Begin SGraphNode Interface
+	virtual void MoveTo(const FVector2D& NewPosition, FNodeSet& NodeFilter, bool bMarkDirty = true) override;
+	//~ End SGraphNode Interface
+
+private:
+	UPCGEditorGraphNode* PCGEditorGraphNode = nullptr;
+};
