Added blueprint library for pcg helpers
Added other blueprint code paths to have efficient loops

diff --git a/Engine/Plugins/Experimental/PCG/Source/PCG/Private/Data/PCGDifferenceData.cpp b/Engine/Plugins/Experimental/PCG/Source/PCG/Private/Data/PCGDifferenceData.cpp
--- a/Engine/Plugins/Experimental/PCG/Source/PCG/Private/Data/PCGDifferenceData.cpp
+++ b/Engine/Plugins/Experimental/PCG/Source/PCG/Private/Data/PCGDifferenceData.cpp
@@ -129,7 +129,7 @@ bool UPCGDifferenceData::HasNonTrivialTransform() const
 	return Source->HasNonTrivialTransform();
 }
 
-const UPCGPointData* UPCGDifferenceData::CreatePointData(FPCGContextPtr Context) const
+const UPCGPointData* UPCGDifferenceData::CreatePointData(FPCGContext* Context) const
 {
 	TRACE_CPUPROFILER_EVENT_SCOPE(UPCGDifferenceData::CreatePointData);
 	
diff --git a/Engine/Plugins/Experimental/PCG/Source/PCG/Private/Data/PCGIntersectionData.cpp b/Engine/Plugins/Experimental/PCG/Source/PCG/Private/Data/PCGIntersectionData.cpp
--- a/Engine/Plugins/Experimental/PCG/Source/PCG/Private/Data/PCGIntersectionData.cpp
+++ b/Engine/Plugins/Experimental/PCG/Source/PCG/Private/Data/PCGIntersectionData.cpp
@@ -98,7 +98,7 @@ bool UPCGIntersectionData::HasNonTrivialTransform() const
 	return A->HasNonTrivialTransform() || B->HasNonTrivialTransform();
 }
 
-const UPCGPointData* UPCGIntersectionData::CreatePointData(FPCGContextPtr Context) const
+const UPCGPointData* UPCGIntersectionData::CreatePointData(FPCGContext* Context) const
 {
 	check(A && B);
 	// TODO: this is a placeholder;
@@ -114,7 +114,7 @@ const UPCGPointData* UPCGIntersectionData::CreatePointData(FPCGContextPtr Contex
 	}
 }
 
-UPCGPointData* UPCGIntersectionData::CreateAndFilterPointData(FPCGContextPtr Context, const UPCGSpatialData* X, const UPCGSpatialData* Y) const
+UPCGPointData* UPCGIntersectionData::CreateAndFilterPointData(FPCGContext* Context, const UPCGSpatialData* X, const UPCGSpatialData* Y) const
 {
 	TRACE_CPUPROFILER_EVENT_SCOPE(UPCGIntersectionData::CreateAndFilterPointData);
 	check(X && Y);
diff --git a/Engine/Plugins/Experimental/PCG/Source/PCG/Private/Data/PCGLandscapeData.cpp b/Engine/Plugins/Experimental/PCG/Source/PCG/Private/Data/PCGLandscapeData.cpp
--- a/Engine/Plugins/Experimental/PCG/Source/PCG/Private/Data/PCGLandscapeData.cpp
+++ b/Engine/Plugins/Experimental/PCG/Source/PCG/Private/Data/PCGLandscapeData.cpp
@@ -77,7 +77,7 @@ float UPCGLandscapeData::GetDensityAtPosition(const FVector& InPosition) const
 	return HeightAtVertex.IsSet() ? 1.0f : 0;
 }
 
-const UPCGPointData* UPCGLandscapeData::CreatePointData(FPCGContextPtr Context) const
+const UPCGPointData* UPCGLandscapeData::CreatePointData(FPCGContext* Context) const
 {
 	TRACE_CPUPROFILER_EVENT_SCOPE(UPCGLandscapeData::CreatePointData);
 
diff --git a/Engine/Plugins/Experimental/PCG/Source/PCG/Private/Data/PCGLandscapeSplineData.cpp b/Engine/Plugins/Experimental/PCG/Source/PCG/Private/Data/PCGLandscapeSplineData.cpp
--- a/Engine/Plugins/Experimental/PCG/Source/PCG/Private/Data/PCGLandscapeSplineData.cpp
+++ b/Engine/Plugins/Experimental/PCG/Source/PCG/Private/Data/PCGLandscapeSplineData.cpp
@@ -88,7 +88,7 @@ FVector UPCGLandscapeSplineData::GetLocationAtDistance(int SegmentIndex, float D
 	return FVector::Zero();
 }
 
-const UPCGPointData* UPCGLandscapeSplineData::CreatePointData(FPCGContextPtr Context) const
+const UPCGPointData* UPCGLandscapeSplineData::CreatePointData(FPCGContext* Context) const
 {
 	check(Spline);
 	TRACE_CPUPROFILER_EVENT_SCOPE(UPCGLandscapeSplineData::CreatePointData);
diff --git a/Engine/Plugins/Experimental/PCG/Source/PCG/Private/Data/PCGProjectionData.cpp b/Engine/Plugins/Experimental/PCG/Source/PCG/Private/Data/PCGProjectionData.cpp
--- a/Engine/Plugins/Experimental/PCG/Source/PCG/Private/Data/PCGProjectionData.cpp
+++ b/Engine/Plugins/Experimental/PCG/Source/PCG/Private/Data/PCGProjectionData.cpp
@@ -87,7 +87,7 @@ bool UPCGProjectionData::HasNonTrivialTransform() const
 	return Target->HasNonTrivialTransform();
 }
 
-const UPCGPointData* UPCGProjectionData::CreatePointData(FPCGContextPtr Context) const
+const UPCGPointData* UPCGProjectionData::CreatePointData(FPCGContext* Context) const
 {
 	TRACE_CPUPROFILER_EVENT_SCOPE(UPCGProjectionData::CreatePointData);
 	// TODO: add mechanism in the ToPointData so we can pass in a transform
diff --git a/Engine/Plugins/Experimental/PCG/Source/PCG/Private/Data/PCGSpatialData.cpp b/Engine/Plugins/Experimental/PCG/Source/PCG/Private/Data/PCGSpatialData.cpp
--- a/Engine/Plugins/Experimental/PCG/Source/PCG/Private/Data/PCGSpatialData.cpp
+++ b/Engine/Plugins/Experimental/PCG/Source/PCG/Private/Data/PCGSpatialData.cpp
@@ -6,7 +6,7 @@
 #include "Data/PCGProjectionData.h"
 #include "Data/PCGUnionData.h"
 
-const UPCGPointData* UPCGSpatialDataWithPointCache::ToPointData(FPCGContextPtr Context) const
+const UPCGPointData* UPCGSpatialDataWithPointCache::ToPointData(FPCGContext* Context) const
 {
 	if (!CachedPointData)
 	{
diff --git a/Engine/Plugins/Experimental/PCG/Source/PCG/Private/Data/PCGSplineData.cpp b/Engine/Plugins/Experimental/PCG/Source/PCG/Private/Data/PCGSplineData.cpp
--- a/Engine/Plugins/Experimental/PCG/Source/PCG/Private/Data/PCGSplineData.cpp
+++ b/Engine/Plugins/Experimental/PCG/Source/PCG/Private/Data/PCGSplineData.cpp
@@ -43,7 +43,7 @@ FVector UPCGSplineData::GetLocationAtDistance(int SegmentIndex, float Distance)
 	return Spline->GetLocationAtDistanceAlongSpline(Distance, ESplineCoordinateSpace::World);
 }
 
-const UPCGPointData* UPCGSplineData::CreatePointData(FPCGContextPtr Context) const
+const UPCGPointData* UPCGSplineData::CreatePointData(FPCGContext* Context) const
 {
 	TRACE_CPUPROFILER_EVENT_SCOPE(UPCGSplineData::CreatePointData);
 	UPCGPointData* Data = NewObject<UPCGPointData>(const_cast<UPCGSplineData*>(this));
diff --git a/Engine/Plugins/Experimental/PCG/Source/PCG/Private/Data/PCGTextureData.cpp b/Engine/Plugins/Experimental/PCG/Source/PCG/Private/Data/PCGTextureData.cpp
--- a/Engine/Plugins/Experimental/PCG/Source/PCG/Private/Data/PCGTextureData.cpp
+++ b/Engine/Plugins/Experimental/PCG/Source/PCG/Private/Data/PCGTextureData.cpp
@@ -107,7 +107,7 @@ FPCGPoint UPCGBaseTextureData::TransformPoint(const FPCGPoint& InPoint) const
 	return Point;
 }
 
-const UPCGPointData* UPCGBaseTextureData::CreatePointData(FPCGContextPtr Context) const
+const UPCGPointData* UPCGBaseTextureData::CreatePointData(FPCGContext* Context) const
 {
 	TRACE_CPUPROFILER_EVENT_SCOPE(UPCGBaseTextureData::CreatePointData);
 	// TODO: this is a trivial implementation
diff --git a/Engine/Plugins/Experimental/PCG/Source/PCG/Private/Data/PCGUnionData.cpp b/Engine/Plugins/Experimental/PCG/Source/PCG/Private/Data/PCGUnionData.cpp
--- a/Engine/Plugins/Experimental/PCG/Source/PCG/Private/Data/PCGUnionData.cpp
+++ b/Engine/Plugins/Experimental/PCG/Source/PCG/Private/Data/PCGUnionData.cpp
@@ -149,7 +149,7 @@ bool UPCGUnionData::HasNonTrivialTransform() const
 	return (FirstNonTrivialTransformData != nullptr || Super::HasNonTrivialTransform());
 }
 
-const UPCGPointData* UPCGUnionData::CreatePointData(FPCGContextPtr Context) const
+const UPCGPointData* UPCGUnionData::CreatePointData(FPCGContext* Context) const
 {
 	TRACE_CPUPROFILER_EVENT_SCOPE(UPCGUnionData::CreatePointData);
 
@@ -195,7 +195,7 @@ const UPCGPointData* UPCGUnionData::CreatePointData(FPCGContextPtr Context) cons
 	return PointData;
 }
 
-void UPCGUnionData::CreateSequentialPointData(FPCGContextPtr Context, UPCGPointData* PointData, bool bLeftToRight) const
+void UPCGUnionData::CreateSequentialPointData(FPCGContext* Context, UPCGPointData* PointData, bool bLeftToRight) const
 {
 	check(PointData);
 
diff --git a/Engine/Plugins/Experimental/PCG/Source/PCG/Private/Data/PCGVolumeData.cpp b/Engine/Plugins/Experimental/PCG/Source/PCG/Private/Data/PCGVolumeData.cpp
--- a/Engine/Plugins/Experimental/PCG/Source/PCG/Private/Data/PCGVolumeData.cpp
+++ b/Engine/Plugins/Experimental/PCG/Source/PCG/Private/Data/PCGVolumeData.cpp
@@ -35,7 +35,7 @@ FBox UPCGVolumeData::GetStrictBounds() const
 	return StrictBounds;
 }
 
-const UPCGPointData* UPCGVolumeData::CreatePointData(FPCGContextPtr Context) const
+const UPCGPointData* UPCGVolumeData::CreatePointData(FPCGContext* Context) const
 {
 	UE_LOG(LogPCG, Error, TEXT("Volume data has no default point sampling"));
 	return nullptr;
diff --git a/Engine/Plugins/Experimental/PCG/Source/PCG/Private/Elements/PCGDebugElement.cpp b/Engine/Plugins/Experimental/PCG/Source/PCG/Private/Elements/PCGDebugElement.cpp
--- a/Engine/Plugins/Experimental/PCG/Source/PCG/Private/Elements/PCGDebugElement.cpp
+++ b/Engine/Plugins/Experimental/PCG/Source/PCG/Private/Elements/PCGDebugElement.cpp
@@ -12,7 +12,7 @@
 
 namespace PCGDebugElement
 {
-	void ExecuteDebugDisplay(FPCGContextPtr Context)
+	void ExecuteDebugDisplay(FPCGContext* Context)
 	{
 		TRACE_CPUPROFILER_EVENT_SCOPE(PCGDebugElement::ExecuteDebugDisplay);
 #if WITH_EDITORONLY_DATA
@@ -128,7 +128,7 @@ FPCGElementPtr UPCGDebugSettings::CreateElement() const
 	return MakeShared<FPCGDebugElement>();
 }
 
-bool FPCGDebugElement::ExecuteInternal(FPCGContextPtr Context) const
+bool FPCGDebugElement::ExecuteInternal(FPCGContext* Context) const
 {
 	TRACE_CPUPROFILER_EVENT_SCOPE(FPCGDebugElement::Execute);
 	PCGDebugElement::ExecuteDebugDisplay(Context);
diff --git a/Engine/Plugins/Experimental/PCG/Source/PCG/Private/Elements/PCGDensityFilter.cpp b/Engine/Plugins/Experimental/PCG/Source/PCG/Private/Elements/PCGDensityFilter.cpp
--- a/Engine/Plugins/Experimental/PCG/Source/PCG/Private/Elements/PCGDensityFilter.cpp
+++ b/Engine/Plugins/Experimental/PCG/Source/PCG/Private/Elements/PCGDensityFilter.cpp
@@ -9,7 +9,7 @@ FPCGElementPtr UPCGDensityFilterSettings::CreateElement() const
 	return MakeShared<FPCGDensityFilterElement>();
 }
 
-bool FPCGDensityFilterElement::ExecuteInternal(FPCGContextPtr Context) const
+bool FPCGDensityFilterElement::ExecuteInternal(FPCGContext* Context) const
 {
 	TRACE_CPUPROFILER_EVENT_SCOPE(FPCGDensityFilterElement::Execute);
 
diff --git a/Engine/Plugins/Experimental/PCG/Source/PCG/Private/Elements/PCGDensityRemap.cpp b/Engine/Plugins/Experimental/PCG/Source/PCG/Private/Elements/PCGDensityRemap.cpp
--- a/Engine/Plugins/Experimental/PCG/Source/PCG/Private/Elements/PCGDensityRemap.cpp
+++ b/Engine/Plugins/Experimental/PCG/Source/PCG/Private/Elements/PCGDensityRemap.cpp
@@ -11,7 +11,7 @@ FPCGElementPtr UPCGLinearDensityRemapSettings::CreateElement() const
 	return MakeShared<FPCGLinearDensityRemapElement>();
 }
 
-bool FPCGLinearDensityRemapElement::ExecuteInternal(FPCGContextPtr Context) const
+bool FPCGLinearDensityRemapElement::ExecuteInternal(FPCGContext* Context) const
 {
 	TRACE_CPUPROFILER_EVENT_SCOPE(FPCGLinearDensityRemapElement::Execute);
 
diff --git a/Engine/Plugins/Experimental/PCG/Source/PCG/Private/Elements/PCGDifferenceElement.cpp b/Engine/Plugins/Experimental/PCG/Source/PCG/Private/Elements/PCGDifferenceElement.cpp
--- a/Engine/Plugins/Experimental/PCG/Source/PCG/Private/Elements/PCGDifferenceElement.cpp
+++ b/Engine/Plugins/Experimental/PCG/Source/PCG/Private/Elements/PCGDifferenceElement.cpp
@@ -7,7 +7,7 @@ FPCGElementPtr UPCGDifferenceSettings::CreateElement() const
 	return MakeShared<FPCGDifferenceElement>();
 }
 
-bool FPCGDifferenceElement::ExecuteInternal(FPCGContextPtr Context) const
+bool FPCGDifferenceElement::ExecuteInternal(FPCGContext* Context) const
 {
 	TRACE_CPUPROFILER_EVENT_SCOPE(FPCGDifferenceElement::Execute);
 
diff --git a/Engine/Plugins/Experimental/PCG/Source/PCG/Private/Elements/PCGExecuteBlueprint.cpp b/Engine/Plugins/Experimental/PCG/Source/PCG/Private/Elements/PCGExecuteBlueprint.cpp
--- a/Engine/Plugins/Experimental/PCG/Source/PCG/Private/Elements/PCGExecuteBlueprint.cpp
+++ b/Engine/Plugins/Experimental/PCG/Source/PCG/Private/Elements/PCGExecuteBlueprint.cpp
@@ -3,6 +3,7 @@
 #include "Elements/PCGExecuteBlueprint.h"
 #include "Data/PCGPointData.h"
 #include "Data/PCGSpatialData.h"
+#include "Helpers/PCGAsync.h"
 
 #if WITH_EDITOR
 #include "Engine/World.h"
@@ -138,6 +139,11 @@ void UPCGBlueprintElement::BeginDestroy()
 	Super::BeginDestroy();
 }
 
+void UPCGBlueprintElement::ExecuteWithContext_Implementation(FPCGContext& InContext, const FPCGDataCollection& Input, FPCGDataCollection& Output) const
+{
+	Execute(Input, Output);
+}
+
 void UPCGBlueprintElement::Initialize()
 {
 #if WITH_EDITOR
@@ -264,6 +270,7 @@ void UPCGBlueprintSettings::PostEditChangeProperty(FPropertyChangedEvent& Proper
 void UPCGBlueprintSettings::OnBlueprintChanged(UBlueprint* InBlueprint)
 {
 	// When the blueprint changes, the element gets recreated, so we must rewire it here.
+	DirtyCache();
 	TeardownBlueprintElementEvent();
 	SetupBlueprintElementEvent();
 
@@ -343,7 +350,7 @@ FPCGElementPtr UPCGBlueprintSettings::CreateElement() const
 	return MakeShared<FPCGExecuteBlueprintElement>();
 }
 
-bool FPCGExecuteBlueprintElement::ExecuteInternal(FPCGContextPtr Context) const
+bool FPCGExecuteBlueprintElement::ExecuteInternal(FPCGContext* Context) const
 {
 	TRACE_CPUPROFILER_EVENT_SCOPE(FPCGExecuteBlueprintElement::Execute);
 	const UPCGBlueprintSettings* Settings = Context->GetInputSettings<UPCGBlueprintSettings>();
@@ -372,7 +379,7 @@ bool FPCGExecuteBlueprintElement::ExecuteInternal(FPCGContextPtr Context) const
 			}
 		}
 
-		Settings->BlueprintElementInstance->Execute(Context->InputData, Context->OutputData);
+		Settings->BlueprintElementInstance->ExecuteWithContext(*Context, Context->InputData, Context->OutputData);
 
 		// Log info on outputs
 		for (int32 OutputIndex = 0; OutputIndex < Context->OutputData.TaggedData.Num(); ++OutputIndex)
@@ -393,6 +400,35 @@ bool FPCGExecuteBlueprintElement::ExecuteInternal(FPCGContextPtr Context) const
 	return true;
 }
 
+void UPCGBlueprintElement::LoopOnPoints(FPCGContext& InContext, const UPCGPointData* InData, UPCGPointData*& OutData) const
+{
+	OutData = NewObject<UPCGPointData>();
+	OutData->TargetActor = InData->TargetActor;
+	
+	const TArray<FPCGPoint>& InPoints = InData->GetPoints();
+	TArray<FPCGPoint>& OutPoints = OutData->GetMutablePoints();
+
+	FPCGAsync::AsyncPointProcessing(&InContext, InPoints.Num(), OutPoints, [this, InData, &InPoints](int32 Index, FPCGPoint& OutPoint)
+	{
+		return PointLoopBody(InData, InPoints[Index], OutPoint);
+	});
+}
+
+void UPCGBlueprintElement::LoopOnPointPairs(FPCGContext& InContext, const UPCGPointData* InA, const UPCGPointData* InB, UPCGPointData*& OutData) const
+{
+	OutData = NewObject<UPCGPointData>();
+	OutData->TargetActor = InA->TargetActor;
+
+	const TArray<FPCGPoint>& InPointsA = InA->GetPoints();
+	const TArray<FPCGPoint>& InPointsB = InB->GetPoints();
+	TArray<FPCGPoint>& OutPoints = OutData->GetMutablePoints();
+
+	FPCGAsync::AsyncPointProcessing(&InContext, InPointsA.Num() * InPointsB.Num(), OutPoints, [this, InA, InB, &InPointsA, &InPointsB](int32 Index, FPCGPoint& OutPoint)
+	{
+		return PointPairLoopBody(InA, InB, InPointsA[Index / InPointsB.Num()], InPointsB[Index % InPointsB.Num()], OutPoint);
+	});
+}
+
 bool FPCGExecuteBlueprintElement::IsCacheable(const UPCGSettings* InSettings) const
 {
 	if (const UPCGBlueprintSettings* BPSettings = Cast<const UPCGBlueprintSettings>(InSettings))
diff --git a/Engine/Plugins/Experimental/PCG/Source/PCG/Private/Elements/PCGIntersectionElement.cpp b/Engine/Plugins/Experimental/PCG/Source/PCG/Private/Elements/PCGIntersectionElement.cpp
--- a/Engine/Plugins/Experimental/PCG/Source/PCG/Private/Elements/PCGIntersectionElement.cpp
+++ b/Engine/Plugins/Experimental/PCG/Source/PCG/Private/Elements/PCGIntersectionElement.cpp
@@ -7,7 +7,7 @@ FPCGElementPtr UPCGIntersectionSettings::CreateElement() const
 	return MakeShared<FPCGIntersectionElement>();
 }
 
-bool FPCGIntersectionElement::ExecuteInternal(FPCGContextPtr Context) const
+bool FPCGIntersectionElement::ExecuteInternal(FPCGContext* Context) const
 {
 	TRACE_CPUPROFILER_EVENT_SCOPE(FPCGIntersectionElement::Execute);
 
diff --git a/Engine/Plugins/Experimental/PCG/Source/PCG/Private/Elements/PCGPointSampler.cpp b/Engine/Plugins/Experimental/PCG/Source/PCG/Private/Elements/PCGPointSampler.cpp
--- a/Engine/Plugins/Experimental/PCG/Source/PCG/Private/Elements/PCGPointSampler.cpp
+++ b/Engine/Plugins/Experimental/PCG/Source/PCG/Private/Elements/PCGPointSampler.cpp
@@ -11,7 +11,7 @@ FPCGElementPtr UPCGPointSamplerSettings::CreateElement() const
 	return MakeShared<FPCGPointSamplerElement>();
 }
 
-bool FPCGPointSamplerElement::ExecuteInternal(FPCGContextPtr Context) const
+bool FPCGPointSamplerElement::ExecuteInternal(FPCGContext* Context) const
 {
 	TRACE_CPUPROFILER_EVENT_SCOPE(FPCGPointSamplerElement::Execute);
 	// TODO: make time-sliced implementation
diff --git a/Engine/Plugins/Experimental/PCG/Source/PCG/Private/Elements/PCGProjectionElement.cpp b/Engine/Plugins/Experimental/PCG/Source/PCG/Private/Elements/PCGProjectionElement.cpp
--- a/Engine/Plugins/Experimental/PCG/Source/PCG/Private/Elements/PCGProjectionElement.cpp
+++ b/Engine/Plugins/Experimental/PCG/Source/PCG/Private/Elements/PCGProjectionElement.cpp
@@ -8,7 +8,7 @@ FPCGElementPtr UPCGProjectionSettings::CreateElement() const
 	return MakeShared<FPCGProjectionElement>();
 }
 
-bool FPCGProjectionElement::ExecuteInternal(FPCGContextPtr Context) const
+bool FPCGProjectionElement::ExecuteInternal(FPCGContext* Context) const
 {
 	TRACE_CPUPROFILER_EVENT_SCOPE(FPCGProjectionElement::Execute);
 
diff --git a/Engine/Plugins/Experimental/PCG/Source/PCG/Private/Elements/PCGSelfPruning.cpp b/Engine/Plugins/Experimental/PCG/Source/PCG/Private/Elements/PCGSelfPruning.cpp
--- a/Engine/Plugins/Experimental/PCG/Source/PCG/Private/Elements/PCGSelfPruning.cpp
+++ b/Engine/Plugins/Experimental/PCG/Source/PCG/Private/Elements/PCGSelfPruning.cpp
@@ -42,7 +42,7 @@ FPCGElementPtr UPCGSelfPruningSettings::CreateElement() const
 	return MakeShared<FPCGSelfPruningElement>();
 }
 
-bool FPCGSelfPruningElement::ExecuteInternal(FPCGContextPtr Context) const
+bool FPCGSelfPruningElement::ExecuteInternal(FPCGContext* Context) const
 {
 	TRACE_CPUPROFILER_EVENT_SCOPE(FPCGSelfPruningElement::Execute);
 	// TODO: time-sliced implementation
diff --git a/Engine/Plugins/Experimental/PCG/Source/PCG/Private/Elements/PCGSpawnActor.cpp b/Engine/Plugins/Experimental/PCG/Source/PCG/Private/Elements/PCGSpawnActor.cpp
--- a/Engine/Plugins/Experimental/PCG/Source/PCG/Private/Elements/PCGSpawnActor.cpp
+++ b/Engine/Plugins/Experimental/PCG/Source/PCG/Private/Elements/PCGSpawnActor.cpp
@@ -64,7 +64,7 @@ TObjectPtr<UPCGGraph> UPCGSpawnActorNode::GetSubgraph() const
 	return (Settings && Settings->Option != EPCGSpawnActorOption::NoMerging) ? Settings->GetSubgraph() : nullptr;
 }
 
-bool FPCGSpawnActorElement::ExecuteInternal(FPCGContextPtr Context) const
+bool FPCGSpawnActorElement::ExecuteInternal(FPCGContext* Context) const
 {
 	TRACE_CPUPROFILER_EVENT_SCOPE(FPCSpawnActorElement::Execute);
 
diff --git a/Engine/Plugins/Experimental/PCG/Source/PCG/Private/Elements/PCGStaticMeshSpawner.cpp b/Engine/Plugins/Experimental/PCG/Source/PCG/Private/Elements/PCGStaticMeshSpawner.cpp
--- a/Engine/Plugins/Experimental/PCG/Source/PCG/Private/Elements/PCGStaticMeshSpawner.cpp
+++ b/Engine/Plugins/Experimental/PCG/Source/PCG/Private/Elements/PCGStaticMeshSpawner.cpp
@@ -22,7 +22,7 @@ struct FWeightedMeshAndInstances
 	TArray<FTransform> Instances;
 };
 
-bool FPCGStaticMeshSpawnerElement::ExecuteInternal(FPCGContextPtr Context) const
+bool FPCGStaticMeshSpawnerElement::ExecuteInternal(FPCGContext* Context) const
 {
 	TRACE_CPUPROFILER_EVENT_SCOPE(FPCGStaticMeshSpawnerElement::Execute);
 	// TODO : time-sliced implementation
diff --git a/Engine/Plugins/Experimental/PCG/Source/PCG/Private/Elements/PCGSurfaceSampler.cpp b/Engine/Plugins/Experimental/PCG/Source/PCG/Private/Elements/PCGSurfaceSampler.cpp
--- a/Engine/Plugins/Experimental/PCG/Source/PCG/Private/Elements/PCGSurfaceSampler.cpp
+++ b/Engine/Plugins/Experimental/PCG/Source/PCG/Private/Elements/PCGSurfaceSampler.cpp
@@ -27,7 +27,7 @@ struct FPCGSurfaceSamplerLoopData
 
 	float Ratio;
 
-	bool Initialize(const UPCGSurfaceSamplerSettings* InSettings, FPCGContextPtr Context, const FBox& InputBounds)
+	bool Initialize(const UPCGSurfaceSamplerSettings* InSettings, FPCGContext* Context, const FBox& InputBounds)
 	{
 		Settings = InSettings;
 
@@ -96,7 +96,7 @@ FPCGElementPtr UPCGSurfaceSamplerSettings::CreateElement() const
 	return MakeShared<FPCGSurfaceSamplerElement>();
 }
 
-bool FPCGSurfaceSamplerElement::ExecuteInternal(FPCGContextPtr Context) const
+bool FPCGSurfaceSamplerElement::ExecuteInternal(FPCGContext* Context) const
 {
 	TRACE_CPUPROFILER_EVENT_SCOPE(FPCGSurfaceSamplerElement::Execute);
 	// TODO: time-sliced implementation
diff --git a/Engine/Plugins/Experimental/PCG/Source/PCG/Private/Elements/PCGUnionElement.cpp b/Engine/Plugins/Experimental/PCG/Source/PCG/Private/Elements/PCGUnionElement.cpp
--- a/Engine/Plugins/Experimental/PCG/Source/PCG/Private/Elements/PCGUnionElement.cpp
+++ b/Engine/Plugins/Experimental/PCG/Source/PCG/Private/Elements/PCGUnionElement.cpp
@@ -7,7 +7,7 @@ FPCGElementPtr UPCGUnionSettings::CreateElement() const
 	return MakeShared<FPCGUnionElement>();
 }
 
-bool FPCGUnionElement::ExecuteInternal(FPCGContextPtr Context) const
+bool FPCGUnionElement::ExecuteInternal(FPCGContext* Context) const
 {
 	TRACE_CPUPROFILER_EVENT_SCOPE(FPCGUnionElement::Execute);
 
diff --git a/Engine/Plugins/Experimental/PCG/Source/PCG/Private/Graph/PCGGraphExecutor.cpp b/Engine/Plugins/Experimental/PCG/Source/PCG/Private/Graph/PCGGraphExecutor.cpp
--- a/Engine/Plugins/Experimental/PCG/Source/PCG/Private/Graph/PCGGraphExecutor.cpp
+++ b/Engine/Plugins/Experimental/PCG/Source/PCG/Private/Graph/PCGGraphExecutor.cpp
@@ -191,14 +191,14 @@ void FPCGGraphExecutor::Execute()
 					continue;
 				}
 				
-				FPCGContextPtr Context = Element->Initialize(TaskInput, Task.SourceComponent);
+				FPCGContext* Context = Element->Initialize(TaskInput, Task.SourceComponent);
 				Context->Node = Task.Node; // still needed?
 				Context->TaskId = Task.NodeId;
 				Context->Cache = &GraphCache;
 
 				FPCGGraphActiveTask& ActiveTask = ActiveTasks.Emplace_GetRef();
 				ActiveTask.Element = Element;
-				ActiveTask.Context = Context;
+				ActiveTask.Context = TUniquePtr<FPCGContext>(Context);
 				ActiveTask.NodeId = Task.NodeId;
 				
 				// TODO update bCanStartNewTasks
@@ -216,7 +216,7 @@ void FPCGGraphExecutor::Execute()
 			// TODO: change this when we support MT in the graph executor
 			ActiveTask.Context->NumAvailableTasks = FMath::Max(1, FMath::Min(FPlatformMisc::NumberOfCoresIncludingHyperthreads() - 2, CVarMaxNumTasks.GetValueOnAnyThread()));
 
-			if (!ActiveTask.Context->bIsPaused && ActiveTask.Element->Execute(ActiveTask.Context))
+			if (!ActiveTask.Context->bIsPaused && ActiveTask.Element->Execute(ActiveTask.Context.Get()))
 			{
 				// Store output data in map
 				StoreResults(ActiveTask.NodeId, ActiveTask.Context->OutputData);
@@ -367,7 +367,7 @@ void FPCGGraphExecutor::NotifyGraphChanged(UPCGGraph* InGraph)
 }
 #endif
 
-bool FPCGFetchInputElement::ExecuteInternal(FPCGContextPtr Context) const
+bool FPCGFetchInputElement::ExecuteInternal(FPCGContext* Context) const
 {
 	// First: any input can be passed through to the output trivially
 	Context->OutputData = Context->InputData;
@@ -389,7 +389,7 @@ FPCGGenericElement::FPCGGenericElement(TFunction<bool()> InOperation)
 {
 }
 
-bool FPCGGenericElement::ExecuteInternal(FPCGContextPtr Context) const
+bool FPCGGenericElement::ExecuteInternal(FPCGContext* Context) const
 {
 	return Operation();
 }
diff --git a/Engine/Plugins/Experimental/PCG/Source/PCG/Private/Graph/PCGGraphExecutor.h b/Engine/Plugins/Experimental/PCG/Source/PCG/Private/Graph/PCGGraphExecutor.h
--- a/Engine/Plugins/Experimental/PCG/Source/PCG/Private/Graph/PCGGraphExecutor.h
+++ b/Engine/Plugins/Experimental/PCG/Source/PCG/Private/Graph/PCGGraphExecutor.h
@@ -36,7 +36,7 @@ struct FPCGGraphScheduleTask
 struct FPCGGraphActiveTask
 {
 	FPCGElementPtr Element;
-	FPCGContextPtr Context;
+	TUniquePtr<FPCGContext> Context;
 	FPCGTaskId NodeId = InvalidTaskId;
 };
 
@@ -118,7 +118,7 @@ private:
 class FPCGFetchInputElement : public FSimplePCGElement
 {
 protected:
-	virtual bool ExecuteInternal(FPCGContextPtr Context) const override;
+	virtual bool ExecuteInternal(FPCGContext* Context) const override;
 	virtual bool IsCacheable(const UPCGSettings* InSettings) const override { return false; }
 };
 
@@ -130,7 +130,7 @@ public:
 protected:
 	// Important note: generic elements must always be run on the main thread
 	// as most of these will impact the editor in some way (loading, unloading, saving)
-	virtual bool ExecuteInternal(FPCGContextPtr Context) const override;
+	virtual bool ExecuteInternal(FPCGContext* Context) const override;
 	virtual bool IsCancellable() const { return false; }
 	virtual bool IsCacheable(const UPCGSettings* InSettings) const override { return false; }
 
diff --git a/Engine/Plugins/Experimental/PCG/Source/PCG/Private/Helpers/PCGAsync.cpp b/Engine/Plugins/Experimental/PCG/Source/PCG/Private/Helpers/PCGAsync.cpp
--- a/Engine/Plugins/Experimental/PCG/Source/PCG/Private/Helpers/PCGAsync.cpp
+++ b/Engine/Plugins/Experimental/PCG/Source/PCG/Private/Helpers/PCGAsync.cpp
@@ -2,14 +2,16 @@
 
 #include "Helpers/PCGAsync.h"
 #include "PCGPoint.h"
+#include "PCGContext.h"
+
 #include "Async/Async.h"
 
 namespace FPCGAsync
 {
-	void AsyncPointProcessing(FPCGContextPtr Context, int32 NumIterations, TArray<FPCGPoint>& OutPoints, const TFunction<bool(int32, FPCGPoint&)>& PointFunc)
+	void AsyncPointProcessing(FPCGContext* Context, int32 NumIterations, TArray<FPCGPoint>& OutPoints, const TFunction<bool(int32, FPCGPoint&)>& PointFunc)
 	{
 		const int32 MinIterationsPerTask = 256;
-		AsyncPointProcessing(Context.IsValid() ? Context->NumAvailableTasks : 1, MinIterationsPerTask, NumIterations, OutPoints, PointFunc);
+		AsyncPointProcessing(Context ? Context->NumAvailableTasks : 1, MinIterationsPerTask, NumIterations, OutPoints, PointFunc);
 	}
 
 	void AsyncPointProcessing(int32 NumAvailableTasks, int32 MinIterationsPerTask, int32 NumIterations, TArray<FPCGPoint>& OutPoints, const TFunction<bool(int32, FPCGPoint&)>& PointFunc)
diff --git a/Engine/Plugins/Experimental/PCG/Source/PCG/Private/Helpers/PCGBlueprintHelpers.cpp b/Engine/Plugins/Experimental/PCG/Source/PCG/Private/Helpers/PCGBlueprintHelpers.cpp
new file mode 100644
--- /dev/null
+++ b/Engine/Plugins/Experimental/PCG/Source/PCG/Private/Helpers/PCGBlueprintHelpers.cpp
@@ -0,0 +1,28 @@
+// Copyright Epic Games, Inc. All Rights Reserved.
+
+#include "Helpers/PCGBlueprintHelpers.h"
+#include "PCGHelpers.h"
+#include "PCGSettings.h"
+
+int UPCGBlueprintHelpers::ComputeSeedFromPosition(const FVector& InPosition)
+{
+	// TODO: should have a config to drive this
+	return PCGHelpers::ComputeSeed((int)InPosition.X, (int)InPosition.Y, (int)InPosition.Z);
+}
+
+void UPCGBlueprintHelpers::SetSeedFromPosition(FPCGPoint& InPoint)
+{
+	InPoint.Seed = ComputeSeedFromPosition(InPoint.Transform.GetLocation());
+}
+
+FRandomStream UPCGBlueprintHelpers::GetRandomStream(const FPCGPoint& InPoint, const UPCGSettings* OptionalSettings)
+{
+	if (OptionalSettings)
+	{
+		return FRandomStream(PCGHelpers::ComputeSeed(InPoint.Seed, OptionalSettings->Seed));
+	}
+	else
+	{
+		return FRandomStream(InPoint.Seed);
+	}
+}
\ No newline at end of file
diff --git a/Engine/Plugins/Experimental/PCG/Source/PCG/Private/PCGElement.cpp b/Engine/Plugins/Experimental/PCG/Source/PCG/Private/PCGElement.cpp
--- a/Engine/Plugins/Experimental/PCG/Source/PCG/Private/PCGElement.cpp
+++ b/Engine/Plugins/Experimental/PCG/Source/PCG/Private/PCGElement.cpp
@@ -7,7 +7,7 @@
 #include "Elements/PCGDebugElement.h"
 #include "Graph/PCGGraphCache.h"
 
-bool IPCGElement::Execute(FPCGContextPtr Context) const
+bool IPCGElement::Execute(FPCGContext* Context) const
 {
 	// Early out to stop execution
 	if (Context->InputData.bCancelExecution)
@@ -131,9 +131,9 @@ bool IPCGElement::Execute(FPCGContextPtr Context) const
 	}
 }
 
-FPCGContextPtr FSimplePCGElement::Initialize(const FPCGDataCollection& InputData, UPCGComponent* SourceComponent)
+FPCGContext* FSimplePCGElement::Initialize(const FPCGDataCollection& InputData, UPCGComponent* SourceComponent)
 {
-	FPCGContextPtr Context = MakeShared<FPCGContext>();
+	FPCGContext* Context = new FPCGContext();
 	Context->InputData = InputData;
 	Context->SourceComponent = SourceComponent;
 
diff --git a/Engine/Plugins/Experimental/PCG/Source/PCG/Private/PCGSettings.cpp b/Engine/Plugins/Experimental/PCG/Source/PCG/Private/PCGSettings.cpp
--- a/Engine/Plugins/Experimental/PCG/Source/PCG/Private/PCGSettings.cpp
+++ b/Engine/Plugins/Experimental/PCG/Source/PCG/Private/PCGSettings.cpp
@@ -82,7 +82,7 @@ FPCGElementPtr UPCGTrivialSettings::CreateElement() const
 	return MakeShared<FPCGTrivialElement>();
 }
 
-bool FPCGTrivialElement::ExecuteInternal(FPCGContextPtr Context) const
+bool FPCGTrivialElement::ExecuteInternal(FPCGContext* Context) const
 {
 	// Pass-through
 	Context->OutputData = Context->InputData;
diff --git a/Engine/Plugins/Experimental/PCG/Source/PCG/Private/PCGSubgraph.cpp b/Engine/Plugins/Experimental/PCG/Source/PCG/Private/PCGSubgraph.cpp
--- a/Engine/Plugins/Experimental/PCG/Source/PCG/Private/PCGSubgraph.cpp
+++ b/Engine/Plugins/Experimental/PCG/Source/PCG/Private/PCGSubgraph.cpp
@@ -173,18 +173,18 @@ void UPCGSubgraphNode::OnStructuralSettingsChanged(UPCGSettings* InSettings)
 }
 #endif // WITH_EDITOR
 
-FPCGContextPtr FPCGSubgraphElement::Initialize(const FPCGDataCollection& InputData, UPCGComponent* SourceComponent)
+FPCGContext* FPCGSubgraphElement::Initialize(const FPCGDataCollection& InputData, UPCGComponent* SourceComponent)
 {
-	TSharedPtr<FPCGSubgraphContext> Context = MakeShared<FPCGSubgraphContext>();
+	FPCGSubgraphContext* Context = new FPCGSubgraphContext();
 	Context->InputData = InputData;
 	Context->SourceComponent = SourceComponent;
 
-	return StaticCastSharedPtr<FPCGContext>(Context);
+	return Context;
 }
 
-bool FPCGSubgraphElement::ExecuteInternal(FPCGContextPtr InContext) const
+bool FPCGSubgraphElement::ExecuteInternal(FPCGContext* InContext) const
 {
-	TSharedPtr<FPCGSubgraphContext> Context = StaticCastSharedPtr<FPCGSubgraphContext>(InContext);
+	FPCGSubgraphContext* Context = static_cast<FPCGSubgraphContext*>(InContext);
 
 	const UPCGSubgraphNode* SubgraphNode = Cast<const UPCGSubgraphNode>(Context->Node);
 	if (SubgraphNode && SubgraphNode->bDynamicGraph)
@@ -261,7 +261,7 @@ FPCGInputForwardingElement::FPCGInputForwardingElement(const FPCGDataCollection&
 
 }
 
-bool FPCGInputForwardingElement::ExecuteInternal(FPCGContextPtr Context) const
+bool FPCGInputForwardingElement::ExecuteInternal(FPCGContext* Context) const
 {
 	Context->OutputData = Input;
 	return true;
diff --git a/Engine/Plugins/Experimental/PCG/Source/PCG/Public/Data/PCGDifferenceData.h b/Engine/Plugins/Experimental/PCG/Source/PCG/Public/Data/PCGDifferenceData.h
--- a/Engine/Plugins/Experimental/PCG/Source/PCG/Public/Data/PCGDifferenceData.h
+++ b/Engine/Plugins/Experimental/PCG/Source/PCG/Public/Data/PCGDifferenceData.h
@@ -45,7 +45,7 @@ public:
 	//~End UPCGSpatialData interface
 
 	//~Begin UPCGSpatialDataWithPointCache interface
-	virtual const UPCGPointData* CreatePointData(FPCGContextPtr Context) const override;
+	virtual const UPCGPointData* CreatePointData(FPCGContext* Context) const override;
 	//~End UPCGSpatialDataWithPointCache interface
 protected:
 	UPROPERTY(BlueprintReadOnly, VisibleAnywhere, Category = SpatialData)
diff --git a/Engine/Plugins/Experimental/PCG/Source/PCG/Public/Data/PCGIntersectionData.h b/Engine/Plugins/Experimental/PCG/Source/PCG/Public/Data/PCGIntersectionData.h
--- a/Engine/Plugins/Experimental/PCG/Source/PCG/Public/Data/PCGIntersectionData.h
+++ b/Engine/Plugins/Experimental/PCG/Source/PCG/Public/Data/PCGIntersectionData.h
@@ -35,7 +35,7 @@ public:
 	//~End UPCGSpatialData interface
 
 	//~Begin UPCGSpatialDataWithPointCache interface
-	virtual const UPCGPointData* CreatePointData(FPCGContextPtr Context) const override;
+	virtual const UPCGPointData* CreatePointData(FPCGContext* Context) const override;
 	//~End UPCGSpatialDataWithPointCache interface
 
 public:
@@ -43,7 +43,7 @@ public:
 	EPCGIntersectionDensityFunction DensityFunction = EPCGIntersectionDensityFunction::Multiply;
 
 protected:
-	UPCGPointData* CreateAndFilterPointData(FPCGContextPtr Context, const UPCGSpatialData* X, const UPCGSpatialData* Y) const;
+	UPCGPointData* CreateAndFilterPointData(FPCGContext* Context, const UPCGSpatialData* X, const UPCGSpatialData* Y) const;
 
 	UPROPERTY(BlueprintReadOnly, VisibleAnywhere, Category = SpatialData)
 	TObjectPtr<const UPCGSpatialData> A = nullptr;
diff --git a/Engine/Plugins/Experimental/PCG/Source/PCG/Public/Data/PCGLandscapeData.h b/Engine/Plugins/Experimental/PCG/Source/PCG/Public/Data/PCGLandscapeData.h
--- a/Engine/Plugins/Experimental/PCG/Source/PCG/Public/Data/PCGLandscapeData.h
+++ b/Engine/Plugins/Experimental/PCG/Source/PCG/Public/Data/PCGLandscapeData.h
@@ -47,7 +47,7 @@ public:
 	// ~End UPGCConcreteData interface
 
 	// ~Begin UPCGSpatialDataWithPointCache interface
-	virtual const UPCGPointData* CreatePointData(FPCGContextPtr Context) const override;
+	virtual const UPCGPointData* CreatePointData(FPCGContext* Context) const override;
 	// ~End UPCGConcreteDataWithPointCache interface
 
 	// TODO: add on property changed to clear cached data
diff --git a/Engine/Plugins/Experimental/PCG/Source/PCG/Public/Data/PCGLandscapeSplineData.h b/Engine/Plugins/Experimental/PCG/Source/PCG/Public/Data/PCGLandscapeSplineData.h
--- a/Engine/Plugins/Experimental/PCG/Source/PCG/Public/Data/PCGLandscapeSplineData.h
+++ b/Engine/Plugins/Experimental/PCG/Source/PCG/Public/Data/PCGLandscapeSplineData.h
@@ -23,7 +23,7 @@ public:
 	//~End UPCGPolyLineData interface
 
 	//~Begin UPCGSpatialDataWithPointCache interface
-	virtual const UPCGPointData* CreatePointData(FPCGContextPtr Context) const override;
+	virtual const UPCGPointData* CreatePointData(FPCGContext* Context) const override;
 	//~End UPCGSpatialDataWithPointCache interface
 
 	//~Begin UPCGSpatialData interface
diff --git a/Engine/Plugins/Experimental/PCG/Source/PCG/Public/Data/PCGPointData.h b/Engine/Plugins/Experimental/PCG/Source/PCG/Public/Data/PCGPointData.h
--- a/Engine/Plugins/Experimental/PCG/Source/PCG/Public/Data/PCGPointData.h
+++ b/Engine/Plugins/Experimental/PCG/Source/PCG/Public/Data/PCGPointData.h
@@ -65,7 +65,7 @@ public:
 	virtual int GetDimension() const override { return 0; }
 	virtual FBox GetBounds() const override;
 	virtual float GetDensityAtPosition(const FVector& InPosition) const override;
-	virtual const UPCGPointData* ToPointData(FPCGContextPtr Context) const { return this; }
+	virtual const UPCGPointData* ToPointData(FPCGContext* Context) const { return this; }
 	virtual FPCGPoint TransformPoint(const FPCGPoint& InPoint) const;
 	// ~End UPCGSpatialData interface
 
diff --git a/Engine/Plugins/Experimental/PCG/Source/PCG/Public/Data/PCGProjectionData.h b/Engine/Plugins/Experimental/PCG/Source/PCG/Public/Data/PCGProjectionData.h
--- a/Engine/Plugins/Experimental/PCG/Source/PCG/Public/Data/PCGProjectionData.h
+++ b/Engine/Plugins/Experimental/PCG/Source/PCG/Public/Data/PCGProjectionData.h
@@ -29,7 +29,7 @@ public:
 	//~End UPCGSpatialData interface
 
 	//~Begin UPCGSpatialDataWithPointCache interface
-	virtual const UPCGPointData* CreatePointData(FPCGContextPtr Context) const override;
+	virtual const UPCGPointData* CreatePointData(FPCGContext* Context) const override;
 	//~End UPCGSpatialDataWithPointCache interface
 
 protected:
diff --git a/Engine/Plugins/Experimental/PCG/Source/PCG/Public/Data/PCGSpatialData.h b/Engine/Plugins/Experimental/PCG/Source/PCG/Public/Data/PCGSpatialData.h
--- a/Engine/Plugins/Experimental/PCG/Source/PCG/Public/Data/PCGSpatialData.h
+++ b/Engine/Plugins/Experimental/PCG/Source/PCG/Public/Data/PCGSpatialData.h
@@ -7,15 +7,13 @@
 
 #include "PCGSpatialData.generated.h"
 
+class AActor;
+struct FPCGContext;
 class UPCGPointData;
 class UPCGIntersectionData;
 class UPCGUnionData;
 class UPCGDifferenceData;
 class UPCGProjectionData;
-class AActor;
-
-struct FPCGContext;
-typedef TSharedPtr<FPCGContext, ESPMode::ThreadSafe> FPCGContextPtr;
 
 /**
 * "Concrete" data base class for PCG generation
@@ -54,9 +52,12 @@ public:
 
 	/** Discretizes the data into points */
 	UFUNCTION(BlueprintCallable, Category = SpatialData)
-	virtual const UPCGPointData* ToPointData() const { return ToPointData(nullptr); }
+	const UPCGPointData* ToPointData() const { return ToPointData(nullptr); }
+
+	UFUNCTION(BlueprintCallable, Category = SpatialData)
+	const UPCGPointData* ToPointDataWithContext(UPARAM(ref) FPCGContext& Context) const { return ToPointData(&Context); }
 
-	virtual const UPCGPointData* ToPointData(FPCGContextPtr Context) const PURE_VIRTUAL(UPCGSpatialData::ToPointData, return nullptr;);
+	virtual const UPCGPointData* ToPointData(FPCGContext* Context) const PURE_VIRTUAL(UPCGSpatialData::ToPointData, return nullptr;);
 
 	/** Transform a world-space position to a world-space position in relation to the current data. (ex: projection on surface) */
 	UFUNCTION(BlueprintCallable, Category = SpatialData)
@@ -101,11 +102,11 @@ class UPCGSpatialDataWithPointCache : public UPCGSpatialData
 
 public:
 	// ~UPCGSpatialData implementation
-	virtual const UPCGPointData* ToPointData(FPCGContextPtr Context) const override;
+	virtual const UPCGPointData* ToPointData(FPCGContext* Context) const override;
 	// ~End UPCGSpatialData implementation
 
 protected:
-	virtual const UPCGPointData* CreatePointData(FPCGContextPtr Context) const PURE_VIRTUAL(UPCGSpatialData::CreatePointData, return nullptr;);
+	virtual const UPCGPointData* CreatePointData(FPCGContext* Context) const PURE_VIRTUAL(UPCGSpatialData::CreatePointData, return nullptr;);
 
 private:
 	UPROPERTY(Transient)
diff --git a/Engine/Plugins/Experimental/PCG/Source/PCG/Public/Data/PCGSplineData.h b/Engine/Plugins/Experimental/PCG/Source/PCG/Public/Data/PCGSplineData.h
--- a/Engine/Plugins/Experimental/PCG/Source/PCG/Public/Data/PCGSplineData.h
+++ b/Engine/Plugins/Experimental/PCG/Source/PCG/Public/Data/PCGSplineData.h
@@ -25,7 +25,7 @@ public:
 	//~End UPCGPolyLineData interface
 
 	//~Begin UPCGSpatialDataWithPointCache interface
-	virtual const UPCGPointData* CreatePointData(FPCGContextPtr Context) const override;
+	virtual const UPCGPointData* CreatePointData(FPCGContext* Context) const override;
 	//~End UPCGSpatialDataWithPointCache interface
 
 	//~Begin UPCGSpatialData interface
diff --git a/Engine/Plugins/Experimental/PCG/Source/PCG/Public/Data/PCGTextureData.h b/Engine/Plugins/Experimental/PCG/Source/PCG/Public/Data/PCGTextureData.h
--- a/Engine/Plugins/Experimental/PCG/Source/PCG/Public/Data/PCGTextureData.h
+++ b/Engine/Plugins/Experimental/PCG/Source/PCG/Public/Data/PCGTextureData.h
@@ -30,7 +30,7 @@ public:
 	//~End UPCGSpatialData interface
 
 	//~Begin UPCGSpatialDataWithPointCache interface
-	virtual const UPCGPointData* CreatePointData(FPCGContextPtr Context) const override;
+	virtual const UPCGPointData* CreatePointData(FPCGContext* Context) const override;
 	//~End UPCGSpatialDataWithPointCache interface
 
 protected:
diff --git a/Engine/Plugins/Experimental/PCG/Source/PCG/Public/Data/PCGUnionData.h b/Engine/Plugins/Experimental/PCG/Source/PCG/Public/Data/PCGUnionData.h
--- a/Engine/Plugins/Experimental/PCG/Source/PCG/Public/Data/PCGUnionData.h
+++ b/Engine/Plugins/Experimental/PCG/Source/PCG/Public/Data/PCGUnionData.h
@@ -47,7 +47,7 @@ public:
 	//~End UPCGSpatialData interface
 
 	//~Begin UPCGSpatialDataWithPointCache interface
-	virtual const UPCGPointData* CreatePointData(FPCGContextPtr Context) const override;
+	virtual const UPCGPointData* CreatePointData(FPCGContext* Context) const override;
 	//~End UPCGSpatialDataWithPointCache interface
 
 protected:
@@ -73,5 +73,5 @@ protected:
 	int CachedDimension = 0;
 
 private:
-	void CreateSequentialPointData(FPCGContextPtr Context, UPCGPointData* PointData, bool bLeftToRight) const;
+	void CreateSequentialPointData(FPCGContext* Context, UPCGPointData* PointData, bool bLeftToRight) const;
 };
\ No newline at end of file
diff --git a/Engine/Plugins/Experimental/PCG/Source/PCG/Public/Data/PCGVolumeData.h b/Engine/Plugins/Experimental/PCG/Source/PCG/Public/Data/PCGVolumeData.h
--- a/Engine/Plugins/Experimental/PCG/Source/PCG/Public/Data/PCGVolumeData.h
+++ b/Engine/Plugins/Experimental/PCG/Source/PCG/Public/Data/PCGVolumeData.h
@@ -26,7 +26,7 @@ public:
 	// ~End UPGCSpatialData interface
 
 	// ~Begin UPCGSpatialDataWithPointCache interface
-	virtual const UPCGPointData* CreatePointData(FPCGContextPtr Context) const override;
+	virtual const UPCGPointData* CreatePointData(FPCGContext* Context) const override;
 	// ~End UPCGSpatialDataWithPointCache interface
 
 protected:
diff --git a/Engine/Plugins/Experimental/PCG/Source/PCG/Public/Elements/PCGDebugElement.h b/Engine/Plugins/Experimental/PCG/Source/PCG/Public/Elements/PCGDebugElement.h
--- a/Engine/Plugins/Experimental/PCG/Source/PCG/Public/Elements/PCGDebugElement.h
+++ b/Engine/Plugins/Experimental/PCG/Source/PCG/Public/Elements/PCGDebugElement.h
@@ -9,7 +9,7 @@
 
 namespace PCGDebugElement
 {
-	void ExecuteDebugDisplay(FPCGContextPtr Context);
+	void ExecuteDebugDisplay(FPCGContext* Context);
 }
 
 UCLASS(BlueprintType, ClassGroup = (Procedural))
@@ -31,6 +31,6 @@ protected:
 class FPCGDebugElement : public FSimpleTypedPCGElement<UPCGDebugSettings>
 {
 protected:
-	virtual bool ExecuteInternal(FPCGContextPtr Contexxt) const override;
+	virtual bool ExecuteInternal(FPCGContext* Contexxt) const override;
 	virtual bool IsCacheable(const UPCGSettings* InSettings) const override { return false; }
 };
\ No newline at end of file
diff --git a/Engine/Plugins/Experimental/PCG/Source/PCG/Public/Elements/PCGDensityFilter.h b/Engine/Plugins/Experimental/PCG/Source/PCG/Public/Elements/PCGDensityFilter.h
--- a/Engine/Plugins/Experimental/PCG/Source/PCG/Public/Elements/PCGDensityFilter.h
+++ b/Engine/Plugins/Experimental/PCG/Source/PCG/Public/Elements/PCGDensityFilter.h
@@ -43,5 +43,5 @@ public:
 class FPCGDensityFilterElement : public FSimplePCGElement
 {
 protected:
-	virtual bool ExecuteInternal(FPCGContextPtr Context) const override;
+	virtual bool ExecuteInternal(FPCGContext* Context) const override;
 };
\ No newline at end of file
diff --git a/Engine/Plugins/Experimental/PCG/Source/PCG/Public/Elements/PCGDensityRemap.h b/Engine/Plugins/Experimental/PCG/Source/PCG/Public/Elements/PCGDensityRemap.h
--- a/Engine/Plugins/Experimental/PCG/Source/PCG/Public/Elements/PCGDensityRemap.h
+++ b/Engine/Plugins/Experimental/PCG/Source/PCG/Public/Elements/PCGDensityRemap.h
@@ -38,5 +38,5 @@ public:
 class FPCGLinearDensityRemapElement : public FSimplePCGElement
 {
 protected:
-	virtual bool ExecuteInternal(FPCGContextPtr Context) const;
+	virtual bool ExecuteInternal(FPCGContext* Context) const;
 };
\ No newline at end of file
diff --git a/Engine/Plugins/Experimental/PCG/Source/PCG/Public/Elements/PCGDifferenceElement.h b/Engine/Plugins/Experimental/PCG/Source/PCG/Public/Elements/PCGDifferenceElement.h
--- a/Engine/Plugins/Experimental/PCG/Source/PCG/Public/Elements/PCGDifferenceElement.h
+++ b/Engine/Plugins/Experimental/PCG/Source/PCG/Public/Elements/PCGDifferenceElement.h
@@ -37,5 +37,5 @@ public:
 class FPCGDifferenceElement : public FSimpleTypedPCGElement<UPCGDifferenceSettings>
 {
 protected:
-	virtual bool ExecuteInternal(FPCGContextPtr Context) const;
+	virtual bool ExecuteInternal(FPCGContext* Context) const;
 };
\ No newline at end of file
diff --git a/Engine/Plugins/Experimental/PCG/Source/PCG/Public/Elements/PCGExecuteBlueprint.h b/Engine/Plugins/Experimental/PCG/Source/PCG/Public/Elements/PCGExecuteBlueprint.h
--- a/Engine/Plugins/Experimental/PCG/Source/PCG/Public/Elements/PCGExecuteBlueprint.h
+++ b/Engine/Plugins/Experimental/PCG/Source/PCG/Public/Elements/PCGExecuteBlueprint.h
@@ -5,6 +5,9 @@
 #include "PCGElement.h"
 #include "PCGSettings.h"
 
+#include "Data/PCGPointData.h"
+#include "PCGPoint.h"
+
 #include "Templates/SubclassOf.h"
 
 #include "PCGExecuteBlueprint.generated.h"
@@ -33,9 +36,25 @@ public:
 	virtual void BeginDestroy() override;
 	// ~End UObject interface
 
+	UFUNCTION(BlueprintNativeEvent, Category = Execution)
+	void ExecuteWithContext(UPARAM(ref)FPCGContext& InContext, const FPCGDataCollection& Input, FPCGDataCollection& Output) const;
+
 	UFUNCTION(BlueprintImplementableEvent, BlueprintCallable, Category = Execution)
 	void Execute(const FPCGDataCollection& Input, FPCGDataCollection& Output) const;
 
+	UFUNCTION(BlueprintImplementableEvent, Category = Execution)
+	bool PointLoopBody(const UPCGPointData* InData, const FPCGPoint& InPoint, FPCGPoint& OutPoint) const;
+
+	UFUNCTION(BlueprintImplementableEvent, Category = Execution)
+	bool PointPairLoopBody(const UPCGPointData* InA, const UPCGPointData* InB, const FPCGPoint& InPointA, const FPCGPoint& InPointB, FPCGPoint& OutPoint) const;
+
+	/** Calls the LoopBody function on all points */
+	UFUNCTION(BlueprintCallable, Category = Execution)
+	void LoopOnPoints(UPARAM(ref) FPCGContext& InContext, const UPCGPointData* InData, UPCGPointData*& OutData) const;
+
+	UFUNCTION(BlueprintCallable, Category = Execution)
+	void LoopOnPointPairs(UPARAM(ref) FPCGContext& InContext, const UPCGPointData* InA, const UPCGPointData* InB, UPCGPointData*& OutData) const;
+
 	/** Called after object creation to setup the object callbacks */
 	void Initialize();
 
@@ -132,6 +151,6 @@ protected:
 class FPCGExecuteBlueprintElement : public FSimpleTypedPCGElement<UPCGBlueprintSettings>
 {
 protected:
-	virtual bool ExecuteInternal(FPCGContextPtr Context) const override;
+	virtual bool ExecuteInternal(FPCGContext* Context) const override;
 	virtual bool IsCacheable(const UPCGSettings* InSettings) const override;
 };
\ No newline at end of file
diff --git a/Engine/Plugins/Experimental/PCG/Source/PCG/Public/Elements/PCGIntersectionElement.h b/Engine/Plugins/Experimental/PCG/Source/PCG/Public/Elements/PCGIntersectionElement.h
--- a/Engine/Plugins/Experimental/PCG/Source/PCG/Public/Elements/PCGIntersectionElement.h
+++ b/Engine/Plugins/Experimental/PCG/Source/PCG/Public/Elements/PCGIntersectionElement.h
@@ -37,5 +37,5 @@ public:
 class FPCGIntersectionElement : public FSimpleTypedPCGElement<UPCGIntersectionSettings>
 {
 protected:
-	virtual bool ExecuteInternal(FPCGContextPtr Context) const;
+	virtual bool ExecuteInternal(FPCGContext* Context) const;
 };
\ No newline at end of file
diff --git a/Engine/Plugins/Experimental/PCG/Source/PCG/Public/Elements/PCGPointSampler.h b/Engine/Plugins/Experimental/PCG/Source/PCG/Public/Elements/PCGPointSampler.h
--- a/Engine/Plugins/Experimental/PCG/Source/PCG/Public/Elements/PCGPointSampler.h
+++ b/Engine/Plugins/Experimental/PCG/Source/PCG/Public/Elements/PCGPointSampler.h
@@ -37,5 +37,5 @@ public:
 class FPCGPointSamplerElement : public FSimpleTypedPCGElement<UPCGPointSamplerSettings>
 {
 protected:
-	virtual bool ExecuteInternal(FPCGContextPtr Context) const override;
+	virtual bool ExecuteInternal(FPCGContext* Context) const override;
 };
\ No newline at end of file
diff --git a/Engine/Plugins/Experimental/PCG/Source/PCG/Public/Elements/PCGProjectionElement.h b/Engine/Plugins/Experimental/PCG/Source/PCG/Public/Elements/PCGProjectionElement.h
--- a/Engine/Plugins/Experimental/PCG/Source/PCG/Public/Elements/PCGProjectionElement.h
+++ b/Engine/Plugins/Experimental/PCG/Source/PCG/Public/Elements/PCGProjectionElement.h
@@ -32,5 +32,5 @@ protected:
 class FPCGProjectionElement : public FSimpleTypedPCGElement<UPCGProjectionSettings>
 {
 protected:
-	virtual bool ExecuteInternal(FPCGContextPtr Context) const;
+	virtual bool ExecuteInternal(FPCGContext* Context) const;
 };
\ No newline at end of file
diff --git a/Engine/Plugins/Experimental/PCG/Source/PCG/Public/Elements/PCGSelfPruning.h b/Engine/Plugins/Experimental/PCG/Source/PCG/Public/Elements/PCGSelfPruning.h
--- a/Engine/Plugins/Experimental/PCG/Source/PCG/Public/Elements/PCGSelfPruning.h
+++ b/Engine/Plugins/Experimental/PCG/Source/PCG/Public/Elements/PCGSelfPruning.h
@@ -46,5 +46,5 @@ protected:
 class FPCGSelfPruningElement : public FSimpleTypedPCGElement<UPCGSelfPruningSettings>
 {
 protected:
-	virtual bool ExecuteInternal(FPCGContextPtr Context) const override;
+	virtual bool ExecuteInternal(FPCGContext* Context) const override;
 };
\ No newline at end of file
diff --git a/Engine/Plugins/Experimental/PCG/Source/PCG/Public/Elements/PCGSpawnActor.h b/Engine/Plugins/Experimental/PCG/Source/PCG/Public/Elements/PCGSpawnActor.h
--- a/Engine/Plugins/Experimental/PCG/Source/PCG/Public/Elements/PCGSpawnActor.h
+++ b/Engine/Plugins/Experimental/PCG/Source/PCG/Public/Elements/PCGSpawnActor.h
@@ -66,6 +66,6 @@ public:
 class FPCGSpawnActorElement : public FSimpleTypedPCGElement<UPCGSpawnActorSettings>
 {
 protected:
-	virtual bool ExecuteInternal(FPCGContextPtr Context) const override;
+	virtual bool ExecuteInternal(FPCGContext* Context) const override;
 	virtual bool IsCacheable(const UPCGSettings* InSettings) const override { return false; }
 };
\ No newline at end of file
diff --git a/Engine/Plugins/Experimental/PCG/Source/PCG/Public/Elements/PCGStaticMeshSpawner.h b/Engine/Plugins/Experimental/PCG/Source/PCG/Public/Elements/PCGStaticMeshSpawner.h
--- a/Engine/Plugins/Experimental/PCG/Source/PCG/Public/Elements/PCGStaticMeshSpawner.h
+++ b/Engine/Plugins/Experimental/PCG/Source/PCG/Public/Elements/PCGStaticMeshSpawner.h
@@ -44,6 +44,6 @@ public:
 class FPCGStaticMeshSpawnerElement : public FSimpleTypedPCGElement<UPCGStaticMeshSpawnerSettings>
 {
 protected:
-	virtual bool ExecuteInternal(FPCGContextPtr Context) const override;
+	virtual bool ExecuteInternal(FPCGContext* Context) const override;
 	virtual bool IsCacheable(const UPCGSettings* InSettings) const override { return false; }
 };
\ No newline at end of file
diff --git a/Engine/Plugins/Experimental/PCG/Source/PCG/Public/Elements/PCGSurfaceSampler.h b/Engine/Plugins/Experimental/PCG/Source/PCG/Public/Elements/PCGSurfaceSampler.h
--- a/Engine/Plugins/Experimental/PCG/Source/PCG/Public/Elements/PCGSurfaceSampler.h
+++ b/Engine/Plugins/Experimental/PCG/Source/PCG/Public/Elements/PCGSurfaceSampler.h
@@ -57,5 +57,5 @@ protected:
 class FPCGSurfaceSamplerElement : public FSimpleTypedPCGElement<UPCGSurfaceSamplerSettings>
 {
 protected:
-	virtual bool ExecuteInternal(FPCGContextPtr Context) const override;
+	virtual bool ExecuteInternal(FPCGContext* Context) const override;
 };
\ No newline at end of file
diff --git a/Engine/Plugins/Experimental/PCG/Source/PCG/Public/Elements/PCGUnionElement.h b/Engine/Plugins/Experimental/PCG/Source/PCG/Public/Elements/PCGUnionElement.h
--- a/Engine/Plugins/Experimental/PCG/Source/PCG/Public/Elements/PCGUnionElement.h
+++ b/Engine/Plugins/Experimental/PCG/Source/PCG/Public/Elements/PCGUnionElement.h
@@ -35,5 +35,5 @@ public:
 class FPCGUnionElement : public FSimpleTypedPCGElement<UPCGUnionSettings>
 {
 protected:
-	virtual bool ExecuteInternal(FPCGContextPtr Context) const;
+	virtual bool ExecuteInternal(FPCGContext* Context) const;
 };
\ No newline at end of file
diff --git a/Engine/Plugins/Experimental/PCG/Source/PCG/Public/Helpers/PCGAsync.h b/Engine/Plugins/Experimental/PCG/Source/PCG/Public/Helpers/PCGAsync.h
--- a/Engine/Plugins/Experimental/PCG/Source/PCG/Public/Helpers/PCGAsync.h
+++ b/Engine/Plugins/Experimental/PCG/Source/PCG/Public/Helpers/PCGAsync.h
@@ -2,10 +2,10 @@
 
 #pragma once
 
-#include "PCGContext.h"
-#include "Async/Async.h"
+#include "CoreMinimal.h"
 
 struct FPCGPoint;
+struct FPCGContext;
 
 namespace FPCGAsync
 {
@@ -16,7 +16,7 @@ namespace FPCGAsync
 	* @param OutPoints - The array in which the results will be written to. Note that the array will be cleared before execution
 	* @param PointFunc - A function that has the index [0; NumIterations] and has to write to the point & return true when the current call generates a point
 	*/
-	void AsyncPointProcessing(FPCGContextPtr Context, int32 NumIterations, TArray<FPCGPoint>& OutPoints, const TFunction<bool(int32, FPCGPoint&)>& PointFunc);
+	void AsyncPointProcessing(FPCGContext* Context, int32 NumIterations, TArray<FPCGPoint>& OutPoints, const TFunction<bool(int32, FPCGPoint&)>& PointFunc);
 
 	/** 
 	* Helper to do more general 1:1 point processing loops
diff --git a/Engine/Plugins/Experimental/PCG/Source/PCG/Public/Helpers/PCGBlueprintHelpers.h b/Engine/Plugins/Experimental/PCG/Source/PCG/Public/Helpers/PCGBlueprintHelpers.h
new file mode 100644
--- /dev/null
+++ b/Engine/Plugins/Experimental/PCG/Source/PCG/Public/Helpers/PCGBlueprintHelpers.h
@@ -0,0 +1,30 @@
+// Copyright Epic Games, Inc. All Rights Reserved.
+
+#pragma once
+
+#include "CoreMinimal.h"
+#include "Math/RandomStream.h"
+#include "Kismet/BlueprintFunctionLibrary.h"
+
+#include "PCGPoint.h"
+
+#include "PCGBlueprintHelpers.generated.h"
+
+class UPCGSettings;
+
+UCLASS()
+class PCG_API UPCGBlueprintHelpers : public UBlueprintFunctionLibrary
+{
+	GENERATED_BODY()
+
+public:
+	UFUNCTION(BlueprintCallable, Category = "PCG|Helpers")
+	static int ComputeSeedFromPosition(const FVector& InPosition);
+
+	UFUNCTION(BlueprintCallable, Category = "PCG|Helpers")
+	static void SetSeedFromPosition(UPARAM(ref) FPCGPoint& InPoint);
+
+	/** Creates a random stream from a point's seed and settings's seed (optional) */
+	UFUNCTION(BlueprintCallable, Category = "PCG|Helpers", meta = (ScriptMethod))
+	static FRandomStream GetRandomStream(const FPCGPoint& InPoint, const UPCGSettings* OptionalSettings = nullptr);
+};
\ No newline at end of file
diff --git a/Engine/Plugins/Experimental/PCG/Source/PCG/Public/PCGContext.h b/Engine/Plugins/Experimental/PCG/Source/PCG/Public/PCGContext.h
--- a/Engine/Plugins/Experimental/PCG/Source/PCG/Public/PCGContext.h
+++ b/Engine/Plugins/Experimental/PCG/Source/PCG/Public/PCGContext.h
@@ -7,14 +7,16 @@
 #include "PCGNode.h"
 #include "PCGSubsystem.h"
 
-struct FPCGContext;
+#include "PCGContext.generated.h"
+
 class UPCGComponent;
 struct FPCGGraphCache;
 
-typedef TSharedPtr<FPCGContext, ESPMode::ThreadSafe> FPCGContextPtr;
-
+USTRUCT(BlueprintType)
 struct FPCGContext
 {
+	GENERATED_BODY()
+
 	virtual ~FPCGContext() = default;
 
 	FPCGDataCollection InputData;
diff --git a/Engine/Plugins/Experimental/PCG/Source/PCG/Public/PCGElement.h b/Engine/Plugins/Experimental/PCG/Source/PCG/Public/PCGElement.h
--- a/Engine/Plugins/Experimental/PCG/Source/PCG/Public/PCGElement.h
+++ b/Engine/Plugins/Experimental/PCG/Source/PCG/Public/PCGElement.h
@@ -29,12 +29,12 @@ class IPCGElement
 {
 public:
 	virtual ~IPCGElement() = default;
-	virtual FPCGContextPtr Initialize(const FPCGDataCollection& InputData, UPCGComponent* SourceComponent) = 0;
+	virtual FPCGContext* Initialize(const FPCGDataCollection& InputData, UPCGComponent* SourceComponent) = 0;
 
-	bool Execute(FPCGContextPtr Context) const;
+	bool Execute(FPCGContext* Context) const;
 
 protected:
-	virtual bool ExecuteInternal(FPCGContextPtr Context) const = 0;
+	virtual bool ExecuteInternal(FPCGContext* Context) const = 0;
 	virtual bool IsCancellable() const { return true; }
 	virtual bool IsCacheable(const UPCGSettings* InSettings) const { return true; }
 };
@@ -45,7 +45,7 @@ protected:
 class FSimplePCGElement : public IPCGElement
 {
 public:
-	virtual FPCGContextPtr Initialize(const FPCGDataCollection& InputData, UPCGComponent* SourceComponent) override;
+	virtual FPCGContext* Initialize(const FPCGDataCollection& InputData, UPCGComponent* SourceComponent) override;
 };
 
 /**
diff --git a/Engine/Plugins/Experimental/PCG/Source/PCG/Public/PCGSettings.h b/Engine/Plugins/Experimental/PCG/Source/PCG/Public/PCGSettings.h
--- a/Engine/Plugins/Experimental/PCG/Source/PCG/Public/PCGSettings.h
+++ b/Engine/Plugins/Experimental/PCG/Source/PCG/Public/PCGSettings.h
@@ -106,6 +106,6 @@ protected:
 class FPCGTrivialElement : public FSimplePCGElement
 {
 protected:
-	virtual bool ExecuteInternal(FPCGContextPtr Context) const override;
+	virtual bool ExecuteInternal(FPCGContext* Context) const override;
 	virtual bool IsCacheable(const UPCGSettings* InSettings) const override { return false; }
 };
\ No newline at end of file
diff --git a/Engine/Plugins/Experimental/PCG/Source/PCG/Public/PCGSubgraph.h b/Engine/Plugins/Experimental/PCG/Source/PCG/Public/PCGSubgraph.h
--- a/Engine/Plugins/Experimental/PCG/Source/PCG/Public/PCGSubgraph.h
+++ b/Engine/Plugins/Experimental/PCG/Source/PCG/Public/PCGSubgraph.h
@@ -126,10 +126,10 @@ struct FPCGSubgraphContext : public FPCGContext
 class FPCGSubgraphElement : public IPCGElement
 {
 public:
-	virtual FPCGContextPtr Initialize(const FPCGDataCollection& InputData, UPCGComponent* SourceComponent) override;
+	virtual FPCGContext* Initialize(const FPCGDataCollection& InputData, UPCGComponent* SourceComponent) override;
 
 protected:
-	virtual bool ExecuteInternal(FPCGContextPtr Context) const override;
+	virtual bool ExecuteInternal(FPCGContext* Context) const override;
 };
 
 class FPCGInputForwardingElement : public FSimplePCGElement
@@ -138,6 +138,6 @@ public:
 	FPCGInputForwardingElement(const FPCGDataCollection& InputToForward);
 
 protected:
-	virtual bool ExecuteInternal(FPCGContextPtr Context) const override;
+	virtual bool ExecuteInternal(FPCGContext* Context) const override;
 	FPCGDataCollection Input;
 };
\ No newline at end of file
