Fixed cell count bug introduced in the surface sampler during the multithreading pass

#rnx

diff --git a/Engine/Plugins/Experimental/PCG/Source/PCG/Private/Elements/PCGSurfaceSampler.cpp b/Engine/Plugins/Experimental/PCG/Source/PCG/Private/Elements/PCGSurfaceSampler.cpp
--- a/Engine/Plugins/Experimental/PCG/Source/PCG/Private/Elements/PCGSurfaceSampler.cpp
+++ b/Engine/Plugins/Experimental/PCG/Source/PCG/Private/Elements/PCGSurfaceSampler.cpp
@@ -73,7 +73,7 @@ struct FPCGSurfaceSamplerLoopData
 	void ComputeCellIndices(int32 Index, int32& CellX, int32& CellY)
 	{
 		check(Index >= 0 && Index < CellCount);
-		const int32 CellCountX = CellMaxX - CellMinX;
+		const int32 CellCountX = 1 + CellMaxX - CellMinX;
 
 		CellX = CellMinX + (Index % CellCountX);
 		CellY = CellMinY + (Index / CellCountX);
