PCG:
* Added public API exposure macros

diff --git a/Engine/Plugins/Experimental/PCG/Source/PCG/Public/Data/PCGDifferenceData.h b/Engine/Plugins/Experimental/PCG/Source/PCG/Public/Data/PCGDifferenceData.h
--- a/Engine/Plugins/Experimental/PCG/Source/PCG/Public/Data/PCGDifferenceData.h
+++ b/Engine/Plugins/Experimental/PCG/Source/PCG/Public/Data/PCGDifferenceData.h
@@ -17,7 +17,7 @@ enum class EPCGDifferenceDensityFunction : uint8
 };
 
 UCLASS(BlueprintType, ClassGroup = (Procedural))
-class UPCGDifferenceData : public UPCGSpatialDataWithPointCache
+class PCG_API UPCGDifferenceData : public UPCGSpatialDataWithPointCache
 {
 	GENERATED_BODY()
 public:
@@ -59,4 +59,4 @@ protected:
 
 	UPROPERTY(BlueprintSetter = SetDensityFunction, EditAnywhere, Category = Settings)
 	EPCGDifferenceDensityFunction DensityFunction = EPCGDifferenceDensityFunction::Minimum;
-};
\ No newline at end of file
+};
diff --git a/Engine/Plugins/Experimental/PCG/Source/PCG/Public/Data/PCGIntersectionData.h b/Engine/Plugins/Experimental/PCG/Source/PCG/Public/Data/PCGIntersectionData.h
--- a/Engine/Plugins/Experimental/PCG/Source/PCG/Public/Data/PCGIntersectionData.h
+++ b/Engine/Plugins/Experimental/PCG/Source/PCG/Public/Data/PCGIntersectionData.h
@@ -17,7 +17,7 @@ enum class EPCGIntersectionDensityFunction : uint8
 * Generic intersection class that delays operations as long as possible.
 */
 UCLASS(BlueprintType, ClassGroup=(Procedural))
-class UPCGIntersectionData : public UPCGSpatialDataWithPointCache
+class PCG_API UPCGIntersectionData : public UPCGSpatialDataWithPointCache
 {
 	GENERATED_BODY()
 public:
@@ -56,4 +56,4 @@ protected:
 
 	UPROPERTY()
 	FBox CachedStrictBounds = FBox(EForceInit::ForceInit);
-};
\ No newline at end of file
+};
diff --git a/Engine/Plugins/Experimental/PCG/Source/PCG/Public/Data/PCGLandscapeData.h b/Engine/Plugins/Experimental/PCG/Source/PCG/Public/Data/PCGLandscapeData.h
--- a/Engine/Plugins/Experimental/PCG/Source/PCG/Public/Data/PCGLandscapeData.h
+++ b/Engine/Plugins/Experimental/PCG/Source/PCG/Public/Data/PCGLandscapeData.h
@@ -12,7 +12,7 @@ class ALandscapeProxy;
 class UPCGPolyLineData;
 
 USTRUCT()
-struct FPCGLandscapeDataPoint
+struct PCG_API FPCGLandscapeDataPoint
 {
 	GENERATED_BODY()
 
@@ -30,7 +30,7 @@ struct FPCGLandscapeDataPoint
 };
 
 UCLASS(BlueprintType, ClassGroup = (Procedural))
-class UPCGLandscapeData : public UPCGSurfaceData
+class PCG_API UPCGLandscapeData : public UPCGSurfaceData
 {
 	GENERATED_BODY()
 
@@ -63,4 +63,4 @@ protected:
 
 	UPROPERTY()
 	FBox Bounds = FBox(EForceInit::ForceInit);
-};
\ No newline at end of file
+};
diff --git a/Engine/Plugins/Experimental/PCG/Source/PCG/Public/Data/PCGLandscapeSplineData.h b/Engine/Plugins/Experimental/PCG/Source/PCG/Public/Data/PCGLandscapeSplineData.h
--- a/Engine/Plugins/Experimental/PCG/Source/PCG/Public/Data/PCGLandscapeSplineData.h
+++ b/Engine/Plugins/Experimental/PCG/Source/PCG/Public/Data/PCGLandscapeSplineData.h
@@ -10,7 +10,7 @@
 class ULandscapeSplinesComponent;
 
 UCLASS(BlueprintType, ClassGroup=(Procedural))
-class UPCGLandscapeSplineData : public UPCGPolyLineData
+class PCG_API UPCGLandscapeSplineData : public UPCGPolyLineData
 {
 	GENERATED_BODY()
 public:
diff --git a/Engine/Plugins/Experimental/PCG/Source/PCG/Public/Data/PCGPointData.h b/Engine/Plugins/Experimental/PCG/Source/PCG/Public/Data/PCGPointData.h
--- a/Engine/Plugins/Experimental/PCG/Source/PCG/Public/Data/PCGPointData.h
+++ b/Engine/Plugins/Experimental/PCG/Source/PCG/Public/Data/PCGPointData.h
@@ -12,7 +12,7 @@
 
 class AActor;
 
-struct FPCGPointRef
+struct PCG_API FPCGPointRef
 {
 	FPCGPointRef(const FPCGPoint& InPoint);
 	FPCGPointRef(const FPCGPointRef& InPointRef);
@@ -21,7 +21,7 @@ struct FPCGPointRef
 	FBoxSphereBounds Bounds;
 };
 
-struct FPCGPointRefSemantics
+struct PCG_API FPCGPointRefSemantics
 {
 	enum { MaxElementsPerLeaf = 16 };
 	enum { MinInclusiveElementsPerNode = 7 };
@@ -52,7 +52,7 @@ struct FPCGPointRefSemantics
 
 // TODO: Split this in "concrete" vs "api" class (needed for views)
 UCLASS(BlueprintType, ClassGroup = (Procedural))
-class UPCGPointData : public UPCGSpatialData
+class PCG_API UPCGPointData : public UPCGSpatialData
 {
 	GENERATED_BODY()
 
@@ -96,4 +96,4 @@ protected:
 	mutable FBox Bounds; // TODO: review if this needs to be threadsafe
 	mutable bool bBoundsAreDirty = true;
 	mutable bool bOctreeIsDirty = true;
-};
\ No newline at end of file
+};
diff --git a/Engine/Plugins/Experimental/PCG/Source/PCG/Public/Data/PCGPolyLineData.h b/Engine/Plugins/Experimental/PCG/Source/PCG/Public/Data/PCGPolyLineData.h
--- a/Engine/Plugins/Experimental/PCG/Source/PCG/Public/Data/PCGPolyLineData.h
+++ b/Engine/Plugins/Experimental/PCG/Source/PCG/Public/Data/PCGPolyLineData.h
@@ -8,7 +8,7 @@
 #include "PCGPolyLineData.generated.h"
 
 UCLASS(Abstract, BlueprintType, ClassGroup = (Procedural))
-class UPCGPolyLineData : public UPCGSpatialDataWithPointCache
+class PCG_API UPCGPolyLineData : public UPCGSpatialDataWithPointCache
 {
 	GENERATED_BODY()
 
@@ -21,4 +21,4 @@ public:
 	virtual int GetNumSegments() const PURE_VIRTUAL(UPCGPolyLineData::GetNumSegments, return 0;);
 	virtual float GetSegmentLength(int SegmentIndex) const PURE_VIRTUAL(UPCGPolyLineData::GetSegmentLength, return 0.0f;);
 	virtual FVector GetLocationAtDistance(int SegmentIndex, float Distance) const PURE_VIRTUAL(UPCGPolyLine::GetLocationAtDistance, return FVector(););
-};
\ No newline at end of file
+};
diff --git a/Engine/Plugins/Experimental/PCG/Source/PCG/Public/Data/PCGProjectionData.h b/Engine/Plugins/Experimental/PCG/Source/PCG/Public/Data/PCGProjectionData.h
--- a/Engine/Plugins/Experimental/PCG/Source/PCG/Public/Data/PCGProjectionData.h
+++ b/Engine/Plugins/Experimental/PCG/Source/PCG/Public/Data/PCGProjectionData.h
@@ -11,7 +11,7 @@
 * Generic projection class (A projected onto B) that intercepts spatial queries
 */
 UCLASS(BlueprintType, ClassGroup=(Procedural))
-class UPCGProjectionData : public UPCGSpatialDataWithPointCache
+class PCG_API UPCGProjectionData : public UPCGSpatialDataWithPointCache
 {
 	GENERATED_BODY()
 public:
@@ -44,4 +44,3 @@ protected:
 	FBox CachedBounds = FBox(EForceInit::ForceInit);
 	FBox CachedStrictBounds = FBox(EForceInit::ForceInit);
 };
-
diff --git a/Engine/Plugins/Experimental/PCG/Source/PCG/Public/Data/PCGRenderTargetData.h b/Engine/Plugins/Experimental/PCG/Source/PCG/Public/Data/PCGRenderTargetData.h
--- a/Engine/Plugins/Experimental/PCG/Source/PCG/Public/Data/PCGRenderTargetData.h
+++ b/Engine/Plugins/Experimental/PCG/Source/PCG/Public/Data/PCGRenderTargetData.h
@@ -10,7 +10,7 @@
 //TODO: It's possible that caching the result in this class is not as efficient as it could be
 // if we expect to sample in different ways (e.g. channel) in the same render target
 UCLASS(BlueprintType, ClassGroup = (Procedural))
-class UPCGRenderTargetData : public UPCGBaseTextureData
+class PCG_API UPCGRenderTargetData : public UPCGBaseTextureData
 {
 	GENERATED_BODY()
 
@@ -20,4 +20,4 @@ public:
 
 	UPROPERTY(BlueprintReadOnly, VisibleAnywhere, Category = Properties)
 	UTextureRenderTarget2D* RenderTarget = nullptr;
-};
\ No newline at end of file
+};
diff --git a/Engine/Plugins/Experimental/PCG/Source/PCG/Public/Data/PCGSpatialData.h b/Engine/Plugins/Experimental/PCG/Source/PCG/Public/Data/PCGSpatialData.h
--- a/Engine/Plugins/Experimental/PCG/Source/PCG/Public/Data/PCGSpatialData.h
+++ b/Engine/Plugins/Experimental/PCG/Source/PCG/Public/Data/PCGSpatialData.h
@@ -25,7 +25,7 @@ class UPCGProjectionData;
 * which hold metadata and a transform, and this is the basic currency of the PCG framework.
 */
 UCLASS(Abstract, BlueprintType, ClassGroup = (Procedural))
-class UPCGSpatialData : public UPCGData
+class PCG_API UPCGSpatialData : public UPCGData
 {
 	GENERATED_BODY()
 
@@ -96,7 +96,7 @@ public:
 };
 
 UCLASS(Abstract, ClassGroup = (Procedural))
-class UPCGSpatialDataWithPointCache : public UPCGSpatialData
+class PCG_API UPCGSpatialDataWithPointCache : public UPCGSpatialData
 {
 	GENERATED_BODY()
 
@@ -113,4 +113,4 @@ private:
 	mutable TObjectPtr<const UPCGPointData> CachedPointData;
 
 	mutable FCriticalSection CacheLock;
-};
\ No newline at end of file
+};
diff --git a/Engine/Plugins/Experimental/PCG/Source/PCG/Public/Data/PCGSplineData.h b/Engine/Plugins/Experimental/PCG/Source/PCG/Public/Data/PCGSplineData.h
--- a/Engine/Plugins/Experimental/PCG/Source/PCG/Public/Data/PCGSplineData.h
+++ b/Engine/Plugins/Experimental/PCG/Source/PCG/Public/Data/PCGSplineData.h
@@ -11,7 +11,7 @@ class USplineComponent;
 class UPCGSurfaceData;
 
 UCLASS(BlueprintType, ClassGroup = (Procedural))
-class UPCGSplineData : public UPCGPolyLineData
+class PCG_API UPCGSplineData : public UPCGPolyLineData
 {
 	GENERATED_BODY()
 
@@ -43,7 +43,7 @@ protected:
 };
 
 UCLASS(BlueprintType, ClassGroup=(Procedural))
-class UPCGSplineProjectionData : public UPCGProjectionData
+class PCG_API UPCGSplineProjectionData : public UPCGProjectionData
 {
 	GENERATED_BODY()
 public:
@@ -61,4 +61,4 @@ public:
 
 protected:
 	FVector2D Project(const FVector& InVector) const;
-};
\ No newline at end of file
+};
diff --git a/Engine/Plugins/Experimental/PCG/Source/PCG/Public/Data/PCGSurfaceData.h b/Engine/Plugins/Experimental/PCG/Source/PCG/Public/Data/PCGSurfaceData.h
--- a/Engine/Plugins/Experimental/PCG/Source/PCG/Public/Data/PCGSurfaceData.h
+++ b/Engine/Plugins/Experimental/PCG/Source/PCG/Public/Data/PCGSurfaceData.h
@@ -8,7 +8,7 @@
 #include "PCGSurfaceData.generated.h"
 
 UCLASS(Abstract, BlueprintType, ClassGroup = (Procedural))
-class UPCGSurfaceData : public UPCGSpatialDataWithPointCache
+class PCG_API UPCGSurfaceData : public UPCGSpatialDataWithPointCache
 {
 	GENERATED_BODY()
 
@@ -23,4 +23,4 @@ public:
 protected:
 	UPROPERTY(BlueprintReadOnly, VisibleAnywhere, Category = SpatialData)
 	FTransform Transform;
-};
\ No newline at end of file
+};
diff --git a/Engine/Plugins/Experimental/PCG/Source/PCG/Public/Data/PCGTextureData.h b/Engine/Plugins/Experimental/PCG/Source/PCG/Public/Data/PCGTextureData.h
--- a/Engine/Plugins/Experimental/PCG/Source/PCG/Public/Data/PCGTextureData.h
+++ b/Engine/Plugins/Experimental/PCG/Source/PCG/Public/Data/PCGTextureData.h
@@ -17,7 +17,7 @@ enum class EPCGTextureColorChannel : uint8
 };
 
 UCLASS(Abstract)
-class UPCGBaseTextureData : public UPCGSurfaceData
+class PCG_API UPCGBaseTextureData : public UPCGSurfaceData
 {
 	GENERATED_BODY()
 
@@ -51,7 +51,7 @@ protected:
 };
 
 UCLASS(BlueprintType, ClassGroup=(Procedural))
-class UPCGTextureData : public UPCGBaseTextureData
+class PCG_API UPCGTextureData : public UPCGBaseTextureData
 {
 	GENERATED_BODY()
 
@@ -61,4 +61,4 @@ public:
 
 	UPROPERTY(BlueprintReadOnly, VisibleAnywhere, Category = Properties)
 	UTexture2D* Texture = nullptr;
-};
\ No newline at end of file
+};
diff --git a/Engine/Plugins/Experimental/PCG/Source/PCG/Public/Data/PCGUnionData.h b/Engine/Plugins/Experimental/PCG/Source/PCG/Public/Data/PCGUnionData.h
--- a/Engine/Plugins/Experimental/PCG/Source/PCG/Public/Data/PCGUnionData.h
+++ b/Engine/Plugins/Experimental/PCG/Source/PCG/Public/Data/PCGUnionData.h
@@ -23,7 +23,7 @@ enum class EPCGUnionDensityFunction : uint8
 };
 
 UCLASS(BlueprintType, ClassGroup=(Procedural))
-class UPCGUnionData : public UPCGSpatialDataWithPointCache
+class PCG_API UPCGUnionData : public UPCGSpatialDataWithPointCache
 {
 	GENERATED_BODY()
 public:
@@ -74,4 +74,4 @@ protected:
 
 private:
 	void CreateSequentialPointData(FPCGContext* Context, UPCGPointData* PointData, bool bLeftToRight) const;
-};
\ No newline at end of file
+};
diff --git a/Engine/Plugins/Experimental/PCG/Source/PCG/Public/Data/PCGVolumeData.h b/Engine/Plugins/Experimental/PCG/Source/PCG/Public/Data/PCGVolumeData.h
--- a/Engine/Plugins/Experimental/PCG/Source/PCG/Public/Data/PCGVolumeData.h
+++ b/Engine/Plugins/Experimental/PCG/Source/PCG/Public/Data/PCGVolumeData.h
@@ -10,7 +10,7 @@
 class AVolume;
 
 UCLASS(BlueprintType, ClassGroup = (Procedural))
-class UPCGVolumeData : public UPCGSpatialDataWithPointCache
+class PCG_API UPCGVolumeData : public UPCGSpatialDataWithPointCache
 {
 	GENERATED_BODY()
 
@@ -38,4 +38,4 @@ protected:
 
 	UPROPERTY()
 	FBox StrictBounds = FBox(EForceInit::ForceInit);
-};
\ No newline at end of file
+};
diff --git a/Engine/Plugins/Experimental/PCG/Source/PCG/Public/Elements/PCGDebugElement.h b/Engine/Plugins/Experimental/PCG/Source/PCG/Public/Elements/PCGDebugElement.h
--- a/Engine/Plugins/Experimental/PCG/Source/PCG/Public/Elements/PCGDebugElement.h
+++ b/Engine/Plugins/Experimental/PCG/Source/PCG/Public/Elements/PCGDebugElement.h
@@ -13,7 +13,7 @@ namespace PCGDebugElement
 }
 
 UCLASS(BlueprintType, ClassGroup = (Procedural))
-class UPCGDebugSettings : public UPCGSettings
+class PCG_API UPCGDebugSettings : public UPCGSettings
 {
 	GENERATED_BODY()
 
@@ -33,4 +33,4 @@ class FPCGDebugElement : public FSimpleTypedPCGElement<UPCGDebugSettings>
 protected:
 	virtual bool ExecuteInternal(FPCGContext* Contexxt) const override;
 	virtual bool IsCacheable(const UPCGSettings* InSettings) const override { return false; }
-};
\ No newline at end of file
+};
diff --git a/Engine/Plugins/Experimental/PCG/Source/PCG/Public/Elements/PCGDensityFilter.h b/Engine/Plugins/Experimental/PCG/Source/PCG/Public/Elements/PCGDensityFilter.h
--- a/Engine/Plugins/Experimental/PCG/Source/PCG/Public/Elements/PCGDensityFilter.h
+++ b/Engine/Plugins/Experimental/PCG/Source/PCG/Public/Elements/PCGDensityFilter.h
@@ -10,7 +10,7 @@
 #include "PCGDensityFilter.generated.h"
 
 UCLASS(BlueprintType, ClassGroup = (Procedural))
-class UPCGDensityFilterSettings : public UPCGSettings
+class PCG_API UPCGDensityFilterSettings : public UPCGSettings
 {
 	GENERATED_BODY()
 
@@ -44,4 +44,4 @@ class FPCGDensityFilterElement : public FSimplePCGElement
 {
 protected:
 	virtual bool ExecuteInternal(FPCGContext* Context) const override;
-};
\ No newline at end of file
+};
diff --git a/Engine/Plugins/Experimental/PCG/Source/PCG/Public/Elements/PCGDensityRemap.h b/Engine/Plugins/Experimental/PCG/Source/PCG/Public/Elements/PCGDensityRemap.h
--- a/Engine/Plugins/Experimental/PCG/Source/PCG/Public/Elements/PCGDensityRemap.h
+++ b/Engine/Plugins/Experimental/PCG/Source/PCG/Public/Elements/PCGDensityRemap.h
@@ -10,7 +10,7 @@
 #include "PCGDensityRemap.generated.h"
 
 UCLASS(BlueprintType, ClassGroup = (Procedural))
-class UPCGLinearDensityRemapSettings : public UPCGSettings
+class PCG_API UPCGLinearDensityRemapSettings : public UPCGSettings
 {
 	GENERATED_BODY()
 
@@ -39,4 +39,4 @@ class FPCGLinearDensityRemapElement : public FSimplePCGElement
 {
 protected:
 	virtual bool ExecuteInternal(FPCGContext* Context) const;
-};
\ No newline at end of file
+};
diff --git a/Engine/Plugins/Experimental/PCG/Source/PCG/Public/Elements/PCGDifferenceElement.h b/Engine/Plugins/Experimental/PCG/Source/PCG/Public/Elements/PCGDifferenceElement.h
--- a/Engine/Plugins/Experimental/PCG/Source/PCG/Public/Elements/PCGDifferenceElement.h
+++ b/Engine/Plugins/Experimental/PCG/Source/PCG/Public/Elements/PCGDifferenceElement.h
@@ -10,7 +10,7 @@
 #include "PCGDifferenceElement.generated.h"
 
 UCLASS(BlueprintType, ClassGroup = (Procedural))
-class UPCGDifferenceSettings : public UPCGSettings
+class PCG_API UPCGDifferenceSettings : public UPCGSettings
 {
 	GENERATED_BODY()
 
@@ -38,4 +38,4 @@ class FPCGDifferenceElement : public FSimpleTypedPCGElement<UPCGDifferenceSettin
 {
 protected:
 	virtual bool ExecuteInternal(FPCGContext* Context) const;
-};
\ No newline at end of file
+};
diff --git a/Engine/Plugins/Experimental/PCG/Source/PCG/Public/Elements/PCGExecuteBlueprint.h b/Engine/Plugins/Experimental/PCG/Source/PCG/Public/Elements/PCGExecuteBlueprint.h
--- a/Engine/Plugins/Experimental/PCG/Source/PCG/Public/Elements/PCGExecuteBlueprint.h
+++ b/Engine/Plugins/Experimental/PCG/Source/PCG/Public/Elements/PCGExecuteBlueprint.h
@@ -26,7 +26,7 @@ namespace PCGBlueprintHelper
 #endif // WITH_EDITOR
 
 UCLASS(Abstract, BlueprintType, Blueprintable, hidecategories = (Object))
-class UPCGBlueprintElement : public UObject
+class PCG_API UPCGBlueprintElement : public UObject
 {
 	GENERATED_BODY()
 
@@ -79,7 +79,7 @@ protected:
 };
 
 UCLASS(BlueprintType, ClassGroup = (Procedural))
-class UPCGBlueprintSettings : public UPCGSettings
+class PCG_API UPCGBlueprintSettings : public UPCGSettings
 {
 	GENERATED_BODY()
 
@@ -153,4 +153,4 @@ class FPCGExecuteBlueprintElement : public FSimpleTypedPCGElement<UPCGBlueprintS
 protected:
 	virtual bool ExecuteInternal(FPCGContext* Context) const override;
 	virtual bool IsCacheable(const UPCGSettings* InSettings) const override;
-};
\ No newline at end of file
+};
diff --git a/Engine/Plugins/Experimental/PCG/Source/PCG/Public/Elements/PCGIntersectionElement.h b/Engine/Plugins/Experimental/PCG/Source/PCG/Public/Elements/PCGIntersectionElement.h
--- a/Engine/Plugins/Experimental/PCG/Source/PCG/Public/Elements/PCGIntersectionElement.h
+++ b/Engine/Plugins/Experimental/PCG/Source/PCG/Public/Elements/PCGIntersectionElement.h
@@ -10,7 +10,7 @@
 #include "PCGIntersectionElement.generated.h"
 
 UCLASS(BlueprintType, ClassGroup = (Procedural))
-class UPCGIntersectionSettings : public UPCGSettings
+class PCG_API UPCGIntersectionSettings : public UPCGSettings
 {
 	GENERATED_BODY()
 
@@ -38,4 +38,4 @@ class FPCGIntersectionElement : public FSimpleTypedPCGElement<UPCGIntersectionSe
 {
 protected:
 	virtual bool ExecuteInternal(FPCGContext* Context) const;
-};
\ No newline at end of file
+};
diff --git a/Engine/Plugins/Experimental/PCG/Source/PCG/Public/Elements/PCGPointSampler.h b/Engine/Plugins/Experimental/PCG/Source/PCG/Public/Elements/PCGPointSampler.h
--- a/Engine/Plugins/Experimental/PCG/Source/PCG/Public/Elements/PCGPointSampler.h
+++ b/Engine/Plugins/Experimental/PCG/Source/PCG/Public/Elements/PCGPointSampler.h
@@ -10,7 +10,7 @@
 #include "PCGPointSampler.generated.h"
 
 UCLASS(BlueprintType, ClassGroup = (Procedural))
-class UPCGPointSamplerSettings : public UPCGSettings
+class PCG_API UPCGPointSamplerSettings : public UPCGSettings
 {
 	GENERATED_BODY()
 
@@ -38,4 +38,4 @@ class FPCGPointSamplerElement : public FSimpleTypedPCGElement<UPCGPointSamplerSe
 {
 protected:
 	virtual bool ExecuteInternal(FPCGContext* Context) const override;
-};
\ No newline at end of file
+};
diff --git a/Engine/Plugins/Experimental/PCG/Source/PCG/Public/Elements/PCGProjectionElement.h b/Engine/Plugins/Experimental/PCG/Source/PCG/Public/Elements/PCGProjectionElement.h
--- a/Engine/Plugins/Experimental/PCG/Source/PCG/Public/Elements/PCGProjectionElement.h
+++ b/Engine/Plugins/Experimental/PCG/Source/PCG/Public/Elements/PCGProjectionElement.h
@@ -9,7 +9,7 @@
 #include "PCGProjectionElement.generated.h"
 
 UCLASS(BlueprintType, ClassGroup = (Procedural))
-class UPCGProjectionSettings : public UPCGSettings
+class PCG_API UPCGProjectionSettings : public UPCGSettings
 {
 	GENERATED_BODY()
 
@@ -33,4 +33,4 @@ class FPCGProjectionElement : public FSimpleTypedPCGElement<UPCGProjectionSettin
 {
 protected:
 	virtual bool ExecuteInternal(FPCGContext* Context) const;
-};
\ No newline at end of file
+};
diff --git a/Engine/Plugins/Experimental/PCG/Source/PCG/Public/Elements/PCGSelfPruning.h b/Engine/Plugins/Experimental/PCG/Source/PCG/Public/Elements/PCGSelfPruning.h
--- a/Engine/Plugins/Experimental/PCG/Source/PCG/Public/Elements/PCGSelfPruning.h
+++ b/Engine/Plugins/Experimental/PCG/Source/PCG/Public/Elements/PCGSelfPruning.h
@@ -19,7 +19,7 @@ enum class EPCGSelfPruningType : uint8
 };
 
 UCLASS(BlueprintType, ClassGroup=(Procedural))
-class UPCGSelfPruningSettings : public UPCGSettings
+class PCG_API UPCGSelfPruningSettings : public UPCGSettings
 {
 	GENERATED_BODY()
 
@@ -47,4 +47,4 @@ class FPCGSelfPruningElement : public FSimpleTypedPCGElement<UPCGSelfPruningSett
 {
 protected:
 	virtual bool ExecuteInternal(FPCGContext* Context) const override;
-};
\ No newline at end of file
+};
diff --git a/Engine/Plugins/Experimental/PCG/Source/PCG/Public/Elements/PCGSpawnActor.h b/Engine/Plugins/Experimental/PCG/Source/PCG/Public/Elements/PCGSpawnActor.h
--- a/Engine/Plugins/Experimental/PCG/Source/PCG/Public/Elements/PCGSpawnActor.h
+++ b/Engine/Plugins/Experimental/PCG/Source/PCG/Public/Elements/PCGSpawnActor.h
@@ -21,7 +21,7 @@ enum class EPCGSpawnActorOption : uint8
 };
 
 UCLASS(BlueprintType, ClassGroup = (Procedural))
-class UPCGSpawnActorSettings : public UPCGBaseSubgraphSettings
+class PCG_API UPCGSpawnActorSettings : public UPCGBaseSubgraphSettings
 {
 	GENERATED_BODY()
 
@@ -54,7 +54,7 @@ protected:
 };
 
 UCLASS(ClassGroup = (Procedural))
-class UPCGSpawnActorNode : public UPCGBaseSubgraphNode
+class PCG_API UPCGSpawnActorNode : public UPCGBaseSubgraphNode
 {
 	GENERATED_BODY()
 public:
@@ -68,4 +68,4 @@ class FPCGSpawnActorElement : public FSimpleTypedPCGElement<UPCGSpawnActorSettin
 protected:
 	virtual bool ExecuteInternal(FPCGContext* Context) const override;
 	virtual bool IsCacheable(const UPCGSettings* InSettings) const override { return false; }
-};
\ No newline at end of file
+};
diff --git a/Engine/Plugins/Experimental/PCG/Source/PCG/Public/Elements/PCGStaticMeshSpawner.h b/Engine/Plugins/Experimental/PCG/Source/PCG/Public/Elements/PCGStaticMeshSpawner.h
--- a/Engine/Plugins/Experimental/PCG/Source/PCG/Public/Elements/PCGStaticMeshSpawner.h
+++ b/Engine/Plugins/Experimental/PCG/Source/PCG/Public/Elements/PCGStaticMeshSpawner.h
@@ -10,7 +10,7 @@
 class UStaticMesh;
 
 USTRUCT(BlueprintType)
-struct FPCGStaticMeshSpawnerEntry
+struct PCG_API FPCGStaticMeshSpawnerEntry
 {
 	GENERATED_BODY()
 
@@ -22,7 +22,7 @@ struct FPCGStaticMeshSpawnerEntry
 };
 
 UCLASS(BlueprintType, ClassGroup = (Procedural))
-class UPCGStaticMeshSpawnerSettings : public UPCGSettings
+class PCG_API UPCGStaticMeshSpawnerSettings : public UPCGSettings
 {
 	GENERATED_BODY()
 
@@ -46,4 +46,4 @@ class FPCGStaticMeshSpawnerElement : public FSimpleTypedPCGElement<UPCGStaticMes
 protected:
 	virtual bool ExecuteInternal(FPCGContext* Context) const override;
 	virtual bool IsCacheable(const UPCGSettings* InSettings) const override { return false; }
-};
\ No newline at end of file
+};
diff --git a/Engine/Plugins/Experimental/PCG/Source/PCG/Public/Elements/PCGSurfaceSampler.h b/Engine/Plugins/Experimental/PCG/Source/PCG/Public/Elements/PCGSurfaceSampler.h
--- a/Engine/Plugins/Experimental/PCG/Source/PCG/Public/Elements/PCGSurfaceSampler.h
+++ b/Engine/Plugins/Experimental/PCG/Source/PCG/Public/Elements/PCGSurfaceSampler.h
@@ -12,7 +12,7 @@
 class FPCGSurfaceSamplerElement;
 
 UCLASS(BlueprintType, ClassGroup = (Procedural))
-class UPCGSurfaceSamplerSettings : public UPCGSettings
+class PCG_API UPCGSurfaceSamplerSettings : public UPCGSettings
 {
 	GENERATED_BODY()
 
@@ -58,4 +58,4 @@ class FPCGSurfaceSamplerElement : public FSimpleTypedPCGElement<UPCGSurfaceSampl
 {
 protected:
 	virtual bool ExecuteInternal(FPCGContext* Context) const override;
-};
\ No newline at end of file
+};
diff --git a/Engine/Plugins/Experimental/PCG/Source/PCG/Public/Elements/PCGUnionElement.h b/Engine/Plugins/Experimental/PCG/Source/PCG/Public/Elements/PCGUnionElement.h
--- a/Engine/Plugins/Experimental/PCG/Source/PCG/Public/Elements/PCGUnionElement.h
+++ b/Engine/Plugins/Experimental/PCG/Source/PCG/Public/Elements/PCGUnionElement.h
@@ -10,7 +10,7 @@
 #include "PCGUnionElement.generated.h"
 
 UCLASS(BlueprintType, ClassGroup = (Procedural))
-class UPCGUnionSettings : public UPCGSettings
+class PCG_API UPCGUnionSettings : public UPCGSettings
 {
 	GENERATED_BODY()
 
@@ -36,4 +36,4 @@ class FPCGUnionElement : public FSimpleTypedPCGElement<UPCGUnionSettings>
 {
 protected:
 	virtual bool ExecuteInternal(FPCGContext* Context) const;
-};
\ No newline at end of file
+};
diff --git a/Engine/Plugins/Experimental/PCG/Source/PCG/Public/PCGContext.h b/Engine/Plugins/Experimental/PCG/Source/PCG/Public/PCGContext.h
--- a/Engine/Plugins/Experimental/PCG/Source/PCG/Public/PCGContext.h
+++ b/Engine/Plugins/Experimental/PCG/Source/PCG/Public/PCGContext.h
@@ -13,7 +13,7 @@ class UPCGComponent;
 struct FPCGGraphCache;
 
 USTRUCT(BlueprintType)
-struct FPCGContext
+struct PCG_API FPCGContext
 {
 	GENERATED_BODY()
 
diff --git a/Engine/Plugins/Experimental/PCG/Source/PCG/Public/PCGData.h b/Engine/Plugins/Experimental/PCG/Source/PCG/Public/PCGData.h
--- a/Engine/Plugins/Experimental/PCG/Source/PCG/Public/PCGData.h
+++ b/Engine/Plugins/Experimental/PCG/Source/PCG/Public/PCGData.h
@@ -17,7 +17,7 @@ class UPCGSettings;
 * flexibility to pass in various concrete data types, settings, and more.
 */
 UCLASS(BlueprintType, ClassGroup = (Procedural))
-class UPCGData : public UObject
+class PCG_API UPCGData : public UObject
 {
 	GENERATED_BODY()
 };
@@ -33,7 +33,7 @@ enum class EPCGDataUsage : uint8
 };
 
 USTRUCT(BlueprintType)
-struct FPCGTaggedData
+struct PCG_API FPCGTaggedData
 {
 	GENERATED_BODY()
 
@@ -51,7 +51,7 @@ struct FPCGTaggedData
 };
 
 USTRUCT(BlueprintType)
-struct FPCGDataCollection
+struct PCG_API FPCGDataCollection
 {
 	GENERATED_BODY()
 
@@ -86,7 +86,7 @@ inline const SettingsType* FPCGDataCollection::GetSettings() const
 }
 
 UCLASS()
-class UPCGDataFunctionLibrary : public UBlueprintFunctionLibrary
+class PCG_API UPCGDataFunctionLibrary : public UBlueprintFunctionLibrary
 {
 	GENERATED_BODY()
 
diff --git a/Engine/Plugins/Experimental/PCG/Source/PCG/Public/PCGDebug.h b/Engine/Plugins/Experimental/PCG/Source/PCG/Public/PCGDebug.h
--- a/Engine/Plugins/Experimental/PCG/Source/PCG/Public/PCGDebug.h
+++ b/Engine/Plugins/Experimental/PCG/Source/PCG/Public/PCGDebug.h
@@ -17,7 +17,7 @@ enum class EPCGDebugVisScaleMethod : uint8
 };
 
 USTRUCT(BlueprintType)
-struct FPCGDebugVisualizationSettings
+struct PCG_API FPCGDebugVisualizationSettings
 {
 	GENERATED_BODY()
 
@@ -37,4 +37,4 @@ public:
 	TSoftObjectPtr<UMaterialInterface> MaterialOverride;
 
 	TSoftObjectPtr<UMaterialInterface> GetMaterial() const;
-};
\ No newline at end of file
+};
diff --git a/Engine/Plugins/Experimental/PCG/Source/PCG/Public/PCGElement.h b/Engine/Plugins/Experimental/PCG/Source/PCG/Public/PCGElement.h
--- a/Engine/Plugins/Experimental/PCG/Source/PCG/Public/PCGElement.h
+++ b/Engine/Plugins/Experimental/PCG/Source/PCG/Public/PCGElement.h
@@ -25,7 +25,7 @@ typedef TSharedPtr<IPCGElement, ESPMode::ThreadSafe> FPCGElementPtr;
 /**
 * Base class for the processing bit of a PCG node/settings
 */
-class IPCGElement
+class PCG_API IPCGElement
 {
 public:
 	virtual ~IPCGElement() = default;
@@ -42,7 +42,7 @@ protected:
 /**
 * Basic PCG element class for elements that do not store any intermediate data in the context
 */
-class FSimplePCGElement : public IPCGElement
+class PCG_API FSimplePCGElement : public IPCGElement
 {
 public:
 	virtual FPCGContext* Initialize(const FPCGDataCollection& InputData, UPCGComponent* SourceComponent) override;
@@ -52,8 +52,8 @@ public:
 * CRTP PCG element class to facilitate settings retrieval
 */
 template<typename SettingsClass>
-class FSimpleTypedPCGElement : public FSimplePCGElement
+class PCG_API FSimpleTypedPCGElement : public FSimplePCGElement
 {
 public:
 	const SettingsClass* GetSettings();
-};
\ No newline at end of file
+};
diff --git a/Engine/Plugins/Experimental/PCG/Source/PCG/Public/PCGGraphSetupBP.h b/Engine/Plugins/Experimental/PCG/Source/PCG/Public/PCGGraphSetupBP.h
--- a/Engine/Plugins/Experimental/PCG/Source/PCG/Public/PCGGraphSetupBP.h
+++ b/Engine/Plugins/Experimental/PCG/Source/PCG/Public/PCGGraphSetupBP.h
@@ -9,11 +9,11 @@
 class UPCGGraph;
 
 UCLASS(Abstract, BlueprintType, Blueprintable, hidecategories=(Object))
-class UPCGGraphSetupBP : public UObject
+class PCG_API UPCGGraphSetupBP : public UObject
 {
 	GENERATED_BODY()
 
 public:
 	UFUNCTION(BlueprintNativeEvent, BlueprintCallable, Category = Graph)
 	void Setup(UPCGGraph* Graph);
-};
\ No newline at end of file
+};
diff --git a/Engine/Plugins/Experimental/PCG/Source/PCG/Public/PCGNode.h b/Engine/Plugins/Experimental/PCG/Source/PCG/Public/PCGNode.h
--- a/Engine/Plugins/Experimental/PCG/Source/PCG/Public/PCGNode.h
+++ b/Engine/Plugins/Experimental/PCG/Source/PCG/Public/PCGNode.h
@@ -15,7 +15,7 @@ class IPCGElement;
 #endif
 
 UCLASS(ClassGroup = (Procedural))
-class UPCGNode : public UObject
+class PCG_API UPCGNode : public UObject
 {
 	GENERATED_BODY()
 
@@ -80,4 +80,4 @@ protected:
 	// - Ability to be multithreaded (here or element)
 	// - Generates artifacts (here or element)
 	// - Priority
-};
\ No newline at end of file
+};
diff --git a/Engine/Plugins/Experimental/PCG/Source/PCG/Public/PCGPoint.h b/Engine/Plugins/Experimental/PCG/Source/PCG/Public/PCGPoint.h
--- a/Engine/Plugins/Experimental/PCG/Source/PCG/Public/PCGPoint.h
+++ b/Engine/Plugins/Experimental/PCG/Source/PCG/Public/PCGPoint.h
@@ -7,7 +7,7 @@
 #include "PCGPoint.generated.h"
 
 USTRUCT(BlueprintType)
-struct FPCGPoint
+struct PCG_API FPCGPoint
 {
 	GENERATED_BODY()
 public:
@@ -34,4 +34,4 @@ public:
 
 	UPROPERTY(BlueprintReadWrite, EditAnywhere, Category = Properties)
 	int32 Seed = 0;
-};
\ No newline at end of file
+};
diff --git a/Engine/Plugins/Experimental/PCG/Source/PCG/Public/PCGSettings.h b/Engine/Plugins/Experimental/PCG/Source/PCG/Public/PCGSettings.h
--- a/Engine/Plugins/Experimental/PCG/Source/PCG/Public/PCGSettings.h
+++ b/Engine/Plugins/Experimental/PCG/Source/PCG/Public/PCGSettings.h
@@ -94,7 +94,7 @@ private:
 
 /** Trivial / Pass-through settings used for input/output nodes */
 UCLASS(BlueprintType, ClassGroup = (Procedural))
-class UPCGTrivialSettings : public UPCGSettings
+class PCG_API UPCGTrivialSettings : public UPCGSettings
 {
 	GENERATED_BODY()
 
@@ -103,9 +103,9 @@ protected:
 	virtual FPCGElementPtr CreateElement() const override;
 };
 
-class FPCGTrivialElement : public FSimplePCGElement
+class PCG_API FPCGTrivialElement : public FSimplePCGElement
 {
 protected:
 	virtual bool ExecuteInternal(FPCGContext* Context) const override;
 	virtual bool IsCacheable(const UPCGSettings* InSettings) const override { return false; }
-};
\ No newline at end of file
+};
diff --git a/Engine/Plugins/Experimental/PCG/Source/PCG/Public/PCGSubgraph.h b/Engine/Plugins/Experimental/PCG/Source/PCG/Public/PCGSubgraph.h
--- a/Engine/Plugins/Experimental/PCG/Source/PCG/Public/PCGSubgraph.h
+++ b/Engine/Plugins/Experimental/PCG/Source/PCG/Public/PCGSubgraph.h
@@ -15,7 +15,7 @@ class UPCGGraph;
 #endif
 
 UCLASS(Abstract)
-class UPCGBaseSubgraphSettings : public UPCGSettings
+class PCG_API UPCGBaseSubgraphSettings : public UPCGSettings
 {
 	GENERATED_BODY()
 public:
@@ -48,7 +48,7 @@ public:
 };
 
 UCLASS(BlueprintType, ClassGroup=(Procedural))
-class UPCGSubgraphSettings : public UPCGBaseSubgraphSettings
+class PCG_API UPCGSubgraphSettings : public UPCGBaseSubgraphSettings
 {
 	GENERATED_BODY()
 
@@ -79,7 +79,7 @@ public:
 };
 
 UCLASS(Abstract)
-class UPCGBaseSubgraphNode : public UPCGNode
+class PCG_API UPCGBaseSubgraphNode : public UPCGNode
 {
 	GENERATED_BODY()
 
@@ -91,7 +91,7 @@ public:
 };
 
 UCLASS(ClassGroup = (Procedural))
-class UPCGSubgraphNode : public UPCGBaseSubgraphNode
+class PCG_API UPCGSubgraphNode : public UPCGBaseSubgraphNode
 {
 	GENERATED_BODY()
 
@@ -117,13 +117,13 @@ protected:
 #endif
 };
 
-struct FPCGSubgraphContext : public FPCGContext
+struct PCG_API FPCGSubgraphContext : public FPCGContext
 {
 	FPCGTaskId SubgraphTaskId = InvalidTaskId;
 	bool bScheduledSubgraph = false;
 };
 
-class FPCGSubgraphElement : public IPCGElement
+class PCG_API FPCGSubgraphElement : public IPCGElement
 {
 public:
 	virtual FPCGContext* Initialize(const FPCGDataCollection& InputData, UPCGComponent* SourceComponent) override;
@@ -132,7 +132,7 @@ protected:
 	virtual bool ExecuteInternal(FPCGContext* Context) const override;
 };
 
-class FPCGInputForwardingElement : public FSimplePCGElement
+class PCG_API FPCGInputForwardingElement : public FSimplePCGElement
 {
 public:
 	FPCGInputForwardingElement(const FPCGDataCollection& InputToForward);
@@ -140,4 +140,4 @@ public:
 protected:
 	virtual bool ExecuteInternal(FPCGContext* Context) const override;
 	FPCGDataCollection Input;
-};
\ No newline at end of file
+};
