Point radius deprecated default value was causing a reset of the point extents parameter

#rnx

diff --git a/Engine/Plugins/Experimental/PCG/Source/PCG/Public/Elements/PCGSurfaceSampler.h b/Engine/Plugins/Experimental/PCG/Source/PCG/Public/Elements/PCGSurfaceSampler.h
--- a/Engine/Plugins/Experimental/PCG/Source/PCG/Public/Elements/PCGSurfaceSampler.h
+++ b/Engine/Plugins/Experimental/PCG/Source/PCG/Public/Elements/PCGSurfaceSampler.h
@@ -21,7 +21,7 @@ public:
 	float PointsPerSquaredMeter = 0.1f;
 
 	UPROPERTY()
-	float PointRadius_DEPRECATED = 100.0f;
+	float PointRadius_DEPRECATED = 0.0f;
 
 	UPROPERTY(BlueprintReadWrite, EditAnywhere, Category = Settings)
 	FVector PointExtents = FVector(100.0f);
