Disabled save-on-generate and loading in PCG

diff --git a/Engine/Plugins/Experimental/PCG/Source/PCG/Private/PCGComponent.cpp b/Engine/Plugins/Experimental/PCG/Source/PCG/Private/PCGComponent.cpp
--- a/Engine/Plugins/Experimental/PCG/Source/PCG/Private/PCGComponent.cpp
+++ b/Engine/Plugins/Experimental/PCG/Source/PCG/Private/PCGComponent.cpp
@@ -33,6 +33,11 @@
 
 #define LOCTEXT_NAMESPACE "UPCGComponent"
 
+namespace PCGComponent
+{
+	const bool bSaveOnCleanupAndGenerate = false;
+}
+
 bool UPCGComponent::CanPartition() const
 {
 	// Support/Force partitioning on non-PCG partition actors in WP worlds.
@@ -173,7 +178,7 @@ void UPCGComponent::Generate()
 	FScopedTransaction Transaction(LOCTEXT("PCGGenerate", "Execute generation on PCG component"));
 #endif
 
-	Generate(/*bForce=*/true);
+	Generate(/*bForce=*/PCGComponent::bSaveOnCleanupAndGenerate);
 }
 
 void UPCGComponent::Generate(bool bForce)
@@ -298,7 +303,7 @@ void UPCGComponent::Cleanup()
 	FScopedTransaction Transaction(LOCTEXT("PCGCleanup", "Clean up PCG component"));
 #endif
 
-	Cleanup(/*bRemoveComponents=*/true, /*bSave=*/true);
+	Cleanup(/*bRemoveComponents=*/true, /*bSave=*/PCGComponent::bSaveOnCleanupAndGenerate);
 }
 
 void UPCGComponent::Cleanup(bool bRemoveComponents, bool bSave)
@@ -592,7 +597,7 @@ void UPCGComponent::PostEditUndo()
 	{
 		// Need to reset the generated flag to go through the Cleanup
 		bGenerated = true;
-		Cleanup(/*bRemoveComponents=*/true, /*bSave=*/true);
+		Cleanup(/*bRemoveComponents=*/true, /*bSave=*/PCGComponent::bSaveOnCleanupAndGenerate);
 	}
 	else
 	{
diff --git a/Engine/Plugins/Experimental/PCG/Source/PCG/Private/PCGSubsystem.cpp b/Engine/Plugins/Experimental/PCG/Source/PCG/Private/PCGSubsystem.cpp
--- a/Engine/Plugins/Experimental/PCG/Source/PCG/Private/PCGSubsystem.cpp
+++ b/Engine/Plugins/Experimental/PCG/Source/PCG/Private/PCGSubsystem.cpp
@@ -167,7 +167,7 @@ namespace PCGSubsystem
 				// At this point, if bCreateActor was true, then it exists, but it is not currently loaded; make sure it is loaded
 				// Otherwise, we still need to load it if it exists
 				// TODO: Revisit after API review on the WP side, we shouldn't have to load here or get the actor desc directly
-				if (!PCGActor)
+				if (!PCGActor && bSaveActors)
 				{
 					const FWorldPartitionActorDesc* PCGActorDesc = nullptr;
 					auto FindFirst = [&CellCoord, &PCGActorDesc](const FWorldPartitionActorDesc* ActorDesc) {
@@ -196,7 +196,7 @@ namespace PCGSubsystem
 						PCGActor = Cast<APCGPartitionActor>(PCGActorDesc->GetActor());
 					}
 				}
-				else
+				else if(PCGActor)
 				{
 					// We still need to keep a reference on the PCG actor - note that newly created PCG actors will not have a reference here, but won't be unloaded
 					ActorReferences->Add(FWorldPartitionReference(World->GetWorldPartition(), PCGActor->GetActorGuid()));
@@ -360,7 +360,7 @@ FPCGTaskId UPCGSubsystem::ProcessGraph(UPCGComponent* Component, const FBox& InP
 
 	FBox UnionBounds = InPreviousBounds + InNewBounds;
 	const bool bCreateActors = (!UnionBounds.Equals(InPreviousBounds));
-	const bool bLoadCell = (bGenerate); // TODO: review this
+	const bool bLoadCell = (bGenerate && bSave); // TODO: review this
 	const bool bSaveActors = (bSave);
 
 	FPCGTaskId ProcessAllCellsTaskId = InvalidTaskId;
