Code cleanup in PCG blueprints:
Removed point caching (superseded with ExecuteWithContext)
Removed element type BP setter

#rnx

diff --git a/Engine/Plugins/Experimental/PCG/Source/PCG/Private/Elements/PCGExecuteBlueprint.cpp b/Engine/Plugins/Experimental/PCG/Source/PCG/Private/Elements/PCGExecuteBlueprint.cpp
--- a/Engine/Plugins/Experimental/PCG/Source/PCG/Private/Elements/PCGExecuteBlueprint.cpp
+++ b/Engine/Plugins/Experimental/PCG/Source/PCG/Private/Elements/PCGExecuteBlueprint.cpp
@@ -289,12 +289,6 @@ void UPCGBlueprintSettings::OnBlueprintElementChanged(UPCGBlueprintElement* InEl
 }
 #endif
 
-void UPCGBlueprintSettings::SetElementTypeInternal(TSubclassOf<UPCGBlueprintElement> InElementType)
-{
-	UPCGBlueprintElement* ElementInstance = nullptr;
-	SetElementType(InElementType, ElementInstance);
-}
-
 void UPCGBlueprintSettings::SetElementType(TSubclassOf<UPCGBlueprintElement> InElementType, UPCGBlueprintElement*& ElementInstance)
 {
 	if (!BlueprintElementInstance || InElementType != BlueprintElementType)
@@ -357,18 +351,6 @@ bool FPCGExecuteBlueprintElement::ExecuteInternal(FPCGContext* Context) const
 
 	if (Settings && Settings->BlueprintElementInstance)
 	{
-		// Precache points if needed
-		if (Settings->bPrecachePoints)
-		{
-			for (const FPCGTaggedData& Input : Context->InputData.TaggedData)
-			{
-				if (const UPCGSpatialDataWithPointCache* SpatialInput = Cast<UPCGSpatialDataWithPointCache>(Input.Data))
-				{
-					SpatialInput->ToPointData(Context);
-				}
-			}
-		}
-
 		// Log info on inputs
 		for(int32 InputIndex = 0; InputIndex < Context->InputData.TaggedData.Num(); ++InputIndex)
 		{
@@ -402,6 +384,12 @@ bool FPCGExecuteBlueprintElement::ExecuteInternal(FPCGContext* Context) const
 
 void UPCGBlueprintElement::LoopOnPoints(FPCGContext& InContext, const UPCGPointData* InData, UPCGPointData*& OutData) const
 {
+	if (!InData)
+	{
+		PCGE_LOG_C(Error, &InContext, "Invalid input data in LoopOnPoints");
+		return;
+	}
+
 	OutData = NewObject<UPCGPointData>();
 	OutData->TargetActor = InData->TargetActor;
 	
@@ -416,6 +404,12 @@ void UPCGBlueprintElement::LoopOnPoints(FPCGContext& InContext, const UPCGPointD
 
 void UPCGBlueprintElement::LoopOnPointPairs(FPCGContext& InContext, const UPCGPointData* InA, const UPCGPointData* InB, UPCGPointData*& OutData) const
 {
+	if (!InA || !InB)
+	{
+		PCGE_LOG_C(Error, &InContext, "Invalid input data in LoopOnPointPairs");
+		return;
+	}
+
 	OutData = NewObject<UPCGPointData>();
 	OutData->TargetActor = InA->TargetActor;
 
diff --git a/Engine/Plugins/Experimental/PCG/Source/PCG/Public/Elements/PCGExecuteBlueprint.h b/Engine/Plugins/Experimental/PCG/Source/PCG/Public/Elements/PCGExecuteBlueprint.h
--- a/Engine/Plugins/Experimental/PCG/Source/PCG/Public/Elements/PCGExecuteBlueprint.h
+++ b/Engine/Plugins/Experimental/PCG/Source/PCG/Public/Elements/PCGExecuteBlueprint.h
@@ -107,10 +107,6 @@ public:
 	// ~End UObject interface
 #endif
 
-	/** To facilitate transition, we will keep this version around even though it's a lot less useful */
-	UFUNCTION(BlueprintCallable, BlueprintInternalUseOnly)
-	void SetElementTypeInternal(TSubclassOf<UPCGBlueprintElement> InElementType);
-
 	UFUNCTION(BlueprintCallable, Category = Settings, meta=(DeterminesOutputType="InElementType", DynamicOutputParam = "ElementInstance"))
 	void SetElementType(TSubclassOf<UPCGBlueprintElement> InElementType, UPCGBlueprintElement*& ElementInstance);
 
@@ -118,7 +114,7 @@ protected:
 	UPROPERTY()
 	TSubclassOf<UPCGBlueprintElement> BlueprintElement_DEPRECATED;
 
-	UPROPERTY(BlueprintSetter = SetElementTypeInternal, EditAnywhere, Category = Template)
+	UPROPERTY(BlueprintReadOnly, EditAnywhere, Category = Template)
 	TSubclassOf<UPCGBlueprintElement> BlueprintElementType;
 
 	UPROPERTY(BlueprintReadOnly, VisibleAnywhere, Instanced, Category = Settings, meta = (ShowOnlyInnerProperties))
@@ -132,9 +128,6 @@ protected:
 	UPROPERTY(BlueprintReadWrite, EditAnywhere, Category = Settings)
 	bool bCreatesArtifacts = false;
 
-	UPROPERTY(BlueprintReadWrite, EditAnywhere, Category = Settings)
-	bool bPrecachePoints = false;
-
 protected:
 #if WITH_EDITOR
 	void OnBlueprintChanged(UBlueprint* InBlueprint);
diff --git a/Engine/Plugins/Experimental/PCG/Source/PCG/Public/PCGElement.h b/Engine/Plugins/Experimental/PCG/Source/PCG/Public/PCGElement.h
--- a/Engine/Plugins/Experimental/PCG/Source/PCG/Public/PCGElement.h
+++ b/Engine/Plugins/Experimental/PCG/Source/PCG/Public/PCGElement.h
@@ -16,8 +16,8 @@ typedef TSharedPtr<IPCGElement, ESPMode::ThreadSafe> FPCGElementPtr;
 	UE_LOG(LogPCG, \
 		Verbosity, \
 		TEXT("[%s - %s]: " Format), \
-		*CustomContext->GetComponentName(), \
-		*CustomContext->GetTaskName(), \
+		*((CustomContext)->GetComponentName()), \
+		*((CustomContext)->GetTaskName()), \
 		##__VA_ARGS__)
 
 #define PCGE_LOG(Verbosity, Format, ...) PCGE_LOG_C(Verbosity, Context, Format, ##__VA_ARGS__)
