Improved bounds issue that would in the surface sampler due to strict bounds testing; in the case of PCG we need a custom method to do it properly

diff --git a/Engine/Plugins/Experimental/PCG/Source/PCG/Private/Data/PCGDifferenceData.cpp b/Engine/Plugins/Experimental/PCG/Source/PCG/Private/Data/PCGDifferenceData.cpp
--- a/Engine/Plugins/Experimental/PCG/Source/PCG/Private/Data/PCGDifferenceData.cpp
+++ b/Engine/Plugins/Experimental/PCG/Source/PCG/Private/Data/PCGDifferenceData.cpp
@@ -4,6 +4,7 @@
 #include "Data/PCGPointData.h"
 #include "Data/PCGUnionData.h"
 #include "Helpers/PCGAsync.h"
+#include "PCGHelpers.h"
 
 namespace PCGDifferenceDataUtils
 {
@@ -82,8 +83,8 @@ FBox UPCGDifferenceData::GetStrictBounds() const
 
 float UPCGDifferenceData::GetDensityAtPosition(const FVector& InPosition) const
 {
-	if (!Source->GetBounds().IsInside(InPosition) ||
-		(Difference && Difference->GetStrictBounds().IsInside(InPosition)))
+	if(!PCGHelpers::IsInsideBounds(Source->GetBounds(), InPosition) ||
+		(Difference && PCGHelpers::IsInsideBounds(Difference->GetStrictBounds(), InPosition)))
 	{
 		return 0;
 	}
diff --git a/Engine/Plugins/Experimental/PCG/Source/PCG/Private/Data/PCGIntersectionData.cpp b/Engine/Plugins/Experimental/PCG/Source/PCG/Private/Data/PCGIntersectionData.cpp
--- a/Engine/Plugins/Experimental/PCG/Source/PCG/Private/Data/PCGIntersectionData.cpp
+++ b/Engine/Plugins/Experimental/PCG/Source/PCG/Private/Data/PCGIntersectionData.cpp
@@ -3,6 +3,7 @@
 #include "Data/PCGIntersectionData.h"
 #include "Data/PCGPointData.h"
 #include "Helpers/PCGAsync.h"
+#include "PCGHelpers.h"
 
 namespace PCGIntersectionDataMaths
 {
@@ -51,11 +52,11 @@ FBox UPCGIntersectionData::GetStrictBounds() const
 float UPCGIntersectionData::GetDensityAtPosition(const FVector& InPosition) const
 {
 	check(A && B);
-	if (!CachedBounds.IsInside(InPosition))
+	if(!PCGHelpers::IsInsideBounds(CachedBounds, InPosition))
 	{
 		return 0;
 	}
-	else if (CachedStrictBounds.IsInside(InPosition))
+	else if(PCGHelpers::IsInsideBounds(CachedStrictBounds, InPosition))
 	{
 		return 1.0f;
 	}
diff --git a/Engine/Plugins/Experimental/PCG/Source/PCG/Private/Data/PCGLandscapeSplineData.cpp b/Engine/Plugins/Experimental/PCG/Source/PCG/Private/Data/PCGLandscapeSplineData.cpp
--- a/Engine/Plugins/Experimental/PCG/Source/PCG/Private/Data/PCGLandscapeSplineData.cpp
+++ b/Engine/Plugins/Experimental/PCG/Source/PCG/Private/Data/PCGLandscapeSplineData.cpp
@@ -217,7 +217,7 @@ float UPCGLandscapeSplineData::GetDensityAtPosition(const FVector& InPosition) c
 	{
 		// Considering the landscape spline always exists on the landscape,
 		// we'll ignore the Z component of the input here for the bounds check.
-		if (!Segment->GetBounds().IsInsideXY(Position))
+		if(!PCGHelpers::IsInsideBoundsXY(Segment->GetBounds(), Position))
 		{
 			continue;
 		}
diff --git a/Engine/Plugins/Experimental/PCG/Source/PCG/Private/Data/PCGUnionData.cpp b/Engine/Plugins/Experimental/PCG/Source/PCG/Private/Data/PCGUnionData.cpp
--- a/Engine/Plugins/Experimental/PCG/Source/PCG/Private/Data/PCGUnionData.cpp
+++ b/Engine/Plugins/Experimental/PCG/Source/PCG/Private/Data/PCGUnionData.cpp
@@ -3,6 +3,7 @@
 #include "Data/PCGUnionData.h"
 #include "Data/PCGPointData.h"
 #include "Helpers/PCGAsync.h"
+#include "PCGHelpers.h"
 
 namespace PCGUnionDataMaths
 {
@@ -76,11 +77,11 @@ FBox UPCGUnionData::GetStrictBounds() const
 float UPCGUnionData::GetDensityAtPosition(const FVector& InPosition) const
 {
 	// Early exits
-	if (!CachedBounds.IsInside(InPosition))
+	if(!PCGHelpers::IsInsideBounds(CachedBounds, InPosition))
 	{
 		return 0;
 	}
-	else if (CachedStrictBounds.IsInside(InPosition))
+	else if(PCGHelpers::IsInsideBounds(CachedStrictBounds, InPosition))
 	{
 		return 1.0f;
 	}
@@ -90,7 +91,7 @@ float UPCGUnionData::GetDensityAtPosition(const FVector& InPosition) const
 	// the strict bounds twice per data, but it will perform better in the worst case.
 	for (int32 DataIndex = 0; DataIndex < Data.Num(); ++DataIndex)
 	{
-		if (Data[DataIndex]->GetStrictBounds().IsInside(InPosition))
+		if(PCGHelpers::IsInsideBounds(Data[DataIndex]->GetStrictBounds(), InPosition))
 		{
 			return 1.0f;
 		}
diff --git a/Engine/Plugins/Experimental/PCG/Source/PCG/Private/Data/PCGVolumeData.cpp b/Engine/Plugins/Experimental/PCG/Source/PCG/Private/Data/PCGVolumeData.cpp
--- a/Engine/Plugins/Experimental/PCG/Source/PCG/Private/Data/PCGVolumeData.cpp
+++ b/Engine/Plugins/Experimental/PCG/Source/PCG/Private/Data/PCGVolumeData.cpp
@@ -1,6 +1,7 @@
 // Copyright Epic Games, Inc. All Rights Reserved.
 
 #include "Data/PCGVolumeData.h"
+#include "PCGHelpers.h"
 
 #include "Components/BrushComponent.h"
 #include "GameFramework/Volume.h"
@@ -45,9 +46,9 @@ float UPCGVolumeData::GetDensityAtPosition(const FVector& InPosition) const
 {
 	// TODO: support fall-off between 0-1 density as needed
 	// could use FBox::GetClosestPointTo
-	if (GetBounds().IsInside(InPosition))
+	if(PCGHelpers::IsInsideBounds(GetBounds(), InPosition))
 	{
-		if(!Volume || GetStrictBounds().IsInside(InPosition))
+		if(!Volume || PCGHelpers::IsInsideBounds(GetStrictBounds(), InPosition))
 		{
 			return 1.0f;
 		}
diff --git a/Engine/Plugins/Experimental/PCG/Source/PCG/Private/PCGHelpers.cpp b/Engine/Plugins/Experimental/PCG/Source/PCG/Private/PCGHelpers.cpp
--- a/Engine/Plugins/Experimental/PCG/Source/PCG/Private/PCGHelpers.cpp
+++ b/Engine/Plugins/Experimental/PCG/Source/PCG/Private/PCGHelpers.cpp
@@ -27,6 +27,19 @@ namespace PCGHelpers
 		return ((A * 196314165U) + 907633515U) ^ ((B * 73148459U) + 453816763U) ^ ((C * 34731343U) + 453816743U);
 	}
 
+	bool IsInsideBounds(const FBox& InBox, const FVector& InPosition)
+	{
+		return (InPosition.X >= InBox.Min.X) && (InPosition.X < InBox.Max.X) &&
+			(InPosition.Y >= InBox.Min.Y) && (InPosition.Y < InBox.Max.Y) &&
+			(InPosition.Z >= InBox.Min.Z) && (InPosition.Z < InBox.Max.Z);
+	}
+
+	bool IsInsideBoundsXY(const FBox& InBox, const FVector& InPosition)
+	{
+		return (InPosition.X >= InBox.Min.X) && (InPosition.X < InBox.Max.X) &&
+			(InPosition.Y >= InBox.Min.Y) && (InPosition.Y < InBox.Max.Y);
+	}
+
 	FBox GetActorBounds(AActor* InActor)
 	{
 		// Specialized version of GetComponentsBoundingBox that skips over PCG generated components
diff --git a/Engine/Plugins/Experimental/PCG/Source/PCG/Private/PCGHelpers.h b/Engine/Plugins/Experimental/PCG/Source/PCG/Private/PCGHelpers.h
--- a/Engine/Plugins/Experimental/PCG/Source/PCG/Private/PCGHelpers.h
+++ b/Engine/Plugins/Experimental/PCG/Source/PCG/Private/PCGHelpers.h
@@ -21,6 +21,9 @@ namespace PCGHelpers
 	int ComputeSeed(int A, int B);
 	int ComputeSeed(int A, int B, int C);
 
+	bool IsInsideBounds(const FBox& InBox, const FVector& InPosition);
+	bool IsInsideBoundsXY(const FBox& InBox, const FVector& InPosition);
+
 	FBox GetActorBounds(AActor* InActor);
 	FBox GetLandscapeBounds(ALandscapeProxy* InLandscape);
 
