Fixed density error in difference

#rnx

diff --git a/Engine/Plugins/Experimental/PCG/Source/PCG/Private/Data/PCGDifferenceData.cpp b/Engine/Plugins/Experimental/PCG/Source/PCG/Private/Data/PCGDifferenceData.cpp
--- a/Engine/Plugins/Experimental/PCG/Source/PCG/Private/Data/PCGDifferenceData.cpp
+++ b/Engine/Plugins/Experimental/PCG/Source/PCG/Private/Data/PCGDifferenceData.cpp
@@ -97,7 +97,7 @@ float UPCGDifferenceData::GetDensityAtPosition(const FVector& InPosition) const
 		}
 
 		const float DensityInDifference = Difference->GetDensityAtPosition(InPosition);
-		return (DensityInDifference > 0 ? 0 : DensityInSource);
+		return FMath::Max(0, DensityInSource - DensityInDifference);
 	}
 	else
 	{
