Fixed bad skip condition in density remap code

#rnx

diff --git a/Engine/Plugins/Experimental/PCG/Source/PCG/Private/Elements/PCGDensityRemap.cpp b/Engine/Plugins/Experimental/PCG/Source/PCG/Private/Elements/PCGDensityRemap.cpp
--- a/Engine/Plugins/Experimental/PCG/Source/PCG/Private/Elements/PCGDensityRemap.cpp
+++ b/Engine/Plugins/Experimental/PCG/Source/PCG/Private/Elements/PCGDensityRemap.cpp
@@ -28,8 +28,6 @@ bool FPCGLinearDensityRemapElement::ExecuteInternal(FPCGContext* Context) const
 	const float RemapMin = FMath::Min(Settings->RemapMin, Settings->RemapMax);
 	const float RemapMax = FMath::Max(Settings->RemapMin, Settings->RemapMax);
 
-	const bool bTrivialRemapping = (RemapMin == 0.0f && RemapMax == 1.0f && Settings->bMultiplyDensity);
-
 	// TODO: embarassingly parallel loop
 	for (const FPCGTaggedData& Input : Inputs)
 	{
@@ -42,13 +40,6 @@ bool FPCGLinearDensityRemapElement::ExecuteInternal(FPCGContext* Context) const
 			continue;
 		}
 
-		// Skip processing if the remapping is trivial
-		if (bTrivialRemapping)
-		{
-			PCGE_LOG(Verbose, "Skipped - trivial remapping");
-			continue;
-		}
-
 		const UPCGPointData* OriginalData = Cast<UPCGSpatialData>(Input.Data)->ToPointData(Context);
 
 		if (!OriginalData)
@@ -81,7 +72,7 @@ bool FPCGLinearDensityRemapElement::ExecuteInternal(FPCGContext* Context) const
 			{
 				OutPoint = Points[Index];
 				FRandomStream RandomSource(PCGHelpers::ComputeSeed(Settings->Seed, OutPoint.Seed));
-				OutPoint.Density  = RandomSource.FRandRange(RemapMin, RemapMax);
+				OutPoint.Density = RandomSource.FRandRange(RemapMin, RemapMax);
 				return true;
 			});
 		}
