Hide the self pin on the loop blueprint methods

#rnx

diff --git a/Engine/Plugins/Experimental/PCG/Source/PCG/Public/Elements/PCGExecuteBlueprint.h b/Engine/Plugins/Experimental/PCG/Source/PCG/Public/Elements/PCGExecuteBlueprint.h
--- a/Engine/Plugins/Experimental/PCG/Source/PCG/Public/Elements/PCGExecuteBlueprint.h
+++ b/Engine/Plugins/Experimental/PCG/Source/PCG/Public/Elements/PCGExecuteBlueprint.h
@@ -49,10 +49,10 @@ public:
 	bool PointPairLoopBody(const UPCGPointData* InA, const UPCGPointData* InB, const FPCGPoint& InPointA, const FPCGPoint& InPointB, FPCGPoint& OutPoint) const;
 
 	/** Calls the LoopBody function on all points */
-	UFUNCTION(BlueprintCallable, Category = Execution)
+	UFUNCTION(BlueprintCallable, Category = Execution, meta = (HideSelfPin = "true"))
 	void LoopOnPoints(UPARAM(ref) FPCGContext& InContext, const UPCGPointData* InData, UPCGPointData*& OutData) const;
 
-	UFUNCTION(BlueprintCallable, Category = Execution)
+	UFUNCTION(BlueprintCallable, Category = Execution, meta = (HideSelfPin = "true"))
 	void LoopOnPointPairs(UPARAM(ref) FPCGContext& InContext, const UPCGPointData* InA, const UPCGPointData* InB, UPCGPointData*& OutData) const;
 
 	/** Called after object creation to setup the object callbacks */
