Added binary difference/union mode, might be placeholder for more general feature.
Fixed bug in projection bounds

diff --git a/Engine/Plugins/Experimental/PCG/Source/PCG/Private/Data/PCGDifferenceData.cpp b/Engine/Plugins/Experimental/PCG/Source/PCG/Private/Data/PCGDifferenceData.cpp
--- a/Engine/Plugins/Experimental/PCG/Source/PCG/Private/Data/PCGDifferenceData.cpp
+++ b/Engine/Plugins/Experimental/PCG/Source/PCG/Private/Data/PCGDifferenceData.cpp
@@ -10,10 +10,28 @@ namespace PCGDifferenceDataUtils
 {
 	EPCGUnionDensityFunction ToUnionDensityFunction(EPCGDifferenceDensityFunction InDensityFunction)
 	{
-		return (InDensityFunction == EPCGDifferenceDensityFunction::ClampedSubstraction ? EPCGUnionDensityFunction::ClampedAddition : EPCGUnionDensityFunction::Maximum);
+		if (InDensityFunction == EPCGDifferenceDensityFunction::ClampedSubstraction)
+		{
+			return EPCGUnionDensityFunction::ClampedAddition;
+		}
+		else if (InDensityFunction == EPCGDifferenceDensityFunction::Binary)
+		{
+			return EPCGUnionDensityFunction::Binary;
+		}
+		else
+		{
+			return EPCGUnionDensityFunction::Maximum;
+		}
 	}
 }
 
+float UPCGDifferenceData::GetDensityAtPositionFromDifference(const FVector& InPosition) const
+{
+	const float DensityInDifference = (Difference ? Difference->GetDensityAtPosition(InPosition) : 0.0f);
+	const bool bBinaryDensity = (DensityFunction == EPCGDifferenceDensityFunction::Binary);
+	return ((bBinaryDensity && DensityInDifference > 0) ? 1.0f : DensityInDifference);
+}
+
 void UPCGDifferenceData::Initialize(const UPCGSpatialData* InData)
 {
 	check(InData);
@@ -25,6 +43,13 @@ void UPCGDifferenceData::AddDifference(const UPCGSpatialData* InDifference)
 {
 	check(InDifference);
 
+	// In the eventuality that the difference has no overlap with the source, then we can drop it directly
+	if (!GetBounds().Intersect(InDifference->GetBounds()))
+	{
+		return;
+	}
+
+	// First difference element we'll keep as is, but subsequent ones will be pushed into a union
 	if (!Difference)
 	{
 		Difference = InDifference;
@@ -78,7 +103,7 @@ FBox UPCGDifferenceData::GetBounds() const
 
 FBox UPCGDifferenceData::GetStrictBounds() const
 {
-	return FBox(EForceInit::ForceInit);
+	return Difference ? FBox(EForceInit::ForceInit) : Source->GetStrictBounds();
 }
 
 float UPCGDifferenceData::GetDensityAtPosition(const FVector& InPosition) const
@@ -96,7 +121,7 @@ float UPCGDifferenceData::GetDensityAtPosition(const FVector& InPosition) const
 			return 0;
 		}
 
-		const float DensityInDifference = Difference->GetDensityAtPosition(InPosition);
+		const float DensityInDifference = GetDensityAtPositionFromDifference(InPosition);
 		return FMath::Max(0, DensityInSource - DensityInDifference);
 	}
 	else
@@ -118,7 +143,7 @@ FPCGPoint UPCGDifferenceData::TransformPoint(const FPCGPoint& InPoint) const
 
 	if (Difference && TransformedPoint.Density > 0)
 	{
-		TransformedPoint.Density = FMath::Max(0, TransformedPoint.Density - Difference->GetDensityAtPosition(TransformedPoint.Transform.GetLocation()));
+		TransformedPoint.Density = FMath::Max(0, TransformedPoint.Density - GetDensityAtPositionFromDifference(TransformedPoint.Transform.GetLocation()));
 	}
 
 	return TransformedPoint;
@@ -158,7 +183,7 @@ const UPCGPointData* UPCGDifferenceData::CreatePointData(FPCGContext* Context) c
 	FPCGAsync::AsyncPointProcessing(Context, SourcePoints.Num(), TargetPoints, [this, &SourcePoints](int32 Index, FPCGPoint& OutPoint)
 	{
 		const FPCGPoint& Point = SourcePoints[Index];
-		const float DensityInDifference = Difference->GetDensityAtPosition(Point.Transform.GetLocation());
+		const float DensityInDifference = GetDensityAtPositionFromDifference(Point.Transform.GetLocation());
 
 		if (DensityInDifference < Point.Density)
 		{
diff --git a/Engine/Plugins/Experimental/PCG/Source/PCG/Private/Data/PCGProjectionData.cpp b/Engine/Plugins/Experimental/PCG/Source/PCG/Private/Data/PCGProjectionData.cpp
--- a/Engine/Plugins/Experimental/PCG/Source/PCG/Private/Data/PCGProjectionData.cpp
+++ b/Engine/Plugins/Experimental/PCG/Source/PCG/Private/Data/PCGProjectionData.cpp
@@ -61,6 +61,12 @@ FBox UPCGProjectionData::ProjectBounds(const FBox& InBounds) const
 				(Corner % 2) ? InBounds.Max.Z : InBounds.Min.Z));
 	}
 
+	// Fixup the Z direction, as transforming the corners is not sufficient
+	const FVector::FReal HalfHeight = 0.5 * (InBounds.Max.Z - InBounds.Min.Z);
+	FVector BoundsCenter = InBounds.GetCenter();
+	Bounds += BoundsCenter + Target->GetNormal() * HalfHeight;
+	Bounds += BoundsCenter - Target->GetNormal() * HalfHeight;
+
 	return Bounds;
 }
 
diff --git a/Engine/Plugins/Experimental/PCG/Source/PCG/Private/Data/PCGUnionData.cpp b/Engine/Plugins/Experimental/PCG/Source/PCG/Private/Data/PCGUnionData.cpp
--- a/Engine/Plugins/Experimental/PCG/Source/PCG/Private/Data/PCGUnionData.cpp
+++ b/Engine/Plugins/Experimental/PCG/Source/PCG/Private/Data/PCGUnionData.cpp
@@ -13,7 +13,11 @@ namespace PCGUnionDataMaths
 		{
 			return FMath::Min(InDensityToUpdate + InOtherDensity, 1.0f);
 		}
-		else
+		else if (DensityFunction == EPCGUnionDensityFunction::Binary)
+		{
+			return InOtherDensity > 0 ? 1.0f : InDensityToUpdate;
+		}
+		else // Maximum
 		{
 			return FMath::Max(InDensityToUpdate, InOtherDensity);
 		}
@@ -128,6 +132,11 @@ FPCGPoint UPCGUnionData::TransformPoint(const FPCGPoint& InPoint) const
 	{
 		FPCGPoint TransformedPoint = FirstNonTrivialTransformData->TransformPoint(InPoint);
 
+		if (DensityFunction == EPCGUnionDensityFunction::Binary && TransformedPoint.Density > 0)
+		{
+			TransformedPoint.Density = 1.0f;
+		}
+
 		const int32 DataCount = Data.Num();
 		for(int32 DataIndex = 0; DataIndex < DataCount && TransformedPoint.Density < 1.0f; ++DataIndex)
 		{
@@ -154,13 +163,15 @@ const UPCGPointData* UPCGUnionData::CreatePointData(FPCGContext* Context) const
 {
 	TRACE_CPUPROFILER_EVENT_SCOPE(UPCGUnionData::CreatePointData);
 
+	const bool bBinaryDensity = (DensityFunction == EPCGUnionDensityFunction::Binary);
+
 	// Trivial results
 	if (Data.Num() == 0)
 	{
 		UE_LOG(LogPCG, Error, TEXT("Invalid union"));
 		return nullptr;
 	}
-	else if (Data.Num() == 1)
+	else if (Data.Num() == 1 && !bBinaryDensity)
 	{
 		UE_LOG(LogPCG, Verbose, TEXT("Union is trivial"));
 		return Data[0]->ToPointData(Context);
@@ -187,6 +198,15 @@ const UPCGPointData* UPCGUnionData::CreatePointData(FPCGContext* Context) const
 			{
 				TargetPoints.Append(Datum->ToPointData(Context)->GetPoints());
 			}
+
+			// Correct density for binary-style union
+			if (bBinaryDensity)
+			{
+				for (FPCGPoint& TargetPoint : TargetPoints)
+				{
+					TargetPoint.Density = (TargetPoint.Density > 0 ? 1.0f : 0);
+				}
+			}
 		}
 		break;
 	}
@@ -237,13 +257,15 @@ void UPCGUnionData::CreateSequentialPointData(FPCGContext* Context, UPCGPointDat
 
 			OutPoint = Point;
 
+			if (DensityFunction == EPCGUnionDensityFunction::Binary && OutPoint.Density > 0)
+			{
+				OutPoint.Density = 1.0f;
+			}
+
 			// Compute final density based on current & following data
-			for (int32 FollowingDataIndex = DataIndex + DataIndexIncrement; FollowingDataIndex != LastDataIndex; FollowingDataIndex += DataIndexIncrement)
+			for (int32 FollowingDataIndex = DataIndex + DataIndexIncrement; FollowingDataIndex != LastDataIndex && OutPoint.Density < 1.0f; FollowingDataIndex += DataIndexIncrement)
 			{
-				if (PCGUnionDataMaths::UpdateDensity(OutPoint.Density, Data[FollowingDataIndex]->GetDensityAtPosition(OutPoint.Transform.GetLocation()), DensityFunction) == 1.0f)
-				{
-					break;
-				}
+				PCGUnionDataMaths::UpdateDensity(OutPoint.Density, Data[FollowingDataIndex]->GetDensityAtPosition(OutPoint.Transform.GetLocation()), DensityFunction);
 			}
 
 			return true;
diff --git a/Engine/Plugins/Experimental/PCG/Source/PCG/Public/Data/PCGDifferenceData.h b/Engine/Plugins/Experimental/PCG/Source/PCG/Public/Data/PCGDifferenceData.h
--- a/Engine/Plugins/Experimental/PCG/Source/PCG/Public/Data/PCGDifferenceData.h
+++ b/Engine/Plugins/Experimental/PCG/Source/PCG/Public/Data/PCGDifferenceData.h
@@ -13,7 +13,8 @@ UENUM()
 enum class EPCGDifferenceDensityFunction : uint8
 {
 	Minimum,
-	ClampedSubstraction
+	ClampedSubstraction,
+	Binary
 };
 
 UCLASS(BlueprintType, ClassGroup = (Procedural))
@@ -59,4 +60,6 @@ protected:
 
 	UPROPERTY(BlueprintSetter = SetDensityFunction, EditAnywhere, Category = Settings)
 	EPCGDifferenceDensityFunction DensityFunction = EPCGDifferenceDensityFunction::Minimum;
+
+	float GetDensityAtPositionFromDifference(const FVector& InPosition) const;
 };
diff --git a/Engine/Plugins/Experimental/PCG/Source/PCG/Public/Data/PCGUnionData.h b/Engine/Plugins/Experimental/PCG/Source/PCG/Public/Data/PCGUnionData.h
--- a/Engine/Plugins/Experimental/PCG/Source/PCG/Public/Data/PCGUnionData.h
+++ b/Engine/Plugins/Experimental/PCG/Source/PCG/Public/Data/PCGUnionData.h
@@ -19,7 +19,8 @@ UENUM()
 enum class EPCGUnionDensityFunction : uint8
 {
 	Maximum,
-	ClampedAddition
+	ClampedAddition,
+	Binary
 };
 
 UCLASS(BlueprintType, ClassGroup=(Procedural))
