Added duplicate points debug test (also exposed in the self-pruning test)

diff --git a/Engine/Plugins/Experimental/PCG/Source/PCG/Private/Elements/PCGDebugElement.cpp b/Engine/Plugins/Experimental/PCG/Source/PCG/Private/Elements/PCGDebugElement.cpp
--- a/Engine/Plugins/Experimental/PCG/Source/PCG/Private/Elements/PCGDebugElement.cpp
+++ b/Engine/Plugins/Experimental/PCG/Source/PCG/Private/Elements/PCGDebugElement.cpp
@@ -40,11 +40,11 @@ namespace PCGDebugElement
 			Materials.Add(Material);
 		}
 
-		TArray<FPCGTaggedData> Outputs = Context->OutputData.GetInputs();
+		TArray<FPCGTaggedData> Inputs = Context->InputData.GetInputs();
 
-		for (const FPCGTaggedData& Output : Outputs)
+		for(const FPCGTaggedData& Input : Inputs)
 		{
-			const UPCGSpatialData* SpatialData = Cast<UPCGSpatialData>(Output.Data);
+			const UPCGSpatialData* SpatialData = Cast<UPCGSpatialData>(Input.Data);
 
 			if (!SpatialData)
 			{
diff --git a/Engine/Plugins/Experimental/PCG/Source/PCG/Private/Elements/PCGSelfPruning.cpp b/Engine/Plugins/Experimental/PCG/Source/PCG/Private/Elements/PCGSelfPruning.cpp
--- a/Engine/Plugins/Experimental/PCG/Source/PCG/Private/Elements/PCGSelfPruning.cpp
+++ b/Engine/Plugins/Experimental/PCG/Source/PCG/Private/Elements/PCGSelfPruning.cpp
@@ -37,142 +37,193 @@ namespace PCGSelfPruningAlgorithms
 	}
 }
 
-FPCGElementPtr UPCGSelfPruningSettings::CreateElement() const
+namespace PCGSelfPruningElement
 {
-	return MakeShared<FPCGSelfPruningElement>();
-}
-
-bool FPCGSelfPruningElement::ExecuteInternal(FPCGContext* Context) const
-{
-	TRACE_CPUPROFILER_EVENT_SCOPE(FPCGSelfPruningElement::Execute);
-	// TODO: time-sliced implementation
-	const UPCGSelfPruningSettings* Settings = Context->GetInputSettings<UPCGSelfPruningSettings>();
-	check(Settings);
-
-	// Early out: if pruning is disabled
-	if (Settings->PruningType == EPCGSelfPruningType::None)
+	void DensityBoundsExclusion(const TArray<const FPCGPoint*>& SortedPoints, const UPCGPointData::PointOctree& Octree, TSet<const FPCGPoint*>& ExclusionPoints)
 	{
-		Context->OutputData = Context->InputData;
-		PCGE_LOG(Verbose, "Skipped - Type is none");
-		return true;
-	}
+		TSet<const FPCGPoint*> ExcludedPoints;
+		ExcludedPoints.Reserve(SortedPoints.Num());
+		
+		for (const FPCGPoint* Point : SortedPoints)
+		{
+			if (ExcludedPoints.Contains(Point))
+			{
+				continue;
+			}
 
-	const FVector::FReal RadiusEquality = 1.0f + Settings->RadiusSimilarityFactor;
-	const FVector::FReal SquaredRadiusEquality = FMath::Square(RadiusEquality);
+			ExclusionPoints.Add(Point);
 
-	TArray<FPCGTaggedData>& Outputs = Context->OutputData.TaggedData;
+			const FBoxSphereBounds PointBounds = Point->GetDensityBounds();
+			Octree.FindElementsWithBoundsTest(FBoxCenterAndExtent(PointBounds.Origin, PointBounds.BoxExtent), [&ExclusionPoints, &ExcludedPoints](const FPCGPointRef& InPointRef)
+			{
+				// TODO: check on an oriented-box basis?
+				if (!ExclusionPoints.Contains(InPointRef.Point))
+				{
+					ExcludedPoints.Add(InPointRef.Point);
+				}
+			});
+		}
+	}
 
-	// TODO: embarassingly parallel loop
-	TArray<FPCGTaggedData> Inputs = Context->InputData.GetInputs();
-	for (const FPCGTaggedData& Input : Inputs)
+	void DuplicatePointsExclusion(const TArray<const FPCGPoint*>& SortedPoints, const UPCGPointData::PointOctree& Octree, TSet<const FPCGPoint*>& ExclusionPoints)
 	{
-		const UPCGSpatialData* SpatialInput = Cast<UPCGSpatialData>(Input.Data);
+		TSet<const FPCGPoint*> ExcludedPoints;
+		ExcludedPoints.Reserve(SortedPoints.Num());
 
-		if (!SpatialInput)
+		for (const FPCGPoint* Point : SortedPoints)
 		{
-			PCGE_LOG(Error, "Invalid input data");
-			continue;
+			if (ExcludedPoints.Contains(Point))
+			{
+				continue;
+			}
+
+			ExclusionPoints.Add(Point);
+
+			Octree.FindElementsWithBoundsTest(FBoxCenterAndExtent(Point->Transform.GetLocation(), FVector::Zero()), [&ExclusionPoints, &ExcludedPoints, Point](const FPCGPointRef& InPointRef)
+			{
+				if ((Point->Transform.GetLocation() - InPointRef.Point->Transform.GetLocation()).SquaredLength() <= SMALL_NUMBER &&
+					!ExclusionPoints.Contains(InPointRef.Point))
+				{
+					ExcludedPoints.Add(InPointRef.Point);
+				}
+			});
 		}
+	}
 
-		const UPCGPointData* InputPointData = SpatialInput->ToPointData(Context);
-		const TArray<FPCGPoint>& Points = InputPointData->GetPoints();
-		const UPCGPointData::PointOctree& Octree = InputPointData->GetOctree();
-
-		// Self-pruning will be done as follows:
-		// For each point:
-		//  if in its vicinity, there is >=1 non-rejected point with a radius significantly larger
-		//  or in its range + has a randomly assigned index -> we'll look at its seed
-		//  then remove this point
-		TArray<const FPCGPoint*> SortedPoints;
-		SortedPoints.Reserve(Points.Num());
-		for (const FPCGPoint& Point : Points)
+	void Execute(FPCGContext* Context, EPCGSelfPruningType PruningType, float RadiusSimilarityFactor, bool bRandomizedPruning)
+	{
+		// Early out: if pruning is disabled
+		if (PruningType == EPCGSelfPruningType::None)
 		{
-			SortedPoints.Add(&Point);
+			Context->OutputData = Context->InputData;
+			PCGE_LOG_C(Verbose, Context, "Skipped - Type is none");
+			return;
 		}
 
-		// Apply proper sort algorithm
-		if (Settings->PruningType == EPCGSelfPruningType::LargeToSmall)
+		const FVector::FReal RadiusEquality = 1.0f + RadiusSimilarityFactor;
+		const FVector::FReal SquaredRadiusEquality = FMath::Square(RadiusEquality);
+
+		TArray<FPCGTaggedData>& Outputs = Context->OutputData.TaggedData;
+
+		// TODO: embarassingly parallel loop
+		TArray<FPCGTaggedData> Inputs = Context->InputData.GetInputs();
+		for (const FPCGTaggedData& Input : Inputs)
 		{
-			if (Settings->bRandomizedPruning)
+			const UPCGSpatialData* SpatialInput = Cast<UPCGSpatialData>(Input.Data);
+
+			if (!SpatialInput)
 			{
-				Algo::Sort(SortedPoints, [SquaredRadiusEquality](const FPCGPoint* A, const FPCGPoint* B) { return !PCGSelfPruningAlgorithms::SortSmallToLargeWithRandom(A, B, SquaredRadiusEquality); });
+				PCGE_LOG_C(Error, Context, "Invalid input data");
+				continue;
 			}
-			else
+
+			const UPCGPointData* InputPointData = SpatialInput->ToPointData(Context);
+			const TArray<FPCGPoint>& Points = InputPointData->GetPoints();
+			const UPCGPointData::PointOctree& Octree = InputPointData->GetOctree();
+
+			// Self-pruning will be done as follows:
+			// For each point:
+			//  if in its vicinity, there is >=1 non-rejected point with a radius significantly larger
+			//  or in its range + has a randomly assigned index -> we'll look at its seed
+			//  then remove this point
+			TArray<const FPCGPoint*> SortedPoints;
+			SortedPoints.Reserve(Points.Num());
+			for (const FPCGPoint& Point : Points)
 			{
-				Algo::Sort(SortedPoints, [SquaredRadiusEquality](const FPCGPoint* A, const FPCGPoint* B) { return !PCGSelfPruningAlgorithms::SortSmallToLargeNoRandom(A, B, SquaredRadiusEquality); });
+				SortedPoints.Add(&Point);
 			}
-		}
-		else if (Settings->PruningType == EPCGSelfPruningType::SmallToLarge)
-		{
-			if (Settings->bRandomizedPruning)
+
+			// Apply proper sort algorithm
+			if (PruningType == EPCGSelfPruningType::LargeToSmall)
 			{
-				Algo::Sort(SortedPoints, [SquaredRadiusEquality](const FPCGPoint* A, const FPCGPoint* B) { return PCGSelfPruningAlgorithms::SortSmallToLargeWithRandom(A, B, SquaredRadiusEquality); });
+				if (bRandomizedPruning)
+				{
+					Algo::Sort(SortedPoints, [SquaredRadiusEquality](const FPCGPoint* A, const FPCGPoint* B) { return !PCGSelfPruningAlgorithms::SortSmallToLargeWithRandom(A, B, SquaredRadiusEquality); });
+				}
+				else
+				{
+					Algo::Sort(SortedPoints, [SquaredRadiusEquality](const FPCGPoint* A, const FPCGPoint* B) { return !PCGSelfPruningAlgorithms::SortSmallToLargeNoRandom(A, B, SquaredRadiusEquality); });
+				}
 			}
-			else
+			else if (PruningType == EPCGSelfPruningType::SmallToLarge)
 			{
-				Algo::Sort(SortedPoints, [SquaredRadiusEquality](const FPCGPoint* A, const FPCGPoint* B) { return PCGSelfPruningAlgorithms::SortSmallToLargeNoRandom(A, B, SquaredRadiusEquality); });
+				if (bRandomizedPruning)
+				{
+					Algo::Sort(SortedPoints, [SquaredRadiusEquality](const FPCGPoint* A, const FPCGPoint* B) { return PCGSelfPruningAlgorithms::SortSmallToLargeWithRandom(A, B, SquaredRadiusEquality); });
+				}
+				else
+				{
+					Algo::Sort(SortedPoints, [SquaredRadiusEquality](const FPCGPoint* A, const FPCGPoint* B) { return PCGSelfPruningAlgorithms::SortSmallToLargeNoRandom(A, B, SquaredRadiusEquality); });
+				}
 			}
-		}
-		else
-		{
-			if (Settings->bRandomizedPruning)
+			else
 			{
-				Algo::Sort(SortedPoints, PCGSelfPruningAlgorithms::RandomSort);
+				if (bRandomizedPruning)
+				{
+					Algo::Sort(SortedPoints, PCGSelfPruningAlgorithms::RandomSort);
+				}
 			}
-		}
-
-		TSet<const FPCGPoint*> ExclusionPoints;
-		ExclusionPoints.Reserve(Points.Num());
 
-		TSet<const FPCGPoint*> ExcludedPoints;
-		ExcludedPoints.Reserve(Points.Num());
+			TSet<const FPCGPoint*> ExclusionPoints;
+			ExclusionPoints.Reserve(Points.Num());
 
-		for (int PointIndex = 0; PointIndex < SortedPoints.Num(); ++PointIndex)
-		{
-			// Find all points in the vicinity of the current point and reject them, if they aren't part of the exclusions yet
-			// If any of these have been seen before (e.g. in the ExclusionPoints)
-			// then reject this point
-			const FPCGPoint* CurrentPoint = SortedPoints[PointIndex];
+			const bool bIsDuplicateTest = (PruningType == EPCGSelfPruningType::RemoveDuplicates);
 
-			if (ExcludedPoints.Contains(CurrentPoint))
+			if (bIsDuplicateTest)
 			{
-				continue;
+				PCGSelfPruningElement::DuplicatePointsExclusion(SortedPoints, Octree, ExclusionPoints);
+			}
+			else
+			{
+				PCGSelfPruningElement::DensityBoundsExclusion(SortedPoints, Octree, ExclusionPoints);
 			}
 
-			ExclusionPoints.Add(CurrentPoint);
+			// Finally, output all points that are present in the ExclusionPoints.
+			FPCGTaggedData& Output = Outputs.Emplace_GetRef();
+			Output = Input;
 
-			const FBoxSphereBounds CurrentPointBounds = CurrentPoint->GetDensityBounds();
-			Octree.FindElementsWithBoundsTest(FBoxCenterAndExtent(CurrentPointBounds.Origin, CurrentPointBounds.BoxExtent), [&ExclusionPoints, &ExcludedPoints](const FPCGPointRef& InPointRef) {
-				if (!ExclusionPoints.Contains(InPointRef.Point))
-				{
-					ExcludedPoints.Add(InPointRef.Point);
-				}
-			});
-		}
-
-		// Finally, output all points that are present in the ExclusionPoints.
-		FPCGTaggedData& Output = Outputs.Emplace_GetRef();
-		Output = Input;
+			UPCGPointData* PrunedData = NewObject<UPCGPointData>();
+			PrunedData->TargetActor = SpatialInput->TargetActor;
 
-		UPCGPointData* PrunedData = NewObject<UPCGPointData>();
-		PrunedData->TargetActor = SpatialInput->TargetActor;
+			Output.Data = PrunedData;
 
-		Output.Data = PrunedData;
+			TArray<FPCGPoint>& OutputPoints = PrunedData->GetMutablePoints();
+			OutputPoints.Reserve(ExclusionPoints.Num());
 
-		TArray<FPCGPoint>& OutputPoints = PrunedData->GetMutablePoints();
-		OutputPoints.Reserve(ExclusionPoints.Num());
+			for (const FPCGPoint* Point : ExclusionPoints)
+			{
+				OutputPoints.Add(*Point);
+			}
 
-		for (const FPCGPoint* Point : ExclusionPoints)
-		{
-			OutputPoints.Add(*Point);
+			if (bIsDuplicateTest)
+			{
+				PCGE_LOG_C(Verbose, Context, "Removed %d duplicate points from %d source points", Points.Num() - OutputPoints.Num(), Points.Num());
+			}
+			else
+			{
+				PCGE_LOG_C(Verbose, Context, "Generated %d points from %d source points", OutputPoints.Num(), Points.Num());
+			}
 		}
 
-		PCGE_LOG(Verbose, "Generated %d points from %d source points", OutputPoints.Num(), Points.Num());
+		// Finally, forward any exclusions/settings
+		Outputs.Append(Context->InputData.GetExclusions());
+		Outputs.Append(Context->InputData.GetAllSettings());
 	}
+}
+
+FPCGElementPtr UPCGSelfPruningSettings::CreateElement() const
+{
+	return MakeShared<FPCGSelfPruningElement>();
+}
+
+bool FPCGSelfPruningElement::ExecuteInternal(FPCGContext* Context) const
+{
+	TRACE_CPUPROFILER_EVENT_SCOPE(FPCGSelfPruningElement::Execute);
+	// TODO: time-sliced implementation
+	const UPCGSelfPruningSettings* Settings = Context->GetInputSettings<UPCGSelfPruningSettings>();
+	check(Settings);
 
-	// Finally, forward any exclusions/settings
-	Outputs.Append(Context->InputData.GetExclusions());
-	Outputs.Append(Context->InputData.GetAllSettings());
+	PCGSelfPruningElement::Execute(Context, Settings->PruningType, Settings->RadiusSimilarityFactor, Settings->bRandomizedPruning);
 
 	return true;
 }
\ No newline at end of file
diff --git a/Engine/Plugins/Experimental/PCG/Source/PCG/Private/PCGElement.cpp b/Engine/Plugins/Experimental/PCG/Source/PCG/Private/PCGElement.cpp
--- a/Engine/Plugins/Experimental/PCG/Source/PCG/Private/PCGElement.cpp
+++ b/Engine/Plugins/Experimental/PCG/Source/PCG/Private/PCGElement.cpp
@@ -5,6 +5,7 @@
 #include "PCGModule.h"
 #include "PCGSettings.h"
 #include "Elements/PCGDebugElement.h"
+#include "Elements/PCGSelfPruning.h"
 #include "Graph/PCGGraphCache.h"
 
 bool IPCGElement::Execute(FPCGContext* Context) const
@@ -115,14 +116,41 @@ bool IPCGElement::Execute(FPCGContext* Context) const
 		}
 
 #if WITH_EDITOR
-		if (bDone && Settings && (Settings->ExecutionMode == EPCGSettingsExecutionMode::Debug || Settings->ExecutionMode == EPCGSettingsExecutionMode::Isolated))
+		if (bDone && Settings)
 		{
-			PCGDebugElement::ExecuteDebugDisplay(Context);
+			if (Settings->DebugSettings.bCheckForDuplicates)
+			{
+				FPCGDataCollection ElementInputs = Context->InputData;
+				FPCGDataCollection ElementOutputs = Context->OutputData;
+
+				Context->InputData = ElementOutputs;
+				Context->OutputData = FPCGDataCollection();
+
+				PCGE_LOG(Verbose, "Performing remove duplicate points test (perf warning)");
+				PCGSelfPruningElement::Execute(Context, EPCGSelfPruningType::RemoveDuplicates, 0.0f, false);
 
-			// Null out the output if this node is executed in isolation
-			if (Settings->ExecutionMode == EPCGSettingsExecutionMode::Isolated)
+				Context->InputData = ElementInputs;
+				Context->OutputData = ElementOutputs;
+			}
+
+			if (Settings->ExecutionMode == EPCGSettingsExecutionMode::Debug || Settings->ExecutionMode == EPCGSettingsExecutionMode::Isolated)
 			{
-				Context->OutputData.bCancelExecution = true;
+				FPCGDataCollection ElementInputs = Context->InputData;
+				FPCGDataCollection ElementOutputs = Context->OutputData;
+
+				Context->InputData = ElementOutputs;
+				Context->OutputData = FPCGDataCollection();
+
+				PCGDebugElement::ExecuteDebugDisplay(Context);
+
+				Context->InputData = ElementInputs;
+				Context->OutputData = ElementOutputs;
+
+				// Null out the output if this node is executed in isolation
+				if (Settings->ExecutionMode == EPCGSettingsExecutionMode::Isolated)
+				{
+					Context->OutputData.bCancelExecution = true;
+				}
 			}
 		}
 #endif
diff --git a/Engine/Plugins/Experimental/PCG/Source/PCG/Public/Elements/PCGDebugElement.h b/Engine/Plugins/Experimental/PCG/Source/PCG/Public/Elements/PCGDebugElement.h
--- a/Engine/Plugins/Experimental/PCG/Source/PCG/Public/Elements/PCGDebugElement.h
+++ b/Engine/Plugins/Experimental/PCG/Source/PCG/Public/Elements/PCGDebugElement.h
@@ -31,6 +31,6 @@ protected:
 class FPCGDebugElement : public FSimpleTypedPCGElement<UPCGDebugSettings>
 {
 protected:
-	virtual bool ExecuteInternal(FPCGContext* Contexxt) const override;
+	virtual bool ExecuteInternal(FPCGContext* Context) const override;
 	virtual bool IsCacheable(const UPCGSettings* InSettings) const override { return false; }
 };
diff --git a/Engine/Plugins/Experimental/PCG/Source/PCG/Public/Elements/PCGSelfPruning.h b/Engine/Plugins/Experimental/PCG/Source/PCG/Public/Elements/PCGSelfPruning.h
--- a/Engine/Plugins/Experimental/PCG/Source/PCG/Public/Elements/PCGSelfPruning.h
+++ b/Engine/Plugins/Experimental/PCG/Source/PCG/Public/Elements/PCGSelfPruning.h
@@ -15,9 +15,15 @@ enum class EPCGSelfPruningType : uint8
 	LargeToSmall,
 	SmallToLarge,
 	AllEqual,
-	None
+	None,
+	RemoveDuplicates
 };
 
+namespace PCGSelfPruningElement
+{
+	void Execute(FPCGContext* Context, EPCGSelfPruningType PruningType, float RadiusSimilarityFactor, bool bRandomizedPruning);
+}
+
 UCLASS(BlueprintType, ClassGroup=(Procedural))
 class PCG_API UPCGSelfPruningSettings : public UPCGSettings
 {
diff --git a/Engine/Plugins/Experimental/PCG/Source/PCG/Public/PCGDebug.h b/Engine/Plugins/Experimental/PCG/Source/PCG/Public/PCGDebug.h
--- a/Engine/Plugins/Experimental/PCG/Source/PCG/Public/PCGDebug.h
+++ b/Engine/Plugins/Experimental/PCG/Source/PCG/Public/PCGDebug.h
@@ -36,5 +36,9 @@ public:
 	UPROPERTY(BlueprintReadWrite, EditAnywhere, Category = Settings)
 	TSoftObjectPtr<UMaterialInterface> MaterialOverride;
 
+	/** Warning: enabling this flag will have severe performance impact */
+	UPROPERTY(BlueprintReadWrite, EditAnywhere, Category = Settings)
+	bool bCheckForDuplicates;
+
 	TSoftObjectPtr<UMaterialInterface> GetMaterial() const;
 };
