Added methods to get other parts of the PCG data from the context (used in very specific cases)

diff --git a/Engine/Plugins/Experimental/PCG/Source/PCG/Private/Helpers/PCGBlueprintHelpers.cpp b/Engine/Plugins/Experimental/PCG/Source/PCG/Private/Helpers/PCGBlueprintHelpers.cpp
--- a/Engine/Plugins/Experimental/PCG/Source/PCG/Private/Helpers/PCGBlueprintHelpers.cpp
+++ b/Engine/Plugins/Experimental/PCG/Source/PCG/Private/Helpers/PCGBlueprintHelpers.cpp
@@ -1,6 +1,7 @@
 // Copyright Epic Games, Inc. All Rights Reserved.
 
 #include "Helpers/PCGBlueprintHelpers.h"
+#include "PCGComponent.h"
 #include "PCGHelpers.h"
 #include "PCGSettings.h"
 
@@ -25,4 +26,19 @@ FRandomStream UPCGBlueprintHelpers::GetRandomStream(const FPCGPoint& InPoint, co
 	{
 		return FRandomStream(InPoint.Seed);
 	}
+}
+
+UPCGData* UPCGBlueprintHelpers::GetActorData(FPCGContext& Context)
+{
+	return Context.SourceComponent ? Context.SourceComponent->GetActorPCGData() : nullptr;
+}
+
+UPCGData* UPCGBlueprintHelpers::GetInputData(FPCGContext& Context)
+{
+	return Context.SourceComponent ? Context.SourceComponent->GetInputPCGData() : nullptr;
+}
+
+TArray<UPCGData*> UPCGBlueprintHelpers::GetExclusionData(FPCGContext& Context)
+{
+	return Context.SourceComponent ? Context.SourceComponent->GetPCGExclusionData() : TArray<UPCGData*>();
 }
\ No newline at end of file
diff --git a/Engine/Plugins/Experimental/PCG/Source/PCG/Public/Helpers/PCGBlueprintHelpers.h b/Engine/Plugins/Experimental/PCG/Source/PCG/Public/Helpers/PCGBlueprintHelpers.h
--- a/Engine/Plugins/Experimental/PCG/Source/PCG/Public/Helpers/PCGBlueprintHelpers.h
+++ b/Engine/Plugins/Experimental/PCG/Source/PCG/Public/Helpers/PCGBlueprintHelpers.h
@@ -7,10 +7,12 @@
 #include "Kismet/BlueprintFunctionLibrary.h"
 
 #include "PCGPoint.h"
+#include "PCGContext.h"
 
 #include "PCGBlueprintHelpers.generated.h"
 
 class UPCGSettings;
+class UPCGData;
 
 UCLASS()
 class PCG_API UPCGBlueprintHelpers : public UBlueprintFunctionLibrary
@@ -27,4 +29,13 @@ public:
 	/** Creates a random stream from a point's seed and settings's seed (optional) */
 	UFUNCTION(BlueprintCallable, Category = "PCG|Helpers", meta = (ScriptMethod))
 	static FRandomStream GetRandomStream(const FPCGPoint& InPoint, const UPCGSettings* OptionalSettings = nullptr);
+
+	UFUNCTION(BlueprintCallable, Category = "PCG|Temporary", meta = (ScriptMethod))
+	static UPCGData* GetActorData(UPARAM(ref) FPCGContext& Context);
+
+	UFUNCTION(BlueprintCallable, Category = "PCG|Temporary", meta = (ScriptMethod))
+	static UPCGData* GetInputData(UPARAM(ref) FPCGContext& Context);
+
+	UFUNCTION(BlueprintCallable, Category = "PCG|Temporary", meta = (ScriptMethod))
+	static TArray<UPCGData*> GetExclusionData(UPARAM(ref) FPCGContext& Context);
 };
\ No newline at end of file
