Exposed source component on PCG context

diff --git a/Engine/Plugins/Experimental/PCG/Source/PCG/Private/Helpers/PCGBlueprintHelpers.cpp b/Engine/Plugins/Experimental/PCG/Source/PCG/Private/Helpers/PCGBlueprintHelpers.cpp
--- a/Engine/Plugins/Experimental/PCG/Source/PCG/Private/Helpers/PCGBlueprintHelpers.cpp
+++ b/Engine/Plugins/Experimental/PCG/Source/PCG/Private/Helpers/PCGBlueprintHelpers.cpp
@@ -4,6 +4,7 @@
 #include "PCGComponent.h"
 #include "PCGHelpers.h"
 #include "PCGSettings.h"
+#include "Grid/PCGPartitionActor.h"
 
 int UPCGBlueprintHelpers::ComputeSeedFromPosition(const FVector& InPosition)
 {
@@ -41,4 +42,23 @@ UPCGData* UPCGBlueprintHelpers::GetInputData(FPCGContext& Context)
 TArray<UPCGData*> UPCGBlueprintHelpers::GetExclusionData(FPCGContext& Context)
 {
 	return Context.SourceComponent ? Context.SourceComponent->GetPCGExclusionData() : TArray<UPCGData*>();
+}
+
+UPCGComponent* UPCGBlueprintHelpers::GetComponent(FPCGContext& Context)
+{
+	return Context.SourceComponent;
+}
+
+UPCGComponent* UPCGBlueprintHelpers::GetOriginalComponent(FPCGContext& Context)
+{
+	if (Context.SourceComponent &&
+		Cast<APCGPartitionActor>(Context.SourceComponent->GetOwner()) &&
+		Cast<APCGPartitionActor>(Context.SourceComponent->GetOwner())->GetOriginalComponent(Context.SourceComponent))
+	{
+		return Cast<APCGPartitionActor>(Context.SourceComponent->GetOwner())->GetOriginalComponent(Context.SourceComponent);
+	}
+	else
+	{
+		return Context.SourceComponent;
+	}
 }
\ No newline at end of file
diff --git a/Engine/Plugins/Experimental/PCG/Source/PCG/Public/Helpers/PCGBlueprintHelpers.h b/Engine/Plugins/Experimental/PCG/Source/PCG/Public/Helpers/PCGBlueprintHelpers.h
--- a/Engine/Plugins/Experimental/PCG/Source/PCG/Public/Helpers/PCGBlueprintHelpers.h
+++ b/Engine/Plugins/Experimental/PCG/Source/PCG/Public/Helpers/PCGBlueprintHelpers.h
@@ -38,4 +38,10 @@ public:
 
 	UFUNCTION(BlueprintCallable, Category = "PCG|Temporary", meta = (ScriptMethod))
 	static TArray<UPCGData*> GetExclusionData(UPARAM(ref) FPCGContext& Context);
+
+	UFUNCTION(BlueprintCallable, Category = "PCG|Temporary", meta = (ScriptMethod))
+	static UPCGComponent* GetComponent(UPARAM(ref) FPCGContext& Context);
+
+	UFUNCTION(BlueprintCallable, Category = "PCG|Temporary", meta = (ScriptMethod))
+	static UPCGComponent* GetOriginalComponent(UPARAM(ref) FPCGContext& Context);
 };
\ No newline at end of file
