PCG TextureData: Added option to ignore density when transforming points.

diff --git a/Engine/Plugins/Experimental/PCG/Source/PCG/Private/Data/PCGTextureData.cpp b/Engine/Plugins/Experimental/PCG/Source/PCG/Private/Data/PCGTextureData.cpp
--- a/Engine/Plugins/Experimental/PCG/Source/PCG/Private/Data/PCGTextureData.cpp
+++ b/Engine/Plugins/Experimental/PCG/Source/PCG/Private/Data/PCGTextureData.cpp
@@ -6,6 +6,21 @@
 #include "Data/PCGPointData.h"
 #include "Helpers/PCGAsync.h"
 
+namespace PCGTextureDataMaths
+{
+	float ComputeDensity(float InDensityA, float InDensityB, EPCGTextureDensityFunction InDensityFunction)
+	{
+		if (InDensityFunction == EPCGTextureDensityFunction::Multiply)
+		{
+			return InDensityA * InDensityB;
+		}
+		else // default: Ignore
+		{
+			return InDensityA;
+		}
+	}
+}
+
 namespace PCGTextureSampling
 {
 	template<typename ValueType>
@@ -102,7 +117,7 @@ FPCGPoint UPCGBaseTextureData::TransformPoint(const FPCGPoint& InPoint) const
 	FLinearColor Color = PCGTextureSampling::Sample<FLinearColor>(Position2D, Surface, Width, Height, [this](int32 Index) { return ColorData[Index]; });
 
 	Point.Color *= Color;
-	Point.Density *= PCGTextureSampling::SampleFloatChannel(Color, ColorChannel);
+	Point.Density = PCGTextureDataMaths::ComputeDensity(Point.Density, PCGTextureSampling::SampleFloatChannel(Color, ColorChannel), DensityFunction);
 
 	return Point;
 }
diff --git a/Engine/Plugins/Experimental/PCG/Source/PCG/Public/Data/PCGTextureData.h b/Engine/Plugins/Experimental/PCG/Source/PCG/Public/Data/PCGTextureData.h
--- a/Engine/Plugins/Experimental/PCG/Source/PCG/Public/Data/PCGTextureData.h
+++ b/Engine/Plugins/Experimental/PCG/Source/PCG/Public/Data/PCGTextureData.h
@@ -16,6 +16,14 @@ enum class EPCGTextureColorChannel : uint8
 	Alpha
 };
 
+
+UENUM(BlueprintType)
+enum class EPCGTextureDensityFunction : uint8
+{
+	Ignore,
+	Multiply
+};
+
 UCLASS(Abstract)
 class PCG_API UPCGBaseTextureData : public UPCGSurfaceData
 {
@@ -33,6 +41,13 @@ public:
 	virtual const UPCGPointData* CreatePointData(FPCGContext* Context) const override;
 	//~End UPCGSpatialDataWithPointCache interface
 
+public:
+	UPROPERTY(BlueprintReadWrite, VisibleAnywhere, Category = SpatialData)
+	EPCGTextureDensityFunction DensityFunction = EPCGTextureDensityFunction::Multiply; 
+
+	UPROPERTY(BlueprintReadWrite, EditAnywhere, Category = Settings)
+	EPCGTextureColorChannel ColorChannel = EPCGTextureColorChannel::Alpha;
+
 protected:
 	UPROPERTY()
 	TArray<FLinearColor> ColorData;
@@ -40,9 +55,6 @@ protected:
 	UPROPERTY(BlueprintReadOnly, VisibleAnywhere, Category = SpatialData)
 	FBox Bounds = FBox(EForceInit::ForceInit);
 
-	UPROPERTY(BlueprintReadWrite, EditAnywhere, Category = Settings)
-	EPCGTextureColorChannel ColorChannel = EPCGTextureColorChannel::Alpha;
-
 	UPROPERTY(BlueprintReadOnly, VisibleAnywhere, Category = SpatialData)
 	int32 Height = 0;
 
