Added metadata support (basic). Created by default on point data, but not in other cases.
Added the context in point loops (needed for some specific use cases like the grid generator) and support for a custom object.
Refactored CreatePointData implementations to go through similar code to the TransformPoint for metadata reasons

diff --git a/Engine/Plugins/Experimental/PCG/Source/PCG/Private/Data/PCGDifferenceData.cpp b/Engine/Plugins/Experimental/PCG/Source/PCG/Private/Data/PCGDifferenceData.cpp
--- a/Engine/Plugins/Experimental/PCG/Source/PCG/Private/Data/PCGDifferenceData.cpp
+++ b/Engine/Plugins/Experimental/PCG/Source/PCG/Private/Data/PCGDifferenceData.cpp
@@ -5,6 +5,7 @@
 #include "Data/PCGUnionData.h"
 #include "Helpers/PCGAsync.h"
 #include "PCGHelpers.h"
+#include "Metadata/PCGMetadataAccessor.h"
 
 namespace PCGDifferenceDataUtils
 {
@@ -149,6 +150,39 @@ FPCGPoint UPCGDifferenceData::TransformPoint(const FPCGPoint& InPoint) const
 	return TransformedPoint;
 }
 
+bool UPCGDifferenceData::GetPointAtPosition(const FVector& InPosition, FPCGPoint& OutPoint, UPCGMetadata* OutMetadata) const
+{
+	check(Source);
+
+	FPCGPoint PointFromSource;
+	if (!Source->GetPointAtPosition(InPosition, PointFromSource, OutMetadata))
+	{
+		return false;
+	}
+
+	OutPoint = PointFromSource;
+
+	FPCGPoint PointFromDiff;
+	if (Difference && Difference->GetPointAtPosition(PointFromSource.Transform.GetLocation(), PointFromDiff, OutMetadata))
+	{
+		const bool bBinaryDensity = (DensityFunction == EPCGDifferenceDensityFunction::Binary);
+		
+		// Apply difference
+		OutPoint.Density = bBinaryDensity ? 0 : FMath::Max(0, PointFromSource.Density - PointFromDiff.Density);
+		// Color?
+		if (OutMetadata && OutPoint.Density > 0)
+		{
+			OutMetadata->MergeAttributes(PointFromSource, PointFromDiff, OutPoint, EPCGMetadataOp::Sub);
+		}
+
+		return OutPoint.Density > 0;
+	}
+	else
+	{
+		return true;
+	}
+}
+
 bool UPCGDifferenceData::HasNonTrivialTransform() const
 {
 	check(Source);
@@ -175,33 +209,40 @@ const UPCGPointData* UPCGDifferenceData::CreatePointData(FPCGContext* Context) c
 	}
 
 	UPCGPointData* Data = NewObject<UPCGPointData>(const_cast<UPCGDifferenceData*>(this));
-	Data->TargetActor = TargetActor;
+	Data->InitializeFromData(this, SourcePointData->Metadata);
 
 	const TArray<FPCGPoint>& SourcePoints = SourcePointData->GetPoints();
 	TArray<FPCGPoint>& TargetPoints = Data->GetMutablePoints();
 
-	FPCGAsync::AsyncPointProcessing(Context, SourcePoints.Num(), TargetPoints, [this, &SourcePoints](int32 Index, FPCGPoint& OutPoint)
+	FPCGAsync::AsyncPointProcessing(Context, SourcePoints.Num(), TargetPoints, [this, Data, SourcePointData, &SourcePoints](int32 Index, FPCGPoint& OutPoint)
 	{
 		const FPCGPoint& Point = SourcePoints[Index];
-		const float DensityInDifference = GetDensityAtPositionFromDifference(Point.Transform.GetLocation());
 
-		if (DensityInDifference < Point.Density)
+		FPCGPoint PointFromDiff;
+		if (Difference && Difference->GetPointAtPosition(Point.Transform.GetLocation(), PointFromDiff, Data->Metadata))
 		{
+			const bool bBinaryDensity = (DensityFunction == EPCGDifferenceDensityFunction::Binary);
+
 			OutPoint = Point;
-			OutPoint.Density -= DensityInDifference;
-			return true;
-		}
+			UPCGMetadataAccessorHelpers::InitializeMetadata(OutPoint, Data->Metadata, Point);
+			OutPoint.Density = bBinaryDensity ? 0 : FMath::Max(0, Point.Density - PointFromDiff.Density);
+
+			if (Data->Metadata && OutPoint.Density > 0)
+			{
+				Data->Metadata->MergeAttributes(Point, SourcePointData->Metadata, PointFromDiff, Data->Metadata, OutPoint, EPCGMetadataOp::Sub);
+			}
+
 #if WITH_EDITORONLY_DATA
-		else if (bKeepZeroDensityPoints)
-		{
-			OutPoint = Point;
-			OutPoint.Density = 0;
-			return true;
-		}
+			return OutPoint.Density > 0 || bKeepZeroDensityPoints;
+#else
+			return OutPoint.Density > 0;
 #endif
+		}
 		else
 		{
-			return false;
+			OutPoint = Point;
+			UPCGMetadataAccessorHelpers::InitializeMetadata(OutPoint, Data->Metadata, Point);
+			return true;
 		}
 	});
 
diff --git a/Engine/Plugins/Experimental/PCG/Source/PCG/Private/Data/PCGIntersectionData.cpp b/Engine/Plugins/Experimental/PCG/Source/PCG/Private/Data/PCGIntersectionData.cpp
--- a/Engine/Plugins/Experimental/PCG/Source/PCG/Private/Data/PCGIntersectionData.cpp
+++ b/Engine/Plugins/Experimental/PCG/Source/PCG/Private/Data/PCGIntersectionData.cpp
@@ -4,6 +4,7 @@
 #include "Data/PCGPointData.h"
 #include "Helpers/PCGAsync.h"
 #include "PCGHelpers.h"
+#include "Metadata/PCGMetadataAccessor.h"
 
 namespace PCGIntersectionDataMaths
 {
@@ -93,6 +94,37 @@ FPCGPoint UPCGIntersectionData::TransformPoint(const FPCGPoint& InPoint) const
 	return TransformedPoint;
 }
 
+bool UPCGIntersectionData::GetPointAtPosition(const FVector& InPosition, FPCGPoint& OutPoint, UPCGMetadata* OutMetadata) const
+{
+	check(A && B);
+	const UPCGSpatialData* X = A->HasNonTrivialTransform() ? A : B;
+	const UPCGSpatialData* Y = X == A ? B : A;
+
+	FPCGPoint PointFromX;
+	if (!X->GetPointAtPosition(InPosition, PointFromX, OutMetadata))
+	{
+		return false;
+	}
+
+	FPCGPoint PointFromY;
+	if (!Y->GetPointAtPosition(PointFromX.Transform.GetLocation(), PointFromY, OutMetadata))
+	{
+		return false;
+	}
+
+	// Merge points into a single point
+	OutPoint = PointFromX;
+	OutPoint.Density = PCGIntersectionDataMaths::ComputeDensity(PointFromX.Density, PointFromY.Density, DensityFunction);
+	OutPoint.Color = PointFromX.Color * PointFromY.Color;
+
+	if (OutMetadata)
+	{
+		OutMetadata->MergeAttributes(PointFromX, PointFromY, OutPoint, EPCGMetadataOp::Min);
+	}
+
+	return true;
+}
+
 bool UPCGIntersectionData::HasNonTrivialTransform() const
 {
 	check(A && B);
@@ -132,28 +164,36 @@ UPCGPointData* UPCGIntersectionData::CreateAndFilterPointData(FPCGContext* Conte
 	const TArray<FPCGPoint>& SourcePoints = SourcePointData->GetPoints();
 
 	UPCGPointData* Data = NewObject<UPCGPointData>(const_cast<UPCGIntersectionData*>(this));
-	Data->TargetActor = TargetActor;
+	Data->InitializeFromData(this, SourcePointData->Metadata);
+	Data->Metadata->AddAttributes(Y->Metadata);
+
 	TArray<FPCGPoint>& TargetPoints = Data->GetMutablePoints();
 
-	FPCGAsync::AsyncPointProcessing(Context, SourcePoints.Num(), TargetPoints, [this, &SourcePoints, Y](int32 Index, FPCGPoint& OutPoint)
+	FPCGAsync::AsyncPointProcessing(Context, SourcePoints.Num(), TargetPoints, [this, Data, SourcePointData, &SourcePoints, Y](int32 Index, FPCGPoint& OutPoint)
 	{
 		const FPCGPoint& Point = SourcePoints[Index];
-		const float YDensity = Y->GetDensityAtPosition(Point.Transform.GetLocation());
 
+		FPCGPoint PointFromY;
 #if WITH_EDITORONLY_DATA
-		if (YDensity > 0 || bKeepZeroDensityPoints)
+		if (!Y->GetPointAtPosition(Point.Transform.GetLocation(), PointFromY, Data->Metadata) && !bKeepZeroDensityPoints)
 #else
-		if (YDensity > 0)
+		if (!Y->GetPointAtPosition(Point.Transform.GetLocation(), PointFromY, Data->Metadata))
 #endif
 		{
-			OutPoint = Point;
-			OutPoint.Density = PCGIntersectionDataMaths::ComputeDensity(Point.Density, YDensity, DensityFunction);
-			return true;
+			return false;
 		}
-		else
+
+		OutPoint = Point;
+		UPCGMetadataAccessorHelpers::InitializeMetadata(OutPoint, Data->Metadata, Point);
+		OutPoint.Density = PCGIntersectionDataMaths::ComputeDensity(Point.Density, PointFromY.Density, DensityFunction);
+		OutPoint.Color = Point.Color * PointFromY.Color;
+
+		if (Data->Metadata)
 		{
-			return false;
+			Data->Metadata->MergeAttributes(Point, SourcePointData->Metadata, PointFromY, Data->Metadata, OutPoint, EPCGMetadataOp::Min);
 		}
+
+		return true;
 	});
 
 	UE_LOG(LogPCG, Verbose, TEXT("Intersection generated %d points from %d source points"), TargetPoints.Num(), SourcePoints.Num());
diff --git a/Engine/Plugins/Experimental/PCG/Source/PCG/Private/Data/PCGLandscapeData.cpp b/Engine/Plugins/Experimental/PCG/Source/PCG/Private/Data/PCGLandscapeData.cpp
--- a/Engine/Plugins/Experimental/PCG/Source/PCG/Private/Data/PCGLandscapeData.cpp
+++ b/Engine/Plugins/Experimental/PCG/Source/PCG/Private/Data/PCGLandscapeData.cpp
@@ -82,7 +82,7 @@ const UPCGPointData* UPCGLandscapeData::CreatePointData(FPCGContext* Context) co
 	TRACE_CPUPROFILER_EVENT_SCOPE(UPCGLandscapeData::CreatePointData);
 
 	UPCGPointData* Data = NewObject<UPCGPointData>(const_cast<UPCGLandscapeData*>(this));
-	Data->TargetActor = TargetActor;
+	Data->InitializeFromData(this);
 	TArray<FPCGPoint>& Points = Data->GetMutablePoints();
 
 	// TODO: add offset to nearest edge, will have an impact if the grid size doesn't match the landscape size
diff --git a/Engine/Plugins/Experimental/PCG/Source/PCG/Private/Data/PCGLandscapeSplineData.cpp b/Engine/Plugins/Experimental/PCG/Source/PCG/Private/Data/PCGLandscapeSplineData.cpp
--- a/Engine/Plugins/Experimental/PCG/Source/PCG/Private/Data/PCGLandscapeSplineData.cpp
+++ b/Engine/Plugins/Experimental/PCG/Source/PCG/Private/Data/PCGLandscapeSplineData.cpp
@@ -94,7 +94,7 @@ const UPCGPointData* UPCGLandscapeSplineData::CreatePointData(FPCGContext* Conte
 	TRACE_CPUPROFILER_EVENT_SCOPE(UPCGLandscapeSplineData::CreatePointData);
 
 	UPCGPointData* Data = NewObject<UPCGPointData>(const_cast<UPCGLandscapeSplineData*>(this));
-	Data->TargetActor = TargetActor;
+	Data->InitializeFromData(this);
 	TArray<FPCGPoint>& Points = Data->GetMutablePoints();
 
 	// TODO: replace all the logic with sampling settings; currently it uses the landscape scaling as a basis
diff --git a/Engine/Plugins/Experimental/PCG/Source/PCG/Private/Data/PCGPointData.cpp b/Engine/Plugins/Experimental/PCG/Source/PCG/Private/Data/PCGPointData.cpp
--- a/Engine/Plugins/Experimental/PCG/Source/PCG/Private/Data/PCGPointData.cpp
+++ b/Engine/Plugins/Experimental/PCG/Source/PCG/Private/Data/PCGPointData.cpp
@@ -3,6 +3,7 @@
 #include "Data/PCGPointData.h"
 
 #include "PCGHelpers.h"
+#include "Metadata/PCGMetadataAccessor.h"
 
 #include "GameFramework/Actor.h"
 #include "Misc/ScopeLock.h"
@@ -63,6 +64,17 @@ namespace PCGPointHelpers
 			return 0;
 		}
 	}
+
+	/** Helper function for additive blending of quaternions (copied from ControlRig) */
+	FQuat AddQuatWithWeight(const FQuat& Q, const FQuat& V, float Weight)
+	{
+		FQuat BlendQuat = V * Weight;
+
+		if ((Q | BlendQuat) >= 0.0f)
+			return Q + BlendQuat;
+		else
+			return Q - BlendQuat;
+	}	
 }
 
 FPCGPointRef::FPCGPointRef(const FPCGPoint& InPoint)
@@ -145,7 +157,7 @@ void UPCGPointData::SetPoints(const TArray<FPCGPoint>& InPoints)
 	GetMutablePoints() = InPoints;
 }
 
-void UPCGPointData::Initialize(AActor* InActor)
+void UPCGPointData::InitializeFromActor(AActor* InActor)
 {
 	check(InActor);
 
@@ -156,6 +168,16 @@ void UPCGPointData::Initialize(AActor* InActor)
 	Points[0].Seed = PCGHelpers::ComputeSeed((int)Position.X, (int)Position.Y, (int)Position.Z);
 
 	TargetActor = InActor;
+	Metadata = NewObject<UPCGMetadata>(this);
+}
+
+void UPCGPointData::InitializeFromData(const UPCGSpatialData* InSource, const UPCGMetadata* InMetadataParentOverride)
+{
+	check(InSource);
+	TargetActor = InSource->TargetActor;
+
+	Metadata = NewObject<UPCGMetadata>(this);
+	Metadata->Initialize(InMetadataParentOverride ? InMetadataParentOverride : InSource->Metadata);
 }
 
 const FPCGPoint* UPCGPointData::GetPointAtPosition(const FVector& InPosition) const
@@ -231,13 +253,95 @@ FPCGPoint UPCGPointData::TransformPoint(const FPCGPoint& InPoint) const
 	Point.Transform.NormalizeRotation();
 	Point.Transform.SetScale3D(Point.Transform.GetScale3D() * WeightedScale);	
 	Point.Density *= WeightedDensity;
-	Point.Extents *= WeightedExtents;
+	Point.Extents *= WeightedExtents; // this assumes that the extents were 1 to begin with
 	Point.Color *= WeightedColor;
 	Point.Steepness *= WeightedSteepness;
 
 	return Point;
 }
 
+bool UPCGPointData::GetPointAtPosition(const FVector& InPosition, FPCGPoint& OutPoint, UPCGMetadata* OutMetadata) const
+{
+	if (bOctreeIsDirty)
+	{
+		RebuildOctree();
+	}
+
+	TArray<TPair<const FPCGPoint*, float>> Contributions;
+	Octree.FindElementsWithBoundsTest(FBoxCenterAndExtent(InPosition, FVector::Zero()), [&InPosition, &Contributions](const FPCGPointRef& InPointRef) {
+		Contributions.Emplace(InPointRef.Point, PCGPointHelpers::InverseEuclidianDistance(*InPointRef.Point, InPosition));
+	});
+
+	float SumContributions = 0;
+	float MaxContribution = 0;
+	const FPCGPoint* MaxContributor = nullptr;
+
+	for (const TPair<const FPCGPoint*, float>& Contribution : Contributions)
+	{
+		SumContributions += Contribution.Value;
+
+		if (Contribution.Value > MaxContribution)
+		{
+			MaxContribution = Contribution.Value;
+			MaxContributor = Contribution.Key;
+		}
+	}
+
+	if (SumContributions <= 0)
+	{
+		return false;
+	}
+
+	// Computed weighted average of spatial properties
+	FQuat WeightedQuat = FQuat::Identity;
+	FVector WeightedScale = FVector::ZeroVector;
+	float WeightedDensity = 0;
+	FVector WeightedExtents = FVector::ZeroVector;
+	FVector WeightedColor = FVector::ZeroVector;
+	float WeightedSteepness = 0;
+
+	for (const TPair<const FPCGPoint*, float> Contribution : Contributions)
+	{
+		const FPCGPoint& SourcePoint = *Contribution.Key;
+		const float Weight = Contribution.Value / SumContributions;
+
+		WeightedQuat = PCGPointHelpers::AddQuatWithWeight(WeightedQuat, SourcePoint.Transform.GetRotation(), Weight);
+		WeightedScale += SourcePoint.Transform.GetScale3D() * Weight;
+		WeightedDensity += PCGPointHelpers::ManhattanDensity(SourcePoint, InPosition);
+		WeightedExtents += SourcePoint.Extents * Weight;
+		WeightedColor += SourcePoint.Color * Weight;
+		WeightedSteepness += SourcePoint.Steepness * Weight;
+	}
+
+	// Finally, apply changes to point
+	WeightedQuat.Normalize();
+
+	OutPoint.Transform.SetRotation(WeightedQuat);
+	OutPoint.Transform.SetScale3D(WeightedScale);
+	OutPoint.Transform.SetLocation(InPosition);
+	OutPoint.Density = WeightedDensity;
+	OutPoint.Extents = WeightedExtents;
+	OutPoint.Color = WeightedColor;
+	OutPoint.Steepness = WeightedSteepness;
+
+	if (OutMetadata)
+	{
+		UPCGMetadataAccessorHelpers::InitializeMetadata(OutPoint, OutMetadata, *MaxContributor, Metadata);
+		OutMetadata->ResetWeightedAttributes(OutPoint);
+
+		for (const TPair<const FPCGPoint*, float> Contribution : Contributions)
+		{
+			const FPCGPoint& SourcePoint = *Contribution.Key;
+			const float Weight = Contribution.Value / SumContributions;
+			const bool bIsMaxContributor = (Contribution.Key == MaxContributor);
+
+			OutMetadata->AccumulateWeightedAttributes(SourcePoint, Metadata, Weight, bIsMaxContributor, OutPoint);
+		}
+	}
+
+	return true;
+}
+
 float UPCGPointData::GetDensityAtPosition(const FVector& InPosition) const
 {
 	if (bOctreeIsDirty)
diff --git a/Engine/Plugins/Experimental/PCG/Source/PCG/Private/Data/PCGProjectionData.cpp b/Engine/Plugins/Experimental/PCG/Source/PCG/Private/Data/PCGProjectionData.cpp
--- a/Engine/Plugins/Experimental/PCG/Source/PCG/Private/Data/PCGProjectionData.cpp
+++ b/Engine/Plugins/Experimental/PCG/Source/PCG/Private/Data/PCGProjectionData.cpp
@@ -2,6 +2,7 @@
 
 #include "Data/PCGProjectionData.h"
 #include "Helpers/PCGAsync.h"
+#include "Metadata/PCGMetadataAccessor.h"
 
 void UPCGProjectionData::Initialize(const UPCGSpatialData* InSource, const UPCGSpatialData* InTarget)
 {
@@ -88,6 +89,35 @@ FPCGPoint UPCGProjectionData::TransformPoint(const FPCGPoint& InPoint) const
 	return Target->TransformPoint(Source->TransformPoint(InPoint));
 }
 
+bool UPCGProjectionData::GetPointAtPosition(const FVector& InPosition, FPCGPoint& OutPoint, UPCGMetadata* OutMetadata) const
+{
+	FPCGPoint PointFromSource;
+	if (!Source->GetPointAtPosition(InPosition, PointFromSource, OutMetadata))
+	{
+		return false;
+	}
+
+	FPCGPoint PointFromTarget;
+	if (!Target->GetPointAtPosition(PointFromSource.Transform.GetLocation(), PointFromTarget, OutMetadata))
+	{
+		return false;
+	}
+
+	// Merge points into a single point
+	OutPoint = PointFromSource;
+	OutPoint.Transform = PointFromTarget.Transform;
+	OutPoint.Density *= PointFromTarget.Density;
+	OutPoint.Color *= PointFromTarget.Color;
+	
+	if (OutMetadata)
+	{
+		//METADATA TODO Review op
+		OutMetadata->MergeAttributes(PointFromSource, PointFromTarget, OutPoint, EPCGMetadataOp::Max);
+	}
+
+	return true;
+}
+
 bool UPCGProjectionData::HasNonTrivialTransform() const
 {
 	return Target->HasNonTrivialTransform();
@@ -102,19 +132,38 @@ const UPCGPointData* UPCGProjectionData::CreatePointData(FPCGContext* Context) c
 	const TArray<FPCGPoint>& SourcePoints = SourcePointData->GetPoints();
 
 	UPCGPointData* PointData = NewObject<UPCGPointData>(const_cast<UPCGProjectionData*>(this));
-	PointData->TargetActor = TargetActor;
+	PointData->InitializeFromData(this, SourcePointData->Metadata);
+
 	TArray<FPCGPoint>& Points = PointData->GetMutablePoints();
 
-	FPCGAsync::AsyncPointProcessing(Context, SourcePoints.Num(), Points, [this, &SourcePoints](int32 Index, FPCGPoint& OutPoint)
+	FPCGAsync::AsyncPointProcessing(Context, SourcePoints.Num(), Points, [this, SourcePointData, PointData, &SourcePoints](int32 Index, FPCGPoint& OutPoint)
 	{
 		const FPCGPoint& SourcePoint = SourcePoints[Index];
-		OutPoint = Target->TransformPoint(SourcePoint);
 
+		FPCGPoint PointFromTarget;
 #if WITH_EDITORONLY_DATA
-		return OutPoint.Density > 0 || bKeepZeroDensityPoints;
+		if (!Target->GetPointAtPosition(SourcePoint.Transform.GetLocation(), PointFromTarget, PointData->Metadata) && !bKeepZeroDensityPoints)
 #else
-		return OutPoint.Density > 0;
+		if (!Target->GetPointAtPosition(SourcePoint.Transform.GetLocation(), PointFromTarget, PointData->Metadata))
 #endif
+		{
+			return false;
+		}
+
+		// Merge points into a single point
+		OutPoint = SourcePoint;
+		UPCGMetadataAccessorHelpers::InitializeMetadata(OutPoint, PointData->Metadata, SourcePoint);
+		OutPoint.Transform = PointFromTarget.Transform;
+		OutPoint.Density *= PointFromTarget.Density;
+		OutPoint.Color *= PointFromTarget.Color;
+
+		if (PointData->Metadata)
+		{
+			//METADATA TODO review op
+			PointData->Metadata->MergeAttributes(SourcePoint, SourcePointData->Metadata, PointFromTarget, PointData->Metadata, OutPoint, EPCGMetadataOp::Max);
+		}
+
+		return true;
 	});
 
 	UE_LOG(LogPCG, Verbose, TEXT("Projection generated %d points from %d source points"), Points.Num(), SourcePoints.Num());
diff --git a/Engine/Plugins/Experimental/PCG/Source/PCG/Private/Data/PCGSplineData.cpp b/Engine/Plugins/Experimental/PCG/Source/PCG/Private/Data/PCGSplineData.cpp
--- a/Engine/Plugins/Experimental/PCG/Source/PCG/Private/Data/PCGSplineData.cpp
+++ b/Engine/Plugins/Experimental/PCG/Source/PCG/Private/Data/PCGSplineData.cpp
@@ -47,7 +47,7 @@ const UPCGPointData* UPCGSplineData::CreatePointData(FPCGContext* Context) const
 {
 	TRACE_CPUPROFILER_EVENT_SCOPE(UPCGSplineData::CreatePointData);
 	UPCGPointData* Data = NewObject<UPCGPointData>(const_cast<UPCGSplineData*>(this));
-	Data->TargetActor = TargetActor;
+	Data->InitializeFromData(this);
 
 	// TODO: introduce settings to sample, 
 	// should be part of the settings passed in param.
diff --git a/Engine/Plugins/Experimental/PCG/Source/PCG/Private/Data/PCGTextureData.cpp b/Engine/Plugins/Experimental/PCG/Source/PCG/Private/Data/PCGTextureData.cpp
--- a/Engine/Plugins/Experimental/PCG/Source/PCG/Private/Data/PCGTextureData.cpp
+++ b/Engine/Plugins/Experimental/PCG/Source/PCG/Private/Data/PCGTextureData.cpp
@@ -131,7 +131,7 @@ const UPCGPointData* UPCGBaseTextureData::CreatePointData(FPCGContext* Context)
 	FBox2D LocalSurfaceBounds(FVector2D(-1.0f, -1.0f), FVector2D(1.0f, 1.0f));
 
 	UPCGPointData* Data = NewObject<UPCGPointData>(const_cast<UPCGBaseTextureData*>(this));
-	Data->TargetActor = TargetActor;
+	Data->InitializeFromData(this);
 	TArray<FPCGPoint>& Points = Data->GetMutablePoints();
 
 	// TODO: There's a bias issue here where we should correct by a 0.5 unit...
diff --git a/Engine/Plugins/Experimental/PCG/Source/PCG/Private/Data/PCGUnionData.cpp b/Engine/Plugins/Experimental/PCG/Source/PCG/Private/Data/PCGUnionData.cpp
--- a/Engine/Plugins/Experimental/PCG/Source/PCG/Private/Data/PCGUnionData.cpp
+++ b/Engine/Plugins/Experimental/PCG/Source/PCG/Private/Data/PCGUnionData.cpp
@@ -154,6 +154,59 @@ FPCGPoint UPCGUnionData::TransformPoint(const FPCGPoint& InPoint) const
 	}
 }
 
+bool UPCGUnionData::GetPointAtPosition(const FVector& InPosition, FPCGPoint& OutPoint, UPCGMetadata* OutMetadata) const
+{
+	FVector PointPosition = InPosition;
+
+	if (FirstNonTrivialTransformData)
+	{
+		if (FirstNonTrivialTransformData->GetPointAtPosition(InPosition, OutPoint, OutMetadata))
+		{
+			PointPosition = OutPoint.Transform.GetLocation();
+
+			if (DensityFunction == EPCGUnionDensityFunction::Binary && OutPoint.Density > 0)
+			{
+				OutPoint.Density = 1.0f;
+			}
+		}
+	}
+
+	const bool bSkipLoop = (!OutMetadata && OutPoint.Density >= 1.0f);
+	const int32 DataCount = Data.Num();
+	for (int32 DataIndex = 0; DataIndex < DataCount && !bSkipLoop; ++DataIndex)
+	{
+		if (Data[DataIndex] == FirstNonTrivialTransformData)
+		{
+			continue;
+		}
+
+		FPCGPoint PointInData;
+		if (Data[DataIndex]->GetPointAtPosition(PointPosition, PointInData, OutMetadata))
+		{
+			// Update density
+			PCGUnionDataMaths::UpdateDensity(OutPoint.Density, PointInData.Density, DensityFunction);
+
+			OutPoint.Color = FVector4(
+				FMath::Max(OutPoint.Color.X, PointInData.Color.X),
+				FMath::Max(OutPoint.Color.Y, PointInData.Color.Y),
+				FMath::Max(OutPoint.Color.Z, PointInData.Color.Z),
+				FMath::Max(OutPoint.Color.W, PointInData.Color.W));
+
+			// Merge properties into OutPoint
+			if (OutMetadata)
+			{
+				OutMetadata->MergeAttributes(OutPoint, PointInData, OutPoint, EPCGMetadataOp::Max);
+			}
+			else if (OutPoint.Density >= 1.0f)
+			{
+				break;
+			}
+		}
+	}
+
+	return OutPoint.Density > 0;
+}
+
 bool UPCGUnionData::HasNonTrivialTransform() const
 {
 	return (FirstNonTrivialTransformData != nullptr || Super::HasNonTrivialTransform());
@@ -178,7 +231,13 @@ const UPCGPointData* UPCGUnionData::CreatePointData(FPCGContext* Context) const
 	}
 
 	UPCGPointData* PointData = NewObject<UPCGPointData>(const_cast<UPCGUnionData*>(this));
-	PointData->TargetActor = TargetActor;
+	PointData->InitializeFromData(this, Data[0]->Metadata);
+
+	// Initialize metadata
+	for (TObjectPtr<const UPCGSpatialData> Datum : Data)
+	{
+		PointData->Metadata->AddAttributes(Datum->Metadata);
+	}
 
 	switch (UnionType)
 	{
@@ -196,7 +255,14 @@ const UPCGPointData* UPCGUnionData::CreatePointData(FPCGContext* Context) const
 			TArray<FPCGPoint>& TargetPoints = PointData->GetMutablePoints();
 			for (TObjectPtr<const UPCGSpatialData> Datum : Data)
 			{
-				TargetPoints.Append(Datum->ToPointData(Context)->GetPoints());
+				const UPCGPointData* DatumPointData = Datum->ToPointData(Context);
+				int32 TargetPointIndex = TargetPoints.Num();
+				TargetPoints.Append(DatumPointData->GetPoints());
+
+				if (PointData->Metadata && DatumPointData->GetPoints().Num() > 0)
+				{
+					PointData->Metadata->SetAttributes(MakeArrayView(DatumPointData->GetPoints()), DatumPointData->Metadata, MakeArrayView(&TargetPoints[TargetPointIndex], DatumPointData->GetPoints().Num()));
+				}
 			}
 
 			// Correct density for binary-style union
@@ -235,7 +301,7 @@ void UPCGUnionData::CreateSequentialPointData(FPCGContext* Context, UPCGPointDat
 		// add it & compute its final density
 		const TArray<FPCGPoint>& Points = Data[DataIndex]->ToPointData(Context)->GetPoints();
 
-		FPCGAsync::AsyncPointProcessing(Context, Points.Num(), SelectedDataPoints, [this, &Points, DataIndex, FirstDataIndex, LastDataIndex, DataIndexIncrement](int32 Index, FPCGPoint& OutPoint)
+		FPCGAsync::AsyncPointProcessing(Context, Points.Num(), SelectedDataPoints, [this, PointData, &Points, DataIndex, FirstDataIndex, LastDataIndex, DataIndexIncrement](int32 Index, FPCGPoint& OutPoint)
 		{
 			const FPCGPoint& Point = Points[Index];
 
@@ -255,17 +321,51 @@ void UPCGUnionData::CreateSequentialPointData(FPCGContext* Context, UPCGPointDat
 				return false;
 			}
 
+			check(PointData && PointData->Metadata);
+
 			OutPoint = Point;
+			if (PointData->Metadata->GetParent() == Data[DataIndex]->Metadata)
+			{
+				UPCGMetadataAccessorHelpers::InitializeMetadata(OutPoint, PointData->Metadata, Point);
+			}
+			else
+			{
+				UPCGMetadataAccessorHelpers::InitializeMetadata(OutPoint, PointData->Metadata);
+				// Since we can't inherit from the parent point, we'll set the values directly here
+				PointData->Metadata->SetAttributes(Point, Data[DataIndex]->Metadata, OutPoint);
+			}
 
 			if (DensityFunction == EPCGUnionDensityFunction::Binary && OutPoint.Density > 0)
 			{
 				OutPoint.Density = 1.0f;
 			}
 
-			// Compute final density based on current & following data
-			for (int32 FollowingDataIndex = DataIndex + DataIndexIncrement; FollowingDataIndex != LastDataIndex && OutPoint.Density < 1.0f; FollowingDataIndex += DataIndexIncrement)
+			// Update density & metadata based on current & following data
+			const bool bSkipLoop = (OutPoint.Density >= 1.0f && !PointData->Metadata);
+			for (int32 FollowingDataIndex = DataIndex + DataIndexIncrement; FollowingDataIndex != LastDataIndex && !bSkipLoop; FollowingDataIndex += DataIndexIncrement)
 			{
-				PCGUnionDataMaths::UpdateDensity(OutPoint.Density, Data[FollowingDataIndex]->GetDensityAtPosition(OutPoint.Transform.GetLocation()), DensityFunction);
+				FPCGPoint PointInData;
+				if (Data[FollowingDataIndex]->GetPointAtPosition(OutPoint.Transform.GetLocation(), PointInData, PointData->Metadata))
+				{
+					// Update density
+					PCGUnionDataMaths::UpdateDensity(OutPoint.Density, PointInData.Density, DensityFunction);
+
+					// Update color
+					OutPoint.Color = FVector4(
+						FMath::Max(OutPoint.Color.X, PointInData.Color.X),
+						FMath::Max(OutPoint.Color.Y, PointInData.Color.Y),
+						FMath::Max(OutPoint.Color.Z, PointInData.Color.Z),
+						FMath::Max(OutPoint.Color.W, PointInData.Color.W));
+
+					if (PointData->Metadata)
+					{
+						PointData->Metadata->MergeAttributes(OutPoint, PointInData, OutPoint, EPCGMetadataOp::Max);
+					}
+					else if (OutPoint.Density >= 1.0f)
+					{
+						break;
+					}
+				}
 			}
 
 			return true;
diff --git a/Engine/Plugins/Experimental/PCG/Source/PCG/Private/Elements/PCGDensityFilter.cpp b/Engine/Plugins/Experimental/PCG/Source/PCG/Private/Elements/PCGDensityFilter.cpp
--- a/Engine/Plugins/Experimental/PCG/Source/PCG/Private/Elements/PCGDensityFilter.cpp
+++ b/Engine/Plugins/Experimental/PCG/Source/PCG/Private/Elements/PCGDensityFilter.cpp
@@ -70,7 +70,7 @@ bool FPCGDensityFilterElement::ExecuteInternal(FPCGContext* Context) const
 		const TArray<FPCGPoint>& Points = OriginalData->GetPoints();
 		
 		UPCGPointData* FilteredData = NewObject<UPCGPointData>();
-		FilteredData->TargetActor = OriginalData->TargetActor;
+		FilteredData->InitializeFromData(OriginalData);
 		TArray<FPCGPoint>& FilteredPoints = FilteredData->GetMutablePoints();
 
 		Output.Data = FilteredData;
diff --git a/Engine/Plugins/Experimental/PCG/Source/PCG/Private/Elements/PCGDensityRemap.cpp b/Engine/Plugins/Experimental/PCG/Source/PCG/Private/Elements/PCGDensityRemap.cpp
--- a/Engine/Plugins/Experimental/PCG/Source/PCG/Private/Elements/PCGDensityRemap.cpp
+++ b/Engine/Plugins/Experimental/PCG/Source/PCG/Private/Elements/PCGDensityRemap.cpp
@@ -52,7 +52,7 @@ bool FPCGLinearDensityRemapElement::ExecuteInternal(FPCGContext* Context) const
 		const int OriginalPointCount = Points.Num();
 
 		UPCGPointData* SampledData = NewObject<UPCGPointData>();
-		SampledData->TargetActor = OriginalData->TargetActor;
+		SampledData->InitializeFromData(OriginalData);
 		TArray<FPCGPoint>& SampledPoints = SampledData->GetMutablePoints();
 		Output.Data = SampledData;
 
diff --git a/Engine/Plugins/Experimental/PCG/Source/PCG/Private/Elements/PCGExecuteBlueprint.cpp b/Engine/Plugins/Experimental/PCG/Source/PCG/Private/Elements/PCGExecuteBlueprint.cpp
--- a/Engine/Plugins/Experimental/PCG/Source/PCG/Private/Elements/PCGExecuteBlueprint.cpp
+++ b/Engine/Plugins/Experimental/PCG/Source/PCG/Private/Elements/PCGExecuteBlueprint.cpp
@@ -382,7 +382,7 @@ bool FPCGExecuteBlueprintElement::ExecuteInternal(FPCGContext* Context) const
 	return true;
 }
 
-void UPCGBlueprintElement::LoopOnPoints(FPCGContext& InContext, const UPCGPointData* InData, UPCGPointData*& OutData) const
+void UPCGBlueprintElement::LoopOnPoints(FPCGContext& InContext, const UPCGPointData* InData, UPCGPointData*& OutData, const UObject* InCustomObject) const
 {
 	if (!InData)
 	{
@@ -391,18 +391,18 @@ void UPCGBlueprintElement::LoopOnPoints(FPCGContext& InContext, const UPCGPointD
 	}
 
 	OutData = NewObject<UPCGPointData>();
-	OutData->TargetActor = InData->TargetActor;
+	OutData->InitializeFromData(InData);
 	
 	const TArray<FPCGPoint>& InPoints = InData->GetPoints();
 	TArray<FPCGPoint>& OutPoints = OutData->GetMutablePoints();
 
-	FPCGAsync::AsyncPointProcessing(&InContext, InPoints.Num(), OutPoints, [this, InData, &InPoints](int32 Index, FPCGPoint& OutPoint)
+	FPCGAsync::AsyncPointProcessing(&InContext, InPoints.Num(), OutPoints, [this, &InContext, InData, OutData, InCustomObject, &InPoints](int32 Index, FPCGPoint& OutPoint)
 	{
-		return PointLoopBody(InData, InPoints[Index], OutPoint);
+		return PointLoopBody(InContext, InData, InPoints[Index], InCustomObject, OutPoint, OutData->Metadata);
 	});
 }
 
-void UPCGBlueprintElement::LoopOnPointPairs(FPCGContext& InContext, const UPCGPointData* InA, const UPCGPointData* InB, UPCGPointData*& OutData) const
+void UPCGBlueprintElement::LoopOnPointPairs(FPCGContext& InContext, const UPCGPointData* InA, const UPCGPointData* InB, UPCGPointData*& OutData, const UObject* InCustomObject) const
 {
 	if (!InA || !InB)
 	{
@@ -411,15 +411,15 @@ void UPCGBlueprintElement::LoopOnPointPairs(FPCGContext& InContext, const UPCGPo
 	}
 
 	OutData = NewObject<UPCGPointData>();
-	OutData->TargetActor = InA->TargetActor;
+	OutData->InitializeFromData(InA); //METADATA TODO should we remove the parenting here?
 
 	const TArray<FPCGPoint>& InPointsA = InA->GetPoints();
 	const TArray<FPCGPoint>& InPointsB = InB->GetPoints();
 	TArray<FPCGPoint>& OutPoints = OutData->GetMutablePoints();
 
-	FPCGAsync::AsyncPointProcessing(&InContext, InPointsA.Num() * InPointsB.Num(), OutPoints, [this, InA, InB, &InPointsA, &InPointsB](int32 Index, FPCGPoint& OutPoint)
+	FPCGAsync::AsyncPointProcessing(&InContext, InPointsA.Num() * InPointsB.Num(), OutPoints, [this, &InContext, InA, InB, InCustomObject, OutData, &InPointsA, &InPointsB](int32 Index, FPCGPoint& OutPoint)
 	{
-		return PointPairLoopBody(InA, InB, InPointsA[Index / InPointsB.Num()], InPointsB[Index % InPointsB.Num()], OutPoint);
+		return PointPairLoopBody(InContext, InA, InB, InPointsA[Index / InPointsB.Num()], InPointsB[Index % InPointsB.Num()], InCustomObject, OutPoint, OutData->Metadata);
 	});
 }
 
diff --git a/Engine/Plugins/Experimental/PCG/Source/PCG/Private/Elements/PCGMetadataElement.cpp b/Engine/Plugins/Experimental/PCG/Source/PCG/Private/Elements/PCGMetadataElement.cpp
new file mode 100644
--- /dev/null
+++ b/Engine/Plugins/Experimental/PCG/Source/PCG/Private/Elements/PCGMetadataElement.cpp
@@ -0,0 +1,197 @@
+// Copyright Epic Games, Inc. All Rights Reserved.
+
+#include "Elements/PCGMetadataElement.h"
+#include "Data/PCGPointData.h"
+#include "Metadata/PCGMetadata.h"
+#include "Metadata/PCGMetadataAttribute.h"
+#include "Metadata/PCGMetadataAttributeTpl.h"
+
+FPCGElementPtr UPCGMetadataOperationSettings::CreateElement() const
+{
+	return MakeShared<FPCGMetadataOperationElement>();
+}
+
+bool FPCGMetadataOperationElement::ExecuteInternal(FPCGContext* Context) const
+{
+	TRACE_CPUPROFILER_EVENT_SCOPE(FPCGMetadataOperationElement::Execute);
+
+	const UPCGMetadataOperationSettings* Settings = Context->GetInputSettings<UPCGMetadataOperationSettings>();
+	check(Settings);
+
+	TArray<FPCGTaggedData> Inputs = Context->InputData.GetInputs();
+	TArray<FPCGTaggedData>& Outputs = Context->OutputData.TaggedData;
+
+	// Forward any non-input data
+	Outputs.Append(Context->InputData.GetExclusions());
+	Outputs.Append(Context->InputData.GetAllSettings());
+
+	for (const FPCGTaggedData& Input : Inputs)
+	{
+		FPCGTaggedData& Output = Outputs.Add_GetRef(Input);
+
+		const UPCGSpatialData* SpatialInput = Cast<const UPCGSpatialData>(Input.Data);
+
+		if (!SpatialInput)
+		{
+			PCGE_LOG(Error, "Invalid input data");
+			continue;
+		}
+
+		const UPCGPointData* OriginalData = SpatialInput->ToPointData(Context);
+
+		if (!OriginalData)
+		{
+			PCGE_LOG(Error, "Unable to get point data from input");
+			continue;
+		}
+
+		if (!OriginalData->Metadata)
+		{
+			PCGE_LOG(Warning, "Input has no metadata");
+			continue;
+		}
+
+		// Check if the attribute exists
+		if ((Settings->Target == EPCGMetadataOperationTarget::AttributeToProperty || Settings->Target == EPCGMetadataOperationTarget::AttributeToAttribute) && !OriginalData->Metadata->HasAttribute(Settings->SourceAttribute))
+		{
+			PCGE_LOG(Warning, "Input does not have the %s attribute", *Settings->SourceAttribute.ToString());
+			continue;
+		}
+
+		const TArray<FPCGPoint>& Points = OriginalData->GetPoints();
+		const int OriginalPointCount = Points.Num();
+
+		UPCGPointData* SampledData = NewObject<UPCGPointData>();
+		SampledData->InitializeFromData(OriginalData);
+		TArray<FPCGPoint>& SampledPoints = SampledData->GetMutablePoints();
+
+		Output.Data = SampledData;
+
+		// Copy points and then apply the operation
+		SampledPoints = Points;
+
+		if (Settings->Target == EPCGMetadataOperationTarget::PropertyToAttribute)
+		{
+			if (Settings->PointProperty == EPCGPointProperties::Density)
+			{
+				if (!SampledData->Metadata->HasAttribute(Settings->DestinationAttribute))
+				{
+					SampledData->Metadata->CreateFloatAttribute(Settings->DestinationAttribute, 0.0f, /*bAllowsInterpolation=*/true, /*bOverrideParent=*/true);
+				}
+
+				FPCGMetadataAttributeBase* AttributeBase = SampledData->Metadata->GetMutableAttribute(Settings->DestinationAttribute);
+				if (AttributeBase && AttributeBase->GetTypeId() == PCG::Private::MetadataTypes<float>::Id)
+				{
+					FPCGMetadataAttribute<float>* Attribute = static_cast<FPCGMetadataAttribute<float>*>(AttributeBase);
+					for (FPCGPoint& Point : SampledPoints)
+					{
+						Attribute->SetValue(Point.MetadataEntry, Point.Density);
+					}
+				}
+				else
+				{
+					PCGE_LOG(Error, "Attribute %s already exists but its type is not compatible", *Settings->DestinationAttribute.ToString());
+				}
+			}
+			else if (Settings->PointProperty == EPCGPointProperties::Extents)
+			{
+				if (!SampledData->Metadata->HasAttribute(Settings->DestinationAttribute))
+				{
+					SampledData->Metadata->CreateVectorAttribute(Settings->DestinationAttribute, FVector::Zero() , /*bAllowsInterpolation=*/true, /*bOverrideParent=*/true);
+				}
+
+				FPCGMetadataAttributeBase* AttributeBase = SampledData->Metadata->GetMutableAttribute(Settings->DestinationAttribute);
+				if (AttributeBase && AttributeBase->GetTypeId() == PCG::Private::MetadataTypes<FVector>::Id)
+				{
+					FPCGMetadataAttribute<FVector>* Attribute = static_cast<FPCGMetadataAttribute<FVector>*>(AttributeBase);
+					for (FPCGPoint& Point : SampledPoints)
+					{
+						Attribute->SetValue(Point.MetadataEntry, Point.Extents);
+					}
+				}
+				else
+				{
+					PCGE_LOG(Error, "Attribute %s already exists but its type is not compatible", *Settings->DestinationAttribute.ToString());
+				}
+			}
+			else if(Settings->PointProperty == EPCGPointProperties::Color)
+			{
+				if (!SampledData->Metadata->HasAttribute(Settings->DestinationAttribute))
+				{
+					SampledData->Metadata->CreateVector4Attribute(Settings->DestinationAttribute, FVector4::Zero(), /*bAllowsInterpolation=*/true, /*bOverrideParent=*/true);
+				}
+
+				FPCGMetadataAttributeBase* AttributeBase = SampledData->Metadata->GetMutableAttribute(Settings->DestinationAttribute);
+				if (AttributeBase && AttributeBase->GetTypeId() == PCG::Private::MetadataTypes<FVector4>::Id)
+				{
+					FPCGMetadataAttribute<FVector4>* Attribute = static_cast<FPCGMetadataAttribute<FVector4>*>(AttributeBase);
+					for (FPCGPoint& Point : SampledPoints)
+					{
+						Attribute->SetValue(Point.MetadataEntry, Point.Color);
+					}
+				}
+				else
+				{
+					PCGE_LOG(Error, "Attribute %s already exists but its type is not compatible", *Settings->DestinationAttribute.ToString());
+				}
+			}
+		}
+		else if(Settings->Target == EPCGMetadataOperationTarget::AttributeToProperty) // Attribute to property
+		{
+			if (Settings->PointProperty == EPCGPointProperties::Density)
+			{
+				const FPCGMetadataAttributeBase* AttributeBase = SampledData->Metadata->GetConstAttribute(Settings->SourceAttribute);
+				if (AttributeBase && AttributeBase->GetTypeId() == PCG::Private::MetadataTypes<float>::Id)
+				{
+					const FPCGMetadataAttribute<float>* Attribute = static_cast<const FPCGMetadataAttribute<float>*>(AttributeBase);
+					for (FPCGPoint& Point : SampledPoints)
+					{
+						Point.Density = Attribute->GetValueFromItemKey(Point.MetadataEntry);
+					}
+				}
+				else
+				{
+					PCGE_LOG(Error, "Attribute %s already exists but its type is not compatible", *Settings->SourceAttribute.ToString());
+				}
+			}
+			else if (Settings->PointProperty == EPCGPointProperties::Extents)
+			{
+				const FPCGMetadataAttributeBase* AttributeBase = SampledData->Metadata->GetConstAttribute(Settings->SourceAttribute);
+				if (AttributeBase && AttributeBase->GetTypeId() == PCG::Private::MetadataTypes<FVector>::Id)
+				{
+					const FPCGMetadataAttribute<FVector>* Attribute = static_cast<const FPCGMetadataAttribute<FVector>*>(AttributeBase);
+					for (FPCGPoint& Point : SampledPoints)
+					{
+						Point.Extents = Attribute->GetValueFromItemKey(Point.MetadataEntry);
+					}
+				}
+				else
+				{
+					PCGE_LOG(Error, "Attribute %s already exists but its type is not compatible", *Settings->SourceAttribute.ToString());
+				}
+			}
+			else if (Settings->PointProperty == EPCGPointProperties::Color)
+			{
+				const FPCGMetadataAttributeBase* AttributeBase = SampledData->Metadata->GetConstAttribute(Settings->SourceAttribute);
+				if (AttributeBase && AttributeBase->GetTypeId() == PCG::Private::MetadataTypes<FVector4>::Id)
+				{
+					const FPCGMetadataAttribute<FVector4>* Attribute = static_cast<const FPCGMetadataAttribute<FVector4>*>(AttributeBase);
+					for (FPCGPoint& Point : SampledPoints)
+					{
+						Point.Color = Attribute->GetValueFromItemKey(Point.MetadataEntry);
+					}
+				}
+				else
+				{
+					PCGE_LOG(Error, "Attribute %s already exists but its type is not compatible", *Settings->SourceAttribute.ToString());
+				}
+			}
+		}
+		else // Attribute to attribute
+		{
+			SampledData->Metadata->CopyAttribute(Settings->SourceAttribute, Settings->DestinationAttribute);
+		}
+	}
+
+	return true;
+}
\ No newline at end of file
diff --git a/Engine/Plugins/Experimental/PCG/Source/PCG/Private/Elements/PCGPointSampler.cpp b/Engine/Plugins/Experimental/PCG/Source/PCG/Private/Elements/PCGPointSampler.cpp
--- a/Engine/Plugins/Experimental/PCG/Source/PCG/Private/Elements/PCGPointSampler.cpp
+++ b/Engine/Plugins/Experimental/PCG/Source/PCG/Private/Elements/PCGPointSampler.cpp
@@ -69,7 +69,7 @@ bool FPCGPointSamplerElement::ExecuteInternal(FPCGContext* Context) const
 		const int OriginalPointCount = Points.Num();
 
 		UPCGPointData* SampledData = NewObject<UPCGPointData>();
-		SampledData->TargetActor = OriginalData->TargetActor;
+		SampledData->InitializeFromData(OriginalData);
 		TArray<FPCGPoint>& SampledPoints = SampledData->GetMutablePoints();
 
 		Output.Data = SampledData;
diff --git a/Engine/Plugins/Experimental/PCG/Source/PCG/Private/Elements/PCGSelfPruning.cpp b/Engine/Plugins/Experimental/PCG/Source/PCG/Private/Elements/PCGSelfPruning.cpp
--- a/Engine/Plugins/Experimental/PCG/Source/PCG/Private/Elements/PCGSelfPruning.cpp
+++ b/Engine/Plugins/Experimental/PCG/Source/PCG/Private/Elements/PCGSelfPruning.cpp
@@ -183,8 +183,7 @@ namespace PCGSelfPruningElement
 			Output = Input;
 
 			UPCGPointData* PrunedData = NewObject<UPCGPointData>();
-			PrunedData->TargetActor = SpatialInput->TargetActor;
-
+			PrunedData->InitializeFromData(InputPointData);
 			Output.Data = PrunedData;
 
 			TArray<FPCGPoint>& OutputPoints = PrunedData->GetMutablePoints();
diff --git a/Engine/Plugins/Experimental/PCG/Source/PCG/Private/Elements/PCGSurfaceSampler.cpp b/Engine/Plugins/Experimental/PCG/Source/PCG/Private/Elements/PCGSurfaceSampler.cpp
--- a/Engine/Plugins/Experimental/PCG/Source/PCG/Private/Elements/PCGSurfaceSampler.cpp
+++ b/Engine/Plugins/Experimental/PCG/Source/PCG/Private/Elements/PCGSurfaceSampler.cpp
@@ -144,12 +144,12 @@ bool FPCGSurfaceSamplerElement::ExecuteInternal(FPCGContext* Context) const
 		Output = Input;
 
 		UPCGPointData* SampledData = NewObject<UPCGPointData>();
-		SampledData->TargetActor = SpatialInput->TargetActor;
+		SampledData->InitializeFromData(SpatialInput);
 		Output.Data = SampledData;
 
 		TArray<FPCGPoint>& SampledPoints = SampledData->GetMutablePoints();
 
-		FPCGAsync::AsyncPointProcessing(Context, LoopData.CellCount, SampledPoints, [&LoopData, SpatialInput](int32 Index, FPCGPoint& OutPoint)
+		FPCGAsync::AsyncPointProcessing(Context, LoopData.CellCount, SampledPoints, [&LoopData, SampledData, SpatialInput](int32 Index, FPCGPoint& OutPoint)
 		{
 			int32 CellX;
 			int32 CellY;
@@ -168,20 +168,27 @@ bool FPCGSurfaceSamplerElement::ExecuteInternal(FPCGContext* Context) const
 			{
 				const float RandX = RandomSource.FRand();
 				const float RandY = RandomSource.FRand();
-
-				OutPoint.Transform = FTransform(FVector(CurrentX + RandX * InnerCellSize.X, CurrentY + RandY * InnerCellSize.Y, 0));
-				OutPoint.Extents = FVector(LoopData.Settings->PointExtents);
-				OutPoint.Density = LoopData.Settings->bApplyDensityToPoints ? ((Ratio - Chance) / Ratio) : 1.0f;
-				OutPoint.Steepness = LoopData.Settings->PointSteepness;
-				OutPoint.Seed = RandomSource.GetCurrentSeed();
-
-				OutPoint = SpatialInput->TransformPoint(OutPoint);
+				 
+				const FVector TentativeLocation = FVector(CurrentX + RandX * InnerCellSize.X, CurrentY + RandY * InnerCellSize.Y, 0.0f);
 
 #if WITH_EDITORONLY_DATA
-				return OutPoint.Density > 0 || LoopData.Settings->bKeepZeroDensityPoints;
+				if (SpatialInput->GetPointAtPosition(TentativeLocation, OutPoint, SampledData->Metadata) || LoopData.Settings->bKeepZeroDensityPoints)
 #else
-				return OutPoint.Density > 0;
+				if (SpatialInput->GetPointAtPosition(TentativeLocation, OutPoint, SampledData->Metadata))
 #endif
+				{
+					// Apply final parameters on the point
+					OutPoint.Extents = FVector(LoopData.Settings->PointExtents);
+					OutPoint.Density *= (LoopData.Settings->bApplyDensityToPoints ? ((Ratio - Chance) / Ratio) : 1.0f);
+					OutPoint.Steepness = LoopData.Settings->PointSteepness;
+					OutPoint.Seed = RandomSource.GetCurrentSeed();
+
+					return true;
+				}
+				else
+				{
+					return false;
+				}
 			}
 			else
 			{
diff --git a/Engine/Plugins/Experimental/PCG/Source/PCG/Private/Metadata/PCGMetadata.cpp b/Engine/Plugins/Experimental/PCG/Source/PCG/Private/Metadata/PCGMetadata.cpp
new file mode 100644
--- /dev/null
+++ b/Engine/Plugins/Experimental/PCG/Source/PCG/Private/Metadata/PCGMetadata.cpp
@@ -0,0 +1,638 @@
+// Copyright Epic Games, Inc. All Rights Reserved.
+
+#include "Metadata/PCGMetadata.h"
+#include "Metadata/PCGMetadataAttributeTpl.h"
+#include "PCGModule.h"
+
+#include "Misc/ScopeRWLock.h"
+
+void UPCGMetadata::Serialize(FArchive& Ar)
+{
+	// TODO
+}
+
+void UPCGMetadata::BeginDestroy()
+{
+	FWriteScopeLock ScopeLock(AttributeLock);
+	for (TPair<FName, FPCGMetadataAttributeBase*>& AttributeEntry : Attributes)
+	{
+		delete AttributeEntry.Value;
+		AttributeEntry.Value = nullptr;
+	}
+	Attributes.Reset();
+
+	Super::BeginDestroy();
+}
+
+void UPCGMetadata::Initialize(const UPCGMetadata* InParent)
+{
+	check(!Parent && Attributes.Num() == 0);
+	Parent = (InParent != this ? InParent : nullptr);
+	ItemKeyOffset = Parent ? Parent->GetItemCountForChild() : 0;
+	GetAllAttributeNames(AttributeNames);
+}
+
+void UPCGMetadata::AddAttributes(const UPCGMetadata* InOther)
+{
+	if (!InOther || InOther == Parent)
+	{
+		return;
+	}
+
+	for (const TPair<FName, FPCGMetadataAttributeBase*> OtherAttribute : InOther->Attributes)
+	{
+		if (HasAttribute(OtherAttribute.Key))
+		{
+			continue;
+		}
+		else
+		{
+			CopyAttribute(OtherAttribute.Value, OtherAttribute.Key, /*bKeepParent=*/false, /*bCopyEntries=*/false, /*bCopyValues=*/false);
+		}
+	}
+
+	OtherParents.Add(InOther);
+}
+
+const UPCGMetadata* UPCGMetadata::GetRoot() const
+{
+	if (Parent)
+	{
+		return Parent->GetRoot();
+	}
+	else
+	{
+		return this;
+	}
+}
+
+bool UPCGMetadata::HasParent(const UPCGMetadata* InTentativeParent) const
+{
+	if (!InTentativeParent)
+	{
+		return false;
+	}
+
+	const UPCGMetadata* HierarchicalParent = Parent.Get();
+	while (HierarchicalParent && HierarchicalParent != InTentativeParent)
+	{
+		HierarchicalParent = HierarchicalParent->Parent.Get();
+	}
+
+	return HierarchicalParent == InTentativeParent;
+}
+
+void UPCGMetadata::AddAttributeInternal(FName AttributeName, FPCGMetadataAttributeBase* Attribute)
+{
+	// This call assumes we have a write lock on the attribute map.
+	Attributes.Add(AttributeName, Attribute);
+	AttributeNames.Add(AttributeName);
+	HiddenAttributes.Remove(AttributeName);
+}
+
+void UPCGMetadata::RemoveAttributeInternal(FName AttributeName)
+{
+	Attributes.Remove(AttributeName);
+	if (!Parent || !Parent->HasAttribute(AttributeName))
+	{
+		AttributeNames.Remove(AttributeName);
+	}
+}
+
+template<typename T>
+FPCGMetadataAttributeBase* UPCGMetadata::CreateAttribute(FName AttributeName, const T& DefaultValue, bool bAllowsInterpolation, bool bOverrideParent)
+{
+	const FPCGMetadataAttributeBase* ParentAttribute = nullptr;
+
+	if (bOverrideParent && Parent.IsValid())
+	{
+		ParentAttribute = Parent->GetConstAttribute(AttributeName);
+	}
+
+	FPCGMetadataAttribute<T>* NewAttribute = new FPCGMetadataAttribute<T>(this, AttributeName, ParentAttribute, DefaultValue, bAllowsInterpolation);
+
+	AttributeLock.WriteLock();
+	if (Attributes.Find(AttributeName))
+	{
+		UE_LOG(LogPCG, Warning, TEXT("Attribute %s already exists"), *AttributeName.ToString());
+		delete NewAttribute;
+	}
+	else
+	{
+		NewAttribute->AttributeId = NextAttributeId++;
+		AddAttributeInternal(AttributeName, NewAttribute);
+	}
+	AttributeLock.WriteUnlock();
+
+	return NewAttribute;
+}
+
+FPCGMetadataAttributeBase* UPCGMetadata::GetMutableAttribute(FName AttributeName)
+{
+	FPCGMetadataAttributeBase* Attribute = nullptr;
+
+	AttributeLock.ReadLock();
+	if (FPCGMetadataAttributeBase** FoundAttribute = Attributes.Find(AttributeName))
+	{
+		Attribute = *FoundAttribute;
+	}
+	AttributeLock.ReadUnlock();
+
+	if (!Attribute && ParentHasAttribute(AttributeName))
+	{
+		Attribute = CopyAttribute(AttributeName, AttributeName, /*bKeepParent=*/true, /*bCopyEntries=*/false, /*bCopyValues=*/false);
+	}
+
+	return Attribute;
+}
+
+const FPCGMetadataAttributeBase* UPCGMetadata::GetConstAttribute(FName AttributeName) const
+{
+	const FPCGMetadataAttributeBase* Attribute = nullptr;
+
+	AttributeLock.ReadLock();
+	if (const FPCGMetadataAttributeBase* const* FoundAttribute = Attributes.Find(AttributeName))
+	{
+		Attribute = *FoundAttribute;
+	}
+	AttributeLock.ReadUnlock();
+
+	if (!Attribute && ParentHasAttribute(AttributeName))
+	{
+		return Parent->GetConstAttribute(AttributeName);
+	}
+	else
+	{
+		return Attribute;
+	}
+}
+
+bool UPCGMetadata::HasAttribute(FName AttributeName) const
+{
+	FReadScopeLock ScopeLock(AttributeLock);
+	return AttributeNames.Contains(AttributeName);
+}
+
+bool UPCGMetadata::ParentHasAttribute(FName AttributeName) const
+{
+	if (Parent && Parent->HasAttribute(AttributeName))
+	{
+		FReadScopeLock ScopeLock(AttributeLock);
+		return !HiddenAttributes.Contains(AttributeName);
+	}
+	else
+	{
+		return false;
+	}
+}
+
+void UPCGMetadata::CreateFloatAttribute(FName AttributeName, float DefaultValue, bool bAllowsInterpolation, bool bOverrideParent)
+{
+	CreateAttribute<float>(AttributeName, DefaultValue, bAllowsInterpolation, bOverrideParent);
+}
+
+void UPCGMetadata::CreateVectorAttribute(FName AttributeName, const FVector& DefaultValue, bool bAllowsInterpolation, bool bOverrideParent)
+{
+	CreateAttribute<FVector>(AttributeName, DefaultValue, bAllowsInterpolation, bOverrideParent);
+}
+
+void UPCGMetadata::CreateVector4Attribute(FName AttributeName, const FVector4& DefaultValue, bool bAllowsInterpolation, bool bOverrideParent)
+{
+	CreateAttribute<FVector4>(AttributeName, DefaultValue, bAllowsInterpolation, bOverrideParent);
+}
+
+void UPCGMetadata::CreateQuatAttribute(FName AttributeName, const FQuat& DefaultValue, bool bAllowsInterpolation, bool bOverrideParent)
+{
+	CreateAttribute<FQuat>(AttributeName, DefaultValue, bAllowsInterpolation, bOverrideParent);
+}
+
+void UPCGMetadata::CreateTransformAttribute(FName AttributeName, const FTransform& DefaultValue, bool bAllowsInterpolation, bool bOverrideParent)
+{
+	CreateAttribute<FTransform>(AttributeName, DefaultValue, bAllowsInterpolation, bOverrideParent);
+}
+
+void UPCGMetadata::CreateStringAttribute(FName AttributeName, const FString& DefaultValue, bool bAllowsInterpolation, bool bOverrideParent)
+{
+	CreateAttribute<FString>(AttributeName, DefaultValue, bAllowsInterpolation, bOverrideParent);
+}
+
+void UPCGMetadata::CopyAttribute(FName AttributeToCopy, FName NewAttributeName, bool bKeepParent)
+{
+	CopyAttribute(AttributeToCopy, NewAttributeName, bKeepParent, /*bCopyEntries=*/true, /*bCopyValues=*/true);
+}
+
+FPCGMetadataAttributeBase* UPCGMetadata::CopyAttribute(FName AttributeToCopy, FName NewAttributeName, bool bKeepParent, bool bCopyEntries, bool bCopyValues)
+{
+	const FPCGMetadataAttributeBase* OriginalAttribute = nullptr;
+
+	AttributeLock.ReadLock();
+	if (FPCGMetadataAttributeBase** AttributeFound = Attributes.Find(AttributeToCopy))
+	{
+		OriginalAttribute = *AttributeFound;
+	}
+	AttributeLock.ReadUnlock();
+
+	if (!OriginalAttribute && Parent)
+	{
+		OriginalAttribute = Parent->GetConstAttribute(AttributeToCopy);
+	}
+
+	if (!OriginalAttribute)
+	{
+		UE_LOG(LogPCG, Warning, TEXT("Attribute %s does not exist, therefore cannot be copied"), *AttributeToCopy.ToString());
+		return nullptr;
+	}
+
+	return CopyAttribute(OriginalAttribute, NewAttributeName, bKeepParent, bCopyEntries, bCopyValues);
+}
+
+FPCGMetadataAttributeBase* UPCGMetadata::CopyAttribute(const FPCGMetadataAttributeBase* OriginalAttribute, FName NewAttributeName, bool bKeepParent, bool bCopyEntries, bool bCopyValues)
+{
+	check(OriginalAttribute);
+	check(OriginalAttribute->GetMetadata()->GetRoot() == GetRoot() || !bKeepParent);
+	FPCGMetadataAttributeBase* NewAttribute = OriginalAttribute->Copy(NewAttributeName, this, bKeepParent, bCopyEntries, bCopyValues);
+
+	AttributeLock.WriteLock();
+	NewAttribute->AttributeId = ++NextAttributeId;
+	AddAttributeInternal(NewAttributeName, NewAttribute);
+	AttributeLock.WriteUnlock();
+
+	return NewAttribute;
+}
+
+void UPCGMetadata::RenameAttribute(FName AttributeToRename, FName NewAttributeName)
+{
+	bool bRenamed = false;
+	AttributeLock.WriteLock();
+	if (FPCGMetadataAttributeBase** AttributeFound = Attributes.Find(AttributeToRename))
+	{
+		FPCGMetadataAttributeBase* Attribute = *AttributeFound;
+		RemoveAttributeInternal(AttributeToRename);
+		Attribute->Name = NewAttributeName;
+		AddAttributeInternal(NewAttributeName, Attribute);
+		
+		bRenamed = true;
+	}
+	AttributeLock.WriteUnlock();
+
+	if(!bRenamed && ParentHasAttribute(AttributeToRename))
+	{
+		// If attribute is present in the parent, we'll otherwise copy the attribute from the parent with a new name
+		CopyAttribute(AttributeToRename, NewAttributeName, /*bKeepParent=*/true);
+		bRenamed = true;
+	}
+
+	if (!bRenamed)
+	{
+		UE_LOG(LogPCG, Warning, TEXT("Attribute %s does not exist and therefore cannot be renamed"), *AttributeToRename.ToString());
+	}
+}
+
+void UPCGMetadata::ClearAttribute(FName AttributeToClear)
+{
+	FPCGMetadataAttributeBase* Attribute = nullptr;
+
+	AttributeLock.ReadLock();
+	if (FPCGMetadataAttributeBase** AttributeFound = Attributes.Find(AttributeToClear))
+	{
+		Attribute = *AttributeFound;
+	}
+	AttributeLock.ReadUnlock();
+
+	// If the attribute exists, then we can lose all the entries
+	// If it doesn't but it exists in the parent hierarchy, then we must create a new attribute.
+	if (Attribute)
+	{
+		Attribute->ClearEntries();
+	}
+	else if(ParentHasAttribute(AttributeToClear))
+	{
+		CopyAttribute(AttributeToClear, AttributeToClear, /*bKeepParent=*/true, /*bCopyEntries=*/false, /*bCopyValues=*/false);
+	}
+}
+
+void UPCGMetadata::DeleteAttribute(FName AttributeToDelete)
+{
+	FPCGMetadataAttributeBase* Attribute = nullptr;
+
+	// If it's a local attribute, then just delete it
+	AttributeLock.WriteLock();
+	if (FPCGMetadataAttributeBase** AttributeFound = Attributes.Find(AttributeToDelete))
+	{
+		Attribute = *AttributeFound;
+		RemoveAttributeInternal(AttributeToDelete);
+	}
+	AttributeLock.WriteUnlock();
+
+	if (Attribute)
+	{
+		delete Attribute;
+	}
+	else if(ParentHasAttribute(AttributeToDelete))
+	{
+		// Otherwise, if it exists in parent hierarchy, then hide it
+		FWriteScopeLock ScopeLock(AttributeLock);
+		HiddenAttributes.Add(AttributeToDelete);
+		AttributeNames.Remove(AttributeToDelete);
+	}
+	else
+	{
+		UE_LOG(LogPCG, Verbose, TEXT("Attribute %s does not exist and therefore cannot be deleted"), *AttributeToDelete.ToString());
+	}
+}
+
+int64 UPCGMetadata::GetItemCountForChild() const
+{
+	FReadScopeLock ScopeLock(ItemLock);
+	return ParentKeys.Num() + ItemKeyOffset;
+}
+
+int64 UPCGMetadata::AddEntry(int64 ParentEntry)
+{
+	FWriteScopeLock ScopeLock(ItemLock);
+	return ParentKeys.Add(ParentEntry) + ItemKeyOffset;
+}
+
+PCGMetadataEntryKey UPCGMetadata::GetParentKey(PCGMetadataEntryKey LocalItemKey) const
+{
+	if (LocalItemKey < ItemKeyOffset)
+	{
+		if (Parent)
+		{
+			return Parent->GetParentKey(LocalItemKey);
+		}
+		else
+		{
+			return PCGInvalidEntryKey;
+		}
+	}
+	else
+	{
+		FReadScopeLock ScopeLock(ItemLock);
+		check(LocalItemKey - ItemKeyOffset < ParentKeys.Num());
+		return ParentKeys[LocalItemKey - ItemKeyOffset];
+	}
+}
+
+void UPCGMetadata::GetAllAttributeNames(TSet<FName>& FoundAttributeNames) const
+{
+	for (const TPair<FName, FPCGMetadataAttributeBase*>& Attribute : Attributes)
+	{
+		FoundAttributeNames.Add(Attribute.Key);
+	}
+
+	if (Parent)
+	{
+		Parent->GetAllAttributeNames(FoundAttributeNames);
+	}
+}
+
+void UPCGMetadata::MergeAttributes(const FPCGPoint& InPointA, const UPCGMetadata* InMetadataA, const FPCGPoint& InPointB, const UPCGMetadata* InMetadataB, FPCGPoint& OutPoint, EPCGMetadataOp Op)
+{
+	// Early out: nothing to do if both input metadata are null / points have no assigned metadata
+	if (!InMetadataA && !InMetadataB)
+	{
+		return;
+	}
+
+	// For each attribute in the current metadata, query the values from point A & B, apply operation on the result and finally store in the out point.
+	if (OutPoint.MetadataEntry == PCGInvalidEntryKey)
+	{
+		if (InPointA.MetadataEntry != PCGInvalidEntryKey && (Parent == InMetadataA || this == InMetadataA))
+		{
+			OutPoint.MetadataEntry = AddEntry(InPointA.MetadataEntry);
+		}
+		else if (InPointB.MetadataEntry != PCGInvalidEntryKey && (Parent == InMetadataB || this == InMetadataB))
+		{
+			OutPoint.MetadataEntry = AddEntry(InPointB.MetadataEntry);
+		}
+		else
+		{
+			OutPoint.MetadataEntry = AddEntry();
+		}
+	}
+
+	//METADATA TODO: This isn't great, we're allocating memory for no reason here
+	AttributeLock.ReadLock();
+	TSet<FName> KnownAttributeNames = AttributeNames;
+	AttributeLock.ReadUnlock();
+
+	for(const FName& AttributeName : KnownAttributeNames)
+	{
+		const FPCGMetadataAttributeBase* ConstAttribute = GetConstAttribute(AttributeName);
+
+		// Get attribute from A
+		const FPCGMetadataAttributeBase* AttributeA = InMetadataA ? InMetadataA->GetConstAttribute(AttributeName) : nullptr;
+
+		if (AttributeA && AttributeA->GetTypeId() != ConstAttribute->GetTypeId())
+		{
+			UE_LOG(LogPCG, Error, TEXT("Metadata type mismatch with attribute %s"), *AttributeName.ToString());
+			AttributeA = nullptr;
+		}
+
+		// Get attribute from B
+		const FPCGMetadataAttributeBase* AttributeB = InMetadataB ? InMetadataB->GetConstAttribute(AttributeName) : nullptr;
+
+		if (AttributeB && AttributeB->GetTypeId() != ConstAttribute->GetTypeId())
+		{
+			UE_LOG(LogPCG, Error, TEXT("Metadata type mismatch with attribute %s"), *AttributeName.ToString());
+			AttributeB = nullptr;
+		}
+
+		if (AttributeA || AttributeB)
+		{
+			FPCGMetadataAttributeBase* Attribute = GetMutableAttribute(AttributeName);
+			if (Attribute)
+			{
+				Attribute->SetValue(OutPoint.MetadataEntry, AttributeA, InPointA.MetadataEntry, AttributeB, InPointB.MetadataEntry, Op);
+			}
+		}
+	}
+}
+
+void UPCGMetadata::ResetWeightedAttributes(FPCGPoint& OutPoint)
+{
+	if (OutPoint.MetadataEntry == PCGInvalidEntryKey)
+	{
+		OutPoint.MetadataEntry = AddEntry();
+	}
+
+	//METADATA TODO: This isn't great, we're allocating memory for no reason here
+	AttributeLock.ReadLock();
+	TSet<FName> KnownAttributeNames = AttributeNames;
+	AttributeLock.ReadUnlock();
+
+	for (const FName& AttributeName : KnownAttributeNames)
+	{
+		const FPCGMetadataAttributeBase* ConstAttribute = GetConstAttribute(AttributeName);
+		
+		if (!ConstAttribute->AllowsInterpolation())
+		{
+			continue;
+		}
+
+		FPCGMetadataAttributeBase* Attribute = GetMutableAttribute(AttributeName);
+		if (Attribute)
+		{
+			check(Attribute->AllowsInterpolation());
+			Attribute->SetZeroValue(OutPoint.MetadataEntry);
+		}
+	}
+}
+
+void UPCGMetadata::AccumulateWeightedAttributes(const FPCGPoint& InPoint, const UPCGMetadata* InMetadata, float Weight, bool bSetNonInterpolableAttributes, FPCGPoint& OutPoint)
+{
+	if (!InMetadata)
+	{
+		return;
+	}
+
+	bool bHasSetParent = false;
+
+	if (OutPoint.MetadataEntry == PCGInvalidEntryKey)
+	{
+		if (InPoint.MetadataEntry != PCGInvalidEntryKey && Parent == InMetadata)
+		{
+			OutPoint.MetadataEntry = AddEntry(InPoint.MetadataEntry);
+			// No early out here, since we'll need to set the weighted value
+			bHasSetParent = true;
+		}
+		else
+		{
+			OutPoint.MetadataEntry = AddEntry();
+		}
+	}
+
+	const bool bShouldSetNonInterpolableAttributes = bSetNonInterpolableAttributes && !bHasSetParent;
+
+	//METADATA TODO: This isn't great, we're allocating memory for no reason here
+	AttributeLock.ReadLock();
+	TSet<FName> KnownAttributeNames = AttributeNames;
+	AttributeLock.ReadUnlock();
+
+	for(const FName& AttributeName : KnownAttributeNames)
+	{
+		const FPCGMetadataAttributeBase* ConstAttribute = GetConstAttribute(AttributeName);
+
+		if (const FPCGMetadataAttributeBase* OtherAttribute = InMetadata->GetConstAttribute(AttributeName))
+		{
+			if (OtherAttribute->GetTypeId() != ConstAttribute->GetTypeId())
+			{
+				UE_LOG(LogPCG, Error, TEXT("Metadata type mismatch with attribute %s"), *AttributeName.ToString());
+				continue;
+			}
+
+			FPCGMetadataAttributeBase* Attribute = GetMutableAttribute(AttributeName);
+
+			if (Attribute->AllowsInterpolation())
+			{
+				Attribute->AccumulateValue(OutPoint.MetadataEntry, OtherAttribute, InPoint.MetadataEntry, Weight);
+			}
+			else if (bShouldSetNonInterpolableAttributes)
+			{
+				Attribute->SetValue(OutPoint.MetadataEntry, OtherAttribute, InPoint.MetadataEntry);
+			}
+		}
+	}
+}
+
+void UPCGMetadata::SetAttributes(const FPCGPoint& InPoint, const UPCGMetadata* InMetadata, FPCGPoint& OutPoint)
+{
+	if (!InMetadata)
+	{
+		return;
+	}
+
+	if (OutPoint.MetadataEntry == PCGInvalidEntryKey)
+	{
+		if (InPoint.MetadataEntry != PCGInvalidEntryKey && Parent == InMetadata)
+		{
+			OutPoint.MetadataEntry = AddEntry(InPoint.MetadataEntry);
+
+			// Early out; we don't need to do anything else at this point
+			return;
+		}
+		else
+		{
+			OutPoint.MetadataEntry = AddEntry();
+		}
+	}
+
+	//METADATA TODO: This isn't great, we're allocating memory for no reason here
+	AttributeLock.ReadLock();
+	TSet<FName> KnownAttributeNames = AttributeNames;
+	AttributeLock.ReadUnlock();
+
+	for(const FName& AttributeName : KnownAttributeNames)
+	{
+		const FPCGMetadataAttributeBase* ConstAttribute = GetConstAttribute(AttributeName);
+
+		if (const FPCGMetadataAttributeBase* OtherAttribute = InMetadata->GetConstAttribute(AttributeName))
+		{
+			if (OtherAttribute->GetTypeId() != ConstAttribute->GetTypeId())
+			{
+				UE_LOG(LogPCG, Error, TEXT("Metadata type mismatch with attribute %s"), *AttributeName.ToString());
+				continue;
+			}
+
+			FPCGMetadataAttributeBase* Attribute = GetMutableAttribute(AttributeName);
+			Attribute->SetValue(OutPoint.MetadataEntry, OtherAttribute, InPoint.MetadataEntry);
+		}
+	}
+}
+
+void UPCGMetadata::SetAttributes(const TArrayView<const FPCGPoint>& InPoints, const UPCGMetadata* InMetadata, const TArrayView<FPCGPoint>& OutPoints)
+{
+	if (!InMetadata)
+	{
+		return;
+	}
+
+	check(InPoints.Num() == OutPoints.Num());
+
+	for (int32 PointIndex = 0; PointIndex < InPoints.Num(); ++PointIndex)
+	{
+		const FPCGPoint& InPoint = InPoints[PointIndex];
+		FPCGPoint& OutPoint = OutPoints[PointIndex];
+		if (OutPoint.MetadataEntry == PCGInvalidEntryKey)
+		{
+			if (InPoint.MetadataEntry != PCGInvalidEntryKey && Parent == InMetadata)
+			{
+				OutPoint.MetadataEntry = AddEntry(InPoint.MetadataEntry);
+			}
+			else
+			{
+				OutPoint.MetadataEntry = AddEntry();
+			}
+		}
+	}
+
+	//METADATA TODO: This isn't great, we're allocating memory for no reason here
+	AttributeLock.ReadLock();
+	TSet<FName> KnownAttributeNames = AttributeNames;
+	AttributeLock.ReadUnlock();
+
+	for(const FName& AttributeName : KnownAttributeNames)
+	{
+		const FPCGMetadataAttributeBase* ConstAttribute = GetConstAttribute(AttributeName);
+
+		if (const FPCGMetadataAttributeBase* OtherAttribute = InMetadata->GetConstAttribute(AttributeName))
+		{
+			if (OtherAttribute->GetTypeId() != ConstAttribute->GetTypeId())
+			{
+				UE_LOG(LogPCG, Error, TEXT("Metadata type mismatch with attribute %s"), *AttributeName.ToString());
+				continue;
+			}
+
+			FPCGMetadataAttributeBase* Attribute = GetMutableAttribute(AttributeName);
+
+			for (int32 PointIndex = 0; PointIndex < InPoints.Num(); ++PointIndex)
+			{
+				const FPCGPoint& InPoint = InPoints[PointIndex];
+				FPCGPoint& OutPoint = OutPoints[PointIndex];
+
+				Attribute->SetValue(OutPoint.MetadataEntry, OtherAttribute, InPoint.MetadataEntry);
+			}
+		}
+	}
+}
\ No newline at end of file
diff --git a/Engine/Plugins/Experimental/PCG/Source/PCG/Private/Metadata/PCGMetadataAccessor.cpp b/Engine/Plugins/Experimental/PCG/Source/PCG/Private/Metadata/PCGMetadataAccessor.cpp
new file mode 100644
--- /dev/null
+++ b/Engine/Plugins/Experimental/PCG/Source/PCG/Private/Metadata/PCGMetadataAccessor.cpp
@@ -0,0 +1,169 @@
+// Copyright Epic Games, Inc. All Rights Reserved.
+
+#include "Metadata/PCGMetadataAccessor.h"
+#include "Metadata/PCGMetadata.h"
+#include "Metadata/PCGMetadataAttributeTpl.h"
+#include "PCGModule.h"
+
+template<typename T>
+T UPCGMetadataAccessorHelpers::GetAttribute(const FPCGPoint& Point, UPCGMetadata* Metadata, FName AttributeName)
+{
+	if (!Metadata)
+	{
+		UE_LOG(LogPCG, Error, TEXT("Source data has no metadata"));
+		return T{};
+	}
+
+	const FPCGMetadataAttribute<T>* Attribute = static_cast<const FPCGMetadataAttribute<T>*>(Metadata->GetConstAttribute(AttributeName));
+	if (Attribute && Attribute->GetTypeId() == PCG::Private::MetadataTypes<T>::Id)
+	{
+		return Attribute->GetValueFromItemKey(Point.MetadataEntry);
+	}
+	else if (Attribute)
+	{
+		UE_LOG(LogPCG, Error, TEXT("Attribute %s does not have the matching type"), *AttributeName.ToString());
+		return T{};
+	}
+	else
+	{
+		UE_LOG(LogPCG, Error, TEXT("Invalid attribute name (%s)"), *AttributeName.ToString());
+		return T{};
+	}
+}
+
+void UPCGMetadataAccessorHelpers::InitializeMetadata(FPCGPoint& Point, UPCGMetadata* Metadata)
+{
+	Point.MetadataEntry = Metadata ? Metadata->AddEntry() : PCGInvalidEntryKey;
+}
+
+void UPCGMetadataAccessorHelpers::InitializeMetadata(FPCGPoint& Point, UPCGMetadata* Metadata, const FPCGPoint& ParentPoint)
+{
+	Point.MetadataEntry = Metadata ? Metadata->AddEntry(ParentPoint.MetadataEntry) : PCGInvalidEntryKey;
+}
+
+void UPCGMetadataAccessorHelpers::InitializeMetadata(FPCGPoint& Point, UPCGMetadata* Metadata, const FPCGPoint& ParentPoint, const UPCGMetadata* ParentMetadata)
+{
+	Point.MetadataEntry = Metadata ? (Metadata->HasParent(ParentMetadata) ? Metadata->AddEntry(ParentPoint.MetadataEntry) : Metadata->AddEntry()) : PCGInvalidEntryKey;
+}
+
+template<typename T>
+void UPCGMetadataAccessorHelpers::SetAttribute(FPCGPoint& Point, UPCGMetadata* Metadata, FName AttributeName, const T& Value)
+{
+	if (!Metadata)
+	{
+		UE_LOG(LogPCG, Error, TEXT("Data has no metadata; cannot write value in attribute"));
+		return;
+	}
+
+	if (Point.MetadataEntry == PCGInvalidEntryKey)
+	{
+		// Default initialization
+		Point.MetadataEntry = Metadata->AddEntry();
+	}
+
+	if (Point.MetadataEntry == PCGInvalidEntryKey)
+	{
+		UE_LOG(LogPCG, Error, TEXT("Metadata item has no entry, therefore can't set values"));
+		return;
+	}
+
+	FPCGMetadataAttribute<T>* Attribute = static_cast<FPCGMetadataAttribute<T>*>(Metadata->GetMutableAttribute(AttributeName));
+	if (Attribute && Attribute->GetTypeId() == PCG::Private::MetadataTypes<T>::Id)
+	{
+		Attribute->SetValue(Point.MetadataEntry, Value);
+	}
+	else if (Attribute)
+	{
+		UE_LOG(LogPCG, Error, TEXT("Attribute %s does not have the matching type"), *AttributeName.ToString());
+	}
+	else
+	{
+		UE_LOG(LogPCG, Error, TEXT("Invalid attribute name (%s)"), *AttributeName.ToString());
+	}
+}
+
+float UPCGMetadataAccessorHelpers::GetFloatAttribute(const FPCGPoint& Point, UPCGMetadata* Metadata, FName AttributeName)
+{
+	return GetAttribute<float>(Point, Metadata, AttributeName);
+}
+
+void UPCGMetadataAccessorHelpers::SetFloatAttribute(FPCGPoint& Point, UPCGMetadata* Metadata, FName AttributeName, float Value)
+{
+	SetAttribute(Point, Metadata, AttributeName, Value);
+}
+
+FVector UPCGMetadataAccessorHelpers::GetVectorAttribute(const FPCGPoint& Point, UPCGMetadata* Metadata, FName AttributeName)
+{
+	return GetAttribute<FVector>(Point, Metadata, AttributeName);
+}
+
+void UPCGMetadataAccessorHelpers::SetVectorAttribute(FPCGPoint& Point, UPCGMetadata* Metadata, FName AttributeName, const FVector& Value)
+{
+	SetAttribute(Point, Metadata, AttributeName, Value);
+}
+
+FVector4 UPCGMetadataAccessorHelpers::GetVector4Attribute(const FPCGPoint& Point, UPCGMetadata* Metadata, FName AttributeName)
+{
+	return GetAttribute<FVector4>(Point, Metadata, AttributeName);
+}
+
+void UPCGMetadataAccessorHelpers::SetVector4Attribute(FPCGPoint& Point, UPCGMetadata* Metadata, FName AttributeName, const FVector4& Value)
+{
+	SetAttribute(Point, Metadata, AttributeName, Value);
+}
+
+FQuat UPCGMetadataAccessorHelpers::GetQuatAttribute(const FPCGPoint& Point, UPCGMetadata* Metadata, FName AttributeName)
+{
+	return GetAttribute<FQuat>(Point, Metadata, AttributeName);
+}
+
+void UPCGMetadataAccessorHelpers::SetQuatAttribute(FPCGPoint& Point, UPCGMetadata* Metadata, FName AttributeName, const FQuat& Value)
+{
+	SetAttribute(Point, Metadata, AttributeName, Value);
+}
+
+FTransform UPCGMetadataAccessorHelpers::GetTransformAttribute(const FPCGPoint& Point, UPCGMetadata* Metadata, FName AttributeName)
+{
+	return GetAttribute<FTransform>(Point, Metadata, AttributeName);
+}
+
+void UPCGMetadataAccessorHelpers::SetTransformAttribute(FPCGPoint& Point, UPCGMetadata* Metadata, FName AttributeName, const FTransform& Value)
+{
+	SetAttribute(Point, Metadata, AttributeName, Value);
+}
+
+FString UPCGMetadataAccessorHelpers::GetStringAttribute(const FPCGPoint& Point, UPCGMetadata* Metadata, FName AttributeName)
+{
+	return GetAttribute<FString>(Point, Metadata, AttributeName);
+}
+
+void UPCGMetadataAccessorHelpers::SetStringAttribute(FPCGPoint& Point, UPCGMetadata* Metadata, FName AttributeName, const FString& Value)
+{
+	SetAttribute(Point, Metadata, AttributeName, Value);
+}
+
+bool UPCGMetadataAccessorHelpers::HasAttributeSet(const FPCGPoint& Point, UPCGMetadata* Metadata, FName AttributeName)
+{
+	if (!Metadata)
+	{
+		UE_LOG(LogPCG, Error, TEXT("Data has no metadata"));
+		return false;
+	}
+
+	if (!Metadata->HasAttribute(AttributeName))
+	{
+		UE_LOG(LogPCG, Error, TEXT("Metadata does not have a %s attribute"), *AttributeName.ToString());
+		return false;
+	}
+
+	// Early out: the point has no metadata entry assigned
+	if (Point.MetadataEntry == PCGInvalidEntryKey)
+	{
+		return false;
+	}
+
+	const FPCGMetadataAttributeBase* Attribute = Metadata->GetConstAttribute(AttributeName);
+	check(Attribute);
+
+	return Attribute->HasNonDefaultValue(Point.MetadataEntry);
+}
\ No newline at end of file
diff --git a/Engine/Plugins/Experimental/PCG/Source/PCG/Private/Metadata/PCGMetadataAttribute.cpp b/Engine/Plugins/Experimental/PCG/Source/PCG/Private/Metadata/PCGMetadataAttribute.cpp
new file mode 100644
--- /dev/null
+++ b/Engine/Plugins/Experimental/PCG/Source/PCG/Private/Metadata/PCGMetadataAttribute.cpp
@@ -0,0 +1,55 @@
+// Copyright Epic Games, Inc. All Rights Reserved.
+
+#include "Metadata/PCGMetadataAttribute.h"
+
+FPCGMetadataAttributeBase::FPCGMetadataAttributeBase(UPCGMetadata* InMetadata, FName InName, const FPCGMetadataAttributeBase* InParent, bool bInAllowsInterpolation)
+	: Metadata(InMetadata)
+	, Parent(InParent)
+	, bAllowsInterpolation(bInAllowsInterpolation)
+	, Name(InName)
+{
+}
+
+void FPCGMetadataAttributeBase::SetValueFromValueKey(PCGMetadataEntryKey EntryKey, PCGMetadataValueKey ValueKey)
+{
+	FWriteScopeLock ScopeLock(EntryMapLock);
+	EntryToValueMap.FindOrAdd(EntryKey) = ValueKey;
+}
+
+PCGMetadataValueKey FPCGMetadataAttributeBase::GetValueKey(PCGMetadataEntryKey EntryKey) const
+{
+	if (EntryKey == PCGInvalidEntryKey)
+	{
+		return PCGDefaultValueKey;
+	}
+
+	PCGMetadataValueKey ValueKey = PCGDefaultValueKey;
+	bool bFoundKey = false;
+
+	EntryMapLock.ReadLock();
+	if (const PCGMetadataValueKey* FoundLocalKey = EntryToValueMap.Find(EntryKey))
+	{
+		ValueKey = *FoundLocalKey;
+		bFoundKey = true;
+	}
+	EntryMapLock.ReadUnlock();
+
+	if (!bFoundKey && Parent)
+	{
+		return Parent->GetValueKey(Metadata->GetParentKey(EntryKey));
+	}
+	else
+	{
+		return ValueKey;
+	}
+}
+
+bool FPCGMetadataAttributeBase::HasNonDefaultValue(PCGMetadataEntryKey EntryKey) const
+{
+	return GetValueKey(EntryKey) != PCGDefaultValueKey;
+}
+
+void FPCGMetadataAttributeBase::ClearEntries()
+{
+	EntryToValueMap.Reset();
+}
\ No newline at end of file
diff --git a/Engine/Plugins/Experimental/PCG/Source/PCG/Private/Metadata/PCGMetadataAttributeTpl.h b/Engine/Plugins/Experimental/PCG/Source/PCG/Private/Metadata/PCGMetadataAttributeTpl.h
new file mode 100644
--- /dev/null
+++ b/Engine/Plugins/Experimental/PCG/Source/PCG/Private/Metadata/PCGMetadataAttributeTpl.h
@@ -0,0 +1,322 @@
+// Copyright Epic Games, Inc. All Rights Reserved.
+
+#pragma once
+
+#include "CoreMinimal.h"
+#include "Metadata/PCGMetadataAttribute.h"
+#include "Metadata/PCGMetadataAttributeTraits.h"
+
+template<typename T>
+class FPCGMetadataAttribute : public FPCGMetadataAttributeBase
+{
+public:
+	FPCGMetadataAttribute(UPCGMetadata* InMetadata, FName InName, const FPCGMetadataAttributeBase* InParent, const T& InDefaultValue, bool bInAllowsInterpolation)
+		: FPCGMetadataAttributeBase(InMetadata, InName, InParent, bInAllowsInterpolation)
+		, DefaultValue(InDefaultValue)
+	{
+		TypeId = PCG::Private::MetadataTypes<T>::Id;
+
+		if (GetParent())
+		{
+			ValueKeyOffset = GetParent()->GetValueKeyOffsetForChild();
+		}
+	}
+
+	const FPCGMetadataAttribute* GetParent() const { return static_cast<const FPCGMetadataAttribute*>(Parent); }
+
+	virtual FPCGMetadataAttributeBase* Copy(FName NewName, UPCGMetadata* InMetadata, bool bKeepParent, bool bCopyEntries = true, bool bCopyValues = true) const override
+	{
+		// this copies to a new attribute
+		check(!bKeepParent || Metadata->GetRoot() == InMetadata->GetRoot());
+		FPCGMetadataAttribute<T>* AttributeCopy = new FPCGMetadataAttribute<T>(InMetadata, NewName, bKeepParent ? this : nullptr, DefaultValue, bAllowsInterpolation);
+
+		if (bCopyEntries)
+		{
+			EntryMapLock.ReadLock();
+			AttributeCopy->EntryToValueMap = EntryToValueMap;
+			EntryMapLock.ReadUnlock();
+		}
+
+		if (bCopyValues)
+		{
+			ValueLock.ReadLock();
+			AttributeCopy->Values = Values;
+			AttributeCopy->ValueKeyOffset = ValueKeyOffset;
+			ValueLock.ReadUnlock();
+		}
+
+		return AttributeCopy;
+	}
+
+	virtual void SetValue(PCGMetadataEntryKey ItemKey, const FPCGMetadataAttributeBase* InAttribute, PCGMetadataEntryKey InEntryKey) override
+	{
+		if (InAttribute == this)
+		{
+			SetValueFromValueKey(ItemKey, GetValueKey(InEntryKey));
+		}
+		else if (InAttribute)
+		{
+			SetValue(ItemKey, static_cast<const FPCGMetadataAttribute<T>*>(InAttribute)->GetValueFromItemKey(InEntryKey));
+		}
+	}
+
+	virtual void SetZeroValue(PCGMetadataEntryKey ItemKey) override
+	{
+		check(ItemKey != PCGInvalidEntryKey);
+		ZeroValue(ItemKey);
+	}
+
+	virtual void AccumulateValue(PCGMetadataEntryKey ItemKey, const FPCGMetadataAttributeBase* InAttribute, PCGMetadataEntryKey InEntryKey, float Weight) override
+	{
+		check(ItemKey != PCGInvalidEntryKey);
+		Accumulate(ItemKey, InAttribute, InEntryKey, Weight);
+	}
+
+	virtual void SetValue(PCGMetadataEntryKey ItemKey, const FPCGMetadataAttributeBase* InAttributeA, PCGMetadataEntryKey InEntryKeyA, const FPCGMetadataAttributeBase* InAttributeB, PCGMetadataEntryKey InEntryKeyB, EPCGMetadataOp Op) override
+	{
+		bool bAppliedValue = false;
+
+		if (InAttributeA && InAttributeB && bAllowsInterpolation)
+		{
+			if (Op == EPCGMetadataOp::Min)
+			{
+				bAppliedValue = SetMin(ItemKey, InAttributeA, InEntryKeyA, InAttributeB, InEntryKeyB);
+			}
+			else if (Op == EPCGMetadataOp::Max)
+			{
+				bAppliedValue = SetMax(ItemKey, InAttributeA, InEntryKeyA, InAttributeB, InEntryKeyB);
+			}
+			else if (Op == EPCGMetadataOp::Sub)
+			{
+				bAppliedValue = SetSub(ItemKey, InAttributeA, InEntryKeyA, InAttributeB, InEntryKeyB);
+			}
+			else if (Op == EPCGMetadataOp::Add)
+			{
+				bAppliedValue = SetAdd(ItemKey, InAttributeA, InEntryKeyA, InAttributeB, InEntryKeyB);
+			}
+		}
+		else if (InAttributeA && InAttributeB && HasNonDefaultValue(ItemKey))
+		{
+			// In this case, the current already has a value, in which case we should not update it
+			bAppliedValue = true;
+		}
+
+		if (bAppliedValue)
+		{
+			// Nothing to do
+		}
+		else if (InAttributeA)
+		{
+			if (InAttributeA == this)
+			{
+				SetValueFromValueKey(ItemKey, GetValueKey(InEntryKeyA));
+			}
+			else
+			{
+				SetValue(ItemKey, static_cast<const FPCGMetadataAttribute<T>*>(InAttributeA)->GetValueFromItemKey(InEntryKeyA));
+			}
+		}
+		else if (InAttributeB)
+		{
+			if (InAttributeB == this)
+			{
+				SetValueFromValueKey(ItemKey, GetValueKey(InEntryKeyB));
+			}
+			else
+			{
+				SetValue(ItemKey, static_cast<const FPCGMetadataAttribute<T>*>(InAttributeB)->GetValueFromItemKey(InEntryKeyB));
+			}
+		}
+	}
+
+	PCGMetadataValueKey GetValueKeyOffsetForChild() const
+	{
+		FReadScopeLock ScopeLock(ValueLock);
+		return Values.Num() + ValueKeyOffset;
+	}
+
+	/** Adds the value, returns the value key for the given value */
+	PCGMetadataValueKey AddValue(const T& InValue)
+	{
+		PCGMetadataValueKey FoundValue = FindValue(InValue);
+
+		if (FoundValue == PCGDefaultValueKey)
+		{
+			FWriteScopeLock ScopeLock(ValueLock);
+			return Values.Add(InValue) + ValueKeyOffset;
+		}
+		else
+		{
+			return FoundValue;
+		}
+	}
+
+	void SetValue(PCGMetadataEntryKey ItemKey, const T& InValue)
+	{
+		SetValueFromValueKey(ItemKey, AddValue(InValue));
+	}
+
+	T GetValueFromItemKey(PCGMetadataEntryKey ItemKey) const
+	{
+		return GetValue(GetValueKey(ItemKey));
+	}
+
+	T GetValue(PCGMetadataValueKey ValueKey) const
+	{
+		if (ValueKey == PCGDefaultValueKey)
+		{
+			return DefaultValue;
+		}
+		else if (ValueKey >= ValueKeyOffset)
+		{
+			FReadScopeLock ScopeLock(ValueLock);
+			return Values[ValueKey - ValueKeyOffset];
+		}
+		else if (GetParent())
+		{
+			return GetParent()->GetValue(ValueKey);
+		}
+		else
+		{
+			return DefaultValue;
+		}
+	}
+
+protected:
+	/** Code related to finding values / compressing data */
+	template<typename IT = T, typename TEnableIf<PCG::Private::MetadataTraits<IT>::CompressData>::Type* = nullptr>
+	PCGMetadataValueKey FindValue(const T& InValue) const
+	{
+		PCGMetadataValueKey ParentValueKey = (GetParent() ? GetParent()->FindValue(InValue) : PCGDefaultValueKey);
+		if (ParentValueKey != PCGDefaultValueKey)
+		{
+			return ParentValueKey;
+		}
+		else
+		{
+			ValueLock.ReadLock();
+			const int32 ValueIndex = Values.FindLast(InValue);
+			ValueLock.ReadUnlock();
+
+			if (ValueIndex != INDEX_NONE)
+			{
+				return ValueIndex + ValueKeyOffset;
+			}
+			else
+			{
+				return ParentValueKey;
+			}
+		}
+	}
+
+	template<typename IT = T, typename TEnableIf<!PCG::Private::MetadataTraits<IT>::CompressData>::Type* = nullptr>
+	PCGMetadataValueKey FindValue(const T& InValue) const
+	{
+		return PCGDefaultValueKey;
+	}
+
+	/** Code related to computing compared values (min, max, sub, add) */
+	template<typename IT = T, typename TEnableIf<PCG::Private::MetadataTraits<IT>::CanMinMax>::Type* = nullptr>
+	bool SetMin(PCGMetadataEntryKey ItemKey, const FPCGMetadataAttributeBase* InAttributeA, PCGMetadataEntryKey InEntryKeyA, const FPCGMetadataAttributeBase* InAttributeB, PCGMetadataEntryKey InEntryKeyB)
+	{
+		SetValue(ItemKey, 
+			PCG::Private::MetadataTraits<IT>::Min(
+				static_cast<const FPCGMetadataAttribute<T>*>(InAttributeA)->GetValueFromItemKey(InEntryKeyA),
+				static_cast<const FPCGMetadataAttribute<T>*>(InAttributeB)->GetValueFromItemKey(InEntryKeyB)));
+
+		return true;
+	}
+
+	template<typename IT = T, typename TEnableIf<!PCG::Private::MetadataTraits<IT>::CanMinMax>::Type* = nullptr>
+	bool SetMin(PCGMetadataEntryKey ItemKey, const FPCGMetadataAttributeBase* InAttributeA, PCGMetadataEntryKey InEntryKeyA, const FPCGMetadataAttributeBase* InAttributeB, PCGMetadataEntryKey InEntryKeyB)
+	{
+		return false;
+	}
+
+	template<typename IT = T, typename TEnableIf<PCG::Private::MetadataTraits<IT>::CanMinMax>::Type* = nullptr>
+	bool SetMax(PCGMetadataEntryKey ItemKey, const FPCGMetadataAttributeBase* InAttributeA, PCGMetadataEntryKey InEntryKeyA, const FPCGMetadataAttributeBase* InAttributeB, PCGMetadataEntryKey InEntryKeyB)
+	{
+		SetValue(ItemKey,
+			PCG::Private::MetadataTraits<IT>::Max(
+				static_cast<const FPCGMetadataAttribute<T>*>(InAttributeA)->GetValueFromItemKey(InEntryKeyA),
+				static_cast<const FPCGMetadataAttribute<T>*>(InAttributeB)->GetValueFromItemKey(InEntryKeyB)));
+
+		return true;
+	}
+
+	template<typename IT = T, typename TEnableIf<!PCG::Private::MetadataTraits<IT>::CanMinMax>::Type* = nullptr>
+	bool SetMax(PCGMetadataEntryKey ItemKey, const FPCGMetadataAttributeBase* InAttributeA, PCGMetadataEntryKey InEntryKeyA, const FPCGMetadataAttributeBase* InAttributeB, PCGMetadataEntryKey InEntryKeyB)
+	{
+		return false;
+	}
+
+	template<typename IT = T, typename TEnableIf<PCG::Private::MetadataTraits<IT>::CanSubAdd>::Type* = nullptr>
+	bool SetSub(PCGMetadataEntryKey ItemKey, const FPCGMetadataAttributeBase* InAttributeA, PCGMetadataEntryKey InEntryKeyA, const FPCGMetadataAttributeBase* InAttributeB, PCGMetadataEntryKey InEntryKeyB)
+	{
+		SetValue(ItemKey,
+			PCG::Private::MetadataTraits<IT>::Sub(
+				static_cast<const FPCGMetadataAttribute<T>*>(InAttributeA)->GetValueFromItemKey(InEntryKeyA),
+				static_cast<const FPCGMetadataAttribute<T>*>(InAttributeB)->GetValueFromItemKey(InEntryKeyB)));
+
+		return true;
+	}
+
+	template<typename IT = T, typename TEnableIf<!PCG::Private::MetadataTraits<IT>::CanSubAdd>::Type* = nullptr>
+	bool SetSub(PCGMetadataEntryKey ItemKey, const FPCGMetadataAttributeBase* InAttributeA, PCGMetadataEntryKey InEntryKeyA, const FPCGMetadataAttributeBase* InAttributeB, PCGMetadataEntryKey InEntryKeyB)
+	{
+		return false;
+	}
+
+	template<typename IT = T, typename TEnableIf<PCG::Private::MetadataTraits<IT>::CanSubAdd>::Type* = nullptr>
+	bool SetAdd(PCGMetadataEntryKey ItemKey, const FPCGMetadataAttributeBase* InAttributeA, PCGMetadataEntryKey InEntryKeyA, const FPCGMetadataAttributeBase* InAttributeB, PCGMetadataEntryKey InEntryKeyB)
+	{
+		SetValue(ItemKey,
+			PCG::Private::MetadataTraits<IT>::Add(
+				static_cast<const FPCGMetadataAttribute<T>*>(InAttributeA)->GetValueFromItemKey(InEntryKeyA),
+				static_cast<const FPCGMetadataAttribute<T>*>(InAttributeB)->GetValueFromItemKey(InEntryKeyB)));
+
+		return true;
+	}
+
+	template<typename IT = T, typename TEnableIf<!PCG::Private::MetadataTraits<IT>::CanSubAdd>::Type* = nullptr>
+	bool SetAdd(PCGMetadataEntryKey ItemKey, const FPCGMetadataAttributeBase* InAttributeA, PCGMetadataEntryKey InEntryKeyA, const FPCGMetadataAttributeBase* InAttributeB, PCGMetadataEntryKey InEntryKeyB)
+	{
+		return false;
+	}
+
+	/** Weighted/interpolated values related code */
+	template<typename IT = T, typename TEnableIf<PCG::Private::MetadataTraits<IT>::CanInterpolate>::Type* = nullptr>
+	void ZeroValue(PCGMetadataEntryKey ItemKey)
+	{
+		SetValue(ItemKey, PCG::Private::MetadataTraits<IT>::ZeroValue());
+	}
+
+	template<typename IT = T, typename TEnableIf<!PCG::Private::MetadataTraits<IT>::CanInterpolate>::Type* = nullptr>
+	void ZeroValue(PCGMetadataEntryKey ItemKey)
+	{
+		// Intentionally empty
+	}
+
+	template<typename IT = T, typename TEnableIf<PCG::Private::MetadataTraits<IT>::CanInterpolate>::Type* = nullptr>
+	void Accumulate(PCGMetadataEntryKey ItemKey, const FPCGMetadataAttributeBase* InAttribute, PCGMetadataEntryKey InEntryKey, float Weight)
+	{
+		SetValue(ItemKey,
+			PCG::Private::MetadataTraits<IT>::WeightedSum(
+				GetValueFromItemKey(ItemKey),
+				static_cast<const FPCGMetadataAttribute<T>*>(InAttribute)->GetValueFromItemKey(InEntryKey),
+				Weight));
+	}
+
+	template<typename IT = T, typename TEnableIf<!PCG::Private::MetadataTraits<IT>::CanInterpolate>::Type* = nullptr>
+	void Accumulate(PCGMetadataEntryKey ItemKey, const FPCGMetadataAttributeBase* InAttribute, PCGMetadataEntryKey InEntryKey, float Weight)
+	{
+		// Empty on purpose
+	}
+
+
+protected:
+	mutable FRWLock ValueLock;
+	TArray<T> Values;
+	T DefaultValue;
+	PCGMetadataValueKey ValueKeyOffset = 0;
+};
\ No newline at end of file
diff --git a/Engine/Plugins/Experimental/PCG/Source/PCG/Private/PCGComponent.cpp b/Engine/Plugins/Experimental/PCG/Source/PCG/Private/PCGComponent.cpp
--- a/Engine/Plugins/Experimental/PCG/Source/PCG/Private/PCGComponent.cpp
+++ b/Engine/Plugins/Experimental/PCG/Source/PCG/Private/PCGComponent.cpp
@@ -1099,7 +1099,7 @@ UPCGData* UPCGComponent::CreateActorPCGData(AActor* Actor)
 		else // Default behavior on unknown actors
 		{
 			UPCGPointData* Data = NewObject<UPCGPointData>(this);
-			Data->Initialize(Actor);
+			Data->InitializeFromActor(Actor);
 			return Data;
 		}
 	}
diff --git a/Engine/Plugins/Experimental/PCG/Source/PCG/Public/Data/PCGDifferenceData.h b/Engine/Plugins/Experimental/PCG/Source/PCG/Public/Data/PCGDifferenceData.h
--- a/Engine/Plugins/Experimental/PCG/Source/PCG/Public/Data/PCGDifferenceData.h
+++ b/Engine/Plugins/Experimental/PCG/Source/PCG/Public/Data/PCGDifferenceData.h
@@ -41,6 +41,7 @@ public:
 	virtual FBox GetStrictBounds() const override;
 	virtual float GetDensityAtPosition(const FVector& InPosition) const override;
 	virtual FVector TransformPosition(const FVector& InPosition) const override;
+	virtual bool GetPointAtPosition(const FVector& InPosition, FPCGPoint& OutPoint, UPCGMetadata* OutMetadata) const override;
 	virtual FPCGPoint TransformPoint(const FPCGPoint& InPoint) const override;
 	virtual bool HasNonTrivialTransform() const override;
 	//~End UPCGSpatialData interface
diff --git a/Engine/Plugins/Experimental/PCG/Source/PCG/Public/Data/PCGIntersectionData.h b/Engine/Plugins/Experimental/PCG/Source/PCG/Public/Data/PCGIntersectionData.h
--- a/Engine/Plugins/Experimental/PCG/Source/PCG/Public/Data/PCGIntersectionData.h
+++ b/Engine/Plugins/Experimental/PCG/Source/PCG/Public/Data/PCGIntersectionData.h
@@ -30,6 +30,7 @@ public:
 	virtual FBox GetStrictBounds() const override;
 	virtual float GetDensityAtPosition(const FVector& InPosition) const override;
 	virtual FVector TransformPosition(const FVector& InPosition) const override;
+	virtual bool GetPointAtPosition(const FVector& InPosition, FPCGPoint& OutPoint, UPCGMetadata* OutMetadata) const override;
 	virtual FPCGPoint TransformPoint(const FPCGPoint& InPoint) const override;
 	virtual bool HasNonTrivialTransform() const override;
 	//~End UPCGSpatialData interface
diff --git a/Engine/Plugins/Experimental/PCG/Source/PCG/Public/Data/PCGPointData.h b/Engine/Plugins/Experimental/PCG/Source/PCG/Public/Data/PCGPointData.h
--- a/Engine/Plugins/Experimental/PCG/Source/PCG/Public/Data/PCGPointData.h
+++ b/Engine/Plugins/Experimental/PCG/Source/PCG/Public/Data/PCGPointData.h
@@ -59,16 +59,20 @@ class PCG_API UPCGPointData : public UPCGSpatialData
 public:
 	typedef TOctree2<FPCGPointRef, FPCGPointRefSemantics> PointOctree;
 
-	void Initialize(AActor* InActor);
-
 	// ~Begin UPCGSpatialData interface
 	virtual int GetDimension() const override { return 0; }
 	virtual FBox GetBounds() const override;
 	virtual float GetDensityAtPosition(const FVector& InPosition) const override;
 	virtual const UPCGPointData* ToPointData(FPCGContext* Context) const { return this; }
+	virtual bool GetPointAtPosition(const FVector& InPosition, FPCGPoint& OutPoint, UPCGMetadata* OutMetadata) const;
 	virtual FPCGPoint TransformPoint(const FPCGPoint& InPoint) const;
 	// ~End UPCGSpatialData interface
 
+	void InitializeFromActor(AActor* InActor);
+
+	UFUNCTION(BlueprintCallable, Category = SpatialData)
+	void InitializeFromData(const UPCGSpatialData* InSource, const UPCGMetadata* InMetadataParentOverride = nullptr);
+
 	UFUNCTION(BlueprintCallable, Category = SpatialData)
 	const TArray<FPCGPoint>& GetPoints() const { return Points; }
 
diff --git a/Engine/Plugins/Experimental/PCG/Source/PCG/Public/Data/PCGProjectionData.h b/Engine/Plugins/Experimental/PCG/Source/PCG/Public/Data/PCGProjectionData.h
--- a/Engine/Plugins/Experimental/PCG/Source/PCG/Public/Data/PCGProjectionData.h
+++ b/Engine/Plugins/Experimental/PCG/Source/PCG/Public/Data/PCGProjectionData.h
@@ -24,6 +24,7 @@ public:
 	virtual FVector GetNormal() const override;
 	virtual float GetDensityAtPosition(const FVector& InPosition) const override;
 	virtual FVector TransformPosition(const FVector& InPosition) const override;
+	virtual bool GetPointAtPosition(const FVector& InPosition, FPCGPoint& OutPoint, UPCGMetadata* OutMetadata) const override;
 	virtual FPCGPoint TransformPoint(const FPCGPoint& InPoint) const override;
 	virtual bool HasNonTrivialTransform() const override;
 	//~End UPCGSpatialData interface
diff --git a/Engine/Plugins/Experimental/PCG/Source/PCG/Public/Data/PCGSpatialData.h b/Engine/Plugins/Experimental/PCG/Source/PCG/Public/Data/PCGSpatialData.h
--- a/Engine/Plugins/Experimental/PCG/Source/PCG/Public/Data/PCGSpatialData.h
+++ b/Engine/Plugins/Experimental/PCG/Source/PCG/Public/Data/PCGSpatialData.h
@@ -6,6 +6,7 @@
 #include "PCGContext.h"
 #include "PCGData.h"
 #include "PCGPoint.h"
+#include "Metadata/PCGMetadata.h"
 
 #include "PCGSpatialData.generated.h"
 
@@ -69,6 +70,13 @@ public:
 	UFUNCTION(BlueprintCallable, Category = SpatialData)
 	virtual FPCGPoint TransformPoint(const FPCGPoint& InPoint) const;
 
+	virtual bool GetPointAtPosition(const FVector& InPosition, FPCGPoint& OutPoint, UPCGMetadata* OutMetadata) const
+	{
+		OutPoint.Transform = FTransform(InPosition);
+		OutPoint = TransformPoint(OutPoint);
+		return OutPoint.Density > 0;
+	}
+
 	/** Returns true if the data has a non-trivial transform */
 	UFUNCTION(BlueprintCallable, Category = SpatialData)
 	virtual bool HasNonTrivialTransform() const { return false; }
@@ -88,6 +96,12 @@ public:
 	UFUNCTION(BlueprintCallable, Category = SpatialData)
 	virtual UPCGDifferenceData* Subtract(const UPCGSpatialData* InOther) const;
 
+	UFUNCTION(BlueprintCallable, Category = Metadata)
+	const UPCGMetadata* ConstMetadata() const { return Metadata; }
+
+	UFUNCTION(BlueprintCallable, Category = Metadata)
+	UPCGMetadata* MutableMetadata() { return Metadata; }
+
 	UPROPERTY(BlueprintReadWrite, EditAnywhere, Category = Data)
 	AActor* TargetActor = nullptr;
 
@@ -95,6 +109,10 @@ public:
 	UPROPERTY(Transient, BlueprintReadWrite, EditAnywhere, Category = Debug)
 	bool bKeepZeroDensityPoints = false;
 #endif
+
+	// Not accessible through blueprint to make sure the constness is preserved
+	UPROPERTY(VisibleAnywhere, Category = Metadata)
+	TObjectPtr<UPCGMetadata> Metadata = nullptr;
 };
 
 UCLASS(Abstract, ClassGroup = (Procedural))
diff --git a/Engine/Plugins/Experimental/PCG/Source/PCG/Public/Data/PCGUnionData.h b/Engine/Plugins/Experimental/PCG/Source/PCG/Public/Data/PCGUnionData.h
--- a/Engine/Plugins/Experimental/PCG/Source/PCG/Public/Data/PCGUnionData.h
+++ b/Engine/Plugins/Experimental/PCG/Source/PCG/Public/Data/PCGUnionData.h
@@ -43,6 +43,7 @@ public:
 	virtual FBox GetStrictBounds() const override;
 	virtual float GetDensityAtPosition(const FVector& InPosition) const override;
 	virtual FVector TransformPosition(const FVector& InPosition) const override;
+	virtual bool GetPointAtPosition(const FVector& InPosition, FPCGPoint& OutPoint, UPCGMetadata* OutMetadata) const override;
 	virtual FPCGPoint TransformPoint(const FPCGPoint& InPoint) const override;
 	virtual bool HasNonTrivialTransform() const override;
 	//~End UPCGSpatialData interface
diff --git a/Engine/Plugins/Experimental/PCG/Source/PCG/Public/Elements/PCGExecuteBlueprint.h b/Engine/Plugins/Experimental/PCG/Source/PCG/Public/Elements/PCGExecuteBlueprint.h
--- a/Engine/Plugins/Experimental/PCG/Source/PCG/Public/Elements/PCGExecuteBlueprint.h
+++ b/Engine/Plugins/Experimental/PCG/Source/PCG/Public/Elements/PCGExecuteBlueprint.h
@@ -43,17 +43,17 @@ public:
 	void Execute(const FPCGDataCollection& Input, FPCGDataCollection& Output) const;
 
 	UFUNCTION(BlueprintImplementableEvent, Category = Execution)
-	bool PointLoopBody(const UPCGPointData* InData, const FPCGPoint& InPoint, FPCGPoint& OutPoint) const;
+	bool PointLoopBody(FPCGContext& InContext, const UPCGPointData* InData, const FPCGPoint& InPoint, const UObject* InCustomObject, FPCGPoint& OutPoint, UPCGMetadata* OutMetadata) const;
 
 	UFUNCTION(BlueprintImplementableEvent, Category = Execution)
-	bool PointPairLoopBody(const UPCGPointData* InA, const UPCGPointData* InB, const FPCGPoint& InPointA, const FPCGPoint& InPointB, FPCGPoint& OutPoint) const;
+	bool PointPairLoopBody(FPCGContext& InContext, const UPCGPointData* InA, const UPCGPointData* InB, const FPCGPoint& InPointA, const FPCGPoint& InPointB, const UObject* InCustomObject, FPCGPoint& OutPoint, UPCGMetadata* OutMetadata) const;
 
 	/** Calls the LoopBody function on all points */
 	UFUNCTION(BlueprintCallable, Category = Execution, meta = (HideSelfPin = "true"))
-	void LoopOnPoints(UPARAM(ref) FPCGContext& InContext, const UPCGPointData* InData, UPCGPointData*& OutData) const;
+	void LoopOnPoints(UPARAM(ref) FPCGContext& InContext, const UPCGPointData* InData, UPCGPointData*& OutData, const UObject* InCustomObject = nullptr) const;
 
 	UFUNCTION(BlueprintCallable, Category = Execution, meta = (HideSelfPin = "true"))
-	void LoopOnPointPairs(UPARAM(ref) FPCGContext& InContext, const UPCGPointData* InA, const UPCGPointData* InB, UPCGPointData*& OutData) const;
+	void LoopOnPointPairs(UPARAM(ref) FPCGContext& InContext, const UPCGPointData* InA, const UPCGPointData* InB, UPCGPointData*& OutData, const UObject* InCustomObject = nullptr) const;
 
 	/** Called after object creation to setup the object callbacks */
 	void Initialize();
diff --git a/Engine/Plugins/Experimental/PCG/Source/PCG/Public/Elements/PCGMetadataElement.h b/Engine/Plugins/Experimental/PCG/Source/PCG/Public/Elements/PCGMetadataElement.h
new file mode 100644
--- /dev/null
+++ b/Engine/Plugins/Experimental/PCG/Source/PCG/Public/Elements/PCGMetadataElement.h
@@ -0,0 +1,60 @@
+// Copyright Epic Games, Inc. All Rights Reserved.
+
+#pragma once
+
+#include "CoreMinimal.h"
+#include "PCGSettings.h"
+#include "PCGElement.h"
+
+#include "PCGMetadataElement.generated.h"
+
+UENUM()
+enum class EPCGPointProperties : uint8
+{
+	Density,
+	Extents,
+	Color
+};
+
+UENUM()
+enum class EPCGMetadataOperationTarget : uint8
+{
+	PropertyToAttribute,
+	AttributeToProperty,
+	AttributeToAttribute
+};
+
+UCLASS(BlueprintType, ClassGroup = (Procedural))
+class PCG_API UPCGMetadataOperationSettings : public UPCGSettings
+{
+	GENERATED_BODY()
+
+public:
+	//~Begin UPCGSettings interface
+#if WITH_EDITOR
+	virtual FName GetDefaultNodeName() const override { return FName(TEXT("MetadataOperationNode")); }
+#endif
+
+protected:
+	virtual FPCGElementPtr CreateElement() const override;
+	//~End UPCGSettings interface
+
+public:
+	UPROPERTY(BlueprintReadWrite, EditAnywhere, Category = Settings, meta = (EditCondition = "Target==EPCGMetadataOperationTarget::AttributeToProperty||Target==EPCGMetadataOperationTarget::AttributeToAttribute"))
+	FName SourceAttribute = NAME_None;
+
+	UPROPERTY(BlueprintReadWrite, EditAnywhere, Category = Settings)
+	EPCGPointProperties PointProperty = EPCGPointProperties::Density;
+
+	UPROPERTY(BlueprintReadWrite, EditAnywhere, Category = Settings, meta = (EditCondition="Target==EPCGMetadataOperationTarget::PropertyToAttribute||Target==EPCGMetadataOperationTarget::AttributeToAttribute"))
+	FName DestinationAttribute = NAME_None;
+
+	UPROPERTY(BlueprintReadWrite, EditAnywhere, Category = Settings)
+	EPCGMetadataOperationTarget Target = EPCGMetadataOperationTarget::PropertyToAttribute;
+};
+
+class FPCGMetadataOperationElement : public FSimpleTypedPCGElement<UPCGMetadataOperationSettings>
+{
+protected:
+	virtual bool ExecuteInternal(FPCGContext* Context) const override;
+};
\ No newline at end of file
diff --git a/Engine/Plugins/Experimental/PCG/Source/PCG/Public/Metadata/PCGMetadata.h b/Engine/Plugins/Experimental/PCG/Source/PCG/Public/Metadata/PCGMetadata.h
new file mode 100644
--- /dev/null
+++ b/Engine/Plugins/Experimental/PCG/Source/PCG/Public/Metadata/PCGMetadata.h
@@ -0,0 +1,141 @@
+// Copyright Epic Games, Inc. All Rights Reserved.
+
+#pragma once
+
+#include "CoreMinimal.h"
+#include "PCGMetadataCommon.h"
+#include "PCGMetadataAccessor.h"
+
+#include "PCGMetadata.generated.h"
+
+class FPCGMetadataAttributeBase;
+
+UENUM()
+enum class EPCGMetadataOp : uint8
+{
+	Min,
+	Max,
+	Sub,
+	Add
+};
+
+UCLASS(BlueprintType)
+class PCG_API UPCGMetadata : public UObject
+{
+	GENERATED_BODY()
+
+public:
+
+	//~ Begin UObject interface
+	virtual void Serialize(FArchive& Ar) override;
+	virtual void BeginDestroy() override;
+	//~ End UObject interface
+
+	/** Initializes the metadata from a parent metadata, if any (can be null) */
+	UFUNCTION(BlueprintCallable, Category = "PCG|Metadata")
+	void Initialize(const UPCGMetadata* InParent);
+
+	/** Creates streams for another metadata if they are not currently present */
+	void AddAttributes(const UPCGMetadata* InOther);
+
+	/** Returns this metadata's parent */
+	const UPCGMetadata* GetParent() const { return Parent.Get(); }
+	const UPCGMetadata* GetRoot() const;
+	bool HasParent(const UPCGMetadata* InTentativeParent) const;
+
+	/** Create new streams */
+	UFUNCTION(BlueprintCallable, Category = "PCG|Metadata")
+	void CreateFloatAttribute(FName AttributeName, float DefaultValue, bool bAllowsInterpolation, bool bOverrideParent);
+
+	UFUNCTION(BlueprintCallable, Category = "PCG|Metadata")
+	void CreateVectorAttribute(FName AttributeName, const FVector& DefaultValue, bool bAllowsInterpolation, bool bOverrideParent);
+
+	UFUNCTION(BlueprintCallable, Category = "PCG|Metadata")
+	void CreateVector4Attribute(FName AttributeName, const FVector4& DefaultValue, bool bAllowsInterpolation, bool bOverrideParent);
+
+	UFUNCTION(BlueprintCallable, Category = "PCG|Metadata")
+	void CreateQuatAttribute(FName AttributeName, const FQuat& DefaultValue, bool bAllowsInterpolation, bool bOverrideParent);
+
+	UFUNCTION(BlueprintCallable, Category = "PCG|Metadata")
+	void CreateTransformAttribute(FName AttributeName, const FTransform& DefaultValue, bool bAllowsInterpolation, bool bOverrideParent);
+
+	UFUNCTION(BlueprintCallable, Category = "PCG|Metadata")
+	void CreateStringAttribute(FName AttributeName, const FString& DefaultValue, bool bAllowsInterpolation, bool bOverrideParent);
+
+	/** Get attributes */
+	FPCGMetadataAttributeBase* GetMutableAttribute(FName AttributeName);
+	const FPCGMetadataAttributeBase* GetConstAttribute(FName AttributeName) const;
+	bool HasAttribute(FName AttributeName) const;
+
+	/** Delete/Hide attribute */
+	// Due to stream inheriting, we might want to consider "hiding" parent stream and deleting local streams only
+	void DeleteAttribute(FName AttributeName);
+
+	/** Copy attribute */
+	UFUNCTION(BlueprintCallable, Category = "PCG|Metadata")
+	void CopyAttribute(FName AttributeToCopy, FName NewAttributeName, bool bKeepParent = true);
+
+	/** Rename attribute */
+	UFUNCTION(BlueprintCallable, Category = "PCG|Metadata")
+	void RenameAttribute(FName AttributeToRename, FName NewAttributeName);
+
+	/** Clear/Reinit attribute */
+	UFUNCTION(BlueprintCallable, Category = "PCG|Metadata")
+	void ClearAttribute(FName AttributeToClear);
+
+	/** Adds a unique entry key to the metadata */
+	UFUNCTION(BlueprintCallable, Category = "PCG|Metadata")
+	int64 AddEntry(int64 ParentEntryKey = -1);
+
+	/** Metadata chaining mechanism */
+	PCGMetadataEntryKey GetParentKey(PCGMetadataEntryKey LocalItemKey) const;
+
+	/** Attributes operations */
+	void MergeAttributes(const FPCGPoint& InPointA, const FPCGPoint& InPointB, FPCGPoint& OutPoint, EPCGMetadataOp Op)
+	{
+		MergeAttributes(InPointA, this, InPointB, this, OutPoint, Op);
+	}
+
+	void MergeAttributes(const FPCGPoint& InPointA, const UPCGMetadata* InMetadataA, const FPCGPoint& InPointB, const UPCGMetadata* InMetadataB, FPCGPoint& OutPoint, EPCGMetadataOp Op);
+
+	void ResetWeightedAttributes(FPCGPoint& OutPoint);
+	void AccumulateWeightedAttributes(const FPCGPoint& InPoint, const UPCGMetadata* InMetadata, float Weight, bool bSetNonInterpolableAttributes, FPCGPoint& OutPoint);
+
+	void SetAttributes(const FPCGPoint& InPoint, const UPCGMetadata* InMetadata, FPCGPoint& OutPoint);
+	void SetAttributes(const TArrayView<const FPCGPoint>& InPoints, const UPCGMetadata* InMetadata, const TArrayView<FPCGPoint>& OutPoints);
+
+protected:
+	template<typename T>
+	FPCGMetadataAttributeBase* CreateAttribute(FName AttributeName, const T& DefaultValue, bool bAllowsInterpolation, bool bOverrideParent);
+
+	FPCGMetadataAttributeBase* CopyAttribute(FName AttributeToCopy, FName NewAttributeName, bool bKeepParent, bool bCopyEntries, bool bCopyValues);
+	FPCGMetadataAttributeBase* CopyAttribute(const FPCGMetadataAttributeBase* OriginalAttribute, FName NewAttributeName, bool bKeepParent, bool bCopyEntries, bool bCopyValues);
+
+	bool ParentHasAttribute(FName AttributeName) const;
+	void GetAllAttributeNames(TSet<FName>& AttributeNames) const;
+	int64 GetItemCountForChild() const;
+
+	void AddAttributeInternal(FName AttributeName, FPCGMetadataAttributeBase* Attribute);
+	void RemoveAttributeInternal(FName AttributeName);
+
+	UPROPERTY()
+	TSoftObjectPtr<const UPCGMetadata> Parent;
+
+	// Set of parents kept for streams relationship and GC collection
+	// But otherwise not used directly
+	UPROPERTY()
+	TSet<TSoftObjectPtr<const UPCGMetadata>> OtherParents;
+
+	TMap<FName, FPCGMetadataAttributeBase*> Attributes;
+	TSet<FName> AttributeNames; // Will contain also parent attributes
+
+	PCGMetadataAttributeKey NextAttributeId = 0;
+
+	TSet<FName> HiddenAttributes;
+
+	TArray<PCGMetadataEntryKey> ParentKeys;
+	int64 ItemKeyOffset = 0;
+
+	mutable FRWLock AttributeLock;
+	mutable FRWLock ItemLock;
+};
\ No newline at end of file
diff --git a/Engine/Plugins/Experimental/PCG/Source/PCG/Public/Metadata/PCGMetadataAccessor.h b/Engine/Plugins/Experimental/PCG/Source/PCG/Public/Metadata/PCGMetadataAccessor.h
new file mode 100644
--- /dev/null
+++ b/Engine/Plugins/Experimental/PCG/Source/PCG/Public/Metadata/PCGMetadataAccessor.h
@@ -0,0 +1,72 @@
+// Copyright Epic Games, Inc. All Rights Reserved.
+
+#pragma once
+
+#include "CoreMinimal.h"
+#include "PCGMetadataCommon.h"
+#include "PCGPoint.h"
+#include "Kismet/BlueprintFunctionLibrary.h"
+
+#include "PCGMetadataAccessor.generated.h"
+
+struct FPCGPoint;
+class UPCGMetadata;
+
+UCLASS()
+class PCG_API UPCGMetadataAccessorHelpers : public UBlueprintFunctionLibrary
+{
+	GENERATED_BODY()
+
+public:
+	static void InitializeMetadata(FPCGPoint& Point, UPCGMetadata* Metadata);
+	static void InitializeMetadata(FPCGPoint& Point, UPCGMetadata* Metadata, const FPCGPoint& ParentPoint, const UPCGMetadata* ParentMetadata);
+
+	UFUNCTION(BlueprintCallable, Category = "PCG|Metadata", meta = (ScriptMethod))
+	static void InitializeMetadata(UPARAM(ref) FPCGPoint& Point, UPCGMetadata* Metadata, const FPCGPoint& ParentPoint);
+
+	UFUNCTION(BlueprintCallable, Category = "PCG|Metadata", meta = (ScriptMethod))
+	static float GetFloatAttribute(const FPCGPoint& Point, UPCGMetadata* Metadata, FName AttributeName);
+
+	UFUNCTION(BlueprintCallable, Category = "PCG|Metadata", meta = (ScriptMethod))
+	static void SetFloatAttribute(UPARAM(ref) FPCGPoint& Point, UPCGMetadata* Metadata, FName AttributeName, float Value);
+
+	UFUNCTION(BlueprintCallable, Category = "PCG|Metadata", meta = (ScriptMethod))
+	static FVector GetVectorAttribute(const FPCGPoint& Point, UPCGMetadata* Metadata, FName AttributeName);
+
+	UFUNCTION(BlueprintCallable, Category = "PCG|Metadata", meta = (ScriptMethod))
+	static void SetVectorAttribute(UPARAM(ref) FPCGPoint& Point, UPCGMetadata* Metadata, FName AttributeName, const FVector& Value);
+
+	UFUNCTION(BlueprintCallable, Category = "PCG|Metadata", meta = (ScriptMethod))
+	static FVector4 GetVector4Attribute(const FPCGPoint& Point, UPCGMetadata* Metadata, FName AttributeName);
+
+	UFUNCTION(BlueprintCallable, Category = "PCG|Metadata", meta = (ScriptMethod))
+	static void SetVector4Attribute(UPARAM(ref) FPCGPoint& Point, UPCGMetadata* Metadata, FName AttributeName, const FVector4& Value);
+
+	UFUNCTION(BlueprintCallable, Category = "PCG|Metadata", meta = (ScriptMethod))
+	static FQuat GetQuatAttribute(const FPCGPoint& Point, UPCGMetadata* Metadata, FName AttributeName);
+
+	UFUNCTION(BlueprintCallable, Category = "PCG|Metadata", meta = (ScriptMethod))
+	static void SetQuatAttribute(UPARAM(ref) FPCGPoint& Point, UPCGMetadata* Metadata, FName AttributeName, const FQuat& Value);
+
+	UFUNCTION(BlueprintCallable, Category = "PCG|Metadata", meta = (ScriptMethod))
+	static FTransform GetTransformAttribute(const FPCGPoint& Point, UPCGMetadata* Metadata, FName AttributeName);
+
+	UFUNCTION(BlueprintCallable, Category = "PCG|Metadata", meta = (ScriptMethod))
+	static void SetTransformAttribute(UPARAM(ref) FPCGPoint& Point, UPCGMetadata* Metadata, FName AttributeName, const FTransform& Value);
+
+	UFUNCTION(BlueprintCallable, Category = "PCG|Metadata", meta = (ScriptMethod))
+	static FString GetStringAttribute(const FPCGPoint& Point, UPCGMetadata* Metadata, FName AttributeName);
+
+	UFUNCTION(BlueprintCallable, Category = "PCG|Metadata", meta = (ScriptMethod))
+	static void SetStringAttribute(UPARAM(ref) FPCGPoint& Point, UPCGMetadata* Metadata, FName AttributeName, const FString& Value);
+
+	UFUNCTION(BlueprintCallable, Category = "PCG|Metadata", meta = (ScriptMethod))
+	static bool HasAttributeSet(const FPCGPoint& Point, UPCGMetadata* Metadata, FName AttributeName);
+
+protected:
+	template<typename T>
+	static T GetAttribute(const FPCGPoint& Point, UPCGMetadata* Metadata, FName AttributeName);
+
+	template<typename T>
+	static void SetAttribute(FPCGPoint& Point, UPCGMetadata* Metadata, FName AttributeName, const T& Value);
+};
\ No newline at end of file
diff --git a/Engine/Plugins/Experimental/PCG/Source/PCG/Public/Metadata/PCGMetadataAttribute.h b/Engine/Plugins/Experimental/PCG/Source/PCG/Public/Metadata/PCGMetadataAttribute.h
new file mode 100644
--- /dev/null
+++ b/Engine/Plugins/Experimental/PCG/Source/PCG/Public/Metadata/PCGMetadataAttribute.h
@@ -0,0 +1,49 @@
+// Copyright Epic Games, Inc. All Rights Reserved.
+
+#pragma once
+
+#include "CoreMinimal.h"
+#include "PCGMetadataCommon.h"
+#include "PCGMetadata.h"
+
+#include "Misc/ScopeRWLock.h"
+#include "Templates/EnableIf.h"
+
+class FPCGMetadataAttributeBase
+{
+public:
+	FPCGMetadataAttributeBase() = default;
+	FPCGMetadataAttributeBase(UPCGMetadata* InMetadata, FName InName, const FPCGMetadataAttributeBase* InParent, bool bInAllowsInterpolation);
+
+	virtual ~FPCGMetadataAttributeBase() = default;
+
+	const UPCGMetadata* GetMetadata() const { return Metadata; }
+	int32 GetTypeId() const { return TypeId; }
+
+	virtual FPCGMetadataAttributeBase* Copy(FName NewName, UPCGMetadata* InMetadata, bool bKeepParent, bool bCopyEntries = true, bool bCopyValues = true) const = 0;
+
+	virtual void SetValue(PCGMetadataEntryKey ItemKey, const FPCGMetadataAttributeBase* InAttribute, PCGMetadataEntryKey InEntryKey) = 0;
+	virtual void SetZeroValue(PCGMetadataEntryKey ItemKey) = 0;
+	virtual void AccumulateValue(PCGMetadataEntryKey ItemKey, const FPCGMetadataAttributeBase* InAttributeA, PCGMetadataEntryKey InEntryKeyA, float Weight) = 0;
+	virtual void SetValue(PCGMetadataEntryKey ItemKey, const FPCGMetadataAttributeBase* InAttributeA, PCGMetadataEntryKey InEntryKeyA, const FPCGMetadataAttributeBase* InAttributeB, PCGMetadataEntryKey InEntryKeyB, EPCGMetadataOp Op) = 0;
+
+	void SetValueFromValueKey(PCGMetadataEntryKey EntryKey, PCGMetadataValueKey ValueKey);
+	PCGMetadataValueKey GetValueKey(PCGMetadataEntryKey EntryKey) const;
+	bool HasNonDefaultValue(PCGMetadataEntryKey EntryKey) const;
+	void ClearEntries();
+
+	bool AllowsInterpolation() const { return bAllowsInterpolation; }
+
+protected:
+	TMap<PCGMetadataEntryKey, PCGMetadataValueKey> EntryToValueMap;
+	mutable FRWLock EntryMapLock;
+
+	TObjectPtr<UPCGMetadata> Metadata = nullptr;
+	const FPCGMetadataAttributeBase* Parent = nullptr;
+	int32 TypeId = 0;
+	bool bAllowsInterpolation = false;
+
+public:
+	FName Name = NAME_None;
+	PCGMetadataAttributeKey AttributeId = -1;
+};
\ No newline at end of file
diff --git a/Engine/Plugins/Experimental/PCG/Source/PCG/Public/Metadata/PCGMetadataAttributeTraits.h b/Engine/Plugins/Experimental/PCG/Source/PCG/Public/Metadata/PCGMetadataAttributeTraits.h
new file mode 100644
--- /dev/null
+++ b/Engine/Plugins/Experimental/PCG/Source/PCG/Public/Metadata/PCGMetadataAttributeTraits.h
@@ -0,0 +1,240 @@
+// Copyright Epic Games, Inc. All Rights Reserved.
+
+#pragma once
+
+#include "CoreMinimal.h"
+
+namespace PCG
+{
+	namespace Private
+	{
+		template<typename T>
+		struct MetadataTypes
+		{
+			enum { Id = sizeof(T) };
+		};
+
+#define PCGMetadataGenerateDataTypes(Type, TypeId) template<> struct MetadataTypes<Type>{ enum {Id = 2048 + TypeId}; }
+
+		PCGMetadataGenerateDataTypes(float, 0);
+		PCGMetadataGenerateDataTypes(double, 1);
+		PCGMetadataGenerateDataTypes(int32, 2);
+		PCGMetadataGenerateDataTypes(int64, 3);
+		PCGMetadataGenerateDataTypes(FVector, 4);
+		PCGMetadataGenerateDataTypes(FVector4, 5);
+		PCGMetadataGenerateDataTypes(FQuat, 6);
+		PCGMetadataGenerateDataTypes(FTransform, 7);
+		PCGMetadataGenerateDataTypes(FString, 8);
+
+#undef PCGMetadataGenerateDataTypes
+
+		// Common traits for int32, int64, float, double
+		template<typename T>
+		struct MetadataTraits
+		{
+			enum { CompressData = false };
+			enum { CanMinMax = true };
+			enum { CanSubAdd = true };
+			enum { CanInterpolate = true };
+			
+			static T Min(const T& A, const T& B)
+			{
+				return FMath::Min(A, B);
+			}
+
+			static T Max(const T& A, const T& B)
+			{
+				return FMath::Max(A, B);
+			}
+
+			static T Add(const T& A, const T& B)
+			{
+				return A + B;
+			}
+
+			static T Sub(const T& A, const T& B)
+			{
+				return A - B;
+			}
+
+			static T WeightedSum(const T& A, const T& B, float Weight)
+			{
+				return A + B * Weight;
+			}
+
+			static T ZeroValue()
+			{
+				return T{};
+			}
+		};
+
+		// Vector types
+		template<>
+		struct MetadataTraits<FVector>
+		{
+			enum { CompressData = false };
+			enum { CanMinMax = true };
+			enum { CanSubAdd = true };
+			enum { CanInterpolate = true};
+
+			static FVector Min(const FVector& A, const FVector& B) 
+			{
+				return FVector(FMath::Min(A.X, B.X), FMath::Min(A.Y, B.Y), FMath::Min(A.Z, B.Z));
+			}
+
+			static FVector Max(const FVector& A, const FVector& B)
+			{
+				return FVector(FMath::Max(A.X, B.X), FMath::Max(A.Y, B.Y), FMath::Max(A.Z, B.Z));
+			}
+
+			static FVector Add(const FVector& A, const FVector& B)
+			{
+				return A + B;
+			}
+
+			static FVector Sub(const FVector& A, const FVector& B)
+			{
+				return A - B;
+			}
+
+			static FVector WeightedSum(const FVector& A, const FVector& B, float Weight)
+			{
+				return A + B * Weight;
+			}
+
+			static FVector ZeroValue()
+			{
+				return FVector::Zero();
+			}
+		};
+
+		template<>
+		struct MetadataTraits<FVector4>
+		{
+			enum { CompressData = false };
+			enum { CanMinMax = true };
+			enum { CanSubAdd = true };
+			enum { CanInterpolate = true };
+
+			static FVector4 Min(const FVector4& A, const FVector4& B)
+			{
+				return FVector4(FMath::Min(A.X, B.X), FMath::Min(A.Y, B.Y), FMath::Min(A.Z, B.Z), FMath::Min(A.W, B.W));
+			}
+
+			static FVector4 Max(const FVector4& A, const FVector4& B)
+			{
+				return FVector4(FMath::Max(A.X, B.X), FMath::Max(A.Y, B.Y), FMath::Max(A.Z, B.Z), FMath::Max(A.W, B.W));
+			}
+
+			static FVector4 Add(const FVector4& A, const FVector4& B)
+			{
+				return A + B;
+			}
+
+			static FVector4 Sub(const FVector4& A, const FVector4& B)
+			{
+				return A - B;
+			}
+
+			static FVector4 WeightedSum(const FVector4& A, const FVector4& B, float Weight)
+			{
+				return A + B * Weight;
+			}
+
+			static FVector4 ZeroValue()
+			{
+				return FVector::Zero();
+			}
+		};
+
+		// Quaternion
+		template<>
+		struct MetadataTraits<FQuat>
+		{
+			enum { CompressData = false };
+			enum { CanMinMax = false };
+			enum { CanSubAdd = true };
+			enum { CanInterpolate = true };
+
+			static FQuat Add(const FQuat& A, const FQuat& B)
+			{
+				return A * B;
+			}
+
+			static FQuat Sub(const FQuat& A, const FQuat& B)
+			{
+				return A * B.Inverse();
+			}
+
+			static FQuat WeightedSum(const FQuat& A, const FQuat& B, float Weight)
+			{
+				// WARNING: the quaternion won't be normalized
+				FQuat BlendQuat = B * Weight;
+
+				if ((A | BlendQuat) >= 0.0f)
+					return A + BlendQuat;
+				else
+					return A - BlendQuat;
+			}
+
+			static FQuat ZeroValue()
+			{
+				return FQuat::Identity;
+			}
+		};
+
+		// Transform
+		template<>
+		struct MetadataTraits<FTransform>
+		{
+			enum { CompressData = false };
+			enum { CanMinMax = false };
+			enum { CanSubAdd = true };
+			enum { CanInterpolate = true };
+
+			static FTransform Add(const FTransform& A, const FTransform& B)
+			{
+				return A * B;
+			}
+
+			static FTransform Sub(const FTransform& A, const FTransform& B)
+			{
+				return A * B.Inverse();
+			}
+
+			static FQuat WeightedQuatSum(const FQuat& Q, const FQuat& V, float Weight)
+			{
+				FQuat BlendQuat = V * Weight;
+
+				if ((Q | BlendQuat) >= 0.0f)
+					return Q + BlendQuat;
+				else
+					return Q - BlendQuat;
+			}
+
+			static FTransform WeightedSum(const FTransform& A, const FTransform& B, float Weight)
+			{
+				// WARNING: the rotation won't be normalized
+				return FTransform(
+					WeightedQuatSum(A.GetRotation(), B.GetRotation(), Weight),
+					A.GetLocation() + B.GetLocation() * Weight,
+					A.GetScale3D() + B.GetScale3D() * Weight);
+			}
+
+			static FTransform ZeroValue()
+			{
+				return FTransform::Identity;
+			}
+		};
+
+		// Strings
+		template<>
+		struct MetadataTraits<FString>
+		{
+			enum { CompressData = true };
+			enum { CanMinMax = false };
+			enum { CanSubAdd = false };
+			enum { CanInterpolate = false };
+		};
+	}
+}
\ No newline at end of file
diff --git a/Engine/Plugins/Experimental/PCG/Source/PCG/Public/Metadata/PCGMetadataCommon.h b/Engine/Plugins/Experimental/PCG/Source/PCG/Public/Metadata/PCGMetadataCommon.h
new file mode 100644
--- /dev/null
+++ b/Engine/Plugins/Experimental/PCG/Source/PCG/Public/Metadata/PCGMetadataCommon.h
@@ -0,0 +1,12 @@
+// Copyright Epic Games, Inc. All Rights Reserved.
+
+#pragma once
+
+#include "CoreMinimal.h"
+
+typedef int64 PCGMetadataEntryKey;
+typedef int32 PCGMetadataAttributeKey;
+typedef int32 PCGMetadataValueKey;
+
+const PCGMetadataEntryKey PCGInvalidEntryKey = -1;
+const PCGMetadataValueKey PCGDefaultValueKey = -1;
\ No newline at end of file
diff --git a/Engine/Plugins/Experimental/PCG/Source/PCG/Public/PCGPoint.h b/Engine/Plugins/Experimental/PCG/Source/PCG/Public/PCGPoint.h
--- a/Engine/Plugins/Experimental/PCG/Source/PCG/Public/PCGPoint.h
+++ b/Engine/Plugins/Experimental/PCG/Source/PCG/Public/PCGPoint.h
@@ -34,4 +34,7 @@ public:
 
 	UPROPERTY(BlueprintReadWrite, EditAnywhere, Category = Properties)
 	int32 Seed = 0;
+
+	UPROPERTY(BlueprintReadOnly, VisibleAnywhere, Category = "Properties|Metadata")
+	int64 MetadataEntry = -1;
 };
