Fixed issue in union when there is no first non-trivial transform data.

#rnx

diff --git a/Engine/Plugins/Experimental/PCG/Source/PCG/Private/Data/PCGUnionData.cpp b/Engine/Plugins/Experimental/PCG/Source/PCG/Private/Data/PCGUnionData.cpp
--- a/Engine/Plugins/Experimental/PCG/Source/PCG/Private/Data/PCGUnionData.cpp
+++ b/Engine/Plugins/Experimental/PCG/Source/PCG/Private/Data/PCGUnionData.cpp
@@ -157,12 +157,14 @@ FPCGPoint UPCGUnionData::TransformPoint(const FPCGPoint& InPoint) const
 bool UPCGUnionData::GetPointAtPosition(const FVector& InPosition, FPCGPoint& OutPoint, UPCGMetadata* OutMetadata) const
 {
 	FVector PointPosition = InPosition;
+	bool bHasSetPoint = false;
 
 	if (FirstNonTrivialTransformData)
 	{
 		if (FirstNonTrivialTransformData->GetPointAtPosition(InPosition, OutPoint, OutMetadata))
 		{
 			PointPosition = OutPoint.Transform.GetLocation();
+			bHasSetPoint = true;
 
 			if (DensityFunction == EPCGUnionDensityFunction::Binary && OutPoint.Density > 0)
 			{
@@ -171,7 +173,7 @@ bool UPCGUnionData::GetPointAtPosition(const FVector& InPosition, FPCGPoint& Out
 		}
 	}
 
-	const bool bSkipLoop = (!OutMetadata && OutPoint.Density >= 1.0f);
+	const bool bSkipLoop = (bHasSetPoint && !OutMetadata && OutPoint.Density >= 1.0f);
 	const int32 DataCount = Data.Num();
 	for (int32 DataIndex = 0; DataIndex < DataCount && !bSkipLoop; ++DataIndex)
 	{
@@ -183,28 +185,37 @@ bool UPCGUnionData::GetPointAtPosition(const FVector& InPosition, FPCGPoint& Out
 		FPCGPoint PointInData;
 		if (Data[DataIndex]->GetPointAtPosition(PointPosition, PointInData, OutMetadata))
 		{
-			// Update density
-			PCGUnionDataMaths::UpdateDensity(OutPoint.Density, PointInData.Density, DensityFunction);
+			if (!bHasSetPoint)
+			{
+				OutPoint = PointInData;
+				bHasSetPoint = true;
+			}
+			else
+			{
+				// Update density
+				PCGUnionDataMaths::UpdateDensity(OutPoint.Density, PointInData.Density, DensityFunction);
 
-			OutPoint.Color = FVector4(
-				FMath::Max(OutPoint.Color.X, PointInData.Color.X),
-				FMath::Max(OutPoint.Color.Y, PointInData.Color.Y),
-				FMath::Max(OutPoint.Color.Z, PointInData.Color.Z),
-				FMath::Max(OutPoint.Color.W, PointInData.Color.W));
+				OutPoint.Color = FVector4(
+					FMath::Max(OutPoint.Color.X, PointInData.Color.X),
+					FMath::Max(OutPoint.Color.Y, PointInData.Color.Y),
+					FMath::Max(OutPoint.Color.Z, PointInData.Color.Z),
+					FMath::Max(OutPoint.Color.W, PointInData.Color.W));
 
-			// Merge properties into OutPoint
-			if (OutMetadata)
-			{
-				OutMetadata->MergeAttributes(OutPoint, PointInData, OutPoint, EPCGMetadataOp::Max);
+				// Merge properties into OutPoint
+				if (OutMetadata)
+				{
+					OutMetadata->MergeAttributes(OutPoint, PointInData, OutPoint, EPCGMetadataOp::Max);
+				}
 			}
-			else if (OutPoint.Density >= 1.0f)
+			
+			if (bHasSetPoint && !OutMetadata && OutPoint.Density >= 1.0f)
 			{
 				break;
 			}
 		}
 	}
 
-	return OutPoint.Density > 0;
+	return (bHasSetPoint && OutPoint.Density > 0);
 }
 
 bool UPCGUnionData::HasNonTrivialTransform() const
