Fixed blueprint argument tag

#rnx

diff --git a/Engine/Plugins/Experimental/PCG/Source/PCG/Public/Elements/PCGExecuteBlueprint.h b/Engine/Plugins/Experimental/PCG/Source/PCG/Public/Elements/PCGExecuteBlueprint.h
--- a/Engine/Plugins/Experimental/PCG/Source/PCG/Public/Elements/PCGExecuteBlueprint.h
+++ b/Engine/Plugins/Experimental/PCG/Source/PCG/Public/Elements/PCGExecuteBlueprint.h
@@ -43,10 +43,10 @@ public:
 	void Execute(const FPCGDataCollection& Input, FPCGDataCollection& Output) const;
 
 	UFUNCTION(BlueprintImplementableEvent, Category = Execution)
-	bool PointLoopBody(FPCGContext& InContext, const UPCGPointData* InData, const FPCGPoint& InPoint, const UObject* InCustomObject, FPCGPoint& OutPoint, UPCGMetadata* OutMetadata) const;
+	bool PointLoopBody(UPARAM(ref) FPCGContext& InContext, const UPCGPointData* InData, const FPCGPoint& InPoint, const UObject* InCustomObject, FPCGPoint& OutPoint, UPCGMetadata* OutMetadata) const;
 
 	UFUNCTION(BlueprintImplementableEvent, Category = Execution)
-	bool PointPairLoopBody(FPCGContext& InContext, const UPCGPointData* InA, const UPCGPointData* InB, const FPCGPoint& InPointA, const FPCGPoint& InPointB, const UObject* InCustomObject, FPCGPoint& OutPoint, UPCGMetadata* OutMetadata) const;
+	bool PointPairLoopBody(UPARAM(ref) FPCGContext& InContext, const UPCGPointData* InA, const UPCGPointData* InB, const FPCGPoint& InPointA, const FPCGPoint& InPointB, const UObject* InCustomObject, FPCGPoint& OutPoint, UPCGMetadata* OutMetadata) const;
 
 	/** Calls the LoopBody function on all points */
 	UFUNCTION(BlueprintCallable, Category = Execution, meta = (HideSelfPin = "true"))
