Added tracked actor dependency tracking

diff --git a/Engine/Plugins/Experimental/PCG/Source/PCG/Private/Elements/PCGExecuteBlueprint.cpp b/Engine/Plugins/Experimental/PCG/Source/PCG/Private/Elements/PCGExecuteBlueprint.cpp
--- a/Engine/Plugins/Experimental/PCG/Source/PCG/Private/Elements/PCGExecuteBlueprint.cpp
+++ b/Engine/Plugins/Experimental/PCG/Source/PCG/Private/Elements/PCGExecuteBlueprint.cpp
@@ -1,6 +1,7 @@
 // Copyright Epic Games, Inc. All Rights Reserved.
 
 #include "Elements/PCGExecuteBlueprint.h"
+#include "PCGHelpers.h"
 #include "Data/PCGPointData.h"
 #include "Data/PCGSpatialData.h"
 #include "Helpers/PCGAsync.h"
@@ -10,106 +11,13 @@
 
 namespace PCGBlueprintHelper
 {
-	void GatherDependencies(UObject* Object, TSet<TObjectPtr<UObject>>& OutDependencies)
-	{
-		UClass* ObjectClass = Object ? Object->GetClass() : nullptr;
-
-		if (!ObjectClass)
-		{
-			return;
-		}
-
-		for (FProperty* Property = ObjectClass->PropertyLink; Property != nullptr; Property = Property->PropertyLinkNext)
-		{
-			GatherDependencies(Property, Object, OutDependencies);
-		}
-	}
-
-	// Inspired by IteratePropertiesRecursive in ObjectPropertyTrace.cpp
-	void GatherDependencies(FProperty* Property, const void* InContainer, TSet<TObjectPtr<UObject>>& OutDependencies)
-	{
-		auto AddToDependenciesAndGatherRecursively = [&OutDependencies](UObject* Object) {
-			if (Object && !OutDependencies.Contains(Object))
-			{
-				OutDependencies.Add(Object);
-				GatherDependencies(Object, OutDependencies);
-			}
-		};
-
-		if (FObjectProperty* ObjectProperty = CastField<FObjectProperty>(Property))
-		{
-			UObject* Object = ObjectProperty->GetPropertyValue_InContainer(InContainer);
-			AddToDependenciesAndGatherRecursively(Object);
-		}
-		else if (FWeakObjectProperty* WeakObjectProperty = CastField<FWeakObjectProperty>(Property))
-		{
-			FWeakObjectPtr WeakObject = WeakObjectProperty->GetPropertyValue_InContainer(InContainer);
-			AddToDependenciesAndGatherRecursively(WeakObject.Get());
-		}
-		else if (FSoftObjectProperty* SoftObjectProperty = CastField<FSoftObjectProperty>(Property))
-		{
-			FSoftObjectPtr SoftObject = SoftObjectProperty->GetPropertyValue_InContainer(InContainer);
-			AddToDependenciesAndGatherRecursively(SoftObject.Get());
-		}
-		else if (FStructProperty* StructProperty = CastField<FStructProperty>(Property))
-		{
-			const void* StructContainer = StructProperty->ContainerPtrToValuePtr<const void>(InContainer);
-			for (TFieldIterator<FProperty> It(StructProperty->Struct); It; ++It)
-			{
-				GatherDependencies(*It, StructContainer, OutDependencies);
-			}
-		}
-		else if (FArrayProperty* ArrayProperty = CastField<FArrayProperty>(Property))
-		{
-			FScriptArrayHelper_InContainer Helper(ArrayProperty, InContainer);
-			for (int32 DynamicIndex = 0; DynamicIndex < Helper.Num(); ++DynamicIndex)
-			{
-				const void* ValuePtr = Helper.GetRawPtr(DynamicIndex);
-				GatherDependencies(ArrayProperty->Inner, ValuePtr, OutDependencies);
-			}
-		}
-		else if (FMapProperty* MapProperty = CastField<FMapProperty>(Property))
-		{
-			FScriptMapHelper_InContainer Helper(MapProperty, InContainer);
-			int32 Num = Helper.Num();
-			for (int32 DynamicIndex = 0; Num; ++DynamicIndex)
-			{
-				if (Helper.IsValidIndex(DynamicIndex))
-				{
-					const void* KeyPtr = Helper.GetKeyPtr(DynamicIndex);
-					GatherDependencies(MapProperty->KeyProp, KeyPtr, OutDependencies);
-
-					const void* ValuePtr = Helper.GetValuePtr(DynamicIndex);
-					GatherDependencies(MapProperty->ValueProp, ValuePtr, OutDependencies);
-
-					--Num;
-				}
-			}
-		}
-		else if (FSetProperty* SetProperty = CastField<FSetProperty>(Property))
-		{
-			FScriptSetHelper_InContainer Helper(SetProperty, InContainer);
-			int32 Num = Helper.Num();
-			for (int32 DynamicIndex = 0; Num; ++DynamicIndex)
-			{
-				if (Helper.IsValidIndex(DynamicIndex))
-				{
-					const void* ValuePtr = Helper.GetElementPtr(DynamicIndex);
-					GatherDependencies(SetProperty->ElementProp, ValuePtr, OutDependencies);
-
-					--Num;
-				}
-			}
-		}
-	}
-
 	TSet<TObjectPtr<UObject>> GetDataDependencies(UPCGBlueprintElement* InElement)
 	{
 		check(InElement && InElement->GetClass());
 		UClass* BPClass = InElement->GetClass();
 
 		TSet<TObjectPtr<UObject>> Dependencies;
-		GatherDependencies(InElement, Dependencies);
+		PCGHelpers::GatherDependencies(InElement, Dependencies);
 		return Dependencies;
 	}
 }
diff --git a/Engine/Plugins/Experimental/PCG/Source/PCG/Private/PCGComponent.cpp b/Engine/Plugins/Experimental/PCG/Source/PCG/Private/PCGComponent.cpp
--- a/Engine/Plugins/Experimental/PCG/Source/PCG/Private/PCGComponent.cpp
+++ b/Engine/Plugins/Experimental/PCG/Source/PCG/Private/PCGComponent.cpp
@@ -759,11 +759,11 @@ void UPCGComponent::OnActorMoved(AActor* InActor)
 
 void UPCGComponent::OnObjectPropertyChanged(UObject* InObject, FPropertyChangedEvent& InEvent)
 {
-	bool bValueSet = (InEvent.ChangeType == EPropertyChangeType::ValueSet);
+	bool bValueNotInteractive = (InEvent.ChangeType != EPropertyChangeType::Interactive);
 	// Special exception for actor tags, as we can't track otherwise an actor "losing" a tag
 	bool bActorTagChange = (InEvent.Property && InEvent.Property->GetFName() == GET_MEMBER_NAME_CHECKED(AActor, Tags));
 
-	if(!bValueSet && !bActorTagChange)
+	if(!bValueNotInteractive && !bActorTagChange)
 	{
 		return;
 	}
@@ -780,6 +780,25 @@ void UPCGComponent::OnObjectPropertyChanged(UObject* InObject, FPropertyChangedE
 		}
 	}
 
+	// Finally, if it's neither an actor or an actor component, it might be a dependency of a tracked actor
+	if (!Actor)
+	{
+		for(const TPair<TWeakObjectPtr<AActor>, TSet<TObjectPtr<UObject>>>& TrackedActor : CachedTrackedActorToDependencies)
+		{
+			if (TrackedActor.Value.Contains(InObject))
+			{
+				OnActorChanged(TrackedActor.Key.Get(), InObject, bActorTagChange);
+			}
+		}
+	}
+	else
+	{
+		OnActorChanged(Actor, InObject, bActorTagChange);
+	}
+}
+
+void UPCGComponent::OnActorChanged(AActor* Actor, UObject* InObject, bool bActorTagChange)
+{
 	if (Actor == GetOwner())
 	{
 		// Something has changed on the owner (including properties of this component)
@@ -1265,6 +1284,7 @@ bool UPCGComponent::PopulateTrackedActorToTagsMap(bool bForce)
 	}
 
 	CachedTrackedActorToTags.Reset();
+	CachedTrackedActorToDependencies.Reset();
 	for (TWeakObjectPtr<AActor> Actor : CachedTrackedActors)
 	{
 		if (Actor.IsValid())
@@ -1296,6 +1316,7 @@ bool UPCGComponent::AddTrackedActor(AActor* InActor, bool bForce)
 
 		bAppliedChange = true;
 		CachedTrackedActorToTags.FindOrAdd(InActor).Add(Tag);
+		PCGHelpers::GatherDependencies(InActor, CachedTrackedActorToDependencies.FindOrAdd(InActor));
 
 		if (!bForce)
 		{
@@ -1321,6 +1342,7 @@ bool UPCGComponent::RemoveTrackedActor(AActor* InActor)
 		}
 
 		CachedTrackedActorToTags.Remove(InActor);
+		CachedTrackedActorToDependencies.Remove(InActor);
 		bAppliedChange = true;
 	}
 
@@ -1360,6 +1382,7 @@ bool UPCGComponent::UpdateTrackedActor(AActor* InActor)
 		if (!CachedTrackedActorToTags.FindOrAdd(InActor).Find(Tag))
 		{
 			CachedTrackedActorToTags[InActor].Add(Tag);
+			PCGHelpers::GatherDependencies(InActor, CachedTrackedActorToDependencies.FindOrAdd(InActor));
 			DirtyCacheFromTag(Tag);
 			bAppliedChange = true;
 		}
@@ -1369,6 +1392,7 @@ bool UPCGComponent::UpdateTrackedActor(AActor* InActor)
 	if (CachedTrackedActorToTags.Contains(InActor) && CachedTrackedActorToTags[InActor].IsEmpty())
 	{
 		CachedTrackedActorToTags.Remove(InActor);
+		CachedTrackedActorToDependencies.Remove(InActor);
 	}
 
 	return bAppliedChange;
diff --git a/Engine/Plugins/Experimental/PCG/Source/PCG/Private/PCGHelpers.cpp b/Engine/Plugins/Experimental/PCG/Source/PCG/Private/PCGHelpers.cpp
--- a/Engine/Plugins/Experimental/PCG/Source/PCG/Private/PCGHelpers.cpp
+++ b/Engine/Plugins/Experimental/PCG/Source/PCG/Private/PCGHelpers.cpp
@@ -111,5 +111,98 @@ namespace PCGHelpers
 	{
 		return InWorld && InWorld->GetSubsystem<UPCGSubsystem>() ? InWorld->GetSubsystem<UPCGSubsystem>()->GetPCGWorldActor() : nullptr;
 	}
+
+	void GatherDependencies(UObject* Object, TSet<TObjectPtr<UObject>>& OutDependencies)
+	{
+		UClass* ObjectClass = Object ? Object->GetClass() : nullptr;
+
+		if (!ObjectClass)
+		{
+			return;
+		}
+
+		for (FProperty* Property = ObjectClass->PropertyLink; Property != nullptr; Property = Property->PropertyLinkNext)
+		{
+			GatherDependencies(Property, Object, OutDependencies);
+		}
+	}
+
+	// Inspired by IteratePropertiesRecursive in ObjectPropertyTrace.cpp
+	void GatherDependencies(FProperty* Property, const void* InContainer, TSet<TObjectPtr<UObject>>& OutDependencies)
+	{
+		auto AddToDependenciesAndGatherRecursively = [&OutDependencies](UObject* Object) {
+			if (Object && !OutDependencies.Contains(Object))
+			{
+				OutDependencies.Add(Object);
+				GatherDependencies(Object, OutDependencies);
+			}
+		};
+
+		if (FObjectProperty* ObjectProperty = CastField<FObjectProperty>(Property))
+		{
+			UObject* Object = ObjectProperty->GetPropertyValue_InContainer(InContainer);
+			AddToDependenciesAndGatherRecursively(Object);
+		}
+		else if (FWeakObjectProperty* WeakObjectProperty = CastField<FWeakObjectProperty>(Property))
+		{
+			FWeakObjectPtr WeakObject = WeakObjectProperty->GetPropertyValue_InContainer(InContainer);
+			AddToDependenciesAndGatherRecursively(WeakObject.Get());
+		}
+		else if (FSoftObjectProperty* SoftObjectProperty = CastField<FSoftObjectProperty>(Property))
+		{
+			FSoftObjectPtr SoftObject = SoftObjectProperty->GetPropertyValue_InContainer(InContainer);
+			AddToDependenciesAndGatherRecursively(SoftObject.Get());
+		}
+		else if (FStructProperty* StructProperty = CastField<FStructProperty>(Property))
+		{
+			const void* StructContainer = StructProperty->ContainerPtrToValuePtr<const void>(InContainer);
+			for (TFieldIterator<FProperty> It(StructProperty->Struct); It; ++It)
+			{
+				GatherDependencies(*It, StructContainer, OutDependencies);
+			}
+		}
+		else if (FArrayProperty* ArrayProperty = CastField<FArrayProperty>(Property))
+		{
+			FScriptArrayHelper_InContainer Helper(ArrayProperty, InContainer);
+			for (int32 DynamicIndex = 0; DynamicIndex < Helper.Num(); ++DynamicIndex)
+			{
+				const void* ValuePtr = Helper.GetRawPtr(DynamicIndex);
+				GatherDependencies(ArrayProperty->Inner, ValuePtr, OutDependencies);
+			}
+		}
+		else if (FMapProperty* MapProperty = CastField<FMapProperty>(Property))
+		{
+			FScriptMapHelper_InContainer Helper(MapProperty, InContainer);
+			int32 Num = Helper.Num();
+			for (int32 DynamicIndex = 0; Num; ++DynamicIndex)
+			{
+				if (Helper.IsValidIndex(DynamicIndex))
+				{
+					const void* KeyPtr = Helper.GetKeyPtr(DynamicIndex);
+					GatherDependencies(MapProperty->KeyProp, KeyPtr, OutDependencies);
+
+					const void* ValuePtr = Helper.GetValuePtr(DynamicIndex);
+					GatherDependencies(MapProperty->ValueProp, ValuePtr, OutDependencies);
+
+					--Num;
+				}
+			}
+		}
+		else if (FSetProperty* SetProperty = CastField<FSetProperty>(Property))
+		{
+			FScriptSetHelper_InContainer Helper(SetProperty, InContainer);
+			int32 Num = Helper.Num();
+			for (int32 DynamicIndex = 0; Num; ++DynamicIndex)
+			{
+				if (Helper.IsValidIndex(DynamicIndex))
+				{
+					const void* ValuePtr = Helper.GetElementPtr(DynamicIndex);
+					GatherDependencies(SetProperty->ElementProp, ValuePtr, OutDependencies);
+
+					--Num;
+				}
+			}
+		}
+	}
 #endif
 }
diff --git a/Engine/Plugins/Experimental/PCG/Source/PCG/Private/PCGHelpers.h b/Engine/Plugins/Experimental/PCG/Source/PCG/Private/PCGHelpers.h
--- a/Engine/Plugins/Experimental/PCG/Source/PCG/Private/PCGHelpers.h
+++ b/Engine/Plugins/Experimental/PCG/Source/PCG/Private/PCGHelpers.h
@@ -3,6 +3,7 @@
 #pragma once
 
 #include "CoreMinimal.h"
+#include "UObject/ObjectPtr.h"
 
 class AActor;
 class APCGWorldActor;
@@ -31,5 +32,8 @@ namespace PCGHelpers
 
 #if WITH_EDITOR
 	APCGWorldActor* GetPCGWorldActor(UWorld* InWorld);
+
+	void GatherDependencies(UObject* Object, TSet<TObjectPtr<UObject>>& OutDependencies);
+	void GatherDependencies(FProperty* Property, const void* InContainer, TSet<TObjectPtr<UObject>>& OutDependencies);
 #endif
 };
\ No newline at end of file
diff --git a/Engine/Plugins/Experimental/PCG/Source/PCG/Public/Elements/PCGExecuteBlueprint.h b/Engine/Plugins/Experimental/PCG/Source/PCG/Public/Elements/PCGExecuteBlueprint.h
--- a/Engine/Plugins/Experimental/PCG/Source/PCG/Public/Elements/PCGExecuteBlueprint.h
+++ b/Engine/Plugins/Experimental/PCG/Source/PCG/Public/Elements/PCGExecuteBlueprint.h
@@ -19,8 +19,6 @@ DECLARE_MULTICAST_DELEGATE_OneParam(FOnPCGBlueprintChanged, UPCGBlueprintElement
 
 namespace PCGBlueprintHelper
 {
-	void GatherDependencies(UObject* Object, TSet<TObjectPtr<UObject>>& OutDependencies);
-	void GatherDependencies(FProperty* Property, const void* InContainer, TSet<TObjectPtr<UObject>>& OutDependencies);
 	TSet<TObjectPtr<UObject>> GetDataDependencies(UPCGBlueprintElement* InElement);
 }
 #endif // WITH_EDITOR
diff --git a/Engine/Plugins/Experimental/PCG/Source/PCG/Public/PCGComponent.h b/Engine/Plugins/Experimental/PCG/Source/PCG/Public/PCGComponent.h
--- a/Engine/Plugins/Experimental/PCG/Source/PCG/Public/PCGComponent.h
+++ b/Engine/Plugins/Experimental/PCG/Source/PCG/Public/PCGComponent.h
@@ -163,6 +163,8 @@ private:
 	bool ActorIsTracked(AActor* InActor) const;
 	void DirtyExclusionData(AActor* InActor);
 
+	void OnActorChanged(AActor* InActor, UObject* InSourceObject, bool bActorTagChange);
+
 	bool PopulateTrackedActorToTagsMap(bool bForce = false);
 	bool AddTrackedActor(AActor* InActor, bool bForce = false);
 	bool RemoveTrackedActor(AActor* InActor);
@@ -210,6 +212,7 @@ private:
 	TSet<TWeakObjectPtr<AActor>> CachedTrackedActors;
 
 	TMap<TWeakObjectPtr<AActor>, TSet<FName>> CachedTrackedActorToTags;
+	TMap<TWeakObjectPtr<AActor>, TSet<TObjectPtr<UObject>>> CachedTrackedActorToDependencies;
 	bool bActorToTagsMapPopulated = false;
 #endif
 };
\ No newline at end of file
