Disabled logging for non-processing related PCG tasks
Refactored out CRTP class that isn't needed anymore in the element

diff --git a/Engine/Plugins/Experimental/PCG/Source/PCG/Private/Graph/PCGGraphExecutor.h b/Engine/Plugins/Experimental/PCG/Source/PCG/Private/Graph/PCGGraphExecutor.h
--- a/Engine/Plugins/Experimental/PCG/Source/PCG/Private/Graph/PCGGraphExecutor.h
+++ b/Engine/Plugins/Experimental/PCG/Source/PCG/Private/Graph/PCGGraphExecutor.h
@@ -133,6 +133,9 @@ protected:
 	virtual bool ExecuteInternal(FPCGContext* Context) const override;
 	virtual bool IsCancellable() const { return false; }
 	virtual bool IsCacheable(const UPCGSettings* InSettings) const override { return false; }
+#if WITH_EDITOR
+	virtual bool ShouldLog() const { return false; }
+#endif
 
 private:
 	TFunction<bool()> Operation;
diff --git a/Engine/Plugins/Experimental/PCG/Source/PCG/Private/PCGElement.cpp b/Engine/Plugins/Experimental/PCG/Source/PCG/Private/PCGElement.cpp
--- a/Engine/Plugins/Experimental/PCG/Source/PCG/Private/PCGElement.cpp
+++ b/Engine/Plugins/Experimental/PCG/Source/PCG/Private/PCGElement.cpp
@@ -63,7 +63,10 @@ bool IPCGElement::Execute(FPCGContext* Context) const
 		bool bDone = false;
 
 #if WITH_EDITOR
-		UE_LOG(LogPCG, Verbose, TEXT("---------------------------------------"));
+		if (ShouldLog())
+		{
+			UE_LOG(LogPCG, Verbose, TEXT("---------------------------------------"));
+		}
 #endif
 
 		if (IsCacheable(Settings) && Context->Cache && Context->Cache->GetFromCache(this, Context->InputData, Settings, Context->OutputData))
diff --git a/Engine/Plugins/Experimental/PCG/Source/PCG/Public/Elements/PCGDebugElement.h b/Engine/Plugins/Experimental/PCG/Source/PCG/Public/Elements/PCGDebugElement.h
--- a/Engine/Plugins/Experimental/PCG/Source/PCG/Public/Elements/PCGDebugElement.h
+++ b/Engine/Plugins/Experimental/PCG/Source/PCG/Public/Elements/PCGDebugElement.h
@@ -28,7 +28,7 @@ protected:
 	// ~End UPCGSettings interface
 };
 
-class FPCGDebugElement : public FSimpleTypedPCGElement<UPCGDebugSettings>
+class FPCGDebugElement : public FSimplePCGElement
 {
 protected:
 	virtual bool ExecuteInternal(FPCGContext* Context) const override;
diff --git a/Engine/Plugins/Experimental/PCG/Source/PCG/Public/Elements/PCGDifferenceElement.h b/Engine/Plugins/Experimental/PCG/Source/PCG/Public/Elements/PCGDifferenceElement.h
--- a/Engine/Plugins/Experimental/PCG/Source/PCG/Public/Elements/PCGDifferenceElement.h
+++ b/Engine/Plugins/Experimental/PCG/Source/PCG/Public/Elements/PCGDifferenceElement.h
@@ -34,7 +34,7 @@ public:
 #endif
 };
 
-class FPCGDifferenceElement : public FSimpleTypedPCGElement<UPCGDifferenceSettings>
+class FPCGDifferenceElement : public FSimplePCGElement
 {
 protected:
 	virtual bool ExecuteInternal(FPCGContext* Context) const;
diff --git a/Engine/Plugins/Experimental/PCG/Source/PCG/Public/Elements/PCGExecuteBlueprint.h b/Engine/Plugins/Experimental/PCG/Source/PCG/Public/Elements/PCGExecuteBlueprint.h
--- a/Engine/Plugins/Experimental/PCG/Source/PCG/Public/Elements/PCGExecuteBlueprint.h
+++ b/Engine/Plugins/Experimental/PCG/Source/PCG/Public/Elements/PCGExecuteBlueprint.h
@@ -139,7 +139,7 @@ protected:
 	void TeardownBlueprintElementEvent();
 };
 
-class FPCGExecuteBlueprintElement : public FSimpleTypedPCGElement<UPCGBlueprintSettings>
+class FPCGExecuteBlueprintElement : public FSimplePCGElement
 {
 protected:
 	virtual bool ExecuteInternal(FPCGContext* Context) const override;
diff --git a/Engine/Plugins/Experimental/PCG/Source/PCG/Public/Elements/PCGIntersectionElement.h b/Engine/Plugins/Experimental/PCG/Source/PCG/Public/Elements/PCGIntersectionElement.h
--- a/Engine/Plugins/Experimental/PCG/Source/PCG/Public/Elements/PCGIntersectionElement.h
+++ b/Engine/Plugins/Experimental/PCG/Source/PCG/Public/Elements/PCGIntersectionElement.h
@@ -34,7 +34,7 @@ public:
 #endif
 };
 
-class FPCGIntersectionElement : public FSimpleTypedPCGElement<UPCGIntersectionSettings>
+class FPCGIntersectionElement : public FSimplePCGElement
 {
 protected:
 	virtual bool ExecuteInternal(FPCGContext* Context) const;
diff --git a/Engine/Plugins/Experimental/PCG/Source/PCG/Public/Elements/PCGMetadataElement.h b/Engine/Plugins/Experimental/PCG/Source/PCG/Public/Elements/PCGMetadataElement.h
--- a/Engine/Plugins/Experimental/PCG/Source/PCG/Public/Elements/PCGMetadataElement.h
+++ b/Engine/Plugins/Experimental/PCG/Source/PCG/Public/Elements/PCGMetadataElement.h
@@ -53,7 +53,7 @@ public:
 	EPCGMetadataOperationTarget Target = EPCGMetadataOperationTarget::PropertyToAttribute;
 };
 
-class FPCGMetadataOperationElement : public FSimpleTypedPCGElement<UPCGMetadataOperationSettings>
+class FPCGMetadataOperationElement : public FSimplePCGElement
 {
 protected:
 	virtual bool ExecuteInternal(FPCGContext* Context) const override;
diff --git a/Engine/Plugins/Experimental/PCG/Source/PCG/Public/Elements/PCGPointSampler.h b/Engine/Plugins/Experimental/PCG/Source/PCG/Public/Elements/PCGPointSampler.h
--- a/Engine/Plugins/Experimental/PCG/Source/PCG/Public/Elements/PCGPointSampler.h
+++ b/Engine/Plugins/Experimental/PCG/Source/PCG/Public/Elements/PCGPointSampler.h
@@ -34,7 +34,7 @@ public:
 #endif
 };
 
-class FPCGPointSamplerElement : public FSimpleTypedPCGElement<UPCGPointSamplerSettings>
+class FPCGPointSamplerElement : public FSimplePCGElement
 {
 protected:
 	virtual bool ExecuteInternal(FPCGContext* Context) const override;
diff --git a/Engine/Plugins/Experimental/PCG/Source/PCG/Public/Elements/PCGProjectionElement.h b/Engine/Plugins/Experimental/PCG/Source/PCG/Public/Elements/PCGProjectionElement.h
--- a/Engine/Plugins/Experimental/PCG/Source/PCG/Public/Elements/PCGProjectionElement.h
+++ b/Engine/Plugins/Experimental/PCG/Source/PCG/Public/Elements/PCGProjectionElement.h
@@ -29,7 +29,7 @@ protected:
 	//~End UPCGSettings interface
 };
 
-class FPCGProjectionElement : public FSimpleTypedPCGElement<UPCGProjectionSettings>
+class FPCGProjectionElement : public FSimplePCGElement
 {
 protected:
 	virtual bool ExecuteInternal(FPCGContext* Context) const;
diff --git a/Engine/Plugins/Experimental/PCG/Source/PCG/Public/Elements/PCGSelfPruning.h b/Engine/Plugins/Experimental/PCG/Source/PCG/Public/Elements/PCGSelfPruning.h
--- a/Engine/Plugins/Experimental/PCG/Source/PCG/Public/Elements/PCGSelfPruning.h
+++ b/Engine/Plugins/Experimental/PCG/Source/PCG/Public/Elements/PCGSelfPruning.h
@@ -49,7 +49,7 @@ protected:
 	// ~End UPCGSettings interface
 };
 
-class FPCGSelfPruningElement : public FSimpleTypedPCGElement<UPCGSelfPruningSettings>
+class FPCGSelfPruningElement : public FSimplePCGElement
 {
 protected:
 	virtual bool ExecuteInternal(FPCGContext* Context) const override;
diff --git a/Engine/Plugins/Experimental/PCG/Source/PCG/Public/Elements/PCGSpawnActor.h b/Engine/Plugins/Experimental/PCG/Source/PCG/Public/Elements/PCGSpawnActor.h
--- a/Engine/Plugins/Experimental/PCG/Source/PCG/Public/Elements/PCGSpawnActor.h
+++ b/Engine/Plugins/Experimental/PCG/Source/PCG/Public/Elements/PCGSpawnActor.h
@@ -63,7 +63,7 @@ public:
 	/** ~End UPCGBaseSubgraphNode interface */
 };
 
-class FPCGSpawnActorElement : public FSimpleTypedPCGElement<UPCGSpawnActorSettings>
+class FPCGSpawnActorElement : public FSimplePCGElement
 {
 protected:
 	virtual bool ExecuteInternal(FPCGContext* Context) const override;
diff --git a/Engine/Plugins/Experimental/PCG/Source/PCG/Public/Elements/PCGStaticMeshSpawner.h b/Engine/Plugins/Experimental/PCG/Source/PCG/Public/Elements/PCGStaticMeshSpawner.h
--- a/Engine/Plugins/Experimental/PCG/Source/PCG/Public/Elements/PCGStaticMeshSpawner.h
+++ b/Engine/Plugins/Experimental/PCG/Source/PCG/Public/Elements/PCGStaticMeshSpawner.h
@@ -41,7 +41,7 @@ public:
 	TArray<FPCGStaticMeshSpawnerEntry> Meshes;
 };
 
-class FPCGStaticMeshSpawnerElement : public FSimpleTypedPCGElement<UPCGStaticMeshSpawnerSettings>
+class FPCGStaticMeshSpawnerElement : public FSimplePCGElement
 {
 protected:
 	virtual bool ExecuteInternal(FPCGContext* Context) const override;
diff --git a/Engine/Plugins/Experimental/PCG/Source/PCG/Public/Elements/PCGSurfaceSampler.h b/Engine/Plugins/Experimental/PCG/Source/PCG/Public/Elements/PCGSurfaceSampler.h
--- a/Engine/Plugins/Experimental/PCG/Source/PCG/Public/Elements/PCGSurfaceSampler.h
+++ b/Engine/Plugins/Experimental/PCG/Source/PCG/Public/Elements/PCGSurfaceSampler.h
@@ -54,7 +54,7 @@ protected:
 	// ~End UPCGSettings interface
 };
 
-class FPCGSurfaceSamplerElement : public FSimpleTypedPCGElement<UPCGSurfaceSamplerSettings>
+class FPCGSurfaceSamplerElement : public FSimplePCGElement
 {
 protected:
 	virtual bool ExecuteInternal(FPCGContext* Context) const override;
diff --git a/Engine/Plugins/Experimental/PCG/Source/PCG/Public/Elements/PCGUnionElement.h b/Engine/Plugins/Experimental/PCG/Source/PCG/Public/Elements/PCGUnionElement.h
--- a/Engine/Plugins/Experimental/PCG/Source/PCG/Public/Elements/PCGUnionElement.h
+++ b/Engine/Plugins/Experimental/PCG/Source/PCG/Public/Elements/PCGUnionElement.h
@@ -32,7 +32,7 @@ public:
 	EPCGUnionDensityFunction DensityFunction = EPCGUnionDensityFunction::Maximum;
 };
 
-class FPCGUnionElement : public FSimpleTypedPCGElement<UPCGUnionSettings>
+class FPCGUnionElement : public FSimplePCGElement
 {
 protected:
 	virtual bool ExecuteInternal(FPCGContext* Context) const;
diff --git a/Engine/Plugins/Experimental/PCG/Source/PCG/Public/PCGElement.h b/Engine/Plugins/Experimental/PCG/Source/PCG/Public/PCGElement.h
--- a/Engine/Plugins/Experimental/PCG/Source/PCG/Public/PCGElement.h
+++ b/Engine/Plugins/Experimental/PCG/Source/PCG/Public/PCGElement.h
@@ -20,7 +20,11 @@ typedef TSharedPtr<IPCGElement, ESPMode::ThreadSafe> FPCGElementPtr;
 		*((CustomContext)->GetTaskName()), \
 		##__VA_ARGS__)
 
+#if WITH_EDITOR
+#define PCGE_LOG(Verbosity, Format, ...) do{ if(ShouldLog()) { PCGE_LOG_C(Verbosity, Context, Format, ##__VA_ARGS__); } }while(0)
+#else
 #define PCGE_LOG(Verbosity, Format, ...) PCGE_LOG_C(Verbosity, Context, Format, ##__VA_ARGS__)
+#endif
 
 /**
 * Base class for the processing bit of a PCG node/settings
@@ -37,6 +41,9 @@ protected:
 	virtual bool ExecuteInternal(FPCGContext* Context) const = 0;
 	virtual bool IsCancellable() const { return true; }
 	virtual bool IsCacheable(const UPCGSettings* InSettings) const { return true; }
+#if WITH_EDITOR
+	virtual bool ShouldLog() const { return true; }
+#endif
 };
 
 /**
@@ -46,14 +53,4 @@ class PCG_API FSimplePCGElement : public IPCGElement
 {
 public:
 	virtual FPCGContext* Initialize(const FPCGDataCollection& InputData, UPCGComponent* SourceComponent) override;
-};
-
-/**
-* CRTP PCG element class to facilitate settings retrieval
-*/
-template<typename SettingsClass>
-class PCG_API FSimpleTypedPCGElement : public FSimplePCGElement
-{
-public:
-	const SettingsClass* GetSettings();
-};
+};
\ No newline at end of file
