Fixed potential issue where adding a duplicate attribute would return a pointer to deleted memory.
Should fix the static analysis warning

#rnx

diff --git a/Engine/Plugins/Experimental/PCG/Source/PCG/Private/Metadata/PCGMetadata.cpp b/Engine/Plugins/Experimental/PCG/Source/PCG/Private/Metadata/PCGMetadata.cpp
--- a/Engine/Plugins/Experimental/PCG/Source/PCG/Private/Metadata/PCGMetadata.cpp
+++ b/Engine/Plugins/Experimental/PCG/Source/PCG/Private/Metadata/PCGMetadata.cpp
@@ -109,13 +109,14 @@ FPCGMetadataAttributeBase* UPCGMetadata::CreateAttribute(FName AttributeName, co
 		ParentAttribute = Parent->GetConstAttribute(AttributeName);
 	}
 
-	FPCGMetadataAttribute<T>* NewAttribute = new FPCGMetadataAttribute<T>(this, AttributeName, ParentAttribute, DefaultValue, bAllowsInterpolation);
+	FPCGMetadataAttributeBase* NewAttribute = new FPCGMetadataAttribute<T>(this, AttributeName, ParentAttribute, DefaultValue, bAllowsInterpolation);
 
 	AttributeLock.WriteLock();
-	if (Attributes.Find(AttributeName))
+	if (FPCGMetadataAttributeBase** ExistingAttribute = Attributes.Find(AttributeName))
 	{
 		UE_LOG(LogPCG, Warning, TEXT("Attribute %s already exists"), *AttributeName.ToString());
 		delete NewAttribute;
+		NewAttribute = *ExistingAttribute;
 	}
 	else
 	{
