PCG:
Added basic support for connections.
*Break pin link / links
*Break node connections
*Creating new connections
*Spawn nodes from connection

Added support for deleting nodes.
Fixed an issue were PCG nodes didnt have their position set after creation.

diff --git a/Engine/Plugins/Experimental/PCG/Source/PCG/Private/PCGGraph.cpp b/Engine/Plugins/Experimental/PCG/Source/PCG/Private/PCGGraph.cpp
--- a/Engine/Plugins/Experimental/PCG/Source/PCG/Private/PCGGraph.cpp
+++ b/Engine/Plugins/Experimental/PCG/Source/PCG/Private/PCGGraph.cpp
@@ -188,6 +188,8 @@ void UPCGGraph::OnNodeRemoved(UPCGNode* InNode)
 		{
 			SubgraphNode->OnNodeStructuralSettingsChangedDelegate.RemoveAll(this);
 		}
+
+		NotifyGraphChanged(/*bIsStructural=*/true);
 	}
 #endif
 }
@@ -217,6 +219,61 @@ bool UPCGGraph::Contains(UPCGNode* Node) const
 	return Node == InputNode || Node == OutputNode || Nodes.Contains(Node);
 }
 
+void UPCGGraph::RemoveNode(UPCGNode* InNode)
+{
+	check(InNode);
+
+	for (UPCGNode* Input : InNode->InboundNodes)
+	{
+		Input->OutboundNodes.Remove(InNode);
+	}
+	for (UPCGNode* Output : InNode->OutboundNodes)
+	{
+		Output->InboundNodes.Remove(InNode);
+	}
+
+	Nodes.Remove(InNode);
+	OnNodeRemoved(InNode);
+}
+
+void UPCGGraph::RemoveEdge(UPCGNode* From, UPCGNode* To)
+{
+	if (!From || !To)
+	{
+		// TODO: log error
+		return;
+	}
+
+	From->RemoveConnection(To);
+	To->RemoveConnection(From);
+
+#if WITH_EDITOR
+	NotifyGraphChanged(/*bIsStructural=*/true);
+#endif
+}
+
+void UPCGGraph::RemoveInboundEdges(UPCGNode* InNode)
+{
+	check(InNode);
+
+	for (int32 i = InNode->InboundNodes.Num() - 1; i >= 0; --i)
+	{
+		UPCGNode* InboundNode = InNode->InboundNodes[i];
+		RemoveEdge(InNode, InboundNode);
+	}
+}
+
+void UPCGGraph::RemoveOutboundEdges(UPCGNode* InNode)
+{
+	check(InNode);
+
+	for (int32 i = InNode->OutboundNodes.Num() - 1; i >= 0; --i)
+	{
+		UPCGNode* OutboundNode = InNode->OutboundNodes[i];
+		RemoveEdge(InNode, OutboundNode);
+	}
+}
+
 #if WITH_EDITOR
 FPCGTagToSettingsMap UPCGGraph::GetTrackedTagsToSettings() const
 {
diff --git a/Engine/Plugins/Experimental/PCG/Source/PCG/Private/PCGNode.cpp b/Engine/Plugins/Experimental/PCG/Source/PCG/Private/PCGNode.cpp
--- a/Engine/Plugins/Experimental/PCG/Source/PCG/Private/PCGNode.cpp
+++ b/Engine/Plugins/Experimental/PCG/Source/PCG/Private/PCGNode.cpp
@@ -82,6 +82,12 @@ void UPCGNode::ConnectFrom(UPCGNode* InPredecessor)
 	InboundNodes.AddUnique(InPredecessor);
 }
 
+void UPCGNode::RemoveConnection(UPCGNode* InNode)
+{
+	InboundNodes.Remove(InNode);
+	OutboundNodes.Remove(InNode);
+}
+
 #if WITH_EDITOR
 
 void UPCGNode::PreEditChange(FProperty* PropertyAboutToChange)
diff --git a/Engine/Plugins/Experimental/PCG/Source/PCG/Public/PCGGraph.h b/Engine/Plugins/Experimental/PCG/Source/PCG/Public/PCGGraph.h
--- a/Engine/Plugins/Experimental/PCG/Source/PCG/Public/PCGGraph.h
+++ b/Engine/Plugins/Experimental/PCG/Source/PCG/Public/PCGGraph.h
@@ -57,6 +57,11 @@ public:
 
 	bool Contains(UPCGNode* Node) const;
 	const TArray<UPCGNode*>& GetNodes() const { return Nodes; }
+	void RemoveNode(UPCGNode* InNode);
+	void RemoveEdge(UPCGNode* From, UPCGNode* To);
+	void RemoveInboundEdges(UPCGNode* InNode);
+	void RemoveOutboundEdges(UPCGNode* InNode);
+
 #if WITH_EDITOR
 	FPCGTagToSettingsMap GetTrackedTagsToSettings() const;
 #endif
diff --git a/Engine/Plugins/Experimental/PCG/Source/PCG/Public/PCGNode.h b/Engine/Plugins/Experimental/PCG/Source/PCG/Public/PCGNode.h
--- a/Engine/Plugins/Experimental/PCG/Source/PCG/Public/PCGNode.h
+++ b/Engine/Plugins/Experimental/PCG/Source/PCG/Public/PCGNode.h
@@ -63,6 +63,8 @@ protected:
 	void ConnectTo(UPCGNode* InSuccessor);
 	void ConnectFrom(UPCGNode* InPredecessor);
 
+	void RemoveConnection(UPCGNode* InNode);
+
 #if WITH_EDITOR
 	virtual void PreEditChange(FProperty* PropertyAboutToChange) override;
 	virtual void PostEditChangeProperty(struct FPropertyChangedEvent& PropertyChangedEvent) override;
diff --git a/Engine/Plugins/Experimental/PCG/Source/PCGEditor/Private/PCGEditor.cpp b/Engine/Plugins/Experimental/PCG/Source/PCGEditor/Private/PCGEditor.cpp
--- a/Engine/Plugins/Experimental/PCG/Source/PCGEditor/Private/PCGEditor.cpp
+++ b/Engine/Plugins/Experimental/PCG/Source/PCGEditor/Private/PCGEditor.cpp
@@ -180,6 +180,28 @@ bool FPCGEditor::CanSelectAllNodes()
 	return GraphEditorWidget.IsValid();
 }
 
+void FPCGEditor::DeleteSelectedNodes()
+{
+	if (GraphEditorWidget.IsValid())
+	{
+		UPCGGraph* PCGGraph = PCGEditorGraph->GetPCGGraph();
+		check(PCGEditorGraph && PCGGraph);
+
+		for (UObject* Object : GraphEditorWidget->GetSelectedNodes())
+		{
+			UPCGEditorGraphNode* PCGEditorGraphNode = CastChecked<UPCGEditorGraphNode>(Object);
+			UPCGNode* PCGNode = PCGEditorGraphNode->GetPCGNode();
+			check(PCGNode);
+
+			PCGGraph->RemoveNode(PCGNode);
+			PCGEditorGraphNode->DestroyNode();
+		}
+
+		GraphEditorWidget->ClearSelectionSet();
+		GraphEditorWidget->NotifyGraphChanged();
+	}
+}
+
 void FPCGEditor::OnAlignTop()
 {
 	if (GraphEditorWidget.IsValid())
@@ -261,6 +283,10 @@ TSharedRef<SGraphEditor> FPCGEditor::CreateGraphEditorWidget()
 		FExecuteAction::CreateSP(this, &FPCGEditor::SelectAllNodes),
 		FCanExecuteAction::CreateSP(this, &FPCGEditor::CanSelectAllNodes));
 
+	GraphEditorCommands->MapAction(FGenericCommands::Get().Delete,
+		FExecuteAction::CreateSP(this, &FPCGEditor::DeleteSelectedNodes)
+	);
+
 	// Alignment Commands
 	GraphEditorCommands->MapAction(FGraphEditorCommands::Get().AlignNodesTop,
 		FExecuteAction::CreateSP(this, &FPCGEditor::OnAlignTop)
diff --git a/Engine/Plugins/Experimental/PCG/Source/PCGEditor/Private/PCGEditor.h b/Engine/Plugins/Experimental/PCG/Source/PCGEditor/Private/PCGEditor.h
--- a/Engine/Plugins/Experimental/PCG/Source/PCGEditor/Private/PCGEditor.h
+++ b/Engine/Plugins/Experimental/PCG/Source/PCGEditor/Private/PCGEditor.h
@@ -34,6 +34,9 @@ private:
 	/** Whether we can select every node */
 	bool CanSelectAllNodes();
 
+	/** Delete all selected nodes in the graph */
+	void DeleteSelectedNodes();
+
 	void OnAlignTop();
 	void OnAlignMiddle();
 	void OnAlignBottom();
diff --git a/Engine/Plugins/Experimental/PCG/Source/PCGEditor/Private/PCGEditorGraphNode.cpp b/Engine/Plugins/Experimental/PCG/Source/PCGEditor/Private/PCGEditorGraphNode.cpp
--- a/Engine/Plugins/Experimental/PCG/Source/PCGEditor/Private/PCGEditorGraphNode.cpp
+++ b/Engine/Plugins/Experimental/PCG/Source/PCGEditor/Private/PCGEditorGraphNode.cpp
@@ -9,6 +9,7 @@
 
 #include "PCGNode.h"
 #include "PCGSettings.h"
+#include "PCGEditorGraphSchema.h"
 
 #define LOCTEXT_NAMESPACE "PCGEditorGraphNode"
 
@@ -36,8 +37,12 @@ void UPCGEditorGraphNode::GetNodeContextMenuActions(UToolMenu* Menu, class UGrap
 		return;
 	}
 
-	FToolMenuSection& Section = Menu->AddSection("EdGraphSchemaOrganization", LOCTEXT("OrganizationHeader", "Organization"));
+	FToolMenuSection& NodeSection = Menu->AddSection("EdGraphSchemaNodeActions", LOCTEXT("NodeActionsHeader", "Node Actions"));
+	{
+		NodeSection.AddMenuEntry(FGraphEditorCommands::Get().BreakNodeLinks);
+	}
 
+	FToolMenuSection& Section = Menu->AddSection("EdGraphSchemaOrganization", LOCTEXT("OrganizationHeader", "Organization"));
 	Section.AddSubMenu("Alignment", LOCTEXT("AlignmentHeader", "Alignment"), FText(), FNewToolMenuDelegate::CreateLambda([](UToolMenu* AlignmentMenu)
 		{
 			{
@@ -72,4 +77,19 @@ void UPCGEditorGraphNode::AllocateDefaultPins()
 	}
 }
 
+void UPCGEditorGraphNode::AutowireNewNode(UEdGraphPin* FromPin)
+{
+	if (FromPin->Direction == EEdGraphPinDirection::EGPD_Output)
+	{
+		UEdGraphPin* ToPin = FindPinChecked(TEXT("In"), EEdGraphPinDirection::EGPD_Input);
+		GetSchema()->TryCreateConnection(FromPin, ToPin);
+	}
+	else if (FromPin->Direction == EEdGraphPinDirection::EGPD_Input)
+	{
+		UEdGraphPin* ToPin = FindPinChecked(TEXT("Out"), EEdGraphPinDirection::EGPD_Output);
+		GetSchema()->TryCreateConnection(FromPin, ToPin);
+	}
+	NodeConnectionListChanged();
+}
+
 #undef LOCTEXT_NAMESPACE
diff --git a/Engine/Plugins/Experimental/PCG/Source/PCGEditor/Private/PCGEditorGraphNode.h b/Engine/Plugins/Experimental/PCG/Source/PCGEditor/Private/PCGEditorGraphNode.h
--- a/Engine/Plugins/Experimental/PCG/Source/PCGEditor/Private/PCGEditorGraphNode.h
+++ b/Engine/Plugins/Experimental/PCG/Source/PCGEditor/Private/PCGEditorGraphNode.h
@@ -31,6 +31,7 @@ public:
 	virtual FText GetNodeTitle(ENodeTitleType::Type TitleType) const override;
 	virtual void GetNodeContextMenuActions(UToolMenu* Menu, class UGraphNodeContextMenuContext* Context) const override;
 	virtual void AllocateDefaultPins() override;
+	virtual void AutowireNewNode(UEdGraphPin* FromPin) override;
 	// ~End UEdGraphNode interface
 
 	UPCGNode* GetPCGNode() { return PCGNode; }
diff --git a/Engine/Plugins/Experimental/PCG/Source/PCGEditor/Private/PCGEditorGraphSchema.cpp b/Engine/Plugins/Experimental/PCG/Source/PCGEditor/Private/PCGEditorGraphSchema.cpp
--- a/Engine/Plugins/Experimental/PCG/Source/PCGEditor/Private/PCGEditorGraphSchema.cpp
+++ b/Engine/Plugins/Experimental/PCG/Source/PCGEditor/Private/PCGEditorGraphSchema.cpp
@@ -6,6 +6,7 @@
 #include "PCGEditorGraphNode.h"
 #include "PCGEditorGraphSchemaActions.h"
 #include "PCGSettings.h"
+#include "PCGGraph.h"
 
 #include "UObject/UObjectIterator.h"
 
@@ -44,4 +45,89 @@ FLinearColor UPCGEditorGraphSchema::GetPinTypeColor(const FEdGraphPinType& PinTy
 	return FLinearColor(1.0f, 0.0f, 1.0f, 1.0f);
 }
 
+const FPinConnectionResponse UPCGEditorGraphSchema::CanCreateConnection(const UEdGraphPin* A, const UEdGraphPin* B) const
+{
+	const UEdGraphNode* NodeA = A->GetOwningNode();
+	const UEdGraphNode* NodeB = B->GetOwningNode();
+
+	if (NodeA == NodeB)
+	{
+		return FPinConnectionResponse(CONNECT_RESPONSE_DISALLOW, LOCTEXT("ConnectionSameNode", "Both pins are on same node"));
+	}
+
+	if (A->Direction == B->Direction)
+	{
+		return FPinConnectionResponse(CONNECT_RESPONSE_DISALLOW, LOCTEXT("ConnectionSameDirection", "Both pins are the same direction"));
+	}
+
+	return FPinConnectionResponse();
+}
+
+bool UPCGEditorGraphSchema::TryCreateConnection(UEdGraphPin* A, UEdGraphPin* B) const
+{
+	bool bModified = Super::TryCreateConnection(A, B);
+
+	if (bModified)
+	{
+		UEdGraphNode* NodeA = A->GetOwningNode();
+		UEdGraphNode* NodeB = B->GetOwningNode();
+
+		UPCGEditorGraphNode* PCGGraphNodeA = CastChecked<UPCGEditorGraphNode>(NodeA);
+		UPCGEditorGraphNode* PCGGraphNodeB = CastChecked<UPCGEditorGraphNode>(NodeB);
+
+		UPCGNode* PCGNodeA = PCGGraphNodeA->GetPCGNode();
+		UPCGNode* PCGNodeB = PCGGraphNodeB->GetPCGNode();
+		check(PCGNodeA && PCGNodeB);
+
+		UPCGGraph* PCGGraph = PCGNodeA->GetGraph();
+		check(PCGGraph);
+
+		PCGGraph->AddEdge(PCGNodeA, PCGNodeB);
+	}
+
+	return bModified;
+}
+
+void UPCGEditorGraphSchema::BreakPinLinks(UEdGraphPin& TargetPin, bool bSendsNodeNotification) const
+{
+	Super::BreakPinLinks(TargetPin, bSendsNodeNotification);
+
+	UEdGraphNode* GraphNode = TargetPin.GetOwningNode();
+
+	UPCGEditorGraphNode* PCGGraphNode = CastChecked<UPCGEditorGraphNode>(GraphNode);
+
+	UPCGNode* PCGNode = PCGGraphNode->GetPCGNode();
+	check(PCGNode);
+
+	UPCGGraph* PCGGraph = PCGNode->GetGraph();
+	check(PCGGraph);
+
+	if (TargetPin.Direction == EEdGraphPinDirection::EGPD_Input)
+	{
+		PCGGraph->RemoveInboundEdges(PCGNode);
+	}
+	else if (TargetPin.Direction == EEdGraphPinDirection::EGPD_Output)
+	{
+		PCGGraph->RemoveOutboundEdges(PCGNode);
+	}
+}
+
+void UPCGEditorGraphSchema::BreakSinglePinLink(UEdGraphPin* SourcePin, UEdGraphPin* TargetPin) const
+{
+	Super::BreakSinglePinLink(SourcePin, TargetPin);
+
+	UEdGraphNode* SourceGraphNode = SourcePin->GetOwningNode();
+	UEdGraphNode* TargetGraphNode = TargetPin->GetOwningNode();
+
+	UPCGEditorGraphNode* SourcePCGGraphNode = CastChecked<UPCGEditorGraphNode>(SourceGraphNode);
+	UPCGEditorGraphNode* TargetPCGGraphNode = CastChecked<UPCGEditorGraphNode>(TargetGraphNode);
+
+	UPCGNode* SourcePCGNode = SourcePCGGraphNode->GetPCGNode();
+	UPCGNode* TargetPCGNode = TargetPCGGraphNode->GetPCGNode();
+	check(SourcePCGNode && TargetPCGNode);
+
+	UPCGGraph* PCGGraph = SourcePCGNode->GetGraph();
+	PCGGraph->RemoveEdge(SourcePCGNode, TargetPCGNode);
+}
+
 #undef LOCTEXT_NAMESPACE
diff --git a/Engine/Plugins/Experimental/PCG/Source/PCGEditor/Private/PCGEditorGraphSchema.h b/Engine/Plugins/Experimental/PCG/Source/PCGEditor/Private/PCGEditorGraphSchema.h
--- a/Engine/Plugins/Experimental/PCG/Source/PCGEditor/Private/PCGEditorGraphSchema.h
+++ b/Engine/Plugins/Experimental/PCG/Source/PCGEditor/Private/PCGEditorGraphSchema.h
@@ -15,5 +15,9 @@ public:
 	//~ Begin EdGraphSchema Interface
 	virtual void GetGraphContextActions(FGraphContextMenuBuilder& ContextMenuBuilder) const;
 	virtual FLinearColor GetPinTypeColor(const FEdGraphPinType& PinType) const override;
+	virtual const FPinConnectionResponse CanCreateConnection(const UEdGraphPin* A, const UEdGraphPin* B) const override;
+	virtual bool TryCreateConnection(UEdGraphPin* A, UEdGraphPin* B) const override;
+	virtual void BreakPinLinks(UEdGraphPin& TargetPin, bool bSendsNodeNotification) const override;
+	virtual void BreakSinglePinLink(UEdGraphPin* SourcePin, UEdGraphPin* TargetPin) const override;
 	//~ End EdGraphSchema Interface
 };
diff --git a/Engine/Plugins/Experimental/PCG/Source/PCGEditor/Private/PCGEditorGraphSchemaActions.cpp b/Engine/Plugins/Experimental/PCG/Source/PCGEditor/Private/PCGEditorGraphSchemaActions.cpp
--- a/Engine/Plugins/Experimental/PCG/Source/PCGEditor/Private/PCGEditorGraphSchemaActions.cpp
+++ b/Engine/Plugins/Experimental/PCG/Source/PCGEditor/Private/PCGEditorGraphSchemaActions.cpp
@@ -33,5 +33,13 @@ UEdGraphNode* FPCGEditorGraphSchemaAction_NewNode::PerformAction(UEdGraph* Paren
 	NewNode->NodePosY = Location.Y;
 	NodeCreator.Finalize();
 
+	NewPCGNode->PositionX = Location.X;
+	NewPCGNode->PositionY = Location.Y;
+
+	if (FromPin)
+	{
+		NewNode->AutowireNewNode(FromPin);
+	}
+
 	return NewNode;
 }
