Added optional out points data in point loop functions
Renamed InCustomObject to OptionalCustomObject for clarity

#rnx

diff --git a/Engine/Plugins/Experimental/PCG/Source/PCG/Private/Data/PCGPointData.cpp b/Engine/Plugins/Experimental/PCG/Source/PCG/Private/Data/PCGPointData.cpp
--- a/Engine/Plugins/Experimental/PCG/Source/PCG/Private/Data/PCGPointData.cpp
+++ b/Engine/Plugins/Experimental/PCG/Source/PCG/Private/Data/PCGPointData.cpp
@@ -174,10 +174,16 @@ void UPCGPointData::InitializeFromActor(AActor* InActor)
 void UPCGPointData::InitializeFromData(const UPCGSpatialData* InSource, const UPCGMetadata* InMetadataParentOverride)
 {
 	check(InSource);
-	TargetActor = InSource->TargetActor;
+	if (!TargetActor)
+	{
+		TargetActor = InSource->TargetActor;
+	}	
 
-	Metadata = NewObject<UPCGMetadata>(this);
-	Metadata->Initialize(InMetadataParentOverride ? InMetadataParentOverride : InSource->Metadata);
+	if (!Metadata)
+	{
+		Metadata = NewObject<UPCGMetadata>(this);
+		Metadata->Initialize(InMetadataParentOverride ? InMetadataParentOverride : InSource->Metadata);
+	}
 }
 
 const FPCGPoint* UPCGPointData::GetPointAtPosition(const FVector& InPosition) const
diff --git a/Engine/Plugins/Experimental/PCG/Source/PCG/Private/Elements/PCGExecuteBlueprint.cpp b/Engine/Plugins/Experimental/PCG/Source/PCG/Private/Elements/PCGExecuteBlueprint.cpp
--- a/Engine/Plugins/Experimental/PCG/Source/PCG/Private/Elements/PCGExecuteBlueprint.cpp
+++ b/Engine/Plugins/Experimental/PCG/Source/PCG/Private/Elements/PCGExecuteBlueprint.cpp
@@ -290,7 +290,7 @@ bool FPCGExecuteBlueprintElement::ExecuteInternal(FPCGContext* Context) const
 	return true;
 }
 
-void UPCGBlueprintElement::LoopOnPoints(FPCGContext& InContext, const UPCGPointData* InData, UPCGPointData*& OutData, const UObject* InCustomObject) const
+void UPCGBlueprintElement::LoopOnPoints(FPCGContext& InContext, const UPCGPointData* InData, UPCGPointData*& OutData, UPCGPointData* OptionalOutData, const UObject* OptionalCustomObject) const
 {
 	if (!InData)
 	{
@@ -298,19 +298,19 @@ void UPCGBlueprintElement::LoopOnPoints(FPCGContext& InContext, const UPCGPointD
 		return;
 	}
 
-	OutData = NewObject<UPCGPointData>();
+	OutData = OptionalOutData ? OptionalOutData : NewObject<UPCGPointData>(const_cast<UPCGPointData*>(InData));
 	OutData->InitializeFromData(InData);
-	
+
 	const TArray<FPCGPoint>& InPoints = InData->GetPoints();
 	TArray<FPCGPoint>& OutPoints = OutData->GetMutablePoints();
 
-	FPCGAsync::AsyncPointProcessing(&InContext, InPoints.Num(), OutPoints, [this, &InContext, InData, OutData, InCustomObject, &InPoints](int32 Index, FPCGPoint& OutPoint)
+	FPCGAsync::AsyncPointProcessing(&InContext, InPoints.Num(), OutPoints, [this, &InContext, InData, OutData, OptionalCustomObject, &InPoints](int32 Index, FPCGPoint& OutPoint)
 	{
-		return PointLoopBody(InContext, InData, InPoints[Index], InCustomObject, OutPoint, OutData->Metadata);
+		return PointLoopBody(InContext, InData, InPoints[Index], OptionalCustomObject, OutPoint, OutData->Metadata);
 	});
 }
 
-void UPCGBlueprintElement::LoopOnPointPairs(FPCGContext& InContext, const UPCGPointData* InA, const UPCGPointData* InB, UPCGPointData*& OutData, const UObject* InCustomObject) const
+void UPCGBlueprintElement::LoopOnPointPairs(FPCGContext& InContext, const UPCGPointData* InA, const UPCGPointData* InB, UPCGPointData*& OutData, UPCGPointData* OptionalOutData, const UObject* OptionalCustomObject) const
 {
 	if (!InA || !InB)
 	{
@@ -318,16 +318,16 @@ void UPCGBlueprintElement::LoopOnPointPairs(FPCGContext& InContext, const UPCGPo
 		return;
 	}
 
-	OutData = NewObject<UPCGPointData>();
+	OutData = OptionalOutData ? OptionalOutData : NewObject<UPCGPointData>(const_cast<UPCGPointData*>(InA));
 	OutData->InitializeFromData(InA); //METADATA TODO should we remove the parenting here?
 
 	const TArray<FPCGPoint>& InPointsA = InA->GetPoints();
 	const TArray<FPCGPoint>& InPointsB = InB->GetPoints();
 	TArray<FPCGPoint>& OutPoints = OutData->GetMutablePoints();
 
-	FPCGAsync::AsyncPointProcessing(&InContext, InPointsA.Num() * InPointsB.Num(), OutPoints, [this, &InContext, InA, InB, InCustomObject, OutData, &InPointsA, &InPointsB](int32 Index, FPCGPoint& OutPoint)
+	FPCGAsync::AsyncPointProcessing(&InContext, InPointsA.Num() * InPointsB.Num(), OutPoints, [this, &InContext, InA, InB, OptionalCustomObject, OutData, &InPointsA, &InPointsB](int32 Index, FPCGPoint& OutPoint)
 	{
-		return PointPairLoopBody(InContext, InA, InB, InPointsA[Index / InPointsB.Num()], InPointsB[Index % InPointsB.Num()], InCustomObject, OutPoint, OutData->Metadata);
+		return PointPairLoopBody(InContext, InA, InB, InPointsA[Index / InPointsB.Num()], InPointsB[Index % InPointsB.Num()], OptionalCustomObject, OutPoint, OutData->Metadata);
 	});
 }
 
diff --git a/Engine/Plugins/Experimental/PCG/Source/PCG/Public/Elements/PCGExecuteBlueprint.h b/Engine/Plugins/Experimental/PCG/Source/PCG/Public/Elements/PCGExecuteBlueprint.h
--- a/Engine/Plugins/Experimental/PCG/Source/PCG/Public/Elements/PCGExecuteBlueprint.h
+++ b/Engine/Plugins/Experimental/PCG/Source/PCG/Public/Elements/PCGExecuteBlueprint.h
@@ -41,17 +41,17 @@ public:
 	void Execute(const FPCGDataCollection& Input, FPCGDataCollection& Output) const;
 
 	UFUNCTION(BlueprintImplementableEvent, Category = Execution)
-	bool PointLoopBody(UPARAM(ref) FPCGContext& InContext, const UPCGPointData* InData, const FPCGPoint& InPoint, const UObject* InCustomObject, FPCGPoint& OutPoint, UPCGMetadata* OutMetadata) const;
+	bool PointLoopBody(UPARAM(ref) FPCGContext& InContext, const UPCGPointData* InData, const FPCGPoint& InPoint, const UObject* OptionalCustomObject, FPCGPoint& OutPoint, UPCGMetadata* OutMetadata) const;
 
 	UFUNCTION(BlueprintImplementableEvent, Category = Execution)
-	bool PointPairLoopBody(UPARAM(ref) FPCGContext& InContext, const UPCGPointData* InA, const UPCGPointData* InB, const FPCGPoint& InPointA, const FPCGPoint& InPointB, const UObject* InCustomObject, FPCGPoint& OutPoint, UPCGMetadata* OutMetadata) const;
+	bool PointPairLoopBody(UPARAM(ref) FPCGContext& InContext, const UPCGPointData* InA, const UPCGPointData* InB, const FPCGPoint& InPointA, const FPCGPoint& InPointB, const UObject* OptionalCustomObject, FPCGPoint& OutPoint, UPCGMetadata* OutMetadata) const;
 
 	/** Calls the LoopBody function on all points */
 	UFUNCTION(BlueprintCallable, Category = Execution, meta = (HideSelfPin = "true"))
-	void LoopOnPoints(UPARAM(ref) FPCGContext& InContext, const UPCGPointData* InData, UPCGPointData*& OutData, const UObject* InCustomObject = nullptr) const;
+	void LoopOnPoints(UPARAM(ref) FPCGContext& InContext, const UPCGPointData* InData, UPCGPointData*& OutData, UPCGPointData* OptionalOutData = nullptr, const UObject* OptionalCustomObject = nullptr) const;
 
 	UFUNCTION(BlueprintCallable, Category = Execution, meta = (HideSelfPin = "true"))
-	void LoopOnPointPairs(UPARAM(ref) FPCGContext& InContext, const UPCGPointData* InA, const UPCGPointData* InB, UPCGPointData*& OutData, const UObject* InCustomObject = nullptr) const;
+	void LoopOnPointPairs(UPARAM(ref) FPCGContext& InContext, const UPCGPointData* InA, const UPCGPointData* InB, UPCGPointData*& OutData, UPCGPointData* OptionalOutData = nullptr, const UObject* OptionalCustomObject = nullptr) const;
 
 	/** Called after object creation to setup the object callbacks */
 	void Initialize();
