PCG:
*Added delete to node context menu
*Blocked pcg editor from deleteing the input/output nodes

diff --git a/Engine/Plugins/Experimental/PCG/Source/PCGEditor/Private/PCGEditor.cpp b/Engine/Plugins/Experimental/PCG/Source/PCGEditor/Private/PCGEditor.cpp
--- a/Engine/Plugins/Experimental/PCG/Source/PCGEditor/Private/PCGEditor.cpp
+++ b/Engine/Plugins/Experimental/PCG/Source/PCGEditor/Private/PCGEditor.cpp
@@ -175,7 +175,7 @@ void FPCGEditor::SelectAllNodes()
 	}
 }
 
-bool FPCGEditor::CanSelectAllNodes()
+bool FPCGEditor::CanSelectAllNodes() const
 {
 	return GraphEditorWidget.IsValid();
 }
@@ -190,11 +190,15 @@ void FPCGEditor::DeleteSelectedNodes()
 		for (UObject* Object : GraphEditorWidget->GetSelectedNodes())
 		{
 			UPCGEditorGraphNode* PCGEditorGraphNode = CastChecked<UPCGEditorGraphNode>(Object);
-			UPCGNode* PCGNode = PCGEditorGraphNode->GetPCGNode();
-			check(PCGNode);
 
-			PCGGraph->RemoveNode(PCGNode);
-			PCGEditorGraphNode->DestroyNode();
+			if (PCGEditorGraphNode->CanUserDeleteNode())
+			{
+				UPCGNode* PCGNode = PCGEditorGraphNode->GetPCGNode();
+				check(PCGNode);
+
+				PCGGraph->RemoveNode(PCGNode);
+				PCGEditorGraphNode->DestroyNode();
+			}
 		}
 
 		GraphEditorWidget->ClearSelectionSet();
@@ -202,6 +206,24 @@ void FPCGEditor::DeleteSelectedNodes()
 	}
 }
 
+bool FPCGEditor::CanDeleteSelectedNodes() const
+{
+	if (GraphEditorWidget.IsValid())
+	{
+		for (UObject* Object : GraphEditorWidget->GetSelectedNodes())
+		{
+			UPCGEditorGraphNode* PCGEditorGraphNode = CastChecked<UPCGEditorGraphNode>(Object);
+
+			if (PCGEditorGraphNode->CanUserDeleteNode())
+			{
+				return true;
+			}
+		}
+	}
+
+	return false;
+}
+
 void FPCGEditor::OnAlignTop()
 {
 	if (GraphEditorWidget.IsValid())
@@ -284,8 +306,8 @@ TSharedRef<SGraphEditor> FPCGEditor::CreateGraphEditorWidget()
 		FCanExecuteAction::CreateSP(this, &FPCGEditor::CanSelectAllNodes));
 
 	GraphEditorCommands->MapAction(FGenericCommands::Get().Delete,
-		FExecuteAction::CreateSP(this, &FPCGEditor::DeleteSelectedNodes)
-	);
+		FExecuteAction::CreateSP(this, &FPCGEditor::DeleteSelectedNodes),
+		FCanExecuteAction::CreateSP(this, &FPCGEditor::CanDeleteSelectedNodes));
 
 	// Alignment Commands
 	GraphEditorCommands->MapAction(FGraphEditorCommands::Get().AlignNodesTop,
diff --git a/Engine/Plugins/Experimental/PCG/Source/PCGEditor/Private/PCGEditor.h b/Engine/Plugins/Experimental/PCG/Source/PCGEditor/Private/PCGEditor.h
--- a/Engine/Plugins/Experimental/PCG/Source/PCGEditor/Private/PCGEditor.h
+++ b/Engine/Plugins/Experimental/PCG/Source/PCGEditor/Private/PCGEditor.h
@@ -32,10 +32,12 @@ private:
 	/** Select every node in the graph */
 	void SelectAllNodes();
 	/** Whether we can select every node */
-	bool CanSelectAllNodes();
+	bool CanSelectAllNodes() const;
 
 	/** Delete all selected nodes in the graph */
 	void DeleteSelectedNodes();
+	/** Whether we can delete all selected nodes */
+	bool CanDeleteSelectedNodes() const;
 
 	void OnAlignTop();
 	void OnAlignMiddle();
diff --git a/Engine/Plugins/Experimental/PCG/Source/PCGEditor/Private/PCGEditorGraphNode.cpp b/Engine/Plugins/Experimental/PCG/Source/PCGEditor/Private/PCGEditorGraphNode.cpp
--- a/Engine/Plugins/Experimental/PCG/Source/PCGEditor/Private/PCGEditorGraphNode.cpp
+++ b/Engine/Plugins/Experimental/PCG/Source/PCGEditor/Private/PCGEditorGraphNode.cpp
@@ -2,15 +2,16 @@
 
 #include "PCGEditorGraphNode.h"
 
+#include "PCGEditorGraphSchema.h"
+#include "PCGNode.h"
+#include "PCGSettings.h"
+
 #include "EdGraph/EdGraphPin.h"
+#include "Framework/Commands/GenericCommands.h"
 #include "GraphEditorActions.h"
 #include "ToolMenu.h"
 #include "ToolMenuSection.h"
 
-#include "PCGNode.h"
-#include "PCGSettings.h"
-#include "PCGEditorGraphSchema.h"
-
 #define LOCTEXT_NAMESPACE "PCGEditorGraphNode"
 
 void UPCGEditorGraphNode::Construct(UPCGNode* InPCGNode, EPCGEditorGraphNodeType InNodeType)
@@ -37,31 +38,38 @@ void UPCGEditorGraphNode::GetNodeContextMenuActions(UToolMenu* Menu, class UGrap
 		return;
 	}
 
-	FToolMenuSection& NodeSection = Menu->AddSection("EdGraphSchemaNodeActions", LOCTEXT("NodeActionsHeader", "Node Actions"));
 	{
-		NodeSection.AddMenuEntry(FGraphEditorCommands::Get().BreakNodeLinks);
+		FToolMenuSection& Section = Menu->AddSection("EdGraphSchemaNodeActions", LOCTEXT("NodeActionsHeader", "Node Actions"));
+		Section.AddMenuEntry(FGraphEditorCommands::Get().BreakNodeLinks);
 	}
 
-	FToolMenuSection& Section = Menu->AddSection("EdGraphSchemaOrganization", LOCTEXT("OrganizationHeader", "Organization"));
-	Section.AddSubMenu("Alignment", LOCTEXT("AlignmentHeader", "Alignment"), FText(), FNewToolMenuDelegate::CreateLambda([](UToolMenu* AlignmentMenu)
+	{
+		FToolMenuSection& Section = Menu->AddSection("EdGraphSchemaGeneral", LOCTEXT("GeneralHeader", "General"));
+		Section.AddMenuEntry(FGenericCommands::Get().Delete);
+	}
+
+	{
+		FToolMenuSection& Section = Menu->AddSection("EdGraphSchemaOrganization", LOCTEXT("OrganizationHeader", "Organization"));
+		Section.AddSubMenu("Alignment", LOCTEXT("AlignmentHeader", "Alignment"), FText(), FNewToolMenuDelegate::CreateLambda([](UToolMenu* AlignmentMenu)
 		{
 			{
-				FToolMenuSection& InSection = AlignmentMenu->AddSection("EdGraphSchemaAlignment", LOCTEXT("AlignHeader", "Align"));
-				InSection.AddMenuEntry(FGraphEditorCommands::Get().AlignNodesTop);
-				InSection.AddMenuEntry(FGraphEditorCommands::Get().AlignNodesMiddle);
-				InSection.AddMenuEntry(FGraphEditorCommands::Get().AlignNodesBottom);
-				InSection.AddMenuEntry(FGraphEditorCommands::Get().AlignNodesLeft);
-				InSection.AddMenuEntry(FGraphEditorCommands::Get().AlignNodesCenter);
-				InSection.AddMenuEntry(FGraphEditorCommands::Get().AlignNodesRight);
-				InSection.AddMenuEntry(FGraphEditorCommands::Get().StraightenConnections);
+				FToolMenuSection& SubSection = AlignmentMenu->AddSection("EdGraphSchemaAlignment", LOCTEXT("AlignHeader", "Align"));
+				SubSection.AddMenuEntry(FGraphEditorCommands::Get().AlignNodesTop);
+				SubSection.AddMenuEntry(FGraphEditorCommands::Get().AlignNodesMiddle);
+				SubSection.AddMenuEntry(FGraphEditorCommands::Get().AlignNodesBottom);
+				SubSection.AddMenuEntry(FGraphEditorCommands::Get().AlignNodesLeft);
+				SubSection.AddMenuEntry(FGraphEditorCommands::Get().AlignNodesCenter);
+				SubSection.AddMenuEntry(FGraphEditorCommands::Get().AlignNodesRight);
+				SubSection.AddMenuEntry(FGraphEditorCommands::Get().StraightenConnections);
 			}
 
 			{
-				FToolMenuSection& InSection = AlignmentMenu->AddSection("EdGraphSchemaDistribution", LOCTEXT("DistributionHeader", "Distribution"));
-				InSection.AddMenuEntry(FGraphEditorCommands::Get().DistributeNodesHorizontally);
-				InSection.AddMenuEntry(FGraphEditorCommands::Get().DistributeNodesVertically);
+				FToolMenuSection& SubSection = AlignmentMenu->AddSection("EdGraphSchemaDistribution", LOCTEXT("DistributionHeader", "Distribution"));
+				SubSection.AddMenuEntry(FGraphEditorCommands::Get().DistributeNodesHorizontally);
+				SubSection.AddMenuEntry(FGraphEditorCommands::Get().DistributeNodesVertically);
 			}
 		}));
+	}
 }
 
 void UPCGEditorGraphNode::AllocateDefaultPins()
@@ -92,4 +100,9 @@ void UPCGEditorGraphNode::AutowireNewNode(UEdGraphPin* FromPin)
 	NodeConnectionListChanged();
 }
 
+bool UPCGEditorGraphNode::CanUserDeleteNode() const
+{
+	return NodeType != EPCGEditorGraphNodeType::Input && NodeType != EPCGEditorGraphNodeType::Output;
+}
+
 #undef LOCTEXT_NAMESPACE
diff --git a/Engine/Plugins/Experimental/PCG/Source/PCGEditor/Private/PCGEditorGraphNode.h b/Engine/Plugins/Experimental/PCG/Source/PCGEditor/Private/PCGEditorGraphNode.h
--- a/Engine/Plugins/Experimental/PCG/Source/PCGEditor/Private/PCGEditorGraphNode.h
+++ b/Engine/Plugins/Experimental/PCG/Source/PCGEditor/Private/PCGEditorGraphNode.h
@@ -32,6 +32,7 @@ public:
 	virtual void GetNodeContextMenuActions(UToolMenu* Menu, class UGraphNodeContextMenuContext* Context) const override;
 	virtual void AllocateDefaultPins() override;
 	virtual void AutowireNewNode(UEdGraphPin* FromPin) override;
+	virtual bool CanUserDeleteNode() const override;
 	// ~End UEdGraphNode interface
 
 	UPCGNode* GetPCGNode() { return PCGNode; }
