Added position/normal/scale/transform property in Metadata operations
Added splatting/default initialization for some types (float -> vector, vector -> transform)

diff --git a/Engine/Plugins/Experimental/PCG/Source/PCG/Private/Elements/PCGMetadataElement.cpp b/Engine/Plugins/Experimental/PCG/Source/PCG/Private/Elements/PCGMetadataElement.cpp
--- a/Engine/Plugins/Experimental/PCG/Source/PCG/Private/Elements/PCGMetadataElement.cpp
+++ b/Engine/Plugins/Experimental/PCG/Source/PCG/Private/Elements/PCGMetadataElement.cpp
@@ -6,6 +6,47 @@
 #include "Metadata/PCGMetadataAttribute.h"
 #include "Metadata/PCGMetadataAttributeTpl.h"
 
+namespace PCGMetadataOperations
+{
+	template<typename U, typename T>
+	bool SetValue(const TArray<FPCGPoint>& InPoints, FPCGMetadataAttributeBase* AttributeBase, TFunctionRef<U(const FPCGPoint& InPoint)> PropGetter)
+	{
+		if (AttributeBase && AttributeBase->GetTypeId() == PCG::Private::MetadataTypes<T>::Id)
+		{
+			FPCGMetadataAttribute<T>* Attribute = static_cast<FPCGMetadataAttribute<T>*>(AttributeBase);
+			for (const FPCGPoint& Point : InPoints)
+			{
+				Attribute->SetValue(Point.MetadataEntry, PropGetter(Point));
+			}
+
+			return true;
+		}
+		else
+		{
+			return false;
+		}
+	}
+
+	template<typename T, typename U>
+	bool SetValue(const FPCGMetadataAttributeBase* AttributeBase, TArray<FPCGPoint>& InPoints, TFunctionRef<void(FPCGPoint& OutPoint, const U& InValue)> PropSetter)
+	{
+		if (AttributeBase && AttributeBase->GetTypeId() == PCG::Private::MetadataTypes<T>::Id)
+		{
+			const FPCGMetadataAttribute<T>* Attribute = static_cast<const FPCGMetadataAttribute<T>*>(AttributeBase);
+			for (FPCGPoint& Point : InPoints)
+			{
+				PropSetter(Point, U(Attribute->GetValueFromItemKey(Point.MetadataEntry)));
+			}
+
+			return true;
+		}
+		else
+		{
+			return false;
+		}
+	}
+}
+
 FPCGElementPtr UPCGMetadataOperationSettings::CreateElement() const
 {
 	return MakeShared<FPCGMetadataOperationElement>();
@@ -74,63 +115,109 @@ bool FPCGMetadataOperationElement::ExecuteInternal(FPCGContext* Context) const
 		{
 			if (Settings->PointProperty == EPCGPointProperties::Density)
 			{
+				auto DensityGetter = [](const FPCGPoint& InPoint) { return InPoint.Density; };
+
 				if (!SampledData->Metadata->HasAttribute(Settings->DestinationAttribute))
 				{
 					SampledData->Metadata->CreateFloatAttribute(Settings->DestinationAttribute, 0.0f, /*bAllowsInterpolation=*/true, /*bOverrideParent=*/true);
 				}
 
 				FPCGMetadataAttributeBase* AttributeBase = SampledData->Metadata->GetMutableAttribute(Settings->DestinationAttribute);
-				if (AttributeBase && AttributeBase->GetTypeId() == PCG::Private::MetadataTypes<float>::Id)
-				{
-					FPCGMetadataAttribute<float>* Attribute = static_cast<FPCGMetadataAttribute<float>*>(AttributeBase);
-					for (FPCGPoint& Point : SampledPoints)
-					{
-						Attribute->SetValue(Point.MetadataEntry, Point.Density);
-					}
-				}
-				else
+				if(!PCGMetadataOperations::SetValue<float, float>(SampledPoints, AttributeBase, DensityGetter) &&
+					!PCGMetadataOperations::SetValue<float, FVector>(SampledPoints, AttributeBase, DensityGetter) &&
+					!PCGMetadataOperations::SetValue<float, FVector4>(SampledPoints, AttributeBase, DensityGetter))
 				{
 					PCGE_LOG(Error, "Attribute %s already exists but its type is not compatible", *Settings->DestinationAttribute.ToString());
 				}
 			}
 			else if (Settings->PointProperty == EPCGPointProperties::Extents)
 			{
+				auto ExtentsGetter = [](const FPCGPoint& InPoint) { return InPoint.Extents; };
+
 				if (!SampledData->Metadata->HasAttribute(Settings->DestinationAttribute))
 				{
 					SampledData->Metadata->CreateVectorAttribute(Settings->DestinationAttribute, FVector::Zero() , /*bAllowsInterpolation=*/true, /*bOverrideParent=*/true);
 				}
 
 				FPCGMetadataAttributeBase* AttributeBase = SampledData->Metadata->GetMutableAttribute(Settings->DestinationAttribute);
-				if (AttributeBase && AttributeBase->GetTypeId() == PCG::Private::MetadataTypes<FVector>::Id)
-				{
-					FPCGMetadataAttribute<FVector>* Attribute = static_cast<FPCGMetadataAttribute<FVector>*>(AttributeBase);
-					for (FPCGPoint& Point : SampledPoints)
-					{
-						Attribute->SetValue(Point.MetadataEntry, Point.Extents);
-					}
-				}
-				else
+				if(!PCGMetadataOperations::SetValue<FVector, FVector>(SampledPoints, AttributeBase, ExtentsGetter) &&
+					!PCGMetadataOperations::SetValue<FVector, FTransform>(SampledPoints, AttributeBase, ExtentsGetter))
 				{
 					PCGE_LOG(Error, "Attribute %s already exists but its type is not compatible", *Settings->DestinationAttribute.ToString());
 				}
 			}
 			else if(Settings->PointProperty == EPCGPointProperties::Color)
 			{
+				auto ColorGetter = [](const FPCGPoint& InPoint) { return InPoint.Color; };
+
 				if (!SampledData->Metadata->HasAttribute(Settings->DestinationAttribute))
 				{
 					SampledData->Metadata->CreateVector4Attribute(Settings->DestinationAttribute, FVector4::Zero(), /*bAllowsInterpolation=*/true, /*bOverrideParent=*/true);
 				}
 
 				FPCGMetadataAttributeBase* AttributeBase = SampledData->Metadata->GetMutableAttribute(Settings->DestinationAttribute);
-				if (AttributeBase && AttributeBase->GetTypeId() == PCG::Private::MetadataTypes<FVector4>::Id)
+				if(!PCGMetadataOperations::SetValue<FVector4, FVector4>(SampledPoints, AttributeBase, ColorGetter))
+				{
+					PCGE_LOG(Error, "Attribute %s already exists but its type is not compatible", *Settings->DestinationAttribute.ToString());
+				}
+			}
+			else if (Settings->PointProperty == EPCGPointProperties::Position)
+			{
+				auto PositionGetter = [](const FPCGPoint& InPoint) { return InPoint.Transform.GetLocation(); };
+
+				if (!SampledData->Metadata->HasAttribute(Settings->DestinationAttribute))
+				{
+					SampledData->Metadata->CreateVectorAttribute(Settings->DestinationAttribute, FVector::Zero(), /*bAllowsInterpolation=*/true, /*bOverrideParent=*/true);
+				}
+
+				FPCGMetadataAttributeBase* AttributeBase = SampledData->Metadata->GetMutableAttribute(Settings->DestinationAttribute);
+				if (!PCGMetadataOperations::SetValue<FVector, FVector>(SampledPoints, AttributeBase, PositionGetter) &&
+					!PCGMetadataOperations::SetValue<FVector, FTransform>(SampledPoints, AttributeBase, PositionGetter))
+				{
+					PCGE_LOG(Error, "Attribute %s already exists but its type is not compatible", *Settings->DestinationAttribute.ToString());
+				}
+			}
+			else if (Settings->PointProperty == EPCGPointProperties::Rotation)
+			{
+				auto RotationGetter = [](const FPCGPoint& InPoint) { return InPoint.Transform.GetRotation(); };
+
+				if (!SampledData->Metadata->HasAttribute(Settings->DestinationAttribute))
+				{
+					SampledData->Metadata->CreateQuatAttribute(Settings->DestinationAttribute, FQuat::Identity, /*bAllowsInterpolation=*/true, /*bOverrideParent=*/true);
+				}
+
+				FPCGMetadataAttributeBase* AttributeBase = SampledData->Metadata->GetMutableAttribute(Settings->DestinationAttribute);
+				if (!PCGMetadataOperations::SetValue<FQuat, FQuat>(SampledPoints, AttributeBase, RotationGetter))
 				{
-					FPCGMetadataAttribute<FVector4>* Attribute = static_cast<FPCGMetadataAttribute<FVector4>*>(AttributeBase);
-					for (FPCGPoint& Point : SampledPoints)
-					{
-						Attribute->SetValue(Point.MetadataEntry, Point.Color);
-					}
+					PCGE_LOG(Error, "Attribute %s already exists but its type is not compatible", *Settings->DestinationAttribute.ToString());
 				}
-				else
+			}
+			else if (Settings->PointProperty == EPCGPointProperties::Scale)
+			{
+				auto ScaleGetter = [](const FPCGPoint& InPoint) { return InPoint.Transform.GetScale3D(); };
+
+				if (!SampledData->Metadata->HasAttribute(Settings->DestinationAttribute))
+				{
+					SampledData->Metadata->CreateVectorAttribute(Settings->DestinationAttribute, FVector::One(), /*bAllowsInterpolation=*/true, /*bOverrideParent=*/true);
+				}
+
+				FPCGMetadataAttributeBase* AttributeBase = SampledData->Metadata->GetMutableAttribute(Settings->DestinationAttribute);
+				if (!PCGMetadataOperations::SetValue<FVector, FVector>(SampledPoints, AttributeBase, ScaleGetter))
+				{
+					PCGE_LOG(Error, "Attribute %s already exists but its type is not compatible", *Settings->DestinationAttribute.ToString());
+				}
+			}
+			else if (Settings->PointProperty == EPCGPointProperties::Transform)
+			{
+				auto TransformGetter = [](const FPCGPoint& InPoint) { return InPoint.Transform; };
+
+				if (!SampledData->Metadata->HasAttribute(Settings->DestinationAttribute))
+				{
+					SampledData->Metadata->CreateTransformAttribute(Settings->DestinationAttribute, FTransform::Identity, /*bAllowsInterpolation=*/true, /*bOverrideParent=*/true);
+				}
+
+				FPCGMetadataAttributeBase* AttributeBase = SampledData->Metadata->GetMutableAttribute(Settings->DestinationAttribute);
+				if (!PCGMetadataOperations::SetValue<FTransform, FTransform>(SampledPoints, AttributeBase, TransformGetter))
 				{
 					PCGE_LOG(Error, "Attribute %s already exists but its type is not compatible", *Settings->DestinationAttribute.ToString());
 				}
@@ -140,48 +227,75 @@ bool FPCGMetadataOperationElement::ExecuteInternal(FPCGContext* Context) const
 		{
 			if (Settings->PointProperty == EPCGPointProperties::Density)
 			{
+				auto DensitySetter = [](FPCGPoint& InPoint, const float& InValue) { InPoint.Density = InValue; };
+
+				const FPCGMetadataAttributeBase* AttributeBase = SampledData->Metadata->GetConstAttribute(Settings->SourceAttribute);
+				if(!PCGMetadataOperations::SetValue<float, float>(AttributeBase, SampledPoints, DensitySetter))
+				{
+					PCGE_LOG(Error, "Attribute %s already exists but its type is not compatible", *Settings->SourceAttribute.ToString());
+				}
+			}
+			else if (Settings->PointProperty == EPCGPointProperties::Extents)
+			{
+				auto ExtentsSetter = [](FPCGPoint& InPoint, const FVector& InValue) { InPoint.Extents = InValue; };
+
 				const FPCGMetadataAttributeBase* AttributeBase = SampledData->Metadata->GetConstAttribute(Settings->SourceAttribute);
-				if (AttributeBase && AttributeBase->GetTypeId() == PCG::Private::MetadataTypes<float>::Id)
+				if(!PCGMetadataOperations::SetValue<float, FVector>(AttributeBase, SampledPoints, ExtentsSetter) &&
+					!PCGMetadataOperations::SetValue<FVector, FVector>(AttributeBase, SampledPoints, ExtentsSetter))
 				{
-					const FPCGMetadataAttribute<float>* Attribute = static_cast<const FPCGMetadataAttribute<float>*>(AttributeBase);
-					for (FPCGPoint& Point : SampledPoints)
-					{
-						Point.Density = Attribute->GetValueFromItemKey(Point.MetadataEntry);
-					}
+					PCGE_LOG(Error, "Attribute %s already exists but its type is not compatible", *Settings->SourceAttribute.ToString());
 				}
-				else
+			}
+			else if (Settings->PointProperty == EPCGPointProperties::Color)
+			{
+				auto ColorSetter = [](FPCGPoint& InPoint, const FVector4& InValue) { InPoint.Color = InValue; };
+
+				const FPCGMetadataAttributeBase* AttributeBase = SampledData->Metadata->GetConstAttribute(Settings->SourceAttribute);
+				if(!PCGMetadataOperations::SetValue<float, FVector4>(AttributeBase, SampledPoints, ColorSetter) &&
+					!PCGMetadataOperations::SetValue<FVector4, FVector4>(AttributeBase, SampledPoints, ColorSetter))
 				{
 					PCGE_LOG(Error, "Attribute %s already exists but its type is not compatible", *Settings->SourceAttribute.ToString());
 				}
 			}
-			else if (Settings->PointProperty == EPCGPointProperties::Extents)
+			else if (Settings->PointProperty == EPCGPointProperties::Position)
 			{
+				auto PositionSetter = [](FPCGPoint& InPoint, const FVector& InValue) { InPoint.Transform.SetLocation(InValue); };
+
 				const FPCGMetadataAttributeBase* AttributeBase = SampledData->Metadata->GetConstAttribute(Settings->SourceAttribute);
-				if (AttributeBase && AttributeBase->GetTypeId() == PCG::Private::MetadataTypes<FVector>::Id)
+				if (!PCGMetadataOperations::SetValue<float, FVector>(AttributeBase, SampledPoints, PositionSetter) &&
+					!PCGMetadataOperations::SetValue<FVector, FVector>(AttributeBase, SampledPoints, PositionSetter))
 				{
-					const FPCGMetadataAttribute<FVector>* Attribute = static_cast<const FPCGMetadataAttribute<FVector>*>(AttributeBase);
-					for (FPCGPoint& Point : SampledPoints)
-					{
-						Point.Extents = Attribute->GetValueFromItemKey(Point.MetadataEntry);
-					}
+					PCGE_LOG(Error, "Attribute %s already exists but its type is not compatible", *Settings->SourceAttribute.ToString());
 				}
-				else
+			}
+			else if (Settings->PointProperty == EPCGPointProperties::Rotation)
+			{
+				auto RotationSetter = [](FPCGPoint& InPoint, const FQuat& InValue) { InPoint.Transform.SetRotation(InValue.GetNormalized()); };
+
+				const FPCGMetadataAttributeBase* AttributeBase = SampledData->Metadata->GetConstAttribute(Settings->SourceAttribute);
+				if (!PCGMetadataOperations::SetValue<FQuat, FQuat>(AttributeBase, SampledPoints, RotationSetter))
 				{
 					PCGE_LOG(Error, "Attribute %s already exists but its type is not compatible", *Settings->SourceAttribute.ToString());
 				}
 			}
-			else if (Settings->PointProperty == EPCGPointProperties::Color)
+			else if (Settings->PointProperty == EPCGPointProperties::Scale)
 			{
+				auto ScaleSetter = [](FPCGPoint& InPoint, const FVector& InValue) { InPoint.Transform.SetScale3D(InValue); };
+
 				const FPCGMetadataAttributeBase* AttributeBase = SampledData->Metadata->GetConstAttribute(Settings->SourceAttribute);
-				if (AttributeBase && AttributeBase->GetTypeId() == PCG::Private::MetadataTypes<FVector4>::Id)
+				if (!PCGMetadataOperations::SetValue<float, FVector>(AttributeBase, SampledPoints, ScaleSetter) &&
+					!PCGMetadataOperations::SetValue<FVector, FVector>(AttributeBase, SampledPoints, ScaleSetter))
 				{
-					const FPCGMetadataAttribute<FVector4>* Attribute = static_cast<const FPCGMetadataAttribute<FVector4>*>(AttributeBase);
-					for (FPCGPoint& Point : SampledPoints)
-					{
-						Point.Color = Attribute->GetValueFromItemKey(Point.MetadataEntry);
-					}
+					PCGE_LOG(Error, "Attribute %s already exists but its type is not compatible", *Settings->SourceAttribute.ToString());
 				}
-				else
+			}
+			else if (Settings->PointProperty == EPCGPointProperties::Transform)
+			{
+				auto TransformSetter = [](FPCGPoint& InPoint, const FTransform& InValue) { InPoint.Transform = InValue; };
+
+				const FPCGMetadataAttributeBase* AttributeBase = SampledData->Metadata->GetConstAttribute(Settings->SourceAttribute);
+				if (!PCGMetadataOperations::SetValue<FVector, FTransform>(AttributeBase, SampledPoints, TransformSetter) &&
+					!PCGMetadataOperations::SetValue<FTransform, FTransform>(AttributeBase, SampledPoints, TransformSetter))
 				{
 					PCGE_LOG(Error, "Attribute %s already exists but its type is not compatible", *Settings->SourceAttribute.ToString());
 				}
diff --git a/Engine/Plugins/Experimental/PCG/Source/PCG/Private/Metadata/PCGMetadataAttributeTpl.h b/Engine/Plugins/Experimental/PCG/Source/PCG/Private/Metadata/PCGMetadataAttributeTpl.h
--- a/Engine/Plugins/Experimental/PCG/Source/PCG/Private/Metadata/PCGMetadataAttributeTpl.h
+++ b/Engine/Plugins/Experimental/PCG/Source/PCG/Private/Metadata/PCGMetadataAttributeTpl.h
@@ -156,6 +156,12 @@ public:
 		SetValueFromValueKey(ItemKey, AddValue(InValue));
 	}
 
+	template<typename U>
+	void SetValue(PCGMetadataEntryKey ItemKey, const U& InValue)
+	{
+		SetValueFromValueKey(ItemKey, AddValue(T(InValue)));
+	}
+
 	T GetValueFromItemKey(PCGMetadataEntryKey ItemKey) const
 	{
 		return GetValue(GetValueKey(ItemKey));
diff --git a/Engine/Plugins/Experimental/PCG/Source/PCG/Public/Elements/PCGMetadataElement.h b/Engine/Plugins/Experimental/PCG/Source/PCG/Public/Elements/PCGMetadataElement.h
--- a/Engine/Plugins/Experimental/PCG/Source/PCG/Public/Elements/PCGMetadataElement.h
+++ b/Engine/Plugins/Experimental/PCG/Source/PCG/Public/Elements/PCGMetadataElement.h
@@ -13,7 +13,11 @@ enum class EPCGPointProperties : uint8
 {
 	Density,
 	Extents,
-	Color
+	Color,
+	Position,
+	Rotation,
+	Scale,
+	Transform
 };
 
 UENUM()
diff --git a/Engine/Plugins/Experimental/PCG/Source/PCG/Public/Metadata/PCGMetadata.h b/Engine/Plugins/Experimental/PCG/Source/PCG/Public/Metadata/PCGMetadata.h
--- a/Engine/Plugins/Experimental/PCG/Source/PCG/Public/Metadata/PCGMetadata.h
+++ b/Engine/Plugins/Experimental/PCG/Source/PCG/Public/Metadata/PCGMetadata.h
@@ -50,7 +50,7 @@ public:
 	UFUNCTION(BlueprintCallable, Category = "PCG|Metadata")
 	void CreateVectorAttribute(FName AttributeName, const FVector& DefaultValue, bool bAllowsInterpolation, bool bOverrideParent);
 
-	UFUNCTION(BlueprintCallable, Category = "PCG|Metadata")
+	UFUNCTION(BlueprintCallable, Category = "PCG|Metadata", meta=(DisplayName = "Create Vector4 Attribute"))
 	void CreateVector4Attribute(FName AttributeName, const FVector4& DefaultValue, bool bAllowsInterpolation, bool bOverrideParent);
 
 	UFUNCTION(BlueprintCallable, Category = "PCG|Metadata")
