Fixed issue with inherited value retrieval, by making sure we create placeholder attributes; this also led to some code cleanup & clarification.
Made metadata code more point-agnostic.
Added metadata key blueprint methods.
Added copy point that sets up the metadata automatically

diff --git a/Engine/Plugins/Experimental/PCG/Source/PCG/Private/Metadata/PCGMetadata.cpp b/Engine/Plugins/Experimental/PCG/Source/PCG/Private/Metadata/PCGMetadata.cpp
--- a/Engine/Plugins/Experimental/PCG/Source/PCG/Private/Metadata/PCGMetadata.cpp
+++ b/Engine/Plugins/Experimental/PCG/Source/PCG/Private/Metadata/PCGMetadata.cpp
@@ -29,12 +29,12 @@ void UPCGMetadata::Initialize(const UPCGMetadata* InParent)
 	check(!Parent && Attributes.Num() == 0);
 	Parent = (InParent != this ? InParent : nullptr);
 	ItemKeyOffset = Parent ? Parent->GetItemCountForChild() : 0;
-	GetAllAttributeNames(AttributeNames);
+	AddAttributes(InParent);
 }
 
 void UPCGMetadata::AddAttributes(const UPCGMetadata* InOther)
 {
-	if (!InOther || InOther == Parent)
+	if (!InOther)
 	{
 		return;
 	}
@@ -47,11 +47,14 @@ void UPCGMetadata::AddAttributes(const UPCGMetadata* InOther)
 		}
 		else
 		{
-			CopyAttribute(OtherAttribute.Value, OtherAttribute.Key, /*bKeepParent=*/false, /*bCopyEntries=*/false, /*bCopyValues=*/false);
+			CopyAttribute(OtherAttribute.Value, OtherAttribute.Key, /*bKeepParent=*/InOther == Parent, /*bCopyEntries=*/false, /*bCopyValues=*/false);
 		}
 	}
 
-	OtherParents.Add(InOther);
+	if(InOther != Parent)
+	{
+		OtherParents.Add(InOther);
+	}
 }
 
 const UPCGMetadata* UPCGMetadata::GetRoot() const
@@ -86,17 +89,11 @@ void UPCGMetadata::AddAttributeInternal(FName AttributeName, FPCGMetadataAttribu
 {
 	// This call assumes we have a write lock on the attribute map.
 	Attributes.Add(AttributeName, Attribute);
-	AttributeNames.Add(AttributeName);
-	HiddenAttributes.Remove(AttributeName);
 }
 
 void UPCGMetadata::RemoveAttributeInternal(FName AttributeName)
 {
 	Attributes.Remove(AttributeName);
-	if (!Parent || !Parent->HasAttribute(AttributeName))
-	{
-		AttributeNames.Remove(AttributeName);
-	}
 }
 
 template<typename T>
@@ -139,11 +136,6 @@ FPCGMetadataAttributeBase* UPCGMetadata::GetMutableAttribute(FName AttributeName
 	}
 	AttributeLock.ReadUnlock();
 
-	if (!Attribute && ParentHasAttribute(AttributeName))
-	{
-		Attribute = CopyAttribute(AttributeName, AttributeName, /*bKeepParent=*/true, /*bCopyEntries=*/false, /*bCopyValues=*/false);
-	}
-
 	return Attribute;
 }
 
@@ -158,33 +150,18 @@ const FPCGMetadataAttributeBase* UPCGMetadata::GetConstAttribute(FName Attribute
 	}
 	AttributeLock.ReadUnlock();
 
-	if (!Attribute && ParentHasAttribute(AttributeName))
-	{
-		return Parent->GetConstAttribute(AttributeName);
-	}
-	else
-	{
-		return Attribute;
-	}
+	return Attribute;
 }
 
 bool UPCGMetadata::HasAttribute(FName AttributeName) const
 {
 	FReadScopeLock ScopeLock(AttributeLock);
-	return AttributeNames.Contains(AttributeName);
+	return Attributes.Contains(AttributeName);
 }
 
 bool UPCGMetadata::ParentHasAttribute(FName AttributeName) const
 {
-	if (Parent && Parent->HasAttribute(AttributeName))
-	{
-		FReadScopeLock ScopeLock(AttributeLock);
-		return !HiddenAttributes.Contains(AttributeName);
-	}
-	else
-	{
-		return false;
-	}
+	return Parent && Parent->HasAttribute(AttributeName);
 }
 
 void UPCGMetadata::CreateFloatAttribute(FName AttributeName, float DefaultValue, bool bAllowsInterpolation, bool bOverrideParent)
@@ -276,13 +253,6 @@ void UPCGMetadata::RenameAttribute(FName AttributeToRename, FName NewAttributeNa
 	}
 	AttributeLock.WriteUnlock();
 
-	if(!bRenamed && ParentHasAttribute(AttributeToRename))
-	{
-		// If attribute is present in the parent, we'll otherwise copy the attribute from the parent with a new name
-		CopyAttribute(AttributeToRename, NewAttributeName, /*bKeepParent=*/true);
-		bRenamed = true;
-	}
-
 	if (!bRenamed)
 	{
 		UE_LOG(LogPCG, Warning, TEXT("Attribute %s does not exist and therefore cannot be renamed"), *AttributeToRename.ToString());
@@ -306,10 +276,6 @@ void UPCGMetadata::ClearAttribute(FName AttributeToClear)
 	{
 		Attribute->ClearEntries();
 	}
-	else if(ParentHasAttribute(AttributeToClear))
-	{
-		CopyAttribute(AttributeToClear, AttributeToClear, /*bKeepParent=*/true, /*bCopyEntries=*/false, /*bCopyValues=*/false);
-	}
 }
 
 void UPCGMetadata::DeleteAttribute(FName AttributeToDelete)
@@ -329,13 +295,6 @@ void UPCGMetadata::DeleteAttribute(FName AttributeToDelete)
 	{
 		delete Attribute;
 	}
-	else if(ParentHasAttribute(AttributeToDelete))
-	{
-		// Otherwise, if it exists in parent hierarchy, then hide it
-		FWriteScopeLock ScopeLock(AttributeLock);
-		HiddenAttributes.Add(AttributeToDelete);
-		AttributeNames.Remove(AttributeToDelete);
-	}
 	else
 	{
 		UE_LOG(LogPCG, Verbose, TEXT("Attribute %s does not exist and therefore cannot be deleted"), *AttributeToDelete.ToString());
@@ -354,41 +313,63 @@ int64 UPCGMetadata::AddEntry(int64 ParentEntry)
 	return ParentKeys.Add(ParentEntry) + ItemKeyOffset;
 }
 
-PCGMetadataEntryKey UPCGMetadata::GetParentKey(PCGMetadataEntryKey LocalItemKey) const
+bool UPCGMetadata::InitializeOnSet(PCGMetadataEntryKey& InKey, PCGMetadataEntryKey InParentKeyA, const UPCGMetadata* InParentMetadataA, PCGMetadataEntryKey InParentKeyB, const UPCGMetadata* InParentMetadataB)
 {
-	if (LocalItemKey < ItemKeyOffset)
+	if (InKey == PCGInvalidEntryKey)
 	{
-		if (Parent)
+		if (InParentKeyA != PCGInvalidEntryKey && Parent == InParentMetadataA)
+		{
+			InKey = AddEntry(InParentKeyA);
+			return true;
+		}
+		else if (InParentKeyB != PCGInvalidEntryKey && Parent == InParentMetadataB)
 		{
-			return Parent->GetParentKey(LocalItemKey);
+			InKey = AddEntry(InParentKeyB);
+			return true;
 		}
 		else
 		{
-			return PCGInvalidEntryKey;
+			InKey = AddEntry();
+			return false;
 		}
 	}
+	else if(InKey < ItemKeyOffset)
+	{
+		InKey = AddEntry(InKey);
+		return false;
+	}
 	else
 	{
-		FReadScopeLock ScopeLock(ItemLock);
-		check(LocalItemKey - ItemKeyOffset < ParentKeys.Num());
-		return ParentKeys[LocalItemKey - ItemKeyOffset];
+		return false;
 	}
 }
 
-void UPCGMetadata::GetAllAttributeNames(TSet<FName>& FoundAttributeNames) const
+PCGMetadataEntryKey UPCGMetadata::GetParentKey(PCGMetadataEntryKey LocalItemKey) const
 {
-	for (const TPair<FName, FPCGMetadataAttributeBase*>& Attribute : Attributes)
+	if (LocalItemKey < ItemKeyOffset)
 	{
-		FoundAttributeNames.Add(Attribute.Key);
+		// Key is already in parent referential
+		return LocalItemKey;
 	}
-
-	if (Parent)
+	else
 	{
-		Parent->GetAllAttributeNames(FoundAttributeNames);
+		FReadScopeLock ScopeLock(ItemLock);
+		check(LocalItemKey - ItemKeyOffset < ParentKeys.Num());
+		return ParentKeys[LocalItemKey - ItemKeyOffset];
 	}
 }
 
+void UPCGMetadata::MergeAttributes(const FPCGPoint& InPointA, const FPCGPoint& InPointB, FPCGPoint& OutPoint, EPCGMetadataOp Op)
+{
+	MergeAttributes(InPointA.MetadataEntry, this, InPointB.MetadataEntry, this, OutPoint.MetadataEntry, Op);
+}
+
 void UPCGMetadata::MergeAttributes(const FPCGPoint& InPointA, const UPCGMetadata* InMetadataA, const FPCGPoint& InPointB, const UPCGMetadata* InMetadataB, FPCGPoint& OutPoint, EPCGMetadataOp Op)
+{
+	MergeAttributes(InPointA.MetadataEntry, InMetadataA, InPointB.MetadataEntry, InMetadataB, OutPoint.MetadataEntry, Op);
+}
+
+void UPCGMetadata::MergeAttributes(PCGMetadataEntryKey InKeyA, const UPCGMetadata* InMetadataA, PCGMetadataEntryKey InKeyB, const UPCGMetadata* InMetadataB, PCGMetadataEntryKey& OutKey, EPCGMetadataOp Op)
 {
 	// Early out: nothing to do if both input metadata are null / points have no assigned metadata
 	if (!InMetadataA && !InMetadataB)
@@ -397,35 +378,18 @@ void UPCGMetadata::MergeAttributes(const FPCGPoint& InPointA, const UPCGMetadata
 	}
 
 	// For each attribute in the current metadata, query the values from point A & B, apply operation on the result and finally store in the out point.
-	if (OutPoint.MetadataEntry == PCGInvalidEntryKey)
-	{
-		if (InPointA.MetadataEntry != PCGInvalidEntryKey && (Parent == InMetadataA || this == InMetadataA))
-		{
-			OutPoint.MetadataEntry = AddEntry(InPointA.MetadataEntry);
-		}
-		else if (InPointB.MetadataEntry != PCGInvalidEntryKey && (Parent == InMetadataB || this == InMetadataB))
-		{
-			OutPoint.MetadataEntry = AddEntry(InPointB.MetadataEntry);
-		}
-		else
-		{
-			OutPoint.MetadataEntry = AddEntry();
-		}
-	}
+	InitializeOnSet(OutKey, InKeyA, InMetadataA, InKeyB, InMetadataB);
 
-	//METADATA TODO: This isn't great, we're allocating memory for no reason here
 	AttributeLock.ReadLock();
-	TSet<FName> KnownAttributeNames = AttributeNames;
-	AttributeLock.ReadUnlock();
-
-	for(const FName& AttributeName : KnownAttributeNames)
+	for(const TPair<FName, FPCGMetadataAttributeBase*>& AttributePair : Attributes)
 	{
-		const FPCGMetadataAttributeBase* ConstAttribute = GetConstAttribute(AttributeName);
+		const FName& AttributeName = AttributePair.Key;
+		FPCGMetadataAttributeBase* Attribute = AttributePair.Value;
 
 		// Get attribute from A
 		const FPCGMetadataAttributeBase* AttributeA = InMetadataA ? InMetadataA->GetConstAttribute(AttributeName) : nullptr;
 
-		if (AttributeA && AttributeA->GetTypeId() != ConstAttribute->GetTypeId())
+		if (AttributeA && AttributeA->GetTypeId() != Attribute->GetTypeId())
 		{
 			UE_LOG(LogPCG, Error, TEXT("Metadata type mismatch with attribute %s"), *AttributeName.ToString());
 			AttributeA = nullptr;
@@ -434,7 +398,7 @@ void UPCGMetadata::MergeAttributes(const FPCGPoint& InPointA, const UPCGMetadata
 		// Get attribute from B
 		const FPCGMetadataAttributeBase* AttributeB = InMetadataB ? InMetadataB->GetConstAttribute(AttributeName) : nullptr;
 
-		if (AttributeB && AttributeB->GetTypeId() != ConstAttribute->GetTypeId())
+		if (AttributeB && AttributeB->GetTypeId() != Attribute->GetTypeId())
 		{
 			UE_LOG(LogPCG, Error, TEXT("Metadata type mismatch with attribute %s"), *AttributeName.ToString());
 			AttributeB = nullptr;
@@ -442,144 +406,122 @@ void UPCGMetadata::MergeAttributes(const FPCGPoint& InPointA, const UPCGMetadata
 
 		if (AttributeA || AttributeB)
 		{
-			FPCGMetadataAttributeBase* Attribute = GetMutableAttribute(AttributeName);
 			if (Attribute)
 			{
-				Attribute->SetValue(OutPoint.MetadataEntry, AttributeA, InPointA.MetadataEntry, AttributeB, InPointB.MetadataEntry, Op);
+				Attribute->SetValue(OutKey, AttributeA, InKeyA, AttributeB, InKeyB, Op);
 			}
 		}
 	}
+	AttributeLock.ReadUnlock();
 }
 
 void UPCGMetadata::ResetWeightedAttributes(FPCGPoint& OutPoint)
 {
-	if (OutPoint.MetadataEntry == PCGInvalidEntryKey)
-	{
-		OutPoint.MetadataEntry = AddEntry();
-	}
+	ResetWeightedAttributes(OutPoint.MetadataEntry);
+}
 
-	//METADATA TODO: This isn't great, we're allocating memory for no reason here
-	AttributeLock.ReadLock();
-	TSet<FName> KnownAttributeNames = AttributeNames;
-	AttributeLock.ReadUnlock();
+void UPCGMetadata::ResetWeightedAttributes(PCGMetadataEntryKey& OutKey)
+{
+	InitializeOnSet(OutKey);
 
-	for (const FName& AttributeName : KnownAttributeNames)
+	AttributeLock.ReadLock();
+	for(const TPair<FName, FPCGMetadataAttributeBase*>& AttributePair : Attributes)
 	{
-		const FPCGMetadataAttributeBase* ConstAttribute = GetConstAttribute(AttributeName);
+		FPCGMetadataAttributeBase* Attribute = AttributePair.Value;
 		
-		if (!ConstAttribute->AllowsInterpolation())
+		if (!Attribute->AllowsInterpolation())
 		{
 			continue;
 		}
 
-		FPCGMetadataAttributeBase* Attribute = GetMutableAttribute(AttributeName);
 		if (Attribute)
 		{
 			check(Attribute->AllowsInterpolation());
-			Attribute->SetZeroValue(OutPoint.MetadataEntry);
+			Attribute->SetZeroValue(OutKey);
 		}
 	}
+	AttributeLock.ReadUnlock();
 }
 
 void UPCGMetadata::AccumulateWeightedAttributes(const FPCGPoint& InPoint, const UPCGMetadata* InMetadata, float Weight, bool bSetNonInterpolableAttributes, FPCGPoint& OutPoint)
+{
+	AccumulateWeightedAttributes(InPoint.MetadataEntry, InMetadata, Weight, bSetNonInterpolableAttributes, OutPoint.MetadataEntry);
+}
+
+void UPCGMetadata::AccumulateWeightedAttributes(PCGMetadataEntryKey InKey, const UPCGMetadata* InMetadata, float Weight, bool bSetNonInterpolableAttributes, PCGMetadataEntryKey& OutKey)
 {
 	if (!InMetadata)
 	{
 		return;
 	}
 
-	bool bHasSetParent = false;
-
-	if (OutPoint.MetadataEntry == PCGInvalidEntryKey)
-	{
-		if (InPoint.MetadataEntry != PCGInvalidEntryKey && Parent == InMetadata)
-		{
-			OutPoint.MetadataEntry = AddEntry(InPoint.MetadataEntry);
-			// No early out here, since we'll need to set the weighted value
-			bHasSetParent = true;
-		}
-		else
-		{
-			OutPoint.MetadataEntry = AddEntry();
-		}
-	}
+	bool bHasSetParent = InitializeOnSet(OutKey, InKey, InMetadata);
 
 	const bool bShouldSetNonInterpolableAttributes = bSetNonInterpolableAttributes && !bHasSetParent;
 
-	//METADATA TODO: This isn't great, we're allocating memory for no reason here
 	AttributeLock.ReadLock();
-	TSet<FName> KnownAttributeNames = AttributeNames;
-	AttributeLock.ReadUnlock();
-
-	for(const FName& AttributeName : KnownAttributeNames)
+	for(const TPair<FName, FPCGMetadataAttributeBase*>& AttributePair : Attributes)
 	{
-		const FPCGMetadataAttributeBase* ConstAttribute = GetConstAttribute(AttributeName);
+		const FName& AttributeName = AttributePair.Key;
+		FPCGMetadataAttributeBase* Attribute = AttributePair.Value;
 
 		if (const FPCGMetadataAttributeBase* OtherAttribute = InMetadata->GetConstAttribute(AttributeName))
 		{
-			if (OtherAttribute->GetTypeId() != ConstAttribute->GetTypeId())
+			if (OtherAttribute->GetTypeId() != Attribute->GetTypeId())
 			{
 				UE_LOG(LogPCG, Error, TEXT("Metadata type mismatch with attribute %s"), *AttributeName.ToString());
 				continue;
 			}
 
-			FPCGMetadataAttributeBase* Attribute = GetMutableAttribute(AttributeName);
-
 			if (Attribute->AllowsInterpolation())
 			{
-				Attribute->AccumulateValue(OutPoint.MetadataEntry, OtherAttribute, InPoint.MetadataEntry, Weight);
+				Attribute->AccumulateValue(OutKey, OtherAttribute, InKey, Weight);
 			}
 			else if (bShouldSetNonInterpolableAttributes)
 			{
-				Attribute->SetValue(OutPoint.MetadataEntry, OtherAttribute, InPoint.MetadataEntry);
+				Attribute->SetValue(OutKey, OtherAttribute, InKey);
 			}
 		}
 	}
+	AttributeLock.ReadUnlock();
 }
 
 void UPCGMetadata::SetAttributes(const FPCGPoint& InPoint, const UPCGMetadata* InMetadata, FPCGPoint& OutPoint)
+{
+	SetAttributes(InPoint.MetadataEntry, InMetadata, OutPoint.MetadataEntry);
+}
+
+void UPCGMetadata::SetAttributes(PCGMetadataEntryKey InKey, const UPCGMetadata* InMetadata, PCGMetadataEntryKey& OutKey)
 {
 	if (!InMetadata)
 	{
 		return;
 	}
 
-	if (OutPoint.MetadataEntry == PCGInvalidEntryKey)
+	if (InitializeOnSet(OutKey, InKey, InMetadata))
 	{
-		if (InPoint.MetadataEntry != PCGInvalidEntryKey && Parent == InMetadata)
-		{
-			OutPoint.MetadataEntry = AddEntry(InPoint.MetadataEntry);
-
-			// Early out; we don't need to do anything else at this point
-			return;
-		}
-		else
-		{
-			OutPoint.MetadataEntry = AddEntry();
-		}
+		// Early out; we don't need to do anything else at this point
+		return;
 	}
 
-	//METADATA TODO: This isn't great, we're allocating memory for no reason here
 	AttributeLock.ReadLock();
-	TSet<FName> KnownAttributeNames = AttributeNames;
-	AttributeLock.ReadUnlock();
-
-	for(const FName& AttributeName : KnownAttributeNames)
+	for(const TPair<FName, FPCGMetadataAttributeBase*>& AttributePair : Attributes)
 	{
-		const FPCGMetadataAttributeBase* ConstAttribute = GetConstAttribute(AttributeName);
+		const FName& AttributeName = AttributePair.Key;
+		FPCGMetadataAttributeBase* Attribute = AttributePair.Value;
 
 		if (const FPCGMetadataAttributeBase* OtherAttribute = InMetadata->GetConstAttribute(AttributeName))
 		{
-			if (OtherAttribute->GetTypeId() != ConstAttribute->GetTypeId())
+			if (OtherAttribute->GetTypeId() != Attribute->GetTypeId())
 			{
 				UE_LOG(LogPCG, Error, TEXT("Metadata type mismatch with attribute %s"), *AttributeName.ToString());
 				continue;
 			}
 
-			FPCGMetadataAttributeBase* Attribute = GetMutableAttribute(AttributeName);
-			Attribute->SetValue(OutPoint.MetadataEntry, OtherAttribute, InPoint.MetadataEntry);
+			Attribute->SetValue(OutKey, OtherAttribute, InKey);
 		}
 	}
+	AttributeLock.ReadUnlock();
 }
 
 void UPCGMetadata::SetAttributes(const TArrayView<const FPCGPoint>& InPoints, const UPCGMetadata* InMetadata, const TArrayView<FPCGPoint>& OutPoints)
@@ -595,45 +537,68 @@ void UPCGMetadata::SetAttributes(const TArrayView<const FPCGPoint>& InPoints, co
 	{
 		const FPCGPoint& InPoint = InPoints[PointIndex];
 		FPCGPoint& OutPoint = OutPoints[PointIndex];
-		if (OutPoint.MetadataEntry == PCGInvalidEntryKey)
+		InitializeOnSet(OutPoint.MetadataEntry, InPoint.MetadataEntry, InMetadata);
+	}
+
+	AttributeLock.ReadLock();
+	for(const TPair<FName, FPCGMetadataAttributeBase*>& AttributePair : Attributes)
+	{
+		const FName& AttributeName = AttributePair.Key;
+		FPCGMetadataAttributeBase* Attribute = AttributePair.Value;
+
+		if (const FPCGMetadataAttributeBase* OtherAttribute = InMetadata->GetConstAttribute(AttributeName))
 		{
-			if (InPoint.MetadataEntry != PCGInvalidEntryKey && Parent == InMetadata)
+			if (OtherAttribute->GetTypeId() != Attribute->GetTypeId())
 			{
-				OutPoint.MetadataEntry = AddEntry(InPoint.MetadataEntry);
+				UE_LOG(LogPCG, Error, TEXT("Metadata type mismatch with attribute %s"), *AttributeName.ToString());
+				continue;
 			}
-			else
+
+			for (int32 PointIndex = 0; PointIndex < InPoints.Num(); ++PointIndex)
 			{
-				OutPoint.MetadataEntry = AddEntry();
+				const FPCGPoint& InPoint = InPoints[PointIndex];
+				FPCGPoint& OutPoint = OutPoints[PointIndex];
+
+				Attribute->SetValue(OutPoint.MetadataEntry, OtherAttribute, InPoint.MetadataEntry);
 			}
 		}
 	}
-
-	//METADATA TODO: This isn't great, we're allocating memory for no reason here
-	AttributeLock.ReadLock();
-	TSet<FName> KnownAttributeNames = AttributeNames;
 	AttributeLock.ReadUnlock();
+}
+
+void UPCGMetadata::SetAttributes(const TArrayView<PCGMetadataEntryKey>& InKeys, const UPCGMetadata* InMetadata, const TArrayView<PCGMetadataEntryKey>& OutKeys)
+{
+	if (!InMetadata)
+	{
+		return;
+	}
 
-	for(const FName& AttributeName : KnownAttributeNames)
+	check(InKeys.Num() == OutKeys.Num());
+
+	for (int32 KeyIndex = 0; KeyIndex < InKeys.Num(); ++KeyIndex)
 	{
-		const FPCGMetadataAttributeBase* ConstAttribute = GetConstAttribute(AttributeName);
+		InitializeOnSet(OutKeys[KeyIndex], InKeys[KeyIndex], InMetadata);
+	}
+
+	AttributeLock.ReadLock();
+	for(const TPair<FName, FPCGMetadataAttributeBase*>& AttributePair : Attributes)
+	{
+		const FName& AttributeName = AttributePair.Key;
+		FPCGMetadataAttributeBase* Attribute = AttributePair.Value;
 
 		if (const FPCGMetadataAttributeBase* OtherAttribute = InMetadata->GetConstAttribute(AttributeName))
 		{
-			if (OtherAttribute->GetTypeId() != ConstAttribute->GetTypeId())
+			if (OtherAttribute->GetTypeId() != Attribute->GetTypeId())
 			{
 				UE_LOG(LogPCG, Error, TEXT("Metadata type mismatch with attribute %s"), *AttributeName.ToString());
 				continue;
 			}
 
-			FPCGMetadataAttributeBase* Attribute = GetMutableAttribute(AttributeName);
-
-			for (int32 PointIndex = 0; PointIndex < InPoints.Num(); ++PointIndex)
+			for (int32 KeyIndex = 0; KeyIndex < InKeys.Num(); ++KeyIndex)
 			{
-				const FPCGPoint& InPoint = InPoints[PointIndex];
-				FPCGPoint& OutPoint = OutPoints[PointIndex];
-
-				Attribute->SetValue(OutPoint.MetadataEntry, OtherAttribute, InPoint.MetadataEntry);
+				Attribute->SetValue(OutKeys[KeyIndex], OtherAttribute, InKeys[KeyIndex]);
 			}
 		}
 	}
+	AttributeLock.ReadUnlock();
 }
\ No newline at end of file
diff --git a/Engine/Plugins/Experimental/PCG/Source/PCG/Private/Metadata/PCGMetadataAccessor.cpp b/Engine/Plugins/Experimental/PCG/Source/PCG/Private/Metadata/PCGMetadataAccessor.cpp
--- a/Engine/Plugins/Experimental/PCG/Source/PCG/Private/Metadata/PCGMetadataAccessor.cpp
+++ b/Engine/Plugins/Experimental/PCG/Source/PCG/Private/Metadata/PCGMetadataAccessor.cpp
@@ -5,8 +5,9 @@
 #include "Metadata/PCGMetadataAttributeTpl.h"
 #include "PCGModule.h"
 
+/** Key-based implmentations */
 template<typename T>
-T UPCGMetadataAccessorHelpers::GetAttribute(const FPCGPoint& Point, UPCGMetadata* Metadata, FName AttributeName)
+T UPCGMetadataAccessorHelpers::GetAttribute(PCGMetadataEntryKey Key, UPCGMetadata* Metadata, FName AttributeName)
 {
 	if (!Metadata)
 	{
@@ -17,7 +18,7 @@ T UPCGMetadataAccessorHelpers::GetAttribute(const FPCGPoint& Point, UPCGMetadata
 	const FPCGMetadataAttribute<T>* Attribute = static_cast<const FPCGMetadataAttribute<T>*>(Metadata->GetConstAttribute(AttributeName));
 	if (Attribute && Attribute->GetTypeId() == PCG::Private::MetadataTypes<T>::Id)
 	{
-		return Attribute->GetValueFromItemKey(Point.MetadataEntry);
+		return Attribute->GetValueFromItemKey(Key);
 	}
 	else if (Attribute)
 	{
@@ -31,23 +32,8 @@ T UPCGMetadataAccessorHelpers::GetAttribute(const FPCGPoint& Point, UPCGMetadata
 	}
 }
 
-void UPCGMetadataAccessorHelpers::InitializeMetadata(FPCGPoint& Point, UPCGMetadata* Metadata)
-{
-	Point.MetadataEntry = Metadata ? Metadata->AddEntry() : PCGInvalidEntryKey;
-}
-
-void UPCGMetadataAccessorHelpers::InitializeMetadata(FPCGPoint& Point, UPCGMetadata* Metadata, const FPCGPoint& ParentPoint)
-{
-	Point.MetadataEntry = Metadata ? Metadata->AddEntry(ParentPoint.MetadataEntry) : PCGInvalidEntryKey;
-}
-
-void UPCGMetadataAccessorHelpers::InitializeMetadata(FPCGPoint& Point, UPCGMetadata* Metadata, const FPCGPoint& ParentPoint, const UPCGMetadata* ParentMetadata)
-{
-	Point.MetadataEntry = Metadata ? (Metadata->HasParent(ParentMetadata) ? Metadata->AddEntry(ParentPoint.MetadataEntry) : Metadata->AddEntry()) : PCGInvalidEntryKey;
-}
-
 template<typename T>
-void UPCGMetadataAccessorHelpers::SetAttribute(FPCGPoint& Point, UPCGMetadata* Metadata, FName AttributeName, const T& Value)
+void UPCGMetadataAccessorHelpers::SetAttribute(PCGMetadataEntryKey& Key, UPCGMetadata* Metadata, FName AttributeName, const T& Value)
 {
 	if (!Metadata)
 	{
@@ -55,22 +41,18 @@ void UPCGMetadataAccessorHelpers::SetAttribute(FPCGPoint& Point, UPCGMetadata* M
 		return;
 	}
 
-	if (Point.MetadataEntry == PCGInvalidEntryKey)
-	{
-		// Default initialization
-		Point.MetadataEntry = Metadata->AddEntry();
-	}
+	Metadata->InitializeOnSet(Key);
 
-	if (Point.MetadataEntry == PCGInvalidEntryKey)
+	if (Key == PCGInvalidEntryKey)
 	{
-		UE_LOG(LogPCG, Error, TEXT("Metadata item has no entry, therefore can't set values"));
+		UE_LOG(LogPCG, Error, TEXT("Metadata key has no entry, therefore can't set values"));
 		return;
 	}
 
 	FPCGMetadataAttribute<T>* Attribute = static_cast<FPCGMetadataAttribute<T>*>(Metadata->GetMutableAttribute(AttributeName));
 	if (Attribute && Attribute->GetTypeId() == PCG::Private::MetadataTypes<T>::Id)
 	{
-		Attribute->SetValue(Point.MetadataEntry, Value);
+		Attribute->SetValue(Key, Value);
 	}
 	else if (Attribute)
 	{
@@ -82,88 +64,189 @@ void UPCGMetadataAccessorHelpers::SetAttribute(FPCGPoint& Point, UPCGMetadata* M
 	}
 }
 
+bool UPCGMetadataAccessorHelpers::HasAttributeSetByMetadataKey(PCGMetadataEntryKey Key, UPCGMetadata* Metadata, FName AttributeName)
+{
+	if (!Metadata)
+	{
+		UE_LOG(LogPCG, Error, TEXT("Data has no metadata"));
+		return false;
+	}
+
+	// Early out: the point has no metadata entry assigned
+	if (Key == PCGInvalidEntryKey)
+	{
+		return false;
+	}
+
+	if (const FPCGMetadataAttributeBase* Attribute = Metadata->GetConstAttribute(AttributeName))
+	{
+		return Attribute->HasNonDefaultValue(Key);
+	}
+	else
+	{
+		UE_LOG(LogPCG, Error, TEXT("Metadata does not have a %s attribute"), *AttributeName.ToString());
+		return false;
+	}
+}
+
+float UPCGMetadataAccessorHelpers::GetFloatAttributeByMetadataKey(int64 Key, UPCGMetadata* Metadata, FName AttributeName)
+{
+	return GetAttribute<float>(Key, Metadata, AttributeName);
+}
+
+void UPCGMetadataAccessorHelpers::SetFloatAttributeByMetadataKey(int64& Key, UPCGMetadata* Metadata, FName AttributeName, float Value)
+{
+	SetAttribute(Key, Metadata, AttributeName, Value);
+}
+
+FVector UPCGMetadataAccessorHelpers::GetVectorAttributeByMetadataKey(int64 Key, UPCGMetadata* Metadata, FName AttributeName)
+{
+	return GetAttribute<FVector>(Key, Metadata, AttributeName);
+}
+
+void UPCGMetadataAccessorHelpers::SetVectorAttributeByMetadataKey(int64& Key, UPCGMetadata* Metadata, FName AttributeName, const FVector& Value)
+{
+	SetAttribute(Key, Metadata, AttributeName, Value);
+}
+
+FVector4 UPCGMetadataAccessorHelpers::GetVector4AttributeByMetadataKey(int64 Key, UPCGMetadata* Metadata, FName AttributeName)
+{
+	return GetAttribute<FVector4>(Key, Metadata, AttributeName);
+}
+
+void UPCGMetadataAccessorHelpers::SetVector4AttributeByMetadataKey(int64& Key, UPCGMetadata* Metadata, FName AttributeName, const FVector4& Value)
+{
+	SetAttribute(Key, Metadata, AttributeName, Value);
+}
+
+FQuat UPCGMetadataAccessorHelpers::GetQuatAttributeByMetadataKey(int64 Key, UPCGMetadata* Metadata, FName AttributeName)
+{
+	return GetAttribute<FQuat>(Key, Metadata, AttributeName);
+}
+
+void UPCGMetadataAccessorHelpers::SetQuatAttributeByMetadataKey(int64& Key, UPCGMetadata* Metadata, FName AttributeName, const FQuat& Value)
+{
+	SetAttribute(Key, Metadata, AttributeName, Value);
+}
+
+FTransform UPCGMetadataAccessorHelpers::GetTransformAttributeByMetadataKey(int64 Key, UPCGMetadata* Metadata, FName AttributeName)
+{
+	return GetAttribute<FTransform>(Key, Metadata, AttributeName);
+}
+
+void UPCGMetadataAccessorHelpers::SetTransformAttributeByMetadataKey(int64& Key, UPCGMetadata* Metadata, FName AttributeName, const FTransform& Value)
+{
+	SetAttribute(Key, Metadata, AttributeName, Value);
+}
+
+FString UPCGMetadataAccessorHelpers::GetStringAttributeByMetadataKey(int64 Key, UPCGMetadata* Metadata, FName AttributeName)
+{
+	return GetAttribute<FString>(Key, Metadata, AttributeName);
+}
+
+void UPCGMetadataAccessorHelpers::SetStringAttributeByMetadataKey(int64& Key, UPCGMetadata* Metadata, FName AttributeName, const FString& Value)
+{
+	SetAttribute(Key, Metadata, AttributeName, Value);
+}
+
+/** Point-based implementations */
+void UPCGMetadataAccessorHelpers::CopyPoint(const FPCGPoint& InPoint, FPCGPoint& OutPoint, bool bCopyMetadata, const UPCGMetadata* InMetadata, UPCGMetadata* OutMetadata)
+{
+	// Copy standard properties
+	OutPoint = InPoint;
+
+	// If we want to copy the metadata, then at least the out metadata must not be null.
+	if (bCopyMetadata && OutMetadata)
+	{
+		// If we have an input metadata, then we can copy values as needed,
+		// Otherwise, we will assume that the point is parented
+		if (InMetadata)
+		{
+			OutMetadata->SetAttributes(InPoint, InMetadata, OutPoint);
+		}
+	}
+	else
+	{
+		OutPoint.MetadataEntry = PCGInvalidEntryKey;
+	}
+}
+
+void UPCGMetadataAccessorHelpers::InitializeMetadata(FPCGPoint& Point, UPCGMetadata* Metadata)
+{
+	Point.MetadataEntry = Metadata ? Metadata->AddEntry() : PCGInvalidEntryKey;
+}
+
+void UPCGMetadataAccessorHelpers::InitializeMetadata(FPCGPoint& Point, UPCGMetadata* Metadata, const FPCGPoint& ParentPoint)
+{
+	Point.MetadataEntry = Metadata ? Metadata->AddEntry(ParentPoint.MetadataEntry) : PCGInvalidEntryKey;
+}
+
+void UPCGMetadataAccessorHelpers::InitializeMetadata(FPCGPoint& Point, UPCGMetadata* Metadata, const FPCGPoint& ParentPoint, const UPCGMetadata* ParentMetadata)
+{
+	Point.MetadataEntry = Metadata ? (Metadata->HasParent(ParentMetadata) ? Metadata->AddEntry(ParentPoint.MetadataEntry) : Metadata->AddEntry()) : PCGInvalidEntryKey;
+}
+
 float UPCGMetadataAccessorHelpers::GetFloatAttribute(const FPCGPoint& Point, UPCGMetadata* Metadata, FName AttributeName)
 {
-	return GetAttribute<float>(Point, Metadata, AttributeName);
+	return GetAttribute<float>(Point.MetadataEntry, Metadata, AttributeName);
 }
 
 void UPCGMetadataAccessorHelpers::SetFloatAttribute(FPCGPoint& Point, UPCGMetadata* Metadata, FName AttributeName, float Value)
 {
-	SetAttribute(Point, Metadata, AttributeName, Value);
+	SetAttribute(Point.MetadataEntry, Metadata, AttributeName, Value);
 }
 
 FVector UPCGMetadataAccessorHelpers::GetVectorAttribute(const FPCGPoint& Point, UPCGMetadata* Metadata, FName AttributeName)
 {
-	return GetAttribute<FVector>(Point, Metadata, AttributeName);
+	return GetAttribute<FVector>(Point.MetadataEntry, Metadata, AttributeName);
 }
 
 void UPCGMetadataAccessorHelpers::SetVectorAttribute(FPCGPoint& Point, UPCGMetadata* Metadata, FName AttributeName, const FVector& Value)
 {
-	SetAttribute(Point, Metadata, AttributeName, Value);
+	SetAttribute(Point.MetadataEntry, Metadata, AttributeName, Value);
 }
 
 FVector4 UPCGMetadataAccessorHelpers::GetVector4Attribute(const FPCGPoint& Point, UPCGMetadata* Metadata, FName AttributeName)
 {
-	return GetAttribute<FVector4>(Point, Metadata, AttributeName);
+	return GetAttribute<FVector4>(Point.MetadataEntry, Metadata, AttributeName);
 }
 
 void UPCGMetadataAccessorHelpers::SetVector4Attribute(FPCGPoint& Point, UPCGMetadata* Metadata, FName AttributeName, const FVector4& Value)
 {
-	SetAttribute(Point, Metadata, AttributeName, Value);
+	SetAttribute(Point.MetadataEntry, Metadata, AttributeName, Value);
 }
 
 FQuat UPCGMetadataAccessorHelpers::GetQuatAttribute(const FPCGPoint& Point, UPCGMetadata* Metadata, FName AttributeName)
 {
-	return GetAttribute<FQuat>(Point, Metadata, AttributeName);
+	return GetAttribute<FQuat>(Point.MetadataEntry, Metadata, AttributeName);
 }
 
 void UPCGMetadataAccessorHelpers::SetQuatAttribute(FPCGPoint& Point, UPCGMetadata* Metadata, FName AttributeName, const FQuat& Value)
 {
-	SetAttribute(Point, Metadata, AttributeName, Value);
+	SetAttribute(Point.MetadataEntry, Metadata, AttributeName, Value);
 }
 
 FTransform UPCGMetadataAccessorHelpers::GetTransformAttribute(const FPCGPoint& Point, UPCGMetadata* Metadata, FName AttributeName)
 {
-	return GetAttribute<FTransform>(Point, Metadata, AttributeName);
+	return GetAttribute<FTransform>(Point.MetadataEntry, Metadata, AttributeName);
 }
 
 void UPCGMetadataAccessorHelpers::SetTransformAttribute(FPCGPoint& Point, UPCGMetadata* Metadata, FName AttributeName, const FTransform& Value)
 {
-	SetAttribute(Point, Metadata, AttributeName, Value);
+	SetAttribute(Point.MetadataEntry, Metadata, AttributeName, Value);
 }
 
 FString UPCGMetadataAccessorHelpers::GetStringAttribute(const FPCGPoint& Point, UPCGMetadata* Metadata, FName AttributeName)
 {
-	return GetAttribute<FString>(Point, Metadata, AttributeName);
+	return GetAttribute<FString>(Point.MetadataEntry, Metadata, AttributeName);
 }
 
 void UPCGMetadataAccessorHelpers::SetStringAttribute(FPCGPoint& Point, UPCGMetadata* Metadata, FName AttributeName, const FString& Value)
 {
-	SetAttribute(Point, Metadata, AttributeName, Value);
+	SetAttribute(Point.MetadataEntry, Metadata, AttributeName, Value);
 }
 
 bool UPCGMetadataAccessorHelpers::HasAttributeSet(const FPCGPoint& Point, UPCGMetadata* Metadata, FName AttributeName)
 {
-	if (!Metadata)
-	{
-		UE_LOG(LogPCG, Error, TEXT("Data has no metadata"));
-		return false;
-	}
-
-	if (!Metadata->HasAttribute(AttributeName))
-	{
-		UE_LOG(LogPCG, Error, TEXT("Metadata does not have a %s attribute"), *AttributeName.ToString());
-		return false;
-	}
-
-	// Early out: the point has no metadata entry assigned
-	if (Point.MetadataEntry == PCGInvalidEntryKey)
-	{
-		return false;
-	}
-
-	const FPCGMetadataAttributeBase* Attribute = Metadata->GetConstAttribute(AttributeName);
-	check(Attribute);
-
-	return Attribute->HasNonDefaultValue(Point.MetadataEntry);
+	return HasAttributeSetByMetadataKey(Point.MetadataEntry, Metadata, AttributeName);
 }
\ No newline at end of file
diff --git a/Engine/Plugins/Experimental/PCG/Source/PCG/Public/Metadata/PCGMetadata.h b/Engine/Plugins/Experimental/PCG/Source/PCG/Public/Metadata/PCGMetadata.h
--- a/Engine/Plugins/Experimental/PCG/Source/PCG/Public/Metadata/PCGMetadata.h
+++ b/Engine/Plugins/Experimental/PCG/Source/PCG/Public/Metadata/PCGMetadata.h
@@ -87,15 +87,27 @@ public:
 	UFUNCTION(BlueprintCallable, Category = "PCG|Metadata")
 	int64 AddEntry(int64 ParentEntryKey = -1);
 
+	/** Initializes the metadata entry key. Returns true if key set from either parent */
+	bool InitializeOnSet(PCGMetadataEntryKey& InKey, PCGMetadataEntryKey InParentKeyA = PCGInvalidEntryKey, const UPCGMetadata* InParentMetadataA = nullptr, PCGMetadataEntryKey InParentKeyB = PCGInvalidEntryKey, const UPCGMetadata* InParentMetadataB = nullptr);
+
 	/** Metadata chaining mechanism */
 	PCGMetadataEntryKey GetParentKey(PCGMetadataEntryKey LocalItemKey) const;
 
 	/** Attributes operations */
-	void MergeAttributes(const FPCGPoint& InPointA, const FPCGPoint& InPointB, FPCGPoint& OutPoint, EPCGMetadataOp Op)
+	void MergeAttributes(PCGMetadataEntryKey InKeyA, const UPCGMetadata* InMetadataA, PCGMetadataEntryKey InKeyB, const UPCGMetadata* InMetadataB, PCGMetadataEntryKey& OutKey, EPCGMetadataOp Op);
+	void MergeAttributes(PCGMetadataEntryKey InKeyA, PCGMetadataEntryKey InKeyB, PCGMetadataEntryKey& OutKey, EPCGMetadataOp Op)
 	{
-		MergeAttributes(InPointA, this, InPointB, this, OutPoint, Op);
+		MergeAttributes(InKeyA, this, InKeyB, this, OutKey, Op);
 	}
 
+	void ResetWeightedAttributes(PCGMetadataEntryKey& OutKey);
+	void AccumulateWeightedAttributes(PCGMetadataEntryKey InKey, const UPCGMetadata* InMetadata, float Weight, bool bSetNonInterpolableAttributes, PCGMetadataEntryKey& OutKey);
+
+	void SetAttributes(PCGMetadataEntryKey InKey, const UPCGMetadata* InMetadata, PCGMetadataEntryKey& OutKey);
+	void SetAttributes(const TArrayView<PCGMetadataEntryKey>& InKeys, const UPCGMetadata* InMetadata, const TArrayView<PCGMetadataEntryKey>& OutKeys);
+
+	/** Attributes operations - shorthand for points */
+	void MergeAttributes(const FPCGPoint& InPointA, const FPCGPoint& InPointB, FPCGPoint& OutPoint, EPCGMetadataOp Op);
 	void MergeAttributes(const FPCGPoint& InPointA, const UPCGMetadata* InMetadataA, const FPCGPoint& InPointB, const UPCGMetadata* InMetadataB, FPCGPoint& OutPoint, EPCGMetadataOp Op);
 
 	void ResetWeightedAttributes(FPCGPoint& OutPoint);
@@ -112,7 +124,6 @@ protected:
 	FPCGMetadataAttributeBase* CopyAttribute(const FPCGMetadataAttributeBase* OriginalAttribute, FName NewAttributeName, bool bKeepParent, bool bCopyEntries, bool bCopyValues);
 
 	bool ParentHasAttribute(FName AttributeName) const;
-	void GetAllAttributeNames(TSet<FName>& AttributeNames) const;
 	int64 GetItemCountForChild() const;
 
 	void AddAttributeInternal(FName AttributeName, FPCGMetadataAttributeBase* Attribute);
@@ -127,12 +138,8 @@ protected:
 	TSet<TSoftObjectPtr<const UPCGMetadata>> OtherParents;
 
 	TMap<FName, FPCGMetadataAttributeBase*> Attributes;
-	TSet<FName> AttributeNames; // Will contain also parent attributes
-
 	PCGMetadataAttributeKey NextAttributeId = 0;
 
-	TSet<FName> HiddenAttributes;
-
 	TArray<PCGMetadataEntryKey> ParentKeys;
 	int64 ItemKeyOffset = 0;
 
diff --git a/Engine/Plugins/Experimental/PCG/Source/PCG/Public/Metadata/PCGMetadataAccessor.h b/Engine/Plugins/Experimental/PCG/Source/PCG/Public/Metadata/PCGMetadataAccessor.h
--- a/Engine/Plugins/Experimental/PCG/Source/PCG/Public/Metadata/PCGMetadataAccessor.h
+++ b/Engine/Plugins/Experimental/PCG/Source/PCG/Public/Metadata/PCGMetadataAccessor.h
@@ -9,7 +9,6 @@
 
 #include "PCGMetadataAccessor.generated.h"
 
-struct FPCGPoint;
 class UPCGMetadata;
 
 UCLASS()
@@ -18,6 +17,50 @@ class PCG_API UPCGMetadataAccessorHelpers : public UBlueprintFunctionLibrary
 	GENERATED_BODY()
 
 public:
+	/** Id-based metadata functions */
+	UFUNCTION(BlueprintCallable, Category = "PCG|Metadata")
+	static float GetFloatAttributeByMetadataKey(int64 Key, UPCGMetadata* Metadata, FName AttributeName);
+
+	UFUNCTION(BlueprintCallable, Category = "PCG|Metadata")
+	static void SetFloatAttributeByMetadataKey(UPARAM(ref) int64& Key, UPCGMetadata* Metadata, FName AttributeName, float Value);
+
+	UFUNCTION(BlueprintCallable, Category = "PCG|Metadata")
+	static FVector GetVectorAttributeByMetadataKey(int64 Key, UPCGMetadata* Metadata, FName AttributeName);
+
+	UFUNCTION(BlueprintCallable, Category = "PCG|Metadata")
+	static void SetVectorAttributeByMetadataKey(UPARAM(ref) int64& Key, UPCGMetadata* Metadata, FName AttributeName, const FVector& Value);
+
+	UFUNCTION(BlueprintCallable, Category = "PCG|Metadata")
+	static FVector4 GetVector4AttributeByMetadataKey(int64 Key, UPCGMetadata* Metadata, FName AttributeName);
+
+	UFUNCTION(BlueprintCallable, Category = "PCG|Metadata")
+	static void SetVector4AttributeByMetadataKey(UPARAM(ref) int64& Key, UPCGMetadata* Metadata, FName AttributeName, const FVector4& Value);
+
+	UFUNCTION(BlueprintCallable, Category = "PCG|Metadata")
+	static FQuat GetQuatAttributeByMetadataKey(int64 Key, UPCGMetadata* Metadata, FName AttributeName);
+
+	UFUNCTION(BlueprintCallable, Category = "PCG|Metadata")
+	static void SetQuatAttributeByMetadataKey(UPARAM(ref) int64& Key, UPCGMetadata* Metadata, FName AttributeName, const FQuat& Value);
+
+	UFUNCTION(BlueprintCallable, Category = "PCG|Metadata")
+	static FTransform GetTransformAttributeByMetadataKey(int64 Key, UPCGMetadata* Metadata, FName AttributeName);
+
+	UFUNCTION(BlueprintCallable, Category = "PCG|Metadata")
+	static void SetTransformAttributeByMetadataKey(UPARAM(ref) int64& Key, UPCGMetadata* Metadata, FName AttributeName, const FTransform& Value);
+
+	UFUNCTION(BlueprintCallable, Category = "PCG|Metadata")
+	static FString GetStringAttributeByMetadataKey(int64 Key, UPCGMetadata* Metadata, FName AttributeName);
+
+	UFUNCTION(BlueprintCallable, Category = "PCG|Metadata")
+	static void SetStringAttributeByMetadataKey(UPARAM(ref) int64& Key, UPCGMetadata* Metadata, FName AttributeName, const FString& Value);
+
+	UFUNCTION(BlueprintCallable, Category = "PCG|Metadata")
+	static bool HasAttributeSetByMetadataKey(int64 Key, UPCGMetadata* Metadata, FName AttributeName);
+
+	/** Point functions */
+	UFUNCTION(BlueprintCallable, Category = "PCG", meta = (ScriptMethod))
+	static void CopyPoint(const FPCGPoint& InPoint, FPCGPoint& OutPoint, bool bCopyMetadata = true, const UPCGMetadata* InMetadata = nullptr, UPCGMetadata* OutMetadata = nullptr);
+
 	static void InitializeMetadata(FPCGPoint& Point, UPCGMetadata* Metadata);
 	static void InitializeMetadata(FPCGPoint& Point, UPCGMetadata* Metadata, const FPCGPoint& ParentPoint, const UPCGMetadata* ParentMetadata);
 
@@ -65,8 +108,8 @@ public:
 
 protected:
 	template<typename T>
-	static T GetAttribute(const FPCGPoint& Point, UPCGMetadata* Metadata, FName AttributeName);
+	static T GetAttribute(PCGMetadataEntryKey Key, UPCGMetadata* Metadata, FName AttributeName);
 
 	template<typename T>
-	static void SetAttribute(FPCGPoint& Point, UPCGMetadata* Metadata, FName AttributeName, const T& Value);
+	static void SetAttribute(PCGMetadataEntryKey& Key, UPCGMetadata* Metadata, FName AttributeName, const T& Value);
 };
\ No newline at end of file
