Added method to create empty metadata on spatial data

diff --git a/Engine/Plugins/Experimental/PCG/Source/PCG/Private/Data/PCGSpatialData.cpp b/Engine/Plugins/Experimental/PCG/Source/PCG/Private/Data/PCGSpatialData.cpp
--- a/Engine/Plugins/Experimental/PCG/Source/PCG/Private/Data/PCGSpatialData.cpp
+++ b/Engine/Plugins/Experimental/PCG/Source/PCG/Private/Data/PCGSpatialData.cpp
@@ -62,4 +62,15 @@ UPCGDifferenceData* UPCGSpatialData::Subtract(const UPCGSpatialData* InOther) co
 	DifferenceData->AddDifference(InOther);
 
 	return DifferenceData;
+}
+
+UPCGMetadata* UPCGSpatialData::CreateEmptyMetadata()
+{
+	if (Metadata)
+	{
+		UE_LOG(LogPCG, Warning, TEXT("Spatial data already had metadata"));
+	}
+
+	Metadata = NewObject<UPCGMetadata>(this);
+	return Metadata;
 }
\ No newline at end of file
diff --git a/Engine/Plugins/Experimental/PCG/Source/PCG/Public/Data/PCGSpatialData.h b/Engine/Plugins/Experimental/PCG/Source/PCG/Public/Data/PCGSpatialData.h
--- a/Engine/Plugins/Experimental/PCG/Source/PCG/Public/Data/PCGSpatialData.h
+++ b/Engine/Plugins/Experimental/PCG/Source/PCG/Public/Data/PCGSpatialData.h
@@ -102,6 +102,9 @@ public:
 	UFUNCTION(BlueprintCallable, Category = Metadata)
 	UPCGMetadata* MutableMetadata() { return Metadata; }
 
+	UFUNCTION(BlueprintCallable, Category = Metadata)
+	UPCGMetadata* CreateEmptyMetadata();
+
 	UPROPERTY(BlueprintReadWrite, EditAnywhere, Category = Data)
 	AActor* TargetActor = nullptr;
 
