Fixed issue when doing property -> attribute metadata operation, was missing a per-point initialization

diff --git a/Engine/Plugins/Experimental/PCG/Source/PCG/Private/Elements/PCGMetadataElement.cpp b/Engine/Plugins/Experimental/PCG/Source/PCG/Private/Elements/PCGMetadataElement.cpp
--- a/Engine/Plugins/Experimental/PCG/Source/PCG/Private/Elements/PCGMetadataElement.cpp
+++ b/Engine/Plugins/Experimental/PCG/Source/PCG/Private/Elements/PCGMetadataElement.cpp
@@ -9,13 +9,14 @@
 namespace PCGMetadataOperations
 {
 	template<typename U, typename T>
-	bool SetValue(const TArray<FPCGPoint>& InPoints, FPCGMetadataAttributeBase* AttributeBase, TFunctionRef<U(const FPCGPoint& InPoint)> PropGetter)
+	bool SetValue(TArray<FPCGPoint>& InPoints, FPCGMetadataAttributeBase* AttributeBase, UPCGMetadata* Metadata, TFunctionRef<U(const FPCGPoint& InPoint)> PropGetter)
 	{
 		if (AttributeBase && AttributeBase->GetTypeId() == PCG::Private::MetadataTypes<T>::Id)
 		{
 			FPCGMetadataAttribute<T>* Attribute = static_cast<FPCGMetadataAttribute<T>*>(AttributeBase);
-			for (const FPCGPoint& Point : InPoints)
+			for (FPCGPoint& Point : InPoints)
 			{
+				Metadata->InitializeOnSet(Point.MetadataEntry);
 				Attribute->SetValue(Point.MetadataEntry, PropGetter(Point));
 			}
 
@@ -123,9 +124,9 @@ bool FPCGMetadataOperationElement::ExecuteInternal(FPCGContext* Context) const
 				}
 
 				FPCGMetadataAttributeBase* AttributeBase = SampledData->Metadata->GetMutableAttribute(Settings->DestinationAttribute);
-				if(!PCGMetadataOperations::SetValue<float, float>(SampledPoints, AttributeBase, DensityGetter) &&
-					!PCGMetadataOperations::SetValue<float, FVector>(SampledPoints, AttributeBase, DensityGetter) &&
-					!PCGMetadataOperations::SetValue<float, FVector4>(SampledPoints, AttributeBase, DensityGetter))
+				if(!PCGMetadataOperations::SetValue<float, float>(SampledPoints, AttributeBase, SampledData->Metadata, DensityGetter) &&
+					!PCGMetadataOperations::SetValue<float, FVector>(SampledPoints, AttributeBase, SampledData->Metadata, DensityGetter) &&
+					!PCGMetadataOperations::SetValue<float, FVector4>(SampledPoints, AttributeBase, SampledData->Metadata, DensityGetter))
 				{
 					PCGE_LOG(Error, "Attribute %s already exists but its type is not compatible", *Settings->DestinationAttribute.ToString());
 				}
@@ -140,8 +141,8 @@ bool FPCGMetadataOperationElement::ExecuteInternal(FPCGContext* Context) const
 				}
 
 				FPCGMetadataAttributeBase* AttributeBase = SampledData->Metadata->GetMutableAttribute(Settings->DestinationAttribute);
-				if(!PCGMetadataOperations::SetValue<FVector, FVector>(SampledPoints, AttributeBase, ExtentsGetter) &&
-					!PCGMetadataOperations::SetValue<FVector, FTransform>(SampledPoints, AttributeBase, ExtentsGetter))
+				if(!PCGMetadataOperations::SetValue<FVector, FVector>(SampledPoints, AttributeBase, SampledData->Metadata, ExtentsGetter) &&
+					!PCGMetadataOperations::SetValue<FVector, FTransform>(SampledPoints, AttributeBase, SampledData->Metadata, ExtentsGetter))
 				{
 					PCGE_LOG(Error, "Attribute %s already exists but its type is not compatible", *Settings->DestinationAttribute.ToString());
 				}
@@ -156,7 +157,7 @@ bool FPCGMetadataOperationElement::ExecuteInternal(FPCGContext* Context) const
 				}
 
 				FPCGMetadataAttributeBase* AttributeBase = SampledData->Metadata->GetMutableAttribute(Settings->DestinationAttribute);
-				if(!PCGMetadataOperations::SetValue<FVector4, FVector4>(SampledPoints, AttributeBase, ColorGetter))
+				if(!PCGMetadataOperations::SetValue<FVector4, FVector4>(SampledPoints, AttributeBase, SampledData->Metadata, ColorGetter))
 				{
 					PCGE_LOG(Error, "Attribute %s already exists but its type is not compatible", *Settings->DestinationAttribute.ToString());
 				}
@@ -171,8 +172,8 @@ bool FPCGMetadataOperationElement::ExecuteInternal(FPCGContext* Context) const
 				}
 
 				FPCGMetadataAttributeBase* AttributeBase = SampledData->Metadata->GetMutableAttribute(Settings->DestinationAttribute);
-				if (!PCGMetadataOperations::SetValue<FVector, FVector>(SampledPoints, AttributeBase, PositionGetter) &&
-					!PCGMetadataOperations::SetValue<FVector, FTransform>(SampledPoints, AttributeBase, PositionGetter))
+				if (!PCGMetadataOperations::SetValue<FVector, FVector>(SampledPoints, AttributeBase, SampledData->Metadata, PositionGetter) &&
+					!PCGMetadataOperations::SetValue<FVector, FTransform>(SampledPoints, AttributeBase, SampledData->Metadata, PositionGetter))
 				{
 					PCGE_LOG(Error, "Attribute %s already exists but its type is not compatible", *Settings->DestinationAttribute.ToString());
 				}
@@ -187,7 +188,7 @@ bool FPCGMetadataOperationElement::ExecuteInternal(FPCGContext* Context) const
 				}
 
 				FPCGMetadataAttributeBase* AttributeBase = SampledData->Metadata->GetMutableAttribute(Settings->DestinationAttribute);
-				if (!PCGMetadataOperations::SetValue<FQuat, FQuat>(SampledPoints, AttributeBase, RotationGetter))
+				if (!PCGMetadataOperations::SetValue<FQuat, FQuat>(SampledPoints, AttributeBase, SampledData->Metadata, RotationGetter))
 				{
 					PCGE_LOG(Error, "Attribute %s already exists but its type is not compatible", *Settings->DestinationAttribute.ToString());
 				}
@@ -202,7 +203,7 @@ bool FPCGMetadataOperationElement::ExecuteInternal(FPCGContext* Context) const
 				}
 
 				FPCGMetadataAttributeBase* AttributeBase = SampledData->Metadata->GetMutableAttribute(Settings->DestinationAttribute);
-				if (!PCGMetadataOperations::SetValue<FVector, FVector>(SampledPoints, AttributeBase, ScaleGetter))
+				if (!PCGMetadataOperations::SetValue<FVector, FVector>(SampledPoints, AttributeBase, SampledData->Metadata, ScaleGetter))
 				{
 					PCGE_LOG(Error, "Attribute %s already exists but its type is not compatible", *Settings->DestinationAttribute.ToString());
 				}
@@ -217,7 +218,7 @@ bool FPCGMetadataOperationElement::ExecuteInternal(FPCGContext* Context) const
 				}
 
 				FPCGMetadataAttributeBase* AttributeBase = SampledData->Metadata->GetMutableAttribute(Settings->DestinationAttribute);
-				if (!PCGMetadataOperations::SetValue<FTransform, FTransform>(SampledPoints, AttributeBase, TransformGetter))
+				if (!PCGMetadataOperations::SetValue<FTransform, FTransform>(SampledPoints, AttributeBase, SampledData->Metadata, TransformGetter))
 				{
 					PCGE_LOG(Error, "Attribute %s already exists but its type is not compatible", *Settings->DestinationAttribute.ToString());
 				}
diff --git a/Engine/Plugins/Experimental/PCG/Source/PCG/Private/Metadata/PCGMetadataAttribute.cpp b/Engine/Plugins/Experimental/PCG/Source/PCG/Private/Metadata/PCGMetadataAttribute.cpp
--- a/Engine/Plugins/Experimental/PCG/Source/PCG/Private/Metadata/PCGMetadataAttribute.cpp
+++ b/Engine/Plugins/Experimental/PCG/Source/PCG/Private/Metadata/PCGMetadataAttribute.cpp
@@ -12,6 +12,7 @@ FPCGMetadataAttributeBase::FPCGMetadataAttributeBase(UPCGMetadata* InMetadata, F
 
 void FPCGMetadataAttributeBase::SetValueFromValueKey(PCGMetadataEntryKey EntryKey, PCGMetadataValueKey ValueKey)
 {
+	check(EntryKey != PCGInvalidEntryKey);
 	FWriteScopeLock ScopeLock(EntryMapLock);
 	EntryToValueMap.FindOrAdd(EntryKey) = ValueKey;
 }
diff --git a/Engine/Plugins/Experimental/PCG/Source/PCG/Private/Metadata/PCGMetadataAttributeTpl.h b/Engine/Plugins/Experimental/PCG/Source/PCG/Private/Metadata/PCGMetadataAttributeTpl.h
--- a/Engine/Plugins/Experimental/PCG/Source/PCG/Private/Metadata/PCGMetadataAttributeTpl.h
+++ b/Engine/Plugins/Experimental/PCG/Source/PCG/Private/Metadata/PCGMetadataAttributeTpl.h
@@ -50,6 +50,7 @@ public:
 
 	virtual void SetValue(PCGMetadataEntryKey ItemKey, const FPCGMetadataAttributeBase* InAttribute, PCGMetadataEntryKey InEntryKey) override
 	{
+		check(ItemKey != PCGInvalidEntryKey);
 		if (InAttribute == this)
 		{
 			SetValueFromValueKey(ItemKey, GetValueKey(InEntryKey));
@@ -74,6 +75,7 @@ public:
 
 	virtual void SetValue(PCGMetadataEntryKey ItemKey, const FPCGMetadataAttributeBase* InAttributeA, PCGMetadataEntryKey InEntryKeyA, const FPCGMetadataAttributeBase* InAttributeB, PCGMetadataEntryKey InEntryKeyB, EPCGMetadataOp Op) override
 	{
+		check(ItemKey != PCGInvalidEntryKey);
 		bool bAppliedValue = false;
 
 		if (InAttributeA && InAttributeB && bAllowsInterpolation)
@@ -153,12 +155,14 @@ public:
 
 	void SetValue(PCGMetadataEntryKey ItemKey, const T& InValue)
 	{
+		check(ItemKey != PCGInvalidEntryKey);
 		SetValueFromValueKey(ItemKey, AddValue(InValue));
 	}
 
 	template<typename U>
 	void SetValue(PCGMetadataEntryKey ItemKey, const U& InValue)
 	{
+		check(ItemKey != PCGInvalidEntryKey);
 		SetValueFromValueKey(ItemKey, AddValue(T(InValue)));
 	}
 
