Fixed an issue where the PCG subsystem isn't initialized before a mesh compilation happens.

diff --git a/Engine/Plugins/Experimental/PCG/Source/PCG/Private/PCGComponent.cpp b/Engine/Plugins/Experimental/PCG/Source/PCG/Private/PCGComponent.cpp
--- a/Engine/Plugins/Experimental/PCG/Source/PCG/Private/PCGComponent.cpp
+++ b/Engine/Plugins/Experimental/PCG/Source/PCG/Private/PCGComponent.cpp
@@ -889,7 +889,10 @@ void UPCGComponent::DirtyGenerated(bool bInDirtyCachedInput)
 	// It would be possible for partitioned actors to add callbacks to their original component, but that inverses the processing flow
 	if (bInDirtyCachedInput && bActivated && IsPartitioned())
 	{
-		GetSubsystem()->DirtyGraph(this, LastGeneratedBounds, bInDirtyCachedInput);
+		if (GetSubsystem())
+		{
+			GetSubsystem()->DirtyGraph(this, LastGeneratedBounds, bInDirtyCachedInput);
+		}
 	}
 }
 
@@ -904,7 +907,7 @@ void UPCGComponent::Refresh()
 	// 2. Otherwise, we need to update the partitioning if the spatial data has changed.
 	if (!bActivated)
 	{
-		if (IsPartitioned())
+		if (IsPartitioned() && GetSubsystem())
 		{
 			GetSubsystem()->DelayUnpartitionGraph(this);
 		}
@@ -921,7 +924,7 @@ void UPCGComponent::Refresh()
 		{
 			Generate(/*bForce=*/false);
 		}
-		else if (IsPartitioned())
+		else if (IsPartitioned() && GetSubsystem())
 		{
 			GetSubsystem()->DelayPartitionGraph(this);
 		}
@@ -1428,7 +1431,7 @@ void UPCGComponent::DirtyCacheFromTag(const FName& InTag)
 	{
 		for (TWeakObjectPtr<const UPCGSettings> Settings : CachedTrackedTagsToSettings[InTag])
 		{
-			if (Settings.IsValid())
+			if (Settings.IsValid() && GetSubsystem())
 			{
 				GetSubsystem()->CleanFromCache(Settings->GetElement().Get());
 			}
@@ -1442,7 +1445,7 @@ void UPCGComponent::DirtyCacheForAllTrackedTags()
 	{
 		for (TWeakObjectPtr<const UPCGSettings> Settings : TagToSettings.Value)
 		{
-			if (Settings.IsValid())
+			if (Settings.IsValid() && GetSubsystem())
 			{
 				GetSubsystem()->CleanFromCache(Settings->GetElement().Get());
 			}
