Added basic parameters support, and overrides through all settings currently supported.
Subsequent CLs will add support for BP overrides and more.

diff --git a/Engine/Plugins/Experimental/PCG/Source/PCG/Private/Data/PCGPointData.cpp b/Engine/Plugins/Experimental/PCG/Source/PCG/Private/Data/PCGPointData.cpp
--- a/Engine/Plugins/Experimental/PCG/Source/PCG/Private/Data/PCGPointData.cpp
+++ b/Engine/Plugins/Experimental/PCG/Source/PCG/Private/Data/PCGPointData.cpp
@@ -171,21 +171,6 @@ void UPCGPointData::InitializeFromActor(AActor* InActor)
 	Metadata = NewObject<UPCGMetadata>(this);
 }
 
-void UPCGPointData::InitializeFromData(const UPCGSpatialData* InSource, const UPCGMetadata* InMetadataParentOverride)
-{
-	check(InSource);
-	if (!TargetActor)
-	{
-		TargetActor = InSource->TargetActor;
-	}	
-
-	if (!Metadata)
-	{
-		Metadata = NewObject<UPCGMetadata>(this);
-		Metadata->Initialize(InMetadataParentOverride ? InMetadataParentOverride : InSource->Metadata);
-	}
-}
-
 const FPCGPoint* UPCGPointData::GetPointAtPosition(const FVector& InPosition) const
 {
 	if (bOctreeIsDirty)
diff --git a/Engine/Plugins/Experimental/PCG/Source/PCG/Private/Data/PCGSpatialData.cpp b/Engine/Plugins/Experimental/PCG/Source/PCG/Private/Data/PCGSpatialData.cpp
--- a/Engine/Plugins/Experimental/PCG/Source/PCG/Private/Data/PCGSpatialData.cpp
+++ b/Engine/Plugins/Experimental/PCG/Source/PCG/Private/Data/PCGSpatialData.cpp
@@ -6,6 +6,12 @@
 #include "Data/PCGProjectionData.h"
 #include "Data/PCGUnionData.h"
 
+UPCGSpatialData::UPCGSpatialData(const FObjectInitializer& ObjectInitializer)
+	: Super(ObjectInitializer)
+{
+	Metadata = ObjectInitializer.CreateDefaultSubobject<UPCGMetadata>(this, TEXT("Metadata"));
+}
+
 const UPCGPointData* UPCGSpatialDataWithPointCache::ToPointData(FPCGContext* Context) const
 {
 	if (!CachedPointData)
@@ -73,4 +79,20 @@ UPCGMetadata* UPCGSpatialData::CreateEmptyMetadata()
 
 	Metadata = NewObject<UPCGMetadata>(this);
 	return Metadata;
+}
+
+void UPCGSpatialData::InitializeFromData(const UPCGSpatialData* InSource, const UPCGMetadata* InMetadataParentOverride)
+{
+	check(InSource);
+	if (!TargetActor)
+	{
+		TargetActor = InSource->TargetActor;
+	}
+
+	if (!Metadata)
+	{
+		Metadata = NewObject<UPCGMetadata>(this);
+	}
+
+	Metadata->Initialize(InMetadataParentOverride ? InMetadataParentOverride : InSource->Metadata);
 }
\ No newline at end of file
diff --git a/Engine/Plugins/Experimental/PCG/Source/PCG/Private/Elements/PCGDensityFilter.cpp b/Engine/Plugins/Experimental/PCG/Source/PCG/Private/Elements/PCGDensityFilter.cpp
--- a/Engine/Plugins/Experimental/PCG/Source/PCG/Private/Elements/PCGDensityFilter.cpp
+++ b/Engine/Plugins/Experimental/PCG/Source/PCG/Private/Elements/PCGDensityFilter.cpp
@@ -3,6 +3,7 @@
 #include "Elements/PCGDensityFilter.h"
 #include "Data/PCGPointData.h"
 #include "Helpers/PCGAsync.h"
+#include "Helpers/PCGSettingsHelpers.h"
 
 FPCGElementPtr UPCGDensityFilterSettings::CreateElement() const
 {
@@ -18,22 +19,28 @@ bool FPCGDensityFilterElement::ExecuteInternal(FPCGContext* Context) const
 
 	TArray<FPCGTaggedData> Inputs = Context->InputData.GetInputs();
 	TArray<FPCGTaggedData>& Outputs = Context->OutputData.TaggedData;
+	UPCGParams* Params = Context->InputData.GetParams();
 
 	// Forward any non-input data
 	Outputs.Append(Context->InputData.GetExclusions());
 	Outputs.Append(Context->InputData.GetAllSettings());
 
-	const float MinBound = FMath::Min(Settings->LowerBound, Settings->UpperBound);
-	const float MaxBound = FMath::Max(Settings->LowerBound, Settings->UpperBound);
-
-	const bool bNoResults = (MaxBound <= 0.0f && !Settings->bInvertFilter) || (MinBound == 0.0f && MaxBound >= 1.0f && Settings->bInvertFilter);
-	const bool bTrivialFilter = (MinBound <= 0.0f && MaxBound >= 1.0f && !Settings->bInvertFilter) || (MinBound == 0.0f && MaxBound == 0.0f && Settings->bInvertFilter);
-
+	const bool bInvertFilter = PCGSettingsHelpers::GetValue(GET_MEMBER_NAME_CHECKED(UPCGDensityFilterSettings, bInvertFilter), Settings->bInvertFilter, Params);
+	const float LowerBound = PCGSettingsHelpers::GetValue(GET_MEMBER_NAME_CHECKED(UPCGDensityFilterSettings, LowerBound), Settings->LowerBound, Params);
+	const float UpperBound = PCGSettingsHelpers::GetValue(GET_MEMBER_NAME_CHECKED(UPCGDensityFilterSettings, LowerBound), Settings->UpperBound, Params);
 #if WITH_EDITORONLY_DATA
-	if(bNoResults && !Settings->bKeepZeroDensityPoints)
+	const bool bKeepZeroDensityPoints = PCGSettingsHelpers::GetValue(GET_MEMBER_NAME_CHECKED(UPCGDensityFilterSettings, bKeepZeroDensityPoints), Settings->bKeepZeroDensityPoints, Params);
 #else
-	if (bNoResults)
+	const bool bKeepZeroDensityPoints = false;
 #endif
+
+	const float MinBound = FMath::Min(LowerBound, UpperBound);
+	const float MaxBound = FMath::Max(LowerBound, UpperBound);
+
+	const bool bNoResults = (MaxBound <= 0.0f && !bInvertFilter) || (MinBound == 0.0f && MaxBound >= 1.0f && bInvertFilter);
+	const bool bTrivialFilter = (MinBound <= 0.0f && MaxBound >= 1.0f && !bInvertFilter) || (MinBound == 0.0f && MaxBound == 0.0f && bInvertFilter);
+
+	if (bNoResults && !bKeepZeroDensityPoints)
 	{
 		PCGE_LOG(Verbose, "Skipped - all inputs rejected");
 		return true;
@@ -75,18 +82,18 @@ bool FPCGDensityFilterElement::ExecuteInternal(FPCGContext* Context) const
 
 		Output.Data = FilteredData;
 
-		FPCGAsync::AsyncPointProcessing(Context, Points.Num(), FilteredPoints, [&Points, Settings, MinBound, MaxBound](int32 Index, FPCGPoint& OutPoint)
+		FPCGAsync::AsyncPointProcessing(Context, Points.Num(), FilteredPoints, [&Points, MinBound, MaxBound, bInvertFilter, bKeepZeroDensityPoints](int32 Index, FPCGPoint& OutPoint)
 		{
 			const FPCGPoint& Point = Points[Index];
 
 			bool bInRange = (Point.Density >= MinBound && Point.Density <= MaxBound);
-			if (bInRange != Settings->bInvertFilter)
+			if (bInRange != bInvertFilter)
 			{
 				OutPoint = Point;
 				return true;
 			}
 #if WITH_EDITORONLY_DATA
-			else if (Settings->bKeepZeroDensityPoints)
+			else if (bKeepZeroDensityPoints)
 			{
 				OutPoint = Point;
 				OutPoint.Density = 0;
diff --git a/Engine/Plugins/Experimental/PCG/Source/PCG/Private/Elements/PCGDensityRemap.cpp b/Engine/Plugins/Experimental/PCG/Source/PCG/Private/Elements/PCGDensityRemap.cpp
--- a/Engine/Plugins/Experimental/PCG/Source/PCG/Private/Elements/PCGDensityRemap.cpp
+++ b/Engine/Plugins/Experimental/PCG/Source/PCG/Private/Elements/PCGDensityRemap.cpp
@@ -4,6 +4,7 @@
 #include "PCGHelpers.h"
 #include "Data/PCGPointData.h"
 #include "Helpers/PCGAsync.h"
+#include "Helpers/PCGSettingsHelpers.h"
 #include "Math/RandomStream.h"
 
 FPCGElementPtr UPCGLinearDensityRemapSettings::CreateElement() const
@@ -20,13 +21,18 @@ bool FPCGLinearDensityRemapElement::ExecuteInternal(FPCGContext* Context) const
 
 	TArray<FPCGTaggedData> Inputs = Context->InputData.GetInputs();
 	TArray<FPCGTaggedData>& Outputs = Context->OutputData.TaggedData;
+	UPCGParams* Params = Context->InputData.GetParams();
 
 	// Forward any non-input data
 	Outputs.Append(Context->InputData.GetExclusions());
 	Outputs.Append(Context->InputData.GetAllSettings());
 
-	const float RemapMin = FMath::Min(Settings->RemapMin, Settings->RemapMax);
-	const float RemapMax = FMath::Max(Settings->RemapMin, Settings->RemapMax);
+	const float SettingsRemapMin = PCGSettingsHelpers::GetValue(GET_MEMBER_NAME_CHECKED(UPCGLinearDensityRemapSettings, RemapMin), Settings->RemapMin, Params);
+	const float SettingsRemapMax = PCGSettingsHelpers::GetValue(GET_MEMBER_NAME_CHECKED(UPCGLinearDensityRemapSettings, RemapMax), Settings->RemapMax, Params);
+	const bool bMultiplyDensity = PCGSettingsHelpers::GetValue(GET_MEMBER_NAME_CHECKED(UPCGLinearDensityRemapSettings, bMultiplyDensity), Settings->bMultiplyDensity, Params);
+
+	const float RemapMin = FMath::Min(SettingsRemapMin, SettingsRemapMax);
+	const float RemapMax = FMath::Max(SettingsRemapMin, SettingsRemapMax);
 
 	// TODO: embarassingly parallel loop
 	for (const FPCGTaggedData& Input : Inputs)
@@ -56,7 +62,7 @@ bool FPCGLinearDensityRemapElement::ExecuteInternal(FPCGContext* Context) const
 		TArray<FPCGPoint>& SampledPoints = SampledData->GetMutablePoints();
 		Output.Data = SampledData;
 
-		if (Settings->bMultiplyDensity)
+		if (bMultiplyDensity)
 		{
 			FPCGAsync::AsyncPointProcessing(Context, OriginalPointCount, SampledPoints, [&Points, Settings, RemapMin, RemapMax](int32 Index, FPCGPoint& OutPoint)
 			{
diff --git a/Engine/Plugins/Experimental/PCG/Source/PCG/Private/Elements/PCGDifferenceElement.cpp b/Engine/Plugins/Experimental/PCG/Source/PCG/Private/Elements/PCGDifferenceElement.cpp
--- a/Engine/Plugins/Experimental/PCG/Source/PCG/Private/Elements/PCGDifferenceElement.cpp
+++ b/Engine/Plugins/Experimental/PCG/Source/PCG/Private/Elements/PCGDifferenceElement.cpp
@@ -1,6 +1,7 @@
 // Copyright Epic Games, Inc. All Rights Reserved.
 
 #include "Elements/PCGDifferenceElement.h"
+#include "Helpers/PCGSettingsHelpers.h"
 
 FPCGElementPtr UPCGDifferenceSettings::CreateElement() const
 {
@@ -16,19 +17,28 @@ bool FPCGDifferenceElement::ExecuteInternal(FPCGContext* Context) const
 
 	TArray<FPCGTaggedData> Inputs = Context->InputData.GetInputs();
 	TArray<FPCGTaggedData> Exclusions = Context->InputData.GetExclusions();
+	UPCGParams* Params = Context->InputData.GetParams();
+
+	const EPCGDifferenceDensityFunction DensityFunction = PCGSettingsHelpers::GetValue(GET_MEMBER_NAME_CHECKED(UPCGDifferenceSettings, DensityFunction), Settings->DensityFunction, Params);
+#if WITH_EDITORONLY_DATA
+	const bool bKeepZeroDensityPoints = PCGSettingsHelpers::GetValue(GET_MEMBER_NAME_CHECKED(UPCGDifferenceSettings, bKeepZeroDensityPoints), Settings->bKeepZeroDensityPoints, Params);
+#else
+	const bool bKeepZeroDensityPoints = false;
+#endif
+
 	TArray<FPCGTaggedData>& Outputs = Context->OutputData.TaggedData;
 
 	const UPCGSpatialData* FirstSpatialData = nullptr;
 	UPCGDifferenceData* DifferenceData = nullptr;
 	int32 DifferenceTaggedDataIndex = -1;
 
-	auto AddToDifference = [&FirstSpatialData, &DifferenceData, &DifferenceTaggedDataIndex, Settings, &Outputs](const UPCGSpatialData* SpatialData) {
+	auto AddToDifference = [&FirstSpatialData, &DifferenceData, &DifferenceTaggedDataIndex, DensityFunction, bKeepZeroDensityPoints, &Outputs](const UPCGSpatialData* SpatialData) {
 		if (!DifferenceData)
 		{
 			DifferenceData = FirstSpatialData->Subtract(SpatialData);
-			DifferenceData->SetDensityFunction(Settings->DensityFunction);
+			DifferenceData->SetDensityFunction(DensityFunction);
 #if WITH_EDITORONLY_DATA
-			DifferenceData->bKeepZeroDensityPoints = Settings->bKeepZeroDensityPoints;
+			DifferenceData->bKeepZeroDensityPoints = bKeepZeroDensityPoints;
 #endif
 
 			FPCGTaggedData& DifferenceTaggedData = Outputs[DifferenceTaggedDataIndex];
diff --git a/Engine/Plugins/Experimental/PCG/Source/PCG/Private/Elements/PCGIntersectionElement.cpp b/Engine/Plugins/Experimental/PCG/Source/PCG/Private/Elements/PCGIntersectionElement.cpp
--- a/Engine/Plugins/Experimental/PCG/Source/PCG/Private/Elements/PCGIntersectionElement.cpp
+++ b/Engine/Plugins/Experimental/PCG/Source/PCG/Private/Elements/PCGIntersectionElement.cpp
@@ -1,6 +1,7 @@
 // Copyright Epic Games, Inc. All Rights Reserved.
 
 #include "Elements/PCGIntersectionElement.h"
+#include "Helpers/PCGSettingsHelpers.h"
 
 FPCGElementPtr UPCGIntersectionSettings::CreateElement() const
 {
@@ -15,6 +16,15 @@ bool FPCGIntersectionElement::ExecuteInternal(FPCGContext* Context) const
 	check(Settings);
 
 	TArray<FPCGTaggedData> Inputs = Context->InputData.GetInputs();
+	UPCGParams* Params = Context->InputData.GetParams();
+
+	const EPCGIntersectionDensityFunction DensityFunction = PCGSettingsHelpers::GetValue(GET_MEMBER_NAME_CHECKED(UPCGIntersectionSettings, DensityFunction), Settings->DensityFunction, Params);
+#if WITH_EDITORONLY_DATA
+	const bool bKeepZeroDensityPoints = PCGSettingsHelpers::GetValue(GET_MEMBER_NAME_CHECKED(UPCGIntersectionSettings, bKeepZeroDensityPoints), Settings->bKeepZeroDensityPoints, Params);
+#else
+	const bool bKeepZeroDensityPoints = false;
+#endif
+
 	TArray<FPCGTaggedData>& Outputs = Context->OutputData.TaggedData;
 
 	const UPCGSpatialData* FirstSpatialData = nullptr;
@@ -43,12 +53,11 @@ bool FPCGIntersectionElement::ExecuteInternal(FPCGContext* Context) const
 		// Create a new intersection
 		IntersectionData = (IntersectionData ? IntersectionData : FirstSpatialData)->IntersectWith(SpatialData);
 		// Propagate settings
-		IntersectionData->DensityFunction = Settings->DensityFunction;
+		IntersectionData->DensityFunction = DensityFunction;
 #if WITH_EDITORONLY_DATA
-		IntersectionData->bKeepZeroDensityPoints = Settings->bKeepZeroDensityPoints;
+		IntersectionData->bKeepZeroDensityPoints = bKeepZeroDensityPoints;
 #endif
 
-
 		// Update tagged data
 		FPCGTaggedData& IntersectionTaggedData = Outputs[IntersectionTaggedDataIndex];
 		IntersectionTaggedData.Data = IntersectionData;
diff --git a/Engine/Plugins/Experimental/PCG/Source/PCG/Private/Elements/PCGMetadataElement.cpp b/Engine/Plugins/Experimental/PCG/Source/PCG/Private/Elements/PCGMetadataElement.cpp
--- a/Engine/Plugins/Experimental/PCG/Source/PCG/Private/Elements/PCGMetadataElement.cpp
+++ b/Engine/Plugins/Experimental/PCG/Source/PCG/Private/Elements/PCGMetadataElement.cpp
@@ -2,6 +2,7 @@
 
 #include "Elements/PCGMetadataElement.h"
 #include "Data/PCGPointData.h"
+#include "Helpers/PCGSettingsHelpers.h"
 #include "Metadata/PCGMetadata.h"
 #include "Metadata/PCGMetadataAttribute.h"
 #include "Metadata/PCGMetadataAttributeTpl.h"
@@ -61,8 +62,15 @@ bool FPCGMetadataOperationElement::ExecuteInternal(FPCGContext* Context) const
 	check(Settings);
 
 	TArray<FPCGTaggedData> Inputs = Context->InputData.GetInputs();
+	UPCGParams* Params = Context->InputData.GetParams();
+
 	TArray<FPCGTaggedData>& Outputs = Context->OutputData.TaggedData;
 
+	const FName SourceAttribute = PCGSettingsHelpers::GetValue(GET_MEMBER_NAME_CHECKED(UPCGMetadataOperationSettings, SourceAttribute), Settings->SourceAttribute, Params);
+	const EPCGPointProperties PointProperty = PCGSettingsHelpers::GetValue(GET_MEMBER_NAME_CHECKED(UPCGMetadataOperationSettings, PointProperty), Settings->PointProperty, Params);
+	const FName DestinationAttribute = PCGSettingsHelpers::GetValue(GET_MEMBER_NAME_CHECKED(UPCGMetadataOperationSettings, DestinationAttribute), Settings->DestinationAttribute, Params);
+	const EPCGMetadataOperationTarget Target = PCGSettingsHelpers::GetValue(GET_MEMBER_NAME_CHECKED(UPCGMetadataOperationSettings, Target), Settings->Target, Params);
+
 	// Forward any non-input data
 	Outputs.Append(Context->InputData.GetExclusions());
 	Outputs.Append(Context->InputData.GetAllSettings());
@@ -94,9 +102,9 @@ bool FPCGMetadataOperationElement::ExecuteInternal(FPCGContext* Context) const
 		}
 
 		// Check if the attribute exists
-		if ((Settings->Target == EPCGMetadataOperationTarget::AttributeToProperty || Settings->Target == EPCGMetadataOperationTarget::AttributeToAttribute) && !OriginalData->Metadata->HasAttribute(Settings->SourceAttribute))
+		if ((Target == EPCGMetadataOperationTarget::AttributeToProperty || Target == EPCGMetadataOperationTarget::AttributeToAttribute) && !OriginalData->Metadata->HasAttribute(SourceAttribute))
 		{
-			PCGE_LOG(Warning, "Input does not have the %s attribute", *Settings->SourceAttribute.ToString());
+			PCGE_LOG(Warning, "Input does not have the %s attribute", *SourceAttribute.ToString());
 			continue;
 		}
 
@@ -112,199 +120,199 @@ bool FPCGMetadataOperationElement::ExecuteInternal(FPCGContext* Context) const
 		// Copy points and then apply the operation
 		SampledPoints = Points;
 
-		if (Settings->Target == EPCGMetadataOperationTarget::PropertyToAttribute)
+		if (Target == EPCGMetadataOperationTarget::PropertyToAttribute)
 		{
-			if (Settings->PointProperty == EPCGPointProperties::Density)
+			if (PointProperty == EPCGPointProperties::Density)
 			{
 				auto DensityGetter = [](const FPCGPoint& InPoint) { return InPoint.Density; };
 
-				if (!SampledData->Metadata->HasAttribute(Settings->DestinationAttribute))
+				if (!SampledData->Metadata->HasAttribute(DestinationAttribute))
 				{
-					SampledData->Metadata->CreateFloatAttribute(Settings->DestinationAttribute, 0.0f, /*bAllowsInterpolation=*/true, /*bOverrideParent=*/true);
+					SampledData->Metadata->CreateFloatAttribute(DestinationAttribute, 0.0f, /*bAllowsInterpolation=*/true, /*bOverrideParent=*/true);
 				}
 
-				FPCGMetadataAttributeBase* AttributeBase = SampledData->Metadata->GetMutableAttribute(Settings->DestinationAttribute);
+				FPCGMetadataAttributeBase* AttributeBase = SampledData->Metadata->GetMutableAttribute(DestinationAttribute);
 				if(!PCGMetadataOperations::SetValue<float, float>(SampledPoints, AttributeBase, SampledData->Metadata, DensityGetter) &&
 					!PCGMetadataOperations::SetValue<float, FVector>(SampledPoints, AttributeBase, SampledData->Metadata, DensityGetter) &&
 					!PCGMetadataOperations::SetValue<float, FVector4>(SampledPoints, AttributeBase, SampledData->Metadata, DensityGetter))
 				{
-					PCGE_LOG(Error, "Attribute %s already exists but its type is not compatible", *Settings->DestinationAttribute.ToString());
+					PCGE_LOG(Error, "Attribute %s already exists but its type is not compatible", *DestinationAttribute.ToString());
 				}
 			}
-			else if (Settings->PointProperty == EPCGPointProperties::Extents)
+			else if (PointProperty == EPCGPointProperties::Extents)
 			{
 				auto ExtentsGetter = [](const FPCGPoint& InPoint) { return InPoint.Extents; };
 
-				if (!SampledData->Metadata->HasAttribute(Settings->DestinationAttribute))
+				if (!SampledData->Metadata->HasAttribute(DestinationAttribute))
 				{
-					SampledData->Metadata->CreateVectorAttribute(Settings->DestinationAttribute, FVector::Zero() , /*bAllowsInterpolation=*/true, /*bOverrideParent=*/true);
+					SampledData->Metadata->CreateVectorAttribute(DestinationAttribute, FVector::Zero() , /*bAllowsInterpolation=*/true, /*bOverrideParent=*/true);
 				}
 
-				FPCGMetadataAttributeBase* AttributeBase = SampledData->Metadata->GetMutableAttribute(Settings->DestinationAttribute);
+				FPCGMetadataAttributeBase* AttributeBase = SampledData->Metadata->GetMutableAttribute(DestinationAttribute);
 				if(!PCGMetadataOperations::SetValue<FVector, FVector>(SampledPoints, AttributeBase, SampledData->Metadata, ExtentsGetter) &&
 					!PCGMetadataOperations::SetValue<FVector, FTransform>(SampledPoints, AttributeBase, SampledData->Metadata, ExtentsGetter))
 				{
-					PCGE_LOG(Error, "Attribute %s already exists but its type is not compatible", *Settings->DestinationAttribute.ToString());
+					PCGE_LOG(Error, "Attribute %s already exists but its type is not compatible", *DestinationAttribute.ToString());
 				}
 			}
-			else if(Settings->PointProperty == EPCGPointProperties::Color)
+			else if(PointProperty == EPCGPointProperties::Color)
 			{
 				auto ColorGetter = [](const FPCGPoint& InPoint) { return InPoint.Color; };
 
-				if (!SampledData->Metadata->HasAttribute(Settings->DestinationAttribute))
+				if (!SampledData->Metadata->HasAttribute(DestinationAttribute))
 				{
-					SampledData->Metadata->CreateVector4Attribute(Settings->DestinationAttribute, FVector4::Zero(), /*bAllowsInterpolation=*/true, /*bOverrideParent=*/true);
+					SampledData->Metadata->CreateVector4Attribute(DestinationAttribute, FVector4::Zero(), /*bAllowsInterpolation=*/true, /*bOverrideParent=*/true);
 				}
 
-				FPCGMetadataAttributeBase* AttributeBase = SampledData->Metadata->GetMutableAttribute(Settings->DestinationAttribute);
+				FPCGMetadataAttributeBase* AttributeBase = SampledData->Metadata->GetMutableAttribute(DestinationAttribute);
 				if(!PCGMetadataOperations::SetValue<FVector4, FVector4>(SampledPoints, AttributeBase, SampledData->Metadata, ColorGetter))
 				{
-					PCGE_LOG(Error, "Attribute %s already exists but its type is not compatible", *Settings->DestinationAttribute.ToString());
+					PCGE_LOG(Error, "Attribute %s already exists but its type is not compatible", *DestinationAttribute.ToString());
 				}
 			}
-			else if (Settings->PointProperty == EPCGPointProperties::Position)
+			else if (PointProperty == EPCGPointProperties::Position)
 			{
 				auto PositionGetter = [](const FPCGPoint& InPoint) { return InPoint.Transform.GetLocation(); };
 
-				if (!SampledData->Metadata->HasAttribute(Settings->DestinationAttribute))
+				if (!SampledData->Metadata->HasAttribute(DestinationAttribute))
 				{
-					SampledData->Metadata->CreateVectorAttribute(Settings->DestinationAttribute, FVector::Zero(), /*bAllowsInterpolation=*/true, /*bOverrideParent=*/true);
+					SampledData->Metadata->CreateVectorAttribute(DestinationAttribute, FVector::Zero(), /*bAllowsInterpolation=*/true, /*bOverrideParent=*/true);
 				}
 
-				FPCGMetadataAttributeBase* AttributeBase = SampledData->Metadata->GetMutableAttribute(Settings->DestinationAttribute);
+				FPCGMetadataAttributeBase* AttributeBase = SampledData->Metadata->GetMutableAttribute(DestinationAttribute);
 				if (!PCGMetadataOperations::SetValue<FVector, FVector>(SampledPoints, AttributeBase, SampledData->Metadata, PositionGetter) &&
 					!PCGMetadataOperations::SetValue<FVector, FTransform>(SampledPoints, AttributeBase, SampledData->Metadata, PositionGetter))
 				{
-					PCGE_LOG(Error, "Attribute %s already exists but its type is not compatible", *Settings->DestinationAttribute.ToString());
+					PCGE_LOG(Error, "Attribute %s already exists but its type is not compatible", *DestinationAttribute.ToString());
 				}
 			}
-			else if (Settings->PointProperty == EPCGPointProperties::Rotation)
+			else if (PointProperty == EPCGPointProperties::Rotation)
 			{
 				auto RotationGetter = [](const FPCGPoint& InPoint) { return InPoint.Transform.GetRotation(); };
 
-				if (!SampledData->Metadata->HasAttribute(Settings->DestinationAttribute))
+				if (!SampledData->Metadata->HasAttribute(DestinationAttribute))
 				{
-					SampledData->Metadata->CreateQuatAttribute(Settings->DestinationAttribute, FQuat::Identity, /*bAllowsInterpolation=*/true, /*bOverrideParent=*/true);
+					SampledData->Metadata->CreateQuatAttribute(DestinationAttribute, FQuat::Identity, /*bAllowsInterpolation=*/true, /*bOverrideParent=*/true);
 				}
 
-				FPCGMetadataAttributeBase* AttributeBase = SampledData->Metadata->GetMutableAttribute(Settings->DestinationAttribute);
+				FPCGMetadataAttributeBase* AttributeBase = SampledData->Metadata->GetMutableAttribute(DestinationAttribute);
 				if (!PCGMetadataOperations::SetValue<FQuat, FQuat>(SampledPoints, AttributeBase, SampledData->Metadata, RotationGetter))
 				{
-					PCGE_LOG(Error, "Attribute %s already exists but its type is not compatible", *Settings->DestinationAttribute.ToString());
+					PCGE_LOG(Error, "Attribute %s already exists but its type is not compatible", *DestinationAttribute.ToString());
 				}
 			}
-			else if (Settings->PointProperty == EPCGPointProperties::Scale)
+			else if (PointProperty == EPCGPointProperties::Scale)
 			{
 				auto ScaleGetter = [](const FPCGPoint& InPoint) { return InPoint.Transform.GetScale3D(); };
 
-				if (!SampledData->Metadata->HasAttribute(Settings->DestinationAttribute))
+				if (!SampledData->Metadata->HasAttribute(DestinationAttribute))
 				{
-					SampledData->Metadata->CreateVectorAttribute(Settings->DestinationAttribute, FVector::One(), /*bAllowsInterpolation=*/true, /*bOverrideParent=*/true);
+					SampledData->Metadata->CreateVectorAttribute(DestinationAttribute, FVector::One(), /*bAllowsInterpolation=*/true, /*bOverrideParent=*/true);
 				}
 
-				FPCGMetadataAttributeBase* AttributeBase = SampledData->Metadata->GetMutableAttribute(Settings->DestinationAttribute);
+				FPCGMetadataAttributeBase* AttributeBase = SampledData->Metadata->GetMutableAttribute(DestinationAttribute);
 				if (!PCGMetadataOperations::SetValue<FVector, FVector>(SampledPoints, AttributeBase, SampledData->Metadata, ScaleGetter))
 				{
-					PCGE_LOG(Error, "Attribute %s already exists but its type is not compatible", *Settings->DestinationAttribute.ToString());
+					PCGE_LOG(Error, "Attribute %s already exists but its type is not compatible", *DestinationAttribute.ToString());
 				}
 			}
-			else if (Settings->PointProperty == EPCGPointProperties::Transform)
+			else if (PointProperty == EPCGPointProperties::Transform)
 			{
 				auto TransformGetter = [](const FPCGPoint& InPoint) { return InPoint.Transform; };
 
-				if (!SampledData->Metadata->HasAttribute(Settings->DestinationAttribute))
+				if (!SampledData->Metadata->HasAttribute(DestinationAttribute))
 				{
-					SampledData->Metadata->CreateTransformAttribute(Settings->DestinationAttribute, FTransform::Identity, /*bAllowsInterpolation=*/true, /*bOverrideParent=*/true);
+					SampledData->Metadata->CreateTransformAttribute(DestinationAttribute, FTransform::Identity, /*bAllowsInterpolation=*/true, /*bOverrideParent=*/true);
 				}
 
-				FPCGMetadataAttributeBase* AttributeBase = SampledData->Metadata->GetMutableAttribute(Settings->DestinationAttribute);
+				FPCGMetadataAttributeBase* AttributeBase = SampledData->Metadata->GetMutableAttribute(DestinationAttribute);
 				if (!PCGMetadataOperations::SetValue<FTransform, FTransform>(SampledPoints, AttributeBase, SampledData->Metadata, TransformGetter))
 				{
-					PCGE_LOG(Error, "Attribute %s already exists but its type is not compatible", *Settings->DestinationAttribute.ToString());
+					PCGE_LOG(Error, "Attribute %s already exists but its type is not compatible", *DestinationAttribute.ToString());
 				}
 			}
 		}
-		else if(Settings->Target == EPCGMetadataOperationTarget::AttributeToProperty) // Attribute to property
+		else if(Target == EPCGMetadataOperationTarget::AttributeToProperty) // Attribute to property
 		{
-			if (Settings->PointProperty == EPCGPointProperties::Density)
+			if (PointProperty == EPCGPointProperties::Density)
 			{
 				auto DensitySetter = [](FPCGPoint& InPoint, const float& InValue) { InPoint.Density = InValue; };
 
-				const FPCGMetadataAttributeBase* AttributeBase = SampledData->Metadata->GetConstAttribute(Settings->SourceAttribute);
+				const FPCGMetadataAttributeBase* AttributeBase = SampledData->Metadata->GetConstAttribute(SourceAttribute);
 				if(!PCGMetadataOperations::SetValue<float, float>(AttributeBase, SampledPoints, DensitySetter))
 				{
-					PCGE_LOG(Error, "Attribute %s already exists but its type is not compatible", *Settings->SourceAttribute.ToString());
+					PCGE_LOG(Error, "Attribute %s already exists but its type is not compatible", *SourceAttribute.ToString());
 				}
 			}
-			else if (Settings->PointProperty == EPCGPointProperties::Extents)
+			else if (PointProperty == EPCGPointProperties::Extents)
 			{
 				auto ExtentsSetter = [](FPCGPoint& InPoint, const FVector& InValue) { InPoint.Extents = InValue; };
 
-				const FPCGMetadataAttributeBase* AttributeBase = SampledData->Metadata->GetConstAttribute(Settings->SourceAttribute);
+				const FPCGMetadataAttributeBase* AttributeBase = SampledData->Metadata->GetConstAttribute(SourceAttribute);
 				if(!PCGMetadataOperations::SetValue<float, FVector>(AttributeBase, SampledPoints, ExtentsSetter) &&
 					!PCGMetadataOperations::SetValue<FVector, FVector>(AttributeBase, SampledPoints, ExtentsSetter))
 				{
-					PCGE_LOG(Error, "Attribute %s already exists but its type is not compatible", *Settings->SourceAttribute.ToString());
+					PCGE_LOG(Error, "Attribute %s already exists but its type is not compatible", *SourceAttribute.ToString());
 				}
 			}
-			else if (Settings->PointProperty == EPCGPointProperties::Color)
+			else if (PointProperty == EPCGPointProperties::Color)
 			{
 				auto ColorSetter = [](FPCGPoint& InPoint, const FVector4& InValue) { InPoint.Color = InValue; };
 
-				const FPCGMetadataAttributeBase* AttributeBase = SampledData->Metadata->GetConstAttribute(Settings->SourceAttribute);
+				const FPCGMetadataAttributeBase* AttributeBase = SampledData->Metadata->GetConstAttribute(SourceAttribute);
 				if(!PCGMetadataOperations::SetValue<float, FVector4>(AttributeBase, SampledPoints, ColorSetter) &&
 					!PCGMetadataOperations::SetValue<FVector4, FVector4>(AttributeBase, SampledPoints, ColorSetter))
 				{
-					PCGE_LOG(Error, "Attribute %s already exists but its type is not compatible", *Settings->SourceAttribute.ToString());
+					PCGE_LOG(Error, "Attribute %s already exists but its type is not compatible", *SourceAttribute.ToString());
 				}
 			}
-			else if (Settings->PointProperty == EPCGPointProperties::Position)
+			else if (PointProperty == EPCGPointProperties::Position)
 			{
 				auto PositionSetter = [](FPCGPoint& InPoint, const FVector& InValue) { InPoint.Transform.SetLocation(InValue); };
 
-				const FPCGMetadataAttributeBase* AttributeBase = SampledData->Metadata->GetConstAttribute(Settings->SourceAttribute);
+				const FPCGMetadataAttributeBase* AttributeBase = SampledData->Metadata->GetConstAttribute(SourceAttribute);
 				if (!PCGMetadataOperations::SetValue<float, FVector>(AttributeBase, SampledPoints, PositionSetter) &&
 					!PCGMetadataOperations::SetValue<FVector, FVector>(AttributeBase, SampledPoints, PositionSetter))
 				{
-					PCGE_LOG(Error, "Attribute %s already exists but its type is not compatible", *Settings->SourceAttribute.ToString());
+					PCGE_LOG(Error, "Attribute %s already exists but its type is not compatible", *SourceAttribute.ToString());
 				}
 			}
-			else if (Settings->PointProperty == EPCGPointProperties::Rotation)
+			else if (PointProperty == EPCGPointProperties::Rotation)
 			{
 				auto RotationSetter = [](FPCGPoint& InPoint, const FQuat& InValue) { InPoint.Transform.SetRotation(InValue.GetNormalized()); };
 
-				const FPCGMetadataAttributeBase* AttributeBase = SampledData->Metadata->GetConstAttribute(Settings->SourceAttribute);
+				const FPCGMetadataAttributeBase* AttributeBase = SampledData->Metadata->GetConstAttribute(SourceAttribute);
 				if (!PCGMetadataOperations::SetValue<FQuat, FQuat>(AttributeBase, SampledPoints, RotationSetter))
 				{
-					PCGE_LOG(Error, "Attribute %s already exists but its type is not compatible", *Settings->SourceAttribute.ToString());
+					PCGE_LOG(Error, "Attribute %s already exists but its type is not compatible", *SourceAttribute.ToString());
 				}
 			}
-			else if (Settings->PointProperty == EPCGPointProperties::Scale)
+			else if (PointProperty == EPCGPointProperties::Scale)
 			{
 				auto ScaleSetter = [](FPCGPoint& InPoint, const FVector& InValue) { InPoint.Transform.SetScale3D(InValue); };
 
-				const FPCGMetadataAttributeBase* AttributeBase = SampledData->Metadata->GetConstAttribute(Settings->SourceAttribute);
+				const FPCGMetadataAttributeBase* AttributeBase = SampledData->Metadata->GetConstAttribute(SourceAttribute);
 				if (!PCGMetadataOperations::SetValue<float, FVector>(AttributeBase, SampledPoints, ScaleSetter) &&
 					!PCGMetadataOperations::SetValue<FVector, FVector>(AttributeBase, SampledPoints, ScaleSetter))
 				{
-					PCGE_LOG(Error, "Attribute %s already exists but its type is not compatible", *Settings->SourceAttribute.ToString());
+					PCGE_LOG(Error, "Attribute %s already exists but its type is not compatible", *SourceAttribute.ToString());
 				}
 			}
-			else if (Settings->PointProperty == EPCGPointProperties::Transform)
+			else if (PointProperty == EPCGPointProperties::Transform)
 			{
 				auto TransformSetter = [](FPCGPoint& InPoint, const FTransform& InValue) { InPoint.Transform = InValue; };
 
-				const FPCGMetadataAttributeBase* AttributeBase = SampledData->Metadata->GetConstAttribute(Settings->SourceAttribute);
+				const FPCGMetadataAttributeBase* AttributeBase = SampledData->Metadata->GetConstAttribute(SourceAttribute);
 				if (!PCGMetadataOperations::SetValue<FVector, FTransform>(AttributeBase, SampledPoints, TransformSetter) &&
 					!PCGMetadataOperations::SetValue<FTransform, FTransform>(AttributeBase, SampledPoints, TransformSetter))
 				{
-					PCGE_LOG(Error, "Attribute %s already exists but its type is not compatible", *Settings->SourceAttribute.ToString());
+					PCGE_LOG(Error, "Attribute %s already exists but its type is not compatible", *SourceAttribute.ToString());
 				}
 			}
 		}
 		else // Attribute to attribute
 		{
-			SampledData->Metadata->CopyAttribute(Settings->SourceAttribute, Settings->DestinationAttribute);
+			SampledData->Metadata->CopyAttribute(SourceAttribute, DestinationAttribute);
 		}
 	}
 
diff --git a/Engine/Plugins/Experimental/PCG/Source/PCG/Private/Elements/PCGPointSampler.cpp b/Engine/Plugins/Experimental/PCG/Source/PCG/Private/Elements/PCGPointSampler.cpp
--- a/Engine/Plugins/Experimental/PCG/Source/PCG/Private/Elements/PCGPointSampler.cpp
+++ b/Engine/Plugins/Experimental/PCG/Source/PCG/Private/Elements/PCGPointSampler.cpp
@@ -4,6 +4,7 @@
 #include "PCGHelpers.h"
 #include "Data/PCGPointData.h"
 #include "Helpers/PCGAsync.h"
+#include "Helpers/PCGSettingsHelpers.h"
 #include "Math/RandomStream.h"
 
 FPCGElementPtr UPCGPointSamplerSettings::CreateElement() const
@@ -19,21 +20,26 @@ bool FPCGPointSamplerElement::ExecuteInternal(FPCGContext* Context) const
 	check(Settings);
 
 	TArray<FPCGTaggedData> Inputs = Context->InputData.GetInputs();
+	UPCGParams* Params = Context->InputData.GetParams();
+
 	TArray<FPCGTaggedData>& Outputs = Context->OutputData.TaggedData;
 
-	// Forward any non-input data
+	// Forward any non-input data, excluding params
 	Outputs.Append(Context->InputData.GetExclusions());
 	Outputs.Append(Context->InputData.GetAllSettings());
 
-	const bool bNoSampling = (Settings->Ratio <= 0.0f);
-	const bool bTrivialSampling = (Settings->Ratio >= 1.0f);
-
-	// Early exit when nothing will be generated out of this sampler
+	const float Ratio = PCGSettingsHelpers::GetValue(GET_MEMBER_NAME_CHECKED(UPCGPointSamplerSettings, Ratio), Settings->Ratio, Params);
 #if WITH_EDITORONLY_DATA
-	if (bNoSampling && !Settings->bKeepZeroDensityPoints)
+	const bool bKeepZeroDensityPoints = PCGSettingsHelpers::GetValue(GET_MEMBER_NAME_CHECKED(UPCGPointSamplerSettings, bKeepZeroDensityPoints), Settings->bKeepZeroDensityPoints, Params);
 #else
-	if (bNoSampling)
+	const bool bKeepZeroDensityPoints = false;
 #endif
+
+	const bool bNoSampling = (Ratio <= 0.0f);
+	const bool bTrivialSampling = (Ratio >= 1.0f);
+
+	// Early exit when nothing will be generated out of this sampler
+	if (bNoSampling && !bKeepZeroDensityPoints)
 	{
 		PCGE_LOG(Verbose, "Skipped - all inputs rejected");
 		return true;
@@ -76,9 +82,9 @@ bool FPCGPointSamplerElement::ExecuteInternal(FPCGContext* Context) const
 
 		// TODO: randomize on the fractional number of points
 #if WITH_EDITORONLY_DATA
-		int TargetNumPoints = (Settings->bKeepZeroDensityPoints ? OriginalPointCount : OriginalPointCount * Settings->Ratio);
+		int TargetNumPoints = (bKeepZeroDensityPoints ? OriginalPointCount : OriginalPointCount * Ratio);
 #else
-		int TargetNumPoints = OriginalPointCount * Settings->Ratio;
+		int TargetNumPoints = OriginalPointCount * Ratio;
 #endif
 
 		// Early out
@@ -91,7 +97,7 @@ bool FPCGPointSamplerElement::ExecuteInternal(FPCGContext* Context) const
 		{
 			TRACE_CPUPROFILER_EVENT_SCOPE(FPCGPointSamplerElement::Execute::SelectPoints);
 
-			FPCGAsync::AsyncPointProcessing(Context, OriginalPointCount, SampledPoints, [&Points, Settings](int32 Index, FPCGPoint& OutPoint)
+			FPCGAsync::AsyncPointProcessing(Context, OriginalPointCount, SampledPoints, [&Points, Settings, Ratio, bKeepZeroDensityPoints](int32 Index, FPCGPoint& OutPoint)
 			{
 				const FPCGPoint& Point = Points[Index];
 
@@ -99,13 +105,13 @@ bool FPCGPointSamplerElement::ExecuteInternal(FPCGContext* Context) const
 				FRandomStream RandomSource(PCGHelpers::ComputeSeed(Settings->Seed, Point.Seed));
 				float Chance = RandomSource.FRand();
 
-				if (Chance < Settings->Ratio)
+				if (Chance < Ratio)
 				{
 					OutPoint = Point;
 					return true;
 				}
 #if WITH_EDITORONLY_DATA
-				else if (Settings->bKeepZeroDensityPoints)
+				else if (bKeepZeroDensityPoints)
 				{
 					OutPoint = Point;
 					OutPoint.Density = 0;
diff --git a/Engine/Plugins/Experimental/PCG/Source/PCG/Private/Elements/PCGProjectionElement.cpp b/Engine/Plugins/Experimental/PCG/Source/PCG/Private/Elements/PCGProjectionElement.cpp
--- a/Engine/Plugins/Experimental/PCG/Source/PCG/Private/Elements/PCGProjectionElement.cpp
+++ b/Engine/Plugins/Experimental/PCG/Source/PCG/Private/Elements/PCGProjectionElement.cpp
@@ -2,6 +2,7 @@
 
 #include "Elements/PCGProjectionElement.h"
 #include "Data/PCGProjectionData.h"
+#include "Helpers/PCGSettingsHelpers.h"
 
 FPCGElementPtr UPCGProjectionSettings::CreateElement() const
 {
@@ -16,8 +17,16 @@ bool FPCGProjectionElement::ExecuteInternal(FPCGContext* Context) const
 	check(Settings);
 
 	TArray<FPCGTaggedData> Inputs = Context->InputData.GetInputs();
+	UPCGParams* Params = Context->InputData.GetParams();
+
 	TArray<FPCGTaggedData>& Outputs = Context->OutputData.TaggedData;
 
+#if WITH_EDITORONLY_DATA
+	const bool bKeepZeroDensityPoints = PCGSettingsHelpers::GetValue(GET_MEMBER_NAME_CHECKED(UPCGProjectionSettings, bKeepZeroDensityPoints), Settings->bKeepZeroDensityPoints, Params);
+#else
+	const bool bKeepZeroDensityPoints = false;
+#endif
+
 	const UPCGSpatialData* FirstSpatialData = nullptr;
 	UPCGProjectionData* ProjectionData = nullptr;
 	int32 ProjectionTaggedDataIndex = -1;
@@ -46,7 +55,7 @@ bool FPCGProjectionElement::ExecuteInternal(FPCGContext* Context) const
 		ProjectionData = (ProjectionData ? ProjectionData : FirstSpatialData)->ProjectOn(SpatialData);
 
 #if WITH_EDITORONLY_DATA
-		ProjectionData->bKeepZeroDensityPoints = Settings->bKeepZeroDensityPoints;
+		ProjectionData->bKeepZeroDensityPoints = bKeepZeroDensityPoints;
 #endif
 		
 		// Update the tagged data
diff --git a/Engine/Plugins/Experimental/PCG/Source/PCG/Private/Elements/PCGSelfPruning.cpp b/Engine/Plugins/Experimental/PCG/Source/PCG/Private/Elements/PCGSelfPruning.cpp
--- a/Engine/Plugins/Experimental/PCG/Source/PCG/Private/Elements/PCGSelfPruning.cpp
+++ b/Engine/Plugins/Experimental/PCG/Source/PCG/Private/Elements/PCGSelfPruning.cpp
@@ -3,6 +3,7 @@
 #include "Elements/PCGSelfPruning.h"
 #include "Data/PCGPointData.h"
 #include "Data/PCGSpatialData.h"
+#include "Helpers/PCGSettingsHelpers.h"
 #include "PCGHelpers.h"
 #include "Math/RandomStream.h"
 
@@ -222,7 +223,13 @@ bool FPCGSelfPruningElement::ExecuteInternal(FPCGContext* Context) const
 	const UPCGSelfPruningSettings* Settings = Context->GetInputSettings<UPCGSelfPruningSettings>();
 	check(Settings);
 
-	PCGSelfPruningElement::Execute(Context, Settings->PruningType, Settings->RadiusSimilarityFactor, Settings->bRandomizedPruning);
+	UPCGParams* Params = Context->InputData.GetParams();
+
+	const EPCGSelfPruningType PruningType = PCGSettingsHelpers::GetValue(GET_MEMBER_NAME_CHECKED(UPCGSelfPruningSettings, PruningType), Settings->PruningType, Params);
+	const float RadiusSimilarityFactor = PCGSettingsHelpers::GetValue(GET_MEMBER_NAME_CHECKED(UPCGSelfPruningSettings, RadiusSimilarityFactor), Settings->RadiusSimilarityFactor, Params);
+	const bool bRandomizedPruning = PCGSettingsHelpers::GetValue(GET_MEMBER_NAME_CHECKED(UPCGSelfPruningSettings, bRandomizedPruning), Settings->bRandomizedPruning, Params);
+
+	PCGSelfPruningElement::Execute(Context, PruningType, RadiusSimilarityFactor, bRandomizedPruning);
 
 	return true;
 }
\ No newline at end of file
diff --git a/Engine/Plugins/Experimental/PCG/Source/PCG/Private/Elements/PCGSurfaceSampler.cpp b/Engine/Plugins/Experimental/PCG/Source/PCG/Private/Elements/PCGSurfaceSampler.cpp
--- a/Engine/Plugins/Experimental/PCG/Source/PCG/Private/Elements/PCGSurfaceSampler.cpp
+++ b/Engine/Plugins/Experimental/PCG/Source/PCG/Private/Elements/PCGSurfaceSampler.cpp
@@ -4,6 +4,7 @@
 #include "Data/PCGPointData.h"
 #include "Data/PCGSpatialData.h"
 #include "Helpers/PCGAsync.h"
+#include "Helpers/PCGSettingsHelpers.h"
 #include "PCGHelpers.h"
 #include "Math/RandomStream.h"
 
@@ -11,7 +12,13 @@ struct FPCGSurfaceSamplerLoopData
 {
 	const UPCGSurfaceSamplerSettings* Settings;
 
+	float PointsPerSquaredMeter;
 	FVector PointExtents;
+	float Looseness;
+	bool bApplyDensityToPoints;
+	float PointSteepness;
+	bool bKeepZeroDensityPoints = false;
+
 	FVector InterstitialDistance;
 	FVector InnerCellSize;
 	FVector CellSize;
@@ -29,11 +36,22 @@ struct FPCGSurfaceSamplerLoopData
 
 	bool Initialize(const UPCGSurfaceSamplerSettings* InSettings, FPCGContext* Context, const FBox& InputBounds)
 	{
+		UPCGParams* Params = Context->InputData.GetParams();
 		Settings = InSettings;
 
+		// Compute used values
+		PointsPerSquaredMeter = PCGSettingsHelpers::GetValue(GET_MEMBER_NAME_CHECKED(UPCGSurfaceSamplerSettings, PointsPerSquaredMeter), Settings->PointsPerSquaredMeter, Params);
+		PointExtents = PCGSettingsHelpers::GetValue(GET_MEMBER_NAME_CHECKED(UPCGSurfaceSamplerSettings, PointExtents), Settings->PointExtents, Params);
+		Looseness = PCGSettingsHelpers::GetValue(GET_MEMBER_NAME_CHECKED(UPCGSurfaceSamplerSettings, Looseness), Settings->Looseness, Params);
+		bApplyDensityToPoints = PCGSettingsHelpers::GetValue(GET_MEMBER_NAME_CHECKED(UPCGSurfaceSamplerSettings, bApplyDensityToPoints), Settings->bApplyDensityToPoints, Params);
+		PointSteepness = PCGSettingsHelpers::GetValue(GET_MEMBER_NAME_CHECKED(UPCGSurfaceSamplerSettings, PointSteepness), Settings->PointSteepness, Params);
+#if WITH_EDITORONLY_DATA
+		bKeepZeroDensityPoints = PCGSettingsHelpers::GetValue(GET_MEMBER_NAME_CHECKED(UPCGSurfaceSamplerSettings, bKeepZeroDensityPoints), Settings->bKeepZeroDensityPoints, Params);
+#endif
+
 		// Conceptually, we will break down the surface bounds in a N x M grid
-		InterstitialDistance = Settings->PointExtents * 2;
-		InnerCellSize = InterstitialDistance * Settings->Looseness;
+		InterstitialDistance = PointExtents * 2;
+		InnerCellSize = InterstitialDistance * Looseness;
 		CellSize = InterstitialDistance + InnerCellSize;
 		check(CellSize.X > 0 && CellSize.Y > 0);
 
@@ -53,7 +71,7 @@ struct FPCGSurfaceSamplerLoopData
 		check(CellCount > 0);
 
 		const FVector::FReal InvSquaredMeterUnits = 1.0 / (100.0 * 100.0);
-		TargetPointCount = (InputBounds.Max.X - InputBounds.Min.X) * (InputBounds.Max.Y - InputBounds.Min.Y) * Settings->PointsPerSquaredMeter * InvSquaredMeterUnits;
+		TargetPointCount = (InputBounds.Max.X - InputBounds.Min.X) * (InputBounds.Max.Y - InputBounds.Min.Y) * PointsPerSquaredMeter * InvSquaredMeterUnits;
 
 		if (TargetPointCount == 0)
 		{
@@ -103,18 +121,21 @@ bool FPCGSurfaceSamplerElement::ExecuteInternal(FPCGContext* Context) const
 	const UPCGSurfaceSamplerSettings* Settings = Context->GetInputSettings<UPCGSurfaceSamplerSettings>();
 	check(Settings);
 
+	TArray<FPCGTaggedData> Inputs = Context->InputData.GetInputs();
+	UPCGParams* Params = Context->InputData.GetParams();
+
 	// Early out on invalid settings
 	// TODO: we could compute an approximate radius based on the points per squared meters if that's useful
-	if(Settings->PointExtents.X <= 0 || Settings->PointExtents.Y <= 0)
+	const FVector PointExtents = PCGSettingsHelpers::GetValue(GET_MEMBER_NAME_CHECKED(UPCGSurfaceSamplerSettings, PointExtents), Settings->PointExtents, Params);
+	if(PointExtents.X <= 0 || PointExtents.Y <= 0)
 	{
 		PCGE_LOG(Warning, "Skipped - Invalid point extents");
 		return true;
 	}
 
-	TArray<FPCGTaggedData>& Outputs = Context->OutputData.TaggedData;
+	TArray<FPCGTaggedData>& Outputs = Context->OutputData.TaggedData; 
 
 	// TODO: embarassingly parallel loop
-	TArray<FPCGTaggedData> Inputs = Context->InputData.GetInputs();
 	for (const FPCGTaggedData& Input : Inputs)
 	{
 		const UPCGSpatialData* SpatialInput = Cast<UPCGSpatialData>(Input.Data);
@@ -171,16 +192,12 @@ bool FPCGSurfaceSamplerElement::ExecuteInternal(FPCGContext* Context) const
 				 
 				const FVector TentativeLocation = FVector(CurrentX + RandX * InnerCellSize.X, CurrentY + RandY * InnerCellSize.Y, 0.0f);
 
-#if WITH_EDITORONLY_DATA
-				if (SpatialInput->GetPointAtPosition(TentativeLocation, OutPoint, SampledData->Metadata) || LoopData.Settings->bKeepZeroDensityPoints)
-#else
-				if (SpatialInput->GetPointAtPosition(TentativeLocation, OutPoint, SampledData->Metadata))
-#endif
+				if (SpatialInput->GetPointAtPosition(TentativeLocation, OutPoint, SampledData->Metadata) || LoopData.bKeepZeroDensityPoints)
 				{
 					// Apply final parameters on the point
-					OutPoint.Extents = FVector(LoopData.Settings->PointExtents);
-					OutPoint.Density *= (LoopData.Settings->bApplyDensityToPoints ? ((Ratio - Chance) / Ratio) : 1.0f);
-					OutPoint.Steepness = LoopData.Settings->PointSteepness;
+					OutPoint.Extents = LoopData.PointExtents;
+					OutPoint.Density *= (LoopData.bApplyDensityToPoints ? ((Ratio - Chance) / Ratio) : 1.0f);
+					OutPoint.Steepness = LoopData.PointSteepness;
 					OutPoint.Seed = RandomSource.GetCurrentSeed();
 
 					return true;
diff --git a/Engine/Plugins/Experimental/PCG/Source/PCG/Private/Elements/PCGUnionElement.cpp b/Engine/Plugins/Experimental/PCG/Source/PCG/Private/Elements/PCGUnionElement.cpp
--- a/Engine/Plugins/Experimental/PCG/Source/PCG/Private/Elements/PCGUnionElement.cpp
+++ b/Engine/Plugins/Experimental/PCG/Source/PCG/Private/Elements/PCGUnionElement.cpp
@@ -1,6 +1,7 @@
 // Copyright Epic Games, Inc. All Rights Reserved.
 
 #include "Elements/PCGUnionElement.h"
+#include "Helpers/PCGSettingsHelpers.h"
 
 FPCGElementPtr UPCGUnionSettings::CreateElement() const
 {
@@ -15,6 +16,11 @@ bool FPCGUnionElement::ExecuteInternal(FPCGContext* Context) const
 	check(Settings);
 
 	TArray<FPCGTaggedData> Inputs = Context->InputData.GetInputs();
+	UPCGParams* Params = Context->InputData.GetParams();
+
+	const EPCGUnionType Type = PCGSettingsHelpers::GetValue(GET_MEMBER_NAME_CHECKED(UPCGUnionSettings, Type), Settings->Type, Params);
+	const EPCGUnionDensityFunction DensityFunction = PCGSettingsHelpers::GetValue(GET_MEMBER_NAME_CHECKED(UPCGUnionSettings, DensityFunction), Settings->DensityFunction, Params);
+
 	TArray<FPCGTaggedData>& Outputs = Context->OutputData.TaggedData;
 
 	const UPCGSpatialData* FirstSpatialData = nullptr;
@@ -46,8 +52,8 @@ bool FPCGUnionElement::ExecuteInternal(FPCGContext* Context) const
 		if (!UnionData)
 		{
 			UnionData = FirstSpatialData->UnionWith(SpatialData);
-			UnionData->SetType(Settings->Type);
-			UnionData->SetDensityFunction(Settings->DensityFunction);
+			UnionData->SetType(Type);
+			UnionData->SetDensityFunction(DensityFunction);
 
 			UnionTaggedData.Data = UnionData;
 		}
diff --git a/Engine/Plugins/Experimental/PCG/Source/PCG/Private/Metadata/PCGMetadata.cpp b/Engine/Plugins/Experimental/PCG/Source/PCG/Private/Metadata/PCGMetadata.cpp
--- a/Engine/Plugins/Experimental/PCG/Source/PCG/Private/Metadata/PCGMetadata.cpp
+++ b/Engine/Plugins/Experimental/PCG/Source/PCG/Private/Metadata/PCGMetadata.cpp
@@ -26,7 +26,12 @@ void UPCGMetadata::BeginDestroy()
 
 void UPCGMetadata::Initialize(const UPCGMetadata* InParent)
 {
-	check(!Parent && Attributes.Num() == 0);
+	if (Parent || Attributes.Num() != 0)
+	{
+		UE_LOG(LogPCG, Error, TEXT("Metadata has already been initialized or already contains attributes"));
+		return;
+	}
+
 	Parent = (InParent != this ? InParent : nullptr);
 	ItemKeyOffset = Parent ? Parent->GetItemCountForChild() : 0;
 	AddAttributes(InParent);
@@ -406,10 +411,7 @@ void UPCGMetadata::MergeAttributes(PCGMetadataEntryKey InKeyA, const UPCGMetadat
 
 		if (AttributeA || AttributeB)
 		{
-			if (Attribute)
-			{
-				Attribute->SetValue(OutKey, AttributeA, InKeyA, AttributeB, InKeyB, Op);
-			}
+			Attribute->SetValue(OutKey, AttributeA, InKeyA, AttributeB, InKeyB, Op);
 		}
 	}
 	AttributeLock.ReadUnlock();
@@ -429,14 +431,8 @@ void UPCGMetadata::ResetWeightedAttributes(PCGMetadataEntryKey& OutKey)
 	{
 		FPCGMetadataAttributeBase* Attribute = AttributePair.Value;
 		
-		if (!Attribute->AllowsInterpolation())
-		{
-			continue;
-		}
-
-		if (Attribute)
+		if (Attribute && Attribute->AllowsInterpolation())
 		{
-			check(Attribute->AllowsInterpolation());
 			Attribute->SetZeroValue(OutKey);
 		}
 	}
diff --git a/Engine/Plugins/Experimental/PCG/Source/PCG/Private/Metadata/PCGMetadataAttributeTpl.h b/Engine/Plugins/Experimental/PCG/Source/PCG/Private/Metadata/PCGMetadataAttributeTpl.h
deleted file mode 100644
diff --git a/Engine/Plugins/Experimental/PCG/Source/PCG/Private/PCGData.cpp b/Engine/Plugins/Experimental/PCG/Source/PCG/Private/PCGData.cpp
--- a/Engine/Plugins/Experimental/PCG/Source/PCG/Private/PCGData.cpp
+++ b/Engine/Plugins/Experimental/PCG/Source/PCG/Private/PCGData.cpp
@@ -2,6 +2,7 @@
 
 #include "PCGData.h"
 #include "PCGSettings.h"
+#include "PCGParams.h"
 
 bool FPCGTaggedData::operator==(const FPCGTaggedData& Other) const
 {
@@ -44,6 +45,19 @@ TArray<FPCGTaggedData> FPCGDataCollection::GetAllSettings() const
 		});
 }
 
+UPCGParams* FPCGDataCollection::GetParams() const
+{
+	for (const FPCGTaggedData& TaggedDatum : TaggedData)
+	{
+		if (UPCGParams* Params = Cast<UPCGParams>(TaggedDatum.Data))
+		{
+			return Params; 
+		}
+	}
+
+	return nullptr;
+}
+
 const UPCGSettings* FPCGDataCollection::GetSettings(const UPCGSettings* InDefaultSettings) const
 {
 	if (!InDefaultSettings)
diff --git a/Engine/Plugins/Experimental/PCG/Source/PCG/Private/PCGParams.cpp b/Engine/Plugins/Experimental/PCG/Source/PCG/Private/PCGParams.cpp
new file mode 100644
--- /dev/null
+++ b/Engine/Plugins/Experimental/PCG/Source/PCG/Private/PCGParams.cpp
@@ -0,0 +1,69 @@
+// Copyright Epic Games, Inc. All Rights Reserved.
+
+#include "PCGParams.h"
+
+UPCGParams::UPCGParams(const FObjectInitializer& ObjectInitializer)
+	: Super(ObjectInitializer)
+{
+	Metadata = ObjectInitializer.CreateDefaultSubobject<UPCGMetadata>(this, TEXT("Metadata"));
+}
+
+int64 UPCGParams::FindMetadataKey(const FName& InName) const
+{
+	if (const PCGMetadataEntryKey* FoundKey = NameMap.Find(InName))
+	{
+		return *FoundKey;
+	}
+	else
+	{
+		return PCGInvalidEntryKey;
+	}
+}
+
+int64 UPCGParams::FindOrAddMetadataKey(const FName& InName)
+{
+	if (const PCGMetadataEntryKey* FoundKey = NameMap.Find(InName))
+	{
+		return *FoundKey;
+	}
+	else
+	{
+		check(Metadata);
+		PCGMetadataEntryKey NewKey = Metadata->AddEntry();
+		NameMap.Add(InName, NewKey);
+		return NewKey;
+	}
+}
+
+UPCGParams* UPCGParams::FilterParamsByName(const FName& InName) const
+{
+	PCGMetadataEntryKey EntryKey = FindMetadataKey(InName);
+	UPCGParams* NewParams = FilterParamsByKey(EntryKey);
+
+	if (EntryKey != PCGInvalidEntryKey)
+	{
+		// NOTE: this relies on the fact that there will be only one entry
+		NewParams->NameMap.Add(InName, 0);
+	}
+
+	return NewParams;
+}
+
+UPCGParams* UPCGParams::FilterParamsByKey(int64 InKey) const
+{
+	UPCGParams* NewParams = NewObject<UPCGParams>(const_cast<UPCGParams*>(this));
+
+	// Here instead of parenting the metadata, we will create a copy
+	// so that the only entry in the metadata (if any) will have the 0 key.
+	check(NewParams && NewParams->Metadata);
+
+	NewParams->Metadata->AddAttributes(Metadata);
+
+	if (InKey != PCGInvalidEntryKey)
+	{
+		PCGMetadataEntryKey OutKey = PCGInvalidEntryKey;
+		NewParams->Metadata->SetAttributes(InKey, Metadata, OutKey);
+	}
+
+	return NewParams;
+}
\ No newline at end of file
diff --git a/Engine/Plugins/Experimental/PCG/Source/PCG/Public/Data/PCGLandscapeData.h b/Engine/Plugins/Experimental/PCG/Source/PCG/Public/Data/PCGLandscapeData.h
--- a/Engine/Plugins/Experimental/PCG/Source/PCG/Public/Data/PCGLandscapeData.h
+++ b/Engine/Plugins/Experimental/PCG/Source/PCG/Public/Data/PCGLandscapeData.h
@@ -20,13 +20,13 @@ struct PCG_API FPCGLandscapeDataPoint
 	FPCGLandscapeDataPoint(int InX, int InY, float InHeight);
 
 	UPROPERTY()
-	int X;
+	int X = 0;
 
 	UPROPERTY()
-	int Y;
+	int Y = 0;
 
 	UPROPERTY()
-	float Height;
+	float Height = 0;
 };
 
 UCLASS(BlueprintType, ClassGroup = (Procedural))
diff --git a/Engine/Plugins/Experimental/PCG/Source/PCG/Public/Data/PCGPointData.h b/Engine/Plugins/Experimental/PCG/Source/PCG/Public/Data/PCGPointData.h
--- a/Engine/Plugins/Experimental/PCG/Source/PCG/Public/Data/PCGPointData.h
+++ b/Engine/Plugins/Experimental/PCG/Source/PCG/Public/Data/PCGPointData.h
@@ -70,9 +70,6 @@ public:
 
 	void InitializeFromActor(AActor* InActor);
 
-	UFUNCTION(BlueprintCallable, Category = SpatialData)
-	void InitializeFromData(const UPCGSpatialData* InSource, const UPCGMetadata* InMetadataParentOverride = nullptr);
-
 	UFUNCTION(BlueprintCallable, Category = SpatialData)
 	const TArray<FPCGPoint>& GetPoints() const { return Points; }
 
diff --git a/Engine/Plugins/Experimental/PCG/Source/PCG/Public/Data/PCGSpatialData.h b/Engine/Plugins/Experimental/PCG/Source/PCG/Public/Data/PCGSpatialData.h
--- a/Engine/Plugins/Experimental/PCG/Source/PCG/Public/Data/PCGSpatialData.h
+++ b/Engine/Plugins/Experimental/PCG/Source/PCG/Public/Data/PCGSpatialData.h
@@ -33,6 +33,8 @@ class PCG_API UPCGSpatialData : public UPCGData
 	GENERATED_BODY()
 
 public:
+	UPCGSpatialData(const FObjectInitializer& ObjectInitializer);
+
 	/** Returns the dimension of the data type, which has nothing to do with the dimension of its points */
 	UFUNCTION(BlueprintCallable, Category = SpatialData)
 	virtual int GetDimension() const PURE_VIRTUAL(UPCGSpatialData::GetDimension, return 0;);
@@ -102,9 +104,12 @@ public:
 	UFUNCTION(BlueprintCallable, Category = Metadata)
 	UPCGMetadata* MutableMetadata() { return Metadata; }
 
-	UFUNCTION(BlueprintCallable, Category = Metadata)
+	UFUNCTION(BlueprintCallable, Category = Metadata, meta=(DeprecatedFunction))
 	UPCGMetadata* CreateEmptyMetadata();
 
+	UFUNCTION(BlueprintCallable, Category = SpatialData)
+	void InitializeFromData(const UPCGSpatialData* InSource, const UPCGMetadata* InMetadataParentOverride = nullptr);
+
 	UPROPERTY(BlueprintReadWrite, EditAnywhere, Category = Data)
 	AActor* TargetActor = nullptr;
 
diff --git a/Engine/Plugins/Experimental/PCG/Source/PCG/Public/Helpers/PCGSettingsHelpers.h b/Engine/Plugins/Experimental/PCG/Source/PCG/Public/Helpers/PCGSettingsHelpers.h
new file mode 100644
--- /dev/null
+++ b/Engine/Plugins/Experimental/PCG/Source/PCG/Public/Helpers/PCGSettingsHelpers.h
@@ -0,0 +1,88 @@
+// Copyright Epic Games, Inc. All Rights Reserved.
+
+#pragma once
+
+#include "PCGParams.h"
+#include "Metadata/PCGMetadataAttributeTraits.h"
+#include "Metadata/PCGMetadataAttributeTpl.h"
+
+namespace PCGSettingsHelpers
+{
+	template<typename T, typename TEnableIf<!TIsEnumClass<T>::Value>::Type* = nullptr>
+	T GetValue(const FName& InName, const T& InValue, UPCGParams* InParams, PCGMetadataEntryKey InKey)
+	{
+		if (InParams)
+		{
+			const FPCGMetadataAttributeBase* MatchingAttribute = InParams->Metadata ? InParams->Metadata->GetConstAttribute(InName) : nullptr;
+			if (MatchingAttribute && MatchingAttribute->GetTypeId() == PCG::Private::MetadataTypes<T>::Id)
+			{
+				return static_cast<const FPCGMetadataAttribute<T>*>(MatchingAttribute)->GetValueFromItemKey(InKey);
+			}
+			else
+			{
+				return InValue;
+			}
+		}
+		else
+		{
+			return InValue;
+		}
+	}
+
+	template<typename T, typename TEnableIf<!TIsEnumClass<T>::Value>::Type* = nullptr>
+	T GetValue(const FName& InName, const T& InValue, UPCGParams* InParams)
+	{
+		return GetValue(InName, InValue, InParams, 0);
+	}
+
+	template<typename T, typename TEnableIf<!TIsEnumClass<T>::Value>::Type* = nullptr>
+	T GetValue(const FName& InName, const T& InValue, UPCGParams* InParams, const FName& InParamName)
+	{
+		if (InParams && InParamName != NAME_None)
+		{
+			return GetValue(InName, InValue, InParams, InParams->FindMetadataKey(InParamName));
+		}
+		else
+		{
+			return InValue;
+		}
+	}
+
+	/** Specialized versions for enums */
+	template<typename T, typename TEnableIf<TIsEnumClass<T>::Value>::Type* = nullptr>
+	T GetValue(const FName& InName, const T& InValue, UPCGParams* InParams, PCGMetadataEntryKey InKey)
+	{
+		return static_cast<T>(GetValue(InName, static_cast<__underlying_type(T)>(InValue), InParams, InKey));
+	}
+
+	template<typename T, typename TEnableIf<TIsEnumClass<T>::Value>::Type* = nullptr>
+	T GetValue(const FName& InName, const T& InValue, UPCGParams* InParams)
+	{
+		return static_cast<T>(GetValue(InName, static_cast<__underlying_type(T)>(InValue), InParams));
+	}
+
+	template<typename T, typename TEnableIf<TIsEnumClass<T>::Value>::Type* = nullptr>
+	T GetValue(const FName& InName, const T& InValue, UPCGParams* InParams, const FName& InParamName)
+	{
+		return static_cast<T>(GetValue(InName, static_cast<__underlying_type(T)>(InValue), InParams, InParamName));
+	}
+
+	/** Specialized version for names */
+	template<>
+	FORCEINLINE FName GetValue(const FName& InName, const FName& InValue, UPCGParams* InParams, PCGMetadataEntryKey InKey)
+	{
+		return FName(GetValue(InName, InValue.ToString(), InParams, InKey));
+	}
+
+	template<>
+	FORCEINLINE FName GetValue(const FName& InName, const FName& InValue, UPCGParams* InParams)
+	{
+		return FName(GetValue(InName, InValue.ToString(), InParams));
+	}
+
+	template<>
+	FORCEINLINE FName GetValue(const FName& InName, const FName& InValue, UPCGParams* InParams, const FName& InParamName)
+	{
+		return FName(GetValue(InName, InValue.ToString(), InParams, InParamName));
+	}
+}
\ No newline at end of file
diff --git a/Engine/Plugins/Experimental/PCG/Source/PCG/Public/Metadata/PCGMetadata.h b/Engine/Plugins/Experimental/PCG/Source/PCG/Public/Metadata/PCGMetadata.h
--- a/Engine/Plugins/Experimental/PCG/Source/PCG/Public/Metadata/PCGMetadata.h
+++ b/Engine/Plugins/Experimental/PCG/Source/PCG/Public/Metadata/PCGMetadata.h
@@ -95,10 +95,7 @@ public:
 
 	/** Attributes operations */
 	void MergeAttributes(PCGMetadataEntryKey InKeyA, const UPCGMetadata* InMetadataA, PCGMetadataEntryKey InKeyB, const UPCGMetadata* InMetadataB, PCGMetadataEntryKey& OutKey, EPCGMetadataOp Op);
-	void MergeAttributes(PCGMetadataEntryKey InKeyA, PCGMetadataEntryKey InKeyB, PCGMetadataEntryKey& OutKey, EPCGMetadataOp Op)
-	{
-		MergeAttributes(InKeyA, this, InKeyB, this, OutKey, Op);
-	}
+	void MergeAttributes(PCGMetadataEntryKey InKeyA, PCGMetadataEntryKey InKeyB, PCGMetadataEntryKey& OutKey, EPCGMetadataOp Op);
 
 	void ResetWeightedAttributes(PCGMetadataEntryKey& OutKey);
 	void AccumulateWeightedAttributes(PCGMetadataEntryKey InKey, const UPCGMetadata* InMetadata, float Weight, bool bSetNonInterpolableAttributes, PCGMetadataEntryKey& OutKey);
diff --git a/Engine/Plugins/Experimental/PCG/Source/PCG/Public/Metadata/PCGMetadataAttributeTpl.h b/Engine/Plugins/Experimental/PCG/Source/PCG/Public/Metadata/PCGMetadataAttributeTpl.h
new file mode 100644
--- /dev/null
+++ b/Engine/Plugins/Experimental/PCG/Source/PCG/Public/Metadata/PCGMetadataAttributeTpl.h
@@ -0,0 +1,332 @@
+// Copyright Epic Games, Inc. All Rights Reserved.
+
+#pragma once
+
+#include "CoreMinimal.h"
+#include "Metadata/PCGMetadataAttribute.h"
+#include "Metadata/PCGMetadataAttributeTraits.h"
+
+template<typename T>
+class FPCGMetadataAttribute : public FPCGMetadataAttributeBase
+{
+public:
+	FPCGMetadataAttribute(UPCGMetadata* InMetadata, FName InName, const FPCGMetadataAttributeBase* InParent, const T& InDefaultValue, bool bInAllowsInterpolation)
+		: FPCGMetadataAttributeBase(InMetadata, InName, InParent, bInAllowsInterpolation)
+		, DefaultValue(InDefaultValue)
+	{
+		TypeId = PCG::Private::MetadataTypes<T>::Id;
+
+		if (GetParent())
+		{
+			ValueKeyOffset = GetParent()->GetValueKeyOffsetForChild();
+		}
+	}
+
+	const FPCGMetadataAttribute* GetParent() const { return static_cast<const FPCGMetadataAttribute*>(Parent); }
+
+	virtual FPCGMetadataAttributeBase* Copy(FName NewName, UPCGMetadata* InMetadata, bool bKeepParent, bool bCopyEntries = true, bool bCopyValues = true) const override
+	{
+		// this copies to a new attribute
+		check(!bKeepParent || Metadata->GetRoot() == InMetadata->GetRoot());
+		FPCGMetadataAttribute<T>* AttributeCopy = new FPCGMetadataAttribute<T>(InMetadata, NewName, bKeepParent ? this : nullptr, DefaultValue, bAllowsInterpolation);
+
+		if (bCopyEntries)
+		{
+			EntryMapLock.ReadLock();
+			AttributeCopy->EntryToValueMap = EntryToValueMap;
+			EntryMapLock.ReadUnlock();
+		}
+
+		if (bCopyValues)
+		{
+			ValueLock.ReadLock();
+			AttributeCopy->Values = Values;
+			AttributeCopy->ValueKeyOffset = ValueKeyOffset;
+			ValueLock.ReadUnlock();
+		}
+
+		return AttributeCopy;
+	}
+
+	virtual void SetValue(PCGMetadataEntryKey ItemKey, const FPCGMetadataAttributeBase* InAttribute, PCGMetadataEntryKey InEntryKey) override
+	{
+		check(ItemKey != PCGInvalidEntryKey);
+		if (InAttribute == this)
+		{
+			SetValueFromValueKey(ItemKey, GetValueKey(InEntryKey));
+		}
+		else if (InAttribute)
+		{
+			SetValue(ItemKey, static_cast<const FPCGMetadataAttribute<T>*>(InAttribute)->GetValueFromItemKey(InEntryKey));
+		}
+	}
+
+	virtual void SetZeroValue(PCGMetadataEntryKey ItemKey) override
+	{
+		check(ItemKey != PCGInvalidEntryKey);
+		ZeroValue(ItemKey);
+	}
+
+	virtual void AccumulateValue(PCGMetadataEntryKey ItemKey, const FPCGMetadataAttributeBase* InAttribute, PCGMetadataEntryKey InEntryKey, float Weight) override
+	{
+		check(ItemKey != PCGInvalidEntryKey);
+		Accumulate(ItemKey, InAttribute, InEntryKey, Weight);
+	}
+
+	virtual void SetValue(PCGMetadataEntryKey ItemKey, const FPCGMetadataAttributeBase* InAttributeA, PCGMetadataEntryKey InEntryKeyA, const FPCGMetadataAttributeBase* InAttributeB, PCGMetadataEntryKey InEntryKeyB, EPCGMetadataOp Op) override
+	{
+		check(ItemKey != PCGInvalidEntryKey);
+		bool bAppliedValue = false;
+
+		if (InAttributeA && InAttributeB && bAllowsInterpolation)
+		{
+			if (Op == EPCGMetadataOp::Min)
+			{
+				bAppliedValue = SetMin(ItemKey, InAttributeA, InEntryKeyA, InAttributeB, InEntryKeyB);
+			}
+			else if (Op == EPCGMetadataOp::Max)
+			{
+				bAppliedValue = SetMax(ItemKey, InAttributeA, InEntryKeyA, InAttributeB, InEntryKeyB);
+			}
+			else if (Op == EPCGMetadataOp::Sub)
+			{
+				bAppliedValue = SetSub(ItemKey, InAttributeA, InEntryKeyA, InAttributeB, InEntryKeyB);
+			}
+			else if (Op == EPCGMetadataOp::Add)
+			{
+				bAppliedValue = SetAdd(ItemKey, InAttributeA, InEntryKeyA, InAttributeB, InEntryKeyB);
+			}
+		}
+		else if (InAttributeA && InAttributeB && HasNonDefaultValue(ItemKey))
+		{
+			// In this case, the current already has a value, in which case we should not update it
+			bAppliedValue = true;
+		}
+
+		if (bAppliedValue)
+		{
+			// Nothing to do
+		}
+		else if (InAttributeA)
+		{
+			if (InAttributeA == this)
+			{
+				SetValueFromValueKey(ItemKey, GetValueKey(InEntryKeyA));
+			}
+			else
+			{
+				SetValue(ItemKey, static_cast<const FPCGMetadataAttribute<T>*>(InAttributeA)->GetValueFromItemKey(InEntryKeyA));
+			}
+		}
+		else if (InAttributeB)
+		{
+			if (InAttributeB == this)
+			{
+				SetValueFromValueKey(ItemKey, GetValueKey(InEntryKeyB));
+			}
+			else
+			{
+				SetValue(ItemKey, static_cast<const FPCGMetadataAttribute<T>*>(InAttributeB)->GetValueFromItemKey(InEntryKeyB));
+			}
+		}
+	}
+
+	PCGMetadataValueKey GetValueKeyOffsetForChild() const
+	{
+		FReadScopeLock ScopeLock(ValueLock);
+		return Values.Num() + ValueKeyOffset;
+	}
+
+	/** Adds the value, returns the value key for the given value */
+	PCGMetadataValueKey AddValue(const T& InValue)
+	{
+		PCGMetadataValueKey FoundValue = FindValue(InValue);
+
+		if (FoundValue == PCGDefaultValueKey)
+		{
+			FWriteScopeLock ScopeLock(ValueLock);
+			return Values.Add(InValue) + ValueKeyOffset;
+		}
+		else
+		{
+			return FoundValue;
+		}
+	}
+
+	void SetValue(PCGMetadataEntryKey ItemKey, const T& InValue)
+	{
+		check(ItemKey != PCGInvalidEntryKey);
+		SetValueFromValueKey(ItemKey, AddValue(InValue));
+	}
+
+	template<typename U>
+	void SetValue(PCGMetadataEntryKey ItemKey, const U& InValue)
+	{
+		check(ItemKey != PCGInvalidEntryKey);
+		SetValueFromValueKey(ItemKey, AddValue(T(InValue)));
+	}
+
+	T GetValueFromItemKey(PCGMetadataEntryKey ItemKey) const
+	{
+		return GetValue(GetValueKey(ItemKey));
+	}
+
+	T GetValue(PCGMetadataValueKey ValueKey) const
+	{
+		if (ValueKey == PCGDefaultValueKey)
+		{
+			return DefaultValue;
+		}
+		else if (ValueKey >= ValueKeyOffset)
+		{
+			FReadScopeLock ScopeLock(ValueLock);
+			return Values[ValueKey - ValueKeyOffset];
+		}
+		else if (GetParent())
+		{
+			return GetParent()->GetValue(ValueKey);
+		}
+		else
+		{
+			return DefaultValue;
+		}
+	}
+
+protected:
+	/** Code related to finding values / compressing data */
+	template<typename IT = T, typename TEnableIf<PCG::Private::MetadataTraits<IT>::CompressData>::Type* = nullptr>
+	PCGMetadataValueKey FindValue(const T& InValue) const
+	{
+		PCGMetadataValueKey ParentValueKey = (GetParent() ? GetParent()->FindValue(InValue) : PCGDefaultValueKey);
+		if (ParentValueKey != PCGDefaultValueKey)
+		{
+			return ParentValueKey;
+		}
+		else
+		{
+			ValueLock.ReadLock();
+			const int32 ValueIndex = Values.FindLast(InValue);
+			ValueLock.ReadUnlock();
+
+			if (ValueIndex != INDEX_NONE)
+			{
+				return ValueIndex + ValueKeyOffset;
+			}
+			else
+			{
+				return ParentValueKey;
+			}
+		}
+	}
+
+	template<typename IT = T, typename TEnableIf<!PCG::Private::MetadataTraits<IT>::CompressData>::Type* = nullptr>
+	PCGMetadataValueKey FindValue(const T& InValue) const
+	{
+		return PCGDefaultValueKey;
+	}
+
+	/** Code related to computing compared values (min, max, sub, add) */
+	template<typename IT = T, typename TEnableIf<PCG::Private::MetadataTraits<IT>::CanMinMax>::Type* = nullptr>
+	bool SetMin(PCGMetadataEntryKey ItemKey, const FPCGMetadataAttributeBase* InAttributeA, PCGMetadataEntryKey InEntryKeyA, const FPCGMetadataAttributeBase* InAttributeB, PCGMetadataEntryKey InEntryKeyB)
+	{
+		SetValue(ItemKey, 
+			PCG::Private::MetadataTraits<IT>::Min(
+				static_cast<const FPCGMetadataAttribute<T>*>(InAttributeA)->GetValueFromItemKey(InEntryKeyA),
+				static_cast<const FPCGMetadataAttribute<T>*>(InAttributeB)->GetValueFromItemKey(InEntryKeyB)));
+
+		return true;
+	}
+
+	template<typename IT = T, typename TEnableIf<!PCG::Private::MetadataTraits<IT>::CanMinMax>::Type* = nullptr>
+	bool SetMin(PCGMetadataEntryKey ItemKey, const FPCGMetadataAttributeBase* InAttributeA, PCGMetadataEntryKey InEntryKeyA, const FPCGMetadataAttributeBase* InAttributeB, PCGMetadataEntryKey InEntryKeyB)
+	{
+		return false;
+	}
+
+	template<typename IT = T, typename TEnableIf<PCG::Private::MetadataTraits<IT>::CanMinMax>::Type* = nullptr>
+	bool SetMax(PCGMetadataEntryKey ItemKey, const FPCGMetadataAttributeBase* InAttributeA, PCGMetadataEntryKey InEntryKeyA, const FPCGMetadataAttributeBase* InAttributeB, PCGMetadataEntryKey InEntryKeyB)
+	{
+		SetValue(ItemKey,
+			PCG::Private::MetadataTraits<IT>::Max(
+				static_cast<const FPCGMetadataAttribute<T>*>(InAttributeA)->GetValueFromItemKey(InEntryKeyA),
+				static_cast<const FPCGMetadataAttribute<T>*>(InAttributeB)->GetValueFromItemKey(InEntryKeyB)));
+
+		return true;
+	}
+
+	template<typename IT = T, typename TEnableIf<!PCG::Private::MetadataTraits<IT>::CanMinMax>::Type* = nullptr>
+	bool SetMax(PCGMetadataEntryKey ItemKey, const FPCGMetadataAttributeBase* InAttributeA, PCGMetadataEntryKey InEntryKeyA, const FPCGMetadataAttributeBase* InAttributeB, PCGMetadataEntryKey InEntryKeyB)
+	{
+		return false;
+	}
+
+	template<typename IT = T, typename TEnableIf<PCG::Private::MetadataTraits<IT>::CanSubAdd>::Type* = nullptr>
+	bool SetSub(PCGMetadataEntryKey ItemKey, const FPCGMetadataAttributeBase* InAttributeA, PCGMetadataEntryKey InEntryKeyA, const FPCGMetadataAttributeBase* InAttributeB, PCGMetadataEntryKey InEntryKeyB)
+	{
+		SetValue(ItemKey,
+			PCG::Private::MetadataTraits<IT>::Sub(
+				static_cast<const FPCGMetadataAttribute<T>*>(InAttributeA)->GetValueFromItemKey(InEntryKeyA),
+				static_cast<const FPCGMetadataAttribute<T>*>(InAttributeB)->GetValueFromItemKey(InEntryKeyB)));
+
+		return true;
+	}
+
+	template<typename IT = T, typename TEnableIf<!PCG::Private::MetadataTraits<IT>::CanSubAdd>::Type* = nullptr>
+	bool SetSub(PCGMetadataEntryKey ItemKey, const FPCGMetadataAttributeBase* InAttributeA, PCGMetadataEntryKey InEntryKeyA, const FPCGMetadataAttributeBase* InAttributeB, PCGMetadataEntryKey InEntryKeyB)
+	{
+		return false;
+	}
+
+	template<typename IT = T, typename TEnableIf<PCG::Private::MetadataTraits<IT>::CanSubAdd>::Type* = nullptr>
+	bool SetAdd(PCGMetadataEntryKey ItemKey, const FPCGMetadataAttributeBase* InAttributeA, PCGMetadataEntryKey InEntryKeyA, const FPCGMetadataAttributeBase* InAttributeB, PCGMetadataEntryKey InEntryKeyB)
+	{
+		SetValue(ItemKey,
+			PCG::Private::MetadataTraits<IT>::Add(
+				static_cast<const FPCGMetadataAttribute<T>*>(InAttributeA)->GetValueFromItemKey(InEntryKeyA),
+				static_cast<const FPCGMetadataAttribute<T>*>(InAttributeB)->GetValueFromItemKey(InEntryKeyB)));
+
+		return true;
+	}
+
+	template<typename IT = T, typename TEnableIf<!PCG::Private::MetadataTraits<IT>::CanSubAdd>::Type* = nullptr>
+	bool SetAdd(PCGMetadataEntryKey ItemKey, const FPCGMetadataAttributeBase* InAttributeA, PCGMetadataEntryKey InEntryKeyA, const FPCGMetadataAttributeBase* InAttributeB, PCGMetadataEntryKey InEntryKeyB)
+	{
+		return false;
+	}
+
+	/** Weighted/interpolated values related code */
+	template<typename IT = T, typename TEnableIf<PCG::Private::MetadataTraits<IT>::CanInterpolate>::Type* = nullptr>
+	void ZeroValue(PCGMetadataEntryKey ItemKey)
+	{
+		SetValue(ItemKey, PCG::Private::MetadataTraits<IT>::ZeroValue());
+	}
+
+	template<typename IT = T, typename TEnableIf<!PCG::Private::MetadataTraits<IT>::CanInterpolate>::Type* = nullptr>
+	void ZeroValue(PCGMetadataEntryKey ItemKey)
+	{
+		// Intentionally empty
+	}
+
+	template<typename IT = T, typename TEnableIf<PCG::Private::MetadataTraits<IT>::CanInterpolate>::Type* = nullptr>
+	void Accumulate(PCGMetadataEntryKey ItemKey, const FPCGMetadataAttributeBase* InAttribute, PCGMetadataEntryKey InEntryKey, float Weight)
+	{
+		SetValue(ItemKey,
+			PCG::Private::MetadataTraits<IT>::WeightedSum(
+				GetValueFromItemKey(ItemKey),
+				static_cast<const FPCGMetadataAttribute<T>*>(InAttribute)->GetValueFromItemKey(InEntryKey),
+				Weight));
+	}
+
+	template<typename IT = T, typename TEnableIf<!PCG::Private::MetadataTraits<IT>::CanInterpolate>::Type* = nullptr>
+	void Accumulate(PCGMetadataEntryKey ItemKey, const FPCGMetadataAttributeBase* InAttribute, PCGMetadataEntryKey InEntryKey, float Weight)
+	{
+		// Empty on purpose
+	}
+
+
+protected:
+	mutable FRWLock ValueLock;
+	TArray<T> Values;
+	T DefaultValue;
+	PCGMetadataValueKey ValueKeyOffset = 0;
+};
\ No newline at end of file
diff --git a/Engine/Plugins/Experimental/PCG/Source/PCG/Public/Metadata/PCGMetadataAttributeTraits.h b/Engine/Plugins/Experimental/PCG/Source/PCG/Public/Metadata/PCGMetadataAttributeTraits.h
--- a/Engine/Plugins/Experimental/PCG/Source/PCG/Public/Metadata/PCGMetadataAttributeTraits.h
+++ b/Engine/Plugins/Experimental/PCG/Source/PCG/Public/Metadata/PCGMetadataAttributeTraits.h
@@ -25,6 +25,7 @@ namespace PCG
 		PCGMetadataGenerateDataTypes(FQuat, 6);
 		PCGMetadataGenerateDataTypes(FTransform, 7);
 		PCGMetadataGenerateDataTypes(FString, 8);
+		PCGMetadataGenerateDataTypes(bool, 9);
 
 #undef PCGMetadataGenerateDataTypes
 
@@ -68,6 +69,35 @@ namespace PCG
 			}
 		};
 
+		template<>
+		struct MetadataTraits<bool>
+		{
+			enum { CompressData = false };
+			enum { CanMinMax = true };
+			enum { CanSubAdd = true };
+			enum { CanInterpolate = false };
+
+			static bool Min(const bool& A, const bool& B)
+			{
+				return A && B;
+			}
+
+			static bool Max(const bool& A, const bool& B)
+			{
+				return A || B;
+			}
+
+			static bool Add(const bool& A, const bool& B)
+			{
+				return A || B;
+			}
+
+			static bool Sub(const bool& A, const bool& B)
+			{
+				return A && !B;
+			}
+		};
+
 		// Vector types
 		template<>
 		struct MetadataTraits<FVector>
diff --git a/Engine/Plugins/Experimental/PCG/Source/PCG/Public/PCGData.h b/Engine/Plugins/Experimental/PCG/Source/PCG/Public/PCGData.h
--- a/Engine/Plugins/Experimental/PCG/Source/PCG/Public/PCGData.h
+++ b/Engine/Plugins/Experimental/PCG/Source/PCG/Public/PCGData.h
@@ -10,6 +10,7 @@
 #include "PCGData.generated.h"
 
 class UPCGSettings;
+class UPCGParams;
 
 /**
 * Base class for any "data" class in the PCG framework.
@@ -38,7 +39,7 @@ struct PCG_API FPCGTaggedData
 	GENERATED_BODY()
 
 	UPROPERTY(BlueprintReadWrite, EditAnywhere, Category = Data)
-	TObjectPtr<const UPCGData> Data;
+	TObjectPtr<const UPCGData> Data = nullptr;
 
 	UPROPERTY(BlueprintReadWrite, EditAnywhere, Category = Data)
 	EPCGDataUsage Usage = EPCGDataUsage::Input;
@@ -59,6 +60,7 @@ struct PCG_API FPCGDataCollection
 	TArray<FPCGTaggedData> GetTaggedInputs(const FString& InTag) const;
 	TArray<FPCGTaggedData> GetExclusions() const;
 	TArray<FPCGTaggedData> GetAllSettings() const;
+	UPCGParams* GetParams() const;
 
 	template<typename SettingsType>
 	const SettingsType* GetSettings() const;
diff --git a/Engine/Plugins/Experimental/PCG/Source/PCG/Public/PCGDebug.h b/Engine/Plugins/Experimental/PCG/Source/PCG/Public/PCGDebug.h
--- a/Engine/Plugins/Experimental/PCG/Source/PCG/Public/PCGDebug.h
+++ b/Engine/Plugins/Experimental/PCG/Source/PCG/Public/PCGDebug.h
@@ -38,7 +38,7 @@ public:
 
 	/** Warning: enabling this flag will have severe performance impact */
 	UPROPERTY(BlueprintReadWrite, EditAnywhere, Category = Settings)
-	bool bCheckForDuplicates;
+	bool bCheckForDuplicates = false;
 
 	TSoftObjectPtr<UMaterialInterface> GetMaterial() const;
 };
diff --git a/Engine/Plugins/Experimental/PCG/Source/PCG/Public/PCGParams.h b/Engine/Plugins/Experimental/PCG/Source/PCG/Public/PCGParams.h
new file mode 100644
--- /dev/null
+++ b/Engine/Plugins/Experimental/PCG/Source/PCG/Public/PCGParams.h
@@ -0,0 +1,49 @@
+// Copyright Epic Games, Inc. All Rights Reserved.
+
+#pragma once
+
+#include "PCGData.h"
+#include "Metadata/PCGMetadata.h"
+
+#include "PCGParams.generated.h"
+
+/**
+* Class to hold execution parameters that will be consumed in nodes of the graph
+*/
+UCLASS(BlueprintType, ClassGroup = (Procedural))
+class PCG_API UPCGParams : public UPCGData
+{
+	GENERATED_BODY()
+
+public:
+	UPCGParams(const FObjectInitializer& ObjectInitializer);
+
+	UFUNCTION(BlueprintCallable, Category = Metadata)
+	const UPCGMetadata* ConstMetadata() const { return Metadata; }
+
+	UFUNCTION(BlueprintCallable, Category = Metadata)
+	UPCGMetadata* MutableMetadata() { return Metadata; }
+
+	/** Returns the entry for the given name */
+	UFUNCTION(BlueprintCallable, Category = Params)
+	int64 FindMetadataKey(const FName& InName) const;
+
+	/** Creates an entry for the given name, if not already added */
+	UFUNCTION(BlueprintCallable, Category = Params)
+	int64 FindOrAddMetadataKey(const FName& InName);
+
+	/** Creates a new params that keeps only a given key/name */
+	UFUNCTION(BlueprintCallable, Category = Params)
+	UPCGParams* FilterParamsByName(const FName& InName) const;
+
+	UFUNCTION(BlueprintCallable, Category = Params)
+	UPCGParams* FilterParamsByKey(int64 InKey) const;
+
+	// Not accessible through blueprint to make sure the constness is preserved
+	UPROPERTY(VisibleAnywhere, Category = Metadata)
+	TObjectPtr<UPCGMetadata> Metadata = nullptr;
+
+protected:
+	UPROPERTY()
+	TMap<FName, int64> NameMap;
+};
\ No newline at end of file
