Added arbitrary length loop in PCG blueprint element

diff --git a/Engine/Plugins/Experimental/PCG/Source/PCG/Private/Elements/PCGExecuteBlueprint.cpp b/Engine/Plugins/Experimental/PCG/Source/PCG/Private/Elements/PCGExecuteBlueprint.cpp
--- a/Engine/Plugins/Experimental/PCG/Source/PCG/Private/Elements/PCGExecuteBlueprint.cpp
+++ b/Engine/Plugins/Experimental/PCG/Source/PCG/Private/Elements/PCGExecuteBlueprint.cpp
@@ -331,6 +331,30 @@ void UPCGBlueprintElement::LoopOnPointPairs(FPCGContext& InContext, const UPCGPo
 	});
 }
 
+void UPCGBlueprintElement::LoopNTimes(FPCGContext& InContext, int64 NumIterations, UPCGPointData*& OutData, const UPCGSpatialData* InA, const UPCGSpatialData* InB, UPCGPointData* OptionalOutData, const UObject* OptionalCustomObject) const
+{
+	if (NumIterations < 0)
+	{
+		UE_LOG(LogPCG, Error, TEXT("Invalid number of iterations in PCG blueprint element"));
+		return;
+	}
+
+	const UPCGSpatialData* Owner = (InA ? InA : InB);
+
+	OutData = OptionalOutData ? OptionalOutData : NewObject<UPCGPointData>(const_cast<UPCGSpatialData*>(Owner));
+	if (Owner)
+	{
+		OutData->InitializeFromData(Owner);
+	}
+
+	TArray<FPCGPoint>& OutPoints = OutData->GetMutablePoints();
+
+	FPCGAsync::AsyncPointProcessing(&InContext, NumIterations, OutPoints, [this, &InContext, InA, InB, OptionalCustomObject, OutData](int32 Index, FPCGPoint& OutPoint)
+	{
+		return IterationLoopBody(InContext, Index, InA, InB, OptionalCustomObject, OutPoint, OutData->Metadata);
+	});
+}
+
 bool FPCGExecuteBlueprintElement::IsCacheable(const UPCGSettings* InSettings) const
 {
 	if (const UPCGBlueprintSettings* BPSettings = Cast<const UPCGBlueprintSettings>(InSettings))
diff --git a/Engine/Plugins/Experimental/PCG/Source/PCG/Public/Elements/PCGExecuteBlueprint.h b/Engine/Plugins/Experimental/PCG/Source/PCG/Public/Elements/PCGExecuteBlueprint.h
--- a/Engine/Plugins/Experimental/PCG/Source/PCG/Public/Elements/PCGExecuteBlueprint.h
+++ b/Engine/Plugins/Experimental/PCG/Source/PCG/Public/Elements/PCGExecuteBlueprint.h
@@ -46,6 +46,9 @@ public:
 	UFUNCTION(BlueprintImplementableEvent, Category = Execution)
 	bool PointPairLoopBody(UPARAM(ref) FPCGContext& InContext, const UPCGPointData* InA, const UPCGPointData* InB, const FPCGPoint& InPointA, const FPCGPoint& InPointB, const UObject* OptionalCustomObject, FPCGPoint& OutPoint, UPCGMetadata* OutMetadata) const;
 
+	UFUNCTION(BlueprintImplementableEvent, Category = Execution)
+	bool IterationLoopBody(UPARAM(ref) FPCGContext& InContext, int64 Iteration, const UPCGSpatialData* InA, const UPCGSpatialData* InB, const UObject* OptionalCustomObject, FPCGPoint& OutPoint, UPCGMetadata* OutMetadata) const;
+
 	/** Calls the LoopBody function on all points */
 	UFUNCTION(BlueprintCallable, Category = Execution, meta = (HideSelfPin = "true"))
 	void LoopOnPoints(UPARAM(ref) FPCGContext& InContext, const UPCGPointData* InData, UPCGPointData*& OutData, UPCGPointData* OptionalOutData = nullptr, const UObject* OptionalCustomObject = nullptr) const;
@@ -53,6 +56,9 @@ public:
 	UFUNCTION(BlueprintCallable, Category = Execution, meta = (HideSelfPin = "true"))
 	void LoopOnPointPairs(UPARAM(ref) FPCGContext& InContext, const UPCGPointData* InA, const UPCGPointData* InB, UPCGPointData*& OutData, UPCGPointData* OptionalOutData = nullptr, const UObject* OptionalCustomObject = nullptr) const;
 
+	UFUNCTION(BlueprintCallable, Category = Execution, meta = (HideSelfPin = "true"))
+	void LoopNTimes(UPARAM(ref) FPCGContext& InContext, int64 NumIterations, UPCGPointData*& OutData, const UPCGSpatialData* InA = nullptr, const UPCGSpatialData* InB = nullptr, UPCGPointData* OptionalOutData = nullptr, const UObject* OptionalCustomObject = nullptr) const;
+
 	/** Called after object creation to setup the object callbacks */
 	void Initialize();
 
