PGC:
*Fixed an issue where new PCG graph nodes were not selected after placement.

diff --git a/Engine/Plugins/Experimental/PCG/Source/PCGEditor/Private/PCGEditorGraphSchemaActions.cpp b/Engine/Plugins/Experimental/PCG/Source/PCGEditor/Private/PCGEditorGraphSchemaActions.cpp
--- a/Engine/Plugins/Experimental/PCG/Source/PCGEditor/Private/PCGEditorGraphSchemaActions.cpp
+++ b/Engine/Plugins/Experimental/PCG/Source/PCGEditor/Private/PCGEditorGraphSchemaActions.cpp
@@ -27,7 +27,7 @@ UEdGraphNode* FPCGEditorGraphSchemaAction_NewNode::PerformAction(UEdGraph* Paren
 	UPCGNode* NewPCGNode = PCGGraph->AddNodeOfType(SettingsClass, DefaultNodeSettings);
 
 	FGraphNodeCreator<UPCGEditorGraphNode> NodeCreator(*EditorGraph);
-	UPCGEditorGraphNode* NewNode = NodeCreator.CreateNode(bSelectNewNode);
+	UPCGEditorGraphNode* NewNode = NodeCreator.CreateUserInvokedNode(bSelectNewNode);
 	NewNode->Construct(NewPCGNode, EPCGEditorGraphNodeType::Settings);
 	NewNode->NodePosX = Location.X;
 	NewNode->NodePosY = Location.Y;
