Duplicate instanced elements during execution to be able to use local variables in BP

diff --git a/Engine/Plugins/Experimental/PCG/Source/PCG/Private/Elements/PCGExecuteBlueprint.cpp b/Engine/Plugins/Experimental/PCG/Source/PCG/Private/Elements/PCGExecuteBlueprint.cpp
--- a/Engine/Plugins/Experimental/PCG/Source/PCG/Private/Elements/PCGExecuteBlueprint.cpp
+++ b/Engine/Plugins/Experimental/PCG/Source/PCG/Private/Elements/PCGExecuteBlueprint.cpp
@@ -47,7 +47,7 @@ void UPCGBlueprintElement::BeginDestroy()
 	Super::BeginDestroy();
 }
 
-void UPCGBlueprintElement::ExecuteWithContext_Implementation(FPCGContext& InContext, const FPCGDataCollection& Input, FPCGDataCollection& Output) const
+void UPCGBlueprintElement::ExecuteWithContext_Implementation(FPCGContext& InContext, const FPCGDataCollection& Input, FPCGDataCollection& Output)
 {
 	Execute(Input, Output);
 }
@@ -252,12 +252,13 @@ FPCGElementPtr UPCGBlueprintSettings::CreateElement() const
 	return MakeShared<FPCGExecuteBlueprintElement>();
 }
 
-bool FPCGExecuteBlueprintElement::ExecuteInternal(FPCGContext* Context) const
+bool FPCGExecuteBlueprintElement::ExecuteInternal(FPCGContext* InContext) const
 {
 	TRACE_CPUPROFILER_EVENT_SCOPE(FPCGExecuteBlueprintElement::Execute);
+	FPCGBlueprintExecutionContext* Context = static_cast<FPCGBlueprintExecutionContext*>(InContext);
 	const UPCGBlueprintSettings* Settings = Context->GetInputSettings<UPCGBlueprintSettings>();
 
-	if (Settings && Settings->BlueprintElementInstance)
+	if (Context && Context->BlueprintElementInstance)
 	{
 		// Log info on inputs
 		for(int32 InputIndex = 0; InputIndex < Context->InputData.TaggedData.Num(); ++InputIndex)
@@ -269,7 +270,7 @@ bool FPCGExecuteBlueprintElement::ExecuteInternal(FPCGContext* Context) const
 			}
 		}
 
-		Settings->BlueprintElementInstance->ExecuteWithContext(*Context, Context->InputData, Context->OutputData);
+		Context->BlueprintElementInstance->ExecuteWithContext(*Context, Context->InputData, Context->OutputData);
 
 		// Log info on outputs
 		for (int32 OutputIndex = 0; OutputIndex < Context->OutputData.TaggedData.Num(); ++OutputIndex)
@@ -355,6 +356,26 @@ void UPCGBlueprintElement::LoopNTimes(FPCGContext& InContext, int64 NumIteration
 	});
 }
 
+FPCGContext* FPCGExecuteBlueprintElement::Initialize(const FPCGDataCollection& InputData, UPCGComponent* SourceComponent, const UPCGNode* Node)
+{
+	FPCGBlueprintExecutionContext* Context = new FPCGBlueprintExecutionContext();
+	Context->InputData = InputData;
+	Context->SourceComponent = SourceComponent;
+	Context->Node = Node;
+
+	const UPCGBlueprintSettings* Settings = Context->GetInputSettings<UPCGBlueprintSettings>();
+	if (Settings && Settings->BlueprintElementInstance)
+	{
+		Context->BlueprintElementInstance = CastChecked<UPCGBlueprintElement>(StaticDuplicateObject(Settings->BlueprintElementInstance, Settings->BlueprintElementInstance, FName()));
+	}
+	else
+	{
+		Context->BlueprintElementInstance = nullptr;
+	}
+
+	return Context;
+}
+
 bool FPCGExecuteBlueprintElement::IsCacheable(const UPCGSettings* InSettings) const
 {
 	if (const UPCGBlueprintSettings* BPSettings = Cast<const UPCGBlueprintSettings>(InSettings))
diff --git a/Engine/Plugins/Experimental/PCG/Source/PCG/Private/Graph/PCGGraphExecutor.cpp b/Engine/Plugins/Experimental/PCG/Source/PCG/Private/Graph/PCGGraphExecutor.cpp
--- a/Engine/Plugins/Experimental/PCG/Source/PCG/Private/Graph/PCGGraphExecutor.cpp
+++ b/Engine/Plugins/Experimental/PCG/Source/PCG/Private/Graph/PCGGraphExecutor.cpp
@@ -191,8 +191,7 @@ void FPCGGraphExecutor::Execute()
 					continue;
 				}
 				
-				FPCGContext* Context = Element->Initialize(TaskInput, Task.SourceComponent);
-				Context->Node = Task.Node; // still needed?
+				FPCGContext* Context = Element->Initialize(TaskInput, Task.SourceComponent, Task.Node);
 				Context->TaskId = Task.NodeId;
 				Context->Cache = &GraphCache;
 
diff --git a/Engine/Plugins/Experimental/PCG/Source/PCG/Private/PCGElement.cpp b/Engine/Plugins/Experimental/PCG/Source/PCG/Private/PCGElement.cpp
--- a/Engine/Plugins/Experimental/PCG/Source/PCG/Private/PCGElement.cpp
+++ b/Engine/Plugins/Experimental/PCG/Source/PCG/Private/PCGElement.cpp
@@ -162,11 +162,12 @@ bool IPCGElement::Execute(FPCGContext* Context) const
 	}
 }
 
-FPCGContext* FSimplePCGElement::Initialize(const FPCGDataCollection& InputData, UPCGComponent* SourceComponent)
+FPCGContext* FSimplePCGElement::Initialize(const FPCGDataCollection& InputData, UPCGComponent* SourceComponent, const UPCGNode* Node)
 {
 	FPCGContext* Context = new FPCGContext();
 	Context->InputData = InputData;
 	Context->SourceComponent = SourceComponent;
+	Context->Node = Node;
 
 	return Context;
 }
\ No newline at end of file
diff --git a/Engine/Plugins/Experimental/PCG/Source/PCG/Private/PCGSubgraph.cpp b/Engine/Plugins/Experimental/PCG/Source/PCG/Private/PCGSubgraph.cpp
--- a/Engine/Plugins/Experimental/PCG/Source/PCG/Private/PCGSubgraph.cpp
+++ b/Engine/Plugins/Experimental/PCG/Source/PCG/Private/PCGSubgraph.cpp
@@ -173,11 +173,12 @@ void UPCGSubgraphNode::OnStructuralSettingsChanged(UPCGSettings* InSettings)
 }
 #endif // WITH_EDITOR
 
-FPCGContext* FPCGSubgraphElement::Initialize(const FPCGDataCollection& InputData, UPCGComponent* SourceComponent)
+FPCGContext* FPCGSubgraphElement::Initialize(const FPCGDataCollection& InputData, UPCGComponent* SourceComponent, const UPCGNode* Node)
 {
 	FPCGSubgraphContext* Context = new FPCGSubgraphContext();
 	Context->InputData = InputData;
 	Context->SourceComponent = SourceComponent;
+	Context->Node = Node;
 
 	return Context;
 }
diff --git a/Engine/Plugins/Experimental/PCG/Source/PCG/Public/Elements/PCGExecuteBlueprint.h b/Engine/Plugins/Experimental/PCG/Source/PCG/Public/Elements/PCGExecuteBlueprint.h
--- a/Engine/Plugins/Experimental/PCG/Source/PCG/Public/Elements/PCGExecuteBlueprint.h
+++ b/Engine/Plugins/Experimental/PCG/Source/PCG/Public/Elements/PCGExecuteBlueprint.h
@@ -35,10 +35,10 @@ public:
 	// ~End UObject interface
 
 	UFUNCTION(BlueprintNativeEvent, Category = Execution)
-	void ExecuteWithContext(UPARAM(ref)FPCGContext& InContext, const FPCGDataCollection& Input, FPCGDataCollection& Output) const;
+	void ExecuteWithContext(UPARAM(ref)FPCGContext& InContext, const FPCGDataCollection& Input, FPCGDataCollection& Output);
 
 	UFUNCTION(BlueprintImplementableEvent, BlueprintCallable, Category = Execution)
-	void Execute(const FPCGDataCollection& Input, FPCGDataCollection& Output) const;
+	void Execute(const FPCGDataCollection& Input, FPCGDataCollection& Output);
 
 	UFUNCTION(BlueprintImplementableEvent, Category = Execution)
 	bool PointLoopBody(UPARAM(ref) FPCGContext& InContext, const UPCGPointData* InData, const FPCGPoint& InPoint, const UObject* OptionalCustomObject, FPCGPoint& OutPoint, UPCGMetadata* OutMetadata) const;
@@ -145,9 +145,15 @@ protected:
 	void TeardownBlueprintElementEvent();
 };
 
-class FPCGExecuteBlueprintElement : public FSimplePCGElement
+struct FPCGBlueprintExecutionContext : public FPCGContext
+{
+	UPCGBlueprintElement* BlueprintElementInstance = nullptr;
+};
+
+class FPCGExecuteBlueprintElement : public IPCGElement
 {
 protected:
 	virtual bool ExecuteInternal(FPCGContext* Context) const override;
+	virtual FPCGContext* Initialize(const FPCGDataCollection& InputData, UPCGComponent* SourceComponent, const UPCGNode* Node) override;
 	virtual bool IsCacheable(const UPCGSettings* InSettings) const override;
 };
diff --git a/Engine/Plugins/Experimental/PCG/Source/PCG/Public/PCGElement.h b/Engine/Plugins/Experimental/PCG/Source/PCG/Public/PCGElement.h
--- a/Engine/Plugins/Experimental/PCG/Source/PCG/Public/PCGElement.h
+++ b/Engine/Plugins/Experimental/PCG/Source/PCG/Public/PCGElement.h
@@ -9,6 +9,7 @@
 class IPCGElement;
 class UPCGComponent;
 class UPCGSettings;
+class UPCGNode;
 
 typedef TSharedPtr<IPCGElement, ESPMode::ThreadSafe> FPCGElementPtr;
 
@@ -33,7 +34,7 @@ class PCG_API IPCGElement
 {
 public:
 	virtual ~IPCGElement() = default;
-	virtual FPCGContext* Initialize(const FPCGDataCollection& InputData, UPCGComponent* SourceComponent) = 0;
+	virtual FPCGContext* Initialize(const FPCGDataCollection& InputData, UPCGComponent* SourceComponent, const UPCGNode* Node) = 0;
 
 	bool Execute(FPCGContext* Context) const;
 
@@ -52,5 +53,5 @@ protected:
 class PCG_API FSimplePCGElement : public IPCGElement
 {
 public:
-	virtual FPCGContext* Initialize(const FPCGDataCollection& InputData, UPCGComponent* SourceComponent) override;
+	virtual FPCGContext* Initialize(const FPCGDataCollection& InputData, UPCGComponent* SourceComponent, const UPCGNode* Node) override;
 };
\ No newline at end of file
diff --git a/Engine/Plugins/Experimental/PCG/Source/PCG/Public/PCGSubgraph.h b/Engine/Plugins/Experimental/PCG/Source/PCG/Public/PCGSubgraph.h
--- a/Engine/Plugins/Experimental/PCG/Source/PCG/Public/PCGSubgraph.h
+++ b/Engine/Plugins/Experimental/PCG/Source/PCG/Public/PCGSubgraph.h
@@ -126,7 +126,7 @@ struct PCG_API FPCGSubgraphContext : public FPCGContext
 class PCG_API FPCGSubgraphElement : public IPCGElement
 {
 public:
-	virtual FPCGContext* Initialize(const FPCGDataCollection& InputData, UPCGComponent* SourceComponent) override;
+	virtual FPCGContext* Initialize(const FPCGDataCollection& InputData, UPCGComponent* SourceComponent, const UPCGNode* Node) override;
 
 protected:
 	virtual bool ExecuteInternal(FPCGContext* Context) const override;
