Exposed some more methods to BP:
- Get specific point by index in point data
- Added mul/div metadata operations
- Exposed metadata operations
- Exposed GetPointAtPosition method

diff --git a/Engine/Plugins/Experimental/PCG/Source/PCG/Private/Data/PCGDifferenceData.cpp b/Engine/Plugins/Experimental/PCG/Source/PCG/Private/Data/PCGDifferenceData.cpp
--- a/Engine/Plugins/Experimental/PCG/Source/PCG/Private/Data/PCGDifferenceData.cpp
+++ b/Engine/Plugins/Experimental/PCG/Source/PCG/Private/Data/PCGDifferenceData.cpp
@@ -172,7 +172,7 @@ bool UPCGDifferenceData::GetPointAtPosition(const FVector& InPosition, FPCGPoint
 		// Color?
 		if (OutMetadata && OutPoint.Density > 0)
 		{
-			OutMetadata->MergeAttributes(PointFromSource, PointFromDiff, OutPoint, EPCGMetadataOp::Sub);
+			OutMetadata->MergePointAttributes(PointFromSource, PointFromDiff, OutPoint, EPCGMetadataOp::Sub);
 		}
 
 		return OutPoint.Density > 0;
@@ -229,7 +229,7 @@ const UPCGPointData* UPCGDifferenceData::CreatePointData(FPCGContext* Context) c
 
 			if (Data->Metadata && OutPoint.Density > 0)
 			{
-				Data->Metadata->MergeAttributes(Point, SourcePointData->Metadata, PointFromDiff, Data->Metadata, OutPoint, EPCGMetadataOp::Sub);
+				Data->Metadata->MergePointAttributes(Point, SourcePointData->Metadata, PointFromDiff, Data->Metadata, OutPoint, EPCGMetadataOp::Sub);
 			}
 
 #if WITH_EDITORONLY_DATA
diff --git a/Engine/Plugins/Experimental/PCG/Source/PCG/Private/Data/PCGIntersectionData.cpp b/Engine/Plugins/Experimental/PCG/Source/PCG/Private/Data/PCGIntersectionData.cpp
--- a/Engine/Plugins/Experimental/PCG/Source/PCG/Private/Data/PCGIntersectionData.cpp
+++ b/Engine/Plugins/Experimental/PCG/Source/PCG/Private/Data/PCGIntersectionData.cpp
@@ -119,7 +119,7 @@ bool UPCGIntersectionData::GetPointAtPosition(const FVector& InPosition, FPCGPoi
 
 	if (OutMetadata)
 	{
-		OutMetadata->MergeAttributes(PointFromX, PointFromY, OutPoint, EPCGMetadataOp::Min);
+		OutMetadata->MergePointAttributes(PointFromX, PointFromY, OutPoint, EPCGMetadataOp::Min);
 	}
 
 	return true;
@@ -190,7 +190,7 @@ UPCGPointData* UPCGIntersectionData::CreateAndFilterPointData(FPCGContext* Conte
 
 		if (Data->Metadata)
 		{
-			Data->Metadata->MergeAttributes(Point, SourcePointData->Metadata, PointFromY, Data->Metadata, OutPoint, EPCGMetadataOp::Min);
+			Data->Metadata->MergePointAttributes(Point, SourcePointData->Metadata, PointFromY, Data->Metadata, OutPoint, EPCGMetadataOp::Min);
 		}
 
 		return true;
diff --git a/Engine/Plugins/Experimental/PCG/Source/PCG/Private/Data/PCGPointData.cpp b/Engine/Plugins/Experimental/PCG/Source/PCG/Private/Data/PCGPointData.cpp
--- a/Engine/Plugins/Experimental/PCG/Source/PCG/Private/Data/PCGPointData.cpp
+++ b/Engine/Plugins/Experimental/PCG/Source/PCG/Private/Data/PCGPointData.cpp
@@ -171,6 +171,19 @@ void UPCGPointData::InitializeFromActor(AActor* InActor)
 	Metadata = NewObject<UPCGMetadata>(this);
 }
 
+FPCGPoint UPCGPointData::GetPoint(int32 Index) const
+{
+	if (Points.IsValidIndex(Index))
+	{
+		return Points[Index];
+	}
+	else
+	{
+		UE_LOG(LogPCG, Error, TEXT("Invalid index in GetPoint call"));
+		return FPCGPoint();
+	}
+}
+
 const FPCGPoint* UPCGPointData::GetPointAtPosition(const FVector& InPosition) const
 {
 	if (bOctreeIsDirty)
@@ -318,7 +331,7 @@ bool UPCGPointData::GetPointAtPosition(const FVector& InPosition, FPCGPoint& Out
 	if (OutMetadata)
 	{
 		UPCGMetadataAccessorHelpers::InitializeMetadata(OutPoint, OutMetadata, *MaxContributor, Metadata);
-		OutMetadata->ResetWeightedAttributes(OutPoint);
+		OutMetadata->ResetPointWeightedAttributes(OutPoint);
 
 		for (const TPair<const FPCGPoint*, float> Contribution : Contributions)
 		{
@@ -326,7 +339,7 @@ bool UPCGPointData::GetPointAtPosition(const FVector& InPosition, FPCGPoint& Out
 			const float Weight = Contribution.Value / SumContributions;
 			const bool bIsMaxContributor = (Contribution.Key == MaxContributor);
 
-			OutMetadata->AccumulateWeightedAttributes(SourcePoint, Metadata, Weight, bIsMaxContributor, OutPoint);
+			OutMetadata->AccumulatePointWeightedAttributes(SourcePoint, Metadata, Weight, bIsMaxContributor, OutPoint);
 		}
 	}
 
diff --git a/Engine/Plugins/Experimental/PCG/Source/PCG/Private/Data/PCGProjectionData.cpp b/Engine/Plugins/Experimental/PCG/Source/PCG/Private/Data/PCGProjectionData.cpp
--- a/Engine/Plugins/Experimental/PCG/Source/PCG/Private/Data/PCGProjectionData.cpp
+++ b/Engine/Plugins/Experimental/PCG/Source/PCG/Private/Data/PCGProjectionData.cpp
@@ -112,7 +112,7 @@ bool UPCGProjectionData::GetPointAtPosition(const FVector& InPosition, FPCGPoint
 	if (OutMetadata)
 	{
 		//METADATA TODO Review op
-		OutMetadata->MergeAttributes(PointFromSource, PointFromTarget, OutPoint, EPCGMetadataOp::Max);
+		OutMetadata->MergePointAttributes(PointFromSource, PointFromTarget, OutPoint, EPCGMetadataOp::Max);
 	}
 
 	return true;
@@ -160,7 +160,7 @@ const UPCGPointData* UPCGProjectionData::CreatePointData(FPCGContext* Context) c
 		if (PointData->Metadata)
 		{
 			//METADATA TODO review op
-			PointData->Metadata->MergeAttributes(SourcePoint, SourcePointData->Metadata, PointFromTarget, PointData->Metadata, OutPoint, EPCGMetadataOp::Max);
+			PointData->Metadata->MergePointAttributes(SourcePoint, SourcePointData->Metadata, PointFromTarget, PointData->Metadata, OutPoint, EPCGMetadataOp::Max);
 		}
 
 		return true;
diff --git a/Engine/Plugins/Experimental/PCG/Source/PCG/Private/Data/PCGUnionData.cpp b/Engine/Plugins/Experimental/PCG/Source/PCG/Private/Data/PCGUnionData.cpp
--- a/Engine/Plugins/Experimental/PCG/Source/PCG/Private/Data/PCGUnionData.cpp
+++ b/Engine/Plugins/Experimental/PCG/Source/PCG/Private/Data/PCGUnionData.cpp
@@ -204,7 +204,7 @@ bool UPCGUnionData::GetPointAtPosition(const FVector& InPosition, FPCGPoint& Out
 				// Merge properties into OutPoint
 				if (OutMetadata)
 				{
-					OutMetadata->MergeAttributes(OutPoint, PointInData, OutPoint, EPCGMetadataOp::Max);
+					OutMetadata->MergePointAttributes(OutPoint, PointInData, OutPoint, EPCGMetadataOp::Max);
 				}
 			}
 			
@@ -272,7 +272,7 @@ const UPCGPointData* UPCGUnionData::CreatePointData(FPCGContext* Context) const
 
 				if (PointData->Metadata && DatumPointData->GetPoints().Num() > 0)
 				{
-					PointData->Metadata->SetAttributes(MakeArrayView(DatumPointData->GetPoints()), DatumPointData->Metadata, MakeArrayView(&TargetPoints[TargetPointIndex], DatumPointData->GetPoints().Num()));
+					PointData->Metadata->SetPointAttributes(MakeArrayView(DatumPointData->GetPoints()), DatumPointData->Metadata, MakeArrayView(&TargetPoints[TargetPointIndex], DatumPointData->GetPoints().Num()));
 				}
 			}
 
@@ -343,7 +343,7 @@ void UPCGUnionData::CreateSequentialPointData(FPCGContext* Context, UPCGPointDat
 			{
 				UPCGMetadataAccessorHelpers::InitializeMetadata(OutPoint, PointData->Metadata);
 				// Since we can't inherit from the parent point, we'll set the values directly here
-				PointData->Metadata->SetAttributes(Point, Data[DataIndex]->Metadata, OutPoint);
+				PointData->Metadata->SetPointAttributes(Point, Data[DataIndex]->Metadata, OutPoint);
 			}
 
 			if (DensityFunction == EPCGUnionDensityFunction::Binary && OutPoint.Density > 0)
@@ -370,7 +370,7 @@ void UPCGUnionData::CreateSequentialPointData(FPCGContext* Context, UPCGPointDat
 
 					if (PointData->Metadata)
 					{
-						PointData->Metadata->MergeAttributes(OutPoint, PointInData, OutPoint, EPCGMetadataOp::Max);
+						PointData->Metadata->MergePointAttributes(OutPoint, PointInData, OutPoint, EPCGMetadataOp::Max);
 					}
 					else if (OutPoint.Density >= 1.0f)
 					{
diff --git a/Engine/Plugins/Experimental/PCG/Source/PCG/Private/Metadata/PCGMetadata.cpp b/Engine/Plugins/Experimental/PCG/Source/PCG/Private/Metadata/PCGMetadata.cpp
--- a/Engine/Plugins/Experimental/PCG/Source/PCG/Private/Metadata/PCGMetadata.cpp
+++ b/Engine/Plugins/Experimental/PCG/Source/PCG/Private/Metadata/PCGMetadata.cpp
@@ -364,16 +364,11 @@ PCGMetadataEntryKey UPCGMetadata::GetParentKey(PCGMetadataEntryKey LocalItemKey)
 	}
 }
 
-void UPCGMetadata::MergeAttributes(const FPCGPoint& InPointA, const FPCGPoint& InPointB, FPCGPoint& OutPoint, EPCGMetadataOp Op)
+void UPCGMetadata::MergePointAttributes(const FPCGPoint& InPointA, const FPCGPoint& InPointB, FPCGPoint& OutPoint, EPCGMetadataOp Op)
 {
 	MergeAttributes(InPointA.MetadataEntry, this, InPointB.MetadataEntry, this, OutPoint.MetadataEntry, Op);
 }
 
-void UPCGMetadata::MergeAttributes(const FPCGPoint& InPointA, const UPCGMetadata* InMetadataA, const FPCGPoint& InPointB, const UPCGMetadata* InMetadataB, FPCGPoint& OutPoint, EPCGMetadataOp Op)
-{
-	MergeAttributes(InPointA.MetadataEntry, InMetadataA, InPointB.MetadataEntry, InMetadataB, OutPoint.MetadataEntry, Op);
-}
-
 void UPCGMetadata::MergeAttributes(PCGMetadataEntryKey InKeyA, const UPCGMetadata* InMetadataA, PCGMetadataEntryKey InKeyB, const UPCGMetadata* InMetadataB, PCGMetadataEntryKey& OutKey, EPCGMetadataOp Op)
 {
 	// Early out: nothing to do if both input metadata are null / points have no assigned metadata
@@ -417,11 +412,6 @@ void UPCGMetadata::MergeAttributes(PCGMetadataEntryKey InKeyA, const UPCGMetadat
 	AttributeLock.ReadUnlock();
 }
 
-void UPCGMetadata::ResetWeightedAttributes(FPCGPoint& OutPoint)
-{
-	ResetWeightedAttributes(OutPoint.MetadataEntry);
-}
-
 void UPCGMetadata::ResetWeightedAttributes(PCGMetadataEntryKey& OutKey)
 {
 	InitializeOnSet(OutKey);
@@ -439,11 +429,6 @@ void UPCGMetadata::ResetWeightedAttributes(PCGMetadataEntryKey& OutKey)
 	AttributeLock.ReadUnlock();
 }
 
-void UPCGMetadata::AccumulateWeightedAttributes(const FPCGPoint& InPoint, const UPCGMetadata* InMetadata, float Weight, bool bSetNonInterpolableAttributes, FPCGPoint& OutPoint)
-{
-	AccumulateWeightedAttributes(InPoint.MetadataEntry, InMetadata, Weight, bSetNonInterpolableAttributes, OutPoint.MetadataEntry);
-}
-
 void UPCGMetadata::AccumulateWeightedAttributes(PCGMetadataEntryKey InKey, const UPCGMetadata* InMetadata, float Weight, bool bSetNonInterpolableAttributes, PCGMetadataEntryKey& OutKey)
 {
 	if (!InMetadata)
@@ -482,11 +467,6 @@ void UPCGMetadata::AccumulateWeightedAttributes(PCGMetadataEntryKey InKey, const
 	AttributeLock.ReadUnlock();
 }
 
-void UPCGMetadata::SetAttributes(const FPCGPoint& InPoint, const UPCGMetadata* InMetadata, FPCGPoint& OutPoint)
-{
-	SetAttributes(InPoint.MetadataEntry, InMetadata, OutPoint.MetadataEntry);
-}
-
 void UPCGMetadata::SetAttributes(PCGMetadataEntryKey InKey, const UPCGMetadata* InMetadata, PCGMetadataEntryKey& OutKey)
 {
 	if (!InMetadata)
@@ -520,7 +500,7 @@ void UPCGMetadata::SetAttributes(PCGMetadataEntryKey InKey, const UPCGMetadata*
 	AttributeLock.ReadUnlock();
 }
 
-void UPCGMetadata::SetAttributes(const TArrayView<const FPCGPoint>& InPoints, const UPCGMetadata* InMetadata, const TArrayView<FPCGPoint>& OutPoints)
+void UPCGMetadata::SetPointAttributes(const TArrayView<const FPCGPoint>& InPoints, const UPCGMetadata* InMetadata, const TArrayView<FPCGPoint>& OutPoints)
 {
 	if (!InMetadata)
 	{
@@ -597,4 +577,42 @@ void UPCGMetadata::SetAttributes(const TArrayView<PCGMetadataEntryKey>& InKeys,
 		}
 	}
 	AttributeLock.ReadUnlock();
-}
\ No newline at end of file
+}
+
+void UPCGMetadata::MergeAttributesByKey(int64 KeyA, const UPCGMetadata* MetadataA, int64 KeyB, const UPCGMetadata* MetadataB, int64 TargetKey, EPCGMetadataOp Op, int64& OutKey)
+{
+	OutKey = TargetKey;
+	MergeAttributes(KeyA, MetadataA, KeyB, MetadataB, OutKey, Op);
+}
+
+void UPCGMetadata::ResetWeightedAttributesByKey(int64 TargetKey, int64& OutKey)
+{
+	OutKey = TargetKey;
+	ResetWeightedAttributes(OutKey);
+}
+
+void UPCGMetadata::AccumulateWeightedAttributesByKey(PCGMetadataEntryKey Key, const UPCGMetadata* Metadata, float Weight, bool bSetNonInterpolableAttributes, int64 TargetKey, int64& OutKey)
+{
+	OutKey = TargetKey;
+	AccumulateWeightedAttributes(Key, Metadata, Weight, bSetNonInterpolableAttributes, OutKey);
+}
+
+void UPCGMetadata::MergePointAttributes(const FPCGPoint& PointA, const UPCGMetadata* MetadataA, const FPCGPoint& PointB, const UPCGMetadata* MetadataB, UPARAM(ref) FPCGPoint& TargetPoint, EPCGMetadataOp Op)
+{
+	MergeAttributes(PointA.MetadataEntry, MetadataA, PointB.MetadataEntry, MetadataB, TargetPoint.MetadataEntry, Op);
+}
+
+void UPCGMetadata::SetPointAttributes(const FPCGPoint& Point, const UPCGMetadata* Metadata, FPCGPoint& OutPoint)
+{
+	SetAttributes(Point.MetadataEntry, Metadata, OutPoint.MetadataEntry);
+}
+
+void UPCGMetadata::ResetPointWeightedAttributes(FPCGPoint& OutPoint)
+{
+	ResetWeightedAttributes(OutPoint.MetadataEntry);
+}
+
+void UPCGMetadata::AccumulatePointWeightedAttributes(const FPCGPoint& InPoint, const UPCGMetadata* InMetadata, float Weight, bool bSetNonInterpolableAttributes, FPCGPoint& OutPoint)
+{
+	AccumulateWeightedAttributes(InPoint.MetadataEntry, InMetadata, Weight, bSetNonInterpolableAttributes, OutPoint.MetadataEntry);
+}
diff --git a/Engine/Plugins/Experimental/PCG/Source/PCG/Private/Metadata/PCGMetadataAccessor.cpp b/Engine/Plugins/Experimental/PCG/Source/PCG/Private/Metadata/PCGMetadataAccessor.cpp
--- a/Engine/Plugins/Experimental/PCG/Source/PCG/Private/Metadata/PCGMetadataAccessor.cpp
+++ b/Engine/Plugins/Experimental/PCG/Source/PCG/Private/Metadata/PCGMetadataAccessor.cpp
@@ -162,7 +162,7 @@ void UPCGMetadataAccessorHelpers::CopyPoint(const FPCGPoint& InPoint, FPCGPoint&
 		// Otherwise, we will assume that the point is parented
 		if (InMetadata)
 		{
-			OutMetadata->SetAttributes(InPoint, InMetadata, OutPoint);
+			OutMetadata->SetPointAttributes(InPoint, InMetadata, OutPoint);
 		}
 	}
 	else
diff --git a/Engine/Plugins/Experimental/PCG/Source/PCG/Public/Data/PCGPointData.h b/Engine/Plugins/Experimental/PCG/Source/PCG/Public/Data/PCGPointData.h
--- a/Engine/Plugins/Experimental/PCG/Source/PCG/Public/Data/PCGPointData.h
+++ b/Engine/Plugins/Experimental/PCG/Source/PCG/Public/Data/PCGPointData.h
@@ -73,6 +73,9 @@ public:
 	UFUNCTION(BlueprintCallable, Category = SpatialData)
 	const TArray<FPCGPoint>& GetPoints() const { return Points; }
 
+	UFUNCTION(BlueprintCallable, Category = SpatialData)
+	FPCGPoint GetPoint(int32 Index) const;
+
 	UFUNCTION(BlueprintCallable, Category = SpatialData)
 	void SetPoints(const TArray<FPCGPoint>& InPoints);
 	
diff --git a/Engine/Plugins/Experimental/PCG/Source/PCG/Public/Data/PCGSpatialData.h b/Engine/Plugins/Experimental/PCG/Source/PCG/Public/Data/PCGSpatialData.h
--- a/Engine/Plugins/Experimental/PCG/Source/PCG/Public/Data/PCGSpatialData.h
+++ b/Engine/Plugins/Experimental/PCG/Source/PCG/Public/Data/PCGSpatialData.h
@@ -72,6 +72,7 @@ public:
 	UFUNCTION(BlueprintCallable, Category = SpatialData)
 	virtual FPCGPoint TransformPoint(const FPCGPoint& InPoint) const;
 
+	UFUNCTION(BlueprintCallable, Category = SpatialData)
 	virtual bool GetPointAtPosition(const FVector& InPosition, FPCGPoint& OutPoint, UPCGMetadata* OutMetadata) const
 	{
 		OutPoint.Transform = FTransform(InPosition);
diff --git a/Engine/Plugins/Experimental/PCG/Source/PCG/Public/Metadata/PCGMetadata.h b/Engine/Plugins/Experimental/PCG/Source/PCG/Public/Metadata/PCGMetadata.h
--- a/Engine/Plugins/Experimental/PCG/Source/PCG/Public/Metadata/PCGMetadata.h
+++ b/Engine/Plugins/Experimental/PCG/Source/PCG/Public/Metadata/PCGMetadata.h
@@ -104,14 +104,30 @@ public:
 	void SetAttributes(const TArrayView<PCGMetadataEntryKey>& InKeys, const UPCGMetadata* InMetadata, const TArrayView<PCGMetadataEntryKey>& OutKeys);
 
 	/** Attributes operations - shorthand for points */
-	void MergeAttributes(const FPCGPoint& InPointA, const FPCGPoint& InPointB, FPCGPoint& OutPoint, EPCGMetadataOp Op);
-	void MergeAttributes(const FPCGPoint& InPointA, const UPCGMetadata* InMetadataA, const FPCGPoint& InPointB, const UPCGMetadata* InMetadataB, FPCGPoint& OutPoint, EPCGMetadataOp Op);
+	void MergePointAttributes(const FPCGPoint& InPointA, const FPCGPoint& InPointB, FPCGPoint& OutPoint, EPCGMetadataOp Op);
+	void SetPointAttributes(const TArrayView<const FPCGPoint>& InPoints, const UPCGMetadata* InMetadata, const TArrayView<FPCGPoint>& OutPoints);
 
-	void ResetWeightedAttributes(FPCGPoint& OutPoint);
-	void AccumulateWeightedAttributes(const FPCGPoint& InPoint, const UPCGMetadata* InMetadata, float Weight, bool bSetNonInterpolableAttributes, FPCGPoint& OutPoint);
+	/** Blueprint-friend versions */
+	UFUNCTION(BlueprintCallable, Category = "PCG|Metadata")
+	void MergeAttributesByKey(int64 KeyA, const UPCGMetadata* MetadataA, int64 KeyB, const UPCGMetadata* MetadataB, int64 TargetKey, EPCGMetadataOp Op, int64& OutKey);
+
+	UFUNCTION(BlueprintCallable, Category = "PCG|Metadata")
+	void ResetWeightedAttributesByKey(int64 TargetKey, int64& OutKey);
+
+	UFUNCTION(BlueprintCallable, Category = "PCG|Metadata")
+	void AccumulateWeightedAttributesByKey(int64 Key, const UPCGMetadata* Metadata, float Weight, bool bSetNonInterpolableAttributes, int64 TargetKey, int64& OutKey);
+
+	UFUNCTION(BlueprintCallable, Category = "PCG|Metadata")
+	void MergePointAttributes(const FPCGPoint& PointA, const UPCGMetadata* MetadataA, const FPCGPoint& PointB, const UPCGMetadata* MetadataB, UPARAM(ref) FPCGPoint& TargetPoint, EPCGMetadataOp Op);
+
+	UFUNCTION(BlueprintCallable, Category = "PCG|Metadata")
+	void SetPointAttributes(const FPCGPoint& Point, const UPCGMetadata* Metadata, UPARAM(ref) FPCGPoint& OutPoint);
 
-	void SetAttributes(const FPCGPoint& InPoint, const UPCGMetadata* InMetadata, FPCGPoint& OutPoint);
-	void SetAttributes(const TArrayView<const FPCGPoint>& InPoints, const UPCGMetadata* InMetadata, const TArrayView<FPCGPoint>& OutPoints);
+	UFUNCTION(BlueprintCallable, Category = "PCG|Metadata")
+	void ResetPointWeightedAttributes(FPCGPoint& OutPoint);
+
+	UFUNCTION(BlueprintCallable, Category = "PCG|Metadata")
+	void AccumulatePointWeightedAttributes(const FPCGPoint& InPoint, const UPCGMetadata* InMetadata, float Weight, bool bSetNonInterpolableAttributes, UPARAM(ref) FPCGPoint& OutPoint);
 
 protected:
 	template<typename T>
diff --git a/Engine/Plugins/Experimental/PCG/Source/PCG/Public/Metadata/PCGMetadataAttributeTpl.h b/Engine/Plugins/Experimental/PCG/Source/PCG/Public/Metadata/PCGMetadataAttributeTpl.h
--- a/Engine/Plugins/Experimental/PCG/Source/PCG/Public/Metadata/PCGMetadataAttributeTpl.h
+++ b/Engine/Plugins/Experimental/PCG/Source/PCG/Public/Metadata/PCGMetadataAttributeTpl.h
@@ -260,6 +260,23 @@ protected:
 		return false;
 	}
 
+	template<typename IT = T, typename TEnableIf<PCG::Private::MetadataTraits<IT>::CanSubAdd>::Type* = nullptr>
+	bool SetAdd(PCGMetadataEntryKey ItemKey, const FPCGMetadataAttributeBase* InAttributeA, PCGMetadataEntryKey InEntryKeyA, const FPCGMetadataAttributeBase* InAttributeB, PCGMetadataEntryKey InEntryKeyB)
+	{
+		SetValue(ItemKey,
+			PCG::Private::MetadataTraits<IT>::Add(
+				static_cast<const FPCGMetadataAttribute<T>*>(InAttributeA)->GetValueFromItemKey(InEntryKeyA),
+				static_cast<const FPCGMetadataAttribute<T>*>(InAttributeB)->GetValueFromItemKey(InEntryKeyB)));
+
+		return true;
+	}
+
+	template<typename IT = T, typename TEnableIf<!PCG::Private::MetadataTraits<IT>::CanSubAdd>::Type* = nullptr>
+	bool SetAdd(PCGMetadataEntryKey ItemKey, const FPCGMetadataAttributeBase* InAttributeA, PCGMetadataEntryKey InEntryKeyA, const FPCGMetadataAttributeBase* InAttributeB, PCGMetadataEntryKey InEntryKeyB)
+	{
+		return false;
+	}
+
 	template<typename IT = T, typename TEnableIf<PCG::Private::MetadataTraits<IT>::CanSubAdd>::Type* = nullptr>
 	bool SetSub(PCGMetadataEntryKey ItemKey, const FPCGMetadataAttributeBase* InAttributeA, PCGMetadataEntryKey InEntryKeyA, const FPCGMetadataAttributeBase* InAttributeB, PCGMetadataEntryKey InEntryKeyB)
 	{
@@ -277,19 +294,36 @@ protected:
 		return false;
 	}
 
-	template<typename IT = T, typename TEnableIf<PCG::Private::MetadataTraits<IT>::CanSubAdd>::Type* = nullptr>
-	bool SetAdd(PCGMetadataEntryKey ItemKey, const FPCGMetadataAttributeBase* InAttributeA, PCGMetadataEntryKey InEntryKeyA, const FPCGMetadataAttributeBase* InAttributeB, PCGMetadataEntryKey InEntryKeyB)
+	template<typename IT = T, typename TEnableIf<PCG::Private::MetadataTraits<IT>::CanMulDiv>::Type* = nullptr>
+	bool SetMul(PCGMetadataEntryKey ItemKey, const FPCGMetadataAttributeBase* InAttributeA, PCGMetadataEntryKey InEntryKeyA, const FPCGMetadataAttributeBase* InAttributeB, PCGMetadataEntryKey InEntryKeyB)
 	{
 		SetValue(ItemKey,
-			PCG::Private::MetadataTraits<IT>::Add(
+			PCG::Private::MetadataTraits<IT>::Mul(
 				static_cast<const FPCGMetadataAttribute<T>*>(InAttributeA)->GetValueFromItemKey(InEntryKeyA),
 				static_cast<const FPCGMetadataAttribute<T>*>(InAttributeB)->GetValueFromItemKey(InEntryKeyB)));
 
 		return true;
 	}
 
-	template<typename IT = T, typename TEnableIf<!PCG::Private::MetadataTraits<IT>::CanSubAdd>::Type* = nullptr>
-	bool SetAdd(PCGMetadataEntryKey ItemKey, const FPCGMetadataAttributeBase* InAttributeA, PCGMetadataEntryKey InEntryKeyA, const FPCGMetadataAttributeBase* InAttributeB, PCGMetadataEntryKey InEntryKeyB)
+	template<typename IT = T, typename TEnableIf<!PCG::Private::MetadataTraits<IT>::CanMulDiv>::Type* = nullptr>
+	bool SetMul(PCGMetadataEntryKey ItemKey, const FPCGMetadataAttributeBase* InAttributeA, PCGMetadataEntryKey InEntryKeyA, const FPCGMetadataAttributeBase* InAttributeB, PCGMetadataEntryKey InEntryKeyB)
+	{
+		return false;
+	}
+
+	template<typename IT = T, typename TEnableIf<PCG::Private::MetadataTraits<IT>::CanMulDiv>::Type* = nullptr>
+	bool SetDiv(PCGMetadataEntryKey ItemKey, const FPCGMetadataAttributeBase* InAttributeA, PCGMetadataEntryKey InEntryKeyA, const FPCGMetadataAttributeBase* InAttributeB, PCGMetadataEntryKey InEntryKeyB)
+	{
+		SetValue(ItemKey,
+			PCG::Private::MetadataTraits<IT>::Div(
+				static_cast<const FPCGMetadataAttribute<T>*>(InAttributeA)->GetValueFromItemKey(InEntryKeyA),
+				static_cast<const FPCGMetadataAttribute<T>*>(InAttributeB)->GetValueFromItemKey(InEntryKeyB)));
+
+		return true;
+	}
+
+	template<typename IT = T, typename TEnableIf<!PCG::Private::MetadataTraits<IT>::CanMulDiv>::Type* = nullptr>
+	bool SetDiv(PCGMetadataEntryKey ItemKey, const FPCGMetadataAttributeBase* InAttributeA, PCGMetadataEntryKey InEntryKeyA, const FPCGMetadataAttributeBase* InAttributeB, PCGMetadataEntryKey InEntryKeyB)
 	{
 		return false;
 	}
diff --git a/Engine/Plugins/Experimental/PCG/Source/PCG/Public/Metadata/PCGMetadataAttributeTraits.h b/Engine/Plugins/Experimental/PCG/Source/PCG/Public/Metadata/PCGMetadataAttributeTraits.h
--- a/Engine/Plugins/Experimental/PCG/Source/PCG/Public/Metadata/PCGMetadataAttributeTraits.h
+++ b/Engine/Plugins/Experimental/PCG/Source/PCG/Public/Metadata/PCGMetadataAttributeTraits.h
@@ -36,6 +36,7 @@ namespace PCG
 			enum { CompressData = false };
 			enum { CanMinMax = true };
 			enum { CanSubAdd = true };
+			enum { CanMulDiv = true };
 			enum { CanInterpolate = true };
 			
 			static T Min(const T& A, const T& B)
@@ -58,6 +59,16 @@ namespace PCG
 				return A - B;
 			}
 
+			static T Mul(const T& A, const T& B)
+			{
+				return A * B;
+			}
+
+			static T Div(const T& A, const T& B)
+			{
+				return A / B;
+			}
+
 			static T WeightedSum(const T& A, const T& B, float Weight)
 			{
 				return A + B * Weight;
@@ -75,6 +86,7 @@ namespace PCG
 			enum { CompressData = false };
 			enum { CanMinMax = true };
 			enum { CanSubAdd = true };
+			enum { CanMulDiv = false };
 			enum { CanInterpolate = false };
 
 			static bool Min(const bool& A, const bool& B)
@@ -105,6 +117,7 @@ namespace PCG
 			enum { CompressData = false };
 			enum { CanMinMax = true };
 			enum { CanSubAdd = true };
+			enum { CanMulDiv = true };
 			enum { CanInterpolate = true};
 
 			static FVector Min(const FVector& A, const FVector& B) 
@@ -127,6 +140,16 @@ namespace PCG
 				return A - B;
 			}
 
+			static FVector Mul(const FVector& A, const FVector& B)
+			{
+				return A * B;
+			}
+
+			static FVector Div(const FVector& A, const FVector& B)
+			{
+				return A / B;
+			}
+
 			static FVector WeightedSum(const FVector& A, const FVector& B, float Weight)
 			{
 				return A + B * Weight;
@@ -144,6 +167,7 @@ namespace PCG
 			enum { CompressData = false };
 			enum { CanMinMax = true };
 			enum { CanSubAdd = true };
+			enum { CanMulDiv = true };
 			enum { CanInterpolate = true };
 
 			static FVector4 Min(const FVector4& A, const FVector4& B)
@@ -166,6 +190,16 @@ namespace PCG
 				return A - B;
 			}
 
+			static FVector4 Mul(const FVector& A, const FVector4& B)
+			{
+				return A * B;
+			}
+
+			static FVector4 Div(const FVector& A, const FVector4& B)
+			{
+				return A / B;
+			}
+
 			static FVector4 WeightedSum(const FVector4& A, const FVector4& B, float Weight)
 			{
 				return A + B * Weight;
@@ -184,6 +218,7 @@ namespace PCG
 			enum { CompressData = false };
 			enum { CanMinMax = false };
 			enum { CanSubAdd = true };
+			enum { CanMulDiv = true };
 			enum { CanInterpolate = true };
 
 			static FQuat Add(const FQuat& A, const FQuat& B)
@@ -196,6 +231,16 @@ namespace PCG
 				return A * B.Inverse();
 			}
 
+			static FQuat Mul(const FQuat& A, const FQuat& B)
+			{
+				return A * B;
+			}
+
+			static FQuat Div(const FQuat& A, const FQuat& B)
+			{
+				return A * B.Inverse();
+			}
+
 			static FQuat WeightedSum(const FQuat& A, const FQuat& B, float Weight)
 			{
 				// WARNING: the quaternion won't be normalized
@@ -220,6 +265,7 @@ namespace PCG
 			enum { CompressData = false };
 			enum { CanMinMax = false };
 			enum { CanSubAdd = true };
+			enum { CanMulDiv = true };
 			enum { CanInterpolate = true };
 
 			static FTransform Add(const FTransform& A, const FTransform& B)
@@ -232,6 +278,16 @@ namespace PCG
 				return A * B.Inverse();
 			}
 
+			static FTransform Mul(const FTransform& A, const FTransform& B)
+			{
+				return A * B;
+			}
+
+			static FTransform Div(const FTransform& A, const FTransform& B)
+			{
+				return A * B.Inverse();
+			}
+
 			static FQuat WeightedQuatSum(const FQuat& Q, const FQuat& V, float Weight)
 			{
 				FQuat BlendQuat = V * Weight;
@@ -264,6 +320,7 @@ namespace PCG
 			enum { CompressData = true };
 			enum { CanMinMax = false };
 			enum { CanSubAdd = false };
+			enum { CanMulAdd = false };
 			enum { CanInterpolate = false };
 		};
 	}
