Fixed error due to assumed default parameter value when there are no inputs in the N iteration loop

#rnx

diff --git a/Engine/Plugins/Experimental/PCG/Source/PCG/Private/Elements/PCGExecuteBlueprint.cpp b/Engine/Plugins/Experimental/PCG/Source/PCG/Private/Elements/PCGExecuteBlueprint.cpp
--- a/Engine/Plugins/Experimental/PCG/Source/PCG/Private/Elements/PCGExecuteBlueprint.cpp
+++ b/Engine/Plugins/Experimental/PCG/Source/PCG/Private/Elements/PCGExecuteBlueprint.cpp
@@ -341,8 +341,13 @@ void UPCGBlueprintElement::LoopNTimes(FPCGContext& InContext, int64 NumIteration
 	}
 
 	const UPCGSpatialData* Owner = (InA ? InA : InB);
+	OutData = OptionalOutData;
+
+	if (!OutData)
+	{
+		OutData = Owner ? NewObject<UPCGPointData>(const_cast<UPCGSpatialData*>(Owner)) : NewObject<UPCGPointData>();
+	}
 
-	OutData = OptionalOutData ? OptionalOutData : NewObject<UPCGPointData>(const_cast<UPCGSpatialData*>(Owner));
 	if (Owner)
 	{
 		OutData->InitializeFromData(Owner);
