Added missing metadata ops, fixed a typo and added a missing Set attribute call

diff --git a/Engine/Plugins/Experimental/PCG/Source/PCG/Private/Metadata/PCGMetadata.cpp b/Engine/Plugins/Experimental/PCG/Source/PCG/Private/Metadata/PCGMetadata.cpp
--- a/Engine/Plugins/Experimental/PCG/Source/PCG/Private/Metadata/PCGMetadata.cpp
+++ b/Engine/Plugins/Experimental/PCG/Source/PCG/Private/Metadata/PCGMetadata.cpp
@@ -585,6 +585,12 @@ void UPCGMetadata::MergeAttributesByKey(int64 KeyA, const UPCGMetadata* Metadata
 	MergeAttributes(KeyA, MetadataA, KeyB, MetadataB, OutKey, Op);
 }
 
+void UPCGMetadata::SetAttributesByKey(int64 Key, const UPCGMetadata* Metadata, int64 TargetKey, int64& OutKey)
+{
+	OutKey = TargetKey;
+	SetAttributes(Key, Metadata, OutKey);
+}
+
 void UPCGMetadata::ResetWeightedAttributesByKey(int64 TargetKey, int64& OutKey)
 {
 	OutKey = TargetKey;
diff --git a/Engine/Plugins/Experimental/PCG/Source/PCG/Public/Metadata/PCGMetadata.h b/Engine/Plugins/Experimental/PCG/Source/PCG/Public/Metadata/PCGMetadata.h
--- a/Engine/Plugins/Experimental/PCG/Source/PCG/Public/Metadata/PCGMetadata.h
+++ b/Engine/Plugins/Experimental/PCG/Source/PCG/Public/Metadata/PCGMetadata.h
@@ -16,7 +16,9 @@ enum class EPCGMetadataOp : uint8
 	Min,
 	Max,
 	Sub,
-	Add
+	Add,
+	Mul,
+	Div
 };
 
 UCLASS(BlueprintType)
@@ -117,6 +119,9 @@ public:
 	UFUNCTION(BlueprintCallable, Category = "PCG|Metadata")
 	void AccumulateWeightedAttributesByKey(int64 Key, const UPCGMetadata* Metadata, float Weight, bool bSetNonInterpolableAttributes, int64 TargetKey, int64& OutKey);
 
+	UFUNCTION(BlueprintCallable, Category = "PCG|Metadata")
+	void SetAttributesByKey(int64 Key, const UPCGMetadata* InMetadata, int64 TargetKey, int64& OutKey);
+
 	UFUNCTION(BlueprintCallable, Category = "PCG|Metadata")
 	void MergePointAttributes(const FPCGPoint& PointA, const UPCGMetadata* MetadataA, const FPCGPoint& PointB, const UPCGMetadata* MetadataB, UPARAM(ref) FPCGPoint& TargetPoint, EPCGMetadataOp Op);
 
diff --git a/Engine/Plugins/Experimental/PCG/Source/PCG/Public/Metadata/PCGMetadataAttributeTpl.h b/Engine/Plugins/Experimental/PCG/Source/PCG/Public/Metadata/PCGMetadataAttributeTpl.h
--- a/Engine/Plugins/Experimental/PCG/Source/PCG/Public/Metadata/PCGMetadataAttributeTpl.h
+++ b/Engine/Plugins/Experimental/PCG/Source/PCG/Public/Metadata/PCGMetadataAttributeTpl.h
@@ -96,6 +96,14 @@ public:
 			{
 				bAppliedValue = SetAdd(ItemKey, InAttributeA, InEntryKeyA, InAttributeB, InEntryKeyB);
 			}
+			else if (Op == EPCGMetadataOp::Mul)
+			{
+				bAppliedValue = SetMul(ItemKey, InAttributeA, InEntryKeyA, InAttributeB, InEntryKeyB);
+			}
+			else if (Op == EPCGMetadataOp::Div)
+			{
+				bAppliedValue = SetDiv(ItemKey, InAttributeA, InEntryKeyA, InAttributeB, InEntryKeyB);
+			}
 		}
 		else if (InAttributeA && InAttributeB && HasNonDefaultValue(ItemKey))
 		{
diff --git a/Engine/Plugins/Experimental/PCG/Source/PCG/Public/Metadata/PCGMetadataAttributeTraits.h b/Engine/Plugins/Experimental/PCG/Source/PCG/Public/Metadata/PCGMetadataAttributeTraits.h
--- a/Engine/Plugins/Experimental/PCG/Source/PCG/Public/Metadata/PCGMetadataAttributeTraits.h
+++ b/Engine/Plugins/Experimental/PCG/Source/PCG/Public/Metadata/PCGMetadataAttributeTraits.h
@@ -320,7 +320,7 @@ namespace PCG
 			enum { CompressData = true };
 			enum { CanMinMax = false };
 			enum { CanSubAdd = false };
-			enum { CanMulAdd = false };
+			enum { CanMulDiv = false };
 			enum { CanInterpolate = false };
 		};
 	}
