Added errors/warnings when we have no target actor

#rnx

diff --git a/Engine/Plugins/Experimental/PCG/Source/PCG/Private/Elements/PCGDebugElement.cpp b/Engine/Plugins/Experimental/PCG/Source/PCG/Private/Elements/PCGDebugElement.cpp
--- a/Engine/Plugins/Experimental/PCG/Source/PCG/Private/Elements/PCGDebugElement.cpp
+++ b/Engine/Plugins/Experimental/PCG/Source/PCG/Private/Elements/PCGDebugElement.cpp
@@ -29,6 +29,7 @@ namespace PCGDebugElement
 
 		if (!Mesh)
 		{
+			UE_LOG(LogPCG, Error, TEXT("Debug display was unable to load mesh %s"), *DebugSettings.PointMesh.ToString());
 			return;
 		}
 
@@ -57,6 +58,7 @@ namespace PCGDebugElement
 			if (!TargetActor)
 			{
 				// No target actor
+				UE_LOG(LogPCG, Error, TEXT("Debug display cannot show data that have no target actor"));
 				continue;
 			}
 
