Added pcg blueprint elements variable override by params
Removed custom object from PCG blueprints

diff --git a/Engine/Plugins/Experimental/PCG/Source/PCG/Private/Elements/PCGExecuteBlueprint.cpp b/Engine/Plugins/Experimental/PCG/Source/PCG/Private/Elements/PCGExecuteBlueprint.cpp
--- a/Engine/Plugins/Experimental/PCG/Source/PCG/Private/Elements/PCGExecuteBlueprint.cpp
+++ b/Engine/Plugins/Experimental/PCG/Source/PCG/Private/Elements/PCGExecuteBlueprint.cpp
@@ -5,6 +5,7 @@
 #include "Data/PCGPointData.h"
 #include "Data/PCGSpatialData.h"
 #include "Helpers/PCGAsync.h"
+#include "Helpers/PCGSettingsHelpers.h"
 
 #if WITH_EDITOR
 #include "Engine/World.h"
@@ -270,6 +271,25 @@ bool FPCGExecuteBlueprintElement::ExecuteInternal(FPCGContext* InContext) const
 			}
 		}
 
+		/** Apply params overrides to variables if any */
+		if (UPCGParams* Params = Context->InputData.GetParams())
+		{
+			UClass* BPClass = Context->BlueprintElementInstance->GetClass();
+
+			for (TFieldIterator<FProperty> PropertyIt(BPClass); PropertyIt; ++PropertyIt)
+			{
+				FProperty* Property = *PropertyIt;
+				if (Property->IsNative())
+				{
+					continue;
+				}
+
+				// Apply params if any
+				PCGSettingsHelpers::SetValue(Params, Context->BlueprintElementInstance, Property);
+			}
+		}
+
+		/** Finally, execute the actual blueprint */
 		Context->BlueprintElementInstance->ExecuteWithContext(*Context, Context->InputData, Context->OutputData);
 
 		// Log info on outputs
@@ -291,7 +311,7 @@ bool FPCGExecuteBlueprintElement::ExecuteInternal(FPCGContext* InContext) const
 	return true;
 }
 
-void UPCGBlueprintElement::LoopOnPoints(FPCGContext& InContext, const UPCGPointData* InData, UPCGPointData*& OutData, UPCGPointData* OptionalOutData, const UObject* OptionalCustomObject) const
+void UPCGBlueprintElement::LoopOnPoints(FPCGContext& InContext, const UPCGPointData* InData, UPCGPointData*& OutData, UPCGPointData* OptionalOutData) const
 {
 	if (!InData)
 	{
@@ -305,13 +325,13 @@ void UPCGBlueprintElement::LoopOnPoints(FPCGContext& InContext, const UPCGPointD
 	const TArray<FPCGPoint>& InPoints = InData->GetPoints();
 	TArray<FPCGPoint>& OutPoints = OutData->GetMutablePoints();
 
-	FPCGAsync::AsyncPointProcessing(&InContext, InPoints.Num(), OutPoints, [this, &InContext, InData, OutData, OptionalCustomObject, &InPoints](int32 Index, FPCGPoint& OutPoint)
+	FPCGAsync::AsyncPointProcessing(&InContext, InPoints.Num(), OutPoints, [this, &InContext, InData, OutData, &InPoints](int32 Index, FPCGPoint& OutPoint)
 	{
-		return PointLoopBody(InContext, InData, InPoints[Index], OptionalCustomObject, OutPoint, OutData->Metadata);
+		return PointLoopBody(InContext, InData, InPoints[Index], OutPoint, OutData->Metadata);
 	});
 }
 
-void UPCGBlueprintElement::LoopOnPointPairs(FPCGContext& InContext, const UPCGPointData* InA, const UPCGPointData* InB, UPCGPointData*& OutData, UPCGPointData* OptionalOutData, const UObject* OptionalCustomObject) const
+void UPCGBlueprintElement::LoopOnPointPairs(FPCGContext& InContext, const UPCGPointData* InA, const UPCGPointData* InB, UPCGPointData*& OutData, UPCGPointData* OptionalOutData) const
 {
 	if (!InA || !InB)
 	{
@@ -326,13 +346,13 @@ void UPCGBlueprintElement::LoopOnPointPairs(FPCGContext& InContext, const UPCGPo
 	const TArray<FPCGPoint>& InPointsB = InB->GetPoints();
 	TArray<FPCGPoint>& OutPoints = OutData->GetMutablePoints();
 
-	FPCGAsync::AsyncPointProcessing(&InContext, InPointsA.Num() * InPointsB.Num(), OutPoints, [this, &InContext, InA, InB, OptionalCustomObject, OutData, &InPointsA, &InPointsB](int32 Index, FPCGPoint& OutPoint)
+	FPCGAsync::AsyncPointProcessing(&InContext, InPointsA.Num() * InPointsB.Num(), OutPoints, [this, &InContext, InA, InB, OutData, &InPointsA, &InPointsB](int32 Index, FPCGPoint& OutPoint)
 	{
-		return PointPairLoopBody(InContext, InA, InB, InPointsA[Index / InPointsB.Num()], InPointsB[Index % InPointsB.Num()], OptionalCustomObject, OutPoint, OutData->Metadata);
+		return PointPairLoopBody(InContext, InA, InB, InPointsA[Index / InPointsB.Num()], InPointsB[Index % InPointsB.Num()], OutPoint, OutData->Metadata);
 	});
 }
 
-void UPCGBlueprintElement::LoopNTimes(FPCGContext& InContext, int64 NumIterations, UPCGPointData*& OutData, const UPCGSpatialData* InA, const UPCGSpatialData* InB, UPCGPointData* OptionalOutData, const UObject* OptionalCustomObject) const
+void UPCGBlueprintElement::LoopNTimes(FPCGContext& InContext, int64 NumIterations, UPCGPointData*& OutData, const UPCGSpatialData* InA, const UPCGSpatialData* InB, UPCGPointData* OptionalOutData) const
 {
 	if (NumIterations < 0)
 	{
@@ -355,9 +375,9 @@ void UPCGBlueprintElement::LoopNTimes(FPCGContext& InContext, int64 NumIteration
 
 	TArray<FPCGPoint>& OutPoints = OutData->GetMutablePoints();
 
-	FPCGAsync::AsyncPointProcessing(&InContext, NumIterations, OutPoints, [this, &InContext, InA, InB, OptionalCustomObject, OutData](int32 Index, FPCGPoint& OutPoint)
+	FPCGAsync::AsyncPointProcessing(&InContext, NumIterations, OutPoints, [this, &InContext, InA, InB, OutData](int32 Index, FPCGPoint& OutPoint)
 	{
-		return IterationLoopBody(InContext, Index, InA, InB, OptionalCustomObject, OutPoint, OutData->Metadata);
+		return IterationLoopBody(InContext, Index, InA, InB, OutPoint, OutData->Metadata);
 	});
 }
 
diff --git a/Engine/Plugins/Experimental/PCG/Source/PCG/Private/Helpers/PCGSettingsHelpers.cpp b/Engine/Plugins/Experimental/PCG/Source/PCG/Private/Helpers/PCGSettingsHelpers.cpp
new file mode 100644
--- /dev/null
+++ b/Engine/Plugins/Experimental/PCG/Source/PCG/Private/Helpers/PCGSettingsHelpers.cpp
@@ -0,0 +1,192 @@
+// Copyright Epic Games, Inc. All Rights Reserved.
+
+#include "Helpers/PCGSettingsHelpers.h"
+
+namespace PCGSettingsHelpers
+{
+	template<typename T>
+	bool GetParamValue(T& Value, const FPCGMetadataAttributeBase* InAttribute, PCGMetadataEntryKey InKey = 0)
+	{
+		if (InAttribute->GetTypeId() == PCG::Private::MetadataTypes<T>::Id)
+		{
+			Value = static_cast<const FPCGMetadataAttribute<T>*>(InAttribute)->GetValueFromItemKey(InKey);
+			return true;
+		}
+		else
+		{
+			return false;
+		}
+	}
+
+	template<typename T, typename U, typename... OtherTypes>
+	bool GetParamValue(T& Value, const FPCGMetadataAttributeBase* InAttribute, PCGMetadataEntryKey InKey = 0)
+	{
+		if (InAttribute->GetTypeId() == PCG::Private::MetadataTypes<U>::Id)
+		{
+			Value = (T)static_cast<const FPCGMetadataAttribute<U>*>(InAttribute)->GetValueFromItemKey(InKey);
+			return true;
+		}
+		else
+		{
+			return GetParamValue<T, OtherTypes...>(Value, InAttribute, InKey);
+		}
+	}
+
+	void SetValue(UPCGParams* InParams, UObject* Object, FProperty* Property)
+	{
+		check(InParams && Property);
+		const FPCGMetadataAttributeBase* MatchingAttribute = InParams->Metadata ? InParams->Metadata->GetConstAttribute(Property->GetFName()) : nullptr;
+
+		// TODO? Support arrays (1) which would be possible if we have a consistent naming scheme
+		if (!MatchingAttribute)
+		{
+			return;
+		}
+
+		bool bTypeError = false;
+
+		if (FNumericProperty* NumericProperty = CastField<FNumericProperty>(Property))
+		{
+			if (NumericProperty->IsFloatingPoint())
+			{
+				double Value = 0;
+				if (GetParamValue<double, float>(Value, MatchingAttribute))
+				{
+					NumericProperty->SetFloatingPointPropertyValue(NumericProperty->ContainerPtrToValuePtr<double>(Object, 0), Value);
+				}
+				else
+				{
+					bTypeError = true;
+				}
+			}
+			else if(NumericProperty->IsInteger())
+			{
+				int64 Value = 0;
+				if (GetParamValue<int64, int32>(Value, MatchingAttribute))
+				{
+					NumericProperty->SetIntPropertyValue(NumericProperty->ContainerPtrToValuePtr<int64>(Object, 0), Value);
+				}
+				else
+				{
+					bTypeError = true;
+				}
+			}
+			else
+			{
+				// Not a float and not an integer...?
+				bTypeError = true;
+			}
+		}
+		else if (FBoolProperty* BoolProperty = CastField<FBoolProperty>(Property))
+		{
+			bool Value = 0;
+			if(GetParamValue<bool>(Value, MatchingAttribute))
+			{
+				BoolProperty->SetPropertyValue(BoolProperty->ContainerPtrToValuePtr<bool>(Object, 0), Value);
+			}
+			else
+			{
+				bTypeError = true;
+			}
+		}
+		else if (FStrProperty* StringProperty = CastField<FStrProperty>(Property))
+		{
+			FString Value;
+			if (GetParamValue<FString>(Value, MatchingAttribute))
+			{
+				StringProperty->SetPropertyValue(StringProperty->ContainerPtrToValuePtr<FString>(Object, 0), Value);
+			}
+			else
+			{
+				bTypeError = true;
+			}
+		}
+		else if (FNameProperty* NameProperty = CastField<FNameProperty>(Property))
+		{
+			FString Value;
+			if (GetParamValue<FString>(Value, MatchingAttribute))
+			{
+				FName NValue(Value);
+				NameProperty->SetPropertyValue(NameProperty->ContainerPtrToValuePtr<FName>(Object), NValue);
+			}
+			else
+			{
+				bTypeError = true;
+			}
+		}
+		else if (FEnumProperty* EnumProperty = CastField<FEnumProperty>(Property))
+		{
+			FNumericProperty* UnderlyingProperty = EnumProperty->GetUnderlyingProperty();
+
+			int64 Value;
+			if (GetParamValue<int64, int32, uint32, int16, uint16, int8, uint8>(Value, MatchingAttribute))
+			{
+				UnderlyingProperty->SetIntPropertyValue(EnumProperty->ContainerPtrToValuePtr<void>(Object, 0), Value);
+			}
+		}
+		else if (FStructProperty* StructProperty = CastField<FStructProperty>(Property))
+		{
+			const FName PropertyTypeName = StructProperty->Struct->GetFName();
+			void* PropertyValue = StructProperty->ContainerPtrToValuePtr<void>(Object, 0);
+
+			if (PropertyTypeName == NAME_Vector)
+			{
+				FVector Value = FVector::Zero();
+				if(GetParamValue<FVector, double, float>(Value, MatchingAttribute))
+				{
+					*reinterpret_cast<FVector*>(PropertyValue) = Value;
+				}
+				else
+				{
+					bTypeError = true;
+				}
+			}
+			else if (PropertyTypeName == NAME_Vector4)
+			{
+				FVector4 Value = FVector4::Zero();
+				if (GetParamValue<FVector4, double, float>(Value, MatchingAttribute))
+				{
+					*reinterpret_cast<FVector4*>(PropertyValue) = Value;
+				}
+				else
+				{
+					bTypeError = true;
+				}
+			}
+			else if (PropertyTypeName == NAME_Quat)
+			{
+				FQuat Value = FQuat::Identity;
+				if (GetParamValue<FQuat>(Value, MatchingAttribute))
+				{
+					*reinterpret_cast<FQuat*>(PropertyValue) = Value;
+				}
+			}
+			else if (PropertyTypeName == NAME_Transform)
+			{
+				FTransform Value = FTransform::Identity;
+				if (GetParamValue<FTransform, FVector>(Value, MatchingAttribute))
+				{
+					*reinterpret_cast<FTransform*>(PropertyValue) = Value;
+				}
+				else
+				{
+					bTypeError = true;
+				}
+			}
+			//else if(PropertyTypeName == NAME_Rotator)
+			//else if (PropertyTypeName == NAME_Color)
+			//else if (PropertyTypeName == NAME_LinearColor)
+			//else if soft object path to something?
+			else
+			{
+				bTypeError = true;
+			}
+		}
+
+		if (bTypeError)
+		{
+			// Mismatching types
+			UE_LOG(LogPCG, Warning, TEXT("Mismatching type in params vs. property %s"), *Property->GetFName().ToString());
+		}
+	}
+}
\ No newline at end of file
diff --git a/Engine/Plugins/Experimental/PCG/Source/PCG/Public/Elements/PCGExecuteBlueprint.h b/Engine/Plugins/Experimental/PCG/Source/PCG/Public/Elements/PCGExecuteBlueprint.h
--- a/Engine/Plugins/Experimental/PCG/Source/PCG/Public/Elements/PCGExecuteBlueprint.h
+++ b/Engine/Plugins/Experimental/PCG/Source/PCG/Public/Elements/PCGExecuteBlueprint.h
@@ -41,23 +41,23 @@ public:
 	void Execute(const FPCGDataCollection& Input, FPCGDataCollection& Output);
 
 	UFUNCTION(BlueprintImplementableEvent, Category = Execution)
-	bool PointLoopBody(UPARAM(ref) FPCGContext& InContext, const UPCGPointData* InData, const FPCGPoint& InPoint, const UObject* OptionalCustomObject, FPCGPoint& OutPoint, UPCGMetadata* OutMetadata) const;
+	bool PointLoopBody(UPARAM(ref) FPCGContext& InContext, const UPCGPointData* InData, const FPCGPoint& InPoint, FPCGPoint& OutPoint, UPCGMetadata* OutMetadata) const;
 
 	UFUNCTION(BlueprintImplementableEvent, Category = Execution)
-	bool PointPairLoopBody(UPARAM(ref) FPCGContext& InContext, const UPCGPointData* InA, const UPCGPointData* InB, const FPCGPoint& InPointA, const FPCGPoint& InPointB, const UObject* OptionalCustomObject, FPCGPoint& OutPoint, UPCGMetadata* OutMetadata) const;
+	bool PointPairLoopBody(UPARAM(ref) FPCGContext& InContext, const UPCGPointData* InA, const UPCGPointData* InB, const FPCGPoint& InPointA, const FPCGPoint& InPointB, FPCGPoint& OutPoint, UPCGMetadata* OutMetadata) const;
 
 	UFUNCTION(BlueprintImplementableEvent, Category = Execution)
-	bool IterationLoopBody(UPARAM(ref) FPCGContext& InContext, int64 Iteration, const UPCGSpatialData* InA, const UPCGSpatialData* InB, const UObject* OptionalCustomObject, FPCGPoint& OutPoint, UPCGMetadata* OutMetadata) const;
+	bool IterationLoopBody(UPARAM(ref) FPCGContext& InContext, int64 Iteration, const UPCGSpatialData* InA, const UPCGSpatialData* InB, FPCGPoint& OutPoint, UPCGMetadata* OutMetadata) const;
 
 	/** Calls the LoopBody function on all points */
 	UFUNCTION(BlueprintCallable, Category = Execution, meta = (HideSelfPin = "true"))
-	void LoopOnPoints(UPARAM(ref) FPCGContext& InContext, const UPCGPointData* InData, UPCGPointData*& OutData, UPCGPointData* OptionalOutData = nullptr, const UObject* OptionalCustomObject = nullptr) const;
+	void LoopOnPoints(UPARAM(ref) FPCGContext& InContext, const UPCGPointData* InData, UPCGPointData*& OutData, UPCGPointData* OptionalOutData = nullptr) const;
 
 	UFUNCTION(BlueprintCallable, Category = Execution, meta = (HideSelfPin = "true"))
-	void LoopOnPointPairs(UPARAM(ref) FPCGContext& InContext, const UPCGPointData* InA, const UPCGPointData* InB, UPCGPointData*& OutData, UPCGPointData* OptionalOutData = nullptr, const UObject* OptionalCustomObject = nullptr) const;
+	void LoopOnPointPairs(UPARAM(ref) FPCGContext& InContext, const UPCGPointData* InA, const UPCGPointData* InB, UPCGPointData*& OutData, UPCGPointData* OptionalOutData = nullptr) const;
 
 	UFUNCTION(BlueprintCallable, Category = Execution, meta = (HideSelfPin = "true"))
-	void LoopNTimes(UPARAM(ref) FPCGContext& InContext, int64 NumIterations, UPCGPointData*& OutData, const UPCGSpatialData* InA = nullptr, const UPCGSpatialData* InB = nullptr, UPCGPointData* OptionalOutData = nullptr, const UObject* OptionalCustomObject = nullptr) const;
+	void LoopNTimes(UPARAM(ref) FPCGContext& InContext, int64 NumIterations, UPCGPointData*& OutData, const UPCGSpatialData* InA = nullptr, const UPCGSpatialData* InB = nullptr, UPCGPointData* OptionalOutData = nullptr) const;
 
 	/** Called after object creation to setup the object callbacks */
 	void Initialize();
diff --git a/Engine/Plugins/Experimental/PCG/Source/PCG/Public/Helpers/PCGSettingsHelpers.h b/Engine/Plugins/Experimental/PCG/Source/PCG/Public/Helpers/PCGSettingsHelpers.h
--- a/Engine/Plugins/Experimental/PCG/Source/PCG/Public/Helpers/PCGSettingsHelpers.h
+++ b/Engine/Plugins/Experimental/PCG/Source/PCG/Public/Helpers/PCGSettingsHelpers.h
@@ -85,4 +85,7 @@ namespace PCGSettingsHelpers
 	{
 		return FName(GetValue(InName, InValue.ToString(), InParams, InParamName));
 	}
+
+	/** Sets data from the params to a given property, matched on a name basis */
+	void SetValue(UPCGParams* Params, UObject* Object, FProperty* Property);
 }
\ No newline at end of file
