Made PCG blueprint element loop calls impure

diff --git a/Engine/Plugins/Experimental/PCG/Source/PCG/Public/Elements/PCGExecuteBlueprint.h b/Engine/Plugins/Experimental/PCG/Source/PCG/Public/Elements/PCGExecuteBlueprint.h
--- a/Engine/Plugins/Experimental/PCG/Source/PCG/Public/Elements/PCGExecuteBlueprint.h
+++ b/Engine/Plugins/Experimental/PCG/Source/PCG/Public/Elements/PCGExecuteBlueprint.h
@@ -50,13 +50,13 @@ public:
 	bool IterationLoopBody(UPARAM(ref) FPCGContext& InContext, int64 Iteration, const UPCGSpatialData* InA, const UPCGSpatialData* InB, FPCGPoint& OutPoint, UPCGMetadata* OutMetadata) const;
 
 	/** Calls the LoopBody function on all points */
-	UFUNCTION(BlueprintCallable, Category = Execution, meta = (HideSelfPin = "true"))
+	UFUNCTION(BlueprintCallable, BlueprintPure = false, Category = Execution, meta = (HideSelfPin = "true"))
 	void LoopOnPoints(UPARAM(ref) FPCGContext& InContext, const UPCGPointData* InData, UPCGPointData*& OutData, UPCGPointData* OptionalOutData = nullptr) const;
 
-	UFUNCTION(BlueprintCallable, Category = Execution, meta = (HideSelfPin = "true"))
+	UFUNCTION(BlueprintCallable, BlueprintPure = false, Category = Execution, meta = (HideSelfPin = "true"))
 	void LoopOnPointPairs(UPARAM(ref) FPCGContext& InContext, const UPCGPointData* InA, const UPCGPointData* InB, UPCGPointData*& OutData, UPCGPointData* OptionalOutData = nullptr) const;
 
-	UFUNCTION(BlueprintCallable, Category = Execution, meta = (HideSelfPin = "true"))
+	UFUNCTION(BlueprintCallable, BlueprintPure = false, Category = Execution, meta = (HideSelfPin = "true"))
 	void LoopNTimes(UPARAM(ref) FPCGContext& InContext, int64 NumIterations, UPCGPointData*& OutData, const UPCGSpatialData* InA = nullptr, const UPCGSpatialData* InB = nullptr, UPCGPointData* OptionalOutData = nullptr) const;
 
 	/** Called after object creation to setup the object callbacks */
