Made graph compilation more robust

diff --git a/Engine/Plugins/Experimental/PCG/Source/PCG/Private/Graph/PCGGraphCompiler.cpp b/Engine/Plugins/Experimental/PCG/Source/PCG/Private/Graph/PCGGraphCompiler.cpp
--- a/Engine/Plugins/Experimental/PCG/Source/PCG/Private/Graph/PCGGraphCompiler.cpp
+++ b/Engine/Plugins/Experimental/PCG/Source/PCG/Private/Graph/PCGGraphCompiler.cpp
@@ -7,7 +7,11 @@
 
 TArray<FPCGGraphTask> FPCGGraphCompiler::CompileGraph(UPCGGraph* InGraph, FPCGTaskId& NextId)
 {
-	check(InGraph);
+	if (!InGraph)
+	{
+		return TArray<FPCGGraphTask>();
+	}
+
 	TArray<FPCGGraphTask> CompiledTasks;
 	TMap<const UPCGNode*, FPCGTaskId> IdMapping;
 	TArray<const UPCGNode*> NodeQueue;
@@ -112,7 +116,15 @@ TArray<FPCGGraphTask> FPCGGraphCompiler::CompileGraph(UPCGGraph* InGraph, FPCGTa
 
 			for (const UPCGNode* Inbound : Node->InboundNodes)
 			{
-				Task.Inputs.Add(IdMapping[Inbound]);
+				if (FPCGTaskId* InboundId = IdMapping.Find(Inbound))
+				{
+					Task.Inputs.Add(*InboundId);
+				}
+				else
+				{
+					UE_LOG(LogPCG, Error, TEXT("Inconsistent node linkage between node %s and node %s"), *Node->GetFName().ToString(), (Inbound ? *Inbound->GetFName().ToString() : TEXT("null")));
+					return TArray<FPCGGraphTask>();
+				}
 			}
 
 			IdMapping.Add(Node, NodeId);
@@ -155,13 +167,16 @@ void FPCGGraphCompiler::Compile(UPCGGraph* InGraph)
 
 	// TODO: optimize no-ops, etc.
 
-	// Store back the results in the cache
-	GraphToTaskMapLock.WriteLock();
-	if (!GraphToTaskMap.Contains(InGraph))
+	// Store back the results in the cache if it's valid
+	if (!CompiledTasks.IsEmpty())
 	{
-		GraphToTaskMap.Add(InGraph, MoveTemp(CompiledTasks));
+		GraphToTaskMapLock.WriteLock();
+		if (!GraphToTaskMap.Contains(InGraph))
+		{
+			GraphToTaskMap.Add(InGraph, MoveTemp(CompiledTasks));
+		}
+		GraphToTaskMapLock.WriteUnlock();
 	}
-	GraphToTaskMapLock.WriteUnlock();
 }
 
 TArray<FPCGGraphTask> FPCGGraphCompiler::GetCompiledTasks(UPCGGraph* InGraph, bool bIsTopGraph)
@@ -176,8 +191,10 @@ TArray<FPCGGraphTask> FPCGGraphCompiler::GetCompiledTasks(UPCGGraph* InGraph, bo
 
 		// Get compiled tasks in a threadsafe way
 		GraphToTaskMapLock.ReadLock();
-		check(TopGraphToTaskMap.Contains(InGraph));
-		CompiledTasks = TopGraphToTaskMap[InGraph];
+		if (TopGraphToTaskMap.Contains(InGraph))
+		{
+			CompiledTasks = TopGraphToTaskMap[InGraph];
+		}
 		GraphToTaskMapLock.ReadUnlock();
 	}
 	else
@@ -188,8 +205,10 @@ TArray<FPCGGraphTask> FPCGGraphCompiler::GetCompiledTasks(UPCGGraph* InGraph, bo
 
 		// Get compiled tasks in a threadsafe way
 		GraphToTaskMapLock.ReadLock();
-		check(GraphToTaskMap.Contains(InGraph));
-		CompiledTasks = GraphToTaskMap[InGraph];
+		if (GraphToTaskMap.Contains(InGraph))
+		{
+			CompiledTasks = GraphToTaskMap[InGraph];
+		}
 		GraphToTaskMapLock.ReadUnlock();
 	}
 
@@ -224,6 +243,12 @@ void FPCGGraphCompiler::CompileTopGraph(UPCGGraph* InGraph)
 	// Build from non-top tasks
 	TArray<FPCGGraphTask> CompiledTasks = GetCompiledTasks(InGraph, /*bIsTopGraph=*/false);
 
+	// Check that the compilation was valid
+	if (CompiledTasks.Num() == 0)
+	{
+		return;
+	}
+
 	const int TaskNum = CompiledTasks.Num();
 	const FPCGTaskId PreExecuteTaskId = FPCGTaskId(TaskNum);
 	const FPCGTaskId PostExecuteTaskId = PreExecuteTaskId + 1;
diff --git a/Engine/Plugins/Experimental/PCG/Source/PCG/Private/Graph/PCGGraphExecutor.cpp b/Engine/Plugins/Experimental/PCG/Source/PCG/Private/Graph/PCGGraphExecutor.cpp
--- a/Engine/Plugins/Experimental/PCG/Source/PCG/Private/Graph/PCGGraphExecutor.cpp
+++ b/Engine/Plugins/Experimental/PCG/Source/PCG/Private/Graph/PCGGraphExecutor.cpp
@@ -144,7 +144,11 @@ void FPCGGraphExecutor::Execute()
 	// TODO: add optimization phase if we've added new graph(s)/tasks to be executed
 
 	// This is a safeguard to check if we're in a stuck state
-	check(Tasks.Num() == 0 || ReadyTasks.Num() > 0 || ActiveTasks.Num() > 0);
+	if (ReadyTasks.Num() == 0 && ActiveTasks.Num() == 0 && Tasks.Num() > 0)
+	{
+		UE_LOG(LogPCG, Error, TEXT("PCG Graph executor error: tasks are in a deadlocked state. Will drop all tasks."));
+		Tasks.Reset();
+	}
 
 	// TODO: change this when we support time-slicing
 	if(ReadyTasks.Num() > 0 || ActiveTasks.Num() > 0)
diff --git a/Engine/Plugins/Experimental/PCG/Source/PCGEditor/Private/PCGEditorGraphNode.cpp b/Engine/Plugins/Experimental/PCG/Source/PCGEditor/Private/PCGEditorGraphNode.cpp
--- a/Engine/Plugins/Experimental/PCG/Source/PCGEditor/Private/PCGEditorGraphNode.cpp
+++ b/Engine/Plugins/Experimental/PCG/Source/PCGEditor/Private/PCGEditorGraphNode.cpp
@@ -27,8 +27,7 @@ void UPCGEditorGraphNode::Construct(UPCGNode* InPCGNode, EPCGEditorGraphNodeType
 
 FText UPCGEditorGraphNode::GetNodeTitle(ENodeTitleType::Type TitleType) const
 {
-	check(PCGNode && PCGNode->DefaultSettings);
-	return FText::FromName(PCGNode->DefaultSettings->GetDefaultNodeName());
+	return (PCGNode && PCGNode->DefaultSettings) ? FText::FromName(PCGNode->DefaultSettings->GetDefaultNodeName()) : FText::FromName(TEXT("Unnamed node"));
 }
 
 void UPCGEditorGraphNode::GetNodeContextMenuActions(UToolMenu* Menu, class UGraphNodeContextMenuContext* Context) const
