Added blueprint methods to get params
Renamed params class to param data for coherence
Removed usage flag

diff --git a/Engine/Plugins/Experimental/PCG/Source/PCG/Private/Elements/PCGDensityFilter.cpp b/Engine/Plugins/Experimental/PCG/Source/PCG/Private/Elements/PCGDensityFilter.cpp
--- a/Engine/Plugins/Experimental/PCG/Source/PCG/Private/Elements/PCGDensityFilter.cpp
+++ b/Engine/Plugins/Experimental/PCG/Source/PCG/Private/Elements/PCGDensityFilter.cpp
@@ -19,10 +19,9 @@ bool FPCGDensityFilterElement::ExecuteInternal(FPCGContext* Context) const
 
 	TArray<FPCGTaggedData> Inputs = Context->InputData.GetInputs();
 	TArray<FPCGTaggedData>& Outputs = Context->OutputData.TaggedData;
-	UPCGParams* Params = Context->InputData.GetParams();
+	UPCGParamData* Params = Context->InputData.GetParams();
 
 	// Forward any non-input data
-	Outputs.Append(Context->InputData.GetExclusions());
 	Outputs.Append(Context->InputData.GetAllSettings());
 
 	const bool bInvertFilter = PCGSettingsHelpers::GetValue(GET_MEMBER_NAME_CHECKED(UPCGDensityFilterSettings, bInvertFilter), Settings->bInvertFilter, Params);
diff --git a/Engine/Plugins/Experimental/PCG/Source/PCG/Private/Elements/PCGDensityRemap.cpp b/Engine/Plugins/Experimental/PCG/Source/PCG/Private/Elements/PCGDensityRemap.cpp
--- a/Engine/Plugins/Experimental/PCG/Source/PCG/Private/Elements/PCGDensityRemap.cpp
+++ b/Engine/Plugins/Experimental/PCG/Source/PCG/Private/Elements/PCGDensityRemap.cpp
@@ -21,10 +21,9 @@ bool FPCGLinearDensityRemapElement::ExecuteInternal(FPCGContext* Context) const
 
 	TArray<FPCGTaggedData> Inputs = Context->InputData.GetInputs();
 	TArray<FPCGTaggedData>& Outputs = Context->OutputData.TaggedData;
-	UPCGParams* Params = Context->InputData.GetParams();
+	UPCGParamData* Params = Context->InputData.GetParams();
 
 	// Forward any non-input data
-	Outputs.Append(Context->InputData.GetExclusions());
 	Outputs.Append(Context->InputData.GetAllSettings());
 
 	const float SettingsRemapMin = PCGSettingsHelpers::GetValue(GET_MEMBER_NAME_CHECKED(UPCGLinearDensityRemapSettings, RemapMin), Settings->RemapMin, Params);
diff --git a/Engine/Plugins/Experimental/PCG/Source/PCG/Private/Elements/PCGDifferenceElement.cpp b/Engine/Plugins/Experimental/PCG/Source/PCG/Private/Elements/PCGDifferenceElement.cpp
--- a/Engine/Plugins/Experimental/PCG/Source/PCG/Private/Elements/PCGDifferenceElement.cpp
+++ b/Engine/Plugins/Experimental/PCG/Source/PCG/Private/Elements/PCGDifferenceElement.cpp
@@ -16,8 +16,7 @@ bool FPCGDifferenceElement::ExecuteInternal(FPCGContext* Context) const
 	check(Settings);
 
 	TArray<FPCGTaggedData> Inputs = Context->InputData.GetInputs();
-	TArray<FPCGTaggedData> Exclusions = Context->InputData.GetExclusions();
-	UPCGParams* Params = Context->InputData.GetParams();
+	UPCGParamData* Params = Context->InputData.GetParams();
 
 	const EPCGDifferenceDensityFunction DensityFunction = PCGSettingsHelpers::GetValue(GET_MEMBER_NAME_CHECKED(UPCGDifferenceSettings, DensityFunction), Settings->DensityFunction, Params);
 #if WITH_EDITORONLY_DATA
@@ -75,18 +74,6 @@ bool FPCGDifferenceElement::ExecuteInternal(FPCGContext* Context) const
 		AddToDifference(SpatialData);
 	}
 
-	// Can't have a difference if we don't have an input
-	if (FirstSpatialData)
-	{
-		for (FPCGTaggedData& Exclusion : Exclusions)
-		{
-			if (const UPCGSpatialData* SpatialData = Cast<UPCGSpatialData>(Exclusion.Data))
-			{
-				AddToDifference(SpatialData);
-			}
-		}
-	}
-
 	// Finally, pass-through settings
 	Outputs.Append(Context->InputData.GetAllSettings());
 
diff --git a/Engine/Plugins/Experimental/PCG/Source/PCG/Private/Elements/PCGExecuteBlueprint.cpp b/Engine/Plugins/Experimental/PCG/Source/PCG/Private/Elements/PCGExecuteBlueprint.cpp
--- a/Engine/Plugins/Experimental/PCG/Source/PCG/Private/Elements/PCGExecuteBlueprint.cpp
+++ b/Engine/Plugins/Experimental/PCG/Source/PCG/Private/Elements/PCGExecuteBlueprint.cpp
@@ -272,7 +272,7 @@ bool FPCGExecuteBlueprintElement::ExecuteInternal(FPCGContext* InContext) const
 		}
 
 		/** Apply params overrides to variables if any */
-		if (UPCGParams* Params = Context->InputData.GetParams())
+		if (UPCGParamData* Params = Context->InputData.GetParams())
 		{
 			UClass* BPClass = Context->BlueprintElementInstance->GetClass();
 
diff --git a/Engine/Plugins/Experimental/PCG/Source/PCG/Private/Elements/PCGIntersectionElement.cpp b/Engine/Plugins/Experimental/PCG/Source/PCG/Private/Elements/PCGIntersectionElement.cpp
--- a/Engine/Plugins/Experimental/PCG/Source/PCG/Private/Elements/PCGIntersectionElement.cpp
+++ b/Engine/Plugins/Experimental/PCG/Source/PCG/Private/Elements/PCGIntersectionElement.cpp
@@ -16,7 +16,7 @@ bool FPCGIntersectionElement::ExecuteInternal(FPCGContext* Context) const
 	check(Settings);
 
 	TArray<FPCGTaggedData> Inputs = Context->InputData.GetInputs();
-	UPCGParams* Params = Context->InputData.GetParams();
+	UPCGParamData* Params = Context->InputData.GetParams();
 
 	const EPCGIntersectionDensityFunction DensityFunction = PCGSettingsHelpers::GetValue(GET_MEMBER_NAME_CHECKED(UPCGIntersectionSettings, DensityFunction), Settings->DensityFunction, Params);
 #if WITH_EDITORONLY_DATA
@@ -65,7 +65,6 @@ bool FPCGIntersectionElement::ExecuteInternal(FPCGContext* Context) const
 	}
 
 	// Pass-through settings & exclusions
-	Outputs.Append(Context->InputData.GetExclusions());
 	Outputs.Append(Context->InputData.GetAllSettings());
 
 	return true;
diff --git a/Engine/Plugins/Experimental/PCG/Source/PCG/Private/Elements/PCGMetadataElement.cpp b/Engine/Plugins/Experimental/PCG/Source/PCG/Private/Elements/PCGMetadataElement.cpp
--- a/Engine/Plugins/Experimental/PCG/Source/PCG/Private/Elements/PCGMetadataElement.cpp
+++ b/Engine/Plugins/Experimental/PCG/Source/PCG/Private/Elements/PCGMetadataElement.cpp
@@ -62,7 +62,7 @@ bool FPCGMetadataOperationElement::ExecuteInternal(FPCGContext* Context) const
 	check(Settings);
 
 	TArray<FPCGTaggedData> Inputs = Context->InputData.GetInputs();
-	UPCGParams* Params = Context->InputData.GetParams();
+	UPCGParamData* Params = Context->InputData.GetParams();
 
 	TArray<FPCGTaggedData>& Outputs = Context->OutputData.TaggedData;
 
@@ -72,7 +72,6 @@ bool FPCGMetadataOperationElement::ExecuteInternal(FPCGContext* Context) const
 	const EPCGMetadataOperationTarget Target = PCGSettingsHelpers::GetValue(GET_MEMBER_NAME_CHECKED(UPCGMetadataOperationSettings, Target), Settings->Target, Params);
 
 	// Forward any non-input data
-	Outputs.Append(Context->InputData.GetExclusions());
 	Outputs.Append(Context->InputData.GetAllSettings());
 
 	for (const FPCGTaggedData& Input : Inputs)
diff --git a/Engine/Plugins/Experimental/PCG/Source/PCG/Private/Elements/PCGPointSampler.cpp b/Engine/Plugins/Experimental/PCG/Source/PCG/Private/Elements/PCGPointSampler.cpp
--- a/Engine/Plugins/Experimental/PCG/Source/PCG/Private/Elements/PCGPointSampler.cpp
+++ b/Engine/Plugins/Experimental/PCG/Source/PCG/Private/Elements/PCGPointSampler.cpp
@@ -20,12 +20,11 @@ bool FPCGPointSamplerElement::ExecuteInternal(FPCGContext* Context) const
 	check(Settings);
 
 	TArray<FPCGTaggedData> Inputs = Context->InputData.GetInputs();
-	UPCGParams* Params = Context->InputData.GetParams();
+	UPCGParamData* Params = Context->InputData.GetParams();
 
 	TArray<FPCGTaggedData>& Outputs = Context->OutputData.TaggedData;
 
 	// Forward any non-input data, excluding params
-	Outputs.Append(Context->InputData.GetExclusions());
 	Outputs.Append(Context->InputData.GetAllSettings());
 
 	const float Ratio = PCGSettingsHelpers::GetValue(GET_MEMBER_NAME_CHECKED(UPCGPointSamplerSettings, Ratio), Settings->Ratio, Params);
diff --git a/Engine/Plugins/Experimental/PCG/Source/PCG/Private/Elements/PCGProjectionElement.cpp b/Engine/Plugins/Experimental/PCG/Source/PCG/Private/Elements/PCGProjectionElement.cpp
--- a/Engine/Plugins/Experimental/PCG/Source/PCG/Private/Elements/PCGProjectionElement.cpp
+++ b/Engine/Plugins/Experimental/PCG/Source/PCG/Private/Elements/PCGProjectionElement.cpp
@@ -17,7 +17,7 @@ bool FPCGProjectionElement::ExecuteInternal(FPCGContext* Context) const
 	check(Settings);
 
 	TArray<FPCGTaggedData> Inputs = Context->InputData.GetInputs();
-	UPCGParams* Params = Context->InputData.GetParams();
+	UPCGParamData* Params = Context->InputData.GetParams();
 
 	TArray<FPCGTaggedData>& Outputs = Context->OutputData.TaggedData;
 
@@ -65,7 +65,6 @@ bool FPCGProjectionElement::ExecuteInternal(FPCGContext* Context) const
 	}
 
 	// Pass-through exclusions/settings
-	Outputs.Append(Context->InputData.GetExclusions());
 	Outputs.Append(Context->InputData.GetAllSettings());
 
 	return true;
diff --git a/Engine/Plugins/Experimental/PCG/Source/PCG/Private/Elements/PCGSelfPruning.cpp b/Engine/Plugins/Experimental/PCG/Source/PCG/Private/Elements/PCGSelfPruning.cpp
--- a/Engine/Plugins/Experimental/PCG/Source/PCG/Private/Elements/PCGSelfPruning.cpp
+++ b/Engine/Plugins/Experimental/PCG/Source/PCG/Private/Elements/PCGSelfPruning.cpp
@@ -205,8 +205,7 @@ namespace PCGSelfPruningElement
 			}
 		}
 
-		// Finally, forward any exclusions/settings
-		Outputs.Append(Context->InputData.GetExclusions());
+		// Finally, forward any settings
 		Outputs.Append(Context->InputData.GetAllSettings());
 	}
 }
@@ -223,7 +222,7 @@ bool FPCGSelfPruningElement::ExecuteInternal(FPCGContext* Context) const
 	const UPCGSelfPruningSettings* Settings = Context->GetInputSettings<UPCGSelfPruningSettings>();
 	check(Settings);
 
-	UPCGParams* Params = Context->InputData.GetParams();
+	UPCGParamData* Params = Context->InputData.GetParams();
 
 	const EPCGSelfPruningType PruningType = PCGSettingsHelpers::GetValue(GET_MEMBER_NAME_CHECKED(UPCGSelfPruningSettings, PruningType), Settings->PruningType, Params);
 	const float RadiusSimilarityFactor = PCGSettingsHelpers::GetValue(GET_MEMBER_NAME_CHECKED(UPCGSelfPruningSettings, RadiusSimilarityFactor), Settings->RadiusSimilarityFactor, Params);
diff --git a/Engine/Plugins/Experimental/PCG/Source/PCG/Private/Elements/PCGSurfaceSampler.cpp b/Engine/Plugins/Experimental/PCG/Source/PCG/Private/Elements/PCGSurfaceSampler.cpp
--- a/Engine/Plugins/Experimental/PCG/Source/PCG/Private/Elements/PCGSurfaceSampler.cpp
+++ b/Engine/Plugins/Experimental/PCG/Source/PCG/Private/Elements/PCGSurfaceSampler.cpp
@@ -36,7 +36,7 @@ struct FPCGSurfaceSamplerLoopData
 
 	bool Initialize(const UPCGSurfaceSamplerSettings* InSettings, FPCGContext* Context, const FBox& InputBounds)
 	{
-		UPCGParams* Params = Context->InputData.GetParams();
+		UPCGParamData* Params = Context->InputData.GetParams();
 		Settings = InSettings;
 
 		// Compute used values
@@ -122,7 +122,7 @@ bool FPCGSurfaceSamplerElement::ExecuteInternal(FPCGContext* Context) const
 	check(Settings);
 
 	TArray<FPCGTaggedData> Inputs = Context->InputData.GetInputs();
-	UPCGParams* Params = Context->InputData.GetParams();
+	UPCGParamData* Params = Context->InputData.GetParams();
 
 	// Early out on invalid settings
 	// TODO: we could compute an approximate radius based on the points per squared meters if that's useful
@@ -217,7 +217,6 @@ bool FPCGSurfaceSamplerElement::ExecuteInternal(FPCGContext* Context) const
 	}
 
 	// Finally, forward any exclusions/settings
-	Outputs.Append(Context->InputData.GetExclusions());
 	Outputs.Append(Context->InputData.GetAllSettings());
 
 	return true;
diff --git a/Engine/Plugins/Experimental/PCG/Source/PCG/Private/Elements/PCGUnionElement.cpp b/Engine/Plugins/Experimental/PCG/Source/PCG/Private/Elements/PCGUnionElement.cpp
--- a/Engine/Plugins/Experimental/PCG/Source/PCG/Private/Elements/PCGUnionElement.cpp
+++ b/Engine/Plugins/Experimental/PCG/Source/PCG/Private/Elements/PCGUnionElement.cpp
@@ -16,7 +16,7 @@ bool FPCGUnionElement::ExecuteInternal(FPCGContext* Context) const
 	check(Settings);
 
 	TArray<FPCGTaggedData> Inputs = Context->InputData.GetInputs();
-	UPCGParams* Params = Context->InputData.GetParams();
+	UPCGParamData* Params = Context->InputData.GetParams();
 
 	const EPCGUnionType Type = PCGSettingsHelpers::GetValue(GET_MEMBER_NAME_CHECKED(UPCGUnionSettings, Type), Settings->Type, Params);
 	const EPCGUnionDensityFunction DensityFunction = PCGSettingsHelpers::GetValue(GET_MEMBER_NAME_CHECKED(UPCGUnionSettings, DensityFunction), Settings->DensityFunction, Params);
@@ -67,7 +67,6 @@ bool FPCGUnionElement::ExecuteInternal(FPCGContext* Context) const
 	}
 
 	// Finally, pass-through settings
-	Outputs.Append(Context->InputData.GetExclusions());
 	Outputs.Append(Context->InputData.GetAllSettings());
 
 	return true;
diff --git a/Engine/Plugins/Experimental/PCG/Source/PCG/Private/Helpers/PCGSettingsHelpers.cpp b/Engine/Plugins/Experimental/PCG/Source/PCG/Private/Helpers/PCGSettingsHelpers.cpp
--- a/Engine/Plugins/Experimental/PCG/Source/PCG/Private/Helpers/PCGSettingsHelpers.cpp
+++ b/Engine/Plugins/Experimental/PCG/Source/PCG/Private/Helpers/PCGSettingsHelpers.cpp
@@ -32,7 +32,7 @@ namespace PCGSettingsHelpers
 		}
 	}
 
-	void SetValue(UPCGParams* InParams, UObject* Object, FProperty* Property)
+	void SetValue(UPCGParamData* InParams, UObject* Object, FProperty* Property)
 	{
 		check(InParams && Property);
 		const FPCGMetadataAttributeBase* MatchingAttribute = InParams->Metadata ? InParams->Metadata->GetConstAttribute(Property->GetFName()) : nullptr;
diff --git a/Engine/Plugins/Experimental/PCG/Source/PCG/Private/PCGData.cpp b/Engine/Plugins/Experimental/PCG/Source/PCG/Private/PCGData.cpp
--- a/Engine/Plugins/Experimental/PCG/Source/PCG/Private/PCGData.cpp
+++ b/Engine/Plugins/Experimental/PCG/Source/PCG/Private/PCGData.cpp
@@ -2,12 +2,12 @@
 
 #include "PCGData.h"
 #include "PCGSettings.h"
-#include "PCGParams.h"
+#include "PCGParamData.h"
+#include "Data/PCGSpatialData.h"
 
 bool FPCGTaggedData::operator==(const FPCGTaggedData& Other) const
 {
 	return Data == Other.Data &&
-		Usage == Other.Usage &&
 		Tags.Num() == Other.Tags.Num() &&
 		Tags.Includes(Other.Tags);
 }
@@ -20,36 +20,43 @@ bool FPCGTaggedData::operator!=(const FPCGTaggedData& Other) const
 TArray<FPCGTaggedData> FPCGDataCollection::GetInputs() const
 {
 	return TaggedData.FilterByPredicate([](const FPCGTaggedData& Data) {
-		return Data.Usage == EPCGDataUsage::Input;
+		return Cast<UPCGSpatialData>(Data.Data) != nullptr;
 		});
 }
 
 TArray<FPCGTaggedData> FPCGDataCollection::GetTaggedInputs(const FString& InTag) const
 {
 	return TaggedData.FilterByPredicate([&InTag](const FPCGTaggedData& Data) {
-		return Data.Usage == EPCGDataUsage::Input && Data.Tags.Contains(InTag);
+		return Data.Tags.Contains(InTag) && Cast<UPCGSpatialData>(Data.Data);
 		});
 }
 
-TArray<FPCGTaggedData> FPCGDataCollection::GetExclusions() const
+TArray<FPCGTaggedData> FPCGDataCollection::GetAllSettings() const
 {
 	return TaggedData.FilterByPredicate([](const FPCGTaggedData& Data) {
-		return Data.Usage == EPCGDataUsage::Exclusion;
+		return Cast<UPCGSettings>(Data.Data) != nullptr;
 		});
 }
 
-TArray<FPCGTaggedData> FPCGDataCollection::GetAllSettings() const
+TArray<FPCGTaggedData> FPCGDataCollection::GetAllParams() const
 {
 	return TaggedData.FilterByPredicate([](const FPCGTaggedData& Data) {
-		return Data.Usage == EPCGDataUsage::Settings;
+		return Cast<UPCGParamData>(Data.Data) != nullptr;
+	});
+}
+
+TArray<FPCGTaggedData> FPCGDataCollection::GetTaggedParams(const FString& InTag) const
+{
+	return TaggedData.FilterByPredicate([&InTag](const FPCGTaggedData& Data) {
+		return Data.Tags.Contains(InTag) && Cast<UPCGParamData>(Data.Data) != nullptr;
 		});
 }
 
-UPCGParams* FPCGDataCollection::GetParams() const
+UPCGParamData* FPCGDataCollection::GetParams() const
 {
 	for (const FPCGTaggedData& TaggedDatum : TaggedData)
 	{
-		if (UPCGParams* Params = Cast<UPCGParams>(TaggedDatum.Data))
+		if (UPCGParamData* Params = Cast<UPCGParamData>(TaggedDatum.Data))
 		{
 			return Params; 
 		}
@@ -67,7 +74,7 @@ const UPCGSettings* FPCGDataCollection::GetSettings(const UPCGSettings* InDefaul
 	else
 	{
 		const FPCGTaggedData* MatchingData = TaggedData.FindByPredicate([InDefaultSettings](const FPCGTaggedData& Data) {
-			return Data.Usage == EPCGDataUsage::Settings &&
+			return Data.Data && 
 				(Data.Data->GetClass() == InDefaultSettings->GetClass() ||
 					Data.Data->GetClass()->IsChildOf(InDefaultSettings->GetClass()));
 			});
@@ -119,9 +126,14 @@ TArray<FPCGTaggedData> UPCGDataFunctionLibrary::GetTaggedInputs(const FPCGDataCo
 	return InCollection.GetTaggedInputs(InTag);
 }
 
-TArray<FPCGTaggedData> UPCGDataFunctionLibrary::GetExclusions(const FPCGDataCollection& InCollection)
+TArray<FPCGTaggedData> UPCGDataFunctionLibrary::GetParams(const FPCGDataCollection& InCollection)
+{
+	return InCollection.GetAllParams();
+}
+
+TArray<FPCGTaggedData> UPCGDataFunctionLibrary::GetTaggedParams(const FPCGDataCollection& InCollection, const FString& InTag)
 {
-	return InCollection.GetExclusions();
+	return InCollection.GetTaggedParams(InTag);
 }
 
 TArray<FPCGTaggedData> UPCGDataFunctionLibrary::GetAllSettings(const FPCGDataCollection& InCollection)
diff --git a/Engine/Plugins/Experimental/PCG/Source/PCG/Private/PCGElement.cpp b/Engine/Plugins/Experimental/PCG/Source/PCG/Private/PCGElement.cpp
--- a/Engine/Plugins/Experimental/PCG/Source/PCG/Private/PCGElement.cpp
+++ b/Engine/Plugins/Experimental/PCG/Source/PCG/Private/PCGElement.cpp
@@ -39,11 +39,7 @@ bool IPCGElement::Execute(FPCGContext* Context) const
 			TArray<FPCGTaggedData> FilteredTaggedData;
 			for (FPCGTaggedData& TaggedData : Context->InputData.TaggedData)
 			{
-				if (TaggedData.Usage != EPCGDataUsage::Input)
-				{
-					FilteredTaggedData.Add(TaggedData);
-				}
-				else if (TaggedData.Tags.Intersect(Settings->FilterOnTags).IsEmpty())
+				if (TaggedData.Tags.Intersect(Settings->FilterOnTags).IsEmpty())
 				{
 					if (Settings->bPassThroughFilteredOutInputs)
 					{
@@ -108,12 +104,9 @@ bool IPCGElement::Execute(FPCGContext* Context) const
 		{
 			for (FPCGTaggedData& TaggedData : Context->OutputData.TaggedData)
 			{
-				if (TaggedData.Usage == EPCGDataUsage::Input)
+				if (!BypassedTaggedData.Contains(TaggedData))
 				{
-					if (!BypassedTaggedData.Contains(TaggedData))
-					{
-						TaggedData.Tags.Append(Settings->TagsAppliedOnOutput);
-					}
+					TaggedData.Tags.Append(Settings->TagsAppliedOnOutput);
 				}
 			}
 		}
diff --git a/Engine/Plugins/Experimental/PCG/Source/PCG/Private/PCGParamData.cpp b/Engine/Plugins/Experimental/PCG/Source/PCG/Private/PCGParamData.cpp
new file mode 100644
--- /dev/null
+++ b/Engine/Plugins/Experimental/PCG/Source/PCG/Private/PCGParamData.cpp
@@ -0,0 +1,69 @@
+// Copyright Epic Games, Inc. All Rights Reserved.
+
+#include "PCGParamData.h"
+
+UPCGParamData::UPCGParamData(const FObjectInitializer& ObjectInitializer)
+	: Super(ObjectInitializer)
+{
+	Metadata = ObjectInitializer.CreateDefaultSubobject<UPCGMetadata>(this, TEXT("Metadata"));
+}
+
+int64 UPCGParamData::FindMetadataKey(const FName& InName) const
+{
+	if (const PCGMetadataEntryKey* FoundKey = NameMap.Find(InName))
+	{
+		return *FoundKey;
+	}
+	else
+	{
+		return PCGInvalidEntryKey;
+	}
+}
+
+int64 UPCGParamData::FindOrAddMetadataKey(const FName& InName)
+{
+	if (const PCGMetadataEntryKey* FoundKey = NameMap.Find(InName))
+	{
+		return *FoundKey;
+	}
+	else
+	{
+		check(Metadata);
+		PCGMetadataEntryKey NewKey = Metadata->AddEntry();
+		NameMap.Add(InName, NewKey);
+		return NewKey;
+	}
+}
+
+UPCGParamData* UPCGParamData::FilterParamsByName(const FName& InName) const
+{
+	PCGMetadataEntryKey EntryKey = FindMetadataKey(InName);
+	UPCGParamData* NewParams = FilterParamsByKey(EntryKey);
+
+	if (EntryKey != PCGInvalidEntryKey)
+	{
+		// NOTE: this relies on the fact that there will be only one entry
+		NewParams->NameMap.Add(InName, 0);
+	}
+
+	return NewParams;
+}
+
+UPCGParamData* UPCGParamData::FilterParamsByKey(int64 InKey) const
+{
+	UPCGParamData* NewParams = NewObject<UPCGParamData>(const_cast<UPCGParamData*>(this));
+
+	// Here instead of parenting the metadata, we will create a copy
+	// so that the only entry in the metadata (if any) will have the 0 key.
+	check(NewParams && NewParams->Metadata);
+
+	NewParams->Metadata->AddAttributes(Metadata);
+
+	if (InKey != PCGInvalidEntryKey)
+	{
+		PCGMetadataEntryKey OutKey = PCGInvalidEntryKey;
+		NewParams->Metadata->SetAttributes(InKey, Metadata, OutKey);
+	}
+
+	return NewParams;
+}
\ No newline at end of file
diff --git a/Engine/Plugins/Experimental/PCG/Source/PCG/Private/PCGParams.cpp b/Engine/Plugins/Experimental/PCG/Source/PCG/Private/PCGParams.cpp
deleted file mode 100644
diff --git a/Engine/Plugins/Experimental/PCG/Source/PCG/Public/Helpers/PCGSettingsHelpers.h b/Engine/Plugins/Experimental/PCG/Source/PCG/Public/Helpers/PCGSettingsHelpers.h
--- a/Engine/Plugins/Experimental/PCG/Source/PCG/Public/Helpers/PCGSettingsHelpers.h
+++ b/Engine/Plugins/Experimental/PCG/Source/PCG/Public/Helpers/PCGSettingsHelpers.h
@@ -2,14 +2,14 @@
 
 #pragma once
 
-#include "PCGParams.h"
+#include "PCGParamData.h"
 #include "Metadata/PCGMetadataAttributeTraits.h"
 #include "Metadata/PCGMetadataAttributeTpl.h"
 
 namespace PCGSettingsHelpers
 {
 	template<typename T, typename TEnableIf<!TIsEnumClass<T>::Value>::Type* = nullptr>
-	T GetValue(const FName& InName, const T& InValue, UPCGParams* InParams, PCGMetadataEntryKey InKey)
+	T GetValue(const FName& InName, const T& InValue, UPCGParamData* InParams, PCGMetadataEntryKey InKey)
 	{
 		if (InParams)
 		{
@@ -30,13 +30,13 @@ namespace PCGSettingsHelpers
 	}
 
 	template<typename T, typename TEnableIf<!TIsEnumClass<T>::Value>::Type* = nullptr>
-	T GetValue(const FName& InName, const T& InValue, UPCGParams* InParams)
+	T GetValue(const FName& InName, const T& InValue, UPCGParamData* InParams)
 	{
 		return GetValue(InName, InValue, InParams, 0);
 	}
 
 	template<typename T, typename TEnableIf<!TIsEnumClass<T>::Value>::Type* = nullptr>
-	T GetValue(const FName& InName, const T& InValue, UPCGParams* InParams, const FName& InParamName)
+	T GetValue(const FName& InName, const T& InValue, UPCGParamData* InParams, const FName& InParamName)
 	{
 		if (InParams && InParamName != NAME_None)
 		{
@@ -50,42 +50,42 @@ namespace PCGSettingsHelpers
 
 	/** Specialized versions for enums */
 	template<typename T, typename TEnableIf<TIsEnumClass<T>::Value>::Type* = nullptr>
-	T GetValue(const FName& InName, const T& InValue, UPCGParams* InParams, PCGMetadataEntryKey InKey)
+	T GetValue(const FName& InName, const T& InValue, UPCGParamData* InParams, PCGMetadataEntryKey InKey)
 	{
 		return static_cast<T>(GetValue(InName, static_cast<__underlying_type(T)>(InValue), InParams, InKey));
 	}
 
 	template<typename T, typename TEnableIf<TIsEnumClass<T>::Value>::Type* = nullptr>
-	T GetValue(const FName& InName, const T& InValue, UPCGParams* InParams)
+	T GetValue(const FName& InName, const T& InValue, UPCGParamData* InParams)
 	{
 		return static_cast<T>(GetValue(InName, static_cast<__underlying_type(T)>(InValue), InParams));
 	}
 
 	template<typename T, typename TEnableIf<TIsEnumClass<T>::Value>::Type* = nullptr>
-	T GetValue(const FName& InName, const T& InValue, UPCGParams* InParams, const FName& InParamName)
+	T GetValue(const FName& InName, const T& InValue, UPCGParamData* InParams, const FName& InParamName)
 	{
 		return static_cast<T>(GetValue(InName, static_cast<__underlying_type(T)>(InValue), InParams, InParamName));
 	}
 
 	/** Specialized version for names */
 	template<>
-	FORCEINLINE FName GetValue(const FName& InName, const FName& InValue, UPCGParams* InParams, PCGMetadataEntryKey InKey)
+	FORCEINLINE FName GetValue(const FName& InName, const FName& InValue, UPCGParamData* InParams, PCGMetadataEntryKey InKey)
 	{
 		return FName(GetValue(InName, InValue.ToString(), InParams, InKey));
 	}
 
 	template<>
-	FORCEINLINE FName GetValue(const FName& InName, const FName& InValue, UPCGParams* InParams)
+	FORCEINLINE FName GetValue(const FName& InName, const FName& InValue, UPCGParamData* InParams)
 	{
 		return FName(GetValue(InName, InValue.ToString(), InParams));
 	}
 
 	template<>
-	FORCEINLINE FName GetValue(const FName& InName, const FName& InValue, UPCGParams* InParams, const FName& InParamName)
+	FORCEINLINE FName GetValue(const FName& InName, const FName& InValue, UPCGParamData* InParams, const FName& InParamName)
 	{
 		return FName(GetValue(InName, InValue.ToString(), InParams, InParamName));
 	}
 
 	/** Sets data from the params to a given property, matched on a name basis */
-	void SetValue(UPCGParams* Params, UObject* Object, FProperty* Property);
+	void SetValue(UPCGParamData* Params, UObject* Object, FProperty* Property);
 }
\ No newline at end of file
diff --git a/Engine/Plugins/Experimental/PCG/Source/PCG/Public/PCGData.h b/Engine/Plugins/Experimental/PCG/Source/PCG/Public/PCGData.h
--- a/Engine/Plugins/Experimental/PCG/Source/PCG/Public/PCGData.h
+++ b/Engine/Plugins/Experimental/PCG/Source/PCG/Public/PCGData.h
@@ -10,7 +10,7 @@
 #include "PCGData.generated.h"
 
 class UPCGSettings;
-class UPCGParams;
+class UPCGParamData;
 
 /**
 * Base class for any "data" class in the PCG framework.
@@ -23,16 +23,6 @@ class PCG_API UPCGData : public UObject
 	GENERATED_BODY()
 };
 
-UENUM(BlueprintType)
-enum class EPCGDataUsage : uint8
-{
-	Input, // must be castable to concrete data
-	Exclusion,
-	Settings,
-	Override,
-	// Manual? Blocker?
-};
-
 USTRUCT(BlueprintType)
 struct PCG_API FPCGTaggedData
 {
@@ -41,9 +31,6 @@ struct PCG_API FPCGTaggedData
 	UPROPERTY(BlueprintReadWrite, EditAnywhere, Category = Data)
 	TObjectPtr<const UPCGData> Data = nullptr;
 
-	UPROPERTY(BlueprintReadWrite, EditAnywhere, Category = Data)
-	EPCGDataUsage Usage = EPCGDataUsage::Input;
-
 	UPROPERTY(BlueprintReadWrite, EditAnywhere, Category = Data)
 	TSet<FString> Tags;
 
@@ -58,9 +45,10 @@ struct PCG_API FPCGDataCollection
 
 	TArray<FPCGTaggedData> GetInputs() const;
 	TArray<FPCGTaggedData> GetTaggedInputs(const FString& InTag) const;
-	TArray<FPCGTaggedData> GetExclusions() const;
 	TArray<FPCGTaggedData> GetAllSettings() const;
-	UPCGParams* GetParams() const;
+	TArray<FPCGTaggedData> GetAllParams() const;
+	TArray<FPCGTaggedData> GetTaggedParams(const FString& InTag) const;
+	UPCGParamData* GetParams() const;
 
 	template<typename SettingsType>
 	const SettingsType* GetSettings() const;
@@ -81,7 +69,7 @@ template<typename SettingsType>
 inline const SettingsType* FPCGDataCollection::GetSettings() const
 {
 	const FPCGTaggedData* MatchingData = TaggedData.FindByPredicate([](const FPCGTaggedData& Data) {
-		return Data.Usage == EPCGDataUsage::Settings && Cast<const SettingsType>(Data.Data) != nullptr;
+		return Cast<const SettingsType>(Data.Data) != nullptr;
 		});
 
 	return MatchingData ? Cast<const SettingsType>(MatchingData->Data) : nullptr;
@@ -101,7 +89,10 @@ public:
 	static TArray<FPCGTaggedData> GetTaggedInputs(const FPCGDataCollection& InCollection, const FString& InTag);
 
 	UFUNCTION(BlueprintCallable, Category = Data)
-	static TArray<FPCGTaggedData> GetExclusions(const FPCGDataCollection& InCollection);
+	static TArray<FPCGTaggedData> GetParams(const FPCGDataCollection& InCollection);
+
+	UFUNCTION(BlueprintCallable, Category = Data)
+	static TArray<FPCGTaggedData> GetTaggedParams(const FPCGDataCollection& InCollection, const FString& InTag);
 
 	UFUNCTION(BlueprintCallable, Category = Data)
 	static TArray<FPCGTaggedData> GetAllSettings(const FPCGDataCollection& InCollection);
diff --git a/Engine/Plugins/Experimental/PCG/Source/PCG/Public/PCGParamData.h b/Engine/Plugins/Experimental/PCG/Source/PCG/Public/PCGParamData.h
new file mode 100644
--- /dev/null
+++ b/Engine/Plugins/Experimental/PCG/Source/PCG/Public/PCGParamData.h
@@ -0,0 +1,49 @@
+// Copyright Epic Games, Inc. All Rights Reserved.
+
+#pragma once
+
+#include "PCGData.h"
+#include "Metadata/PCGMetadata.h"
+
+#include "PCGParamData.generated.h"
+
+/**
+* Class to hold execution parameters that will be consumed in nodes of the graph
+*/
+UCLASS(BlueprintType, ClassGroup = (Procedural))
+class PCG_API UPCGParamData : public UPCGData
+{
+	GENERATED_BODY()
+
+public:
+	UPCGParamData(const FObjectInitializer& ObjectInitializer);
+
+	UFUNCTION(BlueprintCallable, Category = Metadata)
+	const UPCGMetadata* ConstMetadata() const { return Metadata; }
+
+	UFUNCTION(BlueprintCallable, Category = Metadata)
+	UPCGMetadata* MutableMetadata() { return Metadata; }
+
+	/** Returns the entry for the given name */
+	UFUNCTION(BlueprintCallable, Category = Params)
+	int64 FindMetadataKey(const FName& InName) const;
+
+	/** Creates an entry for the given name, if not already added */
+	UFUNCTION(BlueprintCallable, Category = Params)
+	int64 FindOrAddMetadataKey(const FName& InName);
+
+	/** Creates a new params that keeps only a given key/name */
+	UFUNCTION(BlueprintCallable, Category = Params)
+	UPCGParamData* FilterParamsByName(const FName& InName) const;
+
+	UFUNCTION(BlueprintCallable, Category = Params)
+	UPCGParamData* FilterParamsByKey(int64 InKey) const;
+
+	// Not accessible through blueprint to make sure the constness is preserved
+	UPROPERTY(VisibleAnywhere, Category = Metadata)
+	TObjectPtr<UPCGMetadata> Metadata = nullptr;
+
+protected:
+	UPROPERTY()
+	TMap<FName, int64> NameMap;
+};
\ No newline at end of file
diff --git a/Engine/Plugins/Experimental/PCG/Source/PCG/Public/PCGParams.h b/Engine/Plugins/Experimental/PCG/Source/PCG/Public/PCGParams.h
deleted file mode 100644
