Adapted the pcg editor node title to use specific settings for blueprints & subgraphs
Improved the naming for blueprints (no more _C)

diff --git a/Engine/Plugins/Experimental/PCG/Source/PCG/Private/Elements/PCGExecuteBlueprint.cpp b/Engine/Plugins/Experimental/PCG/Source/PCG/Private/Elements/PCGExecuteBlueprint.cpp
--- a/Engine/Plugins/Experimental/PCG/Source/PCG/Private/Elements/PCGExecuteBlueprint.cpp
+++ b/Engine/Plugins/Experimental/PCG/Source/PCG/Private/Elements/PCGExecuteBlueprint.cpp
@@ -245,7 +245,11 @@ void UPCGBlueprintSettings::GetTrackedActorTags(FPCGTagToSettingsMap& OutTagToSe
 
 FName UPCGBlueprintSettings::AdditionalTaskName() const
 {
+#if WITH_EDITOR
+	return BlueprintElementType && BlueprintElementType->ClassGeneratedBy ? BlueprintElementType->ClassGeneratedBy->GetFName() : Super::AdditionalTaskName();
+#else
 	return BlueprintElementType ? BlueprintElementType->GetFName() : Super::AdditionalTaskName();
+#endif	
 }
 
 FPCGElementPtr UPCGBlueprintSettings::CreateElement() const
diff --git a/Engine/Plugins/Experimental/PCG/Source/PCG/Private/PCGSubgraph.cpp b/Engine/Plugins/Experimental/PCG/Source/PCG/Private/PCGSubgraph.cpp
--- a/Engine/Plugins/Experimental/PCG/Source/PCG/Private/PCGSubgraph.cpp
+++ b/Engine/Plugins/Experimental/PCG/Source/PCG/Private/PCGSubgraph.cpp
@@ -89,6 +89,18 @@ UPCGNode* UPCGSubgraphSettings::CreateNode() const
 	return NewObject<UPCGSubgraphNode>();
 }
 
+FName UPCGSubgraphSettings::AdditionalTaskName() const
+{
+	if (UPCGGraph* TargetSubgraph = GetSubgraph())
+	{
+		return TargetSubgraph->GetFName();
+	}
+	else
+	{
+		return TEXT("Invalid subgraph");
+	}
+}
+
 #if WITH_EDITOR
 bool UPCGSubgraphSettings::IsStructuralProperty(const FName& InPropertyName) const
 {
diff --git a/Engine/Plugins/Experimental/PCG/Source/PCG/Public/PCGSubgraph.h b/Engine/Plugins/Experimental/PCG/Source/PCG/Public/PCGSubgraph.h
--- a/Engine/Plugins/Experimental/PCG/Source/PCG/Public/PCGSubgraph.h
+++ b/Engine/Plugins/Experimental/PCG/Source/PCG/Public/PCGSubgraph.h
@@ -60,6 +60,8 @@ public:
 	virtual FName GetDefaultNodeName() const override { return FName(TEXT("SubgraphNode")); }
 #endif
 
+	virtual FName AdditionalTaskName() const override;
+
 protected:
 	virtual FPCGElementPtr CreateElement() const override;
 	//~End UPCGSettings interface implementation
diff --git a/Engine/Plugins/Experimental/PCG/Source/PCGEditor/Private/PCGEditorGraphNode.cpp b/Engine/Plugins/Experimental/PCG/Source/PCGEditor/Private/PCGEditorGraphNode.cpp
--- a/Engine/Plugins/Experimental/PCG/Source/PCGEditor/Private/PCGEditorGraphNode.cpp
+++ b/Engine/Plugins/Experimental/PCG/Source/PCGEditor/Private/PCGEditorGraphNode.cpp
@@ -27,7 +27,24 @@ void UPCGEditorGraphNode::Construct(UPCGNode* InPCGNode, EPCGEditorGraphNodeType
 
 FText UPCGEditorGraphNode::GetNodeTitle(ENodeTitleType::Type TitleType) const
 {
-	return (PCGNode && PCGNode->DefaultSettings) ? FText::FromName(PCGNode->DefaultSettings->GetDefaultNodeName()) : FText::FromName(TEXT("Unnamed node"));
+	if (PCGNode && PCGNode->DefaultSettings)
+	{
+		FName PreciseName = PCGNode->DefaultSettings->AdditionalTaskName();
+		FName NodeName = PCGNode->DefaultSettings->GetDefaultNodeName();
+
+		if (PreciseName != NAME_None)
+		{
+			return FText::Format(FText::FromString(TEXT("{0} - {1}")), { FText::FromName(PreciseName), FText::FromName(NodeName) });
+		}
+		else
+		{
+			return FText::FromName(NodeName);
+		}
+	}
+	else
+	{
+		return FText::FromName(TEXT("Unnamed node"));
+	}
 }
 
 void UPCGEditorGraphNode::GetNodeContextMenuActions(UToolMenu* Menu, class UGraphNodeContextMenuContext* Context) const
