Made pcg node names shorter

#rnx

diff --git a/Engine/Plugins/Experimental/PCG/Source/PCGEditor/Private/PCGEditorGraphNode.cpp b/Engine/Plugins/Experimental/PCG/Source/PCGEditor/Private/PCGEditorGraphNode.cpp
--- a/Engine/Plugins/Experimental/PCG/Source/PCGEditor/Private/PCGEditorGraphNode.cpp
+++ b/Engine/Plugins/Experimental/PCG/Source/PCGEditor/Private/PCGEditorGraphNode.cpp
@@ -29,16 +29,13 @@ FText UPCGEditorGraphNode::GetNodeTitle(ENodeTitleType::Type TitleType) const
 {
 	if (PCGNode && PCGNode->DefaultSettings)
 	{
-		FName PreciseName = PCGNode->DefaultSettings->AdditionalTaskName();
-		FName NodeName = PCGNode->DefaultSettings->GetDefaultNodeName();
-
-		if (PreciseName != NAME_None)
+		if (PCGNode->DefaultSettings->AdditionalTaskName() != NAME_None)
 		{
-			return FText::Format(FText::FromString(TEXT("{0} - {1}")), { FText::FromName(PreciseName), FText::FromName(NodeName) });
+			return FText::FromName(PCGNode->DefaultSettings->AdditionalTaskName());
 		}
 		else
 		{
-			return FText::FromName(NodeName);
+			return FText::FromName(PCGNode->DefaultSettings->GetDefaultNodeName());
 		}
 	}
 	else
