Added edge in/out labelling to allow more precise control in the PCG graphs.
Added basic blueprint/subgraph edge label edition

diff --git a/Engine/Plugins/Experimental/PCG/Source/PCG/Private/Elements/PCGDifferenceElement.cpp b/Engine/Plugins/Experimental/PCG/Source/PCG/Private/Elements/PCGDifferenceElement.cpp
--- a/Engine/Plugins/Experimental/PCG/Source/PCG/Private/Elements/PCGDifferenceElement.cpp
+++ b/Engine/Plugins/Experimental/PCG/Source/PCG/Private/Elements/PCGDifferenceElement.cpp
@@ -1,8 +1,20 @@
 // Copyright Epic Games, Inc. All Rights Reserved.
 
 #include "Elements/PCGDifferenceElement.h"
+#include "Data/PCGUnionData.h"
 #include "Helpers/PCGSettingsHelpers.h"
 
+namespace PCGDifferenceConstants
+{
+	const FName SourceLabel = TEXT("Source");
+	const FName DifferencesLabel = TEXT("Differences");
+}
+
+TArray<FName> UPCGDifferenceSettings::InLabels() const
+{
+	return { PCGDifferenceConstants::SourceLabel, PCGDifferenceConstants::DifferencesLabel };
+}
+
 FPCGElementPtr UPCGDifferenceSettings::CreateElement() const
 {
 	return MakeShared<FPCGDifferenceElement>();
@@ -12,6 +24,96 @@ bool FPCGDifferenceElement::ExecuteInternal(FPCGContext* Context) const
 {
 	TRACE_CPUPROFILER_EVENT_SCOPE(FPCGDifferenceElement::Execute);
 
+	// Early-out for previous behavior (without labeled edges)
+	if (Context->Node && !Context->Node->IsInputPinConnected(PCGDifferenceConstants::SourceLabel))
+	{
+		LabellessProcessing(Context);
+		return true;
+	}
+
+	const UPCGDifferenceSettings* Settings = Context->GetInputSettings<UPCGDifferenceSettings>();
+	check(Settings);
+
+	TArray<FPCGTaggedData> Sources = Context->InputData.GetInputsByLabel(PCGDifferenceConstants::SourceLabel);
+	TArray<FPCGTaggedData> Differences = Context->InputData.GetInputsByLabel(PCGDifferenceConstants::DifferencesLabel);
+	TArray<FPCGTaggedData>& Outputs = Context->OutputData.TaggedData;
+
+	// Get only spatial data or build an union from the sources
+	const UPCGSpatialData* FirstSpatialData = nullptr;
+	UPCGUnionData* UnionData = nullptr;
+	int32 DifferenceTaggedDataIndex = -1;
+
+	// Start by either selecting a source or building a union
+	for (FPCGTaggedData& Source : Sources)
+	{
+		const UPCGSpatialData* SpatialData = Cast<UPCGSpatialData>(Source.Data);
+
+		// Non-spatial data, we're not going to touch
+		if (!SpatialData)
+		{
+			Outputs.Add(Source);
+		}
+		else if (!FirstSpatialData)
+		{
+			FirstSpatialData = SpatialData;
+			DifferenceTaggedDataIndex = Outputs.Num();
+			Outputs.Add(Source);
+		}
+		else
+		{
+			if (!UnionData)
+			{
+				UnionData = FirstSpatialData->UnionWith(SpatialData);
+				// TODO: expose union settings?
+
+				// Replace source by union
+				Outputs[DifferenceTaggedDataIndex].Data = UnionData;
+			}
+			else
+			{
+				UnionData->AddData(SpatialData);
+			}
+		}
+	}
+	
+	if (FirstSpatialData || UnionData)
+	{
+		check(DifferenceTaggedDataIndex >= 0);
+		// Then, depending on the presence of differences, we'll create a difference as needed
+		UPCGDifferenceData* DifferenceData = nullptr;
+		for (FPCGTaggedData& Difference : Differences)
+		{
+			const UPCGSpatialData* SpatialData = Cast<UPCGSpatialData>(Difference.Data);
+
+			if (SpatialData)
+			{
+				if (!DifferenceData)
+				{
+					DifferenceData = (UnionData ? UnionData : FirstSpatialData)->Subtract(SpatialData);
+					DifferenceData->SetDensityFunction(Settings->DensityFunction);
+#if WITH_EDITORONLY_DATA
+					DifferenceData->bKeepZeroDensityPoints = Settings->bKeepZeroDensityPoints;
+#endif
+
+					Outputs[DifferenceTaggedDataIndex].Data = DifferenceData;
+				}
+				else
+				{
+					DifferenceData->AddDifference(SpatialData);
+				}
+			}
+			else
+			{
+				// We are not propagating data from the differences if they don't contribute
+			}
+		}
+	}
+
+	return true;
+}
+
+void FPCGDifferenceElement::LabellessProcessing(FPCGContext* Context) const
+{
 	const UPCGDifferenceSettings* Settings = Context->GetInputSettings<UPCGDifferenceSettings>();
 	check(Settings);
 
@@ -76,6 +178,4 @@ bool FPCGDifferenceElement::ExecuteInternal(FPCGContext* Context) const
 
 	// Finally, pass-through settings
 	Outputs.Append(Context->InputData.GetAllSettings());
-
-	return true;
 }
\ No newline at end of file
diff --git a/Engine/Plugins/Experimental/PCG/Source/PCG/Private/Graph/PCGGraphCompiler.cpp b/Engine/Plugins/Experimental/PCG/Source/PCG/Private/Graph/PCGGraphCompiler.cpp
--- a/Engine/Plugins/Experimental/PCG/Source/PCG/Private/Graph/PCGGraphCompiler.cpp
+++ b/Engine/Plugins/Experimental/PCG/Source/PCG/Private/Graph/PCGGraphCompiler.cpp
@@ -2,6 +2,7 @@
 
 #include "PCGGraphCompiler.h"
 #include "PCGGraph.h"
+#include "PCGEdge.h"
 #include "PCGSubgraph.h"
 #include "PCGSettings.h" // Needed for trivial element
 
@@ -19,7 +20,7 @@ TArray<FPCGGraphTask> FPCGGraphCompiler::CompileGraph(UPCGGraph* InGraph, FPCGTa
 	// Prime the node queue with all nodes that have no inbound edges
 	for (const UPCGNode* Node : InGraph->GetNodes())
 	{
-		if (Node->InboundNodes.IsEmpty())
+		if (Node->InboundEdges.IsEmpty())
 		{
 			NodeQueue.Add(Node);
 		}
@@ -70,15 +71,16 @@ TArray<FPCGGraphTask> FPCGGraphCompiler::CompileGraph(UPCGGraph* InGraph, FPCGTa
 			PreTask.Node = Node;
 			PreTask.NodeId = PreId;
 
-			for (const UPCGNode* Inbound : Node->InboundNodes)
+			for (const UPCGEdge* InboundEdge : Node->InboundEdges)
 			{
-				PreTask.Inputs.Add(IdMapping[Inbound]);
+				check(InboundEdge);
+				PreTask.Inputs.Emplace(IdMapping[InboundEdge->InboundNode], InboundEdge->InboundLabel, InboundEdge->OutboundLabel);
 			}
 
 			// Add pre-task as input to subgraph input node task
 			if (InputNodeTask)
 			{
-				InputNodeTask->Inputs.Add(PreId);
+				InputNodeTask->Inputs.Emplace(PreId, NAME_None, NAME_None);
 			}
 
 			// Merge subgraph tasks into current tasks.
@@ -93,7 +95,7 @@ TArray<FPCGGraphTask> FPCGGraphCompiler::CompileGraph(UPCGGraph* InGraph, FPCGTa
 			// Add subgraph output node task as input to the post-task
 			if (OutputNodeTask)
 			{
-				PostTask.Inputs.Add(OutputNodeTask->NodeId);
+				PostTask.Inputs.Emplace(OutputNodeTask->NodeId, NAME_None, NAME_None);
 			}
 
 			IdMapping.Add(Node, PostId);
@@ -101,7 +103,7 @@ TArray<FPCGGraphTask> FPCGGraphCompiler::CompileGraph(UPCGGraph* InGraph, FPCGTa
 			// Prime any input-less nodes from the subgraph
 			for (const UPCGNode* SubNode : Subgraph->GetNodes())
 			{
-				if (SubNode->InboundNodes.IsEmpty())
+				if (SubNode->InboundEdges.IsEmpty())
 				{
 					NodeQueue.Add(SubNode);
 				}
@@ -114,15 +116,15 @@ TArray<FPCGGraphTask> FPCGGraphCompiler::CompileGraph(UPCGGraph* InGraph, FPCGTa
 			Task.Node = Node;
 			Task.NodeId = NodeId;
 
-			for (const UPCGNode* Inbound : Node->InboundNodes)
+			for (const UPCGEdge* InboundEdge : Node->InboundEdges)
 			{
-				if (FPCGTaskId* InboundId = IdMapping.Find(Inbound))
+				if (FPCGTaskId* InboundId = IdMapping.Find(InboundEdge->InboundNode))
 				{
-					Task.Inputs.Add(*InboundId);
+					Task.Inputs.Emplace(IdMapping[InboundEdge->InboundNode], InboundEdge->InboundLabel, InboundEdge->OutboundLabel);
 				}
 				else
 				{
-					UE_LOG(LogPCG, Error, TEXT("Inconsistent node linkage between node %s and node %s"), *Node->GetFName().ToString(), (Inbound ? *Inbound->GetFName().ToString() : TEXT("null")));
+					UE_LOG(LogPCG, Error, TEXT("Inconsistent node linkage between node %s and node %s"), *Node->GetFName().ToString(), (InboundEdge->InboundNode ? *InboundEdge->InboundNode->GetFName().ToString() : TEXT("null")));
 					return TArray<FPCGGraphTask>();
 				}
 			}
@@ -131,11 +133,14 @@ TArray<FPCGGraphTask> FPCGGraphCompiler::CompileGraph(UPCGGraph* InGraph, FPCGTa
 		}
 
 		// Push next ready nodes on the queue
-		for (const UPCGNode* Outbound : Node->OutboundNodes)
+		for (const UPCGEdge* OutboundEdge : Node->OutboundEdges)
 		{
+			const UPCGNode* Outbound = OutboundEdge->OutboundNode;
+
 			bool bAllPrerequisitesMet = true;
-			for (const UPCGNode* Inbound : Outbound->InboundNodes)
+			for (const UPCGEdge* InboundEdge : Outbound->InboundEdges)
 			{
+				const UPCGNode* Inbound = InboundEdge->InboundNode;
 				bAllPrerequisitesMet &= IdMapping.Contains(Inbound);
 			}
 
@@ -170,13 +175,13 @@ void FPCGGraphCompiler::Compile(UPCGGraph* InGraph)
 	// Store back the results in the cache if it's valid
 	if (!CompiledTasks.IsEmpty())
 	{
-		GraphToTaskMapLock.WriteLock();
-		if (!GraphToTaskMap.Contains(InGraph))
-		{
-			GraphToTaskMap.Add(InGraph, MoveTemp(CompiledTasks));
-		}
-		GraphToTaskMapLock.WriteUnlock();
+	GraphToTaskMapLock.WriteLock();
+	if (!GraphToTaskMap.Contains(InGraph))
+	{
+		GraphToTaskMap.Add(InGraph, MoveTemp(CompiledTasks));
 	}
+	GraphToTaskMapLock.WriteUnlock();
+}
 }
 
 TArray<FPCGGraphTask> FPCGGraphCompiler::GetCompiledTasks(UPCGGraph* InGraph, bool bIsTopGraph)
@@ -221,9 +226,9 @@ void FPCGGraphCompiler::OffsetNodeIds(TArray<FPCGGraphTask>& Tasks, FPCGTaskId O
 	{
 		Task.NodeId += Offset;
 
-		for (FPCGTaskId& Input : Task.Inputs)
+		for(FPCGGraphTaskInput& Input : Task.Inputs)
 		{
-			Input += Offset;
+			Input.TaskId += Offset;
 		}
 	}
 }
@@ -262,7 +267,7 @@ void FPCGGraphCompiler::CompileTopGraph(UPCGGraph* InGraph)
 		FPCGGraphTask& Task = CompiledTasks[TaskIndex];
 		if (Task.Inputs.IsEmpty())
 		{
-			Task.Inputs.Add(PreExecuteTaskId);
+			Task.Inputs.Emplace(PreExecuteTaskId, NAME_None, NAME_None);
 		}
 	}
 
@@ -282,10 +287,13 @@ void FPCGGraphCompiler::CompileTopGraph(UPCGGraph* InGraph)
 		const FPCGGraphTask& Task = CompiledTasks[TaskIndex];
 		if (!TasksWithSuccessors.Contains(Task.NodeId))
 		{
-			PostExecuteTask.Inputs.Add(Task.NodeId);
+			PostExecuteTask.Inputs.Emplace(Task.NodeId, NAME_None, NAME_None);
 		}
 
-		TasksWithSuccessors.Append(Task.Inputs);
+		for (const FPCGGraphTaskInput& Input : Task.Inputs)
+		{
+			TasksWithSuccessors.Add(Input.TaskId);
+		}
 	}
 
 	// Store back the results in the cache
diff --git a/Engine/Plugins/Experimental/PCG/Source/PCG/Private/Graph/PCGGraphExecutor.cpp b/Engine/Plugins/Experimental/PCG/Source/PCG/Private/Graph/PCGGraphExecutor.cpp
--- a/Engine/Plugins/Experimental/PCG/Source/PCG/Private/Graph/PCGGraphExecutor.cpp
+++ b/Engine/Plugins/Experimental/PCG/Source/PCG/Private/Graph/PCGGraphExecutor.cpp
@@ -74,7 +74,10 @@ FPCGTaskId FPCGGraphExecutor::Schedule(UPCGGraph* Graph, UPCGComponent* SourceCo
 		// Push task (not data) dependencies on the pre-execute task
 		// Note must be done after the offset ids, otherwise we'll break the dependencies
 		check(ScheduledTask.Tasks.Num() >= 2 && ScheduledTask.Tasks[ScheduledTask.Tasks.Num() - 2].Node == nullptr);
-		ScheduledTask.Tasks[ScheduledTask.Tasks.Num() - 2].Inputs.Append(ExternalDependencies);
+		for (FPCGTaskId ExternalDependency : ExternalDependencies)
+		{
+			ScheduledTask.Tasks[ScheduledTask.Tasks.Num() - 2].Inputs.Emplace(ExternalDependency, NAME_None, NAME_None);
+		}
 #endif
 
 		ScheduleLock.Unlock();
@@ -87,7 +90,10 @@ FPCGTaskId FPCGGraphExecutor::ScheduleGeneric(TFunction<bool()> InOperation, con
 {
 	// Build task & element to hold the operation to perform
 	FPCGGraphTask Task;
-	Task.Inputs.Append(TaskDependencies);
+	for (FPCGTaskId TaskDependency : TaskDependencies)
+	{
+		Task.Inputs.Emplace(TaskDependency, NAME_None, NAME_None);
+	}
 	Task.Element = MakeShared<FPCGGenericElement>(InOperation);
 
 	ScheduleLock.Lock();
@@ -255,9 +261,9 @@ void FPCGGraphExecutor::QueueReadyTasks(FPCGTaskId FinishedTaskHint)
 	for (int TaskIndex = Tasks.Num() - 1; TaskIndex >= 0; --TaskIndex)
 	{
 		bool bAllPrerequisitesMet = true;
-		for (const FPCGTaskId& Id : Tasks[TaskIndex].Inputs)
+		for(const FPCGGraphTaskInput& Input : Tasks[TaskIndex].Inputs)
 		{
-			bAllPrerequisitesMet &= OutputData.Contains(Id);
+			bAllPrerequisitesMet &= OutputData.Contains(Input.TaskId);
 		}
 
 		if (bAllPrerequisitesMet)
@@ -270,11 +276,30 @@ void FPCGGraphExecutor::QueueReadyTasks(FPCGTaskId FinishedTaskHint)
 
 void FPCGGraphExecutor::BuildTaskInput(const FPCGGraphTask& Task, FPCGDataCollection& TaskInput)
 {
-	for (const FPCGTaskId& Id : Task.Inputs)
+	for (const FPCGGraphTaskInput& Input : Task.Inputs)
 	{
-		check(OutputData.Contains(Id));
-		TaskInput.TaggedData.Append(OutputData[Id].TaggedData);
-		TaskInput.bCancelExecution |= OutputData[Id].bCancelExecution;
+		check(OutputData.Contains(Input.TaskId));
+		const FPCGDataCollection& InputCollection = OutputData[Input.TaskId];
+
+		TaskInput.bCancelExecution |= InputCollection.bCancelExecution;
+		
+		// Trivial case : take everything
+		if (Input.InboundLabel == NAME_None && Input.OutboundLabel == NAME_None)
+		{
+			TaskInput.TaggedData.Append(InputCollection.TaggedData);
+		}
+		// Non-trivial case: take all matches from OutputData[Id]
+		// And apply the out label on it
+		else
+		{
+			const int32 TaggedDataOffset = TaskInput.TaggedData.Num();
+			TaskInput.TaggedData.Append(InputCollection.GetInputsByLabel(Input.InboundLabel));
+
+			for (int32 TaggedDataIndex = TaggedDataOffset; TaggedDataIndex < TaskInput.TaggedData.Num(); ++TaggedDataIndex)
+			{
+				TaskInput.TaggedData[TaggedDataIndex].Label = Input.OutboundLabel;
+			}
+		}
 	}
 }
 
diff --git a/Engine/Plugins/Experimental/PCG/Source/PCG/Private/Graph/PCGGraphExecutor.h b/Engine/Plugins/Experimental/PCG/Source/PCG/Private/Graph/PCGGraphExecutor.h
--- a/Engine/Plugins/Experimental/PCG/Source/PCG/Private/Graph/PCGGraphExecutor.h
+++ b/Engine/Plugins/Experimental/PCG/Source/PCG/Private/Graph/PCGGraphExecutor.h
@@ -18,9 +18,23 @@ class UPCGNode;
 class UPCGComponent;
 class FPCGGraphCompiler;
 
+struct FPCGGraphTaskInput
+{
+	FPCGGraphTaskInput(FPCGTaskId InTaskId, const FName& InInboundLabel, const FName& InOutboundLabel)
+		: TaskId(InTaskId)
+		, InboundLabel(InInboundLabel)
+		, OutboundLabel(InOutboundLabel)		
+	{
+	}
+
+	FPCGTaskId TaskId;
+	FName InboundLabel;
+	FName OutboundLabel;
+};
+
 struct FPCGGraphTask
 {
-	TArray<FPCGTaskId> Inputs;
+	TArray<FPCGGraphTaskInput> Inputs;
 	//TArray<DataId> Outputs;
 	const UPCGNode* Node = nullptr;
 	UPCGComponent* SourceComponent = nullptr;
diff --git a/Engine/Plugins/Experimental/PCG/Source/PCG/Private/PCGData.cpp b/Engine/Plugins/Experimental/PCG/Source/PCG/Private/PCGData.cpp
--- a/Engine/Plugins/Experimental/PCG/Source/PCG/Private/PCGData.cpp
+++ b/Engine/Plugins/Experimental/PCG/Source/PCG/Private/PCGData.cpp
@@ -9,7 +9,8 @@ bool FPCGTaggedData::operator==(const FPCGTaggedData& Other) const
 {
 	return Data == Other.Data &&
 		Tags.Num() == Other.Tags.Num() &&
-		Tags.Includes(Other.Tags);
+		Tags.Includes(Other.Tags) &&
+		Label == Other.Label;
 }
 
 bool FPCGTaggedData::operator!=(const FPCGTaggedData& Other) const
@@ -24,6 +25,13 @@ TArray<FPCGTaggedData> FPCGDataCollection::GetInputs() const
 		});
 }
 
+TArray<FPCGTaggedData> FPCGDataCollection::GetInputsByLabel(const FName& InLabel) const
+{
+	return TaggedData.FilterByPredicate([&InLabel](const FPCGTaggedData& Data) {
+		return Data.Label == InLabel;
+		});
+}
+
 TArray<FPCGTaggedData> FPCGDataCollection::GetTaggedInputs(const FString& InTag) const
 {
 	return TaggedData.FilterByPredicate([&InTag](const FPCGTaggedData& Data) {
diff --git a/Engine/Plugins/Experimental/PCG/Source/PCG/Private/PCGEdge.cpp b/Engine/Plugins/Experimental/PCG/Source/PCG/Private/PCGEdge.cpp
new file mode 100644
--- /dev/null
+++ b/Engine/Plugins/Experimental/PCG/Source/PCG/Private/PCGEdge.cpp
@@ -0,0 +1,17 @@
+// Copyright Epic Games, Inc. All Rights Reserved.
+
+#include "PCGEdge.h"
+#include "PCGNode.h"
+
+void UPCGEdge::BreakEdge()
+{
+	if(InboundNode)
+	{
+		InboundNode->RemoveOutboundEdge(this);
+	}
+
+	if(OutboundNode)
+	{
+		OutboundNode->RemoveInboundEdge(this);
+	}
+}
\ No newline at end of file
diff --git a/Engine/Plugins/Experimental/PCG/Source/PCG/Private/PCGElement.cpp b/Engine/Plugins/Experimental/PCG/Source/PCG/Private/PCGElement.cpp
--- a/Engine/Plugins/Experimental/PCG/Source/PCG/Private/PCGElement.cpp
+++ b/Engine/Plugins/Experimental/PCG/Source/PCG/Private/PCGElement.cpp
@@ -86,6 +86,16 @@ bool IPCGElement::Execute(FPCGContext* Context) const
 
 			if (bDone)
 			{
+				// Cleanup any residual labels if the node isn't supposed to produce them
+				// TODO: this is a bit of a crutch, could be refactored out if we review the way we push tagged data
+				if (Settings && Settings->OutLabels().Num() == 0)
+				{
+					for (FPCGTaggedData& TaggedData : Context->OutputData.TaggedData)
+					{
+						TaggedData.Label = NAME_None;
+					}
+				}
+
 				if (IsCacheable(Settings) && Context->Cache)
 				{
 					Context->Cache->StoreInCache(this, Context->InputData, Settings, Context->OutputData);
diff --git a/Engine/Plugins/Experimental/PCG/Source/PCG/Private/PCGGraph.cpp b/Engine/Plugins/Experimental/PCG/Source/PCG/Private/PCGGraph.cpp
--- a/Engine/Plugins/Experimental/PCG/Source/PCG/Private/PCGGraph.cpp
+++ b/Engine/Plugins/Experimental/PCG/Source/PCG/Private/PCGGraph.cpp
@@ -1,6 +1,8 @@
 // Copyright Epic Games, Inc. All Rights Reserved.
 
 #include "PCGGraph.h"
+#include "PCGEdge.h"
+#include "PCGInputOutputSettings.h"
 #include "PCGSubgraph.h"
 #include "PCGSubsystem.h"
 
@@ -12,7 +14,9 @@ UPCGGraph::UPCGGraph(const FObjectInitializer& ObjectInitializer)
 	: Super(ObjectInitializer)
 {
 	InputNode = ObjectInitializer.CreateDefaultSubobject<UPCGNode>(this, TEXT("DefaultInputNode"));
+	InputNode->DefaultSettings = ObjectInitializer.CreateDefaultSubobject<UPCGGraphInputOutputSettings>(this, TEXT("DefaultInputNodeSettings"));
 	OutputNode = ObjectInitializer.CreateDefaultSubobject<UPCGNode>(this, TEXT("DefaultOutputNode"));
+	OutputNode->DefaultSettings = ObjectInitializer.CreateDefaultSubobject<UPCGGraphInputOutputSettings>(this, TEXT("DefaultOutputNodeSettings"));
 
 	// Note: default connection from input to output
 	// should be added when creating from scratch,
@@ -26,6 +30,22 @@ void UPCGGraph::PostLoad()
 	Super::PostLoad();
 
 #if WITH_EDITOR
+	// Deprecation
+	if (!Cast<UPCGGraphInputOutputSettings>(InputNode->DefaultSettings))
+	{
+		InputNode->DefaultSettings = NewObject<UPCGGraphInputOutputSettings>(this, TEXT("DefaultInputNodeSettings"));
+	}
+
+	if (!Cast<UPCGGraphInputOutputSettings>(OutputNode->DefaultSettings))
+	{
+		OutputNode->DefaultSettings = NewObject<UPCGGraphInputOutputSettings>(this, TEXT("DefaultOutputNodeSettings"));
+	}
+#endif
+
+#if WITH_EDITOR
+	InputNode->OnNodeSettingsChangedDelegate.AddUObject(this, &UPCGGraph::OnSettingsChanged);
+	OutputNode->OnNodeSettingsChangedDelegate.AddUObject(this, &UPCGGraph::OnSettingsChanged);
+
 	for (UPCGNode* Node : Nodes)
 	{
 		OnNodeAdded(Node);
@@ -41,6 +61,9 @@ void UPCGGraph::BeginDestroy()
 		OnNodeRemoved(Node);
 	}
 
+	OutputNode->OnNodeSettingsChangedDelegate.RemoveAll(this);
+	InputNode->OnNodeSettingsChangedDelegate.RemoveAll(this);
+
 	// Notify the compiler to remove this graph from its cache
 	UWorld* World = GEditor->GetEditorWorldContext().World();
 	UPCGSubsystem* PCGSubsystem = World ? World->GetSubsystem<UPCGSubsystem>() : nullptr;
@@ -60,8 +83,8 @@ void UPCGGraph::InitializeFromTemplate()
 
 	auto ResetDefaultNode = [](UPCGNode* InNode) {
 		check(InNode);
-		InNode->OutboundNodes.Reset();
-		InNode->InboundNodes.Reset();
+		InNode->OutboundEdges.Reset();
+		InNode->InboundEdges.Reset();
 
 		// Reset settings as well
 		InNode->SetDefaultSettings(NewObject<UPCGTrivialSettings>(InNode));
@@ -195,18 +218,40 @@ void UPCGGraph::OnNodeRemoved(UPCGNode* InNode)
 }
 
 UPCGNode* UPCGGraph::AddEdge(UPCGNode* From, UPCGNode* To)
+{
+	return AddLabeledEdge(From, NAME_None, To, NAME_None);
+}
+
+UPCGNode* UPCGGraph::AddLabeledEdge(UPCGNode* From, const FName& InboundLabel, UPCGNode* To, const FName& OutboundLabel)
 {
 	if (!From || !To)
 	{
-		// TODO: log error
+		UE_LOG(LogPCG, Error, TEXT("Invalid edge nodes"));
 		return To;
 	}
 
+	if (!From->HasOutLabel(InboundLabel))
+	{
+		UE_LOG(LogPCG, Error, TEXT("From node %s does not have the %s label"), *From->GetName(), *InboundLabel.ToString());
+	}
+
+	if (!To->HasInLabel(OutboundLabel))
+	{
+		UE_LOG(LogPCG, Error, TEXT("To node %s does not have the %s label"), *To->GetName(), *OutboundLabel.ToString());
+	}
+
 	Modify();
 
-	From->ConnectTo(To);
-	To->ConnectFrom(From);
+	// Create edge
+	UPCGEdge* Edge = NewObject<UPCGEdge>(this);
+	Edge->InboundLabel = InboundLabel;
+	Edge->InboundNode = From;
+	Edge->OutboundLabel = OutboundLabel;
+	Edge->OutboundNode = To;
 
+	From->OutboundEdges.Add(Edge);
+	To->InboundEdges.Add(Edge);
+	
 #if WITH_EDITOR
 	NotifyGraphChanged(/*bIsStructural=*/true);
 #endif
@@ -223,32 +268,47 @@ void UPCGGraph::RemoveNode(UPCGNode* InNode)
 {
 	check(InNode);
 
-	for (UPCGNode* Input : InNode->InboundNodes)
+	for (int32 EdgeIndex = InNode->InboundEdges.Num() - 1; EdgeIndex >= 0; --EdgeIndex)
 	{
-		Input->OutboundNodes.Remove(InNode);
+		UPCGEdge* Edge = InNode->InboundEdges[EdgeIndex];
+		Edge->BreakEdge();
 	}
-	for (UPCGNode* Output : InNode->OutboundNodes)
+
+	for (int32 EdgeIndex = InNode->OutboundEdges.Num() - 1; EdgeIndex >= 0; --EdgeIndex)
 	{
-		Output->InboundNodes.Remove(InNode);
+		UPCGEdge* Edge = InNode->OutboundEdges[EdgeIndex];
+		Edge->BreakEdge();
 	}
 
 	Nodes.Remove(InNode);
 	OnNodeRemoved(InNode);
 }
 
-void UPCGGraph::RemoveEdge(UPCGNode* From, UPCGNode* To)
+void UPCGGraph::RemoveEdge(UPCGNode* From, const FName& FromLabel, UPCGNode* To, const FName& ToLabel)
 {
 	if (!From || !To)
 	{
-		// TODO: log error
+		UE_LOG(LogPCG, Error, TEXT("Invalid from/to node in RemoveEdge"));
 		return;
 	}
 
-	From->RemoveConnection(To);
-	To->RemoveConnection(From);
+	bool bChanged = false;
+
+	if (UPCGEdge* Edge = From->GetOutboundEdge(FromLabel, To, ToLabel))
+	{
+		Edge->BreakEdge();
+		bChanged = true;
+	}
+	else
+	{
+		UE_LOG(LogPCG, Warning, TEXT("Edge does not exist with labels %s and %s"), *FromLabel.ToString(), *ToLabel.ToString());
+	}
 
 #if WITH_EDITOR
-	NotifyGraphChanged(/*bIsStructural=*/true);
+	if (bChanged)
+	{
+		NotifyGraphChanged(/*bIsStructural=*/true);
+	}
 #endif
 }
 
@@ -256,10 +316,9 @@ void UPCGGraph::RemoveInboundEdges(UPCGNode* InNode)
 {
 	check(InNode);
 
-	for (int32 i = InNode->InboundNodes.Num() - 1; i >= 0; --i)
+	for (int32 i = InNode->InboundEdges.Num() - 1; i >= 0; --i)
 	{
-		UPCGNode* InboundNode = InNode->InboundNodes[i];
-		RemoveEdge(InNode, InboundNode);
+		InNode->InboundEdges[i]->BreakEdge();
 	}
 }
 
@@ -267,10 +326,9 @@ void UPCGGraph::RemoveOutboundEdges(UPCGNode* InNode)
 {
 	check(InNode);
 
-	for (int32 i = InNode->OutboundNodes.Num() - 1; i >= 0; --i)
+	for (int32 i = InNode->OutboundEdges.Num() - 1; i >= 0; --i)
 	{
-		UPCGNode* OutboundNode = InNode->OutboundNodes[i];
-		RemoveEdge(InNode, OutboundNode);
+		InNode->OutboundEdges[i]->BreakEdge();
 	}
 }
 
diff --git a/Engine/Plugins/Experimental/PCG/Source/PCG/Private/PCGInputOutputSettings.cpp b/Engine/Plugins/Experimental/PCG/Source/PCG/Private/PCGInputOutputSettings.cpp
new file mode 100644
--- /dev/null
+++ b/Engine/Plugins/Experimental/PCG/Source/PCG/Private/PCGInputOutputSettings.cpp
@@ -0,0 +1,3 @@
+// Copyright Epic Games, Inc. All Rights Reserved.
+
+#include "PCGInputOutputSettings.h"
\ No newline at end of file
diff --git a/Engine/Plugins/Experimental/PCG/Source/PCG/Private/PCGNode.cpp b/Engine/Plugins/Experimental/PCG/Source/PCG/Private/PCGNode.cpp
--- a/Engine/Plugins/Experimental/PCG/Source/PCG/Private/PCGNode.cpp
+++ b/Engine/Plugins/Experimental/PCG/Source/PCG/Private/PCGNode.cpp
@@ -1,6 +1,7 @@
 // Copyright Epic Games, Inc. All Rights Reserved.
 
 #include "PCGNode.h"
+#include "PCGEdge.h"
 #include "PCGGraph.h"
 #include "PCGSettings.h"
 
@@ -20,6 +21,16 @@ void UPCGNode::PostLoad()
 		DefaultSettings->OnSettingsChangedDelegate.AddUObject(this, &UPCGNode::OnSettingsChanged);
 	}
 #endif
+
+	// Deprecation
+	for (TObjectPtr<UPCGNode> OutboundNode : OutboundNodes_DEPRECATED)
+	{
+		if (!HasEdgeTo(OutboundNode))
+		{
+			AddEdgeTo(NAME_None, OutboundNode, NAME_None);
+		}
+	}
+	OutboundNodes_DEPRECATED.Reset();
 }
 
 void UPCGNode::BeginDestroy()
@@ -39,12 +50,12 @@ UPCGGraph* UPCGNode::GetGraph() const
 	return Cast<UPCGGraph>(GetOuter());
 }
 
-UPCGNode* UPCGNode::AddEdgeTo(UPCGNode* To)
+UPCGNode* UPCGNode::AddEdgeTo(FName InboundName, UPCGNode* To, FName OutboundName)
 {
 	check(GetGraph());
 	if (UPCGGraph* Graph = GetGraph())
 	{
-		return Graph->AddEdge(this, To);
+		return Graph->AddLabeledEdge(this, InboundName, To, OutboundName);
 	}
 	else
 	{
@@ -52,6 +63,52 @@ UPCGNode* UPCGNode::AddEdgeTo(UPCGNode* To)
 	}
 }
 
+bool UPCGNode::HasInLabel(const FName& Label) const
+{
+	return DefaultSettings && DefaultSettings->HasInLabel(Label);
+}
+
+bool UPCGNode::HasOutLabel(const FName& Label) const
+{
+	return DefaultSettings && DefaultSettings->HasOutLabel(Label);
+}
+
+TArray<FName> UPCGNode::InLabels() const
+{
+	return DefaultSettings ? DefaultSettings->InLabels() : TArray<FName>();
+}
+
+TArray<FName> UPCGNode::OutLabels() const
+{
+	return DefaultSettings ? DefaultSettings->OutLabels() : TArray<FName>();
+}
+
+bool UPCGNode::IsInputPinConnected(const FName& Label) const
+{
+	for (const UPCGEdge* InboundEdge : InboundEdges)
+	{
+		if (InboundEdge->OutboundLabel == Label)
+		{
+			return true;
+		}
+	}
+
+	return false;
+}
+
+bool UPCGNode::IsOutputPinConnected(const FName& Label) const
+{
+	for (const UPCGEdge* OutboundEdge : OutboundEdges)
+	{
+		if (OutboundEdge->InboundLabel == Label)
+		{
+			return true;
+		}
+	}
+
+	return false;
+}
+
 void UPCGNode::SetDefaultSettings(TObjectPtr<UPCGSettings> InSettings)
 {
 #if WITH_EDITOR
@@ -72,20 +129,42 @@ void UPCGNode::SetDefaultSettings(TObjectPtr<UPCGSettings> InSettings)
 #endif
 }
 
-void UPCGNode::ConnectTo(UPCGNode* InSuccessor)
+UPCGEdge* UPCGNode::GetOutboundEdge(const FName& FromLabel, UPCGNode* To, const FName& ToLabel)
 {
-	OutboundNodes.AddUnique(InSuccessor);
+	for (UPCGEdge* OutboundEdge : OutboundEdges)
+	{
+		if (OutboundEdge->OutboundNode == To &&
+			OutboundEdge->InboundLabel == FromLabel &&
+			OutboundEdge->OutboundLabel == ToLabel)
+		{
+			return OutboundEdge;
+		}
+	}
+
+	return nullptr;
+}
+
+bool UPCGNode::HasEdgeTo(UPCGNode* InNode) const
+{
+	for (const UPCGEdge* OutboundEdge : OutboundEdges)
+	{
+		if (OutboundEdge->OutboundNode == InNode)
+		{
+			return true;
+		}
+	}
+
+	return false;
 }
 
-void UPCGNode::ConnectFrom(UPCGNode* InPredecessor)
+void UPCGNode::RemoveInboundEdge(UPCGEdge* InEdge)
 {
-	InboundNodes.AddUnique(InPredecessor);
+	InboundEdges.Remove(InEdge);
 }
 
-void UPCGNode::RemoveConnection(UPCGNode* InNode)
+void UPCGNode::RemoveOutboundEdge(UPCGEdge* InEdge)
 {
-	InboundNodes.Remove(InNode);
-	OutboundNodes.Remove(InNode);
+	OutboundEdges.Remove(InEdge);
 }
 
 #if WITH_EDITOR
diff --git a/Engine/Plugins/Experimental/PCG/Source/PCG/Private/PCGSettings.cpp b/Engine/Plugins/Experimental/PCG/Source/PCG/Private/PCGSettings.cpp
--- a/Engine/Plugins/Experimental/PCG/Source/PCG/Private/PCGSettings.cpp
+++ b/Engine/Plugins/Experimental/PCG/Source/PCG/Private/PCGSettings.cpp
@@ -37,6 +37,16 @@ bool UPCGSettings::operator==(const UPCGSettings& Other) const
 	}
 }
 
+bool UPCGSettings::HasInLabel(const FName& Label) const
+{
+	return Label == NAME_None || InLabels().Contains(Label);
+}
+
+bool UPCGSettings::HasOutLabel(const FName& Label) const
+{
+	return Label == NAME_None || OutLabels().Contains(Label);
+}
+
 FPCGElementPtr UPCGSettings::GetElement() const
 {
 	if (!CachedElement)
diff --git a/Engine/Plugins/Experimental/PCG/Source/PCG/Private/PCGSubgraph.cpp b/Engine/Plugins/Experimental/PCG/Source/PCG/Private/PCGSubgraph.cpp
--- a/Engine/Plugins/Experimental/PCG/Source/PCG/Private/PCGSubgraph.cpp
+++ b/Engine/Plugins/Experimental/PCG/Source/PCG/Private/PCGSubgraph.cpp
@@ -84,6 +84,54 @@ void UPCGBaseSubgraphSettings::OnSubgraphChanged(UPCGGraph* InGraph, bool bIsStr
 }
 #endif // WITH_EDITOR
 
+bool UPCGBaseSubgraphSettings::HasInLabel(const FName& Label) const
+{
+	if (UPCGGraph* Subgraph = GetSubgraph())
+	{
+		return Subgraph->GetInputNode()->HasInLabel(Label);
+	}
+	else
+	{
+		return false;
+	}
+}
+
+bool UPCGBaseSubgraphSettings::HasOutLabel(const FName& Label) const
+{
+	if (UPCGGraph* Subgraph = GetSubgraph())
+	{
+		return Subgraph->GetOutputNode()->HasOutLabel(Label);
+	}
+	else
+	{
+		return false;
+	}
+}
+
+TArray<FName> UPCGBaseSubgraphSettings::InLabels() const
+{
+	if (UPCGGraph* Subgraph = GetSubgraph())
+	{
+		return Subgraph->GetInputNode()->InLabels();
+	}
+	else
+	{
+		return TArray<FName>();
+	}
+}
+
+TArray<FName> UPCGBaseSubgraphSettings::OutLabels() const
+{
+	if (UPCGGraph* Subgraph = GetSubgraph())
+	{
+		return Subgraph->GetOutputNode()->OutLabels();
+	}
+	else
+	{
+		return TArray<FName>();
+	}
+}
+
 UPCGNode* UPCGSubgraphSettings::CreateNode() const
 {
 	return NewObject<UPCGSubgraphNode>();
diff --git a/Engine/Plugins/Experimental/PCG/Source/PCG/Public/Elements/PCGDifferenceElement.h b/Engine/Plugins/Experimental/PCG/Source/PCG/Public/Elements/PCGDifferenceElement.h
--- a/Engine/Plugins/Experimental/PCG/Source/PCG/Public/Elements/PCGDifferenceElement.h
+++ b/Engine/Plugins/Experimental/PCG/Source/PCG/Public/Elements/PCGDifferenceElement.h
@@ -20,6 +20,8 @@ public:
 	virtual FName GetDefaultNodeName() const override { return FName(TEXT("DifferenceNode")); }
 #endif
 
+	virtual TArray<FName> InLabels() const override;
+
 protected:
 	virtual FPCGElementPtr CreateElement() const override;
 	//~End UPCGSettings interface
@@ -38,4 +40,6 @@ class FPCGDifferenceElement : public FSimplePCGElement
 {
 protected:
 	virtual bool ExecuteInternal(FPCGContext* Context) const;
+
+	void LabellessProcessing(FPCGContext* Context) const;
 };
diff --git a/Engine/Plugins/Experimental/PCG/Source/PCG/Public/Elements/PCGExecuteBlueprint.h b/Engine/Plugins/Experimental/PCG/Source/PCG/Public/Elements/PCGExecuteBlueprint.h
--- a/Engine/Plugins/Experimental/PCG/Source/PCG/Public/Elements/PCGExecuteBlueprint.h
+++ b/Engine/Plugins/Experimental/PCG/Source/PCG/Public/Elements/PCGExecuteBlueprint.h
@@ -97,6 +97,10 @@ public:
 #endif
 
 	virtual FName AdditionalTaskName() const override;
+	virtual bool HasInLabel(const FName& Label) const override { return Label == NAME_None || InputPinLabels.Contains(Label); }
+	virtual bool HasOutLabel(const FName& Label) const override { return Label == NAME_None || OutputPinLabels.Contains(Label); }
+	virtual TArray<FName> InLabels() const override { return InputPinLabels.Array(); }
+	virtual TArray<FName> OutLabels() const override { return OutputPinLabels.Array(); }
 
 protected:
 	virtual FPCGElementPtr CreateElement() const override;
@@ -132,6 +136,12 @@ protected:
 	UPROPERTY(BlueprintReadWrite, EditAnywhere, Category = Settings)
 	bool bCreatesArtifacts = false;
 
+	UPROPERTY(EditAnywhere, BlueprintReadWrite, Category = "Settings|Input & Output")
+	TSet<FName> InputPinLabels;
+
+	UPROPERTY(EditAnywhere, BlueprintReadWrite, Category = "Settings|Input & Output")
+	TSet<FName> OutputPinLabels;
+
 protected:
 #if WITH_EDITOR
 	void OnBlueprintChanged(UBlueprint* InBlueprint);
diff --git a/Engine/Plugins/Experimental/PCG/Source/PCG/Public/PCGData.h b/Engine/Plugins/Experimental/PCG/Source/PCG/Public/PCGData.h
--- a/Engine/Plugins/Experimental/PCG/Source/PCG/Public/PCGData.h
+++ b/Engine/Plugins/Experimental/PCG/Source/PCG/Public/PCGData.h
@@ -34,6 +34,9 @@ struct PCG_API FPCGTaggedData
 	UPROPERTY(BlueprintReadWrite, EditAnywhere, Category = Data)
 	TSet<FString> Tags;
 
+	UPROPERTY(BlueprintReadWrite, EditAnywhere, Category = Data)
+	FName Label = NAME_None;
+
 	bool operator==(const FPCGTaggedData& Other) const;
 	bool operator!=(const FPCGTaggedData& Other) const;
 };
@@ -44,6 +47,7 @@ struct PCG_API FPCGDataCollection
 	GENERATED_BODY()
 
 	TArray<FPCGTaggedData> GetInputs() const;
+	TArray<FPCGTaggedData> GetInputsByLabel(const FName& InLabel) const;
 	TArray<FPCGTaggedData> GetTaggedInputs(const FString& InTag) const;
 	TArray<FPCGTaggedData> GetAllSettings() const;
 	TArray<FPCGTaggedData> GetAllParams() const;
diff --git a/Engine/Plugins/Experimental/PCG/Source/PCG/Public/PCGEdge.h b/Engine/Plugins/Experimental/PCG/Source/PCG/Public/PCGEdge.h
new file mode 100644
--- /dev/null
+++ b/Engine/Plugins/Experimental/PCG/Source/PCG/Public/PCGEdge.h
@@ -0,0 +1,30 @@
+// Copyright Epic Games, Inc. All Rights Reserved.
+
+#pragma once
+
+#include "CoreMinimal.h"
+
+#include "PCGEdge.generated.h"
+
+class UPCGNode;
+
+UCLASS(ClassGroup = (Procedural))
+class PCG_API UPCGEdge : public UObject
+{
+	GENERATED_BODY()
+
+public:
+	UPROPERTY(BlueprintReadOnly, VisibleAnywhere, Category = Edge)
+	FName InboundLabel = NAME_None;
+
+	UPROPERTY(BlueprintReadOnly, VisibleAnywhere, Category = Edge)
+	TObjectPtr<UPCGNode> InboundNode;
+
+	UPROPERTY(BlueprintReadOnly, VisibleAnywhere, Category = Edge)
+	FName OutboundLabel = NAME_None;
+
+	UPROPERTY(BlueprintReadOnly, VisibleAnywhere, Category = Edge)
+	TObjectPtr<UPCGNode> OutboundNode;
+
+	void BreakEdge();
+};
\ No newline at end of file
diff --git a/Engine/Plugins/Experimental/PCG/Source/PCG/Public/PCGGraph.h b/Engine/Plugins/Experimental/PCG/Source/PCG/Public/PCGGraph.h
--- a/Engine/Plugins/Experimental/PCG/Source/PCG/Public/PCGGraph.h
+++ b/Engine/Plugins/Experimental/PCG/Source/PCG/Public/PCGGraph.h
@@ -47,6 +47,9 @@ public:
 	UFUNCTION(BlueprintCallable, Category = Graph)
 	UPCGNode* AddEdge(UPCGNode* From, UPCGNode* To);
 
+	UFUNCTION(BlueprintCallable, Category = Graph)
+	UPCGNode* AddLabeledEdge(UPCGNode* From, const FName& InboundLabel, UPCGNode* To, const FName& OutboundLabel);
+
 	/** Returns the graph input node */
 	UFUNCTION(BlueprintCallable, Category = Graph)
 	UPCGNode* GetInputNode() const { return InputNode; }
@@ -58,7 +61,7 @@ public:
 	bool Contains(UPCGNode* Node) const;
 	const TArray<UPCGNode*>& GetNodes() const { return Nodes; }
 	void RemoveNode(UPCGNode* InNode);
-	void RemoveEdge(UPCGNode* From, UPCGNode* To);
+	void RemoveEdge(UPCGNode* From, const FName& FromLabel, UPCGNode* To, const FName& ToLabel);
 	void RemoveInboundEdges(UPCGNode* InNode);
 	void RemoveOutboundEdges(UPCGNode* InNode);
 
diff --git a/Engine/Plugins/Experimental/PCG/Source/PCG/Public/PCGInputOutputSettings.h b/Engine/Plugins/Experimental/PCG/Source/PCG/Public/PCGInputOutputSettings.h
new file mode 100644
--- /dev/null
+++ b/Engine/Plugins/Experimental/PCG/Source/PCG/Public/PCGInputOutputSettings.h
@@ -0,0 +1,33 @@
+// Copyright Epic Games, Inc. All Rights Reserved.
+
+#pragma once
+
+#include "CoreMinimal.h"
+#include "PCGSettings.h"
+
+#include "PCGInputOutputSettings.generated.h"
+
+UCLASS(NotBlueprintable, Hidden, ClassGroup = (Procedural))
+class PCG_API UPCGGraphInputOutputSettings : public UPCGSettings
+{
+	GENERATED_BODY()
+
+public:
+	// ~Begin UPCGSettings interface
+#if WITH_EDITOR
+	virtual FName GetDefaultNodeName() const override { return FName(TEXT("InputNode")); }
+#endif
+
+	virtual bool HasInLabel(const FName& Label) const override { return Label == NAME_None || PinLabels.Contains(Label); }
+	virtual bool HasOutLabel(const FName& Label) const override { return Label == NAME_None || PinLabels.Contains(Label); }
+	virtual TArray<FName> InLabels() const override { return PinLabels.Array(); }
+	virtual TArray<FName> OutLabels() const override { return PinLabels.Array(); }
+
+protected:
+	virtual FPCGElementPtr CreateElement() const override { return MakeShared<FPCGTrivialElement>(); }
+	// ~End UPCGSettings interface
+
+protected:
+	UPROPERTY(EditAnywhere, BlueprintReadWrite, Category = "Settings|Input")
+	TSet<FName> PinLabels;
+};
\ No newline at end of file
diff --git a/Engine/Plugins/Experimental/PCG/Source/PCG/Public/PCGNode.h b/Engine/Plugins/Experimental/PCG/Source/PCG/Public/PCGNode.h
--- a/Engine/Plugins/Experimental/PCG/Source/PCG/Public/PCGNode.h
+++ b/Engine/Plugins/Experimental/PCG/Source/PCG/Public/PCGNode.h
@@ -8,6 +8,7 @@
 
 class UPCGSettings;
 class UPCGGraph;
+class UPCGEdge;
 class IPCGElement;
 
 #if WITH_EDITOR
@@ -20,6 +21,7 @@ class PCG_API UPCGNode : public UObject
 	GENERATED_BODY()
 
 	friend class UPCGGraph;
+	friend class UPCGEdge;
 	friend class FPCGGraphCompiler;
 
 public:
@@ -36,12 +38,30 @@ public:
 
 	/** Adds an edge in the owning graph to the given "To" node. */
 	UFUNCTION(BlueprintCallable, Category = Node)
-	UPCGNode* AddEdgeTo(UPCGNode* To);
+	UPCGNode* AddEdgeTo(FName InboundName, UPCGNode* To, FName OutboundName);
+
+	/** Returns true if one of the input pins matches with the given name */
+	bool HasInLabel(const FName& Label) const;
+
+	/** Returns true if one of the output pins matches with the given name */
+	bool HasOutLabel(const FName& Label) const;
+
+	/** Returns all the input labels */
+	TArray<FName> InLabels() const;
+
+	/** Returns all the output labels */
+	TArray<FName> OutLabels() const;
+
+	/** Returns true if the input pin is connected */
+	bool IsInputPinConnected(const FName& Label) const;
+
+	/** Returns true if the output pin is connected */
+	bool IsOutputPinConnected(const FName& Label) const;
 
 	/** Changes the default settings in the node */
 	void SetDefaultSettings(TObjectPtr<UPCGSettings> InSettings);
 
-	const TArray<TObjectPtr<UPCGNode>>& GetOutboundNodes() const { return OutboundNodes; }
+	const TArray<TObjectPtr<UPCGEdge>>& GetOutboundEdges() const { return OutboundEdges; }
 
 	/** Note: do not set this property directly from code, use SetDefaultSettings instead */
 	UPROPERTY(BlueprintReadWrite, EditAnywhere, Category = Node, meta=(EditInline))
@@ -60,10 +80,11 @@ public:
 #endif // WITH_EDITORONLY_DATA
 
 protected:
-	void ConnectTo(UPCGNode* InSuccessor);
-	void ConnectFrom(UPCGNode* InPredecessor);
+	bool HasEdgeTo(UPCGNode* InNode) const;
 
-	void RemoveConnection(UPCGNode* InNode);
+	void RemoveInboundEdge(UPCGEdge* InEdge);
+	void RemoveOutboundEdge(UPCGEdge* InEdge);
+	UPCGEdge* GetOutboundEdge(const FName& FromLabel, UPCGNode* To, const FName& ToLabel);
 
 #if WITH_EDITOR
 	virtual void PreEditChange(FProperty* PropertyAboutToChange) override;
@@ -71,11 +92,14 @@ protected:
 	void OnSettingsChanged(UPCGSettings* InSettings);
 #endif
 
+	UPROPERTY()
+	TArray<TObjectPtr<UPCGNode>> OutboundNodes_DEPRECATED;
+
 	UPROPERTY(BlueprintReadOnly, VisibleAnywhere, Category = Node)
-	TArray<TObjectPtr<UPCGNode>> InboundNodes;
+	TArray<TObjectPtr<UPCGEdge>> InboundEdges;
 
 	UPROPERTY(BlueprintReadOnly, VisibleAnywhere, Category = Node)
-	TArray<TObjectPtr<UPCGNode>> OutboundNodes;
+	TArray<TObjectPtr<UPCGEdge>> OutboundEdges;
 
 	// TODO: add this information:
 	// - Ability to run on non-game threads (here or element)
diff --git a/Engine/Plugins/Experimental/PCG/Source/PCG/Public/PCGSettings.h b/Engine/Plugins/Experimental/PCG/Source/PCG/Public/PCGSettings.h
--- a/Engine/Plugins/Experimental/PCG/Source/PCG/Public/PCGSettings.h
+++ b/Engine/Plugins/Experimental/PCG/Source/PCG/Public/PCGSettings.h
@@ -38,6 +38,11 @@ public:
 	// TODO: check if we need this to be virtual, we don't really need if we're always caching
 	/*virtual*/ FPCGElementPtr GetElement() const;
 	virtual UPCGNode* CreateNode() const;
+
+	virtual bool HasInLabel(const FName& Label) const;
+	virtual bool HasOutLabel(const FName& Label) const;
+	virtual TArray<FName> InLabels() const { return TArray<FName>(); }
+	virtual TArray<FName> OutLabels() const { return TArray<FName>(); }
 	
 	bool operator==(const UPCGSettings& Other) const;
 
diff --git a/Engine/Plugins/Experimental/PCG/Source/PCG/Public/PCGSubgraph.h b/Engine/Plugins/Experimental/PCG/Source/PCG/Public/PCGSubgraph.h
--- a/Engine/Plugins/Experimental/PCG/Source/PCG/Public/PCGSubgraph.h
+++ b/Engine/Plugins/Experimental/PCG/Source/PCG/Public/PCGSubgraph.h
@@ -31,11 +31,16 @@ protected:
 	//~End UObject interface implementation
 
 	//~Begin UPCGSettings interface
-#if WITH_EDITOR
 	virtual void GetTrackedActorTags(FPCGTagToSettingsMap& OutTagToSettings) const override;
 #endif
+
+	virtual bool HasInLabel(const FName& Label) const override;
+	virtual bool HasOutLabel(const FName& Label) const override;
+	virtual TArray<FName> InLabels() const override;
+	virtual TArray<FName> OutLabels() const override;
 	//~End UPCGSettings interface
 
+#if WITH_EDITOR
 	void OnSubgraphChanged(UPCGGraph* InGraph, bool bIsStructural);
 
 	virtual bool IsStructuralProperty(const FName& InPropertyName) const { return false; }
diff --git a/Engine/Plugins/Experimental/PCG/Source/PCGEditor/Private/PCGEditorGraph.cpp b/Engine/Plugins/Experimental/PCG/Source/PCGEditor/Private/PCGEditorGraph.cpp
--- a/Engine/Plugins/Experimental/PCG/Source/PCGEditor/Private/PCGEditorGraph.cpp
+++ b/Engine/Plugins/Experimental/PCG/Source/PCGEditor/Private/PCGEditorGraph.cpp
@@ -3,12 +3,14 @@
 #include "PCGEditorGraph.h"
 
 #include "PCGGraph.h"
+#include "PCGEdge.h"
 #include "PCGEditorGraphNode.h"
 
 #include "EdGraph/EdGraphPin.h"
 
 void UPCGEditorGraph::InitFromNodeGraph(UPCGGraph* InPCGGraph)
 {
+	check(InPCGGraph && !PCGGraph);
 	PCGGraph = InPCGGraph;
 
 	TMap<UPCGNode*, UPCGEditorGraphNode*> NodeLookup;
@@ -42,21 +44,26 @@ void UPCGEditorGraph::InitFromNodeGraph(UPCGGraph* InPCGGraph)
 		UPCGNode* PCGNode = NodeLookupIt.Key;
 		UPCGEditorGraphNode* GraphNode = NodeLookupIt.Value;
 
-		UEdGraphPin* OutPin = GraphNode->FindPin(TEXT("Out"));
-		if (!OutPin)
+		for (UPCGEdge* OutboundEdge : PCGNode->GetOutboundEdges())
 		{
-			continue;
-		}
+			const FName OutPinName = OutboundEdge->InboundLabel == NAME_None ? TEXT("Out") : OutboundEdge->InboundLabel;
+			UEdGraphPin* OutPin = GraphNode->FindPin(OutPinName);
 
-		for (UPCGNode* OutboundNode : PCGNode->GetOutboundNodes())
-		{
+			if (!OutPin)
+			{
+				continue;
+			}
+			
+			UPCGNode* OutboundNode = OutboundEdge->OutboundNode;
 			if (UPCGEditorGraphNode** ConnectedGraphNode = NodeLookup.Find(OutboundNode))
 			{
-				if (UEdGraphPin* InPin = (*ConnectedGraphNode)->FindPin(TEXT("In")))
+				const FName InPinName = OutboundEdge->OutboundLabel == NAME_None ? TEXT("In") : OutboundEdge->OutboundLabel;
+
+				if(UEdGraphPin* InPin = (*ConnectedGraphNode)->FindPin(InPinName))
 				{
 					OutPin->MakeLinkTo(InPin);
 				}
 			}
 		}
 	}
-}
+}
\ No newline at end of file
diff --git a/Engine/Plugins/Experimental/PCG/Source/PCGEditor/Private/PCGEditorGraphNode.cpp b/Engine/Plugins/Experimental/PCG/Source/PCGEditor/Private/PCGEditorGraphNode.cpp
--- a/Engine/Plugins/Experimental/PCG/Source/PCGEditor/Private/PCGEditorGraphNode.cpp
+++ b/Engine/Plugins/Experimental/PCG/Source/PCGEditor/Private/PCGEditorGraphNode.cpp
@@ -7,6 +7,7 @@
 #include "PCGSettings.h"
 
 #include "EdGraph/EdGraphPin.h"
+#include "EdGraph/EdGraph.h"
 #include "Framework/Commands/GenericCommands.h"
 #include "GraphEditorActions.h"
 #include "ToolMenu.h"
@@ -14,10 +15,21 @@
 
 #define LOCTEXT_NAMESPACE "PCGEditorGraphNode"
 
+void UPCGEditorGraphNode::BeginDestroy()
+{
+	if (PCGNode)
+	{
+		PCGNode->OnNodeSettingsChangedDelegate.RemoveAll(this);
+	}
+
+	Super::BeginDestroy();
+}
+
 void UPCGEditorGraphNode::Construct(UPCGNode* InPCGNode, EPCGEditorGraphNodeType InNodeType)
 {
 	check(InPCGNode);
 	PCGNode = InPCGNode;
+	InPCGNode->OnNodeSettingsChangedDelegate.AddUObject(this, &UPCGEditorGraphNode::OnNodeChanged);
 
 	NodePosX = InPCGNode->PositionX;
 	NodePosY = InPCGNode->PositionY;
@@ -90,11 +102,27 @@ void UPCGEditorGraphNode::AllocateDefaultPins()
 	if (NodeType == EPCGEditorGraphNodeType::Input || NodeType == EPCGEditorGraphNodeType::Settings)
 	{
 		CreatePin(EEdGraphPinDirection::EGPD_Output, NAME_None, FName(TEXT("Out")));
+
+		if (PCGNode)
+		{
+			for (const FName& OutLabel : PCGNode->OutLabels())
+			{
+				CreatePin(EEdGraphPinDirection::EGPD_Output, NAME_None, OutLabel);
+			}
+		}
 	}
 
 	if (NodeType == EPCGEditorGraphNodeType::Output || NodeType == EPCGEditorGraphNodeType::Settings)
 	{
 		CreatePin(EEdGraphPinDirection::EGPD_Input, NAME_None, FName(TEXT("In")));
+
+		if (PCGNode)
+		{
+			for (const FName& InLabel : PCGNode->InLabels())
+			{
+				CreatePin(EEdGraphPinDirection::EGPD_Input, NAME_None, InLabel);
+			}
+		}
 	}
 }
 
@@ -118,4 +146,103 @@ bool UPCGEditorGraphNode::CanUserDeleteNode() const
 	return NodeType != EPCGEditorGraphNodeType::Input && NodeType != EPCGEditorGraphNodeType::Output;
 }
 
+void UPCGEditorGraphNode::OnNodeChanged(UPCGNode* InNode)
+{
+	if (InNode == PCGNode)
+	{
+		ReconstructNode();
+	}
+}
+
+void UPCGEditorGraphNode::ReconstructNode()
+{
+	if (!PCGNode)
+	{
+		return;
+	}
+
+	TArray<FName> InLabels = PCGNode->InLabels();
+	TArray<FName> OutLabels = PCGNode->OutLabels();
+
+	TArray<UEdGraphPin*> InputPins;
+	TArray<UEdGraphPin*> OutputPins;
+
+	for (UEdGraphPin* Pin : Pins)
+	{
+		if (Pin->Direction == EEdGraphPinDirection::EGPD_Input && Pin->PinName != TEXT("In"))
+		{
+			InputPins.Add(Pin);
+		}
+		else if (Pin->Direction == EEdGraphPinDirection::EGPD_Output && Pin->PinName != TEXT("Out"))
+		{
+			OutputPins.Add(Pin);
+		}
+	}
+
+	auto UpdatePins = [this](TArray<UEdGraphPin*>& InPins, const TArray<FName>& InNames, EEdGraphPinDirection InPinDirection)
+	{
+		TArray<UEdGraphPin*> UnmatchedPins;
+
+		for (UEdGraphPin* InPin : InPins)
+		{
+			if (!InNames.Contains(InPin->PinName))
+			{
+				UnmatchedPins.Add(InPin);
+			}
+		}
+
+		TArray<FName> UnmatchedNames;
+
+		for (const FName& InName : InNames)
+		{
+			if (InPins.FindByPredicate([&InName](UEdGraphPin* InPin) { return InPin->PinName == InName; }) == nullptr)
+			{
+				UnmatchedNames.Add(InName);
+			}
+		}
+
+		// Use cases:
+		// - Removing pin(s) : remove pin(s) & break links
+		// - Remove all pins : remove all pins
+		// - Adding new pin(s) : create new pin(s)
+		// - Renaming a pin : find unmatched pin, unmatched name, rename
+		// - Complete change : match on names
+		if (UnmatchedNames.Num() == 1 && UnmatchedPins.Num() == 1)
+		{
+			UnmatchedPins[0]->PinName = UnmatchedNames[0];
+		}
+		else
+		{
+			for (UEdGraphPin* UnmatchedPin : UnmatchedPins)
+			{
+				UnmatchedPin->BreakAllPinLinks();
+				RemovePin(UnmatchedPin);
+			}
+
+			for (const FName& UnmatchedName : UnmatchedNames)
+			{
+				CreatePin(InPinDirection, NAME_None, UnmatchedName);
+			}
+		}
+
+		return UnmatchedNames.Num() > 0 || UnmatchedPins.Num() > 0;
+	};
+
+	bool bNeedsGraphUpdate = false;
+	if (NodeType == EPCGEditorGraphNodeType::Output || NodeType == EPCGEditorGraphNodeType::Settings)
+	{
+		bNeedsGraphUpdate |= UpdatePins(InputPins, InLabels, EEdGraphPinDirection::EGPD_Input);
+	}
+	
+	if (NodeType == EPCGEditorGraphNodeType::Input || NodeType == EPCGEditorGraphNodeType::Settings)
+	{
+		bNeedsGraphUpdate |= UpdatePins(OutputPins, OutLabels, EEdGraphPinDirection::EGPD_Output);
+	}
+
+	if (bNeedsGraphUpdate)
+	{
+		GetGraph()->NotifyGraphChanged();
+	}
+}
+
 #undef LOCTEXT_NAMESPACE
diff --git a/Engine/Plugins/Experimental/PCG/Source/PCGEditor/Private/PCGEditorGraphNode.h b/Engine/Plugins/Experimental/PCG/Source/PCGEditor/Private/PCGEditorGraphNode.h
--- a/Engine/Plugins/Experimental/PCG/Source/PCGEditor/Private/PCGEditorGraphNode.h
+++ b/Engine/Plugins/Experimental/PCG/Source/PCGEditor/Private/PCGEditorGraphNode.h
@@ -27,16 +27,24 @@ class UPCGEditorGraphNode : public UEdGraphNode
 public:
 	void Construct(UPCGNode* InPCGNode, EPCGEditorGraphNodeType InNodeType);
 
+	// ~Begin UObject interface
+	void BeginDestroy() override;
+	// ~End UObject interface
+
 	// ~Begin UEdGraphNode interface
 	virtual FText GetNodeTitle(ENodeTitleType::Type TitleType) const override;
 	virtual void GetNodeContextMenuActions(UToolMenu* Menu, class UGraphNodeContextMenuContext* Context) const override;
 	virtual void AllocateDefaultPins() override;
 	virtual void AutowireNewNode(UEdGraphPin* FromPin) override;
 	virtual bool CanUserDeleteNode() const override;
+	virtual void ReconstructNode() override;
 	// ~End UEdGraphNode interface
 
 	UPCGNode* GetPCGNode() { return PCGNode; }
 
+protected:
+	void OnNodeChanged(UPCGNode* InNode);
+
 private:
 	UPROPERTY()
 	TObjectPtr<UPCGNode> PCGNode = nullptr;
diff --git a/Engine/Plugins/Experimental/PCG/Source/PCGEditor/Private/PCGEditorGraphSchema.cpp b/Engine/Plugins/Experimental/PCG/Source/PCGEditor/Private/PCGEditorGraphSchema.cpp
--- a/Engine/Plugins/Experimental/PCG/Source/PCGEditor/Private/PCGEditorGraphSchema.cpp
+++ b/Engine/Plugins/Experimental/PCG/Source/PCGEditor/Private/PCGEditorGraphSchema.cpp
@@ -82,7 +82,10 @@ bool UPCGEditorGraphSchema::TryCreateConnection(UEdGraphPin* A, UEdGraphPin* B)
 		UPCGGraph* PCGGraph = PCGNodeA->GetGraph();
 		check(PCGGraph);
 
-		PCGGraph->AddEdge(PCGNodeA, PCGNodeB);
+		const FName& NodeAPinName = (A->PinName == TEXT("Out") ? NAME_None : A->PinName);
+		const FName& NodeBPinName = (B->PinName == TEXT("In") ? NAME_None : B->PinName);
+
+		PCGGraph->AddLabeledEdge(PCGNodeA, NodeAPinName, PCGNodeB, NodeBPinName);
 	}
 
 	return bModified;
@@ -126,8 +129,11 @@ void UPCGEditorGraphSchema::BreakSinglePinLink(UEdGraphPin* SourcePin, UEdGraphP
 	UPCGNode* TargetPCGNode = TargetPCGGraphNode->GetPCGNode();
 	check(SourcePCGNode && TargetPCGNode);
 
+	const FName& SourcePinName = (SourcePin->PinName == TEXT("Out") ? NAME_None : SourcePin->PinName);
+	const FName& TargetPinName = (TargetPin->PinName == TEXT("In") ? NAME_None : TargetPin->PinName);
+
 	UPCGGraph* PCGGraph = SourcePCGNode->GetGraph();
-	PCGGraph->RemoveEdge(SourcePCGNode, TargetPCGNode);
+	PCGGraph->RemoveEdge(SourcePCGNode, SourcePinName, TargetPCGNode, TargetPinName);
 }
 
 #undef LOCTEXT_NAMESPACE
