Added metadata rename operation for ease-of-use, supports both input data & params

diff --git a/Engine/Plugins/Experimental/PCG/Source/PCG/Private/Elements/PCGMetadataRenameElement.cpp b/Engine/Plugins/Experimental/PCG/Source/PCG/Private/Elements/PCGMetadataRenameElement.cpp
new file mode 100644
--- /dev/null
+++ b/Engine/Plugins/Experimental/PCG/Source/PCG/Private/Elements/PCGMetadataRenameElement.cpp
@@ -0,0 +1,96 @@
+// Copyright Epic Games, Inc. All Rights Reserved.
+
+#include "Elements/PCGMetadataRenameElement.h"
+#include "Data/PCGSpatialData.h"
+#include "Helpers/PCGSettingsHelpers.h"
+#include "Metadata/PCGMetadata.h"
+#include "Metadata/PCGMetadataAttribute.h"
+
+FPCGElementPtr UPCGMetadataRenameSettings::CreateElement() const
+{
+	return MakeShared<FPCGMetadataRenameElement>();
+}
+
+bool FPCGMetadataRenameElement::ExecuteInternal(FPCGContext* Context) const
+{
+	TRACE_CPUPROFILER_EVENT_SCOPE(FPCGMetadataOperationElement::Execute);
+
+	const UPCGMetadataRenameSettings* Settings = Context->GetInputSettings<UPCGMetadataRenameSettings>();
+	check(Settings);
+
+	TArray<FPCGTaggedData> Inputs = Context->InputData.GetInputs();
+	TArray<FPCGTaggedData> AllParams = Context->InputData.GetAllParams();
+	UPCGParamData* Params = Context->InputData.GetParams();
+
+	TArray<FPCGTaggedData>& Outputs = Context->OutputData.TaggedData;
+	
+	const FName AttributeToRename = PCGSettingsHelpers::GetValue(GET_MEMBER_NAME_CHECKED(UPCGMetadataRenameSettings, AttributeToRename), Settings->AttributeToRename, Params);
+	const FName NewAttributeName = PCGSettingsHelpers::GetValue(GET_MEMBER_NAME_CHECKED(UPCGMetadataRenameSettings, NewAttributeName), Settings->NewAttributeName, Params);
+
+	if (AttributeToRename == NAME_None || NewAttributeName == NAME_None)
+	{
+		UE_LOG(LogPCG, Warning, TEXT("Metadata rename operation cannot rename from %s to %s"), *AttributeToRename.ToString(), *NewAttributeName.ToString());
+		// Bypass
+		Context->OutputData = Context->InputData;
+		return true;
+	}
+
+	for (const FPCGTaggedData& Input : Inputs)
+	{
+		FPCGTaggedData& Output = Outputs.Add_GetRef(Input);
+
+		const UPCGSpatialData* SpatialInput = Cast<const UPCGSpatialData>(Input.Data);
+		if (!SpatialInput)
+		{
+			continue;
+		}
+
+		// If the data has a metadata & the attribute to rename, then duplicate the data
+		// otherwise, keep it as is.
+		const UPCGMetadata* Metadata = SpatialInput->Metadata;
+
+		if (!Metadata || !Metadata->HasAttribute(AttributeToRename))
+		{
+			continue;
+		}
+
+		//TODO: this might require to execute on the main thread
+		UPCGSpatialData* NewSpatialData = Cast<UPCGSpatialData>(StaticDuplicateObject(SpatialInput, const_cast<UPCGSpatialData*>(SpatialInput), FName()));
+		NewSpatialData->InitializeFromData(SpatialInput);
+		NewSpatialData->Metadata->RenameAttribute(AttributeToRename, NewAttributeName);
+
+		Output.Data = NewSpatialData;
+	}
+
+	for (const FPCGTaggedData& ParamTaggedData : AllParams)
+	{
+		FPCGTaggedData& Output = Outputs.Add_GetRef(ParamTaggedData);
+
+		const UPCGParamData* InputParams = Cast<const UPCGParamData>(ParamTaggedData.Data);
+		if (!InputParams)
+		{
+			continue;
+		}
+		
+		const UPCGMetadata* Metadata = InputParams->Metadata;
+
+		if (!Metadata || !Metadata->HasAttribute(AttributeToRename))
+		{
+			continue;
+		}
+
+		//TODO: this might require to execute on the main thread
+		UPCGParamData* NewParams = Cast<UPCGParamData>(StaticDuplicateObject(InputParams, const_cast<UPCGParamData*>(InputParams), FName()));
+		// Note: we will not parent the metadata here to ensure that the 0th entry is present on this metadata
+		// We will instead do a copy (and keep its parent, etc.)
+		NewParams->Metadata->InitializeAsCopy(Metadata);
+		NewParams->Metadata->RenameAttribute(AttributeToRename, NewAttributeName);
+
+		Output.Data = NewParams;
+	}
+
+	// Pass-through settings
+	Outputs.Append(Context->InputData.GetAllSettings());
+
+	return true;
+}
\ No newline at end of file
diff --git a/Engine/Plugins/Experimental/PCG/Source/PCG/Private/Metadata/PCGMetadata.cpp b/Engine/Plugins/Experimental/PCG/Source/PCG/Private/Metadata/PCGMetadata.cpp
--- a/Engine/Plugins/Experimental/PCG/Source/PCG/Private/Metadata/PCGMetadata.cpp
+++ b/Engine/Plugins/Experimental/PCG/Source/PCG/Private/Metadata/PCGMetadata.cpp
@@ -37,6 +37,27 @@ void UPCGMetadata::Initialize(const UPCGMetadata* InParent)
 	AddAttributes(InParent);
 }
 
+void UPCGMetadata::InitializeAsCopy(const UPCGMetadata* InMetadataToCopy)
+{
+	check(InMetadataToCopy);
+	if (Parent || Attributes.Num() != 0)
+	{
+		UE_LOG(LogPCG, Error, TEXT("Metadata has already been initialized or already contains attributes"));
+		return;
+	}
+
+	Parent = InMetadataToCopy->Parent;
+	OtherParents = InMetadataToCopy->OtherParents;
+	ParentKeys = InMetadataToCopy->ParentKeys;
+	ItemKeyOffset = InMetadataToCopy->ItemKeyOffset;
+
+	// Copy attributes
+	for (const TPair<FName, FPCGMetadataAttributeBase*>& OtherAttribute : InMetadataToCopy->Attributes)
+	{
+		CopyAttribute(OtherAttribute.Value, OtherAttribute.Key, /*bKeepParent=*/false, /*bCopyEntries=*/true, /*bCopyValues=*/true);
+	}
+}
+
 void UPCGMetadata::AddAttributes(const UPCGMetadata* InOther)
 {
 	if (!InOther)
diff --git a/Engine/Plugins/Experimental/PCG/Source/PCG/Public/Data/PCGProjectionData.h b/Engine/Plugins/Experimental/PCG/Source/PCG/Public/Data/PCGProjectionData.h
--- a/Engine/Plugins/Experimental/PCG/Source/PCG/Public/Data/PCGProjectionData.h
+++ b/Engine/Plugins/Experimental/PCG/Source/PCG/Public/Data/PCGProjectionData.h
@@ -42,6 +42,9 @@ protected:
 	UPROPERTY(BlueprintReadOnly, VisibleAnywhere, Category = SpatialData)
 	TObjectPtr<const UPCGSpatialData> Target = nullptr;
 
+	UPROPERTY()
 	FBox CachedBounds = FBox(EForceInit::ForceInit);
+
+	UPROPERTY()
 	FBox CachedStrictBounds = FBox(EForceInit::ForceInit);
 };
diff --git a/Engine/Plugins/Experimental/PCG/Source/PCG/Public/Elements/PCGMetadataRenameElement.h b/Engine/Plugins/Experimental/PCG/Source/PCG/Public/Elements/PCGMetadataRenameElement.h
new file mode 100644
--- /dev/null
+++ b/Engine/Plugins/Experimental/PCG/Source/PCG/Public/Elements/PCGMetadataRenameElement.h
@@ -0,0 +1,37 @@
+// Copyright Epic Games, Inc. All Rights Reserved.
+
+#pragma once
+
+#include "PCGSettings.h"
+#include "PCGElement.h"
+
+#include "PCGMetadataRenameElement.generated.h"
+
+UCLASS(BlueprintType, ClassGroup = (Procedural))
+class PCG_API UPCGMetadataRenameSettings : public UPCGSettings
+{
+	GENERATED_BODY()
+
+public:
+	//~Begin UPCGSettings interface
+#if WITH_EDITOR
+	virtual FName GetDefaultNodeName() const override { return FName(TEXT("MetadataRenameNode")); }
+#endif
+
+protected:
+	virtual FPCGElementPtr CreateElement() const override;
+	//~End UPCGSettings interface
+
+public:
+	UPROPERTY(BlueprintReadWrite, EditAnywhere, Category = Settings)
+	FName AttributeToRename = NAME_None;
+
+	UPROPERTY(BlueprintReadWrite, EditAnywhere, Category = Settings)
+	FName NewAttributeName = NAME_None;
+};
+
+class FPCGMetadataRenameElement : public FSimplePCGElement
+{
+protected:
+	virtual bool ExecuteInternal(FPCGContext* Context) const override;
+};
\ No newline at end of file
diff --git a/Engine/Plugins/Experimental/PCG/Source/PCG/Public/Metadata/PCGMetadata.h b/Engine/Plugins/Experimental/PCG/Source/PCG/Public/Metadata/PCGMetadata.h
--- a/Engine/Plugins/Experimental/PCG/Source/PCG/Public/Metadata/PCGMetadata.h
+++ b/Engine/Plugins/Experimental/PCG/Source/PCG/Public/Metadata/PCGMetadata.h
@@ -37,6 +37,9 @@ public:
 	UFUNCTION(BlueprintCallable, Category = "PCG|Metadata")
 	void Initialize(const UPCGMetadata* InParent);
 
+	UFUNCTION(BlueprintCallable, Category = "PCG|Metadata")
+	void InitializeAsCopy(const UPCGMetadata* InMetadataToCopy);
+
 	/** Creates streams for another metadata if they are not currently present */
 	void AddAttributes(const UPCGMetadata* InOther);
 
