Added node label refresh on settings change

diff --git a/Engine/Plugins/Experimental/PCG/Source/PCGEditor/Private/PCGEditorGraphNode.cpp b/Engine/Plugins/Experimental/PCG/Source/PCGEditor/Private/PCGEditorGraphNode.cpp
--- a/Engine/Plugins/Experimental/PCG/Source/PCGEditor/Private/PCGEditorGraphNode.cpp
+++ b/Engine/Plugins/Experimental/PCG/Source/PCGEditor/Private/PCGEditorGraphNode.cpp
@@ -228,21 +228,17 @@ void UPCGEditorGraphNode::ReconstructNode()
 		return UnmatchedNames.Num() > 0 || UnmatchedPins.Num() > 0;
 	};
 
-	bool bNeedsGraphUpdate = false;
 	if (NodeType == EPCGEditorGraphNodeType::Output || NodeType == EPCGEditorGraphNodeType::Settings)
 	{
-		bNeedsGraphUpdate |= UpdatePins(InputPins, InLabels, EEdGraphPinDirection::EGPD_Input);
+		UpdatePins(InputPins, InLabels, EEdGraphPinDirection::EGPD_Input);
 	}
 	
 	if (NodeType == EPCGEditorGraphNodeType::Input || NodeType == EPCGEditorGraphNodeType::Settings)
 	{
-		bNeedsGraphUpdate |= UpdatePins(OutputPins, OutLabels, EEdGraphPinDirection::EGPD_Output);
+		UpdatePins(OutputPins, OutLabels, EEdGraphPinDirection::EGPD_Output);
 	}
 
-	if (bNeedsGraphUpdate)
-	{
-		GetGraph()->NotifyGraphChanged();
-	}
+	OnNodeChangedDelegate.ExecuteIfBound();
 }
 
 #undef LOCTEXT_NAMESPACE
diff --git a/Engine/Plugins/Experimental/PCG/Source/PCGEditor/Private/PCGEditorGraphNode.h b/Engine/Plugins/Experimental/PCG/Source/PCGEditor/Private/PCGEditorGraphNode.h
--- a/Engine/Plugins/Experimental/PCG/Source/PCGEditor/Private/PCGEditorGraphNode.h
+++ b/Engine/Plugins/Experimental/PCG/Source/PCGEditor/Private/PCGEditorGraphNode.h
@@ -42,6 +42,9 @@ public:
 
 	UPCGNode* GetPCGNode() { return PCGNode; }
 
+	DECLARE_DELEGATE(FOnPCGEditorGraphNodeChanged);
+	FOnPCGEditorGraphNodeChanged OnNodeChangedDelegate;
+
 protected:
 	void OnNodeChanged(UPCGNode* InNode);
 
diff --git a/Engine/Plugins/Experimental/PCG/Source/PCGEditor/Private/SPCGEditorGraphNode.cpp b/Engine/Plugins/Experimental/PCG/Source/PCGEditor/Private/SPCGEditorGraphNode.cpp
--- a/Engine/Plugins/Experimental/PCG/Source/PCGEditor/Private/SPCGEditorGraphNode.cpp
+++ b/Engine/Plugins/Experimental/PCG/Source/PCGEditor/Private/SPCGEditorGraphNode.cpp
@@ -10,6 +10,11 @@ void SPCGEditorGraphNode::Construct(const FArguments& InArgs, UPCGEditorGraphNod
 	GraphNode = InNode;
 	PCGEditorGraphNode = InNode;
 
+	if (InNode)
+	{
+		InNode->OnNodeChangedDelegate.BindSP(this, &SPCGEditorGraphNode::OnNodeChanged);
+	}
+
 	UpdateGraphNode();
 }
 
@@ -24,3 +29,8 @@ void SPCGEditorGraphNode::MoveTo(const FVector2D& NewPosition, FNodeSet& NodeFil
 		PCGNode->MarkPackageDirty();
 	}
 }
+
+void SPCGEditorGraphNode::OnNodeChanged()
+{
+	UpdateGraphNode();
+}
diff --git a/Engine/Plugins/Experimental/PCG/Source/PCGEditor/Private/SPCGEditorGraphNode.h b/Engine/Plugins/Experimental/PCG/Source/PCGEditor/Private/SPCGEditorGraphNode.h
--- a/Engine/Plugins/Experimental/PCG/Source/PCGEditor/Private/SPCGEditorGraphNode.h
+++ b/Engine/Plugins/Experimental/PCG/Source/PCGEditor/Private/SPCGEditorGraphNode.h
@@ -19,6 +19,9 @@ public:
 	virtual void MoveTo(const FVector2D& NewPosition, FNodeSet& NodeFilter, bool bMarkDirty = true) override;
 	//~ End SGraphNode Interface
 
+protected:
+	void OnNodeChanged();
+
 private:
 	UPCGEditorGraphNode* PCGEditorGraphNode = nullptr;
 };
