Fixed issue when creating edges from right to left

diff --git a/Engine/Plugins/Experimental/PCG/Source/PCGEditor/Private/PCGEditorGraphSchema.cpp b/Engine/Plugins/Experimental/PCG/Source/PCGEditor/Private/PCGEditorGraphSchema.cpp
--- a/Engine/Plugins/Experimental/PCG/Source/PCGEditor/Private/PCGEditorGraphSchema.cpp
+++ b/Engine/Plugins/Experimental/PCG/Source/PCGEditor/Private/PCGEditorGraphSchema.cpp
@@ -63,12 +63,18 @@ const FPinConnectionResponse UPCGEditorGraphSchema::CanCreateConnection(const UE
 	return FPinConnectionResponse();
 }
 
-bool UPCGEditorGraphSchema::TryCreateConnection(UEdGraphPin* A, UEdGraphPin* B) const
+bool UPCGEditorGraphSchema::TryCreateConnection(UEdGraphPin* InA, UEdGraphPin* InB) const
 {
-	bool bModified = Super::TryCreateConnection(A, B);
+	bool bModified = Super::TryCreateConnection(InA, InB);
 
 	if (bModified)
 	{
+		check(InA && InB);
+		UEdGraphPin* A = (InA->Direction == EGPD_Output ? InA : InB);
+		UEdGraphPin* B = (InA->Direction == EGPD_Input ? InA : InB);
+		
+		check(A->Direction == EGPD_Output && B->Direction == EGPD_Input);
+
 		UEdGraphNode* NodeA = A->GetOwningNode();
 		UEdGraphNode* NodeB = B->GetOwningNode();
 
