Fixed issues with notifications in graph when doing an initialize from template

#rnx

diff --git a/Engine/Plugins/Experimental/PCG/Source/PCG/Private/PCGGraph.cpp b/Engine/Plugins/Experimental/PCG/Source/PCG/Private/PCGGraph.cpp
--- a/Engine/Plugins/Experimental/PCG/Source/PCG/Private/PCGGraph.cpp
+++ b/Engine/Plugins/Experimental/PCG/Source/PCG/Private/PCGGraph.cpp
@@ -81,6 +81,9 @@ void UPCGGraph::InitializeFromTemplate()
 {
 	Modify();
 
+	// Disable notification until the end of this method, otherwise it will cause issues in proper refresh dispatch
+	bEnableGraphChangeNotifications = false;
+
 	auto ResetDefaultNode = [](UPCGNode* InNode) {
 		check(InNode);
 		InNode->OutboundEdges.Reset();
@@ -102,11 +105,11 @@ void UPCGGraph::InitializeFromTemplate()
 
 	if (GraphTemplate)
 	{
-		bEnableGraphChangeNotifications = false;
 		Cast<UPCGGraphSetupBP>(GraphTemplate->GetDefaultObject())->Setup(this);
-		bEnableGraphChangeNotifications = true;
 	}
 
+	// Reenable notifications and notify listeners
+	bEnableGraphChangeNotifications = true;
 	NotifyGraphChanged(/*bIsStructural=*/true);
 }
 #endif
