Made the pins in the graph be driven by the blueprint selection
Added missing GetXYZByPin calls in BP library
Renamed label to pin in the data for consistency
PCG blueprint now reports when its blueprint does not compile
Fixed issue where the graph would not be properly refreshed when compiling a (new) blueprint that is used in a (new) node
Fixed static analysis warning

diff --git a/Engine/Plugins/Experimental/PCG/Source/PCG/Private/Elements/PCGDifferenceElement.cpp b/Engine/Plugins/Experimental/PCG/Source/PCG/Private/Elements/PCGDifferenceElement.cpp
--- a/Engine/Plugins/Experimental/PCG/Source/PCG/Private/Elements/PCGDifferenceElement.cpp
+++ b/Engine/Plugins/Experimental/PCG/Source/PCG/Private/Elements/PCGDifferenceElement.cpp
@@ -34,8 +34,8 @@ bool FPCGDifferenceElement::ExecuteInternal(FPCGContext* Context) const
 	const UPCGDifferenceSettings* Settings = Context->GetInputSettings<UPCGDifferenceSettings>();
 	check(Settings);
 
-	TArray<FPCGTaggedData> Sources = Context->InputData.GetInputsByLabel(PCGDifferenceConstants::SourceLabel);
-	TArray<FPCGTaggedData> Differences = Context->InputData.GetInputsByLabel(PCGDifferenceConstants::DifferencesLabel);
+	TArray<FPCGTaggedData> Sources = Context->InputData.GetInputsByPin(PCGDifferenceConstants::SourceLabel);
+	TArray<FPCGTaggedData> Differences = Context->InputData.GetInputsByPin(PCGDifferenceConstants::DifferencesLabel);
 	TArray<FPCGTaggedData>& Outputs = Context->OutputData.TaggedData;
 
 	// Get only spatial data or build an union from the sources
diff --git a/Engine/Plugins/Experimental/PCG/Source/PCG/Private/Elements/PCGExecuteBlueprint.cpp b/Engine/Plugins/Experimental/PCG/Source/PCG/Private/Elements/PCGExecuteBlueprint.cpp
--- a/Engine/Plugins/Experimental/PCG/Source/PCG/Private/Elements/PCGExecuteBlueprint.cpp
+++ b/Engine/Plugins/Experimental/PCG/Source/PCG/Private/Elements/PCGExecuteBlueprint.cpp
@@ -166,8 +166,27 @@ void UPCGBlueprintSettings::BeginDestroy()
 }
 
 #if WITH_EDITOR
+void UPCGBlueprintSettings::PreEditChange(FProperty* PropertyAboutToChange)
+{
+	if (PropertyAboutToChange)
+	{
+		if (PropertyAboutToChange->GetFName() == GET_MEMBER_NAME_CHECKED(UPCGBlueprintSettings, BlueprintElementType))
+		{
+			TeardownBlueprintEvent();
+		}
+	}
+}
+
 void UPCGBlueprintSettings::PostEditChangeProperty(FPropertyChangedEvent& PropertyChangedEvent)
 {
+	if (PropertyChangedEvent.Property)
+	{
+		if (PropertyChangedEvent.Property->GetFName() == GET_MEMBER_NAME_CHECKED(UPCGBlueprintSettings, BlueprintElementType))
+		{
+			SetupBlueprintEvent();
+		}
+	}
+
 	if (!BlueprintElementInstance || BlueprintElementInstance->GetClass() != BlueprintElementType)
 	{
 		RefreshBlueprintElement();
@@ -249,7 +268,27 @@ FName UPCGBlueprintSettings::AdditionalTaskName() const
 	return BlueprintElementType && BlueprintElementType->ClassGeneratedBy ? BlueprintElementType->ClassGeneratedBy->GetFName() : Super::AdditionalTaskName();
 #else
 	return BlueprintElementType ? BlueprintElementType->GetFName() : Super::AdditionalTaskName();
-#endif	
+#endif
+}
+
+bool UPCGBlueprintSettings::HasInLabel(const FName& Label) const
+{
+	return (Label == NAME_None || (BlueprintElementInstance && BlueprintElementInstance->InputPinLabels.Contains(Label)));
+}
+
+bool UPCGBlueprintSettings::HasOutLabel(const FName& Label) const
+{
+	return (Label == NAME_None || (BlueprintElementInstance && BlueprintElementInstance->OutputPinLabels.Contains(Label)));
+}
+
+TArray<FName> UPCGBlueprintSettings::InLabels() const
+{
+	return BlueprintElementInstance ? BlueprintElementInstance->InputPinLabels.Array() : TArray<FName>();
+}
+
+TArray<FName> UPCGBlueprintSettings::OutLabels() const
+{
+	return BlueprintElementInstance ? BlueprintElementInstance->OutputPinLabels.Array() : TArray<FName>();
 }
 
 FPCGElementPtr UPCGBlueprintSettings::CreateElement() const
@@ -261,25 +300,26 @@ bool FPCGExecuteBlueprintElement::ExecuteInternal(FPCGContext* InContext) const
 {
 	TRACE_CPUPROFILER_EVENT_SCOPE(FPCGExecuteBlueprintElement::Execute);
 	FPCGBlueprintExecutionContext* Context = static_cast<FPCGBlueprintExecutionContext*>(InContext);
-	const UPCGBlueprintSettings* Settings = Context->GetInputSettings<UPCGBlueprintSettings>();
 
 	if (Context && Context->BlueprintElementInstance)
 	{
-		// Log info on inputs
-		for(int32 InputIndex = 0; InputIndex < Context->InputData.TaggedData.Num(); ++InputIndex)
+		UClass* BPClass = Context->BlueprintElementInstance->GetClass();
+
+#if WITH_EDITOR
+		/** Check if the blueprint has been successfully compiled */
+		if (UBlueprint* Blueprint = Cast<UBlueprint>(BPClass ? BPClass->ClassGeneratedBy : nullptr))
 		{
-			const FPCGTaggedData& Input = Context->InputData.TaggedData[InputIndex];
-			if (const UPCGPointData* PointData = Cast<UPCGPointData>(Input.Data))
+			if (Blueprint->Status == BS_Error)
 			{
-				PCGE_LOG(Verbose, "Input %d has %d points", InputIndex, PointData->GetPoints().Num());
+				UE_LOG(LogPCG, Error, TEXT("PCG blueprint element cannot be executed since %s is not properly compiled"), *Blueprint->GetFName().ToString());
+				return true;
 			}
 		}
+#endif
 
 		/** Apply params overrides to variables if any */
 		if (UPCGParamData* Params = Context->InputData.GetParams())
 		{
-			UClass* BPClass = Context->BlueprintElementInstance->GetClass();
-
 			for (TFieldIterator<FProperty> PropertyIt(BPClass); PropertyIt; ++PropertyIt)
 			{
 				FProperty* Property = *PropertyIt;
@@ -293,6 +333,16 @@ bool FPCGExecuteBlueprintElement::ExecuteInternal(FPCGContext* InContext) const
 			}
 		}
 
+		// Log info on inputs
+		for (int32 InputIndex = 0; InputIndex < Context->InputData.TaggedData.Num(); ++InputIndex)
+		{
+			const FPCGTaggedData& Input = Context->InputData.TaggedData[InputIndex];
+			if (const UPCGPointData* PointData = Cast<UPCGPointData>(Input.Data))
+			{
+				PCGE_LOG(Verbose, "Input %d has %d points", InputIndex, PointData->GetPoints().Num());
+			}
+		}
+
 		/** Finally, execute the actual blueprint */
 		Context->BlueprintElementInstance->ExecuteWithContext(*Context, Context->InputData, Context->OutputData);
 
@@ -306,7 +356,7 @@ bool FPCGExecuteBlueprintElement::ExecuteInternal(FPCGContext* InContext) const
 			}
 		}
 	}
-	else
+	else if(Context)
 	{
 		// Nothing to do but forward data
 		Context->OutputData = Context->InputData;
diff --git a/Engine/Plugins/Experimental/PCG/Source/PCG/Private/Graph/PCGGraphExecutor.cpp b/Engine/Plugins/Experimental/PCG/Source/PCG/Private/Graph/PCGGraphExecutor.cpp
--- a/Engine/Plugins/Experimental/PCG/Source/PCG/Private/Graph/PCGGraphExecutor.cpp
+++ b/Engine/Plugins/Experimental/PCG/Source/PCG/Private/Graph/PCGGraphExecutor.cpp
@@ -293,11 +293,11 @@ void FPCGGraphExecutor::BuildTaskInput(const FPCGGraphTask& Task, FPCGDataCollec
 		else
 		{
 			const int32 TaggedDataOffset = TaskInput.TaggedData.Num();
-			TaskInput.TaggedData.Append(InputCollection.GetInputsByLabel(Input.InboundLabel));
+			TaskInput.TaggedData.Append(InputCollection.GetInputsByPin(Input.InboundLabel));
 
 			for (int32 TaggedDataIndex = TaggedDataOffset; TaggedDataIndex < TaskInput.TaggedData.Num(); ++TaggedDataIndex)
 			{
-				TaskInput.TaggedData[TaggedDataIndex].Label = Input.OutboundLabel;
+				TaskInput.TaggedData[TaggedDataIndex].Pin = Input.OutboundLabel;
 			}
 		}
 	}
diff --git a/Engine/Plugins/Experimental/PCG/Source/PCG/Private/PCGData.cpp b/Engine/Plugins/Experimental/PCG/Source/PCG/Private/PCGData.cpp
--- a/Engine/Plugins/Experimental/PCG/Source/PCG/Private/PCGData.cpp
+++ b/Engine/Plugins/Experimental/PCG/Source/PCG/Private/PCGData.cpp
@@ -10,7 +10,7 @@ bool FPCGTaggedData::operator==(const FPCGTaggedData& Other) const
 	return Data == Other.Data &&
 		Tags.Num() == Other.Tags.Num() &&
 		Tags.Includes(Other.Tags) &&
-		Label == Other.Label;
+		Pin == Other.Pin;
 }
 
 bool FPCGTaggedData::operator!=(const FPCGTaggedData& Other) const
@@ -25,10 +25,10 @@ TArray<FPCGTaggedData> FPCGDataCollection::GetInputs() const
 		});
 }
 
-TArray<FPCGTaggedData> FPCGDataCollection::GetInputsByLabel(const FName& InLabel) const
+TArray<FPCGTaggedData> FPCGDataCollection::GetInputsByPin(const FName& InPinLabel) const
 {
-	return TaggedData.FilterByPredicate([&InLabel](const FPCGTaggedData& Data) {
-		return Data.Label == InLabel;
+	return TaggedData.FilterByPredicate([&InPinLabel](const FPCGTaggedData& Data) {
+		return Data.Pin == InPinLabel;
 		});
 }
 
@@ -53,6 +53,13 @@ TArray<FPCGTaggedData> FPCGDataCollection::GetAllParams() const
 	});
 }
 
+TArray<FPCGTaggedData> FPCGDataCollection::GetParamsByPin(const FName& InPinLabel) const
+{
+	return TaggedData.FilterByPredicate([&InPinLabel](const FPCGTaggedData& Data) {
+		return Data.Pin == InPinLabel && Cast<UPCGParamData>(Data.Data);
+		});
+}
+
 TArray<FPCGTaggedData> FPCGDataCollection::GetTaggedParams(const FString& InTag) const
 {
 	return TaggedData.FilterByPredicate([&InTag](const FPCGTaggedData& Data) {
@@ -129,6 +136,11 @@ TArray<FPCGTaggedData> UPCGDataFunctionLibrary::GetInputs(const FPCGDataCollecti
 	return InCollection.GetInputs();
 }
 
+TArray<FPCGTaggedData> UPCGDataFunctionLibrary::GetInputsByPin(const FPCGDataCollection& InCollection, const FName& InPinLabel)
+{
+	return InCollection.GetInputsByPin(InPinLabel);
+}
+
 TArray<FPCGTaggedData> UPCGDataFunctionLibrary::GetTaggedInputs(const FPCGDataCollection& InCollection, const FString& InTag)
 {
 	return InCollection.GetTaggedInputs(InTag);
@@ -139,6 +151,11 @@ TArray<FPCGTaggedData> UPCGDataFunctionLibrary::GetParams(const FPCGDataCollecti
 	return InCollection.GetAllParams();
 }
 
+TArray<FPCGTaggedData> UPCGDataFunctionLibrary::GetParamsByPin(const FPCGDataCollection& InCollection, const FName& InPinLabel)
+{
+	return InCollection.GetParamsByPin(InPinLabel);
+}
+
 TArray<FPCGTaggedData> UPCGDataFunctionLibrary::GetTaggedParams(const FPCGDataCollection& InCollection, const FString& InTag)
 {
 	return InCollection.GetTaggedParams(InTag);
diff --git a/Engine/Plugins/Experimental/PCG/Source/PCG/Private/PCGElement.cpp b/Engine/Plugins/Experimental/PCG/Source/PCG/Private/PCGElement.cpp
--- a/Engine/Plugins/Experimental/PCG/Source/PCG/Private/PCGElement.cpp
+++ b/Engine/Plugins/Experimental/PCG/Source/PCG/Private/PCGElement.cpp
@@ -92,7 +92,7 @@ bool IPCGElement::Execute(FPCGContext* Context) const
 				{
 					for (FPCGTaggedData& TaggedData : Context->OutputData.TaggedData)
 					{
-						TaggedData.Label = NAME_None;
+						TaggedData.Pin = NAME_None;
 					}
 				}
 
diff --git a/Engine/Plugins/Experimental/PCG/Source/PCG/Public/Elements/PCGExecuteBlueprint.h b/Engine/Plugins/Experimental/PCG/Source/PCG/Public/Elements/PCGExecuteBlueprint.h
--- a/Engine/Plugins/Experimental/PCG/Source/PCG/Public/Elements/PCGExecuteBlueprint.h
+++ b/Engine/Plugins/Experimental/PCG/Source/PCG/Public/Elements/PCGExecuteBlueprint.h
@@ -75,6 +75,12 @@ public:
 	FOnPCGBlueprintChanged OnBlueprintChangedDelegate;
 #endif
 
+	UPROPERTY(EditAnywhere, BlueprintReadWrite, Category = "Settings|Input & Output")
+	TSet<FName> InputPinLabels;
+
+	UPROPERTY(EditAnywhere, BlueprintReadWrite, Category = "Settings|Input & Output")
+	TSet<FName> OutputPinLabels;
+
 protected:
 #if WITH_EDITOR
 	void OnDependencyChanged(UObject* Object, FPropertyChangedEvent& PropertyChangedEvent);
@@ -97,10 +103,10 @@ public:
 #endif
 
 	virtual FName AdditionalTaskName() const override;
-	virtual bool HasInLabel(const FName& Label) const override { return Label == NAME_None || InputPinLabels.Contains(Label); }
-	virtual bool HasOutLabel(const FName& Label) const override { return Label == NAME_None || OutputPinLabels.Contains(Label); }
-	virtual TArray<FName> InLabels() const override { return InputPinLabels.Array(); }
-	virtual TArray<FName> OutLabels() const override { return OutputPinLabels.Array(); }
+	virtual bool HasInLabel(const FName& Label) const override;
+	virtual bool HasOutLabel(const FName& Label) const override;
+	virtual TArray<FName> InLabels() const override;
+	virtual TArray<FName> OutLabels() const override;
 
 protected:
 	virtual FPCGElementPtr CreateElement() const override;
@@ -111,6 +117,7 @@ public:
 	virtual void PostLoad() override;
 	virtual void BeginDestroy() override;
 #if WITH_EDITOR
+	virtual void PreEditChange(FProperty* PropertyAboutToChange) override;
 	virtual void PostEditChangeProperty(FPropertyChangedEvent& PropertyChangedEvent) override;
 	// ~End UObject interface
 #endif
@@ -136,12 +143,6 @@ protected:
 	UPROPERTY(BlueprintReadWrite, EditAnywhere, Category = Settings)
 	bool bCreatesArtifacts = false;
 
-	UPROPERTY(EditAnywhere, BlueprintReadWrite, Category = "Settings|Input & Output")
-	TSet<FName> InputPinLabels;
-
-	UPROPERTY(EditAnywhere, BlueprintReadWrite, Category = "Settings|Input & Output")
-	TSet<FName> OutputPinLabels;
-
 protected:
 #if WITH_EDITOR
 	void OnBlueprintChanged(UBlueprint* InBlueprint);
diff --git a/Engine/Plugins/Experimental/PCG/Source/PCG/Public/PCGData.h b/Engine/Plugins/Experimental/PCG/Source/PCG/Public/PCGData.h
--- a/Engine/Plugins/Experimental/PCG/Source/PCG/Public/PCGData.h
+++ b/Engine/Plugins/Experimental/PCG/Source/PCG/Public/PCGData.h
@@ -35,7 +35,7 @@ struct PCG_API FPCGTaggedData
 	TSet<FString> Tags;
 
 	UPROPERTY(BlueprintReadWrite, EditAnywhere, Category = Data)
-	FName Label = NAME_None;
+	FName Pin = NAME_None;
 
 	bool operator==(const FPCGTaggedData& Other) const;
 	bool operator!=(const FPCGTaggedData& Other) const;
@@ -47,11 +47,12 @@ struct PCG_API FPCGDataCollection
 	GENERATED_BODY()
 
 	TArray<FPCGTaggedData> GetInputs() const;
-	TArray<FPCGTaggedData> GetInputsByLabel(const FName& InLabel) const;
+	TArray<FPCGTaggedData> GetInputsByPin(const FName& InPinLabel) const;
 	TArray<FPCGTaggedData> GetTaggedInputs(const FString& InTag) const;
 	TArray<FPCGTaggedData> GetAllSettings() const;
 	TArray<FPCGTaggedData> GetAllParams() const;
 	TArray<FPCGTaggedData> GetTaggedParams(const FString& InTag) const;
+	TArray<FPCGTaggedData> GetParamsByPin(const FName& InPinLabel) const;
 	UPCGParamData* GetParams() const;
 
 	template<typename SettingsType>
@@ -89,12 +90,18 @@ public:
 	UFUNCTION(BlueprintCallable, Category = Data)
 	static TArray<FPCGTaggedData> GetInputs(const FPCGDataCollection& InCollection);
 
+	UFUNCTION(BlueprintCallable, Category = Data)
+	static TArray<FPCGTaggedData> GetInputsByPin(const FPCGDataCollection& InCollection, const FName& InPinLabel);
+
 	UFUNCTION(BlueprintCallable, Category = Data)
 	static TArray<FPCGTaggedData> GetTaggedInputs(const FPCGDataCollection& InCollection, const FString& InTag);
 
 	UFUNCTION(BlueprintCallable, Category = Data)
 	static TArray<FPCGTaggedData> GetParams(const FPCGDataCollection& InCollection);
 
+	UFUNCTION(BlueprintCallable, Category = Data)
+	static TArray<FPCGTaggedData> GetParamsByPin(const FPCGDataCollection& InCollection, const FName& InPinLabel);
+
 	UFUNCTION(BlueprintCallable, Category = Data)
 	static TArray<FPCGTaggedData> GetTaggedParams(const FPCGDataCollection& InCollection, const FString& InTag);
 
