Added integers to supported metadata
Constify some code
Moved the property enum to its rightful place, added steepness too.

#rnx

diff --git a/Engine/Plugins/Experimental/PCG/Source/PCG/Private/Elements/PCGMetadataElement.cpp b/Engine/Plugins/Experimental/PCG/Source/PCG/Private/Elements/PCGMetadataElement.cpp
--- a/Engine/Plugins/Experimental/PCG/Source/PCG/Private/Elements/PCGMetadataElement.cpp
+++ b/Engine/Plugins/Experimental/PCG/Source/PCG/Private/Elements/PCGMetadataElement.cpp
@@ -230,6 +230,21 @@ bool FPCGMetadataOperationElement::ExecuteInternal(FPCGContext* Context) const
 					PCGE_LOG(Error, "Attribute %s already exists but its type is not compatible", *DestinationAttribute.ToString());
 				}
 			}
+			else if (PointProperty == EPCGPointProperties::Steepness)
+			{
+				auto SteepnessGetter = [](const FPCGPoint& InPoint) { return InPoint.Steepness; };
+
+				if (!SampledData->Metadata->HasAttribute(DestinationAttribute))
+				{
+					SampledData->Metadata->CreateFloatAttribute(DestinationAttribute, 0.5f, /*bAllowsInterpolation=*/true, /*bOverrideParent=*/true);
+				}
+
+				FPCGMetadataAttributeBase* AttributeBase = SampledData->Metadata->GetMutableAttribute(DestinationAttribute);
+				if (!PCGMetadataOperations::SetValue<float, float>(SampledPoints, AttributeBase, SampledData->Metadata, SteepnessGetter))
+				{
+					PCGE_LOG(Error, "Attribute %s already exists but its type is not compatible", *DestinationAttribute.ToString());
+				}
+			}
 		}
 		else if(Target == EPCGMetadataOperationTarget::AttributeToProperty) // Attribute to property
 		{
@@ -308,6 +323,16 @@ bool FPCGMetadataOperationElement::ExecuteInternal(FPCGContext* Context) const
 					PCGE_LOG(Error, "Attribute %s already exists but its type is not compatible", *SourceAttribute.ToString());
 				}
 			}
+			else if (PointProperty == EPCGPointProperties::Steepness)
+			{
+				auto SteepnessSetter = [](FPCGPoint& InPoint, const float& InValue) { InPoint.Steepness = InValue; };
+
+				const FPCGMetadataAttributeBase* AttributeBase = SampledData->Metadata->GetConstAttribute(SourceAttribute);
+				if (!PCGMetadataOperations::SetValue<float, float>(AttributeBase, SampledPoints, SteepnessSetter))
+				{
+					PCGE_LOG(Error, "Attribute %s already exists but its type is not compatible", *SourceAttribute.ToString());
+				}
+			}
 		}
 		else // Attribute to attribute
 		{
diff --git a/Engine/Plugins/Experimental/PCG/Source/PCG/Private/Metadata/PCGMetadata.cpp b/Engine/Plugins/Experimental/PCG/Source/PCG/Private/Metadata/PCGMetadata.cpp
--- a/Engine/Plugins/Experimental/PCG/Source/PCG/Private/Metadata/PCGMetadata.cpp
+++ b/Engine/Plugins/Experimental/PCG/Source/PCG/Private/Metadata/PCGMetadata.cpp
@@ -190,6 +190,11 @@ bool UPCGMetadata::ParentHasAttribute(FName AttributeName) const
 	return Parent && Parent->HasAttribute(AttributeName);
 }
 
+void UPCGMetadata::CreateInteger64Attribute(FName AttributeName, int64 DefaultValue, bool bAllowsInterpolation, bool bOverrideParent)
+{
+	CreateAttribute<int64>(AttributeName, DefaultValue, bAllowsInterpolation, bOverrideParent);
+}
+
 void UPCGMetadata::CreateFloatAttribute(FName AttributeName, float DefaultValue, bool bAllowsInterpolation, bool bOverrideParent)
 {
 	CreateAttribute<float>(AttributeName, DefaultValue, bAllowsInterpolation, bOverrideParent);
diff --git a/Engine/Plugins/Experimental/PCG/Source/PCG/Private/Metadata/PCGMetadataAccessor.cpp b/Engine/Plugins/Experimental/PCG/Source/PCG/Private/Metadata/PCGMetadataAccessor.cpp
--- a/Engine/Plugins/Experimental/PCG/Source/PCG/Private/Metadata/PCGMetadataAccessor.cpp
+++ b/Engine/Plugins/Experimental/PCG/Source/PCG/Private/Metadata/PCGMetadataAccessor.cpp
@@ -7,7 +7,7 @@
 
 /** Key-based implmentations */
 template<typename T>
-T UPCGMetadataAccessorHelpers::GetAttribute(PCGMetadataEntryKey Key, UPCGMetadata* Metadata, FName AttributeName)
+T UPCGMetadataAccessorHelpers::GetAttribute(PCGMetadataEntryKey Key, const UPCGMetadata* Metadata, FName AttributeName)
 {
 	if (!Metadata)
 	{
@@ -64,7 +64,7 @@ void UPCGMetadataAccessorHelpers::SetAttribute(PCGMetadataEntryKey& Key, UPCGMet
 	}
 }
 
-bool UPCGMetadataAccessorHelpers::HasAttributeSetByMetadataKey(PCGMetadataEntryKey Key, UPCGMetadata* Metadata, FName AttributeName)
+bool UPCGMetadataAccessorHelpers::HasAttributeSetByMetadataKey(PCGMetadataEntryKey Key, const UPCGMetadata* Metadata, FName AttributeName)
 {
 	if (!Metadata)
 	{
@@ -89,7 +89,17 @@ bool UPCGMetadataAccessorHelpers::HasAttributeSetByMetadataKey(PCGMetadataEntryK
 	}
 }
 
-float UPCGMetadataAccessorHelpers::GetFloatAttributeByMetadataKey(int64 Key, UPCGMetadata* Metadata, FName AttributeName)
+int64 UPCGMetadataAccessorHelpers::GetInteger64AttributeByMetadataKey(int64 Key, const UPCGMetadata* Metadata, FName AttributeName)
+{
+	return GetAttribute<int64>(Key, Metadata, AttributeName);
+}
+
+void UPCGMetadataAccessorHelpers::SetInteger64AttributeByMetadataKey(int64& Key, UPCGMetadata* Metadata, FName AttributeName, int64 Value)
+{
+	SetAttribute(Key, Metadata, AttributeName, Value);
+}
+
+float UPCGMetadataAccessorHelpers::GetFloatAttributeByMetadataKey(int64 Key, const UPCGMetadata* Metadata, FName AttributeName)
 {
 	return GetAttribute<float>(Key, Metadata, AttributeName);
 }
@@ -99,7 +109,7 @@ void UPCGMetadataAccessorHelpers::SetFloatAttributeByMetadataKey(int64& Key, UPC
 	SetAttribute(Key, Metadata, AttributeName, Value);
 }
 
-FVector UPCGMetadataAccessorHelpers::GetVectorAttributeByMetadataKey(int64 Key, UPCGMetadata* Metadata, FName AttributeName)
+FVector UPCGMetadataAccessorHelpers::GetVectorAttributeByMetadataKey(int64 Key, const UPCGMetadata* Metadata, FName AttributeName)
 {
 	return GetAttribute<FVector>(Key, Metadata, AttributeName);
 }
@@ -109,7 +119,7 @@ void UPCGMetadataAccessorHelpers::SetVectorAttributeByMetadataKey(int64& Key, UP
 	SetAttribute(Key, Metadata, AttributeName, Value);
 }
 
-FVector4 UPCGMetadataAccessorHelpers::GetVector4AttributeByMetadataKey(int64 Key, UPCGMetadata* Metadata, FName AttributeName)
+FVector4 UPCGMetadataAccessorHelpers::GetVector4AttributeByMetadataKey(int64 Key, const UPCGMetadata* Metadata, FName AttributeName)
 {
 	return GetAttribute<FVector4>(Key, Metadata, AttributeName);
 }
@@ -119,7 +129,7 @@ void UPCGMetadataAccessorHelpers::SetVector4AttributeByMetadataKey(int64& Key, U
 	SetAttribute(Key, Metadata, AttributeName, Value);
 }
 
-FQuat UPCGMetadataAccessorHelpers::GetQuatAttributeByMetadataKey(int64 Key, UPCGMetadata* Metadata, FName AttributeName)
+FQuat UPCGMetadataAccessorHelpers::GetQuatAttributeByMetadataKey(int64 Key, const UPCGMetadata* Metadata, FName AttributeName)
 {
 	return GetAttribute<FQuat>(Key, Metadata, AttributeName);
 }
@@ -129,7 +139,7 @@ void UPCGMetadataAccessorHelpers::SetQuatAttributeByMetadataKey(int64& Key, UPCG
 	SetAttribute(Key, Metadata, AttributeName, Value);
 }
 
-FTransform UPCGMetadataAccessorHelpers::GetTransformAttributeByMetadataKey(int64 Key, UPCGMetadata* Metadata, FName AttributeName)
+FTransform UPCGMetadataAccessorHelpers::GetTransformAttributeByMetadataKey(int64 Key, const UPCGMetadata* Metadata, FName AttributeName)
 {
 	return GetAttribute<FTransform>(Key, Metadata, AttributeName);
 }
@@ -139,7 +149,7 @@ void UPCGMetadataAccessorHelpers::SetTransformAttributeByMetadataKey(int64& Key,
 	SetAttribute(Key, Metadata, AttributeName, Value);
 }
 
-FString UPCGMetadataAccessorHelpers::GetStringAttributeByMetadataKey(int64 Key, UPCGMetadata* Metadata, FName AttributeName)
+FString UPCGMetadataAccessorHelpers::GetStringAttributeByMetadataKey(int64 Key, const UPCGMetadata* Metadata, FName AttributeName)
 {
 	return GetAttribute<FString>(Key, Metadata, AttributeName);
 }
@@ -186,7 +196,17 @@ void UPCGMetadataAccessorHelpers::InitializeMetadata(FPCGPoint& Point, UPCGMetad
 	Point.MetadataEntry = Metadata ? (Metadata->HasParent(ParentMetadata) ? Metadata->AddEntry(ParentPoint.MetadataEntry) : Metadata->AddEntry()) : PCGInvalidEntryKey;
 }
 
-float UPCGMetadataAccessorHelpers::GetFloatAttribute(const FPCGPoint& Point, UPCGMetadata* Metadata, FName AttributeName)
+int64 UPCGMetadataAccessorHelpers::GetInteger64Attribute(const FPCGPoint& Point, const UPCGMetadata* Metadata, FName AttributeName)
+{
+	return GetAttribute<int64>(Point.MetadataEntry, Metadata, AttributeName);
+}
+
+void UPCGMetadataAccessorHelpers::SetInteger64Attribute(FPCGPoint& Point, UPCGMetadata* Metadata, FName AttributeName, int64 Value)
+{
+	SetAttribute(Point.MetadataEntry, Metadata, AttributeName, Value);
+}
+
+float UPCGMetadataAccessorHelpers::GetFloatAttribute(const FPCGPoint& Point, const UPCGMetadata* Metadata, FName AttributeName)
 {
 	return GetAttribute<float>(Point.MetadataEntry, Metadata, AttributeName);
 }
@@ -196,7 +216,7 @@ void UPCGMetadataAccessorHelpers::SetFloatAttribute(FPCGPoint& Point, UPCGMetada
 	SetAttribute(Point.MetadataEntry, Metadata, AttributeName, Value);
 }
 
-FVector UPCGMetadataAccessorHelpers::GetVectorAttribute(const FPCGPoint& Point, UPCGMetadata* Metadata, FName AttributeName)
+FVector UPCGMetadataAccessorHelpers::GetVectorAttribute(const FPCGPoint& Point, const UPCGMetadata* Metadata, FName AttributeName)
 {
 	return GetAttribute<FVector>(Point.MetadataEntry, Metadata, AttributeName);
 }
@@ -206,7 +226,7 @@ void UPCGMetadataAccessorHelpers::SetVectorAttribute(FPCGPoint& Point, UPCGMetad
 	SetAttribute(Point.MetadataEntry, Metadata, AttributeName, Value);
 }
 
-FVector4 UPCGMetadataAccessorHelpers::GetVector4Attribute(const FPCGPoint& Point, UPCGMetadata* Metadata, FName AttributeName)
+FVector4 UPCGMetadataAccessorHelpers::GetVector4Attribute(const FPCGPoint& Point, const UPCGMetadata* Metadata, FName AttributeName)
 {
 	return GetAttribute<FVector4>(Point.MetadataEntry, Metadata, AttributeName);
 }
@@ -216,7 +236,7 @@ void UPCGMetadataAccessorHelpers::SetVector4Attribute(FPCGPoint& Point, UPCGMeta
 	SetAttribute(Point.MetadataEntry, Metadata, AttributeName, Value);
 }
 
-FQuat UPCGMetadataAccessorHelpers::GetQuatAttribute(const FPCGPoint& Point, UPCGMetadata* Metadata, FName AttributeName)
+FQuat UPCGMetadataAccessorHelpers::GetQuatAttribute(const FPCGPoint& Point, const UPCGMetadata* Metadata, FName AttributeName)
 {
 	return GetAttribute<FQuat>(Point.MetadataEntry, Metadata, AttributeName);
 }
@@ -226,7 +246,7 @@ void UPCGMetadataAccessorHelpers::SetQuatAttribute(FPCGPoint& Point, UPCGMetadat
 	SetAttribute(Point.MetadataEntry, Metadata, AttributeName, Value);
 }
 
-FTransform UPCGMetadataAccessorHelpers::GetTransformAttribute(const FPCGPoint& Point, UPCGMetadata* Metadata, FName AttributeName)
+FTransform UPCGMetadataAccessorHelpers::GetTransformAttribute(const FPCGPoint& Point, const UPCGMetadata* Metadata, FName AttributeName)
 {
 	return GetAttribute<FTransform>(Point.MetadataEntry, Metadata, AttributeName);
 }
@@ -236,7 +256,7 @@ void UPCGMetadataAccessorHelpers::SetTransformAttribute(FPCGPoint& Point, UPCGMe
 	SetAttribute(Point.MetadataEntry, Metadata, AttributeName, Value);
 }
 
-FString UPCGMetadataAccessorHelpers::GetStringAttribute(const FPCGPoint& Point, UPCGMetadata* Metadata, FName AttributeName)
+FString UPCGMetadataAccessorHelpers::GetStringAttribute(const FPCGPoint& Point, const UPCGMetadata* Metadata, FName AttributeName)
 {
 	return GetAttribute<FString>(Point.MetadataEntry, Metadata, AttributeName);
 }
@@ -246,7 +266,7 @@ void UPCGMetadataAccessorHelpers::SetStringAttribute(FPCGPoint& Point, UPCGMetad
 	SetAttribute(Point.MetadataEntry, Metadata, AttributeName, Value);
 }
 
-bool UPCGMetadataAccessorHelpers::HasAttributeSet(const FPCGPoint& Point, UPCGMetadata* Metadata, FName AttributeName)
+bool UPCGMetadataAccessorHelpers::HasAttributeSet(const FPCGPoint& Point, const UPCGMetadata* Metadata, FName AttributeName)
 {
 	return HasAttributeSetByMetadataKey(Point.MetadataEntry, Metadata, AttributeName);
 }
\ No newline at end of file
diff --git a/Engine/Plugins/Experimental/PCG/Source/PCG/Public/Elements/PCGMetadataElement.h b/Engine/Plugins/Experimental/PCG/Source/PCG/Public/Elements/PCGMetadataElement.h
--- a/Engine/Plugins/Experimental/PCG/Source/PCG/Public/Elements/PCGMetadataElement.h
+++ b/Engine/Plugins/Experimental/PCG/Source/PCG/Public/Elements/PCGMetadataElement.h
@@ -2,24 +2,12 @@
 
 #pragma once
 
-#include "CoreMinimal.h"
 #include "PCGSettings.h"
 #include "PCGElement.h"
+#include "PCGPoint.h"
 
 #include "PCGMetadataElement.generated.h"
 
-UENUM()
-enum class EPCGPointProperties : uint8
-{
-	Density,
-	Extents,
-	Color,
-	Position,
-	Rotation,
-	Scale,
-	Transform
-};
-
 UENUM()
 enum class EPCGMetadataOperationTarget : uint8
 {
diff --git a/Engine/Plugins/Experimental/PCG/Source/PCG/Public/Metadata/PCGMetadata.h b/Engine/Plugins/Experimental/PCG/Source/PCG/Public/Metadata/PCGMetadata.h
--- a/Engine/Plugins/Experimental/PCG/Source/PCG/Public/Metadata/PCGMetadata.h
+++ b/Engine/Plugins/Experimental/PCG/Source/PCG/Public/Metadata/PCGMetadata.h
@@ -49,6 +49,9 @@ public:
 	bool HasParent(const UPCGMetadata* InTentativeParent) const;
 
 	/** Create new streams */
+	UFUNCTION(BlueprintCallable, Category = "PCG|Metadata")
+	void CreateInteger64Attribute(FName AttributeName, int64 DefaultValue, bool bAllowsInterpolation, bool bOoverrideParent);
+
 	UFUNCTION(BlueprintCallable, Category = "PCG|Metadata")
 	void CreateFloatAttribute(FName AttributeName, float DefaultValue, bool bAllowsInterpolation, bool bOverrideParent);
 
diff --git a/Engine/Plugins/Experimental/PCG/Source/PCG/Public/Metadata/PCGMetadataAccessor.h b/Engine/Plugins/Experimental/PCG/Source/PCG/Public/Metadata/PCGMetadataAccessor.h
--- a/Engine/Plugins/Experimental/PCG/Source/PCG/Public/Metadata/PCGMetadataAccessor.h
+++ b/Engine/Plugins/Experimental/PCG/Source/PCG/Public/Metadata/PCGMetadataAccessor.h
@@ -19,43 +19,49 @@ class PCG_API UPCGMetadataAccessorHelpers : public UBlueprintFunctionLibrary
 public:
 	/** Id-based metadata functions */
 	UFUNCTION(BlueprintCallable, Category = "PCG|Metadata")
-	static float GetFloatAttributeByMetadataKey(int64 Key, UPCGMetadata* Metadata, FName AttributeName);
+	static int64 GetInteger64AttributeByMetadataKey(int64 Key, const UPCGMetadata* Metadata, FName AttributeName);
+
+	UFUNCTION(BlueprintCallable, Category = "PCG|Metadata")
+	static void SetInteger64AttributeByMetadataKey(UPARAM(ref) int64& Key, UPCGMetadata* Metadata, FName AttributeName, int64 Value);
+
+	UFUNCTION(BlueprintCallable, Category = "PCG|Metadata")
+	static float GetFloatAttributeByMetadataKey(int64 Key, const UPCGMetadata* Metadata, FName AttributeName);
 
 	UFUNCTION(BlueprintCallable, Category = "PCG|Metadata")
 	static void SetFloatAttributeByMetadataKey(UPARAM(ref) int64& Key, UPCGMetadata* Metadata, FName AttributeName, float Value);
 
 	UFUNCTION(BlueprintCallable, Category = "PCG|Metadata")
-	static FVector GetVectorAttributeByMetadataKey(int64 Key, UPCGMetadata* Metadata, FName AttributeName);
+	static FVector GetVectorAttributeByMetadataKey(int64 Key, const UPCGMetadata* Metadata, FName AttributeName);
 
 	UFUNCTION(BlueprintCallable, Category = "PCG|Metadata")
 	static void SetVectorAttributeByMetadataKey(UPARAM(ref) int64& Key, UPCGMetadata* Metadata, FName AttributeName, const FVector& Value);
 
 	UFUNCTION(BlueprintCallable, Category = "PCG|Metadata")
-	static FVector4 GetVector4AttributeByMetadataKey(int64 Key, UPCGMetadata* Metadata, FName AttributeName);
+	static FVector4 GetVector4AttributeByMetadataKey(int64 Key, const UPCGMetadata* Metadata, FName AttributeName);
 
 	UFUNCTION(BlueprintCallable, Category = "PCG|Metadata")
 	static void SetVector4AttributeByMetadataKey(UPARAM(ref) int64& Key, UPCGMetadata* Metadata, FName AttributeName, const FVector4& Value);
 
 	UFUNCTION(BlueprintCallable, Category = "PCG|Metadata")
-	static FQuat GetQuatAttributeByMetadataKey(int64 Key, UPCGMetadata* Metadata, FName AttributeName);
+	static FQuat GetQuatAttributeByMetadataKey(int64 Key, const UPCGMetadata* Metadata, FName AttributeName);
 
 	UFUNCTION(BlueprintCallable, Category = "PCG|Metadata")
 	static void SetQuatAttributeByMetadataKey(UPARAM(ref) int64& Key, UPCGMetadata* Metadata, FName AttributeName, const FQuat& Value);
 
 	UFUNCTION(BlueprintCallable, Category = "PCG|Metadata")
-	static FTransform GetTransformAttributeByMetadataKey(int64 Key, UPCGMetadata* Metadata, FName AttributeName);
+	static FTransform GetTransformAttributeByMetadataKey(int64 Key, const UPCGMetadata* Metadata, FName AttributeName);
 
 	UFUNCTION(BlueprintCallable, Category = "PCG|Metadata")
 	static void SetTransformAttributeByMetadataKey(UPARAM(ref) int64& Key, UPCGMetadata* Metadata, FName AttributeName, const FTransform& Value);
 
 	UFUNCTION(BlueprintCallable, Category = "PCG|Metadata")
-	static FString GetStringAttributeByMetadataKey(int64 Key, UPCGMetadata* Metadata, FName AttributeName);
+	static FString GetStringAttributeByMetadataKey(int64 Key, const UPCGMetadata* Metadata, FName AttributeName);
 
 	UFUNCTION(BlueprintCallable, Category = "PCG|Metadata")
 	static void SetStringAttributeByMetadataKey(UPARAM(ref) int64& Key, UPCGMetadata* Metadata, FName AttributeName, const FString& Value);
 
 	UFUNCTION(BlueprintCallable, Category = "PCG|Metadata")
-	static bool HasAttributeSetByMetadataKey(int64 Key, UPCGMetadata* Metadata, FName AttributeName);
+	static bool HasAttributeSetByMetadataKey(int64 Key, const UPCGMetadata* Metadata, FName AttributeName);
 
 	/** Point functions */
 	UFUNCTION(BlueprintCallable, Category = "PCG", meta = (ScriptMethod))
@@ -68,47 +74,53 @@ public:
 	static void InitializeMetadata(UPARAM(ref) FPCGPoint& Point, UPCGMetadata* Metadata, const FPCGPoint& ParentPoint);
 
 	UFUNCTION(BlueprintCallable, Category = "PCG|Metadata", meta = (ScriptMethod))
-	static float GetFloatAttribute(const FPCGPoint& Point, UPCGMetadata* Metadata, FName AttributeName);
+	static int64 GetInteger64Attribute(const FPCGPoint& Point, const UPCGMetadata* Metadata, FName AttributeName);
+
+	UFUNCTION(BlueprintCallable, Category = "PCG|Metadata", meta = (ScriptMethod))
+	static void SetInteger64Attribute(UPARAM(ref) FPCGPoint& Point, UPCGMetadata* Metadata, FName AttributeName, int64 Value);
+
+	UFUNCTION(BlueprintCallable, Category = "PCG|Metadata", meta = (ScriptMethod))
+	static float GetFloatAttribute(const FPCGPoint& Point, const UPCGMetadata* Metadata, FName AttributeName);
 
 	UFUNCTION(BlueprintCallable, Category = "PCG|Metadata", meta = (ScriptMethod))
 	static void SetFloatAttribute(UPARAM(ref) FPCGPoint& Point, UPCGMetadata* Metadata, FName AttributeName, float Value);
 
 	UFUNCTION(BlueprintCallable, Category = "PCG|Metadata", meta = (ScriptMethod))
-	static FVector GetVectorAttribute(const FPCGPoint& Point, UPCGMetadata* Metadata, FName AttributeName);
+	static FVector GetVectorAttribute(const FPCGPoint& Point, const UPCGMetadata* Metadata, FName AttributeName);
 
 	UFUNCTION(BlueprintCallable, Category = "PCG|Metadata", meta = (ScriptMethod))
 	static void SetVectorAttribute(UPARAM(ref) FPCGPoint& Point, UPCGMetadata* Metadata, FName AttributeName, const FVector& Value);
 
 	UFUNCTION(BlueprintCallable, Category = "PCG|Metadata", meta = (ScriptMethod))
-	static FVector4 GetVector4Attribute(const FPCGPoint& Point, UPCGMetadata* Metadata, FName AttributeName);
+	static FVector4 GetVector4Attribute(const FPCGPoint& Point, const UPCGMetadata* Metadata, FName AttributeName);
 
 	UFUNCTION(BlueprintCallable, Category = "PCG|Metadata", meta = (ScriptMethod))
 	static void SetVector4Attribute(UPARAM(ref) FPCGPoint& Point, UPCGMetadata* Metadata, FName AttributeName, const FVector4& Value);
 
 	UFUNCTION(BlueprintCallable, Category = "PCG|Metadata", meta = (ScriptMethod))
-	static FQuat GetQuatAttribute(const FPCGPoint& Point, UPCGMetadata* Metadata, FName AttributeName);
+	static FQuat GetQuatAttribute(const FPCGPoint& Point, const UPCGMetadata* Metadata, FName AttributeName);
 
 	UFUNCTION(BlueprintCallable, Category = "PCG|Metadata", meta = (ScriptMethod))
 	static void SetQuatAttribute(UPARAM(ref) FPCGPoint& Point, UPCGMetadata* Metadata, FName AttributeName, const FQuat& Value);
 
 	UFUNCTION(BlueprintCallable, Category = "PCG|Metadata", meta = (ScriptMethod))
-	static FTransform GetTransformAttribute(const FPCGPoint& Point, UPCGMetadata* Metadata, FName AttributeName);
+	static FTransform GetTransformAttribute(const FPCGPoint& Point, const UPCGMetadata* Metadata, FName AttributeName);
 
 	UFUNCTION(BlueprintCallable, Category = "PCG|Metadata", meta = (ScriptMethod))
 	static void SetTransformAttribute(UPARAM(ref) FPCGPoint& Point, UPCGMetadata* Metadata, FName AttributeName, const FTransform& Value);
 
 	UFUNCTION(BlueprintCallable, Category = "PCG|Metadata", meta = (ScriptMethod))
-	static FString GetStringAttribute(const FPCGPoint& Point, UPCGMetadata* Metadata, FName AttributeName);
+	static FString GetStringAttribute(const FPCGPoint& Point, const UPCGMetadata* Metadata, FName AttributeName);
 
 	UFUNCTION(BlueprintCallable, Category = "PCG|Metadata", meta = (ScriptMethod))
 	static void SetStringAttribute(UPARAM(ref) FPCGPoint& Point, UPCGMetadata* Metadata, FName AttributeName, const FString& Value);
 
 	UFUNCTION(BlueprintCallable, Category = "PCG|Metadata", meta = (ScriptMethod))
-	static bool HasAttributeSet(const FPCGPoint& Point, UPCGMetadata* Metadata, FName AttributeName);
+	static bool HasAttributeSet(const FPCGPoint& Point, const UPCGMetadata* Metadata, FName AttributeName);
 
 protected:
 	template<typename T>
-	static T GetAttribute(PCGMetadataEntryKey Key, UPCGMetadata* Metadata, FName AttributeName);
+	static T GetAttribute(PCGMetadataEntryKey Key, const UPCGMetadata* Metadata, FName AttributeName);
 
 	template<typename T>
 	static void SetAttribute(PCGMetadataEntryKey& Key, UPCGMetadata* Metadata, FName AttributeName, const T& Value);
diff --git a/Engine/Plugins/Experimental/PCG/Source/PCG/Public/PCGPoint.h b/Engine/Plugins/Experimental/PCG/Source/PCG/Public/PCGPoint.h
--- a/Engine/Plugins/Experimental/PCG/Source/PCG/Public/PCGPoint.h
+++ b/Engine/Plugins/Experimental/PCG/Source/PCG/Public/PCGPoint.h
@@ -6,6 +6,19 @@
 
 #include "PCGPoint.generated.h"
 
+UENUM()
+enum class EPCGPointProperties : uint8
+{
+	Density,
+	Extents,
+	Color,
+	Position,
+	Rotation,
+	Scale,
+	Transform,
+	Steepness
+};
+
 USTRUCT(BlueprintType)
 struct PCG_API FPCGPoint
 {
