Tentative fix for static analysis warning

#rnx

diff --git a/Engine/Plugins/Experimental/PCG/Source/PCGEditor/Private/PCGEditorGraph.cpp b/Engine/Plugins/Experimental/PCG/Source/PCGEditor/Private/PCGEditorGraph.cpp
--- a/Engine/Plugins/Experimental/PCG/Source/PCGEditor/Private/PCGEditorGraph.cpp
+++ b/Engine/Plugins/Experimental/PCG/Source/PCGEditor/Private/PCGEditorGraph.cpp
@@ -46,7 +46,7 @@ void UPCGEditorGraph::InitFromNodeGraph(UPCGGraph* InPCGGraph)
 
 		for (UPCGEdge* OutboundEdge : PCGNode->GetOutboundEdges())
 		{
-			const FName OutPinName = OutboundEdge->InboundLabel == NAME_None ? TEXT("Out") : OutboundEdge->InboundLabel;
+			const FName OutPinName = (OutboundEdge->InboundLabel == NAME_None ? TEXT("Out") : OutboundEdge->InboundLabel);
 			UEdGraphPin* OutPin = GraphNode->FindPin(OutPinName);
 
 			if (!OutPin)
@@ -57,7 +57,7 @@ void UPCGEditorGraph::InitFromNodeGraph(UPCGGraph* InPCGGraph)
 			UPCGNode* OutboundNode = OutboundEdge->OutboundNode;
 			if (UPCGEditorGraphNode** ConnectedGraphNode = NodeLookup.Find(OutboundNode))
 			{
-				const FName InPinName = OutboundEdge->OutboundLabel == NAME_None ? TEXT("In") : OutboundEdge->OutboundLabel;
+				const FName InPinName = (OutboundEdge->OutboundLabel == NAME_None ? TEXT("In") : OutboundEdge->OutboundLabel);
 
 				if(UEdGraphPin* InPin = (*ConnectedGraphNode)->FindPin(InPinName))
 				{
