Added point filtering operation. Support both properties & attributes, using either a constant, a param or data.
Fixed some issues in breaking links in the PCG editor.

diff --git a/Engine/Plugins/Experimental/PCG/Source/PCG/Private/Elements/PCGDebugElement.cpp b/Engine/Plugins/Experimental/PCG/Source/PCG/Private/Elements/PCGDebugElement.cpp
--- a/Engine/Plugins/Experimental/PCG/Source/PCG/Private/Elements/PCGDebugElement.cpp
+++ b/Engine/Plugins/Experimental/PCG/Source/PCG/Private/Elements/PCGDebugElement.cpp
@@ -41,10 +41,40 @@ namespace PCGDebugElement
 			Materials.Add(Material);
 		}
 
-		TArray<FPCGTaggedData> Inputs = Context->InputData.GetInputs();
+		// In the case of a node with multiple output pins, we will select only the inputs from the first non-empty pin.
+		bool bFilterOnPin = false;
+		FName PinFilter = NAME_None;
 
+		if (Context->Node)
+		{
+			if (Context->Node->IsOutputPinConnected(NAME_None))
+			{
+				bFilterOnPin = true;
+			}
+			else
+			{
+				TArray<FName> OutLabels = Context->Node->OutLabels();
+				for (const FName& OutLabel : OutLabels)
+				{
+					if (Context->Node->IsOutputPinConnected(OutLabel))
+					{
+						PinFilter = OutLabel;
+						bFilterOnPin = true;
+						break;
+					}
+				}
+			}
+		}
+
+		TArray<FPCGTaggedData> Inputs = Context->InputData.GetInputs();
 		for(const FPCGTaggedData& Input : Inputs)
 		{
+			// Skip output if we're filtering on the first pin
+			if (bFilterOnPin && Input.Pin != PinFilter)
+			{
+				continue;
+			}
+
 			const UPCGSpatialData* SpatialData = Cast<UPCGSpatialData>(Input.Data);
 
 			if (!SpatialData)
diff --git a/Engine/Plugins/Experimental/PCG/Source/PCG/Private/Elements/PCGPointFilter.cpp b/Engine/Plugins/Experimental/PCG/Source/PCG/Private/Elements/PCGPointFilter.cpp
new file mode 100644
--- /dev/null
+++ b/Engine/Plugins/Experimental/PCG/Source/PCG/Private/Elements/PCGPointFilter.cpp
@@ -0,0 +1,845 @@
+// Copyright Epic Games, Inc. All Rights Reserved.
+
+#include "Elements/PCGPointFilter.h"
+#include "Data/PCGSpatialData.h"
+#include "Data/PCGPointData.h"
+#include "Helpers/PCGSettingsHelpers.h"
+#include "Metadata/PCGMetadata.h"
+#include "Metadata/PCGMetadataAttribute.h"
+#include "Metadata/PCGMetadataAttributeTpl.h"
+#include "Helpers/PCGAsync.h"
+
+namespace PCGPointFilterConstants
+{
+	const FName DataToFilterLabel = TEXT("In");
+	const FName FilterLabel = TEXT("Filter");
+	const FName InFilterLabel = TEXT("InsideFilter");
+	const FName OutFilterLabel = TEXT("OutsideFilter");
+}
+
+namespace PCGPointFilterHelpers
+{
+	TFunction<void(const FPCGPoint&, const UPCGMetadata*, void*)> ConstructPropertyGetter(EPCGPointProperties TargetPointProperty, EPCGPointFilterConstantType& OutType)
+	{
+		if (TargetPointProperty == EPCGPointProperties::Density)
+		{
+			OutType = EPCGPointFilterConstantType::Float;
+			return [](const FPCGPoint& InPoint, const UPCGMetadata*, void* OutData)
+			{
+				*(static_cast<float*>(OutData)) = InPoint.Density;
+			};
+		}
+		else if (TargetPointProperty == EPCGPointProperties::Extents)
+		{
+			OutType = EPCGPointFilterConstantType::Vector;
+			return [](const FPCGPoint& InPoint, const UPCGMetadata*, void* OutData)
+			{
+				*(static_cast<FVector*>(OutData)) = InPoint.Extents;
+			};
+		}
+		else if (TargetPointProperty == EPCGPointProperties::Color)
+		{
+			OutType = EPCGPointFilterConstantType::Vector4;
+			return [](const FPCGPoint& InPoint, const UPCGMetadata*, void* OutData)
+			{
+				*(static_cast<FVector4*>(OutData)) = InPoint.Color;
+			};
+		}
+		else if (TargetPointProperty == EPCGPointProperties::Position)
+		{
+			OutType = EPCGPointFilterConstantType::Vector;
+			return [](const FPCGPoint& InPoint, const UPCGMetadata*, void* OutData)
+			{
+				*(static_cast<FVector*>(OutData)) = InPoint.Transform.GetLocation();
+			};
+		}
+		//else if (TargetPointProperty == EPCGPointProperties::Rotation)
+		//{
+		//	OutType = EPCGPointFilterConstantType::Rotation;
+		//	return [](const FPCGPoint& InPoint, const UPCGMetadata*, void* OutData)
+		//	{
+		//		*(static_cast<FQuat*>(OutData)) = InPoint.Transform.GetRotation();
+		//	};
+		//}
+		else if (TargetPointProperty == EPCGPointProperties::Scale)
+		{
+			OutType = EPCGPointFilterConstantType::Vector;
+			return [](const FPCGPoint& InPoint, const UPCGMetadata*, void* OutData)
+			{
+				*(static_cast<FVector*>(OutData)) = InPoint.Transform.GetScale3D();
+			};
+		}
+		//else if (TargetPointProperty == EPCGPointProperties::Transform)
+		//{
+		//	OutType = EPCGPointFilterConstantType::Transform;
+		//	return [](const FPCGPoint& InPoint, const UPCGMetadata*, void* OutData)
+		//	{
+		//		*(static_cast<FTransform*>(OutData)) = InPoint.Transform;
+		//	};
+		//}
+		else if (TargetPointProperty == EPCGPointProperties::Steepness)
+		{
+			OutType = EPCGPointFilterConstantType::Float;
+			return [](const FPCGPoint& InPoint, const UPCGMetadata*, void* OutData)
+			{
+				*(static_cast<float*>(OutData)) = InPoint.Steepness;
+			};
+		}
+		else
+		{
+			// Log error
+			OutType = EPCGPointFilterConstantType::Unknown;
+			return [](const FPCGPoint&, const UPCGMetadata*, void*) {};
+		}
+	}
+
+	TFunction<void(const FPCGPoint&, const UPCGMetadata*, void*)> ConstructAttributeGetter(
+		const UPCGMetadata* Metadata, 
+		const FName& AttributeName, 
+		EPCGPointFilterConstantType& OutType)
+	{
+		OutType = EPCGPointFilterConstantType::Unknown;
+
+		const FPCGMetadataAttributeBase* AttributeBase = Metadata->GetConstAttribute(AttributeName);
+		if (AttributeBase)
+		{
+			if (AttributeBase->GetTypeId() == PCG::Private::MetadataTypes<bool>::Id)
+			{
+				OutType = EPCGPointFilterConstantType::Integer64;
+				return [AttributeBase](const FPCGPoint& InPoint, const UPCGMetadata* InMetadata, void* OutData)
+				{
+					*(static_cast<int64*>(OutData)) = static_cast<const FPCGMetadataAttribute<bool>*>(AttributeBase)->GetValueFromItemKey(InPoint.MetadataEntry);
+				};
+			}
+			else if (AttributeBase->GetTypeId() == PCG::Private::MetadataTypes<int32>::Id)
+			{
+				OutType = EPCGPointFilterConstantType::Integer64;
+				return [AttributeBase](const FPCGPoint& InPoint, const UPCGMetadata* InMetadata, void* OutData)
+				{
+					*(static_cast<int64*>(OutData)) = static_cast<const FPCGMetadataAttribute<int32>*>(AttributeBase)->GetValueFromItemKey(InPoint.MetadataEntry);
+				};
+			}
+			else if (AttributeBase->GetTypeId() == PCG::Private::MetadataTypes<int64>::Id)
+			{
+				OutType = EPCGPointFilterConstantType::Integer64;
+				return [AttributeBase](const FPCGPoint& InPoint, const UPCGMetadata* InMetadata, void* OutData)
+				{
+					*(static_cast<int64*>(OutData)) = static_cast<const FPCGMetadataAttribute<int64>*>(AttributeBase)->GetValueFromItemKey(InPoint.MetadataEntry);
+				};
+			}
+			else if (AttributeBase->GetTypeId() == PCG::Private::MetadataTypes<float>::Id)
+			{
+				OutType = EPCGPointFilterConstantType::Float;
+				return [AttributeBase](const FPCGPoint& InPoint, const UPCGMetadata* InMetadata, void* OutData)
+				{
+					*(static_cast<float*>(OutData)) = static_cast<const FPCGMetadataAttribute<float>*>(AttributeBase)->GetValueFromItemKey(InPoint.MetadataEntry);
+				};
+			}
+			else if (AttributeBase->GetTypeId() == PCG::Private::MetadataTypes<double>::Id)
+			{
+				OutType = EPCGPointFilterConstantType::Float;
+				return [AttributeBase](const FPCGPoint& InPoint, const UPCGMetadata* InMetadata, void* OutData)
+				{
+					*(static_cast<float*>(OutData)) = static_cast<const FPCGMetadataAttribute<double>*>(AttributeBase)->GetValueFromItemKey(InPoint.MetadataEntry);
+				};
+			}
+			else if (AttributeBase->GetTypeId() == PCG::Private::MetadataTypes<FVector>::Id)
+			{
+				OutType = EPCGPointFilterConstantType::Vector;
+				return [AttributeBase](const FPCGPoint& InPoint, const UPCGMetadata* InMetadata, void* OutData)
+				{
+					*(static_cast<FVector*>(OutData)) = static_cast<const FPCGMetadataAttribute<FVector>*>(AttributeBase)->GetValueFromItemKey(InPoint.MetadataEntry);
+				};
+			}
+			else if (AttributeBase->GetTypeId() == PCG::Private::MetadataTypes<FVector4>::Id)
+			{
+				OutType = EPCGPointFilterConstantType::Vector4;
+				return [AttributeBase](const FPCGPoint& InPoint, const UPCGMetadata* InMetadata, void* OutData)
+				{
+					*(static_cast<FVector4*>(OutData)) = static_cast<const FPCGMetadataAttribute<FVector4>*>(AttributeBase)->GetValueFromItemKey(InPoint.MetadataEntry);
+				};
+			}
+			//else if (AttributeBase->GetTypeId() == PCG::Private::MetadataTypes<FQuat>::Id)
+			//{
+			//	OutType = EPCGPointFilterConstantType::Rotation;
+			//	return [AttributeBase](const FPCGPoint& InPoint, const UPCGMetadata* InMetadata, void* OutData)
+			//	{
+			//		*(static_cast<FQuat*>(OutData)) = static_cast<const FPCGMetadataAttribute<FQuat>*>(AttributeBase)->GetValueFromItemKey(InPoint.MetadataEntry);
+			//	};
+			//}
+			//else if (AttributeBase->GetTypeId() == PCG::Private::MetadataTypes<FTransform>::Id)
+			//{
+			//	*(static_cast<FTransform*>(OutData)) = static_cast<const FPCGMetadataAttribute<FTransform>*>(AttributeBase)->GetValueFromItemKey(InPoint.MetadataEntry);
+			//}
+			else if (AttributeBase->GetTypeId() == PCG::Private::MetadataTypes<FString>::Id)
+			{
+				OutType = EPCGPointFilterConstantType::String;
+				return [AttributeBase](const FPCGPoint& InPoint, const UPCGMetadata* InMetadata, void* OutData)
+				{
+					*(static_cast<FString*>(OutData)) = static_cast<const FPCGMetadataAttribute<FString>*>(AttributeBase)->GetValueFromItemKey(InPoint.MetadataEntry);
+				};
+			}
+			else
+			{
+				// Log error
+				OutType = EPCGPointFilterConstantType::Unknown;
+				return [](const FPCGPoint&, const UPCGMetadata*, void*) {};
+			}
+		}
+		else
+		{
+			// Log error
+			OutType = EPCGPointFilterConstantType::Unknown;
+			return [](const FPCGPoint&, const UPCGMetadata*, void*) {};
+		}
+	}
+
+	TFunction<void(const FPCGPoint&, const UPCGMetadata*, void*)> ConstructSourceGetter(
+		EPCGPointTargetFilterType TargetFilterType, 
+		EPCGPointProperties TargetPointProperty, 
+		const UPCGMetadata* TargetMetadata,
+		const FName& TargetAttributeName,
+		EPCGPointFilterConstantType& OutType)
+	{
+		if (TargetFilterType == EPCGPointTargetFilterType::Property)
+		{
+			return ConstructPropertyGetter(TargetPointProperty, OutType);
+		}
+		else if (TargetFilterType == EPCGPointTargetFilterType::Metadata)
+		{
+			return ConstructAttributeGetter(TargetMetadata, TargetAttributeName, OutType);
+		}
+		else
+		{
+			// log error
+			OutType = EPCGPointFilterConstantType::Unknown;
+			return [](const FPCGPoint&, const UPCGMetadata*, void*) {};
+		}
+	}
+
+	TFunction<void(const FPCGPoint&, const UPCGMetadata*, void*)> ConstructThresholdGetter(
+		EPCGPointThresholdType ThresholdFilterType, 
+		EPCGPointProperties ThresholdPointProperty, 
+		const UPCGMetadata* ThresholdMetadata,
+		const FName& ThresholdAttributeName, 
+		EPCGPointFilterConstantType ThresholdConstantType,
+		const int64& Integer64Constant,
+		const float& FloatConstant,
+		const FVector& VectorConstant,
+		const FVector4& Vector4Constant,
+		//const FQuat& RotationConstant,
+		const FString& StringConstant,
+		EPCGPointFilterConstantType& OutType)
+	{
+		if (ThresholdFilterType == EPCGPointThresholdType::Property)
+		{
+			return ConstructPropertyGetter(ThresholdPointProperty, OutType);
+		}
+		else if (ThresholdFilterType == EPCGPointThresholdType::Metadata)
+		{
+			return ConstructAttributeGetter(ThresholdMetadata, ThresholdAttributeName, OutType);
+		}
+		else if (ThresholdFilterType == EPCGPointThresholdType::Constant)
+		{
+			OutType = ThresholdConstantType;
+
+			if (ThresholdConstantType == EPCGPointFilterConstantType::Integer64)
+			{
+				return [&Integer64Constant](const FPCGPoint&, const UPCGMetadata*, void* OutData)
+				{
+					*(static_cast<int64*>(OutData)) = Integer64Constant;
+				};
+			}
+			else if (ThresholdConstantType == EPCGPointFilterConstantType::Float)
+			{
+				return [&FloatConstant](const FPCGPoint&, const UPCGMetadata*, void* OutData)
+				{
+					*(static_cast<float*>(OutData)) = FloatConstant;
+				};
+			}
+			else if (ThresholdConstantType == EPCGPointFilterConstantType::Vector)
+			{
+				return [&VectorConstant](const FPCGPoint&, const UPCGMetadata*, void* OutData)
+				{
+					*(static_cast<FVector*>(OutData)) = VectorConstant;
+				};
+			}
+			else if (ThresholdConstantType == EPCGPointFilterConstantType::Vector4)
+			{
+				return [&Vector4Constant](const FPCGPoint&, const UPCGMetadata*, void* OutData)
+				{
+					*(static_cast<FVector4*>(OutData)) = Vector4Constant;
+				};
+			}
+			//else if (ThresholdConstantType == EPCGPointFilterConstantType::Rotation)
+			//{
+			//	return [&RotationConstant](const FPCGPoint&, const UPCGMetadata*, void* OutData)
+			//	{
+			//		*(static_cast<FQuat*>(OutData)) = RotationConstant;
+			//	};
+			//}
+			else if (ThresholdConstantType == EPCGPointFilterConstantType::String)
+			{
+				return [&StringConstant](const FPCGPoint&, const UPCGMetadata*, void* OutData)
+				{
+					*(static_cast<FString*>(OutData)) = StringConstant;
+				};
+			}
+			else // unknown constant type
+			{
+				// Log error
+				OutType = EPCGPointFilterConstantType::Unknown;
+				return [](const FPCGPoint&, const UPCGMetadata*, void*) {};
+			}
+		}
+		else // unknown filter type
+		{
+			// Log error
+			OutType = EPCGPointFilterConstantType::Unknown;
+			return [](const FPCGPoint&, const UPCGMetadata*, void*) {};
+		}
+	}
+
+	template<typename T, typename U>
+	struct TemplatedOperator
+	{
+		TemplatedOperator(EPCGPointFilterOperator Operator)
+		{
+			if (Operator == EPCGPointFilterOperator::Greater)
+			{
+				OperatorFn = [](const T& A, const U& B) { return A > B; };
+			}
+			else if (Operator == EPCGPointFilterOperator::GreaterOrEqual)
+			{
+				OperatorFn = [](const T& A, const U& B) { return A >= B; };
+			}
+			else if (Operator == EPCGPointFilterOperator::Lesser)
+			{
+				OperatorFn = [](const T& A, const U& B) { return A < B; };
+			}
+			else if (Operator == EPCGPointFilterOperator::LesserOrEqual)
+			{
+				OperatorFn = [](const T& A, const U& B) { return A <= B; };
+			}
+			else if (Operator == EPCGPointFilterOperator::Equal)
+			{
+				OperatorFn = [](const T& A, const U& B) { return A == B; };
+			}
+			else if (Operator == EPCGPointFilterOperator::NotEqual)
+			{
+				OperatorFn = [](const T& A, const U& B) { return A != B; };
+			}
+			else
+			{
+				OperatorFn = [](const T&, const U&) { return false; };
+			}
+		}
+
+		bool operator()(const T& A, const U& B) const
+		{
+			return OperatorFn(A, B);
+		}
+
+		TFunction<bool(const T&, const U&)> OperatorFn;
+	};
+
+	template<typename U>
+	struct TemplatedOperator<FVector, U>
+	{
+		TemplatedOperator(EPCGPointFilterOperator Operator)
+		{
+			if (Operator == EPCGPointFilterOperator::Greater)
+			{
+				OperatorFn = [](const FVector& A, const U& B) { return A.Length() > B; };
+			}
+			else if (Operator == EPCGPointFilterOperator::GreaterOrEqual)
+			{
+				OperatorFn = [](const FVector& A, const U& B) { return A.Length() >= B; };
+			}
+			else if (Operator == EPCGPointFilterOperator::Lesser)
+			{
+				OperatorFn = [](const FVector& A, const U& B) { return A.Length() < B; };
+			}
+			else if (Operator == EPCGPointFilterOperator::LesserOrEqual)
+			{
+				OperatorFn = [](const FVector& A, const U& B) { return A.Length() <= B; };
+			}
+			else if (Operator == EPCGPointFilterOperator::Equal)
+			{
+				OperatorFn = [](const FVector& A, const U& B) { return A.Length() == B; };
+			}
+			else if (Operator == EPCGPointFilterOperator::NotEqual)
+			{
+				OperatorFn = [](const FVector& A, const U& B) { return A.Length() != B; };
+			}
+			else
+			{
+				OperatorFn = [](const FVector&, const U&) { return false; };
+			}
+		}
+
+		bool operator()(const FVector& A, const U& B) const
+		{
+			return OperatorFn(A, B);
+		}
+
+		TFunction<bool(const FVector&, const U&)> OperatorFn;
+	};
+
+	template<typename U>
+	struct TemplatedOperator<U, FVector>
+	{
+		TemplatedOperator(EPCGPointFilterOperator Operator)
+		{
+			if (Operator == EPCGPointFilterOperator::Greater)
+			{
+				OperatorFn = [](const U& A, const FVector& B) { return A > B.Length(); };
+			}
+			else if (Operator == EPCGPointFilterOperator::GreaterOrEqual)
+			{
+				OperatorFn = [](const U& A, const FVector& B) { return A >= B.Length(); };
+			}
+			else if (Operator == EPCGPointFilterOperator::Lesser)
+			{
+				OperatorFn = [](const U& A, const FVector& B) { return A < B.Length(); };
+			}
+			else if (Operator == EPCGPointFilterOperator::LesserOrEqual)
+			{
+				OperatorFn = [](const U& A, const FVector& B) { return A <= B.Length(); };
+			}
+			else if (Operator == EPCGPointFilterOperator::Equal)
+			{
+				OperatorFn = [](const U& A, const FVector& B) { return A == B.Length(); };
+			}
+			else if (Operator == EPCGPointFilterOperator::NotEqual)
+			{
+				OperatorFn = [](const U& A, const FVector& B) { return A != B.Length(); };
+			}
+			else
+			{
+				OperatorFn = [](const U&, const FVector&) { return false; };
+			}
+		}
+
+		bool operator()(const U& A, const FVector& B) const
+		{
+			return OperatorFn(A, B);
+		}
+
+		TFunction<bool(const U&, const FVector&)> OperatorFn;
+	};
+
+	template<>
+	struct TemplatedOperator<FVector, FVector>
+	{
+		TemplatedOperator(EPCGPointFilterOperator Operator)
+		{
+			if (Operator == EPCGPointFilterOperator::Greater)
+			{
+				OperatorFn = [](const FVector& A, const FVector& B) { return A.X > B.X && A.Y > B.Y && A.Z > B.Z; };
+			}
+			else if (Operator == EPCGPointFilterOperator::GreaterOrEqual)
+			{
+				OperatorFn = [](const FVector& A, const FVector& B) { return A.X >= B.X && A.Y >= B.Y && A.Z >= B.Z; };
+			}
+			else if (Operator == EPCGPointFilterOperator::Lesser)
+			{
+				OperatorFn = [](const FVector& A, const FVector& B) { return A.X < B.X && A.Y < B.Y && A.Z < B.Z; };
+			}
+			else if (Operator == EPCGPointFilterOperator::LesserOrEqual)
+			{
+				OperatorFn = [](const FVector& A, const FVector& B) { return A.X <= B.X && A.Y <= B.Y && A.Z <= B.Z; };
+			}
+			else if (Operator == EPCGPointFilterOperator::Equal)
+			{
+				OperatorFn = [](const FVector& A, const FVector& B) { return A == B; };
+			}
+			else if (Operator == EPCGPointFilterOperator::NotEqual)
+			{
+				OperatorFn = [](const FVector& A, const FVector& B) { return A != B; };
+			}
+			else
+			{
+				OperatorFn = [](const FVector&, const FVector&) { return false; };
+			}
+		}
+
+		bool operator()(const FVector& A, const FVector& B) const
+		{
+			return OperatorFn(A, B);
+		}
+
+		TFunction<bool(const FVector&, const FVector&)> OperatorFn;
+	};
+
+	template<>
+	struct TemplatedOperator<FVector4, FVector4>
+	{
+		TemplatedOperator(EPCGPointFilterOperator Operator)
+		{
+			if (Operator == EPCGPointFilterOperator::Greater)
+			{
+				OperatorFn = [](const FVector4& A, const FVector4& B) { return A.X > B.X && A.Y > B.Y && A.Z > B.Z && A.W > B.W; };
+			}
+			else if (Operator == EPCGPointFilterOperator::GreaterOrEqual)
+			{
+				OperatorFn = [](const FVector4& A, const FVector4& B) { return A.X >= B.X && A.Y >= B.Y && A.Z >= B.Z && A.W >= B.W; };
+			}
+			else if (Operator == EPCGPointFilterOperator::Lesser)
+			{
+				OperatorFn = [](const FVector4& A, const FVector4& B) { return A.X < B.X&& A.Y < B.Y&& A.Z < B.Z && A.W < B.W; };
+			}
+			else if (Operator == EPCGPointFilterOperator::LesserOrEqual)
+			{
+				OperatorFn = [](const FVector4& A, const FVector4& B) { return A.X <= B.X && A.Y <= B.Y && A.Z <= B.Z && A.W <= B.W; };
+			}
+			else if (Operator == EPCGPointFilterOperator::Equal)
+			{
+				OperatorFn = [](const FVector4& A, const FVector4& B) { return A == B; };
+			}
+			else if (Operator == EPCGPointFilterOperator::NotEqual)
+			{
+				OperatorFn = [](const FVector4& A, const FVector4& B) { return A != B; };
+			}
+			else
+			{
+				OperatorFn = [](const FVector4&, const FVector4&) { return false; };
+			}
+		}
+
+		bool operator()(const FVector& A, const FVector& B) const
+		{
+			return OperatorFn(A, B);
+		}
+
+		TFunction<bool(const FVector&, const FVector&)> OperatorFn;
+	};
+
+	TFunction<bool(const FPCGPoint&, const UPCGMetadata*, const FPCGPoint&, const UPCGMetadata*)> ConstructConcreteOperator(
+		const TFunction<void(const FPCGPoint&, const UPCGMetadata*, void*)>& SourceGetter,
+		EPCGPointFilterConstantType SourceType,
+		const TFunction<void(const FPCGPoint&, const UPCGMetadata*, void*)>& ThresholdGetter,
+		EPCGPointFilterConstantType ThresholdType,
+		EPCGPointFilterOperator Operator,
+		bool& bOutOperatorCreatedCorrectly)
+	{
+		bOutOperatorCreatedCorrectly = true;
+
+#define CONCRETE_OPERATOR_SHORTHAND(TypeEnumA, TypeA, TypeEnumB, TypeB) \
+		if(SourceType == EPCGPointFilterConstantType::TypeEnumA && ThresholdType == EPCGPointFilterConstantType::TypeEnumB) \
+		{ \
+			TemplatedOperator<TypeA, TypeB> OperatorFn(Operator); \
+			\
+			return[&SourceGetter, &ThresholdGetter, OperatorFn](const FPCGPoint& InTargetPoint, const UPCGMetadata* InTargetMetadata, const FPCGPoint& InThresholdPoint, const UPCGMetadata* InThresholdMetadata) \
+			{ \
+				TypeA A; \
+				TypeB B; \
+				\
+				SourceGetter(InTargetPoint, InTargetMetadata, &A); \
+				ThresholdGetter(InThresholdPoint, InThresholdMetadata, &B); \
+				return OperatorFn(A, B); \
+			}; \
+		}
+
+		CONCRETE_OPERATOR_SHORTHAND(Integer64, int64, Integer64, int64)
+		CONCRETE_OPERATOR_SHORTHAND(Integer64, int64, Vector, FVector)
+		CONCRETE_OPERATOR_SHORTHAND(Float, float, Float, float)
+		CONCRETE_OPERATOR_SHORTHAND(Float, float, Vector, FVector)
+		CONCRETE_OPERATOR_SHORTHAND(Vector, FVector, Vector, FVector)
+		CONCRETE_OPERATOR_SHORTHAND(Vector, FVector, Integer64, int64)
+		CONCRETE_OPERATOR_SHORTHAND(Vector, FVector, Float, float)
+		CONCRETE_OPERATOR_SHORTHAND(Vector4, FVector4, Vector4, FVector4)
+		//CONCRETE_OPERATOR_SHORTHAND(Rotation, FQuat, Rotation, FQuat)
+		CONCRETE_OPERATOR_SHORTHAND(String, FString, String, FString)
+
+#undef CONCRETE_OPERATOR_SHORTHAND
+
+		bOutOperatorCreatedCorrectly = false;
+		return [](const FPCGPoint&, const UPCGMetadata*, const FPCGPoint&, const UPCGMetadata*) { return false; };
+	}
+}
+
+TArray<FName> UPCGPointFilterSettings::InLabels() const
+{
+	return { PCGPointFilterConstants::DataToFilterLabel, PCGPointFilterConstants::FilterLabel };
+}
+
+TArray<FName> UPCGPointFilterSettings::OutLabels() const
+{
+	return { PCGPointFilterConstants::InFilterLabel, PCGPointFilterConstants::OutFilterLabel };
+}
+
+FPCGElementPtr UPCGPointFilterSettings::CreateElement() const
+{
+	return MakeShared<FPCGPointFilterElement>();
+}
+
+bool FPCGPointFilterElement::ExecuteInternal(FPCGContext* Context) const
+{
+	TRACE_CPUPROFILER_EVENT_SCOPE(FPCGPointFilterElement::Execute);
+	check(Context);
+
+	const bool bHasInFilterOutputPin = Context->Node && Context->Node->IsOutputPinConnected(PCGPointFilterConstants::InFilterLabel);
+	const bool bHasOutsideFilterOutputPin = Context->Node && Context->Node->IsOutputPinConnected(PCGPointFilterConstants::OutFilterLabel);
+
+	// Early out
+	if (!bHasInFilterOutputPin && !bHasOutsideFilterOutputPin)
+	{
+		return true;
+	}
+
+	const UPCGPointFilterSettings* Settings = Context->GetInputSettings<UPCGPointFilterSettings>();
+	check(Settings);
+
+	TArray<FPCGTaggedData> DataToFilter = Context->InputData.GetInputsByPin(PCGPointFilterConstants::DataToFilterLabel);
+	UPCGParamData* Params = Context->InputData.GetParams(); //TODO: by pin?
+	TArray<FPCGTaggedData>& Outputs = Context->OutputData.TaggedData;
+	TArray<FPCGTaggedData> FilterData = Context->InputData.GetInputsByPin(PCGPointFilterConstants::FilterLabel);
+	FilterData.Append(Context->InputData.GetParamsByPin(PCGPointFilterConstants::FilterLabel));
+
+	// Forward any non-input data
+	Outputs.Append(Context->InputData.GetAllSettings());
+
+	const EPCGPointFilterOperator Operator = PCGSettingsHelpers::GetValue(GET_MEMBER_NAME_CHECKED(UPCGPointFilterSettings, Operator), Settings->Operator, Params);
+	const EPCGPointTargetFilterType TargetFilterType = PCGSettingsHelpers::GetValue(GET_MEMBER_NAME_CHECKED(UPCGPointFilterSettings, TargetFilterType), Settings->TargetFilterType, Params);
+	const EPCGPointProperties TargetPointProperty = PCGSettingsHelpers::GetValue(GET_MEMBER_NAME_CHECKED(UPCGPointFilterSettings, TargetPointProperty), Settings->TargetPointProperty, Params);
+	const FName TargetAttributeName = PCGSettingsHelpers::GetValue(GET_MEMBER_NAME_CHECKED(UPCGPointFilterSettings, TargetAttributeName), Settings->TargetAttributeName, Params);
+
+	const EPCGPointThresholdType ThresholdFilterType = PCGSettingsHelpers::GetValue(GET_MEMBER_NAME_CHECKED(UPCGPointFilterSettings, ThresholdFilterType), Settings->ThresholdFilterType, Params);
+	const EPCGPointProperties ThresholdPointProperty = PCGSettingsHelpers::GetValue(GET_MEMBER_NAME_CHECKED(UPCGPointFilterSettings, ThresholdPointProperty), Settings->ThresholdPointProperty, Params);
+	const FName ThresholdAttributeName = PCGSettingsHelpers::GetValue(GET_MEMBER_NAME_CHECKED(UPCGPointFilterSettings, ThresholdAttributeName), Settings->ThresholdAttributeName, Params);
+	const EPCGPointFilterConstantType ThresholdConstantType = PCGSettingsHelpers::GetValue(GET_MEMBER_NAME_CHECKED(UPCGPointFilterSettings, ThresholdConstantType), Settings->ThresholdConstantType, Params);
+
+	const int64 Integer64Constant = PCGSettingsHelpers::GetValue(GET_MEMBER_NAME_CHECKED(UPCGPointFilterSettings, Integer64Constant), Settings->Integer64Constant, Params);
+	const float FloatConstant = PCGSettingsHelpers::GetValue(GET_MEMBER_NAME_CHECKED(UPCGPointFilterSettings, FloatConstant), Settings->FloatConstant, Params);
+	const FVector VectorConstant = PCGSettingsHelpers::GetValue(GET_MEMBER_NAME_CHECKED(UPCGPointFilterSettings, VectorConstant), Settings->VectorConstant, Params);
+	const FVector4 Vector4Constant = PCGSettingsHelpers::GetValue(GET_MEMBER_NAME_CHECKED(UPCGPointFilterSettings, Vector4Constant), Settings->Vector4Constant, Params);
+	//const FQuat RotationConstant = PCGSettingsHelpers::GetValue(GET_MEMBER_NAME_CHECKED(UPCGPointFilterSettings, RotationConstant), Settings->RotationConstant, Params);
+	const FString StringConstant = PCGSettingsHelpers::GetValue(GET_MEMBER_NAME_CHECKED(UPCGPointFilterSettings, StringConstant), Settings->StringConstant, Params);
+
+	bool bUseSpatialQuery = PCGSettingsHelpers::GetValue(GET_MEMBER_NAME_CHECKED(UPCGPointFilterSettings, bUseSpatialQuery), Settings->bUseSpatialQuery, Params);
+
+	// Validate basic input data
+	if (TargetFilterType == EPCGPointTargetFilterType::Metadata && TargetAttributeName == NAME_None)
+	{
+		PCGE_LOG(Error, "Point filter cannot filter on an unnamed metadata attribute (Target Attribute Name)");
+		return true;
+	}
+	else if (ThresholdFilterType == EPCGPointThresholdType::Metadata && ThresholdAttributeName == NAME_None)
+	{
+		PCGE_LOG(Error, "Point filter cannot filter on an unnamed metadata attribute (Threshold Attribute Name)");
+		return true;
+	}
+	else if (ThresholdFilterType == EPCGPointThresholdType::Property && !FilterData.IsEmpty() && Cast<UPCGSpatialData>(FilterData[0].Data) == nullptr)
+	{
+		PCGE_LOG(Error, "Cannot filter by property on a non-spatial filter data");
+		return true;
+	}
+
+	if (ThresholdFilterType == EPCGPointThresholdType::Constant && !FilterData.IsEmpty())
+	{
+		PCGE_LOG(Warning, "Point filter uses a constant but has filter data that will be ignored");
+	}
+
+	for (const FPCGTaggedData& Input : DataToFilter)
+	{
+		const UPCGSpatialData* SpatialInput = Cast<const UPCGSpatialData>(Input.Data);
+		check(SpatialInput);
+
+		const UPCGPointData* OriginalData = SpatialInput->ToPointData(Context);
+		if (!OriginalData)
+		{
+			PCGE_LOG(Error, "Unable to get point data from input");
+			continue;
+		}
+
+		const TArray<FPCGPoint>& TargetPoints = OriginalData->GetPoints();
+		const UPCGMetadata* TargetMetadata = OriginalData->ConstMetadata();
+
+		const UPCGSpatialData* ThresholdSpatialData = nullptr;
+		const TArray<FPCGPoint>* ThresholdPoints = nullptr;
+		const UPCGMetadata* ThresholdMetadata = nullptr;
+		UPCGMetadata* TemporaryMetadata = nullptr;
+
+		if (FilterData.Num() > 0)
+		{
+			if (const UPCGSpatialData* SpatialDataFilter = Cast<const UPCGSpatialData>(FilterData[0].Data))
+			{
+				ThresholdSpatialData = SpatialDataFilter;
+				ThresholdMetadata = ThresholdSpatialData->ConstMetadata();
+
+				if (bUseSpatialQuery && ThresholdMetadata && ThresholdFilterType == EPCGPointThresholdType::Metadata)
+				{
+					TemporaryMetadata = NewObject<UPCGMetadata>();
+					TemporaryMetadata->Initialize(ThresholdMetadata);
+				}
+				else if(!bUseSpatialQuery)
+				{
+					const UPCGPointData* ThresholdData = SpatialDataFilter->ToPointData();
+					ThresholdPoints = &ThresholdData->GetPoints();
+				}
+			}
+			else if (const UPCGParamData* ParamDataFilter = Cast<const UPCGParamData>(FilterData[0].Data))
+			{
+				ThresholdMetadata = ParamDataFilter->ConstMetadata();
+			}
+		}
+		else
+		{
+			// Use self as filter
+			ThresholdMetadata = TargetMetadata;
+			ThresholdPoints = &TargetPoints;
+		}
+
+		// Additional validation
+		if (TargetFilterType == EPCGPointTargetFilterType::Metadata)
+		{
+			if (!TargetMetadata)
+			{
+				PCGE_LOG(Error, "Target data to filter has no metadata which is required to filter by metadata");
+				continue;
+			}
+			else if (!TargetMetadata->GetConstAttribute(TargetAttributeName))
+			{
+				PCGE_LOG(Error, "Target metadata does not have the %s attribute", *TargetAttributeName.ToString());
+				continue;
+			}
+		}
+		
+		if (ThresholdFilterType == EPCGPointThresholdType::Metadata)
+		{
+			if (!ThresholdMetadata)
+			{
+				PCGE_LOG(Error, "Filter data has no metadata which is required to filter by metadata");
+				continue;
+			}
+			else if (!ThresholdMetadata->GetConstAttribute(ThresholdAttributeName))
+			{
+				PCGE_LOG(Error, "Filter metadata does not have the %s attribute", *ThresholdAttributeName.ToString());
+				continue;
+			}
+		}
+		else if (ThresholdFilterType == EPCGPointThresholdType::Property)
+		{
+			if (!bUseSpatialQuery && (!ThresholdPoints || (ThresholdPoints->Num() > 1 && TargetPoints.Num() > 0 && ThresholdPoints->Num() != TargetPoints.Num())))
+			{
+				PCGE_LOG(Error, "Filter data points count (%d) mismatch vs the data to filter (%d)", (ThresholdPoints ? ThresholdPoints->Num() : 0), TargetPoints.Num());
+				continue;
+			}
+
+			if (!bUseSpatialQuery && !ThresholdSpatialData)
+			{
+				PCGE_LOG(Error, "Filter data cannot be used as spatial data");
+				continue;
+			}
+		}
+
+		// Build our functions
+		EPCGPointFilterConstantType TargetType = EPCGPointFilterConstantType::Unknown;
+		EPCGPointFilterConstantType ThresholdType = EPCGPointFilterConstantType::Unknown;
+
+		TFunction<void(const FPCGPoint&, const UPCGMetadata*, void*)> SourceGetter = PCGPointFilterHelpers::ConstructSourceGetter(
+			TargetFilterType, 
+			TargetPointProperty, 
+			TargetMetadata,
+			TargetAttributeName,
+			TargetType);
+
+		if (TargetType == EPCGPointFilterConstantType::Unknown)
+		{
+			PCGE_LOG(Error, "Unable to generate the data value getter");
+			continue;
+		}
+
+		// B - a getter that returns the value for the threshold data
+		TFunction<void(const FPCGPoint&, const UPCGMetadata*, void*)> ThresholdGetter = PCGPointFilterHelpers::ConstructThresholdGetter(
+			ThresholdFilterType,
+			ThresholdPointProperty,
+			ThresholdMetadata,
+			ThresholdAttributeName,
+			ThresholdConstantType,
+			Integer64Constant,
+			FloatConstant,
+			VectorConstant,
+			Vector4Constant,
+			//RotationConstant,
+			StringConstant,
+			ThresholdType);
+
+		if (ThresholdType == EPCGPointFilterConstantType::Unknown)
+		{
+			PCGE_LOG(Error, "Unable to generate the filter value getter");
+			continue;
+		}
+
+		// C - the operator
+		bool bConcreteOperatorConstructedCorrectly = false;
+		TFunction<bool(const FPCGPoint&, const UPCGMetadata*, const FPCGPoint&, const UPCGMetadata*)> ConcreteOperator = PCGPointFilterHelpers::ConstructConcreteOperator(
+			SourceGetter,
+			TargetType,
+			ThresholdGetter,
+			ThresholdType,
+			Operator,
+			bConcreteOperatorConstructedCorrectly);
+
+		if (!bConcreteOperatorConstructedCorrectly)
+		{
+			PCGE_LOG(Error, "Unable to generate the operator (most likely data types mismatch");
+			continue;
+		}
+
+		UPCGPointData* InFilterData = NewObject<UPCGPointData>();
+		InFilterData->InitializeFromData(OriginalData);
+		TArray<FPCGPoint>& InPoints = InFilterData->GetMutablePoints();
+
+		UPCGPointData* OutFilterData = NewObject<UPCGPointData>();
+		OutFilterData->InitializeFromData(OriginalData);
+		TArray<FPCGPoint>& OutPoints = OutFilterData->GetMutablePoints();
+
+		FPCGAsync::AsyncPointFilterProcessing(Context, TargetPoints.Num(), InPoints, OutPoints, [&TargetPoints, TargetMetadata, ThresholdSpatialData, ThresholdPoints, ThresholdMetadata, TemporaryMetadata, &ConcreteOperator](int32 Index, FPCGPoint& InFilterPoint, FPCGPoint& OutFilterPoint)
+		{
+			const FPCGPoint& SourcePoint = TargetPoints[Index];
+
+			bool bSkipTest = false;
+			FPCGPoint ThresholdPoint;
+			if (ThresholdPoints) // not null only in the direct match case
+			{
+				ThresholdPoint = (*ThresholdPoints)[ThresholdPoints->Num() > 1 ? Index : 0];
+			}
+			else if(ThresholdSpatialData)
+			{
+				bSkipTest = !ThresholdSpatialData->GetPointAtPosition(SourcePoint.Transform.GetLocation(), ThresholdPoint, TemporaryMetadata);
+			}
+			else // Only metadata entry (param)
+			{
+				ThresholdPoint.MetadataEntry = 0;
+			}
+
+			// Apply operator
+			bool bPointInFilter = bSkipTest || ConcreteOperator(SourcePoint, TargetMetadata, ThresholdPoint, (TemporaryMetadata ? TemporaryMetadata : ThresholdMetadata));
+			if (bPointInFilter)
+			{
+				InFilterPoint = SourcePoint;
+				return true;
+			}
+			else
+			{
+				OutFilterPoint = SourcePoint;
+				return false;
+			}
+		});
+
+		// Push the filtered data to the output data.
+		FPCGTaggedData& InFilterTaggedData = Outputs.Emplace_GetRef();
+		InFilterTaggedData.Data = InFilterData;
+		InFilterTaggedData.Tags = Input.Tags;
+		InFilterTaggedData.Pin = PCGPointFilterConstants::InFilterLabel;
+
+		FPCGTaggedData& OutFilterTaggedData = Outputs.Emplace_GetRef();
+		OutFilterTaggedData.Data = OutFilterData;
+		OutFilterTaggedData.Tags = Input.Tags;
+		OutFilterTaggedData.Pin = PCGPointFilterConstants::OutFilterLabel;
+	}
+
+	return true;
+}
\ No newline at end of file
diff --git a/Engine/Plugins/Experimental/PCG/Source/PCG/Private/Helpers/PCGAsync.cpp b/Engine/Plugins/Experimental/PCG/Source/PCG/Private/Helpers/PCGAsync.cpp
--- a/Engine/Plugins/Experimental/PCG/Source/PCG/Private/Helpers/PCGAsync.cpp
+++ b/Engine/Plugins/Experimental/PCG/Source/PCG/Private/Helpers/PCGAsync.cpp
@@ -83,4 +83,108 @@ namespace FPCGAsync
 			OutPoints.SetNum(RangeIndex);
 		}
 	}
+
+	void AsyncPointFilterProcessing(FPCGContext* Context, int32 NumIterations, TArray<FPCGPoint>& InFilterPoints, TArray<FPCGPoint>& OutFilterPoints, const TFunction<bool(int32, FPCGPoint&, FPCGPoint&)>& PointFunc)
+	{
+		const int32 MinIterationsPerTask = 256;
+		AsyncPointFilterProcessing(Context ? Context->NumAvailableTasks : 1, MinIterationsPerTask, NumIterations, InFilterPoints, OutFilterPoints, PointFunc);
+	}
+
+	void AsyncPointFilterProcessing(int32 NumAvailableTasks, int32 MinIterationsPerTask, int32 NumIterations, TArray<FPCGPoint>& InFilterPoints, TArray<FPCGPoint>& OutFilterPoints, const TFunction<bool(int32, FPCGPoint&, FPCGPoint&)>& PointFunc)
+	{
+		TRACE_CPUPROFILER_EVENT_SCOPE(IPCGElement::AsyncPointFilterProcessing);
+		check(NumAvailableTasks > 0 && MinIterationsPerTask > 0 && NumIterations >= 0);
+		// Get number of available threads from the context
+		const int32 NumTasks = FMath::Min(NumAvailableTasks, FMath::Max(1, NumIterations / MinIterationsPerTask));
+		const int32 IterationsPerTask = NumIterations / NumTasks;
+
+		{
+			TRACE_CPUPROFILER_EVENT_SCOPE(IPCGElement::AsyncPointProcessing::AllocatingArray);
+			// Pre-reserve the out points array
+			InFilterPoints.SetNum(NumIterations);
+			OutFilterPoints.SetNum(NumIterations);
+		}
+
+		// Setup [current, last, nb points] data per dispatch
+		// Execute
+		TArray<TFuture<TPair<int32, int32>>> AsyncTasks;
+		AsyncTasks.Reserve(NumTasks);
+
+		// Launch the async tasks
+		for (int32 TaskIndex = 0; TaskIndex < NumTasks; ++TaskIndex)
+		{
+			const int32 StartIndex = TaskIndex * IterationsPerTask;
+			const int32 EndIndex = (TaskIndex == NumTasks - 1 ? NumIterations : StartIndex + IterationsPerTask);
+
+			AsyncTasks.Emplace(Async(EAsyncExecution::ThreadPool, [&PointFunc, StartIndex, EndIndex, &InFilterPoints, &OutFilterPoints]()
+				{
+					TRACE_CPUPROFILER_EVENT_SCOPE(IPCGElement::AsyncPointProcessing::InnerLoop);
+					int32 NumPointsInWritten = 0;
+					int32 NumPointsOutWritten = 0;
+
+					for (int32 Index = StartIndex; Index < EndIndex; ++Index)
+					{
+						if (PointFunc(Index, InFilterPoints[StartIndex + NumPointsInWritten], OutFilterPoints[StartIndex + NumPointsOutWritten]))
+						{
+							++NumPointsInWritten;
+						}
+						else
+						{
+							++NumPointsOutWritten;
+						}
+					}
+
+					return TPair<int32, int32>(NumPointsInWritten, NumPointsOutWritten);
+				}));
+		}
+
+		// Wait/Gather results & collapse points
+		{
+			TRACE_CPUPROFILER_EVENT_SCOPE(IPCGElement::AsyncPointFilterProcessing::WaitAndCollapseArray);
+			int InFilterRangeIndex = 0;
+			int OutFilterRangeIndex = 0;
+
+			for (int32 AsyncIndex = 0; AsyncIndex < AsyncTasks.Num(); ++AsyncIndex)
+			{
+				const int32 StartPointsIndex = AsyncIndex * IterationsPerTask;
+
+				TFuture<TPair<int32, int32>>& AsyncTask = AsyncTasks[AsyncIndex];
+				AsyncTask.Wait();
+				TPair<int32, int32> NumberOfPointsAdded = AsyncTask.Get();
+
+				// Move in-filter points
+				{
+					int NumInFilterPoints = NumberOfPointsAdded.Key;
+
+					if (StartPointsIndex != InFilterRangeIndex)
+					{
+						for (int32 MoveIndex = 0; MoveIndex < NumInFilterPoints; ++MoveIndex)
+						{
+							InFilterPoints[InFilterRangeIndex + MoveIndex] = MoveTemp(InFilterPoints[StartPointsIndex + MoveIndex]);
+						}
+					}
+
+					InFilterRangeIndex += NumInFilterPoints;
+				}
+
+				// Move out-filter points
+				{
+					int NumOutFilterPoints = NumberOfPointsAdded.Value;
+
+					if (StartPointsIndex != OutFilterRangeIndex)
+					{
+						for (int32 MoveIndex = 0; MoveIndex < NumOutFilterPoints; ++MoveIndex)
+						{
+							OutFilterPoints[OutFilterRangeIndex + MoveIndex] = MoveTemp(OutFilterPoints[StartPointsIndex + MoveIndex]);
+						}
+					}
+
+					OutFilterRangeIndex += NumOutFilterPoints;
+				}
+			}
+
+			InFilterPoints.SetNum(InFilterRangeIndex);
+			OutFilterPoints.SetNum(OutFilterRangeIndex);
+		}
+	}
 }
\ No newline at end of file
diff --git a/Engine/Plugins/Experimental/PCG/Source/PCG/Private/PCGGraph.cpp b/Engine/Plugins/Experimental/PCG/Source/PCG/Private/PCGGraph.cpp
--- a/Engine/Plugins/Experimental/PCG/Source/PCG/Private/PCGGraph.cpp
+++ b/Engine/Plugins/Experimental/PCG/Source/PCG/Private/PCGGraph.cpp
@@ -315,24 +315,84 @@ void UPCGGraph::RemoveEdge(UPCGNode* From, const FName& FromLabel, UPCGNode* To,
 #endif
 }
 
-void UPCGGraph::RemoveInboundEdges(UPCGNode* InNode)
+void UPCGGraph::RemoveAllInboundEdges(UPCGNode* InNode)
 {
 	check(InNode);
+	const bool bChanged = !InNode->InboundEdges.IsEmpty();
 
 	for (int32 i = InNode->InboundEdges.Num() - 1; i >= 0; --i)
 	{
 		InNode->InboundEdges[i]->BreakEdge();
 	}
+
+#if WITH_EDITOR
+	if (bChanged)
+	{
+		NotifyGraphChanged(/*bIsStructural=*/true);
+	}
+#endif
 }
 
-void UPCGGraph::RemoveOutboundEdges(UPCGNode* InNode)
+void UPCGGraph::RemoveAllOutboundEdges(UPCGNode* InNode)
 {
 	check(InNode);
+	const bool bChanged = !InNode->OutboundEdges.IsEmpty();
 
 	for (int32 i = InNode->OutboundEdges.Num() - 1; i >= 0; --i)
 	{
 		InNode->OutboundEdges[i]->BreakEdge();
 	}
+
+#if WITH_EDITOR
+	if (bChanged)
+	{
+		NotifyGraphChanged(/*bIsStructural=*/true);
+	}
+#endif
+}
+
+void UPCGGraph::RemoveInboundEdges(UPCGNode* InNode, const FName& InboundLabel)
+{
+	check(InNode);
+	bool bChanged = false;
+
+	for (int32 i = InNode->InboundEdges.Num() - 1; i >= 0; --i)
+	{
+		if (InNode->InboundEdges[i]->OutboundLabel == InboundLabel)
+		{
+			InNode->InboundEdges[i]->BreakEdge();
+			bChanged = true;
+		}
+	}
+
+#if WITH_EDITOR
+	if (bChanged)
+	{
+		NotifyGraphChanged(/*bIsStructural=*/true);
+	}
+#endif
+}
+
+void UPCGGraph::RemoveOutboundEdges(UPCGNode* InNode, const FName& OutboundLabel)
+{
+	check(InNode);
+	bool bChanged = false;
+
+	for (int32 i = InNode->OutboundEdges.Num() - 1; i >= 0; --i)
+	{
+		if (InNode->OutboundEdges[i]->InboundLabel == OutboundLabel)
+		{
+			InNode->OutboundEdges[i]->BreakEdge();
+			bChanged = true;
+		}
+	}
+
+#if WITH_EDITOR
+	if (bChanged)
+	{
+		NotifyGraphChanged(/*bIsStructural=*/true);
+	}
+#endif
 }
 
 #if WITH_EDITOR
diff --git a/Engine/Plugins/Experimental/PCG/Source/PCG/Private/PCGNode.cpp b/Engine/Plugins/Experimental/PCG/Source/PCG/Private/PCGNode.cpp
--- a/Engine/Plugins/Experimental/PCG/Source/PCG/Private/PCGNode.cpp
+++ b/Engine/Plugins/Experimental/PCG/Source/PCG/Private/PCGNode.cpp
@@ -83,6 +83,11 @@ TArray<FName> UPCGNode::OutLabels() const
 	return DefaultSettings ? DefaultSettings->OutLabels() : TArray<FName>();
 }
 
+bool UPCGNode::HasDefaultLabels() const
+{
+	return !DefaultSettings || DefaultSettings->HasDefaultLabels();
+}
+
 bool UPCGNode::IsInputPinConnected(const FName& Label) const
 {
 	for (const UPCGEdge* InboundEdge : InboundEdges)
diff --git a/Engine/Plugins/Experimental/PCG/Source/PCG/Private/PCGSettings.cpp b/Engine/Plugins/Experimental/PCG/Source/PCG/Private/PCGSettings.cpp
--- a/Engine/Plugins/Experimental/PCG/Source/PCG/Private/PCGSettings.cpp
+++ b/Engine/Plugins/Experimental/PCG/Source/PCG/Private/PCGSettings.cpp
@@ -39,12 +39,12 @@ bool UPCGSettings::operator==(const UPCGSettings& Other) const
 
 bool UPCGSettings::HasInLabel(const FName& Label) const
 {
-	return Label == NAME_None || InLabels().Contains(Label);
+	return (Label == NAME_None && HasDefaultLabels()) || InLabels().Contains(Label);
 }
 
 bool UPCGSettings::HasOutLabel(const FName& Label) const
 {
-	return Label == NAME_None || OutLabels().Contains(Label);
+	return (Label == NAME_None && HasDefaultLabels()) || OutLabels().Contains(Label);
 }
 
 FPCGElementPtr UPCGSettings::GetElement() const
diff --git a/Engine/Plugins/Experimental/PCG/Source/PCG/Public/Elements/PCGPointFilter.h b/Engine/Plugins/Experimental/PCG/Source/PCG/Public/Elements/PCGPointFilter.h
new file mode 100644
--- /dev/null
+++ b/Engine/Plugins/Experimental/PCG/Source/PCG/Public/Elements/PCGPointFilter.h
@@ -0,0 +1,134 @@
+// Copyright Epic Games, Inc. All Rights Reserved.
+
+#pragma once
+
+#include "PCGSettings.h"
+#include "PCGElement.h"
+#include "PCGPoint.h"
+
+#include "PCGPointFilter.generated.h"
+
+UENUM()
+enum class EPCGPointTargetFilterType : uint8
+{
+	Property,
+	Metadata
+};
+
+UENUM()
+enum class EPCGPointThresholdType : uint8
+{
+	Property,
+	Metadata,
+	Constant
+};
+
+UENUM()
+enum class EPCGPointFilterConstantType : uint8
+{
+	Integer64,
+	Float,
+	Vector,
+	Vector4,
+	//Rotation,
+	String,
+	Unknown UMETA(Hidden)
+};
+
+UENUM()
+enum class EPCGPointFilterOperator : uint8
+{
+	Greater UMETA(DisplayName=">"),
+	GreaterOrEqual UMETA(DisplayName=">="),
+	Lesser UMETA(DisplayName="<"),
+	LesserOrEqual UMETA(DisplayName="<="),
+	Equal UMETA(DisplayName="="),
+	NotEqual UMETA(DisplayName="!=")
+};
+
+/**
+* Point filter that allows to do "A op B" type filtering, where A is the input spatial data,
+* and B is either a constant, another spatial data, a param data (in filter) or the input itself.
+* The filtering can be done either on properties or attributes.
+* Some examples:
+* - Threshold on property by constant (A.Density > 0.5)
+* - Threshold on attribute by constant (A.aaa != "bob")
+* - Threshold on property by metadata attribute(A.density >= B.bbb)
+* - Threshold on property by property(A.density <= B.steepness)
+* - Threshold on attribute by metadata attribute(A.aaa < B.bbb)
+* - Threshold on attribute by property(A.aaa == B.color)
+*/
+UCLASS(BlueprintType, ClassGroup = (Procedural))
+class PCG_API UPCGPointFilterSettings : public UPCGSettings
+{
+	GENERATED_BODY()
+
+public:
+	//~Begin UPCGSettings interface
+#if WITH_EDITOR
+	virtual FName GetDefaultNodeName() const override { return FName(TEXT("PointFilterNode")); }
+#endif
+
+	virtual TArray<FName> InLabels() const override;
+	virtual TArray<FName> OutLabels() const override;
+	virtual bool HasDefaultLabels() const override { return false; }
+
+protected:
+	virtual FPCGElementPtr CreateElement() const override;
+	//~End UPCGSettings interface
+
+public:
+	UPROPERTY(BlueprintReadWrite, EditAnywhere, Category = Settings)
+	EPCGPointFilterOperator Operator = EPCGPointFilterOperator::Greater;
+
+	/** Target property/attribute related properties */
+	UPROPERTY(BlueprintReadWrite, EditAnywhere, Category = Settings)
+	EPCGPointTargetFilterType TargetFilterType = EPCGPointTargetFilterType::Property;
+
+	UPROPERTY(BlueprintReadWrite, EditAnywhere, Category = Settings, meta = (EditCondition = "TargetFilterType==EPCGPointTargetFilterType::Property"))
+	EPCGPointProperties TargetPointProperty = EPCGPointProperties::Density;
+
+	UPROPERTY(BlueprintReadWrite, EditAnywhere, Category = Settings, meta = (EditCondition = "TargetFilterType==EPCGPointTargetFilterType::Metadata"))
+	FName TargetAttributeName = NAME_None;
+
+	/** Threshold property/attribute/constant related properties */
+	UPROPERTY(BlueprintReadWrite, EditAnywhere, Category = Settings)
+	EPCGPointThresholdType ThresholdFilterType = EPCGPointThresholdType::Property;
+
+	UPROPERTY(BlueprintReadWrite, EditAnywhere, Category = Settings, meta = (EditCondition = "ThresholdFilterType==EPCGPointThresholdType::Property"))
+	EPCGPointProperties ThresholdPointProperty = EPCGPointProperties::Density;
+
+	UPROPERTY(BlueprintReadWrite, EditAnywhere, Category = Settings, meta = (EditCondition = "ThresholdFilterType==EPCGPointThresholdType::Metadata"))
+	FName ThresholdAttributeName = NAME_None;
+
+	UPROPERTY(BlueprintReadWrite, EditAnywhere, Category = Settings, meta = (EditCondition = "ThresholdFilterType==EPCGPointThresholdType::Constant"))
+	EPCGPointFilterConstantType ThresholdConstantType = EPCGPointFilterConstantType::Float;
+
+	/** Constants used as threshold comparator */
+	UPROPERTY(BlueprintReadWrite, EditAnywhere, Category = Settings, meta = (EditCondition = "ThresholdFilterType==EPCGPointThresholdType::Constant && ThresholdConstantType==EPCGPointFilterConstantType::Integer64"))
+	int64 Integer64Constant = 0;
+
+	UPROPERTY(BlueprintReadWrite, EditAnywhere, Category = Settings, meta=(EditCondition="ThresholdFilterType==EPCGPointThresholdType::Constant && ThresholdConstantType==EPCGPointFilterConstantType::Float"))
+	float FloatConstant = 0;
+
+	UPROPERTY(BlueprintReadWrite, EditAnywhere, Category = Settings, meta = (EditCondition = "ThresholdFilterType==EPCGPointThresholdType::Constant && ThresholdConstantType==EPCGPointFilterConstantType::Vector"))
+	FVector VectorConstant = FVector::Zero();
+
+	UPROPERTY(BlueprintReadWrite, EditAnywhere, Category = Settings, meta = (EditCondition = "ThresholdFilterType==EPCGPointThresholdType::Constant && ThresholdConstantType==EPCGPointFilterConstantType::Vector4"))
+	FVector4 Vector4Constant = FVector4::Zero();
+
+	//UPROPERTY(BlueprintReadWrite, EditAnywhere, Category = Settings, meta = (EditCondition = "ThresholdFilterType==EPCGPointThresholdType::Constant && ThresholdConstantType==EPCGPointFilterConstantType::Rotation"))
+	//FQuat RotationConstant = FQuat::Identity;
+
+	UPROPERTY(BlueprintReadWrite, EditAnywhere, Category = Settings, meta = (EditCondition = "ThresholdFilterType==EPCGPointThresholdType::Constant && ThresholdConstantType==EPCGPointFilterConstantType::String"))
+	FString StringConstant;
+
+	UPROPERTY(BlueprintReadWrite, EditAnywhere, Category = Settings, meta = (EditCondition = "ThresholdFilterType!=EPCGPointThresholdType::Constant"))
+	bool bUseSpatialQuery = true;
+};
+
+class FPCGPointFilterElement : public FSimplePCGElement
+{
+protected:
+	virtual bool ExecuteInternal(FPCGContext* Context) const override;
+};
\ No newline at end of file
diff --git a/Engine/Plugins/Experimental/PCG/Source/PCG/Public/Helpers/PCGAsync.h b/Engine/Plugins/Experimental/PCG/Source/PCG/Public/Helpers/PCGAsync.h
--- a/Engine/Plugins/Experimental/PCG/Source/PCG/Public/Helpers/PCGAsync.h
+++ b/Engine/Plugins/Experimental/PCG/Source/PCG/Public/Helpers/PCGAsync.h
@@ -27,4 +27,25 @@ namespace FPCGAsync
 	* @param PointFunc - A function that has the index [0; NumIterations] and has to write to the point & return true when the current call generates a point
 	*/
 	void AsyncPointProcessing(int32 NumAvailableTasks, int32 MinIterationsPerTask, int32 NumIterations, TArray<FPCGPoint>& OutPoints, const TFunction<bool(int32, FPCGPoint&)>& PointFunc);
+
+	/**
+	* Helper to do simple point filtering loops
+	* @param Context - The context containing the information about how many tasks to launch
+	* @param NumIterations - The number of calls that will be done to the provided function, also an upper bound on the number of points generated
+	* @param InFilterPoints - The array in which the in-filter results will be written to. Note that the array will be cleared before execution
+	* @param OutFilterPoints - The array in which the out-filter results will be written to. Note that the array will be cleared before execution
+	* @param PointFunc - A function that has the index [0; NumIterations] and has to write to the point & return true when the current call generates a point
+	*/
+	void AsyncPointFilterProcessing(FPCGContext* Context, int32 NumIterations, TArray<FPCGPoint>& InFilterPoints, TArray<FPCGPoint>& OutFilterPoints, const TFunction<bool(int32, FPCGPoint&, FPCGPoint&)>& PointFunc);
+
+	/**
+	* Helper to do more general 1:1 point filtering loops
+	* @param NumAvailableTasks - The upper bound on the number of async tasks we'll start
+	* @param MinIterationsPerTask - The lower bound on the number of iterations per task we'll dispatch
+	* @param NumIterations - The number of calls that will be done to the provided function, also an upper bound on the number of points generated
+	* @param InFilterPoints - The array in which the in-filter results will be written to. Note that the array will be cleared before execution
+	* @param OutFilterPoints - The array in which the out-filter results will be written to. Note that the array will be cleared before execution
+	* @param PointFunc - A function that has the index [0; NumIterations] and has to write to the point & return true when the current call generates a point
+	*/
+	void AsyncPointFilterProcessing(int32 NumAvailableTasks, int32 MinIterationsPerTask, int32 NumIterations, TArray<FPCGPoint>& InFilterPoints, TArray<FPCGPoint>& OutFilterPoints, const TFunction<bool(int32, FPCGPoint&, FPCGPoint&)>& PointFunc);
 }
\ No newline at end of file
diff --git a/Engine/Plugins/Experimental/PCG/Source/PCG/Public/PCGGraph.h b/Engine/Plugins/Experimental/PCG/Source/PCG/Public/PCGGraph.h
--- a/Engine/Plugins/Experimental/PCG/Source/PCG/Public/PCGGraph.h
+++ b/Engine/Plugins/Experimental/PCG/Source/PCG/Public/PCGGraph.h
@@ -62,8 +62,10 @@ public:
 	const TArray<UPCGNode*>& GetNodes() const { return Nodes; }
 	void RemoveNode(UPCGNode* InNode);
 	void RemoveEdge(UPCGNode* From, const FName& FromLabel, UPCGNode* To, const FName& ToLabel);
-	void RemoveInboundEdges(UPCGNode* InNode);
-	void RemoveOutboundEdges(UPCGNode* InNode);
+	void RemoveAllInboundEdges(UPCGNode* InNode);
+	void RemoveAllOutboundEdges(UPCGNode* InNode);
+	void RemoveInboundEdges(UPCGNode* InNode, const FName& InboundLabel);
+	void RemoveOutboundEdges(UPCGNode* InNode, const FName& OutboundLabel);
 
 #if WITH_EDITOR
 	FPCGTagToSettingsMap GetTrackedTagsToSettings() const;
diff --git a/Engine/Plugins/Experimental/PCG/Source/PCG/Public/PCGNode.h b/Engine/Plugins/Experimental/PCG/Source/PCG/Public/PCGNode.h
--- a/Engine/Plugins/Experimental/PCG/Source/PCG/Public/PCGNode.h
+++ b/Engine/Plugins/Experimental/PCG/Source/PCG/Public/PCGNode.h
@@ -52,6 +52,9 @@ public:
 	/** Returns all the output labels */
 	TArray<FName> OutLabels() const;
 
+	/** Returns true if we want the default In/Out labels */
+	bool HasDefaultLabels() const;
+
 	/** Returns true if the input pin is connected */
 	bool IsInputPinConnected(const FName& Label) const;
 
diff --git a/Engine/Plugins/Experimental/PCG/Source/PCG/Public/PCGSettings.h b/Engine/Plugins/Experimental/PCG/Source/PCG/Public/PCGSettings.h
--- a/Engine/Plugins/Experimental/PCG/Source/PCG/Public/PCGSettings.h
+++ b/Engine/Plugins/Experimental/PCG/Source/PCG/Public/PCGSettings.h
@@ -43,6 +43,7 @@ public:
 	virtual bool HasOutLabel(const FName& Label) const;
 	virtual TArray<FName> InLabels() const { return TArray<FName>(); }
 	virtual TArray<FName> OutLabels() const { return TArray<FName>(); }
+	virtual bool HasDefaultLabels() const { return true; }
 	
 	bool operator==(const UPCGSettings& Other) const;
 
diff --git a/Engine/Plugins/Experimental/PCG/Source/PCGEditor/Private/PCGEditorGraphNode.cpp b/Engine/Plugins/Experimental/PCG/Source/PCGEditor/Private/PCGEditorGraphNode.cpp
--- a/Engine/Plugins/Experimental/PCG/Source/PCGEditor/Private/PCGEditorGraphNode.cpp
+++ b/Engine/Plugins/Experimental/PCG/Source/PCGEditor/Private/PCGEditorGraphNode.cpp
@@ -101,7 +101,10 @@ void UPCGEditorGraphNode::AllocateDefaultPins()
 {
 	if (NodeType == EPCGEditorGraphNodeType::Input || NodeType == EPCGEditorGraphNodeType::Settings)
 	{
-		CreatePin(EEdGraphPinDirection::EGPD_Output, NAME_None, FName(TEXT("Out")));
+		if (!PCGNode || PCGNode->HasDefaultLabels())
+		{
+			CreatePin(EEdGraphPinDirection::EGPD_Output, NAME_None, FName(TEXT("Out")));
+		}
 
 		if (PCGNode)
 		{
@@ -114,7 +117,10 @@ void UPCGEditorGraphNode::AllocateDefaultPins()
 
 	if (NodeType == EPCGEditorGraphNodeType::Output || NodeType == EPCGEditorGraphNodeType::Settings)
 	{
-		CreatePin(EEdGraphPinDirection::EGPD_Input, NAME_None, FName(TEXT("In")));
+		if (!PCGNode || PCGNode->HasDefaultLabels())
+		{
+			CreatePin(EEdGraphPinDirection::EGPD_Input, NAME_None, FName(TEXT("In")));
+		}
 
 		if (PCGNode)
 		{
@@ -130,12 +136,24 @@ void UPCGEditorGraphNode::AutowireNewNode(UEdGraphPin* FromPin)
 {
 	if (FromPin->Direction == EEdGraphPinDirection::EGPD_Output)
 	{
-		UEdGraphPin* ToPin = FindPinChecked(TEXT("In"), EEdGraphPinDirection::EGPD_Input);
+		FName InPinName = TEXT("In");
+		if (PCGNode && !PCGNode->HasDefaultLabels() && !PCGNode->InLabels().IsEmpty())
+		{
+			InPinName = PCGNode->InLabels()[0];
+		}
+
+		UEdGraphPin* ToPin = FindPinChecked(InPinName, EEdGraphPinDirection::EGPD_Input);
 		GetSchema()->TryCreateConnection(FromPin, ToPin);
 	}
 	else if (FromPin->Direction == EEdGraphPinDirection::EGPD_Input)
 	{
-		UEdGraphPin* ToPin = FindPinChecked(TEXT("Out"), EEdGraphPinDirection::EGPD_Output);
+		FName OutPinName = TEXT("Out");
+		if (PCGNode && !PCGNode->HasDefaultLabels() && !PCGNode->OutLabels().IsEmpty())
+		{
+			OutPinName = PCGNode->OutLabels()[0];
+		}
+
+		UEdGraphPin* ToPin = FindPinChecked(OutPinName, EEdGraphPinDirection::EGPD_Output);
 		GetSchema()->TryCreateConnection(FromPin, ToPin);
 	}
 	NodeConnectionListChanged();
@@ -169,11 +187,11 @@ void UPCGEditorGraphNode::ReconstructNode()
 
 	for (UEdGraphPin* Pin : Pins)
 	{
-		if (Pin->Direction == EEdGraphPinDirection::EGPD_Input && Pin->PinName != TEXT("In"))
+		if (Pin->Direction == EEdGraphPinDirection::EGPD_Input && (Pin->PinName != TEXT("In") || !PCGNode->HasDefaultLabels()))
 		{
 			InputPins.Add(Pin);
 		}
-		else if (Pin->Direction == EEdGraphPinDirection::EGPD_Output && Pin->PinName != TEXT("Out"))
+		else if (Pin->Direction == EEdGraphPinDirection::EGPD_Output && (Pin->PinName != TEXT("Out") || !PCGNode->HasDefaultLabels()))
 		{
 			OutputPins.Add(Pin);
 		}
diff --git a/Engine/Plugins/Experimental/PCG/Source/PCGEditor/Private/PCGEditorGraphSchema.cpp b/Engine/Plugins/Experimental/PCG/Source/PCGEditor/Private/PCGEditorGraphSchema.cpp
--- a/Engine/Plugins/Experimental/PCG/Source/PCGEditor/Private/PCGEditorGraphSchema.cpp
+++ b/Engine/Plugins/Experimental/PCG/Source/PCGEditor/Private/PCGEditorGraphSchema.cpp
@@ -88,8 +88,8 @@ bool UPCGEditorGraphSchema::TryCreateConnection(UEdGraphPin* InA, UEdGraphPin* I
 		UPCGGraph* PCGGraph = PCGNodeA->GetGraph();
 		check(PCGGraph);
 
-		const FName& NodeAPinName = (A->PinName == TEXT("Out") ? NAME_None : A->PinName);
-		const FName& NodeBPinName = (B->PinName == TEXT("In") ? NAME_None : B->PinName);
+		const FName& NodeAPinName = ((A->PinName == TEXT("Out") && PCGNodeA->HasDefaultLabels()) ? NAME_None : A->PinName);
+		const FName& NodeBPinName = ((B->PinName == TEXT("In") && PCGNodeB->HasDefaultLabels()) ? NAME_None : B->PinName);
 
 		PCGGraph->AddLabeledEdge(PCGNodeA, NodeAPinName, PCGNodeB, NodeBPinName);
 	}
@@ -113,11 +113,13 @@ void UPCGEditorGraphSchema::BreakPinLinks(UEdGraphPin& TargetPin, bool bSendsNod
 
 	if (TargetPin.Direction == EEdGraphPinDirection::EGPD_Input)
 	{
-		PCGGraph->RemoveInboundEdges(PCGNode);
+		const FName& PinName = (TargetPin.PinName == TEXT("In") && PCGNode->HasDefaultLabels()) ? NAME_None : TargetPin.PinName;
+		PCGGraph->RemoveInboundEdges(PCGNode, PinName);
 	}
 	else if (TargetPin.Direction == EEdGraphPinDirection::EGPD_Output)
 	{
-		PCGGraph->RemoveOutboundEdges(PCGNode);
+		const FName& PinName = (TargetPin.PinName == TEXT("Out") && PCGNode->HasDefaultLabels()) ? NAME_None : TargetPin.PinName;
+		PCGGraph->RemoveOutboundEdges(PCGNode, PinName);
 	}
 }
 
