Exposed metadata add attributes (copy) call

diff --git a/Engine/Plugins/Experimental/PCG/Source/PCG/Public/Metadata/PCGMetadata.h b/Engine/Plugins/Experimental/PCG/Source/PCG/Public/Metadata/PCGMetadata.h
--- a/Engine/Plugins/Experimental/PCG/Source/PCG/Public/Metadata/PCGMetadata.h
+++ b/Engine/Plugins/Experimental/PCG/Source/PCG/Public/Metadata/PCGMetadata.h
@@ -41,6 +41,7 @@ public:
 	void InitializeAsCopy(const UPCGMetadata* InMetadataToCopy);
 
 	/** Creates streams for another metadata if they are not currently present */
+	UFUNCTION(BlueprintCallable, Category = "PCG|Metadata")
 	void AddAttributes(const UPCGMetadata* InOther);
 
 	/** Returns this metadata's parent */
