Add metadata partition element

diff --git a/Engine/Plugins/Experimental/PCG/Source/PCG/Private/Elements/PCGMetadataPartition.cpp b/Engine/Plugins/Experimental/PCG/Source/PCG/Private/Elements/PCGMetadataPartition.cpp
new file mode 100644
--- /dev/null
+++ b/Engine/Plugins/Experimental/PCG/Source/PCG/Private/Elements/PCGMetadataPartition.cpp
@@ -0,0 +1,124 @@
+// Copyright Epic Games, Inc. All Rights Reserved.
+
+#include "Elements/PCGMetadataPartition.h"
+#include "Data/PCGSpatialData.h"
+#include "Data/PCGPointData.h"
+#include "Helpers/PCGSettingsHelpers.h"
+
+FPCGElementPtr UPCGMetadataPartitionSettings::CreateElement() const
+{
+	return MakeShared<FPCGMetadataPartitionElement>();
+}
+
+bool FPCGMetadataPartitionElement::ExecuteInternal(FPCGContext* Context) const
+{
+	TRACE_CPUPROFILER_EVENT_SCOPE(FPCGMetadataPartitionElement::Execute);
+	check(Context);
+
+	const UPCGMetadataPartitionSettings* Settings = Context->GetInputSettings<UPCGMetadataPartitionSettings>();
+	check(Settings);
+
+	TArray<FPCGTaggedData> Inputs = Context->InputData.GetInputs();
+	UPCGParamData* Params = Context->InputData.GetParams();
+	TArray<FPCGTaggedData>& Outputs = Context->OutputData.TaggedData;
+
+	// Forward any non-input data
+	Outputs.Append(Context->InputData.GetAllSettings());
+
+	const FName PartitionAttribute = PCGSettingsHelpers::GetValue(GET_MEMBER_NAME_CHECKED(UPCGMetadataPartitionSettings, PartitionAttribute), Settings->PartitionAttribute, Params);
+
+	if (PartitionAttribute == NAME_None)
+	{
+		PCGE_LOG(Error, "Metadata partition has no specified attribute");
+		return true;
+	}
+
+	for (const FPCGTaggedData& Input : Inputs)
+	{
+		const UPCGSpatialData* SpatialInput = Cast<const UPCGSpatialData>(Input.Data);
+		check(SpatialInput);
+
+		if (!SpatialInput->ConstMetadata() || !SpatialInput->ConstMetadata()->HasAttribute(PartitionAttribute))
+		{
+			PCGE_LOG(Error, "Input does not have the %s attribute", *PartitionAttribute.ToString());
+			continue;
+		}
+
+		// Currently, we support only to parse points, so we'll do so here
+		const UPCGPointData* InputData = SpatialInput->ToPointData(Context);
+
+		if (!InputData)
+		{
+			PCGE_LOG(Error, "Unable to get points out of spatial data");
+			continue;
+		}
+
+		const TArray<FPCGPoint>& SourcePoints = InputData->GetPoints();
+		const UPCGMetadata* SourceMetadata = InputData->ConstMetadata();
+		const FPCGMetadataAttributeBase* AttributeBase = SourceMetadata->GetConstAttribute(PartitionAttribute);
+		check(AttributeBase);
+
+		TMap<PCGMetadataValueKey, UPCGPointData*> ValueToData;
+
+		// Get all value keys (-1 + 0 - N)
+		int64 MetadataValueKeyCount = AttributeBase->GetValueKeyOffsetForChild();
+
+		// For every value key, check if it should be merged with the default value
+		// Keep track of that mapping (should be only one, but that might not be true in uncompressed data)
+		TArray<PCGMetadataValueKey> ValueKeyMapping;
+		ValueKeyMapping.Reserve(MetadataValueKeyCount);
+
+		for (PCGMetadataValueKey ValueKey = 0; ValueKey < MetadataValueKeyCount; ++ValueKey)
+		{
+			if (AttributeBase->IsEqualToDefaultValue(ValueKey))
+			{
+				ValueKeyMapping.Add(-1);
+			}
+			else
+			{
+				ValueKeyMapping.Add(ValueKey);
+			}
+		}
+
+		TArray<UPCGPointData*> PartitionedData;
+		PartitionedData.SetNum(1 + ValueKeyMapping.Num());
+		for (int32 Index = 0; Index < PartitionedData.Num(); ++Index)
+		{
+			PartitionedData[Index] = nullptr;
+		}
+
+		// Loop on points, find matching data pointer, allocate if null, etc.and add point to it
+		for (const FPCGPoint& Point : SourcePoints)
+		{
+			PCGMetadataValueKey ValueKey = AttributeBase->GetValueKey(Point.MetadataEntry);
+			// Remap
+			if (ValueKey != PCGDefaultValueKey)
+			{
+				ValueKey = ValueKeyMapping[ValueKey];
+			}
+
+			const int32 PartitionDataIndex = 1 + ValueKey;
+
+			if (!PartitionedData[PartitionDataIndex])
+			{
+				PartitionedData[PartitionDataIndex] = NewObject<UPCGPointData>();
+				PartitionedData[PartitionDataIndex]->InitializeFromData(InputData);
+			}
+
+			PartitionedData[PartitionDataIndex]->GetMutablePoints().Add(Point);
+		}
+
+		// Finally, push back the partitioned data to the outputs
+		for (UPCGPointData* PartitionData : PartitionedData)
+		{
+			if (PartitionData)
+			{
+				FPCGTaggedData& Output = Outputs.Emplace_GetRef();
+				Output.Data = PartitionData;
+				Output.Tags = Input.Tags;
+			}
+		}
+	}
+
+	return true;
+}
\ No newline at end of file
diff --git a/Engine/Plugins/Experimental/PCG/Source/PCG/Public/Elements/PCGMetadataPartition.h b/Engine/Plugins/Experimental/PCG/Source/PCG/Public/Elements/PCGMetadataPartition.h
new file mode 100644
--- /dev/null
+++ b/Engine/Plugins/Experimental/PCG/Source/PCG/Public/Elements/PCGMetadataPartition.h
@@ -0,0 +1,34 @@
+// Copyright Epic Games, Inc. All Rights Reserved.
+
+#pragma once
+
+#include "PCGSettings.h"
+#include "PCGElement.h"
+
+#include "PCGMetadataPartition.generated.h"
+
+UCLASS(BlueprintType, ClassGroup = (Procedural))
+class PCG_API UPCGMetadataPartitionSettings : public UPCGSettings
+{
+	GENERATED_BODY()
+
+public:
+	//~Begin UPCGSettings interface
+#if WITH_EDITOR
+	virtual FName GetDefaultNodeName() const override { return FName(TEXT("MetadataPartitionNode")); }
+#endif
+
+protected:
+	virtual FPCGElementPtr CreateElement() const override;
+	//~End UPCGSettings interface
+
+public:
+	UPROPERTY(BlueprintReadWrite, EditAnywhere, Category = Settings)
+	FName PartitionAttribute = NAME_None;
+};
+
+class FPCGMetadataPartitionElement : public FSimplePCGElement
+{
+protected:
+	virtual bool ExecuteInternal(FPCGContext* Context) const override;
+};
\ No newline at end of file
diff --git a/Engine/Plugins/Experimental/PCG/Source/PCG/Public/Metadata/PCGMetadataAttribute.h b/Engine/Plugins/Experimental/PCG/Source/PCG/Public/Metadata/PCGMetadataAttribute.h
--- a/Engine/Plugins/Experimental/PCG/Source/PCG/Public/Metadata/PCGMetadataAttribute.h
+++ b/Engine/Plugins/Experimental/PCG/Source/PCG/Public/Metadata/PCGMetadataAttribute.h
@@ -22,10 +22,12 @@ public:
 
 	virtual FPCGMetadataAttributeBase* Copy(FName NewName, UPCGMetadata* InMetadata, bool bKeepParent, bool bCopyEntries = true, bool bCopyValues = true) const = 0;
 
+	virtual PCGMetadataValueKey GetValueKeyOffsetForChild() const = 0;
 	virtual void SetValue(PCGMetadataEntryKey ItemKey, const FPCGMetadataAttributeBase* InAttribute, PCGMetadataEntryKey InEntryKey) = 0;
 	virtual void SetZeroValue(PCGMetadataEntryKey ItemKey) = 0;
 	virtual void AccumulateValue(PCGMetadataEntryKey ItemKey, const FPCGMetadataAttributeBase* InAttributeA, PCGMetadataEntryKey InEntryKeyA, float Weight) = 0;
 	virtual void SetValue(PCGMetadataEntryKey ItemKey, const FPCGMetadataAttributeBase* InAttributeA, PCGMetadataEntryKey InEntryKeyA, const FPCGMetadataAttributeBase* InAttributeB, PCGMetadataEntryKey InEntryKeyB, EPCGMetadataOp Op) = 0;
+	virtual bool IsEqualToDefaultValue(PCGMetadataValueKey ValueKey) const = 0;
 
 	void SetValueFromValueKey(PCGMetadataEntryKey EntryKey, PCGMetadataValueKey ValueKey);
 	PCGMetadataValueKey GetValueKey(PCGMetadataEntryKey EntryKey) const;
diff --git a/Engine/Plugins/Experimental/PCG/Source/PCG/Public/Metadata/PCGMetadataAttributeTpl.h b/Engine/Plugins/Experimental/PCG/Source/PCG/Public/Metadata/PCGMetadataAttributeTpl.h
--- a/Engine/Plugins/Experimental/PCG/Source/PCG/Public/Metadata/PCGMetadataAttributeTpl.h
+++ b/Engine/Plugins/Experimental/PCG/Source/PCG/Public/Metadata/PCGMetadataAttributeTpl.h
@@ -139,6 +139,11 @@ public:
 		}
 	}
 
+	virtual bool IsEqualToDefaultValue(PCGMetadataValueKey ValueKey) const
+	{
+		return PCG::Private::MetadataTraits<T>::Equal(GetValue(ValueKey), DefaultValue);
+	}
+
 	PCGMetadataValueKey GetValueKeyOffsetForChild() const
 	{
 		FReadScopeLock ScopeLock(ValueLock);
@@ -365,7 +370,6 @@ protected:
 		// Empty on purpose
 	}
 
-
 protected:
 	mutable FRWLock ValueLock;
 	TArray<T> Values;
diff --git a/Engine/Plugins/Experimental/PCG/Source/PCG/Public/Metadata/PCGMetadataAttributeTraits.h b/Engine/Plugins/Experimental/PCG/Source/PCG/Public/Metadata/PCGMetadataAttributeTraits.h
--- a/Engine/Plugins/Experimental/PCG/Source/PCG/Public/Metadata/PCGMetadataAttributeTraits.h
+++ b/Engine/Plugins/Experimental/PCG/Source/PCG/Public/Metadata/PCGMetadataAttributeTraits.h
@@ -29,24 +29,20 @@ namespace PCG
 
 #undef PCGMetadataGenerateDataTypes
 
-		// Common traits for int32, int64, float, double
 		template<typename T>
-		struct MetadataTraits
+		struct DefaultOperationTraits
 		{
-			enum { CompressData = false };
-			enum { CanMinMax = true };
 			enum { CanSubAdd = true };
 			enum { CanMulDiv = true };
-			enum { CanInterpolate = true };
 			
-			static T Min(const T& A, const T& B)
+			static bool Equal(const T& A, const T& B)
 			{
-				return FMath::Min(A, B);
+				return A == B;
 			}
 
-			static T Max(const T& A, const T& B)
+			static T Sub(const T& A, const T& B)
 			{
-				return FMath::Max(A, B);
+				return A - B;
 			}
 
 			static T Add(const T& A, const T& B)
@@ -54,11 +50,6 @@ namespace PCG
 				return A + B;
 			}
 
-			static T Sub(const T& A, const T& B)
-			{
-				return A - B;
-			}
-
 			static T Mul(const T& A, const T& B)
 			{
 				return A * B;
@@ -68,6 +59,12 @@ namespace PCG
 			{
 				return A / B;
 			}
+		};
+
+		template<typename T>
+		struct DefaultWeightedSumTraits
+		{
+			enum { CanInterpolate = true };
 
 			static T WeightedSum(const T& A, const T& B, float Weight)
 			{
@@ -80,6 +77,24 @@ namespace PCG
 			}
 		};
 
+		// Common traits for int32, int64, float, double
+		template<typename T>
+		struct MetadataTraits : DefaultOperationTraits<T>, DefaultWeightedSumTraits<T>
+		{
+			enum { CompressData = false };
+			enum { CanMinMax = true };
+
+			static T Min(const T& A, const T& B)
+			{
+				return FMath::Min(A, B);
+			}
+
+			static T Max(const T& A, const T& B)
+			{
+				return FMath::Max(A, B);
+			}
+		};
+
 		template<>
 		struct MetadataTraits<bool>
 		{
@@ -89,6 +104,11 @@ namespace PCG
 			enum { CanMulDiv = false };
 			enum { CanInterpolate = false };
 
+			static bool Equal(const bool& A, const bool& B)
+			{
+				return A == B;
+			}
+
 			static bool Min(const bool& A, const bool& B)
 			{
 				return A && B;
@@ -112,13 +132,11 @@ namespace PCG
 
 		// Vector types
 		template<>
-		struct MetadataTraits<FVector>
+		struct MetadataTraits<FVector> : DefaultOperationTraits<FVector>
 		{
 			enum { CompressData = false };
 			enum { CanMinMax = true };
-			enum { CanSubAdd = true };
-			enum { CanMulDiv = true };
-			enum { CanInterpolate = true};
+			enum { CanInterpolate = true };
 
 			static FVector Min(const FVector& A, const FVector& B) 
 			{
@@ -130,26 +148,6 @@ namespace PCG
 				return FVector(FMath::Max(A.X, B.X), FMath::Max(A.Y, B.Y), FMath::Max(A.Z, B.Z));
 			}
 
-			static FVector Add(const FVector& A, const FVector& B)
-			{
-				return A + B;
-			}
-
-			static FVector Sub(const FVector& A, const FVector& B)
-			{
-				return A - B;
-			}
-
-			static FVector Mul(const FVector& A, const FVector& B)
-			{
-				return A * B;
-			}
-
-			static FVector Div(const FVector& A, const FVector& B)
-			{
-				return A / B;
-			}
-
 			static FVector WeightedSum(const FVector& A, const FVector& B, float Weight)
 			{
 				return A + B * Weight;
@@ -157,17 +155,15 @@ namespace PCG
 
 			static FVector ZeroValue()
 			{
-				return FVector::Zero();
+				return FVector::ZeroVector;
 			}
 		};
 
 		template<>
-		struct MetadataTraits<FVector4>
+		struct MetadataTraits<FVector4> : DefaultOperationTraits<FVector4>
 		{
 			enum { CompressData = false };
 			enum { CanMinMax = true };
-			enum { CanSubAdd = true };
-			enum { CanMulDiv = true };
 			enum { CanInterpolate = true };
 
 			static FVector4 Min(const FVector4& A, const FVector4& B)
@@ -180,26 +176,6 @@ namespace PCG
 				return FVector4(FMath::Max(A.X, B.X), FMath::Max(A.Y, B.Y), FMath::Max(A.Z, B.Z), FMath::Max(A.W, B.W));
 			}
 
-			static FVector4 Add(const FVector4& A, const FVector4& B)
-			{
-				return A + B;
-			}
-
-			static FVector4 Sub(const FVector4& A, const FVector4& B)
-			{
-				return A - B;
-			}
-
-			static FVector4 Mul(const FVector& A, const FVector4& B)
-			{
-				return A * B;
-			}
-
-			static FVector4 Div(const FVector& A, const FVector4& B)
-			{
-				return A / B;
-			}
-
 			static FVector4 WeightedSum(const FVector4& A, const FVector4& B, float Weight)
 			{
 				return A + B * Weight;
@@ -207,7 +183,7 @@ namespace PCG
 
 			static FVector4 ZeroValue()
 			{
-				return FVector::Zero();
+				return FVector4::Zero();
 			}
 		};
 
@@ -221,6 +197,11 @@ namespace PCG
 			enum { CanMulDiv = true };
 			enum { CanInterpolate = true };
 
+			static bool Equal(const FQuat& A, const FQuat& B)
+			{
+				return A == B;
+			}
+
 			static FQuat Add(const FQuat& A, const FQuat& B)
 			{
 				return A * B;
@@ -268,6 +249,13 @@ namespace PCG
 			enum { CanMulDiv = true };
 			enum { CanInterpolate = true };
 
+			static bool Equal(const FTransform& A, const FTransform& B)
+			{
+				return A.GetLocation() == B.GetLocation() &&
+					A.GetRotation() == B.GetRotation() &&
+					A.GetScale3D() == B.GetScale3D();
+			}
+
 			static FTransform Add(const FTransform& A, const FTransform& B)
 			{
 				return A * B;
@@ -322,6 +310,11 @@ namespace PCG
 			enum { CanSubAdd = false };
 			enum { CanMulDiv = false };
 			enum { CanInterpolate = false };
+
+			static bool Equal(const FString& A, const FString& B)
+			{
+				return A == B;
+			}
 		};
 	}
 }
\ No newline at end of file
