Add node labels
Made the editor enabled by default
Added option to have the default pins on blueprint elements

diff --git a/Engine/Plugins/Experimental/PCG/Source/PCG/Private/Elements/PCGExecuteBlueprint.cpp b/Engine/Plugins/Experimental/PCG/Source/PCG/Private/Elements/PCGExecuteBlueprint.cpp
--- a/Engine/Plugins/Experimental/PCG/Source/PCG/Private/Elements/PCGExecuteBlueprint.cpp
+++ b/Engine/Plugins/Experimental/PCG/Source/PCG/Private/Elements/PCGExecuteBlueprint.cpp
@@ -89,6 +89,11 @@ void UPCGBlueprintElement::OnDependencyChanged(UObject* Object, FPropertyChanged
 
 #endif // WITH_EDITOR
 
+FName UPCGBlueprintElement::NodeTitleOverride_Implementation() const
+{
+	return NAME_None;
+}
+
 void UPCGBlueprintSettings::SetupBlueprintEvent()
 {
 #if WITH_EDITOR
@@ -215,7 +220,7 @@ void UPCGBlueprintSettings::OnBlueprintElementChanged(UPCGBlueprintElement* InEl
 		OnSettingsChangedDelegate.Broadcast(this);
 	}
 }
-#endif
+#endif // WITH_EDITOR
 
 void UPCGBlueprintSettings::SetElementType(TSubclassOf<UPCGBlueprintElement> InElementType, UPCGBlueprintElement*& ElementInstance)
 {
@@ -264,31 +269,38 @@ void UPCGBlueprintSettings::GetTrackedActorTags(FPCGTagToSettingsMap& OutTagToSe
 
 FName UPCGBlueprintSettings::AdditionalTaskName() const
 {
+	if (BlueprintElementInstance && BlueprintElementInstance->NodeTitleOverride() != NAME_None)
+	{
+		return BlueprintElementInstance->NodeTitleOverride();
+	}
+	else
+	{
 #if WITH_EDITOR
-	return BlueprintElementType && BlueprintElementType->ClassGeneratedBy ? BlueprintElementType->ClassGeneratedBy->GetFName() : Super::AdditionalTaskName();
+		return BlueprintElementType && BlueprintElementType->ClassGeneratedBy ? BlueprintElementType->ClassGeneratedBy->GetFName() : Super::AdditionalTaskName();
 #else
-	return BlueprintElementType ? BlueprintElementType->GetFName() : Super::AdditionalTaskName();
+		return BlueprintElementType ? BlueprintElementType->GetFName() : Super::AdditionalTaskName();
 #endif
+	}
 }
 
-bool UPCGBlueprintSettings::HasInLabel(const FName& Label) const
+TArray<FName> UPCGBlueprintSettings::InLabels() const
 {
-	return (Label == NAME_None || (BlueprintElementInstance && BlueprintElementInstance->InputPinLabels.Contains(Label)));
+	return BlueprintElementInstance ? BlueprintElementInstance->InputPinLabels.Array() : TArray<FName>();
 }
 
-bool UPCGBlueprintSettings::HasOutLabel(const FName& Label) const
+TArray<FName> UPCGBlueprintSettings::OutLabels() const
 {
-	return (Label == NAME_None || (BlueprintElementInstance && BlueprintElementInstance->OutputPinLabels.Contains(Label)));
+	return BlueprintElementInstance ? BlueprintElementInstance->OutputPinLabels.Array() : TArray<FName>();
 }
 
-TArray<FName> UPCGBlueprintSettings::InLabels() const
+bool UPCGBlueprintSettings::HasDefaultInLabel() const
 {
-	return BlueprintElementInstance ? BlueprintElementInstance->InputPinLabels.Array() : TArray<FName>();
+	return !BlueprintElementInstance || BlueprintElementInstance->bHasDefaultInPin;
 }
 
-TArray<FName> UPCGBlueprintSettings::OutLabels() const
+bool UPCGBlueprintSettings::HasDefaultOutLabel() const
 {
-	return BlueprintElementInstance ? BlueprintElementInstance->OutputPinLabels.Array() : TArray<FName>();
+	return !BlueprintElementInstance || BlueprintElementInstance->bHasDefaultOutPin;
 }
 
 FPCGElementPtr UPCGBlueprintSettings::CreateElement() const
diff --git a/Engine/Plugins/Experimental/PCG/Source/PCG/Private/PCGContext.cpp b/Engine/Plugins/Experimental/PCG/Source/PCG/Private/PCGContext.cpp
--- a/Engine/Plugins/Experimental/PCG/Source/PCG/Private/PCGContext.cpp
+++ b/Engine/Plugins/Experimental/PCG/Source/PCG/Private/PCGContext.cpp
@@ -8,12 +8,12 @@ FString FPCGContext::GetTaskName() const
 {
 	if (Node)
 	{
-		const FName NodeName = Node->GetFName();
+		const FName NodeName = (Node->NodeTitle != NAME_None ? Node->NodeTitle : Node->GetFName());
 
 		const UPCGSettings* Settings = GetInputSettings<UPCGSettings>();
 		const FName NodeAdditionalName = Settings ? Settings->AdditionalTaskName() : NAME_None;
 
-		if (NodeAdditionalName == NAME_None)
+		if (NodeAdditionalName == NAME_None || NodeAdditionalName == NodeName)
 		{
 			return NodeName.ToString();
 		}
diff --git a/Engine/Plugins/Experimental/PCG/Source/PCG/Private/PCGGraph.cpp b/Engine/Plugins/Experimental/PCG/Source/PCG/Private/PCGGraph.cpp
--- a/Engine/Plugins/Experimental/PCG/Source/PCG/Private/PCGGraph.cpp
+++ b/Engine/Plugins/Experimental/PCG/Source/PCG/Private/PCGGraph.cpp
@@ -17,6 +17,9 @@ UPCGGraph::UPCGGraph(const FObjectInitializer& ObjectInitializer)
 	InputNode->DefaultSettings = ObjectInitializer.CreateDefaultSubobject<UPCGGraphInputOutputSettings>(this, TEXT("DefaultInputNodeSettings"));
 	OutputNode = ObjectInitializer.CreateDefaultSubobject<UPCGNode>(this, TEXT("DefaultOutputNode"));
 	OutputNode->DefaultSettings = ObjectInitializer.CreateDefaultSubobject<UPCGGraphInputOutputSettings>(this, TEXT("DefaultOutputNodeSettings"));
+#if WITH_EDITORONLY_DATA
+	OutputNode->PositionX = 200;
+#endif
 
 	// Note: default connection from input to output
 	// should be added when creating from scratch,
@@ -431,9 +434,12 @@ void UPCGGraph::NotifyGraphChanged(bool bIsStructural)
 	OnGraphChangedDelegate.Broadcast(this, /*bIsStructural=*/bIsStructural);
 }
 
-void UPCGGraph::OnSettingsChanged(UPCGNode* InNode)
+void UPCGGraph::OnSettingsChanged(UPCGNode* InNode, bool bSettingsChanged)
 {
-	NotifyGraphChanged(/*bIsStructural=*/false);
+	if (bSettingsChanged)
+	{
+		NotifyGraphChanged(/*bIsStructural=*/false);
+	}
 }
 
 void UPCGGraph::OnStructuralSettingsChanged(UPCGNode* InNode)
diff --git a/Engine/Plugins/Experimental/PCG/Source/PCG/Private/PCGNode.cpp b/Engine/Plugins/Experimental/PCG/Source/PCG/Private/PCGNode.cpp
--- a/Engine/Plugins/Experimental/PCG/Source/PCG/Private/PCGNode.cpp
+++ b/Engine/Plugins/Experimental/PCG/Source/PCG/Private/PCGNode.cpp
@@ -63,6 +63,29 @@ UPCGNode* UPCGNode::AddEdgeTo(FName InboundName, UPCGNode* To, FName OutboundNam
 	}
 }
 
+FName UPCGNode::GetNodeTitle() const
+{
+	if (NodeTitle != NAME_None)
+	{
+		return NodeTitle;
+	}
+	else if (DefaultSettings)
+	{
+		if (DefaultSettings->AdditionalTaskName() != NAME_None)
+		{
+			return DefaultSettings->AdditionalTaskName();
+		}
+#if WITH_EDITOR
+		else
+		{
+			return DefaultSettings->GetDefaultNodeName();
+		}
+#endif
+	}
+
+	return TEXT("Unnamed node");
+}
+
 bool UPCGNode::HasInLabel(const FName& Label) const
 {
 	return DefaultSettings && DefaultSettings->HasInLabel(Label);
@@ -83,9 +106,14 @@ TArray<FName> UPCGNode::OutLabels() const
 	return DefaultSettings ? DefaultSettings->OutLabels() : TArray<FName>();
 }
 
-bool UPCGNode::HasDefaultLabels() const
+bool UPCGNode::HasDefaultInLabel() const
 {
-	return !DefaultSettings || DefaultSettings->HasDefaultLabels();
+	return !DefaultSettings || DefaultSettings->HasDefaultInLabel();
+}
+
+bool UPCGNode::HasDefaultOutLabel() const
+{
+	return !DefaultSettings || DefaultSettings->HasDefaultOutLabel();
 }
 
 bool UPCGNode::IsInputPinConnected(const FName& Label) const
@@ -196,12 +224,12 @@ void UPCGNode::PostEditChangeProperty(struct FPropertyChangedEvent& PropertyChan
 		if (DefaultSettings)
 		{
 			DefaultSettings->OnSettingsChangedDelegate.AddUObject(this, &UPCGNode::OnSettingsChanged);
+			OnSettingsChanged(DefaultSettings);
 		}
 	}
-
-	if (DefaultSettings)
+	else if (PropertyChangedEvent.Property && PropertyChangedEvent.Property->GetFName() == GET_MEMBER_NAME_CHECKED(UPCGNode, NodeTitle))
 	{
-		OnSettingsChanged(DefaultSettings);
+		OnNodeSettingsChangedDelegate.Broadcast(this, /*bSettingsChanged=*/false);
 	}
 }
 
@@ -209,7 +237,7 @@ void UPCGNode::OnSettingsChanged(UPCGSettings* InSettings)
 {
 	if (InSettings == DefaultSettings)
 	{
-		OnNodeSettingsChangedDelegate.Broadcast(this);
+		OnNodeSettingsChangedDelegate.Broadcast(this, /*bSettingsChanged=*/true);
 	}
 }
 
diff --git a/Engine/Plugins/Experimental/PCG/Source/PCG/Private/PCGSettings.cpp b/Engine/Plugins/Experimental/PCG/Source/PCG/Private/PCGSettings.cpp
--- a/Engine/Plugins/Experimental/PCG/Source/PCG/Private/PCGSettings.cpp
+++ b/Engine/Plugins/Experimental/PCG/Source/PCG/Private/PCGSettings.cpp
@@ -39,12 +39,12 @@ bool UPCGSettings::operator==(const UPCGSettings& Other) const
 
 bool UPCGSettings::HasInLabel(const FName& Label) const
 {
-	return (Label == NAME_None && HasDefaultLabels()) || InLabels().Contains(Label);
+	return (Label == NAME_None && HasDefaultInLabel()) || InLabels().Contains(Label);
 }
 
 bool UPCGSettings::HasOutLabel(const FName& Label) const
 {
-	return (Label == NAME_None && HasDefaultLabels()) || OutLabels().Contains(Label);
+	return (Label == NAME_None && HasDefaultOutLabel()) || OutLabels().Contains(Label);
 }
 
 FPCGElementPtr UPCGSettings::GetElement() const
diff --git a/Engine/Plugins/Experimental/PCG/Source/PCG/Public/Elements/PCGExecuteBlueprint.h b/Engine/Plugins/Experimental/PCG/Source/PCG/Public/Elements/PCGExecuteBlueprint.h
--- a/Engine/Plugins/Experimental/PCG/Source/PCG/Public/Elements/PCGExecuteBlueprint.h
+++ b/Engine/Plugins/Experimental/PCG/Source/PCG/Public/Elements/PCGExecuteBlueprint.h
@@ -59,6 +59,10 @@ public:
 	UFUNCTION(BlueprintCallable, BlueprintPure = false, Category = Execution, meta = (HideSelfPin = "true"))
 	void LoopNTimes(UPARAM(ref) FPCGContext& InContext, int64 NumIterations, UPCGPointData*& OutData, const UPCGSpatialData* InA = nullptr, const UPCGSpatialData* InB = nullptr, UPCGPointData* OptionalOutData = nullptr) const;
 
+	/** Override for the default node name */
+	UFUNCTION(BlueprintNativeEvent, Category = Graph)
+	FName NodeTitleOverride() const;
+
 	/** Called after object creation to setup the object callbacks */
 	void Initialize();
 
@@ -81,6 +85,12 @@ public:
 	UPROPERTY(EditAnywhere, BlueprintReadWrite, Category = "Settings|Input & Output")
 	TSet<FName> OutputPinLabels;
 
+	UPROPERTY(BlueprintReadWrite, EditAnywhere, Category = "Settings|Input & Output")
+	bool bHasDefaultInPin = true;
+
+	UPROPERTY(BlueprintReadWrite, EditAnywhere, Category = "Settings|Input & Output")
+	bool bHasDefaultOutPin = true;
+
 protected:
 #if WITH_EDITOR
 	void OnDependencyChanged(UObject* Object, FPropertyChangedEvent& PropertyChangedEvent);
@@ -103,10 +113,10 @@ public:
 #endif
 
 	virtual FName AdditionalTaskName() const override;
-	virtual bool HasInLabel(const FName& Label) const override;
-	virtual bool HasOutLabel(const FName& Label) const override;
 	virtual TArray<FName> InLabels() const override;
 	virtual TArray<FName> OutLabels() const override;
+	virtual bool HasDefaultInLabel() const override;
+	virtual bool HasDefaultOutLabel() const override;
 
 protected:
 	virtual FPCGElementPtr CreateElement() const override;
diff --git a/Engine/Plugins/Experimental/PCG/Source/PCG/Public/Elements/PCGPointFilter.h b/Engine/Plugins/Experimental/PCG/Source/PCG/Public/Elements/PCGPointFilter.h
--- a/Engine/Plugins/Experimental/PCG/Source/PCG/Public/Elements/PCGPointFilter.h
+++ b/Engine/Plugins/Experimental/PCG/Source/PCG/Public/Elements/PCGPointFilter.h
@@ -71,7 +71,8 @@ public:
 
 	virtual TArray<FName> InLabels() const override;
 	virtual TArray<FName> OutLabels() const override;
-	virtual bool HasDefaultLabels() const override { return false; }
+	virtual bool HasDefaultInLabel() const override { return false; }
+	virtual bool HasDefaultOutLabel() const override { return false; }
 
 protected:
 	virtual FPCGElementPtr CreateElement() const override;
diff --git a/Engine/Plugins/Experimental/PCG/Source/PCG/Public/PCGDebug.h b/Engine/Plugins/Experimental/PCG/Source/PCG/Public/PCGDebug.h
--- a/Engine/Plugins/Experimental/PCG/Source/PCG/Public/PCGDebug.h
+++ b/Engine/Plugins/Experimental/PCG/Source/PCG/Public/PCGDebug.h
@@ -24,20 +24,20 @@ struct PCG_API FPCGDebugVisualizationSettings
 public:
 	FPCGDebugVisualizationSettings();
 
-	UPROPERTY(BlueprintReadWrite, EditAnywhere, Category = Settings, meta=(ClampMin="0"))
+	UPROPERTY(BlueprintReadWrite, EditAnywhere, Category = Debug, meta=(ClampMin="0"))
 	float PointScale = 1.0f;
 
-	UPROPERTY(BlueprintReadWrite, EditAnywhere, Category = Settings)
+	UPROPERTY(BlueprintReadWrite, EditAnywhere, Category = Debug)
 	EPCGDebugVisScaleMethod ScaleMethod = EPCGDebugVisScaleMethod::Relative;
 
-	UPROPERTY(BlueprintReadWrite, EditAnywhere, Category = Settings)
+	UPROPERTY(BlueprintReadWrite, EditAnywhere, Category = Debug)
 	TSoftObjectPtr<UStaticMesh> PointMesh;
 
-	UPROPERTY(BlueprintReadWrite, EditAnywhere, Category = Settings)
+	UPROPERTY(BlueprintReadWrite, EditAnywhere, Category = Debug)
 	TSoftObjectPtr<UMaterialInterface> MaterialOverride;
 
 	/** Warning: enabling this flag will have severe performance impact */
-	UPROPERTY(BlueprintReadWrite, EditAnywhere, Category = Settings)
+	UPROPERTY(BlueprintReadWrite, EditAnywhere, Category = Debug)
 	bool bCheckForDuplicates = false;
 
 	TSoftObjectPtr<UMaterialInterface> GetMaterial() const;
diff --git a/Engine/Plugins/Experimental/PCG/Source/PCG/Public/PCGGraph.h b/Engine/Plugins/Experimental/PCG/Source/PCG/Public/PCGGraph.h
--- a/Engine/Plugins/Experimental/PCG/Source/PCG/Public/PCGGraph.h
+++ b/Engine/Plugins/Experimental/PCG/Source/PCG/Public/PCGGraph.h
@@ -92,7 +92,7 @@ protected:
 #if WITH_EDITOR
 private:
 	void NotifyGraphChanged(bool bIsStructural);
-	void OnSettingsChanged(UPCGNode* InNode);
+	void OnSettingsChanged(UPCGNode* InNode, bool bSettingsChanged);
 	void OnStructuralSettingsChanged(UPCGNode* InNode);
 
 	bool bEnableGraphChangeNotifications = true;
diff --git a/Engine/Plugins/Experimental/PCG/Source/PCG/Public/PCGNode.h b/Engine/Plugins/Experimental/PCG/Source/PCG/Public/PCGNode.h
--- a/Engine/Plugins/Experimental/PCG/Source/PCG/Public/PCGNode.h
+++ b/Engine/Plugins/Experimental/PCG/Source/PCG/Public/PCGNode.h
@@ -12,7 +12,7 @@ class UPCGEdge;
 class IPCGElement;
 
 #if WITH_EDITOR
-	DECLARE_MULTICAST_DELEGATE_OneParam(FOnPCGNodeSettingsChanged, UPCGNode*);
+	DECLARE_MULTICAST_DELEGATE_TwoParams(FOnPCGNodeSettingsChanged, UPCGNode*, bool);
 #endif
 
 UCLASS(ClassGroup = (Procedural))
@@ -40,6 +40,9 @@ public:
 	UFUNCTION(BlueprintCallable, Category = Node)
 	UPCGNode* AddEdgeTo(FName InboundName, UPCGNode* To, FName OutboundName);
 
+	/** Returns the node title, based either on the current node label, or defaulted to its settings */
+	FName GetNodeTitle() const;
+
 	/** Returns true if one of the input pins matches with the given name */
 	bool HasInLabel(const FName& Label) const;
 
@@ -53,7 +56,8 @@ public:
 	TArray<FName> OutLabels() const;
 
 	/** Returns true if we want the default In/Out labels */
-	bool HasDefaultLabels() const;
+	bool HasDefaultInLabel() const;
+	bool HasDefaultOutLabel() const;
 
 	/** Returns true if the input pin is connected */
 	bool IsInputPinConnected(const FName& Label) const;
@@ -70,6 +74,9 @@ public:
 	UPROPERTY(BlueprintReadWrite, EditAnywhere, Category = Node, meta=(EditInline))
 	TObjectPtr<UPCGSettings> DefaultSettings;
 
+	UPROPERTY(BlueprintReadWrite, EditAnywhere, Category = Node)
+	FName NodeTitle = NAME_None;
+
 #if WITH_EDITOR
 	FOnPCGNodeSettingsChanged OnNodeSettingsChangedDelegate;
 #endif
diff --git a/Engine/Plugins/Experimental/PCG/Source/PCG/Public/PCGSettings.h b/Engine/Plugins/Experimental/PCG/Source/PCG/Public/PCGSettings.h
--- a/Engine/Plugins/Experimental/PCG/Source/PCG/Public/PCGSettings.h
+++ b/Engine/Plugins/Experimental/PCG/Source/PCG/Public/PCGSettings.h
@@ -43,7 +43,8 @@ public:
 	virtual bool HasOutLabel(const FName& Label) const;
 	virtual TArray<FName> InLabels() const { return TArray<FName>(); }
 	virtual TArray<FName> OutLabels() const { return TArray<FName>(); }
-	virtual bool HasDefaultLabels() const { return true; }
+	virtual bool HasDefaultInLabel() const { return true; }
+	virtual bool HasDefaultOutLabel() const { return true; }
 	
 	bool operator==(const UPCGSettings& Other) const;
 
diff --git a/Engine/Plugins/Experimental/PCG/Source/PCGEditor/Private/AssetTypeActions/PCGGraphAssetTypeActions.cpp b/Engine/Plugins/Experimental/PCG/Source/PCGEditor/Private/AssetTypeActions/PCGGraphAssetTypeActions.cpp
--- a/Engine/Plugins/Experimental/PCG/Source/PCGEditor/Private/AssetTypeActions/PCGGraphAssetTypeActions.cpp
+++ b/Engine/Plugins/Experimental/PCG/Source/PCGEditor/Private/AssetTypeActions/PCGGraphAssetTypeActions.cpp
@@ -6,7 +6,7 @@
 
 #include "HAL/IConsoleManager.h"
 
-static TAutoConsoleVariable<bool> CVarPCGUseGraphEditor(TEXT("pcg.UseGraphEditor"), false, TEXT("Wheter to use use the new graph editor or not."));
+static TAutoConsoleVariable<bool> CVarPCGUseGraphEditor(TEXT("pcg.UseGraphEditor"), true, TEXT("Wheter to use use the new graph editor or not."));
 
 FText FPCGGraphAssetTypeActions::GetName() const
 {
diff --git a/Engine/Plugins/Experimental/PCG/Source/PCGEditor/Private/PCGEditor.cpp b/Engine/Plugins/Experimental/PCG/Source/PCGEditor/Private/PCGEditor.cpp
--- a/Engine/Plugins/Experimental/PCG/Source/PCGEditor/Private/PCGEditor.cpp
+++ b/Engine/Plugins/Experimental/PCG/Source/PCGEditor/Private/PCGEditor.cpp
@@ -352,6 +352,7 @@ TSharedRef<SGraphEditor> FPCGEditor::CreateGraphEditorWidget()
 
 	SGraphEditor::FGraphEditorEvents InEvents;
 	InEvents.OnSelectionChanged = SGraphEditor::FOnSelectionChanged::CreateSP(this, &FPCGEditor::OnSelectedNodesChanged);
+	InEvents.OnTextCommitted = FOnNodeTextCommitted::CreateSP(this, &FPCGEditor::OnNodeTitleCommitted);
 
 	return SNew(SGraphEditor)
 		.AdditionalCommands(GraphEditorCommands)
@@ -379,6 +380,14 @@ void FPCGEditor::OnSelectedNodesChanged(const TSet<UObject*>& NewSelection)
 	GetTabManager()->TryInvokeTab(FPCGEditor_private::PropertyDetailsID);
 }
 
+void FPCGEditor::OnNodeTitleCommitted(const FText& NewText, ETextCommit::Type CommitInfo, UEdGraphNode* NodeBeingChanged)
+{
+	if (NodeBeingChanged)
+	{
+		NodeBeingChanged->OnRenameNode(NewText.ToString());
+	}
+}
+
 TSharedRef<SDockTab> FPCGEditor::SpawnTab_GraphEditor(const FSpawnTabArgs& Args)
 {
 	return SNew(SDockTab)
diff --git a/Engine/Plugins/Experimental/PCG/Source/PCGEditor/Private/PCGEditor.h b/Engine/Plugins/Experimental/PCG/Source/PCGEditor/Private/PCGEditor.h
--- a/Engine/Plugins/Experimental/PCG/Source/PCGEditor/Private/PCGEditor.h
+++ b/Engine/Plugins/Experimental/PCG/Source/PCGEditor/Private/PCGEditor.h
@@ -55,6 +55,9 @@ private:
 	/** Called when the selection changes in the GraphEditor */
 	void OnSelectedNodesChanged(const TSet<UObject*>& NewSelection);
 
+	/** Called when the title of a node is changed */
+	void OnNodeTitleCommitted(const FText& NewText, ETextCommit::Type CommitInfo, UEdGraphNode* NodeBeingChanged);
+
 	TSharedRef<SDockTab> SpawnTab_GraphEditor(const FSpawnTabArgs& Args);
 	TSharedRef<SDockTab> SpawnTab_PropertyDetails(const FSpawnTabArgs& Args);
 	TSharedRef<SDockTab> SpawnTab_Library(const FSpawnTabArgs& Args);
diff --git a/Engine/Plugins/Experimental/PCG/Source/PCGEditor/Private/PCGEditorGraphNode.cpp b/Engine/Plugins/Experimental/PCG/Source/PCGEditor/Private/PCGEditorGraphNode.cpp
--- a/Engine/Plugins/Experimental/PCG/Source/PCGEditor/Private/PCGEditorGraphNode.cpp
+++ b/Engine/Plugins/Experimental/PCG/Source/PCGEditor/Private/PCGEditorGraphNode.cpp
@@ -35,20 +35,15 @@ void UPCGEditorGraphNode::Construct(UPCGNode* InPCGNode, EPCGEditorGraphNodeType
 	NodePosY = InPCGNode->PositionY;
 	
 	NodeType = InNodeType;
+
+	bCanRenameNode = (InNodeType == EPCGEditorGraphNodeType::Settings);
 }
 
 FText UPCGEditorGraphNode::GetNodeTitle(ENodeTitleType::Type TitleType) const
 {
-	if (PCGNode && PCGNode->DefaultSettings)
+	if (PCGNode)
 	{
-		if (PCGNode->DefaultSettings->AdditionalTaskName() != NAME_None)
-		{
-			return FText::FromName(PCGNode->DefaultSettings->AdditionalTaskName());
-		}
-		else
-		{
-			return FText::FromName(PCGNode->DefaultSettings->GetDefaultNodeName());
-		}
+		return FText::FromName(PCGNode->GetNodeTitle());
 	}
 	else
 	{
@@ -101,7 +96,7 @@ void UPCGEditorGraphNode::AllocateDefaultPins()
 {
 	if (NodeType == EPCGEditorGraphNodeType::Input || NodeType == EPCGEditorGraphNodeType::Settings)
 	{
-		if (!PCGNode || PCGNode->HasDefaultLabels())
+		if (!PCGNode || PCGNode->HasDefaultOutLabel())
 		{
 			CreatePin(EEdGraphPinDirection::EGPD_Output, NAME_None, FName(TEXT("Out")));
 		}
@@ -117,7 +112,7 @@ void UPCGEditorGraphNode::AllocateDefaultPins()
 
 	if (NodeType == EPCGEditorGraphNodeType::Output || NodeType == EPCGEditorGraphNodeType::Settings)
 	{
-		if (!PCGNode || PCGNode->HasDefaultLabels())
+		if (!PCGNode || PCGNode->HasDefaultInLabel())
 		{
 			CreatePin(EEdGraphPinDirection::EGPD_Input, NAME_None, FName(TEXT("In")));
 		}
@@ -137,7 +132,7 @@ void UPCGEditorGraphNode::AutowireNewNode(UEdGraphPin* FromPin)
 	if (FromPin->Direction == EEdGraphPinDirection::EGPD_Output)
 	{
 		FName InPinName = TEXT("In");
-		if (PCGNode && !PCGNode->HasDefaultLabels() && !PCGNode->InLabels().IsEmpty())
+		if (PCGNode && !PCGNode->HasDefaultInLabel() && !PCGNode->InLabels().IsEmpty())
 		{
 			InPinName = PCGNode->InLabels()[0];
 		}
@@ -148,7 +143,7 @@ void UPCGEditorGraphNode::AutowireNewNode(UEdGraphPin* FromPin)
 	else if (FromPin->Direction == EEdGraphPinDirection::EGPD_Input)
 	{
 		FName OutPinName = TEXT("Out");
-		if (PCGNode && !PCGNode->HasDefaultLabels() && !PCGNode->OutLabels().IsEmpty())
+		if (PCGNode && !PCGNode->HasDefaultOutLabel() && !PCGNode->OutLabels().IsEmpty())
 		{
 			OutPinName = PCGNode->OutLabels()[0];
 		}
@@ -164,7 +159,7 @@ bool UPCGEditorGraphNode::CanUserDeleteNode() const
 	return NodeType != EPCGEditorGraphNodeType::Input && NodeType != EPCGEditorGraphNodeType::Output;
 }
 
-void UPCGEditorGraphNode::OnNodeChanged(UPCGNode* InNode)
+void UPCGEditorGraphNode::OnNodeChanged(UPCGNode* InNode, bool bSettingsChanged)
 {
 	if (InNode == PCGNode)
 	{
@@ -185,13 +180,28 @@ void UPCGEditorGraphNode::ReconstructNode()
 	TArray<UEdGraphPin*> InputPins;
 	TArray<UEdGraphPin*> OutputPins;
 
+	// Reallocate the default pins if they had been removed previously
+	if((NodeType == EPCGEditorGraphNodeType::Output || NodeType == EPCGEditorGraphNodeType::Settings) && PCGNode->HasDefaultInLabel() && !FindPin(TEXT("In"), EEdGraphPinDirection::EGPD_Input))
+	{
+		FCreatePinParams FirstPinParams;
+		FirstPinParams.Index = 0;
+		CreatePin(EEdGraphPinDirection::EGPD_Input, NAME_None, FName(TEXT("In")), FirstPinParams);
+	}
+
+	if ((NodeType == EPCGEditorGraphNodeType::Input || NodeType == EPCGEditorGraphNodeType::Settings) && PCGNode->HasDefaultOutLabel() && !FindPin(TEXT("Out"), EEdGraphPinDirection::EGPD_Output))
+	{
+		FCreatePinParams FirstPinParams;
+		FirstPinParams.Index = 0;
+		CreatePin(EEdGraphPinDirection::EGPD_Output, NAME_None, FName(TEXT("Out")), FirstPinParams);
+	}
+
 	for (UEdGraphPin* Pin : Pins)
 	{
-		if (Pin->Direction == EEdGraphPinDirection::EGPD_Input && (Pin->PinName != TEXT("In") || !PCGNode->HasDefaultLabels()))
+		if (Pin->Direction == EEdGraphPinDirection::EGPD_Input && (Pin->PinName != TEXT("In") || !PCGNode->HasDefaultInLabel()))
 		{
 			InputPins.Add(Pin);
 		}
-		else if (Pin->Direction == EEdGraphPinDirection::EGPD_Output && (Pin->PinName != TEXT("Out") || !PCGNode->HasDefaultLabels()))
+		else if (Pin->Direction == EEdGraphPinDirection::EGPD_Output && (Pin->PinName != TEXT("Out") || !PCGNode->HasDefaultOutLabel()))
 		{
 			OutputPins.Add(Pin);
 		}
@@ -259,4 +269,15 @@ void UPCGEditorGraphNode::ReconstructNode()
 	OnNodeChangedDelegate.ExecuteIfBound();
 }
 
+void UPCGEditorGraphNode::OnRenameNode(const FString& NewName)
+{
+	FName TentativeName(NewName);
+
+	if (GetCanRenameNode() && PCGNode && PCGNode->GetNodeTitle() != TentativeName)
+	{
+		PCGNode->Modify();
+		PCGNode->NodeTitle = FName(NewName);
+	}
+}
+
 #undef LOCTEXT_NAMESPACE
diff --git a/Engine/Plugins/Experimental/PCG/Source/PCGEditor/Private/PCGEditorGraphNode.h b/Engine/Plugins/Experimental/PCG/Source/PCGEditor/Private/PCGEditorGraphNode.h
--- a/Engine/Plugins/Experimental/PCG/Source/PCGEditor/Private/PCGEditorGraphNode.h
+++ b/Engine/Plugins/Experimental/PCG/Source/PCGEditor/Private/PCGEditorGraphNode.h
@@ -38,6 +38,7 @@ public:
 	virtual void AutowireNewNode(UEdGraphPin* FromPin) override;
 	virtual bool CanUserDeleteNode() const override;
 	virtual void ReconstructNode() override;
+	virtual void OnRenameNode(const FString& NewName);
 	// ~End UEdGraphNode interface
 
 	UPCGNode* GetPCGNode() { return PCGNode; }
@@ -46,7 +47,7 @@ public:
 	FOnPCGEditorGraphNodeChanged OnNodeChangedDelegate;
 
 protected:
-	void OnNodeChanged(UPCGNode* InNode);
+	void OnNodeChanged(UPCGNode* InNode, bool bSettingsChanged);
 
 private:
 	UPROPERTY()
diff --git a/Engine/Plugins/Experimental/PCG/Source/PCGEditor/Private/PCGEditorGraphSchema.cpp b/Engine/Plugins/Experimental/PCG/Source/PCGEditor/Private/PCGEditorGraphSchema.cpp
--- a/Engine/Plugins/Experimental/PCG/Source/PCGEditor/Private/PCGEditorGraphSchema.cpp
+++ b/Engine/Plugins/Experimental/PCG/Source/PCGEditor/Private/PCGEditorGraphSchema.cpp
@@ -88,8 +88,8 @@ bool UPCGEditorGraphSchema::TryCreateConnection(UEdGraphPin* InA, UEdGraphPin* I
 		UPCGGraph* PCGGraph = PCGNodeA->GetGraph();
 		check(PCGGraph);
 
-		const FName& NodeAPinName = ((A->PinName == TEXT("Out") && PCGNodeA->HasDefaultLabels()) ? NAME_None : A->PinName);
-		const FName& NodeBPinName = ((B->PinName == TEXT("In") && PCGNodeB->HasDefaultLabels()) ? NAME_None : B->PinName);
+		const FName& NodeAPinName = ((A->PinName == TEXT("Out") && PCGNodeA->HasDefaultOutLabel()) ? NAME_None : A->PinName);
+		const FName& NodeBPinName = ((B->PinName == TEXT("In") && PCGNodeB->HasDefaultInLabel()) ? NAME_None : B->PinName);
 
 		PCGGraph->AddLabeledEdge(PCGNodeA, NodeAPinName, PCGNodeB, NodeBPinName);
 	}
@@ -113,12 +113,12 @@ void UPCGEditorGraphSchema::BreakPinLinks(UEdGraphPin& TargetPin, bool bSendsNod
 
 	if (TargetPin.Direction == EEdGraphPinDirection::EGPD_Input)
 	{
-		const FName& PinName = (TargetPin.PinName == TEXT("In") && PCGNode->HasDefaultLabels()) ? NAME_None : TargetPin.PinName;
+		const FName& PinName = (TargetPin.PinName == TEXT("In") && PCGNode->HasDefaultInLabel()) ? NAME_None : TargetPin.PinName;
 		PCGGraph->RemoveInboundEdges(PCGNode, PinName);
 	}
 	else if (TargetPin.Direction == EEdGraphPinDirection::EGPD_Output)
 	{
-		const FName& PinName = (TargetPin.PinName == TEXT("Out") && PCGNode->HasDefaultLabels()) ? NAME_None : TargetPin.PinName;
+		const FName& PinName = (TargetPin.PinName == TEXT("Out") && PCGNode->HasDefaultOutLabel()) ? NAME_None : TargetPin.PinName;
 		PCGGraph->RemoveOutboundEdges(PCGNode, PinName);
 	}
 }
