Added node colors: overrideable per node, default values per type, blueprint callback, editor settings.

diff --git a/Engine/Plugins/Experimental/PCG/Source/PCG/Private/Elements/PCGExecuteBlueprint.cpp b/Engine/Plugins/Experimental/PCG/Source/PCG/Private/Elements/PCGExecuteBlueprint.cpp
--- a/Engine/Plugins/Experimental/PCG/Source/PCG/Private/Elements/PCGExecuteBlueprint.cpp
+++ b/Engine/Plugins/Experimental/PCG/Source/PCG/Private/Elements/PCGExecuteBlueprint.cpp
@@ -94,6 +94,16 @@ FName UPCGBlueprintElement::NodeTitleOverride_Implementation() const
 	return NAME_None;
 }
 
+FLinearColor UPCGBlueprintElement::NodeColorOverride_Implementation() const
+{
+	return FLinearColor::White;
+}
+
+EPCGSettingsType UPCGBlueprintElement::NodeTypeOverride_Implementation() const
+{
+	return EPCGSettingsType::Blueprint;
+}
+
 void UPCGBlueprintSettings::SetupBlueprintEvent()
 {
 #if WITH_EDITOR
@@ -256,6 +266,30 @@ void UPCGBlueprintSettings::RefreshBlueprintElement()
 }
 
 #if WITH_EDITOR
+FLinearColor UPCGBlueprintSettings::GetNodeTitleColor() const
+{
+	if (BlueprintElementInstance && BlueprintElementInstance->NodeColorOverride() != FLinearColor::White)
+	{
+		return BlueprintElementInstance->NodeColorOverride();
+	}
+	else
+	{
+		return Super::GetNodeTitleColor();
+	}
+}
+
+EPCGSettingsType UPCGBlueprintSettings::GetType() const
+{
+	if (BlueprintElementInstance)
+	{
+		return BlueprintElementInstance->NodeTypeOverride();
+	}
+	else
+	{
+		return EPCGSettingsType::Blueprint;
+	}
+}
+
 void UPCGBlueprintSettings::GetTrackedActorTags(FPCGTagToSettingsMap& OutTagToSettings) const
 {
 #if WITH_EDITORONLY_DATA
diff --git a/Engine/Plugins/Experimental/PCG/Source/PCG/Private/Elements/PCGSpawnActor.cpp b/Engine/Plugins/Experimental/PCG/Source/PCG/Private/Elements/PCGSpawnActor.cpp
--- a/Engine/Plugins/Experimental/PCG/Source/PCG/Private/Elements/PCGSpawnActor.cpp
+++ b/Engine/Plugins/Experimental/PCG/Source/PCG/Private/Elements/PCGSpawnActor.cpp
@@ -219,4 +219,11 @@ bool FPCGSpawnActorElement::ExecuteInternal(FPCGContext* Context) const
 	}
 
 	return true;
-}
\ No newline at end of file
+}
+
+#if WITH_EDITOR
+EPCGSettingsType UPCGSpawnActorSettings::GetType() const
+{
+	return GetSubgraph() ? EPCGSettingsType::Subgraph : EPCGSettingsType::Artifact;
+}
+#endif // WITH_EDITOR
\ No newline at end of file
diff --git a/Engine/Plugins/Experimental/PCG/Source/PCG/Public/Elements/PCGDensityFilter.h b/Engine/Plugins/Experimental/PCG/Source/PCG/Public/Elements/PCGDensityFilter.h
--- a/Engine/Plugins/Experimental/PCG/Source/PCG/Public/Elements/PCGDensityFilter.h
+++ b/Engine/Plugins/Experimental/PCG/Source/PCG/Public/Elements/PCGDensityFilter.h
@@ -18,6 +18,7 @@ public:
 	//~Begin UPCGSettings interface
 #if WITH_EDITOR
 	virtual FName GetDefaultNodeName() const override { return FName(TEXT("DensityFilter")); }
+	virtual EPCGSettingsType GetType() const override { return EPCGSettingsType::Filter; }
 #endif
 
 protected:
diff --git a/Engine/Plugins/Experimental/PCG/Source/PCG/Public/Elements/PCGDensityRemap.h b/Engine/Plugins/Experimental/PCG/Source/PCG/Public/Elements/PCGDensityRemap.h
--- a/Engine/Plugins/Experimental/PCG/Source/PCG/Public/Elements/PCGDensityRemap.h
+++ b/Engine/Plugins/Experimental/PCG/Source/PCG/Public/Elements/PCGDensityRemap.h
@@ -18,6 +18,7 @@ public:
 	// ~Begin UPCGSettings interface
 #if WITH_EDITOR
 	virtual FName GetDefaultNodeName() const override { return FName(TEXT("LinearDensityRemap")); }
+	virtual EPCGSettingsType GetType() const override { return EPCGSettingsType::Density; }
 #endif
 
 protected:
diff --git a/Engine/Plugins/Experimental/PCG/Source/PCG/Public/Elements/PCGDifferenceElement.h b/Engine/Plugins/Experimental/PCG/Source/PCG/Public/Elements/PCGDifferenceElement.h
--- a/Engine/Plugins/Experimental/PCG/Source/PCG/Public/Elements/PCGDifferenceElement.h
+++ b/Engine/Plugins/Experimental/PCG/Source/PCG/Public/Elements/PCGDifferenceElement.h
@@ -18,6 +18,7 @@ public:
 	//~Begin UPCGSettings interface
 #if WITH_EDITOR
 	virtual FName GetDefaultNodeName() const override { return FName(TEXT("DifferenceNode")); }
+	virtual EPCGSettingsType GetType() const override { return EPCGSettingsType::Spatial; }
 #endif
 
 	virtual TArray<FName> InLabels() const override;
diff --git a/Engine/Plugins/Experimental/PCG/Source/PCG/Public/Elements/PCGExecuteBlueprint.h b/Engine/Plugins/Experimental/PCG/Source/PCG/Public/Elements/PCGExecuteBlueprint.h
--- a/Engine/Plugins/Experimental/PCG/Source/PCG/Public/Elements/PCGExecuteBlueprint.h
+++ b/Engine/Plugins/Experimental/PCG/Source/PCG/Public/Elements/PCGExecuteBlueprint.h
@@ -63,6 +63,12 @@ public:
 	UFUNCTION(BlueprintNativeEvent, Category = Graph)
 	FName NodeTitleOverride() const;
 
+	UFUNCTION(BlueprintNativeEvent, Category = Graph)
+	FLinearColor NodeColorOverride() const;
+
+	UFUNCTION(BlueprintNativeEvent, Category = Graph)
+	EPCGSettingsType NodeTypeOverride() const;
+
 	/** Called after object creation to setup the object callbacks */
 	void Initialize();
 
@@ -109,6 +115,8 @@ public:
 	// ~Begin UPCGSettings interface
 #if WITH_EDITOR
 	virtual FName GetDefaultNodeName() const override { return FName(TEXT("BlueprintNode")); }
+	virtual FLinearColor GetNodeTitleColor() const override;
+	virtual EPCGSettingsType GetType() const override;
 	virtual void GetTrackedActorTags(FPCGTagToSettingsMap& OutTagToSettings) const override;
 #endif
 
diff --git a/Engine/Plugins/Experimental/PCG/Source/PCG/Public/Elements/PCGIntersectionElement.h b/Engine/Plugins/Experimental/PCG/Source/PCG/Public/Elements/PCGIntersectionElement.h
--- a/Engine/Plugins/Experimental/PCG/Source/PCG/Public/Elements/PCGIntersectionElement.h
+++ b/Engine/Plugins/Experimental/PCG/Source/PCG/Public/Elements/PCGIntersectionElement.h
@@ -18,6 +18,7 @@ public:
 	//~Begin UPCGSettings interface
 #if WITH_EDITOR
 	virtual FName GetDefaultNodeName() const override { return FName(TEXT("IntersectionNode")); }
+	virtual EPCGSettingsType GetType() const override { return EPCGSettingsType::Spatial; }
 #endif
 
 protected:
diff --git a/Engine/Plugins/Experimental/PCG/Source/PCG/Public/Elements/PCGMetadataElement.h b/Engine/Plugins/Experimental/PCG/Source/PCG/Public/Elements/PCGMetadataElement.h
--- a/Engine/Plugins/Experimental/PCG/Source/PCG/Public/Elements/PCGMetadataElement.h
+++ b/Engine/Plugins/Experimental/PCG/Source/PCG/Public/Elements/PCGMetadataElement.h
@@ -25,6 +25,7 @@ public:
 	//~Begin UPCGSettings interface
 #if WITH_EDITOR
 	virtual FName GetDefaultNodeName() const override { return FName(TEXT("MetadataOperationNode")); }
+	virtual EPCGSettingsType GetType() const override { return EPCGSettingsType::Metadata; }
 #endif
 
 protected:
diff --git a/Engine/Plugins/Experimental/PCG/Source/PCG/Public/Elements/PCGMetadataPartition.h b/Engine/Plugins/Experimental/PCG/Source/PCG/Public/Elements/PCGMetadataPartition.h
--- a/Engine/Plugins/Experimental/PCG/Source/PCG/Public/Elements/PCGMetadataPartition.h
+++ b/Engine/Plugins/Experimental/PCG/Source/PCG/Public/Elements/PCGMetadataPartition.h
@@ -16,6 +16,7 @@ public:
 	//~Begin UPCGSettings interface
 #if WITH_EDITOR
 	virtual FName GetDefaultNodeName() const override { return FName(TEXT("MetadataPartitionNode")); }
+	virtual EPCGSettingsType GetType() const override { return EPCGSettingsType::Metadata; }
 #endif
 
 protected:
diff --git a/Engine/Plugins/Experimental/PCG/Source/PCG/Public/Elements/PCGMetadataRenameElement.h b/Engine/Plugins/Experimental/PCG/Source/PCG/Public/Elements/PCGMetadataRenameElement.h
--- a/Engine/Plugins/Experimental/PCG/Source/PCG/Public/Elements/PCGMetadataRenameElement.h
+++ b/Engine/Plugins/Experimental/PCG/Source/PCG/Public/Elements/PCGMetadataRenameElement.h
@@ -16,6 +16,7 @@ public:
 	//~Begin UPCGSettings interface
 #if WITH_EDITOR
 	virtual FName GetDefaultNodeName() const override { return FName(TEXT("MetadataRenameNode")); }
+	virtual EPCGSettingsType GetType() const override { return EPCGSettingsType::Metadata; }
 #endif
 
 protected:
diff --git a/Engine/Plugins/Experimental/PCG/Source/PCG/Public/Elements/PCGPointFilter.h b/Engine/Plugins/Experimental/PCG/Source/PCG/Public/Elements/PCGPointFilter.h
--- a/Engine/Plugins/Experimental/PCG/Source/PCG/Public/Elements/PCGPointFilter.h
+++ b/Engine/Plugins/Experimental/PCG/Source/PCG/Public/Elements/PCGPointFilter.h
@@ -67,6 +67,7 @@ public:
 	//~Begin UPCGSettings interface
 #if WITH_EDITOR
 	virtual FName GetDefaultNodeName() const override { return FName(TEXT("PointFilterNode")); }
+	virtual EPCGSettingsType GetType() const override { return EPCGSettingsType::Filter; }
 #endif
 
 	virtual TArray<FName> InLabels() const override;
diff --git a/Engine/Plugins/Experimental/PCG/Source/PCG/Public/Elements/PCGPointSampler.h b/Engine/Plugins/Experimental/PCG/Source/PCG/Public/Elements/PCGPointSampler.h
--- a/Engine/Plugins/Experimental/PCG/Source/PCG/Public/Elements/PCGPointSampler.h
+++ b/Engine/Plugins/Experimental/PCG/Source/PCG/Public/Elements/PCGPointSampler.h
@@ -18,6 +18,7 @@ public:
 	//~Begin UPCGSettings interface
 #if WITH_EDITOR
 	virtual FName GetDefaultNodeName() const override { return FName(TEXT("PointSamplerNode")); }
+	virtual EPCGSettingsType GetType() const override { return EPCGSettingsType::Sampler; }
 #endif
 
 protected:
diff --git a/Engine/Plugins/Experimental/PCG/Source/PCG/Public/Elements/PCGProjectionElement.h b/Engine/Plugins/Experimental/PCG/Source/PCG/Public/Elements/PCGProjectionElement.h
--- a/Engine/Plugins/Experimental/PCG/Source/PCG/Public/Elements/PCGProjectionElement.h
+++ b/Engine/Plugins/Experimental/PCG/Source/PCG/Public/Elements/PCGProjectionElement.h
@@ -17,6 +17,7 @@ public:
 	//~Begin UPCGSettings interface
 #if WITH_EDITOR
 	virtual FName GetDefaultNodeName() const override { return FName(TEXT("ProjectionNode")); }
+	virtual EPCGSettingsType GetType() const override { return EPCGSettingsType::Spatial; }
 #endif
 
 #if WITH_EDITORONLY_DATA
diff --git a/Engine/Plugins/Experimental/PCG/Source/PCG/Public/Elements/PCGSelfPruning.h b/Engine/Plugins/Experimental/PCG/Source/PCG/Public/Elements/PCGSelfPruning.h
--- a/Engine/Plugins/Experimental/PCG/Source/PCG/Public/Elements/PCGSelfPruning.h
+++ b/Engine/Plugins/Experimental/PCG/Source/PCG/Public/Elements/PCGSelfPruning.h
@@ -42,6 +42,7 @@ public:
 #if WITH_EDITOR
 	//~Begin UPCGSettings interface
 	virtual FName GetDefaultNodeName() const override { return FName(TEXT("SelfPruningNode")); }
+	virtual EPCGSettingsType GetType() const override { return EPCGSettingsType::Filter; }
 #endif
 
 protected:
diff --git a/Engine/Plugins/Experimental/PCG/Source/PCG/Public/Elements/PCGSpawnActor.h b/Engine/Plugins/Experimental/PCG/Source/PCG/Public/Elements/PCGSpawnActor.h
--- a/Engine/Plugins/Experimental/PCG/Source/PCG/Public/Elements/PCGSpawnActor.h
+++ b/Engine/Plugins/Experimental/PCG/Source/PCG/Public/Elements/PCGSpawnActor.h
@@ -37,6 +37,7 @@ public:
 
 #if WITH_EDITOR	
 	virtual FName GetDefaultNodeName() const override { return FName(TEXT("SpawnActorNode")); }
+	virtual EPCGSettingsType GetType() const override;
 #endif
 protected:
 	virtual FPCGElementPtr CreateElement() const override;
diff --git a/Engine/Plugins/Experimental/PCG/Source/PCG/Public/Elements/PCGStaticMeshSpawner.h b/Engine/Plugins/Experimental/PCG/Source/PCG/Public/Elements/PCGStaticMeshSpawner.h
--- a/Engine/Plugins/Experimental/PCG/Source/PCG/Public/Elements/PCGStaticMeshSpawner.h
+++ b/Engine/Plugins/Experimental/PCG/Source/PCG/Public/Elements/PCGStaticMeshSpawner.h
@@ -30,6 +30,7 @@ public:
 #if WITH_EDITOR
 	// ~Begin UPCGSettings interface
 	virtual FName GetDefaultNodeName() const override { return FName(TEXT("StaticMeshSpawnerNode")); }
+	virtual EPCGSettingsType GetType() const override { return EPCGSettingsType::Artifact; }
 #endif
 
 protected:
diff --git a/Engine/Plugins/Experimental/PCG/Source/PCG/Public/Elements/PCGSurfaceSampler.h b/Engine/Plugins/Experimental/PCG/Source/PCG/Public/Elements/PCGSurfaceSampler.h
--- a/Engine/Plugins/Experimental/PCG/Source/PCG/Public/Elements/PCGSurfaceSampler.h
+++ b/Engine/Plugins/Experimental/PCG/Source/PCG/Public/Elements/PCGSurfaceSampler.h
@@ -47,6 +47,7 @@ public:
 #if WITH_EDITOR
 	//~Begin UPCGSettings interface
 	virtual FName GetDefaultNodeName() const override { return FName(TEXT("SurfaceSamplerNode")); }
+	virtual EPCGSettingsType GetType() const override { return EPCGSettingsType::Sampler; }
 #endif
 
 protected:
diff --git a/Engine/Plugins/Experimental/PCG/Source/PCG/Public/Elements/PCGUnionElement.h b/Engine/Plugins/Experimental/PCG/Source/PCG/Public/Elements/PCGUnionElement.h
--- a/Engine/Plugins/Experimental/PCG/Source/PCG/Public/Elements/PCGUnionElement.h
+++ b/Engine/Plugins/Experimental/PCG/Source/PCG/Public/Elements/PCGUnionElement.h
@@ -18,6 +18,7 @@ public:
 	//~Begin UPCGSettings interface
 #if WITH_EDITOR
 	virtual FName GetDefaultNodeName() const override { return FName(TEXT("UnionNode")); }
+	virtual EPCGSettingsType GetType() const override { return EPCGSettingsType::Spatial; }
 #endif
 
 protected:
diff --git a/Engine/Plugins/Experimental/PCG/Source/PCG/Public/PCGInputOutputSettings.h b/Engine/Plugins/Experimental/PCG/Source/PCG/Public/PCGInputOutputSettings.h
--- a/Engine/Plugins/Experimental/PCG/Source/PCG/Public/PCGInputOutputSettings.h
+++ b/Engine/Plugins/Experimental/PCG/Source/PCG/Public/PCGInputOutputSettings.h
@@ -16,6 +16,7 @@ public:
 	// ~Begin UPCGSettings interface
 #if WITH_EDITOR
 	virtual FName GetDefaultNodeName() const override { return FName(TEXT("InputNode")); }
+	virtual EPCGSettingsType GetType() const override { return EPCGSettingsType::InputOutput; }
 #endif
 
 	virtual bool HasInLabel(const FName& Label) const override { return Label == NAME_None || PinLabels.Contains(Label); }
diff --git a/Engine/Plugins/Experimental/PCG/Source/PCG/Public/PCGNode.h b/Engine/Plugins/Experimental/PCG/Source/PCG/Public/PCGNode.h
--- a/Engine/Plugins/Experimental/PCG/Source/PCG/Public/PCGNode.h
+++ b/Engine/Plugins/Experimental/PCG/Source/PCG/Public/PCGNode.h
@@ -77,6 +77,11 @@ public:
 	UPROPERTY(BlueprintReadWrite, EditAnywhere, Category = Node)
 	FName NodeTitle = NAME_None;
 
+#if WITH_EDITORONLY_DATA
+	UPROPERTY(BlueprintReadWrite, EditAnywhere, Category = Node)
+	FLinearColor NodeTitleColor = FLinearColor::White;
+#endif // WITH_EDITORONLY_DATA
+
 #if WITH_EDITOR
 	FOnPCGNodeSettingsChanged OnNodeSettingsChangedDelegate;
 #endif
diff --git a/Engine/Plugins/Experimental/PCG/Source/PCG/Public/PCGSettings.h b/Engine/Plugins/Experimental/PCG/Source/PCG/Public/PCGSettings.h
--- a/Engine/Plugins/Experimental/PCG/Source/PCG/Public/PCGSettings.h
+++ b/Engine/Plugins/Experimental/PCG/Source/PCG/Public/PCGSettings.h
@@ -22,6 +22,21 @@ enum class EPCGSettingsExecutionMode : uint8
 	Disabled
 };
 
+UENUM()
+enum class EPCGSettingsType : uint8
+{
+	InputOutput,
+	Spatial,
+	Density,
+	Blueprint,
+	Metadata,
+	Filter,
+	Sampler,
+	Artifact,
+	Subgraph,
+	Generic
+};
+
 #if WITH_EDITOR
 	DECLARE_MULTICAST_DELEGATE_OneParam(FOnPCGSettingsChanged, UPCGSettings*);
 #endif
@@ -50,6 +65,8 @@ public:
 
 #if WITH_EDITOR
 	virtual FName GetDefaultNodeName() const { return NAME_None; }
+	virtual FLinearColor GetNodeTitleColor() const { return FLinearColor::White; }
+	virtual EPCGSettingsType GetType() const { return EPCGSettingsType::Generic; }
 	/** Derived classes must implement this to communicate dependencies on external actors */
 	virtual void GetTrackedActorTags(FPCGTagToSettingsMap& OutTagToSettings) const {}
 #endif
diff --git a/Engine/Plugins/Experimental/PCG/Source/PCG/Public/PCGSubgraph.h b/Engine/Plugins/Experimental/PCG/Source/PCG/Public/PCGSubgraph.h
--- a/Engine/Plugins/Experimental/PCG/Source/PCG/Public/PCGSubgraph.h
+++ b/Engine/Plugins/Experimental/PCG/Source/PCG/Public/PCGSubgraph.h
@@ -63,6 +63,7 @@ public:
 
 #if WITH_EDITOR
 	virtual FName GetDefaultNodeName() const override { return FName(TEXT("SubgraphNode")); }
+	virtual EPCGSettingsType GetType() const override { return EPCGSettingsType::Subgraph; }
 #endif
 
 	virtual FName AdditionalTaskName() const override;
diff --git a/Engine/Plugins/Experimental/PCG/Source/PCGEditor/PCGEditor.Build.cs b/Engine/Plugins/Experimental/PCG/Source/PCGEditor/PCGEditor.Build.cs
--- a/Engine/Plugins/Experimental/PCG/Source/PCGEditor/PCGEditor.Build.cs
+++ b/Engine/Plugins/Experimental/PCG/Source/PCGEditor/PCGEditor.Build.cs
@@ -24,6 +24,7 @@ namespace UnrealBuildTool.Rules
 
 			PrivateDependencyModuleNames.AddRange(
 				new string[]{
+					"AppFramework",
 					"AssetTools",
 					"DesktopWidgets",
 					"EditorStyle",
diff --git a/Engine/Plugins/Experimental/PCG/Source/PCGEditor/Private/PCGEditorGraphNode.cpp b/Engine/Plugins/Experimental/PCG/Source/PCGEditor/Private/PCGEditorGraphNode.cpp
--- a/Engine/Plugins/Experimental/PCG/Source/PCGEditor/Private/PCGEditorGraphNode.cpp
+++ b/Engine/Plugins/Experimental/PCG/Source/PCGEditor/Private/PCGEditorGraphNode.cpp
@@ -3,6 +3,7 @@
 #include "PCGEditorGraphNode.h"
 
 #include "PCGEditorGraphSchema.h"
+#include "PCGEditorSettings.h"
 #include "PCGNode.h"
 #include "PCGSettings.h"
 
@@ -13,6 +14,8 @@
 #include "ToolMenu.h"
 #include "ToolMenuSection.h"
 
+#include "Widgets/Colors/SColorPicker.h"
+
 #define LOCTEXT_NAMESPACE "PCGEditorGraphNode"
 
 void UPCGEditorGraphNode::BeginDestroy()
@@ -70,6 +73,13 @@ void UPCGEditorGraphNode::GetNodeContextMenuActions(UToolMenu* Menu, class UGrap
 
 	{
 		FToolMenuSection& Section = Menu->AddSection("EdGraphSchemaOrganization", LOCTEXT("OrganizationHeader", "Organization"));
+		Section.AddMenuEntry(
+			"PCGNode_SetColor",
+			LOCTEXT("PCGNode_SetColor", "Set Node Color"),
+			LOCTEXT("PCGNode_SetColorTooltip", "Sets a specific color on the given node. Note that white maps to the default value"),
+			FSlateIcon(FEditorStyle::GetStyleSetName(), "ColorPicker.Mode"),
+			FUIAction(FExecuteAction::CreateUObject(const_cast<UPCGEditorGraphNode*>(this), &UPCGEditorGraphNode::OnPickColor)));
+
 		Section.AddSubMenu("Alignment", LOCTEXT("AlignmentHeader", "Alignment"), FText(), FNewToolMenuDelegate::CreateLambda([](UToolMenu* AlignmentMenu)
 		{
 			{
@@ -167,6 +177,26 @@ void UPCGEditorGraphNode::OnNodeChanged(UPCGNode* InNode, bool bSettingsChanged)
 	}
 }
 
+void UPCGEditorGraphNode::OnPickColor()
+{
+	FColorPickerArgs PickerArgs;
+	PickerArgs.bIsModal = true;
+	PickerArgs.bUseAlpha = false;
+	PickerArgs.InitialColorOverride = GetNodeTitleColor();
+	PickerArgs.OnColorCommitted = FOnLinearColorValueChanged::CreateUObject(this, &UPCGEditorGraphNode::OnColorPicked);
+
+	OpenColorPicker(PickerArgs);
+}
+
+void UPCGEditorGraphNode::OnColorPicked(FLinearColor NewColor)
+{
+	if (PCGNode && GetNodeTitleColor() != NewColor)
+	{
+		PCGNode->Modify();
+		PCGNode->NodeTitleColor = NewColor;
+	}
+}
+
 void UPCGEditorGraphNode::ReconstructNode()
 {
 	if (!PCGNode)
@@ -280,4 +310,29 @@ void UPCGEditorGraphNode::OnRenameNode(const FString& NewName)
 	}
 }
 
+FLinearColor UPCGEditorGraphNode::GetNodeTitleColor() const
+{
+	if (PCGNode)
+	{
+		if (PCGNode->NodeTitleColor != FLinearColor::White)
+		{
+			return PCGNode->NodeTitleColor;
+		}
+		else if (PCGNode->DefaultSettings)
+		{
+			FLinearColor SettingsColor = PCGNode->DefaultSettings->GetNodeTitleColor();
+			if (SettingsColor == FLinearColor::White)
+			{
+				SettingsColor = GetDefault<UPCGEditorSettings>()->GetColor(PCGNode->DefaultSettings);
+			}
+
+			if (SettingsColor != FLinearColor::White)
+			{
+				return SettingsColor;
+			}
+		}
+	}
+
+	return GetDefault<UPCGEditorSettings>()->DefaultNodeColor;
+}
 #undef LOCTEXT_NAMESPACE
diff --git a/Engine/Plugins/Experimental/PCG/Source/PCGEditor/Private/PCGEditorGraphNode.h b/Engine/Plugins/Experimental/PCG/Source/PCGEditor/Private/PCGEditorGraphNode.h
--- a/Engine/Plugins/Experimental/PCG/Source/PCGEditor/Private/PCGEditorGraphNode.h
+++ b/Engine/Plugins/Experimental/PCG/Source/PCGEditor/Private/PCGEditorGraphNode.h
@@ -39,6 +39,7 @@ public:
 	virtual bool CanUserDeleteNode() const override;
 	virtual void ReconstructNode() override;
 	virtual void OnRenameNode(const FString& NewName);
+	virtual FLinearColor GetNodeTitleColor() const;
 	// ~End UEdGraphNode interface
 
 	UPCGNode* GetPCGNode() { return PCGNode; }
@@ -48,6 +49,8 @@ public:
 
 protected:
 	void OnNodeChanged(UPCGNode* InNode, bool bSettingsChanged);
+	void OnPickColor();
+	void OnColorPicked(FLinearColor NewColor);
 
 private:
 	UPROPERTY()
diff --git a/Engine/Plugins/Experimental/PCG/Source/PCGEditor/Private/PCGEditorModule.cpp b/Engine/Plugins/Experimental/PCG/Source/PCGEditor/Private/PCGEditorModule.cpp
--- a/Engine/Plugins/Experimental/PCG/Source/PCGEditor/Private/PCGEditorModule.cpp
+++ b/Engine/Plugins/Experimental/PCG/Source/PCGEditor/Private/PCGEditorModule.cpp
@@ -9,6 +9,7 @@
 #include "Modules/ModuleManager.h"
 #include "PropertyEditorModule.h"
 #include "Toolkits/IToolkit.h"
+#include "ISettingsModule.h"
 
 #include "AssetTypeActions/PCGGraphAssetTypeActions.h"
 #include "AssetTypeActions/PCGSettingsAssetTypeActions.h"
@@ -16,6 +17,7 @@
 #include "PCGEditorGraphNodeFactory.h"
 #include "PCGGraphDetails.h"
 #include "PCGVolumeFactory.h"
+#include "PCGEditorSettings.h"
 
 #include "PCGSubsystem.h"
 
@@ -28,6 +30,7 @@ void FPCGEditorModule::StartupModule()
 	RegisterDetailsCustomizations();
 	RegisterAssetTypeActions();
 	RegisterMenuExtensions();
+	RegisterSettings();
 
 	GraphNodeFactory = MakeShareable(new FPCGEditorGraphNodeFactory());
 	FEdGraphUtilities::RegisterVisualNodeFactory(GraphNodeFactory);
@@ -40,6 +43,7 @@ void FPCGEditorModule::StartupModule()
 
 void FPCGEditorModule::ShutdownModule()
 {
+	UnregisterSettings();
 	UnregisterAssetTypeActions();
 	UnregisterDetailsCustomizations();
 	UnregisterMenuExtensions();
@@ -158,6 +162,25 @@ void FPCGEditorModule::PopulateMenuActions(FMenuBuilder& MenuBuilder)
 		NAME_None);
 }
 
+void FPCGEditorModule::RegisterSettings()
+{
+	if (ISettingsModule* SettingsModule = FModuleManager::GetModulePtr<ISettingsModule>("Settings"))
+	{
+		SettingsModule->RegisterSettings("Editor", "ContentEditors", "PCGEditor",
+			LOCTEXT("PCGEditorSettingsName", "PCG Editor"),
+			LOCTEXT("PCGEditorSettingsDescription", "Configure the look and feel of the PCG Editor."),
+			GetMutableDefault<UPCGEditorSettings>());
+	}
+}
+
+void FPCGEditorModule::UnregisterSettings()
+{
+	if (ISettingsModule* SettingsModule = FModuleManager::GetModulePtr<ISettingsModule>("Settings"))
+	{
+		SettingsModule->UnregisterSettings("Editor", "ContentEditors", "PCGEditor");
+	}
+}
+
 IMPLEMENT_MODULE(FPCGEditorModule, PCGEditor);
 
 #undef LOCTEXT_NAMESPACE
\ No newline at end of file
diff --git a/Engine/Plugins/Experimental/PCG/Source/PCGEditor/Private/PCGEditorModule.h b/Engine/Plugins/Experimental/PCG/Source/PCGEditor/Private/PCGEditorModule.h
--- a/Engine/Plugins/Experimental/PCG/Source/PCGEditor/Private/PCGEditorModule.h
+++ b/Engine/Plugins/Experimental/PCG/Source/PCGEditor/Private/PCGEditorModule.h
@@ -38,6 +38,8 @@ protected:
 	void UnregisterMenuExtensions();
 	void AddMenuEntry(FMenuBuilder& MenuBuilder);
 	void PopulateMenuActions(FMenuBuilder& MenuBuilder);
+	void RegisterSettings();
+	void UnregisterSettings();
 
 	TArray<TSharedRef<IAssetTypeActions>> RegisteredAssetTypeActions;
 	static EAssetTypeCategories::Type PCGAssetCategory;
diff --git a/Engine/Plugins/Experimental/PCG/Source/PCGEditor/Private/PCGEditorSettings.cpp b/Engine/Plugins/Experimental/PCG/Source/PCGEditor/Private/PCGEditorSettings.cpp
new file mode 100644
--- /dev/null
+++ b/Engine/Plugins/Experimental/PCG/Source/PCGEditor/Private/PCGEditorSettings.cpp
@@ -0,0 +1,73 @@
+// Copyright Epic Games, Inc. All Rights Reserved.
+
+#include "PCGEditorSettings.h"
+
+UPCGEditorSettings::UPCGEditorSettings(const FObjectInitializer& ObjectInitializer)
+	: Super(ObjectInitializer)
+{
+	DefaultNodeColor = FLinearColor(0.4f, 0.62f, 1.0f);
+	InputOutputNodeColor = FLinearColor(1.0f, 0.0f, 0.0f);
+	SetOperationNodeColor = FLinearColor(0.8f, 0.2f, 0.8f);
+	DensityOperationNodeColor = FLinearColor(0.6f, 1.0f, 0.6f);
+	BlueprintNodeColor = FLinearColor(0.0f, 0.6f, 1.0f);
+	MetadataNodeColor = FLinearColor(0.4f, 0.4f, 0.8f);
+	FilterNodeColor = FLinearColor(0.4f, 0.8f, 0.4f);
+	SamplerNodeColor = FLinearColor(0.8f, 1.0f, 0.4f);
+	ArtifactNodeColor = FLinearColor(1.0f, 0.6f, 0.4f);
+	SubgraphNodeColor = FLinearColor(1.0f, 0.1f, 0.1f);
+}
+
+FLinearColor UPCGEditorSettings::GetColor(UPCGSettings* Settings) const
+{
+	if (!Settings)
+	{
+		return DefaultNodeColor;
+	}
+	// First: check if there's an override
+	else if (const FLinearColor* Override = OverrideNodeColorByClass.Find(Settings->GetClass()))
+	{
+		return *Override;
+	}
+	// Otherwise, check against the classes we know
+	else if(Settings->GetType() == EPCGSettingsType::InputOutput)
+	{
+		return InputOutputNodeColor;
+	}
+	else if (Settings->GetType() == EPCGSettingsType::Spatial)
+	{
+		return SetOperationNodeColor;
+	}
+	else if (Settings->GetType() == EPCGSettingsType::Density)
+	{
+		return DensityOperationNodeColor;
+	}
+	else if (Settings->GetType() == EPCGSettingsType::Blueprint)
+	{
+		return BlueprintNodeColor;
+	}
+	else if (Settings->GetType() == EPCGSettingsType::Metadata)
+	{
+		return MetadataNodeColor;
+	}
+	else if (Settings->GetType() == EPCGSettingsType::Filter)
+	{
+		return FilterNodeColor;
+	}
+	else if (Settings->GetType() == EPCGSettingsType::Sampler)
+	{
+		return SamplerNodeColor;
+	}
+	else if (Settings->GetType() == EPCGSettingsType::Artifact)
+	{
+		return ArtifactNodeColor;
+	}
+	else if (Settings->GetType() == EPCGSettingsType::Subgraph)
+	{
+		return SubgraphNodeColor;
+	}
+	else
+	{
+		// Finally, we couldn't find any match, so return the default value
+		return DefaultNodeColor;
+	}
+}
diff --git a/Engine/Plugins/Experimental/PCG/Source/PCGEditor/Public/PCGEditorSettings.h b/Engine/Plugins/Experimental/PCG/Source/PCGEditor/Public/PCGEditorSettings.h
new file mode 100644
--- /dev/null
+++ b/Engine/Plugins/Experimental/PCG/Source/PCGEditor/Public/PCGEditorSettings.h
@@ -0,0 +1,63 @@
+// Copyright Epic Games, Inc. All Rights Reserved.
+
+#pragma once
+
+#include "CoreMinimal.h"
+#include "PCGSettings.h"
+
+#include "PCGEditorSettings.generated.h"
+
+UCLASS(config=EditorPerProjectUserSettings)
+class PCGEDITOR_API UPCGEditorSettings : public UObject
+{
+	GENERATED_BODY()
+
+public:
+	UPCGEditorSettings(const FObjectInitializer& ObjectInitializer);
+
+	FLinearColor GetColor(UPCGSettings* InSettings) const;
+
+	/** Default node color */
+	UPROPERTY(EditAnywhere, config, Category = Node, meta = (HideAlphaChannel))
+	FLinearColor DefaultNodeColor;
+
+	/** Color used for input & output nodes */
+	UPROPERTY(EditAnywhere, config, Category = Node, meta = (HideAlphaChannel))
+	FLinearColor InputOutputNodeColor;
+
+	/** Color used for Difference, Intersection, Projection, Union */
+	UPROPERTY(EditAnywhere, config, Category = Node, meta = (HideAlphaChannel))
+	FLinearColor SetOperationNodeColor;
+
+	/** Color used for density remap */
+	UPROPERTY(EditAnywhere, config, Category = Node, meta = (HideAlphaChannel))
+	FLinearColor DensityOperationNodeColor;
+
+	/** Color used for blueprints */
+	UPROPERTY(EditAnywhere, config, Category = Node, meta = (HideAlphaChannel))
+	FLinearColor BlueprintNodeColor;
+
+	/** Color used for metadata operations */
+	UPROPERTY(EditAnywhere, config, Category = Node, meta = (HideAlphaChannel))
+	FLinearColor MetadataNodeColor;
+
+	/** Color used for filter-like operations */
+	UPROPERTY(EditAnywhere, config, Category = Node, meta = (HideAlphaChannel))
+	FLinearColor FilterNodeColor;
+
+	/** Color used for sampler operations */
+	UPROPERTY(EditAnywhere, config, Category = Node, meta = (HideAlphaChannel))
+	FLinearColor SamplerNodeColor;
+
+	/** Color used for artifact-generating operations */
+	UPROPERTY(EditAnywhere, config, Category = Node, meta = (HideAlphaChannel))
+	FLinearColor ArtifactNodeColor;
+
+	/** Color used for subgraph-like operations */
+	UPROPERTY(EditAnywhere, config, Category = Node, meta = (HideAlphaChannel))
+	FLinearColor SubgraphNodeColor;
+
+	/** User-driven color overrides */
+	UPROPERTY(EditAnywhere, config, Category = Node, meta = (HideAlphaChannel))
+	TMap<TSubclassOf<UPCGSettings>, FLinearColor> OverrideNodeColorByClass;
+};
\ No newline at end of file
