Added call to get list of attributes in blueprint

diff --git a/Engine/Plugins/Experimental/PCG/Source/PCG/Private/Metadata/PCGMetadata.cpp b/Engine/Plugins/Experimental/PCG/Source/PCG/Private/Metadata/PCGMetadata.cpp
--- a/Engine/Plugins/Experimental/PCG/Source/PCG/Private/Metadata/PCGMetadata.cpp
+++ b/Engine/Plugins/Experimental/PCG/Source/PCG/Private/Metadata/PCGMetadata.cpp
@@ -185,6 +185,29 @@ bool UPCGMetadata::HasAttribute(FName AttributeName) const
 	return Attributes.Contains(AttributeName);
 }
 
+void UPCGMetadata::GetAttributes(TArray<FName>& AttributeNames, TArray<EPCGMetadataTypes>& AttributeTypes) const
+{
+	AttributeNames.Reset();
+	AttributeTypes.Reset();
+
+	AttributeLock.ReadLock();
+	for (const TPair<FName, FPCGMetadataAttributeBase*>& Attribute : Attributes)
+	{
+		check(Attribute.Value && Attribute.Value->Name == Attribute.Key);
+		AttributeNames.Add(Attribute.Key);
+
+		if (Attribute.Value->GetTypeId() < static_cast<uint16>(EPCGMetadataTypes::Unknown))
+		{
+			AttributeTypes.Add(static_cast<EPCGMetadataTypes>(Attribute.Value->GetTypeId()));
+		}
+		else
+		{
+			AttributeTypes.Add(EPCGMetadataTypes::Unknown);
+		}
+	}
+	AttributeLock.ReadUnlock();
+}
+
 bool UPCGMetadata::ParentHasAttribute(FName AttributeName) const
 {
 	return Parent && Parent->HasAttribute(AttributeName);
diff --git a/Engine/Plugins/Experimental/PCG/Source/PCG/Public/Metadata/PCGMetadata.h b/Engine/Plugins/Experimental/PCG/Source/PCG/Public/Metadata/PCGMetadata.h
--- a/Engine/Plugins/Experimental/PCG/Source/PCG/Public/Metadata/PCGMetadata.h
+++ b/Engine/Plugins/Experimental/PCG/Source/PCG/Public/Metadata/PCGMetadata.h
@@ -76,6 +76,9 @@ public:
 	const FPCGMetadataAttributeBase* GetConstAttribute(FName AttributeName) const;
 	bool HasAttribute(FName AttributeName) const;
 
+	UFUNCTION(BlueprintCallable, Category = "PCG|Metadata")
+	void GetAttributes(TArray<FName>& AttributeNames, TArray<EPCGMetadataTypes>& AttributeTypes) const;
+
 	/** Delete/Hide attribute */
 	// Due to stream inheriting, we might want to consider "hiding" parent stream and deleting local streams only
 	void DeleteAttribute(FName AttributeName);
diff --git a/Engine/Plugins/Experimental/PCG/Source/PCG/Public/Metadata/PCGMetadataAttribute.h b/Engine/Plugins/Experimental/PCG/Source/PCG/Public/Metadata/PCGMetadataAttribute.h
--- a/Engine/Plugins/Experimental/PCG/Source/PCG/Public/Metadata/PCGMetadataAttribute.h
+++ b/Engine/Plugins/Experimental/PCG/Source/PCG/Public/Metadata/PCGMetadataAttribute.h
@@ -18,7 +18,7 @@ public:
 	virtual ~FPCGMetadataAttributeBase() = default;
 
 	const UPCGMetadata* GetMetadata() const { return Metadata; }
-	int32 GetTypeId() const { return TypeId; }
+	int16 GetTypeId() const { return TypeId; }
 
 	virtual FPCGMetadataAttributeBase* Copy(FName NewName, UPCGMetadata* InMetadata, bool bKeepParent, bool bCopyEntries = true, bool bCopyValues = true) const = 0;
 
@@ -42,7 +42,7 @@ protected:
 
 	TObjectPtr<UPCGMetadata> Metadata = nullptr;
 	const FPCGMetadataAttributeBase* Parent = nullptr;
-	int32 TypeId = 0;
+	int16 TypeId = 0;
 	bool bAllowsInterpolation = false;
 
 public:
diff --git a/Engine/Plugins/Experimental/PCG/Source/PCG/Public/Metadata/PCGMetadataAttributeTraits.h b/Engine/Plugins/Experimental/PCG/Source/PCG/Public/Metadata/PCGMetadataAttributeTraits.h
--- a/Engine/Plugins/Experimental/PCG/Source/PCG/Public/Metadata/PCGMetadataAttributeTraits.h
+++ b/Engine/Plugins/Experimental/PCG/Source/PCG/Public/Metadata/PCGMetadataAttributeTraits.h
@@ -4,6 +4,24 @@
 
 #include "CoreMinimal.h"
 
+#include "PCGMetadataAttributeTraits.generated.h"
+
+UENUM(BlueprintType)
+enum class EPCGMetadataTypes : uint8
+{
+	Float = 0,
+	Double,
+	Integer32,
+	Integer64,
+	Vector,
+	Vector4,
+	Quaternion,
+	Transform,
+	String,
+	Boolean,
+	Unknown = 255
+};
+
 namespace PCG
 {
 	namespace Private
@@ -11,21 +29,21 @@ namespace PCG
 		template<typename T>
 		struct MetadataTypes
 		{
-			enum { Id = sizeof(T) };
+			enum { Id = static_cast<uint16>(EPCGMetadataTypes::Unknown) + sizeof(T) };
 		};
 
-#define PCGMetadataGenerateDataTypes(Type, TypeId) template<> struct MetadataTypes<Type>{ enum {Id = 2048 + TypeId}; }
-
-		PCGMetadataGenerateDataTypes(float, 0);
-		PCGMetadataGenerateDataTypes(double, 1);
-		PCGMetadataGenerateDataTypes(int32, 2);
-		PCGMetadataGenerateDataTypes(int64, 3);
-		PCGMetadataGenerateDataTypes(FVector, 4);
-		PCGMetadataGenerateDataTypes(FVector4, 5);
-		PCGMetadataGenerateDataTypes(FQuat, 6);
-		PCGMetadataGenerateDataTypes(FTransform, 7);
-		PCGMetadataGenerateDataTypes(FString, 8);
-		PCGMetadataGenerateDataTypes(bool, 9);
+#define PCGMetadataGenerateDataTypes(Type, TypeEnum) template<> struct MetadataTypes<Type>{ enum { Id = static_cast<uint16>(EPCGMetadataTypes::TypeEnum)}; }
+
+		PCGMetadataGenerateDataTypes(float, Float);
+		PCGMetadataGenerateDataTypes(double, Double);
+		PCGMetadataGenerateDataTypes(int32, Integer32);
+		PCGMetadataGenerateDataTypes(int64, Integer64);
+		PCGMetadataGenerateDataTypes(FVector, Vector);
+		PCGMetadataGenerateDataTypes(FVector4, Vector4);
+		PCGMetadataGenerateDataTypes(FQuat, Quaternion);
+		PCGMetadataGenerateDataTypes(FTransform, Transform);
+		PCGMetadataGenerateDataTypes(FString, String);
+		PCGMetadataGenerateDataTypes(bool, Boolean);
 
 #undef PCGMetadataGenerateDataTypes
 
