PCG:
*Added copy / paste / duplicate support to the pcg graph editor.

diff --git a/Engine/Plugins/Experimental/PCG/Source/PCG/Private/PCGGraph.cpp b/Engine/Plugins/Experimental/PCG/Source/PCG/Private/PCGGraph.cpp
--- a/Engine/Plugins/Experimental/PCG/Source/PCG/Private/PCGGraph.cpp
+++ b/Engine/Plugins/Experimental/PCG/Source/PCG/Private/PCGGraph.cpp
@@ -180,7 +180,7 @@ UPCGNode* UPCGGraph::AddNode(UPCGSettings* InSettings)
 		const FName DefaultNodeName = InSettings->GetDefaultNodeName();
 		if (DefaultNodeName != NAME_None)
 		{
-			FName NodeName = MakeUniqueObjectName(GetOuter(), UPCGNode::StaticClass(), DefaultNodeName);
+			FName NodeName = MakeUniqueObjectName(this, UPCGNode::StaticClass(), DefaultNodeName);
 			Node->Rename(*NodeName.ToString());
 		}
 #endif
@@ -270,6 +270,23 @@ bool UPCGGraph::Contains(UPCGNode* Node) const
 	return Node == InputNode || Node == OutputNode || Nodes.Contains(Node);
 }
 
+void UPCGGraph::AddNode(UPCGNode* InNode)
+{
+	InNode->Rename(nullptr, this);
+
+#if WITH_EDITOR
+	const FName DefaultNodeName = InNode->DefaultSettings->GetDefaultNodeName();
+	if (DefaultNodeName != NAME_None)
+	{
+		FName NodeName = MakeUniqueObjectName(this, UPCGNode::StaticClass(), DefaultNodeName);
+		InNode->Rename(*NodeName.ToString());
+	}
+#endif
+
+	Nodes.Add(InNode);
+	OnNodeAdded(InNode);
+}
+
 void UPCGGraph::RemoveNode(UPCGNode* InNode)
 {
 	check(InNode);
diff --git a/Engine/Plugins/Experimental/PCG/Source/PCG/Private/PCGNode.cpp b/Engine/Plugins/Experimental/PCG/Source/PCG/Private/PCGNode.cpp
--- a/Engine/Plugins/Experimental/PCG/Source/PCG/Private/PCGNode.cpp
+++ b/Engine/Plugins/Experimental/PCG/Source/PCG/Private/PCGNode.cpp
@@ -33,6 +33,17 @@ void UPCGNode::PostLoad()
 	OutboundNodes_DEPRECATED.Reset();
 }
 
+void UPCGNode::PostEditImport()
+{
+	Super::PostEditImport();
+#if WITH_EDITOR
+	if (DefaultSettings)
+	{
+		DefaultSettings->OnSettingsChangedDelegate.AddUObject(this, &UPCGNode::OnSettingsChanged);
+	}
+#endif
+}
+
 void UPCGNode::BeginDestroy()
 {
 #if WITH_EDITOR
@@ -204,7 +215,8 @@ void UPCGNode::RemoveOutboundEdge(UPCGEdge* InEdge)
 
 void UPCGNode::PreEditChange(FProperty* PropertyAboutToChange)
 {
-	if (PropertyAboutToChange && PropertyAboutToChange->GetFName() == GET_MEMBER_NAME_CHECKED(UPCGNode, DefaultSettings))
+	// To properly clean up old callbacks during paste we clear during null call since the Settings property doesn't get a specific call.
+	if (!PropertyAboutToChange || (PropertyAboutToChange && PropertyAboutToChange->GetFName() == GET_MEMBER_NAME_CHECKED(UPCGNode, DefaultSettings)))
 	{
 		if (DefaultSettings)
 		{
diff --git a/Engine/Plugins/Experimental/PCG/Source/PCG/Public/PCGGraph.h b/Engine/Plugins/Experimental/PCG/Source/PCG/Public/PCGGraph.h
--- a/Engine/Plugins/Experimental/PCG/Source/PCG/Public/PCGGraph.h
+++ b/Engine/Plugins/Experimental/PCG/Source/PCG/Public/PCGGraph.h
@@ -60,6 +60,7 @@ public:
 
 	bool Contains(UPCGNode* Node) const;
 	const TArray<UPCGNode*>& GetNodes() const { return Nodes; }
+	void AddNode(UPCGNode* InNode);
 	void RemoveNode(UPCGNode* InNode);
 	void RemoveEdge(UPCGNode* From, const FName& FromLabel, UPCGNode* To, const FName& ToLabel);
 	void RemoveAllInboundEdges(UPCGNode* InNode);
diff --git a/Engine/Plugins/Experimental/PCG/Source/PCG/Public/PCGNode.h b/Engine/Plugins/Experimental/PCG/Source/PCG/Public/PCGNode.h
--- a/Engine/Plugins/Experimental/PCG/Source/PCG/Public/PCGNode.h
+++ b/Engine/Plugins/Experimental/PCG/Source/PCG/Public/PCGNode.h
@@ -29,6 +29,7 @@ public:
 	
 	/** ~Begin UObject interface */
 	virtual void PostLoad() override;
+	virtual void PostEditImport() override;
 	virtual void BeginDestroy() override;
 	/** ~End UObject interface */
 
@@ -110,10 +111,10 @@ protected:
 	UPROPERTY()
 	TArray<TObjectPtr<UPCGNode>> OutboundNodes_DEPRECATED;
 
-	UPROPERTY(BlueprintReadOnly, VisibleAnywhere, Category = Node)
+	UPROPERTY(BlueprintReadOnly, VisibleAnywhere, Category = Node, TextExportTransient)
 	TArray<TObjectPtr<UPCGEdge>> InboundEdges;
 
-	UPROPERTY(BlueprintReadOnly, VisibleAnywhere, Category = Node)
+	UPROPERTY(BlueprintReadOnly, VisibleAnywhere, Category = Node, TextExportTransient)
 	TArray<TObjectPtr<UPCGEdge>> OutboundEdges;
 
 	// TODO: add this information:
diff --git a/Engine/Plugins/Experimental/PCG/Source/PCGEditor/PCGEditor.Build.cs b/Engine/Plugins/Experimental/PCG/Source/PCGEditor/PCGEditor.Build.cs
--- a/Engine/Plugins/Experimental/PCG/Source/PCGEditor/PCGEditor.Build.cs
+++ b/Engine/Plugins/Experimental/PCG/Source/PCGEditor/PCGEditor.Build.cs
@@ -25,6 +25,7 @@ namespace UnrealBuildTool.Rules
 			PrivateDependencyModuleNames.AddRange(
 				new string[]{
 					"AppFramework",
+					"ApplicationCore",
 					"AssetTools",
 					"DesktopWidgets",
 					"EditorStyle",
diff --git a/Engine/Plugins/Experimental/PCG/Source/PCGEditor/Private/PCGEditor.cpp b/Engine/Plugins/Experimental/PCG/Source/PCGEditor/Private/PCGEditor.cpp
--- a/Engine/Plugins/Experimental/PCG/Source/PCGEditor/Private/PCGEditor.cpp
+++ b/Engine/Plugins/Experimental/PCG/Source/PCGEditor/Private/PCGEditor.cpp
@@ -3,10 +3,11 @@
 #include "PCGEditor.h"
 
 #include "PCGEditorGraph.h"
-#include "PCGEditorGraphNode.h"
+#include "PCGEditorGraphNodeBase.h"
 #include "PCGEditorGraphSchema.h"
 #include "PCGGraph.h"
 
+#include "EdGraphUtilities.h"
 #include "Editor.h"
 #include "Framework/Commands/GenericCommands.h"
 #include "Framework/Commands/UIAction.h"
@@ -15,12 +16,13 @@
 #include "Framework/MultiBox/MultiBoxBuilder.h"
 #include "GraphEditor.h"
 #include "GraphEditorActions.h"
+#include "HAL/PlatformApplicationMisc.h"
 #include "IDetailsView.h"
 #include "Modules/ModuleManager.h"
 #include "PropertyEditorModule.h"
-#include "Widgets/Docking/SDockTab.h"
-
+#include "SNodePanel.h"
 #include "ToolMenus.h"
+#include "Widgets/Docking/SDockTab.h"
 
 #define LOCTEXT_NAMESPACE "PCGGraphEditor"
 
@@ -187,9 +189,11 @@ void FPCGEditor::DeleteSelectedNodes()
 		UPCGGraph* PCGGraph = PCGEditorGraph->GetPCGGraph();
 		check(PCGEditorGraph && PCGGraph);
 
+		bool bChanged = false;
+
 		for (UObject* Object : GraphEditorWidget->GetSelectedNodes())
 		{
-			UPCGEditorGraphNode* PCGEditorGraphNode = CastChecked<UPCGEditorGraphNode>(Object);
+			UPCGEditorGraphNodeBase* PCGEditorGraphNode = CastChecked<UPCGEditorGraphNodeBase>(Object);
 
 			if (PCGEditorGraphNode->CanUserDeleteNode())
 			{
@@ -198,11 +202,16 @@ void FPCGEditor::DeleteSelectedNodes()
 
 				PCGGraph->RemoveNode(PCGNode);
 				PCGEditorGraphNode->DestroyNode();
+				bChanged = true;
 			}
 		}
 
-		GraphEditorWidget->ClearSelectionSet();
-		GraphEditorWidget->NotifyGraphChanged();
+		if (bChanged)
+		{
+			GraphEditorWidget->ClearSelectionSet();
+			GraphEditorWidget->NotifyGraphChanged();
+			PCGGraphBeingEdited->MarkPackageDirty();
+		}
 	}
 }
 
@@ -212,7 +221,7 @@ bool FPCGEditor::CanDeleteSelectedNodes() const
 	{
 		for (UObject* Object : GraphEditorWidget->GetSelectedNodes())
 		{
-			UPCGEditorGraphNode* PCGEditorGraphNode = CastChecked<UPCGEditorGraphNode>(Object);
+			UPCGEditorGraphNodeBase* PCGEditorGraphNode = CastChecked<UPCGEditorGraphNodeBase>(Object);
 
 			if (PCGEditorGraphNode->CanUserDeleteNode())
 			{
@@ -224,6 +233,157 @@ bool FPCGEditor::CanDeleteSelectedNodes() const
 	return false;
 }
 
+void FPCGEditor::CopySelectedNodes()
+{
+	if (GraphEditorWidget.IsValid())
+	{
+		const FGraphPanelSelectionSet SelectedNodes = GraphEditorWidget->GetSelectedNodes();
+
+		//TODO: evaluate creating a clipboard object instead of ownership hack
+		for (UObject* SelectedNode : SelectedNodes)
+		{
+			UEdGraphNode* GraphNode = CastChecked<UEdGraphNode>(SelectedNode);
+			GraphNode->PrepareForCopying();
+		}
+
+		FString ExportedText;
+		FEdGraphUtilities::ExportNodesToText(SelectedNodes, ExportedText);
+		FPlatformApplicationMisc::ClipboardCopy(*ExportedText);
+
+		for (UObject* SelectedNode : SelectedNodes)
+		{
+			UPCGEditorGraphNodeBase* PCGGraphNode = CastChecked<UPCGEditorGraphNodeBase>(SelectedNode);
+			PCGGraphNode->PostCopy();
+		}
+	}
+}
+
+bool FPCGEditor::CanCopySelectedNodes() const
+{
+	if (GraphEditorWidget.IsValid())
+	{
+		for (UObject* Object : GraphEditorWidget->GetSelectedNodes())
+		{
+			if (UPCGEditorGraphNodeBase* PCGEditorGraphNode = Cast<UPCGEditorGraphNodeBase>(Object))
+			{
+				if (PCGEditorGraphNode->CanDuplicateNode())
+				{
+					return true;
+				}
+			}
+		}
+	}
+
+	return false;
+}
+
+void FPCGEditor::CutSelectedNodes()
+{
+	CopySelectedNodes();
+	DeleteSelectedNodes();
+}
+
+bool FPCGEditor::CanCutSelectedNodes() const
+{
+	return CanCopySelectedNodes() && CanDeleteSelectedNodes();
+}
+
+void FPCGEditor::PasteNodes()
+{
+	if (GraphEditorWidget.IsValid())
+	{
+		PasteNodesHere(GraphEditorWidget->GetPasteLocation());
+	}
+}
+
+void FPCGEditor::PasteNodesHere(const FVector2D& Location)
+{
+	if (!GraphEditorWidget.IsValid() || !PCGEditorGraph)
+	{
+		return;
+	}
+
+	PCGEditorGraph->Modify();
+
+	// Clear the selection set (newly pasted stuff will be selected)
+	GraphEditorWidget->ClearSelectionSet();
+
+	// Grab the text to paste from the clipboard.
+	FString TextToImport;
+	FPlatformApplicationMisc::ClipboardPaste(TextToImport);
+
+	// Import the nodes
+	TSet<UEdGraphNode*> PastedNodes;
+	FEdGraphUtilities::ImportNodesFromText(PCGEditorGraph, TextToImport, /*out*/ PastedNodes);
+
+	//Average position of nodes so we can move them while still maintaining relative distances to each other
+	FVector2D AvgNodePosition(0.0f, 0.0f);
+
+	// Number of nodes used to calculate AvgNodePosition
+	int32 AvgCount = 0;
+
+	for (UEdGraphNode* PastedNode : PastedNodes)
+	{
+		if (PastedNode)
+		{
+			AvgNodePosition.X += PastedNode->NodePosX;
+			AvgNodePosition.Y += PastedNode->NodePosY;
+			++AvgCount;
+		}
+	}
+
+	if (AvgCount > 0)
+	{
+		float InvNumNodes = 1.0f / float(AvgCount);
+		AvgNodePosition.X *= InvNumNodes;
+		AvgNodePosition.Y *= InvNumNodes;
+	}
+
+	for (UEdGraphNode* PastedNode : PastedNodes)
+	{
+		GraphEditorWidget->SetNodeSelection(PastedNode, true);
+
+		PastedNode->NodePosX = (PastedNode->NodePosX - AvgNodePosition.X) + Location.X;
+		PastedNode->NodePosY = (PastedNode->NodePosY - AvgNodePosition.Y) + Location.Y;
+
+		PastedNode->SnapToGrid(SNodePanel::GetSnapGridSize());
+
+		PastedNode->CreateNewGuid();
+
+		if (UPCGEditorGraphNodeBase* PastedPCGGraphNode = Cast<UPCGEditorGraphNodeBase>(PastedNode))
+		{
+			if (UPCGNode* PastedPCGNode = PastedPCGGraphNode->GetPCGNode())
+			{
+				PCGGraphBeingEdited->AddNode(PastedPCGNode);
+
+				PastedPCGGraphNode->PostPaste();
+			}
+		}
+	}
+
+	GraphEditorWidget->NotifyGraphChanged();
+	PCGGraphBeingEdited->MarkPackageDirty();
+}
+
+bool FPCGEditor::CanPasteNodes() const
+{
+	FString ClipboardContent;
+	FPlatformApplicationMisc::ClipboardPaste(ClipboardContent);
+
+	return FEdGraphUtilities::CanImportNodesFromText(PCGEditorGraph, ClipboardContent);
+}
+
+void FPCGEditor::DuplicateNodes()
+{
+	CopySelectedNodes();
+	PasteNodes();
+}
+
+bool FPCGEditor::CanDuplicateNodes() const
+{
+	return CanCopySelectedNodes();
+}
+
 void FPCGEditor::OnAlignTop()
 {
 	if (GraphEditorWidget.IsValid())
@@ -309,6 +469,22 @@ TSharedRef<SGraphEditor> FPCGEditor::CreateGraphEditorWidget()
 		FExecuteAction::CreateSP(this, &FPCGEditor::DeleteSelectedNodes),
 		FCanExecuteAction::CreateSP(this, &FPCGEditor::CanDeleteSelectedNodes));
 
+	GraphEditorCommands->MapAction(FGenericCommands::Get().Copy,
+		FExecuteAction::CreateSP(this, &FPCGEditor::CopySelectedNodes),
+		FCanExecuteAction::CreateSP(this, &FPCGEditor::CanCopySelectedNodes));
+
+	GraphEditorCommands->MapAction(FGenericCommands::Get().Cut,
+		FExecuteAction::CreateSP(this, &FPCGEditor::CutSelectedNodes),
+		FCanExecuteAction::CreateSP(this, &FPCGEditor::CanCutSelectedNodes));
+
+	GraphEditorCommands->MapAction(FGenericCommands::Get().Paste,
+		FExecuteAction::CreateSP(this, &FPCGEditor::PasteNodes),
+		FCanExecuteAction::CreateSP(this, &FPCGEditor::CanPasteNodes));
+
+	GraphEditorCommands->MapAction(FGenericCommands::Get().Duplicate,
+		FExecuteAction::CreateSP(this, &FPCGEditor::DuplicateNodes),
+		FCanExecuteAction::CreateSP(this, &FPCGEditor::CanDuplicateNodes));
+
 	// Alignment Commands
 	GraphEditorCommands->MapAction(FGraphEditorCommands::Get().AlignNodesTop,
 		FExecuteAction::CreateSP(this, &FPCGEditor::OnAlignTop)
@@ -369,9 +545,12 @@ void FPCGEditor::OnSelectedNodesChanged(const TSet<UObject*>& NewSelection)
 
 	for (UObject* Object : NewSelection)
 	{
-		if (UPCGEditorGraphNode* GraphNode = Cast<UPCGEditorGraphNode>(Object))
+		if (UPCGEditorGraphNodeBase* GraphNode = Cast<UPCGEditorGraphNodeBase>(Object))
 		{
-			SelectedObjects.Add(GraphNode->GetPCGNode()->DefaultSettings);
+			if (UPCGNode* PCGNode = GraphNode->GetPCGNode())
+			{
+				SelectedObjects.Add(PCGNode->DefaultSettings);
+			}
 		}
 	}
 
diff --git a/Engine/Plugins/Experimental/PCG/Source/PCGEditor/Private/PCGEditor.h b/Engine/Plugins/Experimental/PCG/Source/PCGEditor/Private/PCGEditor.h
--- a/Engine/Plugins/Experimental/PCG/Source/PCGEditor/Private/PCGEditor.h
+++ b/Engine/Plugins/Experimental/PCG/Source/PCGEditor/Private/PCGEditor.h
@@ -39,6 +39,28 @@ private:
 	/** Whether we can delete all selected nodes */
 	bool CanDeleteSelectedNodes() const;
 
+	/** Copy all selected nodes in the graph */
+	void CopySelectedNodes();
+	/** Whether we can copy all selected nodes */
+	bool CanCopySelectedNodes() const;
+
+	/** Cut all selected nodes in the graph */
+	void CutSelectedNodes();
+	/** Whether we can cut all selected nodes */
+	bool CanCutSelectedNodes() const;
+
+	/** Paste nodes in the graph */
+	void PasteNodes();
+	/** Paste nodes in the graph at location*/
+	void PasteNodesHere(const FVector2D& Location);
+	/** Whether we can paste nodes */
+	bool CanPasteNodes() const;
+
+	/** Duplicate the currently selected nodes */
+	void DuplicateNodes();
+	/** Whether we are able to duplicate the currently selected nodes */
+	bool CanDuplicateNodes() const;
+
 	void OnAlignTop();
 	void OnAlignMiddle();
 	void OnAlignBottom();
diff --git a/Engine/Plugins/Experimental/PCG/Source/PCGEditor/Private/PCGEditorGraph.cpp b/Engine/Plugins/Experimental/PCG/Source/PCGEditor/Private/PCGEditorGraph.cpp
--- a/Engine/Plugins/Experimental/PCG/Source/PCGEditor/Private/PCGEditorGraph.cpp
+++ b/Engine/Plugins/Experimental/PCG/Source/PCGEditor/Private/PCGEditorGraph.cpp
@@ -5,6 +5,8 @@
 #include "PCGGraph.h"
 #include "PCGEdge.h"
 #include "PCGEditorGraphNode.h"
+#include "PCGEditorGraphNodeInput.h"
+#include "PCGEditorGraphNodeOutput.h"
 
 #include "EdGraph/EdGraphPin.h"
 
@@ -13,19 +15,19 @@ void UPCGEditorGraph::InitFromNodeGraph(UPCGGraph* InPCGGraph)
 	check(InPCGGraph && !PCGGraph);
 	PCGGraph = InPCGGraph;
 
-	TMap<UPCGNode*, UPCGEditorGraphNode*> NodeLookup;
+	TMap<UPCGNode*, UPCGEditorGraphNodeBase*> NodeLookup;
 	const bool bSelectNewNode = false;
 
 	UPCGNode* InputNode = PCGGraph->GetInputNode();
-	FGraphNodeCreator<UPCGEditorGraphNode> InputNodeCreator(*this);
-	UPCGEditorGraphNode* InputGraphNode = InputNodeCreator.CreateNode(bSelectNewNode);
+	FGraphNodeCreator<UPCGEditorGraphNodeInput> InputNodeCreator(*this);
+	UPCGEditorGraphNodeInput* InputGraphNode = InputNodeCreator.CreateNode(bSelectNewNode);
 	InputGraphNode->Construct(InputNode, EPCGEditorGraphNodeType::Input);
 	InputNodeCreator.Finalize();
 	NodeLookup.Add(InputNode, InputGraphNode);
 
 	UPCGNode* OutputNode = PCGGraph->GetOutputNode();
-	FGraphNodeCreator<UPCGEditorGraphNode> OutputNodeCreator(*this);
-	UPCGEditorGraphNode* OutputGraphNode = OutputNodeCreator.CreateNode(bSelectNewNode);
+	FGraphNodeCreator<UPCGEditorGraphNodeOutput> OutputNodeCreator(*this);
+	UPCGEditorGraphNodeOutput* OutputGraphNode = OutputNodeCreator.CreateNode(bSelectNewNode);
 	OutputGraphNode->Construct(OutputNode, EPCGEditorGraphNodeType::Output);
 	OutputNodeCreator.Finalize();
 	NodeLookup.Add(OutputNode, OutputGraphNode);
@@ -42,7 +44,7 @@ void UPCGEditorGraph::InitFromNodeGraph(UPCGGraph* InPCGGraph)
 	for (const auto& NodeLookupIt : NodeLookup)
 	{
 		UPCGNode* PCGNode = NodeLookupIt.Key;
-		UPCGEditorGraphNode* GraphNode = NodeLookupIt.Value;
+		UPCGEditorGraphNodeBase* GraphNode = NodeLookupIt.Value;
 
 		for (UPCGEdge* OutboundEdge : PCGNode->GetOutboundEdges())
 		{
@@ -55,7 +57,7 @@ void UPCGEditorGraph::InitFromNodeGraph(UPCGGraph* InPCGGraph)
 			}
 			
 			UPCGNode* OutboundNode = OutboundEdge->OutboundNode;
-			if (UPCGEditorGraphNode** ConnectedGraphNode = NodeLookup.Find(OutboundNode))
+			if (UPCGEditorGraphNodeBase** ConnectedGraphNode = NodeLookup.Find(OutboundNode))
 			{
 				const FName InPinName = (OutboundEdge->OutboundLabel == NAME_None ? TEXT("In") : OutboundEdge->OutboundLabel);
 
@@ -66,4 +68,4 @@ void UPCGEditorGraph::InitFromNodeGraph(UPCGGraph* InPCGGraph)
 			}
 		}
 	}
-}
\ No newline at end of file
+}
diff --git a/Engine/Plugins/Experimental/PCG/Source/PCGEditor/Private/PCGEditorGraphNode.cpp b/Engine/Plugins/Experimental/PCG/Source/PCGEditor/Private/PCGEditorGraphNode.cpp
--- a/Engine/Plugins/Experimental/PCG/Source/PCGEditor/Private/PCGEditorGraphNode.cpp
+++ b/Engine/Plugins/Experimental/PCG/Source/PCGEditor/Private/PCGEditorGraphNode.cpp
@@ -3,45 +3,16 @@
 #include "PCGEditorGraphNode.h"
 
 #include "PCGEditorGraphSchema.h"
-#include "PCGEditorSettings.h"
 #include "PCGNode.h"
 #include "PCGSettings.h"
 
-#include "EdGraph/EdGraphPin.h"
-#include "EdGraph/EdGraph.h"
 #include "Framework/Commands/GenericCommands.h"
 #include "GraphEditorActions.h"
 #include "ToolMenu.h"
 #include "ToolMenuSection.h"
 
-#include "Widgets/Colors/SColorPicker.h"
-
 #define LOCTEXT_NAMESPACE "PCGEditorGraphNode"
 
-void UPCGEditorGraphNode::BeginDestroy()
-{
-	if (PCGNode)
-	{
-		PCGNode->OnNodeSettingsChangedDelegate.RemoveAll(this);
-	}
-
-	Super::BeginDestroy();
-}
-
-void UPCGEditorGraphNode::Construct(UPCGNode* InPCGNode, EPCGEditorGraphNodeType InNodeType)
-{
-	check(InPCGNode);
-	PCGNode = InPCGNode;
-	InPCGNode->OnNodeSettingsChangedDelegate.AddUObject(this, &UPCGEditorGraphNode::OnNodeChanged);
-
-	NodePosX = InPCGNode->PositionX;
-	NodePosY = InPCGNode->PositionY;
-	
-	NodeType = InNodeType;
-
-	bCanRenameNode = (InNodeType == EPCGEditorGraphNodeType::Settings);
-}
-
 FText UPCGEditorGraphNode::GetNodeTitle(ENodeTitleType::Type TitleType) const
 {
 	if (PCGNode)
@@ -61,242 +32,48 @@ void UPCGEditorGraphNode::GetNodeContextMenuActions(UToolMenu* Menu, class UGrap
 		return;
 	}
 
-	{
-		FToolMenuSection& Section = Menu->AddSection("EdGraphSchemaNodeActions", LOCTEXT("NodeActionsHeader", "Node Actions"));
-		Section.AddMenuEntry(FGraphEditorCommands::Get().BreakNodeLinks);
-	}
-
-	{
-		FToolMenuSection& Section = Menu->AddSection("EdGraphSchemaGeneral", LOCTEXT("GeneralHeader", "General"));
-		Section.AddMenuEntry(FGenericCommands::Get().Delete);
-	}
-
-	{
-		FToolMenuSection& Section = Menu->AddSection("EdGraphSchemaOrganization", LOCTEXT("OrganizationHeader", "Organization"));
-		Section.AddMenuEntry(
-			"PCGNode_SetColor",
-			LOCTEXT("PCGNode_SetColor", "Set Node Color"),
-			LOCTEXT("PCGNode_SetColorTooltip", "Sets a specific color on the given node. Note that white maps to the default value"),
-			FSlateIcon(FEditorStyle::GetStyleSetName(), "ColorPicker.Mode"),
-			FUIAction(FExecuteAction::CreateUObject(const_cast<UPCGEditorGraphNode*>(this), &UPCGEditorGraphNode::OnPickColor)));
-
-		Section.AddSubMenu("Alignment", LOCTEXT("AlignmentHeader", "Alignment"), FText(), FNewToolMenuDelegate::CreateLambda([](UToolMenu* AlignmentMenu)
-		{
-			{
-				FToolMenuSection& SubSection = AlignmentMenu->AddSection("EdGraphSchemaAlignment", LOCTEXT("AlignHeader", "Align"));
-				SubSection.AddMenuEntry(FGraphEditorCommands::Get().AlignNodesTop);
-				SubSection.AddMenuEntry(FGraphEditorCommands::Get().AlignNodesMiddle);
-				SubSection.AddMenuEntry(FGraphEditorCommands::Get().AlignNodesBottom);
-				SubSection.AddMenuEntry(FGraphEditorCommands::Get().AlignNodesLeft);
-				SubSection.AddMenuEntry(FGraphEditorCommands::Get().AlignNodesCenter);
-				SubSection.AddMenuEntry(FGraphEditorCommands::Get().AlignNodesRight);
-				SubSection.AddMenuEntry(FGraphEditorCommands::Get().StraightenConnections);
-			}
+	FToolMenuSection& Section = Menu->AddSection("EdGraphSchemaGeneral", LOCTEXT("GeneralHeader", "General"));
+	Section.AddMenuEntry(FGenericCommands::Get().Delete);
+	Section.AddMenuEntry(FGenericCommands::Get().Cut);
+	Section.AddMenuEntry(FGenericCommands::Get().Copy);
+	Section.AddMenuEntry(FGenericCommands::Get().Duplicate);
 
-			{
-				FToolMenuSection& SubSection = AlignmentMenu->AddSection("EdGraphSchemaDistribution", LOCTEXT("DistributionHeader", "Distribution"));
-				SubSection.AddMenuEntry(FGraphEditorCommands::Get().DistributeNodesHorizontally);
-				SubSection.AddMenuEntry(FGraphEditorCommands::Get().DistributeNodesVertically);
-			}
-		}));
-	}
+	Super::GetNodeContextMenuActions(Menu, Context);
 }
 
 void UPCGEditorGraphNode::AllocateDefaultPins()
 {
-	if (NodeType == EPCGEditorGraphNodeType::Input || NodeType == EPCGEditorGraphNodeType::Settings)
+	if (!PCGNode || PCGNode->HasDefaultOutLabel())
 	{
-		if (!PCGNode || PCGNode->HasDefaultOutLabel())
-		{
-			CreatePin(EEdGraphPinDirection::EGPD_Output, NAME_None, FName(TEXT("Out")));
-		}
-
-		if (PCGNode)
-		{
-			for (const FName& OutLabel : PCGNode->OutLabels())
-			{
-				CreatePin(EEdGraphPinDirection::EGPD_Output, NAME_None, OutLabel);
-			}
-		}
+		CreatePin(EEdGraphPinDirection::EGPD_Output, NAME_None, FName(TEXT("Out")));
 	}
 
-	if (NodeType == EPCGEditorGraphNodeType::Output || NodeType == EPCGEditorGraphNodeType::Settings)
+	if (PCGNode)
 	{
-		if (!PCGNode || PCGNode->HasDefaultInLabel())
-		{
-			CreatePin(EEdGraphPinDirection::EGPD_Input, NAME_None, FName(TEXT("In")));
-		}
-
-		if (PCGNode)
+		for (const FName& OutLabel : PCGNode->OutLabels())
 		{
-			for (const FName& InLabel : PCGNode->InLabels())
-			{
-				CreatePin(EEdGraphPinDirection::EGPD_Input, NAME_None, InLabel);
-			}
+			CreatePin(EEdGraphPinDirection::EGPD_Output, NAME_None, OutLabel);
 		}
 	}
-}
 
-void UPCGEditorGraphNode::AutowireNewNode(UEdGraphPin* FromPin)
-{
-	if (FromPin->Direction == EEdGraphPinDirection::EGPD_Output)
+	if (!PCGNode || PCGNode->HasDefaultInLabel())
 	{
-		FName InPinName = TEXT("In");
-		if (PCGNode && !PCGNode->HasDefaultInLabel() && !PCGNode->InLabels().IsEmpty())
-		{
-			InPinName = PCGNode->InLabels()[0];
-		}
-
-		UEdGraphPin* ToPin = FindPinChecked(InPinName, EEdGraphPinDirection::EGPD_Input);
-		GetSchema()->TryCreateConnection(FromPin, ToPin);
+		CreatePin(EEdGraphPinDirection::EGPD_Input, NAME_None, FName(TEXT("In")));
 	}
-	else if (FromPin->Direction == EEdGraphPinDirection::EGPD_Input)
+
+	if (PCGNode)
 	{
-		FName OutPinName = TEXT("Out");
-		if (PCGNode && !PCGNode->HasDefaultOutLabel() && !PCGNode->OutLabels().IsEmpty())
+		for (const FName& InLabel : PCGNode->InLabels())
 		{
-			OutPinName = PCGNode->OutLabels()[0];
+			CreatePin(EEdGraphPinDirection::EGPD_Input, NAME_None, InLabel);
 		}
-
-		UEdGraphPin* ToPin = FindPinChecked(OutPinName, EEdGraphPinDirection::EGPD_Output);
-		GetSchema()->TryCreateConnection(FromPin, ToPin);
-	}
-	NodeConnectionListChanged();
-}
-
-bool UPCGEditorGraphNode::CanUserDeleteNode() const
-{
-	return NodeType != EPCGEditorGraphNodeType::Input && NodeType != EPCGEditorGraphNodeType::Output;
-}
-
-void UPCGEditorGraphNode::OnNodeChanged(UPCGNode* InNode, bool bSettingsChanged)
-{
-	if (InNode == PCGNode)
-	{
-		ReconstructNode();
-	}
-}
-
-void UPCGEditorGraphNode::OnPickColor()
-{
-	FColorPickerArgs PickerArgs;
-	PickerArgs.bIsModal = true;
-	PickerArgs.bUseAlpha = false;
-	PickerArgs.InitialColorOverride = GetNodeTitleColor();
-	PickerArgs.OnColorCommitted = FOnLinearColorValueChanged::CreateUObject(this, &UPCGEditorGraphNode::OnColorPicked);
-
-	OpenColorPicker(PickerArgs);
-}
-
-void UPCGEditorGraphNode::OnColorPicked(FLinearColor NewColor)
-{
-	if (PCGNode && GetNodeTitleColor() != NewColor)
-	{
-		PCGNode->Modify();
-		PCGNode->NodeTitleColor = NewColor;
 	}
 }
 
 void UPCGEditorGraphNode::ReconstructNode()
 {
-	if (!PCGNode)
-	{
-		return;
-	}
-
-	TArray<FName> InLabels = PCGNode->InLabels();
-	TArray<FName> OutLabels = PCGNode->OutLabels();
-
-	TArray<UEdGraphPin*> InputPins;
-	TArray<UEdGraphPin*> OutputPins;
-
-	// Reallocate the default pins if they had been removed previously
-	if((NodeType == EPCGEditorGraphNodeType::Output || NodeType == EPCGEditorGraphNodeType::Settings) && PCGNode->HasDefaultInLabel() && !FindPin(TEXT("In"), EEdGraphPinDirection::EGPD_Input))
-	{
-		FCreatePinParams FirstPinParams;
-		FirstPinParams.Index = 0;
-		CreatePin(EEdGraphPinDirection::EGPD_Input, NAME_None, FName(TEXT("In")), FirstPinParams);
-	}
-
-	if ((NodeType == EPCGEditorGraphNodeType::Input || NodeType == EPCGEditorGraphNodeType::Settings) && PCGNode->HasDefaultOutLabel() && !FindPin(TEXT("Out"), EEdGraphPinDirection::EGPD_Output))
-	{
-		FCreatePinParams FirstPinParams;
-		FirstPinParams.Index = 0;
-		CreatePin(EEdGraphPinDirection::EGPD_Output, NAME_None, FName(TEXT("Out")), FirstPinParams);
-	}
-
-	for (UEdGraphPin* Pin : Pins)
-	{
-		if (Pin->Direction == EEdGraphPinDirection::EGPD_Input && (Pin->PinName != TEXT("In") || !PCGNode->HasDefaultInLabel()))
-		{
-			InputPins.Add(Pin);
-		}
-		else if (Pin->Direction == EEdGraphPinDirection::EGPD_Output && (Pin->PinName != TEXT("Out") || !PCGNode->HasDefaultOutLabel()))
-		{
-			OutputPins.Add(Pin);
-		}
-	}
-
-	auto UpdatePins = [this](TArray<UEdGraphPin*>& InPins, const TArray<FName>& InNames, EEdGraphPinDirection InPinDirection)
-	{
-		TArray<UEdGraphPin*> UnmatchedPins;
-
-		for (UEdGraphPin* InPin : InPins)
-		{
-			if (!InNames.Contains(InPin->PinName))
-			{
-				UnmatchedPins.Add(InPin);
-			}
-		}
-
-		TArray<FName> UnmatchedNames;
-
-		for (const FName& InName : InNames)
-		{
-			if (InPins.FindByPredicate([&InName](UEdGraphPin* InPin) { return InPin->PinName == InName; }) == nullptr)
-			{
-				UnmatchedNames.Add(InName);
-			}
-		}
-
-		// Use cases:
-		// - Removing pin(s) : remove pin(s) & break links
-		// - Remove all pins : remove all pins
-		// - Adding new pin(s) : create new pin(s)
-		// - Renaming a pin : find unmatched pin, unmatched name, rename
-		// - Complete change : match on names
-		if (UnmatchedNames.Num() == 1 && UnmatchedPins.Num() == 1)
-		{
-			UnmatchedPins[0]->PinName = UnmatchedNames[0];
-		}
-		else
-		{
-			for (UEdGraphPin* UnmatchedPin : UnmatchedPins)
-			{
-				UnmatchedPin->BreakAllPinLinks();
-				RemovePin(UnmatchedPin);
-			}
-
-			for (const FName& UnmatchedName : UnmatchedNames)
-			{
-				CreatePin(InPinDirection, NAME_None, UnmatchedName);
-			}
-		}
-
-		return UnmatchedNames.Num() > 0 || UnmatchedPins.Num() > 0;
-	};
-
-	if (NodeType == EPCGEditorGraphNodeType::Output || NodeType == EPCGEditorGraphNodeType::Settings)
-	{
-		UpdatePins(InputPins, InLabels, EEdGraphPinDirection::EGPD_Input);
-	}
-	
-	if (NodeType == EPCGEditorGraphNodeType::Input || NodeType == EPCGEditorGraphNodeType::Settings)
-	{
-		UpdatePins(OutputPins, OutLabels, EEdGraphPinDirection::EGPD_Output);
-	}
-
-	OnNodeChangedDelegate.ExecuteIfBound();
+	Super::ReconstructNode();
+	//TODO: Implement special version for settings to avoid the enum type
 }
 
 void UPCGEditorGraphNode::OnRenameNode(const FString& NewName)
@@ -310,29 +87,4 @@ void UPCGEditorGraphNode::OnRenameNode(const FString& NewName)
 	}
 }
 
-FLinearColor UPCGEditorGraphNode::GetNodeTitleColor() const
-{
-	if (PCGNode)
-	{
-		if (PCGNode->NodeTitleColor != FLinearColor::White)
-		{
-			return PCGNode->NodeTitleColor;
-		}
-		else if (PCGNode->DefaultSettings)
-		{
-			FLinearColor SettingsColor = PCGNode->DefaultSettings->GetNodeTitleColor();
-			if (SettingsColor == FLinearColor::White)
-			{
-				SettingsColor = GetDefault<UPCGEditorSettings>()->GetColor(PCGNode->DefaultSettings);
-			}
-
-			if (SettingsColor != FLinearColor::White)
-			{
-				return SettingsColor;
-			}
-		}
-	}
-
-	return GetDefault<UPCGEditorSettings>()->DefaultNodeColor;
-}
 #undef LOCTEXT_NAMESPACE
diff --git a/Engine/Plugins/Experimental/PCG/Source/PCGEditor/Private/PCGEditorGraphNode.h b/Engine/Plugins/Experimental/PCG/Source/PCGEditor/Private/PCGEditorGraphNode.h
--- a/Engine/Plugins/Experimental/PCG/Source/PCGEditor/Private/PCGEditorGraphNode.h
+++ b/Engine/Plugins/Experimental/PCG/Source/PCGEditor/Private/PCGEditorGraphNode.h
@@ -4,58 +4,23 @@
 
 #include "CoreMinimal.h"
 #include "UObject/ObjectMacros.h"
-#include "EdGraph/EdGraphNode.h"
+#include "PCGEditorGraphNodeBase.h"
 
 #include "PCGEditorGraphNode.generated.h"
 
-class UPCGNode;
 class UToolMenu;
 
-UENUM()
-enum class EPCGEditorGraphNodeType : uint8
-{
-	Input,
-	Output,
-	Settings
-};
-
 UCLASS()
-class UPCGEditorGraphNode : public UEdGraphNode
+class UPCGEditorGraphNode : public UPCGEditorGraphNodeBase
 {
 	GENERATED_BODY()
 
 public:
-	void Construct(UPCGNode* InPCGNode, EPCGEditorGraphNodeType InNodeType);
-
-	// ~Begin UObject interface
-	void BeginDestroy() override;
-	// ~End UObject interface
-
 	// ~Begin UEdGraphNode interface
 	virtual FText GetNodeTitle(ENodeTitleType::Type TitleType) const override;
 	virtual void GetNodeContextMenuActions(UToolMenu* Menu, class UGraphNodeContextMenuContext* Context) const override;
 	virtual void AllocateDefaultPins() override;
-	virtual void AutowireNewNode(UEdGraphPin* FromPin) override;
-	virtual bool CanUserDeleteNode() const override;
 	virtual void ReconstructNode() override;
 	virtual void OnRenameNode(const FString& NewName);
-	virtual FLinearColor GetNodeTitleColor() const;
 	// ~End UEdGraphNode interface
-
-	UPCGNode* GetPCGNode() { return PCGNode; }
-
-	DECLARE_DELEGATE(FOnPCGEditorGraphNodeChanged);
-	FOnPCGEditorGraphNodeChanged OnNodeChangedDelegate;
-
-protected:
-	void OnNodeChanged(UPCGNode* InNode, bool bSettingsChanged);
-	void OnPickColor();
-	void OnColorPicked(FLinearColor NewColor);
-
-private:
-	UPROPERTY()
-	TObjectPtr<UPCGNode> PCGNode = nullptr;
-
- 	UPROPERTY()
- 	EPCGEditorGraphNodeType NodeType = EPCGEditorGraphNodeType::Settings;
 };
diff --git a/Engine/Plugins/Experimental/PCG/Source/PCGEditor/Private/PCGEditorGraphNodeBase.cpp b/Engine/Plugins/Experimental/PCG/Source/PCGEditor/Private/PCGEditorGraphNodeBase.cpp
new file mode 100644
--- /dev/null
+++ b/Engine/Plugins/Experimental/PCG/Source/PCGEditor/Private/PCGEditorGraphNodeBase.cpp
@@ -0,0 +1,335 @@
+// Copyright Epic Games, Inc. All Rights Reserved.
+
+#include "PCGEditorGraphNodeBase.h"
+
+#include "PCGEditorGraph.h"
+#include "PCGEditorGraphSchema.h"
+#include "PCGEditorSettings.h"
+#include "PCGGraph.h"
+#include "PCGNode.h"
+#include "PCGSettings.h"
+
+#include "EdGraph/EdGraph.h"
+#include "EdGraph/EdGraphPin.h"
+#include "Framework/Commands/GenericCommands.h"
+#include "GraphEditorActions.h"
+#include "GraphEditorSettings.h"
+#include "ToolMenu.h"
+#include "ToolMenuSection.h"
+#include "Widgets/Colors/SColorPicker.h"
+
+#define LOCTEXT_NAMESPACE "PCGEditorGraphNodeBase"
+
+void UPCGEditorGraphNodeBase::BeginDestroy()
+{
+	if (PCGNode)
+	{
+		PCGNode->OnNodeSettingsChangedDelegate.RemoveAll(this);
+	}
+
+	Super::BeginDestroy();
+}
+
+void UPCGEditorGraphNodeBase::Construct(UPCGNode* InPCGNode, EPCGEditorGraphNodeType InNodeType)
+{
+	check(InPCGNode);
+	PCGNode = InPCGNode;
+	InPCGNode->OnNodeSettingsChangedDelegate.AddUObject(this, &UPCGEditorGraphNodeBase::OnNodeChanged);
+
+	NodePosX = InPCGNode->PositionX;
+	NodePosY = InPCGNode->PositionY;
+	
+	NodeType = InNodeType;
+
+	bCanRenameNode = (InNodeType == EPCGEditorGraphNodeType::Settings);
+}
+
+void UPCGEditorGraphNodeBase::GetNodeContextMenuActions(UToolMenu* Menu, class UGraphNodeContextMenuContext* Context) const
+{
+	if (!Context->Node)
+	{
+		return;
+	}
+
+	{
+		FToolMenuSection& Section = Menu->AddSection("EdGraphSchemaNodeActions", LOCTEXT("NodeActionsHeader", "Node Actions"));
+		Section.AddMenuEntry(FGraphEditorCommands::Get().BreakNodeLinks);
+	}
+
+	{
+		FToolMenuSection& Section = Menu->AddSection("EdGraphSchemaOrganization", LOCTEXT("OrganizationHeader", "Organization"));
+		Section.AddMenuEntry(
+			"PCGNode_SetColor",
+			LOCTEXT("PCGNode_SetColor", "Set Node Color"),
+			LOCTEXT("PCGNode_SetColorTooltip", "Sets a specific color on the given node. Note that white maps to the default value"),
+			FSlateIcon(FEditorStyle::GetStyleSetName(), "ColorPicker.Mode"),
+			FUIAction(FExecuteAction::CreateUObject(const_cast<UPCGEditorGraphNodeBase*>(this), &UPCGEditorGraphNodeBase::OnPickColor)));
+
+		Section.AddSubMenu("Alignment", LOCTEXT("AlignmentHeader", "Alignment"), FText(), FNewToolMenuDelegate::CreateLambda([](UToolMenu* AlignmentMenu)
+		{
+			{
+				FToolMenuSection& SubSection = AlignmentMenu->AddSection("EdGraphSchemaAlignment", LOCTEXT("AlignHeader", "Align"));
+				SubSection.AddMenuEntry(FGraphEditorCommands::Get().AlignNodesTop);
+				SubSection.AddMenuEntry(FGraphEditorCommands::Get().AlignNodesMiddle);
+				SubSection.AddMenuEntry(FGraphEditorCommands::Get().AlignNodesBottom);
+				SubSection.AddMenuEntry(FGraphEditorCommands::Get().AlignNodesLeft);
+				SubSection.AddMenuEntry(FGraphEditorCommands::Get().AlignNodesCenter);
+				SubSection.AddMenuEntry(FGraphEditorCommands::Get().AlignNodesRight);
+				SubSection.AddMenuEntry(FGraphEditorCommands::Get().StraightenConnections);
+			}
+
+			{
+				FToolMenuSection& SubSection = AlignmentMenu->AddSection("EdGraphSchemaDistribution", LOCTEXT("DistributionHeader", "Distribution"));
+				SubSection.AddMenuEntry(FGraphEditorCommands::Get().DistributeNodesHorizontally);
+				SubSection.AddMenuEntry(FGraphEditorCommands::Get().DistributeNodesVertically);
+			}
+		}));
+	}
+}
+
+void UPCGEditorGraphNodeBase::AutowireNewNode(UEdGraphPin* FromPin)
+{
+	if (FromPin->Direction == EEdGraphPinDirection::EGPD_Output)
+	{
+		FName InPinName = TEXT("In");
+		if (PCGNode && !PCGNode->HasDefaultInLabel() && !PCGNode->InLabels().IsEmpty())
+		{
+			InPinName = PCGNode->InLabels()[0];
+		}
+
+		UEdGraphPin* ToPin = FindPinChecked(InPinName, EEdGraphPinDirection::EGPD_Input);
+		GetSchema()->TryCreateConnection(FromPin, ToPin);
+	}
+	else if (FromPin->Direction == EEdGraphPinDirection::EGPD_Input)
+	{
+		FName OutPinName = TEXT("Out");
+		if (PCGNode && !PCGNode->HasDefaultOutLabel() && !PCGNode->OutLabels().IsEmpty())
+		{
+			OutPinName = PCGNode->OutLabels()[0];
+		}
+
+		UEdGraphPin* ToPin = FindPinChecked(OutPinName, EEdGraphPinDirection::EGPD_Output);
+		GetSchema()->TryCreateConnection(FromPin, ToPin);
+	}
+	NodeConnectionListChanged();
+}
+
+void UPCGEditorGraphNodeBase::PrepareForCopying()
+{
+	if (PCGNode)
+	{
+		// Temporarily take ownership of the MaterialExpression, so that it is not deleted when cutting
+		PCGNode->Rename(nullptr, this, REN_DontCreateRedirectors | REN_DoNotDirty);
+	}
+}
+
+bool UPCGEditorGraphNodeBase::CanCreateUnderSpecifiedSchema(const UEdGraphSchema* Schema) const
+{
+	return Schema->IsA(UPCGEditorGraphSchema::StaticClass());
+}
+
+void UPCGEditorGraphNodeBase::PostCopy()
+{
+	if (PCGNode)
+	{
+		UPCGEditorGraph* PCGEditorGraph = CastChecked<UPCGEditorGraph>(GetGraph());
+		UPCGGraph* PCGGraph = PCGEditorGraph->GetPCGGraph();
+		check(PCGGraph);
+		PCGNode->Rename(nullptr, PCGGraph, REN_DontCreateRedirectors | REN_DoNotDirty);
+	}
+}
+
+void UPCGEditorGraphNodeBase::PostPaste()
+{
+	if (PCGNode)
+	{
+		RebuildEdgesFromPins();
+
+		PCGNode->OnNodeSettingsChangedDelegate.AddUObject(this, &UPCGEditorGraphNodeBase::OnNodeChanged);
+		PCGNode->PositionX = NodePosX;
+		PCGNode->PositionY = NodePosY;
+	}
+}
+
+void UPCGEditorGraphNodeBase::RebuildEdgesFromPins()
+{
+	check(PCGNode)
+	
+	for (UEdGraphPin* Pin : Pins)
+	{
+		if (Pin->Direction == EEdGraphPinDirection::EGPD_Output)
+		{
+			for (UEdGraphPin* ConnectedPin : Pin->LinkedTo)
+			{
+				UEdGraphNode* ConnectedGraphNode = ConnectedPin->GetOwningNode();
+				UPCGEditorGraphNodeBase* ConnectedPCGGraphNode = CastChecked<UPCGEditorGraphNodeBase>(ConnectedGraphNode);
+
+				if (UPCGNode* ConnectedPCGNode = ConnectedPCGGraphNode->GetPCGNode())
+				{
+					const FName& NodeAPinName = ((Pin->PinName == TEXT("Out") && PCGNode->HasDefaultOutLabel()) ? NAME_None : Pin->PinName);
+					const FName& NodeBPinName = ((ConnectedPin->PinName == TEXT("In") && ConnectedPCGNode->HasDefaultInLabel()) ? NAME_None : ConnectedPin->PinName);
+
+					PCGNode->AddEdgeTo(NodeAPinName, ConnectedPCGNode, NodeBPinName);
+				}
+			}
+		}
+	}
+}
+
+void UPCGEditorGraphNodeBase::OnNodeChanged(UPCGNode* InNode, bool bSettingsChanged)
+{
+	if (InNode == PCGNode)
+	{
+		ReconstructNode();
+	}
+}
+
+void UPCGEditorGraphNodeBase::OnPickColor()
+{
+	FColorPickerArgs PickerArgs;
+	PickerArgs.bIsModal = true;
+	PickerArgs.bUseAlpha = false;
+	PickerArgs.InitialColorOverride = GetNodeTitleColor();
+	PickerArgs.OnColorCommitted = FOnLinearColorValueChanged::CreateUObject(this, &UPCGEditorGraphNodeBase::OnColorPicked);
+
+	OpenColorPicker(PickerArgs);
+}
+
+void UPCGEditorGraphNodeBase::OnColorPicked(FLinearColor NewColor)
+{
+	if (PCGNode && GetNodeTitleColor() != NewColor)
+	{
+		PCGNode->Modify();
+		PCGNode->NodeTitleColor = NewColor;
+	}
+}
+
+void UPCGEditorGraphNodeBase::ReconstructNode()
+{
+	if (!PCGNode)
+	{
+		return;
+	}
+
+	TArray<FName> InLabels = PCGNode->InLabels();
+	TArray<FName> OutLabels = PCGNode->OutLabels();
+
+	TArray<UEdGraphPin*> InputPins;
+	TArray<UEdGraphPin*> OutputPins;
+
+	// Reallocate the default pins if they had been removed previously
+	if ((NodeType == EPCGEditorGraphNodeType::Output || NodeType == EPCGEditorGraphNodeType::Settings) && PCGNode->HasDefaultInLabel() && !FindPin(TEXT("In"), EEdGraphPinDirection::EGPD_Input))
+	{
+		FCreatePinParams FirstPinParams;
+		FirstPinParams.Index = 0;
+		CreatePin(EEdGraphPinDirection::EGPD_Input, NAME_None, FName(TEXT("In")), FirstPinParams);
+	}
+
+	if ((NodeType == EPCGEditorGraphNodeType::Input || NodeType == EPCGEditorGraphNodeType::Settings) && PCGNode->HasDefaultOutLabel() && !FindPin(TEXT("Out"), EEdGraphPinDirection::EGPD_Output))
+	{
+		FCreatePinParams FirstPinParams;
+		FirstPinParams.Index = 0;
+		CreatePin(EEdGraphPinDirection::EGPD_Output, NAME_None, FName(TEXT("Out")), FirstPinParams);
+	}
+
+	for (UEdGraphPin* Pin : Pins)
+	{
+		if (Pin->Direction == EEdGraphPinDirection::EGPD_Input && (Pin->PinName != TEXT("In") || !PCGNode->HasDefaultInLabel()))
+		{
+			InputPins.Add(Pin);
+		}
+		else if (Pin->Direction == EEdGraphPinDirection::EGPD_Output && (Pin->PinName != TEXT("Out") || !PCGNode->HasDefaultOutLabel()))
+		{
+			OutputPins.Add(Pin);
+		}
+	}
+
+	auto UpdatePins = [this](TArray<UEdGraphPin*>& InPins, const TArray<FName>& InNames, EEdGraphPinDirection InPinDirection)
+	{
+		TArray<UEdGraphPin*> UnmatchedPins;
+
+		for (UEdGraphPin* InPin : InPins)
+		{
+			if (!InNames.Contains(InPin->PinName))
+			{
+				UnmatchedPins.Add(InPin);
+			}
+		}
+
+		TArray<FName> UnmatchedNames;
+
+		for (const FName& InName : InNames)
+		{
+			if (InPins.FindByPredicate([&InName](UEdGraphPin* InPin) { return InPin->PinName == InName; }) == nullptr)
+			{
+				UnmatchedNames.Add(InName);
+			}
+		}
+
+		// Use cases:
+		// - Removing pin(s) : remove pin(s) & break links
+		// - Remove all pins : remove all pins
+		// - Adding new pin(s) : create new pin(s)
+		// - Renaming a pin : find unmatched pin, unmatched name, rename
+		// - Complete change : match on names
+		if (UnmatchedNames.Num() == 1 && UnmatchedPins.Num() == 1)
+		{
+			UnmatchedPins[0]->PinName = UnmatchedNames[0];
+		}
+		else
+		{
+			for (UEdGraphPin* UnmatchedPin : UnmatchedPins)
+			{
+				UnmatchedPin->BreakAllPinLinks();
+				RemovePin(UnmatchedPin);
+			}
+
+			for (const FName& UnmatchedName : UnmatchedNames)
+			{
+				CreatePin(InPinDirection, NAME_None, UnmatchedName);
+			}
+		}
+
+		return UnmatchedNames.Num() > 0 || UnmatchedPins.Num() > 0;
+	};
+
+	if (NodeType == EPCGEditorGraphNodeType::Output || NodeType == EPCGEditorGraphNodeType::Settings)
+	{
+		UpdatePins(InputPins, InLabels, EEdGraphPinDirection::EGPD_Input);
+	}
+	
+	if (NodeType == EPCGEditorGraphNodeType::Input || NodeType == EPCGEditorGraphNodeType::Settings)
+	{
+		UpdatePins(OutputPins, OutLabels, EEdGraphPinDirection::EGPD_Output);
+	}
+
+	OnNodeChangedDelegate.ExecuteIfBound();
+}
+
+FLinearColor UPCGEditorGraphNodeBase::GetNodeTitleColor() const
+{
+	if (PCGNode)
+	{
+		if (PCGNode->NodeTitleColor != FLinearColor::White)
+		{
+			return PCGNode->NodeTitleColor;
+		}
+		else if (PCGNode->DefaultSettings)
+		{
+			FLinearColor SettingsColor = PCGNode->DefaultSettings->GetNodeTitleColor();
+			if (SettingsColor == FLinearColor::White)
+			{
+				SettingsColor = GetDefault<UPCGEditorSettings>()->GetColor(PCGNode->DefaultSettings);
+			}
+
+			if (SettingsColor != FLinearColor::White)
+			{
+				return SettingsColor;
+			}
+		}
+	}
+
+	return GetDefault<UPCGEditorSettings>()->DefaultNodeColor;
+}
+
+#undef LOCTEXT_NAMESPACE
diff --git a/Engine/Plugins/Experimental/PCG/Source/PCGEditor/Private/PCGEditorGraphNodeBase.h b/Engine/Plugins/Experimental/PCG/Source/PCGEditor/Private/PCGEditorGraphNodeBase.h
new file mode 100644
--- /dev/null
+++ b/Engine/Plugins/Experimental/PCG/Source/PCGEditor/Private/PCGEditorGraphNodeBase.h
@@ -0,0 +1,62 @@
+// Copyright Epic Games, Inc. All Rights Reserved.
+
+#pragma once
+
+#include "CoreMinimal.h"
+#include "UObject/ObjectMacros.h"
+#include "EdGraph/EdGraphNode.h"
+
+#include "PCGEditorGraphNodeBase.generated.h"
+
+class UPCGNode;
+class UToolMenu;
+
+UENUM()
+enum class EPCGEditorGraphNodeType : uint8
+{
+	Input,
+	Output,
+	Settings
+};
+
+UCLASS()
+class UPCGEditorGraphNodeBase : public UEdGraphNode
+{
+	GENERATED_BODY()
+
+public:
+	void Construct(UPCGNode* InPCGNode, EPCGEditorGraphNodeType InNodeType);
+
+	// ~Begin UObject interface
+	void BeginDestroy() override;
+	// ~End UObject interface
+
+	// ~Begin UEdGraphNode interface
+	virtual void GetNodeContextMenuActions(UToolMenu* Menu, class UGraphNodeContextMenuContext* Context) const override;
+	virtual void AutowireNewNode(UEdGraphPin* FromPin) override;
+	virtual void PrepareForCopying() override;
+	virtual bool CanCreateUnderSpecifiedSchema(const UEdGraphSchema* Schema) const override;
+	virtual void ReconstructNode() override;
+	virtual FLinearColor GetNodeTitleColor() const;
+	// ~End UEdGraphNode interface
+
+	UPCGNode* GetPCGNode() { return PCGNode; }
+	void PostCopy();
+	void PostPaste();
+
+	DECLARE_DELEGATE(FOnPCGEditorGraphNodeChanged);
+	FOnPCGEditorGraphNodeChanged OnNodeChangedDelegate;
+
+protected:
+	void RebuildEdgesFromPins();
+
+	void OnNodeChanged(UPCGNode* InNode, bool bSettingsChanged);
+	void OnPickColor();
+	void OnColorPicked(FLinearColor NewColor);
+
+	UPROPERTY()
+	TObjectPtr<UPCGNode> PCGNode = nullptr;
+
+ 	UPROPERTY()
+ 	EPCGEditorGraphNodeType NodeType = EPCGEditorGraphNodeType::Settings;
+};
diff --git a/Engine/Plugins/Experimental/PCG/Source/PCGEditor/Private/PCGEditorGraphNodeFactory.cpp b/Engine/Plugins/Experimental/PCG/Source/PCGEditor/Private/PCGEditorGraphNodeFactory.cpp
--- a/Engine/Plugins/Experimental/PCG/Source/PCGEditor/Private/PCGEditorGraphNodeFactory.cpp
+++ b/Engine/Plugins/Experimental/PCG/Source/PCGEditor/Private/PCGEditorGraphNodeFactory.cpp
@@ -2,12 +2,12 @@
 
 #include "PCGEditorGraphNodeFactory.h"
 
-#include "PCGEditorGraphNode.h"
+#include "PCGEditorGraphNodeBase.h"
 #include "SPCGEditorGraphNode.h"
 
 TSharedPtr<SGraphNode> FPCGEditorGraphNodeFactory::CreateNode(UEdGraphNode* InNode) const
 {
-	if (UPCGEditorGraphNode* GraphNode = Cast<UPCGEditorGraphNode>(InNode))
+	if (UPCGEditorGraphNodeBase* GraphNode = Cast<UPCGEditorGraphNodeBase>(InNode))
 	{
 		TSharedRef<SGraphNode> VisualNode =
 			SNew(SPCGEditorGraphNode, GraphNode);
diff --git a/Engine/Plugins/Experimental/PCG/Source/PCGEditor/Private/PCGEditorGraphNodeInput.cpp b/Engine/Plugins/Experimental/PCG/Source/PCGEditor/Private/PCGEditorGraphNodeInput.cpp
new file mode 100644
--- /dev/null
+++ b/Engine/Plugins/Experimental/PCG/Source/PCGEditor/Private/PCGEditorGraphNodeInput.cpp
@@ -0,0 +1,36 @@
+// Copyright Epic Games, Inc. All Rights Reserved.
+
+#include "PCGEditorGraphNodeInput.h"
+
+#include "PCGNode.h"
+
+#define LOCTEXT_NAMESPACE "PCGEditorGraphNodeInput"
+
+FText UPCGEditorGraphNodeInput::GetNodeTitle(ENodeTitleType::Type TitleType) const
+{	
+	return FText::FromName(TEXT("Input"));
+}
+
+void UPCGEditorGraphNodeInput::AllocateDefaultPins()
+{
+	if (!PCGNode || PCGNode->HasDefaultOutLabel())
+	{
+		CreatePin(EEdGraphPinDirection::EGPD_Output, NAME_None, FName(TEXT("Out")));
+	}
+
+	if (PCGNode)
+	{
+		for (const FName& OutLabel : PCGNode->OutLabels())
+		{
+			CreatePin(EEdGraphPinDirection::EGPD_Output, NAME_None, OutLabel);
+		}
+	}
+}
+
+void UPCGEditorGraphNodeInput::ReconstructNode()
+{
+	Super::ReconstructNode();
+	//TODO: Implement special version for input to avoid the enum type
+}
+
+#undef LOCTEXT_NAMESPACE
diff --git a/Engine/Plugins/Experimental/PCG/Source/PCGEditor/Private/PCGEditorGraphNodeInput.h b/Engine/Plugins/Experimental/PCG/Source/PCGEditor/Private/PCGEditorGraphNodeInput.h
new file mode 100644
--- /dev/null
+++ b/Engine/Plugins/Experimental/PCG/Source/PCGEditor/Private/PCGEditorGraphNodeInput.h
@@ -0,0 +1,24 @@
+// Copyright Epic Games, Inc. All Rights Reserved.
+
+#pragma once
+
+#include "CoreMinimal.h"
+#include "UObject/ObjectMacros.h"
+#include "PCGEditorGraphNodeBase.h"
+
+#include "PCGEditorGraphNodeInput.generated.h"
+
+UCLASS()
+class UPCGEditorGraphNodeInput : public UPCGEditorGraphNodeBase
+{
+	GENERATED_BODY()
+
+public:
+	// ~Begin UEdGraphNode interface
+	virtual FText GetNodeTitle(ENodeTitleType::Type TitleType) const override;
+	virtual void AllocateDefaultPins() override;
+	virtual bool CanUserDeleteNode() const override { return false; }
+	virtual bool CanDuplicateNode() const override { return false; }
+	virtual void ReconstructNode() override;
+	// ~End UEdGraphNode interface
+};
diff --git a/Engine/Plugins/Experimental/PCG/Source/PCGEditor/Private/PCGEditorGraphNodeOutput.cpp b/Engine/Plugins/Experimental/PCG/Source/PCGEditor/Private/PCGEditorGraphNodeOutput.cpp
new file mode 100644
--- /dev/null
+++ b/Engine/Plugins/Experimental/PCG/Source/PCGEditor/Private/PCGEditorGraphNodeOutput.cpp
@@ -0,0 +1,36 @@
+// Copyright Epic Games, Inc. All Rights Reserved.
+
+#include "PCGEditorGraphNodeOutput.h"
+
+#include "PCGNode.h"
+
+#define LOCTEXT_NAMESPACE "PCGEditorGraphNodeOutput"
+
+FText UPCGEditorGraphNodeOutput::GetNodeTitle(ENodeTitleType::Type TitleType) const
+{
+	return FText::FromName(TEXT("Output"));
+}
+
+void UPCGEditorGraphNodeOutput::AllocateDefaultPins()
+{
+	if (!PCGNode || PCGNode->HasDefaultInLabel())
+	{
+		CreatePin(EEdGraphPinDirection::EGPD_Input, NAME_None, FName(TEXT("In")));
+	}
+
+	if (PCGNode)
+	{
+		for (const FName& InLabel : PCGNode->InLabels())
+		{
+			CreatePin(EEdGraphPinDirection::EGPD_Input, NAME_None, InLabel);
+		}
+	}
+}
+
+void UPCGEditorGraphNodeOutput::ReconstructNode()
+{
+	Super::ReconstructNode();
+	//TODO: Implement special version for output to avoid the enum type
+}
+
+#undef LOCTEXT_NAMESPACE
diff --git a/Engine/Plugins/Experimental/PCG/Source/PCGEditor/Private/PCGEditorGraphNodeOutput.h b/Engine/Plugins/Experimental/PCG/Source/PCGEditor/Private/PCGEditorGraphNodeOutput.h
new file mode 100644
--- /dev/null
+++ b/Engine/Plugins/Experimental/PCG/Source/PCGEditor/Private/PCGEditorGraphNodeOutput.h
@@ -0,0 +1,24 @@
+// Copyright Epic Games, Inc. All Rights Reserved.
+
+#pragma once
+
+#include "CoreMinimal.h"
+#include "UObject/ObjectMacros.h"
+#include "PCGEditorGraphNodeBase.h"
+
+#include "PCGEditorGraphNodeOutput.generated.h"
+
+UCLASS()
+class UPCGEditorGraphNodeOutput : public UPCGEditorGraphNodeBase
+{
+	GENERATED_BODY()
+
+public:
+	// ~Begin UEdGraphNode interface
+	virtual FText GetNodeTitle(ENodeTitleType::Type TitleType) const override;
+	virtual void AllocateDefaultPins() override;
+	virtual bool CanUserDeleteNode() const override { return false; }
+	virtual bool CanDuplicateNode() const override { return false; }
+	virtual void ReconstructNode() override;
+	// ~End UEdGraphNode interface
+};
diff --git a/Engine/Plugins/Experimental/PCG/Source/PCGEditor/Private/PCGEditorGraphSchema.cpp b/Engine/Plugins/Experimental/PCG/Source/PCGEditor/Private/PCGEditorGraphSchema.cpp
--- a/Engine/Plugins/Experimental/PCG/Source/PCGEditor/Private/PCGEditorGraphSchema.cpp
+++ b/Engine/Plugins/Experimental/PCG/Source/PCGEditor/Private/PCGEditorGraphSchema.cpp
@@ -3,7 +3,7 @@
 #include "PCGEditorGraphSchema.h"
 
 #include "PCGEditorGraph.h"
-#include "PCGEditorGraphNode.h"
+#include "PCGEditorGraphNodeBase.h"
 #include "PCGEditorGraphSchemaActions.h"
 #include "PCGSettings.h"
 #include "PCGGraph.h"
@@ -78,8 +78,8 @@ bool UPCGEditorGraphSchema::TryCreateConnection(UEdGraphPin* InA, UEdGraphPin* I
 		UEdGraphNode* NodeA = A->GetOwningNode();
 		UEdGraphNode* NodeB = B->GetOwningNode();
 
-		UPCGEditorGraphNode* PCGGraphNodeA = CastChecked<UPCGEditorGraphNode>(NodeA);
-		UPCGEditorGraphNode* PCGGraphNodeB = CastChecked<UPCGEditorGraphNode>(NodeB);
+		UPCGEditorGraphNodeBase* PCGGraphNodeA = CastChecked<UPCGEditorGraphNodeBase>(NodeA);
+		UPCGEditorGraphNodeBase* PCGGraphNodeB = CastChecked<UPCGEditorGraphNodeBase>(NodeB);
 
 		UPCGNode* PCGNodeA = PCGGraphNodeA->GetPCGNode();
 		UPCGNode* PCGNodeB = PCGGraphNodeB->GetPCGNode();
@@ -103,7 +103,7 @@ void UPCGEditorGraphSchema::BreakPinLinks(UEdGraphPin& TargetPin, bool bSendsNod
 
 	UEdGraphNode* GraphNode = TargetPin.GetOwningNode();
 
-	UPCGEditorGraphNode* PCGGraphNode = CastChecked<UPCGEditorGraphNode>(GraphNode);
+	UPCGEditorGraphNodeBase* PCGGraphNode = CastChecked<UPCGEditorGraphNodeBase>(GraphNode);
 
 	UPCGNode* PCGNode = PCGGraphNode->GetPCGNode();
 	check(PCGNode);
@@ -130,8 +130,8 @@ void UPCGEditorGraphSchema::BreakSinglePinLink(UEdGraphPin* SourcePin, UEdGraphP
 	UEdGraphNode* SourceGraphNode = SourcePin->GetOwningNode();
 	UEdGraphNode* TargetGraphNode = TargetPin->GetOwningNode();
 
-	UPCGEditorGraphNode* SourcePCGGraphNode = CastChecked<UPCGEditorGraphNode>(SourceGraphNode);
-	UPCGEditorGraphNode* TargetPCGGraphNode = CastChecked<UPCGEditorGraphNode>(TargetGraphNode);
+	UPCGEditorGraphNodeBase* SourcePCGGraphNode = CastChecked<UPCGEditorGraphNodeBase>(SourceGraphNode);
+	UPCGEditorGraphNodeBase* TargetPCGGraphNode = CastChecked<UPCGEditorGraphNodeBase>(TargetGraphNode);
 
 	UPCGNode* SourcePCGNode = SourcePCGGraphNode->GetPCGNode();
 	UPCGNode* TargetPCGNode = TargetPCGGraphNode->GetPCGNode();
diff --git a/Engine/Plugins/Experimental/PCG/Source/PCGEditor/Private/PCGEditorGraphSchemaActions.h b/Engine/Plugins/Experimental/PCG/Source/PCGEditor/Private/PCGEditorGraphSchemaActions.h
--- a/Engine/Plugins/Experimental/PCG/Source/PCGEditor/Private/PCGEditorGraphSchemaActions.h
+++ b/Engine/Plugins/Experimental/PCG/Source/PCGEditor/Private/PCGEditorGraphSchemaActions.h
@@ -9,7 +9,6 @@
 
 class UEdGraph;
 class UEdGraphPin;
-class UPCGEditorGraphNode;
 class UPCGSettings;
 
 USTRUCT()
diff --git a/Engine/Plugins/Experimental/PCG/Source/PCGEditor/Private/SPCGEditorGraphNode.cpp b/Engine/Plugins/Experimental/PCG/Source/PCGEditor/Private/SPCGEditorGraphNode.cpp
--- a/Engine/Plugins/Experimental/PCG/Source/PCGEditor/Private/SPCGEditorGraphNode.cpp
+++ b/Engine/Plugins/Experimental/PCG/Source/PCGEditor/Private/SPCGEditorGraphNode.cpp
@@ -2,10 +2,10 @@
 
 #include "SPCGEditorGraphNode.h"
 
-#include "PCGEditorGraphNode.h"
+#include "PCGEditorGraphNodeBase.h"
 #include "PCGNode.h"
 
-void SPCGEditorGraphNode::Construct(const FArguments& InArgs, UPCGEditorGraphNode* InNode)
+void SPCGEditorGraphNode::Construct(const FArguments& InArgs, UPCGEditorGraphNodeBase* InNode)
 {
 	GraphNode = InNode;
 	PCGEditorGraphNode = InNode;
diff --git a/Engine/Plugins/Experimental/PCG/Source/PCGEditor/Private/SPCGEditorGraphNode.h b/Engine/Plugins/Experimental/PCG/Source/PCGEditor/Private/SPCGEditorGraphNode.h
--- a/Engine/Plugins/Experimental/PCG/Source/PCGEditor/Private/SPCGEditorGraphNode.h
+++ b/Engine/Plugins/Experimental/PCG/Source/PCGEditor/Private/SPCGEditorGraphNode.h
@@ -5,7 +5,7 @@
 #include "SGraphNode.h"
 #include "Widgets/DeclarativeSyntaxSupport.h"
 
-class UPCGEditorGraphNode;
+class UPCGEditorGraphNodeBase;
 
 class SPCGEditorGraphNode : public SGraphNode
 {
@@ -13,7 +13,7 @@ public:
 	SLATE_BEGIN_ARGS(SPCGEditorGraphNode){}
 	SLATE_END_ARGS()
 
-	void Construct(const FArguments& InArgs, UPCGEditorGraphNode* InNode);
+	void Construct(const FArguments& InArgs, UPCGEditorGraphNodeBase* InNode);
 
 	//~ Begin SGraphNode Interface
 	virtual void MoveTo(const FVector2D& NewPosition, FNodeSet& NodeFilter, bool bMarkDirty = true) override;
@@ -23,5 +23,5 @@ protected:
 	void OnNodeChanged();
 
 private:
-	UPCGEditorGraphNode* PCGEditorGraphNode = nullptr;
+	UPCGEditorGraphNodeBase* PCGEditorGraphNode = nullptr;
 };
