Fixed error in the break link implementation (mismatch on non-default pins)
Fixed static analysis warnings

diff --git a/Engine/Plugins/Experimental/PCG/Source/PCG/Private/Data/PCGIntersectionData.cpp b/Engine/Plugins/Experimental/PCG/Source/PCG/Private/Data/PCGIntersectionData.cpp
--- a/Engine/Plugins/Experimental/PCG/Source/PCG/Private/Data/PCGIntersectionData.cpp
+++ b/Engine/Plugins/Experimental/PCG/Source/PCG/Private/Data/PCGIntersectionData.cpp
@@ -83,7 +83,7 @@ FPCGPoint UPCGIntersectionData::TransformPoint(const FPCGPoint& InPoint) const
 { 
 	check(A && B);
 	const UPCGSpatialData* X = A->HasNonTrivialTransform() ? A : B;
-	const UPCGSpatialData* Y = X == A ? B : A;
+	const UPCGSpatialData* Y = (X == A) ? B : A;
 
 	FPCGPoint TransformedPoint = X->TransformPoint(InPoint);
 	if (TransformedPoint.Density > 0)
@@ -98,7 +98,7 @@ bool UPCGIntersectionData::GetPointAtPosition(const FVector& InPosition, FPCGPoi
 {
 	check(A && B);
 	const UPCGSpatialData* X = A->HasNonTrivialTransform() ? A : B;
-	const UPCGSpatialData* Y = X == A ? B : A;
+	const UPCGSpatialData* Y = (X == A) ? B : A;
 
 	FPCGPoint PointFromX;
 	if (!X->GetPointAtPosition(InPosition, PointFromX, OutMetadata))
diff --git a/Engine/Plugins/Experimental/PCG/Source/PCG/Private/Helpers/PCGAsync.cpp b/Engine/Plugins/Experimental/PCG/Source/PCG/Private/Helpers/PCGAsync.cpp
--- a/Engine/Plugins/Experimental/PCG/Source/PCG/Private/Helpers/PCGAsync.cpp
+++ b/Engine/Plugins/Experimental/PCG/Source/PCG/Private/Helpers/PCGAsync.cpp
@@ -37,7 +37,7 @@ namespace FPCGAsync
 		for (int32 TaskIndex = 0; TaskIndex < NumTasks; ++TaskIndex)
 		{
 			const int32 StartIndex = TaskIndex * IterationsPerTask;
-			const int32 EndIndex = (TaskIndex == NumTasks - 1 ? NumIterations : StartIndex + IterationsPerTask);
+			const int32 EndIndex = (TaskIndex == NumTasks - 1) ? NumIterations : (StartIndex + IterationsPerTask);
 
 			AsyncTasks.Emplace(Async(EAsyncExecution::ThreadPool, [&PointFunc, StartIndex, EndIndex, &OutPoints]()
 			{
@@ -114,7 +114,7 @@ namespace FPCGAsync
 		for (int32 TaskIndex = 0; TaskIndex < NumTasks; ++TaskIndex)
 		{
 			const int32 StartIndex = TaskIndex * IterationsPerTask;
-			const int32 EndIndex = (TaskIndex == NumTasks - 1 ? NumIterations : StartIndex + IterationsPerTask);
+			const int32 EndIndex = (TaskIndex == NumTasks - 1) ? NumIterations : (StartIndex + IterationsPerTask);
 
 			AsyncTasks.Emplace(Async(EAsyncExecution::ThreadPool, [&PointFunc, StartIndex, EndIndex, &InFilterPoints, &OutFilterPoints]()
 				{
diff --git a/Engine/Plugins/Experimental/PCG/Source/PCGEditor/Private/PCGEditorGraph.cpp b/Engine/Plugins/Experimental/PCG/Source/PCGEditor/Private/PCGEditorGraph.cpp
--- a/Engine/Plugins/Experimental/PCG/Source/PCGEditor/Private/PCGEditorGraph.cpp
+++ b/Engine/Plugins/Experimental/PCG/Source/PCGEditor/Private/PCGEditorGraph.cpp
@@ -48,7 +48,7 @@ void UPCGEditorGraph::InitFromNodeGraph(UPCGGraph* InPCGGraph)
 
 		for (UPCGEdge* OutboundEdge : PCGNode->GetOutboundEdges())
 		{
-			const FName OutPinName = (OutboundEdge->InboundLabel == NAME_None ? TEXT("Out") : OutboundEdge->InboundLabel);
+			const FName OutPinName = (OutboundEdge->InboundLabel == NAME_None) ? TEXT("Out") : OutboundEdge->InboundLabel;
 			UEdGraphPin* OutPin = GraphNode->FindPin(OutPinName);
 
 			if (!OutPin)
@@ -59,7 +59,7 @@ void UPCGEditorGraph::InitFromNodeGraph(UPCGGraph* InPCGGraph)
 			UPCGNode* OutboundNode = OutboundEdge->OutboundNode;
 			if (UPCGEditorGraphNodeBase** ConnectedGraphNode = NodeLookup.Find(OutboundNode))
 			{
-				const FName InPinName = (OutboundEdge->OutboundLabel == NAME_None ? TEXT("In") : OutboundEdge->OutboundLabel);
+				const FName InPinName = (OutboundEdge->OutboundLabel == NAME_None) ? TEXT("In") : OutboundEdge->OutboundLabel;
 
 				if(UEdGraphPin* InPin = (*ConnectedGraphNode)->FindPin(InPinName))
 				{
diff --git a/Engine/Plugins/Experimental/PCG/Source/PCGEditor/Private/PCGEditorGraphSchema.cpp b/Engine/Plugins/Experimental/PCG/Source/PCGEditor/Private/PCGEditorGraphSchema.cpp
--- a/Engine/Plugins/Experimental/PCG/Source/PCGEditor/Private/PCGEditorGraphSchema.cpp
+++ b/Engine/Plugins/Experimental/PCG/Source/PCGEditor/Private/PCGEditorGraphSchema.cpp
@@ -70,8 +70,8 @@ bool UPCGEditorGraphSchema::TryCreateConnection(UEdGraphPin* InA, UEdGraphPin* I
 	if (bModified)
 	{
 		check(InA && InB);
-		UEdGraphPin* A = (InA->Direction == EGPD_Output ? InA : InB);
-		UEdGraphPin* B = (InA->Direction == EGPD_Input ? InA : InB);
+		UEdGraphPin* A = (InA->Direction == EGPD_Output) ? InA : InB;
+		UEdGraphPin* B = (InA->Direction == EGPD_Input) ? InA : InB;
 		
 		check(A->Direction == EGPD_Output && B->Direction == EGPD_Input);
 
@@ -137,8 +137,8 @@ void UPCGEditorGraphSchema::BreakSinglePinLink(UEdGraphPin* SourcePin, UEdGraphP
 	UPCGNode* TargetPCGNode = TargetPCGGraphNode->GetPCGNode();
 	check(SourcePCGNode && TargetPCGNode);
 
-	const FName& SourcePinName = (SourcePin->PinName == TEXT("Out") ? NAME_None : SourcePin->PinName);
-	const FName& TargetPinName = (TargetPin->PinName == TEXT("In") ? NAME_None : TargetPin->PinName);
+	const FName& SourcePinName = (SourcePin->PinName == TEXT("Out") && SourcePCGNode->HasDefaultOutLabel()) ? NAME_None : SourcePin->PinName;
+	const FName& TargetPinName = (TargetPin->PinName == TEXT("In") && TargetPCGNode->HasDefaultInLabel()) ? NAME_None : TargetPin->PinName;
 
 	UPCGGraph* PCGGraph = SourcePCGNode->GetGraph();
 	PCGGraph->RemoveEdge(SourcePCGNode, SourcePinName, TargetPCGNode, TargetPinName);
