Fixed issue with new graphs that do not update the input/output pins properly
Changed color of pins to white
Changed style of edges

diff --git a/Engine/Plugins/Experimental/PCG/Source/PCG/Private/PCGGraph.cpp b/Engine/Plugins/Experimental/PCG/Source/PCG/Private/PCGGraph.cpp
--- a/Engine/Plugins/Experimental/PCG/Source/PCG/Private/PCGGraph.cpp
+++ b/Engine/Plugins/Experimental/PCG/Source/PCG/Private/PCGGraph.cpp
@@ -14,13 +14,18 @@ UPCGGraph::UPCGGraph(const FObjectInitializer& ObjectInitializer)
 	: Super(ObjectInitializer)
 {
 	InputNode = ObjectInitializer.CreateDefaultSubobject<UPCGNode>(this, TEXT("DefaultInputNode"));
-	InputNode->DefaultSettings = ObjectInitializer.CreateDefaultSubobject<UPCGGraphInputOutputSettings>(this, TEXT("DefaultInputNodeSettings"));
+	InputNode->SetDefaultSettings(ObjectInitializer.CreateDefaultSubobject<UPCGGraphInputOutputSettings>(this, TEXT("DefaultInputNodeSettings")));
 	OutputNode = ObjectInitializer.CreateDefaultSubobject<UPCGNode>(this, TEXT("DefaultOutputNode"));
-	OutputNode->DefaultSettings = ObjectInitializer.CreateDefaultSubobject<UPCGGraphInputOutputSettings>(this, TEXT("DefaultOutputNodeSettings"));
+	OutputNode->SetDefaultSettings(ObjectInitializer.CreateDefaultSubobject<UPCGGraphInputOutputSettings>(this, TEXT("DefaultOutputNodeSettings")));
 #if WITH_EDITORONLY_DATA
 	OutputNode->PositionX = 200;
 #endif
 
+#if WITH_EDITOR
+	InputNode->OnNodeSettingsChangedDelegate.AddUObject(this, &UPCGGraph::OnSettingsChanged);
+	OutputNode->OnNodeSettingsChangedDelegate.AddUObject(this, &UPCGGraph::OnSettingsChanged);
+#endif 
+
 	// Note: default connection from input to output
 	// should be added when creating from scratch,
 	// but not when using a blueprint construct script.
diff --git a/Engine/Plugins/Experimental/PCG/Source/PCGEditor/Private/PCGEditorGraphSchema.cpp b/Engine/Plugins/Experimental/PCG/Source/PCGEditor/Private/PCGEditorGraphSchema.cpp
--- a/Engine/Plugins/Experimental/PCG/Source/PCGEditor/Private/PCGEditorGraphSchema.cpp
+++ b/Engine/Plugins/Experimental/PCG/Source/PCGEditor/Private/PCGEditorGraphSchema.cpp
@@ -42,7 +42,12 @@ void UPCGEditorGraphSchema::GetGraphContextActions(FGraphContextMenuBuilder& Con
 
 FLinearColor UPCGEditorGraphSchema::GetPinTypeColor(const FEdGraphPinType& PinType) const
 {
-	return FLinearColor(1.0f, 0.0f, 1.0f, 1.0f);
+	return FLinearColor(1.0f, 1.0f, 1.0f, 1.0f);
+}
+
+FConnectionDrawingPolicy* UPCGEditorGraphSchema::CreateConnectionDrawingPolicy(int32 InBackLayerID, int32 InFrontLayerID, float InZoomFactor, const FSlateRect& InClippingRect, class FSlateWindowElementList& InDrawElements, class UEdGraph* InGraphObj) const
+{
+	return new FPCGEditorConnectionDrawingPolicy(InBackLayerID, InFrontLayerID, InZoomFactor, InClippingRect, InDrawElements, InGraphObj);
 }
 
 const FPinConnectionResponse UPCGEditorGraphSchema::CanCreateConnection(const UEdGraphPin* A, const UEdGraphPin* B) const
@@ -144,4 +149,22 @@ void UPCGEditorGraphSchema::BreakSinglePinLink(UEdGraphPin* SourcePin, UEdGraphP
 	PCGGraph->RemoveEdge(SourcePCGNode, SourcePinName, TargetPCGNode, TargetPinName);
 }
 
+FPCGEditorConnectionDrawingPolicy::FPCGEditorConnectionDrawingPolicy(int32 InBackLayerID, int32 InFrontLayerID, float InZoomFactor, const FSlateRect& InClippingRect, FSlateWindowElementList& InDrawElements, UEdGraph* InGraph)
+	: FConnectionDrawingPolicy(InBackLayerID, InFrontLayerID, InZoomFactor, InClippingRect, InDrawElements)
+	, Graph(CastChecked<UPCGEditorGraph>(InGraph))
+{
+	ArrowImage = nullptr;
+	ArrowRadius = FVector2D::ZeroVector;
+}
+
+void FPCGEditorConnectionDrawingPolicy::DetermineWiringStyle(UEdGraphPin* OutputPin, UEdGraphPin* InputPin, /*inout*/ FConnectionParams& Params)
+{
+	FConnectionDrawingPolicy::DetermineWiringStyle(OutputPin, InputPin, Params);
+	// Emphasize wire thickness on hovered pins
+	if (HoveredPins.Contains(InputPin) && HoveredPins.Contains(OutputPin))
+	{
+		Params.WireThickness = Params.WireThickness * 3;
+	}
+}
+
 #undef LOCTEXT_NAMESPACE
diff --git a/Engine/Plugins/Experimental/PCG/Source/PCGEditor/Private/PCGEditorGraphSchema.h b/Engine/Plugins/Experimental/PCG/Source/PCGEditor/Private/PCGEditorGraphSchema.h
--- a/Engine/Plugins/Experimental/PCG/Source/PCGEditor/Private/PCGEditorGraphSchema.h
+++ b/Engine/Plugins/Experimental/PCG/Source/PCGEditor/Private/PCGEditorGraphSchema.h
@@ -3,9 +3,12 @@
 #pragma once
 
 #include "EdGraph/EdGraphSchema.h"
+#include "ConnectionDrawingPolicy.h"
 
 #include "PCGEditorGraphSchema.generated.h"
 
+class UPCGEditorGraph;
+
 UCLASS()
 class UPCGEditorGraphSchema : public UEdGraphSchema
 {
@@ -19,5 +22,16 @@ public:
 	virtual bool TryCreateConnection(UEdGraphPin* A, UEdGraphPin* B) const override;
 	virtual void BreakPinLinks(UEdGraphPin& TargetPin, bool bSendsNodeNotification) const override;
 	virtual void BreakSinglePinLink(UEdGraphPin* SourcePin, UEdGraphPin* TargetPin) const override;
+	virtual FConnectionDrawingPolicy* CreateConnectionDrawingPolicy(int32 InBackLayerID, int32 InFrontLayerID, float InZoomFactor, const FSlateRect& InClippingRect, class FSlateWindowElementList& InDrawElements, class UEdGraph* InGraphObj) const override;
 	//~ End EdGraphSchema Interface
 };
+
+class FPCGEditorConnectionDrawingPolicy : public FConnectionDrawingPolicy
+{
+public:
+	FPCGEditorConnectionDrawingPolicy(int32 InBackLayerID, int32 InFrontLayerID, float InZoomFactor, const FSlateRect& InClippingRect, FSlateWindowElementList& InDrawElements, UEdGraph* InGraph);
+	virtual void DetermineWiringStyle(UEdGraphPin* OutputPin, UEdGraphPin* InputPin, /*inout*/ FConnectionParams& Params) override;
+
+protected:
+	UPCGEditorGraph* Graph;
+};
\ No newline at end of file
