Fixed issue where the edges would not bind properly if there was the same name on input & output pins.

diff --git a/Engine/Plugins/Experimental/PCG/Source/PCGEditor/Private/PCGEditorGraph.cpp b/Engine/Plugins/Experimental/PCG/Source/PCGEditor/Private/PCGEditorGraph.cpp
--- a/Engine/Plugins/Experimental/PCG/Source/PCGEditor/Private/PCGEditorGraph.cpp
+++ b/Engine/Plugins/Experimental/PCG/Source/PCGEditor/Private/PCGEditorGraph.cpp
@@ -49,7 +49,7 @@ void UPCGEditorGraph::InitFromNodeGraph(UPCGGraph* InPCGGraph)
 		for (UPCGEdge* OutboundEdge : PCGNode->GetOutboundEdges())
 		{
 			const FName OutPinName = (OutboundEdge->InboundLabel == NAME_None) ? TEXT("Out") : OutboundEdge->InboundLabel;
-			UEdGraphPin* OutPin = GraphNode->FindPin(OutPinName);
+			UEdGraphPin* OutPin = GraphNode->FindPin(OutPinName, EEdGraphPinDirection::EGPD_Output);
 
 			if (!OutPin)
 			{
@@ -61,7 +61,7 @@ void UPCGEditorGraph::InitFromNodeGraph(UPCGGraph* InPCGGraph)
 			{
 				const FName InPinName = (OutboundEdge->OutboundLabel == NAME_None) ? TEXT("In") : OutboundEdge->OutboundLabel;
 
-				if(UEdGraphPin* InPin = (*ConnectedGraphNode)->FindPin(InPinName))
+				if(UEdGraphPin* InPin = (*ConnectedGraphNode)->FindPin(InPinName, EEdGraphPinDirection::EGPD_Input))
 				{
 					OutPin->MakeLinkTo(InPin);
 				}
