Bug fixes:
Fixed issue with output label cleanup in subgraphs

QoL fixes:
Blueprint error in data initialization shouldnt crash, log error instead
Added output label validation
Improved most metadata operations to work on a single-attribute param/metadata without explicitly naming it

diff --git a/Engine/Plugins/Experimental/PCG/Source/PCG/Private/Data/PCGSpatialData.cpp b/Engine/Plugins/Experimental/PCG/Source/PCG/Private/Data/PCGSpatialData.cpp
--- a/Engine/Plugins/Experimental/PCG/Source/PCG/Private/Data/PCGSpatialData.cpp
+++ b/Engine/Plugins/Experimental/PCG/Source/PCG/Private/Data/PCGSpatialData.cpp
@@ -83,8 +83,7 @@ UPCGMetadata* UPCGSpatialData::CreateEmptyMetadata()
 
 void UPCGSpatialData::InitializeFromData(const UPCGSpatialData* InSource, const UPCGMetadata* InMetadataParentOverride)
 {
-	check(InSource);
-	if (!TargetActor)
+	if (InSource && !TargetActor)
 	{
 		TargetActor = InSource->TargetActor;
 	}
@@ -94,5 +93,12 @@ void UPCGSpatialData::InitializeFromData(const UPCGSpatialData* InSource, const
 		Metadata = NewObject<UPCGMetadata>(this);
 	}
 
-	Metadata->Initialize(InMetadataParentOverride ? InMetadataParentOverride : InSource->Metadata);
+	if (InMetadataParentOverride || InSource)
+	{
+		Metadata->Initialize(InMetadataParentOverride ? InMetadataParentOverride : InSource->Metadata);
+	}
+	else
+	{
+		UE_LOG(LogPCG, Warning, TEXT("InitializeFromData has both no source and no metadata override"));
+	}
 }
\ No newline at end of file
diff --git a/Engine/Plugins/Experimental/PCG/Source/PCG/Private/Elements/PCGMetadataElement.cpp b/Engine/Plugins/Experimental/PCG/Source/PCG/Private/Elements/PCGMetadataElement.cpp
--- a/Engine/Plugins/Experimental/PCG/Source/PCG/Private/Elements/PCGMetadataElement.cpp
+++ b/Engine/Plugins/Experimental/PCG/Source/PCG/Private/Elements/PCGMetadataElement.cpp
@@ -100,10 +100,12 @@ bool FPCGMetadataOperationElement::ExecuteInternal(FPCGContext* Context) const
 			continue;
 		}
 
+		const FName LocalSourceAttribute = ((SourceAttribute != NAME_None) ? SourceAttribute : OriginalData->Metadata->GetSingleAttributeNameOrNone());
+
 		// Check if the attribute exists
-		if ((Target == EPCGMetadataOperationTarget::AttributeToProperty || Target == EPCGMetadataOperationTarget::AttributeToAttribute) && !OriginalData->Metadata->HasAttribute(SourceAttribute))
+		if ((Target == EPCGMetadataOperationTarget::AttributeToProperty || Target == EPCGMetadataOperationTarget::AttributeToAttribute) && !OriginalData->Metadata->HasAttribute(LocalSourceAttribute))
 		{
-			PCGE_LOG(Warning, "Input does not have the %s attribute", *SourceAttribute.ToString());
+			PCGE_LOG(Warning, "Input does not have the %s attribute", *LocalSourceAttribute.ToString());
 			continue;
 		}
 
@@ -252,91 +254,91 @@ bool FPCGMetadataOperationElement::ExecuteInternal(FPCGContext* Context) const
 			{
 				auto DensitySetter = [](FPCGPoint& InPoint, const float& InValue) { InPoint.Density = InValue; };
 
-				const FPCGMetadataAttributeBase* AttributeBase = SampledData->Metadata->GetConstAttribute(SourceAttribute);
+				const FPCGMetadataAttributeBase* AttributeBase = SampledData->Metadata->GetConstAttribute(LocalSourceAttribute);
 				if(!PCGMetadataOperations::SetValue<float, float>(AttributeBase, SampledPoints, DensitySetter))
 				{
-					PCGE_LOG(Error, "Attribute %s already exists but its type is not compatible", *SourceAttribute.ToString());
+					PCGE_LOG(Error, "Attribute %s already exists but its type is not compatible", *LocalSourceAttribute.ToString());
 				}
 			}
 			else if (PointProperty == EPCGPointProperties::Extents)
 			{
 				auto ExtentsSetter = [](FPCGPoint& InPoint, const FVector& InValue) { InPoint.Extents = InValue; };
 
-				const FPCGMetadataAttributeBase* AttributeBase = SampledData->Metadata->GetConstAttribute(SourceAttribute);
+				const FPCGMetadataAttributeBase* AttributeBase = SampledData->Metadata->GetConstAttribute(LocalSourceAttribute);
 				if(!PCGMetadataOperations::SetValue<float, FVector>(AttributeBase, SampledPoints, ExtentsSetter) &&
 					!PCGMetadataOperations::SetValue<FVector, FVector>(AttributeBase, SampledPoints, ExtentsSetter))
 				{
-					PCGE_LOG(Error, "Attribute %s already exists but its type is not compatible", *SourceAttribute.ToString());
+					PCGE_LOG(Error, "Attribute %s already exists but its type is not compatible", *LocalSourceAttribute.ToString());
 				}
 			}
 			else if (PointProperty == EPCGPointProperties::Color)
 			{
 				auto ColorSetter = [](FPCGPoint& InPoint, const FVector4& InValue) { InPoint.Color = InValue; };
 
-				const FPCGMetadataAttributeBase* AttributeBase = SampledData->Metadata->GetConstAttribute(SourceAttribute);
+				const FPCGMetadataAttributeBase* AttributeBase = SampledData->Metadata->GetConstAttribute(LocalSourceAttribute);
 				if(!PCGMetadataOperations::SetValue<float, FVector4>(AttributeBase, SampledPoints, ColorSetter) &&
 					!PCGMetadataOperations::SetValue<FVector4, FVector4>(AttributeBase, SampledPoints, ColorSetter))
 				{
-					PCGE_LOG(Error, "Attribute %s already exists but its type is not compatible", *SourceAttribute.ToString());
+					PCGE_LOG(Error, "Attribute %s already exists but its type is not compatible", *LocalSourceAttribute.ToString());
 				}
 			}
 			else if (PointProperty == EPCGPointProperties::Position)
 			{
 				auto PositionSetter = [](FPCGPoint& InPoint, const FVector& InValue) { InPoint.Transform.SetLocation(InValue); };
 
-				const FPCGMetadataAttributeBase* AttributeBase = SampledData->Metadata->GetConstAttribute(SourceAttribute);
+				const FPCGMetadataAttributeBase* AttributeBase = SampledData->Metadata->GetConstAttribute(LocalSourceAttribute);
 				if (!PCGMetadataOperations::SetValue<float, FVector>(AttributeBase, SampledPoints, PositionSetter) &&
 					!PCGMetadataOperations::SetValue<FVector, FVector>(AttributeBase, SampledPoints, PositionSetter))
 				{
-					PCGE_LOG(Error, "Attribute %s already exists but its type is not compatible", *SourceAttribute.ToString());
+					PCGE_LOG(Error, "Attribute %s already exists but its type is not compatible", *LocalSourceAttribute.ToString());
 				}
 			}
 			else if (PointProperty == EPCGPointProperties::Rotation)
 			{
 				auto RotationSetter = [](FPCGPoint& InPoint, const FQuat& InValue) { InPoint.Transform.SetRotation(InValue.GetNormalized()); };
 
-				const FPCGMetadataAttributeBase* AttributeBase = SampledData->Metadata->GetConstAttribute(SourceAttribute);
+				const FPCGMetadataAttributeBase* AttributeBase = SampledData->Metadata->GetConstAttribute(LocalSourceAttribute);
 				if (!PCGMetadataOperations::SetValue<FQuat, FQuat>(AttributeBase, SampledPoints, RotationSetter))
 				{
-					PCGE_LOG(Error, "Attribute %s already exists but its type is not compatible", *SourceAttribute.ToString());
+					PCGE_LOG(Error, "Attribute %s already exists but its type is not compatible", *LocalSourceAttribute.ToString());
 				}
 			}
 			else if (PointProperty == EPCGPointProperties::Scale)
 			{
 				auto ScaleSetter = [](FPCGPoint& InPoint, const FVector& InValue) { InPoint.Transform.SetScale3D(InValue); };
 
-				const FPCGMetadataAttributeBase* AttributeBase = SampledData->Metadata->GetConstAttribute(SourceAttribute);
+				const FPCGMetadataAttributeBase* AttributeBase = SampledData->Metadata->GetConstAttribute(LocalSourceAttribute);
 				if (!PCGMetadataOperations::SetValue<float, FVector>(AttributeBase, SampledPoints, ScaleSetter) &&
 					!PCGMetadataOperations::SetValue<FVector, FVector>(AttributeBase, SampledPoints, ScaleSetter))
 				{
-					PCGE_LOG(Error, "Attribute %s already exists but its type is not compatible", *SourceAttribute.ToString());
+					PCGE_LOG(Error, "Attribute %s already exists but its type is not compatible", *LocalSourceAttribute.ToString());
 				}
 			}
 			else if (PointProperty == EPCGPointProperties::Transform)
 			{
 				auto TransformSetter = [](FPCGPoint& InPoint, const FTransform& InValue) { InPoint.Transform = InValue; };
 
-				const FPCGMetadataAttributeBase* AttributeBase = SampledData->Metadata->GetConstAttribute(SourceAttribute);
+				const FPCGMetadataAttributeBase* AttributeBase = SampledData->Metadata->GetConstAttribute(LocalSourceAttribute);
 				if (!PCGMetadataOperations::SetValue<FVector, FTransform>(AttributeBase, SampledPoints, TransformSetter) &&
 					!PCGMetadataOperations::SetValue<FTransform, FTransform>(AttributeBase, SampledPoints, TransformSetter))
 				{
-					PCGE_LOG(Error, "Attribute %s already exists but its type is not compatible", *SourceAttribute.ToString());
+					PCGE_LOG(Error, "Attribute %s already exists but its type is not compatible", *LocalSourceAttribute.ToString());
 				}
 			}
 			else if (PointProperty == EPCGPointProperties::Steepness)
 			{
 				auto SteepnessSetter = [](FPCGPoint& InPoint, const float& InValue) { InPoint.Steepness = InValue; };
 
-				const FPCGMetadataAttributeBase* AttributeBase = SampledData->Metadata->GetConstAttribute(SourceAttribute);
+				const FPCGMetadataAttributeBase* AttributeBase = SampledData->Metadata->GetConstAttribute(LocalSourceAttribute);
 				if (!PCGMetadataOperations::SetValue<float, float>(AttributeBase, SampledPoints, SteepnessSetter))
 				{
-					PCGE_LOG(Error, "Attribute %s already exists but its type is not compatible", *SourceAttribute.ToString());
+					PCGE_LOG(Error, "Attribute %s already exists but its type is not compatible", *LocalSourceAttribute.ToString());
 				}
 			}
 		}
 		else // Attribute to attribute
 		{
-			SampledData->Metadata->CopyAttribute(SourceAttribute, DestinationAttribute);
+			SampledData->Metadata->CopyAttribute(LocalSourceAttribute, DestinationAttribute);
 		}
 	}
 
diff --git a/Engine/Plugins/Experimental/PCG/Source/PCG/Private/Elements/PCGMetadataPartition.cpp b/Engine/Plugins/Experimental/PCG/Source/PCG/Private/Elements/PCGMetadataPartition.cpp
--- a/Engine/Plugins/Experimental/PCG/Source/PCG/Private/Elements/PCGMetadataPartition.cpp
+++ b/Engine/Plugins/Experimental/PCG/Source/PCG/Private/Elements/PCGMetadataPartition.cpp
@@ -27,20 +27,22 @@ bool FPCGMetadataPartitionElement::ExecuteInternal(FPCGContext* Context) const
 
 	const FName PartitionAttribute = PCGSettingsHelpers::GetValue(GET_MEMBER_NAME_CHECKED(UPCGMetadataPartitionSettings, PartitionAttribute), Settings->PartitionAttribute, Params);
 
-	if (PartitionAttribute == NAME_None)
-	{
-		PCGE_LOG(Error, "Metadata partition has no specified attribute");
-		return true;
-	}
-
 	for (const FPCGTaggedData& Input : Inputs)
 	{
 		const UPCGSpatialData* SpatialInput = Cast<const UPCGSpatialData>(Input.Data);
 		check(SpatialInput);
 
-		if (!SpatialInput->ConstMetadata() || !SpatialInput->ConstMetadata()->HasAttribute(PartitionAttribute))
+		if (!SpatialInput->ConstMetadata())
+		{
+			PCGE_LOG(Error, "Input does not have metadata");
+			continue;
+		}
+
+		const FName LocalPartitionAttribute = ((PartitionAttribute != NAME_None) ? PartitionAttribute : SpatialInput->ConstMetadata()->GetSingleAttributeNameOrNone());
+
+		if (!SpatialInput->ConstMetadata()->HasAttribute(LocalPartitionAttribute))
 		{
-			PCGE_LOG(Error, "Input does not have the %s attribute", *PartitionAttribute.ToString());
+			PCGE_LOG(Error, "Input does not have the %s attribute", *LocalPartitionAttribute.ToString());
 			continue;
 		}
 
@@ -55,7 +57,7 @@ bool FPCGMetadataPartitionElement::ExecuteInternal(FPCGContext* Context) const
 
 		const TArray<FPCGPoint>& SourcePoints = InputData->GetPoints();
 		const UPCGMetadata* SourceMetadata = InputData->ConstMetadata();
-		const FPCGMetadataAttributeBase* AttributeBase = SourceMetadata->GetConstAttribute(PartitionAttribute);
+		const FPCGMetadataAttributeBase* AttributeBase = SourceMetadata->GetConstAttribute(LocalPartitionAttribute);
 		check(AttributeBase);
 
 		TMap<PCGMetadataValueKey, UPCGPointData*> ValueToData;
diff --git a/Engine/Plugins/Experimental/PCG/Source/PCG/Private/Elements/PCGMetadataRenameElement.cpp b/Engine/Plugins/Experimental/PCG/Source/PCG/Private/Elements/PCGMetadataRenameElement.cpp
--- a/Engine/Plugins/Experimental/PCG/Source/PCG/Private/Elements/PCGMetadataRenameElement.cpp
+++ b/Engine/Plugins/Experimental/PCG/Source/PCG/Private/Elements/PCGMetadataRenameElement.cpp
@@ -27,7 +27,7 @@ bool FPCGMetadataRenameElement::ExecuteInternal(FPCGContext* Context) const
 	const FName AttributeToRename = PCGSettingsHelpers::GetValue(GET_MEMBER_NAME_CHECKED(UPCGMetadataRenameSettings, AttributeToRename), Settings->AttributeToRename, Params);
 	const FName NewAttributeName = PCGSettingsHelpers::GetValue(GET_MEMBER_NAME_CHECKED(UPCGMetadataRenameSettings, NewAttributeName), Settings->NewAttributeName, Params);
 
-	if (AttributeToRename == NAME_None || NewAttributeName == NAME_None)
+	if (NewAttributeName == NAME_None)
 	{
 		UE_LOG(LogPCG, Warning, TEXT("Metadata rename operation cannot rename from %s to %s"), *AttributeToRename.ToString(), *NewAttributeName.ToString());
 		// Bypass
@@ -49,7 +49,14 @@ bool FPCGMetadataRenameElement::ExecuteInternal(FPCGContext* Context) const
 		// otherwise, keep it as is.
 		const UPCGMetadata* Metadata = SpatialInput->Metadata;
 
-		if (!Metadata || !Metadata->HasAttribute(AttributeToRename))
+		if (!Metadata)
+		{
+			continue;
+		}
+
+		const FName LocalAttributeToRename = ((AttributeToRename != NAME_None) ? AttributeToRename : Metadata->GetSingleAttributeNameOrNone());
+
+		if (!Metadata->HasAttribute(LocalAttributeToRename))
 		{
 			continue;
 		}
@@ -57,7 +64,7 @@ bool FPCGMetadataRenameElement::ExecuteInternal(FPCGContext* Context) const
 		//TODO: this might require to execute on the main thread
 		UPCGSpatialData* NewSpatialData = Cast<UPCGSpatialData>(StaticDuplicateObject(SpatialInput, const_cast<UPCGSpatialData*>(SpatialInput), FName()));
 		NewSpatialData->InitializeFromData(SpatialInput);
-		NewSpatialData->Metadata->RenameAttribute(AttributeToRename, NewAttributeName);
+		NewSpatialData->Metadata->RenameAttribute(LocalAttributeToRename, NewAttributeName);
 
 		Output.Data = NewSpatialData;
 	}
@@ -74,7 +81,14 @@ bool FPCGMetadataRenameElement::ExecuteInternal(FPCGContext* Context) const
 		
 		const UPCGMetadata* Metadata = InputParams->Metadata;
 
-		if (!Metadata || !Metadata->HasAttribute(AttributeToRename))
+		if (!Metadata)
+		{
+			continue;
+		}
+
+		const FName LocalAttributeToRename = ((AttributeToRename != NAME_None) ? AttributeToRename : Metadata->GetSingleAttributeNameOrNone());
+
+		if (!Metadata->HasAttribute(LocalAttributeToRename))
 		{
 			continue;
 		}
@@ -84,7 +98,7 @@ bool FPCGMetadataRenameElement::ExecuteInternal(FPCGContext* Context) const
 		// Note: we will not parent the metadata here to ensure that the 0th entry is present on this metadata
 		// We will instead do a copy (and keep its parent, etc.)
 		NewParams->Metadata->InitializeAsCopy(Metadata);
-		NewParams->Metadata->RenameAttribute(AttributeToRename, NewAttributeName);
+		NewParams->Metadata->RenameAttribute(LocalAttributeToRename, NewAttributeName);
 
 		Output.Data = NewParams;
 	}
diff --git a/Engine/Plugins/Experimental/PCG/Source/PCG/Private/Elements/PCGPointFilter.cpp b/Engine/Plugins/Experimental/PCG/Source/PCG/Private/Elements/PCGPointFilter.cpp
--- a/Engine/Plugins/Experimental/PCG/Source/PCG/Private/Elements/PCGPointFilter.cpp
+++ b/Engine/Plugins/Experimental/PCG/Source/PCG/Private/Elements/PCGPointFilter.cpp
@@ -620,17 +620,7 @@ bool FPCGPointFilterElement::ExecuteInternal(FPCGContext* Context) const
 	bool bUseSpatialQuery = PCGSettingsHelpers::GetValue(GET_MEMBER_NAME_CHECKED(UPCGPointFilterSettings, bUseSpatialQuery), Settings->bUseSpatialQuery, Params);
 
 	// Validate basic input data
-	if (TargetFilterType == EPCGPointTargetFilterType::Metadata && TargetAttributeName == NAME_None)
-	{
-		PCGE_LOG(Error, "Point filter cannot filter on an unnamed metadata attribute (Target Attribute Name)");
-		return true;
-	}
-	else if (ThresholdFilterType == EPCGPointThresholdType::Metadata && ThresholdAttributeName == NAME_None)
-	{
-		PCGE_LOG(Error, "Point filter cannot filter on an unnamed metadata attribute (Threshold Attribute Name)");
-		return true;
-	}
-	else if (ThresholdFilterType == EPCGPointThresholdType::Property && !FilterData.IsEmpty() && Cast<UPCGSpatialData>(FilterData[0].Data) == nullptr)
+	if (ThresholdFilterType == EPCGPointThresholdType::Property && !FilterData.IsEmpty() && Cast<UPCGSpatialData>(FilterData[0].Data) == nullptr)
 	{
 		PCGE_LOG(Error, "Cannot filter by property on a non-spatial filter data");
 		return true;
@@ -691,6 +681,8 @@ bool FPCGPointFilterElement::ExecuteInternal(FPCGContext* Context) const
 			ThresholdPoints = &TargetPoints;
 		}
 
+		const FName LocalTargetAttributeName = ((TargetAttributeName != NAME_None || !TargetMetadata) ? TargetAttributeName : TargetMetadata->GetSingleAttributeNameOrNone());
+
 		// Additional validation
 		if (TargetFilterType == EPCGPointTargetFilterType::Metadata)
 		{
@@ -699,12 +691,14 @@ bool FPCGPointFilterElement::ExecuteInternal(FPCGContext* Context) const
 				PCGE_LOG(Error, "Target data to filter has no metadata which is required to filter by metadata");
 				continue;
 			}
-			else if (!TargetMetadata->GetConstAttribute(TargetAttributeName))
+			else if (!TargetMetadata->GetConstAttribute(LocalTargetAttributeName))
 			{
-				PCGE_LOG(Error, "Target metadata does not have the %s attribute", *TargetAttributeName.ToString());
+				PCGE_LOG(Error, "Target metadata does not have the %s attribute", *LocalTargetAttributeName.ToString());
 				continue;
 			}
 		}
+
+		const FName LocalThresholdAttributeName = ((ThresholdAttributeName != NAME_None || !ThresholdMetadata) ? ThresholdAttributeName : ThresholdMetadata->GetSingleAttributeNameOrNone());
 		
 		if (ThresholdFilterType == EPCGPointThresholdType::Metadata)
 		{
@@ -713,9 +707,9 @@ bool FPCGPointFilterElement::ExecuteInternal(FPCGContext* Context) const
 				PCGE_LOG(Error, "Filter data has no metadata which is required to filter by metadata");
 				continue;
 			}
-			else if (!ThresholdMetadata->GetConstAttribute(ThresholdAttributeName))
+			else if (!ThresholdMetadata->GetConstAttribute(LocalThresholdAttributeName))
 			{
-				PCGE_LOG(Error, "Filter metadata does not have the %s attribute", *ThresholdAttributeName.ToString());
+				PCGE_LOG(Error, "Filter metadata does not have the %s attribute", *LocalThresholdAttributeName.ToString());
 				continue;
 			}
 		}
@@ -742,7 +736,7 @@ bool FPCGPointFilterElement::ExecuteInternal(FPCGContext* Context) const
 			TargetFilterType, 
 			TargetPointProperty, 
 			TargetMetadata,
-			TargetAttributeName,
+			LocalTargetAttributeName,
 			TargetType);
 
 		if (TargetType == EPCGPointFilterConstantType::Unknown)
@@ -756,7 +750,7 @@ bool FPCGPointFilterElement::ExecuteInternal(FPCGContext* Context) const
 			ThresholdFilterType,
 			ThresholdPointProperty,
 			ThresholdMetadata,
-			ThresholdAttributeName,
+			LocalThresholdAttributeName,
 			ThresholdConstantType,
 			Integer64Constant,
 			FloatConstant,
diff --git a/Engine/Plugins/Experimental/PCG/Source/PCG/Private/Metadata/PCGMetadata.cpp b/Engine/Plugins/Experimental/PCG/Source/PCG/Private/Metadata/PCGMetadata.cpp
--- a/Engine/Plugins/Experimental/PCG/Source/PCG/Private/Metadata/PCGMetadata.cpp
+++ b/Engine/Plugins/Experimental/PCG/Source/PCG/Private/Metadata/PCGMetadata.cpp
@@ -208,6 +208,22 @@ void UPCGMetadata::GetAttributes(TArray<FName>& AttributeNames, TArray<EPCGMetad
 	AttributeLock.ReadUnlock();
 }
 
+FName UPCGMetadata::GetSingleAttributeNameOrNone() const
+{
+	FName SingleAttributeName = NAME_None;
+	AttributeLock.ReadLock();
+	if (Attributes.Num() == 1)
+	{
+		TArray<FName, TInlineAllocator<1>> Keys;
+		Attributes.GenerateKeyArray(Keys);
+		check(Keys.Num() == 1);
+		SingleAttributeName = Keys[0];
+	}
+	AttributeLock.ReadUnlock();
+
+	return SingleAttributeName;
+}
+
 bool UPCGMetadata::ParentHasAttribute(FName AttributeName) const
 {
 	return Parent && Parent->HasAttribute(AttributeName);
diff --git a/Engine/Plugins/Experimental/PCG/Source/PCG/Private/PCGElement.cpp b/Engine/Plugins/Experimental/PCG/Source/PCG/Private/PCGElement.cpp
--- a/Engine/Plugins/Experimental/PCG/Source/PCG/Private/PCGElement.cpp
+++ b/Engine/Plugins/Experimental/PCG/Source/PCG/Private/PCGElement.cpp
@@ -26,6 +26,7 @@ bool IPCGElement::Execute(FPCGContext* Context) const
 	{
 		//Pass-through
 		Context->OutputData = Context->InputData;
+		CleanupAndValidateOutputLabels(Context);
 		return true;
 	}
 	else
@@ -86,15 +87,7 @@ bool IPCGElement::Execute(FPCGContext* Context) const
 
 			if (bDone)
 			{
-				// Cleanup any residual labels if the node isn't supposed to produce them
-				// TODO: this is a bit of a crutch, could be refactored out if we review the way we push tagged data
-				if (Settings && Settings->OutLabels().Num() == 0)
-				{
-					for (FPCGTaggedData& TaggedData : Context->OutputData.TaggedData)
-					{
-						TaggedData.Pin = NAME_None;
-					}
-				}
+				CleanupAndValidateOutputLabels(Context);
 
 				if (IsCacheable(Settings) && Context->Cache)
 				{
@@ -165,6 +158,41 @@ bool IPCGElement::Execute(FPCGContext* Context) const
 	}
 }
 
+void IPCGElement::CleanupAndValidateOutputLabels(FPCGContext* Context) const
+{
+	check(Context);
+	const UPCGSettings* Settings = Context->GetInputSettings<UPCGSettings>();
+
+	if (!IsPassthrough() && Settings)
+	{
+		// Cleanup any residual labels if the node isn't supposed to produce them
+		// TODO: this is a bit of a crutch, could be refactored out if we review the way we push tagged data
+		TArray<FName> OutLabels = Settings->OutLabels();
+		if ((Settings->HasDefaultOutLabel() && OutLabels.Num() == 0) || (!Settings->HasDefaultOutLabel() && OutLabels.Num() == 1))
+		{
+			const FName& DefaultLabel = Settings->HasDefaultOutLabel() ? NAME_None : OutLabels[0];
+			for (FPCGTaggedData& TaggedData : Context->OutputData.TaggedData)
+			{
+				TaggedData.Pin = DefaultLabel;
+			}
+		}
+
+		// Validate all out data for errors in labels
+#if WITH_EDITOR
+		if (Settings->ExecutionMode != EPCGSettingsExecutionMode::Disabled)
+		{
+			for (FPCGTaggedData& TaggedData : Context->OutputData.TaggedData)
+			{
+				if (TaggedData.Pin != NAME_None && !OutLabels.Contains(TaggedData.Pin))
+				{
+					PCGE_LOG(Warning, "Output generated for pin %s but cannot be routed", *TaggedData.Pin.ToString());
+				}
+			}
+		}
+#endif
+	}
+}
+
 FPCGContext* FSimplePCGElement::Initialize(const FPCGDataCollection& InputData, UPCGComponent* SourceComponent, const UPCGNode* Node)
 {
 	FPCGContext* Context = new FPCGContext();
diff --git a/Engine/Plugins/Experimental/PCG/Source/PCG/Public/Elements/PCGSpawnActor.h b/Engine/Plugins/Experimental/PCG/Source/PCG/Public/Elements/PCGSpawnActor.h
--- a/Engine/Plugins/Experimental/PCG/Source/PCG/Public/Elements/PCGSpawnActor.h
+++ b/Engine/Plugins/Experimental/PCG/Source/PCG/Public/Elements/PCGSpawnActor.h
@@ -69,4 +69,5 @@ class FPCGSpawnActorElement : public FSimplePCGElement
 protected:
 	virtual bool ExecuteInternal(FPCGContext* Context) const override;
 	virtual bool IsCacheable(const UPCGSettings* InSettings) const override { return false; }
+	virtual bool IsPassthrough() const override { return true; }
 };
diff --git a/Engine/Plugins/Experimental/PCG/Source/PCG/Public/Metadata/PCGMetadata.h b/Engine/Plugins/Experimental/PCG/Source/PCG/Public/Metadata/PCGMetadata.h
--- a/Engine/Plugins/Experimental/PCG/Source/PCG/Public/Metadata/PCGMetadata.h
+++ b/Engine/Plugins/Experimental/PCG/Source/PCG/Public/Metadata/PCGMetadata.h
@@ -79,6 +79,8 @@ public:
 	UFUNCTION(BlueprintCallable, Category = "PCG|Metadata")
 	void GetAttributes(TArray<FName>& AttributeNames, TArray<EPCGMetadataTypes>& AttributeTypes) const;
 
+	FName GetSingleAttributeNameOrNone() const;
+
 	/** Delete/Hide attribute */
 	// Due to stream inheriting, we might want to consider "hiding" parent stream and deleting local streams only
 	void DeleteAttribute(FName AttributeName);
diff --git a/Engine/Plugins/Experimental/PCG/Source/PCG/Public/PCGElement.h b/Engine/Plugins/Experimental/PCG/Source/PCG/Public/PCGElement.h
--- a/Engine/Plugins/Experimental/PCG/Source/PCG/Public/PCGElement.h
+++ b/Engine/Plugins/Experimental/PCG/Source/PCG/Public/PCGElement.h
@@ -42,9 +42,13 @@ protected:
 	virtual bool ExecuteInternal(FPCGContext* Context) const = 0;
 	virtual bool IsCancellable() const { return true; }
 	virtual bool IsCacheable(const UPCGSettings* InSettings) const { return true; }
+	virtual bool IsPassthrough() const { return false; }
 #if WITH_EDITOR
 	virtual bool ShouldLog() const { return true; }
 #endif
+
+private:
+	void CleanupAndValidateOutputLabels(FPCGContext* Context) const;
 };
 
 /**
diff --git a/Engine/Plugins/Experimental/PCG/Source/PCG/Public/PCGSettings.h b/Engine/Plugins/Experimental/PCG/Source/PCG/Public/PCGSettings.h
--- a/Engine/Plugins/Experimental/PCG/Source/PCG/Public/PCGSettings.h
+++ b/Engine/Plugins/Experimental/PCG/Source/PCG/Public/PCGSettings.h
@@ -132,4 +132,5 @@ class PCG_API FPCGTrivialElement : public FSimplePCGElement
 protected:
 	virtual bool ExecuteInternal(FPCGContext* Context) const override;
 	virtual bool IsCacheable(const UPCGSettings* InSettings) const override { return false; }
+	virtual bool IsPassthrough() const override { return true; }
 };
diff --git a/Engine/Plugins/Experimental/PCG/Source/PCG/Public/PCGSubgraph.h b/Engine/Plugins/Experimental/PCG/Source/PCG/Public/PCGSubgraph.h
--- a/Engine/Plugins/Experimental/PCG/Source/PCG/Public/PCGSubgraph.h
+++ b/Engine/Plugins/Experimental/PCG/Source/PCG/Public/PCGSubgraph.h
@@ -138,6 +138,7 @@ public:
 
 protected:
 	virtual bool ExecuteInternal(FPCGContext* Context) const override;
+	virtual bool IsPassthrough() const override { return true; }
 };
 
 class PCG_API FPCGInputForwardingElement : public FSimplePCGElement
@@ -147,5 +148,6 @@ public:
 
 protected:
 	virtual bool ExecuteInternal(FPCGContext* Context) const override;
+	virtual bool IsPassthrough() const override { return true; }
 	FPCGDataCollection Input;
 };
