Fixed static analysis warning with the ?: pattern

#rnx

diff --git a/Engine/Plugins/Experimental/PCG/Source/PCG/Private/Data/PCGUnionData.cpp b/Engine/Plugins/Experimental/PCG/Source/PCG/Private/Data/PCGUnionData.cpp
--- a/Engine/Plugins/Experimental/PCG/Source/PCG/Private/Data/PCGUnionData.cpp
+++ b/Engine/Plugins/Experimental/PCG/Source/PCG/Private/Data/PCGUnionData.cpp
@@ -15,7 +15,7 @@ namespace PCGUnionDataMaths
 		}
 		else if (DensityFunction == EPCGUnionDensityFunction::Binary)
 		{
-			return InOtherDensity > 0 ? 1.0f : InDensityToUpdate;
+			return (InOtherDensity > 0) ? 1.0f : InDensityToUpdate;
 		}
 		else // Maximum
 		{
@@ -281,7 +281,7 @@ const UPCGPointData* UPCGUnionData::CreatePointData(FPCGContext* Context) const
 			{
 				for (FPCGPoint& TargetPoint : TargetPoints)
 				{
-					TargetPoint.Density = (TargetPoint.Density > 0 ? 1.0f : 0);
+					TargetPoint.Density = ((TargetPoint.Density > 0) ? 1.0f : 0);
 				}
 			}
 		}
diff --git a/Engine/Plugins/Experimental/PCG/Source/PCG/Private/Elements/PCGExecuteBlueprint.cpp b/Engine/Plugins/Experimental/PCG/Source/PCG/Private/Elements/PCGExecuteBlueprint.cpp
--- a/Engine/Plugins/Experimental/PCG/Source/PCG/Private/Elements/PCGExecuteBlueprint.cpp
+++ b/Engine/Plugins/Experimental/PCG/Source/PCG/Private/Elements/PCGExecuteBlueprint.cpp
@@ -310,7 +310,7 @@ FName UPCGBlueprintSettings::AdditionalTaskName() const
 	else
 	{
 #if WITH_EDITOR
-		return BlueprintElementType && BlueprintElementType->ClassGeneratedBy ? BlueprintElementType->ClassGeneratedBy->GetFName() : Super::AdditionalTaskName();
+		return (BlueprintElementType && BlueprintElementType->ClassGeneratedBy) ? BlueprintElementType->ClassGeneratedBy->GetFName() : Super::AdditionalTaskName();
 #else
 		return BlueprintElementType ? BlueprintElementType->GetFName() : Super::AdditionalTaskName();
 #endif
diff --git a/Engine/Plugins/Experimental/PCG/Source/PCG/Private/Elements/PCGPointFilter.cpp b/Engine/Plugins/Experimental/PCG/Source/PCG/Private/Elements/PCGPointFilter.cpp
--- a/Engine/Plugins/Experimental/PCG/Source/PCG/Private/Elements/PCGPointFilter.cpp
+++ b/Engine/Plugins/Experimental/PCG/Source/PCG/Private/Elements/PCGPointFilter.cpp
@@ -798,7 +798,7 @@ bool FPCGPointFilterElement::ExecuteInternal(FPCGContext* Context) const
 			FPCGPoint ThresholdPoint;
 			if (ThresholdPoints) // not null only in the direct match case
 			{
-				ThresholdPoint = (*ThresholdPoints)[ThresholdPoints->Num() > 1 ? Index : 0];
+				ThresholdPoint = (*ThresholdPoints)[(ThresholdPoints->Num() > 1) ? Index : 0];
 			}
 			else if(ThresholdSpatialData)
 			{
diff --git a/Engine/Plugins/Experimental/PCG/Source/PCG/Private/Helpers/PCGActorHelpers.cpp b/Engine/Plugins/Experimental/PCG/Source/PCG/Private/Helpers/PCGActorHelpers.cpp
--- a/Engine/Plugins/Experimental/PCG/Source/PCG/Private/Helpers/PCGActorHelpers.cpp
+++ b/Engine/Plugins/Experimental/PCG/Source/PCG/Private/Helpers/PCGActorHelpers.cpp
@@ -86,7 +86,7 @@ UInstancedStaticMeshComponent* UPCGActorHelpers::GetOrCreateISMC(AActor* InTarge
 	const int32 NumMaterials = ISMC->GetNumMaterials();
 	for (int32 MaterialIndex = 0; MaterialIndex < NumMaterials; ++MaterialIndex)
 	{
-		ISMC->SetMaterial(MaterialIndex, MaterialIndex < InMaterials.Num() ? InMaterials[MaterialIndex] : nullptr);
+		ISMC->SetMaterial(MaterialIndex, (MaterialIndex < InMaterials.Num()) ? InMaterials[MaterialIndex] : nullptr);
 	}
 
 	ISMC->RegisterComponent();
diff --git a/Engine/Plugins/Experimental/PCG/Source/PCG/Private/Metadata/PCGMetadata.cpp b/Engine/Plugins/Experimental/PCG/Source/PCG/Private/Metadata/PCGMetadata.cpp
--- a/Engine/Plugins/Experimental/PCG/Source/PCG/Private/Metadata/PCGMetadata.cpp
+++ b/Engine/Plugins/Experimental/PCG/Source/PCG/Private/Metadata/PCGMetadata.cpp
@@ -32,7 +32,7 @@ void UPCGMetadata::Initialize(const UPCGMetadata* InParent)
 		return;
 	}
 
-	Parent = (InParent != this ? InParent : nullptr);
+	Parent = ((InParent != this) ? InParent : nullptr);
 	ItemKeyOffset = Parent ? Parent->GetItemCountForChild() : 0;
 	AddAttributes(InParent);
 }
diff --git a/Engine/Plugins/Experimental/PCG/Source/PCG/Private/PCGComponent.cpp b/Engine/Plugins/Experimental/PCG/Source/PCG/Private/PCGComponent.cpp
--- a/Engine/Plugins/Experimental/PCG/Source/PCG/Private/PCGComponent.cpp
+++ b/Engine/Plugins/Experimental/PCG/Source/PCG/Private/PCGComponent.cpp
@@ -1275,7 +1275,7 @@ FBox UPCGComponent::GetGridBounds(AActor* Actor) const
 
 UPCGSubsystem* UPCGComponent::GetSubsystem() const
 {
-	return GetOwner() && GetOwner()->GetWorld() ? GetOwner()->GetWorld()->GetSubsystem<UPCGSubsystem>() : nullptr;
+	return (GetOwner() && GetOwner()->GetWorld()) ? GetOwner()->GetWorld()->GetSubsystem<UPCGSubsystem>() : nullptr;
 }
 
 #if WITH_EDITOR
diff --git a/Engine/Plugins/Experimental/PCG/Source/PCG/Private/PCGContext.cpp b/Engine/Plugins/Experimental/PCG/Source/PCG/Private/PCGContext.cpp
--- a/Engine/Plugins/Experimental/PCG/Source/PCG/Private/PCGContext.cpp
+++ b/Engine/Plugins/Experimental/PCG/Source/PCG/Private/PCGContext.cpp
@@ -8,7 +8,7 @@ FString FPCGContext::GetTaskName() const
 {
 	if (Node)
 	{
-		const FName NodeName = (Node->NodeTitle != NAME_None ? Node->NodeTitle : Node->GetFName());
+		const FName NodeName = ((Node->NodeTitle != NAME_None) ? Node->NodeTitle : Node->GetFName());
 
 		const UPCGSettings* Settings = GetInputSettings<UPCGSettings>();
 		const FName NodeAdditionalName = Settings ? Settings->AdditionalTaskName() : NAME_None;
diff --git a/Engine/Plugins/Experimental/PCG/Source/PCG/Private/PCGHelpers.cpp b/Engine/Plugins/Experimental/PCG/Source/PCG/Private/PCGHelpers.cpp
--- a/Engine/Plugins/Experimental/PCG/Source/PCG/Private/PCGHelpers.cpp
+++ b/Engine/Plugins/Experimental/PCG/Source/PCG/Private/PCGHelpers.cpp
@@ -109,7 +109,7 @@ namespace PCGHelpers
 #if WITH_EDITOR
 	APCGWorldActor* GetPCGWorldActor(UWorld* InWorld)
 	{
-		return InWorld && InWorld->GetSubsystem<UPCGSubsystem>() ? InWorld->GetSubsystem<UPCGSubsystem>()->GetPCGWorldActor() : nullptr;
+		return (InWorld && InWorld->GetSubsystem<UPCGSubsystem>()) ? InWorld->GetSubsystem<UPCGSubsystem>()->GetPCGWorldActor() : nullptr;
 	}
 
 	void GatherDependencies(UObject* Object, TSet<TObjectPtr<UObject>>& OutDependencies)
