Enabled plugin content

#rnx

diff --git a/Engine/Plugins/Experimental/PCG/PCG.uplugin b/Engine/Plugins/Experimental/PCG/PCG.uplugin
--- a/Engine/Plugins/Experimental/PCG/PCG.uplugin
+++ b/Engine/Plugins/Experimental/PCG/PCG.uplugin
@@ -12,7 +12,7 @@
 	"MarketplaceURL": "",
 	"SupportURL": "",
 	"EnabledByDefault": false,
-	"CanContainContent": false,
+	"CanContainContent": true,
 	"IsBetaVersion": false,
 	"Installed": false,
 	"Modules": [
