Added some methods to control how metadata is inherited when creating new data
Exposed methods to delete attributes

diff --git a/Engine/Plugins/Experimental/PCG/Source/PCG/Private/Data/PCGSpatialData.cpp b/Engine/Plugins/Experimental/PCG/Source/PCG/Private/Data/PCGSpatialData.cpp
--- a/Engine/Plugins/Experimental/PCG/Source/PCG/Private/Data/PCGSpatialData.cpp
+++ b/Engine/Plugins/Experimental/PCG/Source/PCG/Private/Data/PCGSpatialData.cpp
@@ -81,7 +81,7 @@ UPCGMetadata* UPCGSpatialData::CreateEmptyMetadata()
 	return Metadata;
 }
 
-void UPCGSpatialData::InitializeFromData(const UPCGSpatialData* InSource, const UPCGMetadata* InMetadataParentOverride)
+void UPCGSpatialData::InitializeFromData(const UPCGSpatialData* InSource, const UPCGMetadata* InMetadataParentOverride, bool bInheritMetadata)
 {
 	if (InSource && !TargetActor)
 	{
@@ -93,9 +93,11 @@ void UPCGSpatialData::InitializeFromData(const UPCGSpatialData* InSource, const
 		Metadata = NewObject<UPCGMetadata>(this);
 	}
 
-	if (InMetadataParentOverride || InSource)
+	if (!bInheritMetadata || InMetadataParentOverride || InSource)
 	{
-		Metadata->Initialize(InMetadataParentOverride ? InMetadataParentOverride : InSource->Metadata);
+		const UPCGMetadata* ParentMetadata = bInheritMetadata ? (InMetadataParentOverride ? InMetadataParentOverride : InSource->Metadata) : nullptr;
+
+		Metadata->Initialize(ParentMetadata);
 	}
 	else
 	{
diff --git a/Engine/Plugins/Experimental/PCG/Source/PCG/Private/Elements/PCGMetadataElement.cpp b/Engine/Plugins/Experimental/PCG/Source/PCG/Private/Elements/PCGMetadataElement.cpp
--- a/Engine/Plugins/Experimental/PCG/Source/PCG/Private/Elements/PCGMetadataElement.cpp
+++ b/Engine/Plugins/Experimental/PCG/Source/PCG/Private/Elements/PCGMetadataElement.cpp
@@ -338,7 +338,7 @@ bool FPCGMetadataOperationElement::ExecuteInternal(FPCGContext* Context) const
 		}
 		else // Attribute to attribute
 		{
-			SampledData->Metadata->CopyAttribute(LocalSourceAttribute, DestinationAttribute);
+			SampledData->Metadata->CopyExistingAttribute(LocalSourceAttribute, DestinationAttribute);
 		}
 	}
 
diff --git a/Engine/Plugins/Experimental/PCG/Source/PCG/Private/Metadata/PCGMetadata.cpp b/Engine/Plugins/Experimental/PCG/Source/PCG/Private/Metadata/PCGMetadata.cpp
--- a/Engine/Plugins/Experimental/PCG/Source/PCG/Private/Metadata/PCGMetadata.cpp
+++ b/Engine/Plugins/Experimental/PCG/Source/PCG/Private/Metadata/PCGMetadata.cpp
@@ -39,6 +39,11 @@ void UPCGMetadata::Initialize(const UPCGMetadata* InParent)
 
 void UPCGMetadata::InitializeAsCopy(const UPCGMetadata* InMetadataToCopy)
 {
+	if (!InMetadataToCopy)
+	{
+		return;
+	}
+
 	check(InMetadataToCopy);
 	if (Parent || Attributes.Num() != 0)
 	{
@@ -83,6 +88,68 @@ void UPCGMetadata::AddAttributes(const UPCGMetadata* InOther)
 	}
 }
 
+void UPCGMetadata::AddAttribute(const UPCGMetadata* InOther, FName AttributeName)
+{
+	if (!InOther || !InOther->HasAttribute(AttributeName) || HasAttribute(AttributeName))
+	{
+		return;
+	}
+
+	CopyAttribute(InOther->GetConstAttribute(AttributeName), AttributeName, /*bKeepParent=*/false, /*bCopyEntries=*/false, /*bCopyValues=*/false);
+	OtherParents.Add(InOther);
+}
+
+void UPCGMetadata::CopyAttributes(const UPCGMetadata* InOther)
+{
+	if (!InOther || InOther == Parent)
+	{
+		return;
+	}
+
+	if (GetItemCountForChild() != InOther->GetItemCountForChild())
+	{
+		UE_LOG(LogPCG, Error, TEXT("Mismatch in copy attributes since the entries do not match"));
+		return;
+	}
+
+	for (const TPair<FName, FPCGMetadataAttributeBase*> OtherAttribute : InOther->Attributes)
+	{
+		if (HasAttribute(OtherAttribute.Key))
+		{
+			continue;
+		}
+		else
+		{
+			CopyAttribute(OtherAttribute.Value, OtherAttribute.Key, /*bKeepParent=*/false, /*bCopyEntries=*/true, /*bCopyValues=*/true);
+		}
+	}
+}
+
+void UPCGMetadata::CopyAttribute(const UPCGMetadata* InOther, FName AttributeToCopy, FName NewAttributeName)
+{
+	if (!InOther)
+	{
+		return;
+	}
+	else if (HasAttribute(NewAttributeName) || !InOther->HasAttribute(AttributeToCopy))
+	{
+		return;
+	}
+	else if (InOther == Parent)
+	{
+		CopyExistingAttribute(AttributeToCopy, NewAttributeName);
+		return;
+	}
+
+	if (GetItemCountForChild() != InOther->GetItemCountForChild())
+	{
+		UE_LOG(LogPCG, Error, TEXT("Mismatch in copy attributes since the entries do not match"));
+		return;
+	}
+
+	CopyAttribute(InOther->GetConstAttribute(AttributeToCopy), NewAttributeName, /*bKeepParent=*/false, /*bCopyEntries=*/true, /*bCopyValues=*/true);
+}
+
 const UPCGMetadata* UPCGMetadata::GetRoot() const
 {
 	if (Parent)
@@ -264,7 +331,7 @@ void UPCGMetadata::CreateStringAttribute(FName AttributeName, const FString& Def
 	CreateAttribute<FString>(AttributeName, DefaultValue, bAllowsInterpolation, bOverrideParent);
 }
 
-void UPCGMetadata::CopyAttribute(FName AttributeToCopy, FName NewAttributeName, bool bKeepParent)
+void UPCGMetadata::CopyExistingAttribute(FName AttributeToCopy, FName NewAttributeName, bool bKeepParent)
 {
 	CopyAttribute(AttributeToCopy, NewAttributeName, bKeepParent, /*bCopyEntries=*/true, /*bCopyValues=*/true);
 }
diff --git a/Engine/Plugins/Experimental/PCG/Source/PCG/Public/Data/PCGSpatialData.h b/Engine/Plugins/Experimental/PCG/Source/PCG/Public/Data/PCGSpatialData.h
--- a/Engine/Plugins/Experimental/PCG/Source/PCG/Public/Data/PCGSpatialData.h
+++ b/Engine/Plugins/Experimental/PCG/Source/PCG/Public/Data/PCGSpatialData.h
@@ -109,7 +109,7 @@ public:
 	UPCGMetadata* CreateEmptyMetadata();
 
 	UFUNCTION(BlueprintCallable, Category = SpatialData)
-	void InitializeFromData(const UPCGSpatialData* InSource, const UPCGMetadata* InMetadataParentOverride = nullptr);
+	void InitializeFromData(const UPCGSpatialData* InSource, const UPCGMetadata* InMetadataParentOverride = nullptr, bool bInheritMetadata = true);
 
 	UPROPERTY(BlueprintReadWrite, EditAnywhere, Category = Data)
 	AActor* TargetActor = nullptr;
diff --git a/Engine/Plugins/Experimental/PCG/Source/PCG/Public/Metadata/PCGMetadata.h b/Engine/Plugins/Experimental/PCG/Source/PCG/Public/Metadata/PCGMetadata.h
--- a/Engine/Plugins/Experimental/PCG/Source/PCG/Public/Metadata/PCGMetadata.h
+++ b/Engine/Plugins/Experimental/PCG/Source/PCG/Public/Metadata/PCGMetadata.h
@@ -40,10 +40,22 @@ public:
 	UFUNCTION(BlueprintCallable, Category = "PCG|Metadata")
 	void InitializeAsCopy(const UPCGMetadata* InMetadataToCopy);
 
-	/** Creates streams for another metadata if they are not currently present */
+	/** Creates missing attributes from another metadata if they are not currently present - note that this does not copy values */
 	UFUNCTION(BlueprintCallable, Category = "PCG|Metadata")
 	void AddAttributes(const UPCGMetadata* InOther);
 
+	/** Creates missing attribute from another metadata if it is not currently present - note that this does not copy values */
+	UFUNCTION(BlueprintCallable, Category = "PCG|Metadata")
+	void AddAttribute(const UPCGMetadata* InOther, FName AttributeName);
+
+	/** Copies attributes from another metadata, including entries & values. Warning: this is intended when dealing with the same data set */
+	UFUNCTION(BlueprintCallable, Category = "PCG|Metadata|Advanced")
+	void CopyAttributes(const UPCGMetadata* InOther);
+
+	/** Copies an attribute from another metadata, including entries & values. Warning: this is intended when dealing with the same data set */
+	UFUNCTION(BlueprintCallable, Category = "PCG|Metadata|Advanced")
+	void CopyAttribute(const UPCGMetadata* InOther, FName AttributeToCopy, FName NewAttributeName);
+
 	/** Returns this metadata's parent */
 	const UPCGMetadata* GetParent() const { return Parent.Get(); }
 	const UPCGMetadata* GetRoot() const;
@@ -83,11 +95,12 @@ public:
 
 	/** Delete/Hide attribute */
 	// Due to stream inheriting, we might want to consider "hiding" parent stream and deleting local streams only
+	UFUNCTION(BlueprintCallable, Category = "PCG|Metadata")
 	void DeleteAttribute(FName AttributeName);
 
 	/** Copy attribute */
 	UFUNCTION(BlueprintCallable, Category = "PCG|Metadata")
-	void CopyAttribute(FName AttributeToCopy, FName NewAttributeName, bool bKeepParent = true);
+	void CopyExistingAttribute(FName AttributeToCopy, FName NewAttributeName, bool bKeepParent = true);
 
 	/** Rename attribute */
 	UFUNCTION(BlueprintCallable, Category = "PCG|Metadata")
