PCG: Moved default debug mesh and material to the plugin content folder.
#rnx

diff --git a/Engine/Plugins/Experimental/PCG/Source/PCG/Private/PCGDebug.cpp b/Engine/Plugins/Experimental/PCG/Source/PCG/Private/PCGDebug.cpp
--- a/Engine/Plugins/Experimental/PCG/Source/PCG/Private/PCGDebug.cpp
+++ b/Engine/Plugins/Experimental/PCG/Source/PCG/Private/PCGDebug.cpp
@@ -5,8 +5,8 @@
 
 namespace PCGDebugVisConstants
 {
-	const FSoftObjectPath DefaultPointMesh = FSoftObjectPath(TEXT("/Game/PCG/Content/DebugObjects/PCG_Cube.PCG_Cube"));
-	const FSoftObjectPath MaterialForDefaultPointMesh = FSoftObjectPath(TEXT("Material'/Game/PCG/Content/DebugObjects/PCG_DebugMaterial.PCG_DebugMaterial'"));
+	const FSoftObjectPath DefaultPointMesh = FSoftObjectPath(TEXT("/PCG/DebugObjects/PCG_Cube.PCG_Cube"));
+	const FSoftObjectPath MaterialForDefaultPointMesh = FSoftObjectPath(TEXT("Material'/PCG/DebugObjects/PCG_DebugMaterial.PCG_DebugMaterial'"));
 }
 
 FPCGDebugVisualizationSettings::FPCGDebugVisualizationSettings()
