Added advanced pins on input/output.
Improved editor behavior to be able to delete edges that have no pins anymore.
Improved notification mechanism in the graph to allow for temporary disable (using a stack counter).
Limited property changed mechanism to prevent property-less changes to trigger updates (could cause issues with texture streaming)

diff --git a/Engine/Plugins/Experimental/PCG/Source/PCG/Private/Elements/PCGDifferenceElement.cpp b/Engine/Plugins/Experimental/PCG/Source/PCG/Private/Elements/PCGDifferenceElement.cpp
--- a/Engine/Plugins/Experimental/PCG/Source/PCG/Private/Elements/PCGDifferenceElement.cpp
+++ b/Engine/Plugins/Experimental/PCG/Source/PCG/Private/Elements/PCGDifferenceElement.cpp
@@ -25,7 +25,7 @@ bool FPCGDifferenceElement::ExecuteInternal(FPCGContext* Context) const
 	TRACE_CPUPROFILER_EVENT_SCOPE(FPCGDifferenceElement::Execute);
 
 	// Early-out for previous behavior (without labeled edges)
-	if (Context->Node && !Context->Node->IsInputPinConnected(PCGDifferenceConstants::SourceLabel))
+	if (Context->Node && !Context->Node->IsInputPinConnected(PCGDifferenceConstants::SourceLabel) && !Context->Node->IsInputPinConnected(PCGDifferenceConstants::DifferencesLabel))
 	{
 		LabellessProcessing(Context);
 		return true;
diff --git a/Engine/Plugins/Experimental/PCG/Source/PCG/Private/Graph/PCGGraphExecutor.cpp b/Engine/Plugins/Experimental/PCG/Source/PCG/Private/Graph/PCGGraphExecutor.cpp
--- a/Engine/Plugins/Experimental/PCG/Source/PCG/Private/Graph/PCGGraphExecutor.cpp
+++ b/Engine/Plugins/Experimental/PCG/Source/PCG/Private/Graph/PCGGraphExecutor.cpp
@@ -6,6 +6,7 @@
 #include "PCGSubgraph.h"
 #include "PCGGraphCompiler.h"
 #include "PCGComponent.h"
+#include "PCGInputOutputSettings.h"
 
 // World partition support for in-editor workflows needs these includes
 #if WITH_EDITOR
@@ -407,6 +408,36 @@ bool FPCGFetchInputElement::ExecuteInternal(FPCGContext* Context) const
 	{
 		FPCGTaggedData& TaggedData = Context->OutputData.TaggedData.Emplace_GetRef();
 		TaggedData.Data = PCGData;
+		TaggedData.Pin = PCGInputOutputConstants::DefaultInLabel;
+	}
+
+	if (UPCGData* InputPCGData = Component->GetInputPCGData())
+	{
+		FPCGTaggedData& TaggedData = Context->OutputData.TaggedData.Emplace_GetRef();
+		TaggedData.Data = InputPCGData;
+		TaggedData.Pin = PCGInputOutputConstants::DefaultInputLabel;
+	}
+
+	if (UPCGData* ActorPCGData = Component->GetActorPCGData())
+	{
+		FPCGTaggedData& TaggedData = Context->OutputData.TaggedData.Emplace_GetRef();
+		TaggedData.Data = ActorPCGData;
+		TaggedData.Pin = PCGInputOutputConstants::DefaultActorLabel;
+	}
+
+	if (UPCGData* OriginalActorPCGData = Component->GetOriginalActorPCGData())
+	{
+		FPCGTaggedData& TaggedData = Context->OutputData.TaggedData.Emplace_GetRef();
+		TaggedData.Data = OriginalActorPCGData;
+		TaggedData.Pin = PCGInputOutputConstants::DefaultOriginalActorLabel;
+	}
+
+	TArray<UPCGData*> ExclusionsPCGData = Component->GetPCGExclusionData();
+	for (UPCGData* ExclusionPCGData : ExclusionsPCGData)
+	{
+		FPCGTaggedData& TaggedData = Context->OutputData.TaggedData.Emplace_GetRef();
+		TaggedData.Data = ExclusionPCGData;
+		TaggedData.Pin = PCGInputOutputConstants::DefaultExcludedActorsLabel;
 	}
 
 	return true;
diff --git a/Engine/Plugins/Experimental/PCG/Source/PCG/Private/Graph/PCGGraphExecutor.h b/Engine/Plugins/Experimental/PCG/Source/PCG/Private/Graph/PCGGraphExecutor.h
--- a/Engine/Plugins/Experimental/PCG/Source/PCG/Private/Graph/PCGGraphExecutor.h
+++ b/Engine/Plugins/Experimental/PCG/Source/PCG/Private/Graph/PCGGraphExecutor.h
@@ -134,6 +134,7 @@ class FPCGFetchInputElement : public FSimplePCGElement
 protected:
 	virtual bool ExecuteInternal(FPCGContext* Context) const override;
 	virtual bool IsCacheable(const UPCGSettings* InSettings) const override { return false; }
+	virtual bool IsPassthrough() const override { return true; }
 };
 
 class FPCGGenericElement : public FSimplePCGElement
diff --git a/Engine/Plugins/Experimental/PCG/Source/PCG/Private/PCGComponent.cpp b/Engine/Plugins/Experimental/PCG/Source/PCG/Private/PCGComponent.cpp
--- a/Engine/Plugins/Experimental/PCG/Source/PCG/Private/PCGComponent.cpp
+++ b/Engine/Plugins/Experimental/PCG/Source/PCG/Private/PCGComponent.cpp
@@ -759,6 +759,11 @@ void UPCGComponent::OnActorMoved(AActor* InActor)
 
 void UPCGComponent::OnObjectPropertyChanged(UObject* InObject, FPropertyChangedEvent& InEvent)
 {
+	if(!InEvent.Property)
+	{
+		return;
+	}
+
 	bool bValueNotInteractive = (InEvent.ChangeType != EPropertyChangeType::Interactive);
 	// Special exception for actor tags, as we can't track otherwise an actor "losing" a tag
 	bool bActorTagChange = (InEvent.Property && InEvent.Property->GetFName() == GET_MEMBER_NAME_CHECKED(AActor, Tags));
@@ -881,6 +886,7 @@ void UPCGComponent::DirtyGenerated(bool bInDirtyCachedInput)
 		CachedInputData = nullptr;
 		CachedActorData = nullptr;
 		CachedPCGData = nullptr;
+		CachedExclusionData.Reset();
 	}
 
 	// For partitioned graph, we must forward the call to the partition actor
@@ -962,6 +968,19 @@ UPCGData* UPCGComponent::GetActorPCGData()
 	return CachedActorData;
 }
 
+UPCGData* UPCGComponent::GetOriginalActorPCGData()
+{
+	if (APCGPartitionActor* PartitionActor = Cast<APCGPartitionActor>(GetOwner()))
+	{
+		if (UPCGComponent* OriginalComponent = PartitionActor->GetOriginalComponent(this))
+		{
+			return OriginalComponent->GetActorPCGData();
+		}
+	}
+
+	return nullptr;
+}
+
 TArray<UPCGData*> UPCGComponent::GetPCGExclusionData()
 {
 	// TODO: replace with a boolean, unify.
@@ -1006,6 +1025,9 @@ void UPCGComponent::UpdatePCGExclusionData()
 
 			if (InputSpatialData && ActorSpatialData)
 			{
+				// Change the target actor to this - otherwise we could push changes on another actor
+				ActorSpatialData->TargetActor = GetOwner();
+
 				// Create intersection or projection depending on the dimension
 				// TODO: there's an ambiguity here when it's the same dimension.
 				// For volumes, we'd expect an intersection, for surfaces we'd expect a projection
diff --git a/Engine/Plugins/Experimental/PCG/Source/PCG/Private/PCGGraph.cpp b/Engine/Plugins/Experimental/PCG/Source/PCG/Private/PCGGraph.cpp
--- a/Engine/Plugins/Experimental/PCG/Source/PCG/Private/PCGGraph.cpp
+++ b/Engine/Plugins/Experimental/PCG/Source/PCG/Private/PCGGraph.cpp
@@ -15,8 +15,10 @@ UPCGGraph::UPCGGraph(const FObjectInitializer& ObjectInitializer)
 {
 	InputNode = ObjectInitializer.CreateDefaultSubobject<UPCGNode>(this, TEXT("DefaultInputNode"));
 	InputNode->SetDefaultSettings(ObjectInitializer.CreateDefaultSubobject<UPCGGraphInputOutputSettings>(this, TEXT("DefaultInputNodeSettings")));
+	Cast<UPCGGraphInputOutputSettings>(InputNode->DefaultSettings)->SetInput(true);
 	OutputNode = ObjectInitializer.CreateDefaultSubobject<UPCGNode>(this, TEXT("DefaultOutputNode"));
 	OutputNode->SetDefaultSettings(ObjectInitializer.CreateDefaultSubobject<UPCGGraphInputOutputSettings>(this, TEXT("DefaultOutputNodeSettings")));
+	Cast<UPCGGraphInputOutputSettings>(OutputNode->DefaultSettings)->SetInput(false);
 #if WITH_EDITORONLY_DATA
 	OutputNode->PositionX = 200;
 #endif
@@ -44,10 +46,32 @@ void UPCGGraph::PostLoad()
 		InputNode->DefaultSettings = NewObject<UPCGGraphInputOutputSettings>(this, TEXT("DefaultInputNodeSettings"));
 	}
 
+	Cast<UPCGGraphInputOutputSettings>(InputNode->DefaultSettings)->SetInput(true);
+
+	// Fixup edges on the input node
+	for (UPCGEdge* Edge : InputNode->OutboundEdges)
+	{
+		if (Edge->InboundLabel == NAME_None)
+		{
+			Edge->InboundLabel = TEXT("In");
+		}
+	}
+
 	if (!Cast<UPCGGraphInputOutputSettings>(OutputNode->DefaultSettings))
 	{
 		OutputNode->DefaultSettings = NewObject<UPCGGraphInputOutputSettings>(this, TEXT("DefaultOutputNodeSettings"));
 	}
+
+	Cast<UPCGGraphInputOutputSettings>(OutputNode->DefaultSettings)->SetInput(false);
+
+	// Fixup edges on the output node
+	for (UPCGEdge* Edge : OutputNode->InboundEdges)
+	{
+		if (Edge->OutboundLabel == NAME_None)
+		{
+			Edge->OutboundLabel = TEXT("Out");
+		}
+	}
 #endif
 
 #if WITH_EDITOR
@@ -90,7 +114,7 @@ void UPCGGraph::InitializeFromTemplate()
 	Modify();
 
 	// Disable notification until the end of this method, otherwise it will cause issues in proper refresh dispatch
-	bEnableGraphChangeNotifications = false;
+	DisableNotificationsForEditor();
 
 	auto ResetDefaultNode = [](UPCGNode* InNode) {
 		check(InNode);
@@ -116,9 +140,8 @@ void UPCGGraph::InitializeFromTemplate()
 		Cast<UPCGGraphSetupBP>(GraphTemplate->GetDefaultObject())->Setup(this);
 	}
 
-	// Reenable notifications and notify listeners
-	bEnableGraphChangeNotifications = true;
-	NotifyGraphChanged(/*bIsStructural=*/true);
+	// Reenable notifications and notify listeners as needed
+	EnableNotificationsForEditor();
 }
 #endif
 
@@ -312,12 +335,12 @@ void UPCGGraph::RemoveNode(UPCGNode* InNode)
 	OnNodeRemoved(InNode);
 }
 
-void UPCGGraph::RemoveEdge(UPCGNode* From, const FName& FromLabel, UPCGNode* To, const FName& ToLabel)
+bool UPCGGraph::RemoveEdge(UPCGNode* From, const FName& FromLabel, UPCGNode* To, const FName& ToLabel)
 {
 	if (!From || !To)
 	{
 		UE_LOG(LogPCG, Error, TEXT("Invalid from/to node in RemoveEdge"));
-		return;
+		return false;
 	}
 
 	bool bChanged = false;
@@ -338,9 +361,11 @@ void UPCGGraph::RemoveEdge(UPCGNode* From, const FName& FromLabel, UPCGNode* To,
 		NotifyGraphChanged(/*bIsStructural=*/true);
 	}
 #endif
+
+	return bChanged;
 }
 
-void UPCGGraph::RemoveAllInboundEdges(UPCGNode* InNode)
+bool UPCGGraph::RemoveAllInboundEdges(UPCGNode* InNode)
 {
 	check(InNode);
 	const bool bChanged = !InNode->InboundEdges.IsEmpty();
@@ -356,9 +381,11 @@ void UPCGGraph::RemoveAllInboundEdges(UPCGNode* InNode)
 		NotifyGraphChanged(/*bIsStructural=*/true);
 	}
 #endif
+
+	return bChanged;
 }
 
-void UPCGGraph::RemoveAllOutboundEdges(UPCGNode* InNode)
+bool UPCGGraph::RemoveAllOutboundEdges(UPCGNode* InNode)
 {
 	check(InNode);
 	const bool bChanged = !InNode->OutboundEdges.IsEmpty();
@@ -374,9 +401,11 @@ void UPCGGraph::RemoveAllOutboundEdges(UPCGNode* InNode)
 		NotifyGraphChanged(/*bIsStructural=*/true);
 	}
 #endif
+
+	return bChanged;
 }
 
-void UPCGGraph::RemoveInboundEdges(UPCGNode* InNode, const FName& InboundLabel)
+bool UPCGGraph::RemoveInboundEdges(UPCGNode* InNode, const FName& InboundLabel)
 {
 	check(InNode);
 	bool bChanged = false;
@@ -396,9 +425,11 @@ void UPCGGraph::RemoveInboundEdges(UPCGNode* InNode, const FName& InboundLabel)
 		NotifyGraphChanged(/*bIsStructural=*/true);
 	}
 #endif
+
+	return bChanged;
 }
 
-void UPCGGraph::RemoveOutboundEdges(UPCGNode* InNode, const FName& OutboundLabel)
+bool UPCGGraph::RemoveOutboundEdges(UPCGNode* InNode, const FName& OutboundLabel)
 {
 	check(InNode);
 	bool bChanged = false;
@@ -418,8 +449,31 @@ void UPCGGraph::RemoveOutboundEdges(UPCGNode* InNode, const FName& OutboundLabel
 		NotifyGraphChanged(/*bIsStructural=*/true);
 	}
 #endif
+
+	return bChanged;
+}
+
+#if WITH_EDITOR
+void UPCGGraph::DisableNotificationsForEditor()
+{
+	check(GraphChangeNotificationsDisableCounter >= 0);
+	++GraphChangeNotificationsDisableCounter;
 }
 
+void UPCGGraph::EnableNotificationsForEditor()
+{
+	check(GraphChangeNotificationsDisableCounter > 0);
+	--GraphChangeNotificationsDisableCounter;
+
+	if (GraphChangeNotificationsDisableCounter == 0 && bDelayedChangeNotification)
+	{
+		NotifyGraphChanged(bDelayedChangeNotificationStructural);
+		bDelayedChangeNotification = false;
+		bDelayedChangeNotificationStructural = false;
+	}
+}
+#endif
+
 #if WITH_EDITOR
 FPCGTagToSettingsMap UPCGGraph::GetTrackedTagsToSettings() const
 {
@@ -437,8 +491,10 @@ FPCGTagToSettingsMap UPCGGraph::GetTrackedTagsToSettings() const
 
 void UPCGGraph::NotifyGraphChanged(bool bIsStructural)
 {
-	if (!bEnableGraphChangeNotifications)
+	if(GraphChangeNotificationsDisableCounter > 0)
 	{
+		bDelayedChangeNotification = true;
+		bDelayedChangeNotificationStructural |= bIsStructural;
 		return;
 	}
 
diff --git a/Engine/Plugins/Experimental/PCG/Source/PCG/Private/PCGInputOutputSettings.cpp b/Engine/Plugins/Experimental/PCG/Source/PCG/Private/PCGInputOutputSettings.cpp
--- a/Engine/Plugins/Experimental/PCG/Source/PCG/Private/PCGInputOutputSettings.cpp
+++ b/Engine/Plugins/Experimental/PCG/Source/PCG/Private/PCGInputOutputSettings.cpp
@@ -1,3 +1,22 @@
 // Copyright Epic Games, Inc. All Rights Reserved.
 
-#include "PCGInputOutputSettings.h"
\ No newline at end of file
+#include "PCGInputOutputSettings.h"
+
+bool FPCGInputOutputElement::ExecuteInternal(FPCGContext* Context) const
+{
+	// Essentially a pass-through element
+	Context->OutputData = Context->InputData;
+	return true;
+}
+
+UPCGGraphInputOutputSettings::UPCGGraphInputOutputSettings(const FObjectInitializer& ObjectInitializer)
+	: Super(ObjectInitializer)
+{
+	StaticInLabels.Add(PCGInputOutputConstants::DefaultInLabel);
+	StaticAdvancedInLabels.Add(PCGInputOutputConstants::DefaultInputLabel);
+	StaticAdvancedInLabels.Add(PCGInputOutputConstants::DefaultActorLabel);
+	StaticAdvancedInLabels.Add(PCGInputOutputConstants::DefaultOriginalActorLabel);
+	StaticAdvancedInLabels.Add(PCGInputOutputConstants::DefaultExcludedActorsLabel);
+	
+	StaticOutLabels.Add(PCGInputOutputConstants::DefaultOutLabel);
+}
\ No newline at end of file
diff --git a/Engine/Plugins/Experimental/PCG/Source/PCG/Public/PCGComponent.h b/Engine/Plugins/Experimental/PCG/Source/PCG/Public/PCGComponent.h
--- a/Engine/Plugins/Experimental/PCG/Source/PCG/Public/PCGComponent.h
+++ b/Engine/Plugins/Experimental/PCG/Source/PCG/Public/PCGComponent.h
@@ -51,6 +51,7 @@ public:
 	UPCGData* GetPCGData();
 	UPCGData* GetInputPCGData();
 	UPCGData* GetActorPCGData();
+	UPCGData* GetOriginalActorPCGData();
 	TArray<UPCGData*> GetPCGExclusionData();
 
 	bool CanPartition() const;
diff --git a/Engine/Plugins/Experimental/PCG/Source/PCG/Public/PCGGraph.h b/Engine/Plugins/Experimental/PCG/Source/PCG/Public/PCGGraph.h
--- a/Engine/Plugins/Experimental/PCG/Source/PCG/Public/PCGGraph.h
+++ b/Engine/Plugins/Experimental/PCG/Source/PCG/Public/PCGGraph.h
@@ -62,11 +62,16 @@ public:
 	const TArray<UPCGNode*>& GetNodes() const { return Nodes; }
 	void AddNode(UPCGNode* InNode);
 	void RemoveNode(UPCGNode* InNode);
-	void RemoveEdge(UPCGNode* From, const FName& FromLabel, UPCGNode* To, const FName& ToLabel);
-	void RemoveAllInboundEdges(UPCGNode* InNode);
-	void RemoveAllOutboundEdges(UPCGNode* InNode);
-	void RemoveInboundEdges(UPCGNode* InNode, const FName& InboundLabel);
-	void RemoveOutboundEdges(UPCGNode* InNode, const FName& OutboundLabel);
+	bool RemoveEdge(UPCGNode* From, const FName& FromLabel, UPCGNode* To, const FName& ToLabel);
+	bool RemoveAllInboundEdges(UPCGNode* InNode);
+	bool RemoveAllOutboundEdges(UPCGNode* InNode);
+	bool RemoveInboundEdges(UPCGNode* InNode, const FName& InboundLabel);
+	bool RemoveOutboundEdges(UPCGNode* InNode, const FName& OutboundLabel);
+
+#if WITH_EDITOR
+	void DisableNotificationsForEditor();
+	void EnableNotificationsForEditor();
+#endif
 
 #if WITH_EDITOR
 	FPCGTagToSettingsMap GetTrackedTagsToSettings() const;
@@ -96,6 +101,8 @@ private:
 	void OnSettingsChanged(UPCGNode* InNode, bool bSettingsChanged);
 	void OnStructuralSettingsChanged(UPCGNode* InNode);
 
-	bool bEnableGraphChangeNotifications = true;
+	int32 GraphChangeNotificationsDisableCounter = 0;
+	bool bDelayedChangeNotification = false;
+	bool bDelayedChangeNotificationStructural = false;
 #endif // WITH_EDITOR
 };
diff --git a/Engine/Plugins/Experimental/PCG/Source/PCG/Public/PCGInputOutputSettings.h b/Engine/Plugins/Experimental/PCG/Source/PCG/Public/PCGInputOutputSettings.h
--- a/Engine/Plugins/Experimental/PCG/Source/PCG/Public/PCGInputOutputSettings.h
+++ b/Engine/Plugins/Experimental/PCG/Source/PCG/Public/PCGInputOutputSettings.h
@@ -7,28 +7,87 @@
 
 #include "PCGInputOutputSettings.generated.h"
 
+namespace PCGInputOutputConstants
+{
+	const FName DefaultInLabel = TEXT("In");
+	const FName DefaultInputLabel = TEXT("Input");
+	const FName DefaultActorLabel = TEXT("Actor");
+	const FName DefaultOriginalActorLabel = TEXT("OriginalActor");
+	const FName DefaultExcludedActorsLabel = TEXT("ExcludedActors");
+
+	const FName DefaultOutLabel = TEXT("Out");
+}
+
+class PCG_API FPCGInputOutputElement : public FSimplePCGElement
+{
+protected:
+	virtual bool ExecuteInternal(FPCGContext* Context) const override;
+	virtual bool IsCacheable(const UPCGSettings* InSettings) const override { return false; }
+};
+
 UCLASS(NotBlueprintable, Hidden, ClassGroup = (Procedural))
 class PCG_API UPCGGraphInputOutputSettings : public UPCGSettings
 {
 	GENERATED_BODY()
 
 public:
+	UPCGGraphInputOutputSettings(const FObjectInitializer& ObjectInitializer);
+
 	// ~Begin UPCGSettings interface
 #if WITH_EDITOR
-	virtual FName GetDefaultNodeName() const override { return FName(TEXT("InputNode")); }
+	virtual FName GetDefaultNodeName() const override { return bIsInput ? FName(TEXT("InputNode")) : FName(TEXT("OutputNode")); }
 	virtual EPCGSettingsType GetType() const override { return EPCGSettingsType::InputOutput; }
 #endif
 
-	virtual bool HasInLabel(const FName& Label) const override { return Label == NAME_None || PinLabels.Contains(Label); }
-	virtual bool HasOutLabel(const FName& Label) const override { return Label == NAME_None || PinLabels.Contains(Label); }
-	virtual TArray<FName> InLabels() const override { return PinLabels.Array(); }
-	virtual TArray<FName> OutLabels() const override { return PinLabels.Array(); }
+	virtual bool HasInLabel(const FName& Label) const override { return Label == NAME_None || StaticLabels().Contains(Label) || (bShowAdvancedPins && StaticAdvancedLabels().Contains(Label)) || PinLabels.Contains(Label); }
+	virtual bool HasOutLabel(const FName& Label) const override { return Label == NAME_None || StaticLabels().Contains(Label) || (bShowAdvancedPins && StaticAdvancedLabels().Contains(Label)) || PinLabels.Contains(Label); }
+	virtual bool HasDefaultInLabel() const override { return false; }
+	virtual bool HasDefaultOutLabel() const override { return false; }
+
+	virtual TArray<FName> InLabels() const override
+	{
+		TArray<FName> Labels = StaticLabels();
+		if (bShowAdvancedPins)
+		{
+			Labels += StaticAdvancedLabels();
+		}
+		
+		Labels += PinLabels.Array();
+		return Labels;
+	}
+
+	virtual TArray<FName> OutLabels() const override
+	{
+		TArray<FName> Labels = StaticLabels();
+		if (bShowAdvancedPins)
+		{
+			Labels += StaticAdvancedLabels();
+		}
+
+		Labels += PinLabels.Array();
+		return Labels;
+	}
+
+	void SetInput(bool bInIsInput) { bIsInput = bInIsInput; }
 
 protected:
-	virtual FPCGElementPtr CreateElement() const override { return MakeShared<FPCGTrivialElement>(); }
+	virtual FPCGElementPtr CreateElement() const override { return MakeShared<FPCGInputOutputElement>(); }
 	// ~End UPCGSettings interface
 
+	const TArray<FName>& StaticLabels() const { return bIsInput ? StaticInLabels : StaticOutLabels; }
+	const TArray<FName>& StaticAdvancedLabels() const { return bIsInput ? StaticAdvancedInLabels : StaticAdvancedOutLabels; }
+
 protected:
 	UPROPERTY(EditAnywhere, BlueprintReadWrite, Category = "Settings|Input")
 	TSet<FName> PinLabels;
+
+	UPROPERTY(EditAnywhere, BlueprintReadWrite, Category = "Settings|Input")
+	bool bShowAdvancedPins = false;
+
+	TArray<FName> StaticInLabels;
+	TArray<FName> StaticAdvancedInLabels;
+	TArray<FName> StaticOutLabels;
+	TArray<FName> StaticAdvancedOutLabels;
+
+	bool bIsInput = false;
 };
\ No newline at end of file
diff --git a/Engine/Plugins/Experimental/PCG/Source/PCGEditor/Private/PCGEditorGraphNodeBase.cpp b/Engine/Plugins/Experimental/PCG/Source/PCGEditor/Private/PCGEditorGraphNodeBase.cpp
--- a/Engine/Plugins/Experimental/PCG/Source/PCGEditor/Private/PCGEditorGraphNodeBase.cpp
+++ b/Engine/Plugins/Experimental/PCG/Source/PCGEditor/Private/PCGEditorGraphNodeBase.cpp
@@ -153,7 +153,9 @@ void UPCGEditorGraphNodeBase::PostPaste()
 
 void UPCGEditorGraphNodeBase::RebuildEdgesFromPins()
 {
-	check(PCGNode)
+	check(PCGNode);
+
+	PCGNode->GetGraph()->DisableNotificationsForEditor();
 	
 	for (UEdGraphPin* Pin : Pins)
 	{
@@ -174,6 +176,8 @@ void UPCGEditorGraphNodeBase::RebuildEdgesFromPins()
 			}
 		}
 	}
+
+	PCGNode->GetGraph()->EnableNotificationsForEditor();
 }
 
 void UPCGEditorGraphNodeBase::OnNodeChanged(UPCGNode* InNode, bool bSettingsChanged)
@@ -278,12 +282,29 @@ void UPCGEditorGraphNodeBase::ReconstructNode()
 		}
 		else
 		{
+			PCGNode->GetGraph()->DisableNotificationsForEditor();
+
 			for (UEdGraphPin* UnmatchedPin : UnmatchedPins)
 			{
 				UnmatchedPin->BreakAllPinLinks();
+
+				if (PCGNode)
+				{
+					if (InPinDirection == EEdGraphPinDirection::EGPD_Input)
+					{
+						PCGNode->GetGraph()->RemoveInboundEdges(PCGNode, (UnmatchedPin->PinName == "In" && PCGNode->HasDefaultInLabel()) ? NAME_None : UnmatchedPin->PinName);
+					}
+					else
+					{
+						PCGNode->GetGraph()->RemoveOutboundEdges(PCGNode, (UnmatchedPin->PinName == "Out" && PCGNode->HasDefaultOutLabel()) ? NAME_None : UnmatchedPin->PinName);
+					}
+				}
+
 				RemovePin(UnmatchedPin);
 			}
 
+			PCGNode->GetGraph()->EnableNotificationsForEditor();
+
 			for (const FName& UnmatchedName : UnmatchedNames)
 			{
 				CreatePin(InPinDirection, NAME_None, UnmatchedName);
