Fixed issues with exclusion data not being regenerated properly
Improved behavior of potentially excluded actors outside of bounds
Fixed some edge cases where changing the tags wouldn't trigger an update

#rnx

diff --git a/Engine/Plugins/Experimental/PCG/Source/PCG/Private/PCGComponent.cpp b/Engine/Plugins/Experimental/PCG/Source/PCG/Private/PCGComponent.cpp
--- a/Engine/Plugins/Experimental/PCG/Source/PCG/Private/PCGComponent.cpp
+++ b/Engine/Plugins/Experimental/PCG/Source/PCG/Private/PCGComponent.cpp
@@ -243,7 +243,7 @@ FPCGTaskId UPCGComponent::GenerateInternal(bool bForce, const TArray<FPCGTaskId>
 	return TaskId;
 }
 
-bool UPCGComponent::GetActorsFromTags(const TSet<FName>& InTags, TSet<TWeakObjectPtr<AActor>>& OutActors)
+bool UPCGComponent::GetActorsFromTags(const TSet<FName>& InTags, TSet<TWeakObjectPtr<AActor>>& OutActors, bool bCullAgainstLocalBounds)
 {
 	UWorld* World = GetWorld();
 
@@ -252,6 +252,8 @@ bool UPCGComponent::GetActorsFromTags(const TSet<FName>& InTags, TSet<TWeakObjec
 		return false;
 	}
 
+	FBox LocalBounds = bCullAgainstLocalBounds ? GetGridBounds() : FBox(EForceInit::ForceInit);
+
 	TArray<AActor*> PerTagActors;
 
 	OutActors.Reset();
@@ -266,7 +268,10 @@ bool UPCGComponent::GetActorsFromTags(const TSet<FName>& InTags, TSet<TWeakObjec
 
 			for (AActor* Actor : PerTagActors)
 			{
-				OutActors.Emplace(Actor);
+				if (!bCullAgainstLocalBounds || LocalBounds.Intersect(GetGridBounds(Actor)))
+				{
+					OutActors.Emplace(Actor);
+				}
 			}
 
 			PerTagActors.Reset();
@@ -634,11 +639,11 @@ void UPCGComponent::SetupTrackingCallbacks()
 
 void UPCGComponent::RefreshTrackingData()
 {
-	GetActorsFromTags(ExcludedTags, CachedExcludedActors);
+	GetActorsFromTags(ExcludedTags, CachedExcludedActors, /*bCullAgainstLocalBounds=*/true);
 
 	TSet<FName> TrackedTags;
 	CachedTrackedTagsToSettings.GetKeys(TrackedTags);
-	GetActorsFromTags(TrackedTags, CachedTrackedActors);
+	GetActorsFromTags(TrackedTags, CachedTrackedActors, /*bCullAgainstLocalBounds=*/false);
 	PopulateTrackedActorToTagsMap(/*bForce=*/true);
 }
 
@@ -669,12 +674,28 @@ bool UPCGComponent::ActorHasExcludedTag(AActor* InActor) const
 	return bHasExcludedTag;
 }
 
-void UPCGComponent::DirtyExclusionData(AActor* InActor)
+bool UPCGComponent::UpdateExcludedActor(AActor* InActor)
 {
-	if (UPCGData** ExclusionData = CachedExclusionData.Find(InActor))
+	// Dirty data in all cases - the tag or positional changes will be picked up in the test later
+	if (CachedExcludedActors.Contains(InActor))
+	{
+		if (UPCGData** ExclusionData = CachedExclusionData.Find(InActor))
+		{
+			*ExclusionData = nullptr;
+		}
+
+		CachedPCGData = nullptr;
+		return true;
+	}
+	// Dirty only if the impact actor is inside the bounds
+	else if (ActorHasExcludedTag(InActor) && GetGridBounds().Intersect(GetGridBounds(InActor)))
 	{
-		*ExclusionData = nullptr;
 		CachedPCGData = nullptr;
+		return true;
+	}
+	else
+	{
+		return false;
 	}
 }
 
@@ -700,10 +721,10 @@ bool UPCGComponent::ActorIsTracked(AActor* InActor) const
 
 void UPCGComponent::OnActorAdded(AActor* InActor)
 {
-	const bool bHasExcludedTag = ActorHasExcludedTag(InActor);
+	const bool bIsExcluded = UpdateExcludedActor(InActor);
 	const bool bIsTracked = AddTrackedActor(InActor);
 
-	if (bHasExcludedTag || bIsTracked)
+	if (bIsExcluded || bIsTracked)
 	{
 		DirtyGenerated();
 		Refresh();
@@ -712,10 +733,10 @@ void UPCGComponent::OnActorAdded(AActor* InActor)
 
 void UPCGComponent::OnActorDeleted(AActor* InActor)
 {
-	const bool bHasExcludedTag = ActorHasExcludedTag(InActor);
+	const bool bWasExcluded = UpdateExcludedActor(InActor);
 	const bool bWasTracked = RemoveTrackedActor(InActor);
 
-	if (bHasExcludedTag || bWasTracked)
+	if (bWasExcluded || bWasTracked)
 	{
 		DirtyGenerated();
 		Refresh();
@@ -738,9 +759,8 @@ void UPCGComponent::OnActorMoved(AActor* InActor)
 	{
 		bool bDirtyAndRefresh = false;
 
-		if (ActorHasExcludedTag(InActor))
+		if (UpdateExcludedActor(InActor))
 		{
-			DirtyExclusionData(InActor);
 			bDirtyAndRefresh = true;
 		}
 
@@ -819,31 +839,14 @@ void UPCGComponent::OnActorChanged(AActor* Actor, UObject* InObject, bool bActor
 	{
 		bool bDirtyAndRefresh = false;
 
-		if (bActorTagChange && Actor == InObject)
+		if (UpdateExcludedActor(Actor))
 		{
-			// We don't really need to remove the actor here, but it's a good metric on whether the exclusion is dirty
-			if(CachedExcludedActors.Remove(Actor))
-			{
-				bDirtyAndRefresh = true;
-			}
-
-			if (UpdateTrackedActor(Actor))
-			{
-				bDirtyAndRefresh = true;
-			}
+			bDirtyAndRefresh = true;
 		}
-		else
-		{
-			if (ActorHasExcludedTag(Actor))
-			{
-				DirtyExclusionData(Actor);
-				bDirtyAndRefresh = true;
-			}
 
-			if (DirtyTrackedActor(Actor))
-			{
-				bDirtyAndRefresh = true;
-			}
+		if ((bActorTagChange && Actor == InObject && UpdateTrackedActor(Actor)) || DirtyTrackedActor(Actor))
+		{
+			bDirtyAndRefresh = true;
 		}
 
 		if (bDirtyAndRefresh)
@@ -997,7 +1000,7 @@ void UPCGComponent::UpdatePCGExclusionData()
 	const UPCGSpatialData* InputSpatialData = Cast<const UPCGSpatialData>(InputData);
 
 	// Update the list of cached excluded actors here, since we might not have picked up everything on map load (due to WP)
-	GetActorsFromTags(ExcludedTags, CachedExcludedActors);
+	GetActorsFromTags(ExcludedTags, CachedExcludedActors, /*bCullAgainstLocalBounds=*/true);
 
 	// Build exclusion data based on the CachedExcludedActors
 	TMap<AActor*, UPCGData*> ExclusionData;
diff --git a/Engine/Plugins/Experimental/PCG/Source/PCG/Public/PCGComponent.h b/Engine/Plugins/Experimental/PCG/Source/PCG/Public/PCGComponent.h
--- a/Engine/Plugins/Experimental/PCG/Source/PCG/Public/PCGComponent.h
+++ b/Engine/Plugins/Experimental/PCG/Source/PCG/Public/PCGComponent.h
@@ -139,7 +139,7 @@ private:
 	void CleanupInternal(bool bRemoveComponents, TSet<TSoftObjectPtr<AActor>>& OutActorsToDelete);
 	void PostProcessGraph(const FBox& InNewBounds, bool bInGenerated);
 
-	bool GetActorsFromTags(const TSet<FName>& InTags, TSet<TWeakObjectPtr<AActor>>& OutActors);
+	bool GetActorsFromTags(const TSet<FName>& InTags, TSet<TWeakObjectPtr<AActor>>& OutActors, bool bCullAgainstLocalBounds);
 
 #if WITH_EDITOR
 	virtual void PreEditChange(FProperty* PropertyAboutToChange) override;
@@ -162,7 +162,8 @@ private:
 	void OnObjectPropertyChanged(UObject* InObject, FPropertyChangedEvent& InEvent);
 	bool ActorHasExcludedTag(AActor* InActor) const;
 	bool ActorIsTracked(AActor* InActor) const;
-	void DirtyExclusionData(AActor* InActor);
+
+	bool UpdateExcludedActor(AActor* InActor);
 
 	void OnActorChanged(AActor* InActor, UObject* InSourceObject, bool bActorTagChange);
 
