Implemented PCG graph multithreaded execution.
Turned off by default until we iron-out the places in blueprints where it's causing issues.
Optimized PCG graph cache.
Optimized PCG task execution for already cached tasks.

diff --git a/Engine/Plugins/Experimental/PCG/Source/PCG/Private/Elements/PCGExecuteBlueprint.cpp b/Engine/Plugins/Experimental/PCG/Source/PCG/Private/Elements/PCGExecuteBlueprint.cpp
--- a/Engine/Plugins/Experimental/PCG/Source/PCG/Private/Elements/PCGExecuteBlueprint.cpp
+++ b/Engine/Plugins/Experimental/PCG/Source/PCG/Private/Elements/PCGExecuteBlueprint.cpp
@@ -483,6 +483,7 @@ void UPCGBlueprintElement::LoopNTimes(FPCGContext& InContext, int64 NumIteration
 
 FPCGContext* FPCGExecuteBlueprintElement::Initialize(const FPCGDataCollection& InputData, UPCGComponent* SourceComponent, const UPCGNode* Node)
 {
+	TRACE_CPUPROFILER_EVENT_SCOPE(FPCGExecuteBlueprintElement::Initialize);
 	FPCGBlueprintExecutionContext* Context = new FPCGBlueprintExecutionContext();
 	Context->InputData = InputData;
 	Context->SourceComponent = SourceComponent;
@@ -511,4 +512,16 @@ bool FPCGExecuteBlueprintElement::IsCacheable(const UPCGSettings* InSettings) co
 	{
 		return false;
 	}
+}
+
+bool FPCGExecuteBlueprintElement::CanExecuteOnlyOnMainThread(const UPCGSettings* InSettings) const
+{
+	if (const UPCGBlueprintSettings* BPSettings = Cast<const UPCGBlueprintSettings>(InSettings))
+	{
+		return !BPSettings->bCanBeMultithreaded;
+	}
+	else
+	{
+		return false;
+	}
 }
\ No newline at end of file
diff --git a/Engine/Plugins/Experimental/PCG/Source/PCG/Private/Graph/PCGGraphCache.cpp b/Engine/Plugins/Experimental/PCG/Source/PCG/Private/Graph/PCGGraphCache.cpp
--- a/Engine/Plugins/Experimental/PCG/Source/PCG/Private/Graph/PCGGraphCache.cpp
+++ b/Engine/Plugins/Experimental/PCG/Source/PCG/Private/Graph/PCGGraphCache.cpp
@@ -12,6 +12,7 @@ FPCGGraphCacheEntry::FPCGGraphCacheEntry(const FPCGDataCollection& InInput, cons
 	, Output(InOutput)
 {
 	Settings = InSettings ? Cast<UPCGSettings>(StaticDuplicateObject(InSettings, InOwner.Get())) : nullptr;
+	SettingsCrc32 = InSettings ? InSettings->GetCrc32() : 0;
 
 	Input.RootUnrootedData(OutRootedObjects);
 	Output.RootUnrootedData(OutRootedObjects);
@@ -23,10 +24,9 @@ FPCGGraphCacheEntry::FPCGGraphCacheEntry(const FPCGDataCollection& InInput, cons
 	}
 }
 
-bool FPCGGraphCacheEntry::Matches(const FPCGDataCollection& InInput, const UPCGSettings* InSettings) const
+bool FPCGGraphCacheEntry::Matches(const FPCGDataCollection& InInput, int32 InSettingsCrc32) const
 {
-	const bool bHasSameSettings = ((InSettings == nullptr && Settings == nullptr) || (InSettings && Settings && *InSettings == *Settings));
-	return bHasSameSettings && (Input == InInput);
+	return (SettingsCrc32 == InSettingsCrc32) && (Input == InInput);
 }
 
 FPCGGraphCache::FPCGGraphCache(TWeakObjectPtr<UObject> InOwner)
@@ -46,13 +46,16 @@ bool FPCGGraphCache::GetFromCache(const IPCGElement* InElement, const FPCGDataCo
 		return false;
 	}
 
+	TRACE_CPUPROFILER_EVENT_SCOPE(FPCGGraphCache::GetFromCache);
 	FReadScopeLock ScopedReadLock(CacheLock);
 
 	if (const FPCGGraphCacheEntries* Entries = CacheData.Find(InElement))
 	{
+		int32 InSettingsCrc32 = (InSettings ? InSettings->GetCrc32() : 0);
+
 		for (const FPCGGraphCacheEntry& Entry : *Entries)
 		{
-			if (Entry.Matches(InInput, InSettings))
+			if (Entry.Matches(InInput, InSettingsCrc32))
 			{
 				OutOutput = Entry.Output;
 				return true;
@@ -72,6 +75,7 @@ void FPCGGraphCache::StoreInCache(const IPCGElement* InElement, const FPCGDataCo
 		return;
 	}
 
+	TRACE_CPUPROFILER_EVENT_SCOPE(FPCGGraphCache::StoreInCache);
 	FWriteScopeLock ScopedWriteLock(CacheLock);
 
 	FPCGGraphCacheEntries* Entries = CacheData.Find(InElement);
diff --git a/Engine/Plugins/Experimental/PCG/Source/PCG/Private/Graph/PCGGraphCache.h b/Engine/Plugins/Experimental/PCG/Source/PCG/Private/Graph/PCGGraphCache.h
--- a/Engine/Plugins/Experimental/PCG/Source/PCG/Private/Graph/PCGGraphCache.h
+++ b/Engine/Plugins/Experimental/PCG/Source/PCG/Private/Graph/PCGGraphCache.h
@@ -14,11 +14,12 @@ struct FPCGGraphCacheEntry
 	FPCGGraphCacheEntry() = default;
 	FPCGGraphCacheEntry(const FPCGDataCollection& InInput, const UPCGSettings* InSettings, const FPCGDataCollection& InOutput, TWeakObjectPtr<UObject> InOwner, TSet<UObject*>& OutRootedData);
 
-	bool Matches(const FPCGDataCollection& InInput, const UPCGSettings* InSettings) const;
+	bool Matches(const FPCGDataCollection& InInput, int32 InSettingsCrc32) const;
 
 	FPCGDataCollection Input;
 	TObjectPtr<UPCGSettings> Settings;
 	FPCGDataCollection Output;
+	int32 SettingsCrc32;
 };
 
 // TODO: investigate if we need a more evolved data structure here
diff --git a/Engine/Plugins/Experimental/PCG/Source/PCG/Private/Graph/PCGGraphExecutor.cpp b/Engine/Plugins/Experimental/PCG/Source/PCG/Private/Graph/PCGGraphExecutor.cpp
--- a/Engine/Plugins/Experimental/PCG/Source/PCG/Private/Graph/PCGGraphExecutor.cpp
+++ b/Engine/Plugins/Experimental/PCG/Source/PCG/Private/Graph/PCGGraphExecutor.cpp
@@ -8,6 +8,8 @@
 #include "PCGComponent.h"
 #include "PCGInputOutputSettings.h"
 
+#include "Async/Async.h"
+
 // World partition support for in-editor workflows needs these includes
 #if WITH_EDITOR
 #include "WorldPartition/WorldPartition.h"
@@ -19,6 +21,16 @@ static TAutoConsoleVariable<int32> CVarMaxNumTasks(
 	4096,
 	TEXT("Maximum number of concurrent tasks for PCG processing"));
 
+static TAutoConsoleVariable<float> CVarTimePerFrame(
+	TEXT("pcg.FrameTime"),
+	1000.0f / 60.0f,
+	TEXT("Allocated time in ms per frame"));
+
+static TAutoConsoleVariable<bool> CVarGraphMultithreading(
+	TEXT("pcg.GraphMultithreading"),
+	false,
+	TEXT("Controls whether the graph can dispatch multiple tasks at the same time"));
+
 FPCGGraphExecutor::FPCGGraphExecutor(UObject* InOwner)
 	: GraphCompiler(MakeUnique<FPCGGraphCompiler>())
 	, GraphCache(InOwner)
@@ -126,7 +138,7 @@ bool FPCGGraphExecutor::GetOutputData(FPCGTaskId TaskId, FPCGDataCollection& Out
 
 void FPCGGraphExecutor::Execute()
 {
-	// TODO: change this so we can have time-slicing, priority and thread affinities
+	TRACE_CPUPROFILER_EVENT_SCOPE(FPCGGraphExecutor::Execute);
 
 	// Process any newly scheduled graphs to execute
 	ScheduleLock.Lock();
@@ -157,89 +169,242 @@ void FPCGGraphExecutor::Execute()
 		Tasks.Reset();
 	}
 
-	// TODO: change this when we support time-slicing
-	if(ReadyTasks.Num() > 0 || ActiveTasks.Num() > 0)
+	// TODO: change this if we support tasks that are not framebound
+	bool bAnyTaskEnded = false;
+
+	const double StartTime = FPlatformTime::Seconds();
+	const int32 MaxNumThreads = FMath::Max(0, FMath::Min(FPlatformMisc::NumberOfCoresIncludingHyperthreads() - 2, CVarMaxNumTasks.GetValueOnAnyThread() - 1));
+	const bool bAllowMultiDispatch = CVarGraphMultithreading.GetValueOnAnyThread();
+
+	while(ReadyTasks.Num() > 0 || ActiveTasks.Num() > 0)
 	{
 		// First: if we have free resources, move ready tasks to the active tasks
-		bool bCanStartNewTasks = true; // TODO implement this
+		bool bMainThreadAvailable = (ActiveTasks.Num() == 0 || !ActiveTasks.Last().Element->CanExecuteOnlyOnMainThread(ActiveTasks.Last().Context->GetInputSettings<UPCGSettings>()));
+		int32 NumAvailableThreads = FMath::Max(0, MaxNumThreads - CurrentlyUsedThreads);
+
+		const bool bMainThreadWasAvailable = bMainThreadAvailable;
+		const int32 TasksToLaunchIndex = ActiveTasks.Num();
 
-		if (bCanStartNewTasks && ReadyTasks.Num() > 0)
+		bool bSomeTaskEndedInCurrentLoop = false;
+
+		if (bMainThreadAvailable || NumAvailableThreads > 0)
 		{
-			// Sort by priority
+			TRACE_CPUPROFILER_EVENT_SCOPE(FPCGGraphExecutor::Execute::PrepareTasks);
+			// Sort tasks by priority (highest priority should be at the end)
 			// TODO
 
-			// Create active tasks according to resources & priority
-			while (ReadyTasks.Num() > 0 && bCanStartNewTasks)
+			for(int32 ReadyTaskIndex = ReadyTasks.Num() - 1; ReadyTaskIndex >= 0; --ReadyTaskIndex)
 			{
-				// TODO: peek at task properties to determine if we could start it (esp. thread affinity)
-				FPCGGraphTask Task = ReadyTasks.Pop();
+				FPCGGraphTask& Task = ReadyTasks[ReadyTaskIndex];
 
 				// Build input
 				FPCGDataCollection TaskInput;
 				BuildTaskInput(Task, TaskInput);
 
-				FPCGElementPtr Element = Task.Element;
-
-				// If the task didn't have a fixed element, acquire one from the node/settings
-				if (!Element)
+				// Initialize the element if needed (required to know whether it will run on the main thread or not)
+				if (!Task.Element)
 				{
 					// Get appropriate settings
+					check(Task.Node);
 					const UPCGSettings* Settings = TaskInput.GetSettings(Task.Node->DefaultSettings);
 
-					if (!Settings)
+					if (Settings)
 					{
-						// Can't execute; abort
-						continue;
+						Task.Element = Settings->GetElement();
 					}
+				}
 
-					// Create element
-					Element = Settings->GetElement();
+				// At this point, if the task doesn't have an element, we will never be able to execute it, so we can drop it.
+				if (!Task.Element)
+				{
+					ReadyTasks.RemoveAtSwap(ReadyTaskIndex);
+					continue;
 				}
 
-				if (!Element)
+				// If a task is cacheable and has been cached, then we don't need to create an active task for it unless
+				// there is an execution mode that would prevent us from doing so.
+				const UPCGSettings* TaskSettings = PCGContextHelpers::GetInputSettings<UPCGSettings>(Task.Node, TaskInput);
+				FPCGDataCollection CachedOutput;
+				const bool bResultAlreadyInCache = Task.Element->IsCacheable(TaskSettings) && GraphCache.GetFromCache(Task.Element.Get(), TaskInput, TaskSettings, CachedOutput);
+#if WITH_EDITOR
+				const bool bNeedsToCreateActiveTask = !bResultAlreadyInCache || TaskSettings->ExecutionMode == EPCGSettingsExecutionMode::Debug || TaskSettings->ExecutionMode == EPCGSettingsExecutionMode::Isolated;
+#else
+				const bool bNeedsToCreateActiveTask = !bResultAlreadyInCache;
+#endif
+
+				if (!bNeedsToCreateActiveTask)
 				{
-					// Can't execute; abort
+					// Fast-forward cached result to stored results
+					StoreResults(Task.NodeId, CachedOutput);
+					ReadyTasks.RemoveAtSwap(ReadyTaskIndex);
+					bSomeTaskEndedInCurrentLoop = true;
 					continue;
 				}
-				
-				FPCGContext* Context = Element->Initialize(TaskInput, Task.SourceComponent, Task.Node);
-				Context->TaskId = Task.NodeId;
-				Context->Cache = &GraphCache;
-
-				FPCGGraphActiveTask& ActiveTask = ActiveTasks.Emplace_GetRef();
-				ActiveTask.Element = Element;
-				ActiveTask.Context = TUniquePtr<FPCGContext>(Context);
-				ActiveTask.NodeId = Task.NodeId;
-				
-				// TODO update bCanStartNewTasks
+
+				// Validate that we can start this task now
+				const bool bIsMainThreadTask = Task.Element->CanExecuteOnlyOnMainThread(TaskSettings);
+
+				if (!bIsMainThreadTask || bMainThreadAvailable)
+				{
+					const bool bInsertLast = (bIsMainThreadTask || bMainThreadAvailable || ActiveTasks.Num() == 0);
+					FPCGGraphActiveTask& ActiveTask = (bInsertLast ? ActiveTasks.Emplace_GetRef() : ActiveTasks.EmplaceAt_GetRef(ActiveTasks.Num() - 1));
+
+					ActiveTask.Element = Task.Element;
+					ActiveTask.NodeId = Task.NodeId;
+
+					FPCGContext* Context = Task.Element->Initialize(TaskInput, Task.SourceComponent, Task.Node);
+					Context->TaskId = Task.NodeId;
+					Context->Cache = &GraphCache;
+					ActiveTask.Context = TUniquePtr<FPCGContext>(Context);
+
+#if WITH_EDITOR
+					if (bResultAlreadyInCache)
+					{
+						ActiveTask.bIsBypassed = true;
+						Context->OutputData = CachedOutput;
+					}
+#endif
+
+					if (bIsMainThreadTask || NumAvailableThreads == 0)
+					{
+						bMainThreadAvailable = false;
+					}
+					else
+					{
+						--NumAvailableThreads;
+					}
+
+					ReadyTasks.RemoveAtSwap(ReadyTaskIndex);
+
+					if (!bAllowMultiDispatch || (!bMainThreadAvailable && NumAvailableThreads == 0))
+					{
+						break;
+					}
+				}
 			}
 		}
 
-		// Perform execution
-		// TODO: this is a naive single-threaded implementation, should be replaced by
-		// the task graph or futures w/proper thread assingment
-		for (int ExecutionIndex = 0; ExecutionIndex < ActiveTasks.Num(); ++ExecutionIndex)
+		check(NumAvailableThreads >= 0);
+
+		const int32 NumTasksToLaunch = ActiveTasks.Num() - TasksToLaunchIndex;
+
+		// Re-launch time-sliced tasks & launch new tasks
+		// TODO: currently we don't have any time-slicing so just launch tasks
+		const int32 LastLaunchIndex = ActiveTasks.Num() - (bMainThreadWasAvailable ? 0 : 1);
+		
+		// Assign resources
+		const int32 NumAdditionalThreads = ((NumTasksToLaunch > 0) ? (NumAvailableThreads / NumTasksToLaunch) : 0);
+		check(NumAdditionalThreads >= 0);
+
+		for (int32 ExecutionIndex = TasksToLaunchIndex; ExecutionIndex < LastLaunchIndex; ++ExecutionIndex)
 		{
 			FPCGGraphActiveTask& ActiveTask = ActiveTasks[ExecutionIndex];
+			ActiveTask.Context->NumAvailableTasks = 1 + NumAdditionalThreads;
+			CurrentlyUsedThreads += ActiveTask.Context->NumAvailableTasks;
+		}
 
-			// Reassign resources to the context
-			// TODO: change this when we support MT in the graph executor
-			ActiveTask.Context->NumAvailableTasks = FMath::Max(1, FMath::Min(FPlatformMisc::NumberOfCoresIncludingHyperthreads() - 2, CVarMaxNumTasks.GetValueOnAnyThread()));
+		// Dispatch async tasks
+		TMap<int32, TFuture<bool>> Futures;
+		
+		for(int32 ExecutionIndex = 0; ExecutionIndex < ActiveTasks.Num() - 1; ++ExecutionIndex)
+		{
+			FPCGGraphActiveTask& ActiveTask = ActiveTasks[ExecutionIndex];
+#if WITH_EDITOR
+			if(!ActiveTask.bIsBypassed && !ActiveTask.Context->bIsPaused)
+#else
+			if(!ActiveTask.Context->bIsPaused)
+#endif
+			{
+				Futures.Emplace(ExecutionIndex, Async(EAsyncExecution::ThreadPool, [&ActiveTask]()
+				{
+					return ActiveTask.Element->Execute(ActiveTask.Context.Get());
+				}));
+			}
+		}
+
+		auto PostTaskExecute = [this, &bSomeTaskEndedInCurrentLoop](int32 TaskIndex)
+		{
+			FPCGGraphActiveTask& ActiveTask = ActiveTasks[TaskIndex];
+
+#if WITH_EDITOR
+			if (!ActiveTask.bIsBypassed)
+#endif
+			{
+				// Store result in cache as needed - done here because it needs to be done on the main thread
+				const UPCGSettings* ActiveTaskSettings = ActiveTask.Context->GetInputSettings<UPCGSettings>();
+				if (ActiveTaskSettings && ActiveTask.Element->IsCacheable(ActiveTaskSettings))
+				{
+					GraphCache.StoreInCache(ActiveTask.Element.Get(), ActiveTask.Context->InputData, ActiveTaskSettings, ActiveTask.Context->OutputData);
+				}
+			}
+
+			CurrentlyUsedThreads -= ActiveTask.Context->NumAvailableTasks;
+
+			// Store output in data map
+			StoreResults(ActiveTask.NodeId, ActiveTask.Context->OutputData);
+
+#if WITH_EDITOR
+			// Execute debug display code as needed - done here because it needs to be done on the main thread
+			ActiveTask.Element->DebugDisplay(ActiveTask.Context.Get());
+#endif
+
+			// Book-keeping
+			bSomeTaskEndedInCurrentLoop = true;
+
+			// Remove current active task from list
+			ActiveTasks.RemoveAtSwap(TaskIndex);
+		};
 
-			if (!ActiveTask.Context->bIsPaused && ActiveTask.Element->Execute(ActiveTask.Context.Get()))
+		{
+			TRACE_CPUPROFILER_EVENT_SCOPE(FPCGGraphExecutor::Execute::ExecuteTasks);
+			// Execute main thread task
+			if (!ActiveTasks.IsEmpty())
 			{
-				// Store output data in map
-				StoreResults(ActiveTask.NodeId, ActiveTask.Context->OutputData);
+				FPCGGraphActiveTask& MainThreadTask = ActiveTasks.Last();
+#if WITH_EDITOR
+				if (MainThreadTask.bIsBypassed || (!MainThreadTask.Context->bIsPaused && MainThreadTask.Element->Execute(MainThreadTask.Context.Get())))
+#else
+				if (!MainThreadTask.Context->bIsPaused && MainThreadTask.Element->Execute(MainThreadTask.Context.Get()))
+#endif
+				{
+					PostTaskExecute(ActiveTasks.Num() - 1);
+				}
+			}
 
-				// Move successor nodes to the ready queue
-				QueueReadyTasks(ActiveTask.NodeId);
+			for (int32 ExecutionIndex = ActiveTasks.Num() - 1; ExecutionIndex >= 0; --ExecutionIndex)
+			{
+				bool bTaskDone = true;
+				// Wait on the future if any
+				if (TFuture<bool>* Future = Futures.Find(ExecutionIndex))
+				{
+					Future->Wait();
+					bTaskDone = Future->Get();
+				}
 
-				// Remove current active task from list
-				ActiveTasks.RemoveAt(ExecutionIndex);
-				--ExecutionIndex;
+				if (bTaskDone)
+				{
+					PostTaskExecute(ExecutionIndex);
+				}
 			}
 		}
 
+		check(CurrentlyUsedThreads == 0);
+
+		if (bSomeTaskEndedInCurrentLoop)
+		{
+			QueueReadyTasks();
+			bAnyTaskEnded = true;
+		}
+
+		const double EndLoopTime = FPlatformTime::Seconds();
+		if (static_cast<float>((EndLoopTime - StartTime) * 1000) > CVarTimePerFrame.GetValueOnAnyThread())
+		{
+			break;
+		}
+	}
+
+	if (bAnyTaskEnded)
+	{
 		// Nothing left to do; we'll release everything here.
 		// TODO: this is fine and will make sure any intermediate data is properly
 		// garbage collected, however, this goes a bit against our goals if we want to
@@ -259,6 +424,7 @@ void FPCGGraphExecutor::Execute()
 
 void FPCGGraphExecutor::QueueReadyTasks(FPCGTaskId FinishedTaskHint)
 {
+	TRACE_CPUPROFILER_EVENT_SCOPE(FPCGGraphExecutor::QueueReadyTasks);
 	for (int TaskIndex = Tasks.Num() - 1; TaskIndex >= 0; --TaskIndex)
 	{
 		bool bAllPrerequisitesMet = true;
@@ -277,6 +443,7 @@ void FPCGGraphExecutor::QueueReadyTasks(FPCGTaskId FinishedTaskHint)
 
 void FPCGGraphExecutor::BuildTaskInput(const FPCGGraphTask& Task, FPCGDataCollection& TaskInput)
 {
+	TRACE_CPUPROFILER_EVENT_SCOPE(FPCGGraphExecutor::BuildTaskInput);
 	for (const FPCGGraphTaskInput& Input : Task.Inputs)
 	{
 		check(OutputData.Contains(Input.TaskId));
@@ -306,6 +473,8 @@ void FPCGGraphExecutor::BuildTaskInput(const FPCGGraphTask& Task, FPCGDataCollec
 
 void FPCGGraphExecutor::StoreResults(FPCGTaskId InTaskId, const FPCGDataCollection& InTaskOutput)
 {
+	TRACE_CPUPROFILER_EVENT_SCOPE(FPCGGraphExecutor::StoreResults);
+
 	// Store output in map
 	OutputData.Add(InTaskId, InTaskOutput);
 
@@ -450,5 +619,6 @@ FPCGGenericElement::FPCGGenericElement(TFunction<bool()> InOperation)
 
 bool FPCGGenericElement::ExecuteInternal(FPCGContext* Context) const
 {
+	TRACE_CPUPROFILER_EVENT_SCOPE(FPCGGenericElement::Execute);
 	return Operation();
 }
diff --git a/Engine/Plugins/Experimental/PCG/Source/PCG/Private/Graph/PCGGraphExecutor.h b/Engine/Plugins/Experimental/PCG/Source/PCG/Private/Graph/PCGGraphExecutor.h
--- a/Engine/Plugins/Experimental/PCG/Source/PCG/Private/Graph/PCGGraphExecutor.h
+++ b/Engine/Plugins/Experimental/PCG/Source/PCG/Private/Graph/PCGGraphExecutor.h
@@ -52,6 +52,9 @@ struct FPCGGraphActiveTask
 	FPCGElementPtr Element;
 	TUniquePtr<FPCGContext> Context;
 	FPCGTaskId NodeId = InvalidTaskId;
+#if WITH_EDITOR
+	bool bIsBypassed = false;
+#endif
 };
 
 class FPCGGraphExecutor
@@ -122,6 +125,9 @@ private:
 	/** Monotonically increasing id. Should be reset once all tasks are executed, should be protected by the ScheduleLock */
 	FPCGTaskId NextTaskId = 0;
 
+	/** Runtime information */
+	int32 CurrentlyUsedThreads = 0;
+
 #if WITH_EDITOR
 	FCriticalSection ActorsListLock;
 	TSet<AActor*> ActorsToSave;
@@ -131,9 +137,11 @@ private:
 
 class FPCGFetchInputElement : public FSimplePCGElement
 {
+public:
+	virtual bool IsCacheable(const UPCGSettings* InSettings) const override { return false; }
+
 protected:
 	virtual bool ExecuteInternal(FPCGContext* Context) const override;
-	virtual bool IsCacheable(const UPCGSettings* InSettings) const override { return false; }
 	virtual bool IsPassthrough() const override { return true; }
 };
 
@@ -141,13 +149,15 @@ class FPCGGenericElement : public FSimplePCGElement
 {
 public:
 	FPCGGenericElement(TFunction<bool()> InOperation);
+	virtual bool IsCacheable(const UPCGSettings* InSettings) const override { return false; }
+	virtual bool CanExecuteOnlyOnMainThread(const UPCGSettings* InSettings) const override { return true; }
 
 protected:
 	// Important note: generic elements must always be run on the main thread
 	// as most of these will impact the editor in some way (loading, unloading, saving)
 	virtual bool ExecuteInternal(FPCGContext* Context) const override;
 	virtual bool IsCancellable() const { return false; }
-	virtual bool IsCacheable(const UPCGSettings* InSettings) const override { return false; }
+
 #if WITH_EDITOR
 	virtual bool ShouldLog() const { return false; }
 #endif
diff --git a/Engine/Plugins/Experimental/PCG/Source/PCG/Private/PCGElement.cpp b/Engine/Plugins/Experimental/PCG/Source/PCG/Private/PCGElement.cpp
--- a/Engine/Plugins/Experimental/PCG/Source/PCG/Private/PCGElement.cpp
+++ b/Engine/Plugins/Experimental/PCG/Source/PCG/Private/PCGElement.cpp
@@ -10,6 +10,7 @@
 
 bool IPCGElement::Execute(FPCGContext* Context) const
 {
+	TRACE_CPUPROFILER_EVENT_SCOPE(IPCGElement::Execute);
 	// Early out to stop execution
 	if (Context->InputData.bCancelExecution)
 	{
@@ -66,40 +67,27 @@ bool IPCGElement::Execute(FPCGContext* Context) const
 		}
 #endif
 
-		if (IsCacheable(Settings) && Context->Cache && Context->Cache->GetFromCache(this, Context->InputData, Settings, Context->OutputData))
-		{
-			PCGE_LOG(Verbose, "Used cached results");
-			bDone = true;
-		}
-		else
-		{
 #if WITH_EDITOR
-			const double StartTime = FPlatformTime::Seconds();
+		const double StartTime = FPlatformTime::Seconds();
 #endif
 
-			bDone = ExecuteInternal(Context);
+		bDone = ExecuteInternal(Context);
 
 #if WITH_EDITOR
-			const double EndTime = FPlatformTime::Seconds();
-			Context->ElapsedTime += (EndTime - StartTime);
-			Context->ExecutionCount++;
+		const double EndTime = FPlatformTime::Seconds();
+		Context->ElapsedTime += (EndTime - StartTime);
+		Context->ExecutionCount++;
 #endif
 
-			if (bDone)
-			{
-				CleanupAndValidateOutputLabels(Context);
-
-				if (IsCacheable(Settings) && Context->Cache)
-				{
-					Context->Cache->StoreInCache(this, Context->InputData, Settings, Context->OutputData);
-				}
+		if (bDone)
+		{
+			CleanupAndValidateOutputLabels(Context);
 
 #if WITH_EDITOR
-				PCGE_LOG(Log, "Executed in (%f)s and (%d) call(s)", Context->ElapsedTime, Context->ExecutionCount);
+			PCGE_LOG(Log, "Executed in (%f)s and (%d) call(s)", Context->ElapsedTime, Context->ExecutionCount);
 #else
-				PCGE_LOG(Log, "Executed");
+			PCGE_LOG(Log, "Executed");
 #endif
-			}
 		}
 
 		/** TODO - Placeholder feature */
@@ -131,32 +119,38 @@ bool IPCGElement::Execute(FPCGContext* Context) const
 				Context->InputData = ElementInputs;
 				Context->OutputData = ElementOutputs;
 			}
+		}
+#endif
 
-			if (Settings->ExecutionMode == EPCGSettingsExecutionMode::Debug || Settings->ExecutionMode == EPCGSettingsExecutionMode::Isolated)
-			{
-				FPCGDataCollection ElementInputs = Context->InputData;
-				FPCGDataCollection ElementOutputs = Context->OutputData;
+		return bDone;
+	}
+}
 
-				Context->InputData = ElementOutputs;
-				Context->OutputData = FPCGDataCollection();
+#if WITH_EDITOR
+void IPCGElement::DebugDisplay(FPCGContext* Context) const
+{
+	const UPCGSettings* Settings = Context->GetInputSettings<UPCGSettings>();
+	if (Settings && (Settings->ExecutionMode == EPCGSettingsExecutionMode::Debug || Settings->ExecutionMode == EPCGSettingsExecutionMode::Isolated))
+	{
+		FPCGDataCollection ElementInputs = Context->InputData;
+		FPCGDataCollection ElementOutputs = Context->OutputData;
 
-				PCGDebugElement::ExecuteDebugDisplay(Context);
+		Context->InputData = ElementOutputs;
+		Context->OutputData = FPCGDataCollection();
 
-				Context->InputData = ElementInputs;
-				Context->OutputData = ElementOutputs;
+		PCGDebugElement::ExecuteDebugDisplay(Context);
 
-				// Null out the output if this node is executed in isolation
-				if (Settings->ExecutionMode == EPCGSettingsExecutionMode::Isolated)
-				{
-					Context->OutputData.bCancelExecution = true;
-				}
-			}
-		}
-#endif
+		Context->InputData = ElementInputs;
+		Context->OutputData = ElementOutputs;
 
-		return bDone;
+		// Null out the output if this node is executed in isolation
+		if (Settings->ExecutionMode == EPCGSettingsExecutionMode::Isolated)
+		{
+			Context->OutputData.bCancelExecution = true;
+		}
 	}
 }
+#endif // WITH_EDITOR
 
 void IPCGElement::CleanupAndValidateOutputLabels(FPCGContext* Context) const
 {
diff --git a/Engine/Plugins/Experimental/PCG/Source/PCG/Private/PCGSettings.cpp b/Engine/Plugins/Experimental/PCG/Source/PCG/Private/PCGSettings.cpp
--- a/Engine/Plugins/Experimental/PCG/Source/PCG/Private/PCGSettings.cpp
+++ b/Engine/Plugins/Experimental/PCG/Source/PCG/Private/PCGSettings.cpp
@@ -37,6 +37,12 @@ bool UPCGSettings::operator==(const UPCGSettings& Other) const
 	}
 }
 
+uint32 UPCGSettings::GetCrc32() const
+{
+	FPCGSettingsObjectCrc32 Ar;
+	return Ar.Crc32(const_cast<UPCGSettings*>(this));
+}
+
 bool UPCGSettings::HasInLabel(const FName& Label) const
 {
 	return (Label == NAME_None && HasDefaultInLabel()) || InLabels().Contains(Label);
diff --git a/Engine/Plugins/Experimental/PCG/Source/PCG/Public/Elements/PCGDebugElement.h b/Engine/Plugins/Experimental/PCG/Source/PCG/Public/Elements/PCGDebugElement.h
--- a/Engine/Plugins/Experimental/PCG/Source/PCG/Public/Elements/PCGDebugElement.h
+++ b/Engine/Plugins/Experimental/PCG/Source/PCG/Public/Elements/PCGDebugElement.h
@@ -30,7 +30,9 @@ protected:
 
 class FPCGDebugElement : public FSimplePCGElement
 {
-protected:
-	virtual bool ExecuteInternal(FPCGContext* Context) const override;
+public:
 	virtual bool IsCacheable(const UPCGSettings* InSettings) const override { return false; }
+
+protected:
+	virtual bool ExecuteInternal(FPCGContext* Context) const override;	
 };
diff --git a/Engine/Plugins/Experimental/PCG/Source/PCG/Public/Elements/PCGExecuteBlueprint.h b/Engine/Plugins/Experimental/PCG/Source/PCG/Public/Elements/PCGExecuteBlueprint.h
--- a/Engine/Plugins/Experimental/PCG/Source/PCG/Public/Elements/PCGExecuteBlueprint.h
+++ b/Engine/Plugins/Experimental/PCG/Source/PCG/Public/Elements/PCGExecuteBlueprint.h
@@ -161,6 +161,9 @@ protected:
 	UPROPERTY(BlueprintReadWrite, EditAnywhere, Category = Settings)
 	bool bCreatesArtifacts = false;
 
+	UPROPERTY(BlueprintReadWrite, EditAnywhere, Category = Settings)
+	bool bCanBeMultithreaded = true;
+
 protected:
 #if WITH_EDITOR
 	void OnBlueprintChanged(UBlueprint* InBlueprint);
@@ -181,8 +184,11 @@ struct FPCGBlueprintExecutionContext : public FPCGContext
 
 class FPCGExecuteBlueprintElement : public IPCGElement
 {
+public:
+	virtual bool CanExecuteOnlyOnMainThread(const UPCGSettings* InSettings) const override;
+	virtual bool IsCacheable(const UPCGSettings* InSettings) const override;
+
 protected:
 	virtual bool ExecuteInternal(FPCGContext* Context) const override;
-	virtual FPCGContext* Initialize(const FPCGDataCollection& InputData, UPCGComponent* SourceComponent, const UPCGNode* Node) override;
-	virtual bool IsCacheable(const UPCGSettings* InSettings) const override;
+	virtual FPCGContext* Initialize(const FPCGDataCollection& InputData, UPCGComponent* SourceComponent, const UPCGNode* Node) override;	
 };
diff --git a/Engine/Plugins/Experimental/PCG/Source/PCG/Public/Elements/PCGMetadataRenameElement.h b/Engine/Plugins/Experimental/PCG/Source/PCG/Public/Elements/PCGMetadataRenameElement.h
--- a/Engine/Plugins/Experimental/PCG/Source/PCG/Public/Elements/PCGMetadataRenameElement.h
+++ b/Engine/Plugins/Experimental/PCG/Source/PCG/Public/Elements/PCGMetadataRenameElement.h
@@ -33,6 +33,9 @@ public:
 
 class FPCGMetadataRenameElement : public FSimplePCGElement
 {
+public:
+	virtual bool CanExecuteOnlyOnMainThread(const UPCGSettings* InSettings) const override { return true; }
+
 protected:
 	virtual bool ExecuteInternal(FPCGContext* Context) const override;
 };
\ No newline at end of file
diff --git a/Engine/Plugins/Experimental/PCG/Source/PCG/Public/Elements/PCGSpawnActor.h b/Engine/Plugins/Experimental/PCG/Source/PCG/Public/Elements/PCGSpawnActor.h
--- a/Engine/Plugins/Experimental/PCG/Source/PCG/Public/Elements/PCGSpawnActor.h
+++ b/Engine/Plugins/Experimental/PCG/Source/PCG/Public/Elements/PCGSpawnActor.h
@@ -66,8 +66,11 @@ public:
 
 class FPCGSpawnActorElement : public FSimplePCGElement
 {
-protected:
-	virtual bool ExecuteInternal(FPCGContext* Context) const override;
+public:
+	virtual bool CanExecuteOnlyOnMainThread(const UPCGSettings* InSettings) const override { return true; }
 	virtual bool IsCacheable(const UPCGSettings* InSettings) const override { return false; }
 	virtual bool IsPassthrough() const override { return true; }
+
+protected:
+	virtual bool ExecuteInternal(FPCGContext* Context) const override;
 };
diff --git a/Engine/Plugins/Experimental/PCG/Source/PCG/Public/Elements/PCGStaticMeshSpawner.h b/Engine/Plugins/Experimental/PCG/Source/PCG/Public/Elements/PCGStaticMeshSpawner.h
--- a/Engine/Plugins/Experimental/PCG/Source/PCG/Public/Elements/PCGStaticMeshSpawner.h
+++ b/Engine/Plugins/Experimental/PCG/Source/PCG/Public/Elements/PCGStaticMeshSpawner.h
@@ -44,7 +44,10 @@ public:
 
 class FPCGStaticMeshSpawnerElement : public FSimplePCGElement
 {
-protected:
-	virtual bool ExecuteInternal(FPCGContext* Context) const override;
+public:
+	virtual bool CanExecuteOnlyOnMainThread(const UPCGSettings* InSettings) const { return true; }
 	virtual bool IsCacheable(const UPCGSettings* InSettings) const override { return false; }
+
+protected:
+	virtual bool ExecuteInternal(FPCGContext* Context) const override;	
 };
diff --git a/Engine/Plugins/Experimental/PCG/Source/PCG/Public/PCGContext.h b/Engine/Plugins/Experimental/PCG/Source/PCG/Public/PCGContext.h
--- a/Engine/Plugins/Experimental/PCG/Source/PCG/Public/PCGContext.h
+++ b/Engine/Plugins/Experimental/PCG/Source/PCG/Public/PCGContext.h
@@ -12,6 +12,22 @@
 class UPCGComponent;
 struct FPCGGraphCache;
 
+namespace PCGContextHelpers
+{
+	template<typename SettingsType>
+	const SettingsType* GetInputSettings(const UPCGNode* Node, const FPCGDataCollection& InputData)
+	{
+		if (Node && Node->DefaultSettings)
+		{
+			return Cast<SettingsType>(InputData.GetSettings(Node->DefaultSettings));
+		}
+		else
+		{
+			return InputData.GetSettings<SettingsType>();
+		}
+	}
+}
+
 USTRUCT(BlueprintType)
 struct PCG_API FPCGContext
 {
@@ -39,14 +55,7 @@ struct PCG_API FPCGContext
 	template<typename SettingsType>
 	const SettingsType* GetInputSettings() const
 	{
-		if (Node && Node->DefaultSettings)
-		{
-			return Cast<SettingsType>(InputData.GetSettings(Node->DefaultSettings));
-		}
-		else
-		{
-			return InputData.GetSettings<SettingsType>();
-		}
+		return PCGContextHelpers::GetInputSettings<SettingsType>(Node, InputData);
 	}
 
 	FString GetTaskName() const;
diff --git a/Engine/Plugins/Experimental/PCG/Source/PCG/Public/PCGElement.h b/Engine/Plugins/Experimental/PCG/Source/PCG/Public/PCGElement.h
--- a/Engine/Plugins/Experimental/PCG/Source/PCG/Public/PCGElement.h
+++ b/Engine/Plugins/Experimental/PCG/Source/PCG/Public/PCGElement.h
@@ -36,12 +36,19 @@ public:
 	virtual ~IPCGElement() = default;
 	virtual FPCGContext* Initialize(const FPCGDataCollection& InputData, UPCGComponent* SourceComponent, const UPCGNode* Node) = 0;
 
+	virtual bool CanExecuteOnlyOnMainThread(const UPCGSettings* InSettings) const { return false; }
+	virtual bool IsCacheable(const UPCGSettings* InSettings) const { return true; }
+
 	bool Execute(FPCGContext* Context) const;
 
+	/** Note: the following methods must be called from the main thread */
+#if WITH_EDITOR
+	void DebugDisplay(FPCGContext* Context) const;
+#endif
+
 protected:
 	virtual bool ExecuteInternal(FPCGContext* Context) const = 0;
 	virtual bool IsCancellable() const { return true; }
-	virtual bool IsCacheable(const UPCGSettings* InSettings) const { return true; }
 	virtual bool IsPassthrough() const { return false; }
 #if WITH_EDITOR
 	virtual bool ShouldLog() const { return true; }
diff --git a/Engine/Plugins/Experimental/PCG/Source/PCG/Public/PCGSettings.h b/Engine/Plugins/Experimental/PCG/Source/PCG/Public/PCGSettings.h
--- a/Engine/Plugins/Experimental/PCG/Source/PCG/Public/PCGSettings.h
+++ b/Engine/Plugins/Experimental/PCG/Source/PCG/Public/PCGSettings.h
@@ -62,6 +62,7 @@ public:
 	virtual bool HasDefaultOutLabel() const { return true; }
 	
 	bool operator==(const UPCGSettings& Other) const;
+	uint32 GetCrc32() const;
 
 #if WITH_EDITOR
 	virtual FName GetDefaultNodeName() const { return NAME_None; }
@@ -129,8 +130,10 @@ protected:
 
 class PCG_API FPCGTrivialElement : public FSimplePCGElement
 {
+public:
+	virtual bool IsCacheable(const UPCGSettings* InSettings) const override { return true; }
+
 protected:
 	virtual bool ExecuteInternal(FPCGContext* Context) const override;
-	virtual bool IsCacheable(const UPCGSettings* InSettings) const override { return false; }
 	virtual bool IsPassthrough() const override { return true; }
 };
