Fixed issue with older graphs because the node post load is done after the graph post load

#rnx

diff --git a/Engine/Plugins/Experimental/PCG/Source/PCG/Private/PCGGraph.cpp b/Engine/Plugins/Experimental/PCG/Source/PCG/Private/PCGGraph.cpp
--- a/Engine/Plugins/Experimental/PCG/Source/PCG/Private/PCGGraph.cpp
+++ b/Engine/Plugins/Experimental/PCG/Source/PCG/Private/PCGGraph.cpp
@@ -41,6 +41,8 @@ void UPCGGraph::PostLoad()
 
 #if WITH_EDITOR
 	// Deprecation
+	InputNode->ApplyDeprecation();
+
 	if (!Cast<UPCGGraphInputOutputSettings>(InputNode->DefaultSettings))
 	{
 		InputNode->DefaultSettings = NewObject<UPCGGraphInputOutputSettings>(this, TEXT("DefaultInputNodeSettings"));
@@ -57,6 +59,8 @@ void UPCGGraph::PostLoad()
 		}
 	}
 
+	OutputNode->ApplyDeprecation();
+
 	if (!Cast<UPCGGraphInputOutputSettings>(OutputNode->DefaultSettings))
 	{
 		OutputNode->DefaultSettings = NewObject<UPCGGraphInputOutputSettings>(this, TEXT("DefaultOutputNodeSettings"));
diff --git a/Engine/Plugins/Experimental/PCG/Source/PCG/Private/PCGNode.cpp b/Engine/Plugins/Experimental/PCG/Source/PCG/Private/PCGNode.cpp
--- a/Engine/Plugins/Experimental/PCG/Source/PCG/Private/PCGNode.cpp
+++ b/Engine/Plugins/Experimental/PCG/Source/PCG/Private/PCGNode.cpp
@@ -20,9 +20,14 @@ void UPCGNode::PostLoad()
 	{
 		DefaultSettings->OnSettingsChangedDelegate.AddUObject(this, &UPCGNode::OnSettingsChanged);
 	}
+
+	ApplyDeprecation();
 #endif
+}
 
-	// Deprecation
+#if WITH_EDITOR
+void UPCGNode::ApplyDeprecation()
+{
 	for (TObjectPtr<UPCGNode> OutboundNode : OutboundNodes_DEPRECATED)
 	{
 		if (!HasEdgeTo(OutboundNode))
@@ -32,6 +37,7 @@ void UPCGNode::PostLoad()
 	}
 	OutboundNodes_DEPRECATED.Reset();
 }
+#endif
 
 void UPCGNode::PostEditImport()
 {
diff --git a/Engine/Plugins/Experimental/PCG/Source/PCG/Public/PCGNode.h b/Engine/Plugins/Experimental/PCG/Source/PCG/Public/PCGNode.h
--- a/Engine/Plugins/Experimental/PCG/Source/PCG/Public/PCGNode.h
+++ b/Engine/Plugins/Experimental/PCG/Source/PCG/Public/PCGNode.h
@@ -33,6 +33,11 @@ public:
 	virtual void BeginDestroy() override;
 	/** ~End UObject interface */
 
+#if WITH_EDITOR
+	/** Used to be able to force deprecation when things need to be deprecated at the graph level */
+	void ApplyDeprecation();
+#endif
+
 	/** Returns the owning graph */
 	UFUNCTION(BlueprintCallable, Category = Node)
 	UPCGGraph* GetGraph() const;
