PCG Fixed crash on copy/paste

#rnx

diff --git a/Engine/Plugins/Experimental/PCG/Source/PCGEditor/Private/PCGEditorGraphNodeBase.cpp b/Engine/Plugins/Experimental/PCG/Source/PCGEditor/Private/PCGEditorGraphNodeBase.cpp
--- a/Engine/Plugins/Experimental/PCG/Source/PCGEditor/Private/PCGEditorGraphNodeBase.cpp
+++ b/Engine/Plugins/Experimental/PCG/Source/PCGEditor/Private/PCGEditorGraphNodeBase.cpp
@@ -155,7 +155,10 @@ void UPCGEditorGraphNodeBase::RebuildEdgesFromPins()
 {
 	check(PCGNode);
 
-	PCGNode->GetGraph()->DisableNotificationsForEditor();
+	if (PCGNode->GetGraph())
+	{
+		PCGNode->GetGraph()->DisableNotificationsForEditor();
+	}
 	
 	for (UEdGraphPin* Pin : Pins)
 	{
@@ -177,7 +180,10 @@ void UPCGEditorGraphNodeBase::RebuildEdgesFromPins()
 		}
 	}
 
-	PCGNode->GetGraph()->EnableNotificationsForEditor();
+	if (PCGNode->GetGraph())
+	{
+		PCGNode->GetGraph()->EnableNotificationsForEditor();
+	}
 }
 
 void UPCGEditorGraphNodeBase::OnNodeChanged(UPCGNode* InNode, bool bSettingsChanged)
@@ -282,13 +288,16 @@ void UPCGEditorGraphNodeBase::ReconstructNode()
 		}
 		else
 		{
-			PCGNode->GetGraph()->DisableNotificationsForEditor();
+			if (PCGNode->GetGraph())
+			{
+				PCGNode->GetGraph()->DisableNotificationsForEditor();
+			}
 
 			for (UEdGraphPin* UnmatchedPin : UnmatchedPins)
 			{
 				UnmatchedPin->BreakAllPinLinks();
 
-				if (PCGNode)
+				if (PCGNode->GetGraph())
 				{
 					if (InPinDirection == EEdGraphPinDirection::EGPD_Input)
 					{
@@ -303,7 +312,10 @@ void UPCGEditorGraphNodeBase::ReconstructNode()
 				RemovePin(UnmatchedPin);
 			}
 
-			PCGNode->GetGraph()->EnableNotificationsForEditor();
+			if (PCGNode->GetGraph())
+			{
+				PCGNode->GetGraph()->EnableNotificationsForEditor();
+			}
 
 			for (const FName& UnmatchedName : UnmatchedNames)
 			{
