Split PCG point extents into min/max

diff --git a/Engine/Plugins/Experimental/PCG/Source/PCG/Private/Data/PCGLandscapeData.cpp b/Engine/Plugins/Experimental/PCG/Source/PCG/Private/Data/PCGLandscapeData.cpp
--- a/Engine/Plugins/Experimental/PCG/Source/PCG/Private/Data/PCGLandscapeData.cpp
+++ b/Engine/Plugins/Experimental/PCG/Source/PCG/Private/Data/PCGLandscapeData.cpp
@@ -113,7 +113,7 @@ const UPCGPointData* UPCGLandscapeData::CreatePointData(FPCGContext* Context) co
 			OutPoint = FPCGPoint(FTransform(VertexLocation),
 				1.0f,
 				PCGHelpers::ComputeSeed(X, Y, (int)HeightAtVertex.GetValue()));
-			OutPoint.Extents = Transform.GetScale3D() / 2.0;
+			OutPoint.SetExtents(Transform.GetScale3D() / 2.0);
 
 			return true;
 		}
diff --git a/Engine/Plugins/Experimental/PCG/Source/PCG/Private/Data/PCGPointData.cpp b/Engine/Plugins/Experimental/PCG/Source/PCG/Private/Data/PCGPointData.cpp
--- a/Engine/Plugins/Experimental/PCG/Source/PCG/Private/Data/PCGPointData.cpp
+++ b/Engine/Plugins/Experimental/PCG/Source/PCG/Private/Data/PCGPointData.cpp
@@ -13,7 +13,8 @@ namespace PCGPointHelpers
 	bool GetDistanceRatios(const FPCGPoint& InPoint, const FVector& InPosition, FVector& OutRatios)
 	{
 		FVector LocalPosition = InPoint.Transform.InverseTransformPosition(InPosition);
-		LocalPosition /= InPoint.Extents;
+		LocalPosition -= (InPoint.BoundsMax + InPoint.BoundsMin) / 2;
+		LocalPosition /= InPoint.GetExtents();
 
 		// ]-2+s, 2-s] is the valid range of values
 		const FVector::FReal LowerBound = InPoint.Steepness - 2;
@@ -233,7 +234,8 @@ FPCGPoint UPCGPointData::TransformPoint(const FPCGPoint& InPoint) const
 	FRotator WeightedRotator(0);
 	FVector WeightedScale = FVector::Zero();
 	float WeightedDensity = 0;
-	FVector WeightedExtents = FVector::Zero();
+	FVector WeightedBoundsMin = FVector::Zero();
+	FVector WeightedBoundsMax = FVector::Zero();
 	FVector WeightedColor = FVector::Zero();
 	float WeightedSteepness = 0;
 
@@ -245,7 +247,8 @@ FPCGPoint UPCGPointData::TransformPoint(const FPCGPoint& InPoint) const
 		WeightedRotator += (SourcePoint.Transform.Rotator() * Weight); // TODO: This is wonky
 		WeightedScale += (SourcePoint.Transform.GetScale3D() * Weight);
 		WeightedDensity += PCGPointHelpers::ManhattanDensity(SourcePoint, PointPosition);
-		WeightedExtents += SourcePoint.Extents * Weight;
+		WeightedBoundsMin += SourcePoint.BoundsMin * Weight;
+		WeightedBoundsMax += SourcePoint.BoundsMax * Weight;
 		WeightedColor += SourcePoint.Color * Weight;
 		WeightedSteepness += SourcePoint.Steepness * Weight;
 	}
@@ -257,7 +260,8 @@ FPCGPoint UPCGPointData::TransformPoint(const FPCGPoint& InPoint) const
 	Point.Transform.NormalizeRotation();
 	Point.Transform.SetScale3D(Point.Transform.GetScale3D() * WeightedScale);	
 	Point.Density *= WeightedDensity;
-	Point.Extents *= WeightedExtents; // this assumes that the extents were 1 to begin with
+	Point.BoundsMin = WeightedBoundsMin;
+	Point.BoundsMax = WeightedBoundsMax;
 	Point.Color *= WeightedColor;
 	Point.Steepness *= WeightedSteepness;
 
@@ -300,7 +304,8 @@ bool UPCGPointData::GetPointAtPosition(const FVector& InPosition, FPCGPoint& Out
 	FQuat WeightedQuat = FQuat::Identity;
 	FVector WeightedScale = FVector::ZeroVector;
 	float WeightedDensity = 0;
-	FVector WeightedExtents = FVector::ZeroVector;
+	FVector WeightedBoundsMin = FVector::ZeroVector;
+	FVector WeightedBoundsMax = FVector::ZeroVector;
 	FVector WeightedColor = FVector::ZeroVector;
 	float WeightedSteepness = 0;
 
@@ -312,7 +317,8 @@ bool UPCGPointData::GetPointAtPosition(const FVector& InPosition, FPCGPoint& Out
 		WeightedQuat = PCGPointHelpers::AddQuatWithWeight(WeightedQuat, SourcePoint.Transform.GetRotation(), Weight);
 		WeightedScale += SourcePoint.Transform.GetScale3D() * Weight;
 		WeightedDensity += PCGPointHelpers::ManhattanDensity(SourcePoint, InPosition);
-		WeightedExtents += SourcePoint.Extents * Weight;
+		WeightedBoundsMin += SourcePoint.BoundsMin * Weight;
+		WeightedBoundsMax += SourcePoint.BoundsMax * Weight;
 		WeightedColor += SourcePoint.Color * Weight;
 		WeightedSteepness += SourcePoint.Steepness * Weight;
 	}
@@ -324,7 +330,8 @@ bool UPCGPointData::GetPointAtPosition(const FVector& InPosition, FPCGPoint& Out
 	OutPoint.Transform.SetScale3D(WeightedScale);
 	OutPoint.Transform.SetLocation(InPosition);
 	OutPoint.Density = WeightedDensity;
-	OutPoint.Extents = WeightedExtents;
+	OutPoint.BoundsMin = WeightedBoundsMin;
+	OutPoint.BoundsMax = WeightedBoundsMax;
 	OutPoint.Color = WeightedColor;
 	OutPoint.Steepness = WeightedSteepness;
 
diff --git a/Engine/Plugins/Experimental/PCG/Source/PCG/Private/Data/PCGTextureData.cpp b/Engine/Plugins/Experimental/PCG/Source/PCG/Private/Data/PCGTextureData.cpp
--- a/Engine/Plugins/Experimental/PCG/Source/PCG/Private/Data/PCGTextureData.cpp
+++ b/Engine/Plugins/Experimental/PCG/Source/PCG/Private/Data/PCGTextureData.cpp
@@ -159,7 +159,7 @@ const UPCGPointData* UPCGBaseTextureData::CreatePointData(FPCGContext* Context)
 
 			const FVector TransformScale = Transform.GetScale3D();
 			// Note: divided by 4 here because the scale is doubled before, and the extents represent half a pixel
-			OutPoint.Extents = FVector(TransformScale.X * XScale / 4.0, TransformScale.Y * YScale / 4.0, 1.0);
+			OutPoint.SetExtents(FVector(TransformScale.X * XScale / 4.0, TransformScale.Y * YScale / 4.0, 1.0));
 			OutPoint.Color = ColorData[X + Y * Width];
 
 			return true;
diff --git a/Engine/Plugins/Experimental/PCG/Source/PCG/Private/Elements/PCGDebugElement.cpp b/Engine/Plugins/Experimental/PCG/Source/PCG/Private/Elements/PCGDebugElement.cpp
--- a/Engine/Plugins/Experimental/PCG/Source/PCG/Private/Elements/PCGDebugElement.cpp
+++ b/Engine/Plugins/Experimental/PCG/Source/PCG/Private/Elements/PCGDebugElement.cpp
@@ -136,9 +136,10 @@ namespace PCGDebugElement
 			{
 				const FPCGPoint& Point = Points[PointIndex];
 				InstanceCustomData.Add(Point.Density);
-				InstanceCustomData.Add(Point.Extents[0]);
-				InstanceCustomData.Add(Point.Extents[1]);
-				InstanceCustomData.Add(Point.Extents[2]);
+				const FVector Extents = Point.GetExtents();
+				InstanceCustomData.Add(Extents[0]);
+				InstanceCustomData.Add(Extents[1]);
+				InstanceCustomData.Add(Extents[2]);
 				InstanceCustomData.Add(Point.Color[0]);
 				InstanceCustomData.Add(Point.Color[1]);
 				InstanceCustomData.Add(Point.Color[2]);
diff --git a/Engine/Plugins/Experimental/PCG/Source/PCG/Private/Elements/PCGMetadataElement.cpp b/Engine/Plugins/Experimental/PCG/Source/PCG/Private/Elements/PCGMetadataElement.cpp
--- a/Engine/Plugins/Experimental/PCG/Source/PCG/Private/Elements/PCGMetadataElement.cpp
+++ b/Engine/Plugins/Experimental/PCG/Source/PCG/Private/Elements/PCGMetadataElement.cpp
@@ -140,9 +140,41 @@ bool FPCGMetadataOperationElement::ExecuteInternal(FPCGContext* Context) const
 					PCGE_LOG(Error, "Attribute %s already exists but its type is not compatible", *DestinationAttribute.ToString());
 				}
 			}
+			else if (PointProperty == EPCGPointProperties::BoundsMin)
+			{
+				auto ExtentsGetter = [](const FPCGPoint& InPoint) { return InPoint.BoundsMin; };
+
+				if (!SampledData->Metadata->HasAttribute(DestinationAttribute))
+				{
+					SampledData->Metadata->CreateVectorAttribute(DestinationAttribute, FVector::Zero(), /*bAllowsInterpolation=*/true, /*bOverrideParent=*/true);
+				}
+
+				FPCGMetadataAttributeBase* AttributeBase = SampledData->Metadata->GetMutableAttribute(DestinationAttribute);
+				if (!PCGMetadataOperations::SetValue<FVector, FVector>(SampledPoints, AttributeBase, SampledData->Metadata, ExtentsGetter) &&
+					!PCGMetadataOperations::SetValue<FVector, FTransform>(SampledPoints, AttributeBase, SampledData->Metadata, ExtentsGetter))
+				{
+					PCGE_LOG(Error, "Attribute %s already exists but its type is not compatible", *DestinationAttribute.ToString());
+				}
+			}
+			else if (PointProperty == EPCGPointProperties::BoundsMax)
+			{
+				auto ExtentsGetter = [](const FPCGPoint& InPoint) { return InPoint.BoundsMax; };
+
+				if (!SampledData->Metadata->HasAttribute(DestinationAttribute))
+				{
+					SampledData->Metadata->CreateVectorAttribute(DestinationAttribute, FVector::Zero(), /*bAllowsInterpolation=*/true, /*bOverrideParent=*/true);
+				}
+
+				FPCGMetadataAttributeBase* AttributeBase = SampledData->Metadata->GetMutableAttribute(DestinationAttribute);
+				if (!PCGMetadataOperations::SetValue<FVector, FVector>(SampledPoints, AttributeBase, SampledData->Metadata, ExtentsGetter) &&
+					!PCGMetadataOperations::SetValue<FVector, FTransform>(SampledPoints, AttributeBase, SampledData->Metadata, ExtentsGetter))
+				{
+					PCGE_LOG(Error, "Attribute %s already exists but its type is not compatible", *DestinationAttribute.ToString());
+				}
+			}
 			else if (PointProperty == EPCGPointProperties::Extents)
 			{
-				auto ExtentsGetter = [](const FPCGPoint& InPoint) { return InPoint.Extents; };
+				auto ExtentsGetter = [](const FPCGPoint& InPoint) { return InPoint.GetExtents(); };
 
 				if (!SampledData->Metadata->HasAttribute(DestinationAttribute))
 				{
@@ -260,9 +292,31 @@ bool FPCGMetadataOperationElement::ExecuteInternal(FPCGContext* Context) const
 					PCGE_LOG(Error, "Attribute %s already exists but its type is not compatible", *LocalSourceAttribute.ToString());
 				}
 			}
+			else if (PointProperty == EPCGPointProperties::BoundsMin)
+			{
+				auto ExtentsSetter = [](FPCGPoint& InPoint, const FVector& InValue) { InPoint.BoundsMin = InValue; };
+
+				const FPCGMetadataAttributeBase* AttributeBase = SampledData->Metadata->GetConstAttribute(LocalSourceAttribute);
+				if (!PCGMetadataOperations::SetValue<float, FVector>(AttributeBase, SampledPoints, ExtentsSetter) &&
+					!PCGMetadataOperations::SetValue<FVector, FVector>(AttributeBase, SampledPoints, ExtentsSetter))
+				{
+					PCGE_LOG(Error, "Attribute %s already exists but its type is not compatible", *LocalSourceAttribute.ToString());
+				}
+			}
+			else if (PointProperty == EPCGPointProperties::BoundsMax)
+			{
+				auto ExtentsSetter = [](FPCGPoint& InPoint, const FVector& InValue) { InPoint.BoundsMax = InValue; };
+
+				const FPCGMetadataAttributeBase* AttributeBase = SampledData->Metadata->GetConstAttribute(LocalSourceAttribute);
+				if (!PCGMetadataOperations::SetValue<float, FVector>(AttributeBase, SampledPoints, ExtentsSetter) &&
+					!PCGMetadataOperations::SetValue<FVector, FVector>(AttributeBase, SampledPoints, ExtentsSetter))
+				{
+					PCGE_LOG(Error, "Attribute %s already exists but its type is not compatible", *LocalSourceAttribute.ToString());
+				}
+			}
 			else if (PointProperty == EPCGPointProperties::Extents)
 			{
-				auto ExtentsSetter = [](FPCGPoint& InPoint, const FVector& InValue) { InPoint.Extents = InValue; };
+				auto ExtentsSetter = [](FPCGPoint& InPoint, const FVector& InValue) { InPoint.SetExtents(InValue); };
 
 				const FPCGMetadataAttributeBase* AttributeBase = SampledData->Metadata->GetConstAttribute(LocalSourceAttribute);
 				if(!PCGMetadataOperations::SetValue<float, FVector>(AttributeBase, SampledPoints, ExtentsSetter) &&
diff --git a/Engine/Plugins/Experimental/PCG/Source/PCG/Private/Elements/PCGPointFilter.cpp b/Engine/Plugins/Experimental/PCG/Source/PCG/Private/Elements/PCGPointFilter.cpp
--- a/Engine/Plugins/Experimental/PCG/Source/PCG/Private/Elements/PCGPointFilter.cpp
+++ b/Engine/Plugins/Experimental/PCG/Source/PCG/Private/Elements/PCGPointFilter.cpp
@@ -29,12 +29,28 @@ namespace PCGPointFilterHelpers
 				*(static_cast<float*>(OutData)) = InPoint.Density;
 			};
 		}
+		else if (TargetPointProperty == EPCGPointProperties::BoundsMin)
+		{
+			OutType = EPCGPointFilterConstantType::Vector;
+			return [](const FPCGPoint& InPoint, const UPCGMetadata*, void* OutData)
+			{
+				*(static_cast<FVector*>(OutData)) = InPoint.BoundsMin;
+			};
+		}
+		else if (TargetPointProperty == EPCGPointProperties::BoundsMax)
+		{
+			OutType = EPCGPointFilterConstantType::Vector;
+			return [](const FPCGPoint& InPoint, const UPCGMetadata*, void* OutData)
+			{
+				*(static_cast<FVector*>(OutData)) = InPoint.BoundsMax;
+			};
+		}
 		else if (TargetPointProperty == EPCGPointProperties::Extents)
 		{
 			OutType = EPCGPointFilterConstantType::Vector;
 			return [](const FPCGPoint& InPoint, const UPCGMetadata*, void* OutData)
 			{
-				*(static_cast<FVector*>(OutData)) = InPoint.Extents;
+				*(static_cast<FVector*>(OutData)) = InPoint.GetExtents();
 			};
 		}
 		else if (TargetPointProperty == EPCGPointProperties::Color)
diff --git a/Engine/Plugins/Experimental/PCG/Source/PCG/Private/Elements/PCGSurfaceSampler.cpp b/Engine/Plugins/Experimental/PCG/Source/PCG/Private/Elements/PCGSurfaceSampler.cpp
--- a/Engine/Plugins/Experimental/PCG/Source/PCG/Private/Elements/PCGSurfaceSampler.cpp
+++ b/Engine/Plugins/Experimental/PCG/Source/PCG/Private/Elements/PCGSurfaceSampler.cpp
@@ -195,7 +195,7 @@ bool FPCGSurfaceSamplerElement::ExecuteInternal(FPCGContext* Context) const
 				if (SpatialInput->GetPointAtPosition(TentativeLocation, OutPoint, SampledData->Metadata) || LoopData.bKeepZeroDensityPoints)
 				{
 					// Apply final parameters on the point
-					OutPoint.Extents = LoopData.PointExtents;
+					OutPoint.SetExtents(LoopData.PointExtents);
 					OutPoint.Density *= (LoopData.bApplyDensityToPoints ? ((Ratio - Chance) / Ratio) : 1.0f);
 					OutPoint.Steepness = LoopData.PointSteepness;
 					OutPoint.Seed = RandomSource.GetCurrentSeed();
diff --git a/Engine/Plugins/Experimental/PCG/Source/PCG/Private/Helpers/PCGBlueprintHelpers.cpp b/Engine/Plugins/Experimental/PCG/Source/PCG/Private/Helpers/PCGBlueprintHelpers.cpp
--- a/Engine/Plugins/Experimental/PCG/Source/PCG/Private/Helpers/PCGBlueprintHelpers.cpp
+++ b/Engine/Plugins/Experimental/PCG/Source/PCG/Private/Helpers/PCGBlueprintHelpers.cpp
@@ -61,4 +61,14 @@ UPCGComponent* UPCGBlueprintHelpers::GetOriginalComponent(FPCGContext& Context)
 	{
 		return Context.SourceComponent;
 	}
+}
+
+void UPCGBlueprintHelpers::SetExtents(FPCGPoint& InPoint, const FVector& InExtents)
+{
+	InPoint.SetExtents(InExtents);
+}
+
+FVector UPCGBlueprintHelpers::GetExtents(const FPCGPoint& InPoint)
+{
+	return InPoint.GetExtents();
 }
\ No newline at end of file
diff --git a/Engine/Plugins/Experimental/PCG/Source/PCG/Private/PCGPoint.cpp b/Engine/Plugins/Experimental/PCG/Source/PCG/Private/PCGPoint.cpp
--- a/Engine/Plugins/Experimental/PCG/Source/PCG/Private/PCGPoint.cpp
+++ b/Engine/Plugins/Experimental/PCG/Source/PCG/Private/PCGPoint.cpp
@@ -11,10 +11,10 @@ FPCGPoint::FPCGPoint(const FTransform& InTransform, float InDensity, int32 InSee
 
 FBoxSphereBounds FPCGPoint::GetBounds() const
 {
-	return FBoxSphereBounds(FBox::BuildAABB(FVector::Zero(), Extents).TransformBy(Transform));
+	return FBoxSphereBounds(FBox(BoundsMin, BoundsMax).TransformBy(Transform));
 }
 
 FBoxSphereBounds FPCGPoint::GetDensityBounds() const
 {
-	return FBoxSphereBounds(FBox::BuildAABB(FVector::Zero(), (2 - Steepness) * Extents).TransformBy(Transform));
+	return FBoxSphereBounds(FBox((2 - Steepness) * BoundsMin, (2 - Steepness) * BoundsMax).TransformBy(Transform));
 }
\ No newline at end of file
diff --git a/Engine/Plugins/Experimental/PCG/Source/PCG/Public/Helpers/PCGBlueprintHelpers.h b/Engine/Plugins/Experimental/PCG/Source/PCG/Public/Helpers/PCGBlueprintHelpers.h
--- a/Engine/Plugins/Experimental/PCG/Source/PCG/Public/Helpers/PCGBlueprintHelpers.h
+++ b/Engine/Plugins/Experimental/PCG/Source/PCG/Public/Helpers/PCGBlueprintHelpers.h
@@ -44,4 +44,10 @@ public:
 
 	UFUNCTION(BlueprintCallable, Category = "PCG|Temporary", meta = (ScriptMethod))
 	static UPCGComponent* GetOriginalComponent(UPARAM(ref) FPCGContext& Context);
+
+	UFUNCTION(BlueprintCallable, Category = "PCG|Temporary", meta = (ScriptMethod))
+	static void SetExtents(UPARAM(ref) FPCGPoint& InPoint, const FVector& InExtents);
+
+	UFUNCTION(BlueprintCallable, Category = "PCG|Temporary", meta = (ScriptMethod))
+	static FVector GetExtents(const FPCGPoint& InPoint);
 };
\ No newline at end of file
diff --git a/Engine/Plugins/Experimental/PCG/Source/PCG/Public/PCGPoint.h b/Engine/Plugins/Experimental/PCG/Source/PCG/Public/PCGPoint.h
--- a/Engine/Plugins/Experimental/PCG/Source/PCG/Public/PCGPoint.h
+++ b/Engine/Plugins/Experimental/PCG/Source/PCG/Public/PCGPoint.h
@@ -10,6 +10,8 @@ UENUM()
 enum class EPCGPointProperties : uint8
 {
 	Density,
+	BoundsMin,
+	BoundsMax,
 	Extents,
 	Color,
 	Position,
@@ -37,7 +39,10 @@ public:
 	float Density = 1.0f;
 
 	UPROPERTY(BlueprintReadWrite, EditAnywhere, Category = Properties)
-	FVector Extents = FVector::One();
+	FVector BoundsMin = -FVector::One();
+
+	UPROPERTY(BlueprintReadWrite, EditAnywhere, Category = Properties)
+	FVector BoundsMax = FVector::One();
 
 	UPROPERTY(BlueprintReadWrite, EditAnywhere, Category = Properties)
 	FVector4 Color = FVector4::One();
@@ -50,4 +55,7 @@ public:
 
 	UPROPERTY(BlueprintReadOnly, VisibleAnywhere, Category = "Properties|Metadata")
 	int64 MetadataEntry = -1;
+
+	FVector GetExtents() const { return (BoundsMax - BoundsMin) / 2.0; }
+	void SetExtents(const FVector& InExtents) { BoundsMin = -InExtents; BoundsMax = InExtents; }
 };
