Fixed issues on subgraph pins related to advanced pins change

#rnx

diff --git a/Engine/Plugins/Experimental/PCG/Source/PCG/Private/PCGGraph.cpp b/Engine/Plugins/Experimental/PCG/Source/PCG/Private/PCGGraph.cpp
--- a/Engine/Plugins/Experimental/PCG/Source/PCG/Private/PCGGraph.cpp
+++ b/Engine/Plugins/Experimental/PCG/Source/PCG/Private/PCGGraph.cpp
@@ -76,6 +76,29 @@ void UPCGGraph::PostLoad()
 			Edge->OutboundLabel = TEXT("Out");
 		}
 	}
+
+	// Fixup edges on subgraph nodes
+	for (UPCGNode* Node : Nodes)
+	{
+		if(Cast<UPCGBaseSubgraphSettings>(Node->DefaultSettings))
+		{
+			for (UPCGEdge* Edge : Node->InboundEdges)
+			{
+				if (Edge->OutboundLabel == NAME_None)
+				{
+					Edge->OutboundLabel = TEXT("In");
+				}
+			}
+
+			for (UPCGEdge* Edge : Node->OutboundEdges)
+			{
+				if (Edge->InboundLabel == NAME_None)
+				{
+					Edge->InboundLabel = TEXT("Out");
+				}
+			}
+		}
+	}
 #endif
 
 #if WITH_EDITOR
diff --git a/Engine/Plugins/Experimental/PCG/Source/PCG/Private/PCGSubgraph.cpp b/Engine/Plugins/Experimental/PCG/Source/PCG/Private/PCGSubgraph.cpp
--- a/Engine/Plugins/Experimental/PCG/Source/PCG/Private/PCGSubgraph.cpp
+++ b/Engine/Plugins/Experimental/PCG/Source/PCG/Private/PCGSubgraph.cpp
@@ -108,6 +108,30 @@ bool UPCGBaseSubgraphSettings::HasOutLabel(const FName& Label) const
 	}
 }
 
+bool UPCGBaseSubgraphSettings::HasDefaultInLabel() const
+{
+	if (UPCGGraph* Subgraph = GetSubgraph())
+	{
+		return Subgraph->GetInputNode()->HasDefaultInLabel();
+	}
+	else
+	{
+		return Super::HasDefaultInLabel();
+	}
+}
+
+bool UPCGBaseSubgraphSettings::HasDefaultOutLabel() const
+{
+	if (UPCGGraph* Subgraph = GetSubgraph())
+	{
+		return Subgraph->GetOutputNode()->HasDefaultOutLabel();
+	}
+	else
+	{
+		return Super::HasDefaultOutLabel();
+	}
+}
+
 TArray<FName> UPCGBaseSubgraphSettings::InLabels() const
 {
 	if (UPCGGraph* Subgraph = GetSubgraph())
diff --git a/Engine/Plugins/Experimental/PCG/Source/PCG/Public/PCGSubgraph.h b/Engine/Plugins/Experimental/PCG/Source/PCG/Public/PCGSubgraph.h
--- a/Engine/Plugins/Experimental/PCG/Source/PCG/Public/PCGSubgraph.h
+++ b/Engine/Plugins/Experimental/PCG/Source/PCG/Public/PCGSubgraph.h
@@ -36,6 +36,8 @@ protected:
 
 	virtual bool HasInLabel(const FName& Label) const override;
 	virtual bool HasOutLabel(const FName& Label) const override;
+	virtual bool HasDefaultInLabel() const override;
+	virtual bool HasDefaultOutLabel() const override;
 	virtual TArray<FName> InLabels() const override;
 	virtual TArray<FName> OutLabels() const override;
 	//~End UPCGSettings interface
