Moved create artifacts/can be multithreaded on the blueprint objects and removed from the settings
Cleanup the way we duplicate BP instance for execution so it doesn't pollute the data

diff --git a/Engine/Plugins/Experimental/PCG/Source/PCG/Private/Elements/PCGExecuteBlueprint.cpp b/Engine/Plugins/Experimental/PCG/Source/PCG/Private/Elements/PCGExecuteBlueprint.cpp
--- a/Engine/Plugins/Experimental/PCG/Source/PCG/Private/Elements/PCGExecuteBlueprint.cpp
+++ b/Engine/Plugins/Experimental/PCG/Source/PCG/Private/Elements/PCGExecuteBlueprint.cpp
@@ -170,6 +170,15 @@ void UPCGBlueprintSettings::PostLoad()
 	{
 		SetupBlueprintElementEvent();
 	}
+
+	if (BlueprintElementInstance)
+	{
+		BlueprintElementInstance->bCreatesArtifacts |= bCreatesArtifacts_DEPRECATED;
+		BlueprintElementInstance->bCanBeMultithreaded |= bCanBeMultithreaded_DEPRECATED;
+	}
+
+	bCreatesArtifacts_DEPRECATED = false;
+	bCanBeMultithreaded_DEPRECATED = false;
 }
 
 void UPCGBlueprintSettings::BeginDestroy()
@@ -492,7 +501,7 @@ FPCGContext* FPCGExecuteBlueprintElement::Initialize(const FPCGDataCollection& I
 	const UPCGBlueprintSettings* Settings = Context->GetInputSettings<UPCGBlueprintSettings>();
 	if (Settings && Settings->BlueprintElementInstance)
 	{
-		Context->BlueprintElementInstance = CastChecked<UPCGBlueprintElement>(StaticDuplicateObject(Settings->BlueprintElementInstance, Settings->BlueprintElementInstance, FName()));
+		Context->BlueprintElementInstance = CastChecked<UPCGBlueprintElement>(StaticDuplicateObject(Settings->BlueprintElementInstance, GetTransientPackage(), FName()));
 	}
 	else
 	{
@@ -504,9 +513,10 @@ FPCGContext* FPCGExecuteBlueprintElement::Initialize(const FPCGDataCollection& I
 
 bool FPCGExecuteBlueprintElement::IsCacheable(const UPCGSettings* InSettings) const
 {
-	if (const UPCGBlueprintSettings* BPSettings = Cast<const UPCGBlueprintSettings>(InSettings))
+	const UPCGBlueprintSettings* BPSettings = Cast<const UPCGBlueprintSettings>(InSettings);
+	if (BPSettings && BPSettings->BlueprintElementInstance)
 	{
-		return !BPSettings->bCreatesArtifacts;
+		return BPSettings->BlueprintElementInstance->bCreatesArtifacts;
 	}
 	else
 	{
@@ -516,9 +526,10 @@ bool FPCGExecuteBlueprintElement::IsCacheable(const UPCGSettings* InSettings) co
 
 bool FPCGExecuteBlueprintElement::CanExecuteOnlyOnMainThread(const UPCGSettings* InSettings) const
 {
-	if (const UPCGBlueprintSettings* BPSettings = Cast<const UPCGBlueprintSettings>(InSettings))
+	const UPCGBlueprintSettings* BPSettings = Cast<const UPCGBlueprintSettings>(InSettings);
+	if (BPSettings && BPSettings->BlueprintElementInstance)
 	{
-		return !BPSettings->bCanBeMultithreaded;
+		return !BPSettings->BlueprintElementInstance->bCanBeMultithreaded;
 	}
 	else
 	{
diff --git a/Engine/Plugins/Experimental/PCG/Source/PCG/Public/Elements/PCGExecuteBlueprint.h b/Engine/Plugins/Experimental/PCG/Source/PCG/Public/Elements/PCGExecuteBlueprint.h
--- a/Engine/Plugins/Experimental/PCG/Source/PCG/Public/Elements/PCGExecuteBlueprint.h
+++ b/Engine/Plugins/Experimental/PCG/Source/PCG/Public/Elements/PCGExecuteBlueprint.h
@@ -85,6 +85,12 @@ public:
 	FOnPCGBlueprintChanged OnBlueprintChangedDelegate;
 #endif
 
+	UPROPERTY(BlueprintReadWrite, EditAnywhere, Category = Settings)
+	bool bCreatesArtifacts = false;
+
+	UPROPERTY(BlueprintReadWrite, EditAnywhere, Category = Settings)
+	bool bCanBeMultithreaded = false;
+
 	UPROPERTY(EditAnywhere, BlueprintReadWrite, Category = "Settings|Input & Output")
 	TSet<FName> InputPinLabels;
 
@@ -158,11 +164,11 @@ protected:
 	TArray<FName> TrackedActorTags;
 #endif
 
-	UPROPERTY(BlueprintReadWrite, EditAnywhere, Category = Settings)
-	bool bCreatesArtifacts = false;
+	UPROPERTY()
+	bool bCreatesArtifacts_DEPRECATED = false;
 
-	UPROPERTY(BlueprintReadWrite, EditAnywhere, Category = Settings)
-	bool bCanBeMultithreaded = true;
+	UPROPERTY()
+	bool bCanBeMultithreaded_DEPRECATED = false;
 
 protected:
 #if WITH_EDITOR
