PCG:
*Added undo / redo support for various node operations in the editor.

diff --git a/Engine/Plugins/Experimental/PCG/Source/PCG/Private/PCGEdge.cpp b/Engine/Plugins/Experimental/PCG/Source/PCG/Private/PCGEdge.cpp
--- a/Engine/Plugins/Experimental/PCG/Source/PCG/Private/PCGEdge.cpp
+++ b/Engine/Plugins/Experimental/PCG/Source/PCG/Private/PCGEdge.cpp
@@ -5,6 +5,7 @@
 
 void UPCGEdge::BreakEdge()
 {
+	Modify();
 	if(InboundNode)
 	{
 		InboundNode->RemoveOutboundEdge(this);
diff --git a/Engine/Plugins/Experimental/PCG/Source/PCG/Private/PCGGraph.cpp b/Engine/Plugins/Experimental/PCG/Source/PCG/Private/PCGGraph.cpp
--- a/Engine/Plugins/Experimental/PCG/Source/PCG/Private/PCGGraph.cpp
+++ b/Engine/Plugins/Experimental/PCG/Source/PCG/Private/PCGGraph.cpp
@@ -14,9 +14,11 @@ UPCGGraph::UPCGGraph(const FObjectInitializer& ObjectInitializer)
 	: Super(ObjectInitializer)
 {
 	InputNode = ObjectInitializer.CreateDefaultSubobject<UPCGNode>(this, TEXT("DefaultInputNode"));
+	InputNode->SetFlags(RF_Transactional);
 	InputNode->SetDefaultSettings(ObjectInitializer.CreateDefaultSubobject<UPCGGraphInputOutputSettings>(this, TEXT("DefaultInputNodeSettings")));
 	Cast<UPCGGraphInputOutputSettings>(InputNode->DefaultSettings)->SetInput(true);
 	OutputNode = ObjectInitializer.CreateDefaultSubobject<UPCGNode>(this, TEXT("DefaultOutputNode"));
+	OutputNode->SetFlags(RF_Transactional);
 	OutputNode->SetDefaultSettings(ObjectInitializer.CreateDefaultSubobject<UPCGGraphInputOutputSettings>(this, TEXT("DefaultOutputNodeSettings")));
 	Cast<UPCGGraphInputOutputSettings>(OutputNode->DefaultSettings)->SetInput(false);
 #if WITH_EDITORONLY_DATA
@@ -182,6 +184,7 @@ UPCGNode* UPCGGraph::AddNodeOfType(TSubclassOf<class UPCGSettings> InSettingsCla
 	}
 
 	UPCGNode* Node = Settings->CreateNode();
+	Node->SetFlags(RF_Transactional);
 
 	if (Node)
 	{
@@ -220,6 +223,7 @@ UPCGNode* UPCGGraph::AddNode(UPCGSettings* InSettings)
 	}
 
 	UPCGNode* Node = InSettings->CreateNode();
+	Node->SetFlags(RF_Transactional);
 
 	if (Node)
 	{
@@ -302,6 +306,8 @@ UPCGNode* UPCGGraph::AddLabeledEdge(UPCGNode* From, const FName& InboundLabel, U
 	}
 
 	Modify();
+	From->Modify();
+	To->Modify();
 
 	// Create edge
 	UPCGEdge* Edge = NewObject<UPCGEdge>(this);
@@ -327,6 +333,10 @@ bool UPCGGraph::Contains(UPCGNode* Node) const
 
 void UPCGGraph::AddNode(UPCGNode* InNode)
 {
+	check(InNode);
+
+	Modify();
+
 	InNode->Rename(nullptr, this);
 
 #if WITH_EDITOR
@@ -346,6 +356,8 @@ void UPCGGraph::RemoveNode(UPCGNode* InNode)
 {
 	check(InNode);
 
+	Modify();
+
 	for (int32 EdgeIndex = InNode->InboundEdges.Num() - 1; EdgeIndex >= 0; --EdgeIndex)
 	{
 		UPCGEdge* Edge = InNode->InboundEdges[EdgeIndex];
@@ -480,6 +492,34 @@ bool UPCGGraph::RemoveOutboundEdges(UPCGNode* InNode, const FName& OutboundLabel
 	return bChanged;
 }
 
+#if WITH_EDITOR
+void UPCGGraph::PreNodeUndo(UPCGNode* InPCGNode)
+{
+	if (InPCGNode)
+	{
+		InPCGNode->OnNodeSettingsChangedDelegate.RemoveAll(this);
+
+		if (UPCGSubgraphNode* SubgraphNode = Cast<UPCGSubgraphNode>(InPCGNode))
+		{
+			SubgraphNode->OnNodeStructuralSettingsChangedDelegate.RemoveAll(this);
+		}
+	}
+}
+
+void UPCGGraph::PostNodeUndo(UPCGNode* InPCGNode)
+{
+	if (InPCGNode)
+	{
+		InPCGNode->OnNodeSettingsChangedDelegate.AddUObject(this, &UPCGGraph::OnSettingsChanged);
+
+		if (UPCGSubgraphNode* SubgraphNode = Cast<UPCGSubgraphNode>(InPCGNode))
+		{
+			SubgraphNode->OnNodeStructuralSettingsChangedDelegate.AddUObject(this, &UPCGGraph::OnStructuralSettingsChanged);
+		}
+	}
+}
+#endif
+
 #if WITH_EDITOR
 void UPCGGraph::DisableNotificationsForEditor()
 {
diff --git a/Engine/Plugins/Experimental/PCG/Source/PCG/Private/PCGNode.cpp b/Engine/Plugins/Experimental/PCG/Source/PCG/Private/PCGNode.cpp
--- a/Engine/Plugins/Experimental/PCG/Source/PCG/Private/PCGNode.cpp
+++ b/Engine/Plugins/Experimental/PCG/Source/PCG/Private/PCGNode.cpp
@@ -21,6 +21,12 @@ void UPCGNode::PostLoad()
 		DefaultSettings->OnSettingsChangedDelegate.AddUObject(this, &UPCGNode::OnSettingsChanged);
 	}
 
+	// Make sure legacy nodes support transactions.
+	if (HasAllFlags(RF_Transactional) == false)
+	{
+		SetFlags(RF_Transactional);
+	}
+
 	ApplyDeprecation();
 #endif
 }
@@ -39,17 +45,49 @@ void UPCGNode::ApplyDeprecation()
 }
 #endif
 
+#if WITH_EDITOR
 void UPCGNode::PostEditImport()
 {
 	Super::PostEditImport();
-#if WITH_EDITOR
 	if (DefaultSettings)
 	{
 		DefaultSettings->OnSettingsChangedDelegate.AddUObject(this, &UPCGNode::OnSettingsChanged);
 	}
-#endif
 }
 
+void UPCGNode::PreEditUndo()
+{
+	if (DefaultSettings)
+	{
+		DefaultSettings->OnSettingsChangedDelegate.RemoveAll(this);
+	}
+
+	UObject* Outer = GetOuter();
+	if (UPCGGraph* PCGGraph = Cast<UPCGGraph>(Outer))
+	{
+		PCGGraph->PreNodeUndo(this);
+	}
+
+	Super::PreEditUndo();
+}
+
+void UPCGNode::PostEditUndo()
+{
+	Super::PostEditUndo();
+
+	if (DefaultSettings)
+	{
+		DefaultSettings->OnSettingsChangedDelegate.AddUObject(this, &UPCGNode::OnSettingsChanged);
+	}
+
+ 	UObject* Outer = GetOuter();
+	if (UPCGGraph* PCGGraph = Cast<UPCGGraph>(Outer))
+	{
+		PCGGraph->PostNodeUndo(this);
+	}
+}
+#endif
+
 void UPCGNode::BeginDestroy()
 {
 #if WITH_EDITOR
@@ -209,11 +247,13 @@ bool UPCGNode::HasEdgeTo(UPCGNode* InNode) const
 
 void UPCGNode::RemoveInboundEdge(UPCGEdge* InEdge)
 {
+	Modify();
 	InboundEdges.Remove(InEdge);
 }
 
 void UPCGNode::RemoveOutboundEdge(UPCGEdge* InEdge)
 {
+	Modify();
 	OutboundEdges.Remove(InEdge);
 }
 
diff --git a/Engine/Plugins/Experimental/PCG/Source/PCG/Public/PCGGraph.h b/Engine/Plugins/Experimental/PCG/Source/PCG/Public/PCGGraph.h
--- a/Engine/Plugins/Experimental/PCG/Source/PCG/Public/PCGGraph.h
+++ b/Engine/Plugins/Experimental/PCG/Source/PCG/Public/PCGGraph.h
@@ -17,6 +17,10 @@
 UCLASS(BlueprintType, ClassGroup = (Procedural), hidecategories=(Object))
 class PCG_API UPCGGraph : public UObject
 {
+#if WITH_EDITOR
+	friend class FPCGEditor;
+#endif // WITH_EDITOR
+
 	GENERATED_BODY()
 
 public:
@@ -71,6 +75,8 @@ public:
 #if WITH_EDITOR
 	void DisableNotificationsForEditor();
 	void EnableNotificationsForEditor();
+	void PreNodeUndo(UPCGNode* InPCGNode);
+	void PostNodeUndo(UPCGNode* InPCGNode);
 #endif
 
 #if WITH_EDITOR
diff --git a/Engine/Plugins/Experimental/PCG/Source/PCG/Public/PCGNode.h b/Engine/Plugins/Experimental/PCG/Source/PCG/Public/PCGNode.h
--- a/Engine/Plugins/Experimental/PCG/Source/PCG/Public/PCGNode.h
+++ b/Engine/Plugins/Experimental/PCG/Source/PCG/Public/PCGNode.h
@@ -29,11 +29,14 @@ public:
 	
 	/** ~Begin UObject interface */
 	virtual void PostLoad() override;
-	virtual void PostEditImport() override;
 	virtual void BeginDestroy() override;
-	/** ~End UObject interface */
 
 #if WITH_EDITOR
+	virtual void PostEditImport() override;
+	virtual void PreEditUndo() override;
+	virtual void PostEditUndo() override;
+	/** ~End UObject interface */
+
 	/** Used to be able to force deprecation when things need to be deprecated at the graph level */
 	void ApplyDeprecation();
 #endif
diff --git a/Engine/Plugins/Experimental/PCG/Source/PCGEditor/Private/PCGEditor.cpp b/Engine/Plugins/Experimental/PCG/Source/PCGEditor/Private/PCGEditor.cpp
--- a/Engine/Plugins/Experimental/PCG/Source/PCGEditor/Private/PCGEditor.cpp
+++ b/Engine/Plugins/Experimental/PCG/Source/PCGEditor/Private/PCGEditor.cpp
@@ -2,6 +2,7 @@
 
 #include "PCGEditor.h"
 
+#include "PCGEditorCommon.h"
 #include "PCGEditorGraph.h"
 #include "PCGEditorGraphNodeBase.h"
 #include "PCGEditorGraphSchema.h"
@@ -21,6 +22,7 @@
 #include "Modules/ModuleManager.h"
 #include "PropertyEditorModule.h"
 #include "SNodePanel.h"
+#include "ScopedTransaction.h"
 #include "ToolMenus.h"
 #include "Widgets/Docking/SDockTab.h"
 
@@ -149,6 +151,30 @@ void FPCGEditor::UnregisterTabSpawners(const TSharedRef<class FTabManager>& InTa
 	FAssetEditorToolkit::UnregisterTabSpawners(InTabManager);
 }
 
+bool FPCGEditor::MatchesContext(const FTransactionContext& InContext, const TArray<TPair<UObject*, FTransactionObjectEvent>>& TransactionObjectContexts) const
+{
+	return InContext.Context == FPCGEditorCommon::ContextIdentifier;
+}
+
+void FPCGEditor::PostUndo(bool bSuccess)
+{
+	if (bSuccess)
+	{
+		if (PCGGraphBeingEdited)
+		{
+			PCGGraphBeingEdited->NotifyGraphChanged(true);
+		}
+
+		if (GraphEditorWidget.IsValid())
+		{
+			GraphEditorWidget->ClearSelectionSet();
+			GraphEditorWidget->NotifyGraphChanged();
+
+			FSlateApplication::Get().DismissAllMenus();
+		}
+	}
+}
+
 FName FPCGEditor::GetToolkitFName() const
 {
 	return FName(TEXT("PCGEditor"));
@@ -190,27 +216,32 @@ void FPCGEditor::DeleteSelectedNodes()
 		check(PCGEditorGraph && PCGGraph);
 
 		bool bChanged = false;
-
-		for (UObject* Object : GraphEditorWidget->GetSelectedNodes())
 		{
-			UPCGEditorGraphNodeBase* PCGEditorGraphNode = CastChecked<UPCGEditorGraphNodeBase>(Object);
+			const FScopedTransaction Transaction(*FPCGEditorCommon::ContextIdentifier, LOCTEXT("PCGEditorPaste", "PCG Editor: Delete"), nullptr);
+			PCGEditorGraph->Modify();
 
-			if (PCGEditorGraphNode->CanUserDeleteNode())
+			for (UObject* Object : GraphEditorWidget->GetSelectedNodes())
 			{
-				UPCGNode* PCGNode = PCGEditorGraphNode->GetPCGNode();
-				check(PCGNode);
+				UPCGEditorGraphNodeBase* PCGEditorGraphNode = CastChecked<UPCGEditorGraphNodeBase>(Object);
 
-				PCGGraph->RemoveNode(PCGNode);
-				PCGEditorGraphNode->DestroyNode();
-				bChanged = true;
+				if (PCGEditorGraphNode->CanUserDeleteNode())
+				{
+					UPCGNode* PCGNode = PCGEditorGraphNode->GetPCGNode();
+					check(PCGNode);
+
+					PCGGraph->RemoveNode(PCGNode);
+					PCGEditorGraphNode->DestroyNode();
+					bChanged = true;
+				}
 			}
+			PCGEditorGraph->Modify();
 		}
 
 		if (bChanged)
 		{
 			GraphEditorWidget->ClearSelectionSet();
 			GraphEditorWidget->NotifyGraphChanged();
-			PCGGraphBeingEdited->MarkPackageDirty();
+			PCGGraphBeingEdited->NotifyGraphChanged(true);
 		}
 	}
 }
@@ -303,6 +334,7 @@ void FPCGEditor::PasteNodesHere(const FVector2D& Location)
 		return;
 	}
 
+	const FScopedTransaction Transaction(*FPCGEditorCommon::ContextIdentifier, LOCTEXT("PCGEditorPaste", "PCG Editor: Paste"), nullptr);
 	PCGEditorGraph->Modify();
 
 	// Clear the selection set (newly pasted stuff will be selected)
@@ -362,7 +394,6 @@ void FPCGEditor::PasteNodesHere(const FVector2D& Location)
 	}
 
 	GraphEditorWidget->NotifyGraphChanged();
-	PCGGraphBeingEdited->MarkPackageDirty();
 }
 
 bool FPCGEditor::CanPasteNodes() const
@@ -563,6 +594,8 @@ void FPCGEditor::OnNodeTitleCommitted(const FText& NewText, ETextCommit::Type Co
 {
 	if (NodeBeingChanged)
 	{
+		const FScopedTransaction Transaction(*FPCGEditorCommon::ContextIdentifier, LOCTEXT("PCGEditorRenameNode", "PCG Editor: Rename Node"), nullptr);
+		NodeBeingChanged->Modify();
 		NodeBeingChanged->OnRenameNode(NewText.ToString());
 	}
 }
diff --git a/Engine/Plugins/Experimental/PCG/Source/PCGEditor/Private/PCGEditor.h b/Engine/Plugins/Experimental/PCG/Source/PCGEditor/Private/PCGEditor.h
--- a/Engine/Plugins/Experimental/PCG/Source/PCGEditor/Private/PCGEditor.h
+++ b/Engine/Plugins/Experimental/PCG/Source/PCGEditor/Private/PCGEditor.h
@@ -3,6 +3,7 @@
 #pragma once
 
 #include "Toolkits/AssetEditorToolkit.h"
+#include "EditorUndoClient.h"
 
 class UPCGGraph;
 class UPCGEditorGraph;
@@ -10,7 +11,7 @@ class SGraphEditor;
 class IDetailsView;
 class FUICommandList;
 
-class FPCGEditor : public FAssetEditorToolkit
+class FPCGEditor : public FAssetEditorToolkit, public FSelfRegisteringEditorUndoClient
 {
 public:
 	/** Edits the specified PCGGraph */
@@ -21,6 +22,12 @@ public:
 	virtual void UnregisterTabSpawners(const TSharedRef<class FTabManager>& InTabManager) override;
 	// ~End IToolkit interface
 
+	// ~Begin FEditorUndoClient interface
+	virtual bool MatchesContext(const FTransactionContext& InContext, const TArray<TPair<UObject *, FTransactionObjectEvent>>& TransactionObjectContexts) const override;
+	virtual void PostUndo(bool bSuccess) override;
+	virtual void PostRedo(bool bSuccess) override { PostUndo(bSuccess); }
+	// ~End FEditorUndoClient interface
+
 	// ~Begin FAssetEditorToolkit interface
 	virtual FName GetToolkitFName() const override;
 	virtual FText GetBaseToolkitName() const override;
diff --git a/Engine/Plugins/Experimental/PCG/Source/PCGEditor/Private/PCGEditorCommon.h b/Engine/Plugins/Experimental/PCG/Source/PCGEditor/Private/PCGEditorCommon.h
new file mode 100644
--- /dev/null
+++ b/Engine/Plugins/Experimental/PCG/Source/PCGEditor/Private/PCGEditorCommon.h
@@ -0,0 +1,8 @@
+// Copyright Epic Games, Inc. All Rights Reserved.
+
+#pragma once
+
+namespace FPCGEditorCommon
+{
+	const FString ContextIdentifier = TEXT("PCGEditorContext");
+};
diff --git a/Engine/Plugins/Experimental/PCG/Source/PCGEditor/Private/PCGEditorGraphSchema.cpp b/Engine/Plugins/Experimental/PCG/Source/PCGEditor/Private/PCGEditorGraphSchema.cpp
--- a/Engine/Plugins/Experimental/PCG/Source/PCGEditor/Private/PCGEditorGraphSchema.cpp
+++ b/Engine/Plugins/Experimental/PCG/Source/PCGEditor/Private/PCGEditorGraphSchema.cpp
@@ -2,12 +2,14 @@
 
 #include "PCGEditorGraphSchema.h"
 
+#include "PCGEditorCommon.h"
 #include "PCGEditorGraph.h"
 #include "PCGEditorGraphNodeBase.h"
 #include "PCGEditorGraphSchemaActions.h"
-#include "PCGSettings.h"
 #include "PCGGraph.h"
+#include "PCGSettings.h"
 
+#include "ScopedTransaction.h"
 #include "UObject/UObjectIterator.h"
 
 #define LOCTEXT_NAMESPACE "PCGEditorGraphSchema"
@@ -104,6 +106,8 @@ bool UPCGEditorGraphSchema::TryCreateConnection(UEdGraphPin* InA, UEdGraphPin* I
 
 void UPCGEditorGraphSchema::BreakPinLinks(UEdGraphPin& TargetPin, bool bSendsNodeNotification) const
 {
+	const FScopedTransaction Transaction(*FPCGEditorCommon::ContextIdentifier, LOCTEXT("PCGEditorBreakPinLinks", "Break Pin Links"), nullptr);
+
 	Super::BreakPinLinks(TargetPin, bSendsNodeNotification);
 
 	UEdGraphNode* GraphNode = TargetPin.GetOwningNode();
@@ -130,6 +134,8 @@ void UPCGEditorGraphSchema::BreakPinLinks(UEdGraphPin& TargetPin, bool bSendsNod
 
 void UPCGEditorGraphSchema::BreakSinglePinLink(UEdGraphPin* SourcePin, UEdGraphPin* TargetPin) const
 {
+	const FScopedTransaction Transaction(*FPCGEditorCommon::ContextIdentifier, LOCTEXT("PCGEditorBreakSinglePinLink", "Break Single Pin Link"), nullptr);
+
 	Super::BreakSinglePinLink(SourcePin, TargetPin);
 
 	UEdGraphNode* SourceGraphNode = SourcePin->GetOwningNode();
diff --git a/Engine/Plugins/Experimental/PCG/Source/PCGEditor/Private/PCGEditorGraphSchemaActions.cpp b/Engine/Plugins/Experimental/PCG/Source/PCGEditor/Private/PCGEditorGraphSchemaActions.cpp
--- a/Engine/Plugins/Experimental/PCG/Source/PCGEditor/Private/PCGEditorGraphSchemaActions.cpp
+++ b/Engine/Plugins/Experimental/PCG/Source/PCGEditor/Private/PCGEditorGraphSchemaActions.cpp
@@ -2,12 +2,16 @@
 
 #include "PCGEditorGraphSchemaActions.h"
 
-#include "PCGGraph.h"
+#include "PCGEditorCommon.h"
 #include "PCGEditorGraph.h"
 #include "PCGEditorGraphNode.h"
+#include "PCGGraph.h"
 #include "PCGSettings.h"
 
 #include "EdGraph/EdGraph.h"
+#include "ScopedTransaction.h"
+
+#define LOCTEXT_NAMESPACE "PCGEditorGraphSchemaAction_NewNode"
 
 UEdGraphNode* FPCGEditorGraphSchemaAction_NewNode::PerformAction(UEdGraph* ParentGraph, UEdGraphPin* FromPin, const FVector2D Location, bool bSelectNewNode)
 {
@@ -23,6 +27,9 @@ UEdGraphNode* FPCGEditorGraphSchemaAction_NewNode::PerformAction(UEdGraph* Paren
 		return nullptr;
 	}
 
+	const FScopedTransaction Transaction(*FPCGEditorCommon::ContextIdentifier, LOCTEXT("PCGEditorNewNode", "PCG Editor: New Node"), nullptr);
+	EditorGraph->Modify();
+
 	UPCGSettings* DefaultNodeSettings = nullptr;
 	UPCGNode* NewPCGNode = PCGGraph->AddNodeOfType(SettingsClass, DefaultNodeSettings);
 
@@ -43,3 +50,5 @@ UEdGraphNode* FPCGEditorGraphSchemaAction_NewNode::PerformAction(UEdGraph* Paren
 
 	return NewNode;
 }
+
+#undef LOCTEXT_NAMESPACE
diff --git a/Engine/Plugins/Experimental/PCG/Source/PCGEditor/Private/SPCGEditorGraphNode.cpp b/Engine/Plugins/Experimental/PCG/Source/PCGEditor/Private/SPCGEditorGraphNode.cpp
--- a/Engine/Plugins/Experimental/PCG/Source/PCGEditor/Private/SPCGEditorGraphNode.cpp
+++ b/Engine/Plugins/Experimental/PCG/Source/PCGEditor/Private/SPCGEditorGraphNode.cpp
@@ -24,9 +24,9 @@ void SPCGEditorGraphNode::MoveTo(const FVector2D& NewPosition, FNodeSet& NodeFil
 
 	if (UPCGNode* PCGNode = PCGEditorGraphNode->GetPCGNode())
 	{
+		PCGNode->Modify();
 		PCGNode->PositionX = PCGEditorGraphNode->NodePosX;
 		PCGNode->PositionY = PCGEditorGraphNode->NodePosY;
-		PCGNode->MarkPackageDirty();
 	}
 }
 
