Major overhaul of pins/edges:
Create pin class on the PCG side - preparation for UE-147879
Simplify edges
Moved most edge logic from the editor to the graph

diff --git a/Engine/Plugins/Experimental/PCG/Source/PCG/Private/Elements/PCGDifferenceElement.cpp b/Engine/Plugins/Experimental/PCG/Source/PCG/Private/Elements/PCGDifferenceElement.cpp
--- a/Engine/Plugins/Experimental/PCG/Source/PCG/Private/Elements/PCGDifferenceElement.cpp
+++ b/Engine/Plugins/Experimental/PCG/Source/PCG/Private/Elements/PCGDifferenceElement.cpp
@@ -1,6 +1,8 @@
 // Copyright Epic Games, Inc. All Rights Reserved.
 
 #include "Elements/PCGDifferenceElement.h"
+
+#include "PCGCommon.h"
 #include "Data/PCGUnionData.h"
 #include "Helpers/PCGSettingsHelpers.h"
 
@@ -12,7 +14,7 @@ namespace PCGDifferenceConstants
 
 TArray<FName> UPCGDifferenceSettings::InLabels() const
 {
-	return { PCGDifferenceConstants::SourceLabel, PCGDifferenceConstants::DifferencesLabel };
+	return { PCGPinConstants::DefaultInputLabel, PCGDifferenceConstants::SourceLabel, PCGDifferenceConstants::DifferencesLabel };
 }
 
 FPCGElementPtr UPCGDifferenceSettings::CreateElement() const
diff --git a/Engine/Plugins/Experimental/PCG/Source/PCG/Private/Elements/PCGExecuteBlueprint.cpp b/Engine/Plugins/Experimental/PCG/Source/PCG/Private/Elements/PCGExecuteBlueprint.cpp
--- a/Engine/Plugins/Experimental/PCG/Source/PCG/Private/Elements/PCGExecuteBlueprint.cpp
+++ b/Engine/Plugins/Experimental/PCG/Source/PCG/Private/Elements/PCGExecuteBlueprint.cpp
@@ -226,7 +226,7 @@ void UPCGBlueprintSettings::OnBlueprintChanged(UBlueprint* InBlueprint)
 	TeardownBlueprintElementEvent();
 	SetupBlueprintElementEvent();
 
-	OnSettingsChangedDelegate.Broadcast(this);
+	OnSettingsChangedDelegate.Broadcast(this, EPCGChangeType::Settings);
 }
 
 void UPCGBlueprintSettings::OnBlueprintElementChanged(UPCGBlueprintElement* InElement)
@@ -236,7 +236,7 @@ void UPCGBlueprintSettings::OnBlueprintElementChanged(UPCGBlueprintElement* InEl
 		// When a data dependency is changed, this means we have to dirty the cache, otherwise it will not register as a change.
 		DirtyCache();
 
-		OnSettingsChangedDelegate.Broadcast(this);
+		OnSettingsChangedDelegate.Broadcast(this, EPCGChangeType::Settings);
 	}
 }
 #endif // WITH_EDITOR
@@ -328,22 +328,36 @@ FName UPCGBlueprintSettings::AdditionalTaskName() const
 
 TArray<FName> UPCGBlueprintSettings::InLabels() const
 {
-	return BlueprintElementInstance ? BlueprintElementInstance->InputPinLabels.Array() : TArray<FName>();
+	TArray<FName> InputLabels;
+
+	if (BlueprintElementInstance)
+	{
+		if (BlueprintElementInstance->bHasDefaultInPin)
+		{
+			InputLabels.Add(PCGPinConstants::DefaultInputLabel);
+		}
+
+		InputLabels.Append(BlueprintElementInstance->InputPinLabels.Array());
+	}
+
+	return InputLabels;
 }
 
 TArray<FName> UPCGBlueprintSettings::OutLabels() const
 {
-	return BlueprintElementInstance ? BlueprintElementInstance->OutputPinLabels.Array() : TArray<FName>();
-}
+	TArray<FName> OutputLabels;
+	
+	if (BlueprintElementInstance)
+	{
+		if (BlueprintElementInstance->bHasDefaultOutPin)
+		{
+			OutputLabels.Add(PCGPinConstants::DefaultOutputLabel);
+		}
 
-bool UPCGBlueprintSettings::HasDefaultInLabel() const
-{
-	return !BlueprintElementInstance || BlueprintElementInstance->bHasDefaultInPin;
-}
+		OutputLabels.Append(BlueprintElementInstance->OutputPinLabels.Array());
+	}
 
-bool UPCGBlueprintSettings::HasDefaultOutLabel() const
-{
-	return !BlueprintElementInstance || BlueprintElementInstance->bHasDefaultOutPin;
+	return OutputLabels;
 }
 
 FPCGElementPtr UPCGBlueprintSettings::CreateElement() const
diff --git a/Engine/Plugins/Experimental/PCG/Source/PCG/Private/Graph/PCGGraphCompiler.cpp b/Engine/Plugins/Experimental/PCG/Source/PCG/Private/Graph/PCGGraphCompiler.cpp
--- a/Engine/Plugins/Experimental/PCG/Source/PCG/Private/Graph/PCGGraphCompiler.cpp
+++ b/Engine/Plugins/Experimental/PCG/Source/PCG/Private/Graph/PCGGraphCompiler.cpp
@@ -20,7 +20,7 @@ TArray<FPCGGraphTask> FPCGGraphCompiler::CompileGraph(UPCGGraph* InGraph, FPCGTa
 	// Prime the node queue with all nodes that have no inbound edges
 	for (const UPCGNode* Node : InGraph->GetNodes())
 	{
-		if (Node->InboundEdges.IsEmpty())
+		if(!Node->HasInboundEdges())
 		{
 			NodeQueue.Add(Node);
 		}
@@ -71,16 +71,26 @@ TArray<FPCGGraphTask> FPCGGraphCompiler::CompileGraph(UPCGGraph* InGraph, FPCGTa
 			PreTask.Node = Node;
 			PreTask.NodeId = PreId;
 
-			for (const UPCGEdge* InboundEdge : Node->InboundEdges)
+			for (const UPCGPin* InputPin : Node->InputPins)
 			{
-				check(InboundEdge);
-				PreTask.Inputs.Emplace(IdMapping[InboundEdge->InboundNode], InboundEdge->InboundLabel, InboundEdge->OutboundLabel);
+				check(InputPin);
+				for (const UPCGEdge* InboundEdge : InputPin->Edges)
+				{
+					if (InboundEdge->IsValid())
+					{
+						PreTask.Inputs.Emplace(IdMapping[InboundEdge->InputPin->Node], InboundEdge->InputPin, InboundEdge->OutputPin);
+					}
+					else
+					{
+						UE_LOG(LogPCG, Warning, TEXT("Invalid inbound edge on subgraph"));
+					}
+				}
 			}
 
 			// Add pre-task as input to subgraph input node task
 			if (InputNodeTask)
 			{
-				InputNodeTask->Inputs.Emplace(PreId, NAME_None, NAME_None);
+				InputNodeTask->Inputs.Emplace(PreId, nullptr, nullptr);
 			}
 
 			// Merge subgraph tasks into current tasks.
@@ -95,7 +105,7 @@ TArray<FPCGGraphTask> FPCGGraphCompiler::CompileGraph(UPCGGraph* InGraph, FPCGTa
 			// Add subgraph output node task as input to the post-task
 			if (OutputNodeTask)
 			{
-				PostTask.Inputs.Emplace(OutputNodeTask->NodeId, NAME_None, NAME_None);
+				PostTask.Inputs.Emplace(OutputNodeTask->NodeId, nullptr, nullptr);
 			}
 
 			IdMapping.Add(Node, PostId);
@@ -103,7 +113,7 @@ TArray<FPCGGraphTask> FPCGGraphCompiler::CompileGraph(UPCGGraph* InGraph, FPCGTa
 			// Prime any input-less nodes from the subgraph
 			for (const UPCGNode* SubNode : Subgraph->GetNodes())
 			{
-				if (SubNode->InboundEdges.IsEmpty())
+				if(!SubNode->HasInboundEdges())
 				{
 					NodeQueue.Add(SubNode);
 				}
@@ -116,16 +126,25 @@ TArray<FPCGGraphTask> FPCGGraphCompiler::CompileGraph(UPCGGraph* InGraph, FPCGTa
 			Task.Node = Node;
 			Task.NodeId = NodeId;
 
-			for (const UPCGEdge* InboundEdge : Node->InboundEdges)
+			for (const UPCGPin* InputPin : Node->InputPins)
 			{
-				if (FPCGTaskId* InboundId = IdMapping.Find(InboundEdge->InboundNode))
-				{
-					Task.Inputs.Emplace(IdMapping[InboundEdge->InboundNode], InboundEdge->InboundLabel, InboundEdge->OutboundLabel);
-				}
-				else
+				for (const UPCGEdge* InboundEdge : InputPin->Edges)
 				{
-					UE_LOG(LogPCG, Error, TEXT("Inconsistent node linkage between node %s and node %s"), *Node->GetFName().ToString(), (InboundEdge->InboundNode ? *InboundEdge->InboundNode->GetFName().ToString() : TEXT("null")));
-					return TArray<FPCGGraphTask>();
+					if (!InboundEdge->IsValid())
+					{
+						UE_LOG(LogPCG, Warning, TEXT("Unbound inbound edge"));
+						continue;
+					}
+
+					if (FPCGTaskId* InboundId = IdMapping.Find(InboundEdge->InputPin->Node))
+					{
+						Task.Inputs.Emplace(IdMapping[InboundEdge->InputPin->Node], InboundEdge->InputPin, InboundEdge->OutputPin); 
+					}
+					else
+					{
+						UE_LOG(LogPCG, Error, TEXT("Inconsistent node linkage on node %s"), *Node->GetFName().ToString());
+						return TArray<FPCGGraphTask>();
+					}
 				}
 			}
 
@@ -133,20 +152,35 @@ TArray<FPCGGraphTask> FPCGGraphCompiler::CompileGraph(UPCGGraph* InGraph, FPCGTa
 		}
 
 		// Push next ready nodes on the queue
-		for (const UPCGEdge* OutboundEdge : Node->OutboundEdges)
+		for (const UPCGPin* OutPin : Node->OutputPins)
 		{
-			const UPCGNode* Outbound = OutboundEdge->OutboundNode;
-
-			bool bAllPrerequisitesMet = true;
-			for (const UPCGEdge* InboundEdge : Outbound->InboundEdges)
+			for (const UPCGEdge* OutboundEdge : OutPin->Edges)
 			{
-				const UPCGNode* Inbound = InboundEdge->InboundNode;
-				bAllPrerequisitesMet &= IdMapping.Contains(Inbound);
-			}
+				if (!OutboundEdge->IsValid())
+				{
+					UE_LOG(LogPCG, Warning, TEXT("Unbound outbound edge"));
+					continue;
+				}
 
-			if (bAllPrerequisitesMet)
-			{
-				NodeQueue.Add(Outbound);
+				const UPCGNode* OutboundNode = OutboundEdge->OutputPin->Node;
+				check(OutboundNode);
+				bool bAllPrerequisitesMet = true;
+
+				for (const UPCGPin* OutboundNodeInputPin : OutboundNode->InputPins)
+				{
+					for (const UPCGEdge* OutboundNodeInboundEdge : OutboundNodeInputPin->Edges)
+					{
+						if (OutboundNodeInboundEdge->IsValid())
+						{
+							bAllPrerequisitesMet &= IdMapping.Contains(OutboundNodeInboundEdge->InputPin->Node);
+						}
+					}
+				}
+
+				if (bAllPrerequisitesMet)
+				{
+					NodeQueue.Add(OutboundNode);
+				}
 			}
 		}
 	}
@@ -267,7 +301,7 @@ void FPCGGraphCompiler::CompileTopGraph(UPCGGraph* InGraph)
 		FPCGGraphTask& Task = CompiledTasks[TaskIndex];
 		if (Task.Inputs.IsEmpty())
 		{
-			Task.Inputs.Emplace(PreExecuteTaskId, NAME_None, NAME_None);
+			Task.Inputs.Emplace(PreExecuteTaskId, nullptr, nullptr);
 		}
 	}
 
@@ -287,7 +321,7 @@ void FPCGGraphCompiler::CompileTopGraph(UPCGGraph* InGraph)
 		const FPCGGraphTask& Task = CompiledTasks[TaskIndex];
 		if (!TasksWithSuccessors.Contains(Task.NodeId))
 		{
-			PostExecuteTask.Inputs.Emplace(Task.NodeId, NAME_None, NAME_None);
+			PostExecuteTask.Inputs.Emplace(Task.NodeId, nullptr, nullptr);
 		}
 
 		for (const FPCGGraphTaskInput& Input : Task.Inputs)
diff --git a/Engine/Plugins/Experimental/PCG/Source/PCG/Private/Graph/PCGGraphExecutor.cpp b/Engine/Plugins/Experimental/PCG/Source/PCG/Private/Graph/PCGGraphExecutor.cpp
--- a/Engine/Plugins/Experimental/PCG/Source/PCG/Private/Graph/PCGGraphExecutor.cpp
+++ b/Engine/Plugins/Experimental/PCG/Source/PCG/Private/Graph/PCGGraphExecutor.cpp
@@ -89,7 +89,7 @@ FPCGTaskId FPCGGraphExecutor::Schedule(UPCGGraph* Graph, UPCGComponent* SourceCo
 		check(ScheduledTask.Tasks.Num() >= 2 && ScheduledTask.Tasks[ScheduledTask.Tasks.Num() - 2].Node == nullptr);
 		for (FPCGTaskId ExternalDependency : ExternalDependencies)
 		{
-			ScheduledTask.Tasks[ScheduledTask.Tasks.Num() - 2].Inputs.Emplace(ExternalDependency, NAME_None, NAME_None);
+			ScheduledTask.Tasks[ScheduledTask.Tasks.Num() - 2].Inputs.Emplace(ExternalDependency, nullptr, nullptr);
 		}
 #endif
 
@@ -105,7 +105,7 @@ FPCGTaskId FPCGGraphExecutor::ScheduleGeneric(TFunction<bool()> InOperation, con
 	FPCGGraphTask Task;
 	for (FPCGTaskId TaskDependency : TaskDependencies)
 	{
-		Task.Inputs.Emplace(TaskDependency, NAME_None, NAME_None);
+		Task.Inputs.Emplace(TaskDependency, nullptr, nullptr);
 	}
 	Task.Element = MakeShared<FPCGGenericElement>(InOperation);
 
@@ -450,22 +450,25 @@ void FPCGGraphExecutor::BuildTaskInput(const FPCGGraphTask& Task, FPCGDataCollec
 		const FPCGDataCollection& InputCollection = OutputData[Input.TaskId];
 
 		TaskInput.bCancelExecution |= InputCollection.bCancelExecution;
-		
-		// Trivial case : take everything
-		if (Input.InboundLabel == NAME_None && Input.OutboundLabel == NAME_None)
+
+		// Get input data at the given pin (or everything)
+		const int32 TaggedDataOffset = TaskInput.TaggedData.Num();
+		if (Input.InPin)
 		{
-			TaskInput.TaggedData.Append(InputCollection.TaggedData);
+			TaskInput.TaggedData.Append(InputCollection.GetInputsByPin(Input.InPin->Label));
 		}
-		// Non-trivial case: take all matches from OutputData[Id]
-		// And apply the out label on it
 		else
 		{
-			const int32 TaggedDataOffset = TaskInput.TaggedData.Num();
-			TaskInput.TaggedData.Append(InputCollection.GetInputsByPin(Input.InboundLabel));
+			TaskInput.TaggedData.Append(InputCollection.TaggedData);
+		}
 
+		// Apply labelling on data; technically, we should ensure that we do this only for pass-through nodes,
+		// Otherwise we could also null out the label on the input...
+		if (Input.OutPin)
+		{
 			for (int32 TaggedDataIndex = TaggedDataOffset; TaggedDataIndex < TaskInput.TaggedData.Num(); ++TaggedDataIndex)
 			{
-				TaskInput.TaggedData[TaggedDataIndex].Pin = Input.OutboundLabel;
+				TaskInput.TaggedData[TaggedDataIndex].Pin = Input.OutPin->Label;
 			}
 		}
 	}
@@ -577,7 +580,7 @@ bool FPCGFetchInputElement::ExecuteInternal(FPCGContext* Context) const
 	{
 		FPCGTaggedData& TaggedData = Context->OutputData.TaggedData.Emplace_GetRef();
 		TaggedData.Data = PCGData;
-		TaggedData.Pin = PCGInputOutputConstants::DefaultInLabel;
+		TaggedData.Pin = PCGPinConstants::DefaultInputLabel;
 	}
 
 	if (UPCGData* InputPCGData = Component->GetInputPCGData())
diff --git a/Engine/Plugins/Experimental/PCG/Source/PCG/Private/Graph/PCGGraphExecutor.h b/Engine/Plugins/Experimental/PCG/Source/PCG/Private/Graph/PCGGraphExecutor.h
--- a/Engine/Plugins/Experimental/PCG/Source/PCG/Private/Graph/PCGGraphExecutor.h
+++ b/Engine/Plugins/Experimental/PCG/Source/PCG/Private/Graph/PCGGraphExecutor.h
@@ -13,6 +13,7 @@
 #include "WorldPartition/WorldPartitionHandle.h" // Needed for FWorldPartitonReference
 #endif
 
+class UPCGPin;
 class UPCGGraph;
 class UPCGNode;
 class UPCGComponent;
@@ -20,16 +21,16 @@ class FPCGGraphCompiler;
 
 struct FPCGGraphTaskInput
 {
-	FPCGGraphTaskInput(FPCGTaskId InTaskId, const FName& InInboundLabel, const FName& InOutboundLabel)
+	FPCGGraphTaskInput(FPCGTaskId InTaskId, const UPCGPin* InInboundPin, const UPCGPin* InOutboundPin)
 		: TaskId(InTaskId)
-		, InboundLabel(InInboundLabel)
-		, OutboundLabel(InOutboundLabel)		
+		, InPin(InInboundPin)
+		, OutPin(InOutboundPin)
 	{
 	}
 
 	FPCGTaskId TaskId;
-	FName InboundLabel;
-	FName OutboundLabel;
+	const UPCGPin* InPin;
+	const UPCGPin* OutPin;
 };
 
 struct FPCGGraphTask
diff --git a/Engine/Plugins/Experimental/PCG/Source/PCG/Private/PCGEdge.cpp b/Engine/Plugins/Experimental/PCG/Source/PCG/Private/PCGEdge.cpp
--- a/Engine/Plugins/Experimental/PCG/Source/PCG/Private/PCGEdge.cpp
+++ b/Engine/Plugins/Experimental/PCG/Source/PCG/Private/PCGEdge.cpp
@@ -3,16 +3,19 @@
 #include "PCGEdge.h"
 #include "PCGNode.h"
 
-void UPCGEdge::BreakEdge()
+bool UPCGEdge::IsValid() const
 {
-	Modify();
-	if(InboundNode)
-	{
-		InboundNode->RemoveOutboundEdge(this);
-	}
+	return InputPin.Get() && OutputPin.Get();
+}
 
-	if(OutboundNode)
-	{
-		OutboundNode->RemoveInboundEdge(this);
-	}
+UPCGPin* UPCGEdge::GetOtherPin(const UPCGPin* Pin)
+{
+	check(Pin == InputPin || Pin == OutputPin);
+	return Pin == InputPin ? OutputPin : InputPin;
+}
+
+const UPCGPin* UPCGEdge::GetOtherPin(const UPCGPin* Pin) const
+{
+	check(Pin == InputPin || Pin == OutputPin);
+	return Pin == InputPin ? OutputPin : InputPin;
 }
\ No newline at end of file
diff --git a/Engine/Plugins/Experimental/PCG/Source/PCG/Private/PCGElement.cpp b/Engine/Plugins/Experimental/PCG/Source/PCG/Private/PCGElement.cpp
--- a/Engine/Plugins/Experimental/PCG/Source/PCG/Private/PCGElement.cpp
+++ b/Engine/Plugins/Experimental/PCG/Source/PCG/Private/PCGElement.cpp
@@ -162,12 +162,11 @@ void IPCGElement::CleanupAndValidateOutputLabels(FPCGContext* Context) const
 		// Cleanup any residual labels if the node isn't supposed to produce them
 		// TODO: this is a bit of a crutch, could be refactored out if we review the way we push tagged data
 		TArray<FName> OutLabels = Settings->OutLabels();
-		if ((Settings->HasDefaultOutLabel() && OutLabels.Num() == 0) || (!Settings->HasDefaultOutLabel() && OutLabels.Num() == 1))
+		if (OutLabels.Num() == 1)
 		{
-			const FName& DefaultLabel = Settings->HasDefaultOutLabel() ? NAME_None : OutLabels[0];
 			for (FPCGTaggedData& TaggedData : Context->OutputData.TaggedData)
 			{
-				TaggedData.Pin = DefaultLabel;
+				TaggedData.Pin = OutLabels[0];
 			}
 		}
 
diff --git a/Engine/Plugins/Experimental/PCG/Source/PCG/Private/PCGGraph.cpp b/Engine/Plugins/Experimental/PCG/Source/PCG/Private/PCGGraph.cpp
--- a/Engine/Plugins/Experimental/PCG/Source/PCG/Private/PCGGraph.cpp
+++ b/Engine/Plugins/Experimental/PCG/Source/PCG/Private/PCGGraph.cpp
@@ -26,8 +26,8 @@ UPCGGraph::UPCGGraph(const FObjectInitializer& ObjectInitializer)
 #endif
 
 #if WITH_EDITOR
-	InputNode->OnNodeSettingsChangedDelegate.AddUObject(this, &UPCGGraph::OnSettingsChanged);
-	OutputNode->OnNodeSettingsChangedDelegate.AddUObject(this, &UPCGGraph::OnSettingsChanged);
+	InputNode->OnNodeChangedDelegate.AddUObject(this, &UPCGGraph::OnNodeChanged);
+	OutputNode->OnNodeChangedDelegate.AddUObject(this, &UPCGGraph::OnNodeChanged);
 #endif 
 
 	// Note: default connection from input to output
@@ -43,8 +43,6 @@ void UPCGGraph::PostLoad()
 
 #if WITH_EDITOR
 	// Deprecation
-	InputNode->ApplyDeprecation();
-
 	if (!Cast<UPCGGraphInputOutputSettings>(InputNode->DefaultSettings))
 	{
 		InputNode->DefaultSettings = NewObject<UPCGGraphInputOutputSettings>(this, TEXT("DefaultInputNodeSettings"));
@@ -52,17 +50,6 @@ void UPCGGraph::PostLoad()
 
 	Cast<UPCGGraphInputOutputSettings>(InputNode->DefaultSettings)->SetInput(true);
 
-	// Fixup edges on the input node
-	for (UPCGEdge* Edge : InputNode->OutboundEdges)
-	{
-		if (Edge->InboundLabel == NAME_None)
-		{
-			Edge->InboundLabel = TEXT("In");
-		}
-	}
-
-	OutputNode->ApplyDeprecation();
-
 	if (!Cast<UPCGGraphInputOutputSettings>(OutputNode->DefaultSettings))
 	{
 		OutputNode->DefaultSettings = NewObject<UPCGGraphInputOutputSettings>(this, TEXT("DefaultOutputNodeSettings"));
@@ -70,42 +57,19 @@ void UPCGGraph::PostLoad()
 
 	Cast<UPCGGraphInputOutputSettings>(OutputNode->DefaultSettings)->SetInput(false);
 
-	// Fixup edges on the output node
-	for (UPCGEdge* Edge : OutputNode->InboundEdges)
-	{
-		if (Edge->OutboundLabel == NAME_None)
-		{
-			Edge->OutboundLabel = TEXT("Out");
-		}
-	}
+	// Additional deprecation
+	InputNode->ApplyDeprecation();
+	OutputNode->ApplyDeprecation();
 
-	// Fixup edges on subgraph nodes
 	for (UPCGNode* Node : Nodes)
 	{
-		if(Cast<UPCGBaseSubgraphSettings>(Node->DefaultSettings))
-		{
-			for (UPCGEdge* Edge : Node->InboundEdges)
-			{
-				if (Edge->OutboundLabel == NAME_None)
-				{
-					Edge->OutboundLabel = TEXT("In");
-				}
-			}
-
-			for (UPCGEdge* Edge : Node->OutboundEdges)
-			{
-				if (Edge->InboundLabel == NAME_None)
-				{
-					Edge->InboundLabel = TEXT("Out");
-				}
-			}
-		}
+		Node->ApplyDeprecation();
 	}
 #endif
 
 #if WITH_EDITOR
-	InputNode->OnNodeSettingsChangedDelegate.AddUObject(this, &UPCGGraph::OnSettingsChanged);
-	OutputNode->OnNodeSettingsChangedDelegate.AddUObject(this, &UPCGGraph::OnSettingsChanged);
+	InputNode->OnNodeChangedDelegate.AddUObject(this, &UPCGGraph::OnNodeChanged);
+	OutputNode->OnNodeChangedDelegate.AddUObject(this, &UPCGGraph::OnNodeChanged);
 
 	for (UPCGNode* Node : Nodes)
 	{
@@ -122,8 +86,8 @@ void UPCGGraph::BeginDestroy()
 		OnNodeRemoved(Node);
 	}
 
-	OutputNode->OnNodeSettingsChangedDelegate.RemoveAll(this);
-	InputNode->OnNodeSettingsChangedDelegate.RemoveAll(this);
+	OutputNode->OnNodeChangedDelegate.RemoveAll(this);
+	InputNode->OnNodeChangedDelegate.RemoveAll(this);
 
 	// Notify the compiler to remove this graph from its cache
 	UWorld* World = GEditor->GetEditorWorldContext().World();
@@ -147,8 +111,15 @@ void UPCGGraph::InitializeFromTemplate()
 
 	auto ResetDefaultNode = [](UPCGNode* InNode) {
 		check(InNode);
-		InNode->OutboundEdges.Reset();
-		InNode->InboundEdges.Reset();
+		for (UPCGPin* InputPin : InNode->InputPins)
+		{
+			InputPin->BreakAllEdges();
+		}
+
+		for (UPCGPin* OutputPin : InNode->OutputPins)
+		{
+			OutputPin->BreakAllEdges();
+		}
 
 		// Reset settings as well
 		InNode->SetDefaultSettings(NewObject<UPCGTrivialSettings>(InNode));
@@ -256,13 +227,7 @@ UPCGNode* UPCGGraph::AddNode(UPCGSettings* InSettings)
 void UPCGGraph::OnNodeAdded(UPCGNode* InNode)
 {
 #if WITH_EDITOR
-	InNode->OnNodeSettingsChangedDelegate.AddUObject(this, &UPCGGraph::OnSettingsChanged);
-
-	if (UPCGSubgraphNode* SubgraphNode = Cast<UPCGSubgraphNode>(InNode))
-	{
-		SubgraphNode->OnNodeStructuralSettingsChangedDelegate.AddUObject(this, &UPCGGraph::OnStructuralSettingsChanged);
-	}
-
+	InNode->OnNodeChangedDelegate.AddUObject(this, &UPCGGraph::OnNodeChanged);
 	NotifyGraphChanged(/*bIsStructural=*/true);
 #endif
 }
@@ -272,13 +237,7 @@ void UPCGGraph::OnNodeRemoved(UPCGNode* InNode)
 #if WITH_EDITOR
 	if (InNode)
 	{
-		InNode->OnNodeSettingsChangedDelegate.RemoveAll(this);
-		
-		if (UPCGSubgraphNode* SubgraphNode = Cast<UPCGSubgraphNode>(InNode))
-		{
-			SubgraphNode->OnNodeStructuralSettingsChangedDelegate.RemoveAll(this);
-		}
-
+		InNode->OnNodeChangedDelegate.RemoveAll(this);
 		NotifyGraphChanged(/*bIsStructural=*/true);
 	}
 #endif
@@ -297,29 +256,24 @@ UPCGNode* UPCGGraph::AddLabeledEdge(UPCGNode* From, const FName& InboundLabel, U
 		return To;
 	}
 
-	if (!From->HasOutLabel(InboundLabel))
+	UPCGPin* FromPin = From->GetOutputPin(InboundLabel);
+
+	if(!FromPin)
 	{
 		UE_LOG(LogPCG, Error, TEXT("From node %s does not have the %s label"), *From->GetName(), *InboundLabel.ToString());
+		return To;
 	}
 
-	if (!To->HasInLabel(OutboundLabel))
+	UPCGPin* ToPin = To->GetInputPin(OutboundLabel);
+
+	if(!ToPin)
 	{
 		UE_LOG(LogPCG, Error, TEXT("To node %s does not have the %s label"), *To->GetName(), *OutboundLabel.ToString());
+		return To;
 	}
 
-	Modify();
-	From->Modify();
-	To->Modify();
-
 	// Create edge
-	UPCGEdge* Edge = NewObject<UPCGEdge>(this);
-	Edge->InboundLabel = InboundLabel;
-	Edge->InboundNode = From;
-	Edge->OutboundLabel = OutboundLabel;
-	Edge->OutboundNode = To;
-
-	From->OutboundEdges.Add(Edge);
-	To->InboundEdges.Add(Edge);
+	FromPin->AddEdgeTo(ToPin);
 	
 #if WITH_EDITOR
 	NotifyGraphChanged(/*bIsStructural=*/true);
@@ -360,16 +314,14 @@ void UPCGGraph::RemoveNode(UPCGNode* InNode)
 
 	Modify();
 
-	for (int32 EdgeIndex = InNode->InboundEdges.Num() - 1; EdgeIndex >= 0; --EdgeIndex)
+	for (UPCGPin* InputPin : InNode->InputPins)
 	{
-		UPCGEdge* Edge = InNode->InboundEdges[EdgeIndex];
-		Edge->BreakEdge();
+		InputPin->BreakAllEdges();
 	}
 
-	for (int32 EdgeIndex = InNode->OutboundEdges.Num() - 1; EdgeIndex >= 0; --EdgeIndex)
+	for (UPCGPin* OutputPin : InNode->OutputPins)
 	{
-		UPCGEdge* Edge = InNode->OutboundEdges[EdgeIndex];
-		Edge->BreakEdge();
+		OutputPin->BreakAllEdges();
 	}
 
 	Nodes.Remove(InNode);
@@ -384,17 +336,10 @@ bool UPCGGraph::RemoveEdge(UPCGNode* From, const FName& FromLabel, UPCGNode* To,
 		return false;
 	}
 
-	bool bChanged = false;
+	UPCGPin* OutPin = From->GetOutputPin(FromLabel);
+	UPCGPin* InPin = To->GetInputPin(ToLabel);
 
-	if (UPCGEdge* Edge = From->GetOutboundEdge(FromLabel, To, ToLabel))
-	{
-		Edge->BreakEdge();
-		bChanged = true;
-	}
-	else
-	{
-		UE_LOG(LogPCG, Warning, TEXT("Edge does not exist with labels %s and %s"), *FromLabel.ToString(), *ToLabel.ToString());
-	}
+	const bool bChanged = OutPin && OutPin->BreakEdgeTo(InPin);
 
 #if WITH_EDITOR
 	if (bChanged)
@@ -409,11 +354,11 @@ bool UPCGGraph::RemoveEdge(UPCGNode* From, const FName& FromLabel, UPCGNode* To,
 bool UPCGGraph::RemoveAllInboundEdges(UPCGNode* InNode)
 {
 	check(InNode);
-	const bool bChanged = !InNode->InboundEdges.IsEmpty();
+	bool bChanged = false;
 
-	for (int32 i = InNode->InboundEdges.Num() - 1; i >= 0; --i)
+	for (UPCGPin* InputPin : InNode->InputPins)
 	{
-		InNode->InboundEdges[i]->BreakEdge();
+		bChanged |= InputPin->BreakAllEdges();
 	}
 
 #if WITH_EDITOR
@@ -429,11 +374,10 @@ bool UPCGGraph::RemoveAllInboundEdges(UPCGNode* InNode)
 bool UPCGGraph::RemoveAllOutboundEdges(UPCGNode* InNode)
 {
 	check(InNode);
-	const bool bChanged = !InNode->OutboundEdges.IsEmpty();
-
-	for (int32 i = InNode->OutboundEdges.Num() - 1; i >= 0; --i)
+	bool bChanged = false;
+	for (UPCGPin* OutputPin : InNode->OutputPins)
 	{
-		InNode->OutboundEdges[i]->BreakEdge();
+		bChanged |= OutputPin->BreakAllEdges();
 	}
 
 #if WITH_EDITOR
@@ -451,13 +395,9 @@ bool UPCGGraph::RemoveInboundEdges(UPCGNode* InNode, const FName& InboundLabel)
 	check(InNode);
 	bool bChanged = false;
 
-	for (int32 i = InNode->InboundEdges.Num() - 1; i >= 0; --i)
+	if (UPCGPin* InputPin = InNode->GetInputPin(InboundLabel))
 	{
-		if (InNode->InboundEdges[i]->OutboundLabel == InboundLabel)
-		{
-			InNode->InboundEdges[i]->BreakEdge();
-			bChanged = true;
-		}
+		bChanged = InputPin->BreakAllEdges();
 	}
 
 #if WITH_EDITOR
@@ -475,13 +415,9 @@ bool UPCGGraph::RemoveOutboundEdges(UPCGNode* InNode, const FName& OutboundLabel
 	check(InNode);
 	bool bChanged = false;
 
-	for (int32 i = InNode->OutboundEdges.Num() - 1; i >= 0; --i)
+	if (UPCGPin* OutputPin = InNode->GetOutputPin(OutboundLabel))
 	{
-		if (InNode->OutboundEdges[i]->InboundLabel == OutboundLabel)
-		{
-			InNode->OutboundEdges[i]->BreakEdge();
-			bChanged = true;
-		}
+		bChanged = OutputPin->BreakAllEdges();
 	}
 
 #if WITH_EDITOR
@@ -499,12 +435,7 @@ void UPCGGraph::PreNodeUndo(UPCGNode* InPCGNode)
 {
 	if (InPCGNode)
 	{
-		InPCGNode->OnNodeSettingsChangedDelegate.RemoveAll(this);
-
-		if (UPCGSubgraphNode* SubgraphNode = Cast<UPCGSubgraphNode>(InPCGNode))
-		{
-			SubgraphNode->OnNodeStructuralSettingsChangedDelegate.RemoveAll(this);
-		}
+		InPCGNode->OnNodeChangedDelegate.RemoveAll(this);
 	}
 }
 
@@ -512,12 +443,7 @@ void UPCGGraph::PostNodeUndo(UPCGNode* InPCGNode)
 {
 	if (InPCGNode)
 	{
-		InPCGNode->OnNodeSettingsChangedDelegate.AddUObject(this, &UPCGGraph::OnSettingsChanged);
-
-		if (UPCGSubgraphNode* SubgraphNode = Cast<UPCGSubgraphNode>(InPCGNode))
-		{
-			SubgraphNode->OnNodeStructuralSettingsChangedDelegate.AddUObject(this, &UPCGGraph::OnStructuralSettingsChanged);
-		}
+		InPCGNode->OnNodeChangedDelegate.AddUObject(this, &UPCGGraph::OnNodeChanged);
 	}
 }
 #endif
@@ -581,16 +507,11 @@ void UPCGGraph::NotifyGraphChanged(bool bIsStructural)
 	OnGraphChangedDelegate.Broadcast(this, /*bIsStructural=*/bIsStructural);
 }
 
-void UPCGGraph::OnSettingsChanged(UPCGNode* InNode, bool bSettingsChanged)
+void UPCGGraph::OnNodeChanged(UPCGNode* InNode, EPCGChangeType ChangeType)
 {
-	if (bSettingsChanged)
+	if((ChangeType & ~EPCGChangeType::Cosmetic) != EPCGChangeType::None)
 	{
-		NotifyGraphChanged(/*bIsStructural=*/false);
+		NotifyGraphChanged(/*bIsStructural=*/(ChangeType & (EPCGChangeType::Structural | EPCGChangeType::Edge)) != EPCGChangeType::None);
 	}
 }
-
-void UPCGGraph::OnStructuralSettingsChanged(UPCGNode* InNode)
-{
-	NotifyGraphChanged(/*bIsStructural=*/true);
-}
 #endif // WITH_EDITOR
diff --git a/Engine/Plugins/Experimental/PCG/Source/PCG/Private/PCGInputOutputSettings.cpp b/Engine/Plugins/Experimental/PCG/Source/PCG/Private/PCGInputOutputSettings.cpp
--- a/Engine/Plugins/Experimental/PCG/Source/PCG/Private/PCGInputOutputSettings.cpp
+++ b/Engine/Plugins/Experimental/PCG/Source/PCG/Private/PCGInputOutputSettings.cpp
@@ -1,6 +1,7 @@
 // Copyright Epic Games, Inc. All Rights Reserved.
 
 #include "PCGInputOutputSettings.h"
+#include "PCGCommon.h"
 
 bool FPCGInputOutputElement::ExecuteInternal(FPCGContext* Context) const
 {
@@ -12,11 +13,11 @@ bool FPCGInputOutputElement::ExecuteInternal(FPCGContext* Context) const
 UPCGGraphInputOutputSettings::UPCGGraphInputOutputSettings(const FObjectInitializer& ObjectInitializer)
 	: Super(ObjectInitializer)
 {
-	StaticInLabels.Add(PCGInputOutputConstants::DefaultInLabel);
+	StaticInLabels.Add(PCGPinConstants::DefaultInputLabel);
 	StaticAdvancedInLabels.Add(PCGInputOutputConstants::DefaultInputLabel);
 	StaticAdvancedInLabels.Add(PCGInputOutputConstants::DefaultActorLabel);
 	StaticAdvancedInLabels.Add(PCGInputOutputConstants::DefaultOriginalActorLabel);
 	StaticAdvancedInLabels.Add(PCGInputOutputConstants::DefaultExcludedActorsLabel);
 	
-	StaticOutLabels.Add(PCGInputOutputConstants::DefaultOutLabel);
+	StaticOutLabels.Add(PCGPinConstants::DefaultOutputLabel);
 }
\ No newline at end of file
diff --git a/Engine/Plugins/Experimental/PCG/Source/PCG/Private/PCGNode.cpp b/Engine/Plugins/Experimental/PCG/Source/PCG/Private/PCGNode.cpp
--- a/Engine/Plugins/Experimental/PCG/Source/PCG/Private/PCGNode.cpp
+++ b/Engine/Plugins/Experimental/PCG/Source/PCG/Private/PCGNode.cpp
@@ -4,6 +4,9 @@
 #include "PCGEdge.h"
 #include "PCGGraph.h"
 #include "PCGSettings.h"
+#include "PCGSubgraph.h"
+
+#include "Algo/Transform.h"
 
 UPCGNode::UPCGNode(const FObjectInitializer& ObjectInitializer)
 	: Super(ObjectInitializer)
@@ -34,14 +37,79 @@ void UPCGNode::PostLoad()
 #if WITH_EDITOR
 void UPCGNode::ApplyDeprecation()
 {
+	UpdatePins();
+
+	UPCGPin* DefaultOutputPin = OutputPins.IsEmpty() ? nullptr : OutputPins[0];
 	for (TObjectPtr<UPCGNode> OutboundNode : OutboundNodes_DEPRECATED)
 	{
-		if (!HasEdgeTo(OutboundNode))
+		OutboundNode->UpdatePins();
+		UPCGPin* OtherNodeInputPin = OutboundNode->InputPins.IsEmpty() ? nullptr : OutboundNode->InputPins[0];
+
+		if (DefaultOutputPin && OtherNodeInputPin)
+		{
+			DefaultOutputPin->AddEdgeTo(OtherNodeInputPin);
+		}
+		else
 		{
-			AddEdgeTo(NAME_None, OutboundNode, NAME_None);
+			UE_LOG(LogPCG, Error, TEXT("Unable to apply deprecation on outbound nodes"));
 		}
 	}
 	OutboundNodes_DEPRECATED.Reset();
+
+	// Deprecated edges -> pins & edges
+	// Inbound edges will be taken care of by other nodes outbounds
+	InboundEdges_DEPRECATED.Reset();
+
+	for (UPCGEdge* OutboundEdge : OutboundEdges_DEPRECATED)
+	{
+		check(OutboundEdge->InboundNode_DEPRECATED == this);
+		check(OutboundEdge->OutboundNode_DEPRECATED);
+
+		UPCGPin* OutputPin = nullptr;
+		if (OutboundEdge->InboundLabel_DEPRECATED == NAME_None)
+		{
+			OutputPin = OutputPins.IsEmpty() ? nullptr : OutputPins[0];
+		}
+		else
+		{
+			OutputPin = GetOutputPin(OutboundEdge->InboundLabel_DEPRECATED);
+		}
+
+		if (!OutputPin)
+		{
+			UE_LOG(LogPCG, Error, TEXT("Unable to apply deprecation on outbound edge on node %s - can't find output pin %s"), *GetFName().ToString(), *OutboundEdge->InboundLabel_DEPRECATED.ToString());
+			continue;
+		}
+
+		UPCGNode* OtherNode = OutboundEdge->OutboundNode_DEPRECATED;
+		if (!OtherNode)
+		{
+			UE_LOG(LogPCG, Error, TEXT("Unable to apply deprecation on outbound edge on node %s - can't find other node"), *GetFName().ToString());
+			continue;
+		}
+
+		OtherNode->UpdatePins();
+
+		UPCGPin* OtherNodeInputPin = nullptr;
+		if (OutboundEdge->OutboundLabel_DEPRECATED == NAME_None)
+		{
+			OtherNodeInputPin = OtherNode->InputPins.IsEmpty() ? nullptr : OtherNode->InputPins[0];
+		}
+		else
+		{
+			OtherNodeInputPin = OtherNode->GetInputPin(OutboundEdge->OutboundLabel_DEPRECATED);
+		}
+
+		if (OtherNodeInputPin)
+		{
+			OutputPin->AddEdgeTo(OtherNodeInputPin);
+		}
+		else
+		{
+			UE_LOG(LogPCG, Error, TEXT("Unable to apply deprecation on outbound edge on node %s output pin %s - can't find node %s input pin %s"), *GetFName().ToString(), *OutboundEdge->InboundLabel_DEPRECATED.ToString(), *OtherNode->GetFName().ToString(), *OutboundEdge->OutboundLabel_DEPRECATED.ToString());
+		}
+	}
+	OutboundEdges_DEPRECATED.Reset();
 }
 #endif
 
@@ -141,120 +209,138 @@ FName UPCGNode::GetNodeTitle() const
 	return TEXT("Unnamed node");
 }
 
-bool UPCGNode::HasInLabel(const FName& Label) const
-{
-	return DefaultSettings && DefaultSettings->HasInLabel(Label);
-}
-
-bool UPCGNode::HasOutLabel(const FName& Label) const
-{
-	return DefaultSettings && DefaultSettings->HasOutLabel(Label);
-}
-
 TArray<FName> UPCGNode::InLabels() const
 {
-	return DefaultSettings ? DefaultSettings->InLabels() : TArray<FName>();
+	TArray<FName> Labels;
+	for (UPCGPin* InputPin : InputPins)
+	{
+		Labels.Add(InputPin->Label);
+	}
+
+	return Labels;
 }
 
 TArray<FName> UPCGNode::OutLabels() const
 {
-	return DefaultSettings ? DefaultSettings->OutLabels() : TArray<FName>();
-}
+	TArray<FName> Labels;
+	for (UPCGPin* OutputPin : OutputPins)
+	{
+		Labels.Add(OutputPin->Label);
+	}
 
-bool UPCGNode::HasDefaultInLabel() const
-{
-	return !DefaultSettings || DefaultSettings->HasDefaultInLabel();
+	return Labels;
 }
 
-bool UPCGNode::HasDefaultOutLabel() const
+UPCGPin* UPCGNode::GetInputPin(const FName& Label)
 {
-	return !DefaultSettings || DefaultSettings->HasDefaultOutLabel();
+	for (UPCGPin* InputPin : InputPins)
+	{
+		if (InputPin->Label == Label)
+		{
+			return InputPin;
+		}
+	}
+
+	return nullptr;
 }
 
-bool UPCGNode::IsInputPinConnected(const FName& Label) const
+const UPCGPin* UPCGNode::GetInputPin(const FName& Label) const
 {
-	for (const UPCGEdge* InboundEdge : InboundEdges)
+	for (const UPCGPin* InputPin : InputPins)
 	{
-		if (InboundEdge->OutboundLabel == Label)
+		if (InputPin->Label == Label)
 		{
-			return true;
+			return InputPin;
 		}
 	}
 
-	return false;
+	return nullptr;
 }
 
-bool UPCGNode::IsOutputPinConnected(const FName& Label) const
+UPCGPin* UPCGNode::GetOutputPin(const FName& Label)
 {
-	for (const UPCGEdge* OutboundEdge : OutboundEdges)
+	for (UPCGPin* OutputPin : OutputPins)
 	{
-		if (OutboundEdge->InboundLabel == Label)
+		if (OutputPin->Label == Label)
 		{
-			return true;
+			return OutputPin;
 		}
 	}
 
-	return false;
+	return nullptr;
 }
 
-void UPCGNode::SetDefaultSettings(TObjectPtr<UPCGSettings> InSettings)
+const UPCGPin* UPCGNode::GetOutputPin(const FName& Label) const
 {
-#if WITH_EDITOR
-	const bool bDifferentSettings = (DefaultSettings != InSettings);
-	if (bDifferentSettings && DefaultSettings)
+	for (const UPCGPin* OutputPin : OutputPins)
 	{
-		DefaultSettings->OnSettingsChangedDelegate.RemoveAll(this);
+		if (OutputPin->Label == Label)
+		{
+			return OutputPin;
+		}
 	}
-#endif
 
-	DefaultSettings = InSettings;
+	return nullptr;
+}
 
-#if WITH_EDITOR
-	if (bDifferentSettings && DefaultSettings)
+bool UPCGNode::IsInputPinConnected(const FName& Label) const
+{
+	if (const UPCGPin* InputPin = GetInputPin(Label))
 	{
-		DefaultSettings->OnSettingsChangedDelegate.AddUObject(this, &UPCGNode::OnSettingsChanged);
+		return InputPin->IsConnected();
+	}
+	else
+	{
+		return false;
 	}
-#endif
 }
 
-UPCGEdge* UPCGNode::GetOutboundEdge(const FName& FromLabel, UPCGNode* To, const FName& ToLabel)
+bool UPCGNode::IsOutputPinConnected(const FName& Label) const
 {
-	for (UPCGEdge* OutboundEdge : OutboundEdges)
+	if (const UPCGPin* OutputPin = GetOutputPin(Label))
 	{
-		if (OutboundEdge->OutboundNode == To &&
-			OutboundEdge->InboundLabel == FromLabel &&
-			OutboundEdge->OutboundLabel == ToLabel)
-		{
-			return OutboundEdge;
-		}
+		return OutputPin->IsConnected();
+	}
+	else
+	{
+		return false;
 	}
-
-	return nullptr;
 }
 
-bool UPCGNode::HasEdgeTo(UPCGNode* InNode) const
+bool UPCGNode::HasInboundEdges() const
 {
-	for (const UPCGEdge* OutboundEdge : OutboundEdges)
+	for (const UPCGPin* InputPin : InputPins)
 	{
-		if (OutboundEdge->OutboundNode == InNode)
+		for (const UPCGEdge* InboundEdge : InputPin->Edges)
 		{
-			return true;
+			if (InboundEdge->IsValid())
+			{
+				return true;
+			}
 		}
 	}
 
 	return false;
 }
 
-void UPCGNode::RemoveInboundEdge(UPCGEdge* InEdge)
+void UPCGNode::SetDefaultSettings(TObjectPtr<UPCGSettings> InSettings)
 {
-	Modify();
-	InboundEdges.Remove(InEdge);
-}
+#if WITH_EDITOR
+	const bool bDifferentSettings = (DefaultSettings != InSettings);
+	if (bDifferentSettings && DefaultSettings)
+	{
+		DefaultSettings->OnSettingsChangedDelegate.RemoveAll(this);
+	}
+#endif
 
-void UPCGNode::RemoveOutboundEdge(UPCGEdge* InEdge)
-{
-	Modify();
-	OutboundEdges.Remove(InEdge);
+	DefaultSettings = InSettings;
+
+#if WITH_EDITOR
+	if (bDifferentSettings && DefaultSettings)
+	{
+		DefaultSettings->OnSettingsChangedDelegate.AddUObject(this, &UPCGNode::OnSettingsChanged);
+	}
+#endif
 }
 
 #if WITH_EDITOR
@@ -282,20 +368,120 @@ void UPCGNode::PostEditChangeProperty(struct FPropertyChangedEvent& PropertyChan
 		if (DefaultSettings)
 		{
 			DefaultSettings->OnSettingsChangedDelegate.AddUObject(this, &UPCGNode::OnSettingsChanged);
-			OnSettingsChanged(DefaultSettings);
+			UpdatePins();
+			OnNodeChangedDelegate.Broadcast(this, ((Cast<UPCGBaseSubgraphSettings>(DefaultSettings) ? EPCGChangeType::Structural : EPCGChangeType::None) | EPCGChangeType::Settings));
 		}
 	}
 	else if (PropertyChangedEvent.Property && PropertyChangedEvent.Property->GetFName() == GET_MEMBER_NAME_CHECKED(UPCGNode, NodeTitle))
 	{
-		OnNodeSettingsChangedDelegate.Broadcast(this, /*bSettingsChanged=*/false);
+		OnNodeChangedDelegate.Broadcast(this, EPCGChangeType::Cosmetic);
 	}
 }
 
-void UPCGNode::OnSettingsChanged(UPCGSettings* InSettings)
+bool UPCGNode::UpdatePins()
+{
+	if (!DefaultSettings)
+	{
+		bool bChanged = !InputPins.IsEmpty() || !OutputPins.IsEmpty();
+
+		if (bChanged)
+		{
+			Modify();
+		}
+
+		InputPins.Reset();
+		OutputPins.Reset();
+		return bChanged;
+	}
+
+	TArray<FName> InboundLabels = DefaultSettings->InLabels();
+	TArray<FName> OutboundLabels = DefaultSettings->OutLabels();
+
+	auto UpdatePins = [this](TArray<UPCGPin*>& Pins, const TArray<FName>& Labels)
+	{
+		TArray<FName> OldLabels;
+		Algo::Transform(Pins, OldLabels, [](UPCGPin* Pin) { return Pin->Label; });
+
+		auto GetUnmatched = [](const TArray<FName>& InA, const TArray<FName>& InB)
+		{
+			TArray<FName> Unmatched;
+			for (const FName& A : InA)
+			{
+				if (!InB.Contains(A))
+				{
+					Unmatched.Add(A);
+				}
+			}
+
+			return Unmatched;
+		};
+
+		TArray<FName> UnmatchedLabels = GetUnmatched(Labels, OldLabels);
+		TArray<FName> UnmatchedOldLabels = GetUnmatched(OldLabels, Labels);
+
+		const bool bRenameUnmatchedPin = UnmatchedLabels.Num() == 1 && UnmatchedOldLabels.Num() == 1;
+		if (bRenameUnmatchedPin)
+		{
+			for (UPCGPin* Pin : Pins)
+			{
+				if (Pin->Label == UnmatchedOldLabels[0])
+				{
+					Pin->Modify();
+					Pin->Label = UnmatchedLabels[0];
+				}
+			}
+
+			return false;
+		}
+		else
+		{
+			if (!UnmatchedOldLabels.IsEmpty() || !UnmatchedLabels.IsEmpty())
+			{
+				Modify();
+			}
+
+			bool bAppliedEdgeChanges = false;
+
+			// Remove old pins
+			if (UnmatchedOldLabels.Num() > 0)
+			{
+				for (int32 PinIndex = Pins.Num() - 1; PinIndex >= 0; --PinIndex)
+				{
+					if (UnmatchedOldLabels.Contains(Pins[PinIndex]->Label))
+					{
+						bAppliedEdgeChanges |= Pins[PinIndex]->BreakAllEdges();
+						Pins.RemoveAt(PinIndex);
+					}
+				}
+			}
+
+			// Add new pins
+			for (const FName& UnmatchedLabel : UnmatchedLabels)
+			{
+				const int32 InsertIndex = Labels.IndexOfByKey(UnmatchedLabel);
+				UPCGPin* NewPin = NewObject<UPCGPin>(this);
+				NewPin->Node = this;
+				NewPin->Label = UnmatchedLabel;
+				Pins.Insert(NewPin, InsertIndex);
+			}
+
+			return bAppliedEdgeChanges;
+		}
+	};
+
+	bool bChanged = false;
+	bChanged |= UpdatePins(InputPins, InboundLabels);
+	bChanged |= UpdatePins(OutputPins, OutboundLabels);
+
+	return bChanged;
+}
+
+void UPCGNode::OnSettingsChanged(UPCGSettings* InSettings, EPCGChangeType ChangeType)
 {
 	if (InSettings == DefaultSettings)
 	{
-		OnNodeSettingsChangedDelegate.Broadcast(this, /*bSettingsChanged=*/true);
+		const bool bUpdatedPins = UpdatePins();
+		OnNodeChangedDelegate.Broadcast(this, ((bUpdatedPins ? EPCGChangeType::Edge : EPCGChangeType::None) | ChangeType));
 	}
 }
 
diff --git a/Engine/Plugins/Experimental/PCG/Source/PCG/Private/PCGPin.cpp b/Engine/Plugins/Experimental/PCG/Source/PCG/Private/PCGPin.cpp
new file mode 100644
--- /dev/null
+++ b/Engine/Plugins/Experimental/PCG/Source/PCG/Private/PCGPin.cpp
@@ -0,0 +1,105 @@
+// Copyright Epic Games, Inc. All Rights Reserved.
+
+#include "PCGPin.h"
+#include "PCGEdge.h"
+#include "PCGNode.h"
+
+UPCGPin::UPCGPin(const FObjectInitializer& ObjectInitializer)
+	: Super(ObjectInitializer)
+{
+	SetFlags(RF_Transactional);
+}
+
+bool UPCGPin::AddEdgeTo(UPCGPin* OtherPin)
+{
+	if (OtherPin)
+	{
+		for (UPCGEdge* Edge : Edges)
+		{
+			if (Edge->GetOtherPin(this) == OtherPin)
+			{
+				return false;
+			}
+		}
+	}
+
+	Modify();
+	OtherPin->Modify();
+
+	UPCGEdge* NewEdge = Edges.Add_GetRef(NewObject<UPCGEdge>(this));
+	OtherPin->Edges.Add(NewEdge);
+
+	NewEdge->InputPin = this;
+	NewEdge->OutputPin = OtherPin;
+
+	return true;
+}
+
+bool UPCGPin::BreakEdgeTo(UPCGPin* OtherPin)
+{
+	for (UPCGEdge* Edge : Edges)
+	{
+		if (Edge->GetOtherPin(this) == OtherPin)
+		{
+			Modify();
+			OtherPin->Modify();
+
+			ensure(OtherPin->Edges.Remove(Edge));
+			Edges.Remove(Edge);
+			return true;
+		}
+	}
+
+	return false;
+}
+
+bool UPCGPin::BreakAllEdges()
+{
+	bool bChanged = false;
+
+	if (!Edges.IsEmpty())
+	{
+		Modify();
+	}
+
+	for (UPCGEdge* Edge : Edges)
+	{
+		if (UPCGPin* OtherPin = Edge->GetOtherPin(this))
+		{
+			OtherPin->Modify();
+			ensure(OtherPin->Edges.Remove(Edge));
+			bChanged = true;
+		}
+	}
+
+	Edges.Reset();
+
+	return bChanged;
+}
+
+bool UPCGPin::IsConnected() const
+{
+	for (const UPCGEdge* Edge : Edges)
+	{
+		if (Edge->IsValid())
+		{
+			return true;
+		}
+	}
+
+	return false;
+}
+
+int32 UPCGPin::EdgeCount() const
+{
+	int32 EdgeNum = 0;
+	for (const UPCGEdge* Edge : Edges)
+	{
+		if (Edge->IsValid())
+		{
+			++EdgeNum;
+		}
+	}
+
+	return EdgeNum;
+}
\ No newline at end of file
diff --git a/Engine/Plugins/Experimental/PCG/Source/PCG/Private/PCGSettings.cpp b/Engine/Plugins/Experimental/PCG/Source/PCG/Private/PCGSettings.cpp
--- a/Engine/Plugins/Experimental/PCG/Source/PCG/Private/PCGSettings.cpp
+++ b/Engine/Plugins/Experimental/PCG/Source/PCG/Private/PCGSettings.cpp
@@ -43,14 +43,14 @@ uint32 UPCGSettings::GetCrc32() const
 	return Ar.Crc32(const_cast<UPCGSettings*>(this));
 }
 
-bool UPCGSettings::HasInLabel(const FName& Label) const
+TArray<FName> UPCGSettings::InLabels() const
 {
-	return (Label == NAME_None && HasDefaultInLabel()) || InLabels().Contains(Label);
+	return { PCGPinConstants::DefaultInputLabel };
 }
 
-bool UPCGSettings::HasOutLabel(const FName& Label) const
+TArray<FName> UPCGSettings::OutLabels() const
 {
-	return (Label == NAME_None && HasDefaultOutLabel()) || OutLabels().Contains(Label);
+	return { PCGPinConstants::DefaultOutputLabel };
 }
 
 FPCGElementPtr UPCGSettings::GetElement() const
@@ -79,7 +79,7 @@ UPCGNode* UPCGSettings::CreateNode() const
 void UPCGSettings::PostEditChangeProperty(FPropertyChangedEvent& PropertyChangedEvent)
 {
 	Super::PostEditChangeProperty(PropertyChangedEvent);
-	OnSettingsChangedDelegate.Broadcast(this);
+	OnSettingsChangedDelegate.Broadcast(this, EPCGChangeType::Settings);
 }
 
 void UPCGSettings::DirtyCache()
diff --git a/Engine/Plugins/Experimental/PCG/Source/PCG/Private/PCGSubgraph.cpp b/Engine/Plugins/Experimental/PCG/Source/PCG/Private/PCGSubgraph.cpp
--- a/Engine/Plugins/Experimental/PCG/Source/PCG/Private/PCGSubgraph.cpp
+++ b/Engine/Plugins/Experimental/PCG/Source/PCG/Private/PCGSubgraph.cpp
@@ -72,66 +72,11 @@ void UPCGBaseSubgraphSettings::OnSubgraphChanged(UPCGGraph* InGraph, bool bIsStr
 {
 	if (InGraph == GetSubgraph())
 	{
-		if (bIsStructural)
-		{
-			OnStructuralSettingsChangedDelegate.Broadcast(this);
-		}
-		else
-		{
-			OnSettingsChangedDelegate.Broadcast(this);
-		}
+		OnSettingsChangedDelegate.Broadcast(this, ((bIsStructural ? EPCGChangeType::Structural : EPCGChangeType::None) | EPCGChangeType::Settings));
 	}
 }
 #endif // WITH_EDITOR
 
-bool UPCGBaseSubgraphSettings::HasInLabel(const FName& Label) const
-{
-	if (UPCGGraph* Subgraph = GetSubgraph())
-	{
-		return Subgraph->GetInputNode()->HasInLabel(Label);
-	}
-	else
-	{
-		return false;
-	}
-}
-
-bool UPCGBaseSubgraphSettings::HasOutLabel(const FName& Label) const
-{
-	if (UPCGGraph* Subgraph = GetSubgraph())
-	{
-		return Subgraph->GetOutputNode()->HasOutLabel(Label);
-	}
-	else
-	{
-		return false;
-	}
-}
-
-bool UPCGBaseSubgraphSettings::HasDefaultInLabel() const
-{
-	if (UPCGGraph* Subgraph = GetSubgraph())
-	{
-		return Subgraph->GetInputNode()->HasDefaultInLabel();
-	}
-	else
-	{
-		return Super::HasDefaultInLabel();
-	}
-}
-
-bool UPCGBaseSubgraphSettings::HasDefaultOutLabel() const
-{
-	if (UPCGGraph* Subgraph = GetSubgraph())
-	{
-		return Subgraph->GetOutputNode()->HasDefaultOutLabel();
-	}
-	else
-	{
-		return Super::HasDefaultOutLabel();
-	}
-}
-
 TArray<FName> UPCGBaseSubgraphSettings::InLabels() const
 {
 	if (UPCGGraph* Subgraph = GetSubgraph())
@@ -191,72 +136,6 @@ TObjectPtr<UPCGGraph> UPCGSubgraphNode::GetSubgraph() const
 	return Settings ? Settings->Subgraph : nullptr;
 }
 
-void UPCGSubgraphNode::PostLoad()
-{
-	Super::PostLoad();
-
-#if WITH_EDITOR
-	if (UPCGSubgraphSettings* SubgraphSettings = Cast<UPCGSubgraphSettings>(DefaultSettings))
-	{
-		SubgraphSettings->OnStructuralSettingsChangedDelegate.AddUObject(this, &UPCGSubgraphNode::OnStructuralSettingsChanged);
-	}
-#endif
-}
-
-void UPCGSubgraphNode::BeginDestroy()
-{
-#if WITH_EDITOR
-	if (UPCGSubgraphSettings* SubgraphSettings = Cast<UPCGSubgraphSettings>(DefaultSettings))
-	{
-		SubgraphSettings->OnStructuralSettingsChangedDelegate.RemoveAll(this);
-	}
-#endif
-
-	Super::BeginDestroy();
-}
-
-#if WITH_EDITOR
-void UPCGSubgraphNode::PreEditChange(FProperty* PropertyAboutToChange)
-{
-	if (PropertyAboutToChange && PropertyAboutToChange->GetFName() == GET_MEMBER_NAME_CHECKED(UPCGNode, DefaultSettings))
-	{
-		if (UPCGSubgraphSettings* SubgraphSettings = Cast<UPCGSubgraphSettings>(DefaultSettings))
-		{
-			SubgraphSettings->OnStructuralSettingsChangedDelegate.RemoveAll(this);
-		}
-	}
-
-	Super::PreEditChange(PropertyAboutToChange);
-}
-
-void UPCGSubgraphNode::PostEditChangeProperty(struct FPropertyChangedEvent& PropertyChangedEvent)
-{
-	// Implementation note:
-	// We must process structural changes before the parent class' otherwise the graph might be rescheduled with its tasks
-	// before it is appropriately dirtied (which will trigger a recompilation)
-	if (PropertyChangedEvent.Property && PropertyChangedEvent.Property->GetFName() == GET_MEMBER_NAME_CHECKED(UPCGNode, DefaultSettings))
-	{
-		if (UPCGSubgraphSettings* SubgraphSettings = Cast<UPCGSubgraphSettings>(DefaultSettings))
-		{
-			SubgraphSettings->OnStructuralSettingsChangedDelegate.AddUObject(this, &UPCGSubgraphNode::OnStructuralSettingsChanged);
-
-			// Changing the default settings should trigger immediately a structural change
-			OnStructuralSettingsChanged(DefaultSettings);
-		}
-	}
-
-	Super::PostEditChangeProperty(PropertyChangedEvent);
-}
-
-void UPCGSubgraphNode::OnStructuralSettingsChanged(UPCGSettings* InSettings)
-{
-	if (InSettings == DefaultSettings)
-	{
-		OnNodeStructuralSettingsChangedDelegate.Broadcast(this);
-	}
-}
-#endif // WITH_EDITOR
-
 FPCGContext* FPCGSubgraphElement::Initialize(const FPCGDataCollection& InputData, UPCGComponent* SourceComponent, const UPCGNode* Node)
 {
 	FPCGSubgraphContext* Context = new FPCGSubgraphContext();
diff --git a/Engine/Plugins/Experimental/PCG/Source/PCG/Public/Elements/PCGExecuteBlueprint.h b/Engine/Plugins/Experimental/PCG/Source/PCG/Public/Elements/PCGExecuteBlueprint.h
--- a/Engine/Plugins/Experimental/PCG/Source/PCG/Public/Elements/PCGExecuteBlueprint.h
+++ b/Engine/Plugins/Experimental/PCG/Source/PCG/Public/Elements/PCGExecuteBlueprint.h
@@ -129,8 +129,6 @@ public:
 	virtual FName AdditionalTaskName() const override;
 	virtual TArray<FName> InLabels() const override;
 	virtual TArray<FName> OutLabels() const override;
-	virtual bool HasDefaultInLabel() const override;
-	virtual bool HasDefaultOutLabel() const override;
 
 protected:
 	virtual FPCGElementPtr CreateElement() const override;
diff --git a/Engine/Plugins/Experimental/PCG/Source/PCG/Public/Elements/PCGPointFilter.h b/Engine/Plugins/Experimental/PCG/Source/PCG/Public/Elements/PCGPointFilter.h
--- a/Engine/Plugins/Experimental/PCG/Source/PCG/Public/Elements/PCGPointFilter.h
+++ b/Engine/Plugins/Experimental/PCG/Source/PCG/Public/Elements/PCGPointFilter.h
@@ -72,8 +72,6 @@ public:
 
 	virtual TArray<FName> InLabels() const override;
 	virtual TArray<FName> OutLabels() const override;
-	virtual bool HasDefaultInLabel() const override { return false; }
-	virtual bool HasDefaultOutLabel() const override { return false; }
 
 protected:
 	virtual FPCGElementPtr CreateElement() const override;
diff --git a/Engine/Plugins/Experimental/PCG/Source/PCG/Public/PCGCommon.h b/Engine/Plugins/Experimental/PCG/Source/PCG/Public/PCGCommon.h
new file mode 100644
--- /dev/null
+++ b/Engine/Plugins/Experimental/PCG/Source/PCG/Public/PCGCommon.h
@@ -0,0 +1,23 @@
+// Copyright Epic Games, Inc. All Rights Reserved.
+
+#pragma once
+
+#include "PCGCommon.generated.h"
+
+UENUM(meta = (Bitflags))
+enum class EPCGChangeType : uint8
+{
+	None = 0,
+	Cosmetic = 1 << 0,
+	Settings = 1 << 1,
+	Edge = 1 << 2,
+	Node = 1 << 3,
+	Structural = 1 << 4
+};
+ENUM_CLASS_FLAGS(EPCGChangeType);
+
+namespace PCGPinConstants
+{
+	const FName DefaultInputLabel = TEXT("In");
+	const FName DefaultOutputLabel = TEXT("Out");
+}
\ No newline at end of file
diff --git a/Engine/Plugins/Experimental/PCG/Source/PCG/Public/PCGEdge.h b/Engine/Plugins/Experimental/PCG/Source/PCG/Public/PCGEdge.h
--- a/Engine/Plugins/Experimental/PCG/Source/PCG/Public/PCGEdge.h
+++ b/Engine/Plugins/Experimental/PCG/Source/PCG/Public/PCGEdge.h
@@ -7,6 +7,7 @@
 #include "PCGEdge.generated.h"
 
 class UPCGNode;
+class UPCGPin;
 
 UCLASS(ClassGroup = (Procedural))
 class PCG_API UPCGEdge : public UObject
@@ -14,17 +15,25 @@ class PCG_API UPCGEdge : public UObject
 	GENERATED_BODY()
 
 public:
-	UPROPERTY(BlueprintReadOnly, VisibleAnywhere, Category = Edge)
-	FName InboundLabel = NAME_None;
+	UPROPERTY()
+	FName InboundLabel_DEPRECATED = NAME_None;
 
-	UPROPERTY(BlueprintReadOnly, VisibleAnywhere, Category = Edge)
-	TObjectPtr<UPCGNode> InboundNode;
+	UPROPERTY()
+	TObjectPtr<UPCGNode> InboundNode_DEPRECATED;
 
-	UPROPERTY(BlueprintReadOnly, VisibleAnywhere, Category = Edge)
-	FName OutboundLabel = NAME_None;
+	UPROPERTY()
+	FName OutboundLabel_DEPRECATED = NAME_None;
 
-	UPROPERTY(BlueprintReadOnly, VisibleAnywhere, Category = Edge)
-	TObjectPtr<UPCGNode> OutboundNode;
+	UPROPERTY()
+	TObjectPtr<UPCGNode> OutboundNode_DEPRECATED;
 
-	void BreakEdge();
-};
\ No newline at end of file
+	UPROPERTY()
+	TObjectPtr<UPCGPin> InputPin;
+
+	UPROPERTY()
+	TObjectPtr<UPCGPin> OutputPin;
+
+	bool IsValid() const;
+	UPCGPin* GetOtherPin(const UPCGPin* Pin);
+	const UPCGPin* GetOtherPin(const UPCGPin* Pin) const;
+};
diff --git a/Engine/Plugins/Experimental/PCG/Source/PCG/Public/PCGGraph.h b/Engine/Plugins/Experimental/PCG/Source/PCG/Public/PCGGraph.h
--- a/Engine/Plugins/Experimental/PCG/Source/PCG/Public/PCGGraph.h
+++ b/Engine/Plugins/Experimental/PCG/Source/PCG/Public/PCGGraph.h
@@ -104,8 +104,7 @@ protected:
 #if WITH_EDITOR
 private:
 	void NotifyGraphChanged(bool bIsStructural);
-	void OnSettingsChanged(UPCGNode* InNode, bool bSettingsChanged);
-	void OnStructuralSettingsChanged(UPCGNode* InNode);
+	void OnNodeChanged(UPCGNode* InNode, EPCGChangeType ChangeType);
 
 	int32 GraphChangeNotificationsDisableCounter = 0;
 	bool bDelayedChangeNotification = false;
diff --git a/Engine/Plugins/Experimental/PCG/Source/PCG/Public/PCGInputOutputSettings.h b/Engine/Plugins/Experimental/PCG/Source/PCG/Public/PCGInputOutputSettings.h
--- a/Engine/Plugins/Experimental/PCG/Source/PCG/Public/PCGInputOutputSettings.h
+++ b/Engine/Plugins/Experimental/PCG/Source/PCG/Public/PCGInputOutputSettings.h
@@ -9,13 +9,10 @@
 
 namespace PCGInputOutputConstants
 {
-	const FName DefaultInLabel = TEXT("In");
 	const FName DefaultInputLabel = TEXT("Input");
 	const FName DefaultActorLabel = TEXT("Actor");
 	const FName DefaultOriginalActorLabel = TEXT("OriginalActor");
 	const FName DefaultExcludedActorsLabel = TEXT("ExcludedActors");
-
-	const FName DefaultOutLabel = TEXT("Out");
 }
 
 class PCG_API FPCGInputOutputElement : public FSimplePCGElement
@@ -39,24 +36,19 @@ public:
 	virtual EPCGSettingsType GetType() const override { return EPCGSettingsType::InputOutput; }
 #endif
 
-	virtual bool HasInLabel(const FName& Label) const override { return Label == NAME_None || StaticLabels().Contains(Label) || (bShowAdvancedPins && StaticAdvancedLabels().Contains(Label)) || PinLabels.Contains(Label); }
-	virtual bool HasOutLabel(const FName& Label) const override { return Label == NAME_None || StaticLabels().Contains(Label) || (bShowAdvancedPins && StaticAdvancedLabels().Contains(Label)) || PinLabels.Contains(Label); }
-	virtual bool HasDefaultInLabel() const override { return false; }
-	virtual bool HasDefaultOutLabel() const override { return false; }
+	virtual TArray<FName> InLabels() const override { return GetLabels(); }
+	virtual TArray<FName> OutLabels() const override { return GetLabels(); }
 
-	virtual TArray<FName> InLabels() const override
-	{
-		TArray<FName> Labels = StaticLabels();
-		if (bShowAdvancedPins)
-		{
-			Labels += StaticAdvancedLabels();
-		}
-		
-		Labels += PinLabels.Array();
-		return Labels;
-	}
+	void SetInput(bool bInIsInput) { bIsInput = bInIsInput; }
 
-	virtual TArray<FName> OutLabels() const override
+protected:
+	virtual FPCGElementPtr CreateElement() const override { return MakeShared<FPCGInputOutputElement>(); }
+	// ~End UPCGSettings interface
+
+	const TArray<FName>& StaticLabels() const { return bIsInput ? StaticInLabels : StaticOutLabels; }
+	const TArray<FName>& StaticAdvancedLabels() const { return bIsInput ? StaticAdvancedInLabels : StaticAdvancedOutLabels; }
+
+	TArray<FName> GetLabels() const
 	{
 		TArray<FName> Labels = StaticLabels();
 		if (bShowAdvancedPins)
@@ -68,15 +60,6 @@ public:
 		return Labels;
 	}
 
-	void SetInput(bool bInIsInput) { bIsInput = bInIsInput; }
-
-protected:
-	virtual FPCGElementPtr CreateElement() const override { return MakeShared<FPCGInputOutputElement>(); }
-	// ~End UPCGSettings interface
-
-	const TArray<FName>& StaticLabels() const { return bIsInput ? StaticInLabels : StaticOutLabels; }
-	const TArray<FName>& StaticAdvancedLabels() const { return bIsInput ? StaticAdvancedInLabels : StaticAdvancedOutLabels; }
-
 protected:
 	UPROPERTY(EditAnywhere, BlueprintReadWrite, Category = "Settings|Input")
 	TSet<FName> PinLabels;
diff --git a/Engine/Plugins/Experimental/PCG/Source/PCG/Public/PCGNode.h b/Engine/Plugins/Experimental/PCG/Source/PCG/Public/PCGNode.h
--- a/Engine/Plugins/Experimental/PCG/Source/PCG/Public/PCGNode.h
+++ b/Engine/Plugins/Experimental/PCG/Source/PCG/Public/PCGNode.h
@@ -3,7 +3,8 @@
 #pragma once
 
 #include "CoreMinimal.h"
-
+#include "PCGCommon.h"
+#include "PCGPin.h"
 #include "PCGNode.generated.h"
 
 class UPCGSettings;
@@ -12,7 +13,7 @@ class UPCGEdge;
 class IPCGElement;
 
 #if WITH_EDITOR
-	DECLARE_MULTICAST_DELEGATE_TwoParams(FOnPCGNodeSettingsChanged, UPCGNode*, bool);
+	DECLARE_MULTICAST_DELEGATE_TwoParams(FOnPCGNodeChanged, UPCGNode*, EPCGChangeType);
 #endif
 
 UCLASS(ClassGroup = (Procedural))
@@ -52,22 +53,12 @@ public:
 	/** Returns the node title, based either on the current node label, or defaulted to its settings */
 	FName GetNodeTitle() const;
 
-	/** Returns true if one of the input pins matches with the given name */
-	bool HasInLabel(const FName& Label) const;
-
-	/** Returns true if one of the output pins matches with the given name */
-	bool HasOutLabel(const FName& Label) const;
-
 	/** Returns all the input labels */
 	TArray<FName> InLabels() const;
 
 	/** Returns all the output labels */
 	TArray<FName> OutLabels() const;
 
-	/** Returns true if we want the default In/Out labels */
-	bool HasDefaultInLabel() const;
-	bool HasDefaultOutLabel() const;
-
 	/** Returns true if the input pin is connected */
 	bool IsInputPinConnected(const FName& Label) const;
 
@@ -77,7 +68,14 @@ public:
 	/** Changes the default settings in the node */
 	void SetDefaultSettings(TObjectPtr<UPCGSettings> InSettings);
 
-	const TArray<TObjectPtr<UPCGEdge>>& GetOutboundEdges() const { return OutboundEdges; }
+	UPCGPin* GetInputPin(const FName& Label);
+	const UPCGPin* GetInputPin(const FName& Label) const;
+	UPCGPin* GetOutputPin(const FName& Label);
+	const UPCGPin* GetOutputPin(const FName& Label) const;
+	bool HasInboundEdges() const;
+
+	const TArray<TObjectPtr<UPCGPin>>& GetInputPins() const { return InputPins; }
+	const TArray<TObjectPtr<UPCGPin>>& GetOutputPins() const { return OutputPins; }
 
 	/** Note: do not set this property directly from code, use SetDefaultSettings instead */
 	UPROPERTY(BlueprintReadWrite, EditAnywhere, Category = Node, meta=(EditInline))
@@ -92,7 +90,7 @@ public:
 #endif // WITH_EDITORONLY_DATA
 
 #if WITH_EDITOR
-	FOnPCGNodeSettingsChanged OnNodeSettingsChangedDelegate;
+	FOnPCGNodeChanged OnNodeChangedDelegate;
 #endif
 
 #if WITH_EDITORONLY_DATA
@@ -104,26 +102,27 @@ public:
 #endif // WITH_EDITORONLY_DATA
 
 protected:
-	bool HasEdgeTo(UPCGNode* InNode) const;
-
-	void RemoveInboundEdge(UPCGEdge* InEdge);
-	void RemoveOutboundEdge(UPCGEdge* InEdge);
-	UPCGEdge* GetOutboundEdge(const FName& FromLabel, UPCGNode* To, const FName& ToLabel);
-
 #if WITH_EDITOR
+	bool UpdatePins();
 	virtual void PreEditChange(FProperty* PropertyAboutToChange) override;
 	virtual void PostEditChangeProperty(struct FPropertyChangedEvent& PropertyChangedEvent) override;
-	void OnSettingsChanged(UPCGSettings* InSettings);
+	void OnSettingsChanged(UPCGSettings* InSettings, EPCGChangeType ChangeType);
 #endif
 
 	UPROPERTY()
 	TArray<TObjectPtr<UPCGNode>> OutboundNodes_DEPRECATED;
 
-	UPROPERTY(BlueprintReadOnly, VisibleAnywhere, Category = Node, TextExportTransient)
-	TArray<TObjectPtr<UPCGEdge>> InboundEdges;
+	UPROPERTY(TextExportTransient)
+	TArray<TObjectPtr<UPCGEdge>> InboundEdges_DEPRECATED;
+
+	UPROPERTY(TextExportTransient)
+	TArray<TObjectPtr<UPCGEdge>> OutboundEdges_DEPRECATED;
+
+	UPROPERTY(BlueprintReadOnly, VisibleAnywhere, Category = Node)
+	TArray<TObjectPtr<UPCGPin>> InputPins;
 
-	UPROPERTY(BlueprintReadOnly, VisibleAnywhere, Category = Node, TextExportTransient)
-	TArray<TObjectPtr<UPCGEdge>> OutboundEdges;
+	UPROPERTY(BlueprintReadOnly, VisibleAnywhere, Category = Node)
+	TArray<TObjectPtr<UPCGPin>> OutputPins;
 
 	// TODO: add this information:
 	// - Ability to run on non-game threads (here or element)
diff --git a/Engine/Plugins/Experimental/PCG/Source/PCG/Public/PCGPin.h b/Engine/Plugins/Experimental/PCG/Source/PCG/Public/PCGPin.h
new file mode 100644
--- /dev/null
+++ b/Engine/Plugins/Experimental/PCG/Source/PCG/Public/PCGPin.h
@@ -0,0 +1,34 @@
+// Copyright Epic Games, Inc. All Rights Reserved.
+
+#pragma once
+
+#include "CoreMinimal.h"
+
+#include "PCGPin.generated.h"
+
+class UPCGNode;
+class UPCGEdge;
+
+UCLASS(ClassGroup = (Procedural))
+class PCG_API UPCGPin : public UObject
+{
+	GENERATED_BODY()
+
+public:
+	UPCGPin(const FObjectInitializer& ObjectInitializer);
+
+	UPROPERTY()
+	TObjectPtr<UPCGNode> Node = nullptr;
+
+	UPROPERTY(EditAnywhere, Category = Settings)
+	FName Label = NAME_None;
+
+	UPROPERTY(TextExportTransient)
+	TArray<TObjectPtr<UPCGEdge>> Edges;
+
+	bool AddEdgeTo(UPCGPin* OtherPin);
+	bool BreakEdgeTo(UPCGPin* OtherPin);
+	bool BreakAllEdges();
+	bool IsConnected() const;
+	int32 EdgeCount() const;
+};
\ No newline at end of file
diff --git a/Engine/Plugins/Experimental/PCG/Source/PCG/Public/PCGSettings.h b/Engine/Plugins/Experimental/PCG/Source/PCG/Public/PCGSettings.h
--- a/Engine/Plugins/Experimental/PCG/Source/PCG/Public/PCGSettings.h
+++ b/Engine/Plugins/Experimental/PCG/Source/PCG/Public/PCGSettings.h
@@ -2,6 +2,7 @@
 
 #pragma once
 
+#include "PCGCommon.h"
 #include "PCGData.h"
 #include "PCGElement.h"
 #include "PCGDebug.h"
@@ -38,7 +39,7 @@ enum class EPCGSettingsType : uint8
 };
 
 #if WITH_EDITOR
-	DECLARE_MULTICAST_DELEGATE_OneParam(FOnPCGSettingsChanged, UPCGSettings*);
+	DECLARE_MULTICAST_DELEGATE_TwoParams(FOnPCGSettingsChanged, UPCGSettings*, EPCGChangeType);
 #endif
 
 /**
@@ -54,12 +55,8 @@ public:
 	/*virtual*/ FPCGElementPtr GetElement() const;
 	virtual UPCGNode* CreateNode() const;
 
-	virtual bool HasInLabel(const FName& Label) const;
-	virtual bool HasOutLabel(const FName& Label) const;
-	virtual TArray<FName> InLabels() const { return TArray<FName>(); }
-	virtual TArray<FName> OutLabels() const { return TArray<FName>(); }
-	virtual bool HasDefaultInLabel() const { return true; }
-	virtual bool HasDefaultOutLabel() const { return true; }
+	virtual TArray<FName> InLabels() const;
+	virtual TArray<FName> OutLabels() const;
 	
 	bool operator==(const UPCGSettings& Other) const;
 	uint32 GetCrc32() const;
diff --git a/Engine/Plugins/Experimental/PCG/Source/PCG/Public/PCGSubgraph.h b/Engine/Plugins/Experimental/PCG/Source/PCG/Public/PCGSubgraph.h
--- a/Engine/Plugins/Experimental/PCG/Source/PCG/Public/PCGSubgraph.h
+++ b/Engine/Plugins/Experimental/PCG/Source/PCG/Public/PCGSubgraph.h
@@ -9,11 +9,6 @@ class UPCGGraph;
 
 #include "PCGSubgraph.generated.h"
 
-#if WITH_EDITOR
-	DECLARE_MULTICAST_DELEGATE_OneParam(FOnPCGStructuralSettingsChanged, UPCGSettings*);
-	DECLARE_MULTICAST_DELEGATE_OneParam(FOnPCGNodeStructuralSettingsChanged, UPCGNode*);
-#endif
-
 UCLASS(Abstract)
 class PCG_API UPCGBaseSubgraphSettings : public UPCGSettings
 {
@@ -34,10 +29,6 @@ protected:
 	virtual void GetTrackedActorTags(FPCGTagToSettingsMap& OutTagToSettings) const override;
 #endif
 
-	virtual bool HasInLabel(const FName& Label) const override;
-	virtual bool HasOutLabel(const FName& Label) const override;
-	virtual bool HasDefaultInLabel() const override;
-	virtual bool HasDefaultOutLabel() const override;
 	virtual TArray<FName> InLabels() const override;
 	virtual TArray<FName> OutLabels() const override;
 	//~End UPCGSettings interface
@@ -47,11 +38,6 @@ protected:
 
 	virtual bool IsStructuralProperty(const FName& InPropertyName) const { return false; }
 #endif
-
-public:
-#if WITH_EDITOR
-	FOnPCGStructuralSettingsChanged OnStructuralSettingsChangedDelegate;
-#endif
 };
 
 UCLASS(BlueprintType, ClassGroup=(Procedural))
@@ -106,25 +92,9 @@ class PCG_API UPCGSubgraphNode : public UPCGBaseSubgraphNode
 	GENERATED_BODY()
 
 public:
-#if WITH_EDITOR
-	FOnPCGNodeStructuralSettingsChanged OnNodeStructuralSettingsChangedDelegate;
-#endif
-
 	/** ~Begin UPCGBaseSubgraphNode interface */
-	TObjectPtr<UPCGGraph> GetSubgraph() const override;
+	virtual TObjectPtr<UPCGGraph> GetSubgraph() const override;
 	/** ~End UPCGBaseSubgraphNode interface */
-
-protected:	
-	/** ~Begin UObject interface */
-	virtual void PostLoad() override;
-	virtual void BeginDestroy() override;
-	/** ~End UObject interface */
-
-#if WITH_EDITOR
-	virtual void PreEditChange(FProperty* PropertyAboutToChange) override;
-	virtual void PostEditChangeProperty(struct FPropertyChangedEvent& PropertyChangedEvent) override;
-	void OnStructuralSettingsChanged(UPCGSettings* InSettings);
-#endif
 };
 
 struct PCG_API FPCGSubgraphContext : public FPCGContext
diff --git a/Engine/Plugins/Experimental/PCG/Source/PCGEditor/Private/PCGEditorGraph.cpp b/Engine/Plugins/Experimental/PCG/Source/PCGEditor/Private/PCGEditorGraph.cpp
--- a/Engine/Plugins/Experimental/PCG/Source/PCGEditor/Private/PCGEditorGraph.cpp
+++ b/Engine/Plugins/Experimental/PCG/Source/PCGEditor/Private/PCGEditorGraph.cpp
@@ -45,25 +45,97 @@ void UPCGEditorGraph::InitFromNodeGraph(UPCGGraph* InPCGGraph)
 	{
 		UPCGNode* PCGNode = NodeLookupIt.Key;
 		UPCGEditorGraphNodeBase* GraphNode = NodeLookupIt.Value;
+		CreateLinks(GraphNode, /*bCreateInbound=*/false, /*bCreateOutbound=*/true, NodeLookup);
+	}
+}
+
+void UPCGEditorGraph::CreateLinks(UPCGEditorGraphNodeBase* GraphNode, bool bCreateInbound, bool bCreateOutbound)
+{
+	check(GraphNode);
+	// Build graph node to pcg node map
+	TMap<UPCGNode*, UPCGEditorGraphNodeBase*> GraphNodeToPCGNodeMap;
+
+	for (const TObjectPtr<UEdGraphNode>& EdGraphNode : Nodes)
+	{
+		if (UPCGEditorGraphNodeBase* SomeGraphNode = Cast<UPCGEditorGraphNodeBase>(EdGraphNode))
+		{
+			GraphNodeToPCGNodeMap.Add(SomeGraphNode->GetPCGNode(), SomeGraphNode);
+		}
+	}
+
+	// Forward the call
+	CreateLinks(GraphNode, bCreateInbound, bCreateOutbound, GraphNodeToPCGNodeMap);
+}
+
+void UPCGEditorGraph::CreateLinks(UPCGEditorGraphNodeBase* GraphNode, bool bCreateInbound, bool bCreateOutbound, const TMap<UPCGNode*, UPCGEditorGraphNodeBase*>& GraphNodeToPCGNodeMap)
+{
+	check(GraphNode);
+	const UPCGNode* PCGNode = GraphNode->GetPCGNode();
+	check(PCGNode);
+
+	if (bCreateInbound)
+	{
+		for (const UPCGPin* InputPin : PCGNode->GetInputPins())
+		{
+			UEdGraphPin* InPin = GraphNode->FindPin(InputPin->Label, EEdGraphPinDirection::EGPD_Input);
+
+			if (!InPin)
+			{
+				continue;
+			}
 
-		for (UPCGEdge* OutboundEdge : PCGNode->GetOutboundEdges())
+			for (const UPCGEdge* InboundEdge : InputPin->Edges)
+			{
+				if (!InboundEdge->IsValid())
+				{
+					continue;
+				}
+
+				const UPCGNode* InboundNode = InboundEdge->InputPin->Node;
+				if (UPCGEditorGraphNodeBase* const* ConnectedGraphNode = GraphNodeToPCGNodeMap.Find(InboundNode))
+				{
+					if (UEdGraphPin* OutPin = (*ConnectedGraphNode)->FindPin(InboundEdge->InputPin->Label, EEdGraphPinDirection::EGPD_Output))
+					{
+						OutPin->MakeLinkTo(InPin);
+					}
+					else
+					{
+						continue;
+					}
+				}
+			}
+		}
+	}
+
+	if (bCreateOutbound)
+	{
+		for (const UPCGPin* OutputPin : PCGNode->GetOutputPins())
 		{
-			const FName OutPinName = (OutboundEdge->InboundLabel == NAME_None) ? TEXT("Out") : OutboundEdge->InboundLabel;
-			UEdGraphPin* OutPin = GraphNode->FindPin(OutPinName, EEdGraphPinDirection::EGPD_Output);
+			UEdGraphPin* OutPin = GraphNode->FindPin(OutputPin->Label, EEdGraphPinDirection::EGPD_Output);
 
 			if (!OutPin)
 			{
 				continue;
 			}
-			
-			UPCGNode* OutboundNode = OutboundEdge->OutboundNode;
-			if (UPCGEditorGraphNodeBase** ConnectedGraphNode = NodeLookup.Find(OutboundNode))
+
+			for (const UPCGEdge* OutboundEdge : OutputPin->Edges)
 			{
-				const FName InPinName = (OutboundEdge->OutboundLabel == NAME_None) ? TEXT("In") : OutboundEdge->OutboundLabel;
+				if (!OutboundEdge->IsValid())
+				{
+					continue;
+				}
 
-				if(UEdGraphPin* InPin = (*ConnectedGraphNode)->FindPin(InPinName, EEdGraphPinDirection::EGPD_Input))
+				const UPCGNode* OutboundNode = OutboundEdge->OutputPin->Node;
+				if (UPCGEditorGraphNodeBase* const* ConnectedGraphNode = GraphNodeToPCGNodeMap.Find(OutboundNode))
 				{
-					OutPin->MakeLinkTo(InPin);
+					if (UEdGraphPin* InPin = (*ConnectedGraphNode)->FindPin(OutboundEdge->OutputPin->Label, EEdGraphPinDirection::EGPD_Input))
+					{
+						OutPin->MakeLinkTo(InPin);
+					}
+					else
+					{
+						continue;
+					}
 				}
 			}
 		}
diff --git a/Engine/Plugins/Experimental/PCG/Source/PCGEditor/Private/PCGEditorGraph.h b/Engine/Plugins/Experimental/PCG/Source/PCGEditor/Private/PCGEditorGraph.h
--- a/Engine/Plugins/Experimental/PCG/Source/PCGEditor/Private/PCGEditorGraph.h
+++ b/Engine/Plugins/Experimental/PCG/Source/PCGEditor/Private/PCGEditorGraph.h
@@ -8,7 +8,9 @@
 
 #include "PCGEditorGraph.generated.h"
 
+class UPCGNode;
 class UPCGGraph;
+class UPCGEditorGraphNodeBase;
 
 UCLASS()
 class UPCGEditorGraph : public UEdGraph
@@ -19,8 +21,14 @@ public:
 	/** Initialize the editor graph from a PCGGraph */
 	void InitFromNodeGraph(UPCGGraph* InPCGGraph);
 
+	/** Creates the links for a given node */
+	void CreateLinks(UPCGEditorGraphNodeBase* InGraphNode, bool bCreateInbound, bool bCreateOutbound);
+
 	UPCGGraph* GetPCGGraph() { return PCGGraph; }
 
+protected:
+	void CreateLinks(UPCGEditorGraphNodeBase* InGraphNode, bool bCreateInbound, bool bCreateOutbound, const TMap<UPCGNode*, UPCGEditorGraphNodeBase*>& InGraphNodeToPCGNodeMap);
+
 private:
 	UPROPERTY()
 	TObjectPtr<UPCGGraph> PCGGraph = nullptr;
diff --git a/Engine/Plugins/Experimental/PCG/Source/PCGEditor/Private/PCGEditorGraphNode.cpp b/Engine/Plugins/Experimental/PCG/Source/PCGEditor/Private/PCGEditorGraphNode.cpp
--- a/Engine/Plugins/Experimental/PCG/Source/PCGEditor/Private/PCGEditorGraphNode.cpp
+++ b/Engine/Plugins/Experimental/PCG/Source/PCGEditor/Private/PCGEditorGraphNode.cpp
@@ -43,29 +43,16 @@ void UPCGEditorGraphNode::GetNodeContextMenuActions(UToolMenu* Menu, class UGrap
 
 void UPCGEditorGraphNode::AllocateDefaultPins()
 {
-	if (!PCGNode || PCGNode->HasDefaultOutLabel())
-	{
-		CreatePin(EEdGraphPinDirection::EGPD_Output, NAME_None, FName(TEXT("Out")));
-	}
-
 	if (PCGNode)
 	{
-		for (const FName& OutLabel : PCGNode->OutLabels())
+		for (const UPCGPin* InputPin : PCGNode->GetInputPins())
 		{
-			CreatePin(EEdGraphPinDirection::EGPD_Output, NAME_None, OutLabel);
+			CreatePin(EEdGraphPinDirection::EGPD_Input, NAME_None, InputPin->Label);
 		}
-	}
 
-	if (!PCGNode || PCGNode->HasDefaultInLabel())
-	{
-		CreatePin(EEdGraphPinDirection::EGPD_Input, NAME_None, FName(TEXT("In")));
-	}
-
-	if (PCGNode)
-	{
-		for (const FName& InLabel : PCGNode->InLabels())
+		for (const UPCGPin* OutputPin : PCGNode->GetOutputPins())
 		{
-			CreatePin(EEdGraphPinDirection::EGPD_Input, NAME_None, InLabel);
+			CreatePin(EEdGraphPinDirection::EGPD_Output, NAME_None, OutputPin->Label);
 		}
 	}
 }
diff --git a/Engine/Plugins/Experimental/PCG/Source/PCGEditor/Private/PCGEditorGraphNodeBase.cpp b/Engine/Plugins/Experimental/PCG/Source/PCGEditor/Private/PCGEditorGraphNodeBase.cpp
--- a/Engine/Plugins/Experimental/PCG/Source/PCGEditor/Private/PCGEditorGraphNodeBase.cpp
+++ b/Engine/Plugins/Experimental/PCG/Source/PCGEditor/Private/PCGEditorGraphNodeBase.cpp
@@ -24,7 +24,7 @@ void UPCGEditorGraphNodeBase::BeginDestroy()
 {
 	if (PCGNode)
 	{
-		PCGNode->OnNodeSettingsChangedDelegate.RemoveAll(this);
+		PCGNode->OnNodeChangedDelegate.RemoveAll(this);
 	}
 
 	Super::BeginDestroy();
@@ -34,7 +34,7 @@ void UPCGEditorGraphNodeBase::Construct(UPCGNode* InPCGNode, EPCGEditorGraphNode
 {
 	check(InPCGNode);
 	PCGNode = InPCGNode;
-	InPCGNode->OnNodeSettingsChangedDelegate.AddUObject(this, &UPCGEditorGraphNodeBase::OnNodeChanged);
+	InPCGNode->OnNodeChangedDelegate.AddUObject(this, &UPCGEditorGraphNodeBase::OnNodeChanged);
 
 	NodePosX = InPCGNode->PositionX;
 	NodePosY = InPCGNode->PositionY;
@@ -91,26 +91,23 @@ void UPCGEditorGraphNodeBase::AutowireNewNode(UEdGraphPin* FromPin)
 {
 	if (FromPin->Direction == EEdGraphPinDirection::EGPD_Output)
 	{
-		FName InPinName = TEXT("In");
-		if (PCGNode && !PCGNode->HasDefaultInLabel() && !PCGNode->InLabels().IsEmpty())
+		if (PCGNode && PCGNode->GetInputPins().Num() > 0)
 		{
-			InPinName = PCGNode->InLabels()[0];
+			const FName& InPinName = PCGNode->GetInputPins()[0]->Label;
+			UEdGraphPin* ToPin = FindPinChecked(InPinName, EEdGraphPinDirection::EGPD_Input);
+			GetSchema()->TryCreateConnection(FromPin, ToPin);
 		}
-
-		UEdGraphPin* ToPin = FindPinChecked(InPinName, EEdGraphPinDirection::EGPD_Input);
-		GetSchema()->TryCreateConnection(FromPin, ToPin);
 	}
 	else if (FromPin->Direction == EEdGraphPinDirection::EGPD_Input)
 	{
-		FName OutPinName = TEXT("Out");
-		if (PCGNode && !PCGNode->HasDefaultOutLabel() && !PCGNode->OutLabels().IsEmpty())
+		if (PCGNode && PCGNode->GetOutputPins().Num() > 0)
 		{
-			OutPinName = PCGNode->OutLabels()[0];
+			const FName& OutPinName = PCGNode->GetOutputPins()[0]->Label;
+			UEdGraphPin* ToPin = FindPinChecked(OutPinName, EEdGraphPinDirection::EGPD_Output);
+			GetSchema()->TryCreateConnection(FromPin, ToPin);
 		}
-
-		UEdGraphPin* ToPin = FindPinChecked(OutPinName, EEdGraphPinDirection::EGPD_Output);
-		GetSchema()->TryCreateConnection(FromPin, ToPin);
 	}
+
 	NodeConnectionListChanged();
 }
 
@@ -139,21 +136,29 @@ void UPCGEditorGraphNodeBase::PostCopy()
 	}
 }
 
+void UPCGEditorGraphNodeBase::PostPasteNode()
+{
+	bDisableReconstructFromNode = true;
+}
+
 void UPCGEditorGraphNodeBase::PostPaste()
 {
 	if (PCGNode)
 	{
 		RebuildEdgesFromPins();
 
-		PCGNode->OnNodeSettingsChangedDelegate.AddUObject(this, &UPCGEditorGraphNodeBase::OnNodeChanged);
+		PCGNode->OnNodeChangedDelegate.AddUObject(this, &UPCGEditorGraphNodeBase::OnNodeChanged);
 		PCGNode->PositionX = NodePosX;
 		PCGNode->PositionY = NodePosY;
 	}
+
+	bDisableReconstructFromNode = false;
 }
 
 void UPCGEditorGraphNodeBase::RebuildEdgesFromPins()
 {
 	check(PCGNode);
+	check(bDisableReconstructFromNode);
 
 	if (PCGNode->GetGraph())
 	{
@@ -171,10 +176,7 @@ void UPCGEditorGraphNodeBase::RebuildEdgesFromPins()
 
 				if (UPCGNode* ConnectedPCGNode = ConnectedPCGGraphNode->GetPCGNode())
 				{
-					const FName& NodeAPinName = ((Pin->PinName == TEXT("Out") && PCGNode->HasDefaultOutLabel()) ? NAME_None : Pin->PinName);
-					const FName& NodeBPinName = ((ConnectedPin->PinName == TEXT("In") && ConnectedPCGNode->HasDefaultInLabel()) ? NAME_None : ConnectedPin->PinName);
-
-					PCGNode->AddEdgeTo(NodeAPinName, ConnectedPCGNode, NodeBPinName);
+					PCGNode->AddEdgeTo(Pin->PinName, ConnectedPCGNode, ConnectedPin->PinName);
 				}
 			}
 		}
@@ -186,7 +188,7 @@ void UPCGEditorGraphNodeBase::RebuildEdgesFromPins()
 	}
 }
 
-void UPCGEditorGraphNodeBase::OnNodeChanged(UPCGNode* InNode, bool bSettingsChanged)
+void UPCGEditorGraphNodeBase::OnNodeChanged(UPCGNode* InNode, EPCGChangeType ChangeType)
 {
 	if (InNode == PCGNode)
 	{
@@ -216,126 +218,35 @@ void UPCGEditorGraphNodeBase::OnColorPicked(FLinearColor NewColor)
 
 void UPCGEditorGraphNodeBase::ReconstructNode()
 {
-	if (!PCGNode)
+	// In copy-paste cases, we don't want to remove the pins
+	if (bDisableReconstructFromNode)
 	{
 		return;
 	}
 
-	TArray<FName> InLabels = PCGNode->InLabels();
-	TArray<FName> OutLabels = PCGNode->OutLabels();
+	// Remove all current pins
+	TArray<UEdGraphPin*> OldPins = Pins;
 
-	TArray<UEdGraphPin*> InputPins;
-	TArray<UEdGraphPin*> OutputPins;
-
-	// Reallocate the default pins if they had been removed previously
-	if ((NodeType == EPCGEditorGraphNodeType::Output || NodeType == EPCGEditorGraphNodeType::Settings) && PCGNode->HasDefaultInLabel() && !FindPin(TEXT("In"), EEdGraphPinDirection::EGPD_Input))
-	{
-		FCreatePinParams FirstPinParams;
-		FirstPinParams.Index = 0;
-		CreatePin(EEdGraphPinDirection::EGPD_Input, NAME_None, FName(TEXT("In")), FirstPinParams);
-	}
-
-	if ((NodeType == EPCGEditorGraphNodeType::Input || NodeType == EPCGEditorGraphNodeType::Settings) && PCGNode->HasDefaultOutLabel() && !FindPin(TEXT("Out"), EEdGraphPinDirection::EGPD_Output))
+	for (UEdGraphPin* OldPin : OldPins)
 	{
-		FCreatePinParams FirstPinParams;
-		FirstPinParams.Index = 0;
-		CreatePin(EEdGraphPinDirection::EGPD_Output, NAME_None, FName(TEXT("Out")), FirstPinParams);
+		OldPin->BreakAllPinLinks();
+		RemovePin(OldPin);
 	}
+	check(Pins.IsEmpty());
 
-	for (UEdGraphPin* Pin : Pins)
-	{
-		if (Pin->Direction == EEdGraphPinDirection::EGPD_Input && (Pin->PinName != TEXT("In") || !PCGNode->HasDefaultInLabel()))
-		{
-			InputPins.Add(Pin);
-		}
-		else if (Pin->Direction == EEdGraphPinDirection::EGPD_Output && (Pin->PinName != TEXT("Out") || !PCGNode->HasDefaultOutLabel()))
-		{
-			OutputPins.Add(Pin);
-		}
-	}
-
-	auto UpdatePins = [this](TArray<UEdGraphPin*>& InPins, const TArray<FName>& InNames, EEdGraphPinDirection InPinDirection)
-	{
-		TArray<UEdGraphPin*> UnmatchedPins;
-
-		for (UEdGraphPin* InPin : InPins)
-		{
-			if (!InNames.Contains(InPin->PinName))
-			{
-				UnmatchedPins.Add(InPin);
-			}
-		}
-
-		TArray<FName> UnmatchedNames;
-
-		for (const FName& InName : InNames)
-		{
-			if (InPins.FindByPredicate([&InName](UEdGraphPin* InPin) { return InPin->PinName == InName; }) == nullptr)
-			{
-				UnmatchedNames.Add(InName);
-			}
-		}
-
-		// Use cases:
-		// - Removing pin(s) : remove pin(s) & break links
-		// - Remove all pins : remove all pins
-		// - Adding new pin(s) : create new pin(s)
-		// - Renaming a pin : find unmatched pin, unmatched name, rename
-		// - Complete change : match on names
-		if (UnmatchedNames.Num() == 1 && UnmatchedPins.Num() == 1)
-		{
-			UnmatchedPins[0]->PinName = UnmatchedNames[0];
-		}
-		else
-		{
-			if (PCGNode->GetGraph())
-			{
-				PCGNode->GetGraph()->DisableNotificationsForEditor();
-			}
-
-			for (UEdGraphPin* UnmatchedPin : UnmatchedPins)
-			{
-				UnmatchedPin->BreakAllPinLinks();
+	// Generate new pins
+	AllocateDefaultPins();
 
-				if (PCGNode->GetGraph())
-				{
-					if (InPinDirection == EEdGraphPinDirection::EGPD_Input)
-					{
-						PCGNode->GetGraph()->RemoveInboundEdges(PCGNode, (UnmatchedPin->PinName == "In" && PCGNode->HasDefaultInLabel()) ? NAME_None : UnmatchedPin->PinName);
-					}
-					else
-					{
-						PCGNode->GetGraph()->RemoveOutboundEdges(PCGNode, (UnmatchedPin->PinName == "Out" && PCGNode->HasDefaultOutLabel()) ? NAME_None : UnmatchedPin->PinName);
-					}
-				}
-
-				RemovePin(UnmatchedPin);
-			}
-
-			if (PCGNode->GetGraph())
-			{
-				PCGNode->GetGraph()->EnableNotificationsForEditor();
-			}
-
-			for (const FName& UnmatchedName : UnmatchedNames)
-			{
-				CreatePin(InPinDirection, NAME_None, UnmatchedName);
-			}
-		}
-
-		return UnmatchedNames.Num() > 0 || UnmatchedPins.Num() > 0;
-	};
-
-	if (NodeType == EPCGEditorGraphNodeType::Output || NodeType == EPCGEditorGraphNodeType::Settings)
+	// Generate new links
+	// TODO: we should either keep a map in the PCGEditorGraph or do this elsewhere
+	// TODO: this will not work if we have non-PCG nodes in the graph
+	if (PCGNode)
 	{
-		UpdatePins(InputPins, InLabels, EEdGraphPinDirection::EGPD_Input);
+		UPCGEditorGraph* PCGEditorGraph = CastChecked<UPCGEditorGraph>(GetGraph());
+		PCGEditorGraph->CreateLinks(this, /*bCreateInbound=*/true, /*bCreateOutbound=*/true);
 	}
 	
-	if (NodeType == EPCGEditorGraphNodeType::Input || NodeType == EPCGEditorGraphNodeType::Settings)
-	{
-		UpdatePins(OutputPins, OutLabels, EEdGraphPinDirection::EGPD_Output);
-	}
-
+	// Notify editor
 	OnNodeChangedDelegate.ExecuteIfBound();
 }
 
diff --git a/Engine/Plugins/Experimental/PCG/Source/PCGEditor/Private/PCGEditorGraphNodeBase.h b/Engine/Plugins/Experimental/PCG/Source/PCGEditor/Private/PCGEditorGraphNodeBase.h
--- a/Engine/Plugins/Experimental/PCG/Source/PCGEditor/Private/PCGEditorGraphNodeBase.h
+++ b/Engine/Plugins/Experimental/PCG/Source/PCGEditor/Private/PCGEditorGraphNodeBase.h
@@ -6,6 +6,8 @@
 #include "UObject/ObjectMacros.h"
 #include "EdGraph/EdGraphNode.h"
 
+#include "PCGCommon.h"
+
 #include "PCGEditorGraphNodeBase.generated.h"
 
 class UPCGNode;
@@ -38,6 +40,7 @@ public:
 	virtual bool CanCreateUnderSpecifiedSchema(const UEdGraphSchema* Schema) const override;
 	virtual void ReconstructNode() override;
 	virtual FLinearColor GetNodeTitleColor() const;
+	virtual void PostPasteNode() override;
 	// ~End UEdGraphNode interface
 
 	UPCGNode* GetPCGNode() { return PCGNode; }
@@ -50,7 +53,7 @@ public:
 protected:
 	void RebuildEdgesFromPins();
 
-	void OnNodeChanged(UPCGNode* InNode, bool bSettingsChanged);
+	void OnNodeChanged(UPCGNode* InNode, EPCGChangeType ChangeType);
 	void OnPickColor();
 	void OnColorPicked(FLinearColor NewColor);
 
@@ -59,4 +62,6 @@ protected:
 
  	UPROPERTY()
  	EPCGEditorGraphNodeType NodeType = EPCGEditorGraphNodeType::Settings;
+
+	bool bDisableReconstructFromNode = false;
 };
diff --git a/Engine/Plugins/Experimental/PCG/Source/PCGEditor/Private/PCGEditorGraphNodeInput.cpp b/Engine/Plugins/Experimental/PCG/Source/PCGEditor/Private/PCGEditorGraphNodeInput.cpp
--- a/Engine/Plugins/Experimental/PCG/Source/PCGEditor/Private/PCGEditorGraphNodeInput.cpp
+++ b/Engine/Plugins/Experimental/PCG/Source/PCGEditor/Private/PCGEditorGraphNodeInput.cpp
@@ -13,16 +13,11 @@ FText UPCGEditorGraphNodeInput::GetNodeTitle(ENodeTitleType::Type TitleType) con
 
 void UPCGEditorGraphNodeInput::AllocateDefaultPins()
 {
-	if (!PCGNode || PCGNode->HasDefaultOutLabel())
-	{
-		CreatePin(EEdGraphPinDirection::EGPD_Output, NAME_None, FName(TEXT("Out")));
-	}
-
 	if (PCGNode)
 	{
-		for (const FName& OutLabel : PCGNode->OutLabels())
+		for (const UPCGPin* OutPin : PCGNode->GetOutputPins())
 		{
-			CreatePin(EEdGraphPinDirection::EGPD_Output, NAME_None, OutLabel);
+			CreatePin(EEdGraphPinDirection::EGPD_Output, NAME_None, OutPin->Label);
 		}
 	}
 }
diff --git a/Engine/Plugins/Experimental/PCG/Source/PCGEditor/Private/PCGEditorGraphNodeOutput.cpp b/Engine/Plugins/Experimental/PCG/Source/PCGEditor/Private/PCGEditorGraphNodeOutput.cpp
--- a/Engine/Plugins/Experimental/PCG/Source/PCGEditor/Private/PCGEditorGraphNodeOutput.cpp
+++ b/Engine/Plugins/Experimental/PCG/Source/PCGEditor/Private/PCGEditorGraphNodeOutput.cpp
@@ -13,16 +13,11 @@ FText UPCGEditorGraphNodeOutput::GetNodeTitle(ENodeTitleType::Type TitleType) co
 
 void UPCGEditorGraphNodeOutput::AllocateDefaultPins()
 {
-	if (!PCGNode || PCGNode->HasDefaultInLabel())
-	{
-		CreatePin(EEdGraphPinDirection::EGPD_Input, NAME_None, FName(TEXT("In")));
-	}
-
 	if (PCGNode)
 	{
-		for (const FName& InLabel : PCGNode->InLabels())
+		for (const UPCGPin* InputPin : PCGNode->GetInputPins())
 		{
-			CreatePin(EEdGraphPinDirection::EGPD_Input, NAME_None, InLabel);
+			CreatePin(EEdGraphPinDirection::EGPD_Input, NAME_None, InputPin->Label);
 		}
 	}
 }
diff --git a/Engine/Plugins/Experimental/PCG/Source/PCGEditor/Private/PCGEditorGraphSchema.cpp b/Engine/Plugins/Experimental/PCG/Source/PCGEditor/Private/PCGEditorGraphSchema.cpp
--- a/Engine/Plugins/Experimental/PCG/Source/PCGEditor/Private/PCGEditorGraphSchema.cpp
+++ b/Engine/Plugins/Experimental/PCG/Source/PCGEditor/Private/PCGEditorGraphSchema.cpp
@@ -95,10 +95,7 @@ bool UPCGEditorGraphSchema::TryCreateConnection(UEdGraphPin* InA, UEdGraphPin* I
 		UPCGGraph* PCGGraph = PCGNodeA->GetGraph();
 		check(PCGGraph);
 
-		const FName& NodeAPinName = ((A->PinName == TEXT("Out") && PCGNodeA->HasDefaultOutLabel()) ? NAME_None : A->PinName);
-		const FName& NodeBPinName = ((B->PinName == TEXT("In") && PCGNodeB->HasDefaultInLabel()) ? NAME_None : B->PinName);
-
-		PCGGraph->AddLabeledEdge(PCGNodeA, NodeAPinName, PCGNodeB, NodeBPinName);
+		PCGGraph->AddLabeledEdge(PCGNodeA, A->PinName, PCGNodeB, B->PinName);
 	}
 
 	return bModified;
@@ -122,13 +119,11 @@ void UPCGEditorGraphSchema::BreakPinLinks(UEdGraphPin& TargetPin, bool bSendsNod
 
 	if (TargetPin.Direction == EEdGraphPinDirection::EGPD_Input)
 	{
-		const FName& PinName = (TargetPin.PinName == TEXT("In") && PCGNode->HasDefaultInLabel()) ? NAME_None : TargetPin.PinName;
-		PCGGraph->RemoveInboundEdges(PCGNode, PinName);
+		PCGGraph->RemoveInboundEdges(PCGNode, TargetPin.PinName);
 	}
 	else if (TargetPin.Direction == EEdGraphPinDirection::EGPD_Output)
 	{
-		const FName& PinName = (TargetPin.PinName == TEXT("Out") && PCGNode->HasDefaultOutLabel()) ? NAME_None : TargetPin.PinName;
-		PCGGraph->RemoveOutboundEdges(PCGNode, PinName);
+		PCGGraph->RemoveOutboundEdges(PCGNode, TargetPin.PinName);
 	}
 }
 
@@ -148,11 +143,8 @@ void UPCGEditorGraphSchema::BreakSinglePinLink(UEdGraphPin* SourcePin, UEdGraphP
 	UPCGNode* TargetPCGNode = TargetPCGGraphNode->GetPCGNode();
 	check(SourcePCGNode && TargetPCGNode);
 
-	const FName& SourcePinName = (SourcePin->PinName == TEXT("Out") && SourcePCGNode->HasDefaultOutLabel()) ? NAME_None : SourcePin->PinName;
-	const FName& TargetPinName = (TargetPin->PinName == TEXT("In") && TargetPCGNode->HasDefaultInLabel()) ? NAME_None : TargetPin->PinName;
-
 	UPCGGraph* PCGGraph = SourcePCGNode->GetGraph();
-	PCGGraph->RemoveEdge(SourcePCGNode, SourcePinName, TargetPCGNode, TargetPinName);
+	PCGGraph->RemoveEdge(SourcePCGNode, SourcePin->PinName, TargetPCGNode, TargetPin->PinName);
 }
 
 FPCGEditorConnectionDrawingPolicy::FPCGEditorConnectionDrawingPolicy(int32 InBackLayerID, int32 InFrontLayerID, float InZoomFactor, const FSlateRect& InClippingRect, FSlateWindowElementList& InDrawElements, UEdGraph* InGraph)
