Fixed pin issue on new nodes (missing update)

diff --git a/Engine/Plugins/Experimental/PCG/Source/PCG/Private/PCGNode.cpp b/Engine/Plugins/Experimental/PCG/Source/PCG/Private/PCGNode.cpp
--- a/Engine/Plugins/Experimental/PCG/Source/PCG/Private/PCGNode.cpp
+++ b/Engine/Plugins/Experimental/PCG/Source/PCG/Private/PCGNode.cpp
@@ -341,6 +341,8 @@ void UPCGNode::SetDefaultSettings(TObjectPtr<UPCGSettings> InSettings)
 		DefaultSettings->OnSettingsChangedDelegate.AddUObject(this, &UPCGNode::OnSettingsChanged);
 	}
 #endif
+
+	UpdatePins();
 }
 
 #if WITH_EDITOR
