Fixed ordering issue causing the wrong pins to appear on the inputs in new graphs

diff --git a/Engine/Plugins/Experimental/PCG/Source/PCG/Private/PCGGraph.cpp b/Engine/Plugins/Experimental/PCG/Source/PCG/Private/PCGGraph.cpp
--- a/Engine/Plugins/Experimental/PCG/Source/PCG/Private/PCGGraph.cpp
+++ b/Engine/Plugins/Experimental/PCG/Source/PCG/Private/PCGGraph.cpp
@@ -15,12 +15,18 @@ UPCGGraph::UPCGGraph(const FObjectInitializer& ObjectInitializer)
 {
 	InputNode = ObjectInitializer.CreateDefaultSubobject<UPCGNode>(this, TEXT("DefaultInputNode"));
 	InputNode->SetFlags(RF_Transactional);
-	InputNode->SetDefaultSettings(ObjectInitializer.CreateDefaultSubobject<UPCGGraphInputOutputSettings>(this, TEXT("DefaultInputNodeSettings")));
-	Cast<UPCGGraphInputOutputSettings>(InputNode->DefaultSettings)->SetInput(true);
+
+	UPCGGraphInputOutputSettings* InputSettings = ObjectInitializer.CreateDefaultSubobject<UPCGGraphInputOutputSettings>(this, TEXT("DefaultInputNodeSettings"));
+	InputSettings->SetInput(true);
+	InputNode->SetDefaultSettings(InputSettings);
+	
 	OutputNode = ObjectInitializer.CreateDefaultSubobject<UPCGNode>(this, TEXT("DefaultOutputNode"));
 	OutputNode->SetFlags(RF_Transactional);
-	OutputNode->SetDefaultSettings(ObjectInitializer.CreateDefaultSubobject<UPCGGraphInputOutputSettings>(this, TEXT("DefaultOutputNodeSettings")));
-	Cast<UPCGGraphInputOutputSettings>(OutputNode->DefaultSettings)->SetInput(false);
+
+	UPCGGraphInputOutputSettings* OutputSettings = ObjectInitializer.CreateDefaultSubobject<UPCGGraphInputOutputSettings>(this, TEXT("DefaultOutputNodeSettings"));
+	OutputSettings->SetInput(false);
+	OutputNode->SetDefaultSettings(OutputSettings);
+	
 #if WITH_EDITORONLY_DATA
 	OutputNode->PositionX = 200;
 #endif
