Added collision options in static mesh spawner
Updated debug display & spawn actor to have proper default values

diff --git a/Engine/Plugins/Experimental/PCG/Source/PCG/Private/Elements/PCGDebugElement.cpp b/Engine/Plugins/Experimental/PCG/Source/PCG/Private/Elements/PCGDebugElement.cpp
--- a/Engine/Plugins/Experimental/PCG/Source/PCG/Private/Elements/PCGDebugElement.cpp
+++ b/Engine/Plugins/Experimental/PCG/Source/PCG/Private/Elements/PCGDebugElement.cpp
@@ -124,7 +124,12 @@ namespace PCGDebugElement
 				InstanceTransform.SetScale3D(bIsRelative ? InstanceTransform.GetScale3D() * PointScale : FVector(PointScale));
 			}
 
-			UInstancedStaticMeshComponent* ISMC = UPCGActorHelpers::GetOrCreateISMC(TargetActor, Context->SourceComponent, Mesh, Materials);
+			FPCGISMCBuilderParameters Params;
+			Params.Mesh = Mesh;
+			Params.MaterialOverrides = Materials;
+			Params.CollisionProfile = UCollisionProfile::NoCollision_ProfileName;
+
+			UInstancedStaticMeshComponent* ISMC = UPCGActorHelpers::GetOrCreateISMC(TargetActor, Context->SourceComponent, Params);
 			
 			ISMC->ComponentTags.AddUnique(PCGHelpers::DefaultPCGDebugTag);
 			ISMC->NumCustomDataFloats = NumCustomData;
diff --git a/Engine/Plugins/Experimental/PCG/Source/PCG/Private/Elements/PCGSpawnActor.cpp b/Engine/Plugins/Experimental/PCG/Source/PCG/Private/Elements/PCGSpawnActor.cpp
--- a/Engine/Plugins/Experimental/PCG/Source/PCG/Private/Elements/PCGSpawnActor.cpp
+++ b/Engine/Plugins/Experimental/PCG/Source/PCG/Private/Elements/PCGSpawnActor.cpp
@@ -147,7 +147,10 @@ bool FPCGSpawnActorElement::ExecuteInternal(FPCGContext* Context) const
 
 				if (Mesh)
 				{
-					ISMC = UPCGActorHelpers::GetOrCreateISMC(TargetActor, Context->SourceComponent, Mesh);
+					FPCGISMCBuilderParameters Params;
+					Params.Mesh = Mesh;
+
+					ISMC = UPCGActorHelpers::GetOrCreateISMC(TargetActor, Context->SourceComponent, Params);
 					UEngine::CopyPropertiesForUnrelatedObjects(FirstSMC, ISMC);
 				}
 				else
diff --git a/Engine/Plugins/Experimental/PCG/Source/PCG/Private/Elements/PCGStaticMeshSpawner.cpp b/Engine/Plugins/Experimental/PCG/Source/PCG/Private/Elements/PCGStaticMeshSpawner.cpp
--- a/Engine/Plugins/Experimental/PCG/Source/PCG/Private/Elements/PCGStaticMeshSpawner.cpp
+++ b/Engine/Plugins/Experimental/PCG/Source/PCG/Private/Elements/PCGStaticMeshSpawner.cpp
@@ -17,8 +17,12 @@ FPCGElementPtr UPCGStaticMeshSpawnerSettings::CreateElement() const
 
 struct FWeightedMeshAndInstances
 {
-	int Weight = 0;
-	UStaticMesh* Mesh = nullptr;
+	FWeightedMeshAndInstances(int InWeight, const FPCGStaticMeshSpawnerEntry& InEntry)
+		: Weight(InWeight), Entry(InEntry)
+	{}
+
+	int Weight;
+	const FPCGStaticMeshSpawnerEntry& Entry;
 	TArray<FTransform> Instances;
 };
 
@@ -38,11 +42,7 @@ bool FPCGStaticMeshSpawnerElement::ExecuteInternal(FPCGContext* Context) const
 			continue;
 
 		TotalWeight += Entry.Weight;
-
-		FWeightedMeshAndInstances& WeightedEntry = WeightedEntries.Emplace_GetRef();
-		WeightedEntry.Weight = TotalWeight;
-		// Todo: we could likely pre-load these meshes asynchronously in the settings
-		WeightedEntry.Mesh = Entry.Mesh.LoadSynchronous();
+		FWeightedMeshAndInstances& WeightedEntry = WeightedEntries.Emplace_GetRef(TotalWeight, Entry);
 	}
 
 	if (TotalWeight <= 0)
@@ -114,12 +114,20 @@ bool FPCGStaticMeshSpawnerElement::ExecuteInternal(FPCGContext* Context) const
 			// Second, add the instances
 			for (FWeightedMeshAndInstances& Instances : WeightedEntries)
 			{
-				if (!Instances.Mesh || Instances.Instances.Num() == 0)
+				if (!Instances.Entry.Mesh || Instances.Instances.Num() == 0)
 				{
 					continue;
 				}
 
-				UInstancedStaticMeshComponent* ISMC = UPCGActorHelpers::GetOrCreateISMC(TargetActor, Context->SourceComponent, Instances.Mesh);
+				FPCGISMCBuilderParameters Params;
+				// Todo: we could likely pre-load these meshes asynchronously in the settings
+				Params.Mesh = Instances.Entry.Mesh.LoadSynchronous();
+				if (Instances.Entry.bOverrideCollisionProfile)
+				{
+					Params.CollisionProfile = Instances.Entry.CollisionProfile.Name;
+				}
+
+				UInstancedStaticMeshComponent* ISMC = UPCGActorHelpers::GetOrCreateISMC(TargetActor, Context->SourceComponent, Params);
 
 				// TODO: add scaling
 				// TODO: document these arguments
@@ -127,7 +135,7 @@ bool FPCGStaticMeshSpawnerElement::ExecuteInternal(FPCGContext* Context) const
 				ISMC->AddInstances(Instances.Instances, false, true);
 				ISMC->UpdateBounds();
 
-				PCGE_LOG(Verbose, "Added %d instances of %s on actor %s", Instances.Instances.Num(), *Instances.Mesh->GetFName().ToString(), *TargetActor->GetFName().ToString());
+				PCGE_LOG(Verbose, "Added %d instances of %s on actor %s", Instances.Instances.Num(), *Instances.Entry.Mesh->GetFName().ToString(), *TargetActor->GetFName().ToString());
 			}
 		}
 	}
diff --git a/Engine/Plugins/Experimental/PCG/Source/PCG/Private/Helpers/PCGActorHelpers.cpp b/Engine/Plugins/Experimental/PCG/Source/PCG/Private/Helpers/PCGActorHelpers.cpp
--- a/Engine/Plugins/Experimental/PCG/Source/PCG/Private/Helpers/PCGActorHelpers.cpp
+++ b/Engine/Plugins/Experimental/PCG/Source/PCG/Private/Helpers/PCGActorHelpers.cpp
@@ -22,8 +22,11 @@
 #include "ObjectTools.h"
 #endif
 
-UInstancedStaticMeshComponent* UPCGActorHelpers::GetOrCreateISMC(AActor* InTargetActor, const UPCGComponent* InSourceComponent, UStaticMesh* InMesh, const TArray<UMaterialInterface*>& InMaterials)
+UInstancedStaticMeshComponent* UPCGActorHelpers::GetOrCreateISMC(AActor* InTargetActor, const UPCGComponent* InSourceComponent, const FPCGISMCBuilderParameters& InParams)
 {
+	UStaticMesh* InMesh = InParams.Mesh;
+	const TArray<UMaterialInterface*>& InMaterials = InParams.MaterialOverrides;
+
 	check(InTargetActor != nullptr && InMesh != nullptr);
 
 	TArray<UInstancedStaticMeshComponent*> ISMCs;
@@ -38,6 +41,14 @@ UInstancedStaticMeshComponent* UPCGActorHelpers::GetOrCreateISMC(AActor* InTarge
 			// If not provided, we'll make sure that the current materials aren't overriden
 			bool bMaterialsMatched = true;
 
+			// Check basic parameters
+			if (ISMC->Mobility != InParams.Mobility ||
+				(ISMC->bUseDefaultCollision && InParams.CollisionProfile != TEXT("Default")) ||
+				(!ISMC->bUseDefaultCollision && ISMC->GetCollisionProfileName() != InParams.CollisionProfile))
+			{
+				continue;
+			}
+
 			for (int32 MaterialIndex = 0; MaterialIndex < ISMC->GetNumMaterials() && bMaterialsMatched; ++MaterialIndex)
 			{
 				if (MaterialIndex < InMaterials.Num() && InMaterials[MaterialIndex])
@@ -91,9 +102,16 @@ UInstancedStaticMeshComponent* UPCGActorHelpers::GetOrCreateISMC(AActor* InTarge
 
 	ISMC->RegisterComponent();
 	InTargetActor->AddInstanceComponent(ISMC);
-	ISMC->SetMobility(EComponentMobility::Static);
-	// TODO: add option for collision, or use a template
-	ISMC->SetCollisionEnabled(ECollisionEnabled::NoCollision);
+	ISMC->SetMobility(InParams.Mobility);
+	if (InParams.CollisionProfile != TEXT("Default"))
+	{
+		ISMC->SetCollisionProfileName(InParams.CollisionProfile, /*bOverlaps=*/false);
+	}
+	else
+	{
+		ISMC->bUseDefaultCollision = true;
+	}
+	
 	ISMC->AttachToComponent(InTargetActor->GetRootComponent(), FAttachmentTransformRules::KeepWorldTransform);
 
 	if (InSourceComponent)
diff --git a/Engine/Plugins/Experimental/PCG/Source/PCG/Public/Elements/PCGStaticMeshSpawner.h b/Engine/Plugins/Experimental/PCG/Source/PCG/Public/Elements/PCGStaticMeshSpawner.h
--- a/Engine/Plugins/Experimental/PCG/Source/PCG/Public/Elements/PCGStaticMeshSpawner.h
+++ b/Engine/Plugins/Experimental/PCG/Source/PCG/Public/Elements/PCGStaticMeshSpawner.h
@@ -5,6 +5,8 @@
 #include "CoreMinimal.h"
 #include "PCGSettings.h"
 
+#include "Engine/CollisionProfile.h"
+
 #include "PCGStaticMeshSpawner.generated.h"
 
 class UStaticMesh;
@@ -19,6 +21,12 @@ struct PCG_API FPCGStaticMeshSpawnerEntry
 
 	UPROPERTY(BlueprintReadWrite, EditAnywhere, Category = Settings)
 	TSoftObjectPtr<UStaticMesh> Mesh;
+
+	UPROPERTY(BlueprintReadWrite, EditAnywhere, Category = Settings)
+	bool bOverrideCollisionProfile = false;
+
+	UPROPERTY(BlueprintReadWrite, EditAnywhere, Category = Settings)
+	FCollisionProfileName CollisionProfile;
 };
 
 UCLASS(BlueprintType, ClassGroup = (Procedural))
diff --git a/Engine/Plugins/Experimental/PCG/Source/PCG/Public/Helpers/PCGActorHelpers.h b/Engine/Plugins/Experimental/PCG/Source/PCG/Public/Helpers/PCGActorHelpers.h
--- a/Engine/Plugins/Experimental/PCG/Source/PCG/Public/Helpers/PCGActorHelpers.h
+++ b/Engine/Plugins/Experimental/PCG/Source/PCG/Public/Helpers/PCGActorHelpers.h
@@ -5,6 +5,8 @@
 #include "CoreMinimal.h"
 #include "Kismet/BlueprintFunctionLibrary.h"
 #include "Templates/SubclassOf.h"
+#include "Engine/EngineTypes.h"
+#include "Engine/CollisionProfile.h"
 
 #include "PCGActorHelpers.generated.h"
 
@@ -15,13 +17,21 @@ class UPCGComponent;
 class UMaterialInterface;
 class UActorComponent;
 
+struct FPCGISMCBuilderParameters
+{
+	UStaticMesh* Mesh = nullptr;
+	TArray<UMaterialInterface*> MaterialOverrides;
+	EComponentMobility::Type Mobility = EComponentMobility::Static;
+	FName CollisionProfile = TEXT("Default");
+};
+
 UCLASS(BlueprintType)
 class PCG_API UPCGActorHelpers : public UBlueprintFunctionLibrary
 {
 	GENERATED_BODY()
 
 public:
-	static UInstancedStaticMeshComponent* GetOrCreateISMC(AActor* InActor, const UPCGComponent* SourceComponent, UStaticMesh* InMesh, const TArray<UMaterialInterface*>& InMaterials = TArray<UMaterialInterface*>());
+	static UInstancedStaticMeshComponent* GetOrCreateISMC(AActor* InActor, const UPCGComponent* SourceComponent, const FPCGISMCBuilderParameters& Params);
 	static bool DeleteActors(UWorld* World, const TArray<TSoftObjectPtr<AActor>>& ActorsToDelete);
 
 	/**
