Changed the code path on default creation so that the pins are properly allocated objects (fixes save issue)

diff --git a/Engine/Plugins/Experimental/PCG/Source/PCG/Private/PCGGraph.cpp b/Engine/Plugins/Experimental/PCG/Source/PCG/Private/PCGGraph.cpp
--- a/Engine/Plugins/Experimental/PCG/Source/PCG/Private/PCGGraph.cpp
+++ b/Engine/Plugins/Experimental/PCG/Source/PCG/Private/PCGGraph.cpp
@@ -16,16 +16,27 @@ UPCGGraph::UPCGGraph(const FObjectInitializer& ObjectInitializer)
 	InputNode = ObjectInitializer.CreateDefaultSubobject<UPCGNode>(this, TEXT("DefaultInputNode"));
 	InputNode->SetFlags(RF_Transactional);
 
+	// Since pins would be allocated after initializing the input/output nodes, we must make sure to allocate them using the object initializer
+	int NumAllocatedPins = 1;
+	auto PinAllocator = [&ObjectInitializer, &NumAllocatedPins](UPCGNode* Node)
+	{ 
+		FName DefaultPinName = TEXT("DefaultPin");
+		DefaultPinName.SetNumber(NumAllocatedPins++);
+		return ObjectInitializer.CreateDefaultSubobject<UPCGPin>(Node, DefaultPinName);
+	};
+
 	UPCGGraphInputOutputSettings* InputSettings = ObjectInitializer.CreateDefaultSubobject<UPCGGraphInputOutputSettings>(this, TEXT("DefaultInputNodeSettings"));
 	InputSettings->SetInput(true);
-	InputNode->SetDefaultSettings(InputSettings);
+	InputNode->SetDefaultSettings(InputSettings, /*bUpdatePins=*/false);
+	InputNode->UpdatePins(PinAllocator);
 	
 	OutputNode = ObjectInitializer.CreateDefaultSubobject<UPCGNode>(this, TEXT("DefaultOutputNode"));
 	OutputNode->SetFlags(RF_Transactional);
 
 	UPCGGraphInputOutputSettings* OutputSettings = ObjectInitializer.CreateDefaultSubobject<UPCGGraphInputOutputSettings>(this, TEXT("DefaultOutputNodeSettings"));
 	OutputSettings->SetInput(false);
-	OutputNode->SetDefaultSettings(OutputSettings);
+	OutputNode->SetDefaultSettings(OutputSettings, /*bUpdatePins=*/false);
+	OutputNode->UpdatePins(PinAllocator);
 	
 #if WITH_EDITORONLY_DATA
 	OutputNode->PositionX = 200;
diff --git a/Engine/Plugins/Experimental/PCG/Source/PCG/Private/PCGNode.cpp b/Engine/Plugins/Experimental/PCG/Source/PCG/Private/PCGNode.cpp
--- a/Engine/Plugins/Experimental/PCG/Source/PCG/Private/PCGNode.cpp
+++ b/Engine/Plugins/Experimental/PCG/Source/PCG/Private/PCGNode.cpp
@@ -323,7 +323,7 @@ bool UPCGNode::HasInboundEdges() const
 	return false;
 }
 
-void UPCGNode::SetDefaultSettings(TObjectPtr<UPCGSettings> InSettings)
+void UPCGNode::SetDefaultSettings(TObjectPtr<UPCGSettings> InSettings, bool bUpdatePins)
 {
 #if WITH_EDITOR
 	const bool bDifferentSettings = (DefaultSettings != InSettings);
@@ -340,9 +340,12 @@ void UPCGNode::SetDefaultSettings(TObjectPtr<UPCGSettings> InSettings)
 	{
 		DefaultSettings->OnSettingsChangedDelegate.AddUObject(this, &UPCGNode::OnSettingsChanged);
 	}
-
-	UpdatePins();
 #endif
+
+	if (bUpdatePins)
+	{
+		UpdatePins();
+	}
 }
 
 #if WITH_EDITOR
@@ -380,7 +383,23 @@ void UPCGNode::PostEditChangeProperty(struct FPropertyChangedEvent& PropertyChan
 	}
 }
 
+void UPCGNode::OnSettingsChanged(UPCGSettings* InSettings, EPCGChangeType ChangeType)
+{
+	if (InSettings == DefaultSettings)
+	{
+		const bool bUpdatedPins = UpdatePins();
+		OnNodeChangedDelegate.Broadcast(this, ((bUpdatedPins ? EPCGChangeType::Edge : EPCGChangeType::None) | ChangeType));
+	}
+}
+
+#endif // WITH_EDITOR
+
 bool UPCGNode::UpdatePins()
+{
+	return UpdatePins([](UPCGNode* Node){ return NewObject<UPCGPin>(Node); });
+}
+
+bool UPCGNode::UpdatePins(TFunctionRef<UPCGPin*(UPCGNode*)> PinAllocator)
 {
 	if (!DefaultSettings)
 	{
@@ -399,7 +418,7 @@ bool UPCGNode::UpdatePins()
 	TArray<FName> InboundLabels = DefaultSettings->InLabels();
 	TArray<FName> OutboundLabels = DefaultSettings->OutLabels();
 
-	auto UpdatePins = [this](TArray<UPCGPin*>& Pins, const TArray<FName>& Labels)
+	auto UpdatePins = [this, &PinAllocator](TArray<UPCGPin*>& Pins, const TArray<FName>& Labels)
 	{
 		TArray<FName> OldLabels;
 		Algo::Transform(Pins, OldLabels, [](UPCGPin* Pin) { return Pin->Label; });
@@ -461,7 +480,7 @@ bool UPCGNode::UpdatePins()
 			for (const FName& UnmatchedLabel : UnmatchedLabels)
 			{
 				const int32 InsertIndex = Labels.IndexOfByKey(UnmatchedLabel);
-				UPCGPin* NewPin = NewObject<UPCGPin>(this);
+				UPCGPin* NewPin = PinAllocator(this);
 				NewPin->Node = this;
 				NewPin->Label = UnmatchedLabel;
 				Pins.Insert(NewPin, InsertIndex);
@@ -477,14 +496,3 @@ bool UPCGNode::UpdatePins()
 
 	return bChanged;
 }
-
-void UPCGNode::OnSettingsChanged(UPCGSettings* InSettings, EPCGChangeType ChangeType)
-{
-	if (InSettings == DefaultSettings)
-	{
-		const bool bUpdatedPins = UpdatePins();
-		OnNodeChangedDelegate.Broadcast(this, ((bUpdatedPins ? EPCGChangeType::Edge : EPCGChangeType::None) | ChangeType));
-	}
-}
-
-#endif // WITH_EDITOR
diff --git a/Engine/Plugins/Experimental/PCG/Source/PCG/Public/PCGNode.h b/Engine/Plugins/Experimental/PCG/Source/PCG/Public/PCGNode.h
--- a/Engine/Plugins/Experimental/PCG/Source/PCG/Public/PCGNode.h
+++ b/Engine/Plugins/Experimental/PCG/Source/PCG/Public/PCGNode.h
@@ -66,7 +66,7 @@ public:
 	bool IsOutputPinConnected(const FName& Label) const;
 
 	/** Changes the default settings in the node */
-	void SetDefaultSettings(TObjectPtr<UPCGSettings> InSettings);
+	void SetDefaultSettings(TObjectPtr<UPCGSettings> InSettings, bool bUpdatePins = true);
 
 	UPCGPin* GetInputPin(const FName& Label);
 	const UPCGPin* GetInputPin(const FName& Label) const;
@@ -102,8 +102,10 @@ public:
 #endif // WITH_EDITORONLY_DATA
 
 protected:
-#if WITH_EDITOR
 	bool UpdatePins();
+	bool UpdatePins(TFunctionRef<UPCGPin* (UPCGNode*)> PinAllocator);
+
+#if WITH_EDITOR
 	virtual void PreEditChange(FProperty* PropertyAboutToChange) override;
 	virtual void PostEditChangeProperty(struct FPropertyChangedEvent& PropertyChangedEvent) override;
 	void OnSettingsChanged(UPCGSettings* InSettings, EPCGChangeType ChangeType);
