Update style usage

#rnx

diff --git a/Engine/Plugins/Experimental/PCG/Source/PCGEditor/Private/PCGEditorGraphNodeBase.cpp b/Engine/Plugins/Experimental/PCG/Source/PCGEditor/Private/PCGEditorGraphNodeBase.cpp
--- a/Engine/Plugins/Experimental/PCG/Source/PCGEditor/Private/PCGEditorGraphNodeBase.cpp
+++ b/Engine/Plugins/Experimental/PCG/Source/PCGEditor/Private/PCGEditorGraphNodeBase.cpp
@@ -62,7 +62,7 @@ void UPCGEditorGraphNodeBase::GetNodeContextMenuActions(UToolMenu* Menu, class U
 			"PCGNode_SetColor",
 			LOCTEXT("PCGNode_SetColor", "Set Node Color"),
 			LOCTEXT("PCGNode_SetColorTooltip", "Sets a specific color on the given node. Note that white maps to the default value"),
-			FSlateIcon(FEditorStyle::GetStyleSetName(), "ColorPicker.Mode"),
+			FSlateIcon(FAppStyle::GetAppStyleSetName(), "ColorPicker.Mode"),
 			FUIAction(FExecuteAction::CreateUObject(const_cast<UPCGEditorGraphNodeBase*>(this), &UPCGEditorGraphNodeBase::OnPickColor)));
 
 		Section.AddSubMenu("Alignment", LOCTEXT("AlignmentHeader", "Alignment"), FText(), FNewToolMenuDelegate::CreateLambda([](UToolMenu* AlignmentMenu)
