Fix issue where graph cache dirtying wouldn't unroot some settings (needed for UE-150658)

#rnx

diff --git a/Engine/Plugins/Experimental/PCG/Source/PCG/Private/Graph/PCGGraphCache.cpp b/Engine/Plugins/Experimental/PCG/Source/PCG/Private/Graph/PCGGraphCache.cpp
--- a/Engine/Plugins/Experimental/PCG/Source/PCG/Private/Graph/PCGGraphCache.cpp
+++ b/Engine/Plugins/Experimental/PCG/Source/PCG/Private/Graph/PCGGraphCache.cpp
@@ -7,20 +7,19 @@
 #include "Algo/AnyOf.h"
 #include "GameFramework/Actor.h"
 
-FPCGGraphCacheEntry::FPCGGraphCacheEntry(const FPCGDataCollection& InInput, const UPCGSettings* InSettings, const FPCGDataCollection& InOutput, TWeakObjectPtr<UObject> InOwner, TSet<UObject*>& OutRootedObjects)
+FPCGGraphCacheEntry::FPCGGraphCacheEntry(const FPCGDataCollection& InInput, const UPCGSettings* InSettings, const FPCGDataCollection& InOutput, TWeakObjectPtr<UObject> InOwner, FPCGRootSet& OutRootSet)
 	: Input(InInput)
 	, Output(InOutput)
 {
 	Settings = InSettings ? Cast<UPCGSettings>(StaticDuplicateObject(InSettings, InOwner.Get())) : nullptr;
 	SettingsCrc32 = InSettings ? InSettings->GetCrc32() : 0;
 
-	Input.RootUnrootedData(OutRootedObjects);
-	Output.RootUnrootedData(OutRootedObjects);
+	Input.AddToRootSet(OutRootSet);
+	Output.AddToRootSet(OutRootSet);
 
-	if (Settings && !Settings->IsRooted())
+	if (Settings)
 	{
-		Settings->AddToRoot();
-		OutRootedObjects.Add(Settings);
+		OutRootSet.Add(Settings);
 	}
 }
 
@@ -84,7 +83,7 @@ void FPCGGraphCache::StoreInCache(const IPCGElement* InElement, const FPCGDataCo
 		Entries = &(CacheData.Add(InElement));
 	}
 
-	Entries->Emplace(InInput, InSettings, InOutput, Owner, RootedData);
+	Entries->Emplace(InInput, InSettings, InOutput, Owner, RootSet);
 }
 
 void FPCGGraphCache::ClearCache()
@@ -95,11 +94,7 @@ void FPCGGraphCache::ClearCache()
 	CacheData.Reset();
 
 	// Unroot all previously rooted data
-	for (UObject* Data : RootedData)
-	{
-		Data->RemoveFromRoot();
-	}
-	RootedData.Reset();
+	RootSet.Clear();
 }
 
 #if WITH_EDITOR
@@ -111,6 +106,22 @@ void FPCGGraphCache::CleanFromCache(const IPCGElement* InElement)
 	}
 
 	FWriteScopeLock ScopeWriteLock(CacheLock);
+	FPCGGraphCacheEntries* Entries = CacheData.Find(InElement);
+	if (Entries)
+	{
+		for (FPCGGraphCacheEntry& Entry : *Entries)
+		{
+			Entry.Input.RemoveFromRootSet(RootSet);
+			Entry.Output.RemoveFromRootSet(RootSet);
+
+			if (Entry.Settings)
+			{
+				RootSet.Remove(Entry.Settings);
+			}
+		}
+	}
+
+	// Finally, remove all entries matching that element
 	CacheData.Remove(InElement);
 }
 #endif // WITH_EDITOR
\ No newline at end of file
diff --git a/Engine/Plugins/Experimental/PCG/Source/PCG/Private/Graph/PCGGraphCache.h b/Engine/Plugins/Experimental/PCG/Source/PCG/Private/Graph/PCGGraphCache.h
--- a/Engine/Plugins/Experimental/PCG/Source/PCG/Private/Graph/PCGGraphCache.h
+++ b/Engine/Plugins/Experimental/PCG/Source/PCG/Private/Graph/PCGGraphCache.h
@@ -12,7 +12,7 @@ class AActor;
 struct FPCGGraphCacheEntry
 {
 	FPCGGraphCacheEntry() = default;
-	FPCGGraphCacheEntry(const FPCGDataCollection& InInput, const UPCGSettings* InSettings, const FPCGDataCollection& InOutput, TWeakObjectPtr<UObject> InOwner, TSet<UObject*>& OutRootedData);
+	FPCGGraphCacheEntry(const FPCGDataCollection& InInput, const UPCGSettings* InSettings, const FPCGDataCollection& InOutput, TWeakObjectPtr<UObject> InOwner, FPCGRootSet& OutRootSet);
 
 	bool Matches(const FPCGDataCollection& InInput, int32 InSettingsCrc32) const;
 
@@ -56,7 +56,7 @@ private:
 	TWeakObjectPtr<UObject> Owner = nullptr;
 
 	/** To prevent garbage collection on data in the cache, we'll need to root some data */
-	TSet<UObject*> RootedData;
+	FPCGRootSet RootSet;
 
 	mutable FRWLock CacheLock;
 };
\ No newline at end of file
diff --git a/Engine/Plugins/Experimental/PCG/Source/PCG/Private/Graph/PCGGraphExecutor.cpp b/Engine/Plugins/Experimental/PCG/Source/PCG/Private/Graph/PCGGraphExecutor.cpp
--- a/Engine/Plugins/Experimental/PCG/Source/PCG/Private/Graph/PCGGraphExecutor.cpp
+++ b/Engine/Plugins/Experimental/PCG/Source/PCG/Private/Graph/PCGGraphExecutor.cpp
@@ -482,7 +482,7 @@ void FPCGGraphExecutor::StoreResults(FPCGTaskId InTaskId, const FPCGDataCollecti
 	OutputData.Add(InTaskId, InTaskOutput);
 
 	// Root any non-rooted results, otherwise they'll get garbage-collected
-	InTaskOutput.RootUnrootedData(RootedData);
+	InTaskOutput.AddToRootSet(ResultsRootSet);
 }
 
 void FPCGGraphExecutor::ClearResults()
@@ -492,11 +492,7 @@ void FPCGGraphExecutor::ClearResults()
 	NextTaskId = 0;
 	OutputData.Reset();
 
-	for (UObject* Data : RootedData)
-	{
-		Data->RemoveFromRoot();
-	}
-	RootedData.Reset();
+	ResultsRootSet.Clear();
 
 	ScheduleLock.Unlock();
 }
diff --git a/Engine/Plugins/Experimental/PCG/Source/PCG/Private/Graph/PCGGraphExecutor.h b/Engine/Plugins/Experimental/PCG/Source/PCG/Private/Graph/PCGGraphExecutor.h
--- a/Engine/Plugins/Experimental/PCG/Source/PCG/Private/Graph/PCGGraphExecutor.h
+++ b/Engine/Plugins/Experimental/PCG/Source/PCG/Private/Graph/PCGGraphExecutor.h
@@ -119,7 +119,7 @@ private:
 	TArray<FPCGGraphTask> Tasks;
 	TArray<FPCGGraphTask> ReadyTasks;
 	TArray<FPCGGraphActiveTask> ActiveTasks;
-	TSet<UObject*> RootedData;
+	FPCGRootSet ResultsRootSet;
 	/** Map of node instances to their output, could be cleared once execution is done */
 	/** Note: this should at some point unload based on loaded/unloaded proxies, otherwise memory cost will be unbounded */
 	TMap<FPCGTaskId, FPCGDataCollection> OutputData;
diff --git a/Engine/Plugins/Experimental/PCG/Source/PCG/Private/PCGData.cpp b/Engine/Plugins/Experimental/PCG/Source/PCG/Private/PCGData.cpp
--- a/Engine/Plugins/Experimental/PCG/Source/PCG/Private/PCGData.cpp
+++ b/Engine/Plugins/Experimental/PCG/Source/PCG/Private/PCGData.cpp
@@ -5,6 +5,45 @@
 #include "PCGParamData.h"
 #include "Data/PCGSpatialData.h"
 
+void FPCGRootSet::Clear()
+{
+	for (TPair<UObject*, int32>& Entry : RootSet)
+	{
+		Entry.Key->RemoveFromRoot();
+	}
+
+	RootSet.Reset();
+}
+
+void FPCGRootSet::Add(UObject* InObject)
+{
+	check(InObject);
+	if (int32* Found = RootSet.Find(InObject))
+	{
+		(*Found)++;
+	}
+	else if(!InObject->IsRooted())
+	{
+		InObject->AddToRoot();
+		RootSet.Emplace(InObject, 1);
+	}
+}
+void FPCGRootSet::Remove(UObject* InObject)
+{
+	check(InObject);
+	if (int32* Found = RootSet.Find(InObject))
+	{
+		check(InObject->IsRooted());
+		(*Found)--;
+
+		if (*Found == 0)
+		{
+			InObject->RemoveFromRoot();
+			RootSet.Remove(InObject);
+		}
+	}
+}
+
 bool FPCGTaggedData::operator==(const FPCGTaggedData& Other) const
 {
 	return Data == Other.Data &&
@@ -117,16 +156,26 @@ bool FPCGDataCollection::operator==(const FPCGDataCollection& Other) const
 	return true;
 }
 
-void FPCGDataCollection::RootUnrootedData(TSet<UObject*>& OutRootedData) const
+void FPCGDataCollection::AddToRootSet(FPCGRootSet& RootSet) const
+{
+	for (const FPCGTaggedData& Data : TaggedData)
+	{
+		if (Data.Data)
+		{
+			// This is technically a const_cast
+			RootSet.Add(Cast<UObject>(Data.Data));
+		}
+	}
+}
+
+void FPCGDataCollection::RemoveFromRootSet(FPCGRootSet& RootSet) const
 {
 	for (const FPCGTaggedData& Data : TaggedData)
 	{
-		if (Data.Data && !Data.Data->IsRooted())
+		if (Data.Data)
 		{
 			// This is technically a const_cast
-			UObject* DataToRoot = Cast<UObject>(Data.Data);
-			DataToRoot->AddToRoot();
-			OutRootedData.Add(DataToRoot);
+			RootSet.Remove(Cast<UObject>(Data.Data));
 		}
 	}
 }
diff --git a/Engine/Plugins/Experimental/PCG/Source/PCG/Public/PCGData.h b/Engine/Plugins/Experimental/PCG/Source/PCG/Public/PCGData.h
--- a/Engine/Plugins/Experimental/PCG/Source/PCG/Public/PCGData.h
+++ b/Engine/Plugins/Experimental/PCG/Source/PCG/Public/PCGData.h
@@ -41,6 +41,16 @@ struct PCG_API FPCGTaggedData
 	bool operator!=(const FPCGTaggedData& Other) const;
 };
 
+struct FPCGRootSet
+{
+	~FPCGRootSet() { Clear(); }
+	void Clear();
+	void Add(UObject* InObject);
+	void Remove(UObject* InObject);
+
+	TMap<UObject*, int32> RootSet;
+};
+
 USTRUCT(BlueprintType)
 struct PCG_API FPCGDataCollection
 {
@@ -61,7 +71,8 @@ struct PCG_API FPCGDataCollection
 	const UPCGSettings* GetSettings(const UPCGSettings* InDefaultSettings) const;
 
 	bool operator==(const FPCGDataCollection& Other) const;
-	void RootUnrootedData(TSet<UObject*>& OutRootedData) const;
+	void AddToRootSet(FPCGRootSet& RootSet) const;
+	void RemoveFromRootSet(FPCGRootSet& RootSet) const;
 
 	UPROPERTY(BlueprintReadWrite, EditAnywhere, Category = Data)
 	TArray<FPCGTaggedData> TaggedData;
