Added default pins when there are no subgraphs/no blueprint selected

#rnx

diff --git a/Engine/Plugins/Experimental/PCG/Source/PCG/Private/Elements/PCGExecuteBlueprint.cpp b/Engine/Plugins/Experimental/PCG/Source/PCG/Private/Elements/PCGExecuteBlueprint.cpp
--- a/Engine/Plugins/Experimental/PCG/Source/PCG/Private/Elements/PCGExecuteBlueprint.cpp
+++ b/Engine/Plugins/Experimental/PCG/Source/PCG/Private/Elements/PCGExecuteBlueprint.cpp
@@ -339,6 +339,10 @@ TArray<FName> UPCGBlueprintSettings::InLabels() const
 
 		InputLabels.Append(BlueprintElementInstance->InputPinLabels.Array());
 	}
+	else
+	{
+		InputLabels = Super::InLabels();
+	}
 
 	return InputLabels;
 }
@@ -356,6 +360,10 @@ TArray<FName> UPCGBlueprintSettings::OutLabels() const
 
 		OutputLabels.Append(BlueprintElementInstance->OutputPinLabels.Array());
 	}
+	else
+	{
+		OutputLabels = Super::OutLabels();
+	}
 
 	return OutputLabels;
 }
diff --git a/Engine/Plugins/Experimental/PCG/Source/PCG/Private/PCGSubgraph.cpp b/Engine/Plugins/Experimental/PCG/Source/PCG/Private/PCGSubgraph.cpp
--- a/Engine/Plugins/Experimental/PCG/Source/PCG/Private/PCGSubgraph.cpp
+++ b/Engine/Plugins/Experimental/PCG/Source/PCG/Private/PCGSubgraph.cpp
@@ -85,7 +85,7 @@ TArray<FName> UPCGBaseSubgraphSettings::InLabels() const
 	}
 	else
 	{
-		return TArray<FName>();
+		return Super::InLabels();
 	}
 }
 
@@ -97,7 +97,7 @@ TArray<FName> UPCGBaseSubgraphSettings::OutLabels() const
 	}
 	else
 	{
-		return TArray<FName>();
+		return Super::OutLabels();
 	}
 }
 
