Made texture/render target a bit more robust

#rnx

diff --git a/Engine/Plugins/Experimental/PCG/Source/PCG/Private/Data/PCGTextureData.cpp b/Engine/Plugins/Experimental/PCG/Source/PCG/Private/Data/PCGTextureData.cpp
--- a/Engine/Plugins/Experimental/PCG/Source/PCG/Private/Data/PCGTextureData.cpp
+++ b/Engine/Plugins/Experimental/PCG/Source/PCG/Private/Data/PCGTextureData.cpp
@@ -26,6 +26,12 @@ namespace PCGTextureSampling
 	template<typename ValueType>
 	ValueType Sample(const FVector2D& InPosition, const FBox2D& InSurface, int32 Width, int32 Height, TFunctionRef<ValueType(int32 Index)> SamplingFunction)
 	{
+		check(Width > 0 && Height > 0);
+		if (Width <= 0 || Height <= 0)
+		{
+			return ValueType{};
+		}
+
 		// TODO: There seems to be a bias issue here, as the bounds size are not in the same space as the texels.
 		// Implementation note: Supports only stretch fit
 		FVector2D Pos = (InPosition - InSurface.Min) * FVector2D(Width, Height) / InSurface.GetSize();
@@ -93,6 +99,11 @@ FBox UPCGBaseTextureData::GetStrictBounds() const
 
 float UPCGBaseTextureData::GetDensityAtPosition(const FVector& InPosition) const
 {
+	if (!IsValid())
+	{
+		return 0.0f;
+	}
+
 	FVector LocalPosition = Transform.InverseTransformPosition(InPosition);
 	FVector2D Position2D(LocalPosition.X, LocalPosition.Y);
 
@@ -114,10 +125,13 @@ FPCGPoint UPCGBaseTextureData::TransformPoint(const FPCGPoint& InPoint) const
 	FVector2D Position2D(PointPositionInLocalSpace.X, PointPositionInLocalSpace.Y);
 	FBox2D Surface(FVector2D(-1.0f, -1.0f), FVector2D(1.0f, 1.0f));
 
-	FLinearColor Color = PCGTextureSampling::Sample<FLinearColor>(Position2D, Surface, Width, Height, [this](int32 Index) { return ColorData[Index]; });
+	if (IsValid())
+	{
+		FLinearColor Color = PCGTextureSampling::Sample<FLinearColor>(Position2D, Surface, Width, Height, [this](int32 Index) { return ColorData[Index]; });
 
-	Point.Color *= Color;
-	Point.Density = PCGTextureDataMaths::ComputeDensity(Point.Density, PCGTextureSampling::SampleFloatChannel(Color, ColorChannel), DensityFunction);
+		Point.Color *= Color;
+		Point.Density = PCGTextureDataMaths::ComputeDensity(Point.Density, PCGTextureSampling::SampleFloatChannel(Color, ColorChannel), DensityFunction);
+	}
 
 	return Point;
 }
@@ -132,6 +146,14 @@ const UPCGPointData* UPCGBaseTextureData::CreatePointData(FPCGContext* Context)
 
 	UPCGPointData* Data = NewObject<UPCGPointData>(const_cast<UPCGBaseTextureData*>(this));
 	Data->InitializeFromData(this);
+
+	// Early out for invalid data
+	if (!IsValid())
+	{
+		UE_LOG(LogPCG, Error, TEXT("Texture data does not have valid sizes - will return empty data"));
+		return Data;
+	}
+
 	TArray<FPCGPoint>& Points = Data->GetMutablePoints();
 
 	// TODO: There's a bias issue here where we should correct by a 0.5 unit...
@@ -173,6 +195,11 @@ const UPCGPointData* UPCGBaseTextureData::CreatePointData(FPCGContext* Context)
 	return Data;
 }
 
+bool UPCGBaseTextureData::IsValid() const
+{
+	return Height > 0 && Width > 0;
+}
+
 void UPCGTextureData::Initialize(UTexture2D* InTexture, const FTransform& InTransform)
 {
 	Texture = InTexture;
diff --git a/Engine/Plugins/Experimental/PCG/Source/PCG/Public/Data/PCGTextureData.h b/Engine/Plugins/Experimental/PCG/Source/PCG/Public/Data/PCGTextureData.h
--- a/Engine/Plugins/Experimental/PCG/Source/PCG/Public/Data/PCGTextureData.h
+++ b/Engine/Plugins/Experimental/PCG/Source/PCG/Public/Data/PCGTextureData.h
@@ -41,6 +41,8 @@ public:
 	virtual const UPCGPointData* CreatePointData(FPCGContext* Context) const override;
 	//~End UPCGSpatialDataWithPointCache interface
 
+	virtual bool IsValid() const;
+
 public:
 	UPROPERTY(BlueprintReadWrite, VisibleAnywhere, Category = SpatialData)
 	EPCGTextureDensityFunction DensityFunction = EPCGTextureDensityFunction::Multiply; 
