Fixed PCG compiler issue where in some instances a given node would be compiled multiple times
Fixed issue with mesh loading in static mesh spawner

#rnx

diff --git a/Engine/Plugins/Experimental/PCG/Source/PCG/Private/Elements/PCGStaticMeshSpawner.cpp b/Engine/Plugins/Experimental/PCG/Source/PCG/Private/Elements/PCGStaticMeshSpawner.cpp
--- a/Engine/Plugins/Experimental/PCG/Source/PCG/Private/Elements/PCGStaticMeshSpawner.cpp
+++ b/Engine/Plugins/Experimental/PCG/Source/PCG/Private/Elements/PCGStaticMeshSpawner.cpp
@@ -114,7 +114,14 @@ bool FPCGStaticMeshSpawnerElement::ExecuteInternal(FPCGContext* Context) const
 			// Second, add the instances
 			for (FWeightedMeshAndInstances& Instances : WeightedEntries)
 			{
-				if (!Instances.Entry.Mesh || Instances.Instances.Num() == 0)
+				if (Instances.Instances.Num() == 0)
+				{
+					continue;
+				}
+
+				Instances.Entry.Mesh.LoadSynchronous();
+
+				if (!Instances.Entry.Mesh)
 				{
 					continue;
 				}
diff --git a/Engine/Plugins/Experimental/PCG/Source/PCG/Private/Graph/PCGGraphCompiler.cpp b/Engine/Plugins/Experimental/PCG/Source/PCG/Private/Graph/PCGGraphCompiler.cpp
--- a/Engine/Plugins/Experimental/PCG/Source/PCG/Private/Graph/PCGGraphCompiler.cpp
+++ b/Engine/Plugins/Experimental/PCG/Source/PCG/Private/Graph/PCGGraphCompiler.cpp
@@ -108,16 +108,8 @@ TArray<FPCGGraphTask> FPCGGraphCompiler::CompileGraph(UPCGGraph* InGraph, FPCGTa
 				PostTask.Inputs.Emplace(OutputNodeTask->NodeId, nullptr, nullptr);
 			}
 
+			check(!IdMapping.Contains(Node));
 			IdMapping.Add(Node, PostId);
-
-			// Prime any input-less nodes from the subgraph
-			for (const UPCGNode* SubNode : Subgraph->GetNodes())
-			{
-				if(!SubNode->HasInboundEdges())
-				{
-					NodeQueue.Add(SubNode);
-				}
-			}
 		}
 		else
 		{
@@ -148,6 +140,7 @@ TArray<FPCGGraphTask> FPCGGraphCompiler::CompileGraph(UPCGGraph* InGraph, FPCGTa
 				}
 			}
 
+			check(!IdMapping.Contains(Node));
 			IdMapping.Add(Node, NodeId);
 		}
 
@@ -164,6 +157,12 @@ TArray<FPCGGraphTask> FPCGGraphCompiler::CompileGraph(UPCGGraph* InGraph, FPCGTa
 
 				const UPCGNode* OutboundNode = OutboundEdge->OutputPin->Node;
 				check(OutboundNode);
+
+				if (NodeQueue.Contains(OutboundNode))
+				{
+					continue;
+				}
+
 				bool bAllPrerequisitesMet = true;
 
 				for (const UPCGPin* OutboundNodeInputPin : OutboundNode->InputPins)
