Fixed issue where nodes after an isolated node would still be executed

#rnx

diff --git a/Engine/Plugins/Experimental/PCG/Source/PCG/Private/Graph/PCGGraphExecutor.cpp b/Engine/Plugins/Experimental/PCG/Source/PCG/Private/Graph/PCGGraphExecutor.cpp
--- a/Engine/Plugins/Experimental/PCG/Source/PCG/Private/Graph/PCGGraphExecutor.cpp
+++ b/Engine/Plugins/Experimental/PCG/Source/PCG/Private/Graph/PCGGraphExecutor.cpp
@@ -340,14 +340,15 @@ void FPCGGraphExecutor::Execute()
 
 			CurrentlyUsedThreads -= ActiveTask.Context->NumAvailableTasks;
 
-			// Store output in data map
-			StoreResults(ActiveTask.NodeId, ActiveTask.Context->OutputData);
-
 #if WITH_EDITOR
 			// Execute debug display code as needed - done here because it needs to be done on the main thread
+			// Additional note: this needs to be executed before the StoreResults since debugging might cancel further tasks
 			ActiveTask.Element->DebugDisplay(ActiveTask.Context.Get());
 #endif
 
+			// Store output in data map
+			StoreResults(ActiveTask.NodeId, ActiveTask.Context->OutputData);
+
 			// Book-keeping
 			bSomeTaskEndedInCurrentLoop = true;
 
