Added debug scaling/offset based on extents
Fixed an issue with the material override in debug display

diff --git a/Engine/Plugins/Experimental/PCG/Source/PCG/Private/Elements/PCGDebugElement.cpp b/Engine/Plugins/Experimental/PCG/Source/PCG/Private/Elements/PCGDebugElement.cpp
--- a/Engine/Plugins/Experimental/PCG/Source/PCG/Private/Elements/PCGDebugElement.cpp
+++ b/Engine/Plugins/Experimental/PCG/Source/PCG/Private/Elements/PCGDebugElement.cpp
@@ -117,11 +117,27 @@ namespace PCGDebugElement
 			// First, create target instance transforms
 			const float PointScale = DebugSettings.PointScale;
 			const bool bIsRelative = DebugSettings.ScaleMethod == EPCGDebugVisScaleMethod::Relative;
+			const bool bScaleWithExtents = DebugSettings.ScaleMethod == EPCGDebugVisScaleMethod::Extents;
+			const FVector MeshExtents = Mesh->GetBoundingBox().GetExtent();
+			const FVector MeshCenter = Mesh->GetBoundingBox().GetCenter();
 
 			for (const FPCGPoint& Point : Points)
 			{
 				FTransform& InstanceTransform = Instances.Add_GetRef(Point.Transform);
-				InstanceTransform.SetScale3D(bIsRelative ? InstanceTransform.GetScale3D() * PointScale : FVector(PointScale));
+				if (bIsRelative)
+				{
+					InstanceTransform.SetScale3D(InstanceTransform.GetScale3D() * PointScale);
+				}
+				else if (bScaleWithExtents)
+				{
+					const FVector ScaleWithExtents = Point.GetExtents() / MeshExtents;
+					InstanceTransform.SetTranslation(InstanceTransform.GetTranslation() + Point.GetLocalCenter() - MeshCenter * ScaleWithExtents);
+					InstanceTransform.SetScale3D(InstanceTransform.GetScale3D() * ScaleWithExtents);
+				}
+				else // absolute scaling only
+				{
+					InstanceTransform.SetScale3D(FVector(PointScale));
+				}
 			}
 
 			FPCGISMCBuilderParameters Params;
diff --git a/Engine/Plugins/Experimental/PCG/Source/PCG/Private/Helpers/PCGBlueprintHelpers.cpp b/Engine/Plugins/Experimental/PCG/Source/PCG/Private/Helpers/PCGBlueprintHelpers.cpp
--- a/Engine/Plugins/Experimental/PCG/Source/PCG/Private/Helpers/PCGBlueprintHelpers.cpp
+++ b/Engine/Plugins/Experimental/PCG/Source/PCG/Private/Helpers/PCGBlueprintHelpers.cpp
@@ -71,4 +71,15 @@ void UPCGBlueprintHelpers::SetExtents(FPCGPoint& InPoint, const FVector& InExten
 FVector UPCGBlueprintHelpers::GetExtents(const FPCGPoint& InPoint)
 {
 	return InPoint.GetExtents();
-}
\ No newline at end of file
+}
+
+void UPCGBlueprintHelpers::SetLocalCenter(FPCGPoint& InPoint, const FVector& InLocalCenter)
+{
+	InPoint.SetLocalCenter(InLocalCenter);
+}
+
+FVector UPCGBlueprintHelpers::GetLocalCenter(const FPCGPoint& InPoint)
+{
+	return InPoint.GetLocalCenter();
+}
+
diff --git a/Engine/Plugins/Experimental/PCG/Source/PCG/Private/PCGDebug.cpp b/Engine/Plugins/Experimental/PCG/Source/PCG/Private/PCGDebug.cpp
--- a/Engine/Plugins/Experimental/PCG/Source/PCG/Private/PCGDebug.cpp
+++ b/Engine/Plugins/Experimental/PCG/Source/PCG/Private/PCGDebug.cpp
@@ -16,7 +16,7 @@ FPCGDebugVisualizationSettings::FPCGDebugVisualizationSettings()
 
 TSoftObjectPtr<UMaterialInterface> FPCGDebugVisualizationSettings::GetMaterial() const
 {
-	if (!MaterialOverride.IsValid() && PointMesh.ToSoftObjectPath() == PCGDebugVisConstants::DefaultPointMesh)
+	if (MaterialOverride.IsNull() && PointMesh.ToSoftObjectPath() == PCGDebugVisConstants::DefaultPointMesh)
 	{
 		return TSoftObjectPtr<UMaterialInterface>(PCGDebugVisConstants::MaterialForDefaultPointMesh);
 	}
diff --git a/Engine/Plugins/Experimental/PCG/Source/PCG/Public/Helpers/PCGBlueprintHelpers.h b/Engine/Plugins/Experimental/PCG/Source/PCG/Public/Helpers/PCGBlueprintHelpers.h
--- a/Engine/Plugins/Experimental/PCG/Source/PCG/Public/Helpers/PCGBlueprintHelpers.h
+++ b/Engine/Plugins/Experimental/PCG/Source/PCG/Public/Helpers/PCGBlueprintHelpers.h
@@ -45,9 +45,15 @@ public:
 	UFUNCTION(BlueprintCallable, Category = "PCG|Temporary", meta = (ScriptMethod))
 	static UPCGComponent* GetOriginalComponent(UPARAM(ref) FPCGContext& Context);
 
-	UFUNCTION(BlueprintCallable, Category = "PCG|Temporary", meta = (ScriptMethod))
+	UFUNCTION(BlueprintCallable, Category = "PCG|Points", meta = (ScriptMethod))
 	static void SetExtents(UPARAM(ref) FPCGPoint& InPoint, const FVector& InExtents);
 
-	UFUNCTION(BlueprintCallable, Category = "PCG|Temporary", meta = (ScriptMethod))
+	UFUNCTION(BlueprintCallable, Category = "PCG|Points", meta = (ScriptMethod))
 	static FVector GetExtents(const FPCGPoint& InPoint);
+
+	UFUNCTION(BlueprintCallable, Category = "PCG|Points", meta = (ScriptMethod))
+	static void SetLocalCenter(UPARAM(ref) FPCGPoint& InPoint, const FVector& InLocalCenter);
+
+	UFUNCTION(BlueprintCallable, Category = "PCG|Points", meta = (ScriptMethod))
+	static FVector GetLocalCenter(const FPCGPoint& InPoint);
 };
\ No newline at end of file
diff --git a/Engine/Plugins/Experimental/PCG/Source/PCG/Public/PCGDebug.h b/Engine/Plugins/Experimental/PCG/Source/PCG/Public/PCGDebug.h
--- a/Engine/Plugins/Experimental/PCG/Source/PCG/Public/PCGDebug.h
+++ b/Engine/Plugins/Experimental/PCG/Source/PCG/Public/PCGDebug.h
@@ -13,7 +13,8 @@ UENUM()
 enum class EPCGDebugVisScaleMethod : uint8
 {
 	Relative,
-	Absolute
+	Absolute,
+	Extents
 };
 
 USTRUCT(BlueprintType)
@@ -24,7 +25,7 @@ struct PCG_API FPCGDebugVisualizationSettings
 public:
 	FPCGDebugVisualizationSettings();
 
-	UPROPERTY(BlueprintReadWrite, EditAnywhere, Category = Debug, meta=(ClampMin="0"))
+	UPROPERTY(BlueprintReadWrite, EditAnywhere, Category = Debug, meta=(ClampMin="0", EditCondition = "ScaleMethod != EPCGDebugVisScaleMethod::Extents"))
 	float PointScale = 1.0f;
 
 	UPROPERTY(BlueprintReadWrite, EditAnywhere, Category = Debug)
diff --git a/Engine/Plugins/Experimental/PCG/Source/PCG/Public/PCGPoint.h b/Engine/Plugins/Experimental/PCG/Source/PCG/Public/PCGPoint.h
--- a/Engine/Plugins/Experimental/PCG/Source/PCG/Public/PCGPoint.h
+++ b/Engine/Plugins/Experimental/PCG/Source/PCG/Public/PCGPoint.h
@@ -57,5 +57,18 @@ public:
 	int64 MetadataEntry = -1;
 
 	FVector GetExtents() const { return (BoundsMax - BoundsMin) / 2.0; }
-	void SetExtents(const FVector& InExtents) { BoundsMin = -InExtents; BoundsMax = InExtents; }
+	void SetExtents(const FVector& InExtents)
+	{
+		const FVector Center = GetLocalCenter();
+		BoundsMin = Center - InExtents;
+		BoundsMax = Center + InExtents;
+	}
+
+	FVector GetLocalCenter() const { return (BoundsMax + BoundsMin) / 2.0; }
+	void SetLocalCenter(const FVector& InCenter)
+	{
+		const FVector Delta = InCenter - GetLocalCenter();
+		BoundsMin += Delta;
+		BoundsMax += Delta;
+	}
 };
