Fixed typo in param parsing

#rnx

diff --git a/Engine/Plugins/Experimental/PCG/Source/PCG/Private/Elements/PCGDensityFilter.cpp b/Engine/Plugins/Experimental/PCG/Source/PCG/Private/Elements/PCGDensityFilter.cpp
--- a/Engine/Plugins/Experimental/PCG/Source/PCG/Private/Elements/PCGDensityFilter.cpp
+++ b/Engine/Plugins/Experimental/PCG/Source/PCG/Private/Elements/PCGDensityFilter.cpp
@@ -26,7 +26,7 @@ bool FPCGDensityFilterElement::ExecuteInternal(FPCGContext* Context) const
 
 	const bool bInvertFilter = PCGSettingsHelpers::GetValue(GET_MEMBER_NAME_CHECKED(UPCGDensityFilterSettings, bInvertFilter), Settings->bInvertFilter, Params);
 	const float LowerBound = PCGSettingsHelpers::GetValue(GET_MEMBER_NAME_CHECKED(UPCGDensityFilterSettings, LowerBound), Settings->LowerBound, Params);
-	const float UpperBound = PCGSettingsHelpers::GetValue(GET_MEMBER_NAME_CHECKED(UPCGDensityFilterSettings, LowerBound), Settings->UpperBound, Params);
+	const float UpperBound = PCGSettingsHelpers::GetValue(GET_MEMBER_NAME_CHECKED(UPCGDensityFilterSettings, UpperBound), Settings->UpperBound, Params);
 #if WITH_EDITORONLY_DATA
 	const bool bKeepZeroDensityPoints = PCGSettingsHelpers::GetValue(GET_MEMBER_NAME_CHECKED(UPCGDensityFilterSettings, bKeepZeroDensityPoints), Settings->bKeepZeroDensityPoints, Params);
 #else
