Improved ToPointData behavior on intersection, differences, union based on bounds tests

diff --git a/Engine/Plugins/Experimental/PCG/Source/PCG/Private/Data/PCGDifferenceData.cpp b/Engine/Plugins/Experimental/PCG/Source/PCG/Private/Data/PCGDifferenceData.cpp
--- a/Engine/Plugins/Experimental/PCG/Source/PCG/Private/Data/PCGDifferenceData.cpp
+++ b/Engine/Plugins/Experimental/PCG/Source/PCG/Private/Data/PCGDifferenceData.cpp
@@ -26,13 +26,6 @@ namespace PCGDifferenceDataUtils
 	}
 }
 
-float UPCGDifferenceData::GetDensityAtPositionFromDifference(const FVector& InPosition) const
-{
-	const float DensityInDifference = (Difference ? Difference->GetDensityAtPosition(InPosition) : 0.0f);
-	const bool bBinaryDensity = (DensityFunction == EPCGDifferenceDensityFunction::Binary);
-	return ((bBinaryDensity && DensityInDifference > 0) ? 1.0f : DensityInDifference);
-}
-
 void UPCGDifferenceData::Initialize(const UPCGSpatialData* InData)
 {
 	check(InData);
@@ -107,55 +100,12 @@ FBox UPCGDifferenceData::GetStrictBounds() const
 	return Difference ? FBox(EForceInit::ForceInit) : Source->GetStrictBounds();
 }
 
-float UPCGDifferenceData::GetDensityAtPosition(const FVector& InPosition) const
-{
-	if(!PCGHelpers::IsInsideBounds(Source->GetBounds(), InPosition) ||
-		(Difference && PCGHelpers::IsInsideBounds(Difference->GetStrictBounds(), InPosition)))
-	{
-		return 0;
-	}
-	else if(Difference)
-	{
-		const float DensityInSource = Source->GetDensityAtPosition(InPosition);
-		if (DensityInSource == 0)
-		{
-			return 0;
-		}
-
-		const float DensityInDifference = GetDensityAtPositionFromDifference(InPosition);
-		return FMath::Max(0, DensityInSource - DensityInDifference);
-	}
-	else
-	{
-		return Source->GetDensityAtPosition(InPosition);
-	}
-}
-
-FVector UPCGDifferenceData::TransformPosition(const FVector& InPosition) const
-{
-	check(Source);
-	return Source->TransformPosition(InPosition);
-}
-
-FPCGPoint UPCGDifferenceData::TransformPoint(const FPCGPoint& InPoint) const
-{
-	check(Source);
-	FPCGPoint TransformedPoint = Source->TransformPoint(InPoint);
-
-	if (Difference && TransformedPoint.Density > 0)
-	{
-		TransformedPoint.Density = FMath::Max(0, TransformedPoint.Density - GetDensityAtPositionFromDifference(TransformedPoint.Transform.GetLocation()));
-	}
-
-	return TransformedPoint;
-}
-
-bool UPCGDifferenceData::GetPointAtPosition(const FVector& InPosition, FPCGPoint& OutPoint, UPCGMetadata* OutMetadata) const
+bool UPCGDifferenceData::SamplePoint(const FTransform& InTransform, const FBox& InBounds, FPCGPoint& OutPoint, UPCGMetadata* OutMetadata) const
 {
 	check(Source);
 
 	FPCGPoint PointFromSource;
-	if (!Source->GetPointAtPosition(InPosition, PointFromSource, OutMetadata))
+	if(!Source->SamplePoint(InTransform, InBounds, PointFromSource, OutMetadata))
 	{
 		return false;
 	}
@@ -163,7 +113,8 @@ bool UPCGDifferenceData::GetPointAtPosition(const FVector& InPosition, FPCGPoint
 	OutPoint = PointFromSource;
 
 	FPCGPoint PointFromDiff;
-	if (Difference && Difference->GetPointAtPosition(PointFromSource.Transform.GetLocation(), PointFromDiff, OutMetadata))
+	// Important note: here we will not use the point we got from the source, otherwise we are introducing severe bias
+	if (Difference && Difference->SamplePoint(InTransform, InBounds, PointFromDiff, OutMetadata))
 	{
 		const bool bBinaryDensity = (DensityFunction == EPCGDifferenceDensityFunction::Binary);
 		
@@ -219,7 +170,7 @@ const UPCGPointData* UPCGDifferenceData::CreatePointData(FPCGContext* Context) c
 		const FPCGPoint& Point = SourcePoints[Index];
 
 		FPCGPoint PointFromDiff;
-		if (Difference && Difference->GetPointAtPosition(Point.Transform.GetLocation(), PointFromDiff, Data->Metadata))
+		if (Difference && Difference->SamplePoint(Point.Transform, Point.GetLocalBounds(), PointFromDiff, Data->Metadata))
 		{
 			const bool bBinaryDensity = (DensityFunction == EPCGDifferenceDensityFunction::Binary);
 
diff --git a/Engine/Plugins/Experimental/PCG/Source/PCG/Private/Data/PCGIntersectionData.cpp b/Engine/Plugins/Experimental/PCG/Source/PCG/Private/Data/PCGIntersectionData.cpp
--- a/Engine/Plugins/Experimental/PCG/Source/PCG/Private/Data/PCGIntersectionData.cpp
+++ b/Engine/Plugins/Experimental/PCG/Source/PCG/Private/Data/PCGIntersectionData.cpp
@@ -50,76 +50,30 @@ FBox UPCGIntersectionData::GetStrictBounds() const
 	return CachedStrictBounds;
 }
 
-float UPCGIntersectionData::GetDensityAtPosition(const FVector& InPosition) const
+bool UPCGIntersectionData::SamplePoint(const FTransform& InTransform, const FBox& InBounds, FPCGPoint& OutPoint, UPCGMetadata* OutMetadata) const
 {
 	check(A && B);
-	if(!PCGHelpers::IsInsideBounds(CachedBounds, InPosition))
-	{
-		return 0;
-	}
-	else if(PCGHelpers::IsInsideBounds(CachedStrictBounds, InPosition))
-	{
-		return 1.0f;
-	}
-	else
-	{
-		float Density = A->GetDensityAtPosition(InPosition);
-		if (Density > 0)
-		{
-			Density = PCGIntersectionDataMaths::ComputeDensity(Density, B->GetDensityAtPosition(InPosition), DensityFunction);
-		}
-
-		return Density;
-	}
-}
-
-FVector UPCGIntersectionData::TransformPosition(const FVector& InPosition) const
-{
-	check(A && B);
-	return A->HasNonTrivialTransform() ? A->TransformPosition(InPosition) : B->TransformPosition(InPosition);
-}
-
-FPCGPoint UPCGIntersectionData::TransformPoint(const FPCGPoint& InPoint) const
-{ 
-	check(A && B);
-	const UPCGSpatialData* X = A->HasNonTrivialTransform() ? A : B;
-	const UPCGSpatialData* Y = (X == A) ? B : A;
-
-	FPCGPoint TransformedPoint = X->TransformPoint(InPoint);
-	if (TransformedPoint.Density > 0)
-	{
-		TransformedPoint.Density = PCGIntersectionDataMaths::ComputeDensity(TransformedPoint.Density, Y->GetDensityAtPosition(TransformedPoint.Transform.GetLocation()), DensityFunction);
-	}	
-
-	return TransformedPoint;
-}
-
-bool UPCGIntersectionData::GetPointAtPosition(const FVector& InPosition, FPCGPoint& OutPoint, UPCGMetadata* OutMetadata) const
-{
-	check(A && B);
-	const UPCGSpatialData* X = A->HasNonTrivialTransform() ? A : B;
-	const UPCGSpatialData* Y = (X == A) ? B : A;
 
-	FPCGPoint PointFromX;
-	if (!X->GetPointAtPosition(InPosition, PointFromX, OutMetadata))
+	FPCGPoint PointFromA;
+	if(!A->SamplePoint(InTransform, InBounds, PointFromA, OutMetadata))
 	{
 		return false;
 	}
 
-	FPCGPoint PointFromY;
-	if (!Y->GetPointAtPosition(PointFromX.Transform.GetLocation(), PointFromY, OutMetadata))
+	FPCGPoint PointFromB;
+	if(!B->SamplePoint(InTransform, InBounds, PointFromB, OutMetadata))
 	{
 		return false;
 	}
 
 	// Merge points into a single point
-	OutPoint = PointFromX;
-	OutPoint.Density = PCGIntersectionDataMaths::ComputeDensity(PointFromX.Density, PointFromY.Density, DensityFunction);
-	OutPoint.Color = PointFromX.Color * PointFromY.Color;
+	OutPoint = (A->HasNonTrivialTransform() ? PointFromA : PointFromB);
+	OutPoint.Density = PCGIntersectionDataMaths::ComputeDensity(PointFromA.Density, PointFromB.Density, DensityFunction);
+	OutPoint.Color = PointFromA.Color * PointFromB.Color;
 
 	if (OutMetadata)
 	{
-		OutMetadata->MergePointAttributes(PointFromX, PointFromY, OutPoint, EPCGMetadataOp::Min);
+		OutMetadata->MergePointAttributes(PointFromA, PointFromB, OutPoint, EPCGMetadataOp::Min);
 	}
 
 	return true;
@@ -175,9 +129,9 @@ UPCGPointData* UPCGIntersectionData::CreateAndFilterPointData(FPCGContext* Conte
 
 		FPCGPoint PointFromY;
 #if WITH_EDITORONLY_DATA
-		if (!Y->GetPointAtPosition(Point.Transform.GetLocation(), PointFromY, Data->Metadata) && !bKeepZeroDensityPoints)
+		if (!Y->SamplePoint(Point.Transform, Point.GetLocalBounds(), PointFromY, Data->Metadata) && !bKeepZeroDensityPoints)
 #else
-		if (!Y->GetPointAtPosition(Point.Transform.GetLocation(), PointFromY, Data->Metadata))
+		if (!Y->SamplePoint(Point.Transform, Point.GetLocalBounds(), PointFromY, Data->Metadata))
 #endif
 		{
 			return false;
diff --git a/Engine/Plugins/Experimental/PCG/Source/PCG/Private/Data/PCGLandscapeData.cpp b/Engine/Plugins/Experimental/PCG/Source/PCG/Private/Data/PCGLandscapeData.cpp
--- a/Engine/Plugins/Experimental/PCG/Source/PCG/Private/Data/PCGLandscapeData.cpp
+++ b/Engine/Plugins/Experimental/PCG/Source/PCG/Private/Data/PCGLandscapeData.cpp
@@ -38,45 +38,27 @@ FBox UPCGLandscapeData::GetStrictBounds() const
 	return Bounds;
 }
 
-FVector UPCGLandscapeData::TransformPosition(const FVector& InPosition) const
+bool UPCGLandscapeData::SamplePoint(const FTransform& InTransform, const FBox& InBounds, FPCGPoint& OutPoint, UPCGMetadata* OutMetadata) const
 {
-	TOptional<float> HeightAtVertex = Landscape->GetHeightAtLocation(InPosition);
+	// TODO support metadata / additional parameters
+	// TODO support bounds tests
+	// TODO: maybe we should consider the bounds on this data? wasn't done previously
+	TOptional<float> HeightAtVertex = Landscape->GetHeightAtLocation(InTransform.GetLocation());
 	if (HeightAtVertex.IsSet())
 	{
-		return FVector(InPosition.X, InPosition.Y, HeightAtVertex.GetValue());
+		OutPoint.Transform = InTransform;
+		const FVector LocationOnLandscape = FVector(InTransform.GetLocation().X, InTransform.GetLocation().Y, HeightAtVertex.GetValue());
+		OutPoint.Transform.SetLocation(LocationOnLandscape);
+		OutPoint.SetLocalBounds(InBounds); // TODO: unless we do bounds testing should set Z to 0.
+		OutPoint.Density = 1.0f;
+		return true;
 	}
 	else
 	{
-		return InPosition; // not on landscape
+		return false;
 	}
 }
 
-FPCGPoint UPCGLandscapeData::TransformPoint(const FPCGPoint& InPoint) const
-{
-	// TODO: change orientation, ... based on landscape data
-	FPCGPoint Point = InPoint;
-	FVector PointLocation = InPoint.Transform.GetLocation();
-
-	TOptional<float> HeightAtVertex = Landscape->GetHeightAtLocation(PointLocation);
-	if (HeightAtVertex.IsSet())
-	{
-		PointLocation.Z = HeightAtVertex.GetValue();
-		Point.Transform.SetLocation(PointLocation);
-	}
-	else
-	{
-		Point.Density = 0;
-	}
-	
-	return Point;
-}
-
-float UPCGLandscapeData::GetDensityAtPosition(const FVector& InPosition) const
-{
-	TOptional<float> HeightAtVertex = Landscape->GetHeightAtLocation(InPosition);
-	return HeightAtVertex.IsSet() ? 1.0f : 0;
-}
-
 const UPCGPointData* UPCGLandscapeData::CreatePointData(FPCGContext* Context) const
 {
 	TRACE_CPUPROFILER_EVENT_SCOPE(UPCGLandscapeData::CreatePointData);
diff --git a/Engine/Plugins/Experimental/PCG/Source/PCG/Private/Data/PCGLandscapeSplineData.cpp b/Engine/Plugins/Experimental/PCG/Source/PCG/Private/Data/PCGLandscapeSplineData.cpp
--- a/Engine/Plugins/Experimental/PCG/Source/PCG/Private/Data/PCGLandscapeSplineData.cpp
+++ b/Engine/Plugins/Experimental/PCG/Source/PCG/Private/Data/PCGLandscapeSplineData.cpp
@@ -205,11 +205,16 @@ FBox UPCGLandscapeSplineData::GetBounds() const
 	return Bounds;
 }
 
-float UPCGLandscapeSplineData::GetDensityAtPosition(const FVector& InPosition) const
+bool UPCGLandscapeSplineData::SamplePoint(const FTransform& InTransform, const FBox& InBounds, FPCGPoint& OutPoint, UPCGMetadata* OutMetadata) const
 {
+	// TODO : add metadata support on poly lines
+	// TODO : add support for bounds
 	check(Spline);
 
-	const FVector Position = Spline->GetComponentTransform().InverseTransformPosition(InPosition);
+	OutPoint.Transform = Spline->GetComponentTransform().Inverse() * InTransform;
+	OutPoint.SetLocalBounds(InBounds); // TODO: should maybe do Min.Z = Max.Z = 0 ?
+
+	const FVector Position = OutPoint.Transform.GetLocation();
 
 	float PointDensity = 0.0f;
 
@@ -254,5 +259,6 @@ float UPCGLandscapeSplineData::GetDensityAtPosition(const FVector& InPosition) c
 		}
 	}
 
-	return PointDensity;
+	OutPoint.Density = PointDensity;
+	return OutPoint.Density > 0;
 }
\ No newline at end of file
diff --git a/Engine/Plugins/Experimental/PCG/Source/PCG/Private/Data/PCGPointData.cpp b/Engine/Plugins/Experimental/PCG/Source/PCG/Private/Data/PCGPointData.cpp
--- a/Engine/Plugins/Experimental/PCG/Source/PCG/Private/Data/PCGPointData.cpp
+++ b/Engine/Plugins/Experimental/PCG/Source/PCG/Private/Data/PCGPointData.cpp
@@ -66,6 +66,20 @@ namespace PCGPointHelpers
 		}
 	}
 
+	float VolumeOverlap(const FPCGPoint& InPoint, const FBox& InTransformedBounds)
+	{
+		const FBox PointTransformedBounds = InPoint.GetLocalBounds().TransformBy(InPoint.Transform);
+		const FBox Overlap = PointTransformedBounds.Overlap(InTransformedBounds);
+		if (Overlap.IsValid)
+		{
+			return Overlap.GetVolume();
+		}
+		else
+		{
+			return 0;
+		}
+	}
+
 	/** Helper function for additive blending of quaternions (copied from ControlRig) */
 	FQuat AddQuatWithWeight(const FQuat& Q, const FQuat& V, float Weight)
 	{
@@ -185,101 +199,31 @@ FPCGPoint UPCGPointData::GetPoint(int32 Index) const
 	}
 }
 
-const FPCGPoint* UPCGPointData::GetPointAtPosition(const FVector& InPosition) const
-{
-	if (bOctreeIsDirty)
-	{
-		RebuildOctree();
-	}
-
-	const FPCGPoint* BestPoint = nullptr;
-
-	Octree.FindElementsWithBoundsTest(FBoxCenterAndExtent(InPosition, FVector::Zero()), [&BestPoint](const FPCGPointRef& InPointRef) {
-		if (!BestPoint || BestPoint->Density < InPointRef.Point->Density)
-		{
-			BestPoint = InPointRef.Point;
-		}
-	});
-
-	return BestPoint;
-}
-
-FPCGPoint UPCGPointData::TransformPoint(const FPCGPoint& InPoint) const
+bool UPCGPointData::SamplePoint(const FTransform& InTransform, const FBox& InBounds, FPCGPoint& OutPoint, UPCGMetadata* OutMetadata) const
 {
 	if (bOctreeIsDirty)
 	{
 		RebuildOctree();
 	}
 
-	FPCGPoint Point = InPoint;
-
 	TArray<TPair<const FPCGPoint*, float>> Contributions;
-	const FVector PointPosition = InPoint.Transform.GetLocation();
-
-	Octree.FindElementsWithBoundsTest(FBoxCenterAndExtent(PointPosition, FVector::Zero()), [&PointPosition, &Contributions](const FPCGPointRef& InPointRef) {
-		Contributions.Emplace(InPointRef.Point, PCGPointHelpers::InverseEuclidianDistance(*InPointRef.Point, PointPosition));
-	});
-
-	float SumContrib = 0;
-	for (const auto& Contribution : Contributions)
-	{
-		SumContrib += Contribution.Value;
-	}
+	const bool bSampleInVolume = (InBounds.GetExtent() != FVector::ZeroVector);
 
-	if (SumContrib <= 0)
+	if (!bSampleInVolume)
 	{
-		return InPoint;
+		const FVector InPosition = InTransform.GetLocation();
+		Octree.FindElementsWithBoundsTest(FBoxCenterAndExtent(InPosition, FVector::Zero()), [&InPosition, &Contributions](const FPCGPointRef& InPointRef) {
+			Contributions.Emplace(InPointRef.Point, PCGPointHelpers::InverseEuclidianDistance(*InPointRef.Point, InPosition));
+			});
 	}
-
-	FRotator WeightedRotator(0);
-	FVector WeightedScale = FVector::Zero();
-	float WeightedDensity = 0;
-	FVector WeightedBoundsMin = FVector::Zero();
-	FVector WeightedBoundsMax = FVector::Zero();
-	FVector WeightedColor = FVector::Zero();
-	float WeightedSteepness = 0;
-
-	for (const auto& Contribution : Contributions)
-	{
-		const FPCGPoint& SourcePoint = *Contribution.Key;
-		float Weight = Contribution.Value / SumContrib;
-
-		WeightedRotator += (SourcePoint.Transform.Rotator() * Weight); // TODO: This is wonky
-		WeightedScale += (SourcePoint.Transform.GetScale3D() * Weight);
-		WeightedDensity += PCGPointHelpers::ManhattanDensity(SourcePoint, PointPosition);
-		WeightedBoundsMin += SourcePoint.BoundsMin * Weight;
-		WeightedBoundsMax += SourcePoint.BoundsMax * Weight;
-		WeightedColor += SourcePoint.Color * Weight;
-		WeightedSteepness += SourcePoint.Steepness * Weight;
-	}
-
-	// Finally, apply changes to point
-	FQuat PointRotation = (Point.Transform.Rotator() + WeightedRotator).Quaternion();
-
-	Point.Transform.SetRotation(PointRotation);
-	Point.Transform.NormalizeRotation();
-	Point.Transform.SetScale3D(Point.Transform.GetScale3D() * WeightedScale);	
-	Point.Density *= WeightedDensity;
-	Point.BoundsMin = WeightedBoundsMin;
-	Point.BoundsMax = WeightedBoundsMax;
-	Point.Color *= WeightedColor;
-	Point.Steepness *= WeightedSteepness;
-
-	return Point;
-}
-
-bool UPCGPointData::GetPointAtPosition(const FVector& InPosition, FPCGPoint& OutPoint, UPCGMetadata* OutMetadata) const
-{
-	if (bOctreeIsDirty)
+	else
 	{
-		RebuildOctree();
+		FBox TransformedBounds = InBounds.TransformBy(InTransform);
+		Octree.FindElementsWithBoundsTest(FBoxCenterAndExtent(TransformedBounds.GetCenter(), TransformedBounds.GetExtent()), [&TransformedBounds, &Contributions](const FPCGPointRef& InPointRef) {
+			Contributions.Emplace(InPointRef.Point, PCGPointHelpers::VolumeOverlap(*InPointRef.Point, TransformedBounds));
+			});
 	}
 
-	TArray<TPair<const FPCGPoint*, float>> Contributions;
-	Octree.FindElementsWithBoundsTest(FBoxCenterAndExtent(InPosition, FVector::Zero()), [&InPosition, &Contributions](const FPCGPointRef& InPointRef) {
-		Contributions.Emplace(InPointRef.Point, PCGPointHelpers::InverseEuclidianDistance(*InPointRef.Point, InPosition));
-	});
-
 	float SumContributions = 0;
 	float MaxContribution = 0;
 	const FPCGPoint* MaxContributor = nullptr;
@@ -300,7 +244,14 @@ bool UPCGPointData::GetPointAtPosition(const FVector& InPosition, FPCGPoint& Out
 		return false;
 	}
 
+	const FVector::FReal TransformedBoundsVolume = InBounds.TransformBy(InTransform).GetVolume();
+	auto ComputeDensityContribution = [TransformedBoundsVolume](float VolumeIntersection)
+	{
+		return (TransformedBoundsVolume > 0) ? VolumeIntersection / TransformedBoundsVolume : 1.0f;
+	};
+
 	// Computed weighted average of spatial properties
+	FVector WeightedPosition = FVector::ZeroVector;
 	FQuat WeightedQuat = FQuat::Identity;
 	FVector WeightedScale = FVector::ZeroVector;
 	float WeightedDensity = 0;
@@ -314,9 +265,19 @@ bool UPCGPointData::GetPointAtPosition(const FVector& InPosition, FPCGPoint& Out
 		const FPCGPoint& SourcePoint = *Contribution.Key;
 		const float Weight = Contribution.Value / SumContributions;
 
+		WeightedPosition += SourcePoint.Transform.GetLocation();
 		WeightedQuat = PCGPointHelpers::AddQuatWithWeight(WeightedQuat, SourcePoint.Transform.GetRotation(), Weight);
 		WeightedScale += SourcePoint.Transform.GetScale3D() * Weight;
-		WeightedDensity += PCGPointHelpers::ManhattanDensity(SourcePoint, InPosition);
+
+		if (!bSampleInVolume)
+		{
+			WeightedDensity += PCGPointHelpers::ManhattanDensity(SourcePoint, InTransform.GetLocation());
+		}
+		else
+		{
+			WeightedDensity += SourcePoint.Density * Weight * ComputeDensityContribution(Contribution.Value);
+		}
+
 		WeightedBoundsMin += SourcePoint.BoundsMin * Weight;
 		WeightedBoundsMax += SourcePoint.BoundsMax * Weight;
 		WeightedColor += SourcePoint.Color * Weight;
@@ -328,7 +289,7 @@ bool UPCGPointData::GetPointAtPosition(const FVector& InPosition, FPCGPoint& Out
 
 	OutPoint.Transform.SetRotation(WeightedQuat);
 	OutPoint.Transform.SetScale3D(WeightedScale);
-	OutPoint.Transform.SetLocation(InPosition);
+	OutPoint.Transform.SetLocation(bSampleInVolume ? WeightedPosition : InTransform.GetLocation());
 	OutPoint.Density = WeightedDensity;
 	OutPoint.BoundsMin = WeightedBoundsMin;
 	OutPoint.BoundsMax = WeightedBoundsMax;
@@ -353,22 +314,6 @@ bool UPCGPointData::GetPointAtPosition(const FVector& InPosition, FPCGPoint& Out
 	return true;
 }
 
-float UPCGPointData::GetDensityAtPosition(const FVector& InPosition) const
-{
-	if (bOctreeIsDirty)
-	{
-		RebuildOctree();
-	}
-
-	float Density = 0;
-
-	Octree.FindElementsWithBoundsTest(FBoxCenterAndExtent(InPosition, FVector::Zero()), [&InPosition, &Density](const FPCGPointRef& InPointRef) {
-		Density += PCGPointHelpers::ManhattanDensity(*InPointRef.Point, InPosition);
-	});
-
-	return FMath::Min(Density, 1.0f);
-}
-
 void UPCGPointData::RebuildOctree() const
 {
 	FScopeLock Lock(&CachedDataLock);
diff --git a/Engine/Plugins/Experimental/PCG/Source/PCG/Private/Data/PCGProjectionData.cpp b/Engine/Plugins/Experimental/PCG/Source/PCG/Private/Data/PCGProjectionData.cpp
--- a/Engine/Plugins/Experimental/PCG/Source/PCG/Private/Data/PCGProjectionData.cpp
+++ b/Engine/Plugins/Experimental/PCG/Source/PCG/Private/Data/PCGProjectionData.cpp
@@ -71,34 +71,16 @@ FBox UPCGProjectionData::ProjectBounds(const FBox& InBounds) const
 	return Bounds;
 }
 
-float UPCGProjectionData::GetDensityAtPosition(const FVector& InPosition) const
-{
-	// TODO: improve projection/unprojection mechanism
-	return Source->GetDensityAtPosition(InPosition);
-}
-
-FVector UPCGProjectionData::TransformPosition(const FVector& InPosition) const
-{
-	// TODO: improve projection/unprojection mechanism
-	return Target->TransformPosition(Source->TransformPosition(InPosition));
-}
-
-FPCGPoint UPCGProjectionData::TransformPoint(const FPCGPoint& InPoint) const
-{
-	// TODO: improve projection/unprojection mechanism
-	return Target->TransformPoint(Source->TransformPoint(InPoint));
-}
-
-bool UPCGProjectionData::GetPointAtPosition(const FVector& InPosition, FPCGPoint& OutPoint, UPCGMetadata* OutMetadata) const
+bool UPCGProjectionData::SamplePoint(const FTransform& InTransform, const FBox& InBounds, FPCGPoint& OutPoint, UPCGMetadata* OutMetadata) const
 {
 	FPCGPoint PointFromSource;
-	if (!Source->GetPointAtPosition(InPosition, PointFromSource, OutMetadata))
+	if (!Source->SamplePoint(InTransform, InBounds, PointFromSource, OutMetadata))
 	{
 		return false;
 	}
 
 	FPCGPoint PointFromTarget;
-	if (!Target->GetPointAtPosition(PointFromSource.Transform.GetLocation(), PointFromTarget, OutMetadata))
+	if (!Target->SamplePoint(PointFromSource.Transform, PointFromSource.GetLocalBounds(), PointFromTarget, OutMetadata))
 	{
 		return false;
 	}
@@ -142,9 +124,9 @@ const UPCGPointData* UPCGProjectionData::CreatePointData(FPCGContext* Context) c
 
 		FPCGPoint PointFromTarget;
 #if WITH_EDITORONLY_DATA
-		if (!Target->GetPointAtPosition(SourcePoint.Transform.GetLocation(), PointFromTarget, PointData->Metadata) && !bKeepZeroDensityPoints)
+		if (!Target->SamplePoint(SourcePoint.Transform, SourcePoint.GetLocalBounds(), PointFromTarget, PointData->Metadata) && !bKeepZeroDensityPoints)
 #else
-		if (!Target->GetPointAtPosition(SourcePoint.Transform.GetLocation(), PointFromTarget, PointData->Metadata))
+		if (!Target->SamplePoint(SourcePoint.Transform, SourcePoint.GetLocalBounds(), PointFromTarget, PointData->Metadata))
 #endif
 		{
 			return false;
diff --git a/Engine/Plugins/Experimental/PCG/Source/PCG/Private/Data/PCGSpatialData.cpp b/Engine/Plugins/Experimental/PCG/Source/PCG/Private/Data/PCGSpatialData.cpp
--- a/Engine/Plugins/Experimental/PCG/Source/PCG/Private/Data/PCGSpatialData.cpp
+++ b/Engine/Plugins/Experimental/PCG/Source/PCG/Private/Data/PCGSpatialData.cpp
@@ -29,12 +29,30 @@ const UPCGPointData* UPCGSpatialDataWithPointCache::ToPointData(FPCGContext* Con
 	return CachedPointData;
 }
 
-FPCGPoint UPCGSpatialData::TransformPoint(const FPCGPoint& InPoint) const
+float UPCGSpatialData::GetDensityAtPosition(const FVector& InPosition) const
 {
-	FPCGPoint TransformedPoint = InPoint;
-	TransformedPoint.Transform.SetTranslation(TransformPosition(InPoint.Transform.GetLocation()));
-	TransformedPoint.Density *= GetDensityAtPosition(InPoint.Transform.GetLocation());
-	return TransformedPoint;
+	FPCGPoint TemporaryPoint;
+	if (SamplePoint(FTransform(InPosition), FBox::BuildAABB(FVector::ZeroVector, FVector::ZeroVector), TemporaryPoint, nullptr))
+	{
+		return TemporaryPoint.Density;
+	}
+	else
+	{
+		return 0;
+	}
+}
+
+FVector UPCGSpatialData::TransformPosition(const FVector& InPosition) const
+{
+	FPCGPoint TemporaryPoint;
+	if (SamplePoint(FTransform(InPosition), FBox::BuildAABB(FVector::ZeroVector, FVector::ZeroVector), TemporaryPoint, nullptr))
+	{
+		return TemporaryPoint.Transform.GetLocation();
+	}
+	else
+	{
+		return InPosition;
+	}
 }
 
 UPCGIntersectionData* UPCGSpatialData::IntersectWith(const UPCGSpatialData* InOther) const
diff --git a/Engine/Plugins/Experimental/PCG/Source/PCG/Private/Data/PCGSplineData.cpp b/Engine/Plugins/Experimental/PCG/Source/PCG/Private/Data/PCGSplineData.cpp
--- a/Engine/Plugins/Experimental/PCG/Source/PCG/Private/Data/PCGSplineData.cpp
+++ b/Engine/Plugins/Experimental/PCG/Source/PCG/Private/Data/PCGSplineData.cpp
@@ -81,19 +81,33 @@ FBox UPCGSplineData::GetBounds() const
 	return CachedBounds;
 }
 
-float UPCGSplineData::GetDensityAtPosition(const FVector& InPosition) const
+bool UPCGSplineData::SamplePoint(const FTransform& InTransform, const FBox& InBounds, FPCGPoint& OutPoint, UPCGMetadata* OutMetadata) const
 {
+	// TODO: support metadata
+	// TODO: support proper bounds
+
 	// Find nearest point on spline
+	const FVector InPosition = InTransform.GetLocation();
 	float NearestPointKey = Spline->FindInputKeyClosestToWorldLocation(InPosition);
-
 	FTransform NearestTransform = Spline->GetTransformAtSplineInputKey(NearestPointKey, ESplineCoordinateSpace::World, true);
-
 	FVector LocalPoint = NearestTransform.InverseTransformPosition(InPosition);
-
+	
 	// Linear fall off based on the distance to the nearest point
 	// TODO: should be based on explicit settings
 	float Distance = LocalPoint.Length();
-	return FMath::Max(0, 1.0f - Distance);
+	if (Distance > 1.0f)
+	{
+		return false;
+	}
+	else
+	{
+		OutPoint.Transform = NearestTransform;
+		OutPoint.Transform.SetLocation(LocalPoint);
+		OutPoint.SetLocalBounds(InBounds);
+		OutPoint.Density = 1.0f - Distance;
+
+		return true;
+	}
 }
 
 UPCGProjectionData* UPCGSplineData::ProjectOn(const UPCGSpatialData* InOther) const
@@ -175,17 +189,21 @@ void UPCGSplineProjectionData::Initialize(const UPCGSplineData* InSourceSpline,
 	}
 }
 
-float UPCGSplineProjectionData::GetDensityAtPosition(const FVector& InPosition) const
+bool UPCGSplineProjectionData::SamplePoint(const FTransform& InTransform, const FBox& InBounds, FPCGPoint& OutPoint, UPCGMetadata* OutMetadata) const
 {
+	// TODO: support metadata
+	// TODO: support bounds
+
 	// Find nearest point on projected spline
+	const FVector InPosition = InTransform.GetLocation();
 	const USplineComponent* Spline = GetSpline()->Spline.Get();
 	const FVector& SurfaceNormal = GetSurface()->GetNormal();
 
-	// Project input point to 2D space
-	FVector LocalPosition = Spline->GetComponentTransform().InverseTransformPosition(InPosition);
-	FVector2D LocalPosition2D = Project(LocalPosition);
+	// Project to 2D space
+	const FTransform LocalTransform = Spline->GetComponentTransform().Inverse() * InTransform;
+	FVector2D LocalPosition2D = Project(LocalTransform.GetLocation());
 	float Dummy;
-	// Find nearest key on 2D spline	
+	// Find nearest key on 2D spline
 	float NearestInputKey = ProjectedPosition.InaccurateFindNearest(LocalPosition2D, Dummy);
 	// TODO: if we didn't want to hand off density computation to the spline and do it here instead, we could do it in 2D space.
 	// Find point on original spline using the previously found key
@@ -193,7 +211,8 @@ float UPCGSplineProjectionData::GetDensityAtPosition(const FVector& InPosition)
 	const FVector NearestPointOnSpline = Spline->GetLocationAtSplineInputKey(NearestInputKey, ESplineCoordinateSpace::World);
 	const FVector PointOnLine = FMath::ClosestPointOnInfiniteLine(InPosition, InPosition + SurfaceNormal, NearestPointOnSpline);
 
-	return GetSpline()->GetDensityAtPosition(PointOnLine);
+	// TODO: this is super inefficient
+	return GetSpline()->SamplePoint(FTransform(PointOnLine), InBounds, OutPoint, OutMetadata);
 }
 
 const UPCGSplineData* UPCGSplineProjectionData::GetSpline() const
diff --git a/Engine/Plugins/Experimental/PCG/Source/PCG/Private/Data/PCGSurfaceData.cpp b/Engine/Plugins/Experimental/PCG/Source/PCG/Private/Data/PCGSurfaceData.cpp
--- a/Engine/Plugins/Experimental/PCG/Source/PCG/Private/Data/PCGSurfaceData.cpp
+++ b/Engine/Plugins/Experimental/PCG/Source/PCG/Private/Data/PCGSurfaceData.cpp
@@ -1,23 +1,3 @@
 // Copyright Epic Games, Inc. All Rights Reserved.
 
 #include "Data/PCGSurfaceData.h"
-
-FVector UPCGSurfaceData::TransformPosition(const FVector& InPosition) const
-{
-	return Transform.InverseTransformPosition(InPosition);
-}
-
-FPCGPoint UPCGSurfaceData::TransformPoint(const FPCGPoint& InPoint) const
-{
-	FPCGPoint Point = InPoint;
-
-	// Update point location: put it on the surface plane
-	FVector PointPositionInLocalSpace = TransformPosition(InPoint.Transform.GetLocation());
-	PointPositionInLocalSpace.Z = 0;
-	Point.Transform.SetLocation(Transform.TransformPosition(PointPositionInLocalSpace));
-
-	// Update density
-	Point.Density *= GetDensityAtPosition(InPoint.Transform.GetLocation());
-
-	return Point;
-}
\ No newline at end of file
diff --git a/Engine/Plugins/Experimental/PCG/Source/PCG/Private/Data/PCGTextureData.cpp b/Engine/Plugins/Experimental/PCG/Source/PCG/Private/Data/PCGTextureData.cpp
--- a/Engine/Plugins/Experimental/PCG/Source/PCG/Private/Data/PCGTextureData.cpp
+++ b/Engine/Plugins/Experimental/PCG/Source/PCG/Private/Data/PCGTextureData.cpp
@@ -6,21 +6,6 @@
 #include "Data/PCGPointData.h"
 #include "Helpers/PCGAsync.h"
 
-namespace PCGTextureDataMaths
-{
-	float ComputeDensity(float InDensityA, float InDensityB, EPCGTextureDensityFunction InDensityFunction)
-	{
-		if (InDensityFunction == EPCGTextureDensityFunction::Multiply)
-		{
-			return InDensityA * InDensityB;
-		}
-		else // default: Ignore
-		{
-			return InDensityA;
-		}
-	}
-}
-
 namespace PCGTextureSampling
 {
 	template<typename ValueType>
@@ -97,43 +82,33 @@ FBox UPCGBaseTextureData::GetStrictBounds() const
 	return Bounds;
 }
 
-float UPCGBaseTextureData::GetDensityAtPosition(const FVector& InPosition) const
+bool UPCGBaseTextureData::SamplePoint(const FTransform& InTransform, const FBox& InBounds, FPCGPoint& OutPoint, UPCGMetadata* OutMetadata) const
 {
+	// TODO: add metadata support
+	// TODO: add sampling along the bounds
 	if (!IsValid())
 	{
-		return 0.0f;
+		return false;
 	}
 
-	FVector LocalPosition = Transform.InverseTransformPosition(InPosition);
-	FVector2D Position2D(LocalPosition.X, LocalPosition.Y);
-
-	FBox2D Surface(FVector2D(-1.0f, -1.0f), FVector2D(1.0f, 1.0f));
-
-	return PCGTextureSampling::Sample<float>(Position2D, Surface, Width, Height, [this](int32 Index) { return PCGTextureSampling::SampleFloatChannel(ColorData[Index], ColorChannel); });
-}
-
-FPCGPoint UPCGBaseTextureData::TransformPoint(const FPCGPoint& InPoint) const
-{
-	FPCGPoint Point = InPoint;
-
-	// Update point location: put it on the surface plane
-	FVector PointPositionInLocalSpace = TransformPosition(InPoint.Transform.GetLocation());
+	// Compute transform
+	// TODO: embed local bounds center offset at this time?
+	OutPoint.Transform = Transform.Inverse() * InTransform;
+	FVector PointPositionInLocalSpace = OutPoint.Transform.GetLocation();
 	PointPositionInLocalSpace.Z = 0;
-	Point.Transform.SetLocation(Transform.TransformPosition(PointPositionInLocalSpace));
+	OutPoint.Transform.SetLocation(PointPositionInLocalSpace);
+	OutPoint.SetLocalBounds(InBounds); // TODO: should set Min.Z = Max.Z = 0;
 
-	// Set/Update density & color
+	// Compute density & color (& metadata)
+	// TODO: sample in the bounds given, not only on a single pixel
 	FVector2D Position2D(PointPositionInLocalSpace.X, PointPositionInLocalSpace.Y);
 	FBox2D Surface(FVector2D(-1.0f, -1.0f), FVector2D(1.0f, 1.0f));
 
-	if (IsValid())
-	{
-		FLinearColor Color = PCGTextureSampling::Sample<FLinearColor>(Position2D, Surface, Width, Height, [this](int32 Index) { return ColorData[Index]; });
-
-		Point.Color *= Color;
-		Point.Density = PCGTextureDataMaths::ComputeDensity(Point.Density, PCGTextureSampling::SampleFloatChannel(Color, ColorChannel), DensityFunction);
-	}
+	FLinearColor Color = PCGTextureSampling::Sample<FLinearColor>(Position2D, Surface, Width, Height, [this](int32 Index) { return ColorData[Index]; });
+	OutPoint.Color = Color;
+	OutPoint.Density = ((DensityFunction == EPCGTextureDensityFunction::Ignore) ? 1.0f : PCGTextureSampling::SampleFloatChannel(Color, ColorChannel));
 
-	return Point;
+	return OutPoint.Density > 0;
 }
 
 const UPCGPointData* UPCGBaseTextureData::CreatePointData(FPCGContext* Context) const
diff --git a/Engine/Plugins/Experimental/PCG/Source/PCG/Private/Data/PCGUnionData.cpp b/Engine/Plugins/Experimental/PCG/Source/PCG/Private/Data/PCGUnionData.cpp
--- a/Engine/Plugins/Experimental/PCG/Source/PCG/Private/Data/PCGUnionData.cpp
+++ b/Engine/Plugins/Experimental/PCG/Source/PCG/Private/Data/PCGUnionData.cpp
@@ -78,92 +78,14 @@ FBox UPCGUnionData::GetStrictBounds() const
 	return CachedStrictBounds;
 }
 
-float UPCGUnionData::GetDensityAtPosition(const FVector& InPosition) const
+bool UPCGUnionData::SamplePoint(const FTransform& InTransform, const FBox& InBounds, FPCGPoint& OutPoint, UPCGMetadata* OutMetadata) const
 {
-	// Early exits
-	if(!PCGHelpers::IsInsideBounds(CachedBounds, InPosition))
-	{
-		return 0;
-	}
-	else if(PCGHelpers::IsInsideBounds(CachedStrictBounds, InPosition))
-	{
-		return 1.0f;
-	}
-
-	// Check for presence in any strict bounds of the data.
-	// Note that it can be superfluous in some instances as we will might end up testing
-	// the strict bounds twice per data, but it will perform better in the worst case.
-	for (int32 DataIndex = 0; DataIndex < Data.Num(); ++DataIndex)
-	{
-		if(PCGHelpers::IsInsideBounds(Data[DataIndex]->GetStrictBounds(), InPosition))
-		{
-			return 1.0f;
-		}
-	}
-
-	float Density = 0.0f;
-
-	for (TObjectPtr<const UPCGSpatialData> Datum : Data)
-	{
-		if (PCGUnionDataMaths::UpdateDensity(Density, Datum->GetDensityAtPosition(InPosition), DensityFunction) == 1.0f)
-		{
-			break;
-		}
-	}
-
-	return Density;
-}
-
-FVector UPCGUnionData::TransformPosition(const FVector& InPosition) const
-{
-	if (FirstNonTrivialTransformData)
-	{
-		return FirstNonTrivialTransformData->TransformPosition(InPosition);
-	}
-	else
-	{
-		return Super::TransformPosition(InPosition);
-	}
-}
-
-FPCGPoint UPCGUnionData::TransformPoint(const FPCGPoint& InPoint) const
-{
-	if (FirstNonTrivialTransformData)
-	{
-		FPCGPoint TransformedPoint = FirstNonTrivialTransformData->TransformPoint(InPoint);
-
-		if (DensityFunction == EPCGUnionDensityFunction::Binary && TransformedPoint.Density > 0)
-		{
-			TransformedPoint.Density = 1.0f;
-		}
-
-		const int32 DataCount = Data.Num();
-		for(int32 DataIndex = 0; DataIndex < DataCount && TransformedPoint.Density < 1.0f; ++DataIndex)
-		{
-			if (Data[DataIndex] != FirstNonTrivialTransformData)
-			{
-				PCGUnionDataMaths::UpdateDensity(TransformedPoint.Density, Data[DataIndex]->GetDensityAtPosition(TransformedPoint.Transform.GetLocation()), DensityFunction);
-			}
-		}
-
-		return TransformedPoint;
-	}
-	else
-	{
-		return Super::TransformPoint(InPoint);
-	}
-}
-
-bool UPCGUnionData::GetPointAtPosition(const FVector& InPosition, FPCGPoint& OutPoint, UPCGMetadata* OutMetadata) const
-{
-	FVector PointPosition = InPosition;
 	bool bHasSetPoint = false;
 
 	if (FirstNonTrivialTransformData)
 	{
-		if (FirstNonTrivialTransformData->GetPointAtPosition(InPosition, OutPoint, OutMetadata))
+		if (FirstNonTrivialTransformData->SamplePoint(InTransform, InBounds, OutPoint, OutMetadata))
 		{
-			PointPosition = OutPoint.Transform.GetLocation();
 			bHasSetPoint = true;
 
 			if (DensityFunction == EPCGUnionDensityFunction::Binary && OutPoint.Density > 0)
@@ -183,7 +105,7 @@ bool UPCGUnionData::GetPointAtPosition(const FVector& InPosition, FPCGPoint& Out
 		}
 
 		FPCGPoint PointInData;
-		if (Data[DataIndex]->GetPointAtPosition(PointPosition, PointInData, OutMetadata))
+		if(Data[DataIndex]->SamplePoint(InTransform, InBounds, PointInData, OutMetadata))
 		{
 			if (!bHasSetPoint)
 			{
@@ -356,7 +278,7 @@ void UPCGUnionData::CreateSequentialPointData(FPCGContext* Context, UPCGPointDat
 			for (int32 FollowingDataIndex = DataIndex + DataIndexIncrement; FollowingDataIndex != LastDataIndex && !bSkipLoop; FollowingDataIndex += DataIndexIncrement)
 			{
 				FPCGPoint PointInData;
-				if (Data[FollowingDataIndex]->GetPointAtPosition(OutPoint.Transform.GetLocation(), PointInData, PointData->Metadata))
+				if(Data[FollowingDataIndex]->SamplePoint(OutPoint.Transform, OutPoint.GetLocalBounds(), PointInData, PointData->Metadata))
 				{
 					// Update density
 					PCGUnionDataMaths::UpdateDensity(OutPoint.Density, PointInData.Density, DensityFunction);
diff --git a/Engine/Plugins/Experimental/PCG/Source/PCG/Private/Data/PCGVolumeData.cpp b/Engine/Plugins/Experimental/PCG/Source/PCG/Private/Data/PCGVolumeData.cpp
--- a/Engine/Plugins/Experimental/PCG/Source/PCG/Private/Data/PCGVolumeData.cpp
+++ b/Engine/Plugins/Experimental/PCG/Source/PCG/Private/Data/PCGVolumeData.cpp
@@ -42,23 +42,32 @@ const UPCGPointData* UPCGVolumeData::CreatePointData(FPCGContext* Context) const
 	return nullptr;
 }
 
-float UPCGVolumeData::GetDensityAtPosition(const FVector& InPosition) const
+bool UPCGVolumeData::SamplePoint(const FTransform& InTransform, const FBox& InBounds, FPCGPoint& OutPoint, UPCGMetadata* OutMetadata) const
 {
-	// TODO: support fall-off between 0-1 density as needed
-	// could use FBox::GetClosestPointTo
-	if(PCGHelpers::IsInsideBounds(GetBounds(), InPosition))
+	// TODO: add metadata
+	// TODO: consider bounds
+	const FVector InPosition = InTransform.GetLocation();
+	if (PCGHelpers::IsInsideBounds(GetBounds(), InPosition))
 	{
-		if(!Volume || PCGHelpers::IsInsideBounds(GetStrictBounds(), InPosition))
+		float PointDensity = 0.0f;
+
+		if (!Volume || PCGHelpers::IsInsideBounds(GetStrictBounds(), InPosition))
 		{
-			return 1.0f;
+			PointDensity = 1.0f;
 		}
 		else
 		{
-			return Volume->EncompassesPoint(InPosition) ? 1.0f : 0.0f;
+			PointDensity = Volume->EncompassesPoint(InPosition) ? 1.0f : 0.0f;
 		}
+
+		OutPoint.Transform = InTransform;
+		OutPoint.SetLocalBounds(InBounds);
+		OutPoint.Density = PointDensity;
+
+		return OutPoint.Density > 0;
 	}
 	else
 	{
-		return 0.0f;
+		return false;
 	}
 }
\ No newline at end of file
diff --git a/Engine/Plugins/Experimental/PCG/Source/PCG/Private/Elements/PCGCollapseElement.cpp b/Engine/Plugins/Experimental/PCG/Source/PCG/Private/Elements/PCGCollapseElement.cpp
new file mode 100644
--- /dev/null
+++ b/Engine/Plugins/Experimental/PCG/Source/PCG/Private/Elements/PCGCollapseElement.cpp
@@ -0,0 +1,34 @@
+// Copyright Epic Games, Inc. All Rights Reserved.
+
+#include "Elements/PCGCollapseElement.h"
+
+#include "Data/PCGSpatialData.h"
+#include "Data/PCGPointData.h"
+
+FPCGElementPtr UPCGCollapseSettings::CreateElement() const
+{
+	return MakeShared<FPCGCollapseElement>();
+}
+
+bool FPCGCollapseElement::ExecuteInternal(FPCGContext* Context) const
+{
+	TRACE_CPUPROFILER_EVENT_SCOPE(FPCGCollapseElement::Execute);
+
+	TArray<FPCGTaggedData> Inputs = Context->InputData.GetInputs();
+	TArray<FPCGTaggedData>& Outputs = Context->OutputData.TaggedData;
+
+	for (const FPCGTaggedData& Input : Inputs)
+	{
+		FPCGTaggedData& Output = Outputs.Add_GetRef(Input);
+
+		if (!Input.Data || Cast<const UPCGSpatialData>(Input.Data) == nullptr)
+		{
+			continue;
+		}
+
+		// Currently we support collapsing to point data only, but at some point in the future that might be different
+		Output.Data = Cast<const UPCGSpatialData>(Input.Data)->ToPointData(Context);
+	}
+
+	return true;
+}
\ No newline at end of file
diff --git a/Engine/Plugins/Experimental/PCG/Source/PCG/Private/Elements/PCGPointFilter.cpp b/Engine/Plugins/Experimental/PCG/Source/PCG/Private/Elements/PCGPointFilter.cpp
--- a/Engine/Plugins/Experimental/PCG/Source/PCG/Private/Elements/PCGPointFilter.cpp
+++ b/Engine/Plugins/Experimental/PCG/Source/PCG/Private/Elements/PCGPointFilter.cpp
@@ -818,7 +818,7 @@ bool FPCGPointFilterElement::ExecuteInternal(FPCGContext* Context) const
 			}
 			else if(ThresholdSpatialData)
 			{
-				bSkipTest = !ThresholdSpatialData->GetPointAtPosition(SourcePoint.Transform.GetLocation(), ThresholdPoint, TemporaryMetadata);
+				bSkipTest = !ThresholdSpatialData->SamplePoint(SourcePoint.Transform, SourcePoint.GetLocalBounds(), ThresholdPoint, TemporaryMetadata);
 			}
 			else // Only metadata entry (param)
 			{
diff --git a/Engine/Plugins/Experimental/PCG/Source/PCG/Private/Elements/PCGSurfaceSampler.cpp b/Engine/Plugins/Experimental/PCG/Source/PCG/Private/Elements/PCGSurfaceSampler.cpp
--- a/Engine/Plugins/Experimental/PCG/Source/PCG/Private/Elements/PCGSurfaceSampler.cpp
+++ b/Engine/Plugins/Experimental/PCG/Source/PCG/Private/Elements/PCGSurfaceSampler.cpp
@@ -191,8 +191,9 @@ bool FPCGSurfaceSamplerElement::ExecuteInternal(FPCGContext* Context) const
 				const float RandY = RandomSource.FRand();
 				 
 				const FVector TentativeLocation = FVector(CurrentX + RandX * InnerCellSize.X, CurrentY + RandY * InnerCellSize.Y, 0.0f);
+				const FBox LocalBound(-LoopData.PointExtents, LoopData.PointExtents);
 
-				if (SpatialInput->GetPointAtPosition(TentativeLocation, OutPoint, SampledData->Metadata) || LoopData.bKeepZeroDensityPoints)
+				if (SpatialInput->SamplePoint(FTransform(TentativeLocation), LocalBound, OutPoint, SampledData->Metadata) || LoopData.bKeepZeroDensityPoints)
 				{
 					// Apply final parameters on the point
 					OutPoint.SetExtents(LoopData.PointExtents);
diff --git a/Engine/Plugins/Experimental/PCG/Source/PCG/Private/PCGPoint.cpp b/Engine/Plugins/Experimental/PCG/Source/PCG/Private/PCGPoint.cpp
--- a/Engine/Plugins/Experimental/PCG/Source/PCG/Private/PCGPoint.cpp
+++ b/Engine/Plugins/Experimental/PCG/Source/PCG/Private/PCGPoint.cpp
@@ -9,9 +9,15 @@ FPCGPoint::FPCGPoint(const FTransform& InTransform, float InDensity, int32 InSee
 {
 }
 
-FBoxSphereBounds FPCGPoint::GetBounds() const
+FBox FPCGPoint::GetLocalBounds() const
 {
-	return FBoxSphereBounds(FBox(BoundsMin, BoundsMax).TransformBy(Transform));
+	return FBox(BoundsMin, BoundsMax);
+}
+
+void FPCGPoint::SetLocalBounds(const FBox& InBounds)
+{
+	BoundsMin = InBounds.Min;
+	BoundsMax = InBounds.Max;
 }
 
 FBoxSphereBounds FPCGPoint::GetDensityBounds() const
diff --git a/Engine/Plugins/Experimental/PCG/Source/PCG/Public/Data/PCGDifferenceData.h b/Engine/Plugins/Experimental/PCG/Source/PCG/Public/Data/PCGDifferenceData.h
--- a/Engine/Plugins/Experimental/PCG/Source/PCG/Public/Data/PCGDifferenceData.h
+++ b/Engine/Plugins/Experimental/PCG/Source/PCG/Public/Data/PCGDifferenceData.h
@@ -39,10 +39,7 @@ public:
 	virtual int GetDimension() const override;
 	virtual FBox GetBounds() const override;
 	virtual FBox GetStrictBounds() const override;
-	virtual float GetDensityAtPosition(const FVector& InPosition) const override;
-	virtual FVector TransformPosition(const FVector& InPosition) const override;
-	virtual bool GetPointAtPosition(const FVector& InPosition, FPCGPoint& OutPoint, UPCGMetadata* OutMetadata) const override;
-	virtual FPCGPoint TransformPoint(const FPCGPoint& InPoint) const override;
+	virtual bool SamplePoint(const FTransform& Transform, const FBox& Bounds, FPCGPoint& OutPoint, UPCGMetadata* OutMetadata) const override;
 	virtual bool HasNonTrivialTransform() const override;
 	//~End UPCGSpatialData interface
 
@@ -61,6 +58,4 @@ protected:
 
 	UPROPERTY(BlueprintSetter = SetDensityFunction, EditAnywhere, Category = Settings)
 	EPCGDifferenceDensityFunction DensityFunction = EPCGDifferenceDensityFunction::Minimum;
-
-	float GetDensityAtPositionFromDifference(const FVector& InPosition) const;
 };
diff --git a/Engine/Plugins/Experimental/PCG/Source/PCG/Public/Data/PCGIntersectionData.h b/Engine/Plugins/Experimental/PCG/Source/PCG/Public/Data/PCGIntersectionData.h
--- a/Engine/Plugins/Experimental/PCG/Source/PCG/Public/Data/PCGIntersectionData.h
+++ b/Engine/Plugins/Experimental/PCG/Source/PCG/Public/Data/PCGIntersectionData.h
@@ -28,10 +28,7 @@ public:
 	virtual int GetDimension() const override;
 	virtual FBox GetBounds() const override;
 	virtual FBox GetStrictBounds() const override;
-	virtual float GetDensityAtPosition(const FVector& InPosition) const override;
-	virtual FVector TransformPosition(const FVector& InPosition) const override;
-	virtual bool GetPointAtPosition(const FVector& InPosition, FPCGPoint& OutPoint, UPCGMetadata* OutMetadata) const override;
-	virtual FPCGPoint TransformPoint(const FPCGPoint& InPoint) const override;
+	virtual bool SamplePoint(const FTransform& Transform, const FBox& Bounds, FPCGPoint& OutPoint, UPCGMetadata* OutMetadata) const override;
 	virtual bool HasNonTrivialTransform() const override;
 	//~End UPCGSpatialData interface
 
diff --git a/Engine/Plugins/Experimental/PCG/Source/PCG/Public/Data/PCGLandscapeData.h b/Engine/Plugins/Experimental/PCG/Source/PCG/Public/Data/PCGLandscapeData.h
--- a/Engine/Plugins/Experimental/PCG/Source/PCG/Public/Data/PCGLandscapeData.h
+++ b/Engine/Plugins/Experimental/PCG/Source/PCG/Public/Data/PCGLandscapeData.h
@@ -40,9 +40,7 @@ public:
 	// ~Begin UPGCSpatialData interface
 	virtual FBox GetBounds() const override;
 	virtual FBox GetStrictBounds() const override;
-	virtual float GetDensityAtPosition(const FVector& InPosition) const override;
-	virtual FVector TransformPosition(const FVector& InPosition) const override;
-	virtual FPCGPoint TransformPoint(const FPCGPoint& InPoint) const override;
+	virtual bool SamplePoint(const FTransform& Transform, const FBox& Bounds, FPCGPoint& OutPoint, UPCGMetadata* OutMetadata) const override;
 	virtual bool HasNonTrivialTransform() const override { return true; }
 	// ~End UPGCConcreteData interface
 
diff --git a/Engine/Plugins/Experimental/PCG/Source/PCG/Public/Data/PCGLandscapeSplineData.h b/Engine/Plugins/Experimental/PCG/Source/PCG/Public/Data/PCGLandscapeSplineData.h
--- a/Engine/Plugins/Experimental/PCG/Source/PCG/Public/Data/PCGLandscapeSplineData.h
+++ b/Engine/Plugins/Experimental/PCG/Source/PCG/Public/Data/PCGLandscapeSplineData.h
@@ -28,7 +28,7 @@ public:
 
 	//~Begin UPCGSpatialData interface
 	virtual FBox GetBounds() const override;
-	virtual float GetDensityAtPosition(const FVector& InPosition) const override;
+	virtual bool SamplePoint(const FTransform& Transform, const FBox& Bounds, FPCGPoint& OutPoint, UPCGMetadata* OutMetadata) const override;
 	//~End
 
 protected:
diff --git a/Engine/Plugins/Experimental/PCG/Source/PCG/Public/Data/PCGPointData.h b/Engine/Plugins/Experimental/PCG/Source/PCG/Public/Data/PCGPointData.h
--- a/Engine/Plugins/Experimental/PCG/Source/PCG/Public/Data/PCGPointData.h
+++ b/Engine/Plugins/Experimental/PCG/Source/PCG/Public/Data/PCGPointData.h
@@ -62,10 +62,8 @@ public:
 	// ~Begin UPCGSpatialData interface
 	virtual int GetDimension() const override { return 0; }
 	virtual FBox GetBounds() const override;
-	virtual float GetDensityAtPosition(const FVector& InPosition) const override;
 	virtual const UPCGPointData* ToPointData(FPCGContext* Context) const { return this; }
-	virtual bool GetPointAtPosition(const FVector& InPosition, FPCGPoint& OutPoint, UPCGMetadata* OutMetadata) const;
-	virtual FPCGPoint TransformPoint(const FPCGPoint& InPoint) const;
+	virtual bool SamplePoint(const FTransform& Transform, const FBox& Bounds, FPCGPoint& OutPoint, UPCGMetadata* OutMetadata) const override;
 	// ~End UPCGSpatialData interface
 
 	void InitializeFromActor(AActor* InActor);
@@ -82,8 +80,6 @@ public:
 	UFUNCTION(BlueprintCallable, Category = SpatialData)
 	void CopyPointsFrom(const UPCGPointData* InData, const TArray<int>& InDataIndices);
 
-	const FPCGPoint* GetPointAtPosition(const FVector& InPosition) const;
-
 	TArray<FPCGPoint>& GetMutablePoints();
 
 	const PointOctree& GetOctree() const;
diff --git a/Engine/Plugins/Experimental/PCG/Source/PCG/Public/Data/PCGProjectionData.h b/Engine/Plugins/Experimental/PCG/Source/PCG/Public/Data/PCGProjectionData.h
--- a/Engine/Plugins/Experimental/PCG/Source/PCG/Public/Data/PCGProjectionData.h
+++ b/Engine/Plugins/Experimental/PCG/Source/PCG/Public/Data/PCGProjectionData.h
@@ -22,10 +22,7 @@ public:
 	virtual FBox GetBounds() const override;
 	virtual FBox GetStrictBounds() const override;
 	virtual FVector GetNormal() const override;
-	virtual float GetDensityAtPosition(const FVector& InPosition) const override;
-	virtual FVector TransformPosition(const FVector& InPosition) const override;
-	virtual bool GetPointAtPosition(const FVector& InPosition, FPCGPoint& OutPoint, UPCGMetadata* OutMetadata) const override;
-	virtual FPCGPoint TransformPoint(const FPCGPoint& InPoint) const override;
+	virtual bool SamplePoint(const FTransform& Transform, const FBox& Bounds, FPCGPoint& OutPoint, UPCGMetadata* OutMetadata) const override;
 	virtual bool HasNonTrivialTransform() const override;
 	//~End UPCGSpatialData interface
 
diff --git a/Engine/Plugins/Experimental/PCG/Source/PCG/Public/Data/PCGSpatialData.h b/Engine/Plugins/Experimental/PCG/Source/PCG/Public/Data/PCGSpatialData.h
--- a/Engine/Plugins/Experimental/PCG/Source/PCG/Public/Data/PCGSpatialData.h
+++ b/Engine/Plugins/Experimental/PCG/Source/PCG/Public/Data/PCGSpatialData.h
@@ -53,10 +53,10 @@ public:
 
 	/** Computes the density at a given location */
 	UFUNCTION(BlueprintCallable, Category = Distribution)
-	virtual float GetDensityAtPosition(const FVector& InPosition) const PURE_VIRTUAL(UPCGSpatialData::GetDensityAtPosition, return 0;);
+	virtual float GetDensityAtPosition(const FVector& InPosition) const;
 
 	/** Discretizes the data into points */
-	UFUNCTION(BlueprintCallable, Category = SpatialData)
+	UFUNCTION(BlueprintCallable, Category = SpatialData, meta = (DeprecatedFunction))
 	const UPCGPointData* ToPointData() const { return ToPointData(nullptr); }
 
 	UFUNCTION(BlueprintCallable, Category = SpatialData)
@@ -65,20 +65,10 @@ public:
 	virtual const UPCGPointData* ToPointData(FPCGContext* Context) const PURE_VIRTUAL(UPCGSpatialData::ToPointData, return nullptr;);
 
 	/** Transform a world-space position to a world-space position in relation to the current data. (ex: projection on surface) */
-	UFUNCTION(BlueprintCallable, Category = SpatialData)
-	virtual FVector TransformPosition(const FVector& InPosition) const { return InPosition; }
-
-	/** Transform a full point similar to TransformPosition */
-	UFUNCTION(BlueprintCallable, Category = SpatialData)
-	virtual FPCGPoint TransformPoint(const FPCGPoint& InPoint) const;
+	FVector TransformPosition(const FVector& InPosition) const;
 
 	UFUNCTION(BlueprintCallable, Category = SpatialData)
-	virtual bool GetPointAtPosition(const FVector& InPosition, FPCGPoint& OutPoint, UPCGMetadata* OutMetadata) const
-	{
-		OutPoint.Transform = FTransform(InPosition);
-		OutPoint = TransformPoint(OutPoint);
-		return OutPoint.Density > 0;
-	}
+	virtual bool SamplePoint(const FTransform& Transform, const FBox& Bounds, FPCGPoint& OutPoint, UPCGMetadata* OutMetadata) const PURE_VIRTUAL(UPCGSpatialData::SamplePoint, return false;);
 
 	/** Returns true if the data has a non-trivial transform */
 	UFUNCTION(BlueprintCallable, Category = SpatialData)
diff --git a/Engine/Plugins/Experimental/PCG/Source/PCG/Public/Data/PCGSplineData.h b/Engine/Plugins/Experimental/PCG/Source/PCG/Public/Data/PCGSplineData.h
--- a/Engine/Plugins/Experimental/PCG/Source/PCG/Public/Data/PCGSplineData.h
+++ b/Engine/Plugins/Experimental/PCG/Source/PCG/Public/Data/PCGSplineData.h
@@ -30,7 +30,7 @@ public:
 
 	//~Begin UPCGSpatialData interface
 	virtual FBox GetBounds() const override;
-	virtual float GetDensityAtPosition(const FVector& InPosition) const override;
+	virtual bool SamplePoint(const FTransform& Transform, const FBox& Bounds, FPCGPoint& OutPoint, UPCGMetadata* OutMetadata) const override;
 	virtual UPCGProjectionData* ProjectOn(const UPCGSpatialData* InOther) const override;
 	//~End 
 
@@ -50,7 +50,7 @@ public:
 	void Initialize(const UPCGSplineData* InSourceSpline, const UPCGSpatialData* InTargetSurface);
 
 	//~Begin UPCGSpatialData interface
-	virtual float GetDensityAtPosition(const FVector& InPosition) const override;
+	virtual bool SamplePoint(const FTransform& Transform, const FBox& Bounds, FPCGPoint& OutPoint, UPCGMetadata* OutMetadata) const override;
 	//~End UPCGSpatialData interface
 
 	const UPCGSplineData* GetSpline() const;
diff --git a/Engine/Plugins/Experimental/PCG/Source/PCG/Public/Data/PCGSurfaceData.h b/Engine/Plugins/Experimental/PCG/Source/PCG/Public/Data/PCGSurfaceData.h
--- a/Engine/Plugins/Experimental/PCG/Source/PCG/Public/Data/PCGSurfaceData.h
+++ b/Engine/Plugins/Experimental/PCG/Source/PCG/Public/Data/PCGSurfaceData.h
@@ -15,8 +15,6 @@ class PCG_API UPCGSurfaceData : public UPCGSpatialDataWithPointCache
 public:
 	//~Begin UPCGSpatialData interface
 	virtual int GetDimension() const override { return 2; }
-	virtual FVector TransformPosition(const FVector& InPosition) const override;
-	virtual FPCGPoint TransformPoint(const FPCGPoint& InPoint) const override;
 	virtual bool HasNonTrivialTransform() const override { return true; }
 	//~End UPCGSpatialData interface
 
diff --git a/Engine/Plugins/Experimental/PCG/Source/PCG/Public/Data/PCGTextureData.h b/Engine/Plugins/Experimental/PCG/Source/PCG/Public/Data/PCGTextureData.h
--- a/Engine/Plugins/Experimental/PCG/Source/PCG/Public/Data/PCGTextureData.h
+++ b/Engine/Plugins/Experimental/PCG/Source/PCG/Public/Data/PCGTextureData.h
@@ -33,8 +33,7 @@ public:
 	//~Begin UPCGSpatialData interface
 	virtual FBox GetBounds() const override;
 	virtual FBox GetStrictBounds() const override;
-	virtual float GetDensityAtPosition(const FVector& InPosition) const override;
-	virtual FPCGPoint TransformPoint(const FPCGPoint& InPoint) const;
+	virtual bool SamplePoint(const FTransform& Transform, const FBox& Bounds, FPCGPoint& OutPoint, UPCGMetadata* OutMetadata) const override;
 	//~End UPCGSpatialData interface
 
 	//~Begin UPCGSpatialDataWithPointCache interface
diff --git a/Engine/Plugins/Experimental/PCG/Source/PCG/Public/Data/PCGUnionData.h b/Engine/Plugins/Experimental/PCG/Source/PCG/Public/Data/PCGUnionData.h
--- a/Engine/Plugins/Experimental/PCG/Source/PCG/Public/Data/PCGUnionData.h
+++ b/Engine/Plugins/Experimental/PCG/Source/PCG/Public/Data/PCGUnionData.h
@@ -41,10 +41,7 @@ public:
 	virtual int GetDimension() const override;
 	virtual FBox GetBounds() const override;
 	virtual FBox GetStrictBounds() const override;
-	virtual float GetDensityAtPosition(const FVector& InPosition) const override;
-	virtual FVector TransformPosition(const FVector& InPosition) const override;
-	virtual bool GetPointAtPosition(const FVector& InPosition, FPCGPoint& OutPoint, UPCGMetadata* OutMetadata) const override;
-	virtual FPCGPoint TransformPoint(const FPCGPoint& InPoint) const override;
+	virtual bool SamplePoint(const FTransform& Transform, const FBox& Bounds, FPCGPoint& OutPoint, UPCGMetadata* OutMetadata) const override;
 	virtual bool HasNonTrivialTransform() const override;
 	//~End UPCGSpatialData interface
 
diff --git a/Engine/Plugins/Experimental/PCG/Source/PCG/Public/Data/PCGVolumeData.h b/Engine/Plugins/Experimental/PCG/Source/PCG/Public/Data/PCGVolumeData.h
--- a/Engine/Plugins/Experimental/PCG/Source/PCG/Public/Data/PCGVolumeData.h
+++ b/Engine/Plugins/Experimental/PCG/Source/PCG/Public/Data/PCGVolumeData.h
@@ -22,7 +22,7 @@ public:
 	virtual int GetDimension() const override { return 3; }
 	virtual FBox GetBounds() const override;
 	virtual FBox GetStrictBounds() const override;
-	virtual float GetDensityAtPosition(const FVector& InPosition) const override;
+	virtual bool SamplePoint(const FTransform& Transform, const FBox& Bounds, FPCGPoint& OutPoint, UPCGMetadata* OutMetadata) const override;
 	// ~End UPGCSpatialData interface
 
 	// ~Begin UPCGSpatialDataWithPointCache interface
diff --git a/Engine/Plugins/Experimental/PCG/Source/PCG/Public/Elements/PCGCollapseElement.h b/Engine/Plugins/Experimental/PCG/Source/PCG/Public/Elements/PCGCollapseElement.h
new file mode 100644
--- /dev/null
+++ b/Engine/Plugins/Experimental/PCG/Source/PCG/Public/Elements/PCGCollapseElement.h
@@ -0,0 +1,31 @@
+// Copyright Epic Games, Inc. All Rights Reserved.
+
+#pragma once
+
+#include "PCGElement.h"
+#include "PCGSettings.h"
+
+#include "PCGCollapseElement.generated.h"
+
+UCLASS(BlueprintType, ClassGroup = (Procedural))
+class PCG_API UPCGCollapseSettings : public UPCGSettings
+{
+	GENERATED_BODY()
+
+public:
+	//~Begin UPCGSettings interface
+#if WITH_EDITOR
+	virtual FName GetDefaultNodeName() const override { return FName(TEXT("CollapseNode")); }
+	virtual EPCGSettingsType GetType() const override { return EPCGSettingsType::Spatial; }
+#endif
+
+protected:
+	virtual FPCGElementPtr CreateElement() const override;
+	//~End UPCGSettings interface
+};
+
+class FPCGCollapseElement : public FSimplePCGElement
+{
+protected:
+	virtual bool ExecuteInternal(FPCGContext* Context) const;
+};
\ No newline at end of file
diff --git a/Engine/Plugins/Experimental/PCG/Source/PCG/Public/PCGPoint.h b/Engine/Plugins/Experimental/PCG/Source/PCG/Public/PCGPoint.h
--- a/Engine/Plugins/Experimental/PCG/Source/PCG/Public/PCGPoint.h
+++ b/Engine/Plugins/Experimental/PCG/Source/PCG/Public/PCGPoint.h
@@ -29,7 +29,8 @@ public:
 	FPCGPoint() = default;
 	FPCGPoint(const FTransform& InTransform, float InDensity, int32 InSeed);
 
-	FBoxSphereBounds GetBounds() const;
+	FBox GetLocalBounds() const;
+	void SetLocalBounds(const FBox& InBounds);
 	FBoxSphereBounds GetDensityBounds() const;
 
 	UPROPERTY(BlueprintReadWrite, EditAnywhere, Category = Properties)
