Fixed positional offset in the debug display w.r.t extents-based

#rnx

diff --git a/Engine/Plugins/Experimental/PCG/Source/PCG/Private/Elements/PCGDebugElement.cpp b/Engine/Plugins/Experimental/PCG/Source/PCG/Private/Elements/PCGDebugElement.cpp
--- a/Engine/Plugins/Experimental/PCG/Source/PCG/Private/Elements/PCGDebugElement.cpp
+++ b/Engine/Plugins/Experimental/PCG/Source/PCG/Private/Elements/PCGDebugElement.cpp
@@ -131,7 +131,8 @@ namespace PCGDebugElement
 				else if (bScaleWithExtents)
 				{
 					const FVector ScaleWithExtents = Point.GetExtents() / MeshExtents;
-					InstanceTransform.SetTranslation(InstanceTransform.GetTranslation() + Point.GetLocalCenter() - MeshCenter * ScaleWithExtents);
+					const FVector TransformedBoxCenterWithOffset = InstanceTransform.TransformPositionNoScale(Point.GetLocalCenter()) - InstanceTransform.GetLocation() - MeshCenter * ScaleWithExtents;
+					InstanceTransform.SetTranslation(InstanceTransform.GetTranslation() + TransformedBoxCenterWithOffset);
 					InstanceTransform.SetScale3D(InstanceTransform.GetScale3D() * ScaleWithExtents);
 				}
 				else // absolute scaling only
