Added primitive data to support meshes & collisions
Enabled any change on actor to trigger updates (might need to be revisited)

diff --git a/Engine/Plugins/Experimental/PCG/Source/PCG/Private/Data/PCGPrimitiveData.cpp b/Engine/Plugins/Experimental/PCG/Source/PCG/Private/Data/PCGPrimitiveData.cpp
new file mode 100644
--- /dev/null
+++ b/Engine/Plugins/Experimental/PCG/Source/PCG/Private/Data/PCGPrimitiveData.cpp
@@ -0,0 +1,70 @@
+// Copyright Epic Games, Inc. All Rights Reserved.
+
+#include "Data/PCGPrimitiveData.h"
+
+#include "Data/PCGPointData.h"
+#include "Helpers/PCGAsync.h"
+
+#include "Components/PrimitiveComponent.h"
+
+void UPCGPrimitiveData::Initialize(UPrimitiveComponent* InPrimitive)
+{
+	check(InPrimitive);
+	Primitive = InPrimitive;
+	TargetActor = InPrimitive->GetOwner();
+	CachedBounds = Primitive->Bounds.GetBox();
+	// Not obvious to find strict bounds, leave at the default value
+}
+
+bool UPCGPrimitiveData::SamplePoint(const FTransform& InTransform, const FBox& InBounds, FPCGPoint& OutPoint, UPCGMetadata* OutMetadata) const
+{
+	FCollisionShape CollisionShape;
+	CollisionShape.SetBox(FVector3f(InBounds.GetExtent() * InTransform.GetScale3D()));
+
+	const FVector BoxCenter = InTransform.TransformPosition(InBounds.GetCenter());
+
+	if (Primitive->OverlapComponent(BoxCenter, InTransform.GetRotation(), CollisionShape))
+	{
+		OutPoint.Transform = InTransform;
+		OutPoint.SetLocalBounds(InBounds);
+		OutPoint.Density = 1.0f;
+		return true;
+	}
+	else
+	{
+		return false;
+	}
+}
+
+const UPCGPointData* UPCGPrimitiveData::CreatePointData(FPCGContext* Context) const
+{
+	TRACE_CPUPROFILER_EVENT_SCOPE(UPCGPrimitiveData::CreatePointData);
+
+	UPCGPointData* Data = NewObject<UPCGPointData>(const_cast<UPCGPrimitiveData*>(this));
+	Data->InitializeFromData(this);
+	TArray<FPCGPoint>& Points = Data->GetMutablePoints();
+
+	const int32 MinX = FMath::CeilToInt(CachedBounds.Min.X / VoxelSize);
+	const int32 MaxX = FMath::FloorToInt(CachedBounds.Max.X / VoxelSize);
+	const int32 MinY = FMath::CeilToInt(CachedBounds.Min.Y / VoxelSize);
+	const int32 MaxY = FMath::FloorToInt(CachedBounds.Max.Y / VoxelSize);
+	const int32 MinZ = FMath::CeilToInt(CachedBounds.Min.Z / VoxelSize);
+	const int32 MaxZ = FMath::FloorToInt(CachedBounds.Max.Z / VoxelSize);
+
+	const int32 NumIterations = (MaxX - MinX) * (MaxY - MinY) * (MaxZ - MinZ);
+
+	FPCGAsync::AsyncPointProcessing(Context, NumIterations, Points, [this, MinX, MaxX, MinY, MaxY, MinZ, MaxZ](int32 Index, FPCGPoint& OutPoint)
+	{
+		const int X = MinX + (Index % (MaxX - MinX));
+		const int Y = MinY + (Index / (MaxX - MinX) % (MaxY - MinY));
+		const int Z = MinZ + (Index / ((MaxX - MinX) * (MaxY - MinY)));
+
+		const FVector SampleLocation(X * VoxelSize, Y * VoxelSize, Z * VoxelSize);
+		const FBox VoxelBox(FVector(-VoxelSize / 2.0), FVector(VoxelSize / 2.0));
+		return SamplePoint(FTransform(SampleLocation), VoxelBox, OutPoint, nullptr);
+	});
+
+	UE_LOG(LogPCG, Verbose, TEXT("Primitive %s extracted %d points"), *Primitive->GetFName().ToString(), Points.Num());
+
+	return Data;
+}
\ No newline at end of file
diff --git a/Engine/Plugins/Experimental/PCG/Source/PCG/Private/PCGComponent.cpp b/Engine/Plugins/Experimental/PCG/Source/PCG/Private/PCGComponent.cpp
--- a/Engine/Plugins/Experimental/PCG/Source/PCG/Private/PCGComponent.cpp
+++ b/Engine/Plugins/Experimental/PCG/Source/PCG/Private/PCGComponent.cpp
@@ -9,6 +9,7 @@
 #include "Data/PCGLandscapeData.h"
 #include "Data/PCGLandscapeSplineData.h"
 #include "Data/PCGPointData.h"
+#include "Data/PCGPrimitiveData.h"
 #include "Data/PCGSplineData.h"
 #include "Data/PCGUnionData.h"
 #include "Data/PCGVolumeData.h"
@@ -19,6 +20,7 @@
 #include "ActorPartition/ActorPartitionSubsystem.h"
 #include "Components/InstancedStaticMeshComponent.h"
 #include "Components/SplineComponent.h"
+#include "Components/ShapeComponent.h"
 #include "GameFramework/Volume.h"
 #include "Kismet/GameplayStatics.h"
 #include "Landscape.h"
@@ -49,19 +51,6 @@ bool UPCGComponent::IsPartitioned() const
 	return bIsPartitioned && CanPartition();
 }
 
-bool UPCGComponent::HasSplineComponent() const
-{
-	if (!GetOwner())
-	{
-		return false;
-	}
-
-	TArray<USplineComponent*> Splines;
-	GetOwner()->GetComponents<USplineComponent>(Splines);
-
-	return !Splines.IsEmpty();
-}
-
 void UPCGComponent::SetGraph(UPCGGraph* InGraph)
 {
 	if(Graph == InGraph)
@@ -747,9 +736,8 @@ void UPCGComponent::OnActorMoved(AActor* InActor)
 {
 	if (InActor == GetOwner())
 	{
-		// Generalize the cases where we want the position to dirty the inputs
-		if (InputType == EPCGComponentInput::Landscape ||
-			(InputType == EPCGComponentInput::Actor && HasSplineComponent()))
+		// TODO: find better metrics to dirty the inputs. 
+		// TODO: this should dirty only the actor pcg data.
 		{
 			DirtyGenerated(true);
 			Refresh();
@@ -828,11 +816,8 @@ void UPCGComponent::OnActorChanged(AActor* Actor, UObject* InObject, bool bActor
 	{
 		// Something has changed on the owner (including properties of this component)
 		// In the case of splines, this is where we'd get notified if some component properties (incl. spline vertices) have changed
-		if(HasSplineComponent() || Cast<AVolume>(Actor) != nullptr)
-		{
-			DirtyGenerated(true);
-		}
-
+		// TODO: this should dirty only the actor pcg data.
+		DirtyGenerated(true);
 		Refresh();
 	}
 	else if(Actor)
@@ -1092,59 +1077,97 @@ UPCGData* UPCGComponent::CreateActorPCGData(AActor* Actor)
 
 		return Data;
 	}
-	else if (ALandscapeSplineActor* LandscapeSpline = Cast<ALandscapeSplineActor>(Actor))
+	else // Prepare data on a component basis
 	{
-		TArray<ULandscapeSplinesComponent*> LandscapeSplines;
+		TInlineComponentArray<ULandscapeSplinesComponent*, 1> LandscapeSplines;
 		Actor->GetComponents<ULandscapeSplinesComponent>(LandscapeSplines);
 
-		if (LandscapeSplines.Num() > 1)
+		TInlineComponentArray<USplineComponent*, 1> Splines;
+		Actor->GetComponents<USplineComponent>(Splines);
+
+		TInlineComponentArray<UShapeComponent*, 1> Shapes;
+		Actor->GetComponents<UShapeComponent>(Shapes);
+
+		// Don't get generic primitives unless it's the only thing we can find.
+		TInlineComponentArray<UPrimitiveComponent*, 1> OtherPrimitives;
+		if (LandscapeSplines.Num() == 0 && Splines.Num() == 0 && Shapes.Num() == 0)
 		{
-			UPCGUnionData* Data = NewObject<UPCGUnionData>(this);
-			for (ULandscapeSplinesComponent* SplineComponent : LandscapeSplines)
-			{
-				UPCGLandscapeSplineData* SplineData = NewObject<UPCGLandscapeSplineData>(this);
-				SplineData->Initialize(SplineComponent);
-				Data->AddData(SplineData);
-			}
+			Actor->GetComponents<UPrimitiveComponent>(OtherPrimitives);
+		}
 
-			return Data;
+		UPCGUnionData* Union = nullptr;
+		if (LandscapeSplines.Num() + Splines.Num() + Shapes.Num() + OtherPrimitives.Num() > 1)
+		{
+			Union = NewObject<UPCGUnionData>(this);
 		}
-		else if (LandscapeSplines.Num() == 1)
+
+		for (ULandscapeSplinesComponent* SplineComponent : LandscapeSplines)
 		{
-			UPCGLandscapeSplineData* Data = NewObject<UPCGLandscapeSplineData>(this);
-			Data->Initialize(LandscapeSplines[0]);
-			return Data;
+			UPCGLandscapeSplineData* SplineData = NewObject<UPCGLandscapeSplineData>(this);
+			SplineData->Initialize(SplineComponent);
+
+			if (Union)
+			{
+				Union->AddData(SplineData);
+			}
+			else
+			{
+				return SplineData;
+			}
 		}
-		else
+
+		for (USplineComponent* SplineComponent : Splines)
 		{
-			return nullptr;
+			UPCGSplineData* SplineData = NewObject<UPCGSplineData>(this);
+			SplineData->Initialize(SplineComponent);
+
+			if (Union)
+			{
+				Union->AddData(SplineData);
+			}
+			else
+			{
+				return SplineData;
+			}
 		}
-	}
-	else // TODO how about surfaces, splines?
-	{
-		TArray<USplineComponent*> Splines;
-		Actor->GetComponents<USplineComponent>(Splines);
 
-		if (Splines.Num() > 1)
+		for (UShapeComponent* ShapeComponent : Shapes)
 		{
-			UPCGUnionData* Data = NewObject<UPCGUnionData>(this);
-			for (USplineComponent* SplineComponent : Splines)
+			UPCGPrimitiveData* ShapeData = NewObject<UPCGPrimitiveData>(this);
+			ShapeData->Initialize(ShapeComponent);
+			
+			if (Union)
 			{
-				UPCGSplineData* SplineData = NewObject<UPCGSplineData>(this);
-				SplineData->Initialize(SplineComponent);
-				Data->AddData(SplineData);
+				Union->AddData(ShapeData);
 			}
+			else
+			{
+				return ShapeData;
+			}
+		}
 
-			return Data;
+		for (UPrimitiveComponent* PrimitiveComponent : OtherPrimitives)
+		{
+			UPCGPrimitiveData* PrimitiveData = NewObject<UPCGPrimitiveData>(this);
+			PrimitiveData->Initialize(PrimitiveComponent);
+
+			if (Union)
+			{
+				Union->AddData(PrimitiveData);
+			}
+			else
+			{
+				return PrimitiveData;
+			}
 		}
-		else if (Splines.Num() == 1)
+
+		if (Union)
 		{
-			UPCGSplineData* Data = NewObject<UPCGSplineData>(this);
-			Data->Initialize(Splines[0]);
-			return Data;
+			return Union;
 		}
-		else // Default behavior on unknown actors
+		else // No parsed components: default
 		{
+			//Default behavior on unknown actors is to write a single point at the actor location
 			UPCGPointData* Data = NewObject<UPCGPointData>(this);
 			Data->InitializeFromActor(Actor);
 			return Data;
diff --git a/Engine/Plugins/Experimental/PCG/Source/PCG/Public/Data/PCGPrimitiveData.h b/Engine/Plugins/Experimental/PCG/Source/PCG/Public/Data/PCGPrimitiveData.h
new file mode 100644
--- /dev/null
+++ b/Engine/Plugins/Experimental/PCG/Source/PCG/Public/Data/PCGPrimitiveData.h
@@ -0,0 +1,40 @@
+// Copyright Epic Games, Inc. All Rights Reserved.
+
+#pragma once
+
+#include "PCGSpatialData.h"
+
+#include "PCGPrimitiveData.generated.h"
+
+UCLASS(BlueprintType, ClassGroup = (Procedural))
+class PCG_API UPCGPrimitiveData : public UPCGSpatialDataWithPointCache
+{
+	GENERATED_BODY()
+
+public:
+	void Initialize(UPrimitiveComponent* InPrim);
+
+	// ~Begin UPCGSpatialData interface
+	virtual int GetDimension() const override { return 3; }
+	virtual FBox GetBounds() const override { return CachedBounds; }
+	virtual FBox GetStrictBounds() const override { return CachedStrictBounds; }
+	virtual bool SamplePoint(const FTransform& Transform, const FBox& Bounds, FPCGPoint& OutPoint, UPCGMetadata* OutMetadata) const override;
+	// ~End UPCGSpatialData interface
+
+	// ~Begin UPCGSpatialDataWithPointCache implementation
+	virtual const UPCGPointData* CreatePointData(FPCGContext* Context) const override;
+	// ~End UPCGSpatialDataWithPointCache implementation
+
+	UPROPERTY(BlueprintReadWrite, EditAnywhere, Category = "Data")
+	float VoxelSize = 100.0f;
+
+protected:
+	UPROPERTY(BlueprintReadOnly, VisibleAnywhere, Category = "Data")
+	TObjectPtr<UPrimitiveComponent> Primitive = nullptr;
+
+	UPROPERTY()
+	FBox CachedBounds = FBox(EForceInit::ForceInit);
+
+	UPROPERTY()
+	FBox CachedStrictBounds = FBox(EForceInit::ForceInit);
+};
\ No newline at end of file
diff --git a/Engine/Plugins/Experimental/PCG/Source/PCG/Public/PCGComponent.h b/Engine/Plugins/Experimental/PCG/Source/PCG/Public/PCGComponent.h
--- a/Engine/Plugins/Experimental/PCG/Source/PCG/Public/PCGComponent.h
+++ b/Engine/Plugins/Experimental/PCG/Source/PCG/Public/PCGComponent.h
@@ -55,7 +55,6 @@ public:
 	TArray<UPCGData*> GetPCGExclusionData();
 
 	bool CanPartition() const;
-	bool HasSplineComponent() const;
 
 	UPCGGraph* GetGraph() const { return Graph; }
 	void SetGraph(UPCGGraph* InGraph);
