PCG:
*Renamed Artifact category to Spawner and added Debug category.

diff --git a/Engine/Plugins/Experimental/PCG/Source/PCG/Private/Elements/PCGSpawnActor.cpp b/Engine/Plugins/Experimental/PCG/Source/PCG/Private/Elements/PCGSpawnActor.cpp
--- a/Engine/Plugins/Experimental/PCG/Source/PCG/Private/Elements/PCGSpawnActor.cpp
+++ b/Engine/Plugins/Experimental/PCG/Source/PCG/Private/Elements/PCGSpawnActor.cpp
@@ -227,6 +227,6 @@ bool FPCGSpawnActorElement::ExecuteInternal(FPCGContext* Context) const
 #if WITH_EDITOR
 EPCGSettingsType UPCGSpawnActorSettings::GetType() const
 {
-	return GetSubgraph() ? EPCGSettingsType::Subgraph : EPCGSettingsType::Artifact;
+	return GetSubgraph() ? EPCGSettingsType::Subgraph : EPCGSettingsType::Spawner;
 }
 #endif // WITH_EDITOR
\ No newline at end of file
diff --git a/Engine/Plugins/Experimental/PCG/Source/PCG/Public/Elements/PCGDebugElement.h b/Engine/Plugins/Experimental/PCG/Source/PCG/Public/Elements/PCGDebugElement.h
--- a/Engine/Plugins/Experimental/PCG/Source/PCG/Public/Elements/PCGDebugElement.h
+++ b/Engine/Plugins/Experimental/PCG/Source/PCG/Public/Elements/PCGDebugElement.h
@@ -21,6 +21,7 @@ public:
 #if WITH_EDITOR
 	// ~Begin UPCGSettings interface
 	virtual FName GetDefaultNodeName() const override { return FName(TEXT("DebugNode")); }
+	virtual EPCGSettingsType GetType() const override { return EPCGSettingsType::Debug; }
 #endif
 
 protected:
diff --git a/Engine/Plugins/Experimental/PCG/Source/PCG/Public/Elements/PCGStaticMeshSpawner.h b/Engine/Plugins/Experimental/PCG/Source/PCG/Public/Elements/PCGStaticMeshSpawner.h
--- a/Engine/Plugins/Experimental/PCG/Source/PCG/Public/Elements/PCGStaticMeshSpawner.h
+++ b/Engine/Plugins/Experimental/PCG/Source/PCG/Public/Elements/PCGStaticMeshSpawner.h
@@ -38,7 +38,7 @@ public:
 #if WITH_EDITOR
 	// ~Begin UPCGSettings interface
 	virtual FName GetDefaultNodeName() const override { return FName(TEXT("StaticMeshSpawnerNode")); }
-	virtual EPCGSettingsType GetType() const override { return EPCGSettingsType::Artifact; }
+	virtual EPCGSettingsType GetType() const override { return EPCGSettingsType::Spawner; }
 #endif
 
 protected:
diff --git a/Engine/Plugins/Experimental/PCG/Source/PCG/Public/PCGSettings.h b/Engine/Plugins/Experimental/PCG/Source/PCG/Public/PCGSettings.h
--- a/Engine/Plugins/Experimental/PCG/Source/PCG/Public/PCGSettings.h
+++ b/Engine/Plugins/Experimental/PCG/Source/PCG/Public/PCGSettings.h
@@ -33,8 +33,9 @@ enum class EPCGSettingsType : uint8
 	Metadata,
 	Filter,
 	Sampler,
-	Artifact,
+	Spawner,
 	Subgraph,
+	Debug,
 	Generic
 };
 
diff --git a/Engine/Plugins/Experimental/PCG/Source/PCGEditor/Private/PCGEditorSettings.cpp b/Engine/Plugins/Experimental/PCG/Source/PCGEditor/Private/PCGEditorSettings.cpp
--- a/Engine/Plugins/Experimental/PCG/Source/PCGEditor/Private/PCGEditorSettings.cpp
+++ b/Engine/Plugins/Experimental/PCG/Source/PCGEditor/Private/PCGEditorSettings.cpp
@@ -13,8 +13,9 @@ UPCGEditorSettings::UPCGEditorSettings(const FObjectInitializer& ObjectInitializ
 	MetadataNodeColor = FLinearColor(0.4f, 0.4f, 0.8f);
 	FilterNodeColor = FLinearColor(0.4f, 0.8f, 0.4f);
 	SamplerNodeColor = FLinearColor(0.8f, 1.0f, 0.4f);
-	ArtifactNodeColor = FLinearColor(1.0f, 0.6f, 0.4f);
+	SpawnerNodeColor = FLinearColor(1.0f, 0.6f, 0.4f);
 	SubgraphNodeColor = FLinearColor(1.0f, 0.1f, 0.1f);
+	DebugNodeColor = FLinearColor(1.0f, 0.0f, 1.0f);
 }
 
 FLinearColor UPCGEditorSettings::GetColor(UPCGSettings* Settings) const
@@ -57,14 +58,18 @@ FLinearColor UPCGEditorSettings::GetColor(UPCGSettings* Settings) const
 	{
 		return SamplerNodeColor;
 	}
-	else if (Settings->GetType() == EPCGSettingsType::Artifact)
+	else if (Settings->GetType() == EPCGSettingsType::Spawner)
 	{
-		return ArtifactNodeColor;
+		return SpawnerNodeColor;
 	}
 	else if (Settings->GetType() == EPCGSettingsType::Subgraph)
 	{
 		return SubgraphNodeColor;
 	}
+	else if (Settings->GetType() == EPCGSettingsType::Debug)
+	{
+		return DebugNodeColor;
+	}
 	else
 	{
 		// Finally, we couldn't find any match, so return the default value
diff --git a/Engine/Plugins/Experimental/PCG/Source/PCGEditor/Public/PCGEditorSettings.h b/Engine/Plugins/Experimental/PCG/Source/PCGEditor/Public/PCGEditorSettings.h
--- a/Engine/Plugins/Experimental/PCG/Source/PCGEditor/Public/PCGEditorSettings.h
+++ b/Engine/Plugins/Experimental/PCG/Source/PCGEditor/Public/PCGEditorSettings.h
@@ -51,12 +51,16 @@ public:
 
 	/** Color used for artifact-generating operations */
 	UPROPERTY(EditAnywhere, config, Category = Node, meta = (HideAlphaChannel))
-	FLinearColor ArtifactNodeColor;
+	FLinearColor SpawnerNodeColor;
 
 	/** Color used for subgraph-like operations */
 	UPROPERTY(EditAnywhere, config, Category = Node, meta = (HideAlphaChannel))
 	FLinearColor SubgraphNodeColor;
 
+	/** Color used for debug operations */
+	UPROPERTY(EditAnywhere, config, Category = Node, meta = (HideAlphaChannel))
+	FLinearColor DebugNodeColor;
+
 	/** User-driven color overrides */
 	UPROPERTY(EditAnywhere, config, Category = Node, meta = (HideAlphaChannel))
 	TMap<TSubclassOf<UPCGSettings>, FLinearColor> OverrideNodeColorByClass;
