Fixed static analysis warnings

#rnx

diff --git a/Engine/Plugins/Experimental/PCG/Source/PCG/Private/PCGPin.cpp b/Engine/Plugins/Experimental/PCG/Source/PCG/Private/PCGPin.cpp
--- a/Engine/Plugins/Experimental/PCG/Source/PCG/Private/PCGPin.cpp
+++ b/Engine/Plugins/Experimental/PCG/Source/PCG/Private/PCGPin.cpp
@@ -12,14 +12,16 @@ UPCGPin::UPCGPin(const FObjectInitializer& ObjectInitializer)
 
 bool UPCGPin::AddEdgeTo(UPCGPin* OtherPin)
 {
-	if (OtherPin)
+	if (!OtherPin)
 	{
-		for (UPCGEdge* Edge : Edges)
+		return false;
+	}
+
+	for (UPCGEdge* Edge : Edges)
+	{
+		if (Edge->GetOtherPin(this) == OtherPin)
 		{
-			if (Edge->GetOtherPin(this) == OtherPin)
-			{
-				return false;
-			}
+			return false;
 		}
 	}
 
@@ -37,6 +39,11 @@ bool UPCGPin::AddEdgeTo(UPCGPin* OtherPin)
 
 bool UPCGPin::BreakEdgeTo(UPCGPin* OtherPin)
 {
+	if (!OtherPin)
+	{
+		return false;
+	}
+
 	for (UPCGEdge* Edge : Edges)
 	{
 		if (Edge->GetOtherPin(this) == OtherPin)
