Tentative fix for editor test crash

diff --git a/Engine/Plugins/Experimental/PCG/Source/PCG/Private/PCGGraph.cpp b/Engine/Plugins/Experimental/PCG/Source/PCG/Private/PCGGraph.cpp
--- a/Engine/Plugins/Experimental/PCG/Source/PCG/Private/PCGGraph.cpp
+++ b/Engine/Plugins/Experimental/PCG/Source/PCG/Private/PCGGraph.cpp
@@ -107,12 +107,16 @@ void UPCGGraph::BeginDestroy()
 	InputNode->OnNodeChangedDelegate.RemoveAll(this);
 
 	// Notify the compiler to remove this graph from its cache
-	UWorld* World = GEditor->GetEditorWorldContext().World();
-	UPCGSubsystem* PCGSubsystem = World ? World->GetSubsystem<UPCGSubsystem>() : nullptr;
-	if (PCGSubsystem)
+	if (GEditor)
 	{
-		PCGSubsystem->NotifyGraphChanged(this);
+		UWorld* World = GEditor->GetEditorWorldContext().World();
+		UPCGSubsystem* PCGSubsystem = World ? World->GetSubsystem<UPCGSubsystem>() : nullptr;
+		if (PCGSubsystem)
+		{
+			PCGSubsystem->NotifyGraphChanged(this);
+		}
 	}
+
 #endif
 
 	Super::BeginDestroy();
@@ -511,7 +515,7 @@ void UPCGGraph::NotifyGraphChanged(bool bIsStructural)
 	}
 
 	// Notify the subsystem/compiler cache before so it gets recompiled properly
-	if (bIsStructural)
+	if (bIsStructural && GEditor)
 	{
 		UWorld* World = GEditor->GetEditorWorldContext().World();
 		UPCGSubsystem* PCGSubsystem = World ? World->GetSubsystem<UPCGSubsystem>() : nullptr;
diff --git a/Engine/Plugins/Experimental/PCG/Source/PCG/Private/PCGSettings.cpp b/Engine/Plugins/Experimental/PCG/Source/PCG/Private/PCGSettings.cpp
--- a/Engine/Plugins/Experimental/PCG/Source/PCG/Private/PCGSettings.cpp
+++ b/Engine/Plugins/Experimental/PCG/Source/PCG/Private/PCGSettings.cpp
@@ -84,11 +84,14 @@ void UPCGSettings::PostEditChangeProperty(FPropertyChangedEvent& PropertyChanged
 
 void UPCGSettings::DirtyCache()
 {
-	UWorld* World = GEditor->GetEditorWorldContext().World();
-	UPCGSubsystem* PCGSubsystem = World ? World->GetSubsystem<UPCGSubsystem>() : nullptr;
-	if (PCGSubsystem)
+	if (GEditor)
 	{
-		PCGSubsystem->CleanFromCache(GetElement().Get());
+		UWorld* World = GEditor->GetEditorWorldContext().World();
+		UPCGSubsystem* PCGSubsystem = World ? World->GetSubsystem<UPCGSubsystem>() : nullptr;
+		if (PCGSubsystem)
+		{
+			PCGSubsystem->CleanFromCache(GetElement().Get());
+		}
 	}
 }
 #endif // WITH_EDITOR
diff --git a/Engine/Plugins/Experimental/PCG/Source/PCG/Private/PCGSubsystem.cpp b/Engine/Plugins/Experimental/PCG/Source/PCG/Private/PCGSubsystem.cpp
--- a/Engine/Plugins/Experimental/PCG/Source/PCG/Private/PCGSubsystem.cpp
+++ b/Engine/Plugins/Experimental/PCG/Source/PCG/Private/PCGSubsystem.cpp
@@ -27,9 +27,12 @@ namespace PCGSubsystemConsole
 		TEXT("Clears the PCG results cache."),
 		FConsoleCommandDelegate::CreateLambda([]()
 			{
-				if (UWorld* World = GEditor->GetEditorWorldContext().World())
+				if (GEditor)
 				{
-					World->GetSubsystem<UPCGSubsystem>()->FlushCache();
+					if (UWorld* World = GEditor->GetEditorWorldContext().World())
+					{
+						World->GetSubsystem<UPCGSubsystem>()->FlushCache();
+					}
 				}
 			}));
 }
@@ -478,7 +481,10 @@ void UPCGSubsystem::DeletePartitionActors()
 	};
 
 	// First, clear selection otherwise it might crash
-	GEditor->SelectNone(true, true, false);
+	if (GEditor)
+	{
+		GEditor->SelectNone(true, true, false);
+	}
 
 	FWorldPartitionHelpers::ForEachActorDesc<APCGPartitionActor>(World->GetWorldPartition(), GatherAndDestroyActors);
 
diff --git a/Engine/Plugins/Experimental/PCG/Source/PCGEditor/Private/PCGEditorModule.cpp b/Engine/Plugins/Experimental/PCG/Source/PCGEditor/Private/PCGEditorModule.cpp
--- a/Engine/Plugins/Experimental/PCG/Source/PCGEditor/Private/PCGEditorModule.cpp
+++ b/Engine/Plugins/Experimental/PCG/Source/PCGEditor/Private/PCGEditorModule.cpp
@@ -154,9 +154,12 @@ void FPCGEditorModule::PopulateMenuActions(FMenuBuilder& MenuBuilder)
 		FSlateIcon(),
 		FUIAction(
 			FExecuteAction::CreateLambda([]() {
-				if (UWorld* World = GEditor->GetEditorWorldContext().World())
+				if (GEditor)
 				{
-					World->GetSubsystem<UPCGSubsystem>()->DeletePartitionActors();
+					if (UWorld* World = GEditor->GetEditorWorldContext().World())
+					{
+						World->GetSubsystem<UPCGSubsystem>()->DeletePartitionActors();
+					}
 				}
 			})),
 		NAME_None);
