PCG:
*Added Palette library to the PCG Editor. PCGElementBLueprints and PCGGraphs can now be tagged to be included in the palette library.

diff --git a/Engine/Plugins/Experimental/PCG/Source/PCG/Public/Elements/PCGExecuteBlueprint.h b/Engine/Plugins/Experimental/PCG/Source/PCG/Public/Elements/PCGExecuteBlueprint.h
--- a/Engine/Plugins/Experimental/PCG/Source/PCG/Public/Elements/PCGExecuteBlueprint.h
+++ b/Engine/Plugins/Experimental/PCG/Source/PCG/Public/Elements/PCGExecuteBlueprint.h
@@ -103,6 +103,17 @@ public:
 	UPROPERTY(BlueprintReadWrite, EditAnywhere, Category = "Settings|Input & Output")
 	bool bHasDefaultOutPin = true;
 
+#if WITH_EDITORONLY_DATA
+	UPROPERTY(EditAnywhere, Category = AssetInfo, AssetRegistrySearchable)
+	bool bExposeToLibrary = false;
+
+	UPROPERTY(EditAnywhere, Category = AssetInfo, AssetRegistrySearchable)
+	FText Category;
+
+	UPROPERTY(EditAnywhere, Category = AssetInfo, AssetRegistrySearchable)
+	FText Description;
+#endif
+
 protected:
 #if WITH_EDITOR
 	void OnDependencyChanged(UObject* Object, FPropertyChangedEvent& PropertyChangedEvent);
diff --git a/Engine/Plugins/Experimental/PCG/Source/PCG/Public/PCGGraph.h b/Engine/Plugins/Experimental/PCG/Source/PCG/Public/PCGGraph.h
--- a/Engine/Plugins/Experimental/PCG/Source/PCG/Public/PCGGraph.h
+++ b/Engine/Plugins/Experimental/PCG/Source/PCG/Public/PCGGraph.h
@@ -30,6 +30,17 @@ public:
 	virtual void BeginDestroy() override;
 	/** ~End UObject interface */
 
+#if WITH_EDITORONLY_DATA
+	UPROPERTY(EditAnywhere, Category = AssetInfo, AssetRegistrySearchable)
+	bool bExposeToLibrary = false;
+
+	UPROPERTY(EditAnywhere, Category = AssetInfo, AssetRegistrySearchable)
+	FText Category;
+
+	UPROPERTY(EditAnywhere, Category = AssetInfo, AssetRegistrySearchable)
+	FText Description;
+#endif
+
 	UPROPERTY(BlueprintReadWrite, EditAnywhere, Category = Template)
 	TSubclassOf<UPCGGraphSetupBP> GraphTemplate;
 
diff --git a/Engine/Plugins/Experimental/PCG/Source/PCGEditor/Private/PCGEditor.cpp b/Engine/Plugins/Experimental/PCG/Source/PCGEditor/Private/PCGEditor.cpp
--- a/Engine/Plugins/Experimental/PCG/Source/PCGEditor/Private/PCGEditor.cpp
+++ b/Engine/Plugins/Experimental/PCG/Source/PCGEditor/Private/PCGEditor.cpp
@@ -7,6 +7,7 @@
 #include "PCGEditorGraphNodeBase.h"
 #include "PCGEditorGraphSchema.h"
 #include "PCGGraph.h"
+#include "SPCGEditorGraphNodePalette.h"
 
 #include "EdGraphUtilities.h"
 #include "Editor.h"
@@ -32,7 +33,7 @@ namespace FPCGEditor_private
 {
 	const FName GraphEditorID = FName(TEXT("GraphEditor"));
 	const FName PropertyDetailsID = FName(TEXT("PropertyDetails"));
-	const FName LibraryID = FName(TEXT("Library"));
+	const FName PaletteID = FName(TEXT("Palette"));
 	const FName AttributesID = FName(TEXT("Attributes"));
 	const FName ViewportID = FName(TEXT("Viewport"));
 }
@@ -51,10 +52,12 @@ void FPCGEditor::Initialize(const EToolkitMode::Type InMode, const TSharedPtr<cl
 	DetailsViewArgs.NameAreaSettings = FDetailsViewArgs::HideNameArea;
 	DetailsViewArgs.bHideSelectionTip = true;
 	PropertyDetailsWidget = PropertyEditorModule.CreateDetailView(DetailsViewArgs);
+	PropertyDetailsWidget->SetObject(PCGGraphBeingEdited);
 
 	GraphEditorWidget = CreateGraphEditorWidget();
+	PaletteWidget = CreatePaletteWidget();
 
-	const TSharedRef<FTabManager::FLayout> StandaloneDefaultLayout = FTabManager::NewLayout("Standalone_PCGGraphEditor_Layout_v0.3")
+	const TSharedRef<FTabManager::FLayout> StandaloneDefaultLayout = FTabManager::NewLayout("Standalone_PCGGraphEditor_Layout_v0.4")
 		->AddArea
 		(
 			FTabManager::NewPrimaryArea()->SetOrientation(Orient_Horizontal)
@@ -75,7 +78,7 @@ void FPCGEditor::Initialize(const EToolkitMode::Type InMode, const TSharedPtr<cl
 					FTabManager::NewStack()
 					->SetSizeCoefficient(0.84)
 					->SetHideTabWell(true)
-					->AddTab(FPCGEditor_private::LibraryID, ETabState::OpenedTab)
+					->AddTab(FPCGEditor_private::PaletteID, ETabState::OpenedTab)
 				)
 			)
 			->Split
@@ -127,8 +130,8 @@ void FPCGEditor::RegisterTabSpawners(const TSharedRef<class FTabManager>& InTabM
 		.SetDisplayName(LOCTEXT("DetailsTab", "Details"))
 		.SetGroup(WorkspaceMenuCategoryRef);
 
-	InTabManager->RegisterTabSpawner(FPCGEditor_private::LibraryID, FOnSpawnTab::CreateSP(this, &FPCGEditor::SpawnTab_Library))
-		.SetDisplayName(LOCTEXT("LibraryTab", "Library"))
+	InTabManager->RegisterTabSpawner(FPCGEditor_private::PaletteID, FOnSpawnTab::CreateSP(this, &FPCGEditor::SpawnTab_Palette))
+		.SetDisplayName(LOCTEXT("PaletteTab", "Palette"))
 		.SetGroup(WorkspaceMenuCategoryRef);
 
 	InTabManager->RegisterTabSpawner(FPCGEditor_private::AttributesID, FOnSpawnTab::CreateSP(this, &FPCGEditor::SpawnTab_Attributes))
@@ -144,7 +147,7 @@ void FPCGEditor::UnregisterTabSpawners(const TSharedRef<class FTabManager>& InTa
 {
 	InTabManager->UnregisterTabSpawner(FPCGEditor_private::GraphEditorID);
 	InTabManager->UnregisterTabSpawner(FPCGEditor_private::PropertyDetailsID);
-	InTabManager->UnregisterTabSpawner(FPCGEditor_private::LibraryID);
+	InTabManager->UnregisterTabSpawner(FPCGEditor_private::PaletteID);
 	InTabManager->UnregisterTabSpawner(FPCGEditor_private::AttributesID);
 	InTabManager->UnregisterTabSpawner(FPCGEditor_private::ViewportID);
 
@@ -570,17 +573,29 @@ TSharedRef<SGraphEditor> FPCGEditor::CreateGraphEditorWidget()
 		.ShowGraphStateOverlay(false);
 }
 
+TSharedRef<SPCGEditorGraphNodePalette> FPCGEditor::CreatePaletteWidget()
+{
+	return SNew(SPCGEditorGraphNodePalette);
+}
+
 void FPCGEditor::OnSelectedNodesChanged(const TSet<UObject*>& NewSelection)
 {
 	TArray<TWeakObjectPtr<UObject>> SelectedObjects;
 
-	for (UObject* Object : NewSelection)
+	if (NewSelection.Num() == 0)
+	{
+		SelectedObjects.Add(PCGGraphBeingEdited);
+	}
+	else
 	{
-		if (UPCGEditorGraphNodeBase* GraphNode = Cast<UPCGEditorGraphNodeBase>(Object))
+		for (UObject* Object : NewSelection)
 		{
-			if (UPCGNode* PCGNode = GraphNode->GetPCGNode())
+			if (UPCGEditorGraphNodeBase* GraphNode = Cast<UPCGEditorGraphNodeBase>(Object))
 			{
-				SelectedObjects.Add(PCGNode->DefaultSettings);
+				if (UPCGNode* PCGNode = GraphNode->GetPCGNode())
+				{
+					SelectedObjects.Add(PCGNode->DefaultSettings);
+				}
 			}
 		}
 	}
@@ -620,13 +635,13 @@ TSharedRef<SDockTab> FPCGEditor::SpawnTab_PropertyDetails(const FSpawnTabArgs& A
 		];
 }
 
-TSharedRef<SDockTab> FPCGEditor::SpawnTab_Library(const FSpawnTabArgs& Args)
+TSharedRef<SDockTab> FPCGEditor::SpawnTab_Palette(const FSpawnTabArgs& Args)
 {
 	return SNew(SDockTab)
-		.Label(LOCTEXT("PCGLibraryTitle", "Library"))
+		.Label(LOCTEXT("PCGPaletteTitle", "Palette"))
 		.TabColorScale(GetTabColorScale())
 		[
-			SNullWidget::NullWidget
+			PaletteWidget.ToSharedRef()
 		];
 }
 
diff --git a/Engine/Plugins/Experimental/PCG/Source/PCGEditor/Private/PCGEditor.h b/Engine/Plugins/Experimental/PCG/Source/PCGEditor/Private/PCGEditor.h
--- a/Engine/Plugins/Experimental/PCG/Source/PCGEditor/Private/PCGEditor.h
+++ b/Engine/Plugins/Experimental/PCG/Source/PCGEditor/Private/PCGEditor.h
@@ -5,11 +5,12 @@
 #include "Toolkits/AssetEditorToolkit.h"
 #include "EditorUndoClient.h"
 
-class UPCGGraph;
-class UPCGEditorGraph;
-class SGraphEditor;
-class IDetailsView;
 class FUICommandList;
+class IDetailsView;
+class SGraphEditor;
+class SPCGEditorGraphNodePalette;
+class UPCGEditorGraph;
+class UPCGGraph;
 
 class FPCGEditor : public FAssetEditorToolkit, public FSelfRegisteringEditorUndoClient
 {
@@ -81,6 +82,9 @@ private:
 	/** Create new graph editor widget */
 	TSharedRef<SGraphEditor> CreateGraphEditorWidget();
 
+	/** Create new palette widget */
+	TSharedRef<SPCGEditorGraphNodePalette> CreatePaletteWidget();
+
 	/** Called when the selection changes in the GraphEditor */
 	void OnSelectedNodesChanged(const TSet<UObject*>& NewSelection);
 
@@ -89,12 +93,13 @@ private:
 
 	TSharedRef<SDockTab> SpawnTab_GraphEditor(const FSpawnTabArgs& Args);
 	TSharedRef<SDockTab> SpawnTab_PropertyDetails(const FSpawnTabArgs& Args);
-	TSharedRef<SDockTab> SpawnTab_Library(const FSpawnTabArgs& Args);
+	TSharedRef<SDockTab> SpawnTab_Palette(const FSpawnTabArgs& Args);
 	TSharedRef<SDockTab> SpawnTab_Attributes(const FSpawnTabArgs& Args);
 	TSharedRef<SDockTab> SpawnTab_Viewport(const FSpawnTabArgs& Args);
 
 	TSharedPtr<SGraphEditor> GraphEditorWidget;
 	TSharedPtr<IDetailsView> PropertyDetailsWidget;
+	TSharedPtr<SPCGEditorGraphNodePalette> PaletteWidget;
 
 	TSharedPtr<FUICommandList> GraphEditorCommands;
 
diff --git a/Engine/Plugins/Experimental/PCG/Source/PCGEditor/Private/PCGEditorGraphSchema.cpp b/Engine/Plugins/Experimental/PCG/Source/PCGEditor/Private/PCGEditorGraphSchema.cpp
--- a/Engine/Plugins/Experimental/PCG/Source/PCGEditor/Private/PCGEditorGraphSchema.cpp
+++ b/Engine/Plugins/Experimental/PCG/Source/PCGEditor/Private/PCGEditorGraphSchema.cpp
@@ -2,6 +2,7 @@
 
 #include "PCGEditorGraphSchema.h"
 
+#include "Elements/PCGExecuteBlueprint.h"
 #include "PCGEditorCommon.h"
 #include "PCGEditorGraph.h"
 #include "PCGEditorGraphNodeBase.h"
@@ -9,37 +10,28 @@
 #include "PCGGraph.h"
 #include "PCGSettings.h"
 
+#include "AssetRegistry/AssetRegistryModule.h"
+#include "Engine/Blueprint.h"
 #include "ScopedTransaction.h"
 #include "UObject/UObjectIterator.h"
 
 #define LOCTEXT_NAMESPACE "PCGEditorGraphSchema"
 
+
+void UPCGEditorGraphSchema::GetPaletteActions(FGraphActionMenuBuilder& ActionMenuBuilder) const
+{
+	GetNativeElementActions(ActionMenuBuilder);
+	GetSubgraphElementActions(ActionMenuBuilder);
+	GetBlueprintElementActions(ActionMenuBuilder);
+}
+
 void UPCGEditorGraphSchema::GetGraphContextActions(FGraphContextMenuBuilder& ContextMenuBuilder) const
 {
 	Super::GetGraphContextActions(ContextMenuBuilder);
 
-	const UPCGEditorGraph* PCGEditorGraph = CastChecked<UPCGEditorGraph>(ContextMenuBuilder.CurrentGraph);
-
-	TArray<UClass*> SettingsClasses;
-	for (TObjectIterator<UClass> It; It; ++It)
-	{
-		UClass* Class = *It;
-		
-		if (Class->IsChildOf(UPCGSettings::StaticClass()) &&
-			!Class->HasAnyClassFlags(CLASS_Abstract | CLASS_Deprecated | CLASS_Hidden))
-		{
-			SettingsClasses.Add(Class);
-		}
-	}
-
-	for (UClass* SettingsClass : SettingsClasses)
-	{
-		const FText MenuDesc = FText::FromName(SettingsClass->GetFName());
-
-		TSharedPtr<FPCGEditorGraphSchemaAction_NewNode> NewAction(new FPCGEditorGraphSchemaAction_NewNode(FText::GetEmpty(), MenuDesc, FText::GetEmpty(), 0));
-		NewAction->SettingsClass = SettingsClass;
-		ContextMenuBuilder.AddAction(NewAction);
-	}
+	GetNativeElementActions(ContextMenuBuilder);
+	GetSubgraphElementActions(ContextMenuBuilder);
+	GetBlueprintElementActions(ContextMenuBuilder);
 }
 
 FLinearColor UPCGEditorGraphSchema::GetPinTypeColor(const FEdGraphPinType& PinType) const
@@ -147,6 +139,90 @@ void UPCGEditorGraphSchema::BreakSinglePinLink(UEdGraphPin* SourcePin, UEdGraphP
 	PCGGraph->RemoveEdge(SourcePCGNode, SourcePin->PinName, TargetPCGNode, TargetPin->PinName);
 }
 
+void UPCGEditorGraphSchema::GetNativeElementActions(FGraphActionMenuBuilder& ActionMenuBuilder) const
+{
+	TArray<UClass*> SettingsClasses;
+	for (TObjectIterator<UClass> It; It; ++It)
+	{
+		UClass* Class = *It;
+
+		if (Class->IsChildOf(UPCGSettings::StaticClass()) &&
+			!Class->HasAnyClassFlags(CLASS_Abstract | CLASS_Deprecated | CLASS_Hidden))
+		{
+			SettingsClasses.Add(Class);
+		}
+	}
+
+	for (UClass* SettingsClass : SettingsClasses)
+	{
+		if (const UPCGSettings* PCGSettings = SettingsClass->GetDefaultObject<UPCGSettings>())
+		{
+			const FText MenuDesc = FText::FromName(PCGSettings->GetDefaultNodeName());
+			const FText Category = StaticEnum<EPCGSettingsType>()->GetDisplayNameTextByValue(static_cast<__underlying_type(EPCGSettingsType)>(PCGSettings->GetType()));
+			const FText Description = FText::GetEmpty();
+
+			TSharedPtr<FPCGEditorGraphSchemaAction_NewNativeElement> NewAction(new FPCGEditorGraphSchemaAction_NewNativeElement(Category, MenuDesc, Description, 0));
+			NewAction->SettingsClass = SettingsClass;
+			ActionMenuBuilder.AddAction(NewAction);
+		}
+	}
+}
+
+void UPCGEditorGraphSchema::GetBlueprintElementActions(FGraphActionMenuBuilder& ActionMenuBuilder) const
+{
+	FAssetRegistryModule& AssetRegistryModule = FModuleManager::LoadModuleChecked<FAssetRegistryModule>(TEXT("AssetRegistry"));
+
+	FARFilter Filter;
+	Filter.ClassNames.Add(UBlueprint::StaticClass()->GetFName());
+	Filter.bRecursiveClasses = true;
+	Filter.TagsAndValues.Add(FBlueprintTags::NativeParentClassPath, FString::Printf(TEXT("%s'%s'"), *UClass::StaticClass()->GetName(), *UPCGBlueprintElement::StaticClass()->GetPathName()));
+
+	TArray<FAssetData> BlueprintElementAssets;
+	AssetRegistryModule.Get().GetAssets(Filter, BlueprintElementAssets);
+
+	for (const FAssetData& AssetData : BlueprintElementAssets)
+	{
+		const bool bExposeToLibrary = AssetData.GetTagValueRef<bool>(TEXT("bExposeToLibrary"));
+		if (bExposeToLibrary)
+		{
+			const FText MenuDesc = FText::FromName(AssetData.AssetName);
+			const FText Category = AssetData.GetTagValueRef<FText>(TEXT("Category"));
+			const FText Description = AssetData.GetTagValueRef<FText>(TEXT("Description"));
+
+			const FString GeneratedClass = AssetData.GetTagValueRef<FString>(FBlueprintTags::GeneratedClassPath);
+
+			TSharedPtr<FPCGEditorGraphSchemaAction_NewBlueprintElement> NewBlueprintAction(new FPCGEditorGraphSchemaAction_NewBlueprintElement(Category, MenuDesc, Description, 0));
+			NewBlueprintAction->BlueprintClassPath = FSoftClassPath(GeneratedClass);
+			ActionMenuBuilder.AddAction(NewBlueprintAction);
+		}
+	}
+}
+
+void UPCGEditorGraphSchema::GetSubgraphElementActions(FGraphActionMenuBuilder& ActionMenuBuilder) const
+{
+	FAssetRegistryModule& AssetRegistryModule = FModuleManager::LoadModuleChecked<FAssetRegistryModule>(TEXT("AssetRegistry"));
+
+	TArray<FAssetData> AssetDataList;
+	AssetRegistryModule.Get().GetAssetsByClass(UPCGGraph::StaticClass()->GetFName(), AssetDataList);
+
+	for (const FAssetData& AssetData : AssetDataList)
+	{
+		const bool bExposeToLibrary = AssetData.GetTagValueRef<bool>(TEXT("bExposeToLibrary"));
+		if (bExposeToLibrary)
+		{
+			const FText MenuDesc = FText::FromName(AssetData.AssetName);
+			const FText Category = AssetData.GetTagValueRef<FText>(TEXT("Category"));
+			const FText Description = AssetData.GetTagValueRef<FText>(TEXT("Description"));
+
+			TSharedPtr<FPCGEditorGraphSchemaAction_NewSubgraphElement> NewSubgraphAction(new FPCGEditorGraphSchemaAction_NewSubgraphElement(Category, MenuDesc, Description, 0));
+			NewSubgraphAction->SubgraphObjectPath = AssetData.ObjectPath;
+			ActionMenuBuilder.AddAction(NewSubgraphAction);
+		}
+	}
+}
+
+
+
 FPCGEditorConnectionDrawingPolicy::FPCGEditorConnectionDrawingPolicy(int32 InBackLayerID, int32 InFrontLayerID, float InZoomFactor, const FSlateRect& InClippingRect, FSlateWindowElementList& InDrawElements, UEdGraph* InGraph)
 	: FConnectionDrawingPolicy(InBackLayerID, InFrontLayerID, InZoomFactor, InClippingRect, InDrawElements)
 	, Graph(CastChecked<UPCGEditorGraph>(InGraph))
diff --git a/Engine/Plugins/Experimental/PCG/Source/PCGEditor/Private/PCGEditorGraphSchema.h b/Engine/Plugins/Experimental/PCG/Source/PCGEditor/Private/PCGEditorGraphSchema.h
--- a/Engine/Plugins/Experimental/PCG/Source/PCGEditor/Private/PCGEditorGraphSchema.h
+++ b/Engine/Plugins/Experimental/PCG/Source/PCGEditor/Private/PCGEditorGraphSchema.h
@@ -15,6 +15,8 @@ class UPCGEditorGraphSchema : public UEdGraphSchema
 	GENERATED_BODY()
 
 public:
+	void GetPaletteActions(FGraphActionMenuBuilder& ActionMenuBuilder) const;
+
 	//~ Begin EdGraphSchema Interface
 	virtual void GetGraphContextActions(FGraphContextMenuBuilder& ContextMenuBuilder) const;
 	virtual FLinearColor GetPinTypeColor(const FEdGraphPinType& PinType) const override;
@@ -24,6 +26,11 @@ public:
 	virtual void BreakSinglePinLink(UEdGraphPin* SourcePin, UEdGraphPin* TargetPin) const override;
 	virtual FConnectionDrawingPolicy* CreateConnectionDrawingPolicy(int32 InBackLayerID, int32 InFrontLayerID, float InZoomFactor, const FSlateRect& InClippingRect, class FSlateWindowElementList& InDrawElements, class UEdGraph* InGraphObj) const override;
 	//~ End EdGraphSchema Interface
+
+private:
+	void GetNativeElementActions(FGraphActionMenuBuilder& ActionMenuBuilder) const;
+	void GetBlueprintElementActions(FGraphActionMenuBuilder& ActionMenuBuilder) const;
+	void GetSubgraphElementActions(FGraphActionMenuBuilder& ActionMenuBuilder) const;
 };
 
 class FPCGEditorConnectionDrawingPolicy : public FConnectionDrawingPolicy
@@ -34,4 +41,4 @@ public:
 
 protected:
 	UPCGEditorGraph* Graph;
-};
\ No newline at end of file
+};
diff --git a/Engine/Plugins/Experimental/PCG/Source/PCGEditor/Private/PCGEditorGraphSchemaActions.cpp b/Engine/Plugins/Experimental/PCG/Source/PCGEditor/Private/PCGEditorGraphSchemaActions.cpp
--- a/Engine/Plugins/Experimental/PCG/Source/PCGEditor/Private/PCGEditorGraphSchemaActions.cpp
+++ b/Engine/Plugins/Experimental/PCG/Source/PCGEditor/Private/PCGEditorGraphSchemaActions.cpp
@@ -2,18 +2,21 @@
 
 #include "PCGEditorGraphSchemaActions.h"
 
+#include "Elements/PCGExecuteBlueprint.h"
 #include "PCGEditorCommon.h"
 #include "PCGEditorGraph.h"
 #include "PCGEditorGraphNode.h"
 #include "PCGGraph.h"
 #include "PCGSettings.h"
+#include "PCGSubgraph.h"
 
 #include "EdGraph/EdGraph.h"
+#include "Engine/Blueprint.h"
 #include "ScopedTransaction.h"
 
-#define LOCTEXT_NAMESPACE "PCGEditorGraphSchemaAction_NewNode"
+#define LOCTEXT_NAMESPACE "PCGEditorGraphSchemaActions"
 
-UEdGraphNode* FPCGEditorGraphSchemaAction_NewNode::PerformAction(UEdGraph* ParentGraph, UEdGraphPin* FromPin, const FVector2D Location, bool bSelectNewNode)
+UEdGraphNode* FPCGEditorGraphSchemaAction_NewNativeElement::PerformAction(UEdGraph* ParentGraph, UEdGraphPin* FromPin, const FVector2D Location, bool bSelectNewNode)
 {
 	UPCGEditorGraph* EditorGraph = Cast<UPCGEditorGraph>(ParentGraph);
 	if (!EditorGraph)
@@ -27,7 +30,7 @@ UEdGraphNode* FPCGEditorGraphSchemaAction_NewNode::PerformAction(UEdGraph* Paren
 		return nullptr;
 	}
 
-	const FScopedTransaction Transaction(*FPCGEditorCommon::ContextIdentifier, LOCTEXT("PCGEditorNewNode", "PCG Editor: New Node"), nullptr);
+	const FScopedTransaction Transaction(*FPCGEditorCommon::ContextIdentifier, LOCTEXT("PCGEditorNewNativeElement", "PCG Editor: New Native Element"), nullptr);
 	EditorGraph->Modify();
 
 	UPCGSettings* DefaultNodeSettings = nullptr;
@@ -51,4 +54,93 @@ UEdGraphNode* FPCGEditorGraphSchemaAction_NewNode::PerformAction(UEdGraph* Paren
 	return NewNode;
 }
 
+UEdGraphNode* FPCGEditorGraphSchemaAction_NewBlueprintElement::PerformAction(UEdGraph* ParentGraph, UEdGraphPin* FromPin, const FVector2D Location, bool bSelectNewNode)
+{
+	UPCGEditorGraph* EditorGraph = Cast<UPCGEditorGraph>(ParentGraph);
+	if (!EditorGraph)
+	{
+		return nullptr;
+	}
+
+	UPCGGraph* PCGGraph = EditorGraph->GetPCGGraph();
+	if (!PCGGraph)
+	{
+		return nullptr;
+	}
+
+	const FScopedTransaction Transaction(*FPCGEditorCommon::ContextIdentifier, LOCTEXT("PCGEditorNewBlueprintELement", "PCG Editor: New Blueprint Element"), nullptr);
+	EditorGraph->Modify();
+
+	UPCGSettings* DefaultNodeSettings = nullptr;
+	UPCGNode* NewPCGNode = PCGGraph->AddNodeOfType(UPCGBlueprintSettings::StaticClass(), DefaultNodeSettings);
+	UPCGBlueprintSettings* DefaultBlueprintSettings = CastChecked<UPCGBlueprintSettings>(DefaultNodeSettings);
+
+	UPCGBlueprintElement* ElementInstance = nullptr;
+	TSubclassOf<UPCGBlueprintElement> BlueprintClass = BlueprintClassPath.TryLoadClass<UPCGBlueprintElement>();
+	DefaultBlueprintSettings->SetElementType(BlueprintClass, ElementInstance);
+
+	FGraphNodeCreator<UPCGEditorGraphNode> NodeCreator(*EditorGraph);
+	UPCGEditorGraphNode* NewNode = NodeCreator.CreateUserInvokedNode(bSelectNewNode);
+	NewNode->Construct(NewPCGNode, EPCGEditorGraphNodeType::Settings);
+	NewNode->NodePosX = Location.X;
+	NewNode->NodePosY = Location.Y;
+	NodeCreator.Finalize();
+
+	NewPCGNode->PositionX = Location.X;
+	NewPCGNode->PositionY = Location.Y;
+
+	if (FromPin)
+	{
+		NewNode->AutowireNewNode(FromPin);
+	}
+
+	return NewNode;
+}
+
+UEdGraphNode* FPCGEditorGraphSchemaAction_NewSubgraphElement::PerformAction(UEdGraph* ParentGraph, UEdGraphPin* FromPin, const FVector2D Location, bool bSelectNewNode)
+{
+	UPCGEditorGraph* EditorGraph = Cast<UPCGEditorGraph>(ParentGraph);
+	if (!EditorGraph)
+	{
+		return nullptr;
+	}
+
+	UPCGGraph* PCGGraph = EditorGraph->GetPCGGraph();
+	if (!PCGGraph)
+	{
+		return nullptr;
+	}
+
+	UPCGGraph* Subgraph = CastChecked<UPCGGraph>(SubgraphObjectPath.TryLoad());
+	if (Subgraph == PCGGraph)
+	{
+		return nullptr;
+	}
+
+	const FScopedTransaction Transaction(*FPCGEditorCommon::ContextIdentifier, LOCTEXT("PCGEditorNewSubgraphElement", "PCG Editor: New Subgraph Element"), nullptr);
+	EditorGraph->Modify();
+
+	UPCGSettings* DefaultNodeSettings = nullptr;
+	UPCGNode* NewPCGNode = PCGGraph->AddNodeOfType(UPCGSubgraphSettings::StaticClass(), DefaultNodeSettings);
+	UPCGSubgraphSettings* DefaultSubgraphSettings = CastChecked<UPCGSubgraphSettings>(DefaultNodeSettings);
+	DefaultSubgraphSettings->Subgraph = Subgraph;
+
+	FGraphNodeCreator<UPCGEditorGraphNode> NodeCreator(*EditorGraph);
+	UPCGEditorGraphNode* NewNode = NodeCreator.CreateUserInvokedNode(bSelectNewNode);
+	NewNode->Construct(NewPCGNode, EPCGEditorGraphNodeType::Settings);
+	NewNode->NodePosX = Location.X;
+	NewNode->NodePosY = Location.Y;
+	NodeCreator.Finalize();
+
+	NewPCGNode->PositionX = Location.X;
+	NewPCGNode->PositionY = Location.Y;
+
+	if (FromPin)
+	{
+		NewNode->AutowireNewNode(FromPin);
+	}
+
+	return NewNode;
+}
+
 #undef LOCTEXT_NAMESPACE
diff --git a/Engine/Plugins/Experimental/PCG/Source/PCGEditor/Private/PCGEditorGraphSchemaActions.h b/Engine/Plugins/Experimental/PCG/Source/PCGEditor/Private/PCGEditorGraphSchemaActions.h
--- a/Engine/Plugins/Experimental/PCG/Source/PCGEditor/Private/PCGEditorGraphSchemaActions.h
+++ b/Engine/Plugins/Experimental/PCG/Source/PCGEditor/Private/PCGEditorGraphSchemaActions.h
@@ -4,15 +4,17 @@
 
 #include "EdGraph/EdGraphSchema.h"
 #include "Templates/SubclassOf.h"
+#include "UObject/SoftObjectPath.h"
 
 #include "PCGEditorGraphSchemaActions.generated.h"
 
 class UEdGraph;
 class UEdGraphPin;
+class UPCGBlueprintElement;
 class UPCGSettings;
 
 USTRUCT()
-struct FPCGEditorGraphSchemaAction_NewNode : public FEdGraphSchemaAction
+struct FPCGEditorGraphSchemaAction_NewNativeElement : public FEdGraphSchemaAction
 {
 	GENERATED_BODY()
 
@@ -22,5 +24,39 @@ struct FPCGEditorGraphSchemaAction_NewNode : public FEdGraphSchemaAction
 	UPROPERTY()
 	TSubclassOf<UPCGSettings> SettingsClass;
 
+	//~ Begin FEdGraphSchemaAction Interface
 	virtual UEdGraphNode* PerformAction(UEdGraph* ParentGraph, UEdGraphPin* FromPin, const FVector2D Location, bool bSelectNewNode = true) override;
+	//~ End FEdGraphSchemaAction Interface
+};
+
+USTRUCT()
+struct FPCGEditorGraphSchemaAction_NewBlueprintElement : public FEdGraphSchemaAction
+{
+	GENERATED_BODY()
+
+	// Inherit the base class's constructors
+	using FEdGraphSchemaAction::FEdGraphSchemaAction;
+
+	UPROPERTY()
+	FSoftClassPath BlueprintClassPath;
+
+	//~ Begin FEdGraphSchemaAction Interface
+	virtual UEdGraphNode* PerformAction(UEdGraph* ParentGraph, UEdGraphPin* FromPin, const FVector2D Location, bool bSelectNewNode = true) override;
+	//~ End FEdGraphSchemaAction Interface
+};
+
+USTRUCT()
+struct FPCGEditorGraphSchemaAction_NewSubgraphElement : public FEdGraphSchemaAction
+{
+	GENERATED_BODY()
+
+	// Inherit the base class's constructors
+	using FEdGraphSchemaAction::FEdGraphSchemaAction;
+
+	UPROPERTY()
+	FSoftObjectPath SubgraphObjectPath;
+
+	//~ Begin FEdGraphSchemaAction Interface
+	virtual UEdGraphNode* PerformAction(UEdGraph* ParentGraph, UEdGraphPin* FromPin, const FVector2D Location, bool bSelectNewNode = true) override;
+	//~ End FEdGraphSchemaAction Interface
 };
diff --git a/Engine/Plugins/Experimental/PCG/Source/PCGEditor/Private/SPCGEditorGraphNodePalette.cpp b/Engine/Plugins/Experimental/PCG/Source/PCGEditor/Private/SPCGEditorGraphNodePalette.cpp
new file mode 100644
--- /dev/null
+++ b/Engine/Plugins/Experimental/PCG/Source/PCGEditor/Private/SPCGEditorGraphNodePalette.cpp
@@ -0,0 +1,101 @@
+// Copyright Epic Games, Inc. All Rights Reserved.
+
+#include "SPCGEditorGraphNodePalette.h"
+
+#include "Elements/PCGExecuteBlueprint.h"
+#include "PCGEditor.h"
+#include "PCGEditorGraphSchema.h"
+#include "PCGGraph.h"
+
+#include "AssetRegistry/AssetRegistryModule.h"
+#include "Modules/ModuleManager.h"
+#include "SGraphActionMenu.h"
+
+#define LOCTEXT_NAMESPACE "SPCGEditorGraphNodePalette"
+
+void SPCGEditorGraphNodePaletteItem::Construct(const FArguments& InArgs, FCreateWidgetForActionData* const InCreateData)
+{
+	check(InCreateData->Action.IsValid());
+
+	TSharedPtr<FEdGraphSchemaAction> GraphAction = InCreateData->Action;
+	ActionPtr = InCreateData->Action;
+
+	bool bIsReadOnly = false;
+	TSharedRef<SWidget> NameSlotWidget = CreateTextSlotWidget(InCreateData, bIsReadOnly);
+
+	this->ChildSlot
+	[
+		NameSlotWidget
+	];
+}
+
+FText SPCGEditorGraphNodePaletteItem::GetItemTooltip() const
+{
+	return ActionPtr.Pin()->GetTooltipDescription();
+}
+
+
+
+void SPCGEditorGraphNodePalette::Construct(const FArguments& InArgs)
+{
+	this->ChildSlot
+	[
+		SAssignNew(GraphActionMenu, SGraphActionMenu)
+		.OnActionDragged(this, &SPCGEditorGraphNodePalette::OnActionDragged)
+		.OnCreateWidgetForAction(this, &SPCGEditorGraphNodePalette::OnCreateWidgetForAction)
+		.OnCollectAllActions(this, &SPCGEditorGraphNodePalette::CollectAllActions)
+		.AutoExpandActionMenu(true)
+	];
+
+	FAssetRegistryModule& AssetRegistryModule = FModuleManager::GetModuleChecked<FAssetRegistryModule>(TEXT("AssetRegistry"));
+	AssetRegistryModule.Get().OnAssetAdded().AddSP(this, &SPCGEditorGraphNodePalette::OnAssetChanged);
+	AssetRegistryModule.Get().OnAssetRemoved().AddSP(this, &SPCGEditorGraphNodePalette::OnAssetChanged);
+	AssetRegistryModule.Get().OnAssetUpdated().AddSP(this, &SPCGEditorGraphNodePalette::OnAssetChanged); // Todo, evaluate if this triggers too often
+	AssetRegistryModule.Get().OnAssetRenamed().AddSP(this, &SPCGEditorGraphNodePalette::OnAssetRenamed);
+}
+
+SPCGEditorGraphNodePalette::~SPCGEditorGraphNodePalette()
+{
+	if (FModuleManager::Get().IsModuleLoaded(TEXT("AssetRegistry")))
+	{
+		FAssetRegistryModule& AssetRegistryModule = FModuleManager::GetModuleChecked<FAssetRegistryModule>(TEXT("AssetRegistry"));
+		AssetRegistryModule.Get().OnAssetAdded().RemoveAll(this);
+		AssetRegistryModule.Get().OnAssetRemoved().RemoveAll(this);
+		AssetRegistryModule.Get().OnAssetUpdated().RemoveAll(this);
+		AssetRegistryModule.Get().OnAssetRenamed().RemoveAll(this);
+	}
+}
+
+TSharedRef<SWidget> SPCGEditorGraphNodePalette::OnCreateWidgetForAction(FCreateWidgetForActionData* const InCreateData)
+{
+	return SNew(SPCGEditorGraphNodePaletteItem, InCreateData);
+}
+
+void SPCGEditorGraphNodePalette::CollectAllActions(FGraphActionListBuilderBase& OutAllActions)
+{
+	const UPCGEditorGraphSchema* PCGSchema = GetDefault<UPCGEditorGraphSchema>();
+
+	FGraphActionMenuBuilder ActionMenuBuilder;
+	PCGSchema->GetPaletteActions(ActionMenuBuilder);
+	OutAllActions.Append(ActionMenuBuilder);
+}
+
+void SPCGEditorGraphNodePalette::OnAssetChanged(const FAssetData& InAssetData)
+{
+	FString InNativeParentClassName = InAssetData.GetTagValueRef<FString>(FBlueprintTags::NativeParentClassPath);
+	FString TargetNativeParentClassName = FString::Printf(TEXT("%s'%s'"), *UClass::StaticClass()->GetName(), *UPCGBlueprintElement::StaticClass()->GetPathName());
+
+	if (InAssetData.AssetClass == UPCGGraph::StaticClass()->GetFName() ||
+		(InAssetData.AssetClass == UBlueprint::StaticClass()->GetFName() && InNativeParentClassName == TargetNativeParentClassName))
+	{
+		RefreshActionsList(true);
+	}
+}
+
+void SPCGEditorGraphNodePalette::OnAssetRenamed(const FAssetData& InAssetData, const FString& /*InNewAssetName*/)
+{
+	OnAssetChanged(InAssetData);
+}
+
+#undef LOCTEXT_NAMESPACE
+ 
\ No newline at end of file
diff --git a/Engine/Plugins/Experimental/PCG/Source/PCGEditor/Private/SPCGEditorGraphNodePalette.h b/Engine/Plugins/Experimental/PCG/Source/PCGEditor/Private/SPCGEditorGraphNodePalette.h
new file mode 100644
--- /dev/null
+++ b/Engine/Plugins/Experimental/PCG/Source/PCGEditor/Private/SPCGEditorGraphNodePalette.h
@@ -0,0 +1,41 @@
+// Copyright Epic Games, Inc. All Rights Reserved.
+
+#pragma once
+
+#include "CoreMinimal.h"
+#include "SGraphPalette.h"
+
+class SPCGEditorGraphNodePaletteItem : public SGraphPaletteItem
+{
+public:
+	SLATE_BEGIN_ARGS(SPCGEditorGraphNodePaletteItem) {};
+	SLATE_END_ARGS()
+	
+	void Construct(const FArguments& InArgs, FCreateWidgetForActionData* const InCreateData);
+
+protected:
+	//~ Begin SGraphPaletteItem Interface
+	virtual FText GetItemTooltip() const override;
+	//~ End SGraphPaletteItem Interface
+};
+
+
+class SPCGEditorGraphNodePalette : public SGraphPalette
+{
+public:
+	SLATE_BEGIN_ARGS(SPCGEditorGraphNodePalette) {};
+	SLATE_END_ARGS()
+
+	void Construct(const FArguments& InArgs);
+	virtual ~SPCGEditorGraphNodePalette();
+
+protected:
+	//~ Begin SGraphPalette Interface
+	virtual TSharedRef<SWidget> OnCreateWidgetForAction(FCreateWidgetForActionData* const InCreateData) override;
+	virtual void CollectAllActions(FGraphActionListBuilderBase& OutAllActions) override;
+	//~ End SGraphPalette Interface
+
+private:
+	void OnAssetChanged(const FAssetData& InAssetData);
+	void OnAssetRenamed(const FAssetData& InAssetData, const FString& InNewAssetName);
+};
