Replaced const ref by by-value params in blueprint calls to allow setting values slightly more easily.
Fixed typo

#rnx

diff --git a/Engine/Plugins/Experimental/PCG/Source/PCG/Private/Metadata/PCGMetadata.cpp b/Engine/Plugins/Experimental/PCG/Source/PCG/Private/Metadata/PCGMetadata.cpp
--- a/Engine/Plugins/Experimental/PCG/Source/PCG/Private/Metadata/PCGMetadata.cpp
+++ b/Engine/Plugins/Experimental/PCG/Source/PCG/Private/Metadata/PCGMetadata.cpp
@@ -306,27 +306,27 @@ void UPCGMetadata::CreateFloatAttribute(FName AttributeName, float DefaultValue,
 	CreateAttribute<float>(AttributeName, DefaultValue, bAllowsInterpolation, bOverrideParent);
 }
 
-void UPCGMetadata::CreateVectorAttribute(FName AttributeName, const FVector& DefaultValue, bool bAllowsInterpolation, bool bOverrideParent)
+void UPCGMetadata::CreateVectorAttribute(FName AttributeName, FVector DefaultValue, bool bAllowsInterpolation, bool bOverrideParent)
 {
 	CreateAttribute<FVector>(AttributeName, DefaultValue, bAllowsInterpolation, bOverrideParent);
 }
 
-void UPCGMetadata::CreateVector4Attribute(FName AttributeName, const FVector4& DefaultValue, bool bAllowsInterpolation, bool bOverrideParent)
+void UPCGMetadata::CreateVector4Attribute(FName AttributeName, FVector4 DefaultValue, bool bAllowsInterpolation, bool bOverrideParent)
 {
 	CreateAttribute<FVector4>(AttributeName, DefaultValue, bAllowsInterpolation, bOverrideParent);
 }
 
-void UPCGMetadata::CreateQuatAttribute(FName AttributeName, const FQuat& DefaultValue, bool bAllowsInterpolation, bool bOverrideParent)
+void UPCGMetadata::CreateQuatAttribute(FName AttributeName, FQuat DefaultValue, bool bAllowsInterpolation, bool bOverrideParent)
 {
 	CreateAttribute<FQuat>(AttributeName, DefaultValue, bAllowsInterpolation, bOverrideParent);
 }
 
-void UPCGMetadata::CreateTransformAttribute(FName AttributeName, const FTransform& DefaultValue, bool bAllowsInterpolation, bool bOverrideParent)
+void UPCGMetadata::CreateTransformAttribute(FName AttributeName, FTransform DefaultValue, bool bAllowsInterpolation, bool bOverrideParent)
 {
 	CreateAttribute<FTransform>(AttributeName, DefaultValue, bAllowsInterpolation, bOverrideParent);
 }
 
-void UPCGMetadata::CreateStringAttribute(FName AttributeName, const FString& DefaultValue, bool bAllowsInterpolation, bool bOverrideParent)
+void UPCGMetadata::CreateStringAttribute(FName AttributeName, FString DefaultValue, bool bAllowsInterpolation, bool bOverrideParent)
 {
 	CreateAttribute<FString>(AttributeName, DefaultValue, bAllowsInterpolation, bOverrideParent);
 }
diff --git a/Engine/Plugins/Experimental/PCG/Source/PCG/Public/Metadata/PCGMetadata.h b/Engine/Plugins/Experimental/PCG/Source/PCG/Public/Metadata/PCGMetadata.h
--- a/Engine/Plugins/Experimental/PCG/Source/PCG/Public/Metadata/PCGMetadata.h
+++ b/Engine/Plugins/Experimental/PCG/Source/PCG/Public/Metadata/PCGMetadata.h
@@ -63,25 +63,25 @@ public:
 
 	/** Create new streams */
 	UFUNCTION(BlueprintCallable, Category = "PCG|Metadata")
-	void CreateInteger64Attribute(FName AttributeName, int64 DefaultValue, bool bAllowsInterpolation, bool bOoverrideParent);
+	void CreateInteger64Attribute(FName AttributeName, int64 DefaultValue, bool bAllowsInterpolation, bool bOverrideParent = true);
 
 	UFUNCTION(BlueprintCallable, Category = "PCG|Metadata")
-	void CreateFloatAttribute(FName AttributeName, float DefaultValue, bool bAllowsInterpolation, bool bOverrideParent);
+	void CreateFloatAttribute(FName AttributeName, float DefaultValue, bool bAllowsInterpolation, bool bOverrideParent = true);
 
 	UFUNCTION(BlueprintCallable, Category = "PCG|Metadata")
-	void CreateVectorAttribute(FName AttributeName, const FVector& DefaultValue, bool bAllowsInterpolation, bool bOverrideParent);
+	void CreateVectorAttribute(FName AttributeName, FVector DefaultValue, bool bAllowsInterpolation, bool bOverrideParent = true);
 
 	UFUNCTION(BlueprintCallable, Category = "PCG|Metadata", meta=(DisplayName = "Create Vector4 Attribute"))
-	void CreateVector4Attribute(FName AttributeName, const FVector4& DefaultValue, bool bAllowsInterpolation, bool bOverrideParent);
+	void CreateVector4Attribute(FName AttributeName, FVector4 DefaultValue, bool bAllowsInterpolation, bool bOverrideParent = true);
 
 	UFUNCTION(BlueprintCallable, Category = "PCG|Metadata")
-	void CreateQuatAttribute(FName AttributeName, const FQuat& DefaultValue, bool bAllowsInterpolation, bool bOverrideParent);
+	void CreateQuatAttribute(FName AttributeName, FQuat DefaultValue, bool bAllowsInterpolation, bool bOverrideParent = true);
 
 	UFUNCTION(BlueprintCallable, Category = "PCG|Metadata")
-	void CreateTransformAttribute(FName AttributeName, const FTransform& DefaultValue, bool bAllowsInterpolation, bool bOverrideParent);
+	void CreateTransformAttribute(FName AttributeName, FTransform DefaultValue, bool bAllowsInterpolation, bool bOverrideParent = true);
 
 	UFUNCTION(BlueprintCallable, Category = "PCG|Metadata")
-	void CreateStringAttribute(FName AttributeName, const FString& DefaultValue, bool bAllowsInterpolation, bool bOverrideParent);
+	void CreateStringAttribute(FName AttributeName, FString DefaultValue, bool bAllowsInterpolation, bool bOverrideParent = true);
 
 	/** Get attributes */
 	FPCGMetadataAttributeBase* GetMutableAttribute(FName AttributeName);
