Fixed crash on metadata merge, replaced by warning.

#rnx

diff --git a/Engine/Plugins/Experimental/PCG/Source/PCG/Private/Metadata/PCGMetadata.cpp b/Engine/Plugins/Experimental/PCG/Source/PCG/Private/Metadata/PCGMetadata.cpp
--- a/Engine/Plugins/Experimental/PCG/Source/PCG/Private/Metadata/PCGMetadata.cpp
+++ b/Engine/Plugins/Experimental/PCG/Source/PCG/Private/Metadata/PCGMetadata.cpp
@@ -491,8 +491,15 @@ PCGMetadataEntryKey UPCGMetadata::GetParentKey(PCGMetadataEntryKey LocalItemKey)
 	else
 	{
 		FReadScopeLock ScopeLock(ItemLock);
-		check(LocalItemKey - ItemKeyOffset < ParentKeys.Num());
-		return ParentKeys[LocalItemKey - ItemKeyOffset];
+		if (LocalItemKey - ItemKeyOffset < ParentKeys.Num())
+		{
+			return ParentKeys[LocalItemKey - ItemKeyOffset];
+		}
+		else
+		{
+			UE_LOG(LogPCG, Warning, TEXT("Invalid metadata key - check for entry key not properly initialized"));
+			return PCGInvalidEntryKey;
+		}
 	}
 }
 
