Replaced simple pin labels by pin properties
Added option to limit non-Any pins to one connection
Added pin coloring based on allowed types
Added wire coloring based on output pins

diff --git a/Engine/Plugins/Experimental/PCG/Source/PCG/Private/Elements/PCGDebugElement.cpp b/Engine/Plugins/Experimental/PCG/Source/PCG/Private/Elements/PCGDebugElement.cpp
--- a/Engine/Plugins/Experimental/PCG/Source/PCG/Private/Elements/PCGDebugElement.cpp
+++ b/Engine/Plugins/Experimental/PCG/Source/PCG/Private/Elements/PCGDebugElement.cpp
@@ -47,21 +47,14 @@ namespace PCGDebugElement
 
 		if (Context->Node)
 		{
-			if (Context->Node->IsOutputPinConnected(NAME_None))
+			TArray<FPCGPinProperties> OutputPinProperties = Context->Node->OutputPinProperties();
+			for(const FPCGPinProperties& OutPin : OutputPinProperties)
 			{
-				bFilterOnPin = true;
-			}
-			else
-			{
-				TArray<FName> OutLabels = Context->Node->OutLabels();
-				for (const FName& OutLabel : OutLabels)
+				if (Context->Node->IsOutputPinConnected(OutPin.Label))
 				{
-					if (Context->Node->IsOutputPinConnected(OutLabel))
-					{
-						PinFilter = OutLabel;
-						bFilterOnPin = true;
-						break;
-					}
+					PinFilter = OutPin.Label;
+					bFilterOnPin = true;
+					break;
 				}
 			}
 		}
diff --git a/Engine/Plugins/Experimental/PCG/Source/PCG/Private/Elements/PCGDifferenceElement.cpp b/Engine/Plugins/Experimental/PCG/Source/PCG/Private/Elements/PCGDifferenceElement.cpp
--- a/Engine/Plugins/Experimental/PCG/Source/PCG/Private/Elements/PCGDifferenceElement.cpp
+++ b/Engine/Plugins/Experimental/PCG/Source/PCG/Private/Elements/PCGDifferenceElement.cpp
@@ -12,9 +12,13 @@ namespace PCGDifferenceConstants
 	const FName DifferencesLabel = TEXT("Differences");
 }
 
-TArray<FName> UPCGDifferenceSettings::InLabels() const
+TArray<FPCGPinProperties> UPCGDifferenceSettings::InputPinProperties() const
 {
-	return { PCGPinConstants::DefaultInputLabel, PCGDifferenceConstants::SourceLabel, PCGDifferenceConstants::DifferencesLabel };
+	TArray<FPCGPinProperties> PinProperties;
+	PinProperties.Emplace(PCGDifferenceConstants::SourceLabel, EPCGDataType::Any);
+	PinProperties.Emplace(PCGDifferenceConstants::DifferencesLabel, EPCGDataType::Spatial);
+
+	return PinProperties;
 }
 
 FPCGElementPtr UPCGDifferenceSettings::CreateElement() const
diff --git a/Engine/Plugins/Experimental/PCG/Source/PCG/Private/Elements/PCGExecuteBlueprint.cpp b/Engine/Plugins/Experimental/PCG/Source/PCG/Private/Elements/PCGExecuteBlueprint.cpp
--- a/Engine/Plugins/Experimental/PCG/Source/PCG/Private/Elements/PCGExecuteBlueprint.cpp
+++ b/Engine/Plugins/Experimental/PCG/Source/PCG/Private/Elements/PCGExecuteBlueprint.cpp
@@ -37,6 +37,26 @@ void UPCGBlueprintElement::PostLoad()
 {
 	Super::PostLoad();
 	Initialize();
+
+	if (!InputPinLabels_DEPRECATED.IsEmpty())
+	{
+		for (const FName& Label : InputPinLabels_DEPRECATED)
+		{
+			CustomInputPins.Emplace(Label);
+		}
+
+		InputPinLabels_DEPRECATED.Reset();
+	}
+
+	if (!OutputPinLabels_DEPRECATED.IsEmpty())
+	{
+		for (const FName& Label : OutputPinLabels_DEPRECATED)
+		{
+			CustomOutputPins.Emplace(Label);
+		}
+
+		OutputPinLabels_DEPRECATED.Reset();
+	}
 }
 
 void UPCGBlueprintElement::BeginDestroy()
@@ -104,6 +124,28 @@ EPCGSettingsType UPCGBlueprintElement::NodeTypeOverride_Implementation() const
 	return EPCGSettingsType::Blueprint;
 }
 
+TSet<FName> UPCGBlueprintElement::InputLabels() const
+{
+	TSet<FName> Labels;
+	for (const FPCGPinProperties& PinProperty : CustomInputPins)
+	{
+		Labels.Emplace(PinProperty.Label);
+	}
+
+	return Labels;
+}
+
+TSet<FName> UPCGBlueprintElement::OutputLabels() const
+{
+	TSet<FName> Labels;
+	for (const FPCGPinProperties& PinProperty : CustomOutputPins)
+	{
+		Labels.Emplace(PinProperty.Label);
+	}
+
+	return Labels;
+}
+
 void UPCGBlueprintSettings::SetupBlueprintEvent()
 {
 #if WITH_EDITOR
@@ -173,6 +215,7 @@ void UPCGBlueprintSettings::PostLoad()
 
 	if (BlueprintElementInstance)
 	{
+		BlueprintElementInstance->ConditionalPostLoad();
 		BlueprintElementInstance->bCreatesArtifacts |= bCreatesArtifacts_DEPRECATED;
 		BlueprintElementInstance->bCanBeMultithreaded |= bCanBeMultithreaded_DEPRECATED;
 	}
@@ -326,46 +369,46 @@ FName UPCGBlueprintSettings::AdditionalTaskName() const
 	}
 }
 
-TArray<FName> UPCGBlueprintSettings::InLabels() const
+TArray<FPCGPinProperties> UPCGBlueprintSettings::InputPinProperties() const
 {
-	TArray<FName> InputLabels;
+	TArray<FPCGPinProperties> PinProperties;
 
 	if (BlueprintElementInstance)
 	{
 		if (BlueprintElementInstance->bHasDefaultInPin)
 		{
-			InputLabels.Add(PCGPinConstants::DefaultInputLabel);
+			PinProperties.Append(Super::InputPinProperties());
 		}
 
-		InputLabels.Append(BlueprintElementInstance->InputPinLabels.Array());
+		PinProperties.Append(BlueprintElementInstance->CustomInputPins);
 	}
 	else
 	{
-		InputLabels = Super::InLabels();
+		PinProperties = Super::InputPinProperties();
 	}
 
-	return InputLabels;
+	return PinProperties;
 }
 
-TArray<FName> UPCGBlueprintSettings::OutLabels() const
+TArray<FPCGPinProperties> UPCGBlueprintSettings::OutputPinProperties() const
 {
-	TArray<FName> OutputLabels;
-	
+	TArray<FPCGPinProperties> PinProperties;
+
 	if (BlueprintElementInstance)
 	{
 		if (BlueprintElementInstance->bHasDefaultOutPin)
 		{
-			OutputLabels.Add(PCGPinConstants::DefaultOutputLabel);
+			PinProperties.Append(Super::OutputPinProperties());
 		}
 
-		OutputLabels.Append(BlueprintElementInstance->OutputPinLabels.Array());
+		PinProperties.Append(BlueprintElementInstance->CustomOutputPins);
 	}
 	else
 	{
-		OutputLabels = Super::OutLabels();
+		PinProperties = Super::OutputPinProperties();
 	}
 
-	return OutputLabels;
+	return PinProperties;
 }
 
 FPCGElementPtr UPCGBlueprintSettings::CreateElement() const
diff --git a/Engine/Plugins/Experimental/PCG/Source/PCG/Private/Elements/PCGPointFilter.cpp b/Engine/Plugins/Experimental/PCG/Source/PCG/Private/Elements/PCGPointFilter.cpp
--- a/Engine/Plugins/Experimental/PCG/Source/PCG/Private/Elements/PCGPointFilter.cpp
+++ b/Engine/Plugins/Experimental/PCG/Source/PCG/Private/Elements/PCGPointFilter.cpp
@@ -575,14 +575,22 @@ namespace PCGPointFilterHelpers
 	}
 }
 
-TArray<FName> UPCGPointFilterSettings::InLabels() const
+TArray<FPCGPinProperties> UPCGPointFilterSettings::InputPinProperties() const
 {
-	return { PCGPointFilterConstants::DataToFilterLabel, PCGPointFilterConstants::FilterLabel };
+	TArray<FPCGPinProperties> PinProperties;
+	PinProperties.Emplace(PCGPointFilterConstants::DataToFilterLabel, EPCGDataType::Any);
+	PinProperties.Emplace(PCGPointFilterConstants::FilterLabel, EPCGDataType::Spatial);
+
+	return PinProperties;
 }
 
-TArray<FName> UPCGPointFilterSettings::OutLabels() const
+TArray<FPCGPinProperties> UPCGPointFilterSettings::OutputPinProperties() const
 {
-	return { PCGPointFilterConstants::InFilterLabel, PCGPointFilterConstants::OutFilterLabel };
+	TArray<FPCGPinProperties> PinProperties;
+	PinProperties.Emplace(PCGPointFilterConstants::InFilterLabel, EPCGDataType::Point);
+	PinProperties.Emplace(PCGPointFilterConstants::OutFilterLabel, EPCGDataType::Point);
+
+	return PinProperties;
 }
 
 FPCGElementPtr UPCGPointFilterSettings::CreateElement() const
diff --git a/Engine/Plugins/Experimental/PCG/Source/PCG/Private/Graph/PCGGraphExecutor.cpp b/Engine/Plugins/Experimental/PCG/Source/PCG/Private/Graph/PCGGraphExecutor.cpp
--- a/Engine/Plugins/Experimental/PCG/Source/PCG/Private/Graph/PCGGraphExecutor.cpp
+++ b/Engine/Plugins/Experimental/PCG/Source/PCG/Private/Graph/PCGGraphExecutor.cpp
@@ -456,7 +456,7 @@ void FPCGGraphExecutor::BuildTaskInput(const FPCGGraphTask& Task, FPCGDataCollec
 		const int32 TaggedDataOffset = TaskInput.TaggedData.Num();
 		if (Input.InPin)
 		{
-			TaskInput.TaggedData.Append(InputCollection.GetInputsByPin(Input.InPin->Label));
+			TaskInput.TaggedData.Append(InputCollection.GetInputsByPin(Input.InPin->Properties.Label));
 		}
 		else
 		{
@@ -469,7 +469,7 @@ void FPCGGraphExecutor::BuildTaskInput(const FPCGGraphTask& Task, FPCGDataCollec
 		{
 			for (int32 TaggedDataIndex = TaggedDataOffset; TaggedDataIndex < TaskInput.TaggedData.Num(); ++TaggedDataIndex)
 			{
-				TaskInput.TaggedData[TaggedDataIndex].Pin = Input.OutPin->Label;
+				TaskInput.TaggedData[TaggedDataIndex].Pin = Input.OutPin->Properties.Label;
 			}
 		}
 	}
diff --git a/Engine/Plugins/Experimental/PCG/Source/PCG/Private/PCGElement.cpp b/Engine/Plugins/Experimental/PCG/Source/PCG/Private/PCGElement.cpp
--- a/Engine/Plugins/Experimental/PCG/Source/PCG/Private/PCGElement.cpp
+++ b/Engine/Plugins/Experimental/PCG/Source/PCG/Private/PCGElement.cpp
@@ -27,7 +27,7 @@ bool IPCGElement::Execute(FPCGContext* Context) const
 	{
 		//Pass-through
 		Context->OutputData = Context->InputData;
-		CleanupAndValidateOutputLabels(Context);
+		CleanupAndValidateOutput(Context);
 		return true;
 	}
 	else
@@ -81,7 +81,7 @@ bool IPCGElement::Execute(FPCGContext* Context) const
 
 		if (bDone)
 		{
-			CleanupAndValidateOutputLabels(Context);
+			CleanupAndValidateOutput(Context);
 
 #if WITH_EDITOR
 			PCGE_LOG(Log, "Executed in (%f)s and (%d) call(s)", Context->ElapsedTime, Context->ExecutionCount);
@@ -152,7 +152,7 @@ void IPCGElement::DebugDisplay(FPCGContext* Context) const
 }
 #endif // WITH_EDITOR
 
-void IPCGElement::CleanupAndValidateOutputLabels(FPCGContext* Context) const
+void IPCGElement::CleanupAndValidateOutput(FPCGContext* Context) const
 {
 	check(Context);
 	const UPCGSettings* Settings = Context->GetInputSettings<UPCGSettings>();
@@ -161,12 +161,12 @@ void IPCGElement::CleanupAndValidateOutputLabels(FPCGContext* Context) const
 	{
 		// Cleanup any residual labels if the node isn't supposed to produce them
 		// TODO: this is a bit of a crutch, could be refactored out if we review the way we push tagged data
-		TArray<FName> OutLabels = Settings->OutLabels();
-		if (OutLabels.Num() == 1)
+		TArray<FPCGPinProperties> OutputPinProperties = Settings->OutputPinProperties();
+		if(OutputPinProperties.Num() == 1)
 		{
 			for (FPCGTaggedData& TaggedData : Context->OutputData.TaggedData)
 			{
-				TaggedData.Pin = OutLabels[0];
+				TaggedData.Pin = OutputPinProperties[0].Label;
 			}
 		}
 
@@ -176,10 +176,15 @@ void IPCGElement::CleanupAndValidateOutputLabels(FPCGContext* Context) const
 		{
 			for (FPCGTaggedData& TaggedData : Context->OutputData.TaggedData)
 			{
-				if (TaggedData.Pin != NAME_None && !OutLabels.Contains(TaggedData.Pin))
+				int32 MatchIndex = OutputPinProperties.IndexOfByPredicate([&TaggedData](const FPCGPinProperties& InProp) { return TaggedData.Pin == InProp.Label; });
+				if (MatchIndex == INDEX_NONE)
 				{
 					PCGE_LOG(Warning, "Output generated for pin %s but cannot be routed", *TaggedData.Pin.ToString());
 				}
+				else if(TaggedData.Data && !(OutputPinProperties[MatchIndex].AllowedTypes & TaggedData.Data->GetDataType()))
+				{
+					PCGE_LOG(Warning, "Output generated for pin %s does not have a compatible type", *TaggedData.Pin.ToString());
+				}
 			}
 		}
 #endif
diff --git a/Engine/Plugins/Experimental/PCG/Source/PCG/Private/PCGGraph.cpp b/Engine/Plugins/Experimental/PCG/Source/PCG/Private/PCGGraph.cpp
--- a/Engine/Plugins/Experimental/PCG/Source/PCG/Private/PCGGraph.cpp
+++ b/Engine/Plugins/Experimental/PCG/Source/PCG/Private/PCGGraph.cpp
@@ -59,7 +59,10 @@ void UPCGGraph::PostLoad()
 	Super::PostLoad();
 
 #if WITH_EDITOR
+	// TODO: review this once the data has been updated, it's unwanted weight going forward
 	// Deprecation
+	InputNode->ConditionalPostLoad();
+
 	if (!Cast<UPCGGraphInputOutputSettings>(InputNode->DefaultSettings))
 	{
 		InputNode->DefaultSettings = NewObject<UPCGGraphInputOutputSettings>(this, TEXT("DefaultInputNodeSettings"));
@@ -67,6 +70,8 @@ void UPCGGraph::PostLoad()
 
 	Cast<UPCGGraphInputOutputSettings>(InputNode->DefaultSettings)->SetInput(true);
 
+	OutputNode->ConditionalPostLoad();
+
 	if (!Cast<UPCGGraphInputOutputSettings>(OutputNode->DefaultSettings))
 	{
 		OutputNode->DefaultSettings = NewObject<UPCGGraphInputOutputSettings>(this, TEXT("DefaultOutputNodeSettings"));
@@ -74,10 +79,25 @@ void UPCGGraph::PostLoad()
 
 	Cast<UPCGGraphInputOutputSettings>(OutputNode->DefaultSettings)->SetInput(false);
 
-	// Additional deprecation
+	// Ensure that all nodes are loaded (& updated their deprecated data)
+	for (UPCGNode* Node : Nodes)
+	{
+		Node->ConditionalPostLoad();
+	}
+	
+	// Update pins on all nodes
+	InputNode->UpdatePins();
+	OutputNode->UpdatePins();
+
+	for (UPCGNode* Node : Nodes)
+	{
+		Node->UpdatePins();
+	}
+
+	// Finally, apply deprecation that changes edges/rebinds
 	InputNode->ApplyDeprecation();
 	OutputNode->ApplyDeprecation();
-
+	
 	for (UPCGNode* Node : Nodes)
 	{
 		Node->ApplyDeprecation();
diff --git a/Engine/Plugins/Experimental/PCG/Source/PCG/Private/PCGInputOutputSettings.cpp b/Engine/Plugins/Experimental/PCG/Source/PCG/Private/PCGInputOutputSettings.cpp
--- a/Engine/Plugins/Experimental/PCG/Source/PCG/Private/PCGInputOutputSettings.cpp
+++ b/Engine/Plugins/Experimental/PCG/Source/PCG/Private/PCGInputOutputSettings.cpp
@@ -3,6 +3,8 @@
 #include "PCGInputOutputSettings.h"
 #include "PCGCommon.h"
 
+#include "Algo/Transform.h"
+
 bool FPCGInputOutputElement::ExecuteInternal(FPCGContext* Context) const
 {
 	// Essentially a pass-through element
@@ -20,4 +22,47 @@ UPCGGraphInputOutputSettings::UPCGGraphInputOutputSettings(const FObjectInitiali
 	StaticAdvancedInLabels.Add(PCGInputOutputConstants::DefaultExcludedActorsLabel);
 	
 	StaticOutLabels.Add(PCGPinConstants::DefaultOutputLabel);
+}
+
+void UPCGGraphInputOutputSettings::PostLoad()
+{
+	Super::PostLoad();
+
+	if (!PinLabels_DEPRECATED.IsEmpty())
+	{
+		for (const FName& PinLabel : PinLabels_DEPRECATED)
+		{
+			CustomPins.Emplace(PinLabel);
+		}
+
+		PinLabels_DEPRECATED.Reset();
+	}
+}
+
+TArray<FPCGPinProperties> UPCGGraphInputOutputSettings::InputPinProperties() const
+{
+	TArray<FPCGPinProperties> PinProperties;
+	Algo::Transform(StaticLabels(), PinProperties, [](const FName& InLabel) { return FPCGPinProperties(InLabel, EPCGDataType::Spatial); });
+	
+	if (bShowAdvancedPins)
+	{
+		Algo::Transform(StaticAdvancedLabels(), PinProperties, [](const FName& InLabel) { return FPCGPinProperties(InLabel, EPCGDataType::Spatial); });
+	}
+
+	PinProperties.Append(CustomPins);
+	return PinProperties;
+}
+
+TArray<FPCGPinProperties> UPCGGraphInputOutputSettings::OutputPinProperties() const
+{
+	TArray<FPCGPinProperties> PinProperties;
+	Algo::Transform(StaticLabels(), PinProperties, [](const FName& InLabel) { return FPCGPinProperties(InLabel, EPCGDataType::Spatial); });
+
+	if (bShowAdvancedPins)
+	{
+		Algo::Transform(StaticAdvancedLabels(), PinProperties, [](const FName& InLabel) { return FPCGPinProperties(InLabel, EPCGDataType::Spatial); });
+	}
+
+	PinProperties.Append(CustomPins);
+	return PinProperties;
 }
\ No newline at end of file
diff --git a/Engine/Plugins/Experimental/PCG/Source/PCG/Private/PCGNode.cpp b/Engine/Plugins/Experimental/PCG/Source/PCG/Private/PCGNode.cpp
--- a/Engine/Plugins/Experimental/PCG/Source/PCG/Private/PCGNode.cpp
+++ b/Engine/Plugins/Experimental/PCG/Source/PCG/Private/PCGNode.cpp
@@ -22,6 +22,7 @@ void UPCGNode::PostLoad()
 	if (DefaultSettings)
 	{
 		DefaultSettings->OnSettingsChangedDelegate.AddUObject(this, &UPCGNode::OnSettingsChanged);
+		DefaultSettings->ConditionalPostLoad();
 	}
 
 	// Make sure legacy nodes support transactions.
@@ -30,19 +31,26 @@ void UPCGNode::PostLoad()
 		SetFlags(RF_Transactional);
 	}
 
-	ApplyDeprecation();
+	for (UPCGPin* InputPin : InputPins)
+	{
+		check(InputPin);
+		InputPin->ConditionalPostLoad();
+	}
+
+	for (UPCGPin* OutputPin : OutputPins)
+	{
+		check(OutputPin);
+		OutputPin->ConditionalPostLoad();
+	}
 #endif
 }
 
 #if WITH_EDITOR
 void UPCGNode::ApplyDeprecation()
 {
-	UpdatePins();
-
 	UPCGPin* DefaultOutputPin = OutputPins.IsEmpty() ? nullptr : OutputPins[0];
 	for (TObjectPtr<UPCGNode> OutboundNode : OutboundNodes_DEPRECATED)
 	{
-		OutboundNode->UpdatePins();
 		UPCGPin* OtherNodeInputPin = OutboundNode->InputPins.IsEmpty() ? nullptr : OutboundNode->InputPins[0];
 
 		if (DefaultOutputPin && OtherNodeInputPin)
@@ -88,8 +96,6 @@ void UPCGNode::ApplyDeprecation()
 			continue;
 		}
 
-		OtherNode->UpdatePins();
-
 		UPCGPin* OtherNodeInputPin = nullptr;
 		if (OutboundEdge->OutboundLabel_DEPRECATED == NAME_None)
 		{
@@ -209,33 +215,25 @@ FName UPCGNode::GetNodeTitle() const
 	return TEXT("Unnamed node");
 }
 
-TArray<FName> UPCGNode::InLabels() const
+TArray<FPCGPinProperties> UPCGNode::InputPinProperties() const
 {
-	TArray<FName> Labels;
-	for (UPCGPin* InputPin : InputPins)
-	{
-		Labels.Add(InputPin->Label);
-	}
-
-	return Labels;
+	TArray<FPCGPinProperties> PinProperties;
+	Algo::Transform(InputPins, PinProperties, [](const UPCGPin* InputPin) { return InputPin->Properties; });
+	return PinProperties;
 }
 
-TArray<FName> UPCGNode::OutLabels() const
+TArray<FPCGPinProperties> UPCGNode::OutputPinProperties() const
 {
-	TArray<FName> Labels;
-	for (UPCGPin* OutputPin : OutputPins)
-	{
-		Labels.Add(OutputPin->Label);
-	}
-
-	return Labels;
+	TArray<FPCGPinProperties> PinProperties;
+	Algo::Transform(OutputPins, PinProperties, [](const UPCGPin* OutputPin) { return OutputPin->Properties; });
+	return PinProperties;
 }
 
 UPCGPin* UPCGNode::GetInputPin(const FName& Label)
 {
 	for (UPCGPin* InputPin : InputPins)
 	{
-		if (InputPin->Label == Label)
+		if (InputPin->Properties.Label == Label)
 		{
 			return InputPin;
 		}
@@ -248,7 +246,7 @@ const UPCGPin* UPCGNode::GetInputPin(const FName& Label) const
 {
 	for (const UPCGPin* InputPin : InputPins)
 	{
-		if (InputPin->Label == Label)
+		if (InputPin->Properties.Label == Label)
 		{
 			return InputPin;
 		}
@@ -261,7 +259,7 @@ UPCGPin* UPCGNode::GetOutputPin(const FName& Label)
 {
 	for (UPCGPin* OutputPin : OutputPins)
 	{
-		if (OutputPin->Label == Label)
+		if (OutputPin->Properties.Label == Label)
 		{
 			return OutputPin;
 		}
@@ -274,7 +272,7 @@ const UPCGPin* UPCGNode::GetOutputPin(const FName& Label) const
 {
 	for (const UPCGPin* OutputPin : OutputPins)
 	{
-		if (OutputPin->Label == Label)
+		if (OutputPin->Properties.Label == Label)
 		{
 			return OutputPin;
 		}
@@ -415,84 +413,83 @@ bool UPCGNode::UpdatePins(TFunctionRef<UPCGPin*(UPCGNode*)> PinAllocator)
 		return bChanged;
 	}
 
-	TArray<FName> InboundLabels = DefaultSettings->InLabels();
-	TArray<FName> OutboundLabels = DefaultSettings->OutLabels();
+	TArray<FPCGPinProperties> InboundPinProperties = DefaultSettings->InputPinProperties();
+	TArray<FPCGPinProperties> OutboundPinProperties = DefaultSettings->OutputPinProperties();
 
-	auto UpdatePins = [this, &PinAllocator](TArray<UPCGPin*>& Pins, const TArray<FName>& Labels)
+	auto UpdatePins = [this, &PinAllocator](TArray<UPCGPin*>& Pins, const TArray<FPCGPinProperties>& PinProperties)
 	{
-		TArray<FName> OldLabels;
-		Algo::Transform(Pins, OldLabels, [](UPCGPin* Pin) { return Pin->Label; });
+		bool bAppliedEdgeChanges = false;
 
-		auto GetUnmatched = [](const TArray<FName>& InA, const TArray<FName>& InB)
+		// Find unmatched pins vs. properties on a name basis
+		TArray<UPCGPin*> UnmatchedPins;
+		for (UPCGPin* Pin : Pins)
 		{
-			TArray<FName> Unmatched;
-			for (const FName& A : InA)
+			if (const FPCGPinProperties* MatchingProperties = PinProperties.FindByPredicate([Pin](const FPCGPinProperties& Prop) { return Prop.Label == Pin->Properties.Label; }))
 			{
-				if (!InB.Contains(A))
+				if (!(Pin->Properties == *MatchingProperties))
 				{
-					Unmatched.Add(A);
+					Pin->Modify();
+					Pin->Properties = *MatchingProperties;
+					bAppliedEdgeChanges |= Pin->BreakAllIncompatibleEdges();
 				}
 			}
+			else
+			{
+				UnmatchedPins.Add(Pin);
+			}
+		}
 
-			return Unmatched;
-		};
-
-		TArray<FName> UnmatchedLabels = GetUnmatched(Labels, OldLabels);
-		TArray<FName> UnmatchedOldLabels = GetUnmatched(OldLabels, Labels);
-
-		const bool bRenameUnmatchedPin = UnmatchedLabels.Num() == 1 && UnmatchedOldLabels.Num() == 1;
-		if (bRenameUnmatchedPin)
+		// Find unmatched properties vs pins on a name basis
+		TArray<FPCGPinProperties> UnmatchedProperties;
+		for (const FPCGPinProperties& Properties : PinProperties)
 		{
-			for (UPCGPin* Pin : Pins)
+			if (!Pins.FindByPredicate([&Properties](const UPCGPin* Pin) { return Pin->Properties.Label == Properties.Label; }))
 			{
-				if (Pin->Label == UnmatchedOldLabels[0])
-				{
-					Pin->Modify();
-					Pin->Label = UnmatchedLabels[0];
-				}
+				UnmatchedProperties.Add(Properties);
 			}
+		}
 
-			return false;
+		const bool bUpdateUnmatchedPin = UnmatchedPins.Num() == 1 && UnmatchedProperties.Num() == 1;
+		if (bUpdateUnmatchedPin)
+		{
+			UnmatchedPins[0]->Modify();
+			UnmatchedPins[0]->Properties = UnmatchedProperties[0];
+			bAppliedEdgeChanges |= UnmatchedPins[0]->BreakAllIncompatibleEdges();
 		}
 		else
 		{
-			if (!UnmatchedOldLabels.IsEmpty() || !UnmatchedLabels.IsEmpty())
+			if(!UnmatchedPins.IsEmpty() || !UnmatchedProperties.IsEmpty())
 			{
 				Modify();
 			}
 
-			bool bAppliedEdgeChanges = false;
-
 			// Remove old pins
-			if (UnmatchedOldLabels.Num() > 0)
+			for (int32 UnmatchedPinIndex = UnmatchedPins.Num() - 1; UnmatchedPinIndex >= 0; --UnmatchedPinIndex)
 			{
-				for (int32 PinIndex = Pins.Num() - 1; PinIndex >= 0; --PinIndex)
-				{
-					if (UnmatchedOldLabels.Contains(Pins[PinIndex]->Label))
-					{
-						bAppliedEdgeChanges |= Pins[PinIndex]->BreakAllEdges();
-						Pins.RemoveAt(PinIndex);
-					}
-				}
+				const int32 PinIndex = Pins.IndexOfByKey(UnmatchedPins[UnmatchedPinIndex]);
+				check(PinIndex >= 0);
+
+				bAppliedEdgeChanges |= Pins[PinIndex]->BreakAllEdges();
+				Pins.RemoveAt(PinIndex);
 			}
 
 			// Add new pins
-			for (const FName& UnmatchedLabel : UnmatchedLabels)
+			for (const FPCGPinProperties& UnmatchedProperty : UnmatchedProperties)
 			{
-				const int32 InsertIndex = Labels.IndexOfByKey(UnmatchedLabel);
+				const int32 InsertIndex = PinProperties.IndexOfByKey(UnmatchedProperty);
 				UPCGPin* NewPin = PinAllocator(this);
 				NewPin->Node = this;
-				NewPin->Label = UnmatchedLabel;
+				NewPin->Properties = UnmatchedProperty;
 				Pins.Insert(NewPin, InsertIndex);
 			}
-
-			return bAppliedEdgeChanges;
 		}
+
+		return bAppliedEdgeChanges;
 	};
 
 	bool bChanged = false;
-	bChanged |= UpdatePins(InputPins, InboundLabels);
-	bChanged |= UpdatePins(OutputPins, OutboundLabels);
+	bChanged |= UpdatePins(InputPins, InboundPinProperties);
+	bChanged |= UpdatePins(OutputPins, OutboundPinProperties);
 
 	return bChanged;
 }
diff --git a/Engine/Plugins/Experimental/PCG/Source/PCG/Private/PCGPin.cpp b/Engine/Plugins/Experimental/PCG/Source/PCG/Private/PCGPin.cpp
--- a/Engine/Plugins/Experimental/PCG/Source/PCG/Private/PCGPin.cpp
+++ b/Engine/Plugins/Experimental/PCG/Source/PCG/Private/PCGPin.cpp
@@ -4,12 +4,32 @@
 #include "PCGEdge.h"
 #include "PCGNode.h"
 
+FPCGPinProperties::FPCGPinProperties(const FName& InLabel, EPCGDataType InAllowedTypes, bool bInAllowMultipleConnections)
+	: Label(InLabel), AllowedTypes(InAllowedTypes), bAllowMultipleConnections(bInAllowMultipleConnections)
+{}
+
+bool FPCGPinProperties::operator==(const FPCGPinProperties& Other) const
+{
+	return Label == Other.Label && AllowedTypes == Other.AllowedTypes && bAllowMultipleConnections == Other.bAllowMultipleConnections;
+}
+
 UPCGPin::UPCGPin(const FObjectInitializer& ObjectInitializer)
 	: Super(ObjectInitializer)
 {
 	SetFlags(RF_Transactional);
 }
 
+void UPCGPin::PostLoad()
+{
+	Super::PostLoad();
+
+	if (Label_DEPRECATED != NAME_None)
+	{
+		Properties = FPCGPinProperties(Label_DEPRECATED);
+		Label_DEPRECATED = NAME_None;
+	}
+}
+
 bool UPCGPin::AddEdgeTo(UPCGPin* OtherPin)
 {
 	if (!OtherPin)
@@ -84,6 +104,39 @@ bool UPCGPin::BreakAllEdges()
 	return bChanged;
 }
 
+bool UPCGPin::BreakAllIncompatibleEdges()
+{
+	bool bChanged = false;
+	bool bHasAValidEdge = false;
+
+	for (int32 EdgeIndex = Edges.Num() - 1; EdgeIndex >= 0; --EdgeIndex)
+	{
+		UPCGEdge* Edge = Edges[EdgeIndex];
+		UPCGPin* OtherPin = Edge->GetOtherPin(this);
+
+		bool bRemoveEdge = !IsCompatible(OtherPin) || (!Properties.bAllowMultipleConnections && bHasAValidEdge);
+
+		if (bRemoveEdge)
+		{
+			Modify();
+			Edges.RemoveAtSwap(EdgeIndex);
+
+			if (OtherPin)
+			{
+				OtherPin->Modify();
+				ensure(OtherPin->Edges.Remove(Edge));
+				bChanged = true;
+			}
+		}
+		else
+		{
+			bHasAValidEdge = true;
+		}
+	}
+
+	return bChanged;
+}
+
 bool UPCGPin::IsConnected() const
 {
 	for (const UPCGEdge* Edge : Edges)
@@ -109,4 +162,14 @@ int32 UPCGPin::EdgeCount() const
 	}
 
 	return EdgeNum;
+}
+
+bool UPCGPin::IsCompatible(const UPCGPin* OtherPin) const
+{
+	return OtherPin && !!(Properties.AllowedTypes & OtherPin->Properties.AllowedTypes);
+}
+
+bool UPCGPin::CanConnect(const UPCGPin* OtherPin) const
+{
+	return OtherPin && (Properties.bAllowMultipleConnections || Edges.IsEmpty());
 }
\ No newline at end of file
diff --git a/Engine/Plugins/Experimental/PCG/Source/PCG/Private/PCGSettings.cpp b/Engine/Plugins/Experimental/PCG/Source/PCG/Private/PCGSettings.cpp
--- a/Engine/Plugins/Experimental/PCG/Source/PCG/Private/PCGSettings.cpp
+++ b/Engine/Plugins/Experimental/PCG/Source/PCG/Private/PCGSettings.cpp
@@ -43,14 +43,21 @@ uint32 UPCGSettings::GetCrc32() const
 	return Ar.Crc32(const_cast<UPCGSettings*>(this));
 }
 
-TArray<FName> UPCGSettings::InLabels() const
+TArray<FPCGPinProperties> UPCGSettings::InputPinProperties() const
 {
-	return { PCGPinConstants::DefaultInputLabel };
+	TArray<FPCGPinProperties> PinProperties;
+	PinProperties.Emplace(PCGPinConstants::DefaultInputLabel, EPCGDataType::Any);
+
+	return PinProperties;
 }
 
-TArray<FName> UPCGSettings::OutLabels() const
+TArray<FPCGPinProperties> UPCGSettings::OutputPinProperties() const
 {
-	return { PCGPinConstants::DefaultOutputLabel };
+	TArray<FPCGPinProperties> PinProperties;
+	// This is not true for everything, use a virtual call?
+	PinProperties.Emplace(PCGPinConstants::DefaultOutputLabel, EPCGDataType::Spatial);
+
+	return PinProperties;
 }
 
 FPCGElementPtr UPCGSettings::GetElement() const
@@ -96,6 +103,13 @@ void UPCGSettings::DirtyCache()
 }
 #endif // WITH_EDITOR
 
+TArray<FPCGPinProperties> UPCGSettings::DefaultPointOutputPinProperties() const
+{
+	TArray<FPCGPinProperties> Properties;
+	Properties.Emplace(PCGPinConstants::DefaultOutputLabel, EPCGDataType::Point);
+	return Properties;
+}
+
 FPCGElementPtr UPCGTrivialSettings::CreateElement() const
 {
 	return MakeShared<FPCGTrivialElement>();
diff --git a/Engine/Plugins/Experimental/PCG/Source/PCG/Private/PCGSubgraph.cpp b/Engine/Plugins/Experimental/PCG/Source/PCG/Private/PCGSubgraph.cpp
--- a/Engine/Plugins/Experimental/PCG/Source/PCG/Private/PCGSubgraph.cpp
+++ b/Engine/Plugins/Experimental/PCG/Source/PCG/Private/PCGSubgraph.cpp
@@ -77,27 +77,27 @@ void UPCGBaseSubgraphSettings::OnSubgraphChanged(UPCGGraph* InGraph, bool bIsStr
 }
 #endif // WITH_EDITOR
 
-TArray<FName> UPCGBaseSubgraphSettings::InLabels() const
+TArray<FPCGPinProperties> UPCGBaseSubgraphSettings::InputPinProperties() const
 {
 	if (UPCGGraph* Subgraph = GetSubgraph())
 	{
-		return Subgraph->GetInputNode()->InLabels();
+		return Subgraph->GetInputNode()->InputPinProperties();
 	}
 	else
 	{
-		return Super::InLabels();
+		return Super::InputPinProperties();
 	}
 }
 
-TArray<FName> UPCGBaseSubgraphSettings::OutLabels() const
+TArray<FPCGPinProperties> UPCGBaseSubgraphSettings::OutputPinProperties() const
 {
 	if (UPCGGraph* Subgraph = GetSubgraph())
 	{
-		return Subgraph->GetOutputNode()->OutLabels();
+		return Subgraph->GetOutputNode()->OutputPinProperties();
 	}
 	else
 	{
-		return Super::OutLabels();
+		return Super::OutputPinProperties();
 	}
 }
 
diff --git a/Engine/Plugins/Experimental/PCG/Source/PCG/Public/Data/PCGLandscapeData.h b/Engine/Plugins/Experimental/PCG/Source/PCG/Public/Data/PCGLandscapeData.h
--- a/Engine/Plugins/Experimental/PCG/Source/PCG/Public/Data/PCGLandscapeData.h
+++ b/Engine/Plugins/Experimental/PCG/Source/PCG/Public/Data/PCGLandscapeData.h
@@ -37,6 +37,10 @@ class PCG_API UPCGLandscapeData : public UPCGSurfaceData
 public:
 	void Initialize(ALandscapeProxy* InLandscape, const FBox& InBounds);
 
+	// ~Begin UPCGData interface
+	virtual EPCGDataType GetDataType() const override { return EPCGDataType::Landscape | Super::GetDataType(); }
+	// ~End UPCGData interface
+
 	// ~Begin UPGCSpatialData interface
 	virtual FBox GetBounds() const override;
 	virtual FBox GetStrictBounds() const override;
diff --git a/Engine/Plugins/Experimental/PCG/Source/PCG/Public/Data/PCGLandscapeSplineData.h b/Engine/Plugins/Experimental/PCG/Source/PCG/Public/Data/PCGLandscapeSplineData.h
--- a/Engine/Plugins/Experimental/PCG/Source/PCG/Public/Data/PCGLandscapeSplineData.h
+++ b/Engine/Plugins/Experimental/PCG/Source/PCG/Public/Data/PCGLandscapeSplineData.h
@@ -16,6 +16,10 @@ class PCG_API UPCGLandscapeSplineData : public UPCGPolyLineData
 public:
 	void Initialize(ULandscapeSplinesComponent* InSplineComponent);
 
+	// ~Begin UPCGData interface
+	virtual EPCGDataType GetDataType() const override { return EPCGDataType::LandscapeSpline | Super::GetDataType(); }
+	// ~End UPCGData interface
+
 	//~Begin UPCGPolyLineData interface
 	virtual int GetNumSegments() const override;
 	virtual float GetSegmentLength(int SegmentIndex) const override;
diff --git a/Engine/Plugins/Experimental/PCG/Source/PCG/Public/Data/PCGPointData.h b/Engine/Plugins/Experimental/PCG/Source/PCG/Public/Data/PCGPointData.h
--- a/Engine/Plugins/Experimental/PCG/Source/PCG/Public/Data/PCGPointData.h
+++ b/Engine/Plugins/Experimental/PCG/Source/PCG/Public/Data/PCGPointData.h
@@ -59,6 +59,10 @@ class PCG_API UPCGPointData : public UPCGSpatialData
 public:
 	typedef TOctree2<FPCGPointRef, FPCGPointRefSemantics> PointOctree;
 
+	// ~Begin UPCGData interface
+	virtual EPCGDataType GetDataType() const override { return EPCGDataType::Point | Super::GetDataType(); }
+	// ~End UPCGData interface
+
 	// ~Begin UPCGSpatialData interface
 	virtual int GetDimension() const override { return 0; }
 	virtual FBox GetBounds() const override;
diff --git a/Engine/Plugins/Experimental/PCG/Source/PCG/Public/Data/PCGPolyLineData.h b/Engine/Plugins/Experimental/PCG/Source/PCG/Public/Data/PCGPolyLineData.h
--- a/Engine/Plugins/Experimental/PCG/Source/PCG/Public/Data/PCGPolyLineData.h
+++ b/Engine/Plugins/Experimental/PCG/Source/PCG/Public/Data/PCGPolyLineData.h
@@ -13,6 +13,10 @@ class PCG_API UPCGPolyLineData : public UPCGSpatialDataWithPointCache
 	GENERATED_BODY()
 
 public:
+	// ~Begin UPCGData interface
+	virtual EPCGDataType GetDataType() const override { return EPCGDataType::PolyLine | Super::GetDataType(); }
+	// ~End UPCGData interface
+
 	//~Begin UPCGSpatialData interface
 	virtual int GetDimension() const override { return 1; }
 	virtual FBox GetBounds() const override;
diff --git a/Engine/Plugins/Experimental/PCG/Source/PCG/Public/Data/PCGPrimitiveData.h b/Engine/Plugins/Experimental/PCG/Source/PCG/Public/Data/PCGPrimitiveData.h
--- a/Engine/Plugins/Experimental/PCG/Source/PCG/Public/Data/PCGPrimitiveData.h
+++ b/Engine/Plugins/Experimental/PCG/Source/PCG/Public/Data/PCGPrimitiveData.h
@@ -14,6 +14,10 @@ class PCG_API UPCGPrimitiveData : public UPCGSpatialDataWithPointCache
 public:
 	void Initialize(UPrimitiveComponent* InPrim);
 
+	// ~Begin UPCGData interface
+	virtual EPCGDataType GetDataType() const override { return EPCGDataType::Primitive | Super::GetDataType(); }
+	// ~End UPCGData interface
+
 	// ~Begin UPCGSpatialData interface
 	virtual int GetDimension() const override { return 3; }
 	virtual FBox GetBounds() const override { return CachedBounds; }
diff --git a/Engine/Plugins/Experimental/PCG/Source/PCG/Public/Data/PCGRenderTargetData.h b/Engine/Plugins/Experimental/PCG/Source/PCG/Public/Data/PCGRenderTargetData.h
--- a/Engine/Plugins/Experimental/PCG/Source/PCG/Public/Data/PCGRenderTargetData.h
+++ b/Engine/Plugins/Experimental/PCG/Source/PCG/Public/Data/PCGRenderTargetData.h
@@ -15,6 +15,10 @@ class PCG_API UPCGRenderTargetData : public UPCGBaseTextureData
 	GENERATED_BODY()
 
 public:
+	// ~Begin UPCGData interface
+	virtual EPCGDataType GetDataType() const override { return EPCGDataType::RenderTarget | Super::GetDataType(); }
+	// ~End UPCGData interface
+
 	UFUNCTION(BlueprintCallable, Category = RenderTarget)
 	void Initialize(UTextureRenderTarget2D* InRenderTarget, const FTransform& InTransform);
 
diff --git a/Engine/Plugins/Experimental/PCG/Source/PCG/Public/Data/PCGSpatialData.h b/Engine/Plugins/Experimental/PCG/Source/PCG/Public/Data/PCGSpatialData.h
--- a/Engine/Plugins/Experimental/PCG/Source/PCG/Public/Data/PCGSpatialData.h
+++ b/Engine/Plugins/Experimental/PCG/Source/PCG/Public/Data/PCGSpatialData.h
@@ -35,6 +35,10 @@ class PCG_API UPCGSpatialData : public UPCGData
 public:
 	UPCGSpatialData(const FObjectInitializer& ObjectInitializer);
 
+	// ~Begin UPCGData interface
+	virtual EPCGDataType GetDataType() const override { return EPCGDataType::Spatial | Super::GetDataType(); }
+	// ~End UPCGData interface
+
 	/** Returns the dimension of the data type, which has nothing to do with the dimension of its points */
 	UFUNCTION(BlueprintCallable, Category = SpatialData)
 	virtual int GetDimension() const PURE_VIRTUAL(UPCGSpatialData::GetDimension, return 0;);
diff --git a/Engine/Plugins/Experimental/PCG/Source/PCG/Public/Data/PCGSplineData.h b/Engine/Plugins/Experimental/PCG/Source/PCG/Public/Data/PCGSplineData.h
--- a/Engine/Plugins/Experimental/PCG/Source/PCG/Public/Data/PCGSplineData.h
+++ b/Engine/Plugins/Experimental/PCG/Source/PCG/Public/Data/PCGSplineData.h
@@ -18,6 +18,10 @@ class PCG_API UPCGSplineData : public UPCGPolyLineData
 public:
 	void Initialize(USplineComponent* InSpline);
 
+	// ~Begin UPCGData interface
+	virtual EPCGDataType GetDataType() const override { return EPCGDataType::Spline | Super::GetDataType(); }
+	// ~End UPCGData interface
+
 	//~Begin UPCGPolyLineData interface
 	virtual int GetNumSegments() const override;
 	virtual float GetSegmentLength(int SegmentIndex) const override;
diff --git a/Engine/Plugins/Experimental/PCG/Source/PCG/Public/Data/PCGSurfaceData.h b/Engine/Plugins/Experimental/PCG/Source/PCG/Public/Data/PCGSurfaceData.h
--- a/Engine/Plugins/Experimental/PCG/Source/PCG/Public/Data/PCGSurfaceData.h
+++ b/Engine/Plugins/Experimental/PCG/Source/PCG/Public/Data/PCGSurfaceData.h
@@ -13,6 +13,10 @@ class PCG_API UPCGSurfaceData : public UPCGSpatialDataWithPointCache
 	GENERATED_BODY()
 
 public:
+	// ~Begin UPCGData interface
+	virtual EPCGDataType GetDataType() const override { return EPCGDataType::Surface | Super::GetDataType(); }
+	// ~End UPCGData interface
+
 	//~Begin UPCGSpatialData interface
 	virtual int GetDimension() const override { return 2; }
 	virtual bool HasNonTrivialTransform() const override { return true; }
diff --git a/Engine/Plugins/Experimental/PCG/Source/PCG/Public/Data/PCGTextureData.h b/Engine/Plugins/Experimental/PCG/Source/PCG/Public/Data/PCGTextureData.h
--- a/Engine/Plugins/Experimental/PCG/Source/PCG/Public/Data/PCGTextureData.h
+++ b/Engine/Plugins/Experimental/PCG/Source/PCG/Public/Data/PCGTextureData.h
@@ -30,6 +30,10 @@ class PCG_API UPCGBaseTextureData : public UPCGSurfaceData
 	GENERATED_BODY()
 
 public:
+	// ~Begin UPCGData interface
+	virtual EPCGDataType GetDataType() const override { return EPCGDataType::BaseTexture | Super::GetDataType(); }
+	// ~End UPCGData interface
+
 	//~Begin UPCGSpatialData interface
 	virtual FBox GetBounds() const override;
 	virtual FBox GetStrictBounds() const override;
@@ -69,6 +73,10 @@ class PCG_API UPCGTextureData : public UPCGBaseTextureData
 	GENERATED_BODY()
 
 public:
+	// ~Begin UPCGData interface
+	virtual EPCGDataType GetDataType() const override { return EPCGDataType::Texture | Super::GetDataType(); }
+	// ~End UPCGData interface
+
 	UFUNCTION(BlueprintCallable, Category = Texture)
 	void Initialize(UTexture2D* InTexture, const FTransform& InTransform);
 
diff --git a/Engine/Plugins/Experimental/PCG/Source/PCG/Public/Data/PCGVolumeData.h b/Engine/Plugins/Experimental/PCG/Source/PCG/Public/Data/PCGVolumeData.h
--- a/Engine/Plugins/Experimental/PCG/Source/PCG/Public/Data/PCGVolumeData.h
+++ b/Engine/Plugins/Experimental/PCG/Source/PCG/Public/Data/PCGVolumeData.h
@@ -18,6 +18,10 @@ public:
 	void Initialize(AVolume* InVolume, AActor* InTargetActor = nullptr);
 	void Initialize(const FBox& InBounds, AActor* InTargetActor);
 
+	// ~Begin UPCGData interface
+	virtual EPCGDataType GetDataType() const override { return EPCGDataType::Volume | Super::GetDataType(); }
+	// ~End UPCGData interface
+
 	// ~Begin UPGCSpatialData interface
 	virtual int GetDimension() const override { return 3; }
 	virtual FBox GetBounds() const override;
diff --git a/Engine/Plugins/Experimental/PCG/Source/PCG/Public/Elements/PCGCollapseElement.h b/Engine/Plugins/Experimental/PCG/Source/PCG/Public/Elements/PCGCollapseElement.h
--- a/Engine/Plugins/Experimental/PCG/Source/PCG/Public/Elements/PCGCollapseElement.h
+++ b/Engine/Plugins/Experimental/PCG/Source/PCG/Public/Elements/PCGCollapseElement.h
@@ -19,6 +19,8 @@ public:
 	virtual EPCGSettingsType GetType() const override { return EPCGSettingsType::Spatial; }
 #endif
 
+	virtual TArray<FPCGPinProperties> OutputPinProperties() const override { return Super::DefaultPointOutputPinProperties(); }
+
 protected:
 	virtual FPCGElementPtr CreateElement() const override;
 	//~End UPCGSettings interface
diff --git a/Engine/Plugins/Experimental/PCG/Source/PCG/Public/Elements/PCGDebugElement.h b/Engine/Plugins/Experimental/PCG/Source/PCG/Public/Elements/PCGDebugElement.h
--- a/Engine/Plugins/Experimental/PCG/Source/PCG/Public/Elements/PCGDebugElement.h
+++ b/Engine/Plugins/Experimental/PCG/Source/PCG/Public/Elements/PCGDebugElement.h
@@ -24,6 +24,8 @@ public:
 	virtual EPCGSettingsType GetType() const override { return EPCGSettingsType::Debug; }
 #endif
 
+	virtual TArray<FPCGPinProperties> OutputPinProperties() const override { return TArray<FPCGPinProperties>(); }
+
 protected:
 	virtual FPCGElementPtr CreateElement() const override;
 	// ~End UPCGSettings interface
diff --git a/Engine/Plugins/Experimental/PCG/Source/PCG/Public/Elements/PCGDensityFilter.h b/Engine/Plugins/Experimental/PCG/Source/PCG/Public/Elements/PCGDensityFilter.h
--- a/Engine/Plugins/Experimental/PCG/Source/PCG/Public/Elements/PCGDensityFilter.h
+++ b/Engine/Plugins/Experimental/PCG/Source/PCG/Public/Elements/PCGDensityFilter.h
@@ -21,6 +21,8 @@ public:
 	virtual EPCGSettingsType GetType() const override { return EPCGSettingsType::Filter; }
 #endif
 
+	virtual TArray<FPCGPinProperties> OutputPinProperties() const override { return Super::DefaultPointOutputPinProperties(); }
+
 protected:
 	virtual FPCGElementPtr CreateElement() const override;
 	//~End UPCGSettings interface
diff --git a/Engine/Plugins/Experimental/PCG/Source/PCG/Public/Elements/PCGDensityRemap.h b/Engine/Plugins/Experimental/PCG/Source/PCG/Public/Elements/PCGDensityRemap.h
--- a/Engine/Plugins/Experimental/PCG/Source/PCG/Public/Elements/PCGDensityRemap.h
+++ b/Engine/Plugins/Experimental/PCG/Source/PCG/Public/Elements/PCGDensityRemap.h
@@ -21,6 +21,8 @@ public:
 	virtual EPCGSettingsType GetType() const override { return EPCGSettingsType::Density; }
 #endif
 
+	virtual TArray<FPCGPinProperties> OutputPinProperties() const override { return Super::DefaultPointOutputPinProperties(); }
+
 protected:
 	virtual FPCGElementPtr CreateElement() const override;
 	//~End UPCGSettings interface
diff --git a/Engine/Plugins/Experimental/PCG/Source/PCG/Public/Elements/PCGDifferenceElement.h b/Engine/Plugins/Experimental/PCG/Source/PCG/Public/Elements/PCGDifferenceElement.h
--- a/Engine/Plugins/Experimental/PCG/Source/PCG/Public/Elements/PCGDifferenceElement.h
+++ b/Engine/Plugins/Experimental/PCG/Source/PCG/Public/Elements/PCGDifferenceElement.h
@@ -21,7 +21,7 @@ public:
 	virtual EPCGSettingsType GetType() const override { return EPCGSettingsType::Spatial; }
 #endif
 
-	virtual TArray<FName> InLabels() const override;
+	virtual TArray<FPCGPinProperties> InputPinProperties() const override;
 
 protected:
 	virtual FPCGElementPtr CreateElement() const override;
diff --git a/Engine/Plugins/Experimental/PCG/Source/PCG/Public/Elements/PCGExecuteBlueprint.h b/Engine/Plugins/Experimental/PCG/Source/PCG/Public/Elements/PCGExecuteBlueprint.h
--- a/Engine/Plugins/Experimental/PCG/Source/PCG/Public/Elements/PCGExecuteBlueprint.h
+++ b/Engine/Plugins/Experimental/PCG/Source/PCG/Public/Elements/PCGExecuteBlueprint.h
@@ -69,6 +69,12 @@ public:
 	UFUNCTION(BlueprintNativeEvent, Category = Graph)
 	EPCGSettingsType NodeTypeOverride() const;
 
+	UFUNCTION(BlueprintCallable, Category = "Input & Output")
+	TSet<FName> InputLabels() const;
+
+	UFUNCTION(BlueprintCallable, Category = "Input & Output")
+	TSet<FName> OutputLabels() const;
+
 	/** Called after object creation to setup the object callbacks */
 	void Initialize();
 
@@ -91,11 +97,17 @@ public:
 	UPROPERTY(BlueprintReadWrite, EditAnywhere, Category = Settings)
 	bool bCanBeMultithreaded = false;
 
-	UPROPERTY(EditAnywhere, BlueprintReadWrite, Category = "Settings|Input & Output")
-	TSet<FName> InputPinLabels;
+	UPROPERTY(BlueprintGetter=InputLabels, Category = "Settings|Input & Output")
+	TSet<FName> InputPinLabels_DEPRECATED;
+
+	UPROPERTY(BlueprintGetter=OutputLabels, Category = "Settings|Input & Output")
+	TSet<FName> OutputPinLabels_DEPRECATED;
 
-	UPROPERTY(EditAnywhere, BlueprintReadWrite, Category = "Settings|Input & Output")
-	TSet<FName> OutputPinLabels;
+	UPROPERTY(BlueprintReadWrite, EditAnywhere, Category = "Settings|Input & Output")
+	TArray<FPCGPinProperties> CustomInputPins;
+
+	UPROPERTY(BlueprintReadWrite, EditAnywhere, Category = "Settings|Input & Output")
+	TArray<FPCGPinProperties> CustomOutputPins;
 
 	UPROPERTY(BlueprintReadWrite, EditAnywhere, Category = "Settings|Input & Output")
 	bool bHasDefaultInPin = true;
@@ -138,8 +150,8 @@ public:
 #endif
 
 	virtual FName AdditionalTaskName() const override;
-	virtual TArray<FName> InLabels() const override;
-	virtual TArray<FName> OutLabels() const override;
+	virtual TArray<FPCGPinProperties> InputPinProperties() const override;
+	virtual TArray<FPCGPinProperties> OutputPinProperties() const override;
 
 protected:
 	virtual FPCGElementPtr CreateElement() const override;
diff --git a/Engine/Plugins/Experimental/PCG/Source/PCG/Public/Elements/PCGMetadataElement.h b/Engine/Plugins/Experimental/PCG/Source/PCG/Public/Elements/PCGMetadataElement.h
--- a/Engine/Plugins/Experimental/PCG/Source/PCG/Public/Elements/PCGMetadataElement.h
+++ b/Engine/Plugins/Experimental/PCG/Source/PCG/Public/Elements/PCGMetadataElement.h
@@ -28,6 +28,8 @@ public:
 	virtual EPCGSettingsType GetType() const override { return EPCGSettingsType::Metadata; }
 #endif
 
+	virtual TArray<FPCGPinProperties> OutputPinProperties() const override { return Super::DefaultPointOutputPinProperties(); }
+
 protected:
 	virtual FPCGElementPtr CreateElement() const override;
 	//~End UPCGSettings interface
diff --git a/Engine/Plugins/Experimental/PCG/Source/PCG/Public/Elements/PCGMetadataPartition.h b/Engine/Plugins/Experimental/PCG/Source/PCG/Public/Elements/PCGMetadataPartition.h
--- a/Engine/Plugins/Experimental/PCG/Source/PCG/Public/Elements/PCGMetadataPartition.h
+++ b/Engine/Plugins/Experimental/PCG/Source/PCG/Public/Elements/PCGMetadataPartition.h
@@ -19,6 +19,8 @@ public:
 	virtual EPCGSettingsType GetType() const override { return EPCGSettingsType::Metadata; }
 #endif
 
+	virtual TArray<FPCGPinProperties> OutputPinProperties() const override { return Super::DefaultPointOutputPinProperties(); }
+
 protected:
 	virtual FPCGElementPtr CreateElement() const override;
 	//~End UPCGSettings interface
diff --git a/Engine/Plugins/Experimental/PCG/Source/PCG/Public/Elements/PCGPointFilter.h b/Engine/Plugins/Experimental/PCG/Source/PCG/Public/Elements/PCGPointFilter.h
--- a/Engine/Plugins/Experimental/PCG/Source/PCG/Public/Elements/PCGPointFilter.h
+++ b/Engine/Plugins/Experimental/PCG/Source/PCG/Public/Elements/PCGPointFilter.h
@@ -70,8 +70,8 @@ public:
 	virtual EPCGSettingsType GetType() const override { return EPCGSettingsType::Filter; }
 #endif
 
-	virtual TArray<FName> InLabels() const override;
-	virtual TArray<FName> OutLabels() const override;
+	virtual TArray<FPCGPinProperties> InputPinProperties() const override;
+	virtual TArray<FPCGPinProperties> OutputPinProperties() const override;
 
 protected:
 	virtual FPCGElementPtr CreateElement() const override;
diff --git a/Engine/Plugins/Experimental/PCG/Source/PCG/Public/Elements/PCGPointSampler.h b/Engine/Plugins/Experimental/PCG/Source/PCG/Public/Elements/PCGPointSampler.h
--- a/Engine/Plugins/Experimental/PCG/Source/PCG/Public/Elements/PCGPointSampler.h
+++ b/Engine/Plugins/Experimental/PCG/Source/PCG/Public/Elements/PCGPointSampler.h
@@ -21,6 +21,8 @@ public:
 	virtual EPCGSettingsType GetType() const override { return EPCGSettingsType::Sampler; }
 #endif
 
+	virtual TArray<FPCGPinProperties> OutputPinProperties() const override { return Super::DefaultPointOutputPinProperties(); }
+
 protected:
 	virtual FPCGElementPtr CreateElement() const override;
 	//~End UPCGSettings interface
diff --git a/Engine/Plugins/Experimental/PCG/Source/PCG/Public/Elements/PCGSelfPruning.h b/Engine/Plugins/Experimental/PCG/Source/PCG/Public/Elements/PCGSelfPruning.h
--- a/Engine/Plugins/Experimental/PCG/Source/PCG/Public/Elements/PCGSelfPruning.h
+++ b/Engine/Plugins/Experimental/PCG/Source/PCG/Public/Elements/PCGSelfPruning.h
@@ -45,6 +45,8 @@ public:
 	virtual EPCGSettingsType GetType() const override { return EPCGSettingsType::Filter; }
 #endif
 
+	virtual TArray<FPCGPinProperties> OutputPinProperties() const override { return Super::DefaultPointOutputPinProperties(); }
+
 protected:
 	virtual FPCGElementPtr CreateElement() const override;
 	// ~End UPCGSettings interface
diff --git a/Engine/Plugins/Experimental/PCG/Source/PCG/Public/Elements/PCGStaticMeshSpawner.h b/Engine/Plugins/Experimental/PCG/Source/PCG/Public/Elements/PCGStaticMeshSpawner.h
--- a/Engine/Plugins/Experimental/PCG/Source/PCG/Public/Elements/PCGStaticMeshSpawner.h
+++ b/Engine/Plugins/Experimental/PCG/Source/PCG/Public/Elements/PCGStaticMeshSpawner.h
@@ -41,6 +41,8 @@ public:
 	virtual EPCGSettingsType GetType() const override { return EPCGSettingsType::Spawner; }
 #endif
 
+	virtual TArray<FPCGPinProperties> OutputPinProperties() const override { return TArray<FPCGPinProperties>(); }
+
 protected:
 	virtual FPCGElementPtr CreateElement() const override;
 	// ~End UPCGSettings interface
diff --git a/Engine/Plugins/Experimental/PCG/Source/PCG/Public/Elements/PCGSurfaceSampler.h b/Engine/Plugins/Experimental/PCG/Source/PCG/Public/Elements/PCGSurfaceSampler.h
--- a/Engine/Plugins/Experimental/PCG/Source/PCG/Public/Elements/PCGSurfaceSampler.h
+++ b/Engine/Plugins/Experimental/PCG/Source/PCG/Public/Elements/PCGSurfaceSampler.h
@@ -50,6 +50,8 @@ public:
 	virtual EPCGSettingsType GetType() const override { return EPCGSettingsType::Sampler; }
 #endif
 
+	virtual TArray<FPCGPinProperties> OutputPinProperties() const override { return Super::DefaultPointOutputPinProperties(); }
+
 protected:
 	virtual FPCGElementPtr CreateElement() const override;
 	// ~End UPCGSettings interface
diff --git a/Engine/Plugins/Experimental/PCG/Source/PCG/Public/PCGCommon.h b/Engine/Plugins/Experimental/PCG/Source/PCG/Public/PCGCommon.h
--- a/Engine/Plugins/Experimental/PCG/Source/PCG/Public/PCGCommon.h
+++ b/Engine/Plugins/Experimental/PCG/Source/PCG/Public/PCGCommon.h
@@ -16,6 +16,34 @@ enum class EPCGChangeType : uint8
 };
 ENUM_CLASS_FLAGS(EPCGChangeType);
 
+UENUM(meta = (Bitflags))
+enum class EPCGDataType : uint32
+{
+	None = 0 UMETA(Hidden),
+	Point = 1 << 1,
+
+	Spline = 1 << 2,
+	LandscapeSpline = 1 << 3,
+	PolyLine = Spline | LandscapeSpline UMETA(Hidden),
+
+	Landscape = 1 << 4,
+	Texture = 1 << 5,
+	RenderTarget = 1 << 6,
+	BaseTexture = Texture | RenderTarget UMETA(Hidden),
+	Surface = Landscape | BaseTexture,
+
+	Volume = 1 << 7,
+	Primitive = 1 << 8,
+
+	Spatial = Point | PolyLine | Surface | Volume | Primitive,
+
+	Param = 1 << 27,
+	Settings = 1 << 28 UMETA(Hidden),
+	Other = 1 << 29,
+	Any = (1 << 30) - 1
+};
+ENUM_CLASS_FLAGS(EPCGDataType);
+
 namespace PCGPinConstants
 {
 	const FName DefaultInputLabel = TEXT("In");
diff --git a/Engine/Plugins/Experimental/PCG/Source/PCG/Public/PCGData.h b/Engine/Plugins/Experimental/PCG/Source/PCG/Public/PCGData.h
--- a/Engine/Plugins/Experimental/PCG/Source/PCG/Public/PCGData.h
+++ b/Engine/Plugins/Experimental/PCG/Source/PCG/Public/PCGData.h
@@ -6,6 +6,7 @@
 #include "Kismet/BlueprintFunctionLibrary.h"
 
 #include "PCGModule.h"
+#include "PCGCommon.h"
 
 #include "PCGData.generated.h"
 
@@ -21,6 +22,9 @@ UCLASS(BlueprintType, ClassGroup = (Procedural))
 class PCG_API UPCGData : public UObject
 {
 	GENERATED_BODY()
+
+public:
+	virtual EPCGDataType GetDataType() const { return EPCGDataType::None; }
 };
 
 USTRUCT(BlueprintType)
diff --git a/Engine/Plugins/Experimental/PCG/Source/PCG/Public/PCGElement.h b/Engine/Plugins/Experimental/PCG/Source/PCG/Public/PCGElement.h
--- a/Engine/Plugins/Experimental/PCG/Source/PCG/Public/PCGElement.h
+++ b/Engine/Plugins/Experimental/PCG/Source/PCG/Public/PCGElement.h
@@ -55,7 +55,7 @@ protected:
 #endif
 
 private:
-	void CleanupAndValidateOutputLabels(FPCGContext* Context) const;
+	void CleanupAndValidateOutput(FPCGContext* Context) const;
 };
 
 /**
diff --git a/Engine/Plugins/Experimental/PCG/Source/PCG/Public/PCGInputOutputSettings.h b/Engine/Plugins/Experimental/PCG/Source/PCG/Public/PCGInputOutputSettings.h
--- a/Engine/Plugins/Experimental/PCG/Source/PCG/Public/PCGInputOutputSettings.h
+++ b/Engine/Plugins/Experimental/PCG/Source/PCG/Public/PCGInputOutputSettings.h
@@ -30,14 +30,18 @@ class PCG_API UPCGGraphInputOutputSettings : public UPCGSettings
 public:
 	UPCGGraphInputOutputSettings(const FObjectInitializer& ObjectInitializer);
 
+	// ~Begin UObject interface
+	virtual void PostLoad() override;
+	// ~End UObject interface
+
 	// ~Begin UPCGSettings interface
 #if WITH_EDITOR
 	virtual FName GetDefaultNodeName() const override { return bIsInput ? FName(TEXT("InputNode")) : FName(TEXT("OutputNode")); }
 	virtual EPCGSettingsType GetType() const override { return EPCGSettingsType::InputOutput; }
 #endif
 
-	virtual TArray<FName> InLabels() const override { return GetLabels(); }
-	virtual TArray<FName> OutLabels() const override { return GetLabels(); }
+	TArray<FPCGPinProperties> InputPinProperties() const override;
+	TArray<FPCGPinProperties> OutputPinProperties() const override;
 
 	void SetInput(bool bInIsInput) { bIsInput = bInIsInput; }
 
@@ -48,21 +52,12 @@ protected:
 	const TArray<FName>& StaticLabels() const { return bIsInput ? StaticInLabels : StaticOutLabels; }
 	const TArray<FName>& StaticAdvancedLabels() const { return bIsInput ? StaticAdvancedInLabels : StaticAdvancedOutLabels; }
 
-	TArray<FName> GetLabels() const
-	{
-		TArray<FName> Labels = StaticLabels();
-		if (bShowAdvancedPins)
-		{
-			Labels += StaticAdvancedLabels();
-		}
-
-		Labels += PinLabels.Array();
-		return Labels;
-	}
-
 protected:
+	UPROPERTY()
+	TSet<FName> PinLabels_DEPRECATED;
+
 	UPROPERTY(EditAnywhere, BlueprintReadWrite, Category = "Settings|Input")
-	TSet<FName> PinLabels;
+	TArray<FPCGPinProperties> CustomPins;
 
 	UPROPERTY(EditAnywhere, BlueprintReadWrite, Category = "Settings|Input")
 	bool bShowAdvancedPins = false;
diff --git a/Engine/Plugins/Experimental/PCG/Source/PCG/Public/PCGNode.h b/Engine/Plugins/Experimental/PCG/Source/PCG/Public/PCGNode.h
--- a/Engine/Plugins/Experimental/PCG/Source/PCG/Public/PCGNode.h
+++ b/Engine/Plugins/Experimental/PCG/Source/PCG/Public/PCGNode.h
@@ -53,11 +53,11 @@ public:
 	/** Returns the node title, based either on the current node label, or defaulted to its settings */
 	FName GetNodeTitle() const;
 
-	/** Returns all the input labels */
-	TArray<FName> InLabels() const;
+	/** Returns all the input pin properties */
+	TArray<FPCGPinProperties> InputPinProperties() const;
 
-	/** Returns all the output labels */
-	TArray<FName> OutLabels() const;
+	/** Returns all the output pin properties */
+	TArray<FPCGPinProperties> OutputPinProperties() const;
 
 	/** Returns true if the input pin is connected */
 	bool IsInputPinConnected(const FName& Label) const;
diff --git a/Engine/Plugins/Experimental/PCG/Source/PCG/Public/PCGParamData.h b/Engine/Plugins/Experimental/PCG/Source/PCG/Public/PCGParamData.h
--- a/Engine/Plugins/Experimental/PCG/Source/PCG/Public/PCGParamData.h
+++ b/Engine/Plugins/Experimental/PCG/Source/PCG/Public/PCGParamData.h
@@ -18,6 +18,10 @@ class PCG_API UPCGParamData : public UPCGData
 public:
 	UPCGParamData(const FObjectInitializer& ObjectInitializer);
 
+	// ~Begin UPCGData interface
+	virtual EPCGDataType GetDataType() const override { return EPCGDataType::Param | Super::GetDataType(); }
+	// ~End UPCGData interface
+
 	UFUNCTION(BlueprintCallable, Category = Metadata)
 	const UPCGMetadata* ConstMetadata() const { return Metadata; }
 
diff --git a/Engine/Plugins/Experimental/PCG/Source/PCG/Public/PCGPin.h b/Engine/Plugins/Experimental/PCG/Source/PCG/Public/PCGPin.h
--- a/Engine/Plugins/Experimental/PCG/Source/PCG/Public/PCGPin.h
+++ b/Engine/Plugins/Experimental/PCG/Source/PCG/Public/PCGPin.h
@@ -3,12 +3,33 @@
 #pragma once
 
 #include "CoreMinimal.h"
+#include "PCGCommon.h"
 
 #include "PCGPin.generated.h"
 
 class UPCGNode;
 class UPCGEdge;
 
+USTRUCT(BlueprintType)
+struct PCG_API FPCGPinProperties
+{
+	GENERATED_BODY()
+
+	FPCGPinProperties() = default;
+	explicit FPCGPinProperties(const FName& InLabel, EPCGDataType InAllowedTypes = EPCGDataType::Any, bool bInAllowMultipleConnections = true);
+
+	bool operator==(const FPCGPinProperties& Other) const;
+
+	UPROPERTY(EditAnywhere, Category = Settings)
+	FName Label = NAME_None;
+
+	UPROPERTY(EditAnywhere, Category = Settings)
+	EPCGDataType AllowedTypes = EPCGDataType::Any;
+
+	UPROPERTY(EditAnywhere, Category = Settings)
+	bool bAllowMultipleConnections = true;
+};
+
 UCLASS(ClassGroup = (Procedural))
 class PCG_API UPCGPin : public UObject
 {
@@ -17,18 +38,28 @@ class PCG_API UPCGPin : public UObject
 public:
 	UPCGPin(const FObjectInitializer& ObjectInitializer);
 
+	// ~Begin UObject interface
+	virtual void PostLoad() override;
+	// ~End UObject interface
+
 	UPROPERTY()
 	TObjectPtr<UPCGNode> Node = nullptr;
 
-	UPROPERTY(EditAnywhere, Category = Settings)
-	FName Label = NAME_None;
+	UPROPERTY()
+	FName Label_DEPRECATED = NAME_None;
 
 	UPROPERTY(TextExportTransient)
 	TArray<TObjectPtr<UPCGEdge>> Edges;
 
+	UPROPERTY(EditAnywhere, Category = Settings, meta = (ShowOnlyInnerProperties))
+	FPCGPinProperties Properties;
+
+	bool IsCompatible(const UPCGPin* OtherPin) const;
+	bool CanConnect(const UPCGPin* OtherPin) const;
 	bool AddEdgeTo(UPCGPin* OtherPin);
 	bool BreakEdgeTo(UPCGPin* OtherPin);
 	bool BreakAllEdges();
+	bool BreakAllIncompatibleEdges();
 	bool IsConnected() const;
 	int32 EdgeCount() const;
 };
\ No newline at end of file
diff --git a/Engine/Plugins/Experimental/PCG/Source/PCG/Public/PCGSettings.h b/Engine/Plugins/Experimental/PCG/Source/PCG/Public/PCGSettings.h
--- a/Engine/Plugins/Experimental/PCG/Source/PCG/Public/PCGSettings.h
+++ b/Engine/Plugins/Experimental/PCG/Source/PCG/Public/PCGSettings.h
@@ -52,12 +52,16 @@ class PCG_API UPCGSettings : public UPCGData
 	GENERATED_BODY()
 
 public:
+	// ~Begin UPCGData interface
+	virtual EPCGDataType GetDataType() const override { return EPCGDataType::Settings | Super::GetDataType(); }
+	// ~End UPCGData interface
+
 	// TODO: check if we need this to be virtual, we don't really need if we're always caching
 	/*virtual*/ FPCGElementPtr GetElement() const;
 	virtual UPCGNode* CreateNode() const;
 
-	virtual TArray<FName> InLabels() const;
-	virtual TArray<FName> OutLabels() const;
+	virtual TArray<FPCGPinProperties> InputPinProperties() const;
+	virtual TArray<FPCGPinProperties> OutputPinProperties() const;
 	
 	bool operator==(const UPCGSettings& Other) const;
 	uint32 GetCrc32() const;
@@ -110,6 +114,9 @@ protected:
 	void DirtyCache();
 #endif 
 
+	/** Methods to remove boilerplate code across settings */
+	TArray<FPCGPinProperties> DefaultPointOutputPinProperties() const;
+
 private:
 	mutable FPCGElementPtr CachedElement;
 	mutable FCriticalSection CacheLock;
diff --git a/Engine/Plugins/Experimental/PCG/Source/PCG/Public/PCGSubgraph.h b/Engine/Plugins/Experimental/PCG/Source/PCG/Public/PCGSubgraph.h
--- a/Engine/Plugins/Experimental/PCG/Source/PCG/Public/PCGSubgraph.h
+++ b/Engine/Plugins/Experimental/PCG/Source/PCG/Public/PCGSubgraph.h
@@ -29,8 +29,8 @@ protected:
 	virtual void GetTrackedActorTags(FPCGTagToSettingsMap& OutTagToSettings) const override;
 #endif
 
-	virtual TArray<FName> InLabels() const override;
-	virtual TArray<FName> OutLabels() const override;
+	TArray<FPCGPinProperties> InputPinProperties() const override;
+	TArray<FPCGPinProperties> OutputPinProperties() const override;
 	//~End UPCGSettings interface
 
 #if WITH_EDITOR
diff --git a/Engine/Plugins/Experimental/PCG/Source/PCGEditor/Private/PCGEditorCommon.h b/Engine/Plugins/Experimental/PCG/Source/PCGEditor/Private/PCGEditorCommon.h
--- a/Engine/Plugins/Experimental/PCG/Source/PCGEditor/Private/PCGEditorCommon.h
+++ b/Engine/Plugins/Experimental/PCG/Source/PCGEditor/Private/PCGEditorCommon.h
@@ -5,4 +5,16 @@
 namespace FPCGEditorCommon
 {
 	const FString ContextIdentifier = TEXT("PCGEditorContext");
+
+	const FName SpatialDataType = FName(TEXT("Spatial Data"));
+	const FName ParamDataType = FName(TEXT("Param Data"));
+	const FName SettingsDataType = FName(TEXT("Settings Data"));
+	const FName OtherDataType = FName(TEXT("Other Data"));
+
+	const FName PointDataType = FName(TEXT("Point Data"));
+	const FName PolyLineDataType = FName(TEXT("Poly Line Data"));
+	const FName SurfaceDataType = FName(TEXT("Surface Data"));
+	const FName RenderTargetDataType = FName(TEXT("Render Target Data"));
+	const FName VolumeDataType = FName(TEXT("Volume Data"));
+	const FName PrimitiveDataType = FName(TEXT("Primitive Data"));	
 };
diff --git a/Engine/Plugins/Experimental/PCG/Source/PCGEditor/Private/PCGEditorGraph.cpp b/Engine/Plugins/Experimental/PCG/Source/PCGEditor/Private/PCGEditorGraph.cpp
--- a/Engine/Plugins/Experimental/PCG/Source/PCGEditor/Private/PCGEditorGraph.cpp
+++ b/Engine/Plugins/Experimental/PCG/Source/PCGEditor/Private/PCGEditorGraph.cpp
@@ -77,7 +77,7 @@ void UPCGEditorGraph::CreateLinks(UPCGEditorGraphNodeBase* GraphNode, bool bCrea
 	{
 		for (const UPCGPin* InputPin : PCGNode->GetInputPins())
 		{
-			UEdGraphPin* InPin = GraphNode->FindPin(InputPin->Label, EEdGraphPinDirection::EGPD_Input);
+			UEdGraphPin* InPin = GraphNode->FindPin(InputPin->Properties.Label, EEdGraphPinDirection::EGPD_Input);
 
 			if (!InPin)
 			{
@@ -94,7 +94,7 @@ void UPCGEditorGraph::CreateLinks(UPCGEditorGraphNodeBase* GraphNode, bool bCrea
 				const UPCGNode* InboundNode = InboundEdge->InputPin->Node;
 				if (UPCGEditorGraphNodeBase* const* ConnectedGraphNode = GraphNodeToPCGNodeMap.Find(InboundNode))
 				{
-					if (UEdGraphPin* OutPin = (*ConnectedGraphNode)->FindPin(InboundEdge->InputPin->Label, EEdGraphPinDirection::EGPD_Output))
+					if (UEdGraphPin* OutPin = (*ConnectedGraphNode)->FindPin(InboundEdge->InputPin->Properties.Label, EEdGraphPinDirection::EGPD_Output))
 					{
 						OutPin->MakeLinkTo(InPin);
 					}
@@ -111,7 +111,7 @@ void UPCGEditorGraph::CreateLinks(UPCGEditorGraphNodeBase* GraphNode, bool bCrea
 	{
 		for (const UPCGPin* OutputPin : PCGNode->GetOutputPins())
 		{
-			UEdGraphPin* OutPin = GraphNode->FindPin(OutputPin->Label, EEdGraphPinDirection::EGPD_Output);
+			UEdGraphPin* OutPin = GraphNode->FindPin(OutputPin->Properties.Label, EEdGraphPinDirection::EGPD_Output);
 
 			if (!OutPin)
 			{
@@ -128,7 +128,7 @@ void UPCGEditorGraph::CreateLinks(UPCGEditorGraphNodeBase* GraphNode, bool bCrea
 				const UPCGNode* OutboundNode = OutboundEdge->OutputPin->Node;
 				if (UPCGEditorGraphNodeBase* const* ConnectedGraphNode = GraphNodeToPCGNodeMap.Find(OutboundNode))
 				{
-					if (UEdGraphPin* InPin = (*ConnectedGraphNode)->FindPin(OutboundEdge->OutputPin->Label, EEdGraphPinDirection::EGPD_Input))
+					if (UEdGraphPin* InPin = (*ConnectedGraphNode)->FindPin(OutboundEdge->OutputPin->Properties.Label, EEdGraphPinDirection::EGPD_Input))
 					{
 						OutPin->MakeLinkTo(InPin);
 					}
diff --git a/Engine/Plugins/Experimental/PCG/Source/PCGEditor/Private/PCGEditorGraphNode.cpp b/Engine/Plugins/Experimental/PCG/Source/PCGEditor/Private/PCGEditorGraphNode.cpp
--- a/Engine/Plugins/Experimental/PCG/Source/PCGEditor/Private/PCGEditorGraphNode.cpp
+++ b/Engine/Plugins/Experimental/PCG/Source/PCGEditor/Private/PCGEditorGraphNode.cpp
@@ -47,12 +47,12 @@ void UPCGEditorGraphNode::AllocateDefaultPins()
 	{
 		for (const UPCGPin* InputPin : PCGNode->GetInputPins())
 		{
-			CreatePin(EEdGraphPinDirection::EGPD_Input, NAME_None, InputPin->Label);
+			CreatePin(EEdGraphPinDirection::EGPD_Input, GetPinType(InputPin), InputPin->Properties.Label);
 		}
 
 		for (const UPCGPin* OutputPin : PCGNode->GetOutputPins())
 		{
-			CreatePin(EEdGraphPinDirection::EGPD_Output, NAME_None, OutputPin->Label);
+			CreatePin(EEdGraphPinDirection::EGPD_Output, GetPinType(OutputPin), OutputPin->Properties.Label);
 		}
 	}
 }
diff --git a/Engine/Plugins/Experimental/PCG/Source/PCGEditor/Private/PCGEditorGraphNodeBase.cpp b/Engine/Plugins/Experimental/PCG/Source/PCGEditor/Private/PCGEditorGraphNodeBase.cpp
--- a/Engine/Plugins/Experimental/PCG/Source/PCGEditor/Private/PCGEditorGraphNodeBase.cpp
+++ b/Engine/Plugins/Experimental/PCG/Source/PCGEditor/Private/PCGEditorGraphNodeBase.cpp
@@ -2,6 +2,7 @@
 
 #include "PCGEditorGraphNodeBase.h"
 
+#include "PCGEditorCommon.h"
 #include "PCGEditorGraph.h"
 #include "PCGEditorGraphSchema.h"
 #include "PCGEditorSettings.h"
@@ -93,7 +94,7 @@ void UPCGEditorGraphNodeBase::AutowireNewNode(UEdGraphPin* FromPin)
 	{
 		if (PCGNode && PCGNode->GetInputPins().Num() > 0)
 		{
-			const FName& InPinName = PCGNode->GetInputPins()[0]->Label;
+			const FName& InPinName = PCGNode->GetInputPins()[0]->Properties.Label;
 			UEdGraphPin* ToPin = FindPinChecked(InPinName, EEdGraphPinDirection::EGPD_Input);
 			GetSchema()->TryCreateConnection(FromPin, ToPin);
 		}
@@ -102,7 +103,7 @@ void UPCGEditorGraphNodeBase::AutowireNewNode(UEdGraphPin* FromPin)
 	{
 		if (PCGNode && PCGNode->GetOutputPins().Num() > 0)
 		{
-			const FName& OutPinName = PCGNode->GetOutputPins()[0]->Label;
+			const FName& OutPinName = PCGNode->GetOutputPins()[0]->Properties.Label;
 			UEdGraphPin* ToPin = FindPinChecked(OutPinName, EEdGraphPinDirection::EGPD_Output);
 			GetSchema()->TryCreateConnection(FromPin, ToPin);
 		}
@@ -276,4 +277,64 @@ FLinearColor UPCGEditorGraphNodeBase::GetNodeTitleColor() const
 	return GetDefault<UPCGEditorSettings>()->DefaultNodeColor;
 }
 
+FEdGraphPinType UPCGEditorGraphNodeBase::GetPinType(const UPCGPin* InPin)
+{
+	FEdGraphPinType EdPinType;
+	EdPinType.ResetToDefaults();
+	EdPinType.PinCategory = NAME_None;
+	EdPinType.PinSubCategory = NAME_None;
+	EdPinType.ContainerType = EPinContainerType::None;
+
+	const EPCGDataType PinType = InPin->Properties.AllowedTypes;
+	auto CheckType = [PinType](EPCGDataType AllowedType)
+	{
+		return !!(PinType & AllowedType) && !(PinType & ~AllowedType);
+	};
+
+	if (CheckType(EPCGDataType::Spatial))
+	{
+		EdPinType.PinCategory = FPCGEditorCommon::SpatialDataType;
+
+		// Assign subcategory if we have precise information
+		if (CheckType(EPCGDataType::Point))
+		{
+			EdPinType.PinSubCategory = FPCGEditorCommon::PointDataType;
+		}
+		else if (CheckType(EPCGDataType::PolyLine))
+		{
+			EdPinType.PinSubCategory = FPCGEditorCommon::PolyLineDataType;
+		}
+		else if (CheckType(EPCGDataType::RenderTarget))
+		{
+			EdPinType.PinSubCategory = FPCGEditorCommon::RenderTargetDataType;
+		}
+		else if (CheckType(EPCGDataType::Surface))
+		{
+			EdPinType.PinSubCategory = FPCGEditorCommon::SurfaceDataType;
+		}
+		else if (CheckType(EPCGDataType::Volume))
+		{
+			EdPinType.PinSubCategory = FPCGEditorCommon::VolumeDataType;
+		}
+		else if (CheckType(EPCGDataType::Primitive))
+		{
+			EdPinType.PinSubCategory = FPCGEditorCommon::PrimitiveDataType;
+		}
+	}
+	else if (CheckType(EPCGDataType::Param))
+	{
+		EdPinType.PinCategory = FPCGEditorCommon::ParamDataType;
+	}
+	else if (CheckType(EPCGDataType::Settings))
+	{
+		EdPinType.PinCategory = FPCGEditorCommon::SettingsDataType;
+	}
+	else if (CheckType(EPCGDataType::Other))
+	{
+		EdPinType.PinCategory = FPCGEditorCommon::OtherDataType;
+	}	
+
+	return EdPinType;
+}
+
 #undef LOCTEXT_NAMESPACE
diff --git a/Engine/Plugins/Experimental/PCG/Source/PCGEditor/Private/PCGEditorGraphNodeBase.h b/Engine/Plugins/Experimental/PCG/Source/PCGEditor/Private/PCGEditorGraphNodeBase.h
--- a/Engine/Plugins/Experimental/PCG/Source/PCGEditor/Private/PCGEditorGraphNodeBase.h
+++ b/Engine/Plugins/Experimental/PCG/Source/PCGEditor/Private/PCGEditorGraphNodeBase.h
@@ -5,12 +5,14 @@
 #include "CoreMinimal.h"
 #include "UObject/ObjectMacros.h"
 #include "EdGraph/EdGraphNode.h"
+#include "EdGraph/EdGraphPin.h"
 
 #include "PCGCommon.h"
 
 #include "PCGEditorGraphNodeBase.generated.h"
 
 class UPCGNode;
+class UPCGPin;
 class UToolMenu;
 
 UENUM()
@@ -44,6 +46,7 @@ public:
 	// ~End UEdGraphNode interface
 
 	UPCGNode* GetPCGNode() { return PCGNode; }
+	const UPCGNode* GetPCGNode() const { return PCGNode; }
 	void PostCopy();
 	void PostPaste();
 
@@ -51,6 +54,8 @@ public:
 	FOnPCGEditorGraphNodeChanged OnNodeChangedDelegate;
 
 protected:
+	static FEdGraphPinType GetPinType(const UPCGPin* InPin);
+
 	void RebuildEdgesFromPins();
 
 	void OnNodeChanged(UPCGNode* InNode, EPCGChangeType ChangeType);
diff --git a/Engine/Plugins/Experimental/PCG/Source/PCGEditor/Private/PCGEditorGraphNodeInput.cpp b/Engine/Plugins/Experimental/PCG/Source/PCGEditor/Private/PCGEditorGraphNodeInput.cpp
--- a/Engine/Plugins/Experimental/PCG/Source/PCGEditor/Private/PCGEditorGraphNodeInput.cpp
+++ b/Engine/Plugins/Experimental/PCG/Source/PCGEditor/Private/PCGEditorGraphNodeInput.cpp
@@ -17,7 +17,7 @@ void UPCGEditorGraphNodeInput::AllocateDefaultPins()
 	{
 		for (const UPCGPin* OutPin : PCGNode->GetOutputPins())
 		{
-			CreatePin(EEdGraphPinDirection::EGPD_Output, NAME_None, OutPin->Label);
+			CreatePin(EEdGraphPinDirection::EGPD_Output, GetPinType(OutPin), OutPin->Properties.Label);
 		}
 	}
 }
diff --git a/Engine/Plugins/Experimental/PCG/Source/PCGEditor/Private/PCGEditorGraphNodeOutput.cpp b/Engine/Plugins/Experimental/PCG/Source/PCGEditor/Private/PCGEditorGraphNodeOutput.cpp
--- a/Engine/Plugins/Experimental/PCG/Source/PCGEditor/Private/PCGEditorGraphNodeOutput.cpp
+++ b/Engine/Plugins/Experimental/PCG/Source/PCGEditor/Private/PCGEditorGraphNodeOutput.cpp
@@ -17,7 +17,7 @@ void UPCGEditorGraphNodeOutput::AllocateDefaultPins()
 	{
 		for (const UPCGPin* InputPin : PCGNode->GetInputPins())
 		{
-			CreatePin(EEdGraphPinDirection::EGPD_Input, NAME_None, InputPin->Label);
+			CreatePin(EEdGraphPinDirection::EGPD_Input, GetPinType(InputPin), InputPin->Properties.Label);
 		}
 	}
 }
diff --git a/Engine/Plugins/Experimental/PCG/Source/PCGEditor/Private/PCGEditorGraphSchema.cpp b/Engine/Plugins/Experimental/PCG/Source/PCGEditor/Private/PCGEditorGraphSchema.cpp
--- a/Engine/Plugins/Experimental/PCG/Source/PCGEditor/Private/PCGEditorGraphSchema.cpp
+++ b/Engine/Plugins/Experimental/PCG/Source/PCGEditor/Private/PCGEditorGraphSchema.cpp
@@ -7,6 +7,7 @@
 #include "PCGEditorGraph.h"
 #include "PCGEditorGraphNodeBase.h"
 #include "PCGEditorGraphSchemaActions.h"
+#include "PCGEditorSettings.h"
 #include "PCGGraph.h"
 #include "PCGSettings.h"
 
@@ -36,7 +37,7 @@ void UPCGEditorGraphSchema::GetGraphContextActions(FGraphContextMenuBuilder& Con
 
 FLinearColor UPCGEditorGraphSchema::GetPinTypeColor(const FEdGraphPinType& PinType) const
 {
-	return FLinearColor(1.0f, 1.0f, 1.0f, 1.0f);
+	return GetDefault<UPCGEditorSettings>()->GetPinColor(PinType);
 }
 
 FConnectionDrawingPolicy* UPCGEditorGraphSchema::CreateConnectionDrawingPolicy(int32 InBackLayerID, int32 InFrontLayerID, float InZoomFactor, const FSlateRect& InClippingRect, class FSlateWindowElementList& InDrawElements, class UEdGraph* InGraphObj) const
@@ -46,6 +47,7 @@ FConnectionDrawingPolicy* UPCGEditorGraphSchema::CreateConnectionDrawingPolicy(i
 
 const FPinConnectionResponse UPCGEditorGraphSchema::CanCreateConnection(const UEdGraphPin* A, const UEdGraphPin* B) const
 {
+	check(A && B);
 	const UEdGraphNode* NodeA = A->GetOwningNode();
 	const UEdGraphNode* NodeB = B->GetOwningNode();
 
@@ -59,11 +61,40 @@ const FPinConnectionResponse UPCGEditorGraphSchema::CanCreateConnection(const UE
 		return FPinConnectionResponse(CONNECT_RESPONSE_DISALLOW, LOCTEXT("ConnectionSameDirection", "Both pins are the same direction"));
 	}
 
+	const UPCGEditorGraphNodeBase* EditorNodeA = CastChecked<const UPCGEditorGraphNodeBase>(NodeA);
+	const UPCGEditorGraphNodeBase* EditorNodeB = CastChecked<const UPCGEditorGraphNodeBase>(NodeB);
+
+	// Check type compatibility & whether we can connect more pins
+	const UPCGPin* InputPin = nullptr;
+	const UPCGPin* OutputPin = nullptr;
+
+	if (A->Direction == EGPD_Output)
+	{
+		OutputPin = EditorNodeA->GetPCGNode()->GetOutputPin(A->PinName);
+		InputPin = EditorNodeB->GetPCGNode()->GetInputPin(B->PinName);
+	}
+	else
+	{
+		OutputPin = EditorNodeB->GetPCGNode()->GetOutputPin(B->PinName);
+		InputPin = EditorNodeA->GetPCGNode()->GetInputPin(A->PinName);
+	}
+
+	if (!InputPin || !OutputPin)
+	{
+		return FPinConnectionResponse(CONNECT_RESPONSE_DISALLOW, LOCTEXT("ConnectionFailed", "Unable to verify pins"));
+	}
+
+	if (!InputPin->IsCompatible(OutputPin))
+	{
+		return FPinConnectionResponse(CONNECT_RESPONSE_DISALLOW, LOCTEXT("ConnectionTypesIncompatible", "Pins are incompatible"));
+	}
+
 	return FPinConnectionResponse();
 }
 
 bool UPCGEditorGraphSchema::TryCreateConnection(UEdGraphPin* InA, UEdGraphPin* InB) const
 {
+	// TODO: check if we need to verify connectivity first
 	bool bModified = Super::TryCreateConnection(InA, InB);
 
 	if (bModified)
@@ -88,6 +119,19 @@ bool UPCGEditorGraphSchema::TryCreateConnection(UEdGraphPin* InA, UEdGraphPin* I
 		check(PCGGraph);
 
 		PCGGraph->AddLabeledEdge(PCGNodeA, A->PinName, PCGNodeB, B->PinName);
+
+		// TODO: unclear if that kind of behavior should be down the code hierarchy or not,
+		// Since we really want to do cleanup only on manual interaction
+		if (UPCGPin* InputPin = PCGNodeB->GetInputPin(B->PinName))
+		{
+			if (!InputPin->Properties.bAllowMultipleConnections)
+			{
+				if (InputPin->BreakAllIncompatibleEdges())
+				{
+					PCGGraphNodeB->ReconstructNode();
+				}
+			}
+		}
 	}
 
 	return bModified;
@@ -239,6 +283,12 @@ void FPCGEditorConnectionDrawingPolicy::DetermineWiringStyle(UEdGraphPin* Output
 	{
 		Params.WireThickness = Params.WireThickness * 3;
 	}
+
+	// Base the color of the wire on the color of the output pin
+	if (OutputPin)
+	{
+		Params.WireColor = GetDefault<UPCGEditorSettings>()->GetPinColor(OutputPin->PinType);
+	}
 }
 
 #undef LOCTEXT_NAMESPACE
diff --git a/Engine/Plugins/Experimental/PCG/Source/PCGEditor/Private/PCGEditorSettings.cpp b/Engine/Plugins/Experimental/PCG/Source/PCGEditor/Private/PCGEditorSettings.cpp
--- a/Engine/Plugins/Experimental/PCG/Source/PCGEditor/Private/PCGEditorSettings.cpp
+++ b/Engine/Plugins/Experimental/PCG/Source/PCGEditor/Private/PCGEditorSettings.cpp
@@ -1,6 +1,9 @@
 // Copyright Epic Games, Inc. All Rights Reserved.
 
 #include "PCGEditorSettings.h"
+#include "PCGEditorCommon.h"
+
+#include "EdGraph/EdGraphPin.h"
 
 UPCGEditorSettings::UPCGEditorSettings(const FObjectInitializer& ObjectInitializer)
 	: Super(ObjectInitializer)
@@ -16,6 +19,12 @@ UPCGEditorSettings::UPCGEditorSettings(const FObjectInitializer& ObjectInitializ
 	SpawnerNodeColor = FLinearColor(1.0f, 0.6f, 0.4f);
 	SubgraphNodeColor = FLinearColor(1.0f, 0.1f, 0.1f);
 	DebugNodeColor = FLinearColor(1.0f, 0.0f, 1.0f);
+
+	DefaultPinColor = FLinearColor(1.0f, 1.0f, 1.0f);
+	SpatialDataPinColor = FLinearColor(0.2f, 0.2f, 1.0f);
+	RenderTargetDataPinColor = FLinearColor(1.0f, 0.3f, 0.f);
+	ParamDataPinColor = FLinearColor(1.0f, 0.6f, 0.0f);
+	UnknownDataPinColor = FLinearColor(0.3f, 0.3f, 0.3f);
 }
 
 FLinearColor UPCGEditorSettings::GetColor(UPCGSettings* Settings) const
@@ -76,3 +85,30 @@ FLinearColor UPCGEditorSettings::GetColor(UPCGSettings* Settings) const
 		return DefaultNodeColor;
 	}
 }
+
+FLinearColor UPCGEditorSettings::GetPinColor(const FEdGraphPinType& PinType) const
+{
+	if (PinType.PinCategory == FPCGEditorCommon::SpatialDataType)
+	{
+		if (PinType.PinSubCategory == FPCGEditorCommon::RenderTargetDataType)
+		{
+			return RenderTargetDataPinColor;
+		}
+		else
+		{
+			return SpatialDataPinColor;
+		}
+	}
+	else if (PinType.PinCategory == FPCGEditorCommon::ParamDataType)
+	{
+		return ParamDataPinColor;
+	}
+	else if (PinType.PinCategory == FPCGEditorCommon::OtherDataType)
+	{
+		return UnknownDataPinColor;
+	}
+	else
+	{
+		return DefaultPinColor;
+	}
+}
diff --git a/Engine/Plugins/Experimental/PCG/Source/PCGEditor/Public/PCGEditorSettings.h b/Engine/Plugins/Experimental/PCG/Source/PCGEditor/Public/PCGEditorSettings.h
--- a/Engine/Plugins/Experimental/PCG/Source/PCGEditor/Public/PCGEditorSettings.h
+++ b/Engine/Plugins/Experimental/PCG/Source/PCGEditor/Public/PCGEditorSettings.h
@@ -7,6 +7,8 @@
 
 #include "PCGEditorSettings.generated.h"
 
+struct FEdGraphPinType;
+
 UCLASS(config=EditorPerProjectUserSettings)
 class PCGEDITOR_API UPCGEditorSettings : public UObject
 {
@@ -16,6 +18,7 @@ public:
 	UPCGEditorSettings(const FObjectInitializer& ObjectInitializer);
 
 	FLinearColor GetColor(UPCGSettings* InSettings) const;
+	FLinearColor GetPinColor(const FEdGraphPinType& InPinType) const;
 
 	/** Default node color */
 	UPROPERTY(EditAnywhere, config, Category = Node, meta = (HideAlphaChannel))
@@ -61,6 +64,26 @@ public:
 	UPROPERTY(EditAnywhere, config, Category = Node, meta = (HideAlphaChannel))
 	FLinearColor DebugNodeColor;
 
+	/** Default pin color */
+	UPROPERTY(EditAnywhere, config, Category = Node, meta = (HideAlphaChannel))
+	FLinearColor DefaultPinColor;
+
+	/** Color used for spatial data pins */
+	UPROPERTY(EditAnywhere, config, Category = Node, meta = (HideAlphaChannel))
+	FLinearColor SpatialDataPinColor;
+
+	/** Color used for render targets / GPU pins */
+	UPROPERTY(EditAnywhere, config, Category = Node, meta = (HideAlphaChannel))
+	FLinearColor RenderTargetDataPinColor;
+
+	/** Color used for param data pins */
+	UPROPERTY(EditAnywhere, config, Category = Node, meta = (HideAlphaChannel))
+	FLinearColor ParamDataPinColor;
+
+	/** Color used for other/unknown data types */
+	UPROPERTY(EditAnywhere, config, Category = Node, meta = (HideAlphaChannel))
+	FLinearColor UnknownDataPinColor;
+
 	/** User-driven color overrides */
 	UPROPERTY(EditAnywhere, config, Category = Node, meta = (HideAlphaChannel))
 	TMap<TSubclassOf<UPCGSettings>, FLinearColor> OverrideNodeColorByClass;
