Fixed output Out pin default type

#rnx

diff --git a/Engine/Plugins/Experimental/PCG/Source/PCG/Private/PCGInputOutputSettings.cpp b/Engine/Plugins/Experimental/PCG/Source/PCG/Private/PCGInputOutputSettings.cpp
--- a/Engine/Plugins/Experimental/PCG/Source/PCG/Private/PCGInputOutputSettings.cpp
+++ b/Engine/Plugins/Experimental/PCG/Source/PCG/Private/PCGInputOutputSettings.cpp
@@ -39,30 +39,27 @@ void UPCGGraphInputOutputSettings::PostLoad()
 	}
 }
 
-TArray<FPCGPinProperties> UPCGGraphInputOutputSettings::InputPinProperties() const
+TArray<FPCGPinProperties> UPCGGraphInputOutputSettings::GetPinProperties() const
 {
 	TArray<FPCGPinProperties> PinProperties;
-	Algo::Transform(StaticLabels(), PinProperties, [](const FName& InLabel) { return FPCGPinProperties(InLabel, EPCGDataType::Spatial); });
+	const EPCGDataType DefaultPinDataType = bIsInput ? EPCGDataType::Spatial : EPCGDataType::Any;
+	Algo::Transform(StaticLabels(), PinProperties, [DefaultPinDataType](const FName& InLabel) { return FPCGPinProperties(InLabel, DefaultPinDataType); });
 	
 	if (bShowAdvancedPins)
 	{
-		Algo::Transform(StaticAdvancedLabels(), PinProperties, [](const FName& InLabel) { return FPCGPinProperties(InLabel, EPCGDataType::Spatial); });
+		Algo::Transform(StaticAdvancedLabels(), PinProperties, [DefaultPinDataType](const FName& InLabel) { return FPCGPinProperties(InLabel, DefaultPinDataType); });
 	}
 
 	PinProperties.Append(CustomPins);
 	return PinProperties;
 }
 
-TArray<FPCGPinProperties> UPCGGraphInputOutputSettings::OutputPinProperties() const
+TArray<FPCGPinProperties> UPCGGraphInputOutputSettings::InputPinProperties() const
 {
-	TArray<FPCGPinProperties> PinProperties;
-	Algo::Transform(StaticLabels(), PinProperties, [](const FName& InLabel) { return FPCGPinProperties(InLabel, EPCGDataType::Spatial); });
-
-	if (bShowAdvancedPins)
-	{
-		Algo::Transform(StaticAdvancedLabels(), PinProperties, [](const FName& InLabel) { return FPCGPinProperties(InLabel, EPCGDataType::Spatial); });
-	}
+	return GetPinProperties();
+}
 
-	PinProperties.Append(CustomPins);
-	return PinProperties;
+TArray<FPCGPinProperties> UPCGGraphInputOutputSettings::OutputPinProperties() const
+{
+	return GetPinProperties();
 }
\ No newline at end of file
diff --git a/Engine/Plugins/Experimental/PCG/Source/PCG/Public/PCGInputOutputSettings.h b/Engine/Plugins/Experimental/PCG/Source/PCG/Public/PCGInputOutputSettings.h
--- a/Engine/Plugins/Experimental/PCG/Source/PCG/Public/PCGInputOutputSettings.h
+++ b/Engine/Plugins/Experimental/PCG/Source/PCG/Public/PCGInputOutputSettings.h
@@ -52,6 +52,8 @@ protected:
 	const TArray<FName>& StaticLabels() const { return bIsInput ? StaticInLabels : StaticOutLabels; }
 	const TArray<FName>& StaticAdvancedLabels() const { return bIsInput ? StaticAdvancedInLabels : StaticAdvancedOutLabels; }
 
+	TArray<FPCGPinProperties> GetPinProperties() const;
+
 protected:
 	UPROPERTY()
 	TSet<FName> PinLabels_DEPRECATED;
