Fixed pins not being intiialized properly when a node is dropped from the palette

#rnx

diff --git a/Engine/Plugins/Experimental/PCG/Source/PCG/Private/PCGNode.cpp b/Engine/Plugins/Experimental/PCG/Source/PCG/Private/PCGNode.cpp
--- a/Engine/Plugins/Experimental/PCG/Source/PCG/Private/PCGNode.cpp
+++ b/Engine/Plugins/Experimental/PCG/Source/PCG/Private/PCGNode.cpp
@@ -392,6 +392,11 @@ void UPCGNode::OnSettingsChanged(UPCGSettings* InSettings, EPCGChangeType Change
 
 #endif // WITH_EDITOR
 
+void UPCGNode::UpdateAfterSettingsChangeDuringCreation()
+{
+	UpdatePins();
+}
+
 bool UPCGNode::UpdatePins()
 {
 	return UpdatePins([](UPCGNode* Node){ return NewObject<UPCGPin>(Node); });
diff --git a/Engine/Plugins/Experimental/PCG/Source/PCG/Public/PCGNode.h b/Engine/Plugins/Experimental/PCG/Source/PCG/Public/PCGNode.h
--- a/Engine/Plugins/Experimental/PCG/Source/PCG/Public/PCGNode.h
+++ b/Engine/Plugins/Experimental/PCG/Source/PCG/Public/PCGNode.h
@@ -68,6 +68,9 @@ public:
 	/** Changes the default settings in the node */
 	void SetDefaultSettings(TObjectPtr<UPCGSettings> InSettings, bool bUpdatePins = true);
 
+	/** Triggers some uppdates after creating a new node and changing its settings */
+	void UpdateAfterSettingsChangeDuringCreation();
+
 	UPCGPin* GetInputPin(const FName& Label);
 	const UPCGPin* GetInputPin(const FName& Label) const;
 	UPCGPin* GetOutputPin(const FName& Label);
diff --git a/Engine/Plugins/Experimental/PCG/Source/PCGEditor/Private/PCGEditorGraphSchemaActions.cpp b/Engine/Plugins/Experimental/PCG/Source/PCGEditor/Private/PCGEditorGraphSchemaActions.cpp
--- a/Engine/Plugins/Experimental/PCG/Source/PCGEditor/Private/PCGEditorGraphSchemaActions.cpp
+++ b/Engine/Plugins/Experimental/PCG/Source/PCGEditor/Private/PCGEditorGraphSchemaActions.cpp
@@ -79,6 +79,8 @@ UEdGraphNode* FPCGEditorGraphSchemaAction_NewBlueprintElement::PerformAction(UEd
 	TSubclassOf<UPCGBlueprintElement> BlueprintClass = BlueprintClassPath.TryLoadClass<UPCGBlueprintElement>();
 	DefaultBlueprintSettings->SetElementType(BlueprintClass, ElementInstance);
 
+	NewPCGNode->UpdateAfterSettingsChangeDuringCreation();
+
 	FGraphNodeCreator<UPCGEditorGraphNode> NodeCreator(*EditorGraph);
 	UPCGEditorGraphNode* NewNode = NodeCreator.CreateUserInvokedNode(bSelectNewNode);
 	NewNode->Construct(NewPCGNode, EPCGEditorGraphNodeType::Settings);
@@ -125,6 +127,8 @@ UEdGraphNode* FPCGEditorGraphSchemaAction_NewSubgraphElement::PerformAction(UEdG
 	UPCGSubgraphSettings* DefaultSubgraphSettings = CastChecked<UPCGSubgraphSettings>(DefaultNodeSettings);
 	DefaultSubgraphSettings->Subgraph = Subgraph;
 
+	NewPCGNode->UpdateAfterSettingsChangeDuringCreation();
+
 	FGraphNodeCreator<UPCGEditorGraphNode> NodeCreator(*EditorGraph);
 	UPCGEditorGraphNode* NewNode = NodeCreator.CreateUserInvokedNode(bSelectNewNode);
 	NewNode->Construct(NewPCGNode, EPCGEditorGraphNodeType::Settings);
