Made output pin in blueprints not spatial by default

#rnx

diff --git a/Engine/Plugins/Experimental/PCG/Source/PCG/Private/Elements/PCGExecuteBlueprint.cpp b/Engine/Plugins/Experimental/PCG/Source/PCG/Private/Elements/PCGExecuteBlueprint.cpp
--- a/Engine/Plugins/Experimental/PCG/Source/PCG/Private/Elements/PCGExecuteBlueprint.cpp
+++ b/Engine/Plugins/Experimental/PCG/Source/PCG/Private/Elements/PCGExecuteBlueprint.cpp
@@ -398,7 +398,8 @@ TArray<FPCGPinProperties> UPCGBlueprintSettings::OutputPinProperties() const
 	{
 		if (BlueprintElementInstance->bHasDefaultOutPin)
 		{
-			PinProperties.Append(Super::OutputPinProperties());
+			// Note: we do not use the default base class pin here, as a blueprint node might return anything
+			PinProperties.Emplace(PCGPinConstants::DefaultOutputLabel, EPCGDataType::Any);
 		}
 
 		PinProperties.Append(BlueprintElementInstance->CustomOutputPins);
