Fixed scaling issue introduced with the new SamplePoint code

#rnx

diff --git a/Engine/Plugins/Experimental/PCG/Source/PCG/Private/Data/PCGLandscapeSplineData.cpp b/Engine/Plugins/Experimental/PCG/Source/PCG/Private/Data/PCGLandscapeSplineData.cpp
--- a/Engine/Plugins/Experimental/PCG/Source/PCG/Private/Data/PCGLandscapeSplineData.cpp
+++ b/Engine/Plugins/Experimental/PCG/Source/PCG/Private/Data/PCGLandscapeSplineData.cpp
@@ -211,10 +211,10 @@ bool UPCGLandscapeSplineData::SamplePoint(const FTransform& InTransform, const F
 	// TODO : add support for bounds
 	check(Spline);
 
-	OutPoint.Transform = Spline->GetComponentTransform().Inverse() * InTransform;
+	OutPoint.Transform = InTransform;
 	OutPoint.SetLocalBounds(InBounds); // TODO: should maybe do Min.Z = Max.Z = 0 ?
 
-	const FVector Position = OutPoint.Transform.GetLocation();
+	const FVector Position = Spline->GetComponentTransform().InverseTransformPosition(OutPoint.Transform.GetLocation());
 
 	float PointDensity = 0.0f;
 
diff --git a/Engine/Plugins/Experimental/PCG/Source/PCG/Private/Data/PCGTextureData.cpp b/Engine/Plugins/Experimental/PCG/Source/PCG/Private/Data/PCGTextureData.cpp
--- a/Engine/Plugins/Experimental/PCG/Source/PCG/Private/Data/PCGTextureData.cpp
+++ b/Engine/Plugins/Experimental/PCG/Source/PCG/Private/Data/PCGTextureData.cpp
@@ -93,10 +93,10 @@ bool UPCGBaseTextureData::SamplePoint(const FTransform& InTransform, const FBox&
 
 	// Compute transform
 	// TODO: embed local bounds center offset at this time?
-	OutPoint.Transform = Transform.Inverse() * InTransform;
-	FVector PointPositionInLocalSpace = OutPoint.Transform.GetLocation();
+	OutPoint.Transform = InTransform;
+	FVector PointPositionInLocalSpace = Transform.InverseTransformPosition(InTransform.GetLocation());
 	PointPositionInLocalSpace.Z = 0;
-	OutPoint.Transform.SetLocation(PointPositionInLocalSpace);
+	OutPoint.Transform.SetLocation(Transform.TransformPosition(PointPositionInLocalSpace));
 	OutPoint.SetLocalBounds(InBounds); // TODO: should set Min.Z = Max.Z = 0;
 
 	// Compute density & color (& metadata)
