Added arbitrary 1:N point loop in PCG blueprint

diff --git a/Engine/Plugins/Experimental/PCG/Source/PCG/Private/Elements/PCGExecuteBlueprint.cpp b/Engine/Plugins/Experimental/PCG/Source/PCG/Private/Elements/PCGExecuteBlueprint.cpp
--- a/Engine/Plugins/Experimental/PCG/Source/PCG/Private/Elements/PCGExecuteBlueprint.cpp
+++ b/Engine/Plugins/Experimental/PCG/Source/PCG/Private/Elements/PCGExecuteBlueprint.cpp
@@ -506,6 +506,26 @@ void UPCGBlueprintElement::LoopOnPoints(FPCGContext& InContext, const UPCGPointD
 	});
 }
 
+void UPCGBlueprintElement::MultiLoopOnPoints(FPCGContext& InContext, const UPCGPointData* InData, UPCGPointData*& OutData, UPCGPointData* OptionalOutData) const
+{
+	if (!InData)
+	{
+		PCGE_LOG_C(Error, &InContext, "Invalid input data in MultiLoopOnPoints");
+		return;
+	}
+
+	OutData = OptionalOutData ? OptionalOutData : NewObject<UPCGPointData>(const_cast<UPCGPointData*>(InData));
+	OutData->InitializeFromData(InData);
+
+	const TArray<FPCGPoint>& InPoints = InData->GetPoints();
+	TArray<FPCGPoint>& OutPoints = OutData->GetMutablePoints();
+
+	FPCGAsync::AsyncMultiPointProcessing(&InContext, InPoints.Num(), OutPoints, [this, &InContext, InData, OutData, &InPoints](int32 Index)
+	{
+		return MultiPointLoopBody(InContext, InData, InPoints[Index], OutData->Metadata);
+	});
+}
+
 void UPCGBlueprintElement::LoopOnPointPairs(FPCGContext& InContext, const UPCGPointData* InA, const UPCGPointData* InB, UPCGPointData*& OutData, UPCGPointData* OptionalOutData) const
 {
 	if (!InA || !InB)
diff --git a/Engine/Plugins/Experimental/PCG/Source/PCG/Private/Helpers/PCGAsync.cpp b/Engine/Plugins/Experimental/PCG/Source/PCG/Private/Helpers/PCGAsync.cpp
--- a/Engine/Plugins/Experimental/PCG/Source/PCG/Private/Helpers/PCGAsync.cpp
+++ b/Engine/Plugins/Experimental/PCG/Source/PCG/Private/Helpers/PCGAsync.cpp
@@ -187,4 +187,52 @@ namespace FPCGAsync
 			OutFilterPoints.SetNum(OutFilterRangeIndex);
 		}
 	}
+
+	void AsyncMultiPointProcessing(FPCGContext* Context, int32 NumIterations, TArray<FPCGPoint>& OutPoints, const TFunction<TArray<FPCGPoint>(int32)>& PointFunc)
+	{
+		const int32 MinIterationsPerTask = 256;
+		AsyncMultiPointProcessing(Context ? Context->NumAvailableTasks : 1, MinIterationsPerTask, NumIterations, OutPoints, PointFunc);
+	}
+
+	void AsyncMultiPointProcessing(int32 NumAvailableTasks, int32 MinIterationsPerTask, int32 NumIterations, TArray<FPCGPoint>& OutPoints, const TFunction<TArray<FPCGPoint>(int32)>& PointFunc)
+	{
+		TRACE_CPUPROFILER_EVENT_SCOPE(IPCGElement::AsyncMultiPointProcessing);
+		check(NumAvailableTasks > 0 && MinIterationsPerTask > 0 && NumIterations >= 0);
+		// Get number of available threads from the context
+		const int32 NumTasks = FMath::Min(NumAvailableTasks, FMath::Max(1, NumIterations / MinIterationsPerTask));
+		const int32 IterationsPerTask = NumIterations / NumTasks;
+
+		TArray<TFuture<TArray<FPCGPoint>>> AsyncTasks;
+		AsyncTasks.Reserve(NumTasks);
+
+		// Launch the async tasks
+		for (int32 TaskIndex = 0; TaskIndex < NumTasks; ++TaskIndex)
+		{
+			const int32 StartIndex = TaskIndex * IterationsPerTask;
+			const int32 EndIndex = (TaskIndex == NumTasks - 1) ? NumIterations : (StartIndex + IterationsPerTask);
+
+			AsyncTasks.Emplace(Async(EAsyncExecution::ThreadPool, [&PointFunc, StartIndex, EndIndex]()
+			{
+				TRACE_CPUPROFILER_EVENT_SCOPE(IPCGElement::AsyncMultiPointProcessing::InnerLoop);
+				TArray<FPCGPoint> OutPoints;
+
+				for (int32 Index = StartIndex; Index < EndIndex; ++Index)
+				{
+					OutPoints.Append(PointFunc(Index));
+				}
+
+				return OutPoints;
+			}));
+		}
+
+		// Wait/Gather results & collapse points
+		{
+			TRACE_CPUPROFILER_EVENT_SCOPE(IPCGElement::AsyncMultiPointProcessing::WaitAndCollapseArray);
+			for (TFuture<TArray<FPCGPoint>>& AsyncTask : AsyncTasks)
+			{
+				AsyncTask.Wait();
+				OutPoints.Append(AsyncTask.Get());
+			}
+		}
+	}
 }
\ No newline at end of file
diff --git a/Engine/Plugins/Experimental/PCG/Source/PCG/Public/Elements/PCGExecuteBlueprint.h b/Engine/Plugins/Experimental/PCG/Source/PCG/Public/Elements/PCGExecuteBlueprint.h
--- a/Engine/Plugins/Experimental/PCG/Source/PCG/Public/Elements/PCGExecuteBlueprint.h
+++ b/Engine/Plugins/Experimental/PCG/Source/PCG/Public/Elements/PCGExecuteBlueprint.h
@@ -43,6 +43,9 @@ public:
 	UFUNCTION(BlueprintImplementableEvent, Category = Execution)
 	bool PointLoopBody(UPARAM(ref) FPCGContext& InContext, const UPCGPointData* InData, const FPCGPoint& InPoint, FPCGPoint& OutPoint, UPCGMetadata* OutMetadata) const;
 
+	UFUNCTION(BlueprintImplementableEvent, Category = Execution)
+	TArray<FPCGPoint> MultiPointLoopBody(UPARAM(ref) FPCGContext& InContext, const UPCGPointData* InData, const FPCGPoint& InPoint, UPCGMetadata* OutMetadata) const;
+
 	UFUNCTION(BlueprintImplementableEvent, Category = Execution)
 	bool PointPairLoopBody(UPARAM(ref) FPCGContext& InContext, const UPCGPointData* InA, const UPCGPointData* InB, const FPCGPoint& InPointA, const FPCGPoint& InPointB, FPCGPoint& OutPoint, UPCGMetadata* OutMetadata) const;
 
@@ -53,6 +56,9 @@ public:
 	UFUNCTION(BlueprintCallable, BlueprintPure = false, Category = Execution, meta = (HideSelfPin = "true"))
 	void LoopOnPoints(UPARAM(ref) FPCGContext& InContext, const UPCGPointData* InData, UPCGPointData*& OutData, UPCGPointData* OptionalOutData = nullptr) const;
 
+	UFUNCTION(BlueprintCallable, BlueprintPure = false, Category = Execution, meta = (HideSelfPin = "true"))
+	void MultiLoopOnPoints(UPARAM(ref) FPCGContext& InContext, const UPCGPointData* InData, UPCGPointData*& OutData, UPCGPointData* OptionalOutData = nullptr) const;
+
 	UFUNCTION(BlueprintCallable, BlueprintPure = false, Category = Execution, meta = (HideSelfPin = "true"))
 	void LoopOnPointPairs(UPARAM(ref) FPCGContext& InContext, const UPCGPointData* InA, const UPCGPointData* InB, UPCGPointData*& OutData, UPCGPointData* OptionalOutData = nullptr) const;
 
diff --git a/Engine/Plugins/Experimental/PCG/Source/PCG/Public/Helpers/PCGAsync.h b/Engine/Plugins/Experimental/PCG/Source/PCG/Public/Helpers/PCGAsync.h
--- a/Engine/Plugins/Experimental/PCG/Source/PCG/Public/Helpers/PCGAsync.h
+++ b/Engine/Plugins/Experimental/PCG/Source/PCG/Public/Helpers/PCGAsync.h
@@ -48,4 +48,23 @@ namespace FPCGAsync
 	* @param PointFunc - A function that has the index [0; NumIterations] and has to write to the point & return true when the current call generates a point
 	*/
 	void AsyncPointFilterProcessing(int32 NumAvailableTasks, int32 MinIterationsPerTask, int32 NumIterations, TArray<FPCGPoint>& InFilterPoints, TArray<FPCGPoint>& OutFilterPoints, const TFunction<bool(int32, FPCGPoint&, FPCGPoint&)>& PointFunc);
+
+	/**
+	* Helper to do simple 1:N point processing loops
+	* @param Context - The context containing the information about how many tasks to launch
+	* @param NumIterations - The number of calls that will be done to the provided function, also an upper bound on the number of points generated
+	* @param OutPoints - The array in which the results will be written to. Note that the array will be cleared before execution
+	* @param PointFunc - A function that has the index [0; NumIterations] and has to write to the point & return true when the current call generates a point
+	*/
+	void AsyncMultiPointProcessing(FPCGContext* Context, int32 NumIterations, TArray<FPCGPoint>& OutPoints, const TFunction<TArray<FPCGPoint>(int32)>& PointFunc);
+
+	/** 
+	* Helper to do more general 1:N point processing loops
+	* @param NumAvailableTasks - The upper bound on the number of async tasks we'll start
+	* @param MinIterationsPerTask - The lower bound on the number of iterations per task we'll dispatch
+	* @param NumIterations - The number of calls that will be done to the provided function, also an upper bound on the number of points generated
+ 	* @param OutPoints - The array in which the results will be written to. Note that the array will be cleared before execution
+	* @param PointFunc - A function that has the index [0; NumIterations] and has to write to the point & return true when the current call generates a point
+	*/
+	void AsyncMultiPointProcessing(int32 NumAvailableTasks, int32 MinIterationsPerTask, int32 NumIterations, TArray<FPCGPoint>& OutPoints, const TFunction<TArray<FPCGPoint>(int32)>& PointFunc);
 }
\ No newline at end of file
