Added the primitive data voxel size per axis.
Set the seed on points on creation.

#rnx

diff --git a/Engine/Plugins/Experimental/PCG/Source/PCG/Private/Data/PCGPrimitiveData.cpp b/Engine/Plugins/Experimental/PCG/Source/PCG/Private/Data/PCGPrimitiveData.cpp
--- a/Engine/Plugins/Experimental/PCG/Source/PCG/Private/Data/PCGPrimitiveData.cpp
+++ b/Engine/Plugins/Experimental/PCG/Source/PCG/Private/Data/PCGPrimitiveData.cpp
@@ -4,6 +4,7 @@
 
 #include "Data/PCGPointData.h"
 #include "Helpers/PCGAsync.h"
+#include "PCGHelpers.h"
 
 #include "Components/PrimitiveComponent.h"
 
@@ -44,12 +45,12 @@ const UPCGPointData* UPCGPrimitiveData::CreatePointData(FPCGContext* Context) co
 	Data->InitializeFromData(this);
 	TArray<FPCGPoint>& Points = Data->GetMutablePoints();
 
-	const int32 MinX = FMath::CeilToInt(CachedBounds.Min.X / VoxelSize);
-	const int32 MaxX = FMath::FloorToInt(CachedBounds.Max.X / VoxelSize);
-	const int32 MinY = FMath::CeilToInt(CachedBounds.Min.Y / VoxelSize);
-	const int32 MaxY = FMath::FloorToInt(CachedBounds.Max.Y / VoxelSize);
-	const int32 MinZ = FMath::CeilToInt(CachedBounds.Min.Z / VoxelSize);
-	const int32 MaxZ = FMath::FloorToInt(CachedBounds.Max.Z / VoxelSize);
+	const int32 MinX = FMath::CeilToInt(CachedBounds.Min.X / VoxelSize.X);
+	const int32 MaxX = FMath::FloorToInt(CachedBounds.Max.X / VoxelSize.X);
+	const int32 MinY = FMath::CeilToInt(CachedBounds.Min.Y / VoxelSize.Y);
+	const int32 MaxY = FMath::FloorToInt(CachedBounds.Max.Y / VoxelSize.Y);
+	const int32 MinZ = FMath::CeilToInt(CachedBounds.Min.Z / VoxelSize.Z);
+	const int32 MaxZ = FMath::FloorToInt(CachedBounds.Max.Z / VoxelSize.Z);
 
 	const int32 NumIterations = (MaxX - MinX) * (MaxY - MinY) * (MaxZ - MinZ);
 
@@ -59,9 +60,17 @@ const UPCGPointData* UPCGPrimitiveData::CreatePointData(FPCGContext* Context) co
 		const int Y = MinY + (Index / (MaxX - MinX) % (MaxY - MinY));
 		const int Z = MinZ + (Index / ((MaxX - MinX) * (MaxY - MinY)));
 
-		const FVector SampleLocation(X * VoxelSize, Y * VoxelSize, Z * VoxelSize);
-		const FBox VoxelBox(FVector(-VoxelSize / 2.0), FVector(VoxelSize / 2.0));
-		return SamplePoint(FTransform(SampleLocation), VoxelBox, OutPoint, nullptr);
+		const FVector SampleLocation(X * VoxelSize.X, Y * VoxelSize.Y, Z * VoxelSize.Z);
+		const FBox VoxelBox(VoxelSize * -0.5, VoxelSize * 0.5);
+		if (SamplePoint(FTransform(SampleLocation), VoxelBox, OutPoint, nullptr))
+		{
+			OutPoint.Seed = PCGHelpers::ComputeSeed(X, Y, Z);
+			return true;
+		}
+		else
+		{
+			return false;
+		}
 	});
 
 	UE_LOG(LogPCG, Verbose, TEXT("Primitive %s extracted %d points"), *Primitive->GetFName().ToString(), Points.Num());
diff --git a/Engine/Plugins/Experimental/PCG/Source/PCG/Public/Data/PCGPrimitiveData.h b/Engine/Plugins/Experimental/PCG/Source/PCG/Public/Data/PCGPrimitiveData.h
--- a/Engine/Plugins/Experimental/PCG/Source/PCG/Public/Data/PCGPrimitiveData.h
+++ b/Engine/Plugins/Experimental/PCG/Source/PCG/Public/Data/PCGPrimitiveData.h
@@ -30,7 +30,7 @@ public:
 	// ~End UPCGSpatialDataWithPointCache implementation
 
 	UPROPERTY(BlueprintReadWrite, EditAnywhere, Category = "Data")
-	float VoxelSize = 100.0f;
+	FVector VoxelSize = FVector(100.0, 100.0, 100.0);
 
 protected:
 	UPROPERTY(BlueprintReadOnly, VisibleAnywhere, Category = "Data")
