Fixed issue with behavior of intersection & union when there is a point location transform (such as a projection)

#rnx

diff --git a/Engine/Plugins/Experimental/PCG/Source/PCG/Private/Data/PCGIntersectionData.cpp b/Engine/Plugins/Experimental/PCG/Source/PCG/Private/Data/PCGIntersectionData.cpp
--- a/Engine/Plugins/Experimental/PCG/Source/PCG/Private/Data/PCGIntersectionData.cpp
+++ b/Engine/Plugins/Experimental/PCG/Source/PCG/Private/Data/PCGIntersectionData.cpp
@@ -53,27 +53,29 @@ FBox UPCGIntersectionData::GetStrictBounds() const
 bool UPCGIntersectionData::SamplePoint(const FTransform& InTransform, const FBox& InBounds, FPCGPoint& OutPoint, UPCGMetadata* OutMetadata) const
 {
 	check(A && B);
+	const UPCGSpatialData* X = (A->HasNonTrivialTransform() || !B->HasNonTrivialTransform()) ? A : B;
+	const UPCGSpatialData* Y = (X == A) ? B : A;
 
-	FPCGPoint PointFromA;
-	if(!A->SamplePoint(InTransform, InBounds, PointFromA, OutMetadata))
+	FPCGPoint PointFromX;
+	if(!X->SamplePoint(InTransform, InBounds, PointFromX, OutMetadata))
 	{
 		return false;
 	}
 
-	FPCGPoint PointFromB;
-	if(!B->SamplePoint(InTransform, InBounds, PointFromB, OutMetadata))
+	FPCGPoint PointFromY;
+	if(!Y->SamplePoint(PointFromX.Transform, InBounds, PointFromY, OutMetadata))
 	{
 		return false;
 	}
 
 	// Merge points into a single point
-	OutPoint = (A->HasNonTrivialTransform() ? PointFromA : PointFromB);
-	OutPoint.Density = PCGIntersectionDataMaths::ComputeDensity(PointFromA.Density, PointFromB.Density, DensityFunction);
-	OutPoint.Color = PointFromA.Color * PointFromB.Color;
+	OutPoint = PointFromY;
+	OutPoint.Density = PCGIntersectionDataMaths::ComputeDensity(PointFromX.Density, PointFromY.Density, DensityFunction);
+	OutPoint.Color = PointFromX.Color * PointFromY.Color;
 
 	if (OutMetadata)
 	{
-		OutMetadata->MergePointAttributes(PointFromA, PointFromB, OutPoint, EPCGMetadataOp::Min);
+		OutMetadata->MergePointAttributes(PointFromX, PointFromY, OutPoint, EPCGMetadataOp::Min);
 	}
 
 	return true;
diff --git a/Engine/Plugins/Experimental/PCG/Source/PCG/Private/Data/PCGUnionData.cpp b/Engine/Plugins/Experimental/PCG/Source/PCG/Private/Data/PCGUnionData.cpp
--- a/Engine/Plugins/Experimental/PCG/Source/PCG/Private/Data/PCGUnionData.cpp
+++ b/Engine/Plugins/Experimental/PCG/Source/PCG/Private/Data/PCGUnionData.cpp
@@ -80,12 +80,14 @@ FBox UPCGUnionData::GetStrictBounds() const
 
 bool UPCGUnionData::SamplePoint(const FTransform& InTransform, const FBox& InBounds, FPCGPoint& OutPoint, UPCGMetadata* OutMetadata) const
 {
+	FTransform PointTransform = InTransform;
 	bool bHasSetPoint = false;
 
 	if (FirstNonTrivialTransformData)
 	{
 		if (FirstNonTrivialTransformData->SamplePoint(InTransform, InBounds, OutPoint, OutMetadata))
 		{
+			PointTransform = OutPoint.Transform;
 			bHasSetPoint = true;
 
 			if (DensityFunction == EPCGUnionDensityFunction::Binary && OutPoint.Density > 0)
@@ -105,7 +107,7 @@ bool UPCGUnionData::SamplePoint(const FTransform& InTransform, const FBox& InBou
 		}
 
 		FPCGPoint PointInData;
-		if(Data[DataIndex]->SamplePoint(InTransform, InBounds, PointInData, OutMetadata))
+		if(Data[DataIndex]->SamplePoint(PointTransform, InBounds, PointInData, OutMetadata))
 		{
 			if (!bHasSetPoint)
 			{
