Used different pin visual for multi vs single input pins
Set a tooltip on the pin to indicate what data type is supported
Set a tooltip on the editor nodes to show what PCG node they are bound to + index of the node

diff --git a/Engine/Plugins/Experimental/PCG/Source/PCGEditor/Private/PCGEditorGraphNodeBase.cpp b/Engine/Plugins/Experimental/PCG/Source/PCGEditor/Private/PCGEditorGraphNodeBase.cpp
--- a/Engine/Plugins/Experimental/PCG/Source/PCGEditor/Private/PCGEditorGraphNodeBase.cpp
+++ b/Engine/Plugins/Experimental/PCG/Source/PCGEditor/Private/PCGEditorGraphNodeBase.cpp
@@ -337,4 +337,73 @@ FEdGraphPinType UPCGEditorGraphNodeBase::GetPinType(const UPCGPin* InPin)
 	return EdPinType;
 }
 
+FText UPCGEditorGraphNodeBase::GetTooltipText() const
+{
+	return FText::Format(LOCTEXT("NodeTooltip", "{0}\n{1} - Node index {2}"),
+		GetNodeTitle(ENodeTitleType::FullTitle),
+		PCGNode ? FText::FromName(PCGNode->GetFName()) :FText(LOCTEXT("InvalidNodeName", "Unbound node")),
+		PCGNode ? FText::AsNumber(PCGNode->GetGraph()->GetNodes().IndexOfByKey(PCGNode)) : FText(LOCTEXT("InvalidNodeIndex", "Invalid index")));
+}
+
+void UPCGEditorGraphNodeBase::GetPinHoverText(const UEdGraphPin& Pin, FString& HoverTextOut) const
+{
+	const bool bIsInputPin = (Pin.Direction == EGPD_Input);
+	UPCGPin* MatchingPin = (PCGNode ? (bIsInputPin ? PCGNode->GetInputPin(Pin.PinName) : PCGNode->GetOutputPin(Pin.PinName)) : nullptr);
+
+	auto PCGDataTypeToText = [](EPCGDataType DataType)
+	{
+		TArray<FText> BitFlags;
+
+		for (uint64 BitIndex = 1; BitIndex < 8 * sizeof(EPCGDataType); ++BitIndex)
+		{
+			const int64 BitValue = static_cast<__underlying_type(EPCGDataType)>(DataType) & (1 << BitIndex);
+			if (BitValue)
+			{
+				BitFlags.Add(StaticEnum<EPCGDataType>()->GetDisplayNameTextByValue(BitValue));
+			}
+		}
+
+		return FText::Join(FText(LOCTEXT("Delimiter", " | ")), BitFlags);
+	};
+
+	auto PinTypeToText = [&PCGDataTypeToText](const FName& Category, UPCGPin* MatchingPin)
+	{
+		if (Category != NAME_None)
+		{
+			return FText::FromName(Category);
+		}
+		else if (MatchingPin && MatchingPin->Properties.AllowedTypes == EPCGDataType::Any)
+		{
+			return FText::FromName(FName("Any"));
+		}
+		else if (MatchingPin)
+		{
+			return PCGDataTypeToText(MatchingPin->Properties.AllowedTypes);
+		}
+		else
+		{
+			return FText(LOCTEXT("Unknown data type", "Unknown data type"));
+		}
+	};
+
+	const FText DataTypeText = PinTypeToText(Pin.PinType.PinCategory, MatchingPin);
+	const FText DataSubtypeText = PinTypeToText(Pin.PinType.PinSubCategory, MatchingPin);
+
+	if (MatchingPin != nullptr && bIsInputPin)
+	{
+		const FText AdditionalInfo = MatchingPin->Properties.bAllowMultipleConnections ? FText(LOCTEXT("SupportsMultiInput", "Supports multiple inputs")) : FText(LOCTEXT("SingleInputOnly", "Supports only one input"));
+
+		HoverTextOut = FText::Format(LOCTEXT("PinHoverToolTipFull", "Type: {0}\nSubtype: {1}\nAdditional information: {2}"),
+			DataTypeText,
+			DataSubtypeText,
+			AdditionalInfo).ToString();
+	}
+	else
+	{
+		HoverTextOut = FText::Format(LOCTEXT("PinHoverToolTipFull", "Type: {0}\nSubtype: {1}"),
+			DataTypeText,
+			DataSubtypeText).ToString();
+	}
+}
+
 #undef LOCTEXT_NAMESPACE
diff --git a/Engine/Plugins/Experimental/PCG/Source/PCGEditor/Private/PCGEditorGraphNodeBase.h b/Engine/Plugins/Experimental/PCG/Source/PCGEditor/Private/PCGEditorGraphNodeBase.h
--- a/Engine/Plugins/Experimental/PCG/Source/PCGEditor/Private/PCGEditorGraphNodeBase.h
+++ b/Engine/Plugins/Experimental/PCG/Source/PCGEditor/Private/PCGEditorGraphNodeBase.h
@@ -43,6 +43,8 @@ public:
 	virtual void ReconstructNode() override;
 	virtual FLinearColor GetNodeTitleColor() const;
 	virtual void PostPasteNode() override;
+	virtual FText GetTooltipText() const override;
+	virtual void GetPinHoverText(const UEdGraphPin& Pin, FString& HoverTextOut) const override;
 	// ~End UEdGraphNode interface
 
 	UPCGNode* GetPCGNode() { return PCGNode; }
diff --git a/Engine/Plugins/Experimental/PCG/Source/PCGEditor/Private/SPCGEditorGraphNode.cpp b/Engine/Plugins/Experimental/PCG/Source/PCGEditor/Private/SPCGEditorGraphNode.cpp
--- a/Engine/Plugins/Experimental/PCG/Source/PCGEditor/Private/SPCGEditorGraphNode.cpp
+++ b/Engine/Plugins/Experimental/PCG/Source/PCGEditor/Private/SPCGEditorGraphNode.cpp
@@ -4,6 +4,9 @@
 
 #include "PCGEditorGraphNodeBase.h"
 #include "PCGNode.h"
+#include "PCGPin.h"
+
+#include "SGraphPin.h"
 
 void SPCGEditorGraphNode::Construct(const FArguments& InArgs, UPCGEditorGraphNodeBase* InNode)
 {
@@ -22,6 +25,7 @@ void SPCGEditorGraphNode::MoveTo(const FVector2D& NewPosition, FNodeSet& NodeFil
 {
 	SGraphNode::MoveTo(NewPosition, NodeFilter, bMarkDirty);
 
+	check(PCGEditorGraphNode);
 	if (UPCGNode* PCGNode = PCGEditorGraphNode->GetPCGNode())
 	{
 		PCGNode->Modify();
@@ -30,7 +34,26 @@ void SPCGEditorGraphNode::MoveTo(const FVector2D& NewPosition, FNodeSet& NodeFil
 	}
 }
 
+void SPCGEditorGraphNode::AddPin(const TSharedRef<SGraphPin>& PinToAdd)
+{
+	check(PCGEditorGraphNode);
+	UPCGNode* PCGNode = PCGEditorGraphNode ? PCGEditorGraphNode->GetPCGNode() : nullptr;
+	// Implementation note: we do not distinguish single/multiple pins on the output since that is not relevant
+	if (PCGNode && PinToAdd->GetPinObj())
+	{
+		if (UPCGPin* Pin = PCGNode->GetInputPin(PinToAdd->GetPinObj()->PinName))
+		{
+			if (Pin->Properties.bAllowMultipleConnections)
+			{
+				PinToAdd->SetCustomPinIcon(FAppStyle::GetBrush("Graph.ArrayPin.Connected"), FAppStyle::GetBrush("Graph.ArrayPin.Disconnected"));
+			}
+		}
+	}
+
+	SGraphNode::AddPin(PinToAdd);
+}
+
 void SPCGEditorGraphNode::OnNodeChanged()
 {
 	UpdateGraphNode();
-}
+}
\ No newline at end of file
diff --git a/Engine/Plugins/Experimental/PCG/Source/PCGEditor/Private/SPCGEditorGraphNode.h b/Engine/Plugins/Experimental/PCG/Source/PCGEditor/Private/SPCGEditorGraphNode.h
--- a/Engine/Plugins/Experimental/PCG/Source/PCGEditor/Private/SPCGEditorGraphNode.h
+++ b/Engine/Plugins/Experimental/PCG/Source/PCGEditor/Private/SPCGEditorGraphNode.h
@@ -18,6 +18,7 @@ public:
 	//~ Begin SGraphNode Interface
 	virtual void MoveTo(const FVector2D& NewPosition, FNodeSet& NodeFilter, bool bMarkDirty = true) override;
 	virtual void RequestRenameOnSpawn() override { /* Empty to avoid the default behavior to rename on node spawn */ }
+	virtual void AddPin(const TSharedRef<SGraphPin>& PinToAdd) override;
 	//~ End SGraphNode Interface
 
 protected:
