PCG:
*Added find support to PCG graph editor.

diff --git a/Engine/Plugins/Experimental/PCG/Source/PCGEditor/Private/PCGEditor.cpp b/Engine/Plugins/Experimental/PCG/Source/PCGEditor/Private/PCGEditor.cpp
--- a/Engine/Plugins/Experimental/PCG/Source/PCGEditor/Private/PCGEditor.cpp
+++ b/Engine/Plugins/Experimental/PCG/Source/PCGEditor/Private/PCGEditor.cpp
@@ -2,11 +2,13 @@
 
 #include "PCGEditor.h"
 
+#include "PCGEditorCommands.h"
 #include "PCGEditorCommon.h"
 #include "PCGEditorGraph.h"
 #include "PCGEditorGraphNodeBase.h"
 #include "PCGEditorGraphSchema.h"
 #include "PCGGraph.h"
+#include "SPCGEditorGraphFind.h"
 #include "SPCGEditorGraphNodePalette.h"
 
 #include "EdGraphUtilities.h"
@@ -36,6 +38,7 @@ namespace FPCGEditor_private
 	const FName PaletteID = FName(TEXT("Palette"));
 	const FName AttributesID = FName(TEXT("Attributes"));
 	const FName ViewportID = FName(TEXT("Viewport"));
+	const FName FindID = FName(TEXT("Find"));
 }
 
 void FPCGEditor::Initialize(const EToolkitMode::Type InMode, const TSharedPtr<class IToolkitHost>& InToolkitHost, UPCGGraph* InPCGGraph)
@@ -56,6 +59,9 @@ void FPCGEditor::Initialize(const EToolkitMode::Type InMode, const TSharedPtr<cl
 
 	GraphEditorWidget = CreateGraphEditorWidget();
 	PaletteWidget = CreatePaletteWidget();
+	FindWidget = CreateFindWidget();
+
+	BindCommands();
 
 	const TSharedRef<FTabManager::FLayout> StandaloneDefaultLayout = FTabManager::NewLayout("Standalone_PCGGraphEditor_Layout_v0.4")
 		->AddArea
@@ -98,6 +104,7 @@ void FPCGEditor::Initialize(const EToolkitMode::Type InMode, const TSharedPtr<cl
 					->SetSizeCoefficient(0.28)
 					->SetHideTabWell(true)
 					->AddTab(FPCGEditor_private::AttributesID, ETabState::OpenedTab)
+					->AddTab(FPCGEditor_private::FindID, ETabState::ClosedTab)
 				)
 			)
 			->Split
@@ -114,7 +121,20 @@ void FPCGEditor::Initialize(const EToolkitMode::Type InMode, const TSharedPtr<cl
 	InitAssetEditor(InMode, InToolkitHost, PCGGraphEditorAppName, StandaloneDefaultLayout, /*bCreateDefaultStandaloneMenu=*/ true, /*bCreateDefaultToolbar=*/ true, InPCGGraph);
 }
 
-void FPCGEditor::RegisterTabSpawners(const TSharedRef<class FTabManager>& InTabManager)
+UPCGEditorGraph* FPCGEditor::GetPCGEditorGraph()
+{
+	return PCGEditorGraph;
+}
+
+void FPCGEditor::JumpToNode(const UEdGraphNode* InNode)
+{
+	if (GraphEditorWidget.IsValid())
+	{
+		GraphEditorWidget->JumpToNode(InNode);
+	}
+}
+
+void FPCGEditor::RegisterTabSpawners(const TSharedRef<FTabManager>& InTabManager)
 {
 	WorkspaceMenuCategory = InTabManager->AddLocalWorkspaceMenuCategory(LOCTEXT("WorkspaceMenu_PCGEditor", "PCG Editor"));
 	auto WorkspaceMenuCategoryRef = WorkspaceMenuCategory.ToSharedRef();
@@ -141,6 +161,10 @@ void FPCGEditor::RegisterTabSpawners(const TSharedRef<class FTabManager>& InTabM
 	InTabManager->RegisterTabSpawner(FPCGEditor_private::ViewportID, FOnSpawnTab::CreateSP(this, &FPCGEditor::SpawnTab_Viewport))
 		.SetDisplayName(LOCTEXT("ViewportTab", "Viewport"))
 		.SetGroup(WorkspaceMenuCategoryRef);
+
+	InTabManager->RegisterTabSpawner(FPCGEditor_private::FindID, FOnSpawnTab::CreateSP(this, &FPCGEditor::SpawnTab_Find))
+		.SetDisplayName(LOCTEXT("FindTab", "Find"))
+		.SetGroup(WorkspaceMenuCategoryRef);
 }
 
 void FPCGEditor::UnregisterTabSpawners(const TSharedRef<class FTabManager>& InTabManager)
@@ -198,6 +222,24 @@ FString FPCGEditor::GetWorldCentricTabPrefix() const
 	return LOCTEXT("WorldCentricTabPrefix", "PCG ").ToString();
 }
 
+void FPCGEditor::BindCommands()
+{
+	const FPCGEditorCommands& PCGEditorCommands = FPCGEditorCommands::Get();
+
+	ToolkitCommands->MapAction(
+		PCGEditorCommands.Find,
+		FExecuteAction::CreateSP(this, &FPCGEditor::OnFind));
+}
+
+void FPCGEditor::OnFind()
+{
+	if (TabManager.IsValid() && FindWidget.IsValid())
+	{
+		TabManager->TryInvokeTab(FPCGEditor_private::FindID);
+		FindWidget->FocusForUse();
+	}
+}
+
 void FPCGEditor::SelectAllNodes()
 {
 	if (GraphEditorWidget.IsValid())
@@ -578,6 +620,11 @@ TSharedRef<SPCGEditorGraphNodePalette> FPCGEditor::CreatePaletteWidget()
 	return SNew(SPCGEditorGraphNodePalette);
 }
 
+TSharedRef<SPCGEditorGraphFind> FPCGEditor::CreateFindWidget()
+{
+	return SNew(SPCGEditorGraphFind, SharedThis(this));
+}
+
 void FPCGEditor::OnSelectedNodesChanged(const TSet<UObject*>& NewSelection)
 {
 	TArray<TWeakObjectPtr<UObject>> SelectedObjects;
@@ -665,4 +712,14 @@ TSharedRef<SDockTab> FPCGEditor::SpawnTab_Viewport(const FSpawnTabArgs& Args)
 		];
 }
 
+TSharedRef<SDockTab> FPCGEditor::SpawnTab_Find(const FSpawnTabArgs& Args)
+{
+	return SNew(SDockTab)
+		.Label(LOCTEXT("PCGFindTitle", "Find"))
+		.TabColorScale(GetTabColorScale())
+		[
+			FindWidget.ToSharedRef()
+		];
+}
+
 #undef LOCTEXT_NAMESPACE
diff --git a/Engine/Plugins/Experimental/PCG/Source/PCGEditor/Private/PCGEditor.h b/Engine/Plugins/Experimental/PCG/Source/PCGEditor/Private/PCGEditor.h
--- a/Engine/Plugins/Experimental/PCG/Source/PCGEditor/Private/PCGEditor.h
+++ b/Engine/Plugins/Experimental/PCG/Source/PCGEditor/Private/PCGEditor.h
@@ -8,6 +8,7 @@
 class FUICommandList;
 class IDetailsView;
 class SGraphEditor;
+class SPCGEditorGraphFind;
 class SPCGEditorGraphNodePalette;
 class UPCGEditorGraph;
 class UPCGGraph;
@@ -18,6 +19,12 @@ public:
 	/** Edits the specified PCGGraph */
 	void Initialize(const EToolkitMode::Type InMode, const TSharedPtr<class IToolkitHost>& InToolkitHost, UPCGGraph* InPCGGraph);
 
+	/** Get the PCG graph being edited */
+	UPCGEditorGraph* GetPCGEditorGraph();
+
+	/** Focus the graph view on a specific node */
+	void JumpToNode(const UEdGraphNode* InNode);
+
 	// ~Begin IToolkit interface
 	virtual void RegisterTabSpawners(const TSharedRef<class FTabManager>& InTabManager) override;
 	virtual void UnregisterTabSpawners(const TSharedRef<class FTabManager>& InTabManager) override;
@@ -37,6 +44,12 @@ public:
 	// ~End FAssetEditorToolkit interface
 
 private:
+	/** Bind commands to delegates */
+	void BindCommands();
+
+	/** Bring up the find tab */
+	void OnFind();
+
 	/** Select every node in the graph */
 	void SelectAllNodes();
 	/** Whether we can select every node */
@@ -85,6 +98,9 @@ private:
 	/** Create new palette widget */
 	TSharedRef<SPCGEditorGraphNodePalette> CreatePaletteWidget();
 
+	/** Create new find widget */
+	TSharedRef<SPCGEditorGraphFind> CreateFindWidget();
+
 	/** Called when the selection changes in the GraphEditor */
 	void OnSelectedNodesChanged(const TSet<UObject*>& NewSelection);
 
@@ -96,10 +112,12 @@ private:
 	TSharedRef<SDockTab> SpawnTab_Palette(const FSpawnTabArgs& Args);
 	TSharedRef<SDockTab> SpawnTab_Attributes(const FSpawnTabArgs& Args);
 	TSharedRef<SDockTab> SpawnTab_Viewport(const FSpawnTabArgs& Args);
+	TSharedRef<SDockTab> SpawnTab_Find(const FSpawnTabArgs& Args);
 
 	TSharedPtr<SGraphEditor> GraphEditorWidget;
 	TSharedPtr<IDetailsView> PropertyDetailsWidget;
 	TSharedPtr<SPCGEditorGraphNodePalette> PaletteWidget;
+	TSharedPtr<SPCGEditorGraphFind> FindWidget;
 
 	TSharedPtr<FUICommandList> GraphEditorCommands;
 
diff --git a/Engine/Plugins/Experimental/PCG/Source/PCGEditor/Private/PCGEditorCommands.cpp b/Engine/Plugins/Experimental/PCG/Source/PCGEditor/Private/PCGEditorCommands.cpp
new file mode 100644
--- /dev/null
+++ b/Engine/Plugins/Experimental/PCG/Source/PCGEditor/Private/PCGEditorCommands.cpp
@@ -0,0 +1,24 @@
+// Copyright Epic Games, Inc. All Rights Reserved.
+
+#include "PCGEditorCommands.h"
+
+#include "Framework/Commands/InputChord.h"
+#include "Styling/AppStyle.h"
+
+#define LOCTEXT_NAMESPACE "PCGEditorCommands"
+
+FPCGEditorCommands::FPCGEditorCommands()
+	: TCommands<FPCGEditorCommands>(
+		"PCGEditor",
+		NSLOCTEXT("Contexts", "PCGEditor", "PCG Editor"),
+		NAME_None,
+		FAppStyle::GetAppStyleSetName())
+{
+}
+
+void FPCGEditorCommands::RegisterCommands()
+{
+	UI_COMMAND(Find, "Find", "Finds PCG nodes and comments in the current graph", EUserInterfaceActionType::Button, FInputChord(EModifierKey::Control, EKeys::F));
+}
+
+#undef LOCTEXT_NAMESPACE
diff --git a/Engine/Plugins/Experimental/PCG/Source/PCGEditor/Private/PCGEditorCommands.h b/Engine/Plugins/Experimental/PCG/Source/PCGEditor/Private/PCGEditorCommands.h
new file mode 100644
--- /dev/null
+++ b/Engine/Plugins/Experimental/PCG/Source/PCGEditor/Private/PCGEditorCommands.h
@@ -0,0 +1,18 @@
+// Copyright Epic Games, Inc. All Rights Reserved.
+
+#pragma once
+
+#include "CoreMinimal.h"
+#include "Framework/Commands/Commands.h"
+
+class FPCGEditorCommands : public TCommands<FPCGEditorCommands>
+{
+public:
+	FPCGEditorCommands();
+
+	// ~Begin TCommands<> interface
+	virtual void RegisterCommands() override;
+	// ~End TCommands<> interface
+
+	TSharedPtr<FUICommandInfo> Find;
+};
diff --git a/Engine/Plugins/Experimental/PCG/Source/PCGEditor/Private/PCGEditorModule.cpp b/Engine/Plugins/Experimental/PCG/Source/PCGEditor/Private/PCGEditorModule.cpp
--- a/Engine/Plugins/Experimental/PCG/Source/PCGEditor/Private/PCGEditorModule.cpp
+++ b/Engine/Plugins/Experimental/PCG/Source/PCGEditor/Private/PCGEditorModule.cpp
@@ -2,24 +2,24 @@
 
 #include "PCGEditorModule.h"
 
-#include "EdGraphUtilities.h"
-#include "Framework/MultiBox/MultiBoxBuilder.h"
-#include "IAssetTools.h"
-#include "LevelEditor.h"
-#include "Modules/ModuleManager.h"
-#include "PropertyEditorModule.h"
-#include "Toolkits/IToolkit.h"
-#include "ISettingsModule.h"
-
 #include "AssetTypeActions/PCGGraphAssetTypeActions.h"
 #include "AssetTypeActions/PCGSettingsAssetTypeActions.h"
 #include "PCGComponentDetails.h"
+#include "PCGEditorCommands.h"
 #include "PCGEditorGraphNodeFactory.h"
+#include "PCGEditorSettings.h"
 #include "PCGGraphDetails.h"
+#include "PCGSubsystem.h"
 #include "PCGVolumeFactory.h"
-#include "PCGEditorSettings.h"
 
-#include "PCGSubsystem.h"
+#include "EdGraphUtilities.h"
+#include "Framework/MultiBox/MultiBoxBuilder.h"
+#include "IAssetTools.h"
+#include "ISettingsModule.h"
+#include "LevelEditor.h"
+#include "Modules/ModuleManager.h"
+#include "PropertyEditorModule.h"
+#include "Toolkits/IToolkit.h"
 
 #define LOCTEXT_NAMESPACE "FPCGEditorModule"
 
@@ -32,6 +32,8 @@ void FPCGEditorModule::StartupModule()
 	RegisterMenuExtensions();
 	RegisterSettings();
 
+	FPCGEditorCommands::Register();
+
 	GraphNodeFactory = MakeShareable(new FPCGEditorGraphNodeFactory());
 	FEdGraphUtilities::RegisterVisualNodeFactory(GraphNodeFactory);
 
@@ -48,6 +50,8 @@ void FPCGEditorModule::ShutdownModule()
 	UnregisterDetailsCustomizations();
 	UnregisterMenuExtensions();
 
+	FPCGEditorCommands::Unregister();
+
 	FEdGraphUtilities::UnregisterVisualNodeFactory(GraphNodeFactory);
 
 	if (GEditor)
diff --git a/Engine/Plugins/Experimental/PCG/Source/PCGEditor/Private/SPCGEditorGraphFind.cpp b/Engine/Plugins/Experimental/PCG/Source/PCGEditor/Private/SPCGEditorGraphFind.cpp
new file mode 100644
--- /dev/null
+++ b/Engine/Plugins/Experimental/PCG/Source/PCGEditor/Private/SPCGEditorGraphFind.cpp
@@ -0,0 +1,317 @@
+// Copyright Epic Games, Inc. All Rights Reserved.
+
+#include "SPCGEditorGraphFind.h"
+
+#include "PCGEditor.h"
+#include "PCGEditorGraph.h"
+#include "PCGEditorGraphNodeBase.h"
+
+#include "EdGraph/EdGraph.h"
+#include "EdGraph/EdGraphNode.h"
+#include "EdGraph/EdGraphSchema.h"
+#include "Widgets/SWidget.h"
+
+#define LOCTEXT_NAMESPACE "PCGEditorGraphFind"
+
+FPCGEditorGraphFindResult::FPCGEditorGraphFindResult(const FString& InValue)
+: Parent(nullptr)
+, Value(InValue)
+, GraphNode(nullptr)
+, Pin()
+{
+}
+
+FPCGEditorGraphFindResult::FPCGEditorGraphFindResult(const FText& InValue)
+: FPCGEditorGraphFindResult(InValue.ToString())
+{
+}
+
+FPCGEditorGraphFindResult::FPCGEditorGraphFindResult(const FString& InValue, TSharedPtr<FPCGEditorGraphFindResult>& InParent, UEdGraphNode* InNode)
+: Parent(InParent)
+, Value(InValue)
+, GraphNode(InNode)
+, Pin()
+{
+}
+
+FPCGEditorGraphFindResult::FPCGEditorGraphFindResult(const FString& InValue, TSharedPtr<FPCGEditorGraphFindResult>& InParent, UEdGraphPin* InPin)
+: Parent(InParent)
+, Value(InValue)
+, GraphNode(nullptr)
+, Pin(InPin)
+{
+}
+
+FReply FPCGEditorGraphFindResult::OnClick(TWeakPtr<FPCGEditor> InPCGEditorPtr)
+{
+	if (UEdGraphPin* ResolvedPin = Pin.Get())
+	{
+		InPCGEditorPtr.Pin()->JumpToNode(ResolvedPin->GetOwningNode());
+	}
+	else if (GraphNode.IsValid())
+	{
+		InPCGEditorPtr.Pin()->JumpToNode(GraphNode.Get());
+	}
+
+	return FReply::Handled();
+}
+
+FText FPCGEditorGraphFindResult::GetToolTip() const
+{
+	FText ToolTip;
+
+	if (UEdGraphPin* ResolvedPin = Pin.Get())
+	{
+		if (UEdGraphNode* OwningNode = ResolvedPin->GetOwningNode())
+		{
+			FString ToolTipString;
+			OwningNode->GetPinHoverText(*ResolvedPin, ToolTipString);
+			ToolTip = FText::FromString(ToolTipString);
+		}
+	}
+	else if (GraphNode.IsValid())
+	{
+		ToolTip = GraphNode->GetTooltipText();
+	}
+
+	return ToolTip;
+}
+
+FText FPCGEditorGraphFindResult::GetCategory() const
+{
+	if (Pin.Get())
+	{
+		return LOCTEXT("PinCategory", "Pin");
+	}
+	else if (GraphNode.IsValid())
+	{
+		return LOCTEXT("NodeCategory", "Node");
+	}
+	return FText::GetEmpty();
+}
+
+TSharedRef<SWidget> FPCGEditorGraphFindResult::CreateIcon() const
+{
+	FSlateColor IconColor = FSlateColor::UseForeground();
+	const FSlateBrush* Brush = nullptr;
+
+	if (UEdGraphPin* ResolvedPin = Pin.Get())
+	{
+		// TODO get pin icon from nodebase?
+		Brush = FAppStyle::GetBrush(TEXT("GraphEditor.PinIcon"));
+		const UEdGraphSchema* Schema = ResolvedPin->GetSchema();
+		IconColor = Schema->GetPinTypeColor(ResolvedPin->PinType);
+	}
+	else if (GraphNode.IsValid())
+	{
+		// TODO get icon and tint from nodebase?
+		Brush = FAppStyle::GetBrush(TEXT("GraphEditor.NodeGlyph"));
+	}
+
+	return SNew(SImage)
+		.Image(Brush)
+		.ColorAndOpacity(IconColor)
+		.ToolTipText(GetCategory());
+}
+
+void SPCGEditorGraphFind::Construct(const FArguments& InArgs, TSharedPtr<FPCGEditor> InPCGEditor)
+{
+	PCGEditorPtr = InPCGEditor;
+
+	this->ChildSlot
+	[
+		SNew(SVerticalBox)
+		+SVerticalBox::Slot()
+		.AutoHeight()
+		[
+			SAssignNew(SearchTextField, SSearchBox)
+			.HintText(LOCTEXT("PCGGraphSearchHint", "Enter text to find nodes..."))
+			.OnTextChanged(this, &SPCGEditorGraphFind::OnSearchTextChanged)
+			.OnTextCommitted(this, &SPCGEditorGraphFind::OnSearchTextCommitted)
+		]
+		+ SVerticalBox::Slot()
+		.FillHeight(1.0f)
+		.Padding(0.f, 4.f, 0.f, 0.f)
+		[
+			SNew(SBorder)
+			.BorderImage(FAppStyle::GetBrush("Menu.Background"))
+			[
+				SAssignNew(TreeView, STreeViewType)
+				.ItemHeight(24)
+				.TreeItemsSource(&ItemsFound)
+				.OnGenerateRow(this, &SPCGEditorGraphFind::OnGenerateRow)
+				.OnGetChildren(this, &SPCGEditorGraphFind::OnGetChildren)
+				.OnSelectionChanged(this, &SPCGEditorGraphFind::OnTreeSelectionChanged)
+				.SelectionMode(ESelectionMode::Single)
+			]
+		]
+	];
+}
+
+void SPCGEditorGraphFind::FocusForUse()
+{
+	// NOTE: Careful, GeneratePathToWidget can be reentrant in that it can call visibility delegates and such
+	FWidgetPath FilterTextBoxWidgetPath;
+	FSlateApplication::Get().GeneratePathToWidgetUnchecked(SearchTextField.ToSharedRef(), FilterTextBoxWidgetPath);
+
+	// Set keyboard focus directly
+	FSlateApplication::Get().SetKeyboardFocus(FilterTextBoxWidgetPath, EFocusCause::SetDirectly);
+}
+
+void SPCGEditorGraphFind::OnSearchTextChanged(const FText& InText)
+{
+	SearchValue = InText.ToString();
+	InitiateSearch();
+}
+
+void SPCGEditorGraphFind::OnSearchTextCommitted(const FText& /*InText*/, ETextCommit::Type /*InCommitType*/)
+{
+	InitiateSearch();
+}
+
+void SPCGEditorGraphFind::OnGetChildren(FPCGEditorGraphFindResultPtr InItem, TArray<FPCGEditorGraphFindResultPtr>& OutChildren)
+{
+	OutChildren += InItem->Children;
+}
+
+void SPCGEditorGraphFind::OnTreeSelectionChanged(FPCGEditorGraphFindResultPtr Item, ESelectInfo::Type)
+{
+	if (Item.IsValid())
+	{
+		Item->OnClick(PCGEditorPtr);
+	}
+}
+
+TSharedRef<ITableRow> SPCGEditorGraphFind::OnGenerateRow(FPCGEditorGraphFindResultPtr InItem, const TSharedRef<STableViewBase>& InOwnerTable)
+{
+	return SNew(STableRow<TSharedPtr<FPCGEditorGraphFindResult>>, InOwnerTable)
+		[
+			SNew(SHorizontalBox)
+			+ SHorizontalBox::Slot()
+			.VAlign(VAlign_Center)
+			.AutoWidth()
+			[
+				SNew(SBox)
+				.WidthOverride(450.0f)
+				[
+					SNew(SHorizontalBox)
+					+ SHorizontalBox::Slot()
+					.AutoWidth()
+					[
+						InItem->CreateIcon()
+					]
+					+SHorizontalBox::Slot()
+					.VAlign(VAlign_Center)
+					.AutoWidth()
+					.Padding(2, 0)
+					[
+						SNew(STextBlock)
+						.Text(FText::FromString(InItem->Value))
+						.HighlightText(HighlightText)
+						.ToolTipText(InItem->GetToolTip())
+					]
+				]
+			]
+		];
+}
+
+void SPCGEditorGraphFind::InitiateSearch()
+{
+	TArray<FString> Tokens;
+	SearchValue.ParseIntoArray(Tokens, TEXT(" "), true);
+
+	ItemsFound.Empty();
+	if (Tokens.Num() > 0)
+	{
+		HighlightText = FText::FromString(SearchValue);
+		MatchTokens(Tokens);
+	}
+
+	// Insert a fake result to inform user if none found
+	if (ItemsFound.Num() == 0)
+	{
+		ItemsFound.Add(MakeShared<FPCGEditorGraphFindResult>(LOCTEXT("PCGGraphSearchNoResults", "No Results found")));
+	}
+
+	TreeView->RequestTreeRefresh();
+
+	for (const FPCGEditorGraphFindResultPtr& Item : ItemsFound)
+	{
+		TreeView->SetItemExpansion(Item, true);
+	}
+}
+
+void SPCGEditorGraphFind::MatchTokens(const TArray<FString>& InTokens)
+{
+	RootFindResult.Reset();
+
+	UPCGEditorGraph* PCGEditorGraph = PCGEditorPtr.Pin()->GetPCGEditorGraph();
+	if (!PCGEditorGraph)
+	{
+		return;
+	}
+
+	RootFindResult = MakeShared<FPCGEditorGraphFindResult>(FString("PCGTreeRoot"));
+
+	for (UEdGraphNode* Node : PCGEditorGraph->Nodes)
+	{
+		const FString NodeName = Node->GetNodeTitle(ENodeTitleType::FullTitle).ToString();
+		const FString NodeType = Node->GetNodeTitle(ENodeTitleType::ListView).ToString();
+
+		FString NodeSearchString = NodeName + NodeType + Node->NodeComment;
+		NodeSearchString = NodeSearchString.Replace(TEXT(" "), TEXT(""));
+
+		FPCGEditorGraphFindResultPtr NodeResult;
+		auto GetOrCreateNodeResult = [&]() -> FPCGEditorGraphFindResultPtr&
+		{
+			if (!NodeResult.IsValid())
+			{
+				NodeResult = MakeShared<FPCGEditorGraphFindResult>(NodeName == NodeType ? NodeName : NodeName + " - " + NodeType, RootFindResult, Node);
+			}
+			return NodeResult;
+		};
+
+		if (StringMatchesSearchTokens(InTokens, NodeSearchString))
+		{
+			GetOrCreateNodeResult();
+		}
+	
+		for (UEdGraphPin* Pin : Node->Pins)
+		{
+			if (Pin && Pin->PinFriendlyName.CompareTo(FText::FromString(TEXT(" "))) != 0)
+			{
+				const FText PinName = Pin->GetSchema()->GetPinDisplayName(Pin);
+				FString PinSearchString = Pin->PinName.ToString() + Pin->PinFriendlyName.ToString() + Pin->DefaultValue + Pin->PinType.PinCategory.ToString() + Pin->PinType.PinSubCategory.ToString() + (Pin->PinType.PinSubCategoryObject.IsValid() ? Pin->PinType.PinSubCategoryObject.Get()->GetFullName() : TEXT(""));
+				PinSearchString = PinSearchString.Replace(TEXT(" "), TEXT(""));
+				if (StringMatchesSearchTokens(InTokens, PinSearchString))
+				{
+					FPCGEditorGraphFindResultPtr PinResult(MakeShared<FPCGEditorGraphFindResult>(PinName.ToString(), GetOrCreateNodeResult(), Pin));
+					NodeResult->Children.Add(PinResult);
+				}
+			}
+		}
+
+		if (NodeResult.IsValid())
+		{
+			ItemsFound.Add(NodeResult);
+		}
+	}
+}
+
+bool SPCGEditorGraphFind::StringMatchesSearchTokens(const TArray<FString>& InTokens, const FString& InComparisonString)
+{
+	bool bFoundAllTokens = true;
+
+	//search the entry for each token, it must have all of them to pass
+	for (const FString& Token : InTokens)
+	{
+		if (!InComparisonString.Contains(Token))
+		{
+			bFoundAllTokens = false;
+			break;
+		}
+	}
+	return bFoundAllTokens;
+}
+
+#undef LOCTEXT_NAMESPACE
diff --git a/Engine/Plugins/Experimental/PCG/Source/PCGEditor/Private/SPCGEditorGraphFind.h b/Engine/Plugins/Experimental/PCG/Source/PCGEditor/Private/SPCGEditorGraphFind.h
new file mode 100644
--- /dev/null
+++ b/Engine/Plugins/Experimental/PCG/Source/PCGEditor/Private/SPCGEditorGraphFind.h
@@ -0,0 +1,118 @@
+// Copyright Epic Games, Inc. All Rights Reserved.
+
+#pragma once
+
+#include "CoreMinimal.h"
+#include "EdGraph/EdGraphPin.h"
+#include "Widgets/SCompoundWidget.h"
+#include "Widgets/Views/STableViewBase.h"
+#include "Widgets/Views/STableRow.h"
+#include "Widgets/Views/STreeView.h"
+#include "Widgets/Input/SSearchBox.h"
+
+class FPCGEditor;
+class UEdGraphNode;
+
+struct FPCGEditorGraphFindResult
+{
+	/** Create a root (or only text) result */
+	FPCGEditorGraphFindResult(const FString& InValue);
+
+	/** Create a root (or only text) result */
+	FPCGEditorGraphFindResult(const FText& InValue);
+
+	/** Create a listing for a node result */
+	FPCGEditorGraphFindResult(const FString& InValue, TSharedPtr<FPCGEditorGraphFindResult>& InParent, UEdGraphNode* InNode);
+
+	/** Create a listing for a pin result */
+	FPCGEditorGraphFindResult(const FString& InValue, TSharedPtr<FPCGEditorGraphFindResult>& InParent, UEdGraphPin* InPin);
+
+	/** Called when user clicks on the search item */
+	FReply OnClick(TWeakPtr<FPCGEditor> InPCGEditorPtr);
+
+	/** Get ToolTip for this search result */
+	FText GetToolTip() const;
+
+	/** Get Category for this search result */
+	FText GetCategory() const;
+
+	/** Create an icon to represent the result */
+	TSharedRef<SWidget> CreateIcon() const;
+
+	/** Search result parent */
+	TWeakPtr<FPCGEditorGraphFindResult> Parent;
+
+	/** Any children listed under this category */
+	TArray<TSharedPtr<FPCGEditorGraphFindResult>> Children;
+	
+	/** The string value for this result */
+	FString Value;
+
+	/** The graph node that this search result refers to */
+	TWeakObjectPtr<UEdGraphNode> GraphNode;
+
+	/** The pin that this search result refers to */
+	FEdGraphPinReference Pin;
+};
+
+
+class SPCGEditorGraphFind : public SCompoundWidget
+{
+public:
+	SLATE_BEGIN_ARGS(SPCGEditorGraphFind) {}
+	SLATE_END_ARGS()
+
+	void Construct(const FArguments& InArgs, TSharedPtr<FPCGEditor> InPCGEditor);
+
+	/** Focuses this widget's search box */
+	void FocusForUse();
+
+private:
+	typedef TSharedPtr<FPCGEditorGraphFindResult> FPCGEditorGraphFindResultPtr;
+	typedef STreeView<FPCGEditorGraphFindResultPtr> STreeViewType;
+
+	/** Called when user changes the text they are searching for */
+	void OnSearchTextChanged(const FText& InText);
+
+	/** Called when user changes commits text to the search box */
+	void OnSearchTextCommitted(const FText& InText, ETextCommit::Type InCommitType);
+
+	/** Get the children of a row */
+	void OnGetChildren(FPCGEditorGraphFindResultPtr InItem, TArray<FPCGEditorGraphFindResultPtr>& OutChildren);
+
+	/** Called when user clicks on a new result */
+	void OnTreeSelectionChanged(FPCGEditorGraphFindResultPtr Item, ESelectInfo::Type);
+
+	/** Called when a new row is being generated */
+	TSharedRef<ITableRow> OnGenerateRow(FPCGEditorGraphFindResultPtr InItem, const TSharedRef<STableViewBase>& InOwnerTable);
+
+	/** Begins the search based on the SearchValue */
+	void InitiateSearch();
+
+	/** Find any results that contain all of the tokens */
+	void MatchTokens(const TArray<FString>& InTokens);
+
+	/** Determines if a string matches the search tokens */
+	static bool StringMatchesSearchTokens(const TArray<FString>& InTokens, const FString& InComparisonString);
+
+	/** Pointer back to the PCG editor that owns us */
+	TWeakPtr<FPCGEditor> PCGEditorPtr;
+
+	/** The tree view displays the results */
+	TSharedPtr<STreeViewType> TreeView;
+
+	/** The search text box */
+	TSharedPtr<SSearchBox> SearchTextField;
+
+	/** This buffer stores the currently displayed results */
+	TArray<FPCGEditorGraphFindResultPtr> ItemsFound;
+
+	/** we need to keep a handle on the root result, because it won't show up in the tree */
+	FPCGEditorGraphFindResultPtr RootFindResult;
+
+	/** The string to highlight in the results */
+	FText HighlightText;
+
+	/** The string to search for */
+	FString SearchValue;
+};
