Fixed issue in the remove duplicate test with non-centered bounds

#rnx

diff --git a/Engine/Plugins/Experimental/PCG/Source/PCG/Private/Elements/PCGSelfPruning.cpp b/Engine/Plugins/Experimental/PCG/Source/PCG/Private/Elements/PCGSelfPruning.cpp
--- a/Engine/Plugins/Experimental/PCG/Source/PCG/Private/Elements/PCGSelfPruning.cpp
+++ b/Engine/Plugins/Experimental/PCG/Source/PCG/Private/Elements/PCGSelfPruning.cpp
@@ -80,7 +80,7 @@ namespace PCGSelfPruningElement
 
 			ExclusionPoints.Add(Point);
 
-			Octree.FindElementsWithBoundsTest(FBoxCenterAndExtent(Point->Transform.GetLocation(), FVector::Zero()), [&ExclusionPoints, &ExcludedPoints, Point](const FPCGPointRef& InPointRef)
+			Octree.FindElementsWithBoundsTest(FBoxCenterAndExtent(Point->Transform.TransformPosition(Point->GetLocalCenter()), FVector::Zero()), [&ExclusionPoints, &ExcludedPoints, Point](const FPCGPointRef& InPointRef)
 			{
 				if ((Point->Transform.GetLocation() - InPointRef.Point->Transform.GetLocation()).SquaredLength() <= SMALL_NUMBER &&
 					!ExclusionPoints.Contains(InPointRef.Point))
