Basic PCG unit/determinism framework - added basic test for a single data type and a single element
Refactored voxel sampling to be used by the volumes too
Added volume sampler element/node

diff --git a/Engine/Plugins/Experimental/PCG/Source/PCG/Private/Data/PCGIntersectionData.cpp b/Engine/Plugins/Experimental/PCG/Source/PCG/Private/Data/PCGIntersectionData.cpp
--- a/Engine/Plugins/Experimental/PCG/Source/PCG/Private/Data/PCGIntersectionData.cpp
+++ b/Engine/Plugins/Experimental/PCG/Source/PCG/Private/Data/PCGIntersectionData.cpp
@@ -28,8 +28,8 @@ void UPCGIntersectionData::Initialize(const UPCGSpatialData* InA, const UPCGSpat
 	B = InB;
 	TargetActor = A->TargetActor;
 
-	CachedBounds = A->GetBounds().Overlap(B->GetBounds());
-	CachedStrictBounds = A->GetStrictBounds().Overlap(B->GetStrictBounds());
+	CachedBounds = PCGHelpers::OverlapBounds(A->GetBounds(), B->GetBounds());
+	CachedStrictBounds = PCGHelpers::OverlapBounds(A->GetStrictBounds(), B->GetStrictBounds());
 }
 
 int UPCGIntersectionData::GetDimension() const
diff --git a/Engine/Plugins/Experimental/PCG/Source/PCG/Private/Data/PCGPrimitiveData.cpp b/Engine/Plugins/Experimental/PCG/Source/PCG/Private/Data/PCGPrimitiveData.cpp
--- a/Engine/Plugins/Experimental/PCG/Source/PCG/Private/Data/PCGPrimitiveData.cpp
+++ b/Engine/Plugins/Experimental/PCG/Source/PCG/Private/Data/PCGPrimitiveData.cpp
@@ -3,9 +3,7 @@
 #include "Data/PCGPrimitiveData.h"
 
 #include "Data/PCGPointData.h"
-#include "Helpers/PCGAsync.h"
-#include "PCGHelpers.h"
-
+#include "Elements/PCGVolumeSampler.h"
 #include "Components/PrimitiveComponent.h"
 
 void UPCGPrimitiveData::Initialize(UPrimitiveComponent* InPrimitive)
@@ -41,39 +39,15 @@ const UPCGPointData* UPCGPrimitiveData::CreatePointData(FPCGContext* Context) co
 {
 	TRACE_CPUPROFILER_EVENT_SCOPE(UPCGPrimitiveData::CreatePointData);
 
-	UPCGPointData* Data = NewObject<UPCGPointData>(const_cast<UPCGPrimitiveData*>(this));
-	Data->InitializeFromData(this);
-	TArray<FPCGPoint>& Points = Data->GetMutablePoints();
-
-	const int32 MinX = FMath::CeilToInt(CachedBounds.Min.X / VoxelSize.X);
-	const int32 MaxX = FMath::FloorToInt(CachedBounds.Max.X / VoxelSize.X);
-	const int32 MinY = FMath::CeilToInt(CachedBounds.Min.Y / VoxelSize.Y);
-	const int32 MaxY = FMath::FloorToInt(CachedBounds.Max.Y / VoxelSize.Y);
-	const int32 MinZ = FMath::CeilToInt(CachedBounds.Min.Z / VoxelSize.Z);
-	const int32 MaxZ = FMath::FloorToInt(CachedBounds.Max.Z / VoxelSize.Z);
+	PCGVolumeSampler::FVolumeSamplerSettings SamplerSettings;
+	SamplerSettings.VoxelSize = VoxelSize;
 
-	const int32 NumIterations = (MaxX - MinX) * (MaxY - MinY) * (MaxZ - MinZ);
+	UPCGPointData* Data = PCGVolumeSampler::SampleVolume(Context, this, SamplerSettings);
 
-	FPCGAsync::AsyncPointProcessing(Context, NumIterations, Points, [this, MinX, MaxX, MinY, MaxY, MinZ, MaxZ](int32 Index, FPCGPoint& OutPoint)
+	if (Data)
 	{
-		const int X = MinX + (Index % (MaxX - MinX));
-		const int Y = MinY + (Index / (MaxX - MinX) % (MaxY - MinY));
-		const int Z = MinZ + (Index / ((MaxX - MinX) * (MaxY - MinY)));
-
-		const FVector SampleLocation(X * VoxelSize.X, Y * VoxelSize.Y, Z * VoxelSize.Z);
-		const FBox VoxelBox(VoxelSize * -0.5, VoxelSize * 0.5);
-		if (SamplePoint(FTransform(SampleLocation), VoxelBox, OutPoint, nullptr))
-		{
-			OutPoint.Seed = PCGHelpers::ComputeSeed(X, Y, Z);
-			return true;
-		}
-		else
-		{
-			return false;
-		}
-	});
-
-	UE_LOG(LogPCG, Verbose, TEXT("Primitive %s extracted %d points"), *Primitive->GetFName().ToString(), Points.Num());
+		UE_LOG(LogPCG, Verbose, TEXT("Primitive %s extracted %d points"), *Primitive->GetFName().ToString(), Data->GetPoints().Num());
+	}
 
 	return Data;
 }
\ No newline at end of file
diff --git a/Engine/Plugins/Experimental/PCG/Source/PCG/Private/Data/PCGUnionData.cpp b/Engine/Plugins/Experimental/PCG/Source/PCG/Private/Data/PCGUnionData.cpp
--- a/Engine/Plugins/Experimental/PCG/Source/PCG/Private/Data/PCGUnionData.cpp
+++ b/Engine/Plugins/Experimental/PCG/Source/PCG/Private/Data/PCGUnionData.cpp
@@ -53,7 +53,7 @@ void UPCGUnionData::AddData(const UPCGSpatialData* InData)
 	else
 	{
 		CachedBounds += InData->GetBounds();
-		CachedStrictBounds = CachedStrictBounds.Overlap(InData->GetStrictBounds());
+		CachedStrictBounds = PCGHelpers::OverlapBounds(CachedStrictBounds, InData->GetStrictBounds());
 		CachedDimension = FMath::Max(CachedDimension, InData->GetDimension());
 	}
 
diff --git a/Engine/Plugins/Experimental/PCG/Source/PCG/Private/Data/PCGVolumeData.cpp b/Engine/Plugins/Experimental/PCG/Source/PCG/Private/Data/PCGVolumeData.cpp
--- a/Engine/Plugins/Experimental/PCG/Source/PCG/Private/Data/PCGVolumeData.cpp
+++ b/Engine/Plugins/Experimental/PCG/Source/PCG/Private/Data/PCGVolumeData.cpp
@@ -1,7 +1,9 @@
 // Copyright Epic Games, Inc. All Rights Reserved.
 
 #include "Data/PCGVolumeData.h"
+#include "Data/PCGPointData.h"
 #include "PCGHelpers.h"
+#include "Elements/PCGVolumeSampler.h"
 
 #include "Components/BrushComponent.h"
 #include "GameFramework/Volume.h"
@@ -38,8 +40,19 @@ FBox UPCGVolumeData::GetStrictBounds() const
 
 const UPCGPointData* UPCGVolumeData::CreatePointData(FPCGContext* Context) const
 {
-	UE_LOG(LogPCG, Error, TEXT("Volume data has no default point sampling"));
-	return nullptr;
+	TRACE_CPUPROFILER_EVENT_SCOPE(UPCGVolumeData::CreatePointData);
+
+	PCGVolumeSampler::FVolumeSamplerSettings SamplerSettings;
+	SamplerSettings.VoxelSize = VoxelSize;
+
+	UPCGPointData* Data = PCGVolumeSampler::SampleVolume(Context, this, SamplerSettings);
+
+	if (Data)
+	{
+		UE_LOG(LogPCG, Verbose, TEXT("Volume extracted %d points"), Data->GetPoints().Num());
+	}
+
+	return Data;
 }
 
 bool UPCGVolumeData::SamplePoint(const FTransform& InTransform, const FBox& InBounds, FPCGPoint& OutPoint, UPCGMetadata* OutMetadata) const
diff --git a/Engine/Plugins/Experimental/PCG/Source/PCG/Private/Elements/PCGVolumeSampler.cpp b/Engine/Plugins/Experimental/PCG/Source/PCG/Private/Elements/PCGVolumeSampler.cpp
new file mode 100644
--- /dev/null
+++ b/Engine/Plugins/Experimental/PCG/Source/PCG/Private/Elements/PCGVolumeSampler.cpp
@@ -0,0 +1,106 @@
+// Copyright Epic Games, Inc. All Rights Reserved.
+
+#include "Elements/PCGVolumeSampler.h"
+#include "Data/PCGPointData.h"
+#include "Data/PCGSpatialData.h"
+#include "Helpers/PCGAsync.h"
+#include "Helpers/PCGSettingsHelpers.h"
+#include "PCGHelpers.h"
+
+namespace PCGVolumeSampler
+{
+	UPCGPointData* SampleVolume(FPCGContext* Context, const UPCGSpatialData* SpatialData, const FVolumeSamplerSettings& SamplerSettings)
+	{
+		UPCGPointData* Data = NewObject<UPCGPointData>(const_cast<UPCGSpatialData*>(SpatialData));
+		Data->InitializeFromData(SpatialData);
+		TArray<FPCGPoint>& Points = Data->GetMutablePoints();
+
+		const FVector& VoxelSize = SamplerSettings.VoxelSize;
+		const FBox Bounds = SpatialData->GetBounds();
+
+		const int32 MinX = FMath::CeilToInt(Bounds.Min.X / VoxelSize.X);
+		const int32 MaxX = FMath::FloorToInt(Bounds.Max.X / VoxelSize.X);
+		const int32 MinY = FMath::CeilToInt(Bounds.Min.Y / VoxelSize.Y);
+		const int32 MaxY = FMath::FloorToInt(Bounds.Max.Y / VoxelSize.Y);
+		const int32 MinZ = FMath::CeilToInt(Bounds.Min.Z / VoxelSize.Z);
+		const int32 MaxZ = FMath::FloorToInt(Bounds.Max.Z / VoxelSize.Z);
+
+		const int32 NumIterations = (MaxX - MinX) * (MaxY - MinY) * (MaxZ - MinZ);
+
+		FPCGAsync::AsyncPointProcessing(Context, NumIterations, Points, [SpatialData, VoxelSize, MinX, MaxX, MinY, MaxY, MinZ, MaxZ](int32 Index, FPCGPoint& OutPoint)
+		{
+			const int X = MinX + (Index % (MaxX - MinX));
+			const int Y = MinY + (Index / (MaxX - MinX) % (MaxY - MinY));
+			const int Z = MinZ + (Index / ((MaxX - MinX) * (MaxY - MinY)));
+
+			const FVector SampleLocation(X * VoxelSize.X, Y * VoxelSize.Y, Z * VoxelSize.Z);
+			const FBox VoxelBox(VoxelSize * -0.5, VoxelSize * 0.5);
+			if (SpatialData->SamplePoint(FTransform(SampleLocation), VoxelBox, OutPoint, nullptr))
+			{
+				OutPoint.Seed = PCGHelpers::ComputeSeed(X, Y, Z);
+				return true;
+			}
+			else
+			{
+				return false;
+			}
+		});
+
+		return Data;
+	}
+}
+
+FPCGElementPtr UPCGVolumeSamplerSettings::CreateElement() const
+{
+	return MakeShared<FPCGVolumeSamplerElement>();
+}
+
+bool FPCGVolumeSamplerElement::ExecuteInternal(FPCGContext* Context) const
+{
+	TRACE_CPUPROFILER_EVENT_SCOPE(FPCGVolumeSamplerElement::Execute);
+	// TODO: time-sliced implementation
+	const UPCGVolumeSamplerSettings* Settings = Context->GetInputSettings<UPCGVolumeSamplerSettings>();
+	check(Settings);
+
+	TArray<FPCGTaggedData> Inputs = Context->InputData.GetInputs();
+	UPCGParamData* Params = Context->InputData.GetParams();
+
+	const FVector VoxelSize = PCGSettingsHelpers::GetValue(GET_MEMBER_NAME_CHECKED(UPCGVolumeSamplerSettings, VoxelSize), Settings->VoxelSize, Params);
+	if (VoxelSize.X <= 0 || VoxelSize.Y <= 0 || VoxelSize.Z <= 0)
+	{
+		PCGE_LOG(Warning, "Skipped - Invalid voxel size");
+		return true;
+	}
+
+	PCGVolumeSampler::FVolumeSamplerSettings SamplerSettings;
+	SamplerSettings.VoxelSize = VoxelSize;
+
+	TArray<FPCGTaggedData>& Outputs = Context->OutputData.TaggedData;
+
+	// TODO: embarassingly parallel loop
+	for (const FPCGTaggedData& Input : Inputs)
+	{
+		const UPCGSpatialData* SpatialInput = Cast<UPCGSpatialData>(Input.Data);
+
+		if (!SpatialInput)
+		{
+			PCGE_LOG(Error, "Invalid input data");
+			continue;
+		}
+
+		FPCGTaggedData& Output = Outputs.Emplace_GetRef();
+		Output = Input;
+		UPCGPointData* SampledData = PCGVolumeSampler::SampleVolume(Context, SpatialInput, SamplerSettings);
+		Output.Data = SampledData;
+
+		if (SampledData)
+		{
+			PCGE_LOG(Verbose, "Generated %d points in volume", SampledData->GetPoints().Num());
+		}
+	}
+
+	// Finally, forward any settings
+	Outputs.Append(Context->InputData.GetAllSettings());
+
+	return true;
+}
\ No newline at end of file
diff --git a/Engine/Plugins/Experimental/PCG/Source/PCG/Private/PCGHelpers.cpp b/Engine/Plugins/Experimental/PCG/Source/PCG/Private/PCGHelpers.cpp
--- a/Engine/Plugins/Experimental/PCG/Source/PCG/Private/PCGHelpers.cpp
+++ b/Engine/Plugins/Experimental/PCG/Source/PCG/Private/PCGHelpers.cpp
@@ -40,6 +40,18 @@ namespace PCGHelpers
 			(InPosition.Y >= InBox.Min.Y) && (InPosition.Y < InBox.Max.Y);
 	}
 
+	FBox OverlapBounds(const FBox& InA, const FBox& InB)
+	{
+		if (!InA.IsValid || !InB.IsValid)
+		{
+			return FBox(EForceInit::ForceInit);
+		}
+		else
+		{
+			return InA.Overlap(InB);
+		}
+	}
+
 	FBox GetActorBounds(AActor* InActor)
 	{
 		// Specialized version of GetComponentsBoundingBox that skips over PCG generated components
diff --git a/Engine/Plugins/Experimental/PCG/Source/PCG/Private/PCGHelpers.h b/Engine/Plugins/Experimental/PCG/Source/PCG/Private/PCGHelpers.h
--- a/Engine/Plugins/Experimental/PCG/Source/PCG/Private/PCGHelpers.h
+++ b/Engine/Plugins/Experimental/PCG/Source/PCG/Private/PCGHelpers.h
@@ -25,6 +25,8 @@ namespace PCGHelpers
 	bool IsInsideBounds(const FBox& InBox, const FVector& InPosition);
 	bool IsInsideBoundsXY(const FBox& InBox, const FVector& InPosition);
 
+	FBox OverlapBounds(const FBox& InBoxA, const FBox& InBoxB);
+
 	FBox GetActorBounds(AActor* InActor);
 	FBox GetLandscapeBounds(ALandscapeProxy* InLandscape);
 
diff --git a/Engine/Plugins/Experimental/PCG/Source/PCG/Private/Tests/Data/PCGIntersectionDataTest.cpp b/Engine/Plugins/Experimental/PCG/Source/PCG/Private/Tests/Data/PCGIntersectionDataTest.cpp
new file mode 100644
--- /dev/null
+++ b/Engine/Plugins/Experimental/PCG/Source/PCG/Private/Tests/Data/PCGIntersectionDataTest.cpp
@@ -0,0 +1,102 @@
+// Copyright Epic Games, Inc. All Rights Reserved.
+
+#include "Tests/PCGTestsCommon.h"
+#include "Data/PCGIntersectionData.h"
+#include "Data/PCGPointData.h"
+#include "Data/PCGVolumeData.h"
+
+IMPLEMENT_CUSTOM_SIMPLE_AUTOMATION_TEST(FPCGIntersectionDataBasicTest, FPCGTestBaseClass, "pcg.tests.Intersection.Basic", PCGTestsCommon::TestFlags)
+
+bool FPCGIntersectionDataBasicTest::RunTest(const FString& Parameters)
+{
+	UPCGPointData* InsidePoint = PCGTestsCommon::CreatePointData();
+	check(InsidePoint->GetPoints().Num() == 1);
+
+	UPCGPointData* OutsidePoint = PCGTestsCommon::CreatePointData();
+	check(OutsidePoint->GetPoints().Num() == 1);
+	OutsidePoint->GetMutablePoints()[0].Transform.SetLocation(FVector::OneVector * 10000);
+
+	UPCGVolumeData* Volume = PCGTestsCommon::CreateVolumeData(FBox::BuildAABB(FVector::ZeroVector, FVector::OneVector * 100));
+
+	// Create intersections
+	UPCGIntersectionData* InsideVolume = InsidePoint->IntersectWith(Volume);
+	UPCGIntersectionData* VolumeInside = Volume->IntersectWith(InsidePoint);
+	UPCGIntersectionData* OutsideVolume = OutsidePoint->IntersectWith(Volume);
+	UPCGIntersectionData* VolumeOutside = Volume->IntersectWith(OutsidePoint);
+
+	auto ValidateInsideIntersection = [this, InsidePoint](UPCGIntersectionData* Intersection)
+	{
+		// Basic data validations
+		TestTrue("Valid intersection", Intersection != nullptr);
+
+		if (!Intersection)
+		{
+			return;
+		}
+
+		TestTrue("Valid dimension", Intersection->GetDimension() == 0);
+		TestTrue("Valid bounds", Intersection->GetBounds() == InsidePoint->GetBounds());
+		TestTrue("Valid strict bounds", Intersection->GetStrictBounds() == InsidePoint->GetStrictBounds());
+
+		// Validate sample point		
+		const FPCGPoint& Point = InsidePoint->GetPoints()[0];
+
+		FPCGPoint SampledPoint;
+		TestTrue("Successful point sampling", Intersection->SamplePoint(Point.Transform, Point.GetLocalBounds(), SampledPoint, nullptr));
+		// TODO: should do a full point comparison, not only on a positional basis
+		TestTrue("Correct sampled point", (Point.Transform.GetLocation() - SampledPoint.Transform.GetLocation()).SquaredLength() < KINDA_SMALL_NUMBER);
+
+		// Validate create point data
+		const UPCGPointData* OutputPointData = Intersection->ToPointData(nullptr);
+		TestTrue("Successful ToPoint", OutputPointData != nullptr);
+		
+		if (OutputPointData)
+		{
+			TestTrue("Valid number of points in ToPoint", OutputPointData->GetPoints().Num() == 1);
+			if (OutputPointData->GetPoints().Num() == 1)
+			{
+				TestTrue("Correct point in ToPoint", (Point.Transform.GetLocation() - OutputPointData->GetPoints()[0].Transform.GetLocation()).SquaredLength() < KINDA_SMALL_NUMBER);
+			}
+		}
+	};
+
+	ValidateInsideIntersection(InsideVolume);
+	ValidateInsideIntersection(VolumeInside);
+
+	auto ValidateOutsideIntersection = [this, OutsidePoint](UPCGIntersectionData* Intersection)
+	{
+		TestTrue("Valid intersection", Intersection != nullptr);
+
+		if (!Intersection)
+		{
+			return;
+		}
+
+		TestTrue("Valid dimension", Intersection->GetDimension() == 0);
+		TestTrue("Null bounds", !Intersection->GetBounds().IsValid);
+		TestTrue("Null strict bounds", !Intersection->GetStrictBounds().IsValid);
+
+		// Validate that we're not able to sample a point
+		const FPCGPoint& Point = OutsidePoint->GetPoints()[0];
+
+		FPCGPoint SampledPoint;
+		TestTrue("Unsucessful point sampling", !Intersection->SamplePoint(Point.Transform, Point.GetLocalBounds(), SampledPoint, nullptr));
+
+		// Validate empty point data
+		const UPCGPointData* OutputPointData = Intersection->ToPointData(nullptr);
+		TestTrue("Successful ToPoint", OutputPointData != nullptr);
+
+		if (OutputPointData)
+		{
+			TestTrue("Empty point data", OutputPointData->GetPoints().Num() == 0);
+		}
+	};
+
+	ValidateOutsideIntersection(OutsideVolume);
+	ValidateOutsideIntersection(VolumeOutside);
+
+	return true;
+}
+
+//TOODs:
+// Test with one/two data that do not have a trivial transformation (e.g. projection, surfaces, ...)
\ No newline at end of file
diff --git a/Engine/Plugins/Experimental/PCG/Source/PCG/Private/Tests/Elements/PCGCollapseElementTest.cpp b/Engine/Plugins/Experimental/PCG/Source/PCG/Private/Tests/Elements/PCGCollapseElementTest.cpp
new file mode 100644
--- /dev/null
+++ b/Engine/Plugins/Experimental/PCG/Source/PCG/Private/Tests/Elements/PCGCollapseElementTest.cpp
@@ -0,0 +1,12 @@
+// Copyright Epic Games, Inc. All Rights Reserved.
+
+#include "Tests/PCGTestsCommon.h"
+#include "Elements/PCGCollapseElement.h"
+
+IMPLEMENT_CUSTOM_SIMPLE_AUTOMATION_TEST(FPCGCollapseElementSmokeTest, FPCGTestBaseClass, "pcg.tests.Collapse.Smoke", PCGTestsCommon::TestFlags)
+
+bool FPCGCollapseElementSmokeTest::RunTest(const FString& Parameters)
+{
+	UPCGCollapseSettings* Settings = NewObject<UPCGCollapseSettings>();
+	return SmokeTestAnyValidInput(Settings);
+}
\ No newline at end of file
diff --git a/Engine/Plugins/Experimental/PCG/Source/PCG/Private/Tests/PCGTestsCommon.cpp b/Engine/Plugins/Experimental/PCG/Source/PCG/Private/Tests/PCGTestsCommon.cpp
new file mode 100644
--- /dev/null
+++ b/Engine/Plugins/Experimental/PCG/Source/PCG/Private/Tests/PCGTestsCommon.cpp
@@ -0,0 +1,204 @@
+// Copyright Epic Games, Inc. All Rights Reserved.
+
+#include "Tests/PCGTestsCommon.h"
+
+#include "PCGParamData.h"
+#include "PCGSettings.h"
+#include "Data/PCGPointData.h"
+#include "Data/PCGPolyLineData.h"
+#include "Data/PCGSurfaceData.h"
+#include "Data/PCGVolumeData.h"
+#include "Data/PCGPrimitiveData.h"
+
+#include "GameFramework/Actor.h"
+
+namespace PCGTestsCommon
+{
+	AActor* CreateTemporaryActor()
+	{
+		return NewObject<AActor>();
+	}
+
+	UPCGPointData* CreateEmptyPointData()
+	{
+		return NewObject<UPCGPointData>();
+	}
+
+	UPCGPointData* CreatePointData()
+	{
+		UPCGPointData* SinglePointData = CreateEmptyPointData();
+		SinglePointData->GetMutablePoints().Emplace();
+
+		return SinglePointData;
+	}
+
+	UPCGPolyLineData* CreatePolyLineData()
+	{
+		// TODO: spline, landscape spline
+		return nullptr;
+	}
+
+	UPCGSurfaceData* CreateSurfaceData()
+	{
+		// TODO: either landscape, texture, render target
+		return nullptr;
+	}
+
+	UPCGVolumeData* CreateVolumeData(const FBox& InBounds)
+	{
+		UPCGVolumeData* VolumeData = NewObject<UPCGVolumeData>();
+		VolumeData->Initialize(InBounds, nullptr);
+		return VolumeData;
+	}
+
+	UPCGPrimitiveData* CreatePrimitiveData()
+	{
+		// TODO: need UPrimitiveComponent on an actor
+		return nullptr;
+	}
+
+	UPCGParamData* CreateEmptyParamData()
+	{
+		return NewObject<UPCGParamData>();
+	}
+
+	TArray<FPCGDataCollection> GenerateAllowedData(const FPCGPinProperties& PinProperties)
+	{
+		TArray<FPCGDataCollection> Data;
+
+		TMap<EPCGDataType, TFunction<UPCGData*(void)>> TypeToDataFn;
+		TypeToDataFn.Add(EPCGDataType::Point, []() { return PCGTestsCommon::CreatePointData(); });
+		TypeToDataFn.Add(EPCGDataType::PolyLine, []() { return PCGTestsCommon::CreatePolyLineData(); });
+		TypeToDataFn.Add(EPCGDataType::Surface, []() { return PCGTestsCommon::CreateSurfaceData(); });
+		TypeToDataFn.Add(EPCGDataType::Volume, []() { return PCGTestsCommon::CreateVolumeData(); });
+		TypeToDataFn.Add(EPCGDataType::Primitive, []() { return PCGTestsCommon::CreatePrimitiveData(); });
+		TypeToDataFn.Add(EPCGDataType::Param, []() { return PCGTestsCommon::CreateEmptyParamData(); });
+
+		// Create empty data
+		Data.Emplace();
+
+		// Create single data & data pairs
+		for (const auto& TypeToData : TypeToDataFn)
+		{
+			if (!(TypeToData.Key & PinProperties.AllowedTypes))
+			{
+				continue;
+			}
+
+			FPCGDataCollection& SingleCollection = Data.Emplace_GetRef();
+			FPCGTaggedData& SingleTaggedData = SingleCollection.TaggedData.Emplace_GetRef();
+			SingleTaggedData.Data = TypeToData.Value();
+			SingleTaggedData.Pin = PinProperties.Label;
+
+			if (!PinProperties.bAllowMultipleConnections)
+			{
+				continue;
+			}
+
+			for (const auto& SecondaryTypeToData : TypeToDataFn)
+			{
+				if (!(SecondaryTypeToData.Key & PinProperties.AllowedTypes))
+				{
+					continue;
+				}
+
+				FPCGDataCollection& MultiCollection = Data.Emplace_GetRef();
+				FPCGTaggedData& FirstTaggedData = MultiCollection.TaggedData.Emplace_GetRef();
+				FirstTaggedData.Data = TypeToData.Value();
+				FirstTaggedData.Pin = PinProperties.Label;
+
+				FPCGTaggedData& SecondTaggedData = MultiCollection.TaggedData.Emplace_GetRef();
+				SecondTaggedData.Data = SecondaryTypeToData.Value();
+				SecondTaggedData.Pin = PinProperties.Label;
+			}
+		}
+
+		return Data;
+	}
+}
+
+bool FPCGTestBaseClass::SmokeTestAnyValidInput(UPCGSettings* InSettings, TFunction<bool(const FPCGDataCollection&, const FPCGDataCollection&)> ValidationFn)
+{
+	TestTrue("Valid settings", InSettings != nullptr);
+
+	if (!InSettings)
+	{
+		return false;
+	}
+
+	FPCGElementPtr Element = InSettings->GetElement();
+
+	TestTrue("Valid element", Element != nullptr);
+
+	if (!Element)
+	{
+		return false;
+	}
+
+	TArray<FPCGPinProperties> InputProperties = InSettings->InputPinProperties();
+	// For each pin: take nothing, take 1 of any supported type, take 2 of any supported types (if enabled)
+	TArray<TArray<FPCGDataCollection>> InputsPerProperties;
+	TArray<uint32> InputIndices;
+
+	if (!InputProperties.IsEmpty())
+	{
+		for (const FPCGPinProperties& InputProperty : InputProperties)
+		{
+			InputsPerProperties.Add(PCGTestsCommon::GenerateAllowedData(InputProperty));
+			InputIndices.Add(0);
+		}
+	}
+	else
+	{
+		TArray<FPCGDataCollection>& EmptyCollection = InputsPerProperties.Emplace_GetRef();
+		EmptyCollection.Emplace();
+		InputIndices.Add(0);
+	}
+
+	check(InputIndices.Num() == InputsPerProperties.Num());
+
+	while (1)
+	{
+		// Prepare input
+		FPCGDataCollection InputData;
+		for (int32 PinIndex = 0; PinIndex < InputIndices.Num(); ++PinIndex)
+		{
+			InputData.TaggedData.Append(InputsPerProperties[PinIndex][InputIndices[PinIndex]].TaggedData);
+		}
+
+		// Perform execution
+		FPCGContext* Context = Element->Initialize(InputData, nullptr, nullptr);
+		Context->NumAvailableTasks = 1;
+		Element->Execute(Context);
+
+		if (ValidationFn)
+		{
+			TestTrue("Validation", ValidationFn(Context->InputData, Context->OutputData));
+		}
+
+		delete Context;
+
+		// Bump indices
+		int BumpIndex = 0;
+		while (BumpIndex < InputIndices.Num())
+		{
+			if (InputIndices[BumpIndex] == InputsPerProperties[BumpIndex].Num() - 1)
+			{
+				InputIndices[BumpIndex] = 0;
+				++BumpIndex;
+			}
+			else
+			{
+				++InputIndices[BumpIndex];
+				break;
+			}
+		}
+
+		if (BumpIndex == InputIndices.Num())
+		{
+			break;
+		}
+	}
+
+	return true;
+}
diff --git a/Engine/Plugins/Experimental/PCG/Source/PCG/Public/Data/PCGVolumeData.h b/Engine/Plugins/Experimental/PCG/Source/PCG/Public/Data/PCGVolumeData.h
--- a/Engine/Plugins/Experimental/PCG/Source/PCG/Public/Data/PCGVolumeData.h
+++ b/Engine/Plugins/Experimental/PCG/Source/PCG/Public/Data/PCGVolumeData.h
@@ -33,6 +33,9 @@ public:
 	virtual const UPCGPointData* CreatePointData(FPCGContext* Context) const override;
 	// ~End UPCGSpatialDataWithPointCache interface
 
+	UPROPERTY(BlueprintReadWrite, EditAnywhere, Category = "Data")
+	FVector VoxelSize = FVector(100.0, 100.0, 100.0);
+
 protected:
 	UPROPERTY(BlueprintReadWrite, EditAnywhere, Category = SourceData)
 	TObjectPtr<AVolume> Volume = nullptr;
diff --git a/Engine/Plugins/Experimental/PCG/Source/PCG/Public/Elements/PCGVolumeSampler.h b/Engine/Plugins/Experimental/PCG/Source/PCG/Public/Elements/PCGVolumeSampler.h
new file mode 100644
--- /dev/null
+++ b/Engine/Plugins/Experimental/PCG/Source/PCG/Public/Elements/PCGVolumeSampler.h
@@ -0,0 +1,49 @@
+// Copyright Epic Games, Inc. All Rights Reserved.
+
+#pragma once
+
+#include "PCGSettings.h"
+#include "PCGElement.h"
+
+#include "PCGVolumeSampler.generated.h"
+
+class UPCGPointData;
+class UPCGSpatialData;
+
+namespace PCGVolumeSampler
+{
+	struct FVolumeSamplerSettings
+	{
+		FVector VoxelSize;
+	};
+
+	UPCGPointData* SampleVolume(FPCGContext* Context, const UPCGSpatialData* SpatialData, const FVolumeSamplerSettings& SamplerSettings);
+}
+
+UCLASS(BlueprintType, ClassGroup = (Procedural))
+class PCG_API UPCGVolumeSamplerSettings : public UPCGSettings
+{
+	GENERATED_BODY()
+
+public:
+	UPROPERTY(BlueprintReadWrite, EditAnywhere, Category = "Data")
+	FVector VoxelSize = FVector(100.0, 100.0, 100.0);
+
+	//~Begin UPCGSettings interface
+#if WITH_EDITOR
+	virtual FName GetDefaultNodeName() const override { return FName(TEXT("VolumeSamplerNode")); }
+	virtual EPCGSettingsType GetType() const override { return EPCGSettingsType::Sampler; }
+#endif
+
+	virtual TArray<FPCGPinProperties> OutputPinProperties() const override { return Super::DefaultPointOutputPinProperties(); }
+	
+protected:
+	virtual FPCGElementPtr CreateElement() const override;
+	//~End UPCGSettings interface
+};
+
+class FPCGVolumeSamplerElement : public FSimplePCGElement
+{
+protected:
+	bool ExecuteInternal(FPCGContext* Context) const override;
+};
\ No newline at end of file
diff --git a/Engine/Plugins/Experimental/PCG/Source/PCG/Public/Tests/PCGTestsCommon.h b/Engine/Plugins/Experimental/PCG/Source/PCG/Public/Tests/PCGTestsCommon.h
new file mode 100644
--- /dev/null
+++ b/Engine/Plugins/Experimental/PCG/Source/PCG/Public/Tests/PCGTestsCommon.h
@@ -0,0 +1,41 @@
+// Copyright Epic Games, Inc. All Rights Reserved.
+
+#pragma once
+
+#include "Misc/AutomationTest.h"
+#include "Tests/AutomationCommon.h"
+
+class AActor;
+class UPCGPointData;
+class UPCGPolyLineData;
+class UPCGSurfaceData;
+class UPCGVolumeData;
+class UPCGPrimitiveData;
+class UPCGParamData;
+class UPCGSettings;
+struct FPCGDataCollection;
+
+namespace PCGTestsCommon
+{
+	static const int TestFlags = EAutomationTestFlags::EditorContext | EAutomationTestFlags::EngineFilter;
+
+	AActor* CreateTemporaryActor();
+	UPCGPointData* CreateEmptyPointData();
+	UPCGParamData* CreateEmptyParamData();
+
+	UPCGPointData* CreatePointData();
+	UPCGPolyLineData* CreatePolyLineData();
+	UPCGSurfaceData* CreateSurfaceData();
+	UPCGVolumeData* CreateVolumeData(const FBox& InBounds = FBox::BuildAABB(FVector::ZeroVector, FVector::OneVector * 100));
+	UPCGPrimitiveData* CreatePrimitiveData();
+}
+
+class FPCGTestBaseClass : public FAutomationTestBase
+{
+public:
+	using FAutomationTestBase::FAutomationTestBase;
+
+protected:
+	/** Generates all valid input combinations */
+	bool SmokeTestAnyValidInput(UPCGSettings* InSettings, TFunction<bool(const FPCGDataCollection&, const FPCGDataCollection&)> ValidationFn = TFunction<bool(const FPCGDataCollection&, const FPCGDataCollection&)>());
+};
\ No newline at end of file
