Fixed crash/potential crash related to misuse of reference in BP multithreaded calls

#rnx

diff --git a/Engine/Plugins/Experimental/PCG/Source/PCG/Public/Elements/PCGExecuteBlueprint.h b/Engine/Plugins/Experimental/PCG/Source/PCG/Public/Elements/PCGExecuteBlueprint.h
--- a/Engine/Plugins/Experimental/PCG/Source/PCG/Public/Elements/PCGExecuteBlueprint.h
+++ b/Engine/Plugins/Experimental/PCG/Source/PCG/Public/Elements/PCGExecuteBlueprint.h
@@ -41,16 +41,16 @@ public:
 	void Execute(const FPCGDataCollection& Input, FPCGDataCollection& Output);
 
 	UFUNCTION(BlueprintImplementableEvent, Category = Execution)
-	bool PointLoopBody(UPARAM(ref) FPCGContext& InContext, const UPCGPointData* InData, const FPCGPoint& InPoint, FPCGPoint& OutPoint, UPCGMetadata* OutMetadata) const;
+	bool PointLoopBody(const FPCGContext& InContext, const UPCGPointData* InData, const FPCGPoint& InPoint, FPCGPoint& OutPoint, UPCGMetadata* OutMetadata) const;
 
 	UFUNCTION(BlueprintImplementableEvent, Category = Execution)
-	TArray<FPCGPoint> MultiPointLoopBody(UPARAM(ref) FPCGContext& InContext, const UPCGPointData* InData, const FPCGPoint& InPoint, UPCGMetadata* OutMetadata) const;
+	TArray<FPCGPoint> MultiPointLoopBody(const FPCGContext& InContext, const UPCGPointData* InData, const FPCGPoint& InPoint, UPCGMetadata* OutMetadata) const;
 
 	UFUNCTION(BlueprintImplementableEvent, Category = Execution)
-	bool PointPairLoopBody(UPARAM(ref) FPCGContext& InContext, const UPCGPointData* InA, const UPCGPointData* InB, const FPCGPoint& InPointA, const FPCGPoint& InPointB, FPCGPoint& OutPoint, UPCGMetadata* OutMetadata) const;
+	bool PointPairLoopBody(const FPCGContext& InContext, const UPCGPointData* InA, const UPCGPointData* InB, const FPCGPoint& InPointA, const FPCGPoint& InPointB, FPCGPoint& OutPoint, UPCGMetadata* OutMetadata) const;
 
 	UFUNCTION(BlueprintImplementableEvent, Category = Execution)
-	bool IterationLoopBody(UPARAM(ref) FPCGContext& InContext, int64 Iteration, const UPCGSpatialData* InA, const UPCGSpatialData* InB, FPCGPoint& OutPoint, UPCGMetadata* OutMetadata) const;
+	bool IterationLoopBody(const FPCGContext& InContext, int64 Iteration, const UPCGSpatialData* InA, const UPCGSpatialData* InB, FPCGPoint& OutPoint, UPCGMetadata* OutMetadata) const;
 
 	/** Calls the LoopBody function on all points */
 	UFUNCTION(BlueprintCallable, BlueprintPure = false, Category = Execution, meta = (HideSelfPin = "true"))
