Removed graph setup BP, intiailize etc.

#rnx

diff --git a/Engine/Plugins/Experimental/PCG/Source/PCG/Private/PCGGraph.cpp b/Engine/Plugins/Experimental/PCG/Source/PCG/Private/PCGGraph.cpp
--- a/Engine/Plugins/Experimental/PCG/Source/PCG/Private/PCGGraph.cpp
+++ b/Engine/Plugins/Experimental/PCG/Source/PCG/Private/PCGGraph.cpp
@@ -142,50 +142,6 @@ void UPCGGraph::BeginDestroy()
 	Super::BeginDestroy();
 }
 
-#if WITH_EDITOR
-void UPCGGraph::InitializeFromTemplate()
-{
-	Modify();
-
-	// Disable notification until the end of this method, otherwise it will cause issues in proper refresh dispatch
-	DisableNotificationsForEditor();
-
-	auto ResetDefaultNode = [](UPCGNode* InNode) {
-		check(InNode);
-		for (UPCGPin* InputPin : InNode->InputPins)
-		{
-			InputPin->BreakAllEdges();
-		}
-
-		for (UPCGPin* OutputPin : InNode->OutputPins)
-		{
-			OutputPin->BreakAllEdges();
-		}
-
-		// Reset settings as well
-		InNode->SetDefaultSettings(NewObject<UPCGTrivialSettings>(InNode));
-	};
-
-	ResetDefaultNode(InputNode);
-	ResetDefaultNode(OutputNode);
-
-	for (UPCGNode* Node : Nodes)
-	{
-		OnNodeRemoved(Node);
-	}
-
-	Nodes.Reset();
-
-	if (GraphTemplate)
-	{
-		Cast<UPCGGraphSetupBP>(GraphTemplate->GetDefaultObject())->Setup(this);
-	}
-
-	// Reenable notifications and notify listeners as needed
-	EnableNotificationsForEditor();
-}
-#endif
-
 UPCGNode* UPCGGraph::AddNodeOfType(TSubclassOf<class UPCGSettings> InSettingsClass, UPCGSettings*& OutDefaultNodeSettings)
 {
 	UPCGSettings* Settings = NewObject<UPCGSettings>(GetTransientPackage(), InSettingsClass);
diff --git a/Engine/Plugins/Experimental/PCG/Source/PCG/Private/PCGGraphSetupBP.cpp b/Engine/Plugins/Experimental/PCG/Source/PCG/Private/PCGGraphSetupBP.cpp
deleted file mode 100644
diff --git a/Engine/Plugins/Experimental/PCG/Source/PCG/Public/PCGGraph.h b/Engine/Plugins/Experimental/PCG/Source/PCG/Public/PCGGraph.h
--- a/Engine/Plugins/Experimental/PCG/Source/PCG/Public/PCGGraph.h
+++ b/Engine/Plugins/Experimental/PCG/Source/PCG/Public/PCGGraph.h
@@ -6,7 +6,6 @@
 #include "Templates/SubclassOf.h"
 #include "PCGNode.h"
 #include "PCGSettings.h"
-#include "PCGGraphSetupBP.h"
 
 #include "PCGGraph.generated.h"
 
@@ -41,15 +40,6 @@ public:
 	FText Description;
 #endif
 
-	UPROPERTY(BlueprintReadWrite, EditAnywhere, Category = Template)
-	TSubclassOf<UPCGGraphSetupBP> GraphTemplate;
-
-#if WITH_EDITOR
-	/** Resets & initializes the graph from the Graph Template parameter. */
-	UFUNCTION(BlueprintCallable, Category = Template)
-	void InitializeFromTemplate();
-#endif
-
 	/** Creates a default node based on the settings class wanted. Returns the newly created node. */
 	UFUNCTION(BlueprintCallable, Category = Graph, meta=(DeterminesOutputType = "InSettingsClass", DynamicOutputParam = "DefaultNodeSettings"))
 	UPCGNode* AddNodeOfType(TSubclassOf<class UPCGSettings> InSettingsClass, UPCGSettings*& DefaultNodeSettings);
diff --git a/Engine/Plugins/Experimental/PCG/Source/PCG/Public/PCGGraphSetupBP.h b/Engine/Plugins/Experimental/PCG/Source/PCG/Public/PCGGraphSetupBP.h
deleted file mode 100644
diff --git a/Engine/Plugins/Experimental/PCG/Source/PCGEditor/Private/PCGGraphDetails.cpp b/Engine/Plugins/Experimental/PCG/Source/PCGEditor/Private/PCGGraphDetails.cpp
--- a/Engine/Plugins/Experimental/PCG/Source/PCGEditor/Private/PCGGraphDetails.cpp
+++ b/Engine/Plugins/Experimental/PCG/Source/PCGEditor/Private/PCGGraphDetails.cpp
@@ -45,39 +45,6 @@ void FPCGGraphDetails::CustomizeDetails(IDetailLayoutBuilder& DetailBuilder)
 	{
 		PCGCategory.AddProperty(Property);
 	}
-
-	FDetailWidgetRow& NewRow = PCGCategory.AddCustomRow(FText::GetEmpty());
-
-	NewRow.ValueContent()
-		.MaxDesiredWidth(120.f)
-		[
-			SNew(SHorizontalBox)
-			+ SHorizontalBox::Slot()
-			.AutoWidth()
-			.Padding(4.0f)
-			[
-				SNew(SButton)
-				.OnClicked(this, &FPCGGraphDetails::OnInitializeClicked)
-				[
-					SNew(STextBlock)
-					.Font(IDetailLayoutBuilder::GetDetailFont())
-					.Text(LOCTEXT("InitializeButton", "Initialize from template"))
-				]
-			]
-		];
-}
-
-FReply FPCGGraphDetails::OnInitializeClicked()
-{
-	for (TWeakObjectPtr<UPCGGraph>& Graph : SelectedGraphs)
-	{
-		if (Graph.IsValid())
-		{
-			Graph.Get()->InitializeFromTemplate();
-		}
-	}
-
-	return FReply::Handled();
 }
 
 #undef LOCTEXT_NAMESPACE
\ No newline at end of file
diff --git a/Engine/Plugins/Experimental/PCG/Source/PCGEditor/Public/PCGGraphDetails.h b/Engine/Plugins/Experimental/PCG/Source/PCGEditor/Public/PCGGraphDetails.h
--- a/Engine/Plugins/Experimental/PCG/Source/PCGEditor/Public/PCGGraphDetails.h
+++ b/Engine/Plugins/Experimental/PCG/Source/PCGEditor/Public/PCGGraphDetails.h
@@ -16,7 +16,5 @@ public:
 	virtual void CustomizeDetails(IDetailLayoutBuilder& DetailBuilder) override;
 
 private:
-	FReply OnInitializeClicked();
-
 	TArray<TWeakObjectPtr<UPCGGraph>> SelectedGraphs;
 };
\ No newline at end of file
