Removed error on already initialized metadata as that is part of a normal workflow

diff --git a/Engine/Plugins/Experimental/PCG/Source/PCG/Private/Metadata/PCGMetadata.cpp b/Engine/Plugins/Experimental/PCG/Source/PCG/Private/Metadata/PCGMetadata.cpp
--- a/Engine/Plugins/Experimental/PCG/Source/PCG/Private/Metadata/PCGMetadata.cpp
+++ b/Engine/Plugins/Experimental/PCG/Source/PCG/Private/Metadata/PCGMetadata.cpp
@@ -28,7 +28,7 @@ void UPCGMetadata::Initialize(const UPCGMetadata* InParent)
 {
 	if (Parent || Attributes.Num() != 0)
 	{
-		UE_LOG(LogPCG, Error, TEXT("Metadata has already been initialized or already contains attributes"));
+		// Already initialized; note that while that might be construed as a warning, there are legit cases where this is correct
 		return;
 	}
 
