PCG: export SetValueFromValueKey() and GetValueKey() using PCG_API.

diff --git a/Engine/Plugins/Experimental/PCG/Source/PCG/Public/Metadata/PCGMetadataAttribute.h b/Engine/Plugins/Experimental/PCG/Source/PCG/Public/Metadata/PCGMetadataAttribute.h
--- a/Engine/Plugins/Experimental/PCG/Source/PCG/Public/Metadata/PCGMetadataAttribute.h
+++ b/Engine/Plugins/Experimental/PCG/Source/PCG/Public/Metadata/PCGMetadataAttribute.h
@@ -29,8 +29,8 @@ public:
 	virtual void SetValue(PCGMetadataEntryKey ItemKey, const FPCGMetadataAttributeBase* InAttributeA, PCGMetadataEntryKey InEntryKeyA, const FPCGMetadataAttributeBase* InAttributeB, PCGMetadataEntryKey InEntryKeyB, EPCGMetadataOp Op) = 0;
 	virtual bool IsEqualToDefaultValue(PCGMetadataValueKey ValueKey) const = 0;
 
-	void SetValueFromValueKey(PCGMetadataEntryKey EntryKey, PCGMetadataValueKey ValueKey);
-	PCGMetadataValueKey GetValueKey(PCGMetadataEntryKey EntryKey) const;
+	PCG_API	void SetValueFromValueKey(PCGMetadataEntryKey EntryKey, PCGMetadataValueKey ValueKey);
+	PCG_API PCGMetadataValueKey GetValueKey(PCGMetadataEntryKey EntryKey) const;
 	bool HasNonDefaultValue(PCGMetadataEntryKey EntryKey) const;
 	void ClearEntries();
 
