PCG: Export AsyncPointProcessing() as the symbols were hidden.

diff --git a/Engine/Plugins/Experimental/PCG/Source/PCG/Public/Helpers/PCGAsync.h b/Engine/Plugins/Experimental/PCG/Source/PCG/Public/Helpers/PCGAsync.h
--- a/Engine/Plugins/Experimental/PCG/Source/PCG/Public/Helpers/PCGAsync.h
+++ b/Engine/Plugins/Experimental/PCG/Source/PCG/Public/Helpers/PCGAsync.h
@@ -16,7 +16,7 @@ namespace FPCGAsync
 	* @param OutPoints - The array in which the results will be written to. Note that the array will be cleared before execution
 	* @param PointFunc - A function that has the index [0; NumIterations] and has to write to the point & return true when the current call generates a point
 	*/
-	void AsyncPointProcessing(FPCGContext* Context, int32 NumIterations, TArray<FPCGPoint>& OutPoints, const TFunction<bool(int32, FPCGPoint&)>& PointFunc);
+	PCG_API void AsyncPointProcessing(FPCGContext* Context, int32 NumIterations, TArray<FPCGPoint>& OutPoints, const TFunction<bool(int32, FPCGPoint&)>& PointFunc);
 
 	/** 
 	* Helper to do more general 1:1 point processing loops
