Exported some missing classes in the PCG api

#rnx

diff --git a/Engine/Plugins/Experimental/PCG/Source/PCG/Public/Helpers/PCGAsync.h b/Engine/Plugins/Experimental/PCG/Source/PCG/Public/Helpers/PCGAsync.h
--- a/Engine/Plugins/Experimental/PCG/Source/PCG/Public/Helpers/PCGAsync.h
+++ b/Engine/Plugins/Experimental/PCG/Source/PCG/Public/Helpers/PCGAsync.h
@@ -36,7 +36,7 @@
 	* @param OutFilterPoints - The array in which the out-filter results will be written to. Note that the array will be cleared before execution
 	* @param PointFunc - A function that has the index [0; NumIterations] and has to write to the point & return true when the current call generates a point
 	*/
-	void AsyncPointFilterProcessing(FPCGContext* Context, int32 NumIterations, TArray<FPCGPoint>& InFilterPoints, TArray<FPCGPoint>& OutFilterPoints, const TFunction<bool(int32, FPCGPoint&, FPCGPoint&)>& PointFunc);
+	PCG_API void AsyncPointFilterProcessing(FPCGContext* Context, int32 NumIterations, TArray<FPCGPoint>& InFilterPoints, TArray<FPCGPoint>& OutFilterPoints, const TFunction<bool(int32, FPCGPoint&, FPCGPoint&)>& PointFunc);
 
 	/**
 	* Helper to do more general 1:1 point filtering loops
@@ -56,7 +56,7 @@
 	* @param OutPoints - The array in which the results will be written to. Note that the array will be cleared before execution
 	* @param PointFunc - A function that has the index [0; NumIterations] and has to write to the point & return true when the current call generates a point
 	*/
-	void AsyncMultiPointProcessing(FPCGContext* Context, int32 NumIterations, TArray<FPCGPoint>& OutPoints, const TFunction<TArray<FPCGPoint>(int32)>& PointFunc);
+	PCG_API void AsyncMultiPointProcessing(FPCGContext* Context, int32 NumIterations, TArray<FPCGPoint>& OutPoints, const TFunction<TArray<FPCGPoint>(int32)>& PointFunc);
 
 	/** 
 	* Helper to do more general 1:N point processing loops
diff --git a/Engine/Plugins/Experimental/PCG/Source/PCG/Public/Metadata/PCGMetadataAttribute.h b/Engine/Plugins/Experimental/PCG/Source/PCG/Public/Metadata/PCGMetadataAttribute.h
--- a/Engine/Plugins/Experimental/PCG/Source/PCG/Public/Metadata/PCGMetadataAttribute.h
+++ b/Engine/Plugins/Experimental/PCG/Source/PCG/Public/Metadata/PCGMetadataAttribute.h
@@ -9,7 +9,7 @@
 #include "Misc/ScopeRWLock.h"
 #include "Templates/EnableIf.h"
 
-class FPCGMetadataAttributeBase
+class PCG_API FPCGMetadataAttributeBase
 {
 public:
 	FPCGMetadataAttributeBase() = default;
@@ -29,8 +29,8 @@
 	virtual void SetValue(PCGMetadataEntryKey ItemKey, const FPCGMetadataAttributeBase* InAttributeA, PCGMetadataEntryKey InEntryKeyA, const FPCGMetadataAttributeBase* InAttributeB, PCGMetadataEntryKey InEntryKeyB, EPCGMetadataOp Op) = 0;
 	virtual bool IsEqualToDefaultValue(PCGMetadataValueKey ValueKey) const = 0;
 
-	PCG_API	void SetValueFromValueKey(PCGMetadataEntryKey EntryKey, PCGMetadataValueKey ValueKey);
-	PCG_API PCGMetadataValueKey GetValueKey(PCGMetadataEntryKey EntryKey) const;
+	void SetValueFromValueKey(PCGMetadataEntryKey EntryKey, PCGMetadataValueKey ValueKey);
+	PCGMetadataValueKey GetValueKey(PCGMetadataEntryKey EntryKey) const;
 	bool HasNonDefaultValue(PCGMetadataEntryKey EntryKey) const;
 	void ClearEntries();
 
