Revisit resource tracking/cleanup to support other types of tracked data (such as LWI).
Added a cleanup phase which removes unused ISMs.

#rnx

diff --git a/Engine/Plugins/Experimental/PCG/Source/PCG/Private/Elements/PCGSpawnActor.cpp b/Engine/Plugins/Experimental/PCG/Source/PCG/Private/Elements/PCGSpawnActor.cpp
--- a/Engine/Plugins/Experimental/PCG/Source/PCG/Private/Elements/PCGSpawnActor.cpp
+++ b/Engine/Plugins/Experimental/PCG/Source/PCG/Private/Elements/PCGSpawnActor.cpp
@@ -5,6 +5,7 @@
 #include "Helpers/PCGActorHelpers.h"
 #include "PCGComponent.h"
 #include "PCGHelpers.h"
+#include "PCGManagedResource.h"
 
 #include "Components/InstancedStaticMeshComponent.h"
 #include "GameFramework/Actor.h"
@@ -174,13 +175,15 @@ bool FPCGSpawnActorElement::ExecuteInternal(FPCGContext* Context) const
 					SpawnParams.Owner = TargetActor;
 					SpawnParams.SpawnCollisionHandlingOverride = ESpawnActorCollisionHandlingMethod::AlwaysSpawn;
 
+					UPCGManagedActors* ManagedActors = NewObject<UPCGManagedActors>(Context->SourceComponent);
+
 					for (const FPCGPoint& Point : Points)
 					{
 						AActor* GeneratedActor = TargetActor->GetWorld()->SpawnActor(Settings->TemplateActorClass, &Point.Transform, SpawnParams);
 						GeneratedActor->Tags.Add(PCGHelpers::DefaultPCGActorTag);
 						GeneratedActor->AttachToActor(TargetActor, FAttachmentTransformRules::KeepWorldTransform);
 
-						Context->SourceComponent->AddToGeneratedActors(GeneratedActor);
+						ManagedActors->GeneratedActors.Add(GeneratedActor);
 
 						// If the actor spawned has a PCG component, either generate it or mark it as generated if we pass through its inputs
 						TArray<UPCGComponent*> PCGComponents;
@@ -199,6 +202,8 @@ bool FPCGSpawnActorElement::ExecuteInternal(FPCGContext* Context) const
 						}
 					}
 
+					Context->SourceComponent->AddToManagedResources(ManagedActors);
+
 					PCGE_LOG(Verbose, "Generated %d actors", Points.Num());
 				}
 			}
diff --git a/Engine/Plugins/Experimental/PCG/Source/PCG/Private/Helpers/PCGActorHelpers.cpp b/Engine/Plugins/Experimental/PCG/Source/PCG/Private/Helpers/PCGActorHelpers.cpp
--- a/Engine/Plugins/Experimental/PCG/Source/PCG/Private/Helpers/PCGActorHelpers.cpp
+++ b/Engine/Plugins/Experimental/PCG/Source/PCG/Private/Helpers/PCGActorHelpers.cpp
@@ -5,6 +5,7 @@
 #include "PCGHelpers.h"
 
 #include "PCGComponent.h"
+#include "PCGManagedResource.h"
 #include "Components/InstancedStaticMeshComponent.h"
 #include "Components/HierarchicalInstancedStaticMeshComponent.h"
 #include "Engine/InheritableComponentHandler.h"
@@ -22,19 +23,27 @@
 #include "ObjectTools.h"
 #endif
 
-UInstancedStaticMeshComponent* UPCGActorHelpers::GetOrCreateISMC(AActor* InTargetActor, const UPCGComponent* InSourceComponent, const FPCGISMCBuilderParameters& InParams)
+UInstancedStaticMeshComponent* UPCGActorHelpers::GetOrCreateISMC(AActor* InTargetActor, UPCGComponent* InSourceComponent, const FPCGISMCBuilderParameters& InParams)
 {
 	UStaticMesh* InMesh = InParams.Mesh;
 	const TArray<UMaterialInterface*>& InMaterials = InParams.MaterialOverrides;
 
 	check(InTargetActor != nullptr && InMesh != nullptr);
+	check(InSourceComponent);
 
 	TArray<UInstancedStaticMeshComponent*> ISMCs;
-	InTargetActor->GetComponents<UInstancedStaticMeshComponent>(ISMCs);
+	InSourceComponent->ForEachManagedResource([&ISMCs](UPCGManagedResource* InResource)
+	{
+		if (UPCGManagedISMComponent* Resource = Cast<UPCGManagedISMComponent>(InResource))
+		{
+			ISMCs.Add(Resource->GetComponent());
+		}
+	});
 
 	for (UInstancedStaticMeshComponent* ISMC : ISMCs)
 	{
-		if (ISMC->GetStaticMesh() == InMesh &&
+		if (ISMC && 
+			ISMC->GetStaticMesh() == InMesh &&
 			(!InSourceComponent || ISMC->ComponentTags.Contains(InSourceComponent->GetFName())))
 		{
 			// If materials are provided, we'll make sure they match to the already set materials.
@@ -121,6 +130,11 @@ UInstancedStaticMeshComponent* UPCGActorHelpers::GetOrCreateISMC(AActor* InTarge
 	
 	ISMC->ComponentTags.Add(PCGHelpers::DefaultPCGTag);
 
+	// Create managed resource on source component
+	UPCGManagedISMComponent* Resource = NewObject<UPCGManagedISMComponent>(InSourceComponent);
+	Resource->GeneratedComponent = ISMC;
+	InSourceComponent->AddToManagedResources(Resource);
+
 	return ISMC;
 }
 
diff --git a/Engine/Plugins/Experimental/PCG/Source/PCG/Private/PCGComponent.cpp b/Engine/Plugins/Experimental/PCG/Source/PCG/Private/PCGComponent.cpp
--- a/Engine/Plugins/Experimental/PCG/Source/PCG/Private/PCGComponent.cpp
+++ b/Engine/Plugins/Experimental/PCG/Source/PCG/Private/PCGComponent.cpp
@@ -4,6 +4,7 @@
 #include "PCGGraph.h"
 #include "PCGHelpers.h"
 #include "PCGVolume.h"
+#include "PCGManagedResource.h"
 #include "Data/PCGDifferenceData.h"
 #include "Data/PCGIntersectionData.h"
 #include "Data/PCGLandscapeData.h"
@@ -27,6 +28,7 @@
 #include "LandscapeInfo.h"
 #include "LandscapeSplineActor.h"
 #include "LandscapeSplinesComponent.h"
+#include "Misc/ScopeLock.h"
 #include "WorldPartition/WorldPartition.h"
 
 #if WITH_EDITOR
@@ -74,14 +76,24 @@ void UPCGComponent::SetGraph(UPCGGraph* InGraph)
 	}
 
 	OnGraphChanged(Graph, true, false);
-#endif	
+#endif
 }
 
-void UPCGComponent::AddToGeneratedActors(AActor* InActor)
+void UPCGComponent::AddToManagedResources(UPCGManagedResource* InResource)
 {
-	if (InActor)
+	if (InResource)
 	{
-		GeneratedActors.Add(InActor);
+		FScopeLock ResourcesLock(&GeneratedResourcesLock);
+		GeneratedResources.Add(InResource);
+	}
+}
+
+void UPCGComponent::ForEachManagedResource(TFunctionRef<void(UPCGManagedResource*)> Func)
+{
+	FScopeLock ResourcesLock(&GeneratedResourcesLock);
+	for (TObjectPtr<UPCGManagedResource> ManagedResource : GeneratedResources)
+	{
+		Func(ManagedResource);
 	}
 }
 
@@ -276,6 +288,8 @@ void UPCGComponent::PostProcessGraph(const FBox& InNewBounds, bool bInGenerated)
 
 	if (bInGenerated)
 	{
+		CleanupUnusedManagedResources();
+
 		bGenerated = true;
 
 #if WITH_EDITOR
@@ -335,7 +349,7 @@ void UPCGComponent::CleanupInternal(bool bRemoveComponents)
 	UPCGActorHelpers::DeleteActors(GetWorld(), ActorsToDelete.Array());
 }
 
-void UPCGComponent::CleanupInternal(bool bRemoveComponents, TSet<TSoftObjectPtr<AActor>>& OutActorsToDelete)
+void UPCGComponent::CleanupInternal(bool bHardCleanup, TSet<TSoftObjectPtr<AActor>>& OutActorsToDelete)
 {
 	if (!bGenerated || IsPartitioned())
 	{
@@ -345,48 +359,28 @@ void UPCGComponent::CleanupInternal(bool bRemoveComponents, TSet<TSoftObjectPtr<
 	Modify();
 	bGenerated = false;
 
-	// TODO: remove any ISM, HISM on this actor that were created by the PCG execution
-	// TODO: call recursively on actors generated through this + delete them.
-	TArray<UInstancedStaticMeshComponent*> ISMCs;
-	GetOwner()->GetComponents<UInstancedStaticMeshComponent>(ISMCs);
-
-	for (UInstancedStaticMeshComponent* ISMC : ISMCs)
+	FScopeLock ResourcesLock(&GeneratedResourcesLock);
+	for (int32 ResourceIndex = GeneratedResources.Num() - 1; ResourceIndex >= 0; --ResourceIndex)
 	{
-		if (ISMC->ComponentTags.Contains(GetFName()))
+		check(GeneratedResources[ResourceIndex]);
+		if (GeneratedResources[ResourceIndex]->Release(bHardCleanup, OutActorsToDelete))
 		{
-			if (bRemoveComponents || ISMC->ComponentTags.Contains(PCGHelpers::DefaultPCGDebugTag))
-			{
-				ISMC->DestroyComponent();
-			}
-			else
-			{
-				ISMC->ClearInstances();
-				ISMC->UpdateBounds();
-			}
+			GeneratedResources.RemoveAtSwap(ResourceIndex);
 		}
 	}
+}
 
-	// Cleanup recursively before deleting the actor(s)
-	OutActorsToDelete.Append(GeneratedActors);
-
-	TArray<UPCGComponent*> ComponentsToCleanup;
-
-	for (TSoftObjectPtr<AActor> GeneratedActor : GeneratedActors)
+void UPCGComponent::CleanupUnusedManagedResources()
+{
+	FScopeLock ResourcesLock(&GeneratedResourcesLock);
+	for (int32 ResourceIndex = GeneratedResources.Num() - 1; ResourceIndex >= 0; --ResourceIndex)
 	{
-		if (GeneratedActor.IsValid())
+		check(GeneratedResources[ResourceIndex]);
+		if (GeneratedResources[ResourceIndex]->ReleaseIfUnused())
 		{
-			GeneratedActor.Get()->GetComponents<UPCGComponent>(ComponentsToCleanup);
-
-			for (UPCGComponent* Component : ComponentsToCleanup)
-			{
-				Component->CleanupInternal(/*bRemoveComponents=*/false, OutActorsToDelete);
-			}
-
-			ComponentsToCleanup.Reset();
+			GeneratedResources.RemoveAtSwap(ResourceIndex);
 		}
 	}
-
-	GeneratedActors.Reset();
 }
 
 void UPCGComponent::BeginPlay()
@@ -437,6 +431,31 @@ void UPCGComponent::PostLoad()
 		ExcludedTags.Append(ExclusionTags_DEPRECATED);
 		ExclusionTags_DEPRECATED.Reset();
 	}
+
+	/** Deprecation code, should be removed once generated data has been updated */
+	if (bGenerated && GeneratedResources.Num() == 0)
+	{
+		TArray<UInstancedStaticMeshComponent*> ISMCs;
+		GetOwner()->GetComponents<UInstancedStaticMeshComponent>(ISMCs);
+
+		for (UInstancedStaticMeshComponent* ISMC : ISMCs)
+		{
+			if (ISMC->ComponentTags.Contains(GetFName()))
+			{
+				UPCGManagedISMComponent* ManagedComponent = NewObject<UPCGManagedISMComponent>(this);
+				ManagedComponent->GeneratedComponent = ISMC;
+				GeneratedResources.Add(ManagedComponent);
+			}
+		}
+
+		if (GeneratedActors_DEPRECATED.Num() > 0)
+		{
+			UPCGManagedActors* ManagedActors = NewObject<UPCGManagedActors>(this);
+			ManagedActors->GeneratedActors = GeneratedActors_DEPRECATED;
+			GeneratedResources.Add(ManagedActors);
+			GeneratedActors_DEPRECATED.Reset();
+		}
+	}
 #endif
 
 #if WITH_EDITOR
@@ -560,17 +579,25 @@ void UPCGComponent::PostEditChangeProperty(FPropertyChangedEvent& PropertyChange
 
 void UPCGComponent::PreEditUndo()
 {
+	// Here we will keep a copy of flags that we require to keep through the undo
+	// so we can have a consistent state
+	LastGeneratedBoundsPriorToUndo = LastGeneratedBounds;
+
 	// We don't know what is changing so remove all callbacks
 	if (Graph)
 	{
 		Graph->OnGraphChangedDelegate.RemoveAll(this);
 	}
 
+	if (bGenerated)
+	{
+		// Cleanup so managed resources are cleaned in all cases
+		Cleanup(/*bRemoveComponents=*/true, /*bSave=*/PCGComponent::bSaveOnCleanupAndGenerate);
+		// Put back generated flag to its original value so it is captured properly
+		bGenerated = true;
+	}	
+	
 	TeardownTrackingCallbacks();
-	// Here we will keep a copy of flags that we require to keep through the undo
-	// so we can have a consistent state
-	bWasGeneratedPriorToUndo = bGenerated;
-	LastGeneratedBoundsPriorToUndo = LastGeneratedBounds;
 }
 
 void UPCGComponent::PostEditUndo()
@@ -586,14 +613,8 @@ void UPCGComponent::PostEditUndo()
 	RefreshTrackingData();
 	DirtyGenerated(/*bDirtyCachedInput=*/true);
 	DirtyCacheForAllTrackedTags();
-	
-	if (bWasGeneratedPriorToUndo && !bGenerated)
-	{
-		// Need to reset the generated flag to go through the Cleanup
-		bGenerated = true;
-		Cleanup(/*bRemoveComponents=*/true, /*bSave=*/PCGComponent::bSaveOnCleanupAndGenerate);
-	}
-	else
+
+	if (bGenerated)
 	{
 		Refresh();
 	}
diff --git a/Engine/Plugins/Experimental/PCG/Source/PCG/Private/PCGManagedResource.cpp b/Engine/Plugins/Experimental/PCG/Source/PCG/Private/PCGManagedResource.cpp
new file mode 100644
--- /dev/null
+++ b/Engine/Plugins/Experimental/PCG/Source/PCG/Private/PCGManagedResource.cpp
@@ -0,0 +1,93 @@
+// Copyright Epic Games, Inc. All Rights Reserved.
+
+#include "PCGManagedResource.h"
+#include "PCGComponent.h"
+
+#include "Components/InstancedStaticMeshComponent.h"
+
+bool UPCGManagedActors::Release(bool bHardRelease, TSet<TSoftObjectPtr<AActor>>& OutActorsToDelete)
+{
+	OutActorsToDelete.Append(GeneratedActors);
+
+	// Cleanup recursively
+	TInlineComponentArray<UPCGComponent*, 1> ComponentsToCleanup;
+
+	for (TSoftObjectPtr<AActor> GeneratedActor : GeneratedActors)
+	{
+		if (GeneratedActor.IsValid())
+		{
+			GeneratedActor.Get()->GetComponents<UPCGComponent>(ComponentsToCleanup);
+
+			for (UPCGComponent* Component : ComponentsToCleanup)
+			{
+				Component->CleanupInternal(/*bRemoveComponents=*/false, OutActorsToDelete);
+			}
+
+			ComponentsToCleanup.Reset();
+		}
+	}
+
+	GeneratedActors.Reset();
+	return true;
+}
+
+bool UPCGManagedActors::ReleaseIfUnused()
+{
+	return GeneratedActors.IsEmpty();
+}
+
+bool UPCGManagedComponent::Release(bool bHardRelease, TSet<TSoftObjectPtr<AActor>>& OutActorsToDelete)
+{
+	const bool bSupportsComponentReset = SupportsComponentReset();
+	const bool bDeleteComponent = bHardRelease || !bSupportsComponentReset;
+
+	if (GeneratedComponent.IsValid())
+	{
+		if (bDeleteComponent)
+		{
+			GeneratedComponent->DestroyComponent();
+		}
+		else
+		{
+			ResetComponent();
+		}
+	}
+
+	return bDeleteComponent;
+}
+
+bool UPCGManagedComponent::ReleaseIfUnused()
+{
+	return !GeneratedComponent.IsValid();
+}
+
+bool UPCGManagedISMComponent::ReleaseIfUnused()
+{
+	if (Super::ReleaseIfUnused() || !GetComponent())
+	{
+		return true;
+	}
+	else if (GetComponent()->GetInstanceCount() == 0)
+	{
+		GeneratedComponent->DestroyComponent();
+		return true;
+	}
+	else
+	{
+		return false;
+	}
+}
+
+void UPCGManagedISMComponent::ResetComponent()
+{
+	if (UInstancedStaticMeshComponent * ISMC = GetComponent())
+	{
+		ISMC->ClearInstances();
+		ISMC->UpdateBounds();
+	}
+}
+
+UInstancedStaticMeshComponent* UPCGManagedISMComponent::GetComponent() const
+{
+	return Cast<UInstancedStaticMeshComponent>(GeneratedComponent.Get());
+}
\ No newline at end of file
diff --git a/Engine/Plugins/Experimental/PCG/Source/PCG/Public/Helpers/PCGActorHelpers.h b/Engine/Plugins/Experimental/PCG/Source/PCG/Public/Helpers/PCGActorHelpers.h
--- a/Engine/Plugins/Experimental/PCG/Source/PCG/Public/Helpers/PCGActorHelpers.h
+++ b/Engine/Plugins/Experimental/PCG/Source/PCG/Public/Helpers/PCGActorHelpers.h
@@ -31,7 +31,7 @@ class PCG_API UPCGActorHelpers : public UBlueprintFunctionLibrary
 	GENERATED_BODY()
 
 public:
-	static UInstancedStaticMeshComponent* GetOrCreateISMC(AActor* InActor, const UPCGComponent* SourceComponent, const FPCGISMCBuilderParameters& Params);
+	static UInstancedStaticMeshComponent* GetOrCreateISMC(AActor* InActor, UPCGComponent* SourceComponent, const FPCGISMCBuilderParameters& Params);
 	static bool DeleteActors(UWorld* World, const TArray<TSoftObjectPtr<AActor>>& ActorsToDelete);
 
 	/**
diff --git a/Engine/Plugins/Experimental/PCG/Source/PCG/Public/PCGComponent.h b/Engine/Plugins/Experimental/PCG/Source/PCG/Public/PCGComponent.h
--- a/Engine/Plugins/Experimental/PCG/Source/PCG/Public/PCGComponent.h
+++ b/Engine/Plugins/Experimental/PCG/Source/PCG/Public/PCGComponent.h
@@ -11,6 +11,7 @@
 
 class UPCGComponent;
 class UPCGGraph;
+class UPCGManagedResource;
 class UPCGData;
 class ALandscape;
 
@@ -35,6 +36,7 @@ class PCG_API UPCGComponent : public UActorComponent
 	GENERATED_BODY()
 
 	friend class UPCGSubsystem;
+	friend class UPCGManagedActors;
 
 public:
 	/** ~Begin UObject interface */
@@ -59,7 +61,8 @@ public:
 	UPCGGraph* GetGraph() const { return Graph; }
 	void SetGraph(UPCGGraph* InGraph);
 
-	void AddToGeneratedActors(AActor* InActor);
+	void AddToManagedResources(UPCGManagedResource* InResource);
+	void ForEachManagedResource(TFunctionRef<void(UPCGManagedResource*)> InFunction);
 
 	/** Transactionable methods to be called from details UI */
 	void Generate();
@@ -137,6 +140,7 @@ private:
 	void CleanupInternal(bool bRemoveComponents);
 	void CleanupInternal(bool bRemoveComponents, TSet<TSoftObjectPtr<AActor>>& OutActorsToDelete);
 	void PostProcessGraph(const FBox& InNewBounds, bool bInGenerated);
+	void CleanupUnusedManagedResources();
 
 	bool GetActorsFromTags(const TSet<FName>& InTags, TSet<TWeakObjectPtr<AActor>>& OutActors, bool bCullAgainstLocalBounds);
 
@@ -194,15 +198,19 @@ private:
 	UPROPERTY()
 	TSet<TWeakObjectPtr<AActor>> CachedExcludedActors; 
 
+#if WITH_EDITORONLY_DATA
 	UPROPERTY()
-	TSet<TSoftObjectPtr<AActor>> GeneratedActors;
+	TSet<TSoftObjectPtr<AActor>> GeneratedActors_DEPRECATED;
+#endif
+
+	UPROPERTY()
+	TArray<TObjectPtr<UPCGManagedResource>> GeneratedResources;
 
 	UPROPERTY()
 	FBox LastGeneratedBounds = FBox(EForceInit::ForceInit);
 
 #if WITH_EDITOR
 	bool bIsGenerating = false;
-	bool bWasGeneratedPriorToUndo = false;
 	FBox LastGeneratedBoundsPriorToUndo = FBox(EForceInit::ForceInit);
 	FPCGTagToSettingsMap CachedTrackedTagsToSettings;
 #endif
@@ -216,4 +224,6 @@ private:
 	TMap<TWeakObjectPtr<AActor>, TSet<TObjectPtr<UObject>>> CachedTrackedActorToDependencies;
 	bool bActorToTagsMapPopulated = false;
 #endif
+
+	FCriticalSection GeneratedResourcesLock;
 };
\ No newline at end of file
diff --git a/Engine/Plugins/Experimental/PCG/Source/PCG/Public/PCGManagedResource.h b/Engine/Plugins/Experimental/PCG/Source/PCG/Public/PCGManagedResource.h
new file mode 100644
--- /dev/null
+++ b/Engine/Plugins/Experimental/PCG/Source/PCG/Public/PCGManagedResource.h
@@ -0,0 +1,79 @@
+// Copyright Epic Games, Inc. All Rights Reserved.
+
+#pragma once
+
+#include "CoreMinimal.h"
+
+#include "PCGManagedResource.generated.h"
+
+class UActorComponent;
+class UInstancedStaticMeshComponent;
+
+/** 
+* This class is used to hold resources and their mechanism to delete them on demand.
+* In order to allow for some reuse (e.g. components), the Release call supports a "soft"
+* release that should empty any data but might retain some data.
+* At the end of the generate, a call to ReleaseIfUnused will serve to finally cleanup
+* what is not needed.
+*/
+UCLASS(BlueprintType)
+class PCG_API UPCGManagedResource : public UObject
+{
+	GENERATED_BODY()
+public:
+	/** Releases/Resets the resource depending on the bHardRelease flag. Returns true if resource can be removed from the PCG component */
+	virtual bool Release(bool bHardRelease, TSet<TSoftObjectPtr<AActor>>& OutActorsToDelete) { return true; }
+	/** Releases resource if empty or unused. Returns true if the resource can be removed from the PCG component */
+	virtual bool ReleaseIfUnused() { return true; }
+};
+
+UCLASS(BlueprintType)
+class PCG_API UPCGManagedActors : public UPCGManagedResource
+{
+	GENERATED_BODY()
+
+public:
+	//~Begin UPCGManagedResource interface
+	virtual bool Release(bool bHardRelease, TSet<TSoftObjectPtr<AActor>>& OutActorsToDelete) override;
+	virtual bool ReleaseIfUnused() override;
+	//~End UPCGManagedResource interface
+
+	UPROPERTY(BlueprintReadWrite, EditAnywhere, Category = GeneratedData)
+	TSet<TSoftObjectPtr<AActor>> GeneratedActors;
+};
+
+UCLASS(BlueprintType)
+class PCG_API UPCGManagedComponent : public UPCGManagedResource
+{
+	GENERATED_BODY()
+
+public:
+	//~Begin UPCGManagedResource interface
+	virtual bool Release(bool bHardRelease, TSet<TSoftObjectPtr<AActor>>& OutActorsToDelete) override;
+	virtual bool ReleaseIfUnused() override;
+	//~End UPCGManagedResource interface
+
+	virtual void ResetComponent() { check(0); }
+	virtual bool SupportsComponentReset() const { return false; }
+
+	UPROPERTY(BlueprintReadWrite, EditAnywhere, Category = GeneratedData)
+	TSoftObjectPtr<UActorComponent> GeneratedComponent;
+};
+
+UCLASS(BlueprintType)
+class PCG_API UPCGManagedISMComponent : public UPCGManagedComponent
+{
+	GENERATED_BODY()
+
+public:
+	//~Begin UPCGManagedResource interface
+	virtual bool ReleaseIfUnused() override;
+	//~End UPCGManagedResource interface
+
+	//~Begin UPCGManagedComponents interface
+	virtual void ResetComponent() override;
+	virtual bool SupportsComponentReset() const { return true; }
+	//~End UPCGManagedComponents interface
+
+	UInstancedStaticMeshComponent* GetComponent() const;
+};
\ No newline at end of file
