Improved deprecation warning on ToPointData and input pins
Fixed metadata issue in Extents Modifier
Fixed metadata issue in CopyPoints (to be extended later)
Update ProjectLandscapeNormals & Extents Modifier to fix the warnings

#rnx

diff --git a/Engine/Plugins/Experimental/PCG/Source/PCG/Public/Data/PCGSpatialData.h b/Engine/Plugins/Experimental/PCG/Source/PCG/Public/Data/PCGSpatialData.h
--- a/Engine/Plugins/Experimental/PCG/Source/PCG/Public/Data/PCGSpatialData.h
+++ b/Engine/Plugins/Experimental/PCG/Source/PCG/Public/Data/PCGSpatialData.h
@@ -60,7 +60,7 @@ public:
 	virtual float GetDensityAtPosition(const FVector& InPosition) const;
 
 	/** Discretizes the data into points */
-	UFUNCTION(BlueprintCallable, Category = SpatialData, meta = (DeprecatedFunction))
+	UFUNCTION(BlueprintCallable, Category = SpatialData, meta = (DeprecatedFunction, DeprecationMessage = "The To Point Data function is deprecated - use To Point Data With Context instead."))
 	const UPCGPointData* ToPointData() const { return ToPointData(nullptr); }
 
 	UFUNCTION(BlueprintCallable, Category = SpatialData)
@@ -99,7 +99,7 @@ public:
 	UFUNCTION(BlueprintCallable, Category = Metadata)
 	UPCGMetadata* MutableMetadata() { return Metadata; }
 
-	UFUNCTION(BlueprintCallable, Category = Metadata, meta=(DeprecatedFunction))
+	UFUNCTION(BlueprintCallable, Category = Metadata, meta=(DeprecatedFunction, DeprecationMessage = "The Create Empty Metadata function is not needed anymore - it can safely be removed"))
 	UPCGMetadata* CreateEmptyMetadata();
 
 	UFUNCTION(BlueprintCallable, Category = SpatialData)
diff --git a/Engine/Plugins/Experimental/PCG/Source/PCG/Public/Elements/PCGExecuteBlueprint.h b/Engine/Plugins/Experimental/PCG/Source/PCG/Public/Elements/PCGExecuteBlueprint.h
--- a/Engine/Plugins/Experimental/PCG/Source/PCG/Public/Elements/PCGExecuteBlueprint.h
+++ b/Engine/Plugins/Experimental/PCG/Source/PCG/Public/Elements/PCGExecuteBlueprint.h
@@ -103,7 +103,7 @@ public:
 	UPROPERTY(BlueprintReadWrite, EditAnywhere, Category = Settings)
 	bool bCanBeMultithreaded = false;
 
-	UPROPERTY(BlueprintGetter=InputLabels, Category = "Settings|Input & Output")
+	UPROPERTY(BlueprintGetter=InputLabels, Category = "Settings|Input & Output", meta = (DeprecatedProperty, DeprecatedMessage = "Input Pin Labels are deprecated - use Input Labels instead."))
 	TSet<FName> InputPinLabels_DEPRECATED;
 
 	UPROPERTY(BlueprintGetter=OutputLabels, Category = "Settings|Input & Output")
