Add support for double click on PCG Nodes.
* Blueprint node will open the underlying blueprint in the blueprint editor
* Subgraph node will open the subgraph in PCGGraphEditor.
* Otherwise, native nodes will try to open the setting header class in C++
* Default behavior can be toggled off with a setting in PCGEditorSettings

diff --git a/Engine/Plugins/Experimental/PCG/Source/PCG/Private/Elements/PCGExecuteBlueprint.cpp b/Engine/Plugins/Experimental/PCG/Source/PCG/Private/Elements/PCGExecuteBlueprint.cpp
--- a/Engine/Plugins/Experimental/PCG/Source/PCG/Private/Elements/PCGExecuteBlueprint.cpp
+++ b/Engine/Plugins/Experimental/PCG/Source/PCG/Private/Elements/PCGExecuteBlueprint.cpp
@@ -351,6 +351,11 @@ void UPCGBlueprintSettings::GetTrackedActorTags(FPCGTagToSettingsMap& OutTagToSe
 	}
 #endif // WITH_EDITORONLY_DATA
 }
+
+UObject* UPCGBlueprintSettings::GetJumpTargetForDoubleClick() const
+{
+	return BlueprintElementType ? BlueprintElementType->ClassGeneratedBy : nullptr;
+}
 #endif // WITH_EDITOR
 
 FName UPCGBlueprintSettings::AdditionalTaskName() const
diff --git a/Engine/Plugins/Experimental/PCG/Source/PCG/Private/PCGSettings.cpp b/Engine/Plugins/Experimental/PCG/Source/PCG/Private/PCGSettings.cpp
--- a/Engine/Plugins/Experimental/PCG/Source/PCG/Private/PCGSettings.cpp
+++ b/Engine/Plugins/Experimental/PCG/Source/PCG/Private/PCGSettings.cpp
@@ -43,6 +43,13 @@ uint32 UPCGSettings::GetCrc32() const
 	return Ar.Crc32(const_cast<UPCGSettings*>(this));
 }
 
+#if WITH_EDITOR
+UObject* UPCGSettings::GetJumpTargetForDoubleClick() const
+{
+	return const_cast<UObject*>(Cast<UObject>(this));
+}
+#endif // WITH_EDITOR
+
 TArray<FPCGPinProperties> UPCGSettings::InputPinProperties() const
 {
 	TArray<FPCGPinProperties> PinProperties;
diff --git a/Engine/Plugins/Experimental/PCG/Source/PCG/Private/PCGSubgraph.cpp b/Engine/Plugins/Experimental/PCG/Source/PCG/Private/PCGSubgraph.cpp
--- a/Engine/Plugins/Experimental/PCG/Source/PCG/Private/PCGSubgraph.cpp
+++ b/Engine/Plugins/Experimental/PCG/Source/PCG/Private/PCGSubgraph.cpp
@@ -119,6 +119,13 @@ FName UPCGSubgraphSettings::AdditionalTaskName() const
 }
 
 #if WITH_EDITOR
+UObject* UPCGSubgraphSettings::GetJumpTargetForDoubleClick() const
+{
+	// Note that there is a const_cast done behind the scenes in Cast.
+	// And this behavior is already used in similar part of UE.
+	return Cast<UObject>(Subgraph);
+}
+
 bool UPCGSubgraphSettings::IsStructuralProperty(const FName& InPropertyName) const
 {
 	return (InPropertyName == GET_MEMBER_NAME_CHECKED(UPCGSubgraphSettings, Subgraph)) || Super::IsStructuralProperty(InPropertyName);
diff --git a/Engine/Plugins/Experimental/PCG/Source/PCG/Public/Elements/PCGExecuteBlueprint.h b/Engine/Plugins/Experimental/PCG/Source/PCG/Public/Elements/PCGExecuteBlueprint.h
--- a/Engine/Plugins/Experimental/PCG/Source/PCG/Public/Elements/PCGExecuteBlueprint.h
+++ b/Engine/Plugins/Experimental/PCG/Source/PCG/Public/Elements/PCGExecuteBlueprint.h
@@ -153,6 +153,7 @@ public:
 	virtual FLinearColor GetNodeTitleColor() const override;
 	virtual EPCGSettingsType GetType() const override;
 	virtual void GetTrackedActorTags(FPCGTagToSettingsMap& OutTagToSettings) const override;
+	virtual UObject* GetJumpTargetForDoubleClick() const override;
 #endif
 
 	virtual FName AdditionalTaskName() const override;
diff --git a/Engine/Plugins/Experimental/PCG/Source/PCG/Public/PCGSettings.h b/Engine/Plugins/Experimental/PCG/Source/PCG/Public/PCGSettings.h
--- a/Engine/Plugins/Experimental/PCG/Source/PCG/Public/PCGSettings.h
+++ b/Engine/Plugins/Experimental/PCG/Source/PCG/Public/PCGSettings.h
@@ -72,6 +72,11 @@ public:
 	virtual EPCGSettingsType GetType() const { return EPCGSettingsType::Generic; }
 	/** Derived classes must implement this to communicate dependencies on external actors */
 	virtual void GetTrackedActorTags(FPCGTagToSettingsMap& OutTagToSettings) const {}
+	/** Override this class to provide an UObject to jump to in case of double click on node
+	 *  ie. returning a blueprint instance will open the given blueprint in its editor.
+	 *  By default, it will return the underlying class, to try to jump to its header in code
+     */
+	virtual UObject* GetJumpTargetForDoubleClick() const;
 #endif
 
 	/** Derived classes can implement this to expose additional name information in the logs */
diff --git a/Engine/Plugins/Experimental/PCG/Source/PCG/Public/PCGSubgraph.h b/Engine/Plugins/Experimental/PCG/Source/PCG/Public/PCGSubgraph.h
--- a/Engine/Plugins/Experimental/PCG/Source/PCG/Public/PCGSubgraph.h
+++ b/Engine/Plugins/Experimental/PCG/Source/PCG/Public/PCGSubgraph.h
@@ -52,6 +52,7 @@ public:
 #if WITH_EDITOR
 	virtual FName GetDefaultNodeName() const override { return FName(TEXT("SubgraphNode")); }
 	virtual EPCGSettingsType GetType() const override { return EPCGSettingsType::Subgraph; }
+	virtual UObject* GetJumpTargetForDoubleClick() const override;
 #endif
 
 	virtual FName AdditionalTaskName() const override;
diff --git a/Engine/Plugins/Experimental/PCG/Source/PCGEditor/Private/PCGEditor.cpp b/Engine/Plugins/Experimental/PCG/Source/PCGEditor/Private/PCGEditor.cpp
--- a/Engine/Plugins/Experimental/PCG/Source/PCGEditor/Private/PCGEditor.cpp
+++ b/Engine/Plugins/Experimental/PCG/Source/PCGEditor/Private/PCGEditor.cpp
@@ -10,6 +10,7 @@
 #include "PCGGraph.h"
 #include "SPCGEditorGraphFind.h"
 #include "SPCGEditorGraphNodePalette.h"
+#include "PCGEditorSettings.h"
 
 #include "EdGraphUtilities.h"
 #include "Editor.h"
@@ -18,6 +19,7 @@
 #include "Framework/Commands/UICommandList.h"
 #include "Framework/Docking/TabManager.h"
 #include "Framework/MultiBox/MultiBoxBuilder.h"
+#include "Framework/Notifications/NotificationManager.h"
 #include "GraphEditor.h"
 #include "GraphEditorActions.h"
 #include "HAL/PlatformApplicationMisc.h"
@@ -28,6 +30,11 @@
 #include "ScopedTransaction.h"
 #include "ToolMenus.h"
 #include "Widgets/Docking/SDockTab.h"
+#include "SourceCodeNavigation.h"
+#include "UnrealEdGlobals.h"
+#include "Editor/UnrealEdEngine.h"
+#include "Preferences/UnrealEdOptions.h"
+#include "Widgets/Notifications/SNotificationList.h"
 
 #define LOCTEXT_NAMESPACE "PCGGraphEditor"
 
@@ -605,6 +612,7 @@ TSharedRef<SGraphEditor> FPCGEditor::CreateGraphEditorWidget()
 	SGraphEditor::FGraphEditorEvents InEvents;
 	InEvents.OnSelectionChanged = SGraphEditor::FOnSelectionChanged::CreateSP(this, &FPCGEditor::OnSelectedNodesChanged);
 	InEvents.OnTextCommitted = FOnNodeTextCommitted::CreateSP(this, &FPCGEditor::OnNodeTitleCommitted);
+	InEvents.OnNodeDoubleClicked = FSingleNodeEvent::CreateSP(this, &FPCGEditor::OnNodeDoubleClicked);
 
 	return SNew(SGraphEditor)
 		.AdditionalCommands(GraphEditorCommands)
@@ -662,6 +670,58 @@ void FPCGEditor::OnNodeTitleCommitted(const FText& NewText, ETextCommit::Type Co
 	}
 }
 
+void FPCGEditor::OnNodeDoubleClicked(UEdGraphNode* Node)
+{
+	if (Node != nullptr)
+	{
+		if (UObject* Object = Node->GetJumpTargetForDoubleClick())
+		{
+			if (UPCGSettings* PCGSettings = Cast<UPCGSettings>(Object))
+			{
+				JumpToDefinition(PCGSettings->GetClass());
+			}
+			else
+			{
+				GEditor->GetEditorSubsystem<UAssetEditorSubsystem>()->OpenEditorForAsset(Object);
+			}
+		}
+	}
+}
+
+void FPCGEditor::JumpToDefinition(const UClass* Class) const
+{
+	if (ensure(GUnrealEd) && GUnrealEd->GetUnrealEdOptions()->IsCPPAllowed())
+	{
+		const bool bEnableNavigateToNativeNodes = GetDefault<UPCGEditorSettings>()->bEnableNavigateToNativeNodes;
+		if (bEnableNavigateToNativeNodes)
+		{
+			FSourceCodeNavigation::NavigateToClass(Class);
+		}
+		else
+		{
+			// Inform user that the node is native, give them opportunity to enable navigation to native nodes:
+			FNotificationInfo Info(LOCTEXT("NavigateToNativeDisabled", "Navigation to Native (c++) PCG Nodes Disabled"));
+			Info.ExpireDuration = 10.0f;
+			Info.CheckBoxState = bEnableNavigateToNativeNodes ? ECheckBoxState::Checked : ECheckBoxState::Unchecked;
+
+			Info.CheckBoxStateChanged = FOnCheckStateChanged::CreateStatic(
+				[](ECheckBoxState NewState)
+				{
+					const FScopedTransaction Transaction(LOCTEXT("ChangeEnableNavigateToNativeNodes", "Change Enable Navigate to Native Nodes Setting"));
+
+					UPCGEditorSettings* MutableEditorSetings = GetMutableDefault<UPCGEditorSettings>();
+					MutableEditorSetings->Modify();
+					MutableEditorSetings->bEnableNavigateToNativeNodes = (NewState == ECheckBoxState::Checked) ? true : false;
+					MutableEditorSetings->SaveConfig();
+				}
+			);
+			Info.CheckBoxText = LOCTEXT("EnableNavigationToNative", "Enable Navigate to Native Nodes?");
+
+			FSlateNotificationManager::Get().AddNotification(Info);
+		}
+	}
+}
+
 TSharedRef<SDockTab> FPCGEditor::SpawnTab_GraphEditor(const FSpawnTabArgs& Args)
 {
 	return SNew(SDockTab)
diff --git a/Engine/Plugins/Experimental/PCG/Source/PCGEditor/Private/PCGEditor.h b/Engine/Plugins/Experimental/PCG/Source/PCGEditor/Private/PCGEditor.h
--- a/Engine/Plugins/Experimental/PCG/Source/PCGEditor/Private/PCGEditor.h
+++ b/Engine/Plugins/Experimental/PCG/Source/PCGEditor/Private/PCGEditor.h
@@ -107,6 +107,20 @@ private:
 	/** Called when the title of a node is changed */
 	void OnNodeTitleCommitted(const FText& NewText, ETextCommit::Type CommitInfo, UEdGraphNode* NodeBeingChanged);
 
+	/**
+	 * Called when a node is double clicked
+	 *
+	 * @param Node - The Node that was clicked
+	 */
+	void OnNodeDoubleClicked(UEdGraphNode* Node);
+
+	/**
+	 * Try to jump to a given class (if allowed)
+	 *
+	 * @param Class - The Class to jump to
+	 */
+	void JumpToDefinition(const UClass* Class) const;
+
 	TSharedRef<SDockTab> SpawnTab_GraphEditor(const FSpawnTabArgs& Args);
 	TSharedRef<SDockTab> SpawnTab_PropertyDetails(const FSpawnTabArgs& Args);
 	TSharedRef<SDockTab> SpawnTab_Palette(const FSpawnTabArgs& Args);
diff --git a/Engine/Plugins/Experimental/PCG/Source/PCGEditor/Private/PCGEditorGraphNodeBase.cpp b/Engine/Plugins/Experimental/PCG/Source/PCGEditor/Private/PCGEditorGraphNodeBase.cpp
--- a/Engine/Plugins/Experimental/PCG/Source/PCGEditor/Private/PCGEditorGraphNodeBase.cpp
+++ b/Engine/Plugins/Experimental/PCG/Source/PCGEditor/Private/PCGEditorGraphNodeBase.cpp
@@ -406,4 +406,17 @@ void UPCGEditorGraphNodeBase::GetPinHoverText(const UEdGraphPin& Pin, FString& H
 	}
 }
 
+UObject* UPCGEditorGraphNodeBase::GetJumpTargetForDoubleClick() const
+{
+	if (PCGNode)
+	{
+		if (UPCGSettings* Settings = PCGNode->DefaultSettings)
+		{
+			return Settings->GetJumpTargetForDoubleClick();
+		}
+	}
+
+	return nullptr;
+}
+
 #undef LOCTEXT_NAMESPACE
diff --git a/Engine/Plugins/Experimental/PCG/Source/PCGEditor/Private/PCGEditorGraphNodeBase.h b/Engine/Plugins/Experimental/PCG/Source/PCGEditor/Private/PCGEditorGraphNodeBase.h
--- a/Engine/Plugins/Experimental/PCG/Source/PCGEditor/Private/PCGEditorGraphNodeBase.h
+++ b/Engine/Plugins/Experimental/PCG/Source/PCGEditor/Private/PCGEditorGraphNodeBase.h
@@ -45,6 +45,7 @@ public:
 	virtual void PostPasteNode() override;
 	virtual FText GetTooltipText() const override;
 	virtual void GetPinHoverText(const UEdGraphPin& Pin, FString& HoverTextOut) const override;
+	virtual UObject* GetJumpTargetForDoubleClick() const override;
 	// ~End UEdGraphNode interface
 
 	UPCGNode* GetPCGNode() { return PCGNode; }
diff --git a/Engine/Plugins/Experimental/PCG/Source/PCGEditor/Private/PCGEditorSettings.cpp b/Engine/Plugins/Experimental/PCG/Source/PCGEditor/Private/PCGEditorSettings.cpp
--- a/Engine/Plugins/Experimental/PCG/Source/PCGEditor/Private/PCGEditorSettings.cpp
+++ b/Engine/Plugins/Experimental/PCG/Source/PCGEditor/Private/PCGEditorSettings.cpp
@@ -25,6 +25,8 @@ UPCGEditorSettings::UPCGEditorSettings(const FObjectInitializer& ObjectInitializ
 	RenderTargetDataPinColor = FLinearColor(1.0f, 0.3f, 0.f);
 	ParamDataPinColor = FLinearColor(1.0f, 0.6f, 0.0f);
 	UnknownDataPinColor = FLinearColor(0.3f, 0.3f, 0.3f);
+
+	bEnableNavigateToNativeNodes = true;
 }
 
 FLinearColor UPCGEditorSettings::GetColor(UPCGSettings* Settings) const
diff --git a/Engine/Plugins/Experimental/PCG/Source/PCGEditor/Public/PCGEditorSettings.h b/Engine/Plugins/Experimental/PCG/Source/PCGEditor/Public/PCGEditorSettings.h
--- a/Engine/Plugins/Experimental/PCG/Source/PCGEditor/Public/PCGEditorSettings.h
+++ b/Engine/Plugins/Experimental/PCG/Source/PCGEditor/Public/PCGEditorSettings.h
@@ -87,4 +87,8 @@ public:
 	/** User-driven color overrides */
 	UPROPERTY(EditAnywhere, config, Category = Node, meta = (HideAlphaChannel))
 	TMap<TSubclassOf<UPCGSettings>, FLinearColor> OverrideNodeColorByClass;
+
+	/** Specify if we want to jump to definition in case of double click on native PCG Nodes */
+	UPROPERTY(EditAnywhere, config, Category = Workflow)
+	bool bEnableNavigateToNativeNodes;
 };
\ No newline at end of file
