Added callbacks on landscape change
Fixed potential issue with the way we were rooting settings

diff --git a/Engine/Plugins/Experimental/PCG/Source/PCG/Private/Data/PCGLandscapeData.cpp b/Engine/Plugins/Experimental/PCG/Source/PCG/Private/Data/PCGLandscapeData.cpp
--- a/Engine/Plugins/Experimental/PCG/Source/PCG/Private/Data/PCGLandscapeData.cpp
+++ b/Engine/Plugins/Experimental/PCG/Source/PCG/Private/Data/PCGLandscapeData.cpp
@@ -94,7 +94,7 @@ const UPCGPointData* UPCGLandscapeData::CreatePointData(FPCGContext* Context) co
 
 			OutPoint = FPCGPoint(FTransform(VertexLocation),
 				1.0f,
-				PCGHelpers::ComputeSeed(X, Y, (int)HeightAtVertex.GetValue()));
+				PCGHelpers::ComputeSeed(X, Y));
 			OutPoint.SetExtents(Transform.GetScale3D() / 2.0);
 
 			return true;
diff --git a/Engine/Plugins/Experimental/PCG/Source/PCG/Private/Graph/PCGGraphCache.cpp b/Engine/Plugins/Experimental/PCG/Source/PCG/Private/Graph/PCGGraphCache.cpp
--- a/Engine/Plugins/Experimental/PCG/Source/PCG/Private/Graph/PCGGraphCache.cpp
+++ b/Engine/Plugins/Experimental/PCG/Source/PCG/Private/Graph/PCGGraphCache.cpp
@@ -11,16 +11,12 @@ FPCGGraphCacheEntry::FPCGGraphCacheEntry(const FPCGDataCollection& InInput, cons
 	: Input(InInput)
 	, Output(InOutput)
 {
+	// Note: we don't need to root the settings since they'll be owned by the subsystem
 	Settings = InSettings ? Cast<UPCGSettings>(StaticDuplicateObject(InSettings, InOwner.Get())) : nullptr;
 	SettingsCrc32 = InSettings ? InSettings->GetCrc32() : 0;
 
 	Input.AddToRootSet(OutRootSet);
 	Output.AddToRootSet(OutRootSet);
-
-	if (Settings)
-	{
-		OutRootSet.Add(Settings);
-	}
 }
 
 bool FPCGGraphCacheEntry::Matches(const FPCGDataCollection& InInput, int32 InSettingsCrc32) const
diff --git a/Engine/Plugins/Experimental/PCG/Source/PCG/Private/PCGComponent.cpp b/Engine/Plugins/Experimental/PCG/Source/PCG/Private/PCGComponent.cpp
--- a/Engine/Plugins/Experimental/PCG/Source/PCG/Private/PCGComponent.cpp
+++ b/Engine/Plugins/Experimental/PCG/Source/PCG/Private/PCGComponent.cpp
@@ -462,6 +462,15 @@ void UPCGComponent::PostLoad()
 	SetupActorCallbacks();
 	SetupTrackingCallbacks();
 
+	if (TrackedLandscape)
+	{
+		SetupLandscapeTracking();
+	}
+	else
+	{
+		UpdateTrackedLandscape(/*bBoundsCheck=*/false);
+	}	
+
 	if (Graph)
 	{
 		Graph->OnGraphChangedDelegate.AddUObject(this, &UPCGComponent::OnGraphChanged);
@@ -477,6 +486,7 @@ void UPCGComponent::BeginDestroy()
 		Graph->OnGraphChangedDelegate.RemoveAll(this);
 	}
 
+	TeardownLandscapeTracking();
 	TeardownTrackingCallbacks();
 	TeardownActorCallbacks();
 #endif
@@ -552,6 +562,7 @@ void UPCGComponent::PostEditChangeProperty(FPropertyChangedEvent& PropertyChange
 	}
 	else if (PropName == GET_MEMBER_NAME_CHECKED(UPCGComponent, InputType))
 	{
+		UpdateTrackedLandscape();
 		DirtyGenerated(/*bDirtyCachedInput=*/true);
 		Refresh();
 	}
@@ -611,6 +622,7 @@ void UPCGComponent::PostEditUndo()
 
 	SetupTrackingCallbacks();
 	RefreshTrackingData();
+	UpdateTrackedLandscape();
 	DirtyGenerated(/*bDirtyCachedInput=*/true);
 	DirtyCacheForAllTrackedTags();
 
@@ -755,11 +767,12 @@ void UPCGComponent::OnActorDeleted(AActor* InActor)
 
 void UPCGComponent::OnActorMoved(AActor* InActor)
 {
-	if (InActor == GetOwner())
+	if (InActor == GetOwner() || (InActor && InActor == TrackedLandscape))
 	{
 		// TODO: find better metrics to dirty the inputs. 
 		// TODO: this should dirty only the actor pcg data.
 		{
+			UpdateTrackedLandscape();
 			DirtyGenerated(true);
 			Refresh();
 		}
@@ -786,6 +799,85 @@ void UPCGComponent::OnActorMoved(AActor* InActor)
 	}
 }
 
+void UPCGComponent::UpdateTrackedLandscape(bool bBoundsCheck)
+{
+	TeardownLandscapeTracking();
+	TrackedLandscape = nullptr;
+
+	if (InputType == EPCGComponentInput::Landscape)
+	{
+		if (bBoundsCheck)
+		{
+			UPCGData* ActorData = GetActorPCGData();
+			if (const UPCGSpatialData* ActorSpatialData = Cast<const UPCGSpatialData>(ActorData))
+			{
+				TrackedLandscape = PCGHelpers::GetLandscape(GetOwner()->GetWorld(), ActorSpatialData->GetBounds());
+			}
+		}
+		else
+		{
+			TrackedLandscape = PCGHelpers::GetAnyLandscape(GetOwner()->GetWorld());
+		}
+	}
+	else if (InputType == EPCGComponentInput::Actor)
+	{
+		if (ALandscapeProxy* Landscape = Cast<ALandscapeProxy>(GetOwner()))
+		{
+			TrackedLandscape = Landscape;
+		}
+	}
+
+	SetupLandscapeTracking();
+}
+
+void UPCGComponent::SetupLandscapeTracking()
+{
+	if (TrackedLandscape)
+	{
+		TrackedLandscape->OnComponentDataChanged.AddUObject(this, &UPCGComponent::OnLandscapeChanged);
+	}
+}
+
+void UPCGComponent::TeardownLandscapeTracking()
+{
+	if (TrackedLandscape)
+	{
+		TrackedLandscape->OnComponentDataChanged.RemoveAll(this);
+	}
+}
+
+void UPCGComponent::OnLandscapeChanged(ALandscapeProxy* Landscape, const FLandscapeProxyComponentDataChangedParams& ChangeParams)
+{
+	if (Landscape == TrackedLandscape)
+	{
+		// Check if there is an overlap in the changed components vs. the current actor data
+		bool bDirtyInput = (GetOwner() == TrackedLandscape);
+
+		if (!bDirtyInput && InputType == EPCGComponentInput::Landscape)
+		{
+			UPCGData* ActorData = GetActorPCGData();
+			if (const UPCGSpatialData* ActorSpatialData = Cast<const UPCGSpatialData>(ActorData))
+			{
+				const FBox ActorBounds = ActorSpatialData->GetBounds();
+
+				ChangeParams.ForEachComponent([&bDirtyInput, &ActorBounds](const ULandscapeComponent* LandscapeComponent)
+				{
+					if(LandscapeComponent && ActorBounds.Intersect(LandscapeComponent->Bounds.GetBox()))
+					{
+						bDirtyInput = true;
+					}
+				});
+			}
+		}
+
+		if (bDirtyInput)
+		{
+			DirtyGenerated(true);
+			Refresh();
+		}
+	}
+}
+
 void UPCGComponent::OnObjectPropertyChanged(UObject* InObject, FPropertyChangedEvent& InEvent)
 {
 	if(!InEvent.Property)
@@ -1244,7 +1336,7 @@ UPCGData* UPCGComponent::CreateInputPCGData()
 	{
 		UPCGData* ActorData = GetActorPCGData();
 
-		if (ALandscapeProxy* Landscape = Cast<ALandscapeProxy>(this))
+		if (ALandscapeProxy* Landscape = Cast<ALandscapeProxy>(GetOwner()))
 		{
 			// Current actor is a landscape, we don't need to do an additional projection
 			return ActorData;
diff --git a/Engine/Plugins/Experimental/PCG/Source/PCG/Private/PCGHelpers.cpp b/Engine/Plugins/Experimental/PCG/Source/PCG/Private/PCGHelpers.cpp
--- a/Engine/Plugins/Experimental/PCG/Source/PCG/Private/PCGHelpers.cpp
+++ b/Engine/Plugins/Experimental/PCG/Source/PCG/Private/PCGHelpers.cpp
@@ -93,6 +93,18 @@ namespace PCGHelpers
 		}
 	}
 
+	ALandscape* GetAnyLandscape(UWorld* InWorld)
+	{
+		for (TObjectIterator<ALandscape> It; It; ++It)
+		{
+			if (It->GetWorld() == InWorld)
+			{
+				return *It;
+			}
+		}
+		return nullptr;
+	}
+
 	ALandscape* GetLandscape(UWorld* InWorld, const FBox& InBounds)
 	{
 		ALandscape* Landscape = nullptr;
diff --git a/Engine/Plugins/Experimental/PCG/Source/PCG/Private/PCGHelpers.h b/Engine/Plugins/Experimental/PCG/Source/PCG/Private/PCGHelpers.h
--- a/Engine/Plugins/Experimental/PCG/Source/PCG/Private/PCGHelpers.h
+++ b/Engine/Plugins/Experimental/PCG/Source/PCG/Private/PCGHelpers.h
@@ -31,6 +31,7 @@ namespace PCGHelpers
 	FBox GetLandscapeBounds(ALandscapeProxy* InLandscape);
 
 	ALandscape* GetLandscape(UWorld* InWorld, const FBox& InActorBounds);
+	ALandscape* GetAnyLandscape(UWorld* InWorld);
 
 #if WITH_EDITOR
 	APCGWorldActor* GetPCGWorldActor(UWorld* InWorld);
diff --git a/Engine/Plugins/Experimental/PCG/Source/PCG/Public/PCGComponent.h b/Engine/Plugins/Experimental/PCG/Source/PCG/Public/PCGComponent.h
--- a/Engine/Plugins/Experimental/PCG/Source/PCG/Public/PCGComponent.h
+++ b/Engine/Plugins/Experimental/PCG/Source/PCG/Public/PCGComponent.h
@@ -13,7 +13,8 @@ class UPCGComponent;
 class UPCGGraph;
 class UPCGManagedResource;
 class UPCGData;
-class ALandscape;
+class ALandscapeProxy;
+class FLandscapeProxyComponentDataChangedParams;
 
 #if WITH_EDITOR
 	DECLARE_MULTICAST_DELEGATE_OneParam(FOnPCGGraphGenerated, UPCGComponent*);
@@ -213,6 +214,16 @@ private:
 	bool bIsGenerating = false;
 	FBox LastGeneratedBoundsPriorToUndo = FBox(EForceInit::ForceInit);
 	FPCGTagToSettingsMap CachedTrackedTagsToSettings;
+
+	void SetupLandscapeTracking();
+	void TeardownLandscapeTracking();
+	void UpdateTrackedLandscape(bool bBoundsCheck = true);
+	void OnLandscapeChanged(ALandscapeProxy* Landscape, const FLandscapeProxyComponentDataChangedParams& ChangeParams);
+#endif
+
+#if WITH_EDITORONLY_DATA
+	UPROPERTY()
+	TObjectPtr<ALandscapeProxy> TrackedLandscape = nullptr;
 #endif
 
 #if WITH_EDITORONLY_DATA
