Add support for Drag and Drop Assets in PCGEditor

diff --git a/Engine/Plugins/Experimental/PCG/Source/PCG/Private/Elements/PCGExecuteBlueprint.cpp b/Engine/Plugins/Experimental/PCG/Source/PCG/Private/Elements/PCGExecuteBlueprint.cpp
--- a/Engine/Plugins/Experimental/PCG/Source/PCG/Private/Elements/PCGExecuteBlueprint.cpp
+++ b/Engine/Plugins/Experimental/PCG/Source/PCG/Private/Elements/PCGExecuteBlueprint.cpp
@@ -107,6 +107,10 @@ void UPCGBlueprintElement::OnDependencyChanged(UObject* Object, FPropertyChanged
 	OnBlueprintChangedDelegate.Broadcast(this);
 }
 
+FString UPCGBlueprintElement::GetParentClassName()
+{
+	return FString::Printf(TEXT("%s'%s'"), *UClass::StaticClass()->GetName(), *UPCGBlueprintElement::StaticClass()->GetPathName());
+}
 #endif // WITH_EDITOR
 
 FName UPCGBlueprintElement::NodeTitleOverride_Implementation() const
diff --git a/Engine/Plugins/Experimental/PCG/Source/PCG/Public/Elements/PCGExecuteBlueprint.h b/Engine/Plugins/Experimental/PCG/Source/PCG/Public/Elements/PCGExecuteBlueprint.h
--- a/Engine/Plugins/Experimental/PCG/Source/PCG/Public/Elements/PCGExecuteBlueprint.h
+++ b/Engine/Plugins/Experimental/PCG/Source/PCG/Public/Elements/PCGExecuteBlueprint.h
@@ -88,6 +88,9 @@ public:
 	// ~Begin UObject interface
 	virtual void PostEditChangeProperty(FPropertyChangedEvent& PropertyChangedEvent) override;
 	// ~End UObject interface
+
+	/** Used for filtering */
+	static FString GetParentClassName();
 #endif
 
 	/** Needed to be able to call certain blueprint functions */
diff --git a/Engine/Plugins/Experimental/PCG/Source/PCGEditor/Private/PCGEditorGraphSchema.cpp b/Engine/Plugins/Experimental/PCG/Source/PCGEditor/Private/PCGEditorGraphSchema.cpp
--- a/Engine/Plugins/Experimental/PCG/Source/PCGEditor/Private/PCGEditorGraphSchema.cpp
+++ b/Engine/Plugins/Experimental/PCG/Source/PCGEditor/Private/PCGEditorGraphSchema.cpp
@@ -15,6 +15,7 @@
 #include "Engine/Blueprint.h"
 #include "ScopedTransaction.h"
 #include "UObject/UObjectIterator.h"
+#include "PCGEditorUtils.h"
 
 #define LOCTEXT_NAMESPACE "PCGEditorGraphSchema"
 
@@ -219,7 +220,7 @@ void UPCGEditorGraphSchema::GetBlueprintElementActions(FGraphActionMenuBuilder&
 	FARFilter Filter;
 	Filter.ClassNames.Add(UBlueprint::StaticClass()->GetFName());
 	Filter.bRecursiveClasses = true;
-	Filter.TagsAndValues.Add(FBlueprintTags::NativeParentClassPath, FString::Printf(TEXT("%s'%s'"), *UClass::StaticClass()->GetName(), *UPCGBlueprintElement::StaticClass()->GetPathName()));
+	Filter.TagsAndValues.Add(FBlueprintTags::NativeParentClassPath, UPCGBlueprintElement::GetParentClassName());
 
 	TArray<FAssetData> BlueprintElementAssets;
 	AssetRegistryModule.Get().GetAssets(Filter, BlueprintElementAssets);
@@ -265,6 +266,59 @@ void UPCGEditorGraphSchema::GetSubgraphElementActions(FGraphActionMenuBuilder& A
 	}
 }
 
+void UPCGEditorGraphSchema::DroppedAssetsOnGraph(const TArray<FAssetData>& Assets, const FVector2D& GraphPosition, UEdGraph* Graph) const
+{
+	FVector2D GraphPositionOffset = GraphPosition;
+	constexpr float PositionOffsetIncrementY = 50.f;
+	UEdGraphPin* NullFromPin = nullptr;
+
+	for (const FAssetData& AssetData : Assets)
+	{
+		if (const UObject* Asset = AssetData.GetAsset())
+		{
+			if (Asset->IsA<UPCGGraph>())
+			{
+				FPCGEditorGraphSchemaAction_NewSubgraphElement NewSubgraphAction;
+				NewSubgraphAction.SubgraphObjectPath = AssetData.ObjectPath;
+				NewSubgraphAction.PerformAction(Graph, NullFromPin, GraphPositionOffset);
+				GraphPositionOffset.Y += PositionOffsetIncrementY;
+			}
+			else if (PCGEditorUtils::IsAssetPCGBlueprint(AssetData))
+			{
+				const FString GeneratedClass = AssetData.GetTagValueRef<FString>(FBlueprintTags::GeneratedClassPath);
+
+				FPCGEditorGraphSchemaAction_NewBlueprintElement NewBlueprintAction;
+				NewBlueprintAction.BlueprintClassPath = FSoftClassPath(GeneratedClass);
+				NewBlueprintAction.PerformAction(Graph, NullFromPin, GraphPositionOffset);
+				GraphPositionOffset.Y += PositionOffsetIncrementY;
+			}
+		}
+	}
+}
+
+void UPCGEditorGraphSchema::GetAssetsGraphHoverMessage(const TArray<FAssetData>& Assets, const UEdGraph* /*HoverGraph*/, FString& OutTooltipText, bool& OutOkIcon) const
+{
+	for (const FAssetData& AssetData : Assets)
+	{
+		if (const UObject* Asset = AssetData.GetAsset())
+		{
+			if (Asset->IsA<UPCGGraph>() || PCGEditorUtils::IsAssetPCGBlueprint(AssetData))
+			{
+				OutOkIcon = true;
+				return;
+			}
+			else if (Asset->IsA<UBlueprint>())
+			{
+				OutTooltipText = LOCTEXT("PCGEditorDropAssetInvalidBP", "Blueprint does not derive from UPCGBlueprintElement").ToString();
+				OutOkIcon = false;
+				return;
+			}
+		}
+	}
+
+	OutTooltipText = LOCTEXT("PCGEditorDropAssetInvalid", "Can't create a node for this asset").ToString();
+	OutOkIcon = false;
+}
 
 
 FPCGEditorConnectionDrawingPolicy::FPCGEditorConnectionDrawingPolicy(int32 InBackLayerID, int32 InFrontLayerID, float InZoomFactor, const FSlateRect& InClippingRect, FSlateWindowElementList& InDrawElements, UEdGraph* InGraph)
diff --git a/Engine/Plugins/Experimental/PCG/Source/PCGEditor/Private/PCGEditorGraphSchema.h b/Engine/Plugins/Experimental/PCG/Source/PCGEditor/Private/PCGEditorGraphSchema.h
--- a/Engine/Plugins/Experimental/PCG/Source/PCGEditor/Private/PCGEditorGraphSchema.h
+++ b/Engine/Plugins/Experimental/PCG/Source/PCGEditor/Private/PCGEditorGraphSchema.h
@@ -25,6 +25,8 @@ public:
 	virtual void BreakPinLinks(UEdGraphPin& TargetPin, bool bSendsNodeNotification) const override;
 	virtual void BreakSinglePinLink(UEdGraphPin* SourcePin, UEdGraphPin* TargetPin) const override;
 	virtual FConnectionDrawingPolicy* CreateConnectionDrawingPolicy(int32 InBackLayerID, int32 InFrontLayerID, float InZoomFactor, const FSlateRect& InClippingRect, class FSlateWindowElementList& InDrawElements, class UEdGraph* InGraphObj) const override;
+	virtual void DroppedAssetsOnGraph(const TArray<FAssetData>& Assets, const FVector2D& GraphPosition, UEdGraph* Graph) const override;
+	virtual void GetAssetsGraphHoverMessage(const TArray<FAssetData>& Assets, const UEdGraph* HoverGraph, FString& OutTooltipText, bool& OutOkIcon) const override;
 	//~ End EdGraphSchema Interface
 
 private:
diff --git a/Engine/Plugins/Experimental/PCG/Source/PCGEditor/Private/PCGEditorUtils.cpp b/Engine/Plugins/Experimental/PCG/Source/PCGEditor/Private/PCGEditorUtils.cpp
new file mode 100644
--- /dev/null
+++ b/Engine/Plugins/Experimental/PCG/Source/PCGEditor/Private/PCGEditorUtils.cpp
@@ -0,0 +1,15 @@
+// Copyright Epic Games, Inc. All Rights Reserved.
+
+#include "PCGEditorUtils.h"
+
+#include "AssetTypeCategories.h"
+#include "UObject/NoExportTypes.h"
+#include "Elements/PCGExecuteBlueprint.h"
+
+bool PCGEditorUtils::IsAssetPCGBlueprint(const FAssetData& InAssetData)
+{
+	FString InNativeParentClassName = InAssetData.GetTagValueRef<FString>(FBlueprintTags::NativeParentClassPath);
+	FString TargetNativeParentClassName = UPCGBlueprintElement::GetParentClassName();
+
+	return InAssetData.AssetClass == UBlueprint::StaticClass()->GetFName() && InNativeParentClassName == TargetNativeParentClassName;
+}
\ No newline at end of file
diff --git a/Engine/Plugins/Experimental/PCG/Source/PCGEditor/Private/PCGEditorUtils.h b/Engine/Plugins/Experimental/PCG/Source/PCGEditor/Private/PCGEditorUtils.h
new file mode 100644
--- /dev/null
+++ b/Engine/Plugins/Experimental/PCG/Source/PCGEditor/Private/PCGEditorUtils.h
@@ -0,0 +1,10 @@
+// Copyright Epic Games, Inc. All Rights Reserved.
+
+#pragma once
+
+struct FAssetData;
+
+namespace PCGEditorUtils
+{
+	bool IsAssetPCGBlueprint(const FAssetData& InAssetData);
+}
\ No newline at end of file
diff --git a/Engine/Plugins/Experimental/PCG/Source/PCGEditor/Private/SPCGEditorGraphNodePalette.cpp b/Engine/Plugins/Experimental/PCG/Source/PCGEditor/Private/SPCGEditorGraphNodePalette.cpp
--- a/Engine/Plugins/Experimental/PCG/Source/PCGEditor/Private/SPCGEditorGraphNodePalette.cpp
+++ b/Engine/Plugins/Experimental/PCG/Source/PCGEditor/Private/SPCGEditorGraphNodePalette.cpp
@@ -10,6 +10,7 @@
 #include "AssetRegistry/AssetRegistryModule.h"
 #include "Modules/ModuleManager.h"
 #include "SGraphActionMenu.h"
+#include "PCGEditorUtils.h"
 
 #define LOCTEXT_NAMESPACE "SPCGEditorGraphNodePalette"
 
@@ -82,11 +83,7 @@ void SPCGEditorGraphNodePalette::CollectAllActions(FGraphActionListBuilderBase&
 
 void SPCGEditorGraphNodePalette::OnAssetChanged(const FAssetData& InAssetData)
 {
-	FString InNativeParentClassName = InAssetData.GetTagValueRef<FString>(FBlueprintTags::NativeParentClassPath);
-	FString TargetNativeParentClassName = FString::Printf(TEXT("%s'%s'"), *UClass::StaticClass()->GetName(), *UPCGBlueprintElement::StaticClass()->GetPathName());
-
-	if (InAssetData.AssetClass == UPCGGraph::StaticClass()->GetFName() ||
-		(InAssetData.AssetClass == UBlueprint::StaticClass()->GetFName() && InNativeParentClassName == TargetNativeParentClassName))
+	if (InAssetData.AssetClass == UPCGGraph::StaticClass()->GetFName() || PCGEditorUtils::IsAssetPCGBlueprint(InAssetData))
 	{
 		RefreshActionsList(true);
 	}
