Added geometry scripting dependency by default in the PCG plugin

diff --git a/Engine/Plugins/Experimental/PCG/PCG.uplugin b/Engine/Plugins/Experimental/PCG/PCG.uplugin
--- a/Engine/Plugins/Experimental/PCG/PCG.uplugin
+++ b/Engine/Plugins/Experimental/PCG/PCG.uplugin
@@ -24,10 +24,14 @@
 		{
 			"Name": "PCGEditor",
 			"Type": "Editor",
-			"LoadingPhase": "Default"			
+			"LoadingPhase": "Default"
 		}
 
 	],
 	"Plugins": [
+		{
+			"Name": "GeometryScripting",
+			"Enabled": true
+		}
 	]
 }
