Fixed static analysis warning

#rnx

diff --git a/Engine/Plugins/Experimental/PCG/Source/PCG/Private/Helpers/PCGActorHelpers.cpp b/Engine/Plugins/Experimental/PCG/Source/PCG/Private/Helpers/PCGActorHelpers.cpp
--- a/Engine/Plugins/Experimental/PCG/Source/PCG/Private/Helpers/PCGActorHelpers.cpp
+++ b/Engine/Plugins/Experimental/PCG/Source/PCG/Private/Helpers/PCGActorHelpers.cpp
@@ -31,6 +31,11 @@ UInstancedStaticMeshComponent* UPCGActorHelpers::GetOrCreateISMC(AActor* InTarge
 	check(InTargetActor != nullptr && InMesh != nullptr);
 	check(InSourceComponent);
 
+	if (!InSourceComponent)
+	{
+		return nullptr;
+	}
+
 	TArray<UInstancedStaticMeshComponent*> ISMCs;
 	InSourceComponent->ForEachManagedResource([&ISMCs](UPCGManagedResource* InResource)
 	{
@@ -44,7 +49,7 @@ UInstancedStaticMeshComponent* UPCGActorHelpers::GetOrCreateISMC(AActor* InTarge
 	{
 		if (ISMC && 
 			ISMC->GetStaticMesh() == InMesh &&
-			(!InSourceComponent || ISMC->ComponentTags.Contains(InSourceComponent->GetFName())))
+			ISMC->ComponentTags.Contains(InSourceComponent->GetFName()))
 		{
 			// If materials are provided, we'll make sure they match to the already set materials.
 			// If not provided, we'll make sure that the current materials aren't overriden
@@ -122,12 +127,7 @@ UInstancedStaticMeshComponent* UPCGActorHelpers::GetOrCreateISMC(AActor* InTarge
 	}
 	
 	ISMC->AttachToComponent(InTargetActor->GetRootComponent(), FAttachmentTransformRules::KeepWorldTransform);
-
-	if (InSourceComponent)
-	{
-		ISMC->ComponentTags.Add(InSourceComponent->GetFName());
-	}
-	
+	ISMC->ComponentTags.Add(InSourceComponent->GetFName());
 	ISMC->ComponentTags.Add(PCGHelpers::DefaultPCGTag);
 
 	// Create managed resource on source component
