New PCG Unit Test for Difference spatial data.

diff --git a/Engine/Plugins/Experimental/PCG/Source/PCG/Private/Tests/Data/PCGDifferenceDataTest.cpp b/Engine/Plugins/Experimental/PCG/Source/PCG/Private/Tests/Data/PCGDifferenceDataTest.cpp
new file mode 100644
--- /dev/null
+++ b/Engine/Plugins/Experimental/PCG/Source/PCG/Private/Tests/Data/PCGDifferenceDataTest.cpp
@@ -0,0 +1,118 @@
+// Copyright Epic Games, Inc. All Rights Reserved.
+
+#include "Tests/PCGTestsCommon.h"
+#include "Data/PCGDifferenceData.h"
+#include "Data/PCGPointData.h"
+#include "Data/PCGVolumeData.h"
+
+IMPLEMENT_CUSTOM_SIMPLE_AUTOMATION_TEST(FPCGDifferenceDataBasicTest, FPCGTestBaseClass, "pcg.tests.Difference.Basic", PCGTestsCommon::TestFlags)
+
+bool FPCGDifferenceDataBasicTest::RunTest(const FString& Parameters)
+{
+	const UPCGVolumeData* FirstVolume = PCGTestsCommon::CreateVolumeData(FBox::BuildAABB(FVector::OneVector * 250, FVector::OneVector * 500));
+
+	const UPCGVolumeData* SecondVolume = PCGTestsCommon::CreateVolumeData(FBox::BuildAABB(FVector::OneVector * -250, FVector::OneVector * 500));
+
+	const UPCGVolumeData* ThirdVolume = PCGTestsCommon::CreateVolumeData(FBox::BuildAABB(FVector::OneVector * 10000, FVector::OneVector * 500));
+
+	// Inside first volume
+	UPCGPointData* FirstVolumeTestPoint = PCGTestsCommon::CreatePointData(FVector::OneVector * 500);
+	// Inside second volume
+	UPCGPointData* SecondVolumeTestPoint = PCGTestsCommon::CreatePointData(FVector::OneVector * -500);
+	// Point at origin
+	UPCGPointData* OriginTestPoint = PCGTestsCommon::CreatePointData();
+	// Point far from anything
+	UPCGPointData* FarOffTestPoint = PCGTestsCommon::CreatePointData(FVector::OneVector * -10000);
+
+	// Tests for the DifferenceData itself
+	auto ValidateVolumeDifferenceData = [this](const UPCGVolumeData* OriginalVolume, const UPCGVolumeData* DifferenceVolume) -> const UPCGDifferenceData*
+	{
+		// Basic data validation
+		if (!TestNotNull("OriginalVolume", OriginalVolume) || !TestNotNull("DifferenceVolume", DifferenceVolume))
+		{
+			return nullptr;
+		}
+
+		// Create difference data with Subtract
+		const UPCGDifferenceData* DifferenceData = OriginalVolume->Subtract(DifferenceVolume);
+
+		if (!TestNotNull("DifferenceData", DifferenceData))
+		{
+			return nullptr;
+		}
+
+		// Dimension of a volume overlapping another should be a 3D volume
+		TestEqual("DifferenceData dimension", DifferenceData->GetDimension(), 3);
+
+		// Validate the bounds
+		const FBox Bounds = DifferenceData->GetBounds();
+		TestTrue("DifferenceData bounds is valid", !!Bounds.IsValid); // !! to convert uint8 to bool
+		TestEqual("DifferenceData Bounds", Bounds, OriginalVolume->GetBounds());
+
+		return DifferenceData;
+	};
+
+	// Tests points in the various volumes
+	auto ValidatePointsShouldExist = [this](const UPCGDifferenceData* DifferenceData, const UPCGPointData* PointData, bool PointShouldExist)
+	{
+		check(PointData);
+
+		// Basic data validation
+		if (!TestNotNull("DifferenceData", DifferenceData))
+		{
+			return;
+		}
+
+		check(PointData->GetPoints().Num() == 1);
+		// Validate that we're able to sample a point
+		const FPCGPoint& Point = PointData->GetPoints()[0];
+
+		FPCGPoint SampledPoint;
+		bool PointWasSampled = DifferenceData->SamplePoint(Point.Transform, Point.GetLocalBounds(), SampledPoint, nullptr);
+
+		// If the sampled point should exist, and if so, compare it to the original point
+		TestTrue("Point sampled", PointWasSampled == PointShouldExist);
+		if (PointWasSampled && PointShouldExist)
+		{
+			TestTrue("Sampled point identical", PCGTestsCommon::PointsAreIdentical(Point, SampledPoint));
+		}
+
+		// Generate PointData
+		const UPCGPointData* OutputPointData = DifferenceData->ToPointData(nullptr);
+		if (TestNotNull("PointData generated", OutputPointData))
+		{
+			// If the point should be generated
+			TestTrue("PointData exists", OutputPointData->GetPoints().Num() > 0);
+		}
+	};
+
+	// Create and validate difference data with subtraction
+	const UPCGDifferenceData* FirstSubSecond = ValidateVolumeDifferenceData(FirstVolume, SecondVolume); // Should exist
+	const UPCGDifferenceData* SecondSubFirst = ValidateVolumeDifferenceData(SecondVolume, FirstVolume); // Should exist
+	const UPCGDifferenceData* FirstSubThird = ValidateVolumeDifferenceData(FirstVolume, ThirdVolume);   // Should not exist
+	const UPCGDifferenceData* ThirdSubFirst = ValidateVolumeDifferenceData(ThirdVolume, FirstVolume);   // Should not exist
+
+	// Run lambda tests
+	ValidatePointsShouldExist(FirstSubSecond, FirstVolumeTestPoint, true);   // First point inside first volume only
+	ValidatePointsShouldExist(FirstSubSecond, SecondVolumeTestPoint, false); // Second point inside second volume only
+	ValidatePointsShouldExist(FirstSubSecond, OriginTestPoint, false);       // Origin in subtracted volume
+
+	ValidatePointsShouldExist(SecondSubFirst, FirstVolumeTestPoint, false);  // First point in subtracted volume
+	ValidatePointsShouldExist(SecondSubFirst, SecondVolumeTestPoint, true);  // Second point in second volume only
+	ValidatePointsShouldExist(SecondSubFirst, OriginTestPoint, false);       // Origin in subtracted volume
+
+	ValidatePointsShouldExist(FirstSubThird, FirstVolumeTestPoint, true);    // First point in first volume only
+	ValidatePointsShouldExist(FirstSubThird, SecondVolumeTestPoint, false);  // Second point not in either volume
+	ValidatePointsShouldExist(FirstSubThird, OriginTestPoint, true);         // Origin in first volume, no subtracted volume
+
+	ValidatePointsShouldExist(ThirdSubFirst, FirstVolumeTestPoint, false);   // No points in this volume
+	ValidatePointsShouldExist(ThirdSubFirst, SecondVolumeTestPoint, false);
+	ValidatePointsShouldExist(ThirdSubFirst, OriginTestPoint, false);
+
+	ValidatePointsShouldExist(FirstSubSecond, FarOffTestPoint, false); // Far off point shouldn't be in any volume
+	ValidatePointsShouldExist(SecondSubFirst, FarOffTestPoint, false);
+	ValidatePointsShouldExist(FirstSubThird, FarOffTestPoint, false);
+	ValidatePointsShouldExist(ThirdSubFirst, FarOffTestPoint, false);
+
+	return true;
+}
\ No newline at end of file
diff --git a/Engine/Plugins/Experimental/PCG/Source/PCG/Private/Tests/Data/PCGIntersectionDataTest.cpp b/Engine/Plugins/Experimental/PCG/Source/PCG/Private/Tests/Data/PCGIntersectionDataTest.cpp
--- a/Engine/Plugins/Experimental/PCG/Source/PCG/Private/Tests/Data/PCGIntersectionDataTest.cpp
+++ b/Engine/Plugins/Experimental/PCG/Source/PCG/Private/Tests/Data/PCGIntersectionDataTest.cpp
@@ -12,9 +12,8 @@ bool FPCGIntersectionDataBasicTest::RunTest(const FString& Parameters)
 	UPCGPointData* InsidePoint = PCGTestsCommon::CreatePointData();
 	check(InsidePoint->GetPoints().Num() == 1);
 
-	UPCGPointData* OutsidePoint = PCGTestsCommon::CreatePointData();
+	UPCGPointData* OutsidePoint = PCGTestsCommon::CreatePointData(FVector::OneVector * 10000);
 	check(OutsidePoint->GetPoints().Num() == 1);
-	OutsidePoint->GetMutablePoints()[0].Transform.SetLocation(FVector::OneVector * 10000);
 
 	UPCGVolumeData* Volume = PCGTestsCommon::CreateVolumeData(FBox::BuildAABB(FVector::ZeroVector, FVector::OneVector * 100));
 
@@ -43,8 +42,7 @@ bool FPCGIntersectionDataBasicTest::RunTest(const FString& Parameters)
 
 		FPCGPoint SampledPoint;
 		TestTrue("Successful point sampling", Intersection->SamplePoint(Point.Transform, Point.GetLocalBounds(), SampledPoint, nullptr));
-		// TODO: should do a full point comparison, not only on a positional basis
-		TestTrue("Correct sampled point", (Point.Transform.GetLocation() - SampledPoint.Transform.GetLocation()).SquaredLength() < KINDA_SMALL_NUMBER);
+		TestTrue("Correct sampled point", PCGTestsCommon::PointsAreIdentical(Point, SampledPoint));
 
 		// Validate create point data
 		const UPCGPointData* OutputPointData = Intersection->ToPointData(nullptr);
@@ -55,7 +53,7 @@ bool FPCGIntersectionDataBasicTest::RunTest(const FString& Parameters)
 			TestTrue("Valid number of points in ToPoint", OutputPointData->GetPoints().Num() == 1);
 			if (OutputPointData->GetPoints().Num() == 1)
 			{
-				TestTrue("Correct point in ToPoint", (Point.Transform.GetLocation() - OutputPointData->GetPoints()[0].Transform.GetLocation()).SquaredLength() < KINDA_SMALL_NUMBER);
+				TestTrue("Correct point in ToPoint", PCGTestsCommon::PointsAreIdentical(Point, OutputPointData->GetPoints()[0]));
 			}
 		}
 	};
@@ -80,7 +78,7 @@ bool FPCGIntersectionDataBasicTest::RunTest(const FString& Parameters)
 		const FPCGPoint& Point = OutsidePoint->GetPoints()[0];
 
 		FPCGPoint SampledPoint;
-		TestTrue("Unsucessful point sampling", !Intersection->SamplePoint(Point.Transform, Point.GetLocalBounds(), SampledPoint, nullptr));
+		TestTrue("Unsuccessful point sampling", !Intersection->SamplePoint(Point.Transform, Point.GetLocalBounds(), SampledPoint, nullptr));
 
 		// Validate empty point data
 		const UPCGPointData* OutputPointData = Intersection->ToPointData(nullptr);
diff --git a/Engine/Plugins/Experimental/PCG/Source/PCG/Private/Tests/PCGTestsCommon.cpp b/Engine/Plugins/Experimental/PCG/Source/PCG/Private/Tests/PCGTestsCommon.cpp
--- a/Engine/Plugins/Experimental/PCG/Source/PCG/Private/Tests/PCGTestsCommon.cpp
+++ b/Engine/Plugins/Experimental/PCG/Source/PCG/Private/Tests/PCGTestsCommon.cpp
@@ -27,11 +27,26 @@ namespace PCGTestsCommon
 	UPCGPointData* CreatePointData()
 	{
 		UPCGPointData* SinglePointData = CreateEmptyPointData();
+
+		check(SinglePointData);
+
 		SinglePointData->GetMutablePoints().Emplace();
 
 		return SinglePointData;
 	}
 
+	UPCGPointData* CreatePointData(const FVector& InLocation)
+	{
+		UPCGPointData* SinglePointData = CreatePointData();
+
+		check(SinglePointData);
+		check(SinglePointData->GetMutablePoints().Num() == 1);
+
+		SinglePointData->GetMutablePoints()[0].Transform.SetLocation(InLocation);
+
+		return SinglePointData;
+	}
+
 	UPCGPolyLineData* CreatePolyLineData()
 	{
 		// TODO: spline, landscape spline
@@ -115,6 +130,12 @@ namespace PCGTestsCommon
 
 		return Data;
 	}
+
+	bool PointsAreIdentical(const FPCGPoint& FirstPoint, const FPCGPoint& SecondPoint)
+	{
+		// TODO: should do a full point comparison, not only on a positional basis
+		return (FirstPoint.Transform.GetLocation() - SecondPoint.Transform.GetLocation()).SquaredLength() < KINDA_SMALL_NUMBER;
+	}
 }
 
 bool FPCGTestBaseClass::SmokeTestAnyValidInput(UPCGSettings* InSettings, TFunction<bool(const FPCGDataCollection&, const FPCGDataCollection&)> ValidationFn)
diff --git a/Engine/Plugins/Experimental/PCG/Source/PCG/Public/Tests/PCGTestsCommon.h b/Engine/Plugins/Experimental/PCG/Source/PCG/Public/Tests/PCGTestsCommon.h
--- a/Engine/Plugins/Experimental/PCG/Source/PCG/Public/Tests/PCGTestsCommon.h
+++ b/Engine/Plugins/Experimental/PCG/Source/PCG/Public/Tests/PCGTestsCommon.h
@@ -4,6 +4,7 @@
 
 #include "Misc/AutomationTest.h"
 #include "Tests/AutomationCommon.h"
+#include "PCGPoint.h"
 
 class AActor;
 class UPCGPointData;
@@ -23,11 +24,17 @@ namespace PCGTestsCommon
 	UPCGPointData* CreateEmptyPointData();
 	UPCGParamData* CreateEmptyParamData();
 
+	/** Creates a PointData with a single point at the origin */
 	UPCGPointData* CreatePointData();
+	/** Creates a PointData with a single point at the provided location */
+	UPCGPointData* CreatePointData(const FVector& InLocation);
+
 	UPCGPolyLineData* CreatePolyLineData();
 	UPCGSurfaceData* CreateSurfaceData();
 	UPCGVolumeData* CreateVolumeData(const FBox& InBounds = FBox::BuildAABB(FVector::ZeroVector, FVector::OneVector * 100));
 	UPCGPrimitiveData* CreatePrimitiveData();
+
+	bool PointsAreIdentical(const FPCGPoint& FirstPoint, const FPCGPoint& SecondPoint);
 }
 
 class FPCGTestBaseClass : public FAutomationTestBase
