Changed tracked landscape to weak pointer to tentatively prevent saving issues

#rnx

diff --git a/Engine/Plugins/Experimental/PCG/Source/PCG/Private/PCGComponent.cpp b/Engine/Plugins/Experimental/PCG/Source/PCG/Private/PCGComponent.cpp
--- a/Engine/Plugins/Experimental/PCG/Source/PCG/Private/PCGComponent.cpp
+++ b/Engine/Plugins/Experimental/PCG/Source/PCG/Private/PCGComponent.cpp
@@ -462,7 +462,7 @@ void UPCGComponent::PostLoad()
 	SetupActorCallbacks();
 	SetupTrackingCallbacks();
 
-	if (TrackedLandscape)
+	if (TrackedLandscape.IsValid())
 	{
 		SetupLandscapeTracking();
 	}
@@ -832,7 +832,7 @@ void UPCGComponent::UpdateTrackedLandscape(bool bBoundsCheck)
 
 void UPCGComponent::SetupLandscapeTracking()
 {
-	if (TrackedLandscape)
+	if (TrackedLandscape.IsValid())
 	{
 		TrackedLandscape->OnComponentDataChanged.AddUObject(this, &UPCGComponent::OnLandscapeChanged);
 	}
@@ -840,7 +840,7 @@ void UPCGComponent::SetupLandscapeTracking()
 
 void UPCGComponent::TeardownLandscapeTracking()
 {
-	if (TrackedLandscape)
+	if (TrackedLandscape.IsValid())
 	{
 		TrackedLandscape->OnComponentDataChanged.RemoveAll(this);
 	}
diff --git a/Engine/Plugins/Experimental/PCG/Source/PCG/Public/PCGComponent.h b/Engine/Plugins/Experimental/PCG/Source/PCG/Public/PCGComponent.h
--- a/Engine/Plugins/Experimental/PCG/Source/PCG/Public/PCGComponent.h
+++ b/Engine/Plugins/Experimental/PCG/Source/PCG/Public/PCGComponent.h
@@ -223,7 +223,7 @@ private:
 
 #if WITH_EDITORONLY_DATA
 	UPROPERTY()
-	TObjectPtr<ALandscapeProxy> TrackedLandscape = nullptr;
+	TWeakObjectPtr<ALandscapeProxy> TrackedLandscape = nullptr;
 #endif
 
 #if WITH_EDITORONLY_DATA
