Added some behavior control on the difference to clarify continuous vs discrete functions (inferred by default)

diff --git a/Engine/Plugins/Experimental/PCG/Source/PCG/Private/Elements/PCGDifferenceElement.cpp b/Engine/Plugins/Experimental/PCG/Source/PCG/Private/Elements/PCGDifferenceElement.cpp
--- a/Engine/Plugins/Experimental/PCG/Source/PCG/Private/Elements/PCGDifferenceElement.cpp
+++ b/Engine/Plugins/Experimental/PCG/Source/PCG/Private/Elements/PCGDifferenceElement.cpp
@@ -3,6 +3,7 @@
 #include "Elements/PCGDifferenceElement.h"
 
 #include "PCGCommon.h"
+#include "Data/PCGPointData.h"
 #include "Data/PCGUnionData.h"
 #include "Helpers/PCGSettingsHelpers.h"
 
@@ -49,6 +50,9 @@ bool FPCGDifferenceElement::ExecuteInternal(FPCGContext* Context) const
 	UPCGUnionData* UnionData = nullptr;
 	int32 DifferenceTaggedDataIndex = -1;
 
+	bool bHasPointsInSource = false;
+	bool bHasPointsInDifferences = false;
+
 	// Start by either selecting a source or building a union
 	for (FPCGTaggedData& Source : Sources)
 	{
@@ -58,8 +62,12 @@ bool FPCGDifferenceElement::ExecuteInternal(FPCGContext* Context) const
 		if (!SpatialData)
 		{
 			Outputs.Add(Source);
+			continue;
 		}
-		else if (!FirstSpatialData)
+
+		bHasPointsInSource |= (Cast<UPCGPointData>(SpatialData) != nullptr);
+
+		if (!FirstSpatialData)
 		{
 			FirstSpatialData = SpatialData;
 			DifferenceTaggedDataIndex = Outputs.Num();
@@ -93,6 +101,8 @@ bool FPCGDifferenceElement::ExecuteInternal(FPCGContext* Context) const
 
 			if (SpatialData)
 			{
+				bHasPointsInDifferences |= (Cast<UPCGPointData>(SpatialData) != nullptr);
+
 				if (!DifferenceData)
 				{
 					DifferenceData = (UnionData ? UnionData : FirstSpatialData)->Subtract(SpatialData);
@@ -113,6 +123,14 @@ bool FPCGDifferenceElement::ExecuteInternal(FPCGContext* Context) const
 				// We are not propagating data from the differences if they don't contribute
 			}
 		}
+
+		// Finally, apply any discretization based on the mode
+		if (DifferenceData && 
+			(Settings->Mode == EPCGDifferenceMode::Discrete || 
+			(Settings->Mode == EPCGDifferenceMode::Inferred && bHasPointsInSource && bHasPointsInDifferences)))
+		{
+			Outputs[DifferenceTaggedDataIndex].Data = DifferenceData->ToPointData(Context);
+		}
 	}
 
 	return true;
diff --git a/Engine/Plugins/Experimental/PCG/Source/PCG/Public/Data/PCGDifferenceData.h b/Engine/Plugins/Experimental/PCG/Source/PCG/Public/Data/PCGDifferenceData.h
--- a/Engine/Plugins/Experimental/PCG/Source/PCG/Public/Data/PCGDifferenceData.h
+++ b/Engine/Plugins/Experimental/PCG/Source/PCG/Public/Data/PCGDifferenceData.h
@@ -17,6 +17,14 @@ enum class EPCGDifferenceDensityFunction : uint8
 	Binary
 };
 
+UENUM()
+enum class EPCGDifferenceMode : uint8
+{
+	Inferred,
+	Continuous,
+	Discrete
+};
+
 UCLASS(BlueprintType, ClassGroup = (Procedural))
 class PCG_API UPCGDifferenceData : public UPCGSpatialDataWithPointCache
 {
diff --git a/Engine/Plugins/Experimental/PCG/Source/PCG/Public/Elements/PCGDifferenceElement.h b/Engine/Plugins/Experimental/PCG/Source/PCG/Public/Elements/PCGDifferenceElement.h
--- a/Engine/Plugins/Experimental/PCG/Source/PCG/Public/Elements/PCGDifferenceElement.h
+++ b/Engine/Plugins/Experimental/PCG/Source/PCG/Public/Elements/PCGDifferenceElement.h
@@ -31,6 +31,9 @@ public:
 	UPROPERTY(BlueprintReadWrite, EditAnywhere, Category = Settings)
 	EPCGDifferenceDensityFunction DensityFunction = EPCGDifferenceDensityFunction::Minimum;
 
+	UPROPERTY(BlueprintReadWrite, EditAnywhere, Category = Settings)
+	EPCGDifferenceMode Mode = EPCGDifferenceMode::Inferred;
+
 #if WITH_EDITORONLY_DATA
 	UPROPERTY(Transient, BlueprintReadWrite, EditAnywhere, Category = Debug)
 	bool bKeepZeroDensityPoints = false;
