Improved ability to conceptually parent to points in other metadata spaces.
Improved behavior when passing optional output data (prevents changes in metadata)

diff --git a/Engine/Plugins/Experimental/PCG/Source/PCG/Private/Data/PCGPointData.cpp b/Engine/Plugins/Experimental/PCG/Source/PCG/Private/Data/PCGPointData.cpp
--- a/Engine/Plugins/Experimental/PCG/Source/PCG/Private/Data/PCGPointData.cpp
+++ b/Engine/Plugins/Experimental/PCG/Source/PCG/Private/Data/PCGPointData.cpp
@@ -298,7 +298,7 @@ bool UPCGPointData::SamplePoint(const FTransform& InTransform, const FBox& InBou
 
 	if (OutMetadata)
 	{
-		UPCGMetadataAccessorHelpers::InitializeMetadata(OutPoint, OutMetadata, *MaxContributor, Metadata);
+		UPCGMetadataAccessorHelpers::InitializeMetadataWithParent(OutPoint, OutMetadata, *MaxContributor, Metadata);
 		OutMetadata->ResetPointWeightedAttributes(OutPoint);
 
 		for (const TPair<const FPCGPoint*, float> Contribution : Contributions)
diff --git a/Engine/Plugins/Experimental/PCG/Source/PCG/Private/Elements/PCGExecuteBlueprint.cpp b/Engine/Plugins/Experimental/PCG/Source/PCG/Private/Elements/PCGExecuteBlueprint.cpp
--- a/Engine/Plugins/Experimental/PCG/Source/PCG/Private/Elements/PCGExecuteBlueprint.cpp
+++ b/Engine/Plugins/Experimental/PCG/Source/PCG/Private/Elements/PCGExecuteBlueprint.cpp
@@ -503,8 +503,15 @@ void UPCGBlueprintElement::LoopOnPoints(FPCGContext& InContext, const UPCGPointD
 		return;
 	}
 
-	OutData = OptionalOutData ? OptionalOutData : NewObject<UPCGPointData>(const_cast<UPCGPointData*>(InData));
-	OutData->InitializeFromData(InData);
+	if (OptionalOutData)
+	{
+		OutData = OptionalOutData;
+	}
+	else
+	{
+		OutData = NewObject<UPCGPointData>(const_cast<UPCGPointData*>(InData));
+		OutData->InitializeFromData(InData);
+	}
 
 	const TArray<FPCGPoint>& InPoints = InData->GetPoints();
 	TArray<FPCGPoint>& OutPoints = OutData->GetMutablePoints();
@@ -523,8 +530,15 @@ void UPCGBlueprintElement::MultiLoopOnPoints(FPCGContext& InContext, const UPCGP
 		return;
 	}
 
-	OutData = OptionalOutData ? OptionalOutData : NewObject<UPCGPointData>(const_cast<UPCGPointData*>(InData));
-	OutData->InitializeFromData(InData);
+	if (OptionalOutData)
+	{
+		OutData = OptionalOutData;
+	}
+	else
+	{
+		OutData = NewObject<UPCGPointData>(const_cast<UPCGPointData*>(InData));
+		OutData->InitializeFromData(InData);
+	}
 
 	const TArray<FPCGPoint>& InPoints = InData->GetPoints();
 	TArray<FPCGPoint>& OutPoints = OutData->GetMutablePoints();
@@ -543,8 +557,16 @@ void UPCGBlueprintElement::LoopOnPointPairs(FPCGContext& InContext, const UPCGPo
 		return;
 	}
 
-	OutData = OptionalOutData ? OptionalOutData : NewObject<UPCGPointData>(const_cast<UPCGPointData*>(InA));
-	OutData->InitializeFromData(InA); //METADATA TODO should we remove the parenting here?
+	if (OptionalOutData)
+	{
+		OutData = OptionalOutData;
+	}
+	else
+	{
+		OutData = NewObject<UPCGPointData>(const_cast<UPCGPointData*>(InA));
+		OutData->InitializeFromData(InA);
+		OutData->Metadata->AddAttributes(InB->Metadata);
+	}
 
 	const TArray<FPCGPoint>& InPointsA = InA->GetPoints();
 	const TArray<FPCGPoint>& InPointsB = InB->GetPoints();
@@ -564,17 +586,19 @@ void UPCGBlueprintElement::LoopNTimes(FPCGContext& InContext, int64 NumIteration
 		return;
 	}
 
-	const UPCGSpatialData* Owner = (InA ? InA : InB);
-	OutData = OptionalOutData;
-
-	if (!OutData)
+	if (OptionalOutData)
 	{
-		OutData = Owner ? NewObject<UPCGPointData>(const_cast<UPCGSpatialData*>(Owner)) : NewObject<UPCGPointData>();
+		OutData = OptionalOutData;
 	}
-
-	if (Owner)
+	else
 	{
-		OutData->InitializeFromData(Owner);
+		const UPCGSpatialData* Owner = (InA ? InA : InB);
+		OutData = Owner ? NewObject<UPCGPointData>(const_cast<UPCGSpatialData*>(Owner)) : NewObject<UPCGPointData>();
+
+		if (Owner)
+		{
+			OutData->InitializeFromData(Owner);
+		}
 	}
 
 	TArray<FPCGPoint>& OutPoints = OutData->GetMutablePoints();
diff --git a/Engine/Plugins/Experimental/PCG/Source/PCG/Private/Metadata/PCGMetadataAccessor.cpp b/Engine/Plugins/Experimental/PCG/Source/PCG/Private/Metadata/PCGMetadataAccessor.cpp
--- a/Engine/Plugins/Experimental/PCG/Source/PCG/Private/Metadata/PCGMetadataAccessor.cpp
+++ b/Engine/Plugins/Experimental/PCG/Source/PCG/Private/Metadata/PCGMetadataAccessor.cpp
@@ -166,14 +166,9 @@ void UPCGMetadataAccessorHelpers::CopyPoint(const FPCGPoint& InPoint, FPCGPoint&
 	OutPoint = InPoint;
 
 	// If we want to copy the metadata, then at least the out metadata must not be null.
-	if (bCopyMetadata && OutMetadata)
+	if (bCopyMetadata)
 	{
-		// If we have an input metadata, then we can copy values as needed,
-		// Otherwise, we will assume that the point is parented
-		if (InMetadata)
-		{
-			OutMetadata->SetPointAttributes(InPoint, InMetadata, OutPoint);
-		}
+		InitializeMetadata(OutPoint, OutMetadata, InPoint, InMetadata);
 	}
 	else
 	{
@@ -186,12 +181,29 @@ void UPCGMetadataAccessorHelpers::InitializeMetadata(FPCGPoint& Point, UPCGMetad
 	Point.MetadataEntry = Metadata ? Metadata->AddEntry() : PCGInvalidEntryKey;
 }
 
-void UPCGMetadataAccessorHelpers::InitializeMetadata(FPCGPoint& Point, UPCGMetadata* Metadata, const FPCGPoint& ParentPoint)
+void UPCGMetadataAccessorHelpers::InitializeMetadata(FPCGPoint& Point, UPCGMetadata* Metadata, const FPCGPoint& ParentPoint, const UPCGMetadata* ParentMetadata)
 {
-	Point.MetadataEntry = Metadata ? Metadata->AddEntry(ParentPoint.MetadataEntry) : PCGInvalidEntryKey;
+	if (Metadata)
+	{
+		// If we're not given the parent metadata, we'll assume it is the current metadata's parent
+		if (!ParentMetadata || ParentMetadata->HasParent(ParentMetadata))
+		{
+			Metadata->AddEntry(ParentPoint.MetadataEntry);
+		}
+		else
+		{
+			// Conceptual parent isn't in the metadata hierarchy, therefore we must set attributes if any
+			Point.MetadataEntry = Metadata->AddEntry();
+			Metadata->SetPointAttributes(ParentPoint, ParentMetadata, Point);
+		}
+	}
+	else
+	{
+		Point.MetadataEntry = PCGInvalidEntryKey;
+	}
 }
 
-void UPCGMetadataAccessorHelpers::InitializeMetadata(FPCGPoint& Point, UPCGMetadata* Metadata, const FPCGPoint& ParentPoint, const UPCGMetadata* ParentMetadata)
+void UPCGMetadataAccessorHelpers::InitializeMetadataWithParent(FPCGPoint& Point, UPCGMetadata* Metadata, const FPCGPoint& ParentPoint, const UPCGMetadata* ParentMetadata)
 {
 	Point.MetadataEntry = Metadata ? (Metadata->HasParent(ParentMetadata) ? Metadata->AddEntry(ParentPoint.MetadataEntry) : Metadata->AddEntry()) : PCGInvalidEntryKey;
 }
diff --git a/Engine/Plugins/Experimental/PCG/Source/PCG/Public/Metadata/PCGMetadataAccessor.h b/Engine/Plugins/Experimental/PCG/Source/PCG/Public/Metadata/PCGMetadataAccessor.h
--- a/Engine/Plugins/Experimental/PCG/Source/PCG/Public/Metadata/PCGMetadataAccessor.h
+++ b/Engine/Plugins/Experimental/PCG/Source/PCG/Public/Metadata/PCGMetadataAccessor.h
@@ -68,10 +68,12 @@ public:
 	static void CopyPoint(const FPCGPoint& InPoint, FPCGPoint& OutPoint, bool bCopyMetadata = true, const UPCGMetadata* InMetadata = nullptr, UPCGMetadata* OutMetadata = nullptr);
 
 	static void InitializeMetadata(FPCGPoint& Point, UPCGMetadata* Metadata);
-	static void InitializeMetadata(FPCGPoint& Point, UPCGMetadata* Metadata, const FPCGPoint& ParentPoint, const UPCGMetadata* ParentMetadata);
+	/** Assigns a metadata entry but does not copy values if from a non-parented metadata */
+	static void InitializeMetadataWithParent(FPCGPoint& Point, UPCGMetadata* Metadata, const FPCGPoint& ParentPoint, const UPCGMetadata* ParentMetadata);
 
+	/** Assigns a metadata entry key and will copy attribute values if from an unrelated metadata. Note: a null ParentMetadata assumes this is the same as Metadata */
 	UFUNCTION(BlueprintCallable, Category = "PCG|Metadata", meta = (ScriptMethod))
-	static void InitializeMetadata(UPARAM(ref) FPCGPoint& Point, UPCGMetadata* Metadata, const FPCGPoint& ParentPoint);
+	static void InitializeMetadata(UPARAM(ref) FPCGPoint& Point, UPCGMetadata* Metadata, const FPCGPoint& ParentPoint, const UPCGMetadata* ParentMetadata = nullptr);
 
 	UFUNCTION(BlueprintCallable, Category = "PCG|Metadata", meta = (ScriptMethod))
 	static int64 GetInteger64Attribute(const FPCGPoint& Point, const UPCGMetadata* Metadata, FName AttributeName);
