Add support for comments in PCGGraph

diff --git a/Engine/Plugins/Experimental/PCG/Source/PCG/Private/PCGGraph.cpp b/Engine/Plugins/Experimental/PCG/Source/PCG/Private/PCGGraph.cpp
--- a/Engine/Plugins/Experimental/PCG/Source/PCG/Private/PCGGraph.cpp
+++ b/Engine/Plugins/Experimental/PCG/Source/PCG/Private/PCGGraph.cpp
@@ -85,6 +85,12 @@ void UPCGGraph::PostLoad()
 		Node->ConditionalPostLoad();
 	}
 	
+	// Also do this for ExtraNodes
+	for (UObject* ExtraNode : ExtraEditorNodes)
+	{
+		ExtraNode->ConditionalPostLoad();
+	}
+
 	// Update pins on all nodes
 	InputNode->UpdatePins();
 	OutputNode->UpdatePins();
@@ -102,9 +108,7 @@ void UPCGGraph::PostLoad()
 	{
 		Node->ApplyDeprecation();
 	}
-#endif
 
-#if WITH_EDITOR
 	InputNode->OnNodeChangedDelegate.AddUObject(this, &UPCGGraph::OnNodeChanged);
 	OutputNode->OnNodeChangedDelegate.AddUObject(this, &UPCGGraph::OnNodeChanged);
 
@@ -464,9 +468,17 @@ void UPCGGraph::EnableNotificationsForEditor()
 		bDelayedChangeNotificationStructural = false;
 	}
 }
-#endif
 
-#if WITH_EDITOR
+void UPCGGraph::SetExtraEditorNodes(const TArray<TObjectPtr<const UObject>>& InNodes)
+{
+	ExtraEditorNodes.Empty();
+
+	for (const UObject* Node : InNodes)
+	{
+		ExtraEditorNodes.Add(DuplicateObject(Node, this));
+	}
+}
+
 FPCGTagToSettingsMap UPCGGraph::GetTrackedTagsToSettings() const
 {
 	FPCGTagToSettingsMap TagsToSettings;
diff --git a/Engine/Plugins/Experimental/PCG/Source/PCG/Public/PCGGraph.h b/Engine/Plugins/Experimental/PCG/Source/PCG/Public/PCGGraph.h
--- a/Engine/Plugins/Experimental/PCG/Source/PCG/Public/PCGGraph.h
+++ b/Engine/Plugins/Experimental/PCG/Source/PCG/Public/PCGGraph.h
@@ -78,6 +78,9 @@ public:
 	void EnableNotificationsForEditor();
 	void PreNodeUndo(UPCGNode* InPCGNode);
 	void PostNodeUndo(UPCGNode* InPCGNode);
+
+	const TArray<TObjectPtr<UObject>>& GetExtraEditorNodes() const { return ExtraEditorNodes; }
+	void SetExtraEditorNodes(const TArray<TObjectPtr<const UObject>>& InNodes);
 #endif
 
 #if WITH_EDITOR
@@ -102,6 +105,12 @@ protected:
 	UPROPERTY(BlueprintReadOnly, VisibleAnywhere, Category = Graph)
 	TObjectPtr<UPCGNode> OutputNode;
 
+#if WITH_EDITORONLY_DATA
+	// Extra data to hold information that is useful only in editor (like comments)
+	UPROPERTY()
+	TArray<TObjectPtr<UObject>> ExtraEditorNodes;
+#endif // WITH_EDITORONLY_DATA
+
 #if WITH_EDITOR
 private:
 	void NotifyGraphChanged(bool bIsStructural);
diff --git a/Engine/Plugins/Experimental/PCG/Source/PCG/Public/PCGNode.h b/Engine/Plugins/Experimental/PCG/Source/PCG/Public/PCGNode.h
--- a/Engine/Plugins/Experimental/PCG/Source/PCG/Public/PCGNode.h
+++ b/Engine/Plugins/Experimental/PCG/Source/PCG/Public/PCGNode.h
@@ -102,6 +102,15 @@ public:
 
 	UPROPERTY()
 	int32 PositionY;
+
+	UPROPERTY()
+	FString NodeComment;
+
+	UPROPERTY()
+	uint8 bCommentBubblePinned : 1;
+
+	UPROPERTY()
+	uint8 bCommentBubbleVisible : 1;
 #endif // WITH_EDITORONLY_DATA
 
 protected:
diff --git a/Engine/Plugins/Experimental/PCG/Source/PCGEditor/Private/PCGEditor.cpp b/Engine/Plugins/Experimental/PCG/Source/PCGEditor/Private/PCGEditor.cpp
--- a/Engine/Plugins/Experimental/PCG/Source/PCGEditor/Private/PCGEditor.cpp
+++ b/Engine/Plugins/Experimental/PCG/Source/PCGEditor/Private/PCGEditor.cpp
@@ -35,6 +35,8 @@
 #include "Editor/UnrealEdEngine.h"
 #include "Preferences/UnrealEdOptions.h"
 #include "Widgets/Notifications/SNotificationList.h"
+#include "PCGEditorGraphSchemaActions.h"
+#include "GraphEditAction.h"
 
 #define LOCTEXT_NAMESPACE "PCGGraphEditor"
 
@@ -274,18 +276,28 @@ void FPCGEditor::DeleteSelectedNodes()
 
 			for (UObject* Object : GraphEditorWidget->GetSelectedNodes())
 			{
-				UPCGEditorGraphNodeBase* PCGEditorGraphNode = CastChecked<UPCGEditorGraphNodeBase>(Object);
-
-				if (PCGEditorGraphNode->CanUserDeleteNode())
+				if (UPCGEditorGraphNodeBase* PCGEditorGraphNode = Cast<UPCGEditorGraphNodeBase>(Object))
 				{
-					UPCGNode* PCGNode = PCGEditorGraphNode->GetPCGNode();
-					check(PCGNode);
-
-					PCGGraph->RemoveNode(PCGNode);
-					PCGEditorGraphNode->DestroyNode();
-					bChanged = true;
+					if (PCGEditorGraphNode->CanUserDeleteNode())
+					{
+						UPCGNode* PCGNode = PCGEditorGraphNode->GetPCGNode();
+						check(PCGNode);
+
+						PCGGraph->RemoveNode(PCGNode);
+						PCGEditorGraphNode->DestroyNode();
+						bChanged = true;
+					}
+				}
+				else if (UEdGraphNode* GraphNode = Cast<UEdGraphNode>(Object))
+				{
+					if (GraphNode->CanUserDeleteNode())
+					{
+						GraphNode->DestroyNode();
+						bChanged = true;
+					}
 				}
 			}
+
 			PCGEditorGraph->Modify();
 		}
 
@@ -304,9 +316,9 @@ bool FPCGEditor::CanDeleteSelectedNodes() const
 	{
 		for (UObject* Object : GraphEditorWidget->GetSelectedNodes())
 		{
-			UPCGEditorGraphNodeBase* PCGEditorGraphNode = CastChecked<UPCGEditorGraphNodeBase>(Object);
+			UEdGraphNode* GraphNode = CastChecked<UEdGraphNode>(Object);
 
-			if (PCGEditorGraphNode->CanUserDeleteNode())
+			if (GraphNode->CanUserDeleteNode())
 			{
 				return true;
 			}
@@ -335,8 +347,10 @@ void FPCGEditor::CopySelectedNodes()
 
 		for (UObject* SelectedNode : SelectedNodes)
 		{
-			UPCGEditorGraphNodeBase* PCGGraphNode = CastChecked<UPCGEditorGraphNodeBase>(SelectedNode);
-			PCGGraphNode->PostCopy();
+			if (UPCGEditorGraphNodeBase* PCGGraphNode = Cast<UPCGEditorGraphNodeBase>(SelectedNode))
+			{
+				PCGGraphNode->PostCopy();
+			}
 		}
 	}
 }
@@ -347,9 +361,9 @@ bool FPCGEditor::CanCopySelectedNodes() const
 	{
 		for (UObject* Object : GraphEditorWidget->GetSelectedNodes())
 		{
-			if (UPCGEditorGraphNodeBase* PCGEditorGraphNode = Cast<UPCGEditorGraphNodeBase>(Object))
+			if (UEdGraphNode* GraphNode = CastChecked<UEdGraphNode>(Object))
 			{
-				if (PCGEditorGraphNode->CanDuplicateNode())
+				if (GraphNode->CanDuplicateNode())
 				{
 					return true;
 				}
@@ -539,6 +553,23 @@ void FPCGEditor::OnDistributeNodesV()
 	}
 }
 
+void FPCGEditor::OnCreateComment()
+{
+	if (PCGEditorGraph)
+	{
+		FPCGEditorGraphSchemaAction_NewComment CommentAction;
+
+		TSharedPtr<SGraphEditor> GraphEditorPtr = SGraphEditor::FindGraphEditorForGraph(PCGEditorGraph);
+		FVector2D Location;
+		if (GraphEditorPtr)
+		{
+			Location = GraphEditorPtr->GetPasteLocation();
+		}
+
+		CommentAction.PerformAction(PCGEditorGraph, nullptr, Location);
+	}
+}
+
 TSharedRef<SGraphEditor> FPCGEditor::CreateGraphEditorWidget()
 {
 	GraphEditorCommands = MakeShareable(new FUICommandList);
@@ -597,6 +628,10 @@ TSharedRef<SGraphEditor> FPCGEditor::CreateGraphEditorWidget()
 		FExecuteAction::CreateSP(this, &FPCGEditor::OnStraightenConnections)
 	);
 
+	GraphEditorCommands->MapAction(FGraphEditorCommands::Get().CreateComment,
+		FExecuteAction::CreateSP(this, &FPCGEditor::OnCreateComment)
+	);
+
 	// Distribution Commands
 	GraphEditorCommands->MapAction(FGraphEditorCommands::Get().DistributeNodesHorizontally,
 		FExecuteAction::CreateSP(this, &FPCGEditor::OnDistributeNodesH)
@@ -623,6 +658,44 @@ TSharedRef<SGraphEditor> FPCGEditor::CreateGraphEditorWidget()
 		.ShowGraphStateOverlay(false);
 }
 
+void FPCGEditor::ReplicateExtraNodes() const
+{
+	if (PCGEditorGraph)
+	{
+		if (UPCGGraph* PCGGraph = PCGEditorGraph->GetPCGGraph())
+		{
+			TArray<TObjectPtr<const UObject>> ExtraNodes;
+			for (const UEdGraphNode* GraphNode : PCGEditorGraph->Nodes)
+			{
+				check(GraphNode);
+				if (!GraphNode->IsA<UPCGEditorGraphNodeBase>())
+				{
+					ExtraNodes.Add(GraphNode);
+				}
+
+			}
+
+			PCGGraph->SetExtraEditorNodes(ExtraNodes);
+		}
+	}
+}
+
+void FPCGEditor::SaveAsset_Execute()
+{
+	// Extra nodes are replicated on asset save, to be saved in the underlying PCGGraph
+	ReplicateExtraNodes();
+
+	FAssetEditorToolkit::SaveAsset_Execute();
+}
+
+void FPCGEditor::OnClose()
+{
+	// Extra nodes are replicated on editor close, to be saved in the underlying PCGGraph
+	ReplicateExtraNodes();
+
+	FAssetEditorToolkit::OnClose();
+}
+
 TSharedRef<SPCGEditorGraphNodePalette> FPCGEditor::CreatePaletteWidget()
 {
 	return SNew(SPCGEditorGraphNodePalette);
@@ -645,13 +718,17 @@ void FPCGEditor::OnSelectedNodesChanged(const TSet<UObject*>& NewSelection)
 	{
 		for (UObject* Object : NewSelection)
 		{
-			if (UPCGEditorGraphNodeBase* GraphNode = Cast<UPCGEditorGraphNodeBase>(Object))
+			if (UPCGEditorGraphNodeBase* PCGGraphNode = Cast<UPCGEditorGraphNodeBase>(Object))
 			{
-				if (UPCGNode* PCGNode = GraphNode->GetPCGNode())
+				if (UPCGNode* PCGNode = PCGGraphNode->GetPCGNode())
 				{
 					SelectedObjects.Add(PCGNode->DefaultSettings);
 				}
 			}
+			else if (UEdGraphNode* GraphNode = Cast<UEdGraphNode>(Object))
+			{
+				SelectedObjects.Add(GraphNode);
+			}
 		}
 	}
 
diff --git a/Engine/Plugins/Experimental/PCG/Source/PCGEditor/Private/PCGEditor.h b/Engine/Plugins/Experimental/PCG/Source/PCGEditor/Private/PCGEditor.h
--- a/Engine/Plugins/Experimental/PCG/Source/PCGEditor/Private/PCGEditor.h
+++ b/Engine/Plugins/Experimental/PCG/Source/PCGEditor/Private/PCGEditor.h
@@ -41,6 +41,13 @@ public:
 	virtual FText GetBaseToolkitName() const override;
 	virtual FLinearColor GetWorldCentricTabColorScale() const override;
 	virtual FString GetWorldCentricTabPrefix() const override;
+	virtual void OnClose() override;
+	// ~End FAssetEditorToolkit interface
+
+protected:
+	// ~Begin FAssetEditorToolkit interface
+	/** Called when "Save" is clicked for this asset */
+	virtual void SaveAsset_Execute() override;
 	// ~End FAssetEditorToolkit interface
 
 private:
@@ -91,6 +98,7 @@ private:
 	void OnStraightenConnections();
 	void OnDistributeNodesH();
 	void OnDistributeNodesV();
+	void OnCreateComment();
 
 	/** Create new graph editor widget */
 	TSharedRef<SGraphEditor> CreateGraphEditorWidget();
@@ -121,6 +129,9 @@ private:
 	 */
 	void JumpToDefinition(const UClass* Class) const;
 
+	/** To be called everytime we need to replicate our extra nodes to the underlying PCGGraph */
+	void ReplicateExtraNodes() const;
+
 	TSharedRef<SDockTab> SpawnTab_GraphEditor(const FSpawnTabArgs& Args);
 	TSharedRef<SDockTab> SpawnTab_PropertyDetails(const FSpawnTabArgs& Args);
 	TSharedRef<SDockTab> SpawnTab_Palette(const FSpawnTabArgs& Args);
diff --git a/Engine/Plugins/Experimental/PCG/Source/PCGEditor/Private/PCGEditorGraph.cpp b/Engine/Plugins/Experimental/PCG/Source/PCGEditor/Private/PCGEditorGraph.cpp
--- a/Engine/Plugins/Experimental/PCG/Source/PCGEditor/Private/PCGEditorGraph.cpp
+++ b/Engine/Plugins/Experimental/PCG/Source/PCGEditor/Private/PCGEditorGraph.cpp
@@ -47,6 +47,16 @@ void UPCGEditorGraph::InitFromNodeGraph(UPCGGraph* InPCGGraph)
 		UPCGEditorGraphNodeBase* GraphNode = NodeLookupIt.Value;
 		CreateLinks(GraphNode, /*bCreateInbound=*/false, /*bCreateOutbound=*/true, NodeLookup);
 	}
+
+	for (const UObject* ExtraNode : PCGGraph->GetExtraEditorNodes())
+	{
+		if (const UEdGraphNode* ExtraGraphNode = Cast<UEdGraphNode>(ExtraNode))
+		{
+			UEdGraphNode* NewNode = DuplicateObject(ExtraGraphNode, /*Outer=*/this);
+			const bool bIsUserAction = false;
+			AddNode(NewNode, bIsUserAction, bSelectNewNode);
+		}
+	}
 }
 
 void UPCGEditorGraph::CreateLinks(UPCGEditorGraphNodeBase* GraphNode, bool bCreateInbound, bool bCreateOutbound)
diff --git a/Engine/Plugins/Experimental/PCG/Source/PCGEditor/Private/PCGEditorGraphNodeBase.cpp b/Engine/Plugins/Experimental/PCG/Source/PCGEditor/Private/PCGEditorGraphNodeBase.cpp
--- a/Engine/Plugins/Experimental/PCG/Source/PCGEditor/Private/PCGEditorGraphNodeBase.cpp
+++ b/Engine/Plugins/Experimental/PCG/Source/PCGEditor/Private/PCGEditorGraphNodeBase.cpp
@@ -31,6 +31,23 @@ void UPCGEditorGraphNodeBase::BeginDestroy()
 	Super::BeginDestroy();
 }
 
+void UPCGEditorGraphNodeBase::PostTransacted(const FTransactionObjectEvent& TransactionEvent)
+{
+	Super::PostTransacted(TransactionEvent);
+
+	TArray<FName> PropertiesChanged = TransactionEvent.GetChangedProperties();
+
+	if (PropertiesChanged.Contains(TEXT("bCommentBubblePinned")))
+	{
+		UpdateCommentBubblePinned();
+	}
+
+	if (PropertiesChanged.Contains(TEXT("NodePosX")) || PropertiesChanged.Contains(TEXT("NodePosY")))
+	{
+		UpdatePosition();
+	}
+}
+
 void UPCGEditorGraphNodeBase::Construct(UPCGNode* InPCGNode, EPCGEditorGraphNodeType InNodeType)
 {
 	check(InPCGNode);
@@ -39,6 +56,9 @@ void UPCGEditorGraphNodeBase::Construct(UPCGNode* InPCGNode, EPCGEditorGraphNode
 
 	NodePosX = InPCGNode->PositionX;
 	NodePosY = InPCGNode->PositionY;
+	NodeComment = InPCGNode->NodeComment;
+	bCommentBubblePinned = InPCGNode->bCommentBubblePinned;
+	bCommentBubbleVisible = InPCGNode->bCommentBubbleVisible;
 	
 	NodeType = InNodeType;
 
@@ -419,4 +439,45 @@ UObject* UPCGEditorGraphNodeBase::GetJumpTargetForDoubleClick() const
 	return nullptr;
 }
 
+void UPCGEditorGraphNodeBase::OnUpdateCommentText(const FString& NewComment)
+{
+	Super::OnUpdateCommentText(NewComment);
+
+	if (PCGNode && PCGNode->NodeComment != NewComment)
+	{
+		PCGNode->Modify();
+		PCGNode->NodeComment = NewComment;
+	}
+}
+
+void UPCGEditorGraphNodeBase::OnCommentBubbleToggled(bool bInCommentBubbleVisible)
+{
+	Super::OnCommentBubbleToggled(bInCommentBubbleVisible);
+
+	if (PCGNode && PCGNode->bCommentBubbleVisible != bInCommentBubbleVisible)
+	{
+		PCGNode->Modify();
+		PCGNode->bCommentBubbleVisible = bInCommentBubbleVisible;
+	}
+}
+
+void UPCGEditorGraphNodeBase::UpdateCommentBubblePinned()
+{
+	if (PCGNode)
+	{
+		PCGNode->Modify();
+		PCGNode->bCommentBubblePinned = bCommentBubblePinned;
+	}
+}
+
+void UPCGEditorGraphNodeBase::UpdatePosition()
+{
+	if (PCGNode)
+	{
+		PCGNode->Modify();
+		PCGNode->PositionX = NodePosX;
+		PCGNode->PositionY = NodePosY;
+	}
+}
+
 #undef LOCTEXT_NAMESPACE
diff --git a/Engine/Plugins/Experimental/PCG/Source/PCGEditor/Private/PCGEditorGraphNodeBase.h b/Engine/Plugins/Experimental/PCG/Source/PCGEditor/Private/PCGEditorGraphNodeBase.h
--- a/Engine/Plugins/Experimental/PCG/Source/PCGEditor/Private/PCGEditorGraphNodeBase.h
+++ b/Engine/Plugins/Experimental/PCG/Source/PCGEditor/Private/PCGEditorGraphNodeBase.h
@@ -32,7 +32,8 @@ public:
 	void Construct(UPCGNode* InPCGNode, EPCGEditorGraphNodeType InNodeType);
 
 	// ~Begin UObject interface
-	void BeginDestroy() override;
+	virtual void BeginDestroy() override;
+	virtual void PostTransacted(const FTransactionObjectEvent& TransactionEvent) override;
 	// ~End UObject interface
 
 	// ~Begin UEdGraphNode interface
@@ -46,6 +47,8 @@ public:
 	virtual FText GetTooltipText() const override;
 	virtual void GetPinHoverText(const UEdGraphPin& Pin, FString& HoverTextOut) const override;
 	virtual UObject* GetJumpTargetForDoubleClick() const override;
+	virtual void OnUpdateCommentText(const FString& NewComment);
+	virtual void OnCommentBubbleToggled(bool bInCommentBubbleVisible) override;
 	// ~End UEdGraphNode interface
 
 	UPCGNode* GetPCGNode() { return PCGNode; }
@@ -64,6 +67,8 @@ protected:
 	void OnNodeChanged(UPCGNode* InNode, EPCGChangeType ChangeType);
 	void OnPickColor();
 	void OnColorPicked(FLinearColor NewColor);
+	void UpdateCommentBubblePinned();
+	void UpdatePosition();
 
 	UPROPERTY()
 	TObjectPtr<UPCGNode> PCGNode = nullptr;
diff --git a/Engine/Plugins/Experimental/PCG/Source/PCGEditor/Private/PCGEditorGraphSchema.cpp b/Engine/Plugins/Experimental/PCG/Source/PCGEditor/Private/PCGEditorGraphSchema.cpp
--- a/Engine/Plugins/Experimental/PCG/Source/PCGEditor/Private/PCGEditorGraphSchema.cpp
+++ b/Engine/Plugins/Experimental/PCG/Source/PCGEditor/Private/PCGEditorGraphSchema.cpp
@@ -25,6 +25,7 @@ void UPCGEditorGraphSchema::GetPaletteActions(FGraphActionMenuBuilder& ActionMen
 	GetNativeElementActions(ActionMenuBuilder);
 	GetSubgraphElementActions(ActionMenuBuilder);
 	GetBlueprintElementActions(ActionMenuBuilder);
+	GetExtraElementActions(ActionMenuBuilder);
 }
 
 void UPCGEditorGraphSchema::GetGraphContextActions(FGraphContextMenuBuilder& ContextMenuBuilder) const
@@ -34,6 +35,7 @@ void UPCGEditorGraphSchema::GetGraphContextActions(FGraphContextMenuBuilder& Con
 	GetNativeElementActions(ContextMenuBuilder);
 	GetSubgraphElementActions(ContextMenuBuilder);
 	GetBlueprintElementActions(ContextMenuBuilder);
+	GetExtraElementActions(ContextMenuBuilder);
 }
 
 FLinearColor UPCGEditorGraphSchema::GetPinTypeColor(const FEdGraphPinType& PinType) const
@@ -266,6 +268,17 @@ void UPCGEditorGraphSchema::GetSubgraphElementActions(FGraphActionMenuBuilder& A
 	}
 }
 
+void UPCGEditorGraphSchema::GetExtraElementActions(FGraphActionMenuBuilder& ActionMenuBuilder) const
+{
+	// Comment action
+	const FText MenuDesc = LOCTEXT("PCGAddComment", "Add Comment...");
+	const FText Category;
+	const FText Description = LOCTEXT("PCGAddCommentTooltip", "Create a resizable comment box.");
+
+	TSharedPtr<FPCGEditorGraphSchemaAction_NewComment> NewCommentAction(new FPCGEditorGraphSchemaAction_NewComment(Category, MenuDesc, Description, 0));
+	ActionMenuBuilder.AddAction(NewCommentAction);
+}
+
 void UPCGEditorGraphSchema::DroppedAssetsOnGraph(const TArray<FAssetData>& Assets, const FVector2D& GraphPosition, UEdGraph* Graph) const
 {
 	FVector2D GraphPositionOffset = GraphPosition;
diff --git a/Engine/Plugins/Experimental/PCG/Source/PCGEditor/Private/PCGEditorGraphSchema.h b/Engine/Plugins/Experimental/PCG/Source/PCGEditor/Private/PCGEditorGraphSchema.h
--- a/Engine/Plugins/Experimental/PCG/Source/PCGEditor/Private/PCGEditorGraphSchema.h
+++ b/Engine/Plugins/Experimental/PCG/Source/PCGEditor/Private/PCGEditorGraphSchema.h
@@ -33,6 +33,7 @@ private:
 	void GetNativeElementActions(FGraphActionMenuBuilder& ActionMenuBuilder) const;
 	void GetBlueprintElementActions(FGraphActionMenuBuilder& ActionMenuBuilder) const;
 	void GetSubgraphElementActions(FGraphActionMenuBuilder& ActionMenuBuilder) const;
+	void GetExtraElementActions(FGraphActionMenuBuilder& ActionMenuBuilder) const;
 };
 
 class FPCGEditorConnectionDrawingPolicy : public FConnectionDrawingPolicy
diff --git a/Engine/Plugins/Experimental/PCG/Source/PCGEditor/Private/PCGEditorGraphSchemaActions.cpp b/Engine/Plugins/Experimental/PCG/Source/PCGEditor/Private/PCGEditorGraphSchemaActions.cpp
--- a/Engine/Plugins/Experimental/PCG/Source/PCGEditor/Private/PCGEditorGraphSchemaActions.cpp
+++ b/Engine/Plugins/Experimental/PCG/Source/PCGEditor/Private/PCGEditorGraphSchemaActions.cpp
@@ -13,6 +13,7 @@
 #include "EdGraph/EdGraph.h"
 #include "Engine/Blueprint.h"
 #include "ScopedTransaction.h"
+#include "EdGraphNode_Comment.h"
 
 #define LOCTEXT_NAMESPACE "PCGEditorGraphSchemaActions"
 
@@ -147,4 +148,30 @@ UEdGraphNode* FPCGEditorGraphSchemaAction_NewSubgraphElement::PerformAction(UEdG
 	return NewNode;
 }
 
+UEdGraphNode* FPCGEditorGraphSchemaAction_NewComment::PerformAction(class UEdGraph* ParentGraph, UEdGraphPin* FromPin, const FVector2D Location, bool bSelectNewNode /*= true*/)
+{
+	UPCGEditorGraph* EditorGraph = Cast<UPCGEditorGraph>(ParentGraph);
+	if (!EditorGraph)
+	{
+		return nullptr;
+	}
+
+	UEdGraphNode_Comment* CommentTemplate = NewObject<UEdGraphNode_Comment>();
+
+	TSharedPtr<SGraphEditor> GraphEditorPtr = SGraphEditor::FindGraphEditorForGraph(ParentGraph);
+
+	FVector2D SpawnLocation = Location;
+	FSlateRect Bounds;
+	if (GraphEditorPtr && GraphEditorPtr->GetBoundsForSelectedNodes(Bounds, 50.0f))
+	{
+		CommentTemplate->SetBounds(Bounds);
+		SpawnLocation.X = CommentTemplate->NodePosX;
+		SpawnLocation.Y = CommentTemplate->NodePosY;
+	}
+
+	UEdGraphNode_Comment* NewNode = FEdGraphSchemaAction_NewNode::SpawnNodeFromTemplate<UEdGraphNode_Comment>(ParentGraph, CommentTemplate, SpawnLocation, bSelectNewNode);
+
+	return NewNode;
+}
+
 #undef LOCTEXT_NAMESPACE
diff --git a/Engine/Plugins/Experimental/PCG/Source/PCGEditor/Private/PCGEditorGraphSchemaActions.h b/Engine/Plugins/Experimental/PCG/Source/PCGEditor/Private/PCGEditorGraphSchemaActions.h
--- a/Engine/Plugins/Experimental/PCG/Source/PCGEditor/Private/PCGEditorGraphSchemaActions.h
+++ b/Engine/Plugins/Experimental/PCG/Source/PCGEditor/Private/PCGEditorGraphSchemaActions.h
@@ -60,3 +60,21 @@ struct FPCGEditorGraphSchemaAction_NewSubgraphElement : public FEdGraphSchemaAct
 	virtual UEdGraphNode* PerformAction(UEdGraph* ParentGraph, UEdGraphPin* FromPin, const FVector2D Location, bool bSelectNewNode = true) override;
 	//~ End FEdGraphSchemaAction Interface
 };
+
+/** Action to add a 'comment' node to the graph */
+USTRUCT()
+struct FPCGEditorGraphSchemaAction_NewComment : public FEdGraphSchemaAction
+{
+	GENERATED_BODY()
+
+	// Inherit the base class's constructors
+	using FEdGraphSchemaAction::FEdGraphSchemaAction;
+
+	// Simple type info
+	static FName StaticGetTypeId() { static FName Type("FPCGEditorGraphSchemaAction_NewComment"); return Type; }
+	virtual FName GetTypeId() const override { return StaticGetTypeId(); }
+
+	// FEdGraphSchemaAction interface
+	virtual UEdGraphNode* PerformAction(class UEdGraph* ParentGraph, UEdGraphPin* FromPin, const FVector2D Location, bool bSelectNewNode = true) override;
+	// End of FEdGraphSchemaAction interface
+};
diff --git a/Engine/Plugins/Experimental/PCG/Source/PCGEditor/Private/SPCGEditorGraphNode.cpp b/Engine/Plugins/Experimental/PCG/Source/PCGEditor/Private/SPCGEditorGraphNode.cpp
--- a/Engine/Plugins/Experimental/PCG/Source/PCGEditor/Private/SPCGEditorGraphNode.cpp
+++ b/Engine/Plugins/Experimental/PCG/Source/PCGEditor/Private/SPCGEditorGraphNode.cpp
@@ -21,19 +21,6 @@ void SPCGEditorGraphNode::Construct(const FArguments& InArgs, UPCGEditorGraphNod
 	UpdateGraphNode();
 }
 
-void SPCGEditorGraphNode::MoveTo(const FVector2D& NewPosition, FNodeSet& NodeFilter, bool bMarkDirty /*= true*/)
-{
-	SGraphNode::MoveTo(NewPosition, NodeFilter, bMarkDirty);
-
-	check(PCGEditorGraphNode);
-	if (UPCGNode* PCGNode = PCGEditorGraphNode->GetPCGNode())
-	{
-		PCGNode->Modify();
-		PCGNode->PositionX = PCGEditorGraphNode->NodePosX;
-		PCGNode->PositionY = PCGEditorGraphNode->NodePosY;
-	}
-}
-
 void SPCGEditorGraphNode::AddPin(const TSharedRef<SGraphPin>& PinToAdd)
 {
 	check(PCGEditorGraphNode);
diff --git a/Engine/Plugins/Experimental/PCG/Source/PCGEditor/Private/SPCGEditorGraphNode.h b/Engine/Plugins/Experimental/PCG/Source/PCGEditor/Private/SPCGEditorGraphNode.h
--- a/Engine/Plugins/Experimental/PCG/Source/PCGEditor/Private/SPCGEditorGraphNode.h
+++ b/Engine/Plugins/Experimental/PCG/Source/PCGEditor/Private/SPCGEditorGraphNode.h
@@ -16,7 +16,6 @@ public:
 	void Construct(const FArguments& InArgs, UPCGEditorGraphNodeBase* InNode);
 
 	//~ Begin SGraphNode Interface
-	virtual void MoveTo(const FVector2D& NewPosition, FNodeSet& NodeFilter, bool bMarkDirty = true) override;
 	virtual void RequestRenameOnSpawn() override { /* Empty to avoid the default behavior to rename on node spawn */ }
 	virtual void AddPin(const TSharedRef<SGraphPin>& PinToAdd) override;
 	//~ End SGraphNode Interface
