Fixed duplication issue where the managed resources would not be setup properly after duplication, linking different actors together.

#rnx

diff --git a/Engine/Plugins/Experimental/PCG/Source/PCG/Private/PCGManagedResource.cpp b/Engine/Plugins/Experimental/PCG/Source/PCG/Private/PCGManagedResource.cpp
--- a/Engine/Plugins/Experimental/PCG/Source/PCG/Private/PCGManagedResource.cpp
+++ b/Engine/Plugins/Experimental/PCG/Source/PCG/Private/PCGManagedResource.cpp
@@ -5,6 +5,14 @@
 
 #include "Components/InstancedStaticMeshComponent.h"
 
+void UPCGManagedActors::PostEditImport()
+{
+	// In this case, the managed actors won't be copied along the actor/component,
+	// So we just have to "forget" the actors.
+	Super::PostEditImport();
+	GeneratedActors.Reset();
+}
+
 bool UPCGManagedActors::Release(bool bHardRelease, TSet<TSoftObjectPtr<AActor>>& OutActorsToDelete)
 {
 	OutActorsToDelete.Append(GeneratedActors);
@@ -36,6 +44,44 @@ bool UPCGManagedActors::ReleaseIfUnused()
 	return GeneratedActors.IsEmpty();
 }
 
+void UPCGManagedComponent::PostEditImport()
+{
+	Super::PostEditImport();
+
+	// Rehook components from the original to the locally duplicated components
+	UPCGComponent* OwningComponent = Cast<UPCGComponent>(GetOuter());
+	AActor* Actor = OwningComponent ? OwningComponent->GetOwner() : nullptr;
+
+	bool bFoundMatch = false;
+
+	if (Actor && GeneratedComponent.IsValid())
+	{
+		TInlineComponentArray<UActorComponent*, 16> Components;
+		Actor->GetComponents(Components);
+
+		for (UActorComponent* Component : Components)
+		{
+			if (Component && Component->GetFName() == GeneratedComponent->GetFName())
+			{
+				GeneratedComponent = Component;
+				bFoundMatch = true;
+			}
+		}
+
+		if (!bFoundMatch)
+		{
+			// Not quite clear what to do when we have a component that cannot be remapped.
+			// Maybe we should check against guids instead?
+			GeneratedComponent.Reset();
+		}
+	}
+	else
+	{
+		// Somewhat irrelevant case, if we don't have an actor or a component, there's not a lot we can do.
+		GeneratedComponent.Reset();
+	}
+}
+
 bool UPCGManagedComponent::Release(bool bHardRelease, TSet<TSoftObjectPtr<AActor>>& OutActorsToDelete)
 {
 	const bool bSupportsComponentReset = SupportsComponentReset();
diff --git a/Engine/Plugins/Experimental/PCG/Source/PCG/Public/PCGManagedResource.h b/Engine/Plugins/Experimental/PCG/Source/PCG/Public/PCGManagedResource.h
--- a/Engine/Plugins/Experimental/PCG/Source/PCG/Public/PCGManagedResource.h
+++ b/Engine/Plugins/Experimental/PCG/Source/PCG/Public/PCGManagedResource.h
@@ -33,6 +33,10 @@ class PCG_API UPCGManagedActors : public UPCGManagedResource
 	GENERATED_BODY()
 
 public:
+	//~Begin UObject interface
+	virtual void PostEditImport() override;
+	//~End UObject interface
+
 	//~Begin UPCGManagedResource interface
 	virtual bool Release(bool bHardRelease, TSet<TSoftObjectPtr<AActor>>& OutActorsToDelete) override;
 	virtual bool ReleaseIfUnused() override;
@@ -48,6 +52,10 @@ class PCG_API UPCGManagedComponent : public UPCGManagedResource
 	GENERATED_BODY()
 
 public:
+	//~Begin UObject interface
+	virtual void PostEditImport() override;
+	//~End UObject interface
+
 	//~Begin UPCGManagedResource interface
 	virtual bool Release(bool bHardRelease, TSet<TSoftObjectPtr<AActor>>& OutActorsToDelete) override;
 	virtual bool ReleaseIfUnused() override;
