PCG:
*Added node comment text to find result.

diff --git a/Engine/Plugins/Experimental/PCG/Source/PCGEditor/Private/SPCGEditorGraphFind.cpp b/Engine/Plugins/Experimental/PCG/Source/PCGEditor/Private/SPCGEditorGraphFind.cpp
--- a/Engine/Plugins/Experimental/PCG/Source/PCGEditor/Private/SPCGEditorGraphFind.cpp
+++ b/Engine/Plugins/Experimental/PCG/Source/PCGEditor/Private/SPCGEditorGraphFind.cpp
@@ -90,6 +90,20 @@ FText FPCGEditorGraphFindResult::GetCategory() const
 	return FText::GetEmpty();
 }
 
+FText FPCGEditorGraphFindResult::GetComment() const
+{
+	if (GraphNode.IsValid())
+	{
+		const FString NodeComment = GraphNode->NodeComment;
+		if (!NodeComment.IsEmpty())
+		{
+			return FText::Format(LOCTEXT("NodeCommentFmt", "Node Comment:[{0}]"), FText::FromString(NodeComment));
+		}
+	}
+
+	return FText::GetEmpty();
+}
+
 TSharedRef<SWidget> FPCGEditorGraphFindResult::CreateIcon() const
 {
 	FSlateColor IconColor = FSlateColor::UseForeground();
@@ -191,26 +205,28 @@ TSharedRef<ITableRow> SPCGEditorGraphFind::OnGenerateRow(FPCGEditorGraphFindResu
 			.VAlign(VAlign_Center)
 			.AutoWidth()
 			[
-				SNew(SBox)
-				.WidthOverride(450.0f)
-				[
-					SNew(SHorizontalBox)
-					+ SHorizontalBox::Slot()
-					.AutoWidth()
-					[
-						InItem->CreateIcon()
-					]
-					+SHorizontalBox::Slot()
-					.VAlign(VAlign_Center)
-					.AutoWidth()
-					.Padding(2, 0)
-					[
-						SNew(STextBlock)
-						.Text(FText::FromString(InItem->Value))
-						.HighlightText(HighlightText)
-						.ToolTipText(InItem->GetToolTip())
-					]
-				]
+				InItem->CreateIcon()
+			]
+			+SHorizontalBox::Slot()
+			.VAlign(VAlign_Center)
+			.AutoWidth()
+			.Padding(2, 0)
+			[
+				SNew(STextBlock)
+				.Text(FText::FromString(InItem->Value))
+				.HighlightText(HighlightText)
+				.ToolTipText(InItem->GetToolTip())
+			]
+			+ SHorizontalBox::Slot()
+			.FillWidth(1)
+			.HAlign(HAlign_Right)
+			.VAlign(VAlign_Center)
+			.Padding(2, 0)
+			[
+				SNew(STextBlock)
+				.Text(InItem->GetComment()) 
+				.ColorAndOpacity(FLinearColor::Yellow)
+				.HighlightText(HighlightText)
 			]
 		];
 }
diff --git a/Engine/Plugins/Experimental/PCG/Source/PCGEditor/Private/SPCGEditorGraphFind.h b/Engine/Plugins/Experimental/PCG/Source/PCGEditor/Private/SPCGEditorGraphFind.h
--- a/Engine/Plugins/Experimental/PCG/Source/PCGEditor/Private/SPCGEditorGraphFind.h
+++ b/Engine/Plugins/Experimental/PCG/Source/PCGEditor/Private/SPCGEditorGraphFind.h
@@ -36,6 +36,9 @@ struct FPCGEditorGraphFindResult
 	/** Get Category for this search result */
 	FText GetCategory() const;
 
+	/** Get Comment for this search result */
+	FText GetComment() const;
+
 	/** Create an icon to represent the result */
 	TSharedRef<SWidget> CreateIcon() const;
 
