Fix non-unity build error. (Issue 177163)

diff --git a/Engine/Plugins/Experimental/PCG/Source/PCGEditor/Private/PCGEditorGraphSchemaActions.cpp b/Engine/Plugins/Experimental/PCG/Source/PCGEditor/Private/PCGEditorGraphSchemaActions.cpp
--- a/Engine/Plugins/Experimental/PCG/Source/PCGEditor/Private/PCGEditorGraphSchemaActions.cpp
+++ b/Engine/Plugins/Experimental/PCG/Source/PCGEditor/Private/PCGEditorGraphSchemaActions.cpp
@@ -14,6 +14,7 @@
 #include "Engine/Blueprint.h"
 #include "ScopedTransaction.h"
 #include "EdGraphNode_Comment.h"
+#include "GraphEditor.h"
 
 #define LOCTEXT_NAMESPACE "PCGEditorGraphSchemaActions"
 
