Add "Create comment" in context menu for selected nodes in PCGEditor

diff --git a/Engine/Plugins/Experimental/PCG/Source/PCGEditor/Private/PCGEditorGraphNodeBase.cpp b/Engine/Plugins/Experimental/PCG/Source/PCGEditor/Private/PCGEditorGraphNodeBase.cpp
--- a/Engine/Plugins/Experimental/PCG/Source/PCGEditor/Private/PCGEditorGraphNodeBase.cpp
+++ b/Engine/Plugins/Experimental/PCG/Source/PCGEditor/Private/PCGEditorGraphNodeBase.cpp
@@ -106,6 +106,13 @@ void UPCGEditorGraphNodeBase::GetNodeContextMenuActions(UToolMenu* Menu, class U
 			}
 		}));
 	}
+
+	{
+		FToolMenuSection& Section = Menu->AddSection("EdGraphSchemaCommentGroup", LOCTEXT("CommentGroupHeader", "Comment Group"));
+		Section.AddMenuEntry(FGraphEditorCommands::Get().CreateComment,
+			LOCTEXT("MultiCommentDesc", "Create Comment from Selection"),
+			LOCTEXT("CommentToolTip", "Create a resizable comment box around selection."));
+	}
 }
 
 void UPCGEditorGraphNodeBase::AutowireNewNode(UEdGraphPin* FromPin)
