Put back potential refresh on tracked object change even if there are no properties involved (fixes an issue when applying geometry tools changes)

diff --git a/Engine/Plugins/Experimental/PCG/Source/PCG/Private/PCGComponent.cpp b/Engine/Plugins/Experimental/PCG/Source/PCG/Private/PCGComponent.cpp
--- a/Engine/Plugins/Experimental/PCG/Source/PCG/Private/PCGComponent.cpp
+++ b/Engine/Plugins/Experimental/PCG/Source/PCG/Private/PCGComponent.cpp
@@ -880,11 +880,6 @@ void UPCGComponent::OnLandscapeChanged(ALandscapeProxy* Landscape, const FLandsc
 
 void UPCGComponent::OnObjectPropertyChanged(UObject* InObject, FPropertyChangedEvent& InEvent)
 {
-	if(!InEvent.Property)
-	{
-		return;
-	}
-
 	bool bValueNotInteractive = (InEvent.ChangeType != EPropertyChangeType::Interactive);
 	// Special exception for actor tags, as we can't track otherwise an actor "losing" a tag
 	bool bActorTagChange = (InEvent.Property && InEvent.Property->GetFName() == GET_MEMBER_NAME_CHECKED(AActor, Tags));
