Add collapse to subgraph in PCG Editor

diff --git a/Engine/Plugins/Experimental/PCG/Source/PCG/Private/PCGGraph.cpp b/Engine/Plugins/Experimental/PCG/Source/PCG/Private/PCGGraph.cpp
--- a/Engine/Plugins/Experimental/PCG/Source/PCG/Private/PCGGraph.cpp
+++ b/Engine/Plugins/Experimental/PCG/Source/PCG/Private/PCGGraph.cpp
@@ -283,6 +283,19 @@ UPCGNode* UPCGGraph::AddLabeledEdge(UPCGNode* From, const FName& InboundLabel, U
 	return To;
 }
 
+TObjectPtr<UPCGNode> UPCGGraph::ReconstructNewNode(const UPCGNode* InNode)
+{
+	TObjectPtr<UPCGSettings> NewSettings = DuplicateObject(InNode->DefaultSettings, nullptr);
+	TObjectPtr<UPCGNode> NewNode = AddNode(NewSettings);
+	NewSettings->Rename(nullptr, NewNode);
+
+#if WITH_EDITOR
+	InNode->TransferEditorProperties(NewNode);
+#endif // WITH_EDITOR
+
+	return NewNode;
+}
+
 bool UPCGGraph::Contains(UPCGNode* Node) const
 {
 	return Node == InputNode || Node == OutputNode || Nodes.Contains(Node);
diff --git a/Engine/Plugins/Experimental/PCG/Source/PCG/Private/PCGInputOutputSettings.cpp b/Engine/Plugins/Experimental/PCG/Source/PCG/Private/PCGInputOutputSettings.cpp
--- a/Engine/Plugins/Experimental/PCG/Source/PCG/Private/PCGInputOutputSettings.cpp
+++ b/Engine/Plugins/Experimental/PCG/Source/PCG/Private/PCGInputOutputSettings.cpp
@@ -62,4 +62,24 @@ TArray<FPCGPinProperties> UPCGGraphInputOutputSettings::InputPinProperties() con
 TArray<FPCGPinProperties> UPCGGraphInputOutputSettings::OutputPinProperties() const
 {
 	return GetPinProperties();
+}
+
+bool UPCGGraphInputOutputSettings::IsPinAdvanced(const UPCGPin* Pin) const
+{
+	return Pin != nullptr ? StaticAdvancedLabels().Contains(Pin->Properties.Label) : false;
+}
+
+void UPCGGraphInputOutputSettings::SetShowAdvancedPins(bool bValue)
+{
+	if (bValue != bShowAdvancedPins)
+	{
+		Modify();
+		bShowAdvancedPins = bValue;
+	}
+}
+
+void UPCGGraphInputOutputSettings::AddCustomPin(const FPCGPinProperties& NewCustomPinProperties)
+{
+	Modify();
+	CustomPins.Add(NewCustomPinProperties);
 }
\ No newline at end of file
diff --git a/Engine/Plugins/Experimental/PCG/Source/PCG/Private/PCGNode.cpp b/Engine/Plugins/Experimental/PCG/Source/PCG/Private/PCGNode.cpp
--- a/Engine/Plugins/Experimental/PCG/Source/PCG/Private/PCGNode.cpp
+++ b/Engine/Plugins/Experimental/PCG/Source/PCG/Private/PCGNode.cpp
@@ -390,6 +390,15 @@ void UPCGNode::OnSettingsChanged(UPCGSettings* InSettings, EPCGChangeType Change
 	}
 }
 
+void UPCGNode::TransferEditorProperties(UPCGNode* OtherNode) const
+{
+	OtherNode->PositionX = PositionX;
+	OtherNode->PositionY = PositionY;
+	OtherNode->bCommentBubblePinned = bCommentBubblePinned;
+	OtherNode->bCommentBubbleVisible = bCommentBubbleVisible;
+	OtherNode->NodeComment = NodeComment;
+}
+
 #endif // WITH_EDITOR
 
 void UPCGNode::UpdateAfterSettingsChangeDuringCreation()
diff --git a/Engine/Plugins/Experimental/PCG/Source/PCG/Public/PCGGraph.h b/Engine/Plugins/Experimental/PCG/Source/PCG/Public/PCGGraph.h
--- a/Engine/Plugins/Experimental/PCG/Source/PCG/Public/PCGGraph.h
+++ b/Engine/Plugins/Experimental/PCG/Source/PCG/Public/PCGGraph.h
@@ -63,6 +63,9 @@ public:
 	UFUNCTION(BlueprintCallable, Category = Graph)
 	UPCGNode* GetOutputNode() const { return OutputNode; }
 
+	/** Duplicate a given node by creating a new node with the same settings and properties, but without any edges and add it to the graph */
+	TObjectPtr<UPCGNode> ReconstructNewNode(const UPCGNode* InNode);
+
 	bool Contains(UPCGNode* Node) const;
 	const TArray<UPCGNode*>& GetNodes() const { return Nodes; }
 	void AddNode(UPCGNode* InNode);
diff --git a/Engine/Plugins/Experimental/PCG/Source/PCG/Public/PCGInputOutputSettings.h b/Engine/Plugins/Experimental/PCG/Source/PCG/Public/PCGInputOutputSettings.h
--- a/Engine/Plugins/Experimental/PCG/Source/PCG/Public/PCGInputOutputSettings.h
+++ b/Engine/Plugins/Experimental/PCG/Source/PCG/Public/PCGInputOutputSettings.h
@@ -45,6 +45,11 @@ public:
 
 	void SetInput(bool bInIsInput) { bIsInput = bInIsInput; }
 
+	bool IsPinAdvanced(const UPCGPin* Pin) const;
+	void SetShowAdvancedPins(bool bValue);
+
+	void AddCustomPin(const FPCGPinProperties& NewCustomPinProperties);
+
 protected:
 	virtual FPCGElementPtr CreateElement() const override { return MakeShared<FPCGInputOutputElement>(); }
 	// ~End UPCGSettings interface
diff --git a/Engine/Plugins/Experimental/PCG/Source/PCG/Public/PCGNode.h b/Engine/Plugins/Experimental/PCG/Source/PCG/Public/PCGNode.h
--- a/Engine/Plugins/Experimental/PCG/Source/PCG/Public/PCGNode.h
+++ b/Engine/Plugins/Experimental/PCG/Source/PCG/Public/PCGNode.h
@@ -80,6 +80,11 @@ public:
 	const TArray<TObjectPtr<UPCGPin>>& GetInputPins() const { return InputPins; }
 	const TArray<TObjectPtr<UPCGPin>>& GetOutputPins() const { return OutputPins; }
 
+#if WITH_EDITOR
+	/** Transfer all editor only properties to the other node */
+	void TransferEditorProperties(UPCGNode* OtherNode) const;
+#endif // WITH_EDITOR
+
 	/** Note: do not set this property directly from code, use SetDefaultSettings instead */
 	UPROPERTY(BlueprintReadWrite, EditAnywhere, Category = Node, meta=(EditInline))
 	TObjectPtr<UPCGSettings> DefaultSettings;
diff --git a/Engine/Plugins/Experimental/PCG/Source/PCGEditor/PCGEditor.Build.cs b/Engine/Plugins/Experimental/PCG/Source/PCGEditor/PCGEditor.Build.cs
--- a/Engine/Plugins/Experimental/PCG/Source/PCGEditor/PCGEditor.Build.cs
+++ b/Engine/Plugins/Experimental/PCG/Source/PCGEditor/PCGEditor.Build.cs
@@ -28,6 +28,7 @@ namespace UnrealBuildTool.Rules
 					"ApplicationCore",
 					"AssetTools",
 					"DesktopWidgets",
+					"EditorScriptingUtilities",
 					"EditorStyle",
 					"EditorSubsystem",
 					"GraphEditor",
diff --git a/Engine/Plugins/Experimental/PCG/Source/PCGEditor/Private/PCGEditor.cpp b/Engine/Plugins/Experimental/PCG/Source/PCGEditor/Private/PCGEditor.cpp
--- a/Engine/Plugins/Experimental/PCG/Source/PCGEditor/Private/PCGEditor.cpp
+++ b/Engine/Plugins/Experimental/PCG/Source/PCGEditor/Private/PCGEditor.cpp
@@ -2,15 +2,25 @@
 
 #include "PCGEditor.h"
 
+#include "PCGEdge.h"
+#include "PCGGraph.h"
+#include "PCGGraphFactory.h"
+#include "PCGInputOutputSettings.h"
+#include "PCGPin.h"
+#include "PCGSubgraph.h"
+
 #include "PCGEditorCommands.h"
 #include "PCGEditorCommon.h"
 #include "PCGEditorGraph.h"
 #include "PCGEditorGraphNodeBase.h"
+#include "PCGEditorGraphNodeInput.h"
+#include "PCGEditorGraphNodeOutput.h"
 #include "PCGEditorGraphSchema.h"
-#include "PCGGraph.h"
+#include "PCGEditorGraphSchemaActions.h"
+#include "PCGEditorSettings.h"
+#include "PCGEditorUtils.h"
 #include "SPCGEditorGraphFind.h"
 #include "SPCGEditorGraphNodePalette.h"
-#include "PCGEditorSettings.h"
 
 #include "EdGraphUtilities.h"
 #include "Editor.h"
@@ -35,11 +45,14 @@
 #include "Editor/UnrealEdEngine.h"
 #include "Preferences/UnrealEdOptions.h"
 #include "Widgets/Notifications/SNotificationList.h"
-#include "PCGEditorGraphSchemaActions.h"
 #include "GraphEditAction.h"
+#include "AssetToolsModule.h"
+#include "EditorAssetLibrary.h"
 
 #define LOCTEXT_NAMESPACE "PCGGraphEditor"
 
+DEFINE_LOG_CATEGORY(LogPCGEditor);
+
 namespace FPCGEditor_private
 {
 	const FName GraphEditorID = FName(TEXT("GraphEditor"));
@@ -238,6 +251,11 @@ void FPCGEditor::BindCommands()
 	ToolkitCommands->MapAction(
 		PCGEditorCommands.Find,
 		FExecuteAction::CreateSP(this, &FPCGEditor::OnFind));
+
+	GraphEditorCommands->MapAction(
+		PCGEditorCommands.CollapseNodes,
+		FExecuteAction::CreateSP(this, &FPCGEditor::OnCollapseNodesInSubgraph),
+		FCanExecuteAction::CreateSP(this, &FPCGEditor::CanCollapseNodesInSubgraph));
 }
 
 void FPCGEditor::OnFind()
@@ -249,6 +267,382 @@ void FPCGEditor::OnFind()
 	}
 }
 
+bool FPCGEditor::CanCollapseNodesInSubgraph()
+{
+	bool HasPCGNode = false;
+
+	for (UObject* Object : GraphEditorWidget->GetSelectedNodes())
+	{
+		check(Object);
+
+		// Exclude input and output nodes from the subgraph.
+		if (Object->IsA<UPCGEditorGraphNodeInput>() || Object->IsA<UPCGEditorGraphNodeOutput>())
+		{
+			continue;
+		}
+
+		if (Object->IsA<UPCGEditorGraphNodeBase>())
+		{
+			if (HasPCGNode)
+			{
+				return true;
+			}
+
+			HasPCGNode = true;
+		}
+	}
+
+	return false;
+}
+
+void FPCGEditor::OnCollapseNodesInSubgraph()
+{
+	if (!GraphEditorWidget.IsValid() || PCGEditorGraph == nullptr)
+	{
+		UE_LOG(LogPCGEditor, Error, TEXT("GraphEditorWidget or PCGEditorGraph is null, aborting"));
+		return;
+	}
+	
+	UPCGGraph* PCGGraph = PCGEditorGraph->GetPCGGraph();
+	if (PCGGraph == nullptr)
+	{
+		UE_LOG(LogPCGEditor, Error, TEXT("PCGGraph is null, aborting"));
+		return;
+	}
+
+	// 1. Gather all nodes that will be included in the subgraph, and the extra nodes
+	// Also keep track of all node positions to know where to spawn our subgraph node
+	FVector2D AveragePosition = FVector2D::ZeroVector;
+	int32 MinX = std::numeric_limits<int32>::max();
+	int32 MaxX = std::numeric_limits<int32>::min();
+
+	// We keep track of all the PCG nodes that will go into the subgraph
+	TArray<TObjectPtr<UPCGEditorGraphNodeBase>> PCGEditorGraphNodes;
+	// Also move the extra nodes (like comments)
+	TArray<TObjectPtr<UEdGraphNode>> ExtraGraphNodes;
+
+	// Those 3 sets are used to keep track of the subgraph pins
+	// and will be used to extract pins that will be outside the subgraph
+	// and will need special treatment.
+	TSet<TObjectPtr<UPCGPin>> InputFromSubgraphPins;
+	TSet<TObjectPtr<UPCGPin>> OutputToSubgraphPins;
+	TSet<TObjectPtr<UPCGPin>> AllSubgraphPins;
+
+	// Also keep track of all the edges in the subgraph. Use a set to be able to add the
+	// same edge twice without duplicates.
+	TSet<TObjectPtr<UPCGEdge>> PCGGraphEdges;
+
+	for (UObject* Object : GraphEditorWidget->GetSelectedNodes())
+	{
+		check(Object);
+
+		// Exclude input and output nodes from the subgraph.
+		if (Object->IsA<UPCGEditorGraphNodeInput>() || Object->IsA<UPCGEditorGraphNodeOutput>())
+		{
+			continue;
+		}
+
+		if (UPCGEditorGraphNodeBase* PCGEditorGraphNode = Cast<UPCGEditorGraphNodeBase>(Object))
+		{
+			UPCGNode* PCGNode = PCGEditorGraphNode->GetPCGNode();
+			check(PCGNode);
+
+			// For each pin of the subgraph, gather all its edges in a set and store the other pin of
+			// the edge as an input from or output to the subgraph. It will be useful for tracking pins
+			// that are outside the subgraph.
+			for (UPCGPin* InputPin : PCGNode->GetInputPins())
+			{
+				check(InputPin);
+
+				for (UPCGEdge* Edge : InputPin->Edges)
+				{
+					check(Edge);
+					PCGGraphEdges.Add(Edge);
+					OutputToSubgraphPins.Add(Edge->InputPin);
+				}
+
+				AllSubgraphPins.Add(InputPin);
+			}
+
+			for (UPCGPin* OutputPin : PCGNode->GetOutputPins())
+			{
+				check(OutputPin);
+
+				for (UPCGEdge* Edge : OutputPin->Edges)
+				{
+					check(Edge);
+					PCGGraphEdges.Add(Edge);
+					InputFromSubgraphPins.Add(Edge->OutputPin);
+				}
+
+				AllSubgraphPins.Add(OutputPin);
+			}
+
+			PCGEditorGraphNodes.Add(PCGEditorGraphNode);
+
+			// And do all the computation to get the min, max and mean position.
+			AveragePosition.X += PCGEditorGraphNode->NodePosX;
+			AveragePosition.Y += PCGEditorGraphNode->NodePosY;
+			MinX = FMath::Min(MinX, PCGEditorGraphNode->NodePosX);
+			MaxX = FMath::Max(MaxX, PCGEditorGraphNode->NodePosX);
+		}
+		else if (UEdGraphNode* GraphNode = Cast<UEdGraphNode>(Object))
+		{
+			ExtraGraphNodes.Add(GraphNode);
+		}
+	}
+
+	// If we have at most 1 node to collapse, just exit
+	if (PCGEditorGraphNodes.Num() <= 1)
+	{
+		UE_LOG(LogPCGEditor, Warning, TEXT("There were less than 2 PCG nodes selected, abort"));
+		return;
+	}
+
+	// Compute the average position
+	AveragePosition /= PCGEditorGraphNodes.Num();
+
+	// 2. Create a new subgraph, by creating a new PCGGraph asset.
+	IAssetTools& AssetTools = FModuleManager::Get().LoadModuleChecked<FAssetToolsModule>("AssetTools").Get();
+
+	TObjectPtr<UPCGGraphFactory> Factory = NewObject<UPCGGraphFactory>();
+
+	FString NewPackageName;
+	FString NewAssetName;
+	PCGEditorUtils::GetParentPackagePathAndUniqueName(PCGGraph, LOCTEXT("NewPCGSubgraphAsset", "NewPCGSubgraph").ToString(), NewPackageName, NewAssetName);
+
+	TObjectPtr<UPCGGraph> NewPCGGraph = Cast<UPCGGraph>(AssetTools.CreateAssetWithDialog(NewAssetName, NewPackageName, PCGGraph->GetClass(), Factory, "PCGEditor_CollapseInSubgraph"));
+
+	if (NewPCGGraph == nullptr)
+	{
+		UE_LOG(LogPCGEditor, Warning, TEXT("Subgraph asset creation was aborted or failed, abort."));
+		return;
+	}
+
+	// Do some clean-up on input/output nodes
+	constexpr int32 Padding = 200;
+	NewPCGGraph->GetInputNode()->PositionX = MinX - Padding;
+	NewPCGGraph->GetInputNode()->PositionY = AveragePosition.Y;
+	NewPCGGraph->GetOutputNode()->PositionX = MaxX + Padding;
+	NewPCGGraph->GetOutputNode()->PositionY = AveragePosition.Y;
+
+	// 3. Define some lambdas to create/adapt pins on the new subgraph.
+
+	// Gather the pins that are outside the subgraph
+	InputFromSubgraphPins = InputFromSubgraphPins.Difference(AllSubgraphPins);
+	OutputToSubgraphPins = OutputToSubgraphPins.Difference(AllSubgraphPins);
+
+	// Function that check if the pin is part of the advanced pins of the input/output node of the PCGGraph.
+	// If so, enable advanced pins on the subgraph.
+	auto IsAdvancedAndSet = [&PCGGraph, &NewPCGGraph](UPCGPin* Pin, bool bIsInput)
+	{
+		TObjectPtr<UPCGNode> InputOutputNode = bIsInput ? PCGGraph->GetInputNode() : PCGGraph->GetOutputNode();
+		TObjectPtr<UPCGNode> NewInputOutputNode = bIsInput ? NewPCGGraph->GetInputNode() : NewPCGGraph->GetOutputNode();
+
+		TObjectPtr<UPCGGraphInputOutputSettings> InputOutputSettings = CastChecked<UPCGGraphInputOutputSettings>(InputOutputNode->DefaultSettings);
+		if (InputOutputSettings->IsPinAdvanced(Pin))
+		{
+			TObjectPtr<UPCGGraphInputOutputSettings> NewInputOutputSettings = CastChecked<UPCGGraphInputOutputSettings>(NewInputOutputNode->DefaultSettings);
+			NewInputOutputSettings->SetShowAdvancedPins(true);
+			NewInputOutputNode->UpdateAfterSettingsChangeDuringCreation();
+
+			return true;
+		}
+
+		return false;
+	};
+
+	// Function that will create a new pin in the input/output node of the subgraph, with a unique name that will
+	// be the related to the pin passed as argument. It will be formatted like this:
+	// "{NodeName} {PinName} {OptionalIndex}"
+	TMap<FString, int> NameCollisionMapping;
+	auto CreateNewCustomPin = [&NewPCGGraph, &NameCollisionMapping](UPCGPin* Pin, bool bIsInput)
+	{
+		TObjectPtr<UPCGNode> NewInputOutputNode = bIsInput ? NewPCGGraph->GetInputNode() : NewPCGGraph->GetOutputNode();
+		TObjectPtr<UPCGGraphInputOutputSettings> NewInputOutputSettings = CastChecked<UPCGGraphInputOutputSettings>(NewInputOutputNode->DefaultSettings);
+		FString NewName = Pin->Node->GetNodeTitle().ToString() + " " + Pin->Properties.Label.ToString();
+		if (NameCollisionMapping.Contains(NewName))
+		{
+			NewName += " " + FString::FormatAsNumber(++NameCollisionMapping[NewName]);
+		}
+		else
+		{
+			NameCollisionMapping.Emplace(NewName, 1);
+		}
+
+		FPCGPinProperties NewProperties = Pin->Properties;
+		NewProperties.Label = FName(NewName);
+		NewInputOutputSettings->AddCustomPin(NewProperties);
+		NewInputOutputNode->UpdateAfterSettingsChangeDuringCreation();
+		return NewProperties.Label;
+	};
+
+	// 4. Duplicate all the nodes, and keep a mapping between the old pins and new pins
+	TMap<TObjectPtr<UPCGPin>, TObjectPtr<UPCGPin>> PinMapping;
+	for (UPCGEditorGraphNodeBase* PCGEditorGraphNode : PCGEditorGraphNodes)
+	{
+		const UPCGNode* PCGNode = PCGEditorGraphNode->GetPCGNode();
+		check(PCGNode);
+
+		// Reconstruct a new node, same as PCGNode, but without any edges in the new graph 
+		TObjectPtr<UPCGNode> NewNode = NewPCGGraph->ReconstructNewNode(PCGNode);
+		
+		for (int i = 0; i < PCGNode->GetInputPins().Num(); ++i)
+		{
+			PinMapping.Emplace(PCGNode->GetInputPins()[i], NewNode->GetInputPins()[i]);
+		}
+		for (int i = 0; i < PCGNode->GetOutputPins().Num(); ++i)
+		{
+			PinMapping.Emplace(PCGNode->GetOutputPins()[i], NewNode->GetOutputPins()[i]);
+		}
+	}
+
+	// Also duplicate the extra nodes and assign them to the new graph
+	TArray<TObjectPtr<const UObject>> NewExtraGraphNodes;
+	for (const UEdGraphNode* ExtraNode : ExtraGraphNodes)
+	{
+		NewExtraGraphNodes.Add(DuplicateObject(ExtraNode, NewPCGGraph));
+	}
+
+	NewPCGGraph->SetExtraEditorNodes(NewExtraGraphNodes);
+
+	// 5. Iterate over all the edges and create edges "placeholders"
+	// Most of them will already be complete, but for those that needs to be connected to the new
+	// subgraph nodes, the pins don't exist yet. Therefore we identify them with their pin labels.
+	//
+	// The logic behind the new pins is this:
+	// -> If the pin is connected to the simple pin of the input/output node, let it like this
+	// -> If the pin is part of the original input/output advanced pins, we will trigger the advanced pins flag on the input/output node in the subgraph
+	// -> Otherwise, we add a new custom pin, with the name of the node, the name of the pin and a number if there is name collision
+	struct EdgePlaceholder
+	{
+		TObjectPtr<UPCGPin> InputPin = nullptr;
+		TObjectPtr<UPCGPin> OutputPin = nullptr;
+		FName InputPinLabel;
+		FName OutputPinLabel;
+	};
+
+	TArray<EdgePlaceholder> EdgePlaceholders;
+	for (UPCGEdge* Edge : PCGGraphEdges)
+	{
+		const TObjectPtr<UPCGPin>* InPin = PinMapping.Find(Edge->InputPin);
+		const TObjectPtr<UPCGPin>* OutPin = PinMapping.Find(Edge->OutputPin);
+
+		check(InPin || OutPin)
+
+		if (InPin == nullptr)
+		{
+			// The edge comes from outside the graph.
+			// If it is from the input node, we have a special behavior
+			EdgePlaceholder OutsideSubgraphEdge;
+			EdgePlaceholder InsideSubgraphEdge;
+			bool bProcessed = false;
+
+			OutsideSubgraphEdge.InputPin = Edge->InputPin;
+			InsideSubgraphEdge.OutputPin = *OutPin;
+
+			if (Edge->InputPin->Node == PCGGraph->GetInputNode())
+			{
+				if (Edge->InputPin == PCGGraph->GetInputNode()->GetOutputPins()[0] || IsAdvancedAndSet(Edge->InputPin, true))
+				{
+					OutsideSubgraphEdge.OutputPinLabel = Edge->InputPin->Properties.Label;
+					InsideSubgraphEdge.InputPin = NewPCGGraph->GetInputNode()->GetOutputPin(Edge->InputPin->Properties.Label);
+					bProcessed = true;
+				}
+			}
+
+			if (!bProcessed)
+			{
+				FName NewPinName = CreateNewCustomPin(Edge->OutputPin, true);
+				OutsideSubgraphEdge.OutputPinLabel = NewPinName;
+				InsideSubgraphEdge.InputPin = NewPCGGraph->GetInputNode()->GetOutputPin(NewPinName);
+			}
+
+			EdgePlaceholders.Add(OutsideSubgraphEdge);
+			EdgePlaceholders.Add(InsideSubgraphEdge);
+		}
+		else if (OutPin == nullptr)
+		{
+			// The edge comes from outside the graph.
+			// If it is from the output node, we have a special behavior
+			EdgePlaceholder OutsideSubgraphEdge;
+			EdgePlaceholder InsideSubgraphEdge;
+			bool bProcessed = false;
+
+			OutsideSubgraphEdge.OutputPin = Edge->OutputPin;
+			InsideSubgraphEdge.InputPin = *InPin;
+
+			if (Edge->OutputPin->Node == PCGGraph->GetOutputNode())
+			{
+				if (Edge->OutputPin == PCGGraph->GetOutputNode()->GetInputPins()[0] || IsAdvancedAndSet(Edge->OutputPin, false))
+				{
+					OutsideSubgraphEdge.InputPinLabel = Edge->OutputPin->Properties.Label;
+					InsideSubgraphEdge.OutputPin = NewPCGGraph->GetOutputNode()->GetInputPin(Edge->OutputPin->Properties.Label);
+					bProcessed = true;
+				}
+			}
+
+			if (!bProcessed)
+			{
+				FName NewPinName = CreateNewCustomPin(Edge->InputPin, false);
+				OutsideSubgraphEdge.InputPinLabel = NewPinName;
+				InsideSubgraphEdge.OutputPin = NewPCGGraph->GetOutputNode()->GetInputPin(NewPinName);
+			}
+
+			EdgePlaceholders.Add(OutsideSubgraphEdge);
+			EdgePlaceholders.Add(InsideSubgraphEdge);
+		}
+		else
+		{
+			// Both nodes are inside
+			EdgePlaceholders.Add(EdgePlaceholder{*InPin, *OutPin});
+		}
+	}
+
+	// 6. Create subgraph and delete old nodes
+	// Done within a transaction to be undoable (subgraph will stay though)
+	TObjectPtr<UPCGEditorGraphNodeBase> SubgraphEditorNode;
+	TObjectPtr<UPCGNode> SubgraphNode;
+	{
+		const FScopedTransaction Transaction(*FPCGEditorCommon::ContextIdentifier, LOCTEXT("PCGEditorCollapseInSubgraphMessage", "PCG Editor: Collapse in Subgraph"), nullptr);
+
+		DeleteSelectedNodes();
+
+		FPCGEditorGraphSchemaAction_NewSubgraphElement SubgraphAddAction;
+		SubgraphAddAction.SubgraphObjectPath = NewPCGGraph->GetPathName();
+		SubgraphEditorNode = Cast<UPCGEditorGraphNodeBase>(SubgraphAddAction.PerformAction(PCGEditorGraph, nullptr, AveragePosition, true));
+		SubgraphNode = SubgraphEditorNode->GetPCGNode();
+	}
+
+	// 7. Connect all the edges
+	for (EdgePlaceholder& Edge : EdgePlaceholders)
+	{
+		if (Edge.InputPin == nullptr)
+		{
+			SubgraphNode->GetOutputPin(Edge.InputPinLabel)->AddEdgeTo(Edge.OutputPin);
+		}
+		else if (Edge.OutputPin == nullptr)
+		{
+			Edge.InputPin->AddEdgeTo(SubgraphNode->GetInputPin(Edge.OutputPinLabel));
+		}
+		else
+		{
+			Edge.InputPin->AddEdgeTo(Edge.OutputPin);
+		}
+	}
+
+	// 8. Finalize the operation
+	SubgraphEditorNode->ReconstructNode();
+
+	// Save the new asset
+	UEditorAssetLibrary::SaveLoadedAsset(NewPCGGraph);
+
+	// And notify everyone
+	GraphEditorWidget->NotifyGraphChanged();
+	PCGGraph->NotifyGraphChanged(true);
+	NewPCGGraph->NotifyGraphChanged(true);
+}
+
 void FPCGEditor::SelectAllNodes()
 {
 	if (GraphEditorWidget.IsValid())
diff --git a/Engine/Plugins/Experimental/PCG/Source/PCGEditor/Private/PCGEditor.h b/Engine/Plugins/Experimental/PCG/Source/PCGEditor/Private/PCGEditor.h
--- a/Engine/Plugins/Experimental/PCG/Source/PCGEditor/Private/PCGEditor.h
+++ b/Engine/Plugins/Experimental/PCG/Source/PCGEditor/Private/PCGEditor.h
@@ -13,6 +13,9 @@ class SPCGEditorGraphNodePalette;
 class UPCGEditorGraph;
 class UPCGGraph;
 
+// Logs
+DECLARE_LOG_CATEGORY_EXTERN(LogPCGEditor, Log, All);
+
 class FPCGEditor : public FAssetEditorToolkit, public FSelfRegisteringEditorUndoClient
 {
 public:
@@ -89,6 +92,11 @@ private:
 	/** Whether we are able to duplicate the currently selected nodes */
 	bool CanDuplicateNodes() const;
 
+	/** Collapse the currently selected nodes in a subgraph */
+	void OnCollapseNodesInSubgraph();
+	/** Whether we can collapse nodes in a subgraph */
+	bool CanCollapseNodesInSubgraph();
+
 	void OnAlignTop();
 	void OnAlignMiddle();
 	void OnAlignBottom();
diff --git a/Engine/Plugins/Experimental/PCG/Source/PCGEditor/Private/PCGEditorCommands.cpp b/Engine/Plugins/Experimental/PCG/Source/PCGEditor/Private/PCGEditorCommands.cpp
--- a/Engine/Plugins/Experimental/PCG/Source/PCGEditor/Private/PCGEditorCommands.cpp
+++ b/Engine/Plugins/Experimental/PCG/Source/PCGEditor/Private/PCGEditorCommands.cpp
@@ -19,6 +19,7 @@ FPCGEditorCommands::FPCGEditorCommands()
 void FPCGEditorCommands::RegisterCommands()
 {
 	UI_COMMAND(Find, "Find", "Finds PCG nodes and comments in the current graph", EUserInterfaceActionType::Button, FInputChord(EModifierKey::Control, EKeys::F));
+	UI_COMMAND(CollapseNodes, "Collapse in Subgraph", "Collapse selected nodes in a separate PCGGraph asset.", EUserInterfaceActionType::Button, FInputChord(EModifierKey::Alt, EKeys::J));
 }
 
 #undef LOCTEXT_NAMESPACE
diff --git a/Engine/Plugins/Experimental/PCG/Source/PCGEditor/Private/PCGEditorCommands.h b/Engine/Plugins/Experimental/PCG/Source/PCGEditor/Private/PCGEditorCommands.h
--- a/Engine/Plugins/Experimental/PCG/Source/PCGEditor/Private/PCGEditorCommands.h
+++ b/Engine/Plugins/Experimental/PCG/Source/PCGEditor/Private/PCGEditorCommands.h
@@ -15,4 +15,5 @@ public:
 	// ~End TCommands<> interface
 
 	TSharedPtr<FUICommandInfo> Find;
+	TSharedPtr<FUICommandInfo> CollapseNodes;
 };
diff --git a/Engine/Plugins/Experimental/PCG/Source/PCGEditor/Private/PCGEditorGraphNodeBase.cpp b/Engine/Plugins/Experimental/PCG/Source/PCGEditor/Private/PCGEditorGraphNodeBase.cpp
--- a/Engine/Plugins/Experimental/PCG/Source/PCGEditor/Private/PCGEditorGraphNodeBase.cpp
+++ b/Engine/Plugins/Experimental/PCG/Source/PCGEditor/Private/PCGEditorGraphNodeBase.cpp
@@ -18,6 +18,7 @@
 #include "ToolMenu.h"
 #include "ToolMenuSection.h"
 #include "Widgets/Colors/SColorPicker.h"
+#include "PCGEditorCommands.h"
 
 #define LOCTEXT_NAMESPACE "PCGEditorGraphNodeBase"
 
@@ -75,6 +76,7 @@ void UPCGEditorGraphNodeBase::GetNodeContextMenuActions(UToolMenu* Menu, class U
 	{
 		FToolMenuSection& Section = Menu->AddSection("EdGraphSchemaNodeActions", LOCTEXT("NodeActionsHeader", "Node Actions"));
 		Section.AddMenuEntry(FGraphEditorCommands::Get().BreakNodeLinks);
+		Section.AddMenuEntry(FPCGEditorCommands::Get().CollapseNodes);
 	}
 
 	{
diff --git a/Engine/Plugins/Experimental/PCG/Source/PCGEditor/Private/PCGEditorUtils.cpp b/Engine/Plugins/Experimental/PCG/Source/PCGEditor/Private/PCGEditorUtils.cpp
--- a/Engine/Plugins/Experimental/PCG/Source/PCGEditor/Private/PCGEditorUtils.cpp
+++ b/Engine/Plugins/Experimental/PCG/Source/PCGEditor/Private/PCGEditorUtils.cpp
@@ -6,6 +6,10 @@
 #include "UObject/NoExportTypes.h"
 #include "Elements/PCGExecuteBlueprint.h"
 #include "Engine/Blueprint.h"
+#include "AssetToolsModule.h"
+#include "Modules/ModuleManager.h"
+#include "Misc/PackageName.h"
+#include "ObjectTools.h"
 
 bool PCGEditorUtils::IsAssetPCGBlueprint(const FAssetData& InAssetData)
 {
@@ -13,4 +17,23 @@ bool PCGEditorUtils::IsAssetPCGBlueprint(const FAssetData& InAssetData)
 	FString TargetNativeParentClassName = UPCGBlueprintElement::GetParentClassName();
 
 	return InAssetData.AssetClassPath == UBlueprint::StaticClass()->GetClassPathName() && InNativeParentClassName == TargetNativeParentClassName;
-}
\ No newline at end of file
+}
+
+void PCGEditorUtils::GetParentPackagePathAndUniqueName(const UObject* OriginalObject, const FString& NewAssetTentativeName, FString& OutPackagePath, FString& OutUniqueName)
+{
+	if (OriginalObject == nullptr)
+	{
+		return;
+	}
+
+	IAssetTools& AssetTools = FModuleManager::Get().LoadModuleChecked<FAssetToolsModule>("AssetTools").Get();
+
+	FString PackageRoot, PackagePath, PackageName;
+	FPackageName::SplitLongPackageName(OriginalObject->GetPackage()->GetPathName(), PackageRoot, PackagePath, PackageName);
+
+	OutPackagePath = PackageRoot / PackagePath;
+
+	FString DummyPackageName;
+	AssetTools.CreateUniqueAssetName(OutPackagePath, NewAssetTentativeName, DummyPackageName, OutUniqueName);
+}
+
diff --git a/Engine/Plugins/Experimental/PCG/Source/PCGEditor/Private/PCGEditorUtils.h b/Engine/Plugins/Experimental/PCG/Source/PCGEditor/Private/PCGEditorUtils.h
--- a/Engine/Plugins/Experimental/PCG/Source/PCGEditor/Private/PCGEditorUtils.h
+++ b/Engine/Plugins/Experimental/PCG/Source/PCGEditor/Private/PCGEditorUtils.h
@@ -3,8 +3,18 @@
 #pragma once
 
 struct FAssetData;
+class UObject;
+class FString;
 
 namespace PCGEditorUtils
 {
 	bool IsAssetPCGBlueprint(const FAssetData& InAssetData);
+
+	/** 
+	* From an object, get its parent package and a unique name 
+	* For example, if you want to create a new asset next to the original object, it will return the parent package of the original package
+	* and a unique name for the new asset.
+	*/
+	void GetParentPackagePathAndUniqueName(const UObject* OriginalObject, const FString& NewAssetTentativeName, FString& OutPackagePath, FString& OutUniqueName);
+
 }
\ No newline at end of file
