Tentative fix for dependency warning in PCG plugin

#rnx

diff --git a/Engine/Plugins/Experimental/PCG/PCG.uplugin b/Engine/Plugins/Experimental/PCG/PCG.uplugin
--- a/Engine/Plugins/Experimental/PCG/PCG.uplugin
+++ b/Engine/Plugins/Experimental/PCG/PCG.uplugin
@@ -32,6 +32,10 @@
 		{
 			"Name": "GeometryScripting",
 			"Enabled": true
+		},
+		{
+			"Name": "EditorScriptingUtilities",
+			"Enabled": true
 		}
 	]
 }
