Added spline sampler
Fixed some outstanding issues in the landscape spline data initialization

diff --git a/Engine/Plugins/Experimental/PCG/Source/PCG/Private/Data/PCGLandscapeSplineData.cpp b/Engine/Plugins/Experimental/PCG/Source/PCG/Private/Data/PCGLandscapeSplineData.cpp
--- a/Engine/Plugins/Experimental/PCG/Source/PCG/Private/Data/PCGLandscapeSplineData.cpp
+++ b/Engine/Plugins/Experimental/PCG/Source/PCG/Private/Data/PCGLandscapeSplineData.cpp
@@ -4,6 +4,7 @@
 
 #include "PCGHelpers.h"
 #include "Data/PCGPointData.h"
+#include "Elements/PCGSplineSampler.h"
 
 #include "Landscape.h"
 #include "LandscapeSplinesComponent.h"
@@ -37,6 +38,7 @@ void UPCGLandscapeSplineData::Initialize(ULandscapeSplinesComponent* InSplineCom
 {
 	check(InSplineComponent);
 	Spline = InSplineComponent;
+	TargetActor = InSplineComponent->GetOwner();
 }
 
 int UPCGLandscapeSplineData::GetNumSegments() const
@@ -45,14 +47,14 @@ int UPCGLandscapeSplineData::GetNumSegments() const
 	return Spline->GetSegments().Num();
 }
 
-float UPCGLandscapeSplineData::GetSegmentLength(int SegmentIndex) const
+FVector::FReal UPCGLandscapeSplineData::GetSegmentLength(int SegmentIndex) const
 {
 	check(Spline);
 	check(SegmentIndex >= 0 && SegmentIndex < Spline->GetSegments().Num());
 	
 	const ULandscapeSplineSegment* Segment = Spline->GetSegments()[SegmentIndex];
 	const TArray<FLandscapeSplineInterpPoint>& InterpPoints = Segment->GetPoints();
-	float Length = 0;
+	FVector::FReal Length = 0;
 
 	for (int PointIndex = 1; PointIndex < InterpPoints.Num(); ++PointIndex)
 	{
@@ -62,21 +64,49 @@ float UPCGLandscapeSplineData::GetSegmentLength(int SegmentIndex) const
 	return Length;
 }
 
-FVector UPCGLandscapeSplineData::GetLocationAtDistance(int SegmentIndex, float Distance) const
+FTransform UPCGLandscapeSplineData::GetTransformAtDistance(int SegmentIndex, FVector::FReal Distance, FBox* OutBounds) const
 {
 	check(Spline);
 	check(SegmentIndex >= 0 && SegmentIndex < Spline->GetSegments().Num());
 
 	const ULandscapeSplineSegment* Segment = Spline->GetSegments()[SegmentIndex];
+	check(Segment);
+
+	const FLandscapeSplineSegmentConnection& Start = Segment->Connections[0];
+	const FLandscapeSplineSegmentConnection& End = Segment->Connections[1];
+
 	const TArray<FLandscapeSplineInterpPoint>& InterpPoints = Segment->GetPoints();
-	float Length = Distance;
+	FVector::FReal Length = FMath::Max(0, Distance);
 
 	for (int PointIndex = 1; PointIndex < InterpPoints.Num(); ++PointIndex)
 	{
-		const float SegmentLength = (InterpPoints[PointIndex].Center - InterpPoints[PointIndex - 1].Center).Length();
-		if (SegmentLength > Length)
+		const FLandscapeSplineInterpPoint& PreviousPoint = InterpPoints[PointIndex - 1];
+		const FLandscapeSplineInterpPoint& CurrentPoint = InterpPoints[PointIndex];
+
+		const FVector::FReal SegmentLength = (CurrentPoint.Center - PreviousPoint.Center).Length();
+		if (SegmentLength > Length || PointIndex == InterpPoints.Num() - 1)
 		{
-			return InterpPoints[PointIndex - 1].Center + (InterpPoints[PointIndex].Center - InterpPoints[PointIndex - 1].Center) * (Length / SegmentLength);
+			const FVector XAxis = CurrentPoint.Center - PreviousPoint.Center;
+			const FVector PreviousYAxis = PreviousPoint.Right - PreviousPoint.Center;
+			const FVector CurrentYAxis = CurrentPoint.Right - CurrentPoint.Center;
+			const FVector PreviousZAxis = (XAxis ^ PreviousYAxis).GetSafeNormal(UE_SMALL_NUMBER, FVector::ZAxisVector);
+			const FVector CurrentZAxis = (XAxis ^ CurrentYAxis).GetSafeNormal(UE_SMALL_NUMBER, FVector::ZAxisVector);
+
+			FTransform PreviousTransform = FTransform(XAxis, PreviousYAxis, PreviousZAxis, PreviousPoint.Center);
+			FTransform CurrentTransform = FTransform(XAxis, CurrentYAxis, CurrentZAxis, CurrentPoint.Center);
+
+			const FVector::FReal BlendRatio = ((SegmentLength > Length) ? (Length / SegmentLength) : 1.0);
+			PreviousTransform.BlendWith(CurrentTransform, BlendRatio);
+
+			if (OutBounds)
+			{
+				// Important note: the box here is going to be useful to be able to specify the relative sizes of the falloffs
+				*OutBounds = FBox::BuildAABB(FVector::ZeroVector, FVector::OneVector);
+				OutBounds->Min.Y *= (CurrentPoint.FalloffLeft - CurrentPoint.Center).Length() / (CurrentPoint.Left - CurrentPoint.Center).Length();
+				OutBounds->Max.Y *= (CurrentPoint.FalloffRight - CurrentPoint.Center).Length() / (CurrentPoint.Right - CurrentPoint.Center).Length();
+			}
+
+			return PreviousTransform * Spline->GetComponentTransform();
 		}
 		else
 		{
@@ -85,7 +115,7 @@ FVector UPCGLandscapeSplineData::GetLocationAtDistance(int SegmentIndex, float D
 	}
 	
 	check(0);
-	return FVector::Zero();
+	return FTransform();
 }
 
 const UPCGPointData* UPCGLandscapeSplineData::CreatePointData(FPCGContext* Context) const
@@ -97,97 +127,13 @@ const UPCGPointData* UPCGLandscapeSplineData::CreatePointData(FPCGContext* Conte
 	Data->InitializeFromData(this);
 	TArray<FPCGPoint>& Points = Data->GetMutablePoints();
 
-	// TODO: replace all the logic with sampling settings; currently it uses the landscape scaling as a basis
-	const FTransform SplineTransform = Spline->GetComponentTransform();
-	const FTransform LandscapeTransform = Spline->GetSplineOwner()->LandscapeActorToWorld();
-	const FTransform SplineToLandscape = SplineTransform.GetRelativeTransform(LandscapeTransform);
-	ULandscapeInfo* LandscapeInfo = Spline->GetSplineOwner()->GetLandscapeInfo();
-
-	auto AddPoints = [&Points, &SplineTransform, &LandscapeTransform, &SplineToLandscape, LandscapeInfo](const FVector& A, const FVector& B, const FVector& C, const FVector& D, bool bComputeDensity, bool bAddA, bool bAddB) {
-
-		auto AddPoint = [&Points](const FVector& P, float Density) {
-			FPCGPoint& Point = Points.Emplace_GetRef();
-			Point.Transform = FTransform(P); // pass in segment vector as orientation?
-			Point.Seed = PCGHelpers::ComputeSeed((int)P.X, (int)P.Y, (int)P.Z);
-			Point.Density = Density;
-		};
-
-		if (bAddA)
-		{
-			AddPoint(SplineTransform.TransformPosition(A), 1.0f);
-		}
-
-		// Interpolate points in the ABCD quad based on the given scales, computing density based on ratio from the A->B segment if the bComputeDensity flag is on, and 1.0f otherwise
-		if (LandscapeInfo && LandscapeInfo->LandscapeActor.IsValid())
-		{
-			FBox QuadBoxOnLandscape(EForceInit::ForceInit);
-			QuadBoxOnLandscape += SplineToLandscape.TransformPosition(A);
-			QuadBoxOnLandscape += SplineToLandscape.TransformPosition(B);
-			QuadBoxOnLandscape += SplineToLandscape.TransformPosition(C);
-			QuadBoxOnLandscape += SplineToLandscape.TransformPosition(D);
-
-			int32 MinX = FMath::CeilToInt(QuadBoxOnLandscape.Min.X);
-			int32 MinY = FMath::CeilToInt(QuadBoxOnLandscape.Min.Y);
-			int32 MaxX = FMath::FloorToInt(QuadBoxOnLandscape.Max.X);
-			int32 MaxY = FMath::FloorToInt(QuadBoxOnLandscape.Max.Y);
-
-			const FVector QuadNormal = (C - B).Cross(B - A).GetSafeNormal();
+	FPCGSplineSamplerParams SamplerParams;
+	SamplerParams.Mode = EPCGSplineSamplingMode::Distance;
+	SamplerParams.Dimension = EPCGSplineSamplingDimension::OnHorizontal;
 
-			for (int X = MinX; X <= MaxX; ++X)
-			{
-				for (int Y = MinY; Y <= MaxY; ++Y)
-				{
-					FVector TentativeLocation = LandscapeTransform.TransformPosition(FVector(X, Y, 0));
-					const FVector TentativeLocationInSplineSpace = SplineTransform.InverseTransformPosition(TentativeLocation);
-
-					const FVector::FReal ComputedDensity = PCGLandscapeDataHelpers::GetDensityInQuad(A, B, C, D, TentativeLocationInSplineSpace);
-
-					// Check if the point would be in the quad
-					if (ComputedDensity >= 0)
-					{
-						// TODO: note that we can't call GetHeight on the ULandscapeHEightfieldCollisionComponent, as it is not exported
-						// so we have to resort to calling the method from the landscape actor
-						TOptional<float> HeightAtVertex = LandscapeInfo->LandscapeActor->GetHeightAtLocation(TentativeLocation);
-						if (HeightAtVertex.IsSet())
-						{
-							TentativeLocation.Z = HeightAtVertex.GetValue();
-							AddPoint(TentativeLocation, bComputeDensity ? ComputedDensity : 1.0f);
-						}
-					}
-				}
-			}
-		}
-
-		if (bAddB)
-		{
-			AddPoint(SplineTransform.TransformPosition(B), 1.0f);
-		}
-	};
-	
-	const TArray<TObjectPtr<ULandscapeSplineSegment>>& Segments = Spline->GetSegments();
-
-	// For each segment on the spline,
-	for (int SegmentIndex = 0; SegmentIndex < Segments.Num(); ++SegmentIndex)
-	{
-		// Sample points on and between points
-		const TArray<FLandscapeSplineInterpPoint>& InterpPoints = Segments[SegmentIndex]->GetPoints();
-		bool bIsLastSegment = (SegmentIndex == Segments.Num() - 1);
-
-		for (int PointIndex = 1; PointIndex < InterpPoints.Num(); ++PointIndex)
-		{
-			bool bIsLastPointInSpline = (bIsLastSegment && PointIndex == InterpPoints.Num() - 1);
-
-			const FLandscapeSplineInterpPoint& Start = InterpPoints[PointIndex - 1];
-			const FLandscapeSplineInterpPoint& End = InterpPoints[PointIndex];
-
-			AddPoints(Start.Center, End.Center, End.Left, Start.Left, /*bComputeDensity=*/false, /*bAddA=*/true, /*bAddB=*/false);
-			AddPoints(Start.Left, End.Left, End.FalloffLeft, Start.FalloffLeft, /*bComputeDensity=*/true, /*bAddA=*/true, /*bAddB=*/bIsLastPointInSpline);
-			AddPoints(End.Center, Start.Center, Start.Right, End.Right, /*bComputeDensity=*/false, /*bAddA=*/bIsLastPointInSpline, /*bAddB=*/false);
-			AddPoints(End.Right, Start.Right, Start.FalloffRight, End.FalloffRight, /*bComputeDensity=*/true, /*bAddA=*/bIsLastPointInSpline, /*bAddB=*/true);
-		}
-	}
+	PCGSplineSampler::SampleLineData(this, this, SamplerParams, Data);
 
-	UE_LOG(LogPCG, Verbose, TEXT("Landscape spline %s generated %d points on %d segments"), *Spline->GetFName().ToString(), Points.Num(), Segments.Num());
+	UE_LOG(LogPCG, Verbose, TEXT("Landscape spline %s generated %d points"), *Spline->GetFName().ToString(), Points.Num());
 
 	return Data;
 }
@@ -242,10 +188,11 @@ bool UPCGLandscapeSplineData::SamplePoint(const FTransform& InTransform, const F
 			// Considering that the points on a given control point are probably aligned, we could do an early check
 			// in the original quad (start left falloff -> start right falloff -> end right falloff -> end left falloff)
 			// TODO: this sequence here can be optimized knowing that some checks will be redundant.
+			// Important note: the order and selection of points is important to the density computation. This assumes the first 2 points are the "1" density edge
 			Density = FMath::Max(Density, PCGLandscapeDataHelpers::GetDensityInQuad(Start.Center, End.Center, End.Left, Start.Left, Position) >= 0 ? 1.0f : 0.0f);
 			Density = FMath::Max(Density, PCGLandscapeDataHelpers::GetDensityInQuad(Start.Left, End.Left, End.FalloffLeft, Start.FalloffLeft, Position));
-			Density = FMath::Max(Density, PCGLandscapeDataHelpers::GetDensityInQuad(Start.Center, Start.Right, End.Right, End.Center, Position) >= 0 ? 1.0f : 0.0f);
-			Density = FMath::Max(Density, PCGLandscapeDataHelpers::GetDensityInQuad(Start.Right, Start.FalloffRight, End.FalloffRight, End.Right, Position));
+			Density = FMath::Max(Density, PCGLandscapeDataHelpers::GetDensityInQuad(End.Center, Start.Center, Start.Right, End.Right, Position) >= 0 ? 1.0f : 0.0f);
+			Density = FMath::Max(Density, PCGLandscapeDataHelpers::GetDensityInQuad(End.Right, Start.Right, Start.FalloffRight, End.FalloffRight, Position));
 
 			if (Density > SegmentDensity)
 			{
diff --git a/Engine/Plugins/Experimental/PCG/Source/PCG/Private/Data/PCGSplineData.cpp b/Engine/Plugins/Experimental/PCG/Source/PCG/Private/Data/PCGSplineData.cpp
--- a/Engine/Plugins/Experimental/PCG/Source/PCG/Private/Data/PCGSplineData.cpp
+++ b/Engine/Plugins/Experimental/PCG/Source/PCG/Private/Data/PCGSplineData.cpp
@@ -5,6 +5,7 @@
 #include "PCGHelpers.h"
 #include "Data/PCGPointData.h"
 #include "Data/PCGSurfaceData.h"
+#include "Elements/PCGSplineSampler.h"
 
 #include "Components/SplineComponent.h"
 
@@ -30,17 +31,27 @@ void UPCGSplineData::Initialize(USplineComponent* InSpline)
 
 int UPCGSplineData::GetNumSegments() const
 {
-	return 1;
+	return Spline ? Spline->GetNumberOfSplineSegments() : 0;
 }
 
-float UPCGSplineData::GetSegmentLength(int SegmentIndex) const
+FVector::FReal UPCGSplineData::GetSegmentLength(int SegmentIndex) const
 {
-	return Spline->GetSplineLength();
+	return Spline->GetDistanceAlongSplineAtSplinePoint(SegmentIndex + 1) - Spline->GetDistanceAlongSplineAtSplinePoint(SegmentIndex);
 }
 
-FVector UPCGSplineData::GetLocationAtDistance(int SegmentIndex, float Distance) const
+FVector UPCGSplineData::GetLocationAtDistance(int SegmentIndex, FVector::FReal Distance) const
 {
-	return Spline->GetLocationAtDistanceAlongSpline(Distance, ESplineCoordinateSpace::World);
+	return Spline->GetLocationAtDistanceAlongSpline(Spline->GetDistanceAlongSplineAtSplinePoint(SegmentIndex) + Distance, ESplineCoordinateSpace::World);
+}
+
+FTransform UPCGSplineData::GetTransformAtDistance(int SegmentIndex, FVector::FReal Distance, FBox* OutBounds) const
+{
+	if (OutBounds)
+	{
+		*OutBounds = FBox::BuildAABB(FVector::ZeroVector, FVector::OneVector);
+	}
+
+	return Spline->GetTransformAtDistanceAlongSpline(Spline->GetDistanceAlongSplineAtSplinePoint(SegmentIndex) + Distance, ESplineCoordinateSpace::World, /*bUseScale=*/true);
 }
 
 const UPCGPointData* UPCGSplineData::CreatePointData(FPCGContext* Context) const
@@ -49,29 +60,12 @@ const UPCGPointData* UPCGSplineData::CreatePointData(FPCGContext* Context) const
 	UPCGPointData* Data = NewObject<UPCGPointData>(const_cast<UPCGSplineData*>(this));
 	Data->InitializeFromData(this);
 
-	// TODO: introduce settings to sample, 
-	// should be part of the settings passed in param.
-	float SplineLength = Spline->GetSplineLength();
-	float Offset = (SplineLength - FMath::Floor(SplineLength)) / 2.0f;
-
-	TArray<FPCGPoint>& Points = Data->GetMutablePoints();
-	int NumPoints = int(SplineLength) + 1;
-	Points.SetNum(NumPoints);
+	FPCGSplineSamplerParams SamplerParams;
+	SamplerParams.Mode = EPCGSplineSamplingMode::Distance;
 
-	int Index = 0;
-	float Distance = Offset;
-	while (Distance < SplineLength)
-	{
-		Points[Index].Transform = Spline->GetTransformAtDistanceAlongSpline(Distance, ESplineCoordinateSpace::World);
-		// Exception: considering the spline holds size data in its scale, we'll reset it to identity here
-		Points[Index].Transform.SetScale3D(FVector::One());
-		Points[Index].Seed = PCGHelpers::ComputeSeed((int)Distance);
-		Points[Index].Density = 1.0f;
-		++Index;
-		Distance += 1.0f;
-	}
+	PCGSplineSampler::SampleLineData(this, this, SamplerParams, Data);
 
-	UE_LOG(LogPCG, Verbose, TEXT("Spline %s generated %d points"), *Spline->GetFName().ToString(), Points.Num());
+	UE_LOG(LogPCG, Verbose, TEXT("Spline %s generated %d points"), *Spline->GetFName().ToString(), Data->GetPoints().Num());
 
 	return Data;
 }
@@ -102,7 +96,7 @@ bool UPCGSplineData::SamplePoint(const FTransform& InTransform, const FBox& InBo
 	else
 	{
 		OutPoint.Transform = NearestTransform;
-		OutPoint.Transform.SetLocation(LocalPoint);
+		OutPoint.Transform.SetLocation(InPosition);
 		OutPoint.SetLocalBounds(InBounds);
 		OutPoint.Density = 1.0f - Distance;
 
@@ -200,7 +194,7 @@ bool UPCGSplineProjectionData::SamplePoint(const FTransform& InTransform, const
 	const FVector& SurfaceNormal = GetSurface()->GetNormal();
 
 	// Project to 2D space
-	const FTransform LocalTransform = Spline->GetComponentTransform().Inverse() * InTransform;
+	const FTransform LocalTransform = InTransform * Spline->GetComponentTransform().Inverse();
 	FVector2D LocalPosition2D = Project(LocalTransform.GetLocation());
 	float Dummy;
 	// Find nearest key on 2D spline
@@ -212,7 +206,28 @@ bool UPCGSplineProjectionData::SamplePoint(const FTransform& InTransform, const
 	const FVector PointOnLine = FMath::ClosestPointOnInfiniteLine(InPosition, InPosition + SurfaceNormal, NearestPointOnSpline);
 
 	// TODO: this is super inefficient
-	return GetSpline()->SamplePoint(FTransform(PointOnLine), InBounds, OutPoint, OutMetadata);
+	FPCGPoint SplinePoint;
+	if (GetSpline()->SamplePoint(FTransform(PointOnLine), InBounds, SplinePoint, OutMetadata))
+	{
+		FPCGPoint SurfacePoint;
+		if (GetSurface()->SamplePoint(SplinePoint.Transform, InBounds, SurfacePoint, OutMetadata))
+		{
+			OutPoint = SplinePoint;
+			OutPoint.Transform = SurfacePoint.Transform;
+			OutPoint.Density *= SurfacePoint.Density;
+			OutPoint.Color *= SurfacePoint.Color;
+
+			if (OutMetadata)
+			{
+				//TODO review op
+				OutMetadata->MergePointAttributes(SplinePoint, SurfacePoint, OutPoint, EPCGMetadataOp::Max);
+			}
+
+			return true;
+		}
+	}
+
+	return false;
 }
 
 const UPCGSplineData* UPCGSplineProjectionData::GetSpline() const
diff --git a/Engine/Plugins/Experimental/PCG/Source/PCG/Private/Elements/PCGSplineSampler.cpp b/Engine/Plugins/Experimental/PCG/Source/PCG/Private/Elements/PCGSplineSampler.cpp
new file mode 100644
--- /dev/null
+++ b/Engine/Plugins/Experimental/PCG/Source/PCG/Private/Elements/PCGSplineSampler.cpp
@@ -0,0 +1,363 @@
+// Copyright Epic Games, Inc. All Rights Reserved.
+
+#include "Elements/PCGSplineSampler.h"
+
+#include "PCGCommon.h"
+#include "Data/PCGIntersectionData.h"
+#include "Data/PCGLandscapeSplineData.h"
+#include "Data/PCGPointData.h"
+#include "Data/PCGSplineData.h"
+#include "Helpers/PCGBlueprintHelpers.h"
+#include "Helpers/PCGSettingsHelpers.h"
+
+namespace PCGSplineSampler
+{
+	struct FStepSampler
+	{
+		FStepSampler(const UPCGPolyLineData* InLineData)
+			: LineData(InLineData)
+		{
+			check(LineData);
+			CurrentSegmentIndex = 0;
+		}
+
+		virtual void Step(FTransform& OutTransform, FBox& OutBox) = 0;
+
+		bool IsDone() const
+		{
+			return CurrentSegmentIndex >= LineData->GetNumSegments();
+		}
+
+		const UPCGPolyLineData* LineData;
+		int CurrentSegmentIndex;
+	};
+
+	struct FSubdivisionStepSampler : public FStepSampler
+	{
+		FSubdivisionStepSampler(const UPCGPolyLineData* InLineData, const FPCGSplineSamplerParams& Params)
+			: FStepSampler(InLineData)
+		{
+			NumSegments = LineData->GetNumSegments();
+			SubdivisionsPerSegment = Params.SubdivisionsPerSegment;
+
+			CurrentSegmentIndex = 0;
+			SubpointIndex = 0;
+		}
+
+		virtual void Step(FTransform& OutTransform, FBox& OutBox) override
+		{
+			const FVector::FReal SegmentLength = LineData->GetSegmentLength(CurrentSegmentIndex);
+			const FVector::FReal SegmentStep = SegmentLength / (SubdivisionsPerSegment + 1);
+
+			OutTransform = LineData->GetTransformAtDistance(CurrentSegmentIndex, SubpointIndex * SegmentStep, &OutBox);
+
+			if (SubpointIndex == 0)
+			{
+				const int PreviousSegmentIndex = (CurrentSegmentIndex > 0 ? CurrentSegmentIndex : NumSegments) - 1;
+				const FVector::FReal PreviousSegmentLength = LineData->GetSegmentLength(PreviousSegmentIndex);
+				FTransform PreviousSegmentEndTransform = LineData->GetTransformAtDistance(PreviousSegmentIndex, PreviousSegmentLength);
+
+				if ((PreviousSegmentEndTransform.GetLocation() - OutTransform.GetLocation()).Length() <= KINDA_SMALL_NUMBER)
+				{
+					OutBox.Min.X *= 0.5 * PreviousSegmentLength / (PreviousSegmentEndTransform.GetScale3D().X * (SubdivisionsPerSegment + 1));
+				}
+				else
+				{
+					OutBox.Min.X *= 0.5 * SegmentStep / OutTransform.GetScale3D().X;
+				}
+			}
+			else
+			{
+				OutBox.Min.X *= 0.5 * SegmentStep / OutTransform.GetScale3D().X;
+			}
+
+			OutBox.Max.X *= 0.5 * SegmentStep / OutTransform.GetScale3D().X;
+
+			++SubpointIndex;
+			if (SubpointIndex > SubdivisionsPerSegment)
+			{
+				SubpointIndex = 0;
+				++CurrentSegmentIndex;
+			}
+		}
+
+		int NumSegments;
+		int SubdivisionsPerSegment;
+		int SubpointIndex;
+	};
+
+	struct FDistanceStepSampler : public FStepSampler
+	{
+		FDistanceStepSampler(const UPCGPolyLineData* InLineData, const FPCGSplineSamplerParams& Params)
+			: FStepSampler(InLineData)
+		{
+			DistanceIncrement = Params.DistanceIncrement;
+			CurrentDistance = 0;
+		}
+
+		virtual void Step(FTransform& OutTransform, FBox& OutBox) override
+		{
+			FVector::FReal CurrentSegmentLength = LineData->GetSegmentLength(CurrentSegmentIndex);
+			OutTransform = LineData->GetTransformAtDistance(CurrentSegmentIndex, CurrentDistance, &OutBox);
+
+			OutBox.Min.X *= DistanceIncrement / OutTransform.GetScale3D().X;
+			OutBox.Max.X *= DistanceIncrement / OutTransform.GetScale3D().X;
+
+			CurrentDistance += DistanceIncrement;
+			while(CurrentDistance > CurrentSegmentLength)
+			{
+				CurrentDistance -= CurrentSegmentLength;
+				++CurrentSegmentIndex;
+				if (!IsDone())
+				{
+					CurrentSegmentLength = LineData->GetSegmentLength(CurrentSegmentIndex);
+				}
+				else
+				{
+					break;
+				}
+			}
+		}
+
+		FVector::FReal DistanceIncrement;
+		FVector::FReal CurrentDistance;
+	};
+
+	struct FDimensionSampler
+	{
+		FDimensionSampler(const UPCGPolyLineData* InLineData, const UPCGSpatialData* InSpatialData)
+		{
+			check(InLineData);
+			LineData = InLineData;
+			SpatialData = InSpatialData;
+		}
+
+		virtual void Sample(const FTransform& InTransform, const FBox& InBox, UPCGPointData* OutPointData)
+		{
+			FPCGPoint TrivialPoint;
+			if(SpatialData->SamplePoint(InTransform, InBox, TrivialPoint, OutPointData->Metadata))
+			{
+				OutPointData->GetMutablePoints().Add(TrivialPoint);
+			}
+		}
+
+		const UPCGPolyLineData* LineData;
+		const UPCGSpatialData* SpatialData;
+	};
+
+	struct FVolumeSampler : public FDimensionSampler
+	{
+		FVolumeSampler(const UPCGPolyLineData* InLineData, const UPCGSpatialData* InSpatialData, const FPCGSplineSamplerParams& Params)
+			: FDimensionSampler(InLineData, InSpatialData)
+		{
+			Fill = Params.Fill;
+			NumPlanarSteps = 1 + ((Params.Dimension == EPCGSplineSamplingDimension::OnVertical) ? 0 : Params.NumPlanarSubdivisions);
+			NumHeightSteps = 1 + ((Params.Dimension == EPCGSplineSamplingDimension::OnHorizontal) ? 0 : Params.NumHeightSubdivisions);
+		}
+
+		virtual void Sample(const FTransform& InTransform, const FBox& InBox, UPCGPointData* OutPointData) override
+		{
+			// We're assuming that we can scale against the origin in this method so this should always be true.
+			// We will also assume that we can separate the curve into 4 ellipse sections for radius checks
+			check(InBox.Max.Y > 0 && InBox.Min.Y < 0 && InBox.Max.Z > 0 && InBox.Min.Z < 0);
+
+			const FVector::FReal YHalfStep = 0.5 * (InBox.Max.Y - InBox.Min.Y) / (FVector::FReal)NumPlanarSteps;
+			const FVector::FReal ZHalfStep = 0.5 * (InBox.Max.Z - InBox.Min.Z) / (FVector::FReal)NumHeightSteps;
+
+			FBox SubBox = InBox;
+			SubBox.Min /= FVector(1.0, NumPlanarSteps, NumHeightSteps);
+			SubBox.Max /= FVector(1.0, NumPlanarSteps, NumHeightSteps);
+
+			FPCGPoint SeedPoint;
+			if (SpatialData->SamplePoint(InTransform, InBox, SeedPoint, nullptr))
+			{
+				// Assuming the the normal to the curve is on the Y axis
+				const FVector YAxis = SeedPoint.Transform.GetScaledAxis(EAxis::Y);
+				const FVector ZAxis = SeedPoint.Transform.GetScaledAxis(EAxis::Z);
+
+				// TODO: we can optimize this if we are in the "edges only case" to pick only values that
+				FVector::FReal CurrentZ = (InBox.Min.Z + ZHalfStep);
+				while(CurrentZ <= InBox.Max.Z - ZHalfStep + KINDA_SMALL_NUMBER)
+				{
+					// Compute inner/outer distance Z contribution (squared value since we'll compare against 1)
+					const FVector::FReal InnerZ = ((NumHeightSteps > 1) ? (CurrentZ - FMath::Sign(CurrentZ) * ZHalfStep) : 0);
+					const FVector::FReal OuterZ = ((NumHeightSteps > 1) ? (CurrentZ + FMath::Sign(CurrentZ) * ZHalfStep) : 0);
+
+					// TODO: based on the current Z, we can compute the "unit" circle (as seen below) so we don't run outside of it by design, which would be a bit more efficient
+					//  care needs to be taken to make sure we are on the same "steps" though
+					FVector::FReal CurrentY = (InBox.Min.Y + YHalfStep);
+					while(CurrentY <= InBox.Max.Y - YHalfStep + KINDA_SMALL_NUMBER)
+					{
+						// Compute inner/outer distance Y contribution
+						const FVector::FReal InnerY = ((NumPlanarSteps > 1) ? (CurrentY - FMath::Sign(CurrentY) * YHalfStep) : 0);
+						const FVector::FReal OuterY = ((NumPlanarSteps > 1) ? (CurrentY + FMath::Sign(CurrentY) * YHalfStep) : 0);
+
+						bool bTestPoint = true;
+						const FVector::FReal InnerDistance = FMath::Square((CurrentZ >= 0) ? (InnerZ / InBox.Max.Z) : (InnerZ / InBox.Min.Z)) + FMath::Square((CurrentY >= 0) ? (InnerY / InBox.Max.Y) : (InnerY / InBox.Min.Y));
+						const FVector::FReal OuterDistance = FMath::Square((CurrentZ >= 0) ? (OuterZ / InBox.Max.Z) : (OuterZ / InBox.Min.Z)) + FMath::Square((CurrentY >= 0) ? (OuterY / InBox.Max.Y) : (OuterY / InBox.Min.Y));
+
+						// Check if we should consider this point based on the fill mode / the position in the iteration
+						// If the normalized z^2 + y^2 > 1, then there's no point in testing this point
+						if (InnerDistance >= 1.0 + KINDA_SMALL_NUMBER)
+						{
+							bTestPoint = false; // fully outside the unit circle
+						}
+						else if (Fill == EPCGSplineSamplingFill::EdgesOnly && OuterDistance < 1.0 - KINDA_SMALL_NUMBER)
+						{
+							bTestPoint = false; // Not the edge point
+						}
+
+						if (bTestPoint)
+						{
+							FVector TentativeLocation = InTransform.GetLocation() + YAxis * CurrentY + ZAxis * CurrentZ;
+
+							FTransform TentativeTransform = InTransform;
+							TentativeTransform.SetLocation(TentativeLocation);
+
+							FPCGPoint OutPoint;
+							if (SpatialData->SamplePoint(TentativeTransform, SubBox, OutPoint, OutPointData->Metadata))
+							{
+								OutPointData->GetMutablePoints().Add(OutPoint);
+							}
+						}
+
+						CurrentY += 2.0 * YHalfStep;
+					}
+
+					CurrentZ += 2.0 * ZHalfStep;
+				}
+			}
+		}
+
+		EPCGSplineSamplingFill Fill;
+		int NumPlanarSteps;
+		int NumHeightSteps;
+	};
+
+	void SampleLineData(const UPCGPolyLineData* LineData, const UPCGSpatialData* SpatialData, const FPCGSplineSamplerParams& Params, UPCGPointData* OutPointData)
+	{
+		check(LineData && OutPointData);
+
+		FSubdivisionStepSampler SubdivisionSampler(LineData, Params);
+		FDistanceStepSampler DistanceSampler(LineData, Params);
+
+		FStepSampler* Sampler = ((Params.Mode == EPCGSplineSamplingMode::Subdivision) ? static_cast<FStepSampler*>(&SubdivisionSampler) : static_cast<FStepSampler*>(&DistanceSampler));
+
+		FDimensionSampler TrivialDimensionSampler(LineData, SpatialData);
+		FVolumeSampler VolumeSampler(LineData, SpatialData, Params);
+
+		FDimensionSampler* ExtentsSampler = ((Params.Dimension == EPCGSplineSamplingDimension::OnSpline) ? &TrivialDimensionSampler : static_cast<FDimensionSampler*>(&VolumeSampler));
+
+		FTransform SeedTransform;
+
+		while (!Sampler->IsDone())
+		{
+			FBox SeedBox = FBox::BuildAABB(FVector::ZeroVector, FVector::OneVector);
+			// Get seed transform/box
+			Sampler->Step(SeedTransform, SeedBox);
+			// From seed point, sample in other dimensions as needed
+			ExtentsSampler->Sample(SeedTransform, SeedBox, OutPointData);
+		}
+
+		// Finally, set seed on points based on position
+		for (FPCGPoint& Point : OutPointData->GetMutablePoints())
+		{
+			Point.Seed = UPCGBlueprintHelpers::ComputeSeedFromPosition(Point.Transform.GetLocation());
+		}
+	}
+
+	const UPCGPolyLineData* GetPolyLineData(const UPCGSpatialData* InSpatialData)
+	{
+		if (!InSpatialData)
+		{
+			return nullptr;
+		}
+
+		if (const UPCGPolyLineData* LineData = Cast<const UPCGPolyLineData>(InSpatialData))
+		{
+			return LineData;
+		}
+		else if (const UPCGSplineProjectionData* SplineProjectionData = Cast<const UPCGSplineProjectionData>(InSpatialData))
+		{
+			return SplineProjectionData->GetSpline();
+		}
+		else if (const UPCGIntersectionData* Intersection = Cast<const UPCGIntersectionData>(InSpatialData))
+		{
+			if (const UPCGPolyLineData* IntersectionA = GetPolyLineData(Intersection->A))
+			{
+				return IntersectionA;
+			}
+			else if (const UPCGPolyLineData* IntersectionB = GetPolyLineData(Intersection->B))
+			{
+				return IntersectionB;
+			}
+		}
+
+		return nullptr;
+	}
+}
+
+FPCGElementPtr UPCGSplineSamplerSettings::CreateElement() const
+{
+	return MakeShared<FPCGSplineSamplerElement>();
+}
+
+TArray<FPCGPinProperties> UPCGSplineSamplerSettings::InputPinProperties() const
+{
+	TArray<FPCGPinProperties> PinProperties;
+	PinProperties.Emplace(PCGPinConstants::DefaultInputLabel, /*EPCGDataType::Point |*/ EPCGDataType::Spline | EPCGDataType::LandscapeSpline);
+
+	return PinProperties;
+}
+
+bool FPCGSplineSamplerElement::ExecuteInternal(FPCGContext* Context) const
+{
+	TRACE_CPUPROFILER_EVENT_SCOPE(FPCGSplineSamplerElement::Execute);
+
+	const UPCGSplineSamplerSettings* Settings = Context->GetInputSettings<UPCGSplineSamplerSettings>();
+	check(Settings);
+
+	TArray<FPCGTaggedData> Inputs = Context->InputData.GetInputs();
+	UPCGParamData* Params = Context->InputData.GetParams();
+
+	FPCGSplineSamplerParams SamplerParams = Settings->Params;
+	SamplerParams.Mode = PCG_GET_OVERRIDEN_VALUE(&SamplerParams, Mode, Params);
+	SamplerParams.Dimension = PCG_GET_OVERRIDEN_VALUE(&SamplerParams, Dimension, Params);
+	SamplerParams.Fill = PCG_GET_OVERRIDEN_VALUE(&SamplerParams, Fill, Params);
+	SamplerParams.SubdivisionsPerSegment = PCG_GET_OVERRIDEN_VALUE(&SamplerParams, SubdivisionsPerSegment, Params);
+	SamplerParams.DistanceIncrement = PCG_GET_OVERRIDEN_VALUE(&SamplerParams, DistanceIncrement, Params);
+	SamplerParams.NumPlanarSubdivisions = PCG_GET_OVERRIDEN_VALUE(&SamplerParams, NumPlanarSubdivisions, Params);
+	SamplerParams.NumHeightSubdivisions = PCG_GET_OVERRIDEN_VALUE(&SamplerParams, NumHeightSubdivisions, Params);
+
+	TArray<FPCGTaggedData>& Outputs = Context->OutputData.TaggedData;
+
+	for (FPCGTaggedData& Input : Inputs)
+	{
+		const UPCGSpatialData* SpatialData = Cast<const UPCGSpatialData>(Input.Data);
+
+		if (!SpatialData)
+		{
+			continue;
+		}
+
+		// TODO: do something for point data approximations
+		const UPCGPolyLineData* LineData = PCGSplineSampler::GetPolyLineData(SpatialData);
+
+		if (!LineData)
+		{
+			continue;
+		}
+
+		FPCGTaggedData& Output = Outputs.Emplace_GetRef();
+		Output = Input;
+
+		UPCGPointData* SampledPointData = NewObject<UPCGPointData>(const_cast<UPCGPolyLineData*>(LineData));
+		SampledPointData->InitializeFromData(SpatialData);
+
+		Output.Data = SampledPointData;
+
+		PCGSplineSampler::SampleLineData(LineData, SpatialData, SamplerParams, SampledPointData);
+	}
+
+	return true;
+}
\ No newline at end of file
diff --git a/Engine/Plugins/Experimental/PCG/Source/PCG/Public/Data/PCGLandscapeSplineData.h b/Engine/Plugins/Experimental/PCG/Source/PCG/Public/Data/PCGLandscapeSplineData.h
--- a/Engine/Plugins/Experimental/PCG/Source/PCG/Public/Data/PCGLandscapeSplineData.h
+++ b/Engine/Plugins/Experimental/PCG/Source/PCG/Public/Data/PCGLandscapeSplineData.h
@@ -22,8 +22,8 @@ public:
 
 	//~Begin UPCGPolyLineData interface
 	virtual int GetNumSegments() const override;
-	virtual float GetSegmentLength(int SegmentIndex) const override;
-	virtual FVector GetLocationAtDistance(int SegmentIndex, float Distance) const override;
+	virtual FVector::FReal GetSegmentLength(int SegmentIndex) const override;
+	virtual FTransform GetTransformAtDistance(int SegmentIndex, FVector::FReal Distance, FBox* OutBounds = nullptr) const override;
 	//~End UPCGPolyLineData interface
 
 	//~Begin UPCGSpatialDataWithPointCache interface
diff --git a/Engine/Plugins/Experimental/PCG/Source/PCG/Public/Data/PCGPolyLineData.h b/Engine/Plugins/Experimental/PCG/Source/PCG/Public/Data/PCGPolyLineData.h
--- a/Engine/Plugins/Experimental/PCG/Source/PCG/Public/Data/PCGPolyLineData.h
+++ b/Engine/Plugins/Experimental/PCG/Source/PCG/Public/Data/PCGPolyLineData.h
@@ -23,6 +23,8 @@ public:
 	//~End UPCGSpatialData interface
 
 	virtual int GetNumSegments() const PURE_VIRTUAL(UPCGPolyLineData::GetNumSegments, return 0;);
-	virtual float GetSegmentLength(int SegmentIndex) const PURE_VIRTUAL(UPCGPolyLineData::GetSegmentLength, return 0.0f;);
-	virtual FVector GetLocationAtDistance(int SegmentIndex, float Distance) const PURE_VIRTUAL(UPCGPolyLine::GetLocationAtDistance, return FVector(););
+	virtual FVector::FReal GetSegmentLength(int SegmentIndex) const PURE_VIRTUAL(UPCGPolyLineData::GetSegmentLength, return 0.0f;);
+	virtual FTransform GetTransformAtDistance(int SegmentIndex, FVector::FReal Distance, FBox* OutBounds = nullptr) const PURE_VIRTUAL(UPCGPolyLine::GetTransformAtDistance, return FTransform(););
+
+	virtual FVector GetLocationAtDistance(int SegmentIndex, FVector::FReal Distance) const { return GetTransformAtDistance(SegmentIndex, Distance).GetLocation(); }
 };
diff --git a/Engine/Plugins/Experimental/PCG/Source/PCG/Public/Data/PCGSplineData.h b/Engine/Plugins/Experimental/PCG/Source/PCG/Public/Data/PCGSplineData.h
--- a/Engine/Plugins/Experimental/PCG/Source/PCG/Public/Data/PCGSplineData.h
+++ b/Engine/Plugins/Experimental/PCG/Source/PCG/Public/Data/PCGSplineData.h
@@ -24,8 +24,9 @@ public:
 
 	//~Begin UPCGPolyLineData interface
 	virtual int GetNumSegments() const override;
-	virtual float GetSegmentLength(int SegmentIndex) const override;
-	virtual FVector GetLocationAtDistance(int SegmentIndex, float Distance) const override;
+	virtual FVector::FReal GetSegmentLength(int SegmentIndex) const override;
+	virtual FVector GetLocationAtDistance(int SegmentIndex, FVector::FReal Distance) const override;
+	virtual FTransform GetTransformAtDistance(int SegmentIndex, FVector::FReal Distance, FBox* OutBounds = nullptr) const override;
 	//~End UPCGPolyLineData interface
 
 	//~Begin UPCGSpatialDataWithPointCache interface
diff --git a/Engine/Plugins/Experimental/PCG/Source/PCG/Public/Elements/PCGSplineSampler.h b/Engine/Plugins/Experimental/PCG/Source/PCG/Public/Elements/PCGSplineSampler.h
new file mode 100644
--- /dev/null
+++ b/Engine/Plugins/Experimental/PCG/Source/PCG/Public/Elements/PCGSplineSampler.h
@@ -0,0 +1,98 @@
+// Copyright Epic Games, Inc. All Rights Reserved.
+
+#pragma once
+
+#include "PCGSettings.h"
+#include "PCGElement.h"
+
+#include "PCGSplineSampler.generated.h"
+
+class UPCGPolyLineData;
+class UPCGSpatialData;
+class UPCGPointData;
+
+UENUM()
+enum class EPCGSplineSamplingMode : uint8
+{
+	Subdivision = 0,
+	Distance
+};
+
+UENUM()
+enum class EPCGSplineSamplingDimension : uint8
+{
+	OnSpline = 0,
+	OnHorizontal,
+	OnVertical,
+	OnVolume
+};
+
+UENUM()
+enum class EPCGSplineSamplingFill : uint8
+{
+	Fill = 0,
+	EdgesOnly = 1
+};
+
+USTRUCT(BlueprintType)
+struct PCG_API FPCGSplineSamplerParams
+{
+	GENERATED_BODY()
+
+	UPROPERTY(BlueprintReadWrite, EditAnywhere, Category = Settings)
+	EPCGSplineSamplingMode Mode = EPCGSplineSamplingMode::Subdivision;
+
+	UPROPERTY(BlueprintReadWrite, EditAnywhere, Category = Settings)
+	EPCGSplineSamplingDimension Dimension = EPCGSplineSamplingDimension::OnSpline;
+
+	UPROPERTY(BlueprintReadWrite, EditAnywhere, Category = Settings, meta = (EditCondition = "Dimension!=EPCGSplineSamplingDimension::OnSpline"))
+	EPCGSplineSamplingFill Fill = EPCGSplineSamplingFill::Fill;
+
+	UPROPERTY(BlueprintReadWrite, EditAnywhere, Category = Settings, meta = (ClampMin = "0", EditCondition = "Mode==EPCGSplineSamplingMode::Subdivision"))
+	int32 SubdivisionsPerSegment = 1;
+
+	UPROPERTY(BlueprintReadWrite, EditAnywhere, Category = Settings, meta = (ClampMin = "0.1", EditCondition = "Mode==EPCGSplineSamplingMode::Distance"))
+	float DistanceIncrement = 100.0f;
+
+	UPROPERTY(BlueprintReadWrite, EditAnywhere, Category = Settings, meta = (ClampMin = "0", EditCondition = "Dimension==EPCGSplineSamplingDimension::OnHorizontal||Dimension==EPCGSplineSamplingDimension::OnVolume"))
+	int32 NumPlanarSubdivisions = 8;
+
+	UPROPERTY(BlueprintReadWrite, EditAnywhere, Category = Settings, meta = (ClampMin = "0", EditCondition = "Dimension==EPCGSplineSamplingDimension::OnVertical||Dimension==EPCGSplineSamplingDimension::OnVolume"))
+	int32 NumHeightSubdivisions = 8;
+};
+
+namespace PCGSplineSampler
+{
+	void SampleLineData(const UPCGPolyLineData* LineData, const UPCGSpatialData* SpatialData, const FPCGSplineSamplerParams& Params, UPCGPointData* OutPointData);
+	const UPCGPolyLineData* GetPolyLineData(const UPCGSpatialData* InSpatialData);
+}
+
+UCLASS(BlueprintType, ClassGroup = (Procedural))
+class PCG_API UPCGSplineSamplerSettings : public UPCGSettings
+{
+	GENERATED_BODY()
+
+public:
+	// ~Begin UPCGSettings interface
+#if WITH_EDITOR
+	virtual FName GetDefaultNodeName() const override { return FName(TEXT("SplineSamplerNode")); }
+	virtual EPCGSettingsType GetType() const override { return EPCGSettingsType::Sampler; }
+#endif
+
+	virtual TArray<FPCGPinProperties> InputPinProperties() const override;
+	virtual TArray<FPCGPinProperties> OutputPinProperties() const override { return Super::DefaultPointOutputPinProperties(); }
+
+protected:
+	virtual FPCGElementPtr CreateElement() const override;
+	// ~End UPCGSettings interface
+
+public:
+	UPROPERTY(BlueprintReadWrite, EditAnywhere, Category = Settings, meta = (ShowOnlyInnerProperties))
+	FPCGSplineSamplerParams Params;
+};
+
+class FPCGSplineSamplerElement : public FSimplePCGElement
+{
+protected:
+	virtual bool ExecuteInternal(FPCGContext* Context) const override;
+};
\ No newline at end of file
diff --git a/Engine/Plugins/Experimental/PCG/Source/PCG/Public/Helpers/PCGSettingsHelpers.h b/Engine/Plugins/Experimental/PCG/Source/PCG/Public/Helpers/PCGSettingsHelpers.h
--- a/Engine/Plugins/Experimental/PCG/Source/PCG/Public/Helpers/PCGSettingsHelpers.h
+++ b/Engine/Plugins/Experimental/PCG/Source/PCG/Public/Helpers/PCGSettingsHelpers.h
@@ -88,4 +88,6 @@ namespace PCGSettingsHelpers
 
 	/** Sets data from the params to a given property, matched on a name basis */
 	void SetValue(UPCGParamData* Params, UObject* Object, FProperty* Property);
-}
\ No newline at end of file
+}
+
+#define PCG_GET_OVERRIDEN_VALUE(Settings, Variable, Params) PCGSettingsHelpers::GetValue(GET_MEMBER_NAME_CHECKED(TRemovePointer<TRemoveConst<decltype(Settings)>::Type>::Type, Variable), (Settings)->Variable, Params);
\ No newline at end of file
