Add various logging to PCGEditor

diff --git a/Engine/Plugins/Experimental/PCG/Source/PCGEditor/Private/PCGEditor.cpp b/Engine/Plugins/Experimental/PCG/Source/PCGEditor/Private/PCGEditor.cpp
--- a/Engine/Plugins/Experimental/PCG/Source/PCGEditor/Private/PCGEditor.cpp
+++ b/Engine/Plugins/Experimental/PCG/Source/PCGEditor/Private/PCGEditor.cpp
@@ -51,8 +51,6 @@
 
 #define LOCTEXT_NAMESPACE "PCGGraphEditor"
 
-DEFINE_LOG_CATEGORY(LogPCGEditor);
-
 namespace FPCGEditor_private
 {
 	const FName GraphEditorID = FName(TEXT("GraphEditor"));
diff --git a/Engine/Plugins/Experimental/PCG/Source/PCGEditor/Private/PCGEditor.h b/Engine/Plugins/Experimental/PCG/Source/PCGEditor/Private/PCGEditor.h
--- a/Engine/Plugins/Experimental/PCG/Source/PCGEditor/Private/PCGEditor.h
+++ b/Engine/Plugins/Experimental/PCG/Source/PCGEditor/Private/PCGEditor.h
@@ -2,6 +2,7 @@
 
 #pragma once
 
+#include "PCGEditorModule.h"
 #include "Toolkits/AssetEditorToolkit.h"
 #include "EditorUndoClient.h"
 
@@ -13,9 +14,6 @@ class SPCGEditorGraphNodePalette;
 class UPCGEditorGraph;
 class UPCGGraph;
 
-// Logs
-DECLARE_LOG_CATEGORY_EXTERN(LogPCGEditor, Log, All);
-
 class FPCGEditor : public FAssetEditorToolkit, public FSelfRegisteringEditorUndoClient
 {
 public:
diff --git a/Engine/Plugins/Experimental/PCG/Source/PCGEditor/Private/PCGEditorGraph.cpp b/Engine/Plugins/Experimental/PCG/Source/PCGEditor/Private/PCGEditorGraph.cpp
--- a/Engine/Plugins/Experimental/PCG/Source/PCGEditor/Private/PCGEditorGraph.cpp
+++ b/Engine/Plugins/Experimental/PCG/Source/PCGEditor/Private/PCGEditorGraph.cpp
@@ -7,6 +7,7 @@
 #include "PCGEditorGraphNode.h"
 #include "PCGEditorGraphNodeInput.h"
 #include "PCGEditorGraphNodeOutput.h"
+#include "PCGEditorModule.h"
 
 #include "EdGraph/EdGraphPin.h"
 
@@ -91,6 +92,7 @@ void UPCGEditorGraph::CreateLinks(UPCGEditorGraphNodeBase* GraphNode, bool bCrea
 
 			if (!InPin)
 			{
+				UE_LOG(LogPCGEditor, Error, TEXT("Invalid InputPin for %s"), *InputPin->Properties.Label.ToString());
 				continue;
 			}
 
@@ -98,6 +100,7 @@ void UPCGEditorGraph::CreateLinks(UPCGEditorGraphNodeBase* GraphNode, bool bCrea
 			{
 				if (!InboundEdge->IsValid())
 				{
+					UE_LOG(LogPCGEditor, Error, TEXT("Invalid inbound edge for %s"), *InputPin->Properties.Label.ToString());
 					continue;
 				}
 
@@ -110,6 +113,7 @@ void UPCGEditorGraph::CreateLinks(UPCGEditorGraphNodeBase* GraphNode, bool bCrea
 					}
 					else
 					{
+						UE_LOG(LogPCGEditor, Error, TEXT("Could not create link to InputPin %s from Node %s"), *InputPin->Properties.Label.ToString(), *InboundNode->GetFName().ToString());
 						continue;
 					}
 				}
@@ -125,6 +129,7 @@ void UPCGEditorGraph::CreateLinks(UPCGEditorGraphNodeBase* GraphNode, bool bCrea
 
 			if (!OutPin)
 			{
+				UE_LOG(LogPCGEditor, Error, TEXT("Invalid OutputPin for %s"), *OutputPin->Properties.Label.ToString());
 				continue;
 			}
 
@@ -132,6 +137,7 @@ void UPCGEditorGraph::CreateLinks(UPCGEditorGraphNodeBase* GraphNode, bool bCrea
 			{
 				if (!OutboundEdge->IsValid())
 				{
+					UE_LOG(LogPCGEditor, Error, TEXT("Invalid outbound edge for %s"), *OutputPin->Properties.Label.ToString());
 					continue;
 				}
 
@@ -144,6 +150,7 @@ void UPCGEditorGraph::CreateLinks(UPCGEditorGraphNodeBase* GraphNode, bool bCrea
 					}
 					else
 					{
+						UE_LOG(LogPCGEditor, Error, TEXT("Could not create link from OutputPin %s to Node %s"), *OutputPin->Properties.Label.ToString(), *OutboundNode->GetFName().ToString());
 						continue;
 					}
 				}
diff --git a/Engine/Plugins/Experimental/PCG/Source/PCGEditor/Private/PCGEditorGraphSchemaActions.cpp b/Engine/Plugins/Experimental/PCG/Source/PCGEditor/Private/PCGEditorGraphSchemaActions.cpp
--- a/Engine/Plugins/Experimental/PCG/Source/PCGEditor/Private/PCGEditorGraphSchemaActions.cpp
+++ b/Engine/Plugins/Experimental/PCG/Source/PCGEditor/Private/PCGEditorGraphSchemaActions.cpp
@@ -6,6 +6,7 @@
 #include "PCGEditorCommon.h"
 #include "PCGEditorGraph.h"
 #include "PCGEditorGraphNode.h"
+#include "PCGEditorModule.h"
 #include "PCGGraph.h"
 #include "PCGSettings.h"
 #include "PCGSubgraph.h"
@@ -23,12 +24,14 @@ UEdGraphNode* FPCGEditorGraphSchemaAction_NewNativeElement::PerformAction(UEdGra
 	UPCGEditorGraph* EditorGraph = Cast<UPCGEditorGraph>(ParentGraph);
 	if (!EditorGraph)
 	{
+		UE_LOG(LogPCGEditor, Error, TEXT("Invalid EditorGraph"));
 		return nullptr;
 	}
 
 	UPCGGraph* PCGGraph = EditorGraph->GetPCGGraph();
 	if (!PCGGraph)
 	{
+		UE_LOG(LogPCGEditor, Error, TEXT("Invalid PCGGraph"));
 		return nullptr;
 	}
 
@@ -61,12 +64,14 @@ UEdGraphNode* FPCGEditorGraphSchemaAction_NewBlueprintElement::PerformAction(UEd
 	UPCGEditorGraph* EditorGraph = Cast<UPCGEditorGraph>(ParentGraph);
 	if (!EditorGraph)
 	{
+		UE_LOG(LogPCGEditor, Error, TEXT("Invalid EditorGraph"));
 		return nullptr;
 	}
 
 	UPCGGraph* PCGGraph = EditorGraph->GetPCGGraph();
 	if (!PCGGraph)
 	{
+		UE_LOG(LogPCGEditor, Error, TEXT("Invalid PCGGraph"));
 		return nullptr;
 	}
 
@@ -106,18 +111,21 @@ UEdGraphNode* FPCGEditorGraphSchemaAction_NewSubgraphElement::PerformAction(UEdG
 	UPCGEditorGraph* EditorGraph = Cast<UPCGEditorGraph>(ParentGraph);
 	if (!EditorGraph)
 	{
+		UE_LOG(LogPCGEditor, Error, TEXT("Invalid EditorGraph"));
 		return nullptr;
 	}
 
 	UPCGGraph* PCGGraph = EditorGraph->GetPCGGraph();
 	if (!PCGGraph)
 	{
+		UE_LOG(LogPCGEditor, Error, TEXT("Invalid PCGGraph"));
 		return nullptr;
 	}
 
 	UPCGGraph* Subgraph = CastChecked<UPCGGraph>(SubgraphObjectPath.TryLoad());
 	if (Subgraph == PCGGraph)
 	{
+		UE_LOG(LogPCGEditor, Error, TEXT("Invalid Subgraph"));
 		return nullptr;
 	}
 
@@ -154,6 +162,7 @@ UEdGraphNode* FPCGEditorGraphSchemaAction_NewComment::PerformAction(class UEdGra
 	UPCGEditorGraph* EditorGraph = Cast<UPCGEditorGraph>(ParentGraph);
 	if (!EditorGraph)
 	{
+		UE_LOG(LogPCGEditor, Error, TEXT("Invalid EditorGraph"));
 		return nullptr;
 	}
 
diff --git a/Engine/Plugins/Experimental/PCG/Source/PCGEditor/Private/PCGEditorModule.cpp b/Engine/Plugins/Experimental/PCG/Source/PCGEditor/Private/PCGEditorModule.cpp
--- a/Engine/Plugins/Experimental/PCG/Source/PCGEditor/Private/PCGEditorModule.cpp
+++ b/Engine/Plugins/Experimental/PCG/Source/PCGEditor/Private/PCGEditorModule.cpp
@@ -190,4 +190,6 @@ void FPCGEditorModule::UnregisterSettings()
 
 IMPLEMENT_MODULE(FPCGEditorModule, PCGEditor);
 
+DEFINE_LOG_CATEGORY(LogPCGEditor);
+
 #undef LOCTEXT_NAMESPACE
\ No newline at end of file
diff --git a/Engine/Plugins/Experimental/PCG/Source/PCGEditor/Private/PCGEditorModule.h b/Engine/Plugins/Experimental/PCG/Source/PCGEditor/Private/PCGEditorModule.h
--- a/Engine/Plugins/Experimental/PCG/Source/PCGEditor/Private/PCGEditorModule.h
+++ b/Engine/Plugins/Experimental/PCG/Source/PCGEditor/Private/PCGEditorModule.h
@@ -7,6 +7,9 @@
 #include "Modules/ModuleInterface.h"
 #include "Toolkits/AssetEditorToolkit.h"
 
+// Logs
+DECLARE_LOG_CATEGORY_EXTERN(LogPCGEditor, Log, All);
+
 class FMenuBuilder;
 class FPCGEditorGraphNodeFactory;
 
