PCG Determinism Testing Framework - First Iteration (WIP)
* Includes determinism tests for the PCG Volume Sampler element

diff --git a/Engine/Plugins/Experimental/PCG/Source/PCG/Private/PCGData.cpp b/Engine/Plugins/Experimental/PCG/Source/PCG/Private/PCGData.cpp
--- a/Engine/Plugins/Experimental/PCG/Source/PCG/Private/PCGData.cpp
+++ b/Engine/Plugins/Experimental/PCG/Source/PCG/Private/PCGData.cpp
@@ -156,6 +156,11 @@ bool FPCGDataCollection::operator==(const FPCGDataCollection& Other) const
 	return true;
 }
 
+bool FPCGDataCollection::operator!=(const FPCGDataCollection& Other) const
+{
+	return !operator==(Other);
+}
+
 void FPCGDataCollection::AddToRootSet(FPCGRootSet& RootSet) const
 {
 	for (const FPCGTaggedData& Data : TaggedData)
diff --git a/Engine/Plugins/Experimental/PCG/Source/PCG/Private/PCGElement.cpp b/Engine/Plugins/Experimental/PCG/Source/PCG/Private/PCGElement.cpp
--- a/Engine/Plugins/Experimental/PCG/Source/PCG/Private/PCGElement.cpp
+++ b/Engine/Plugins/Experimental/PCG/Source/PCG/Private/PCGElement.cpp
@@ -23,6 +23,10 @@ bool IPCGElement::Execute(FPCGContext* Context) const
 	}
 
 	const UPCGSettings* Settings = Context->GetInputSettings<UPCGSettings>();
+
+
+
+
 	if (Settings && Settings->ExecutionMode == EPCGSettingsExecutionMode::Disabled)
 	{
 		//Pass-through
@@ -181,9 +185,10 @@ void IPCGElement::CleanupAndValidateOutput(FPCGContext* Context) const
 				{
 					PCGE_LOG(Warning, "Output generated for pin %s but cannot be routed", *TaggedData.Pin.ToString());
 				}
-				else if(TaggedData.Data && !(OutputPinProperties[MatchIndex].AllowedTypes & TaggedData.Data->GetDataType()))
+				// TODO: Temporary fix for Settings directly from InputData (ie. from elements with code and not PCG nodes)
+				else if(TaggedData.Data && !(OutputPinProperties[MatchIndex].AllowedTypes & TaggedData.Data->GetDataType()) && TaggedData.Data->GetDataType() != EPCGDataType::Settings)
 				{
-					PCGE_LOG(Warning, "Output generated for pin %s does not have a compatible type", *TaggedData.Pin.ToString());
+					PCGE_LOG(Warning, "Output generated for pin %s does not have a compatible type: %s", *TaggedData.Pin.ToString(), *UEnum::GetValueAsString(TaggedData.Data->GetDataType()));
 				}
 			}
 		}
diff --git a/Engine/Plugins/Experimental/PCG/Source/PCG/Private/Tests/Determinism/PCGDeterminismTestsCommon.cpp b/Engine/Plugins/Experimental/PCG/Source/PCG/Private/Tests/Determinism/PCGDeterminismTestsCommon.cpp
new file mode 100644
--- /dev/null
+++ b/Engine/Plugins/Experimental/PCG/Source/PCG/Private/Tests/Determinism/PCGDeterminismTestsCommon.cpp
@@ -0,0 +1,297 @@
+// Copyright Epic Games, Inc. All Rights Reserved.
+
+#include "Tests/Determinism/PCGDeterminismTestsCommon.h"
+
+#include "PCGComponent.h"
+#include "PCGGraph.h"
+#include "PCGPoint.h"
+#include "Data/PCGPointData.h"
+#include "Data/PCGPolyLineData.h"
+#include "Data/PCGPrimitiveData.h"
+#include "Data/PCGSurfaceData.h"
+#include "Data/PCGVolumeData.h"
+#include "Tests/PCGTestsCommon.h"
+
+namespace PCGDeterminismTests
+{
+	FTestData::FTestData(int32 RandomSeed) :
+		Seed(RandomSeed),
+		RandomStream(Seed)
+	{
+		check(GEditor);
+		UWorld* EditorWorld = GEditor->GetEditorWorldContext().World();
+		check(EditorWorld);
+
+		// No getting the level dirty
+		FActorSpawnParameters TransientActorParameters;
+		TransientActorParameters.bHideFromSceneOutliner = true;
+		TransientActorParameters.bTemporaryEditorActor = true;
+		TransientActorParameters.ObjectFlags = RF_Transient;
+		TestActor = EditorWorld->SpawnActor<AActor>(AActor::StaticClass(), TransientActorParameters);
+		check(TestActor);
+
+		TestComponent = NewObject<UPCGComponent>(TestActor, FName(TEXT("PCG Test Component")), RF_Transient);
+		check(TestComponent);
+		TestActor->AddInstanceComponent(TestComponent);
+		TestComponent->RegisterComponent();
+
+		UPCGGraph* TestGraph = NewObject<UPCGGraph>(TestComponent, FName(TEXT("PCG Test Graph")), RF_Transient);
+		check(TestGraph);
+		TestComponent->SetGraph(TestGraph);
+	}
+
+	void FTestData::Reset()
+	{
+		// Clear all the data
+		RandomStream.Reset();
+		InputData.TaggedData.Empty();
+		Settings = nullptr;
+	}
+
+	void AddVolumeInputData(FPCGDataCollection& InputData, const FVector& Location, const FVector& HalfSize, const FVector& VoxelSize)
+	{
+		UPCGVolumeData* VolumeData = PCGTestsCommon::CreateVolumeData(FBox::BuildAABB(Location, HalfSize));
+		VolumeData->VoxelSize = VoxelSize;
+		InputData.TaggedData.Emplace_GetRef().Data = VolumeData;
+	}
+
+	void AddGenericVolumeInputData(FTestData& TestData)
+	{
+		AddVolumeInputData(TestData.InputData, FVector::ZeroVector, FVector::OneVector * 500.f, FVector::OneVector * 200.f);
+	}
+
+	void AddRandomizedVolumeInputData(FTestData& TestData)
+	{
+		AddVolumeInputData(TestData.InputData, TestData.RandomStream.VRand() * 500.f, FVector::OneVector * 500.f + TestData.RandomStream.VRand() * 200.f, FVector::OneVector * 200.f + TestData.RandomStream.VRand() * 100.f);
+	}
+
+	bool DataCollectionsAreIdentical(const FPCGDataCollection& FirstCollection, const FPCGDataCollection& SecondCollection)
+	{
+		if (FirstCollection.TaggedData.Num() != SecondCollection.TaggedData.Num())
+		{
+			return false;
+		}
+
+		TArray<int32> FirstCollectionComparableIndices;
+		TArray<int32> SecondCollectionComparableIndices;
+
+		// Find comparable data from first collection
+		for (int32 I = 0; I < FirstCollection.TaggedData.Num(); ++I)
+		{
+			if (DataIsComparable(FirstCollection.TaggedData[I].Data->GetDataType()))
+			{
+				FirstCollectionComparableIndices.Emplace(I);
+			}
+		}
+
+		// Find comparable data from second collection
+		for (int32 I = 0; I < SecondCollection.TaggedData.Num(); ++I)
+		{
+			// Check if its a type we care about comparing
+			if (DataIsComparable(SecondCollection.TaggedData[I].Data->GetDataType()))
+			{
+				SecondCollectionComparableIndices.Emplace(I);
+			}
+		}
+
+		TArray<int32> MatchedIndices;
+
+		// Find matches for each TaggedData array
+		for (int32 I : FirstCollectionComparableIndices)
+		{
+			bool bFound = false;
+			for (int32 J : SecondCollectionComparableIndices)
+			{
+				if (MatchedIndices.Contains(J))
+				{
+					continue;
+				}
+
+				EPCGDataType FirstDataType = FirstCollection.TaggedData[I].Data->GetDataType();
+				EPCGDataType SecondDataType = SecondCollection.TaggedData[I].Data->GetDataType();
+
+				// Only compare if they are the same type and same pin label
+				if (FirstDataType == SecondDataType && FirstCollection.TaggedData[I].Pin == SecondCollection.TaggedData[I].Pin)
+				{
+					if (GetCompareFunction(FirstDataType)(FirstCollection.TaggedData[I].Data, SecondCollection.TaggedData[J].Data))
+					{
+						MatchedIndices.Emplace(J);
+						bFound = true;
+						break;
+					}
+				}
+			}
+
+			if (!bFound)
+			{
+				return false;
+			}
+		}
+
+		return true;
+	}
+
+	bool SpatialDataIsIdentical(const UPCGData* FirstSpatialData, const UPCGData* SecondSpatialData)
+	{
+		check(FirstSpatialData);
+		check(SecondSpatialData);
+
+		if (Cast<UPCGPointData>(FirstSpatialData) != nullptr && Cast<UPCGPointData>(SecondSpatialData) != nullptr)
+		{
+			return PointDataIsIdentical(FirstSpatialData, SecondSpatialData);
+		}
+		else if (Cast<UPCGVolumeData>(FirstSpatialData) != nullptr && Cast<UPCGVolumeData>(SecondSpatialData) != nullptr)
+		{
+			// TODO: Implement Volume Data comparison
+			UE_LOG(LogPCG, Warning, TEXT("Volume comparison unimplemented."));
+			return ComparisonIsUnimplemented(FirstSpatialData, SecondSpatialData);
+		}
+		else if (Cast<UPCGSurfaceData>(FirstSpatialData) != nullptr && Cast<UPCGSurfaceData>(SecondSpatialData) != nullptr)
+		{
+			// TODO: Implement Surface Data comparison
+			UE_LOG(LogPCG, Warning, TEXT("Surface comparison unimplemented."));
+			return ComparisonIsUnimplemented(FirstSpatialData, SecondSpatialData);
+		}
+		else if (Cast<UPCGPolyLineData>(FirstSpatialData) != nullptr && Cast<UPCGPolyLineData>(SecondSpatialData) != nullptr)
+		{
+			// TODO: Implement PolyLine Data comparison
+			UE_LOG(LogPCG, Warning, TEXT("PolyLine comparison unimplemented."));
+			return ComparisonIsUnimplemented(FirstSpatialData, SecondSpatialData);
+		}
+		else if (Cast<UPCGPrimitiveData>(FirstSpatialData) != nullptr && Cast<UPCGPrimitiveData>(SecondSpatialData) != nullptr)
+		{
+			// TODO: Implement Primitive Data comparison
+			UE_LOG(LogPCG, Warning, TEXT("Primitive comparison unimplemented."));
+			return ComparisonIsUnimplemented(FirstSpatialData, SecondSpatialData);
+		}
+
+		return false;
+	}
+
+	bool PointDataIsIdentical(const UPCGData* FirstData, const UPCGData* SecondData)
+	{
+		const UPCGPointData* FirstPointData = Cast<const UPCGPointData>(FirstData);
+		const UPCGPointData* SecondPointData = Cast<const UPCGPointData>(SecondData);
+		check(FirstPointData);
+		check(SecondPointData);
+
+		const TArray<FPCGPoint>& FirstPoints = Cast<const UPCGPointData>(FirstPointData)->GetPoints();
+		const TArray<FPCGPoint>& SecondPoints = Cast<const UPCGPointData>(SecondPointData)->GetPoints();
+
+		if (FirstPoints.Num() != SecondPoints.Num())
+		{
+			return false;
+		}
+
+		// Match points based on indices
+		TSet<int32> RemainingIndices;
+		for (int32 i = 0; i < SecondPoints.Num(); ++i)
+		{
+			RemainingIndices.Add(i);
+		}
+
+		for (int32 i = 0; i < FirstPoints.Num(); ++i)
+		{
+			bool bFound = false;
+
+			// TODO: Maybe use Octree to prune. Is the saved computation worth the complexity of more complex index tracking?
+			for (int32 j : RemainingIndices)
+			{
+				if (PCGTestsCommon::PointsAreIdentical(FirstPoints[i], SecondPoints[j]))
+				{
+					RemainingIndices.Remove(j);
+					bFound = true;
+					break;
+				}
+			}
+
+			// Couldn't find the matching point
+			if (!bFound)
+			{
+				return false;
+			}
+		}
+
+		return true;
+	}
+
+	bool ComparisonIsUnimplemented(const UPCGData* FirstPointData, const UPCGData* SecondPointData)
+	{
+		return false;
+	}
+
+	bool DataIsComparable(EPCGDataType DataType)
+	{
+		// Comparable data types
+		if (DataType == EPCGDataType::Spatial)
+		{
+			return true;
+		}
+
+		// TODO: Data types that don't need to be compared
+		if (DataType == EPCGDataType::None || DataType == EPCGDataType::Other || DataType == EPCGDataType::Settings)
+		{
+			return false;
+		}
+
+		UE_LOG(LogPCG, Warning, TEXT("Unknown data comparison type: %s"), *UEnum::GetValueAsString(DataType));
+		return false;
+	}
+
+	void ShuffleInputOrder(FTestData& TestData)
+	{
+		const int32 LastIndex = TestData.InputData.TaggedData.Num() - 1;
+		for (int32 I = 0; I <= LastIndex; ++I)
+		{
+			int32 Index = TestData.RandomStream.RandRange(I, LastIndex);
+
+			if (I != Index)
+			{
+				TestData.InputData.TaggedData.Swap(I, Index);
+			}
+		}
+	}
+
+	TFunction<bool(const UPCGData*, const UPCGData*)> GetCompareFunction(EPCGDataType DataType)
+	{
+		if (!DataIsComparable(DataType))
+		{
+			// Should never reach here
+			UE_LOG(LogPCG, Warning, TEXT("Attempting to compare incomparable data."));
+			return ComparisonIsUnimplemented;
+		}
+
+		switch (DataType)
+		{
+			// TODO: Many more data type comparisons to be added
+		case EPCGDataType::Spatial:
+			return SpatialDataIsIdentical;
+		default:
+			UE_LOG(LogPCG, Warning, TEXT("Comparable PCG DataType has no 'comparison function'."));
+			return ComparisonIsUnimplemented;
+		}
+	}
+
+	bool ExecutionIsDeterministic(FTestData& FirstTestData, FTestData& SecondTestData)
+	{
+		FPCGElementPtr FirstElement = FirstTestData.Settings->GetElement();
+		FPCGElementPtr SecondElement = SecondTestData.Settings->GetElement();
+
+		TUniquePtr<FPCGContext> FirstContext(FirstElement->Initialize(FirstTestData.InputData, FirstTestData.TestComponent, nullptr));
+		TUniquePtr<FPCGContext> SecondContext(SecondElement->Initialize(SecondTestData.InputData, SecondTestData.TestComponent, nullptr));
+
+		FirstContext->NumAvailableTasks = 1;
+		SecondContext->NumAvailableTasks = 1;
+
+		// Execute both elements until complete
+		while (!FirstElement->Execute(FirstContext.Get()))
+		{
+		}
+
+		while (!SecondElement->Execute(SecondContext.Get()))
+		{
+		}
+
+		return PCGDeterminismTests::DataCollectionsAreIdentical(FirstContext->OutputData, SecondContext->OutputData);
+	}
+}
\ No newline at end of file
diff --git a/Engine/Plugins/Experimental/PCG/Source/PCG/Private/Tests/Determinism/PCGVolumeSamplerDeterminismTest.cpp b/Engine/Plugins/Experimental/PCG/Source/PCG/Private/Tests/Determinism/PCGVolumeSamplerDeterminismTest.cpp
new file mode 100644
--- /dev/null
+++ b/Engine/Plugins/Experimental/PCG/Source/PCG/Private/Tests/Determinism/PCGVolumeSamplerDeterminismTest.cpp
@@ -0,0 +1,115 @@
+// Copyright Epic Games, Inc. All Rights Reserved.
+
+#include "PCGComponent.h"
+#include "PCGGraph.h"
+#include "Data/PCGSpatialData.h"
+#include "Data/PCGPointData.h"
+#include "Elements/PCGVolumeSampler.h"
+#include "Tests/PCGTestsCommon.h"
+#include "Tests/Determinism/PCGDeterminismTestsCommon.h"
+
+IMPLEMENT_CUSTOM_SIMPLE_AUTOMATION_TEST(FPCGVolumeSamplerDeterminismSingleSameDataTest, FPCGTestBaseClass, "pcg.determinism.VolumeSampler.SingleSameData", PCGTestsCommon::TestFlags)
+IMPLEMENT_CUSTOM_SIMPLE_AUTOMATION_TEST(FPCGVolumeSamplerDeterminismSingleIdenticalDataTest, FPCGTestBaseClass, "pcg.determinism.VolumeSampler.SingleIdenticalData", PCGTestsCommon::TestFlags)
+IMPLEMENT_CUSTOM_SIMPLE_AUTOMATION_TEST(FPCGVolumeSamplerDeterminismMultipleSameDataTest, FPCGTestBaseClass, "pcg.determinism.VolumeSampler.MultipleSameData", PCGTestsCommon::TestFlags)
+IMPLEMENT_CUSTOM_SIMPLE_AUTOMATION_TEST(FPCGVolumeSamplerDeterminismMultipleIdenticalDataTest, FPCGTestBaseClass, "pcg.determinism.VolumeSampler.MultipleIdenticalData", PCGTestsCommon::TestFlags)
+IMPLEMENT_CUSTOM_SIMPLE_AUTOMATION_TEST(FPCGVolumeSamplerDeterminismOrderIndependenceTest, FPCGTestBaseClass, "pcg.determinism.VolumeSampler.OrderIndependence", PCGTestsCommon::TestFlags)
+
+namespace
+{
+	void VolumeSettingsDelegate(UPCGVolumeSamplerSettings* Settings, FRandomStream& RandomStream)
+	{
+		Settings->VoxelSize = FVector::OneVector * 200.f + RandomStream.VRand() * 100.f;
+	}
+}
+
+bool FPCGVolumeSamplerDeterminismSingleSameDataTest::RunTest(const FString& Parameters)
+{
+	// Test single same data
+	PCGDeterminismTests::FTestData FirstTestData(PCGDeterminismTests::DefaultSeed);
+
+	PCGDeterminismTests::GenerateSettings<UPCGVolumeSamplerSettings>(FirstTestData, VolumeSettingsDelegate);
+	PCGDeterminismTests::AddRandomizedVolumeInputData(FirstTestData);
+
+	TestTrue("Same single input and settings, same output", PCGDeterminismTests::ExecutionIsDeterministic(FirstTestData, FirstTestData));
+
+	return true;
+}
+
+bool FPCGVolumeSamplerDeterminismSingleIdenticalDataTest::RunTest(const FString& Parameters)
+{
+	// Test single identical data
+	PCGDeterminismTests::FTestData FirstTestData(PCGDeterminismTests::DefaultSeed);
+	PCGDeterminismTests::FTestData SecondTestData(PCGDeterminismTests::DefaultSeed);
+
+	PCGDeterminismTests::GenerateSettings<UPCGVolumeSamplerSettings>(FirstTestData, VolumeSettingsDelegate);
+	PCGDeterminismTests::GenerateSettings<UPCGVolumeSamplerSettings>(SecondTestData, VolumeSettingsDelegate);
+
+	PCGDeterminismTests::AddRandomizedVolumeInputData(FirstTestData);
+	PCGDeterminismTests::AddRandomizedVolumeInputData(SecondTestData);
+
+	TestTrue("Identical single input and settings, same output", PCGDeterminismTests::ExecutionIsDeterministic(FirstTestData, SecondTestData));
+
+	return true;
+}
+
+bool FPCGVolumeSamplerDeterminismMultipleSameDataTest::RunTest(const FString& Parameters)
+{
+	PCGDeterminismTests::FTestData FirstTestData(PCGDeterminismTests::DefaultSeed);
+
+	// Test multiple same data
+	PCGDeterminismTests::GenerateSettings<UPCGVolumeSamplerSettings>(FirstTestData, VolumeSettingsDelegate);
+
+	// Add many random volumes
+	for (int32 I = 0; I < 10; ++I)
+	{
+		PCGDeterminismTests::AddRandomizedVolumeInputData(FirstTestData);
+	}
+
+	TestTrue("Same multiple input, same output", PCGDeterminismTests::ExecutionIsDeterministic(FirstTestData, FirstTestData));
+
+	return true;
+}
+
+bool FPCGVolumeSamplerDeterminismMultipleIdenticalDataTest::RunTest(const FString& Parameters)
+{
+	// Test multiple identical data
+	PCGDeterminismTests::FTestData FirstTestData(PCGDeterminismTests::DefaultSeed);
+	PCGDeterminismTests::FTestData SecondTestData(PCGDeterminismTests::DefaultSeed);
+
+	PCGDeterminismTests::GenerateSettings<UPCGVolumeSamplerSettings>(FirstTestData, VolumeSettingsDelegate);
+	PCGDeterminismTests::GenerateSettings<UPCGVolumeSamplerSettings>(SecondTestData, VolumeSettingsDelegate);
+
+	// Add many random volumes
+	for (int32 I = 0; I < 10; ++I)
+	{
+		PCGDeterminismTests::AddRandomizedVolumeInputData(FirstTestData);
+		PCGDeterminismTests::AddRandomizedVolumeInputData(SecondTestData);
+	}
+
+	TestTrue("Identical multiple input, same output", PCGDeterminismTests::ExecutionIsDeterministic(FirstTestData, SecondTestData));
+
+	return true;
+}
+
+bool FPCGVolumeSamplerDeterminismOrderIndependenceTest::RunTest(const FString& Parameters)
+{
+	// Test shuffled input data
+	PCGDeterminismTests::FTestData FirstTestData(PCGDeterminismTests::DefaultSeed);
+	PCGDeterminismTests::FTestData SecondTestData(PCGDeterminismTests::DefaultSeed);
+
+	PCGDeterminismTests::GenerateSettings<UPCGVolumeSamplerSettings>(FirstTestData, VolumeSettingsDelegate);
+	PCGDeterminismTests::GenerateSettings<UPCGVolumeSamplerSettings>(SecondTestData, VolumeSettingsDelegate);
+
+	// Add many random volumes
+	for (int32 I = 0; I < 10; ++I)
+	{
+		PCGDeterminismTests::AddRandomizedVolumeInputData(FirstTestData);
+		PCGDeterminismTests::AddRandomizedVolumeInputData(SecondTestData);
+	}
+
+	PCGDeterminismTests::ShuffleInputOrder(SecondTestData);
+
+	TestTrue("Shuffled input order, same output", PCGDeterminismTests::ExecutionIsDeterministic(FirstTestData, SecondTestData));
+
+	return true;
+}
diff --git a/Engine/Plugins/Experimental/PCG/Source/PCG/Private/Tests/PCGTestsCommon.cpp b/Engine/Plugins/Experimental/PCG/Source/PCG/Private/Tests/PCGTestsCommon.cpp
--- a/Engine/Plugins/Experimental/PCG/Source/PCG/Private/Tests/PCGTestsCommon.cpp
+++ b/Engine/Plugins/Experimental/PCG/Source/PCG/Private/Tests/PCGTestsCommon.cpp
@@ -6,9 +6,9 @@
 #include "PCGSettings.h"
 #include "Data/PCGPointData.h"
 #include "Data/PCGPolyLineData.h"
+#include "Data/PCGPrimitiveData.h"
 #include "Data/PCGSurfaceData.h"
 #include "Data/PCGVolumeData.h"
-#include "Data/PCGPrimitiveData.h"
 
 #include "GameFramework/Actor.h"
 
@@ -133,8 +133,16 @@ namespace PCGTestsCommon
 
 	bool PointsAreIdentical(const FPCGPoint& FirstPoint, const FPCGPoint& SecondPoint)
 	{
-		// TODO: should do a full point comparison, not only on a positional basis
-		return (FirstPoint.Transform.GetLocation() - SecondPoint.Transform.GetLocation()).SquaredLength() < KINDA_SMALL_NUMBER;
+		// Trivial checks first for pruning
+		if (FirstPoint.Density != SecondPoint.Density || FirstPoint.Steepness != SecondPoint.Steepness ||
+			FirstPoint.BoundsMin != SecondPoint.BoundsMin || FirstPoint.BoundsMax != SecondPoint.BoundsMax ||
+			FirstPoint.Color != SecondPoint.Color)
+		{
+			return false;
+		}
+
+		// Transform checks with epsilon
+		return FirstPoint.Transform.Equals(SecondPoint.Transform);
 	}
 }
 
@@ -178,7 +186,8 @@ bool FPCGTestBaseClass::SmokeTestAnyValidInput(UPCGSettings* InSettings, TFuncti
 
 	check(InputIndices.Num() == InputsPerProperties.Num());
 
-	while (1)
+	bool bDone = false;
+	while (!bDone)
 	{
 		// Prepare input
 		FPCGDataCollection InputData;
@@ -187,18 +196,19 @@ bool FPCGTestBaseClass::SmokeTestAnyValidInput(UPCGSettings* InSettings, TFuncti
 			InputData.TaggedData.Append(InputsPerProperties[PinIndex][InputIndices[PinIndex]].TaggedData);
 		}
 
-		// Perform execution
-		FPCGContext* Context = Element->Initialize(InputData, nullptr, nullptr);
+		TUniquePtr<FPCGContext> Context(Element->Initialize(InputData, nullptr, nullptr));
 		Context->NumAvailableTasks = 1;
-		Element->Execute(Context);
+		
+		// Execute element until done
+		while (!Element->Execute(Context.Get()))
+		{
+		}
 
 		if (ValidationFn)
 		{
 			TestTrue("Validation", ValidationFn(Context->InputData, Context->OutputData));
 		}
 
-		delete Context;
-
 		// Bump indices
 		int BumpIndex = 0;
 		while (BumpIndex < InputIndices.Num())
@@ -217,7 +227,7 @@ bool FPCGTestBaseClass::SmokeTestAnyValidInput(UPCGSettings* InSettings, TFuncti
 
 		if (BumpIndex == InputIndices.Num())
 		{
-			break;
+			bDone = true;
 		}
 	}
 
diff --git a/Engine/Plugins/Experimental/PCG/Source/PCG/Public/PCGData.h b/Engine/Plugins/Experimental/PCG/Source/PCG/Public/PCGData.h
--- a/Engine/Plugins/Experimental/PCG/Source/PCG/Public/PCGData.h
+++ b/Engine/Plugins/Experimental/PCG/Source/PCG/Public/PCGData.h
@@ -75,6 +75,7 @@ struct PCG_API FPCGDataCollection
 	const UPCGSettings* GetSettings(const UPCGSettings* InDefaultSettings) const;
 
 	bool operator==(const FPCGDataCollection& Other) const;
+	bool operator!=(const FPCGDataCollection& Other) const;
 	void AddToRootSet(FPCGRootSet& RootSet) const;
 	void RemoveFromRootSet(FPCGRootSet& RootSet) const;
 
diff --git a/Engine/Plugins/Experimental/PCG/Source/PCG/Public/Tests/Determinism/PCGDeterminismTestsCommon.h b/Engine/Plugins/Experimental/PCG/Source/PCG/Public/Tests/Determinism/PCGDeterminismTestsCommon.h
new file mode 100644
--- /dev/null
+++ b/Engine/Plugins/Experimental/PCG/Source/PCG/Public/Tests/Determinism/PCGDeterminismTestsCommon.h
@@ -0,0 +1,75 @@
+// Copyright Epic Games, Inc. All Rights Reserved.
+
+#pragma once
+
+#include "PCGCommon.h"
+#include "PCGSettings.h"
+
+class AActor;
+class UPCGComponent;
+class UPCGData;
+struct FPCGDataCollection;
+
+namespace PCGDeterminismTests
+{
+	constexpr int32 DefaultSeed = 42;
+
+	struct FTestData
+	{
+		explicit FTestData(int32 Seed = DefaultSeed);
+
+		void Reset();
+
+		AActor* TestActor;
+		UPCGComponent* TestComponent;
+		FPCGDataCollection InputData;
+		UPCGSettings* Settings;
+		int32 Seed;
+		FRandomStream RandomStream;
+	};
+
+	/** Helper for adding VolumeData to an InputData */
+	void AddVolumeInputData(FPCGDataCollection& InputData, const FVector& Location, const FVector& HalfSize, const FVector& VoxelSize);
+
+	/** Adds generic Volume Spatial Data at the origin */
+	void AddGenericVolumeInputData(FTestData& TestData);
+	/** Adds randomized Volume Spatial Data at the origin */
+	void AddRandomizedVolumeInputData(FTestData& TestData);
+
+	/** Validates whether two DataCollection objects are identical */
+	bool DataCollectionsAreIdentical(const FPCGDataCollection& FirstCollection, const FPCGDataCollection& SecondCollection);
+	/** Validates whether two SpatialData objects are identical */
+	bool SpatialDataIsIdentical(const UPCGData* FirstSpatialData, const UPCGData* SecondSpatialData);
+	/** Validates whether two PointData objects are identical */
+	bool PointDataIsIdentical(const UPCGData* FirstPointData, const UPCGData* SecondPointData);
+	/** A catch function for unimplemented comparisons */
+	bool ComparisonIsUnimplemented(const UPCGData* FirstPointData, const UPCGData* SecondPointData);
+
+	/** Determines if a PCG DataType is a type that should be compared */
+	bool DataIsComparable(EPCGDataType DataType);
+
+	/** Randomizes the order of InputData */
+	void ShuffleInputOrder(FTestData& TestData);
+
+	/** Gets a comparison function to compare two of a specific DataType */
+	TFunction<bool(const UPCGData*, const UPCGData*)> GetCompareFunction(EPCGDataType DataType);
+
+	/** Execute the elements for each valid input and compare if all the outputs are identical */
+	bool ExecutionIsDeterministic(FTestData& FirstTestData, FTestData& SecondTestData);
+
+	/** Generates settings based upon a UPCGSettings subclass */
+	template<typename SettingsType>
+	void GenerateSettings(FTestData& TestData, TFunction<void(SettingsType*, FRandomStream&)> ExtraSettingsDelegate = nullptr)
+	{
+		TestData.Settings = NewObject<SettingsType>();
+		check(TestData.Settings);
+		TestData.Settings->Seed = TestData.Seed;
+
+		TestData.InputData.TaggedData.Emplace_GetRef().Data = TestData.Settings;
+
+		if (ExtraSettingsDelegate)
+		{
+			ExtraSettingsDelegate(Cast<SettingsType>(TestData.Settings), TestData.RandomStream);
+		}
+	}
+}
\ No newline at end of file
diff --git a/Engine/Plugins/Experimental/PCG/Source/PCG/Public/Tests/PCGTestsCommon.h b/Engine/Plugins/Experimental/PCG/Source/PCG/Public/Tests/PCGTestsCommon.h
--- a/Engine/Plugins/Experimental/PCG/Source/PCG/Public/Tests/PCGTestsCommon.h
+++ b/Engine/Plugins/Experimental/PCG/Source/PCG/Public/Tests/PCGTestsCommon.h
@@ -4,17 +4,18 @@
 
 #include "Misc/AutomationTest.h"
 #include "Tests/AutomationCommon.h"
-#include "PCGPoint.h"
 
 class AActor;
+class UPCGParamData;
 class UPCGPointData;
 class UPCGPolyLineData;
+class UPCGPrimitiveData;
 class UPCGSurfaceData;
 class UPCGVolumeData;
-class UPCGPrimitiveData;
-class UPCGParamData;
 class UPCGSettings;
 struct FPCGDataCollection;
+struct FPCGPinProperties;
+struct FPCGPoint;
 
 namespace PCGTestsCommon
 {
@@ -34,6 +35,9 @@ namespace PCGTestsCommon
 	UPCGVolumeData* CreateVolumeData(const FBox& InBounds = FBox::BuildAABB(FVector::ZeroVector, FVector::OneVector * 100));
 	UPCGPrimitiveData* CreatePrimitiveData();
 
+	TArray<FPCGDataCollection> GenerateAllowedData(const FPCGPinProperties& PinProperties);
+
+	/** Validates that two Spatial Points are identical */
 	bool PointsAreIdentical(const FPCGPoint& FirstPoint, const FPCGPoint& SecondPoint);
 }
 
