Tentative build fix on behalf of Ryan Buehler

#rnx

diff --git a/Engine/Plugins/Experimental/PCG/Source/PCG/Private/Tests/Determinism/PCGDeterminismTestsCommon.cpp b/Engine/Plugins/Experimental/PCG/Source/PCG/Private/Tests/Determinism/PCGDeterminismTestsCommon.cpp
--- a/Engine/Plugins/Experimental/PCG/Source/PCG/Private/Tests/Determinism/PCGDeterminismTestsCommon.cpp
+++ b/Engine/Plugins/Experimental/PCG/Source/PCG/Private/Tests/Determinism/PCGDeterminismTestsCommon.cpp
@@ -12,12 +12,19 @@
 #include "Data/PCGVolumeData.h"
 #include "Tests/PCGTestsCommon.h"
 
+#include "GameFramework/Actor.h"
+
+#if WITH_EDITOR
+#include "Editor.h"
+#endif
+
 namespace PCGDeterminismTests
 {
 	FTestData::FTestData(int32 RandomSeed) :
 		Seed(RandomSeed),
 		RandomStream(Seed)
 	{
+#if WITH_EDITOR
 		check(GEditor);
 		UWorld* EditorWorld = GEditor->GetEditorWorldContext().World();
 		check(EditorWorld);
@@ -38,6 +45,11 @@ namespace PCGDeterminismTests
 		UPCGGraph* TestGraph = NewObject<UPCGGraph>(TestComponent, FName(TEXT("PCG Test Graph")), RF_Transient);
 		check(TestGraph);
 		TestComponent->SetGraph(TestGraph);
+#else
+		TestActor = nullptr;
+		TestComponent = nullptr;
+		Settings = nullptr;
+#endif
 	}
 
 	void FTestData::Reset()
diff --git a/Engine/Plugins/Experimental/PCG/Source/PCG/Private/Tests/Determinism/PCGVolumeSamplerDeterminismTest.cpp b/Engine/Plugins/Experimental/PCG/Source/PCG/Private/Tests/Determinism/PCGVolumeSamplerDeterminismTest.cpp
--- a/Engine/Plugins/Experimental/PCG/Source/PCG/Private/Tests/Determinism/PCGVolumeSamplerDeterminismTest.cpp
+++ b/Engine/Plugins/Experimental/PCG/Source/PCG/Private/Tests/Determinism/PCGVolumeSamplerDeterminismTest.cpp
@@ -24,6 +24,7 @@ namespace
 
 bool FPCGVolumeSamplerDeterminismSingleSameDataTest::RunTest(const FString& Parameters)
 {
+#if WITH_EDITOR
 	// Test single same data
 	PCGDeterminismTests::FTestData FirstTestData(PCGDeterminismTests::DefaultSeed);
 
@@ -32,11 +33,13 @@ bool FPCGVolumeSamplerDeterminismSingleSameDataTest::RunTest(const FString& Para
 
 	TestTrue("Same single input and settings, same output", PCGDeterminismTests::ExecutionIsDeterministic(FirstTestData, FirstTestData));
 
+#endif
 	return true;
 }
 
 bool FPCGVolumeSamplerDeterminismSingleIdenticalDataTest::RunTest(const FString& Parameters)
 {
+#if WITH_EDITOR
 	// Test single identical data
 	PCGDeterminismTests::FTestData FirstTestData(PCGDeterminismTests::DefaultSeed);
 	PCGDeterminismTests::FTestData SecondTestData(PCGDeterminismTests::DefaultSeed);
@@ -48,12 +51,13 @@ bool FPCGVolumeSamplerDeterminismSingleIdenticalDataTest::RunTest(const FString&
 	PCGDeterminismTests::AddRandomizedVolumeInputData(SecondTestData);
 
 	TestTrue("Identical single input and settings, same output", PCGDeterminismTests::ExecutionIsDeterministic(FirstTestData, SecondTestData));
-
+#endif
 	return true;
 }
 
 bool FPCGVolumeSamplerDeterminismMultipleSameDataTest::RunTest(const FString& Parameters)
 {
+#if WITH_EDITOR
 	PCGDeterminismTests::FTestData FirstTestData(PCGDeterminismTests::DefaultSeed);
 
 	// Test multiple same data
@@ -66,12 +70,13 @@ bool FPCGVolumeSamplerDeterminismMultipleSameDataTest::RunTest(const FString& Pa
 	}
 
 	TestTrue("Same multiple input, same output", PCGDeterminismTests::ExecutionIsDeterministic(FirstTestData, FirstTestData));
-
+#endif
 	return true;
 }
 
 bool FPCGVolumeSamplerDeterminismMultipleIdenticalDataTest::RunTest(const FString& Parameters)
 {
+#if WITH_EDITOR
 	// Test multiple identical data
 	PCGDeterminismTests::FTestData FirstTestData(PCGDeterminismTests::DefaultSeed);
 	PCGDeterminismTests::FTestData SecondTestData(PCGDeterminismTests::DefaultSeed);
@@ -87,12 +92,13 @@ bool FPCGVolumeSamplerDeterminismMultipleIdenticalDataTest::RunTest(const FStrin
 	}
 
 	TestTrue("Identical multiple input, same output", PCGDeterminismTests::ExecutionIsDeterministic(FirstTestData, SecondTestData));
-
+#endif
 	return true;
 }
 
 bool FPCGVolumeSamplerDeterminismOrderIndependenceTest::RunTest(const FString& Parameters)
 {
+#if WITH_EDITOR
 	// Test shuffled input data
 	PCGDeterminismTests::FTestData FirstTestData(PCGDeterminismTests::DefaultSeed);
 	PCGDeterminismTests::FTestData SecondTestData(PCGDeterminismTests::DefaultSeed);
@@ -110,6 +116,6 @@ bool FPCGVolumeSamplerDeterminismOrderIndependenceTest::RunTest(const FString& P
 	PCGDeterminismTests::ShuffleInputOrder(SecondTestData);
 
 	TestTrue("Shuffled input order, same output", PCGDeterminismTests::ExecutionIsDeterministic(FirstTestData, SecondTestData));
-
+#endif
 	return true;
 }
