Fixing IsPartitioned for PCG Component

#rnx

diff --git a/Engine/Plugins/Experimental/PCG/Source/PCG/Private/PCGComponent.cpp b/Engine/Plugins/Experimental/PCG/Source/PCG/Private/PCGComponent.cpp
--- a/Engine/Plugins/Experimental/PCG/Source/PCG/Private/PCGComponent.cpp
+++ b/Engine/Plugins/Experimental/PCG/Source/PCG/Private/PCGComponent.cpp
@@ -300,6 +300,18 @@ void UPCGComponent::PostProcessGraph(const FBox& InNewBounds, bool bInGenerated)
 	}
 }
 
+void UPCGComponent::OnProcessGraphAborted()
+{
+	UE_LOG(LogPCG, Warning, TEXT("Process Graph was called but aborted, check for errors in log if you expected a result."));
+
+	bGenerated = false;
+
+#if WITH_EDITOR
+	bDirtyGenerated = false;
+	bIsGenerating = false;
+#endif
+}
+
 void UPCGComponent::Cleanup()
 {
 	if (!bGenerated || !GetSubsystem())
@@ -535,10 +547,19 @@ void UPCGComponent::PostEditChangeProperty(FPropertyChangedEvent& PropertyChange
 			{
 				bool bIsNowPartitioned = bIsPartitioned;
 				bIsPartitioned = !bIsPartitioned;
+
 				// First, we'll cleanup
 				bActivated = false;
 				Refresh();
-				// Then do a normal refresh
+
+				// Invalidate the previous bounds to force actor creation (as if we moved the volume)
+				// if we are now partitioned
+				if (bIsNowPartitioned)
+				{
+					LastGeneratedBounds = FBox(EForceInit::ForceInit);
+				}
+
+				// And do a normal refresh
 				bActivated = true;
 				bIsPartitioned = bIsNowPartitioned;
 				DirtyGenerated();
diff --git a/Engine/Plugins/Experimental/PCG/Source/PCG/Private/PCGSubsystem.cpp b/Engine/Plugins/Experimental/PCG/Source/PCG/Private/PCGSubsystem.cpp
--- a/Engine/Plugins/Experimental/PCG/Source/PCG/Private/PCGSubsystem.cpp
+++ b/Engine/Plugins/Experimental/PCG/Source/PCG/Private/PCGSubsystem.cpp
@@ -113,6 +113,8 @@ FPCGTaskId UPCGSubsystem::ScheduleComponent(UPCGComponent* PCGComponent, const T
 	}
 	else
 	{
+		UE_LOG(LogPCG, Error, TEXT("[ScheduleComponent] Didn't schedule any task."));
+		PCGComponent->OnProcessGraphAborted();
 		return InvalidTaskId;
 	}
 }
@@ -152,6 +154,7 @@ namespace PCGSubsystem
 	{
 		if (!GraphExecutor || !World)
 		{
+			UE_LOG(LogPCG, Error, TEXT("[ForEachIntersectingCell] GraphExecutor or World is null"));
 			return InvalidTaskId;
 		}
 
@@ -318,6 +321,7 @@ FPCGTaskId UPCGSubsystem::ProcessGraph(UPCGComponent* Component, const FBox& InP
 
 			if (!LocalComponent)
 			{
+				UE_LOG(LogPCG, Error, TEXT("[ProcessGraph] PCG Component on PCG Actor is null"));
 				return InvalidTaskId;
 			}
 
@@ -385,6 +389,7 @@ FPCGTaskId UPCGSubsystem::ProcessGraph(UPCGComponent* Component, const FBox& InP
 	}
 	else
 	{
+		Component->OnProcessGraphAborted();
 		return InvalidTaskId;
 	}
 }
diff --git a/Engine/Plugins/Experimental/PCG/Source/PCG/Public/PCGComponent.h b/Engine/Plugins/Experimental/PCG/Source/PCG/Public/PCGComponent.h
--- a/Engine/Plugins/Experimental/PCG/Source/PCG/Public/PCGComponent.h
+++ b/Engine/Plugins/Experimental/PCG/Source/PCG/Public/PCGComponent.h
@@ -141,6 +141,7 @@ private:
 	void CleanupInternal(bool bRemoveComponents);
 	void CleanupInternal(bool bRemoveComponents, TSet<TSoftObjectPtr<AActor>>& OutActorsToDelete);
 	void PostProcessGraph(const FBox& InNewBounds, bool bInGenerated);
+	void OnProcessGraphAborted();
 	void CleanupUnusedManagedResources();
 
 	bool GetActorsFromTags(const TSet<FName>& InTags, TSet<TWeakObjectPtr<AActor>>& OutActors, bool bCullAgainstLocalBounds);
