Improve the static mesh spawner

diff --git a/Engine/Plugins/Experimental/PCG/Source/PCG/Private/Elements/PCGStaticMeshSpawner.cpp b/Engine/Plugins/Experimental/PCG/Source/PCG/Private/Elements/PCGStaticMeshSpawner.cpp
--- a/Engine/Plugins/Experimental/PCG/Source/PCG/Private/Elements/PCGStaticMeshSpawner.cpp
+++ b/Engine/Plugins/Experimental/PCG/Source/PCG/Private/Elements/PCGStaticMeshSpawner.cpp
@@ -3,28 +3,25 @@
 #include "Elements/PCGStaticMeshSpawner.h"
 
 #include "PCGHelpers.h"
-#include "Helpers/PCGActorHelpers.h"
 #include "Data/PCGPointData.h"
+#include "Data/PCGSpatialData.h"
+#include "Helpers/PCGActorHelpers.h"
+#include "MeshSelectors/PCGMeshSelectorBase.h"
+#include "MeshSelectors/PCGMeshSelectorWeighted.h"
 
 #include "Components/InstancedStaticMeshComponent.h"
 #include "Engine/StaticMesh.h"
-#include "Math/RandomStream.h"
 
-FPCGElementPtr UPCGStaticMeshSpawnerSettings::CreateElement() const
+UPCGStaticMeshSpawnerSettings::UPCGStaticMeshSpawnerSettings(const FObjectInitializer &ObjectInitializer)
 {
-	return MakeShared<FPCGStaticMeshSpawnerElement>();
+	MeshSelectorType = UPCGMeshSelectorWeighted::StaticClass();
+	MeshSelectorInstance = ObjectInitializer.CreateDefaultSubobject<UPCGMeshSelectorWeighted>(this, TEXT("DefaultSelectorInstance"));
 }
 
-struct FWeightedMeshAndInstances
+FPCGElementPtr UPCGStaticMeshSpawnerSettings::CreateElement() const
 {
-	FWeightedMeshAndInstances(int InWeight, const FPCGStaticMeshSpawnerEntry& InEntry)
-		: Weight(InWeight), Entry(InEntry)
-	{}
-
-	int Weight;
-	const FPCGStaticMeshSpawnerEntry& Entry;
-	TArray<FTransform> Instances;
-};
+	return MakeShared<FPCGStaticMeshSpawnerElement>();
+}
 
 bool FPCGStaticMeshSpawnerElement::ExecuteInternal(FPCGContext* Context) const
 {
@@ -33,23 +30,14 @@ bool FPCGStaticMeshSpawnerElement::ExecuteInternal(FPCGContext* Context) const
 	const UPCGStaticMeshSpawnerSettings* Settings = Context->GetInputSettings<UPCGStaticMeshSpawnerSettings>();
 	check(Settings);
 
-	// TODO : see if we could do async mesh loading here
-	TArray<FWeightedMeshAndInstances> WeightedEntries;
-	int TotalWeight = 0;
-	for (const FPCGStaticMeshSpawnerEntry& Entry : Settings->Meshes)
-	{
-		if (Entry.Weight <= 0)
-			continue;
-
-		TotalWeight += Entry.Weight;
-		FWeightedMeshAndInstances& WeightedEntry = WeightedEntries.Emplace_GetRef(TotalWeight, Entry);
-	}
-
-	if (TotalWeight <= 0)
+	if (!Settings->MeshSelectorInstance)
 	{
+		PCGE_LOG(Error, "Invalid MeshSelectorInstance");
 		return true;
 	}
 
+	// perform mesh selection
+	TMap<TSoftObjectPtr<UStaticMesh>, FPCGMeshInstanceList> MeshInstances;
 	TArray<FPCGTaggedData> Inputs = Context->InputData.GetInputs();
 
 	for (const FPCGTaggedData& Input : Inputs)
@@ -70,82 +58,120 @@ bool FPCGStaticMeshSpawnerElement::ExecuteInternal(FPCGContext* Context) const
 			continue;
 		}
 
-		// First, create target instance transforms
-		const UPCGPointData* PointData = SpatialData->ToPointData(Context);
+		Settings->MeshSelectorInstance->SelectInstances(*Context, Settings, SpatialData, MeshInstances);
+
+		// Spawn a static mesh for each instance
+		SpawnStaticMeshInstances(Context, MeshInstances, TargetActor);
+
+		MeshInstances.Reset();
+	}
 
-		if (!PointData)
+	return true;
+}
+
+void FPCGStaticMeshSpawnerElement::SpawnStaticMeshInstances(FPCGContext* Context, const TMap<TSoftObjectPtr<UStaticMesh>, FPCGMeshInstanceList>& MeshInstances, AActor* TargetActor) const
+{
+	// Populate the (H)ISM from the previously prepared entries
+	TRACE_CPUPROFILER_EVENT_SCOPE(FPCGStaticMeshSpawnerElement::Execute::PopulateISMs);
+
+	for (const TPair<TSoftObjectPtr<UStaticMesh>, FPCGMeshInstanceList>& Elem : MeshInstances)
+	{
+		const TSoftObjectPtr<UStaticMesh> Mesh = Elem.Key;
+		const FPCGMeshInstanceList& InstanceList = Elem.Value;
+
+		if (InstanceList.Instances.Num() == 0)
 		{
-			PCGE_LOG(Error, "Unable to get point data from input");
 			continue;
 		}
 
-		// Assign points to entries
+		// Todo: we could likely pre-load these meshes asynchronously in the settings
+		UStaticMesh* LoadedMesh = Mesh.LoadSynchronous();
+
+		if (!LoadedMesh)
 		{
-			TRACE_CPUPROFILER_EVENT_SCOPE(FPCGStaticMeshSpawnerElement::Execute::SelectEntries);
-			const TArray<FPCGPoint>& Points = PointData->GetPoints();
-
-			for (const FPCGPoint& Point : Points)
-			{
-				if (Point.Density <= 0.0f)
-				{
-					continue;
-				}
-
-				FRandomStream RandomSource(PCGHelpers::ComputeSeed(Point.Seed, Settings->Seed));
-				int RandomWeightedPick = RandomSource.RandRange(0, TotalWeight - 1);
-
-				int RandomPick = 0;
-				while (RandomPick < WeightedEntries.Num() && WeightedEntries[RandomPick].Weight <= RandomWeightedPick)
-				{
-					++RandomPick;
-				}
-
-				if (RandomPick < WeightedEntries.Num())
-				{
-					WeightedEntries[RandomPick].Instances.Emplace(Point.Transform);
-				}
-			}
+			continue;
 		}
 
-		// Populate the (H)ISM from the previously prepared entries
+		FPCGISMCBuilderParameters Params;
+		Params.Mesh = LoadedMesh;
+		if (InstanceList.bOverrideCollisionProfile)
 		{
-			TRACE_CPUPROFILER_EVENT_SCOPE(FPCGStaticMeshSpawnerElement::Execute::PopulateISMs);
-
-			// Second, add the instances
-			for (FWeightedMeshAndInstances& Instances : WeightedEntries)
-			{
-				if (Instances.Instances.Num() == 0)
-				{
-					continue;
-				}
-
-				Instances.Entry.Mesh.LoadSynchronous();
-
-				if (!Instances.Entry.Mesh)
-				{
-					continue;
-				}
-
-				FPCGISMCBuilderParameters Params;
-				// Todo: we could likely pre-load these meshes asynchronously in the settings
-				Params.Mesh = Instances.Entry.Mesh.LoadSynchronous();
-				if (Instances.Entry.bOverrideCollisionProfile)
-				{
-					Params.CollisionProfile = Instances.Entry.CollisionProfile.Name;
-				}
-
-				UInstancedStaticMeshComponent* ISMC = UPCGActorHelpers::GetOrCreateISMC(TargetActor, Context->SourceComponent, Params);
-
-				// TODO: add scaling
-				// TODO: document these arguments
-				ISMC->NumCustomDataFloats = 0;
-				ISMC->AddInstances(Instances.Instances, false, true);
-				ISMC->UpdateBounds();
-
-				PCGE_LOG(Verbose, "Added %d instances of %s on actor %s", Instances.Instances.Num(), *Instances.Entry.Mesh->GetFName().ToString(), *TargetActor->GetFName().ToString());
-			}
+			Params.CollisionProfile = InstanceList.CollisionProfile.Name;
 		}
+
+		UInstancedStaticMeshComponent* ISMC = UPCGActorHelpers::GetOrCreateISMC(TargetActor, Context->SourceComponent, Params);
+
+		// TODO: add scaling
+		// TODO: document these arguments
+		ISMC->NumCustomDataFloats = 0;
+		ISMC->AddInstances(InstanceList.Instances, false, true);
+		ISMC->UpdateBounds();
+
+		PCGE_LOG(Verbose, "Added %d instances of %s on actor %s", InstanceList.Instances.Num(), *Mesh->GetFName().ToString(), *TargetActor->GetFName().ToString());
 	}
+}
 
-	return true;
+void UPCGStaticMeshSpawnerSettings::PostLoad()
+{
+	Super::PostLoad();
+
+#if WITH_EDITOR
+	if (Meshes_DEPRECATED.Num() != 0)
+	{
+		SetMeshSelectorType(UPCGMeshSelectorWeighted::StaticClass());
+
+		UPCGMeshSelectorWeighted* MeshSelector = CastChecked<UPCGMeshSelectorWeighted>(MeshSelectorInstance);
+
+		for (const FPCGStaticMeshSpawnerEntry& Entry : Meshes_DEPRECATED)
+		{
+			FPCGMeshSelectorWeightedEntry& NewEntry = MeshSelector->MeshEntries.Emplace_GetRef(Entry.Mesh, Entry.Weight);
+			NewEntry.CollisionProfile = Entry.CollisionProfile;
+			NewEntry.bOverrideCollisionProfile = Entry.bOverrideCollisionProfile;
+		}
+
+		Meshes_DEPRECATED.Reset();
+	}
+#endif
+
+	if (!MeshSelectorInstance)
+	{
+		RefreshMeshSelector();
+	}
+}
+
+#if WITH_EDITOR
+void UPCGStaticMeshSpawnerSettings::PostEditChangeProperty(FPropertyChangedEvent& PropertyChangedEvent) 
+{
+	if (PropertyChangedEvent.Property && PropertyChangedEvent.Property->GetFName() == GET_MEMBER_NAME_CHECKED(UPCGStaticMeshSpawnerSettings, MeshSelectorType))
+	{
+		RefreshMeshSelector();
+	}
+
+	Super::PostEditChangeProperty(PropertyChangedEvent);
+}
+#endif
+
+void UPCGStaticMeshSpawnerSettings::SetMeshSelectorType(TSubclassOf<UPCGMeshSelectorBase> InMeshSelectorType) 
+{
+	if (!MeshSelectorInstance || InMeshSelectorType != MeshSelectorType)
+	{
+		if (InMeshSelectorType != MeshSelectorType)
+		{
+			MeshSelectorType = InMeshSelectorType;
+		}
+		
+		RefreshMeshSelector();
+	}
+}
+
+void UPCGStaticMeshSpawnerSettings::RefreshMeshSelector()
+{
+	if (MeshSelectorType)
+	{
+		MeshSelectorInstance = NewObject<UPCGMeshSelectorBase>(this, MeshSelectorType);
+	}
+	else
+	{
+		MeshSelectorInstance = nullptr;
+	}
 }
diff --git a/Engine/Plugins/Experimental/PCG/Source/PCG/Private/MeshSelectors/PCGMeshSelectorBase.cpp b/Engine/Plugins/Experimental/PCG/Source/PCG/Private/MeshSelectors/PCGMeshSelectorBase.cpp
new file mode 100644
--- /dev/null
+++ b/Engine/Plugins/Experimental/PCG/Source/PCG/Private/MeshSelectors/PCGMeshSelectorBase.cpp
@@ -0,0 +1,14 @@
+// Copyright Epic Games, Inc. All Rights Reserved.
+
+#include "MeshSelectors/PCGMeshSelectorBase.h"
+
+#include "Elements/PCGStaticMeshSpawner.h"
+
+void UPCGMeshSelectorBase::SelectInstances_Implementation(
+	FPCGContext& Context, 
+	const UPCGStaticMeshSpawnerSettings* Settings, 
+	const UPCGSpatialData* InSpatialData,
+	TMap<TSoftObjectPtr<UStaticMesh>, FPCGMeshInstanceList>& OutMeshInstances) const
+{
+	PCGE_LOG_C(Error, &Context, "Invalid use of abstract MeshSelectorBase class. Please use an already existing class or implement the CreateMeshInstanceData method");
+}
diff --git a/Engine/Plugins/Experimental/PCG/Source/PCG/Private/MeshSelectors/PCGMeshSelectorByAttribute.cpp b/Engine/Plugins/Experimental/PCG/Source/PCG/Private/MeshSelectors/PCGMeshSelectorByAttribute.cpp
new file mode 100644
--- /dev/null
+++ b/Engine/Plugins/Experimental/PCG/Source/PCG/Private/MeshSelectors/PCGMeshSelectorByAttribute.cpp
@@ -0,0 +1,97 @@
+// Copyright Epic Games, Inc. All Rights Reserved.
+
+#include "MeshSelectors/PCGMeshSelectorByAttribute.h"
+
+#include "Elements/PCGPointFilter.h"
+#include "Elements/PCGStaticMeshSpawner.h"
+#include "Metadata/PCGMetadata.h"
+#include "Metadata/PCGMetadataAttribute.h"
+#include "Metadata/PCGMetadataAttributeTraits.h"
+#include "Metadata/PCGMetadataAttributeTpl.h"
+
+void UPCGMeshSelectorByAttribute::SelectInstances_Implementation(
+	FPCGContext& Context, 
+	const UPCGStaticMeshSpawnerSettings* Settings, 
+	const UPCGSpatialData* InSpatialData, 
+	TMap<TSoftObjectPtr<UStaticMesh>, FPCGMeshInstanceList>& OutMeshInstances) const
+{
+	const UPCGPointData* PointData = InSpatialData->ToPointData(&Context);
+
+	if (!PointData)
+	{
+		PCGE_LOG_C(Error, &Context, "Unable to get point data from input");
+		return;
+	}
+
+	if (!PointData->Metadata)
+	{
+		PCGE_LOG_C(Error, &Context, "Unable to get metadata from input");
+		return;
+	}
+
+	if (!PointData->Metadata->HasAttribute(AttributeName)) 
+	{
+		PCGE_LOG_C(Error, &Context, "Attribute %s is not in the metadata", *AttributeName.ToString());
+		return;
+	}
+
+	const FPCGMetadataAttributeBase* AttributeBase = PointData->Metadata->GetConstAttribute(AttributeName);
+	check(AttributeBase);
+
+	if (AttributeBase->GetTypeId() != PCG::Private::MetadataTypes<FString>::Id)
+	{
+		PCGE_LOG_C(Error, &Context, "Attribute is not of valid type FString");
+		return;
+	}
+
+	const FPCGMetadataAttribute<FString>* Attribute = static_cast<const FPCGMetadataAttribute<FString>*>(AttributeBase);
+
+	TMap<PCGMetadataValueKey, TSoftObjectPtr<UStaticMesh>> ValueKeyToMesh;
+
+	TRACE_CPUPROFILER_EVENT_SCOPE(FPCGStaticMeshSpawnerElement::Execute::SelectEntries);
+
+	// for every point, add an instance of the stored Mesh
+	for (const FPCGPoint& Point : PointData->GetPoints()) 
+	{
+		if (Point.Density <= 0.0f)
+		{
+			continue;
+		}
+
+		PCGMetadataValueKey ValueKey = Attribute->GetValueKey(Point.MetadataEntry);
+		TSoftObjectPtr<UStaticMesh>* NewMesh = ValueKeyToMesh.Find(ValueKey);
+		TSoftObjectPtr<UStaticMesh> Mesh = nullptr;
+
+		// if this ValueKey has not been seen before, let's cache it for the future
+		if (!NewMesh)
+		{
+			FSoftObjectPath MeshPath(Attribute->GetValue(ValueKey));
+			Mesh = TSoftObjectPtr<UStaticMesh>(MeshPath);
+
+			ValueKeyToMesh.Add(ValueKey, Mesh);
+		}
+		else
+		{
+			Mesh = *NewMesh;
+		}
+
+		if (Mesh.IsNull())
+		{
+			PCGE_LOG_C(Error, &Context, "Invalid object path.");
+			continue;
+		}
+
+		// insert the Mesh if it does not exist already
+		FPCGMeshInstanceList* InstanceList = OutMeshInstances.Find(Mesh);
+		if (!InstanceList)
+		{
+			FPCGMeshInstanceList& NewInstanceList = OutMeshInstances.Emplace(Mesh);
+			NewInstanceList.bOverrideCollisionProfile = bOverrideCollisionProfile;
+			NewInstanceList.CollisionProfile = CollisionProfile;
+			
+			InstanceList = &NewInstanceList;
+		}
+
+		InstanceList->Instances.Emplace(Point.Transform);
+	}
+}
diff --git a/Engine/Plugins/Experimental/PCG/Source/PCG/Private/MeshSelectors/PCGMeshSelectorWeighted.cpp b/Engine/Plugins/Experimental/PCG/Source/PCG/Private/MeshSelectors/PCGMeshSelectorWeighted.cpp
new file mode 100644
--- /dev/null
+++ b/Engine/Plugins/Experimental/PCG/Source/PCG/Private/MeshSelectors/PCGMeshSelectorWeighted.cpp
@@ -0,0 +1,80 @@
+// Copyright Epic Games, Inc. All Rights Reserved.
+
+#include "MeshSelectors/PCGMeshSelectorWeighted.h"
+
+#include "PCGHelpers.h"
+#include "PCGPoint.h"
+#include "Elements/PCGStaticMeshSpawner.h"
+
+#include "Math/RandomStream.h"
+
+void UPCGMeshSelectorWeighted::SelectInstances_Implementation(
+	FPCGContext& Context, 
+	const UPCGStaticMeshSpawnerSettings* Settings, 
+	const UPCGSpatialData* InSpatialData,
+	TMap<TSoftObjectPtr<UStaticMesh>, FPCGMeshInstanceList>& OutMeshInstances) const
+{
+	TArray<TSoftObjectPtr<UStaticMesh>> Meshes;
+	TArray<int> CumulativeWeights;
+
+	int TotalWeight = 0;
+
+	for (const FPCGMeshSelectorWeightedEntry& Entry : MeshEntries)
+	{
+		if (Entry.Weight <= 0)
+		{
+			PCGE_LOG_C(Verbose, &Context, "Entry found with weight <= 0");
+			continue;
+		}
+	
+		// TODO: handle collisions between different collision profiles
+		FPCGMeshInstanceList InstanceList;
+		InstanceList.bOverrideCollisionProfile = Entry.bOverrideCollisionProfile;
+		InstanceList.CollisionProfile = Entry.CollisionProfile;
+		OutMeshInstances.Emplace(Entry.Mesh, InstanceList);
+
+		TotalWeight += Entry.Weight;
+		CumulativeWeights.Add(TotalWeight);
+		Meshes.Add(Entry.Mesh);
+	}
+
+	if (TotalWeight <= 0)
+	{
+		return;
+	}
+
+	const UPCGPointData* PointData = InSpatialData->ToPointData(&Context);
+
+	if (!PointData)
+	{
+		PCGE_LOG_C(Error, &Context, "Unable to get point data from input");
+		return;
+	}
+
+	// Assign points to entries
+	{
+		TRACE_CPUPROFILER_EVENT_SCOPE(FPCGStaticMeshSpawnerElement::Execute::SelectEntries);
+
+		for (const FPCGPoint& Point : PointData->GetPoints()) 
+		{
+			if (Point.Density <= 0.0f)
+			{
+				continue;
+			}
+
+			FRandomStream RandomSource(PCGHelpers::ComputeSeed(Point.Seed, Settings->Seed));
+			int RandomWeightedPick = RandomSource.RandRange(0, TotalWeight - 1);
+
+			int RandomPick = 0;
+			while (RandomPick < Meshes.Num() && CumulativeWeights[RandomPick] <= RandomWeightedPick)
+			{
+				++RandomPick;
+			}
+
+			if (RandomPick < Meshes.Num())
+			{
+				OutMeshInstances[Meshes[RandomPick]].Instances.Emplace(Point.Transform);
+			}
+		}
+	}
+}
diff --git a/Engine/Plugins/Experimental/PCG/Source/PCG/Private/MeshSelectors/PCGMeshSelectorWeightedByCategory.cpp b/Engine/Plugins/Experimental/PCG/Source/PCG/Private/MeshSelectors/PCGMeshSelectorWeightedByCategory.cpp
new file mode 100644
--- /dev/null
+++ b/Engine/Plugins/Experimental/PCG/Source/PCG/Private/MeshSelectors/PCGMeshSelectorWeightedByCategory.cpp
@@ -0,0 +1,166 @@
+// Copyright Epic Games, Inc. All Rights Reserved.
+
+#include "MeshSelectors/PCGMeshSelectorWeightedByCategory.h"
+
+#include "PCGHelpers.h"
+#include "Metadata/PCGMetadata.h"
+#include "Metadata/PCGMetadataAttribute.h"
+#include "Metadata/PCGMetadataAttributeTraits.h"
+#include "Metadata/PCGMetadataAttributeTpl.h"
+#include "Elements/PCGPointFilter.h"
+#include "Elements/PCGStaticMeshSpawner.h"
+
+#include "Math/RandomStream.h"
+
+struct FPCGMeshesAndWeights
+{
+	TArray<TSoftObjectPtr<UStaticMesh>> Meshes;
+	TArray<int> CumulativeWeights;
+};
+
+void UPCGMeshSelectorWeightedByCategory::SelectInstances_Implementation(
+	FPCGContext& Context, 
+	const UPCGStaticMeshSpawnerSettings* Settings, 
+	const UPCGSpatialData* InSpatialData,
+	TMap<TSoftObjectPtr<UStaticMesh>, FPCGMeshInstanceList>& OutMeshInstances) const
+{
+	const UPCGPointData* PointData = InSpatialData->ToPointData(&Context);
+
+	if (!PointData)
+	{
+		PCGE_LOG_C(Error, &Context, "Unable to get point data from input");
+		return;
+	}
+
+	if (!PointData->Metadata)
+	{
+		PCGE_LOG_C(Error, &Context, "Unable to get metadata from input");
+		return;
+	}
+
+	if (!PointData->Metadata->HasAttribute(CategoryAttribute))
+	{
+		PCGE_LOG_C(Error, &Context, "Attribute %s is not in the metadata", *CategoryAttribute.ToString());
+		return;
+	}
+
+	const FPCGMetadataAttributeBase* AttributeBase = PointData->Metadata->GetConstAttribute(CategoryAttribute);
+	check(AttributeBase);
+
+	// TODO: support enum type as well
+	if (AttributeBase->GetTypeId() != PCG::Private::MetadataTypes<FString>::Id)
+	{
+		PCGE_LOG_C(Error, &Context, "Attribute is not of valid type FString");
+		return;
+	}
+
+	const FPCGMetadataAttribute<FString>* Attribute = static_cast<const FPCGMetadataAttribute<FString>*>(AttributeBase);
+
+	// maps a CategoryEntry ValueKey to the meshes and precomputed weight data 
+	TMap<PCGMetadataValueKey, FPCGMeshesAndWeights> CategoryEntryToMeshesAndWeights;
+
+	// unmarked points will fallback to the MeshEntries associated with the DefaultValueKey
+	PCGMetadataValueKey DefaultValueKey = PCGDefaultValueKey;
+
+	for (const FPCGWeightedByCategoryEntryList& Entry : Entries)
+	{
+		if (Entry.WeightedMeshEntries.Num() == 0)
+		{
+			PCGE_LOG_C(Verbose, &Context, "Empty entry found in category %s", *Entry.CategoryEntry);
+			continue;
+		}
+
+		PCGMetadataValueKey ValueKey = Attribute->FindValue<FString>(Entry.CategoryEntry);
+
+		if (ValueKey == PCGDefaultValueKey)
+		{
+			PCGE_LOG_C(Verbose, &Context, "Invalid category %s", *Entry.CategoryEntry);
+			continue;
+		}
+
+		FPCGMeshesAndWeights* MeshesAndWeights = CategoryEntryToMeshesAndWeights.Find(ValueKey);
+
+		if (MeshesAndWeights)
+		{
+			PCGE_LOG_C(Warning, &Context, "Duplicate entry found in category %s. Subsequent entries are ignored.", *Entry.CategoryEntry);
+			continue;
+		}
+
+		if (Entry.IsDefault)
+		{
+			if (DefaultValueKey == PCGDefaultValueKey)
+			{
+				DefaultValueKey = ValueKey;
+			}
+			else
+			{
+				PCGE_LOG_C(Warning, &Context, "Duplicate default entry found. Subsequent default entries are ignored.");
+			}
+		}
+
+		MeshesAndWeights = &CategoryEntryToMeshesAndWeights.Add(ValueKey, FPCGMeshesAndWeights());
+
+		int TotalWeight = 0;
+		for (const FPCGMeshSelectorWeightedEntry& WeightedEntry : Entry.WeightedMeshEntries)
+		{
+			if (WeightedEntry.Weight <= 0)
+			{
+				PCGE_LOG_C(Verbose, &Context, "Entry found with weight <= 0 in category %s", *Entry.CategoryEntry);
+				continue;
+			}
+
+			// create or overwrite the entry for this mesh
+			FPCGMeshInstanceList& InstanceList = OutMeshInstances.Emplace(WeightedEntry.Mesh);
+			InstanceList.bOverrideCollisionProfile = WeightedEntry.bOverrideCollisionProfile;
+			InstanceList.CollisionProfile = WeightedEntry.CollisionProfile;
+
+			// precompute the weights
+			TotalWeight += WeightedEntry.Weight;
+			MeshesAndWeights->CumulativeWeights.Add(TotalWeight);
+			MeshesAndWeights->Meshes.Add(WeightedEntry.Mesh);
+		}
+	}
+
+	TRACE_CPUPROFILER_EVENT_SCOPE(FPCGStaticMeshSpawnerElement::Execute::SelectEntries);
+
+	for (const FPCGPoint& Point : PointData->GetPoints())
+	{
+		if (Point.Density <= 0.0f)
+		{
+			continue;
+		}
+
+		PCGMetadataValueKey ValueKey = Attribute->GetValueKey(Point.MetadataEntry);
+
+		// if no mesh list was processed for this attribute value, fallback to the default mesh list
+		FPCGMeshesAndWeights* MeshesAndWeights = CategoryEntryToMeshesAndWeights.Find(ValueKey);
+		if (!MeshesAndWeights)
+		{
+			if (DefaultValueKey != PCGDefaultValueKey)
+			{
+				MeshesAndWeights = CategoryEntryToMeshesAndWeights.Find(DefaultValueKey);
+				check(MeshesAndWeights);
+			}
+			else
+			{
+				continue;
+			}
+		}
+
+		const int TotalWeight = MeshesAndWeights->CumulativeWeights.Last();
+
+		FRandomStream RandomSource(PCGHelpers::ComputeSeed(Point.Seed, Settings->Seed));
+		int RandomWeightedPick = RandomSource.RandRange(0, TotalWeight - 1);
+
+		int RandomPick = 0;
+		while (RandomPick < MeshesAndWeights->Meshes.Num() && MeshesAndWeights->CumulativeWeights[RandomPick] <= RandomWeightedPick)
+		{
+			++RandomPick;
+		}
+
+		if (RandomPick < MeshesAndWeights->Meshes.Num())
+		{
+			OutMeshInstances[MeshesAndWeights->Meshes[RandomPick]].Instances.Emplace(Point.Transform);
+		}
+	}
+}
diff --git a/Engine/Plugins/Experimental/PCG/Source/PCG/Private/Tests/Elements/PCGStaticMeshSpawnerTest.cpp b/Engine/Plugins/Experimental/PCG/Source/PCG/Private/Tests/Elements/PCGStaticMeshSpawnerTest.cpp
new file mode 100644
--- /dev/null
+++ b/Engine/Plugins/Experimental/PCG/Source/PCG/Private/Tests/Elements/PCGStaticMeshSpawnerTest.cpp
@@ -0,0 +1,759 @@
+// Copyright Epic Games, Inc. All Rights Reserved.
+
+#include "Tests/PCGTestsCommon.h"
+#include "Tests/Determinism/PCGDeterminismTestsCommon.h"
+
+#include "PCGComponent.h"
+#include "PCGData.h"
+#include "PCGGraph.h"
+#include "PCGManagedResource.h"
+#include "PCGParamData.h"
+
+#include "Data/PCGPointData.h"
+#include "Elements/PCGStaticMeshSpawner.h"
+
+#include "MeshSelectors/PCGMeshSelectorBase.h"
+#include "MeshSelectors/PCGMeshSelectorWeighted.h"
+#include "MeshSelectors/PCGMeshSelectorByAttribute.h"
+#include "MeshSelectors/PCGMeshSelectorWeightedByCategory.h"
+
+#include "Components/InstancedStaticMeshComponent.h"
+#include "Engine/StaticMesh.h"
+
+IMPLEMENT_CUSTOM_SIMPLE_AUTOMATION_TEST(FPCGStaticMeshSpawnerByAttributeTest, FPCGTestBaseClass, "pcg.tests.StaticMeshSpawner.ByAttribute", PCGTestsCommon::TestFlags)
+IMPLEMENT_CUSTOM_SIMPLE_AUTOMATION_TEST(FPCGStaticMeshSpawnerWeightedTest, FPCGTestBaseClass, "pcg.tests.StaticMeshSpawner.Weighted", PCGTestsCommon::TestFlags)
+IMPLEMENT_CUSTOM_SIMPLE_AUTOMATION_TEST(FPCGStaticMeshSpawnerWeightedByCategoryTest, FPCGTestBaseClass, "pcg.tests.StaticMeshSpawner.WeightedByCategory", PCGTestsCommon::TestFlags)
+
+namespace
+{
+	struct FPCGByAttributeValidationData
+	{
+		FPCGByAttributeValidationData(TSoftObjectPtr<UStaticMesh> InMesh, int InStartIndex, int InEndIndex)
+			: Mesh(InMesh), StartIndex(InStartIndex), EndIndex(InEndIndex)
+		{}
+	
+		TSoftObjectPtr<UStaticMesh> Mesh;
+		int StartIndex;
+		int EndIndex;
+	};
+
+	struct FPCGWeightedByCategoryValidationData
+	{
+		FPCGWeightedByCategoryValidationData(UStaticMesh* InMesh, int InExpectedCount)
+			: Mesh(InMesh), ExpectedCount(InExpectedCount)
+		{}
+
+		UStaticMesh* Mesh;
+		int ExpectedCount;
+	};
+}
+
+void TestMeshSelectorByAttribute(
+	FPCGStaticMeshSpawnerByAttributeTest* Test, 
+	const PCGDeterminismTests::FTestData& TestData, 
+	const FName& AttributeName, 
+	const TArray<FPCGByAttributeValidationData>& ValidationDataEntries)
+{
+	// set the MeshSelector
+	UPCGStaticMeshSpawnerSettings* Settings = CastChecked<UPCGStaticMeshSpawnerSettings>(TestData.Settings);
+	Settings->SetMeshSelectorType(UPCGMeshSelectorByAttribute::StaticClass());
+	UPCGMeshSelectorByAttribute* MeshSelector = CastChecked<UPCGMeshSelectorByAttribute>(Settings->MeshSelectorInstance);
+	MeshSelector->AttributeName = AttributeName;
+
+	// initialize and execute the StaticMeshSpawner
+	FPCGElementPtr StaticMeshSpawner = Settings->GetElement();
+	FPCGContext* Context = StaticMeshSpawner->Initialize(TestData.InputData, TestData.TestComponent, nullptr);
+
+	while (!StaticMeshSpawner->Execute(Context))
+	{}
+
+	struct FPCGByAttributeValidIndexList
+	{
+		TArray<TPair<int, int>> ValidIndices;
+		int Count = 0;
+	};
+
+	TMap<TSoftObjectPtr<UStaticMesh>, FPCGByAttributeValidIndexList> MapToIndexList;
+
+	int ValidInstanceCount = 0;
+	for (const FPCGByAttributeValidationData& Entry : ValidationDataEntries)
+	{
+		FPCGByAttributeValidIndexList* IndexList = MapToIndexList.Find(Entry.Mesh);
+
+		if (!IndexList)
+		{
+			IndexList = &MapToIndexList.Emplace(Entry.Mesh, FPCGByAttributeValidIndexList());
+		}
+
+		IndexList->ValidIndices.Add(TPair<int, int>(Entry.StartIndex, Entry.EndIndex));
+		IndexList->Count += Entry.EndIndex - Entry.StartIndex;
+
+		ValidInstanceCount += Entry.EndIndex - Entry.StartIndex;
+	}
+
+	check(TestData.TestComponent);
+
+	TArray<UInstancedStaticMeshComponent*> ISMCs;
+	TestData.TestComponent->ForEachManagedResource([&ISMCs](UPCGManagedResource* InResource)
+	{
+		if (UPCGManagedISMComponent* Resource = Cast<UPCGManagedISMComponent>(InResource))
+		{
+			if (Resource->GetComponent())
+			{
+				ISMCs.Add(Resource->GetComponent());
+			}
+		}
+	});
+
+	int TotalInstanceCount = 0;
+	for (UInstancedStaticMeshComponent* ISMC : ISMCs)
+	{
+		const int InstanceCount = ISMC->GetInstanceCount();
+		TotalInstanceCount += InstanceCount;
+
+		TObjectPtr<UStaticMesh> StaticMesh = ISMC->GetStaticMesh();
+
+		const FPCGByAttributeValidationData* Entry = ValidationDataEntries.FindByPredicate([StaticMesh](const FPCGByAttributeValidationData& Entry) {
+			return Entry.Mesh == StaticMesh;
+			});
+
+		if (!Entry)
+		{
+			Test->TestTrue("Validate instanced mesh exists in MeshEntries", false);
+			continue;
+		}
+
+		const FPCGByAttributeValidIndexList* IndexList = MapToIndexList.Find(Entry->Mesh);
+
+		Test->TestTrue("Valid instance count per mesh", InstanceCount == IndexList->Count);
+
+		for (int InstanceIndex = 0; InstanceIndex < InstanceCount; InstanceIndex++)
+		{
+			FTransform Transform;
+			ISMC->GetInstanceTransform(InstanceIndex, Transform);
+			FVector Location = Transform.GetLocation();
+
+			bool bFound = false;
+
+			for (int i = 0; i < IndexList->ValidIndices.Num(); i++)
+			{
+				const int StartIndex = IndexList->ValidIndices[i].Key;
+				const int EndIndex = IndexList->ValidIndices[i].Value;
+
+				if (Location.X >= StartIndex && Location.X < EndIndex)
+				{
+					bFound = true;
+					break;
+				}
+			}
+
+			Test->TestTrue("Valid transform locations", bFound);
+		}
+	}
+
+	TestData.TestComponent->bGenerated = true;
+	TestData.TestComponent->Cleanup();
+	
+	Test->TestTrue("Valid total instance count", TotalInstanceCount == ValidInstanceCount);
+}
+
+void TestMeshSelectorWeighted(
+	FPCGStaticMeshSpawnerWeightedTest* Test, 
+	const PCGDeterminismTests::FTestData& TestData, 
+	const TArray<FPCGMeshSelectorWeightedEntry>& Entries, 
+	int PointCount)
+{
+	// set the MeshSelector
+	UPCGStaticMeshSpawnerSettings* Settings = CastChecked<UPCGStaticMeshSpawnerSettings>(TestData.Settings);
+	Settings->SetMeshSelectorType(UPCGMeshSelectorWeighted::StaticClass());
+	UPCGMeshSelectorWeighted* MeshSelector = CastChecked<UPCGMeshSelectorWeighted>(Settings->MeshSelectorInstance);
+	MeshSelector->MeshEntries = Entries;
+
+	// initialize and execute the StaticMeshSpawner
+	FPCGElementPtr StaticMeshSpawner = Settings->GetElement();
+	FPCGContext* Context = StaticMeshSpawner->Initialize(TestData.InputData, TestData.TestComponent, nullptr);
+
+	while (!StaticMeshSpawner->Execute(Context))
+	{}
+
+	int TotalWeight = 0;
+	for (const FPCGMeshSelectorWeightedEntry& Entry : Entries)
+	{
+		TotalWeight += Entry.Weight;
+	}
+
+	check(TestData.TestComponent);
+
+	TArray<UInstancedStaticMeshComponent*> ISMCs;
+	TestData.TestComponent->ForEachManagedResource([&ISMCs](UPCGManagedResource* InResource)
+	{
+		if (UPCGManagedISMComponent* Resource = Cast<UPCGManagedISMComponent>(InResource))
+		{
+			if (Resource->GetComponent())
+			{
+				ISMCs.Add(Resource->GetComponent());
+			}
+		}
+	});
+
+	int TotalInstanceCount = 0;
+	for (UInstancedStaticMeshComponent* ISMC : ISMCs)
+	{
+		const int InstanceCount = ISMC->GetInstanceCount();
+		TotalInstanceCount += InstanceCount;
+
+		TObjectPtr<UStaticMesh> StaticMesh = ISMC->GetStaticMesh();
+
+		FPCGMeshSelectorWeightedEntry* Entry = MeshSelector->MeshEntries.FindByPredicate([StaticMesh](const FPCGMeshSelectorWeightedEntry& Entry) {
+			return Entry.Mesh == StaticMesh;
+			});
+
+		if (!Entry)
+		{
+			Test->TestTrue("Validate instanced mesh exists in MeshEntries", false);
+			continue;
+		}
+
+		// validate correct InstanceCount (within a range) for the weight of this mesh
+		const int TargetCount = PointCount * Entry->Weight / TotalWeight;
+		const float ErrorBound = FMath::Max(TargetCount * 0.1f, 10); // 10% error boundary 
+		Test->TestTrue("Valid instance count per mesh", InstanceCount >= TargetCount - ErrorBound && InstanceCount <= TargetCount + ErrorBound);
+	}
+
+	TestData.TestComponent->bGenerated = true;
+	TestData.TestComponent->Cleanup();
+	
+	Test->TestTrue("Valid total instance count", TotalInstanceCount == PointCount);
+}
+
+void TestMeshSelectorWeightedByCategory(
+	FPCGStaticMeshSpawnerWeightedByCategoryTest* Test, 
+	const PCGDeterminismTests::FTestData& TestData, 
+	const FName& AttributeName, 
+	const TArray<FPCGWeightedByCategoryEntryList>& Entries, 
+	const TArray<FPCGWeightedByCategoryValidationData>& ValidationDataEntries,
+	int PointCount)
+{
+	// set the MeshSelector
+	UPCGStaticMeshSpawnerSettings* Settings = CastChecked<UPCGStaticMeshSpawnerSettings>(TestData.Settings);
+	Settings->SetMeshSelectorType(UPCGMeshSelectorWeightedByCategory::StaticClass());
+	UPCGMeshSelectorWeightedByCategory* MeshSelector = CastChecked<UPCGMeshSelectorWeightedByCategory>(Settings->MeshSelectorInstance);
+	MeshSelector->CategoryAttribute = AttributeName;
+	MeshSelector->Entries = Entries;
+
+	// initialize and execute the StaticMeshSpawner
+	FPCGElementPtr StaticMeshSpawner = Settings->GetElement();
+	FPCGContext* Context = StaticMeshSpawner->Initialize(TestData.InputData, TestData.TestComponent, nullptr);
+
+	while (!StaticMeshSpawner->Execute(Context))
+	{
+	}
+
+	check(TestData.TestComponent);
+
+	TArray<UInstancedStaticMeshComponent*> ISMCs;
+	TestData.TestComponent->ForEachManagedResource([&ISMCs](UPCGManagedResource* InResource)
+		{
+			if (UPCGManagedISMComponent* Resource = Cast<UPCGManagedISMComponent>(InResource))
+			{
+				if (Resource->GetComponent())
+				{
+					ISMCs.Add(Resource->GetComponent());
+				}
+			}
+		});
+
+	int TotalInstanceCount = 0;
+	for (UInstancedStaticMeshComponent* ISMC : ISMCs)
+	{
+		const int InstanceCount = ISMC->GetInstanceCount();
+		TotalInstanceCount += InstanceCount;
+
+		TObjectPtr<UStaticMesh> StaticMesh = ISMC->GetStaticMesh();
+
+		const FPCGWeightedByCategoryValidationData* Entry = ValidationDataEntries.FindByPredicate([StaticMesh](const FPCGWeightedByCategoryValidationData& Entry) {
+			return Entry.Mesh == StaticMesh;
+			});
+
+		if (!Entry)
+		{
+			Test->TestTrue("Validate instanced mesh exists in entries", false);
+		}
+
+		const float ErrorBound = FMath::Max(Entry->ExpectedCount * 0.1f, 10); // 10% error boundary 
+		Test->TestTrue("Valid instance count per mesh", InstanceCount >= Entry->ExpectedCount - ErrorBound && InstanceCount <= Entry->ExpectedCount + ErrorBound);
+	}
+
+	TestData.TestComponent->bGenerated = true;
+	TestData.TestComponent->Cleanup();
+
+	Test->TestTrue("Valid total instance count", TotalInstanceCount == PointCount);
+}
+
+bool FPCGStaticMeshSpawnerByAttributeTest::RunTest(const FString& Parameters)
+{
+	PCGDeterminismTests::FTestData TestData(PCGDeterminismTests::DefaultSeed);
+	PCGDeterminismTests::GenerateSettings<UPCGStaticMeshSpawnerSettings>(TestData);
+
+	// create test attribute data
+	const FName RockAttribute = "Rock";
+	const FName TreeAttribute = "Tree";
+	const FName BushAttribute = "Bush";
+
+	FSoftObjectPath CubePath(TEXT("StaticMesh'/Engine/BasicShapes/Cube.Cube'"));
+	FSoftObjectPath SpherePath(TEXT("StaticMesh'/Engine/BasicShapes/Sphere.Sphere'"));
+	FSoftObjectPath CylinderPath(TEXT("StaticMesh'/Engine/BasicShapes/Cylinder.Cylinder'"));
+	FSoftObjectPath ConePath(TEXT("StaticMesh'/Engine/BasicShapes/Cone.Cone'"));
+
+	TSoftObjectPtr<UStaticMesh> CubeMesh(CubePath);
+	TSoftObjectPtr<UStaticMesh> SphereMesh(SpherePath);
+	TSoftObjectPtr<UStaticMesh> CylinderMesh(CylinderPath);
+	TSoftObjectPtr<UStaticMesh> ConeMesh(ConePath);
+
+	// compose the test data
+	TObjectPtr<UPCGPointData> PointData = PCGTestsCommon::CreateEmptyPointData();
+	PointData->TargetActor = TestData.TestActor;
+	PointData->Metadata->CreateStringAttribute(RockAttribute, ConePath.GetAssetPathString(), false);
+	PointData->Metadata->CreateStringAttribute(TreeAttribute, ConePath.GetAssetPathString(), false);
+	PointData->Metadata->CreateStringAttribute(BushAttribute, ConePath.GetAssetPathString(), false);
+
+	// first 20 points are Rock, next 50 are Tree, next 30 are Bush, for a total of 100 points
+	const int PointCount = 100;
+	TArray<FPCGPoint>& Points = PointData->GetMutablePoints();
+	for (int i = 0; i < PointCount; i++)
+	{
+		FVector Position(i, 0, 0);
+		Points.Emplace(FTransform(Position), 1, i);
+
+		FName AttributeName;
+		FSoftObjectPath ObjectPath;
+
+		if (i < 20)
+		{
+			AttributeName = RockAttribute;
+			ObjectPath = (i < 10) ? CubePath : SpherePath; // half are CubePath, half are SpherePath
+		}
+		else if (i < 70)
+		{
+			AttributeName = TreeAttribute;
+			ObjectPath = (i < 35) ? CubePath : (i < 50) ? SpherePath : CylinderPath; // some are CubePath, then SpherePath, then CylinderPath
+		}
+		else 
+		{
+			AttributeName = BushAttribute;
+			ObjectPath = CylinderPath;
+			ObjectPath = (i < 85) ? CubePath : CylinderPath; // half are CubePath, half are CylinderPath
+		}
+		
+		UPCGMetadataAccessorHelpers::SetStringAttribute(Points[i], PointData->Metadata, AttributeName, ObjectPath.GetAssetPathString());
+	}
+	
+	FPCGTaggedData& TaggedDataPoints = TestData.InputData.TaggedData.Emplace_GetRef(FPCGTaggedData());
+	TaggedDataPoints.Data = PointData;
+
+	// Test 1
+	// should have 10 CubeMesh at Loc.x = [0, 10), then 10 SphereMesh at Loc.x = [10, 20), and then 80 ConeMesh at Loc.x = [20, 100)
+	TArray<FPCGByAttributeValidationData> ValidationDataEntries;
+	ValidationDataEntries.Emplace(CubeMesh, 0, 10);
+	ValidationDataEntries.Emplace(SphereMesh, 10, 20);
+	ValidationDataEntries.Emplace(ConeMesh, 20, 100);
+	TestMeshSelectorByAttribute(this, TestData, RockAttribute, ValidationDataEntries);
+
+	// Test 2
+	// should have ConeMesh at Loc.x = [0, 20), [70, 100), CubeMesh at Loc.x = [20, 35), SphereMesh at Loc.x = [35, 50), CylinderMesh at Loc.x = [50, 70)
+	ValidationDataEntries.Empty();
+	ValidationDataEntries.Emplace(ConeMesh, 0, 20);
+	ValidationDataEntries.Emplace(CubeMesh, 20, 35);
+	ValidationDataEntries.Emplace(SphereMesh, 35, 50);
+	ValidationDataEntries.Emplace(CylinderMesh, 50, 70);
+	ValidationDataEntries.Emplace(ConeMesh, 70, 100);
+	TestMeshSelectorByAttribute(this, TestData, TreeAttribute, ValidationDataEntries);
+
+	// Test 3
+	// should have ConeMesh at Loc.x = [0, 70), CubeMesh at Loc.x = [70, 85), CylinderMesh at Loc.x = [85, 100)
+	ValidationDataEntries.Empty();
+	ValidationDataEntries.Emplace(ConeMesh, 0, 70);
+	ValidationDataEntries.Emplace(CubeMesh, 70, 85);
+	ValidationDataEntries.Emplace(CylinderMesh, 85, 100);
+	TestMeshSelectorByAttribute(this, TestData, BushAttribute, ValidationDataEntries);
+
+	return true;
+}
+
+bool FPCGStaticMeshSpawnerWeightedTest::RunTest(const FString& Parameters)
+{
+	PCGDeterminismTests::FTestData TestData(PCGDeterminismTests::DefaultSeed);
+	PCGDeterminismTests::GenerateSettings<UPCGStaticMeshSpawnerSettings>(TestData);
+
+	// create test meshes
+	const FString CubePath = TEXT("StaticMesh'/Engine/BasicShapes/Cube.Cube'");
+	const FString SpherePath = TEXT("StaticMesh'/Engine/BasicShapes/Sphere.Sphere'");
+	const FString CylinderPath = TEXT("StaticMesh'/Engine/BasicShapes/Cylinder.Cylinder'");
+
+	UStaticMesh* CubeMesh = Cast<UStaticMesh>(StaticLoadObject(UStaticMesh::StaticClass(), NULL, *CubePath));
+	UStaticMesh* SphereMesh = Cast<UStaticMesh>(StaticLoadObject(UStaticMesh::StaticClass(), NULL, *SpherePath));
+	UStaticMesh* CylinderMesh = Cast<UStaticMesh>(StaticLoadObject(UStaticMesh::StaticClass(), NULL, *CylinderPath));
+
+	// compose the test data
+	const int PointCount = 1000;
+	TObjectPtr<UPCGPointData> PointData = PCGTestsCommon::CreateEmptyPointData();
+	PointData->TargetActor = TestData.TestActor;
+
+	TArray<FPCGPoint>& Points = PointData->GetMutablePoints();
+	for (int i = 0; i < PointCount; i++)
+	{
+		FVector Position(i, 0, 0);
+		Points.Emplace(FTransform(Position), 1, i);
+	}
+	
+	FPCGTaggedData& TaggedDataPoints = TestData.InputData.TaggedData.Emplace_GetRef(FPCGTaggedData());
+	TaggedDataPoints.Data = PointData;
+
+	// Test 1: even distribution
+	TArray<FPCGMeshSelectorWeightedEntry> Entries;
+	Entries.Emplace(CubeMesh, 1);
+	Entries.Emplace(SphereMesh, 1);
+	Entries.Emplace(CylinderMesh, 1);
+	TestMeshSelectorWeighted(this, TestData, Entries, PointCount);
+
+	// Test 2: skewed distribution
+	Entries[0].Weight = 1;
+	Entries[1].Weight = 10;
+	Entries[2].Weight = 1;
+	TestMeshSelectorWeighted(this, TestData, Entries, PointCount);
+
+	// Test 3: null weight
+	Entries[0].Weight = 1;
+	Entries[1].Weight = 0;
+	Entries[2].Weight = 1;
+	TestMeshSelectorWeighted(this, TestData, Entries, PointCount);
+
+	// Test 4: heavily skewed distribution
+	Entries[0].Weight = 1000;
+	Entries[1].Weight = 1;
+	Entries[2].Weight = 1;
+	TestMeshSelectorWeighted(this, TestData, Entries, PointCount);
+
+	// Test 5: null weight
+	Entries[0].Weight = 1;
+	Entries[1].Weight = 0;
+	Entries[2].Weight = 1;
+	TestMeshSelectorWeighted(this, TestData, Entries, PointCount);
+
+	return true;
+}
+
+bool FPCGStaticMeshSpawnerWeightedByCategoryTest::RunTest(const FString& Parameters)
+{
+	PCGDeterminismTests::FTestData TestData(PCGDeterminismTests::DefaultSeed);
+	PCGDeterminismTests::GenerateSettings<UPCGStaticMeshSpawnerSettings>(TestData);
+
+	// create test meshes
+	const FString CubePath = TEXT("StaticMesh'/Engine/BasicShapes/Cube.Cube'");
+	const FString SpherePath = TEXT("StaticMesh'/Engine/BasicShapes/Sphere.Sphere'");
+	const FString CylinderPath = TEXT("StaticMesh'/Engine/BasicShapes/Cylinder.Cylinder'");
+	const FString ConePath = TEXT("StaticMesh'/Engine/BasicShapes/Cone.Cone'");
+
+	UStaticMesh* CubeMesh = Cast<UStaticMesh>(StaticLoadObject(UStaticMesh::StaticClass(), NULL, *CubePath));
+	UStaticMesh* SphereMesh = Cast<UStaticMesh>(StaticLoadObject(UStaticMesh::StaticClass(), NULL, *SpherePath));
+	UStaticMesh* CylinderMesh = Cast<UStaticMesh>(StaticLoadObject(UStaticMesh::StaticClass(), NULL, *CylinderPath));
+	UStaticMesh* ConeMesh = Cast<UStaticMesh>(StaticLoadObject(UStaticMesh::StaticClass(), NULL, *ConePath));
+
+	// Attribute categories
+	const FName AttributeName = TEXT("RockType");
+	const FString IgneousRockType = TEXT("Igneous");
+	const FString SedimentaryRockType = TEXT("Sedimentary");
+	const FString MetamorphicRockType = TEXT("Metamorphic");
+	const FString DefaultRockType = TEXT("DefaultRockType");
+	const FString InvalidRockType = TEXT("InvalidRockType");
+
+	// compose the test data
+	TObjectPtr<UPCGPointData> PointData = PCGTestsCommon::CreateEmptyPointData();
+	PointData->TargetActor = TestData.TestActor;
+	PointData->Metadata->CreateStringAttribute(AttributeName, DefaultRockType, false);
+
+	TArray<FPCGPoint>& Points = PointData->GetMutablePoints();
+
+	// first 1k are igneous, next 1.5k are sedimenatary, next 500 are metamorphic, then last 500 are Default
+	const int PointCount = 3500;
+	const int IgneousCategorySize = 1000;
+	const int SedimentaryCategorySize = 1500;
+	const int MetamorphicCategorySize = 500;
+	const int DefaultCategorySize = 500;
+
+	for (int i = 0; i < PointCount; i++)
+	{
+		FVector Position(i, 0, 0);
+		Points.Emplace(FTransform(Position), 1, i);
+
+		FString AttributeValue;
+		
+		if (i < IgneousCategorySize)
+		{
+			AttributeValue = IgneousRockType;
+		} 
+		else if (i < IgneousCategorySize + SedimentaryCategorySize)
+		{
+			AttributeValue = SedimentaryRockType;
+		} 
+		else if (i < IgneousCategorySize + SedimentaryCategorySize + MetamorphicCategorySize)
+		{
+			AttributeValue = MetamorphicRockType;
+		}
+		else
+		{
+			AttributeValue = DefaultRockType;
+		} 
+		
+		UPCGMetadataAccessorHelpers::SetStringAttribute(Points[i], PointData->Metadata, AttributeName, AttributeValue);
+	}
+	
+	FPCGTaggedData& TaggedDataPoints = TestData.InputData.TaggedData.Emplace_GetRef(FPCGTaggedData());
+	TaggedDataPoints.Data = PointData;
+
+	// Test 1 - basic weighted functionality
+	{
+		const int CubeWeight = 1.f;
+		const int SphereWeight = 1.f;
+		const int CylinderWeight = 1.f;
+
+		const int IgneousCategoryTotalWeight = 3;
+
+		const int TotalSize = IgneousCategorySize;
+
+		TArray<FPCGMeshSelectorWeightedEntry> IgneousEntries;
+		IgneousEntries.Emplace(CubeMesh, CubeWeight);
+		IgneousEntries.Emplace(SphereMesh, SphereWeight);
+		IgneousEntries.Emplace(CylinderMesh, CylinderWeight);
+
+		TArray<FPCGWeightedByCategoryEntryList> Entries;
+		Entries.Emplace(IgneousRockType, IgneousEntries);
+
+		const int CubeCount = (IgneousCategorySize * CubeWeight / IgneousCategoryTotalWeight);
+		const int SphereCount = (IgneousCategorySize * SphereWeight / IgneousCategoryTotalWeight);
+		const int CylinderCount = (IgneousCategorySize * CylinderWeight / IgneousCategoryTotalWeight);
+
+		TArray<FPCGWeightedByCategoryValidationData> ValidationData;
+		ValidationData.Emplace(CubeMesh, CubeCount);
+		ValidationData.Emplace(SphereMesh, SphereCount);
+		ValidationData.Emplace(CylinderMesh, CylinderCount);
+
+		TestMeshSelectorWeightedByCategory(this, TestData, AttributeName, Entries, ValidationData, TotalSize);
+	}
+
+	// Test 2 - selection based on rocktype
+	{
+		const int IgneousCubeWeight = 1.f;
+		const int IgneousSphereWeight = 1.f;
+
+		const int SedimentaryCubeWeight = 10.f;
+		const int SedimentaryCylinderWeight = 1.f;
+
+		const int MetamorphicSphereWeight = 1.f;
+		const int MetamorphicCylinderWeight = 10.f;
+
+		const int IgneousCategoryTotalWeight = 2;
+		const int SedimentaryCategoryTotalWeight = 11;
+		const int MetamorphicCategoryTotalWeight = 11;
+
+		const int TotalSize = IgneousCategorySize + SedimentaryCategorySize + MetamorphicCategorySize;
+
+		TArray<FPCGMeshSelectorWeightedEntry> IgneousEntries;
+		IgneousEntries.Emplace(CubeMesh, IgneousCubeWeight);
+		IgneousEntries.Emplace(SphereMesh, IgneousSphereWeight);
+
+		TArray<FPCGMeshSelectorWeightedEntry> SedimentaryEntries;
+		SedimentaryEntries.Emplace(CubeMesh, SedimentaryCubeWeight);
+		SedimentaryEntries.Emplace(CylinderMesh, SedimentaryCylinderWeight);
+
+		TArray<FPCGMeshSelectorWeightedEntry> MetamorphicEntries;
+		MetamorphicEntries.Emplace(SphereMesh, MetamorphicSphereWeight);
+		MetamorphicEntries.Emplace(CylinderMesh, MetamorphicCylinderWeight);
+
+		TArray<FPCGWeightedByCategoryEntryList> Entries;
+		Entries.Emplace(IgneousRockType, IgneousEntries);
+		Entries.Emplace(SedimentaryRockType, SedimentaryEntries);
+		Entries.Emplace(MetamorphicRockType, MetamorphicEntries);
+
+		const int CubeCount = (IgneousCategorySize * IgneousCubeWeight / IgneousCategoryTotalWeight) + (SedimentaryCategorySize * SedimentaryCubeWeight / SedimentaryCategoryTotalWeight);
+		const int SphereCount = (IgneousCategorySize * IgneousSphereWeight / IgneousCategoryTotalWeight) + (MetamorphicCategorySize * MetamorphicSphereWeight / MetamorphicCategoryTotalWeight);
+		const int CylinderCount = (SedimentaryCategorySize * SedimentaryCylinderWeight / SedimentaryCategoryTotalWeight) + (MetamorphicCategorySize * MetamorphicCylinderWeight / MetamorphicCategoryTotalWeight);
+
+		TArray<FPCGWeightedByCategoryValidationData> ValidationData;
+		ValidationData.Emplace(CubeMesh, CubeCount);
+		ValidationData.Emplace(SphereMesh, SphereCount);
+		ValidationData.Emplace(CylinderMesh, CylinderCount);
+
+		TestMeshSelectorWeightedByCategory(this, TestData, AttributeName, Entries, ValidationData, TotalSize);
+	}
+
+	// Test 3 - empty weighted entries
+	{
+		const int CubeWeight = 1.f;
+		const int IgneousSphereWeight = 1.f;
+		const int MetamorphicSphereWeight = 1.f;
+		const int CylinderWeight = 10.f;
+
+		const int IgneousCategoryTotalWeight = 2;
+		const int MetamorphicCategoryTotalWeight = 11;
+
+		const int TotalSize = IgneousCategorySize + MetamorphicCategorySize;
+
+		TArray<FPCGMeshSelectorWeightedEntry> IgneousEntries;
+		IgneousEntries.Emplace(CubeMesh, CubeWeight);
+		IgneousEntries.Emplace(SphereMesh, IgneousSphereWeight);
+
+		TArray<FPCGMeshSelectorWeightedEntry> SedimentaryEntries;
+
+		TArray<FPCGMeshSelectorWeightedEntry> MetamorphicEntries;
+		MetamorphicEntries.Emplace(SphereMesh, MetamorphicSphereWeight);
+		MetamorphicEntries.Emplace(CylinderMesh, CylinderWeight);
+
+		TArray<FPCGWeightedByCategoryEntryList> Entries;
+		Entries.Emplace(IgneousRockType, IgneousEntries);
+		Entries.Emplace(SedimentaryRockType, SedimentaryEntries);
+		Entries.Emplace(MetamorphicRockType, MetamorphicEntries);
+
+		const int CubeCount = (IgneousCategorySize * CubeWeight / IgneousCategoryTotalWeight);
+		const int SphereCount = (IgneousCategorySize * IgneousSphereWeight / IgneousCategoryTotalWeight) + (MetamorphicCategorySize * MetamorphicSphereWeight / MetamorphicCategoryTotalWeight);
+		const int CylinderCount = (MetamorphicCategorySize * CylinderWeight / MetamorphicCategoryTotalWeight);
+
+		TArray<FPCGWeightedByCategoryValidationData> ValidationData;
+		ValidationData.Emplace(CubeMesh, CubeCount);
+		ValidationData.Emplace(SphereMesh, SphereCount);
+		ValidationData.Emplace(CylinderMesh, CylinderCount);
+
+		TestMeshSelectorWeightedByCategory(this, TestData, AttributeName, Entries, ValidationData, TotalSize);
+	}
+
+	// Test 4 - duplicate category
+	{
+		const int CubeWeight = 1.f;
+		const int SphereWeight = 1.f;
+
+		const int IgneousCategoryTotalWeight = 1;
+		const int SedimentaryCategoryTotalWeight = 1;
+
+		const int TotalSize = IgneousCategorySize + SedimentaryCategorySize;
+
+		TArray<FPCGMeshSelectorWeightedEntry> IgneousEntries;
+		IgneousEntries.Emplace(CubeMesh, CubeWeight);
+
+		TArray<FPCGMeshSelectorWeightedEntry> SedimentaryEntries;
+		SedimentaryEntries.Emplace(SphereMesh, SphereWeight);
+
+		TArray<FPCGMeshSelectorWeightedEntry> SedimentaryEntriesDuplicate;
+		SedimentaryEntriesDuplicate.Emplace(SphereMesh, 10000); // heavily skewed to ensure its effect on the data is noticeable
+
+		TArray<FPCGWeightedByCategoryEntryList> Entries;
+		Entries.Emplace(IgneousRockType, IgneousEntries);
+		Entries.Emplace(SedimentaryRockType, SedimentaryEntries);
+		Entries.Emplace(SedimentaryRockType, SedimentaryEntriesDuplicate);
+
+		const int CubeCount = (IgneousCategorySize * CubeWeight / IgneousCategoryTotalWeight);
+		const int SphereCount = (SedimentaryCategorySize * SphereWeight / SedimentaryCategoryTotalWeight);
+
+		TArray<FPCGWeightedByCategoryValidationData> ValidationData;
+		ValidationData.Emplace(CubeMesh, CubeCount);
+		ValidationData.Emplace(SphereMesh, SphereCount);
+
+		const FString& ErrorMessage = TEXT("Duplicate entry found in category ") + SedimentaryRockType + TEXT(". Subsequent entries are ignored.");
+		AddExpectedError(ErrorMessage, EAutomationExpectedErrorFlags::Contains, 1);
+
+		TestMeshSelectorWeightedByCategory(this, TestData, AttributeName, Entries, ValidationData, TotalSize);
+	}
+
+	// Test 5 - invalid rock type
+	{
+		const int CubeWeight = 1.f;
+		const int SphereWeight = 1.f;
+		const int CategoryTotalWeight = 2;
+		const int TotalSize = IgneousCategorySize;
+
+		TArray<FPCGMeshSelectorWeightedEntry> IgneousEntries;
+		IgneousEntries.Emplace(CubeMesh, 1);
+		IgneousEntries.Emplace(SphereMesh, 1);
+
+		TArray<FPCGMeshSelectorWeightedEntry> InvalidEntries;
+		InvalidEntries.Emplace(SphereMesh, 10000); // heavily skewed to ensure its effect on the data is noticeable
+
+		TArray<FPCGWeightedByCategoryEntryList> Entries;
+		Entries.Emplace(IgneousRockType, IgneousEntries);
+		Entries.Emplace(InvalidRockType, InvalidEntries);
+
+		const int CubeCount = (IgneousCategorySize * CubeWeight / CategoryTotalWeight);
+		const int SphereCount = (IgneousCategorySize * SphereWeight / CategoryTotalWeight);
+
+		TArray<FPCGWeightedByCategoryValidationData> ValidationData;
+		ValidationData.Emplace(CubeMesh, CubeCount);
+		ValidationData.Emplace(SphereMesh, SphereCount);
+
+		TestMeshSelectorWeightedByCategory(this, TestData, AttributeName, Entries, ValidationData, TotalSize);
+	}
+
+	// Test 6 - default entry list + duplicate default ignored
+	{
+		const int IgneousCubeWeight = 1.f;
+		const int IgneousSphereWeight = 1.f;
+
+		const int SedimentaryCubeWeight = 10.f;
+		const int SedimentaryCylinderWeight = 1.f;
+
+		const int MetamorphicSphereWeight = 1.f;
+		const int MetamorphicCylinderWeight = 10.f;
+
+		const int IgneousCategoryTotalWeight = 2;
+		const int SedimentaryCategoryTotalWeight = 11;
+		const int MetamorphicCategoryTotalWeight = 11;
+
+		const int TotalSize = IgneousCategorySize + SedimentaryCategorySize + MetamorphicCategorySize + DefaultCategorySize;
+
+		TArray<FPCGMeshSelectorWeightedEntry> IgneousEntries;
+		IgneousEntries.Emplace(CubeMesh, IgneousCubeWeight);
+		IgneousEntries.Emplace(SphereMesh, IgneousSphereWeight);
+
+		TArray<FPCGMeshSelectorWeightedEntry> SedimentaryEntries;
+		SedimentaryEntries.Emplace(CubeMesh, SedimentaryCubeWeight);
+		SedimentaryEntries.Emplace(CylinderMesh, SedimentaryCylinderWeight);
+
+		TArray<FPCGMeshSelectorWeightedEntry> MetamorphicEntries;
+		MetamorphicEntries.Emplace(SphereMesh, MetamorphicSphereWeight);
+		MetamorphicEntries.Emplace(CylinderMesh, MetamorphicCylinderWeight);
+
+		TArray<FPCGMeshSelectorWeightedEntry> DuplicateDefault;
+		DuplicateDefault.Emplace(CubeMesh, 10000); // heavily skewed to ensure its effect on the data is noticeable
+
+		TArray<FPCGWeightedByCategoryEntryList> Entries;
+		Entries.Emplace(IgneousRockType, IgneousEntries);
+		Entries.Emplace(MetamorphicRockType, MetamorphicEntries);
+
+		FPCGWeightedByCategoryEntryList& DefaultEntryList = Entries.Emplace_GetRef(SedimentaryRockType, SedimentaryEntries);
+		FPCGWeightedByCategoryEntryList& DefaultEntryListDuplicate = Entries.Emplace_GetRef(MetamorphicRockType, DuplicateDefault);
+
+		DefaultEntryList.IsDefault = true;
+		DefaultEntryListDuplicate.IsDefault = true;
+
+		const int CubeCount = (IgneousCategorySize * IgneousCubeWeight / IgneousCategoryTotalWeight)
+			+ (SedimentaryCategorySize * SedimentaryCubeWeight / SedimentaryCategoryTotalWeight)
+			+ (DefaultCategorySize * SedimentaryCubeWeight / SedimentaryCategoryTotalWeight);
+
+		const int SphereCount = (IgneousCategorySize * IgneousSphereWeight / IgneousCategoryTotalWeight)
+			+ (MetamorphicCategorySize * MetamorphicSphereWeight / MetamorphicCategoryTotalWeight);
+
+		const int CylinderCount = (SedimentaryCategorySize * SedimentaryCylinderWeight / SedimentaryCategoryTotalWeight)
+			+ (MetamorphicCategorySize * MetamorphicCylinderWeight / MetamorphicCategoryTotalWeight)
+			+ (DefaultCategorySize * SedimentaryCylinderWeight / SedimentaryCategoryTotalWeight);
+
+		TArray<FPCGWeightedByCategoryValidationData> ValidationData;
+		ValidationData.Emplace(CubeMesh, CubeCount);
+		ValidationData.Emplace(SphereMesh, SphereCount);
+		ValidationData.Emplace(CylinderMesh, CylinderCount);
+
+		const FString& ErrorMessage = TEXT("Duplicate entry found in category ") + MetamorphicRockType + TEXT(". Subsequent entries are ignored.");
+		AddExpectedError(ErrorMessage, EAutomationExpectedErrorFlags::Contains, 1);
+
+		TestMeshSelectorWeightedByCategory(this, TestData, AttributeName, Entries, ValidationData, TotalSize);
+	}
+
+	return true;
+}
diff --git a/Engine/Plugins/Experimental/PCG/Source/PCG/Public/Elements/PCGStaticMeshSpawner.h b/Engine/Plugins/Experimental/PCG/Source/PCG/Public/Elements/PCGStaticMeshSpawner.h
--- a/Engine/Plugins/Experimental/PCG/Source/PCG/Public/Elements/PCGStaticMeshSpawner.h
+++ b/Engine/Plugins/Experimental/PCG/Source/PCG/Public/Elements/PCGStaticMeshSpawner.h
@@ -4,6 +4,7 @@
 
 #include "CoreMinimal.h"
 #include "PCGSettings.h"
+#include "MeshSelectors/PCGMeshSelectorBase.h"
 
 #include "Engine/CollisionProfile.h"
 
@@ -11,7 +12,7 @@
 
 class UStaticMesh;
 
-USTRUCT(BlueprintType)
+USTRUCT(BlueprintType, meta=(Deprecated = "5.0", DeprecationMessage="Use MeshSelectorWeighted instead."))
 struct PCG_API FPCGStaticMeshSpawnerEntry
 {
 	GENERATED_BODY()
@@ -35,6 +36,8 @@ class PCG_API UPCGStaticMeshSpawnerSettings : public UPCGSettings
 	GENERATED_BODY()
 
 public:
+	UPCGStaticMeshSpawnerSettings(const FObjectInitializer &ObjectInitializer);
+
 #if WITH_EDITOR
 	// ~Begin UPCGSettings interface
 	virtual FName GetDefaultNodeName() const override { return FName(TEXT("StaticMeshSpawnerNode")); }
@@ -48,8 +51,28 @@ protected:
 	// ~End UPCGSettings interface
 
 public:
-	UPROPERTY(BlueprintReadWrite, EditAnywhere, Category = Settings)
-	TArray<FPCGStaticMeshSpawnerEntry> Meshes;
+	// ~Begin UObject interface
+	virtual void PostLoad() override;
+#if WITH_EDITOR
+	virtual void PostEditChangeProperty(FPropertyChangedEvent& PropertyChangedEvent) override;
+	// ~End UObject interface
+#endif
+
+	UFUNCTION(BlueprintCallable, Category = Settings)
+	void SetMeshSelectorType(TSubclassOf<UPCGMeshSelectorBase> InMeshSelectorType);
+
+public:
+	UPROPERTY(BlueprintReadOnly, EditAnywhere, NoClear, Category = Settings)
+	TSubclassOf<UPCGMeshSelectorBase> MeshSelectorType;
+
+	UPROPERTY(BlueprintReadOnly, VisibleAnywhere, Instanced, Category = Settings)
+	TObjectPtr<UPCGMeshSelectorBase> MeshSelectorInstance;
+
+	UPROPERTY(meta=(DeprecatedProperty, DeprecationMessage="Use MeshSelectorType and MeshSelectorInstance instead."))
+	TArray<FPCGStaticMeshSpawnerEntry> Meshes_DEPRECATED;
+
+protected:
+	void RefreshMeshSelector(); 
 };
 
 class FPCGStaticMeshSpawnerElement : public FSimplePCGElement
@@ -60,4 +83,6 @@ public:
 
 protected:
 	virtual bool ExecuteInternal(FPCGContext* Context) const override;	
+	void SpawnStaticMeshInstances(FPCGContext* Context, const TMap<TSoftObjectPtr<UStaticMesh>, FPCGMeshInstanceList>& MeshInstances, AActor* TargetActor) const;
 };
+
diff --git a/Engine/Plugins/Experimental/PCG/Source/PCG/Public/MeshSelectors/PCGMeshSelectorBase.h b/Engine/Plugins/Experimental/PCG/Source/PCG/Public/MeshSelectors/PCGMeshSelectorBase.h
new file mode 100644
--- /dev/null
+++ b/Engine/Plugins/Experimental/PCG/Source/PCG/Public/MeshSelectors/PCGMeshSelectorBase.h
@@ -0,0 +1,49 @@
+// Copyright Epic Games, Inc. All Rights Reserved.
+
+#pragma once
+
+#include "CoreMinimal.h"
+#include "PCGElement.h"
+#include "PCGPoint.h"
+#include "Data/PCGPointData.h"
+
+#include "Engine/CollisionProfile.h"
+
+#include "PCGMeshSelectorBase.generated.h"
+
+class UPCGStaticMeshSpawnerSettings;
+
+USTRUCT(BlueprintType)
+struct FPCGMeshInstanceList
+{
+	GENERATED_BODY()
+
+	UPROPERTY(BlueprintReadWrite, EditAnywhere, Category = Settings)
+    TArray<FTransform> Instances;
+
+	UPROPERTY(BlueprintReadWrite, EditAnywhere, Category = Settings)
+	bool bOverrideCollisionProfile = false;
+
+	UPROPERTY(BlueprintReadWrite, EditAnywhere, Category = Settings)
+	FCollisionProfileName CollisionProfile;
+};
+
+UCLASS(Abstract, BlueprintType, Blueprintable, ClassGroup = (Procedural))
+class PCG_API UPCGMeshSelectorBase : public UObject 
+{
+	GENERATED_BODY()
+
+public:
+	UFUNCTION(BlueprintNativeEvent, Category = MeshSelection)
+		void SelectInstances(
+		FPCGContext& Context,
+		const UPCGStaticMeshSpawnerSettings* Settings,
+		const UPCGSpatialData* InSpatialData,
+		TMap<TSoftObjectPtr<UStaticMesh>, FPCGMeshInstanceList>& OutMeshInstances) const;
+
+	virtual void SelectInstances_Implementation(
+		FPCGContext& Context,
+		const UPCGStaticMeshSpawnerSettings* Settings,
+		const UPCGSpatialData* InSpatialData,
+		TMap<TSoftObjectPtr<UStaticMesh>, FPCGMeshInstanceList>& OutMeshInstances) const;
+};
diff --git a/Engine/Plugins/Experimental/PCG/Source/PCG/Public/MeshSelectors/PCGMeshSelectorByAttribute.h b/Engine/Plugins/Experimental/PCG/Source/PCG/Public/MeshSelectors/PCGMeshSelectorByAttribute.h
new file mode 100644
--- /dev/null
+++ b/Engine/Plugins/Experimental/PCG/Source/PCG/Public/MeshSelectors/PCGMeshSelectorByAttribute.h
@@ -0,0 +1,31 @@
+// Copyright Epic Games, Inc. All Rights Reserved.
+
+#pragma once
+
+#include "PCGMeshSelectorBase.h"
+
+#include "PCGMeshSelectorByAttribute.generated.h"
+
+UCLASS(BlueprintType, ClassGroup = (Procedural))
+class PCG_API UPCGMeshSelectorByAttribute : public UPCGMeshSelectorBase 
+{
+	GENERATED_BODY()
+
+public:
+	virtual void SelectInstances_Implementation(
+		UPARAM(ref) FPCGContext& Context,
+		const UPCGStaticMeshSpawnerSettings* Settings,
+		const UPCGSpatialData* InSpatialData,
+		TMap<TSoftObjectPtr<UStaticMesh>, FPCGMeshInstanceList>& OutMeshInstances) const override;
+
+public:
+	UPROPERTY(BlueprintReadWrite, EditAnywhere, Category = Settings)
+	FName AttributeName; 
+
+	UPROPERTY(BlueprintReadWrite, EditAnywhere, Category = Settings)
+	bool bOverrideCollisionProfile = false;
+
+	UPROPERTY(BlueprintReadWrite, EditAnywhere, Category = Settings)
+	FCollisionProfileName CollisionProfile;
+};
+
diff --git a/Engine/Plugins/Experimental/PCG/Source/PCG/Public/MeshSelectors/PCGMeshSelectorWeighted.h b/Engine/Plugins/Experimental/PCG/Source/PCG/Public/MeshSelectors/PCGMeshSelectorWeighted.h
new file mode 100644
--- /dev/null
+++ b/Engine/Plugins/Experimental/PCG/Source/PCG/Public/MeshSelectors/PCGMeshSelectorWeighted.h
@@ -0,0 +1,49 @@
+// Copyright Epic Games, Inc. All Rights Reserved.
+
+#pragma once
+
+#include "PCGMeshSelectorBase.h"
+
+#include "PCGMeshSelectorWeighted.generated.h"
+
+USTRUCT(BlueprintType)
+struct PCG_API FPCGMeshSelectorWeightedEntry
+{
+	GENERATED_BODY()
+
+	FPCGMeshSelectorWeightedEntry() = default;
+
+	FPCGMeshSelectorWeightedEntry(TSoftObjectPtr<UStaticMesh> InMesh, int InWeight)
+		: Mesh(InMesh), Weight(InWeight)
+	{}
+
+	UPROPERTY(BlueprintReadWrite, EditAnywhere, Category = Settings)
+	TSoftObjectPtr<UStaticMesh> Mesh;
+
+	UPROPERTY(BlueprintReadWrite, EditAnywhere, Category = Settings)
+	FCollisionProfileName CollisionProfile;
+
+	UPROPERTY(BlueprintReadWrite, EditAnywhere, Category = Settings)
+	bool bOverrideCollisionProfile = false;
+
+	UPROPERTY(BlueprintReadWrite, EditAnywhere, Category = Settings, meta = (ClampMin = "0"))
+	int Weight = 1;
+};
+
+UCLASS(BlueprintType, ClassGroup = (Procedural))
+class PCG_API UPCGMeshSelectorWeighted : public UPCGMeshSelectorBase 
+{
+	GENERATED_BODY()
+
+public:
+	virtual void SelectInstances_Implementation(
+		UPARAM(ref) FPCGContext& Context, 
+		const UPCGStaticMeshSpawnerSettings* Settings, 
+		const UPCGSpatialData* InSpatialData,
+		TMap<TSoftObjectPtr<UStaticMesh>, FPCGMeshInstanceList>& OutMeshInstances) const override;
+
+public:
+	UPROPERTY(BlueprintReadWrite, EditAnywhere, Category = Settings)
+	TArray<FPCGMeshSelectorWeightedEntry> MeshEntries;
+};
+
diff --git a/Engine/Plugins/Experimental/PCG/Source/PCG/Public/MeshSelectors/PCGMeshSelectorWeightedByCategory.h b/Engine/Plugins/Experimental/PCG/Source/PCG/Public/MeshSelectors/PCGMeshSelectorWeightedByCategory.h
new file mode 100644
--- /dev/null
+++ b/Engine/Plugins/Experimental/PCG/Source/PCG/Public/MeshSelectors/PCGMeshSelectorWeightedByCategory.h
@@ -0,0 +1,49 @@
+// Copyright Epic Games, Inc. All Rights Reserved.
+
+#pragma once
+
+#include "PCGMeshSelectorBase.h"
+#include "PCGMeshSelectorWeighted.h"
+
+#include "PCGMeshSelectorWeightedByCategory.generated.h"
+
+USTRUCT(BlueprintType)
+struct PCG_API FPCGWeightedByCategoryEntryList
+{
+	GENERATED_BODY()
+
+	FPCGWeightedByCategoryEntryList() = default;
+
+	FPCGWeightedByCategoryEntryList(const FString& InCategoryEntry, const TArray<FPCGMeshSelectorWeightedEntry>& InWeightedMeshEntries)
+		: CategoryEntry(InCategoryEntry), WeightedMeshEntries(InWeightedMeshEntries)
+	{}
+
+	UPROPERTY(BlueprintReadWrite, EditAnywhere, Category = Settings)
+	FString CategoryEntry;
+
+	UPROPERTY(BlueprintReadWrite, EditAnywhere, Category = Settings)
+	bool IsDefault = false;
+
+	UPROPERTY(BlueprintReadWrite, EditAnywhere, Category = Settings)
+	TArray<FPCGMeshSelectorWeightedEntry> WeightedMeshEntries;
+};
+
+UCLASS(BlueprintType, ClassGroup = (Procedural))
+class PCG_API UPCGMeshSelectorWeightedByCategory : public UPCGMeshSelectorBase 
+{
+	GENERATED_BODY()
+
+public:
+	virtual void SelectInstances_Implementation(
+		UPARAM(ref) FPCGContext& Context,
+		const UPCGStaticMeshSpawnerSettings* Settings,
+		const UPCGSpatialData* InSpatialData,
+		TMap<TSoftObjectPtr<UStaticMesh>, FPCGMeshInstanceList>& OutMeshInstances) const override;
+
+public:
+	UPROPERTY(BlueprintReadWrite, EditAnywhere, Category = Settings)
+	FName CategoryAttribute;
+
+	UPROPERTY(BlueprintReadWrite, EditAnywhere, Category = Settings)
+	TArray<FPCGWeightedByCategoryEntryList> Entries;
+};
\ No newline at end of file
