Fix Static Analysis Win64 warnings in FPCGStaticMeshSpawnerTest

#rnx

diff --git a/Engine/Plugins/Experimental/PCG/Source/PCG/Private/Tests/Elements/PCGStaticMeshSpawnerTest.cpp b/Engine/Plugins/Experimental/PCG/Source/PCG/Private/Tests/Elements/PCGStaticMeshSpawnerTest.cpp
--- a/Engine/Plugins/Experimental/PCG/Source/PCG/Private/Tests/Elements/PCGStaticMeshSpawnerTest.cpp
+++ b/Engine/Plugins/Experimental/PCG/Source/PCG/Private/Tests/Elements/PCGStaticMeshSpawnerTest.cpp
@@ -278,6 +278,7 @@ void TestMeshSelectorWeightedByCategory(
 		if (!Entry)
 		{
 			Test->TestTrue("Validate instanced mesh exists in entries", false);
+			continue;
 		}
 
 		const float ErrorBound = FMath::Max(Entry->ExpectedCount * 0.1f, 10); // 10% error boundary 
