Add setting to set PCGActorPartition grid size + safe guard their deletion

diff --git a/Engine/Plugins/Experimental/PCG/Source/PCG/Private/Grid/PCGPartitionActor.cpp b/Engine/Plugins/Experimental/PCG/Source/PCG/Private/Grid/PCGPartitionActor.cpp
--- a/Engine/Plugins/Experimental/PCG/Source/PCG/Private/Grid/PCGPartitionActor.cpp
+++ b/Engine/Plugins/Experimental/PCG/Source/PCG/Private/Grid/PCGPartitionActor.cpp
@@ -7,6 +7,23 @@
 
 #include "Landscape.h"
 
+constexpr uint32 InvalidPCGGridSizeValue = 0u;
+
+APCGPartitionActor::APCGPartitionActor(const FObjectInitializer& ObjectInitializer)
+	: Super(ObjectInitializer)
+{
+	PCGGridSize = InvalidPCGGridSizeValue;
+}
+
+void APCGPartitionActor::PostLoad()
+{
+	// If the grid size is not set, set it to the default value.
+	if (PCGGridSize == InvalidPCGGridSizeValue)
+	{
+		PCGGridSize = APCGWorldActor::DefaultPartitionGridSize;
+	}
+}
+
 void APCGPartitionActor::BeginPlay()
 {
 	Super::BeginPlay();
@@ -26,16 +43,20 @@ void APCGPartitionActor::EndPlay(const EEndPlayReason::Type EndPlayReason)
 #if WITH_EDITOR
 uint32 APCGPartitionActor::GetDefaultGridSize(UWorld* InWorld) const
 {
-	// TODO: promote to world settings
-	return 25600;
+	if (APCGWorldActor* PCGActor = PCGHelpers::GetPCGWorldActor(InWorld))
+	{
+		return PCGActor->PartitionGridSize;
+	}
+
+	UE_LOG(LogPCG, Error, TEXT("[APCGPartitionActor::InternalGetDefaultGridSize] PCG World Actor was null. Returning default value"));
+	return APCGWorldActor::DefaultPartitionGridSize;
 }
 #endif
 
 FBox APCGPartitionActor::GetFixedBounds() const
 {
 	const FVector Center = GetActorLocation();
-	// TODO: keep this in the actor instead of hardcoding.
-	const uint32 HalfGridSize = 25600 / 2;
+	const FVector::FReal HalfGridSize = PCGGridSize / 2.0;
 	return FBox(Center - HalfGridSize, Center + HalfGridSize);
 }
 
@@ -169,4 +190,40 @@ bool APCGPartitionActor::CleanupDeadGraphInstances()
 
 	return OriginalToLocalMap.IsEmpty();
 }
+
+void APCGPartitionActor::PostCreation()
+{
+	PCGGridSize = GridSize;
+}
+
+bool APCGPartitionActor::IsSafeForDeletion() const
+{
+	ensure(IsInGameThread());
+	for (TObjectPtr<UPCGComponent> PCGComponent : GetAllOriginalPCGComponents())
+	{
+		if (PCGComponent->IsGenerating())
+		{
+			return false;
+		}
+	}
+
+	return true;
+}
+
+TSet<TObjectPtr<UPCGComponent>> APCGPartitionActor::GetAllLocalPCGComponents() const
+{
+	TSet<TObjectPtr<UPCGComponent>> ResultComponents;
+	LocalToOriginalMap.GetKeys(ResultComponents);
+
+	return ResultComponents;
+}
+
+TSet<TObjectPtr<UPCGComponent>> APCGPartitionActor::GetAllOriginalPCGComponents() const
+{
+	TSet<TObjectPtr<UPCGComponent>> ResultComponents;
+	OriginalToLocalMap.GetKeys(ResultComponents);
+
+	return ResultComponents;
+}
+
 #endif // WITH_EDITOR
\ No newline at end of file
diff --git a/Engine/Plugins/Experimental/PCG/Source/PCG/Private/Grid/PCGPartitionActor.h b/Engine/Plugins/Experimental/PCG/Source/PCG/Private/Grid/PCGPartitionActor.h
--- a/Engine/Plugins/Experimental/PCG/Source/PCG/Private/Grid/PCGPartitionActor.h
+++ b/Engine/Plugins/Experimental/PCG/Source/PCG/Private/Grid/PCGPartitionActor.h
@@ -19,6 +19,12 @@ class APCGPartitionActor : public APartitionActor
 	GENERATED_BODY()
 
 public:
+	APCGPartitionActor(const FObjectInitializer& ObjectInitializer = FObjectInitializer::Get());
+
+	//~Begin UObject Interface
+	virtual void PostLoad() override;
+	//~End UObject Interface
+
 	//~Begin AActor Interface
 	virtual void BeginPlay();
 	virtual void EndPlay(const EEndPlayReason::Type EndPlayReason) override;
@@ -41,6 +47,18 @@ public:
 	void AddGraphInstance(UPCGComponent* OriginalComponent);
 	bool RemoveGraphInstance(UPCGComponent* OriginalComponent);
 	bool CleanupDeadGraphInstances();
+
+	/** To be called after the creation of a new actor to copy the GridSize property (Editor only) into the PCGGridSize property */
+	void PostCreation();
+
+	/** [Game thread only] Return if the actor is safe for deletion, meaning no generation is currently running on all original components. */
+	bool IsSafeForDeletion() const;
+
+	/** Return an array of all the PCGComponents on this actor */
+	TSet<TObjectPtr<UPCGComponent>> GetAllLocalPCGComponents() const;
+
+	/** Return a set of all the PCGComponents linked to this actor */
+	TSet<TObjectPtr<UPCGComponent>> GetAllOriginalPCGComponents() const;
 #endif
 
 	// TODO: Make this in-editor only; during runtime, we should keep a map of component to bounds/volume only
@@ -61,4 +79,7 @@ private:
 
 	UPROPERTY()
 	TMap<TObjectPtr<UPCGComponent>, TObjectPtr<UPCGComponent>> LocalToOriginalMap;
+
+	UPROPERTY(VisibleAnywhere, Category = WorldPartition)
+	uint32 PCGGridSize;
 };
diff --git a/Engine/Plugins/Experimental/PCG/Source/PCG/Private/Helpers/PCGActorHelpers.cpp b/Engine/Plugins/Experimental/PCG/Source/PCG/Private/Helpers/PCGActorHelpers.cpp
--- a/Engine/Plugins/Experimental/PCG/Source/PCG/Private/Helpers/PCGActorHelpers.cpp
+++ b/Engine/Plugins/Experimental/PCG/Source/PCG/Private/Helpers/PCGActorHelpers.cpp
@@ -9,6 +9,7 @@
 #include "Components/InstancedStaticMeshComponent.h"
 #include "Components/HierarchicalInstancedStaticMeshComponent.h"
 #include "Engine/InheritableComponentHandler.h"
+#include "Engine/Level.h"
 #include "Engine/SCS_Node.h"
 #include "Engine/SimpleConstructionScript.h"
 #include "Engine/StaticMesh.h"
@@ -285,4 +286,20 @@ void UPCGActorHelpers::GetActorClassDefaultComponents(const TSubclassOf<AActor>&
 	}
 
 	OutComponents = MoveTemp(ResultComponents);
-}
\ No newline at end of file
+}
+
+void UPCGActorHelpers::ForEachActorInLevel(TObjectPtr<ULevel> Level, TSubclassOf<AActor> ActorClass, TFunctionRef<void(AActor*)> Callback)
+{
+	if (!Level)
+	{
+		return;
+	}
+
+	for (AActor* Actor : Level->Actors)
+	{
+		if (Actor && Actor->IsA(ActorClass))
+		{
+			Callback(Actor);
+		}
+	}
+}
diff --git a/Engine/Plugins/Experimental/PCG/Source/PCG/Private/PCGComponent.cpp b/Engine/Plugins/Experimental/PCG/Source/PCG/Private/PCGComponent.cpp
--- a/Engine/Plugins/Experimental/PCG/Source/PCG/Private/PCGComponent.cpp
+++ b/Engine/Plugins/Experimental/PCG/Source/PCG/Private/PCGComponent.cpp
@@ -552,16 +552,11 @@ void UPCGComponent::PostEditChangeProperty(FPropertyChangedEvent& PropertyChange
 				bActivated = false;
 				Refresh();
 
-				// Invalidate the previous bounds to force actor creation (as if we moved the volume)
-				// if we are now partitioned
-				if (bIsNowPartitioned)
-				{
-					LastGeneratedBounds = FBox(EForceInit::ForceInit);
-				}
-
-				// And do a normal refresh
+				// Then invalidate the previous bounds to force actor creation (as if we moved the volume)
+				// and do a normal refresh
 				bActivated = true;
 				bIsPartitioned = bIsNowPartitioned;
+				ResetLastGeneratedBounds();
 				DirtyGenerated();
 				Refresh();
 			}
@@ -1019,6 +1014,11 @@ void UPCGComponent::DirtyGenerated(bool bInDirtyCachedInput)
 	}
 }
 
+void UPCGComponent::ResetLastGeneratedBounds()
+{
+	LastGeneratedBounds = FBox(EForceInit::ForceInit);
+}
+
 void UPCGComponent::Refresh()
 {
 	// Following a change in some properties or in some spatial information related to this component,
diff --git a/Engine/Plugins/Experimental/PCG/Source/PCG/Private/PCGSubsystem.cpp b/Engine/Plugins/Experimental/PCG/Source/PCG/Private/PCGSubsystem.cpp
--- a/Engine/Plugins/Experimental/PCG/Source/PCG/Private/PCGSubsystem.cpp
+++ b/Engine/Plugins/Experimental/PCG/Source/PCG/Private/PCGSubsystem.cpp
@@ -5,6 +5,7 @@
 #include "PCGWorldActor.h"
 #include "Graph/PCGGraphExecutor.h"
 #include "Grid/PCGPartitionActor.h"
+#include "Helpers/PCGActorHelpers.h"
 
 #include "Engine/World.h"
 #include "HAL/IConsoleManager.h"
@@ -168,7 +169,14 @@ namespace PCGSubsystem
 			if (IntersectedBounds.IsValid)
 			{
 				TSharedPtr<TSet<FWorldPartitionReference>> ActorReferences = MakeShared<TSet<FWorldPartitionReference>>();
-				APCGPartitionActor* PCGActor = Cast<APCGPartitionActor>(PartitionSubsystem->GetActor(APCGPartitionActor::StaticClass(), CellCoord, bCreateActor));
+
+				auto PostCreation = [](APartitionActor* Actor) { CastChecked<APCGPartitionActor>(Actor)->PostCreation(); };
+
+				const bool bInBoundsSearch = true;
+				const FGuid DefaultGridGuid;
+				const uint32 DefaultGridSize = 0;
+
+				APCGPartitionActor* PCGActor = Cast<APCGPartitionActor>(PartitionSubsystem->GetActor(APCGPartitionActor::StaticClass(), CellCoord, bCreateActor, DefaultGridGuid, DefaultGridSize, bInBoundsSearch, PostCreation));
 
 				// At this point, if bCreateActor was true, then it exists, but it is not currently loaded; make sure it is loaded
 				// Otherwise, we still need to load it if it exists
@@ -461,6 +469,14 @@ void UPCGSubsystem::CleanupPartitionActors(const FBox& InBounds)
 
 void UPCGSubsystem::DeletePartitionActors()
 {
+	TRACE_CPUPROFILER_EVENT_SCOPE(UPCGSubsystem::DeletePartitionActors);
+
+	// TODO: This should be handled differently when we are able to stop a generation
+	// For now, we need to be careful and not delete partition actors that are linked to components that are currently
+	// generating stuff.
+	// Also, we keep a set on all those actors, in case the partition actor status changes during the loop.
+	TSet<TObjectPtr<APCGPartitionActor>> ActorsNotSafeToBeDeleted;
+
 	TSet<UPackage*> PackagesToCleanup;
 	TSet<FString> PackagesToDeleteFromSCC;
 	UWorld* World = GetWorld();
@@ -470,11 +486,35 @@ void UPCGSubsystem::DeletePartitionActors()
 		return;
 	}
 
-	auto GatherAndDestroyActors = [&PackagesToCleanup, &PackagesToDeleteFromSCC, World](const FWorldPartitionActorDesc* ActorDesc) {
+	auto GatherAndDestroyLoadedActors = [&PackagesToCleanup, &PackagesToDeleteFromSCC, World, &ActorsNotSafeToBeDeleted](AActor* Actor)
+	{
+		// Make sure that this actor was not flagged to not be deleted, or not safe for deletion
+		TObjectPtr<APCGPartitionActor> PartitionActor = CastChecked<APCGPartitionActor>(Actor);
+		if (!PartitionActor->IsSafeForDeletion())
+		{
+			ActorsNotSafeToBeDeleted.Add(PartitionActor);
+		}
+		else if (!ActorsNotSafeToBeDeleted.Contains(PartitionActor))
+		{
+			// Also reset the last generated bounds to indicate to the component to re-create its PartitionActors when it generates.
+			for (TObjectPtr<UPCGComponent> PCGComponent : PartitionActor->GetAllOriginalPCGComponents())
+			{
+				PCGComponent->ResetLastGeneratedBounds();
+			}
+
+			if (UPackage* ExternalPackage = PartitionActor->GetExternalPackage())
+			{
+				PackagesToCleanup.Add(ExternalPackage);
+			}
+
+			World->DestroyActor(PartitionActor);
+		}
+	};
+
+	auto GatherAndDestroyActors = [&PackagesToCleanup, &PackagesToDeleteFromSCC, World, &ActorsNotSafeToBeDeleted, &GatherAndDestroyLoadedActors](const FWorldPartitionActorDesc* ActorDesc) {
 		if (ActorDesc->GetActor())
 		{
-			PackagesToCleanup.Add(ActorDesc->GetActor()->GetExternalPackage());
-			World->DestroyActor(ActorDesc->GetActor());
+			GatherAndDestroyLoadedActors(ActorDesc->GetActor());
 		}
 		else
 		{
@@ -493,6 +533,19 @@ void UPCGSubsystem::DeletePartitionActors()
 
 	FWorldPartitionHelpers::ForEachActorDesc<APCGPartitionActor>(World->GetWorldPartition(), GatherAndDestroyActors);
 
+	// Also cleanup the remaining actors that don't have descriptors, if we have a loaded level
+	if (ULevel* Level = World->GetCurrentLevel())
+	{
+		TRACE_CPUPROFILER_EVENT_SCOPE(UPCGSubsystem::DeletePartitionActors::ForEachActorInLevel);
+		UPCGActorHelpers::ForEachActorInLevel<APCGPartitionActor>(Level, GatherAndDestroyLoadedActors);
+	}
+
+	if (!ActorsNotSafeToBeDeleted.IsEmpty())
+	{
+		// FIXME: see at the top for this function.
+		UE_LOG(LogPCG, Error, TEXT("Tried to delete PCGPartitionActors while their PCGComponent were refreshing. All PCGPartitionActors that are linked to those PCGComponents won't be deleted. You should retry deleting them when the refresh is done."));
+	}
+
 	if (PackagesToCleanup.Num() > 0)
 	{
 		ObjectTools::CleanupAfterSuccessfulDelete(PackagesToCleanup.Array(), /*bPerformanceReferenceCheck=*/true);
diff --git a/Engine/Plugins/Experimental/PCG/Source/PCG/Private/PCGWorldActor.cpp b/Engine/Plugins/Experimental/PCG/Source/PCG/Private/PCGWorldActor.cpp
--- a/Engine/Plugins/Experimental/PCG/Source/PCG/Private/PCGWorldActor.cpp
+++ b/Engine/Plugins/Experimental/PCG/Source/PCG/Private/PCGWorldActor.cpp
@@ -2,7 +2,10 @@
 
 #include "PCGWorldActor.h"
 
+#include "PCGComponent.h"
 #include "PCGSubsystem.h"
+#include "Grid/PCGPartitionActor.h"
+#include "Helpers/PCGActorHelpers.h"
 
 #include "Engine/World.h"
 
@@ -14,6 +17,8 @@ APCGWorldActor::APCGWorldActor(const FObjectInitializer& ObjectInitializer)
 	bIsSpatiallyLoaded = false;
 	bDefaultOutlinerExpansionState = false;
 #endif
+
+	PartitionGridSize = DefaultPartitionGridSize;
 }
 
 void APCGWorldActor::PostLoad()
@@ -55,4 +60,76 @@ void APCGWorldActor::UnregisterFromSubsystem()
 	{
 		PCGSubsystem->UnregisterPCGWorldActor(this);
 	}
-}
\ No newline at end of file
+}
+
+#if WITH_EDITOR
+void APCGWorldActor::OnPartitionGridSizeChanged()
+{
+	TRACE_CPUPROFILER_EVENT_SCOPE(APCGWorldActor::OnPartitionGridSizeChanged);
+
+	UWorld* World = GetWorld();
+	if (!World)
+	{
+		return;
+	}
+
+	UPCGSubsystem* PCGSubsystem = World->GetSubsystem<UPCGSubsystem>();
+	ULevel* Level = World->GetCurrentLevel();
+	if (!PCGSubsystem || !Level)
+	{
+		return;
+	}
+
+	// First gather all the PCG components that linked to PCGPartitionActor and generated
+	TSet<TObjectPtr<UPCGComponent>> AllPartitionedComponents;
+	
+	bool bAllSafeToDelete = true;
+
+	auto AddPartitionComponentAndCheckIfSafeToDelete = [&AllPartitionedComponents, &bAllSafeToDelete](AActor* Actor)
+	{
+		TObjectPtr<APCGPartitionActor> PartitionActor = CastChecked<APCGPartitionActor>(Actor);
+
+		if (!PartitionActor->IsSafeForDeletion())
+		{
+			bAllSafeToDelete = false;
+			return;
+		}
+
+		for (UPCGComponent* PCGComponent : PartitionActor->GetAllOriginalPCGComponents())
+		{
+			if (PCGComponent && PCGComponent->bGenerated)
+			{
+				AllPartitionedComponents.Add(PCGComponent);
+			}
+		}
+	};
+
+	UPCGActorHelpers::ForEachActorInLevel<APCGPartitionActor>(Level, AddPartitionComponentAndCheckIfSafeToDelete);
+
+	// TODO: When we have the capability to stop the generation, we should just do that
+	// For now, just throw an error
+	if (!bAllSafeToDelete)
+	{
+		UE_LOG(LogPCG, Error, TEXT("Trying to change the partition grid size while there are partitionned PCGComponents that are refreshing. We cannot stop the refresh for now, so we abort there. You should delete your partition actors manually and regenerate when the refresh is done"));
+		return;
+	}
+
+	// Then delete all PCGPartitionActors
+	PCGSubsystem->DeletePartitionActors();
+
+	// And finally, refresh all components that were generated.
+	for (TObjectPtr<UPCGComponent> PCGComponent : AllPartitionedComponents)
+	{
+		PCGComponent->DirtyGenerated();
+		PCGComponent->Refresh();
+	}
+}
+
+void APCGWorldActor::PostEditChangeProperty(FPropertyChangedEvent& PropertyChangedEvent)
+{
+	if (PropertyChangedEvent.GetPropertyName() == GET_MEMBER_NAME_CHECKED(APCGWorldActor, PartitionGridSize))
+	{
+		OnPartitionGridSizeChanged();
+	}
+}
+#endif // WITH_EDITOR
\ No newline at end of file
diff --git a/Engine/Plugins/Experimental/PCG/Source/PCG/Public/Helpers/PCGActorHelpers.h b/Engine/Plugins/Experimental/PCG/Source/PCG/Public/Helpers/PCGActorHelpers.h
--- a/Engine/Plugins/Experimental/PCG/Source/PCG/Public/Helpers/PCGActorHelpers.h
+++ b/Engine/Plugins/Experimental/PCG/Source/PCG/Public/Helpers/PCGActorHelpers.h
@@ -8,6 +8,8 @@
 #include "Engine/EngineTypes.h"
 #include "Engine/CollisionProfile.h"
 
+#include <type_traits>
+
 #include "PCGActorHelpers.generated.h"
 
 class AActor;
@@ -16,6 +18,7 @@ class UStaticMesh;
 class UPCGComponent;
 class UMaterialInterface;
 class UActorComponent;
+class ULevel;
 
 struct FPCGISMCBuilderParameters
 {
@@ -41,4 +44,12 @@ public:
 	* @param InComponentClass if supplied will be used to filter the results
 	*/
 	static void GetActorClassDefaultComponents(const TSubclassOf<AActor>& ActorClass, TArray<UActorComponent*>& OutComponents, const TSubclassOf<UActorComponent>& InComponentClass = TSubclassOf<UActorComponent>());
+
+	template <typename T, typename = typename std::enable_if_t<std::is_base_of_v<AActor, T>>>
+	static void ForEachActorInLevel(TObjectPtr<ULevel> Level, TFunctionRef<void(AActor*)> Callback)
+	{
+		return ForEachActorInLevel(Level, T::StaticClass(), Callback);
+	}
+
+	static void ForEachActorInLevel(TObjectPtr<ULevel> Level, TSubclassOf<AActor> ActorClass, TFunctionRef<void(AActor*)> Callback);
 };
\ No newline at end of file
diff --git a/Engine/Plugins/Experimental/PCG/Source/PCG/Public/PCGComponent.h b/Engine/Plugins/Experimental/PCG/Source/PCG/Public/PCGComponent.h
--- a/Engine/Plugins/Experimental/PCG/Source/PCG/Public/PCGComponent.h
+++ b/Engine/Plugins/Experimental/PCG/Source/PCG/Public/PCGComponent.h
@@ -116,6 +116,12 @@ public:
 #if WITH_EDITOR
 	void Refresh();
 	void DirtyGenerated(bool bInDirtyCachedInput = false);
+
+	/** Reset last generated bounds to force PCGPartitionActor creation on next refresh */
+	void ResetLastGeneratedBounds();
+
+	/** Return if we are currently generating the graph for this component */
+	bool IsGenerating() const { return bIsGenerating; }
 #endif
 
 	bool IsPartitioned() const;
diff --git a/Engine/Plugins/Experimental/PCG/Source/PCG/Public/PCGWorldActor.h b/Engine/Plugins/Experimental/PCG/Source/PCG/Public/PCGWorldActor.h
--- a/Engine/Plugins/Experimental/PCG/Source/PCG/Public/PCGWorldActor.h
+++ b/Engine/Plugins/Experimental/PCG/Source/PCG/Public/PCGWorldActor.h
@@ -27,7 +27,23 @@ public:
 	static APCGWorldActor* CreatePCGWorldActor(UWorld* InWorld);
 #endif
 
+	static inline constexpr uint32 DefaultPartitionGridSize = 25600; // 256m
+
+	//~ Begin UObject Interface.
+#if WITH_EDITOR	
+	virtual void PostEditChangeProperty(FPropertyChangedEvent& PropertyChangedEvent) override;
+#endif
+	//~ End UObject Interface.
+
+	/** Size of the grid for PCG partition actors */
+	UPROPERTY(config, EditAnywhere, Category = WorldPartition)
+	uint32 PartitionGridSize;
+
 private:
 	void RegisterToSubsystem();
 	void UnregisterFromSubsystem();
+
+#if WITH_EDITOR	
+	void OnPartitionGridSizeChanged();
+#endif
 };
diff --git a/Engine/Plugins/Experimental/PCG/Source/PCGEditor/Private/PCGEditor.cpp b/Engine/Plugins/Experimental/PCG/Source/PCGEditor/Private/PCGEditor.cpp
--- a/Engine/Plugins/Experimental/PCG/Source/PCGEditor/Private/PCGEditor.cpp
+++ b/Engine/Plugins/Experimental/PCG/Source/PCGEditor/Private/PCGEditor.cpp
@@ -602,7 +602,7 @@ void FPCGEditor::OnCollapseNodesInSubgraph()
 	TObjectPtr<UPCGEditorGraphNodeBase> SubgraphEditorNode;
 	TObjectPtr<UPCGNode> SubgraphNode;
 	{
-		const FScopedTransaction Transaction(*FPCGEditorCommon::ContextIdentifier, LOCTEXT("PCGEditorCollapseInSubgraphMessage", "PCG Editor: Collapse in Subgraph"), nullptr);
+		const FScopedTransaction Transaction(*FPCGEditorCommon::ContextIdentifier, LOCTEXT("PCGEditorCollapseInSubgraphMessage", "PCG Editor: Collapse into Subgraph"), nullptr);
 
 		DeleteSelectedNodes();
 
diff --git a/Engine/Plugins/Experimental/PCG/Source/PCGEditor/Private/PCGEditorCommands.cpp b/Engine/Plugins/Experimental/PCG/Source/PCGEditor/Private/PCGEditorCommands.cpp
--- a/Engine/Plugins/Experimental/PCG/Source/PCGEditor/Private/PCGEditorCommands.cpp
+++ b/Engine/Plugins/Experimental/PCG/Source/PCGEditor/Private/PCGEditorCommands.cpp
@@ -19,7 +19,7 @@ FPCGEditorCommands::FPCGEditorCommands()
 void FPCGEditorCommands::RegisterCommands()
 {
 	UI_COMMAND(Find, "Find", "Finds PCG nodes and comments in the current graph", EUserInterfaceActionType::Button, FInputChord(EModifierKey::Control, EKeys::F));
-	UI_COMMAND(CollapseNodes, "Collapse in Subgraph", "Collapse selected nodes in a separate PCGGraph asset.", EUserInterfaceActionType::Button, FInputChord(EModifierKey::Alt, EKeys::J));
+	UI_COMMAND(CollapseNodes, "Collapse into Subgraph", "Collapse selected nodes into a separate PCGGraph asset.", EUserInterfaceActionType::Button, FInputChord(EModifierKey::Alt, EKeys::J));
 }
 
 #undef LOCTEXT_NAMESPACE
