PCG: Fix missing Super::PostLoad

#rnx

diff --git a/Engine/Plugins/Experimental/PCG/Source/PCG/Private/Grid/PCGPartitionActor.cpp b/Engine/Plugins/Experimental/PCG/Source/PCG/Private/Grid/PCGPartitionActor.cpp
--- a/Engine/Plugins/Experimental/PCG/Source/PCG/Private/Grid/PCGPartitionActor.cpp
+++ b/Engine/Plugins/Experimental/PCG/Source/PCG/Private/Grid/PCGPartitionActor.cpp
@@ -17,6 +17,8 @@ APCGPartitionActor::APCGPartitionActor(const FObjectInitializer& ObjectInitializ
 
 void APCGPartitionActor::PostLoad()
 {
+	Super::PostLoad();
+
 	// If the grid size is not set, set it to the default value.
 	if (PCGGridSize == InvalidPCGGridSizeValue)
 	{
