Fixed some issues when we have recursive graphs/subgraphs

diff --git a/Engine/Plugins/Experimental/PCG/Source/PCG/Private/Elements/PCGExecuteBlueprint.cpp b/Engine/Plugins/Experimental/PCG/Source/PCG/Private/Elements/PCGExecuteBlueprint.cpp
--- a/Engine/Plugins/Experimental/PCG/Source/PCG/Private/Elements/PCGExecuteBlueprint.cpp
+++ b/Engine/Plugins/Experimental/PCG/Source/PCG/Private/Elements/PCGExecuteBlueprint.cpp
@@ -346,7 +346,7 @@ EPCGSettingsType UPCGBlueprintSettings::GetType() const
 	}
 }
 
-void UPCGBlueprintSettings::GetTrackedActorTags(FPCGTagToSettingsMap& OutTagToSettings) const
+void UPCGBlueprintSettings::GetTrackedActorTags(FPCGTagToSettingsMap& OutTagToSettings, TArray<TObjectPtr<const UPCGGraph>>& OutVisitedGraphs) const
 {
 #if WITH_EDITORONLY_DATA
 	for (const FName& Tag : TrackedActorTags)
diff --git a/Engine/Plugins/Experimental/PCG/Source/PCG/Private/PCGComponent.cpp b/Engine/Plugins/Experimental/PCG/Source/PCG/Private/PCGComponent.cpp
--- a/Engine/Plugins/Experimental/PCG/Source/PCG/Private/PCGComponent.cpp
+++ b/Engine/Plugins/Experimental/PCG/Source/PCG/Private/PCGComponent.cpp
@@ -1021,6 +1021,14 @@ void UPCGComponent::ResetLastGeneratedBounds()
 
 void UPCGComponent::Refresh()
 {
+#if WITH_EDITOR
+	// Disable auto-refreshing on preview actors until we have something more robust on the execution side.
+	if (GetOwner() && GetOwner()->bIsEditorPreviewActor)
+	{
+		return;
+	}
+#endif
+
 	// Following a change in some properties or in some spatial information related to this component,
 	// We need to regenerate the graph, depending of the state in the editor.
 	// In the case of a non-partitioned graph, we need to generate the graph only if it was previously generated & tagged for regeneration
@@ -1432,7 +1440,10 @@ FBox UPCGComponent::GetGridBounds(AActor* Actor) const
 		// Then intersect with the original component's bounds.
 		if (const UPCGComponent* OriginalComponent = PartitionActor->GetOriginalComponent(this))
 		{
-			Bounds = Bounds.Overlap(OriginalComponent->GetGridBounds());
+			if (OriginalComponent->GetOwner() != PartitionActor)
+			{
+				Bounds = Bounds.Overlap(OriginalComponent->GetGridBounds());
+			}
 		}
 	}
 	// TODO: this might not work in non-editor builds
diff --git a/Engine/Plugins/Experimental/PCG/Source/PCG/Private/PCGGraph.cpp b/Engine/Plugins/Experimental/PCG/Source/PCG/Private/PCGGraph.cpp
--- a/Engine/Plugins/Experimental/PCG/Source/PCG/Private/PCGGraph.cpp
+++ b/Engine/Plugins/Experimental/PCG/Source/PCG/Private/PCGGraph.cpp
@@ -495,15 +495,28 @@ void UPCGGraph::SetExtraEditorNodes(const TArray<TObjectPtr<const UObject>>& InN
 FPCGTagToSettingsMap UPCGGraph::GetTrackedTagsToSettings() const
 {
 	FPCGTagToSettingsMap TagsToSettings;
+	TArray<TObjectPtr<const UPCGGraph>> VisitedGraphs;
+
+	GetTrackedTagsToSettings(TagsToSettings, VisitedGraphs);
+	return TagsToSettings;
+}
+
+void UPCGGraph::GetTrackedTagsToSettings(FPCGTagToSettingsMap& OutTagsToSettings, TArray<TObjectPtr<const UPCGGraph>>& OutVisitedGraphs) const
+{
+	if (OutVisitedGraphs.Contains(this))
+	{
+		return;
+	}
+
+	OutVisitedGraphs.Emplace(this);
+
 	for (UPCGNode* Node : Nodes)
 	{
 		if (Node && Node->DefaultSettings)
 		{
-			Node->DefaultSettings->GetTrackedActorTags(TagsToSettings);
+			Node->DefaultSettings->GetTrackedActorTags(OutTagsToSettings, OutVisitedGraphs);
 		}
 	}
-
-	return TagsToSettings;
 }
 
 void UPCGGraph::NotifyGraphChanged(bool bIsStructural)
@@ -515,6 +528,14 @@ void UPCGGraph::NotifyGraphChanged(bool bIsStructural)
 		return;
 	}
 
+	// Skip recursive cases which can happen either through direct recursivity (A -> A) or indirectly (A -> B -> A)
+	if (bIsNotifying)
+	{
+		return;
+	}
+
+	bIsNotifying = true;
+
 	// Notify the subsystem/compiler cache before so it gets recompiled properly
 	if (bIsStructural && GEditor)
 	{
@@ -527,6 +548,8 @@ void UPCGGraph::NotifyGraphChanged(bool bIsStructural)
 	}
 
 	OnGraphChangedDelegate.Broadcast(this, /*bIsStructural=*/bIsStructural);
+
+	bIsNotifying = false;
 }
 
 void UPCGGraph::OnNodeChanged(UPCGNode* InNode, EPCGChangeType ChangeType)
diff --git a/Engine/Plugins/Experimental/PCG/Source/PCG/Private/PCGHelpers.cpp b/Engine/Plugins/Experimental/PCG/Source/PCG/Private/PCGHelpers.cpp
--- a/Engine/Plugins/Experimental/PCG/Source/PCG/Private/PCGHelpers.cpp
+++ b/Engine/Plugins/Experimental/PCG/Source/PCG/Private/PCGHelpers.cpp
@@ -136,11 +136,24 @@ namespace PCGHelpers
 		return (InWorld && InWorld->GetSubsystem<UPCGSubsystem>()) ? InWorld->GetSubsystem<UPCGSubsystem>()->GetPCGWorldActor() : nullptr;
 	}
 
+	bool CanBeExpanded(UClass* ObjectClass)
+	{
+		// There shouldn't be any need to dig through Niagara assets + there are some issues (most likely related to loading) with parsing all their dependencies
+		if (!ObjectClass ||
+			ObjectClass->GetFName() == TEXT("NiagaraSystem") ||
+			ObjectClass->GetFName() == TEXT("NiagaraComponent"))
+		{
+			return false;
+		}
+
+		return true;
+	}
+
 	void GatherDependencies(UObject* Object, TSet<TObjectPtr<UObject>>& OutDependencies)
 	{
 		UClass* ObjectClass = Object ? Object->GetClass() : nullptr;
 
-		if (!ObjectClass)
+		if(!CanBeExpanded(ObjectClass))
 		{
 			return;
 		}
@@ -162,6 +175,11 @@ namespace PCGHelpers
 			}
 		};
 
+		if (!InContainer || !Property)
+		{
+			return;
+		}
+
 		if (FObjectProperty* ObjectProperty = CastField<FObjectProperty>(Property))
 		{
 			UObject* Object = ObjectProperty->GetPropertyValue_InContainer(InContainer);
diff --git a/Engine/Plugins/Experimental/PCG/Source/PCG/Private/PCGSubgraph.cpp b/Engine/Plugins/Experimental/PCG/Source/PCG/Private/PCGSubgraph.cpp
--- a/Engine/Plugins/Experimental/PCG/Source/PCG/Private/PCGSubgraph.cpp
+++ b/Engine/Plugins/Experimental/PCG/Source/PCG/Private/PCGSubgraph.cpp
@@ -55,16 +55,11 @@ void UPCGBaseSubgraphSettings::PostEditChangeProperty(struct FPropertyChangedEve
 	}
 }
 
-void UPCGBaseSubgraphSettings::GetTrackedActorTags(FPCGTagToSettingsMap& OutTagToSettings) const
+void UPCGBaseSubgraphSettings::GetTrackedActorTags(FPCGTagToSettingsMap& OutTagToSettings, TArray<TObjectPtr<const UPCGGraph>>& VisitedGraphs) const
 {
 	if (UPCGGraph* Subgraph = GetSubgraph())
 	{
-		FPCGTagToSettingsMap SubgraphMap = Subgraph->GetTrackedTagsToSettings();
-
-		for (const auto& TagToSettings : SubgraphMap)
-		{
-			OutTagToSettings.FindOrAdd(TagToSettings.Key).Append(TagToSettings.Value);
-		}
+		Subgraph->GetTrackedTagsToSettings(OutTagToSettings, VisitedGraphs);
 	}
 }
 
diff --git a/Engine/Plugins/Experimental/PCG/Source/PCG/Public/Elements/PCGExecuteBlueprint.h b/Engine/Plugins/Experimental/PCG/Source/PCG/Public/Elements/PCGExecuteBlueprint.h
--- a/Engine/Plugins/Experimental/PCG/Source/PCG/Public/Elements/PCGExecuteBlueprint.h
+++ b/Engine/Plugins/Experimental/PCG/Source/PCG/Public/Elements/PCGExecuteBlueprint.h
@@ -155,7 +155,7 @@ public:
 	virtual FName GetDefaultNodeName() const override { return FName(TEXT("BlueprintNode")); }
 	virtual FLinearColor GetNodeTitleColor() const override;
 	virtual EPCGSettingsType GetType() const override;
-	virtual void GetTrackedActorTags(FPCGTagToSettingsMap& OutTagToSettings) const override;
+	virtual void GetTrackedActorTags(FPCGTagToSettingsMap& OutTagToSettings, TArray<TObjectPtr<const UPCGGraph>>& OutVisitedGraphs) const override;
 	virtual UObject* GetJumpTargetForDoubleClick() const override;
 #endif
 
diff --git a/Engine/Plugins/Experimental/PCG/Source/PCG/Public/PCGGraph.h b/Engine/Plugins/Experimental/PCG/Source/PCG/Public/PCGGraph.h
--- a/Engine/Plugins/Experimental/PCG/Source/PCG/Public/PCGGraph.h
+++ b/Engine/Plugins/Experimental/PCG/Source/PCG/Public/PCGGraph.h
@@ -88,6 +88,7 @@ public:
 
 #if WITH_EDITOR
 	FPCGTagToSettingsMap GetTrackedTagsToSettings() const;
+	void GetTrackedTagsToSettings(FPCGTagToSettingsMap& OutTagsToSettings, TArray<TObjectPtr<const UPCGGraph>>& OutVisitedGraphs) const;
 #endif
 
 #if WITH_EDITOR
@@ -122,5 +123,6 @@ private:
 	int32 GraphChangeNotificationsDisableCounter = 0;
 	bool bDelayedChangeNotification = false;
 	bool bDelayedChangeNotificationStructural = false;
+	bool bIsNotifying = false;
 #endif // WITH_EDITOR
 };
diff --git a/Engine/Plugins/Experimental/PCG/Source/PCG/Public/PCGSettings.h b/Engine/Plugins/Experimental/PCG/Source/PCG/Public/PCGSettings.h
--- a/Engine/Plugins/Experimental/PCG/Source/PCG/Public/PCGSettings.h
+++ b/Engine/Plugins/Experimental/PCG/Source/PCG/Public/PCGSettings.h
@@ -11,6 +11,7 @@
 
 class UPCGNode;
 class UPCGSettings;
+class UPCGGraph;
 
 typedef TMap<FName, TSet<TWeakObjectPtr<const UPCGSettings>>> FPCGTagToSettingsMap;
 
@@ -71,7 +72,7 @@ public:
 	virtual FLinearColor GetNodeTitleColor() const { return FLinearColor::White; }
 	virtual EPCGSettingsType GetType() const { return EPCGSettingsType::Generic; }
 	/** Derived classes must implement this to communicate dependencies on external actors */
-	virtual void GetTrackedActorTags(FPCGTagToSettingsMap& OutTagToSettings) const {}
+	virtual void GetTrackedActorTags(FPCGTagToSettingsMap& OutTagToSettings, TArray<TObjectPtr<const UPCGGraph>>& OutVisitedGraphs) const {}
 	/** Override this class to provide an UObject to jump to in case of double click on node
 	 *  ie. returning a blueprint instance will open the given blueprint in its editor.
 	 *  By default, it will return the underlying class, to try to jump to its header in code
diff --git a/Engine/Plugins/Experimental/PCG/Source/PCG/Public/PCGSubgraph.h b/Engine/Plugins/Experimental/PCG/Source/PCG/Public/PCGSubgraph.h
--- a/Engine/Plugins/Experimental/PCG/Source/PCG/Public/PCGSubgraph.h
+++ b/Engine/Plugins/Experimental/PCG/Source/PCG/Public/PCGSubgraph.h
@@ -26,7 +26,7 @@ protected:
 	//~End UObject interface implementation
 
 	//~Begin UPCGSettings interface
-	virtual void GetTrackedActorTags(FPCGTagToSettingsMap& OutTagToSettings) const override;
+	virtual void GetTrackedActorTags(FPCGTagToSettingsMap& OutTagToSettings, TArray<TObjectPtr<const UPCGGraph>>& OutVisitedGraphs) const override;
 #endif
 
 	TArray<FPCGPinProperties> InputPinProperties() const override;
