PCG Editor: When placing a node from an existing pin, it binds to the first valid match

diff --git a/Engine/Plugins/Experimental/PCG/Source/PCGEditor/Private/PCGEditorGraphNodeBase.cpp b/Engine/Plugins/Experimental/PCG/Source/PCGEditor/Private/PCGEditorGraphNodeBase.cpp
--- a/Engine/Plugins/Experimental/PCG/Source/PCGEditor/Private/PCGEditorGraphNodeBase.cpp
+++ b/Engine/Plugins/Experimental/PCG/Source/PCGEditor/Private/PCGEditorGraphNodeBase.cpp
@@ -119,22 +119,25 @@ void UPCGEditorGraphNodeBase::GetNodeContextMenuActions(UToolMenu* Menu, class U
 
 void UPCGEditorGraphNodeBase::AutowireNewNode(UEdGraphPin* FromPin)
 {
-	if (FromPin->Direction == EEdGraphPinDirection::EGPD_Output)
+	if (PCGNode == nullptr || FromPin == nullptr)
 	{
-		if (PCGNode && PCGNode->GetInputPins().Num() > 0)
-		{
-			const FName& InPinName = PCGNode->GetInputPins()[0]->Properties.Label;
-			UEdGraphPin* ToPin = FindPinChecked(InPinName, EEdGraphPinDirection::EGPD_Input);
-			GetSchema()->TryCreateConnection(FromPin, ToPin);
-		}
+		return;
 	}
-	else if (FromPin->Direction == EEdGraphPinDirection::EGPD_Input)
+
+	const bool bFromPinIsInput = FromPin->Direction == EEdGraphPinDirection::EGPD_Input;
+	const TArray<TObjectPtr<UPCGPin>>& OtherPinsList = bFromPinIsInput ? PCGNode->GetOutputPins() : PCGNode->GetInputPins();
+
+	// Try to connect to the first compatible pin
+	for (TObjectPtr<UPCGPin> OtherPin : OtherPinsList)
 	{
-		if (PCGNode && PCGNode->GetOutputPins().Num() > 0)
+		check(OtherPin);
+
+		const FName& OtherPinName = OtherPin->Properties.Label;
+		UEdGraphPin* ToPin = FindPinChecked(OtherPinName, bFromPinIsInput ? EEdGraphPinDirection::EGPD_Output : EEdGraphPinDirection::EGPD_Input);
+		if (ToPin && GetSchema()->TryCreateConnection(FromPin, ToPin))
 		{
-			const FName& OutPinName = PCGNode->GetOutputPins()[0]->Properties.Label;
-			UEdGraphPin* ToPin = FindPinChecked(OutPinName, EEdGraphPinDirection::EGPD_Output);
-			GetSchema()->TryCreateConnection(FromPin, ToPin);
+			// Connection succeeded
+			break;
 		}
 	}
 
