Improved robustness to components being deleted and possibly not notifying properly the partition actors / removing tasks from the graph executor

#rnx

diff --git a/Engine/Plugins/Experimental/PCG/Source/PCG/Private/PCGComponent.cpp b/Engine/Plugins/Experimental/PCG/Source/PCG/Private/PCGComponent.cpp
--- a/Engine/Plugins/Experimental/PCG/Source/PCG/Private/PCGComponent.cpp
+++ b/Engine/Plugins/Experimental/PCG/Source/PCG/Private/PCGComponent.cpp
@@ -1040,7 +1040,10 @@ void UPCGComponent::Refresh()
 	{
 		if (IsPartitioned() && GetSubsystem())
 		{
-			GetSubsystem()->DelayUnpartitionGraph(this);
+			if (LastGeneratedBounds.IsValid)
+			{
+				GetSubsystem()->DelayUnpartitionGraph(this);
+			}
 		}
 		else
 		{
@@ -1178,7 +1181,13 @@ UPCGData* UPCGComponent::CreateActorPCGData()
 
 UPCGData* UPCGComponent::CreateActorPCGData(AActor* Actor)
 {
-	TRACE_CPUPROFILER_EVENT_SCOPE(UPCGComponent::CreateActorPCGData)
+	TRACE_CPUPROFILER_EVENT_SCOPE(UPCGComponent::CreateActorPCGData);
+
+	if (!Actor)
+	{
+		return nullptr;
+	}
+
 	// In this case, we'll build the data type that's closest to known actor types
 	// TODO: add factory for extensibility
 	if (APCGPartitionActor* PartitionActor = Cast<APCGPartitionActor>(Actor))
diff --git a/Engine/Plugins/Experimental/PCG/Source/PCG/Private/PCGSubsystem.cpp b/Engine/Plugins/Experimental/PCG/Source/PCG/Private/PCGSubsystem.cpp
--- a/Engine/Plugins/Experimental/PCG/Source/PCG/Private/PCGSubsystem.cpp
+++ b/Engine/Plugins/Experimental/PCG/Source/PCG/Private/PCGSubsystem.cpp
@@ -98,14 +98,21 @@ void UPCGSubsystem::UnregisterPCGWorldActor(APCGWorldActor* InActor)
 FPCGTaskId UPCGSubsystem::ScheduleComponent(UPCGComponent* PCGComponent, const TArray<FPCGTaskId>& Dependencies)
 {
 	check(GraphExecutor);
-	FPCGTaskId ProcessTaskId = GraphExecutor->Schedule(PCGComponent, Dependencies);
+
+	FPCGTaskId ProcessTaskId = PCGComponent ? GraphExecutor->Schedule(PCGComponent, Dependencies) : InvalidTaskId;
 
 	if (ProcessTaskId != InvalidTaskId)
 	{
 #if WITH_EDITOR
-		return GraphExecutor->ScheduleGeneric([PCGComponent]() {
-			const FBox NewBounds = PCGComponent->GetGridBounds();
-			PCGComponent->PostProcessGraph(NewBounds, /*bGenerate=*/true);
+		TWeakObjectPtr<UPCGComponent> ComponentPtr(PCGComponent);
+
+		return GraphExecutor->ScheduleGeneric([ComponentPtr]() {
+			if (UPCGComponent* Component = ComponentPtr.Get())
+			{
+				const FBox NewBounds = Component->GetGridBounds();
+				Component->PostProcessGraph(NewBounds, /*bGenerate=*/true);
+			}
+
 			return true;
 			}, { ProcessTaskId });
 #else
@@ -115,14 +122,24 @@ FPCGTaskId UPCGSubsystem::ScheduleComponent(UPCGComponent* PCGComponent, const T
 	else
 	{
 		UE_LOG(LogPCG, Error, TEXT("[ScheduleComponent] Didn't schedule any task."));
-		PCGComponent->OnProcessGraphAborted();
+		if (PCGComponent)
+		{
+			PCGComponent->OnProcessGraphAborted();
+		}
 		return InvalidTaskId;
 	}
 }
 
 FPCGTaskId UPCGSubsystem::ScheduleGraph(UPCGGraph* Graph, UPCGComponent* SourceComponent, FPCGElementPtr InputElement, const TArray<FPCGTaskId>& Dependencies)
 {
-	return GraphExecutor->Schedule(Graph, SourceComponent, InputElement, Dependencies);
+	if (SourceComponent)
+	{
+		return GraphExecutor->Schedule(Graph, SourceComponent, InputElement, Dependencies);
+	}
+	else
+	{
+		return InvalidTaskId;
+	}
 }
 
 ETickableTickType UPCGSubsystem::GetTickableTickType() const
@@ -304,36 +321,56 @@ namespace PCGSubsystem
 FPCGTaskId UPCGSubsystem::ProcessGraph(UPCGComponent* Component, const FBox& InPreviousBounds, const FBox& InNewBounds, bool bGenerate, bool bSave)
 {
 	check(Component && (InPreviousBounds.IsValid || InNewBounds.IsValid));
+	TWeakObjectPtr<UPCGComponent> ComponentPtr(Component);
 
 	// TODO: optimal implementation would find the difference between the previous bounds and the new bounds
 	// and process these only. This is esp. important because of the CreateActor parameter.
-	auto ScheduleTask = [this, Component, InPreviousBounds, InNewBounds, bGenerate](APCGPartitionActor* PCGActor, const FBox& InBounds, const TArray<FPCGTaskId>& TaskDependencies){
-		auto UnpartitionTask = [Component, PCGActor]() {
+	auto ScheduleTask = [this, ComponentPtr, InPreviousBounds, InNewBounds, bGenerate](APCGPartitionActor* PCGActor, const FBox& InBounds, const TArray<FPCGTaskId>& TaskDependencies){
+		auto UnpartitionTask = [ComponentPtr, PCGActor]() {
 			// TODO: PCG actors that become empty could be deleted, but we also need to keep track
 			// of packages that would need to be deleted from SCC.
-			PCGActor->RemoveGraphInstance(Component);
+			if (UPCGComponent* Component = ComponentPtr.Get())
+			{
+				PCGActor->RemoveGraphInstance(Component);
+			}
+			else
+			{
+#if WITH_EDITOR
+				PCGActor->CleanupDeadGraphInstances();
+#endif
+			}
 			return true;
 		};
 
-		auto PartitionTask = [Component, PCGActor]() {
-			PCGActor->AddGraphInstance(Component);
+		auto PartitionTask = [ComponentPtr, PCGActor]() {
+			if (UPCGComponent* Component = ComponentPtr.Get())
+			{
+				PCGActor->AddGraphInstance(Component);
+			}			
 			return true;
 		};
 
-		auto ScheduleGraph = [this, PCGActor, Component, &TaskDependencies]() {
-			// Ensure that the PCG actor has a matching local component.
-			// This is done immediately, but technically we could add it as a task
-			PCGActor->AddGraphInstance(Component);
+		auto ScheduleGraph = [this, PCGActor, ComponentPtr, &TaskDependencies]() {
+			if (UPCGComponent* Component = ComponentPtr.Get())
+			{
+				// Ensure that the PCG actor has a matching local component.
+				// This is done immediately, but technically we could add it as a task
+				PCGActor->AddGraphInstance(Component);
 
-			UPCGComponent* LocalComponent = PCGActor->GetLocalComponent(Component);
+				UPCGComponent* LocalComponent = PCGActor->GetLocalComponent(Component);
 
-			if (!LocalComponent)
+				if (!LocalComponent)
+				{
+					return InvalidTaskId;
+				}
+
+				return LocalComponent->GenerateInternal(/*bForce=*/false, TaskDependencies);
+			}
+			else
 			{
 				UE_LOG(LogPCG, Error, TEXT("[ProcessGraph] PCG Component on PCG Actor is null"));
 				return InvalidTaskId;
 			}
-
-			return LocalComponent->GenerateInternal(/*bForce=*/false, TaskDependencies);
 		};
 
 		const bool bIsInOldBounds = InPreviousBounds.IsValid && InBounds.Intersect(InPreviousBounds);
@@ -388,8 +425,11 @@ FPCGTaskId UPCGSubsystem::ProcessGraph(UPCGComponent* Component, const FBox& InP
 	// Finally, call PostProcessGraph if something happened
 	if (ProcessAllCellsTaskId != InvalidTaskId)
 	{
-		auto PostProcessGraph = [Component, InNewBounds, bGenerate]() {
-			Component->PostProcessGraph(InNewBounds, bGenerate);
+		auto PostProcessGraph = [ComponentPtr, InNewBounds, bGenerate]() {
+			if (UPCGComponent* Component = ComponentPtr.Get())
+			{
+				Component->PostProcessGraph(InNewBounds, bGenerate);
+			}
 			return true;
 		};
 
@@ -397,16 +437,23 @@ FPCGTaskId UPCGSubsystem::ProcessGraph(UPCGComponent* Component, const FBox& InP
 	}
 	else
 	{
-		Component->OnProcessGraphAborted();
+		if(Component)
+		{
+			Component->OnProcessGraphAborted();
+		}
+
 		return InvalidTaskId;
 	}
 }
 
 void UPCGSubsystem::CleanupGraph(UPCGComponent* Component, const FBox& InBounds, bool bRemoveComponents, bool bSave)
 {
-	auto ScheduleTask = [this, Component, bRemoveComponents](APCGPartitionActor* PCGActor, const FBox& InIntersectedBounds, const TArray<FPCGTaskId>& TaskDependencies) {
-		auto CleanupTask = [Component, PCGActor, bRemoveComponents]() {
-			check(Component != nullptr && PCGActor != nullptr);
+	TWeakObjectPtr<UPCGComponent> ComponentPtr(Component);
+
+	auto ScheduleTask = [this, ComponentPtr, bRemoveComponents](APCGPartitionActor* PCGActor, const FBox& InIntersectedBounds, const TArray<FPCGTaskId>& TaskDependencies) {
+		auto CleanupTask = [ComponentPtr, PCGActor, bRemoveComponents]() {
+			check(ComponentPtr.Get() != nullptr && PCGActor != nullptr);
+			UPCGComponent* Component = ComponentPtr.Get();
 
 			if (!PCGActor || !Component)
 			{
@@ -429,10 +476,13 @@ void UPCGSubsystem::CleanupGraph(UPCGComponent* Component, const FBox& InBounds,
 
 void UPCGSubsystem::DirtyGraph(UPCGComponent* Component, const FBox& InBounds, bool bDirtyInputs)
 {
-	auto ScheduleTask = [this, Component, bDirtyInputs](APCGPartitionActor* PCGActor, const FBox& InIntersectedBounds, const TArray<FPCGTaskId>& TaskDependencies) {
+	TWeakObjectPtr<UPCGComponent> ComponentPtr(Component);
+
+	auto ScheduleTask = [this, ComponentPtr, bDirtyInputs](APCGPartitionActor* PCGActor, const FBox& InIntersectedBounds, const TArray<FPCGTaskId>& TaskDependencies) {
 		// In the specific case of the dirty, we want to bypass the execution queue, esp. since there's nothing happening here
 		// so we will run the command now, and not delay it.
-		check(Component != nullptr && PCGActor != nullptr);
+		check(ComponentPtr.Get() != nullptr && PCGActor != nullptr);
+		UPCGComponent* Component = ComponentPtr.Get();
 
 		if (!PCGActor || !Component)
 		{
@@ -464,7 +514,7 @@ void UPCGSubsystem::CleanupPartitionActors(const FBox& InBounds)
 		return GraphExecutor->ScheduleGeneric(CleanupTask, TaskDependencies);
 	};
 
-	PCGSubsystem::ForEachIntersectingCell(GraphExecutor, GetWorld(), InBounds, /*bCreateActor=*/false, /*bLoadCell=*/false, /*bSave=*/true, ScheduleTask);
+	PCGSubsystem::ForEachIntersectingCell(GraphExecutor, GetWorld(), InBounds, /*bCreateActor=*/false, /*bLoadCell=*/false, /*bSave=*/false, ScheduleTask);
 }
 
 void UPCGSubsystem::DeletePartitionActors()
@@ -604,11 +654,15 @@ FPCGTaskId UPCGSubsystem::DelayGenerateGraph(UPCGComponent* Component, bool bSav
 FPCGTaskId UPCGSubsystem::DelayProcessGraph(UPCGComponent* Component, bool bGenerate, bool bSave, bool bUseEmptyNewBounds)
 {
 	check(Component && Component->IsPartitioned());
+	TWeakObjectPtr<UPCGComponent> ComponentPtr(Component);
 
-	auto ExecuteProcessGraph = [this, Component, bGenerate, bSave, bUseEmptyNewBounds]() {
-		const FBox PreviousBounds = Component->LastGeneratedBounds;
-		FBox NewBounds = bUseEmptyNewBounds ? FBox(EForceInit::ForceInit) : Component->GetGridBounds();
-		/*FPCGTaskId GraphTaskId = */ ProcessGraph(Component, PreviousBounds, NewBounds, bGenerate, bSave);
+	auto ExecuteProcessGraph = [this, ComponentPtr, bGenerate, bSave, bUseEmptyNewBounds]() {
+		if (UPCGComponent* Component = ComponentPtr.Get())
+		{
+			const FBox PreviousBounds = Component->LastGeneratedBounds;
+			FBox NewBounds = bUseEmptyNewBounds ? FBox(EForceInit::ForceInit) : Component->GetGridBounds();
+			/*FPCGTaskId GraphTaskId = */ ProcessGraph(Component, PreviousBounds, NewBounds, bGenerate, bSave);
+		}
 
 		return true;
 	};
