Added landscape input pin
Improved dirtying mechanism based on flags, will continue to evolve based on tracked actors
Improved ToPointData call for special data case (e.g. landscape) to do early culling for intersections
Moved some typedefs around which caused a few changes around in the includes

diff --git a/Engine/Plugins/Experimental/PCG/Source/PCG/Private/Data/PCGIntersectionData.cpp b/Engine/Plugins/Experimental/PCG/Source/PCG/Private/Data/PCGIntersectionData.cpp
--- a/Engine/Plugins/Experimental/PCG/Source/PCG/Private/Data/PCGIntersectionData.cpp
+++ b/Engine/Plugins/Experimental/PCG/Source/PCG/Private/Data/PCGIntersectionData.cpp
@@ -109,7 +109,7 @@ UPCGPointData* UPCGIntersectionData::CreateAndFilterPointData(FPCGContext* Conte
 	check(X && Y);
 	check(X->GetDimension() <= Y->GetDimension());
 
-	const UPCGPointData* SourcePointData = X->ToPointData(Context);
+	const UPCGPointData* SourcePointData = X->ToPointData(Context, CachedBounds);
 
 	if (!SourcePointData)
 	{
diff --git a/Engine/Plugins/Experimental/PCG/Source/PCG/Private/Data/PCGLandscapeData.cpp b/Engine/Plugins/Experimental/PCG/Source/PCG/Private/Data/PCGLandscapeData.cpp
--- a/Engine/Plugins/Experimental/PCG/Source/PCG/Private/Data/PCGLandscapeData.cpp
+++ b/Engine/Plugins/Experimental/PCG/Source/PCG/Private/Data/PCGLandscapeData.cpp
@@ -59,7 +59,7 @@ bool UPCGLandscapeData::SamplePoint(const FTransform& InTransform, const FBox& I
 	}
 }
 
-const UPCGPointData* UPCGLandscapeData::CreatePointData(FPCGContext* Context) const
+const UPCGPointData* UPCGLandscapeData::CreatePointData(FPCGContext* Context, const FBox& InBounds) const
 {
 	TRACE_CPUPROFILER_EVENT_SCOPE(UPCGLandscapeData::CreatePointData);
 
@@ -67,9 +67,15 @@ const UPCGPointData* UPCGLandscapeData::CreatePointData(FPCGContext* Context) co
 	Data->InitializeFromData(this);
 	TArray<FPCGPoint>& Points = Data->GetMutablePoints();
 
+	FBox EffectiveBounds = Bounds;
+	if (InBounds.IsValid)
+	{
+		EffectiveBounds = Bounds.Overlap(InBounds);
+	}
+
 	// TODO: add offset to nearest edge, will have an impact if the grid size doesn't match the landscape size
-	const FVector MinPt = Transform.InverseTransformPosition(Bounds.Min);
-	const FVector MaxPt = Transform.InverseTransformPosition(Bounds.Max);
+	const FVector MinPt = Transform.InverseTransformPosition(EffectiveBounds.Min);
+	const FVector MaxPt = Transform.InverseTransformPosition(EffectiveBounds.Max);
 
 	const int32 MinX = FMath::FloorToInt(MinPt.X);
 	const int32 MaxX = FMath::FloorToInt(MaxPt.X);
@@ -78,17 +84,17 @@ const UPCGPointData* UPCGLandscapeData::CreatePointData(FPCGContext* Context) co
 
 	int32 NumIterations = (MaxX - MinX) * (MaxY - MinY);
 
-	FPCGAsync::AsyncPointProcessing(Context, NumIterations, Points, [this, MinX, MaxX, MinY](int32 Index, FPCGPoint& OutPoint)
+	FPCGAsync::AsyncPointProcessing(Context, NumIterations, Points, [this, &EffectiveBounds, MinX, MaxX, MinY](int32 Index, FPCGPoint& OutPoint)
 	{
 		const int X = MinX + (Index % (MaxX - MinX));
 		const int Y = MinY + (Index / (MaxX - MinX));
-		const bool bPlaneCase = Bounds.Min.Z == Bounds.Max.Z;
+		const bool bPlaneCase = EffectiveBounds.Min.Z == EffectiveBounds.Max.Z;
 
 		FVector VertexLocation = Transform.TransformPosition(FVector(X, Y, 0));
 		TOptional<float> HeightAtVertex = Landscape->GetHeightAtLocation(VertexLocation);
 		if (HeightAtVertex.IsSet() &&
-			HeightAtVertex.GetValue() >= Bounds.Min.Z &&
-			(bPlaneCase ? HeightAtVertex.GetValue() <= Bounds.Max.Z : HeightAtVertex.GetValue() < Bounds.Max.Z))
+			HeightAtVertex.GetValue() >= EffectiveBounds.Min.Z &&
+			(bPlaneCase ? HeightAtVertex.GetValue() <= EffectiveBounds.Max.Z : HeightAtVertex.GetValue() < EffectiveBounds.Max.Z))
 		{
 			VertexLocation.Z = HeightAtVertex.GetValue();
 
diff --git a/Engine/Plugins/Experimental/PCG/Source/PCG/Private/Data/PCGSpatialData.cpp b/Engine/Plugins/Experimental/PCG/Source/PCG/Private/Data/PCGSpatialData.cpp
--- a/Engine/Plugins/Experimental/PCG/Source/PCG/Private/Data/PCGSpatialData.cpp
+++ b/Engine/Plugins/Experimental/PCG/Source/PCG/Private/Data/PCGSpatialData.cpp
@@ -12,21 +12,48 @@ UPCGSpatialData::UPCGSpatialData(const FObjectInitializer& ObjectInitializer)
 	Metadata = ObjectInitializer.CreateDefaultSubobject<UPCGMetadata>(this, TEXT("Metadata"));
 }
 
-const UPCGPointData* UPCGSpatialDataWithPointCache::ToPointData(FPCGContext* Context) const
+const UPCGPointData* UPCGSpatialDataWithPointCache::ToPointData(FPCGContext* Context, const FBox& InBounds) const
 {
-	if (!CachedPointData)
+	if (InBounds.IsValid && SupportsBoundedPointData())
 	{
+		const UPCGPointData* BoundedPointData = nullptr;
 		CacheLock.Lock();
-
-		if (!CachedPointData)
+		check(CachedBoundedPointDataBoxes.Num() == CachedBoundedPointData.Num());
+		for (int CachedDataIndex = 0; CachedDataIndex < CachedBoundedPointDataBoxes.Num(); ++CachedDataIndex)
 		{
-			CachedPointData = CreatePointData(Context);
+			if (InBounds.Equals(CachedBoundedPointDataBoxes[CachedDataIndex]))
+			{
+				BoundedPointData = CachedBoundedPointData[CachedDataIndex];
+				break;
+			}
 		}
 
+		if (!BoundedPointData)
+		{
+			BoundedPointData = CreatePointData(Context, InBounds);
+			CachedBoundedPointDataBoxes.Add(InBounds);
+			CachedBoundedPointData.Add(BoundedPointData);
+		}
 		CacheLock.Unlock();
+
+		return BoundedPointData;
 	}
+	else
+	{
+		if (!CachedPointData)
+		{
+			CacheLock.Lock();
+
+			if (!CachedPointData)
+			{
+				CachedPointData = CreatePointData(Context, InBounds);
+			}
 
-	return CachedPointData;
+			CacheLock.Unlock();
+		}
+
+		return CachedPointData;
+	}
 }
 
 float UPCGSpatialData::GetDensityAtPosition(const FVector& InPosition) const
diff --git a/Engine/Plugins/Experimental/PCG/Source/PCG/Private/Graph/PCGGraphExecutor.cpp b/Engine/Plugins/Experimental/PCG/Source/PCG/Private/Graph/PCGGraphExecutor.cpp
--- a/Engine/Plugins/Experimental/PCG/Source/PCG/Private/Graph/PCGGraphExecutor.cpp
+++ b/Engine/Plugins/Experimental/PCG/Source/PCG/Private/Graph/PCGGraphExecutor.cpp
@@ -54,7 +54,7 @@ FPCGTaskId FPCGGraphExecutor::Schedule(UPCGComponent* Component, const TArray<FP
 
 FPCGTaskId FPCGGraphExecutor::Schedule(UPCGGraph* Graph, UPCGComponent* SourceComponent, FPCGElementPtr InputElement, const TArray<FPCGTaskId>& ExternalDependencies)
 {
-	FPCGTaskId ScheduledId = InvalidTaskId;
+	FPCGTaskId ScheduledId = InvalidPCGTaskId;
 
 	// Get compiled tasks from compiler
 	TArray<FPCGGraphTask> CompiledTasks = GraphCompiler->GetCompiledTasks(Graph);
@@ -573,40 +573,68 @@ bool FPCGFetchInputElement::ExecuteInternal(FPCGContext* Context) const
 	// Second: fetch the inputs provided by the component
 	UPCGComponent* Component = Context->SourceComponent;
 	check(Component);
-	if(UPCGData* PCGData = Component->GetPCGData())
+
+	check(Context->Node);
+
+	if (Context->Node->IsOutputPinConnected(PCGPinConstants::DefaultInputLabel))
 	{
-		FPCGTaggedData& TaggedData = Context->OutputData.TaggedData.Emplace_GetRef();
-		TaggedData.Data = PCGData;
-		TaggedData.Pin = PCGPinConstants::DefaultInputLabel;
+		if (UPCGData* PCGData = Component->GetPCGData())
+		{
+			FPCGTaggedData& TaggedData = Context->OutputData.TaggedData.Emplace_GetRef();
+			TaggedData.Data = PCGData;
+			TaggedData.Pin = PCGPinConstants::DefaultInputLabel;
+		}
 	}
 
-	if (UPCGData* InputPCGData = Component->GetInputPCGData())
+	if (Context->Node->IsOutputPinConnected(PCGInputOutputConstants::DefaultInputLabel))
 	{
-		FPCGTaggedData& TaggedData = Context->OutputData.TaggedData.Emplace_GetRef();
-		TaggedData.Data = InputPCGData;
-		TaggedData.Pin = PCGInputOutputConstants::DefaultInputLabel;
+		if (UPCGData* InputPCGData = Component->GetInputPCGData())
+		{
+			FPCGTaggedData& TaggedData = Context->OutputData.TaggedData.Emplace_GetRef();
+			TaggedData.Data = InputPCGData;
+			TaggedData.Pin = PCGInputOutputConstants::DefaultInputLabel;
+		}
 	}
 
-	if (UPCGData* ActorPCGData = Component->GetActorPCGData())
+	if (Context->Node->IsOutputPinConnected(PCGInputOutputConstants::DefaultActorLabel))
 	{
-		FPCGTaggedData& TaggedData = Context->OutputData.TaggedData.Emplace_GetRef();
-		TaggedData.Data = ActorPCGData;
-		TaggedData.Pin = PCGInputOutputConstants::DefaultActorLabel;
+		if (UPCGData* ActorPCGData = Component->GetActorPCGData())
+		{
+			FPCGTaggedData& TaggedData = Context->OutputData.TaggedData.Emplace_GetRef();
+			TaggedData.Data = ActorPCGData;
+			TaggedData.Pin = PCGInputOutputConstants::DefaultActorLabel;
+		}
 	}
 
-	if (UPCGData* OriginalActorPCGData = Component->GetOriginalActorPCGData())
+	if (Context->Node->IsOutputPinConnected(PCGInputOutputConstants::DefaultLandscapeLabel))
 	{
-		FPCGTaggedData& TaggedData = Context->OutputData.TaggedData.Emplace_GetRef();
-		TaggedData.Data = OriginalActorPCGData;
-		TaggedData.Pin = PCGInputOutputConstants::DefaultOriginalActorLabel;
+		if (UPCGData* LandscapePCGData = Component->GetLandscapePCGData())
+		{
+			FPCGTaggedData& TaggedData = Context->OutputData.TaggedData.Emplace_GetRef();
+			TaggedData.Data = LandscapePCGData;
+			TaggedData.Pin = PCGInputOutputConstants::DefaultLandscapeLabel;
+		}
+	}
+
+	if (Context->Node->IsOutputPinConnected(PCGInputOutputConstants::DefaultOriginalActorLabel))
+	{
+		if (UPCGData* OriginalActorPCGData = Component->GetOriginalActorPCGData())
+		{
+			FPCGTaggedData& TaggedData = Context->OutputData.TaggedData.Emplace_GetRef();
+			TaggedData.Data = OriginalActorPCGData;
+			TaggedData.Pin = PCGInputOutputConstants::DefaultOriginalActorLabel;
+		}
 	}
 
-	TArray<UPCGData*> ExclusionsPCGData = Component->GetPCGExclusionData();
-	for (UPCGData* ExclusionPCGData : ExclusionsPCGData)
+	if (Context->Node->IsOutputPinConnected(PCGInputOutputConstants::DefaultExcludedActorsLabel))
 	{
-		FPCGTaggedData& TaggedData = Context->OutputData.TaggedData.Emplace_GetRef();
-		TaggedData.Data = ExclusionPCGData;
-		TaggedData.Pin = PCGInputOutputConstants::DefaultExcludedActorsLabel;
+		TArray<UPCGData*> ExclusionsPCGData = Component->GetPCGExclusionData();
+		for (UPCGData* ExclusionPCGData : ExclusionsPCGData)
+		{
+			FPCGTaggedData& TaggedData = Context->OutputData.TaggedData.Emplace_GetRef();
+			TaggedData.Data = ExclusionPCGData;
+			TaggedData.Pin = PCGInputOutputConstants::DefaultExcludedActorsLabel;
+		}
 	}
 
 	return true;
diff --git a/Engine/Plugins/Experimental/PCG/Source/PCG/Private/Graph/PCGGraphExecutor.h b/Engine/Plugins/Experimental/PCG/Source/PCG/Private/Graph/PCGGraphExecutor.h
--- a/Engine/Plugins/Experimental/PCG/Source/PCG/Private/Graph/PCGGraphExecutor.h
+++ b/Engine/Plugins/Experimental/PCG/Source/PCG/Private/Graph/PCGGraphExecutor.h
@@ -40,7 +40,7 @@ struct FPCGGraphTask
 	const UPCGNode* Node = nullptr;
 	UPCGComponent* SourceComponent = nullptr;
 	FPCGElementPtr Element; // Added to have tasks that aren't node-bound
-	FPCGTaskId NodeId = InvalidTaskId;
+	FPCGTaskId NodeId = InvalidPCGTaskId;
 };
 
 struct FPCGGraphScheduleTask
@@ -52,7 +52,7 @@ struct FPCGGraphActiveTask
 {
 	FPCGElementPtr Element;
 	TUniquePtr<FPCGContext> Context;
-	FPCGTaskId NodeId = InvalidTaskId;
+	FPCGTaskId NodeId = InvalidPCGTaskId;
 #if WITH_EDITOR
 	bool bIsBypassed = false;
 #endif
@@ -92,7 +92,7 @@ public:
 	FPCGGraphCache& GetCache() { return GraphCache; }
 
 private:
-	void QueueReadyTasks(FPCGTaskId FinishedTaskHint = InvalidTaskId);
+	void QueueReadyTasks(FPCGTaskId FinishedTaskHint = InvalidPCGTaskId);
 	void BuildTaskInput(const FPCGGraphTask& Task, FPCGDataCollection& TaskInput);
 	void StoreResults(FPCGTaskId InTaskId, const FPCGDataCollection& InTaskOutput);
 	void ClearResults();
diff --git a/Engine/Plugins/Experimental/PCG/Source/PCG/Private/PCGComponent.cpp b/Engine/Plugins/Experimental/PCG/Source/PCG/Private/PCGComponent.cpp
--- a/Engine/Plugins/Experimental/PCG/Source/PCG/Private/PCGComponent.cpp
+++ b/Engine/Plugins/Experimental/PCG/Source/PCG/Private/PCGComponent.cpp
@@ -3,6 +3,7 @@
 #include "PCGComponent.h"
 #include "PCGGraph.h"
 #include "PCGHelpers.h"
+#include "PCGInputOutputSettings.h"
 #include "PCGVolume.h"
 #include "PCGManagedResource.h"
 #include "Data/PCGDifferenceData.h"
@@ -163,7 +164,7 @@ void UPCGComponent::SetPropertiesFromOriginal(const UPCGComponent* Original)
 	if (bIsDirty)
 	{
 		Modify();
-		DirtyGenerated(bHasDirtyInput);
+		DirtyGenerated((bHasDirtyInput ? EPCGComponentDirtyFlag::Input : EPCGComponentDirtyFlag::None) | (bHasDirtyExclusions ? EPCGComponentDirtyFlag::Exclusions : EPCGComponentDirtyFlag::None));
 	}
 #endif
 }
@@ -194,7 +195,7 @@ void UPCGComponent::Generate(bool bForce)
 	FPCGTaskId TaskId = GenerateInternal(bForce, {});
 
 #if WITH_EDITOR
-	if (TaskId != InvalidTaskId)
+	if (TaskId != InvalidPCGTaskId)
 	{
 		bIsGenerating = true;
 	}
@@ -203,11 +204,11 @@ void UPCGComponent::Generate(bool bForce)
 
 FPCGTaskId UPCGComponent::GenerateInternal(bool bForce, const TArray<FPCGTaskId>& TaskDependencies)
 {
-	FPCGTaskId TaskId = InvalidTaskId;
+	FPCGTaskId TaskId = InvalidPCGTaskId;
 
 	if (!ShouldGenerate(bForce))
 	{
-		return InvalidTaskId;
+		return InvalidPCGTaskId;
 	}
 
 #if WITH_EDITOR
@@ -579,7 +580,7 @@ void UPCGComponent::PostEditChangeProperty(FPropertyChangedEvent& PropertyChange
 	else if (PropName == GET_MEMBER_NAME_CHECKED(UPCGComponent, InputType))
 	{
 		UpdateTrackedLandscape();
-		DirtyGenerated(/*bDirtyCachedInput=*/true);
+		DirtyGenerated(EPCGComponentDirtyFlag::Input);
 		Refresh();
 	}
 	// General properties that don't affect behavior
@@ -598,7 +599,7 @@ void UPCGComponent::PostEditChangeProperty(FPropertyChangedEvent& PropertyChange
 
 		if(bHadExclusionData || bHasExcludedActors)
 		{
-			DirtyGenerated();
+			DirtyGenerated(EPCGComponentDirtyFlag::Exclusions);
 			Refresh();
 		}
 	}
@@ -639,7 +640,7 @@ void UPCGComponent::PostEditUndo()
 	SetupTrackingCallbacks();
 	RefreshTrackingData();
 	UpdateTrackedLandscape();
-	DirtyGenerated(/*bDirtyCachedInput=*/true);
+	DirtyGenerated(EPCGComponentDirtyFlag::All);
 	DirtyCacheForAllTrackedTags();
 
 	if (bGenerated)
@@ -764,7 +765,7 @@ void UPCGComponent::OnActorAdded(AActor* InActor)
 
 	if (bIsExcluded || bIsTracked)
 	{
-		DirtyGenerated();
+		DirtyGenerated(bIsExcluded ? EPCGComponentDirtyFlag::Exclusions : EPCGComponentDirtyFlag::None);
 		Refresh();
 	}
 }
@@ -776,30 +777,35 @@ void UPCGComponent::OnActorDeleted(AActor* InActor)
 
 	if (bWasExcluded || bWasTracked)
 	{
-		DirtyGenerated();
+		DirtyGenerated(bWasExcluded ? EPCGComponentDirtyFlag::Exclusions : EPCGComponentDirtyFlag::None);
 		Refresh();
 	}
 }
 
 void UPCGComponent::OnActorMoved(AActor* InActor)
 {
-	if (InActor == GetOwner() || (InActor && InActor == TrackedLandscape))
+	const bool bOwnerMoved = (InActor == GetOwner());
+	const bool bLandscapeMoved = (InActor && InActor == TrackedLandscape);
+
+	if (bOwnerMoved || bLandscapeMoved)
 	{
 		// TODO: find better metrics to dirty the inputs. 
 		// TODO: this should dirty only the actor pcg data.
 		{
 			UpdateTrackedLandscape();
-			DirtyGenerated(true);
+			DirtyGenerated((bOwnerMoved ? EPCGComponentDirtyFlag::Actor : EPCGComponentDirtyFlag::None) | (bLandscapeMoved ? EPCGComponentDirtyFlag::Landscape : EPCGComponentDirtyFlag::None));
 			Refresh();
 		}
 	}
 	else
 	{
 		bool bDirtyAndRefresh = false;
+		bool bDirtyExclusions = false;
 
 		if (UpdateExcludedActor(InActor))
 		{
 			bDirtyAndRefresh = true;
+			bDirtyExclusions = true;
 		}
 
 		if (DirtyTrackedActor(InActor))
@@ -809,7 +815,7 @@ void UPCGComponent::OnActorMoved(AActor* InActor)
 
 		if (bDirtyAndRefresh)
 		{
-			DirtyGenerated();
+			DirtyGenerated(bDirtyExclusions ? EPCGComponentDirtyFlag::Exclusions : EPCGComponentDirtyFlag::None);
 			Refresh();
 		}
 	}
@@ -820,26 +826,26 @@ void UPCGComponent::UpdateTrackedLandscape(bool bBoundsCheck)
 	TeardownLandscapeTracking();
 	TrackedLandscape = nullptr;
 
-	if (InputType == EPCGComponentInput::Landscape)
+	if (ALandscapeProxy* Landscape = Cast<ALandscapeProxy>(GetOwner()))
 	{
-		if (bBoundsCheck)
-		{
-			UPCGData* ActorData = GetActorPCGData();
-			if (const UPCGSpatialData* ActorSpatialData = Cast<const UPCGSpatialData>(ActorData))
-			{
-				TrackedLandscape = PCGHelpers::GetLandscape(GetOwner()->GetWorld(), ActorSpatialData->GetBounds());
-			}
-		}
-		else
-		{
-			TrackedLandscape = PCGHelpers::GetAnyLandscape(GetOwner()->GetWorld());
-		}
+		TrackedLandscape = Landscape;
 	}
-	else if (InputType == EPCGComponentInput::Actor)
+	else if (InputType == EPCGComponentInput::Landscape || GraphUsesLandscapePin())
 	{
-		if (ALandscapeProxy* Landscape = Cast<ALandscapeProxy>(GetOwner()))
+		if (UWorld* World = GetOwner() ? GetOwner()->GetWorld() : nullptr)
 		{
-			TrackedLandscape = Landscape;
+			if (bBoundsCheck)
+			{
+				UPCGData* ActorData = GetActorPCGData();
+				if (const UPCGSpatialData* ActorSpatialData = Cast<const UPCGSpatialData>(ActorData))
+				{
+					TrackedLandscape = PCGHelpers::GetLandscape(World, ActorSpatialData->GetBounds());
+				}
+			}
+			else
+			{
+				TrackedLandscape = PCGHelpers::GetAnyLandscape(World);
+			}
 		}
 	}
 
@@ -867,28 +873,39 @@ void UPCGComponent::OnLandscapeChanged(ALandscapeProxy* Landscape, const FLandsc
 	if (Landscape == TrackedLandscape)
 	{
 		// Check if there is an overlap in the changed components vs. the current actor data
-		bool bDirtyInput = (GetOwner() == TrackedLandscape);
+		EPCGComponentDirtyFlag DirtyFlag = EPCGComponentDirtyFlag::None;
 
-		if (!bDirtyInput && InputType == EPCGComponentInput::Landscape)
+		if (GetOwner() == TrackedLandscape)
+		{
+			DirtyFlag = EPCGComponentDirtyFlag::Actor;
+		}
+		// Note: this means that graphs that are interacting with the landscape outside their bounds might not be updated properly
+		else if (InputType == EPCGComponentInput::Landscape || GraphUsesLandscapePin())
 		{
 			UPCGData* ActorData = GetActorPCGData();
 			if (const UPCGSpatialData* ActorSpatialData = Cast<const UPCGSpatialData>(ActorData))
 			{
 				const FBox ActorBounds = ActorSpatialData->GetBounds();
+				bool bDirtyLandscape = false;
 
-				ChangeParams.ForEachComponent([&bDirtyInput, &ActorBounds](const ULandscapeComponent* LandscapeComponent)
+				ChangeParams.ForEachComponent([&bDirtyLandscape, &ActorBounds](const ULandscapeComponent* LandscapeComponent)
 				{
 					if(LandscapeComponent && ActorBounds.Intersect(LandscapeComponent->Bounds.GetBox()))
 					{
-						bDirtyInput = true;
+						bDirtyLandscape = true;
 					}
 				});
+
+				if (bDirtyLandscape)
+				{
+					DirtyFlag = EPCGComponentDirtyFlag::Landscape;
+				}
 			}
 		}
 
-		if (bDirtyInput)
+		if (DirtyFlag != EPCGComponentDirtyFlag::None)
 		{
-			DirtyGenerated(true);
+			DirtyGenerated(DirtyFlag);
 			Refresh();
 		}
 	}
@@ -941,7 +958,7 @@ void UPCGComponent::OnActorChanged(AActor* Actor, UObject* InObject, bool bActor
 		// Something has changed on the owner (including properties of this component)
 		// In the case of splines, this is where we'd get notified if some component properties (incl. spline vertices) have changed
 		// TODO: this should dirty only the actor pcg data.
-		DirtyGenerated(true);
+		DirtyGenerated(EPCGComponentDirtyFlag::Actor);
 		Refresh();
 	}
 	else if(Actor)
@@ -980,6 +997,7 @@ void UPCGComponent::OnGraphChanged(UPCGGraph* InGraph, bool bIsStructural, bool
 		SetupTrackingCallbacks();
 		RefreshTrackingData();
 		DirtyCacheForAllTrackedTags();
+		UpdateTrackedLandscape();
 
 		DirtyGenerated();
 		if (bShouldRefresh)
@@ -989,27 +1007,60 @@ void UPCGComponent::OnGraphChanged(UPCGGraph* InGraph, bool bIsStructural, bool
 	}
 }
 
-void UPCGComponent::DirtyGenerated(bool bInDirtyCachedInput)
+void UPCGComponent::DirtyGenerated(EPCGComponentDirtyFlag DirtyFlag)
 {
 	bDirtyGenerated = true;
 
-	if (bInDirtyCachedInput)
+	// Dirty data as a waterfall from basic values
+	if (!!(DirtyFlag & EPCGComponentDirtyFlag::Actor))
 	{
-		CachedInputData = nullptr;
 		CachedActorData = nullptr;
+		
+		if (Cast<ALandscapeProxy>(GetOwner()))
+		{
+			CachedLandscapeData = nullptr;
+		}
+
+		CachedInputData = nullptr;
+		CachedPCGData = nullptr;
+	}
+	
+	if (!!(DirtyFlag & EPCGComponentDirtyFlag::Landscape))
+	{
+		CachedLandscapeData = nullptr;
+		if (InputType == EPCGComponentInput::Landscape)
+		{
+			CachedInputData = nullptr;
+			CachedPCGData = nullptr;
+		}
+	}
+
+	if (!!(DirtyFlag & EPCGComponentDirtyFlag::Input))
+	{
+		CachedInputData = nullptr;
 		CachedPCGData = nullptr;
+	}
+
+	if (!!(DirtyFlag & EPCGComponentDirtyFlag::Exclusions))
+	{
 		CachedExclusionData.Reset();
+		CachedPCGData = nullptr;
+	}
+
+	if (!!(DirtyFlag & EPCGComponentDirtyFlag::Data))
+	{
+		CachedPCGData = nullptr;
 	}
 
 	// For partitioned graph, we must forward the call to the partition actor
 	// Note that we do not need to forward "normal" dirty as these will be picked up by the local PCG components
 	// However, input changes / moves of the partitioned object will not be caught
 	// It would be possible for partitioned actors to add callbacks to their original component, but that inverses the processing flow
-	if (bInDirtyCachedInput && bActivated && IsPartitioned())
+	if (DirtyFlag != EPCGComponentDirtyFlag::None && bActivated && IsPartitioned())
 	{
 		if (GetSubsystem())
 		{
-			GetSubsystem()->DirtyGraph(this, LastGeneratedBounds, bInDirtyCachedInput);
+			GetSubsystem()->DirtyGraph(this, LastGeneratedBounds, DirtyFlag);
 		}
 	}
 }
@@ -1096,6 +1147,16 @@ UPCGData* UPCGComponent::GetActorPCGData()
 	return CachedActorData;
 }
 
+UPCGData* UPCGComponent::GetLandscapePCGData()
+{
+	if (!CachedLandscapeData)
+	{
+		CachedLandscapeData = CreateLandscapePCGData();
+	}
+
+	return CachedLandscapeData;
+}
+
 UPCGData* UPCGComponent::GetOriginalActorPCGData()
 {
 	if (APCGPartitionActor* PartitionActor = Cast<APCGPartitionActor>(GetOwner()))
@@ -1105,6 +1166,10 @@ UPCGData* UPCGComponent::GetOriginalActorPCGData()
 			return OriginalComponent->GetActorPCGData();
 		}
 	}
+	else
+	{
+		return GetActorPCGData();
+	}
 
 	return nullptr;
 }
@@ -1354,6 +1419,55 @@ UPCGData* UPCGComponent::CreatePCGData()
 	return Difference ? Difference : InputData;
 }
 
+UPCGData* UPCGComponent::CreateLandscapePCGData()
+{
+	TRACE_CPUPROFILER_EVENT_SCOPE(UPCGComponent::CreateLandscapePCGData);
+	AActor* Actor = GetOwner();
+
+	if (!Actor)
+	{
+		return nullptr;
+	}
+
+	UPCGData* ActorData = GetActorPCGData();
+
+	if (ALandscapeProxy* Landscape = Cast<ALandscapeProxy>(Actor))
+	{
+		return ActorData;
+	}
+
+	const UPCGSpatialData* ActorSpatialData = Cast<const UPCGSpatialData>(ActorData);
+	ALandscapeProxy* Landscape = nullptr;
+
+	if (ActorSpatialData)
+	{
+		const FBox ActorDataBounds = ActorSpatialData->GetBounds();
+		Landscape = PCGHelpers::GetLandscape(Actor->GetWorld(), ActorDataBounds);
+	}
+	else
+	{
+		FVector Origin;
+		FVector Extent;
+		Actor->GetActorBounds(/*bOnlyCollidingComponents=*/false, Origin, Extent);
+
+		Landscape = PCGHelpers::GetLandscape(Actor->GetWorld(), FBox::BuildAABB(Origin, Extent));
+	}
+
+	if (!Landscape)
+	{
+		// No landscape found
+		return nullptr;
+	}
+
+	// TODO: we're creating separate landscape data instances here so we can do some tweaks on it (such as storing the right target actor) but this probably should change
+	UPCGLandscapeData* LandscapeData = NewObject<UPCGLandscapeData>(this);
+	LandscapeData->Initialize(Landscape, GetGridBounds(Landscape));
+	// Need to override target actor for this one, not the landscape
+	LandscapeData->TargetActor = Actor;
+
+	return LandscapeData;
+}
+
 UPCGData* UPCGComponent::CreateInputPCGData()
 {
 	TRACE_CPUPROFILER_EVENT_SCOPE(UPCGComponent::CreateInputPCGData);
@@ -1369,12 +1483,6 @@ UPCGData* UPCGComponent::CreateInputPCGData()
 	{
 		UPCGData* ActorData = GetActorPCGData();
 
-		if (ALandscapeProxy* Landscape = Cast<ALandscapeProxy>(GetOwner()))
-		{
-			// Current actor is a landscape, we don't need to do an additional projection
-			return ActorData;
-		}
-
 		const UPCGSpatialData* ActorSpatialData = Cast<const UPCGSpatialData>(ActorData);
 
 		if (!ActorSpatialData)
@@ -1383,32 +1491,18 @@ UPCGData* UPCGComponent::CreateInputPCGData()
 			return nullptr;
 		}
 
-		const FBox ActorDataBounds = ActorSpatialData->GetBounds();
-		ALandscapeProxy* Landscape = PCGHelpers::GetLandscape(Actor->GetWorld(), ActorDataBounds);
+		const UPCGSpatialData* LandscapeData = Cast<const UPCGSpatialData>(GetLandscapePCGData());
 
-		if (!Landscape)
+		if (!LandscapeData)
 		{
-			// No landscape found
 			return nullptr;
 		}
 
-		const FBox LandscapeBounds = GetGridBounds(Landscape);
-		check(LandscapeBounds.IsValid);
-
-		const FBox OverlappedBounds = LandscapeBounds.Overlap(ActorDataBounds);
-
-		if (!OverlappedBounds.IsValid)
+		if (LandscapeData == ActorSpatialData)
 		{
-			return nullptr;
+			return ActorData;
 		}
 
-		UPCGLandscapeData* LandscapeData = NewObject<UPCGLandscapeData>(this);
-		// Since we're not sharing this landscape data with any other users, we can limit the bounds earlier
-		LandscapeData->Initialize(Landscape, OverlappedBounds);
-
-		// Need to override target actor for this one, not the landscape
-		LandscapeData->TargetActor = Actor;
-
 		// Decide whether to intersect or project
 		// Currently, it makes sense to intersect only for volumes;
 		// Note that we don't currently check for a volume object but only on dimension
@@ -1650,6 +1744,11 @@ void UPCGComponent::DirtyCacheForAllTrackedTags()
 	}
 }
 
+bool UPCGComponent::GraphUsesLandscapePin() const
+{
+	return Graph && Graph->GetInputNode()->IsOutputPinConnected(PCGInputOutputConstants::DefaultLandscapeLabel);
+}
+
 #endif // WITH_EDITOR
 
 #undef LOCTEXT_NAMESPACE
\ No newline at end of file
diff --git a/Engine/Plugins/Experimental/PCG/Source/PCG/Private/PCGInputOutputSettings.cpp b/Engine/Plugins/Experimental/PCG/Source/PCG/Private/PCGInputOutputSettings.cpp
--- a/Engine/Plugins/Experimental/PCG/Source/PCG/Private/PCGInputOutputSettings.cpp
+++ b/Engine/Plugins/Experimental/PCG/Source/PCG/Private/PCGInputOutputSettings.cpp
@@ -19,6 +19,7 @@ UPCGGraphInputOutputSettings::UPCGGraphInputOutputSettings(const FObjectInitiali
 	StaticAdvancedInLabels.Add(PCGInputOutputConstants::DefaultInputLabel);
 	StaticAdvancedInLabels.Add(PCGInputOutputConstants::DefaultActorLabel);
 	StaticAdvancedInLabels.Add(PCGInputOutputConstants::DefaultOriginalActorLabel);
+	StaticAdvancedInLabels.Add(PCGInputOutputConstants::DefaultLandscapeLabel);
 	StaticAdvancedInLabels.Add(PCGInputOutputConstants::DefaultExcludedActorsLabel);
 	
 	StaticOutLabels.Add(PCGPinConstants::DefaultOutputLabel);
@@ -47,7 +48,10 @@ TArray<FPCGPinProperties> UPCGGraphInputOutputSettings::GetPinProperties() const
 	
 	if (bShowAdvancedPins)
 	{
-		Algo::Transform(StaticAdvancedLabels(), PinProperties, [DefaultPinDataType](const FName& InLabel) { return FPCGPinProperties(InLabel, DefaultPinDataType); });
+		Algo::Transform(StaticAdvancedLabels(), PinProperties, [DefaultPinDataType](const FName& InLabel)
+		{
+			return FPCGPinProperties(InLabel, (InLabel == PCGInputOutputConstants::DefaultLandscapeLabel) ? EPCGDataType::Surface: DefaultPinDataType); 
+		});
 	}
 
 	PinProperties.Append(CustomPins);
diff --git a/Engine/Plugins/Experimental/PCG/Source/PCG/Private/PCGSubgraph.cpp b/Engine/Plugins/Experimental/PCG/Source/PCG/Private/PCGSubgraph.cpp
--- a/Engine/Plugins/Experimental/PCG/Source/PCG/Private/PCGSubgraph.cpp
+++ b/Engine/Plugins/Experimental/PCG/Source/PCG/Private/PCGSubgraph.cpp
@@ -3,6 +3,7 @@
 #include "PCGSubgraph.h"
 #include "PCGGraph.h"
 #include "PCGComponent.h"
+#include "PCGSubsystem.h"
 
 void UPCGBaseSubgraphSettings::PostLoad()
 {
diff --git a/Engine/Plugins/Experimental/PCG/Source/PCG/Private/PCGSubsystem.cpp b/Engine/Plugins/Experimental/PCG/Source/PCG/Private/PCGSubsystem.cpp
--- a/Engine/Plugins/Experimental/PCG/Source/PCG/Private/PCGSubsystem.cpp
+++ b/Engine/Plugins/Experimental/PCG/Source/PCG/Private/PCGSubsystem.cpp
@@ -99,9 +99,9 @@ FPCGTaskId UPCGSubsystem::ScheduleComponent(UPCGComponent* PCGComponent, const T
 {
 	check(GraphExecutor);
 
-	FPCGTaskId ProcessTaskId = PCGComponent ? GraphExecutor->Schedule(PCGComponent, Dependencies) : InvalidTaskId;
+	FPCGTaskId ProcessTaskId = PCGComponent ? GraphExecutor->Schedule(PCGComponent, Dependencies) : InvalidPCGTaskId;
 
-	if (ProcessTaskId != InvalidTaskId)
+	if (ProcessTaskId != InvalidPCGTaskId)
 	{
 #if WITH_EDITOR
 		TWeakObjectPtr<UPCGComponent> ComponentPtr(PCGComponent);
@@ -126,7 +126,7 @@ FPCGTaskId UPCGSubsystem::ScheduleComponent(UPCGComponent* PCGComponent, const T
 		{
 			PCGComponent->OnProcessGraphAborted();
 		}
-		return InvalidTaskId;
+		return InvalidPCGTaskId;
 	}
 }
 
@@ -138,7 +138,7 @@ FPCGTaskId UPCGSubsystem::ScheduleGraph(UPCGGraph* Graph, UPCGComponent* SourceC
 	}
 	else
 	{
-		return InvalidTaskId;
+		return InvalidPCGTaskId;
 	}
 }
 
@@ -173,7 +173,7 @@ namespace PCGSubsystem
 		if (!GraphExecutor || !World)
 		{
 			UE_LOG(LogPCG, Error, TEXT("[ForEachIntersectingCell] GraphExecutor or World is null"));
-			return InvalidTaskId;
+			return InvalidPCGTaskId;
 		}
 
 		TArray<FPCGTaskId> CellTasks;
@@ -241,7 +241,7 @@ namespace PCGSubsystem
 				TArray<FPCGTaskId> PreviousTasks;
 
 				auto SetPreviousTaskIfValid = [&PreviousTasks](FPCGTaskId TaskId){
-					if (TaskId != InvalidTaskId)
+					if (TaskId != InvalidPCGTaskId)
 					{
 						PreviousTasks.Reset();
 						PreviousTasks.Add(TaskId);
@@ -273,7 +273,7 @@ namespace PCGSubsystem
 				SetPreviousTaskIfValid(ExecuteTaskId);
 
 				// Save changes; note that there's no need to save if the operation was cancelled
-				if (bSaveActors && ExecuteTaskId != InvalidTaskId)
+				if (bSaveActors && ExecuteTaskId != InvalidPCGTaskId)
 				{
 					auto SaveActorTask = [PCGActor, GraphExecutor]() {
 						GraphExecutor->AddToDirtyActors(PCGActor);
@@ -312,7 +312,7 @@ namespace PCGSubsystem
 		}
 		else
 		{
-			return InvalidTaskId;
+			return InvalidPCGTaskId;
 		}
 	}
 
@@ -361,7 +361,7 @@ FPCGTaskId UPCGSubsystem::ProcessGraph(UPCGComponent* Component, const FBox& InP
 
 				if (!LocalComponent)
 				{
-					return InvalidTaskId;
+					return InvalidPCGTaskId;
 				}
 
 				return LocalComponent->GenerateInternal(/*bForce=*/false, TaskDependencies);
@@ -369,7 +369,7 @@ FPCGTaskId UPCGSubsystem::ProcessGraph(UPCGComponent* Component, const FBox& InP
 			else
 			{
 				UE_LOG(LogPCG, Error, TEXT("[ProcessGraph] PCG Component on PCG Actor is null"));
-				return InvalidTaskId;
+				return InvalidPCGTaskId;
 			}
 		};
 
@@ -406,7 +406,7 @@ FPCGTaskId UPCGSubsystem::ProcessGraph(UPCGComponent* Component, const FBox& InP
 		}
 		else
 		{
-			return InvalidTaskId;
+			return InvalidPCGTaskId;
 		}
 	};
 
@@ -415,7 +415,7 @@ FPCGTaskId UPCGSubsystem::ProcessGraph(UPCGComponent* Component, const FBox& InP
 	const bool bLoadCell = (bGenerate && bSave); // TODO: review this
 	const bool bSaveActors = (bSave);
 
-	FPCGTaskId ProcessAllCellsTaskId = InvalidTaskId;
+	FPCGTaskId ProcessAllCellsTaskId = InvalidPCGTaskId;
 
 	if (UnionBounds.IsValid)
 	{
@@ -423,7 +423,7 @@ FPCGTaskId UPCGSubsystem::ProcessGraph(UPCGComponent* Component, const FBox& InP
 	}
 
 	// Finally, call PostProcessGraph if something happened
-	if (ProcessAllCellsTaskId != InvalidTaskId)
+	if (ProcessAllCellsTaskId != InvalidPCGTaskId)
 	{
 		auto PostProcessGraph = [ComponentPtr, InNewBounds, bGenerate]() {
 			if (UPCGComponent* Component = ComponentPtr.Get())
@@ -442,7 +442,7 @@ FPCGTaskId UPCGSubsystem::ProcessGraph(UPCGComponent* Component, const FBox& InP
 			Component->OnProcessGraphAborted();
 		}
 
-		return InvalidTaskId;
+		return InvalidPCGTaskId;
 	}
 }
 
@@ -474,11 +474,11 @@ void UPCGSubsystem::CleanupGraph(UPCGComponent* Component, const FBox& InBounds,
 	PCGSubsystem::ForEachIntersectingCell(GraphExecutor, Component->GetWorld(), InBounds, /*bCreateActor=*/false, /*bLoadCell=*/false, bSave, ScheduleTask);
 }
 
-void UPCGSubsystem::DirtyGraph(UPCGComponent* Component, const FBox& InBounds, bool bDirtyInputs)
+void UPCGSubsystem::DirtyGraph(UPCGComponent* Component, const FBox& InBounds, EPCGComponentDirtyFlag DirtyFlag)
 {
 	TWeakObjectPtr<UPCGComponent> ComponentPtr(Component);
 
-	auto ScheduleTask = [this, ComponentPtr, bDirtyInputs](APCGPartitionActor* PCGActor, const FBox& InIntersectedBounds, const TArray<FPCGTaskId>& TaskDependencies) {
+	auto ScheduleTask = [this, ComponentPtr, DirtyFlag](APCGPartitionActor* PCGActor, const FBox& InIntersectedBounds, const TArray<FPCGTaskId>& TaskDependencies) {
 		// In the specific case of the dirty, we want to bypass the execution queue, esp. since there's nothing happening here
 		// so we will run the command now, and not delay it.
 		check(ComponentPtr.Get() != nullptr && PCGActor != nullptr);
@@ -486,15 +486,15 @@ void UPCGSubsystem::DirtyGraph(UPCGComponent* Component, const FBox& InBounds, b
 
 		if (!PCGActor || !Component)
 		{
-			return InvalidTaskId;
+			return InvalidPCGTaskId;
 		}
 
 		if (UPCGComponent* LocalComponent = PCGActor->GetLocalComponent(Component))
 		{
-			LocalComponent->DirtyGenerated(bDirtyInputs);
+			LocalComponent->DirtyGenerated(DirtyFlag);
 		}
 
-		return InvalidTaskId;
+		return InvalidPCGTaskId;
 	};
 
 	// TODO: Could implement a "ForEachLoadedIntersectingCellImmediate" method to clarify this use case
diff --git a/Engine/Plugins/Experimental/PCG/Source/PCG/Public/Data/PCGLandscapeData.h b/Engine/Plugins/Experimental/PCG/Source/PCG/Public/Data/PCGLandscapeData.h
--- a/Engine/Plugins/Experimental/PCG/Source/PCG/Public/Data/PCGLandscapeData.h
+++ b/Engine/Plugins/Experimental/PCG/Source/PCG/Public/Data/PCGLandscapeData.h
@@ -49,7 +49,8 @@ public:
 	// ~End UPGCConcreteData interface
 
 	// ~Begin UPCGSpatialDataWithPointCache interface
-	virtual const UPCGPointData* CreatePointData(FPCGContext* Context) const override;
+	virtual bool SupportsBoundedPointData() const { return true; }
+	virtual const UPCGPointData* CreatePointData(FPCGContext* Context, const FBox& InBounds) const override;
 	// ~End UPCGConcreteDataWithPointCache interface
 
 	// TODO: add on property changed to clear cached data
diff --git a/Engine/Plugins/Experimental/PCG/Source/PCG/Public/Data/PCGPointData.h b/Engine/Plugins/Experimental/PCG/Source/PCG/Public/Data/PCGPointData.h
--- a/Engine/Plugins/Experimental/PCG/Source/PCG/Public/Data/PCGPointData.h
+++ b/Engine/Plugins/Experimental/PCG/Source/PCG/Public/Data/PCGPointData.h
@@ -66,7 +66,7 @@ public:
 	// ~Begin UPCGSpatialData interface
 	virtual int GetDimension() const override { return 0; }
 	virtual FBox GetBounds() const override;
-	virtual const UPCGPointData* ToPointData(FPCGContext* Context) const { return this; }
+	virtual const UPCGPointData* ToPointData(FPCGContext* Context, const FBox& InBounds = FBox(EForceInit::ForceInit)) const { return this; }
 	virtual bool SamplePoint(const FTransform& Transform, const FBox& Bounds, FPCGPoint& OutPoint, UPCGMetadata* OutMetadata) const override;
 	// ~End UPCGSpatialData interface
 
diff --git a/Engine/Plugins/Experimental/PCG/Source/PCG/Public/Data/PCGSpatialData.h b/Engine/Plugins/Experimental/PCG/Source/PCG/Public/Data/PCGSpatialData.h
--- a/Engine/Plugins/Experimental/PCG/Source/PCG/Public/Data/PCGSpatialData.h
+++ b/Engine/Plugins/Experimental/PCG/Source/PCG/Public/Data/PCGSpatialData.h
@@ -66,7 +66,7 @@ public:
 	UFUNCTION(BlueprintCallable, Category = SpatialData)
 	const UPCGPointData* ToPointDataWithContext(UPARAM(ref) FPCGContext& Context) const { return ToPointData(&Context); }
 
-	virtual const UPCGPointData* ToPointData(FPCGContext* Context) const PURE_VIRTUAL(UPCGSpatialData::ToPointData, return nullptr;);
+	virtual const UPCGPointData* ToPointData(FPCGContext* Context, const FBox& InBounds = FBox(EForceInit::ForceInit)) const PURE_VIRTUAL(UPCGSpatialData::ToPointData, return nullptr;);
 
 	/** Transform a world-space position to a world-space position in relation to the current data. (ex: projection on surface) */
 	FVector TransformPosition(const FVector& InPosition) const;
@@ -125,15 +125,23 @@ class PCG_API UPCGSpatialDataWithPointCache : public UPCGSpatialData
 
 public:
 	// ~UPCGSpatialData implementation
-	virtual const UPCGPointData* ToPointData(FPCGContext* Context) const override;
+	virtual const UPCGPointData* ToPointData(FPCGContext* Context, const FBox& InBounds = FBox(EForceInit::ForceInit)) const override;
 	// ~End UPCGSpatialData implementation
 
 protected:
+	virtual bool SupportsBoundedPointData() const { return false; }
 	virtual const UPCGPointData* CreatePointData(FPCGContext* Context) const PURE_VIRTUAL(UPCGSpatialData::CreatePointData, return nullptr;);
+	virtual const UPCGPointData* CreatePointData(FPCGContext* Context, const FBox& InBounds) const { return CreatePointData(Context); }
 
 private:
 	UPROPERTY(Transient)
 	mutable TObjectPtr<const UPCGPointData> CachedPointData;
 
+	UPROPERTY(Transient)
+	mutable TArray<FBox> CachedBoundedPointDataBoxes;
+
+	UPROPERTY(Transient)
+	mutable TArray<TObjectPtr<const UPCGPointData>> CachedBoundedPointData;
+
 	mutable FCriticalSection CacheLock;
 };
diff --git a/Engine/Plugins/Experimental/PCG/Source/PCG/Public/PCGCommon.h b/Engine/Plugins/Experimental/PCG/Source/PCG/Public/PCGCommon.h
--- a/Engine/Plugins/Experimental/PCG/Source/PCG/Public/PCGCommon.h
+++ b/Engine/Plugins/Experimental/PCG/Source/PCG/Public/PCGCommon.h
@@ -4,6 +4,9 @@
 
 #include "PCGCommon.generated.h"
 
+using FPCGTaskId = uint64;
+static const FPCGTaskId InvalidPCGTaskId = (uint64)-1;
+
 UENUM(meta = (Bitflags))
 enum class EPCGChangeType : uint8
 {
diff --git a/Engine/Plugins/Experimental/PCG/Source/PCG/Public/PCGComponent.h b/Engine/Plugins/Experimental/PCG/Source/PCG/Public/PCGComponent.h
--- a/Engine/Plugins/Experimental/PCG/Source/PCG/Public/PCGComponent.h
+++ b/Engine/Plugins/Experimental/PCG/Source/PCG/Public/PCGComponent.h
@@ -5,7 +5,6 @@
 #include "CoreMinimal.h"
 #include "Components/ActorComponent.h"
 #include "PCGSettings.h"
-#include "PCGSubsystem.h"
 
 #include "PCGComponent.generated.h"
 
@@ -13,6 +12,7 @@ class UPCGComponent;
 class UPCGGraph;
 class UPCGManagedResource;
 class UPCGData;
+class UPCGSubsystem;
 class ALandscapeProxy;
 class FLandscapeProxyComponentDataChangedParams;
 
@@ -31,6 +31,19 @@ enum class EPCGComponentInput : uint8
 	EPCGComponentInput_MAX
 };
 
+UENUM(meta = (Bitflags))
+enum class EPCGComponentDirtyFlag : uint8
+{
+	None = 0,
+	Actor = 1 << 0,
+	Landscape = 1 << 1,
+	Input = 1 << 2,
+	Exclusions = 1 << 3,
+	Data = 1 << 4,
+	All = Actor | Landscape | Input | Exclusions | Data
+};
+ENUM_CLASS_FLAGS(EPCGComponentDirtyFlag);
+
 UCLASS(BlueprintType, ClassGroup = (Procedural), meta = (BlueprintSpawnableComponent))
 class PCG_API UPCGComponent : public UActorComponent
 {
@@ -54,6 +67,7 @@ public:
 	UPCGData* GetPCGData();
 	UPCGData* GetInputPCGData();
 	UPCGData* GetActorPCGData();
+	UPCGData* GetLandscapePCGData();
 	UPCGData* GetOriginalActorPCGData();
 	TArray<UPCGData*> GetPCGExclusionData();
 
@@ -115,7 +129,7 @@ public:
 
 #if WITH_EDITOR
 	void Refresh();
-	void DirtyGenerated(bool bInDirtyCachedInput = false);
+	void DirtyGenerated(EPCGComponentDirtyFlag DataToDirtyFlag = EPCGComponentDirtyFlag::None);
 
 	/** Reset last generated bounds to force PCGPartitionActor creation on next refresh */
 	void ResetLastGeneratedBounds();
@@ -140,6 +154,7 @@ private:
 	UPCGData* CreateInputPCGData();
 	UPCGData* CreateActorPCGData();
 	UPCGData* CreateActorPCGData(AActor* Actor);
+	UPCGData* CreateLandscapePCGData();
 	void UpdatePCGExclusionData();
 
 	bool ShouldGenerate(bool bForce = false) const;
@@ -185,6 +200,8 @@ private:
 	bool DirtyTrackedActor(AActor* InActor);
 	void DirtyCacheFromTag(const FName& InTag);
 	void DirtyCacheForAllTrackedTags();
+
+	bool GraphUsesLandscapePin() const;
 #endif
 
 	FBox GetGridBounds() const;
@@ -199,6 +216,9 @@ private:
 	UPROPERTY(Transient, NonPIEDuplicateTransient)
 	UPCGData* CachedActorData = nullptr;
 
+	UPROPERTY(Transient, NonPIEDuplicateTransient)
+	UPCGData* CachedLandscapeData = nullptr;
+
 	UPROPERTY(Transient, NonPIEDuplicateTransient)
 	TMap<AActor*, UPCGData*> CachedExclusionData;
 
diff --git a/Engine/Plugins/Experimental/PCG/Source/PCG/Public/PCGContext.h b/Engine/Plugins/Experimental/PCG/Source/PCG/Public/PCGContext.h
--- a/Engine/Plugins/Experimental/PCG/Source/PCG/Public/PCGContext.h
+++ b/Engine/Plugins/Experimental/PCG/Source/PCG/Public/PCGContext.h
@@ -5,7 +5,6 @@
 #include "Templates/SharedPointer.h"
 #include "PCGData.h"
 #include "PCGNode.h"
-#include "PCGSubsystem.h"
 
 #include "PCGContext.generated.h"
 
@@ -44,7 +43,7 @@ struct PCG_API FPCGContext
 	// TODO: add RNG source
 	// TODO: replace this by a better identification mechanism
 	const UPCGNode* Node = nullptr;
-	FPCGTaskId TaskId = InvalidTaskId;
+	FPCGTaskId TaskId = InvalidPCGTaskId;
 	bool bIsPaused = false;
 
 #if WITH_EDITOR
diff --git a/Engine/Plugins/Experimental/PCG/Source/PCG/Public/PCGInputOutputSettings.h b/Engine/Plugins/Experimental/PCG/Source/PCG/Public/PCGInputOutputSettings.h
--- a/Engine/Plugins/Experimental/PCG/Source/PCG/Public/PCGInputOutputSettings.h
+++ b/Engine/Plugins/Experimental/PCG/Source/PCG/Public/PCGInputOutputSettings.h
@@ -12,6 +12,7 @@ namespace PCGInputOutputConstants
 	const FName DefaultInputLabel = TEXT("Input");
 	const FName DefaultActorLabel = TEXT("Actor");
 	const FName DefaultOriginalActorLabel = TEXT("OriginalActor");
+	const FName DefaultLandscapeLabel = TEXT("Landscape");
 	const FName DefaultExcludedActorsLabel = TEXT("ExcludedActors");
 }
 
diff --git a/Engine/Plugins/Experimental/PCG/Source/PCG/Public/PCGSubgraph.h b/Engine/Plugins/Experimental/PCG/Source/PCG/Public/PCGSubgraph.h
--- a/Engine/Plugins/Experimental/PCG/Source/PCG/Public/PCGSubgraph.h
+++ b/Engine/Plugins/Experimental/PCG/Source/PCG/Public/PCGSubgraph.h
@@ -100,7 +100,7 @@ public:
 
 struct PCG_API FPCGSubgraphContext : public FPCGContext
 {
-	FPCGTaskId SubgraphTaskId = InvalidTaskId;
+	FPCGTaskId SubgraphTaskId = InvalidPCGTaskId;
 	bool bScheduledSubgraph = false;
 };
 
diff --git a/Engine/Plugins/Experimental/PCG/Source/PCG/Public/PCGSubsystem.h b/Engine/Plugins/Experimental/PCG/Source/PCG/Public/PCGSubsystem.h
--- a/Engine/Plugins/Experimental/PCG/Source/PCG/Public/PCGSubsystem.h
+++ b/Engine/Plugins/Experimental/PCG/Source/PCG/Public/PCGSubsystem.h
@@ -4,18 +4,17 @@
 
 #include "Subsystems/WorldSubsystem.h"
 
+#include "PCGCommon.h"
+#include "PCGComponent.h"
+
 #include "PCGSubsystem.generated.h"
 
-class UPCGComponent;
 class FPCGGraphExecutor;
 class APCGPartitionActor;
 class APCGWorldActor;
 class UPCGGraph;
 struct FPCGDataCollection;
 
-using FPCGTaskId = uint64;
-static const FPCGTaskId InvalidTaskId = (uint64)-1;
-
 class IPCGElement;
 typedef TSharedPtr<IPCGElement, ESPMode::ThreadSafe> FPCGElementPtr;
 
@@ -71,7 +70,7 @@ public:
 	void CleanupGraph(UPCGComponent* Component, const FBox& InBounds, bool bRemoveComponents, bool bSave);
 
 	/** Immediately dirties the partition actors in the given bounds */
-	void DirtyGraph(UPCGComponent* Component, const FBox& InBounds, bool bDirtyInputs);
+	void DirtyGraph(UPCGComponent* Component, const FBox& InBounds, EPCGComponentDirtyFlag DirtyFlag);
 
 	/** Partition actors methods */
 	void CleanupPartitionActors(const FBox& InBounds);
