Override pure virtual
#rnx

diff --git a/Engine/Plugins/Experimental/PCG/Source/PCG/Public/Data/PCGLandscapeData.h b/Engine/Plugins/Experimental/PCG/Source/PCG/Public/Data/PCGLandscapeData.h
--- a/Engine/Plugins/Experimental/PCG/Source/PCG/Public/Data/PCGLandscapeData.h
+++ b/Engine/Plugins/Experimental/PCG/Source/PCG/Public/Data/PCGLandscapeData.h
@@ -50,6 +50,7 @@ public:
 
 	// ~Begin UPCGSpatialDataWithPointCache interface
 	virtual bool SupportsBoundedPointData() const { return true; }
+	virtual const UPCGPointData* CreatePointData(FPCGContext* Context) const override { return CreatePointData(Context, FBox(EForceInit::ForceInit)); }
 	virtual const UPCGPointData* CreatePointData(FPCGContext* Context, const FBox& InBounds) const override;
 	// ~End UPCGConcreteDataWithPointCache interface
 
