Can mark PCG managed resources unused to re-use some of them

diff --git a/Engine/Plugins/Experimental/PCG/Source/PCG/Private/Helpers/PCGActorHelpers.cpp b/Engine/Plugins/Experimental/PCG/Source/PCG/Private/Helpers/PCGActorHelpers.cpp
--- a/Engine/Plugins/Experimental/PCG/Source/PCG/Private/Helpers/PCGActorHelpers.cpp
+++ b/Engine/Plugins/Experimental/PCG/Source/PCG/Private/Helpers/PCGActorHelpers.cpp
@@ -37,17 +37,19 @@ UInstancedStaticMeshComponent* UPCGActorHelpers::GetOrCreateISMC(AActor* InTarge
 		return nullptr;
 	}
 
-	TArray<UInstancedStaticMeshComponent*> ISMCs;
-	InSourceComponent->ForEachManagedResource([&ISMCs](UPCGManagedResource* InResource)
+	TArray<UPCGManagedISMComponent*> MISMCs;
+	InSourceComponent->ForEachManagedResource([&MISMCs](UPCGManagedResource* InResource)
 	{
 		if (UPCGManagedISMComponent* Resource = Cast<UPCGManagedISMComponent>(InResource))
 		{
-			ISMCs.Add(Resource->GetComponent());
+			MISMCs.Add(Resource);
 		}
 	});
 
-	for (UInstancedStaticMeshComponent* ISMC : ISMCs)
+	for (UPCGManagedISMComponent* MISMC : MISMCs)
 	{
+		UInstancedStaticMeshComponent* ISMC = MISMC->GetComponent();
+
 		if (ISMC && 
 			ISMC->GetStaticMesh() == InMesh &&
 			ISMC->ComponentTags.Contains(InSourceComponent->GetFName()))
@@ -85,6 +87,8 @@ UInstancedStaticMeshComponent* UPCGActorHelpers::GetOrCreateISMC(AActor* InTarge
 
 			if (bMaterialsMatched)
 			{
+				// In this case, we make sure to mark the managed resource as used
+				MISMC->MarkAsUsed();
 				return ISMC;
 			}
 		}
diff --git a/Engine/Plugins/Experimental/PCG/Source/PCG/Private/PCGComponent.cpp b/Engine/Plugins/Experimental/PCG/Source/PCG/Private/PCGComponent.cpp
--- a/Engine/Plugins/Experimental/PCG/Source/PCG/Private/PCGComponent.cpp
+++ b/Engine/Plugins/Experimental/PCG/Source/PCG/Private/PCGComponent.cpp
@@ -385,15 +385,21 @@ void UPCGComponent::CleanupInternal(bool bHardCleanup, TSet<TSoftObjectPtr<AActo
 
 void UPCGComponent::CleanupUnusedManagedResources()
 {
-	FScopeLock ResourcesLock(&GeneratedResourcesLock);
-	for (int32 ResourceIndex = GeneratedResources.Num() - 1; ResourceIndex >= 0; --ResourceIndex)
+	TSet<TSoftObjectPtr<AActor>> ActorsToDelete;
+
 	{
-		check(GeneratedResources[ResourceIndex]);
-		if (GeneratedResources[ResourceIndex]->ReleaseIfUnused())
+		FScopeLock ResourcesLock(&GeneratedResourcesLock);
+		for (int32 ResourceIndex = GeneratedResources.Num() - 1; ResourceIndex >= 0; --ResourceIndex)
 		{
-			GeneratedResources.RemoveAtSwap(ResourceIndex);
+			check(GeneratedResources[ResourceIndex]);
+			if (GeneratedResources[ResourceIndex]->ReleaseIfUnused(ActorsToDelete))
+			{
+				GeneratedResources.RemoveAtSwap(ResourceIndex);
+			}
 		}
 	}
+
+	UPCGActorHelpers::DeleteActors(GetWorld(), ActorsToDelete.Array());
 }
 
 void UPCGComponent::BeginPlay()
diff --git a/Engine/Plugins/Experimental/PCG/Source/PCG/Private/PCGManagedResource.cpp b/Engine/Plugins/Experimental/PCG/Source/PCG/Private/PCGManagedResource.cpp
--- a/Engine/Plugins/Experimental/PCG/Source/PCG/Private/PCGManagedResource.cpp
+++ b/Engine/Plugins/Experimental/PCG/Source/PCG/Private/PCGManagedResource.cpp
@@ -5,6 +5,29 @@
 
 #include "Components/InstancedStaticMeshComponent.h"
 
+// By default, if it is not a hard release, we mark the resource unused.
+bool UPCGManagedResource::Release(bool bHardRelease, TSet<TSoftObjectPtr<AActor>>& /*OutActorsToDelete*/)
+{
+	if (!bHardRelease)
+	{
+		bIsMarkedUnused = true;
+		return false;
+	}
+
+	return true;
+}
+
+bool UPCGManagedResource::ReleaseIfUnused(TSet<TSoftObjectPtr<AActor>>& OutActorsToDelete)
+{
+	if (bIsMarkedUnused)
+	{
+		Release(true, OutActorsToDelete);
+		return true;
+	}
+
+	return false;
+}
+
 void UPCGManagedActors::PostEditImport()
 {
 	// In this case, the managed actors won't be copied along the actor/component,
@@ -15,6 +38,11 @@ void UPCGManagedActors::PostEditImport()
 
 bool UPCGManagedActors::Release(bool bHardRelease, TSet<TSoftObjectPtr<AActor>>& OutActorsToDelete)
 {
+	if (!Super::Release(bHardRelease, OutActorsToDelete))
+	{
+		return false;
+	}
+
 	OutActorsToDelete.Append(GeneratedActors);
 
 	// Cleanup recursively
@@ -39,9 +67,9 @@ bool UPCGManagedActors::Release(bool bHardRelease, TSet<TSoftObjectPtr<AActor>>&
 	return true;
 }
 
-bool UPCGManagedActors::ReleaseIfUnused()
+bool UPCGManagedActors::ReleaseIfUnused(TSet<TSoftObjectPtr<AActor>>& OutActorsToDelete)
 {
-	return GeneratedActors.IsEmpty();
+	return Super::ReleaseIfUnused(OutActorsToDelete) || GeneratedActors.IsEmpty();
 }
 
 void UPCGManagedComponent::PostEditImport()
@@ -82,7 +110,7 @@ void UPCGManagedComponent::PostEditImport()
 	}
 }
 
-bool UPCGManagedComponent::Release(bool bHardRelease, TSet<TSoftObjectPtr<AActor>>& OutActorsToDelete)
+bool UPCGManagedComponent::Release(bool bHardRelease, TSet<TSoftObjectPtr<AActor>>& /*OutActorsToDelete*/)
 {
 	const bool bSupportsComponentReset = SupportsComponentReset();
 	const bool bDeleteComponent = bHardRelease || !bSupportsComponentReset;
@@ -95,21 +123,36 @@ bool UPCGManagedComponent::Release(bool bHardRelease, TSet<TSoftObjectPtr<AActor
 		}
 		else
 		{
-			ResetComponent();
+			// We can only mark it unused if we can reset the component.
+			bIsMarkedUnused = true;
 		}
 	}
 
 	return bDeleteComponent;
 }
 
-bool UPCGManagedComponent::ReleaseIfUnused()
+bool UPCGManagedComponent::ReleaseIfUnused(TSet<TSoftObjectPtr<AActor>>& OutActorsToDelete)
 {
-	return !GeneratedComponent.IsValid();
+	return Super::ReleaseIfUnused(OutActorsToDelete) || !GeneratedComponent.IsValid();
+}
+
+void UPCGManagedComponent::MarkAsUsed()
+{
+	if (!bIsMarkedUnused)
+	{
+		return;
+	}
+
+	// Can't reuse a resource if we can't reset it. Make sure we never take this path in this case.
+	check(SupportsComponentReset());
+
+	ResetComponent();
+	bIsMarkedUnused = false;
 }
 
-bool UPCGManagedISMComponent::ReleaseIfUnused()
+bool UPCGManagedISMComponent::ReleaseIfUnused(TSet<TSoftObjectPtr<AActor>>& OutActorsToDelete)
 {
-	if (Super::ReleaseIfUnused() || !GetComponent())
+	if (Super::ReleaseIfUnused(OutActorsToDelete) || !GetComponent())
 	{
 		return true;
 	}
diff --git a/Engine/Plugins/Experimental/PCG/Source/PCG/Public/PCGManagedResource.h b/Engine/Plugins/Experimental/PCG/Source/PCG/Public/PCGManagedResource.h
--- a/Engine/Plugins/Experimental/PCG/Source/PCG/Public/PCGManagedResource.h
+++ b/Engine/Plugins/Experimental/PCG/Source/PCG/Public/PCGManagedResource.h
@@ -12,19 +12,25 @@ class UInstancedStaticMeshComponent;
 /** 
 * This class is used to hold resources and their mechanism to delete them on demand.
 * In order to allow for some reuse (e.g. components), the Release call supports a "soft"
-* release that should empty any data but might retain some data.
+* release by marking them unused in order to be potentially re-used down the line.
 * At the end of the generate, a call to ReleaseIfUnused will serve to finally cleanup
-* what is not needed.
+* what is not needed anymore.
 */
 UCLASS(BlueprintType)
 class PCG_API UPCGManagedResource : public UObject
 {
 	GENERATED_BODY()
 public:
-	/** Releases/Resets the resource depending on the bHardRelease flag. Returns true if resource can be removed from the PCG component */
-	virtual bool Release(bool bHardRelease, TSet<TSoftObjectPtr<AActor>>& OutActorsToDelete) { return true; }
+	/** Releases/Mark Unsued the resource depending on the bHardRelease flag. Returns true if resource can be removed from the PCG component */
+	virtual bool Release(bool bHardRelease, TSet<TSoftObjectPtr<AActor>>& OutActorsToDelete);
 	/** Releases resource if empty or unused. Returns true if the resource can be removed from the PCG component */
-	virtual bool ReleaseIfUnused() { return true; }
+	virtual bool ReleaseIfUnused(TSet<TSoftObjectPtr<AActor>>& OutActorsToDelete);
+
+	virtual void MarkAsUsed() { bIsMarkedUnused = false; }
+
+protected:
+	UPROPERTY(Transient, VisibleAnywhere, Category = GeneratedData)
+	bool bIsMarkedUnused = false;
 };
 
 UCLASS(BlueprintType)
@@ -39,7 +45,7 @@ public:
 
 	//~Begin UPCGManagedResource interface
 	virtual bool Release(bool bHardRelease, TSet<TSoftObjectPtr<AActor>>& OutActorsToDelete) override;
-	virtual bool ReleaseIfUnused() override;
+	virtual bool ReleaseIfUnused(TSet<TSoftObjectPtr<AActor>>& OutActorsToDelete) override;
 	//~End UPCGManagedResource interface
 
 	UPROPERTY(BlueprintReadWrite, EditAnywhere, Category = GeneratedData)
@@ -58,11 +64,12 @@ public:
 
 	//~Begin UPCGManagedResource interface
 	virtual bool Release(bool bHardRelease, TSet<TSoftObjectPtr<AActor>>& OutActorsToDelete) override;
-	virtual bool ReleaseIfUnused() override;
+	virtual bool ReleaseIfUnused(TSet<TSoftObjectPtr<AActor>>& OutActorsToDelete) override;
 	//~End UPCGManagedResource interface
 
 	virtual void ResetComponent() { check(0); }
 	virtual bool SupportsComponentReset() const { return false; }
+	virtual void MarkAsUsed() override;
 
 	UPROPERTY(BlueprintReadWrite, EditAnywhere, Category = GeneratedData)
 	TSoftObjectPtr<UActorComponent> GeneratedComponent;
@@ -75,7 +82,7 @@ class PCG_API UPCGManagedISMComponent : public UPCGManagedComponent
 
 public:
 	//~Begin UPCGManagedResource interface
-	virtual bool ReleaseIfUnused() override;
+	virtual bool ReleaseIfUnused(TSet<TSoftObjectPtr<AActor>>& OutActorsToDelete) override;
 	//~End UPCGManagedResource interface
 
 	//~Begin UPCGManagedComponents interface
