PCG Determinism Framework Iteration #2 - Includes:
* Collapse Determinism Test
* Difference Determinism Test

diff --git a/Engine/Plugins/Experimental/PCG/Source/PCG/Private/Elements/PCGDifferenceElement.cpp b/Engine/Plugins/Experimental/PCG/Source/PCG/Private/Elements/PCGDifferenceElement.cpp
--- a/Engine/Plugins/Experimental/PCG/Source/PCG/Private/Elements/PCGDifferenceElement.cpp
+++ b/Engine/Plugins/Experimental/PCG/Source/PCG/Private/Elements/PCGDifferenceElement.cpp
@@ -7,12 +7,6 @@
 #include "Data/PCGUnionData.h"
 #include "Helpers/PCGSettingsHelpers.h"
 
-namespace PCGDifferenceConstants
-{
-	const FName SourceLabel = TEXT("Source");
-	const FName DifferencesLabel = TEXT("Differences");
-}
-
 TArray<FPCGPinProperties> UPCGDifferenceSettings::InputPinProperties() const
 {
 	TArray<FPCGPinProperties> PinProperties;
diff --git a/Engine/Plugins/Experimental/PCG/Source/PCG/Private/Tests/Data/PCGDifferenceDataTest.cpp b/Engine/Plugins/Experimental/PCG/Source/PCG/Private/Tests/Data/PCGDifferenceDataTest.cpp
--- a/Engine/Plugins/Experimental/PCG/Source/PCG/Private/Tests/Data/PCGDifferenceDataTest.cpp
+++ b/Engine/Plugins/Experimental/PCG/Source/PCG/Private/Tests/Data/PCGDifferenceDataTest.cpp
@@ -5,9 +5,9 @@
 #include "Data/PCGPointData.h"
 #include "Data/PCGVolumeData.h"
 
-IMPLEMENT_CUSTOM_SIMPLE_AUTOMATION_TEST(FPCGDifferenceDataBasicTest, FPCGTestBaseClass, "pcg.tests.Difference.Basic", PCGTestsCommon::TestFlags)
+IMPLEMENT_CUSTOM_SIMPLE_AUTOMATION_TEST(FPCGDifferenceDataTest, FPCGTestBaseClass, "pcg.tests.Difference.Data", PCGTestsCommon::TestFlags)
 
-bool FPCGDifferenceDataBasicTest::RunTest(const FString& Parameters)
+bool FPCGDifferenceDataTest::RunTest(const FString& Parameters)
 {
 	const UPCGVolumeData* FirstVolume = PCGTestsCommon::CreateVolumeData(FBox::BuildAABB(FVector::OneVector * 250, FVector::OneVector * 500));
 
diff --git a/Engine/Plugins/Experimental/PCG/Source/PCG/Private/Tests/Data/PCGIntersectionDataTest.cpp b/Engine/Plugins/Experimental/PCG/Source/PCG/Private/Tests/Data/PCGIntersectionDataTest.cpp
--- a/Engine/Plugins/Experimental/PCG/Source/PCG/Private/Tests/Data/PCGIntersectionDataTest.cpp
+++ b/Engine/Plugins/Experimental/PCG/Source/PCG/Private/Tests/Data/PCGIntersectionDataTest.cpp
@@ -5,9 +5,9 @@
 #include "Data/PCGPointData.h"
 #include "Data/PCGVolumeData.h"
 
-IMPLEMENT_CUSTOM_SIMPLE_AUTOMATION_TEST(FPCGIntersectionDataBasicTest, FPCGTestBaseClass, "pcg.tests.Intersection.Basic", PCGTestsCommon::TestFlags)
+IMPLEMENT_CUSTOM_SIMPLE_AUTOMATION_TEST(FPCGIntersectionDataTest, FPCGTestBaseClass, "pcg.tests.Intersection.Data", PCGTestsCommon::TestFlags)
 
-bool FPCGIntersectionDataBasicTest::RunTest(const FString& Parameters)
+bool FPCGIntersectionDataTest::RunTest(const FString& Parameters)
 {
 	UPCGPointData* InsidePoint = PCGTestsCommon::CreatePointData();
 	check(InsidePoint->GetPoints().Num() == 1);
diff --git a/Engine/Plugins/Experimental/PCG/Source/PCG/Private/Tests/Determinism/PCGCollapseDeterminismTest.cpp b/Engine/Plugins/Experimental/PCG/Source/PCG/Private/Tests/Determinism/PCGCollapseDeterminismTest.cpp
new file mode 100644
--- /dev/null
+++ b/Engine/Plugins/Experimental/PCG/Source/PCG/Private/Tests/Determinism/PCGCollapseDeterminismTest.cpp
@@ -0,0 +1,136 @@
+// Copyright Epic Games, Inc. All Rights Reserved.
+
+#if WITH_EDITOR
+
+#include "Elements/PCGCollapseElement.h"
+#include "Tests/Determinism/PCGDeterminismTestsCommon.h"
+
+IMPLEMENT_CUSTOM_SIMPLE_AUTOMATION_TEST(FPCGCollapseDeterminismSingleSameDataTest, FPCGTestBaseClass, "pcg.tests.Collapse.Determinism.SingleSameData", PCGTestsCommon::TestFlags)
+IMPLEMENT_CUSTOM_SIMPLE_AUTOMATION_TEST(FPCGCollapseDeterminismSingleIdenticalDataTest, FPCGTestBaseClass, "pcg.tests.Collapse.Determinism.SingleIdenticalData", PCGTestsCommon::TestFlags)
+IMPLEMENT_CUSTOM_SIMPLE_AUTOMATION_TEST(FPCGCollapseDeterminismMultipleSameDataTest, FPCGTestBaseClass, "pcg.tests.Collapse.Determinism.MultipleSameData", PCGTestsCommon::TestFlags)
+IMPLEMENT_CUSTOM_SIMPLE_AUTOMATION_TEST(FPCGCollapseDeterminismMultipleIdenticalDataTest, FPCGTestBaseClass, "pcg.tests.Collapse.Determinism.MultipleIdenticalData", PCGTestsCommon::TestFlags)
+IMPLEMENT_CUSTOM_SIMPLE_AUTOMATION_TEST(FPCGCollapseDeterminismOrderIndependenceTest, FPCGTestBaseClass, "pcg.tests.Collapse.Determinism.OrderIndependence", PCGTestsCommon::TestFlags)
+
+namespace
+{
+	constexpr static int32 NumMultipleSpatialSetsToAdd = 3;
+	constexpr static int32 NumPointsToAddDuringMultipleCases = 20;
+	constexpr static int32 NumTestableSpatialDataTypes = 5;
+	constexpr static EPCGDataType TestableSpatialDataTypes[NumTestableSpatialDataTypes] =
+	{
+		EPCGDataType::Point,
+		EPCGDataType::Volume,
+		EPCGDataType::Surface,
+		EPCGDataType::PolyLine,
+		EPCGDataType::Primitive
+	};
+
+	void CreateNewRandomizedSpatialDataSet(TArray<PCGDeterminismTests::FTestData>& TestDataSet,
+		int32 NumSpatialSetsToAdd = 1,
+		int32 NumPointsToCreatePerSet = PCGDeterminismTests::DefaultNumPointsToGenerate,
+		int32 NumPolyLinePointsToCreatePerSet = PCGDeterminismTests::DefaultNumPolyLinePointsToGenerate)
+	{
+		check(NumSpatialSetsToAdd > 0 && NumPointsToCreatePerSet > 0 && NumPolyLinePointsToCreatePerSet > 0);
+
+		TestDataSet.Empty();
+		TestDataSet.Init(PCGDeterminismTests::FTestData(PCGDeterminismTests::DefaultSeed), NumTestableSpatialDataTypes);
+
+		for (PCGDeterminismTests::FTestData& Data : TestDataSet)
+		{
+			PCGDeterminismTests::GenerateSettings<UPCGCollapseSettings>(Data);
+		}
+
+		for (int32 I = 0; I < NumSpatialSetsToAdd; ++I)
+		{
+			PCGDeterminismTests::AddRandomizedMultiplePointInputData(TestDataSet[0], NumPointsToCreatePerSet);
+			PCGDeterminismTests::AddRandomizedVolumeInputData(TestDataSet[1]);
+			PCGDeterminismTests::AddRandomizedSurfaceInputData(TestDataSet[2]);
+			PCGDeterminismTests::AddRandomizedPolyLineInputData(TestDataSet[3], NumPolyLinePointsToCreatePerSet);
+			PCGDeterminismTests::AddRandomizedPrimitiveInputData(TestDataSet[4]);
+		}
+	}
+}
+
+bool FPCGCollapseDeterminismSingleSameDataTest::RunTest(const FString& Parameters)
+{
+	// Test single same data
+	TArray<PCGDeterminismTests::FTestData> TestDataSet;
+	CreateNewRandomizedSpatialDataSet(TestDataSet);
+
+	for (int32 I = 0; I < TestDataSet.Num(); ++I)
+	{
+		bool bExecutionIsDeterministic = PCGDeterminismTests::ExecutionIsDeterministicSameData(TestDataSet[I]);
+		TestTrue(FString::Printf(TEXT("[%s] Same single input and settings, same output"), *UEnum::GetValueAsString(TestableSpatialDataTypes[I])), bExecutionIsDeterministic);
+	}
+
+	return true;
+}
+
+bool FPCGCollapseDeterminismSingleIdenticalDataTest::RunTest(const FString& Parameters)
+{
+	// Test single identical data
+	TArray<PCGDeterminismTests::FTestData> FirstTestDataSet;
+	TArray<PCGDeterminismTests::FTestData> SecondTestDataSet;
+	CreateNewRandomizedSpatialDataSet(FirstTestDataSet);
+	CreateNewRandomizedSpatialDataSet(SecondTestDataSet);
+
+	for (int32 I = 0; I < FirstTestDataSet.Num(); ++I)
+	{
+		bool bExecutionIsDeterministic = PCGDeterminismTests::ExecutionIsDeterministic(FirstTestDataSet[I], SecondTestDataSet[I]);
+		TestTrue(FString::Printf(TEXT("[%s] Identical single input and settings, same output"), *UEnum::GetValueAsString(TestableSpatialDataTypes[I])), bExecutionIsDeterministic);
+	}
+
+	return true;
+}
+
+bool FPCGCollapseDeterminismMultipleSameDataTest::RunTest(const FString& Parameters)
+{
+	// Test multiple same data
+	TArray<PCGDeterminismTests::FTestData> TestDataSet;
+	CreateNewRandomizedSpatialDataSet(TestDataSet, NumMultipleSpatialSetsToAdd, NumPointsToAddDuringMultipleCases);
+
+	for (int32 I = 0; I < TestDataSet.Num(); ++I)
+	{
+		bool bExecutionIsDeterministic = PCGDeterminismTests::ExecutionIsDeterministicSameData(TestDataSet[I]);
+		TestTrue(FString::Printf(TEXT("[%s] Same multiple input, same output"), *UEnum::GetValueAsString(TestableSpatialDataTypes[I])), bExecutionIsDeterministic);
+	}
+
+	return true;
+}
+
+bool FPCGCollapseDeterminismMultipleIdenticalDataTest::RunTest(const FString& Parameters)
+{
+	// Test multiple identical data
+	TArray<PCGDeterminismTests::FTestData> FirstTestDataSet;
+	TArray<PCGDeterminismTests::FTestData> SecondTestDataSet;
+	CreateNewRandomizedSpatialDataSet(FirstTestDataSet, NumMultipleSpatialSetsToAdd);
+	CreateNewRandomizedSpatialDataSet(SecondTestDataSet, NumMultipleSpatialSetsToAdd);
+
+	for (int32 I = 0; I < FirstTestDataSet.Num(); ++I)
+	{
+		bool bExecutionIsDeterministic = PCGDeterminismTests::ExecutionIsDeterministic(FirstTestDataSet[I], SecondTestDataSet[I]);
+		TestTrue(FString::Printf(TEXT("[%s] Identical multiple input, same output"), *UEnum::GetValueAsString(TestableSpatialDataTypes[I])), bExecutionIsDeterministic);
+	}
+
+	return true;
+}
+
+bool FPCGCollapseDeterminismOrderIndependenceTest::RunTest(const FString& Parameters)
+{
+	// Test shuffled input data
+	TArray<PCGDeterminismTests::FTestData> FirstTestDataSet;
+	TArray<PCGDeterminismTests::FTestData> SecondTestDataSet;
+	CreateNewRandomizedSpatialDataSet(FirstTestDataSet, NumMultipleSpatialSetsToAdd);
+	CreateNewRandomizedSpatialDataSet(SecondTestDataSet, NumMultipleSpatialSetsToAdd);
+
+	for (int32 I = 0; I < FirstTestDataSet.Num(); ++I)
+	{
+		PCGDeterminismTests::ShuffleInputOrder(SecondTestDataSet[I]);
+		bool bExecutionIsDeterministic = PCGDeterminismTests::ExecutionIsDeterministic(FirstTestDataSet[I], SecondTestDataSet[I]);
+		TestTrue(FString::Printf(TEXT("[%s] Shuffled input order, same output"), *UEnum::GetValueAsString(TestableSpatialDataTypes[I])), bExecutionIsDeterministic);
+	}
+
+	return true;
+}
+
+#endif
\ No newline at end of file
diff --git a/Engine/Plugins/Experimental/PCG/Source/PCG/Private/Tests/Determinism/PCGDeterminismTestsCommon.cpp b/Engine/Plugins/Experimental/PCG/Source/PCG/Private/Tests/Determinism/PCGDeterminismTestsCommon.cpp
--- a/Engine/Plugins/Experimental/PCG/Source/PCG/Private/Tests/Determinism/PCGDeterminismTestsCommon.cpp
+++ b/Engine/Plugins/Experimental/PCG/Source/PCG/Private/Tests/Determinism/PCGDeterminismTestsCommon.cpp
@@ -4,15 +4,18 @@
 
 #include "PCGComponent.h"
 #include "PCGGraph.h"
+#include "PCGHelpers.h"
 #include "PCGPoint.h"
+#include "Data/PCGDifferenceData.h"
 #include "Data/PCGPointData.h"
-#include "Data/PCGPolyLineData.h"
 #include "Data/PCGPrimitiveData.h"
+#include "Data/PCGSplineData.h"
 #include "Data/PCGSurfaceData.h"
 #include "Data/PCGVolumeData.h"
 #include "Tests/PCGTestsCommon.h"
 
 #include "GameFramework/Actor.h"
+#include "Components/SplineComponent.h"
 
 #if WITH_EDITOR
 #include "Editor.h"
@@ -20,6 +23,11 @@
 
 namespace PCGDeterminismTests
 {
+	constexpr static FVector::FReal SmallDistance = 50.0;
+	constexpr static FVector::FReal MediumDistance = 200.0;
+	constexpr static FVector::FReal LargeDistance = 500.0;
+	constexpr static int32 NumSamplingStepsPerDimension = 50;
+
 	FTestData::FTestData(int32 RandomSeed) :
 		Seed(RandomSeed),
 		RandomStream(Seed)
@@ -37,17 +45,17 @@ namespace PCGDeterminismTests
 		TestActor = EditorWorld->SpawnActor<AActor>(AActor::StaticClass(), TransientActorParameters);
 		check(TestActor);
 
-		TestComponent = NewObject<UPCGComponent>(TestActor, FName(TEXT("PCG Test Component")), RF_Transient);
-		check(TestComponent);
-		TestActor->AddInstanceComponent(TestComponent);
-		TestComponent->RegisterComponent();
+		TestPCGComponent = NewObject<UPCGComponent>(TestActor, FName(TEXT("Test PCG Component")), RF_Transient);
+		check(TestPCGComponent);
+		TestActor->AddInstanceComponent(TestPCGComponent);
+		TestPCGComponent->RegisterComponent();
 
-		UPCGGraph* TestGraph = NewObject<UPCGGraph>(TestComponent, FName(TEXT("PCG Test Graph")), RF_Transient);
+		UPCGGraph* TestGraph = NewObject<UPCGGraph>(TestPCGComponent, FName(TEXT("Test PCG Graph")), RF_Transient);
 		check(TestGraph);
-		TestComponent->SetGraph(TestGraph);
+		TestPCGComponent->SetGraph(TestGraph);
 #else
 		TestActor = nullptr;
-		TestComponent = nullptr;
+		TestPCGComponent = nullptr;
 		Settings = nullptr;
 #endif
 	}
@@ -60,21 +68,136 @@ namespace PCGDeterminismTests
 		Settings = nullptr;
 	}
 
-	void AddVolumeInputData(FPCGDataCollection& InputData, const FVector& Location, const FVector& HalfSize, const FVector& VoxelSize)
+	void AddSinglePointInputData(FPCGDataCollection& InputData, const FVector& Location, const FName& PinName)
+	{
+		UPCGPointData* PointData = PCGTestsCommon::CreatePointData(Location);
+
+		FPCGTaggedData& TaggedData = InputData.TaggedData.Emplace_GetRef();
+		TaggedData.Data = PointData;
+		TaggedData.Pin = PinName;
+	}
+
+	void AddMultiplePointsInputData(FPCGDataCollection& InputData, const TArray<FPCGPoint>& Points, const FName& PinName)
+	{
+		UPCGPointData* PointData = PCGTestsCommon::CreateEmptyPointData();
+		PointData->SetPoints(Points);
+
+		FPCGTaggedData& TaggedData = InputData.TaggedData.Emplace_GetRef();
+		TaggedData.Data = PointData;
+		TaggedData.Pin = PinName;
+	}
+
+	void AddVolumeInputData(FPCGDataCollection& InputData, const FVector& Location, const FVector& HalfSize, const FVector& VoxelSize, const FName& PinName)
 	{
 		UPCGVolumeData* VolumeData = PCGTestsCommon::CreateVolumeData(FBox::BuildAABB(Location, HalfSize));
 		VolumeData->VoxelSize = VoxelSize;
-		InputData.TaggedData.Emplace_GetRef().Data = VolumeData;
+
+		FPCGTaggedData& TaggedData = InputData.TaggedData.Emplace_GetRef();
+		TaggedData.Data = VolumeData;
+		TaggedData.Pin = PinName;
+	}
+
+	void AddPolyLineInputData(FPCGDataCollection& InputData, USplineComponent* SplineComponent, const FName& PinName)
+	{
+		UPCGSplineData* SplineData = NewObject<UPCGSplineData>();
+		SplineData->Initialize(SplineComponent);
+
+		FPCGTaggedData& TaggedData = InputData.TaggedData.Emplace_GetRef();
+		TaggedData.Data = SplineData;
+		TaggedData.Pin = PinName;
+	}
+
+	void AddPrimitiveInputData(FPCGDataCollection& InputData, UPrimitiveComponent* PrimitiveComponent, const FVector& VoxelSize, const FName& PinName)
+	{
+		UPCGPrimitiveData* PrimitiveData = NewObject<UPCGPrimitiveData>();
+		PrimitiveData->Initialize(PrimitiveComponent);
+		PrimitiveData->VoxelSize = VoxelSize;
+
+		FPCGTaggedData& TaggedData = InputData.TaggedData.Emplace_GetRef();
+		TaggedData.Data = PrimitiveData;
+		TaggedData.Pin = PinName;
+	}
+
+	void AddRandomizedSinglePointInputData(FTestData& TestData, int32 PointNum, const FName& PinName)
+	{
+		check(PointNum > 0);
+		for (int32 I = 0; I < PointNum; ++I)
+		{
+			AddSinglePointInputData(TestData.InputData, TestData.RandomStream.VRand() * LargeDistance, PinName);
+		}
+	}
+
+	void AddRandomizedMultiplePointInputData(FTestData& TestData, int32 PointNum, const FName& PinName)
+	{
+		check(PointNum > 0);
+
+		TArray<FPCGPoint> Points;
+		Points.SetNumUninitialized(PointNum);
+		for (int32 I = 0; I < PointNum; ++I)
+		{
+			FVector NewLocation = TestData.RandomStream.VRand() * LargeDistance;
+			FTransform NewTransform(FRotator::ZeroRotator, NewLocation, FVector::OneVector * TestData.RandomStream.FRandRange(0.5, 1.5));
+			int PointSeed = PCGHelpers::ComputeSeed(static_cast<int>(NewLocation.X), static_cast<int>(NewLocation.Y), static_cast<int>(NewLocation.Z));
+			Points[I] = FPCGPoint(NewTransform, 1.f, PCGHelpers::ComputeSeed(PointSeed, TestData.Seed));
+		}
+
+		AddMultiplePointsInputData(TestData.InputData, Points, PinName);
 	}
 
-	void AddGenericVolumeInputData(FTestData& TestData)
+	void AddRandomizedVolumeInputData(FTestData& TestData, const FName& PinName)
 	{
-		AddVolumeInputData(TestData.InputData, FVector::ZeroVector, FVector::OneVector * 500.f, FVector::OneVector * 200.f);
+		AddVolumeInputData(TestData.InputData, TestData.RandomStream.VRand() * LargeDistance,
+			FVector::OneVector * LargeDistance + TestData.RandomStream.VRand() * MediumDistance,
+			FVector::OneVector * MediumDistance + TestData.RandomStream.VRand() * 0.5f * MediumDistance, PinName);
 	}
 
-	void AddRandomizedVolumeInputData(FTestData& TestData)
+	void AddRandomizedSurfaceInputData(FTestData& TestData, const FName& PinName)
 	{
-		AddVolumeInputData(TestData.InputData, TestData.RandomStream.VRand() * 500.f, FVector::OneVector * 500.f + TestData.RandomStream.VRand() * 200.f, FVector::OneVector * 200.f + TestData.RandomStream.VRand() * 100.f);
+		// TODO: PCG doesn't currently generate Surface data; function remains for future scalability
+	}
+
+	void AddRandomizedPolyLineInputData(FTestData& TestData, int32 PointNum, const FName& PinName)
+	{
+		check(TestData.TestActor);
+		USplineComponent* TestSplineComponent = Cast<USplineComponent>(TestData.TestActor->GetComponentByClass(USplineComponent::StaticClass()));
+
+		if (TestSplineComponent == nullptr)
+		{
+			TestSplineComponent = NewObject<USplineComponent>(TestData.TestActor, FName(TEXT("Test Spline Component")), RF_Transient);
+		}
+
+		check(PointNum > 1);
+		for (int32 I = 0; I < PointNum; ++I)
+		{
+			TestSplineComponent->AddSplinePoint(TestData.RandomStream.VRand() * LargeDistance, ESplineCoordinateSpace::Type::World, false);
+			TestSplineComponent->AddRelativeRotation(FRotator(
+				TestData.RandomStream.FRandRange(-90.0, 90.0),
+				TestData.RandomStream.FRandRange(-90.0, 90.0),
+				TestData.RandomStream.FRandRange(-90.0, 90.0)));
+		}
+		TestSplineComponent->UpdateSpline();
+
+		AddPolyLineInputData(TestData.InputData, TestSplineComponent, PinName);
+	}
+
+	void AddRandomizedPrimitiveInputData(FTestData& TestData, const FName& PinName)
+	{
+		check(TestData.TestActor);
+		UPrimitiveComponent* TestPrimitiveComponent = Cast<UPrimitiveComponent>(TestData.TestActor->GetComponentByClass(UPrimitiveComponent::StaticClass()));
+		
+		if (TestPrimitiveComponent == nullptr)
+		{
+			// TODO: If it reaches here, this will break, as it has no bounds. Please suggest the best way to give it some bounds?
+			TestPrimitiveComponent = NewObject<UStaticMeshComponent>(TestData.TestActor, FName(TEXT("Test Primitive Component")), RF_Transient);
+		}
+
+		TestPrimitiveComponent->SetWorldTransform(FTransform(
+			FRotator(TestData.RandomStream.FRandRange(0.0, 90.0), TestData.RandomStream.FRandRange(0.0, 90.0), TestData.RandomStream.FRandRange(0.0, 90.0)),
+			TestData.RandomStream.VRand() * LargeDistance,
+			FVector::OneVector * TestData.RandomStream.FRandRange(0.5, 1.5)));
+
+		// TODO: Probably more varieties to add in the future
+		AddPrimitiveInputData(TestData.InputData, TestPrimitiveComponent, FVector::OneVector * MediumDistance + TestData.RandomStream.VRand() * 0.5f * MediumDistance, PinName);
 	}
 
 	bool DataCollectionsAreIdentical(const FPCGDataCollection& FirstCollection, const FPCGDataCollection& SecondCollection)
@@ -143,38 +266,34 @@ namespace PCGDeterminismTests
 		return true;
 	}
 
-	bool SpatialDataIsIdentical(const UPCGData* FirstSpatialData, const UPCGData* SecondSpatialData)
+	bool SpatialDataIsIdentical(const UPCGData* FirstData, const UPCGData* SecondData)
 	{
-		check(FirstSpatialData);
-		check(SecondSpatialData);
+		check(FirstData);
+		check(SecondData);
 
-		if (Cast<UPCGPointData>(FirstSpatialData) != nullptr && Cast<UPCGPointData>(SecondSpatialData) != nullptr)
+		if (BothDataCastsToDataType<const UPCGPointData>(FirstData, SecondData))
 		{
-			return PointDataIsIdentical(FirstSpatialData, SecondSpatialData);
+			return PointDataIsIdentical(FirstData, SecondData);
 		}
-		else if (Cast<UPCGVolumeData>(FirstSpatialData) != nullptr && Cast<UPCGVolumeData>(SecondSpatialData) != nullptr)
+		else if (BothDataCastsToDataType<const UPCGVolumeData>(FirstData, SecondData))
 		{
-			// TODO: Implement Volume Data comparison
-			UE_LOG(LogPCG, Warning, TEXT("Volume comparison unimplemented."));
-			return ComparisonIsUnimplemented(FirstSpatialData, SecondSpatialData);
+			return VolumeDataIsIdentical(FirstData, SecondData);
 		}
-		else if (Cast<UPCGSurfaceData>(FirstSpatialData) != nullptr && Cast<UPCGSurfaceData>(SecondSpatialData) != nullptr)
+		else if (BothDataCastsToDataType<const UPCGSurfaceData>(FirstData, SecondData))
 		{
-			// TODO: Implement Surface Data comparison
-			UE_LOG(LogPCG, Warning, TEXT("Surface comparison unimplemented."));
-			return ComparisonIsUnimplemented(FirstSpatialData, SecondSpatialData);
+			return SurfaceDataIsIdentical(FirstData, SecondData);
 		}
-		else if (Cast<UPCGPolyLineData>(FirstSpatialData) != nullptr && Cast<UPCGPolyLineData>(SecondSpatialData) != nullptr)
+		else if (BothDataCastsToDataType<const UPCGPolyLineData>(FirstData, SecondData))
 		{
-			// TODO: Implement PolyLine Data comparison
-			UE_LOG(LogPCG, Warning, TEXT("PolyLine comparison unimplemented."));
-			return ComparisonIsUnimplemented(FirstSpatialData, SecondSpatialData);
+			return PolyLineDataIsIdentical(FirstData, SecondData);
 		}
-		else if (Cast<UPCGPrimitiveData>(FirstSpatialData) != nullptr && Cast<UPCGPrimitiveData>(SecondSpatialData) != nullptr)
+		else if (BothDataCastsToDataType<const UPCGPrimitiveData>(FirstData, SecondData))
 		{
-			// TODO: Implement Primitive Data comparison
-			UE_LOG(LogPCG, Warning, TEXT("Primitive comparison unimplemented."));
-			return ComparisonIsUnimplemented(FirstSpatialData, SecondSpatialData);
+			return PrimitiveDataIsIdentical(FirstData, SecondData);
+		}
+		else if (BothDataCastsToDataType<const UPCGDifferenceData>(FirstData, SecondData))
+		{
+			return DifferenceDataIsIdentical(FirstData, SecondData);
 		}
 
 		return false;
@@ -182,13 +301,16 @@ namespace PCGDeterminismTests
 
 	bool PointDataIsIdentical(const UPCGData* FirstData, const UPCGData* SecondData)
 	{
-		const UPCGPointData* FirstPointData = Cast<const UPCGPointData>(FirstData);
-		const UPCGPointData* SecondPointData = Cast<const UPCGPointData>(SecondData);
-		check(FirstPointData);
-		check(SecondPointData);
+		const UPCGPointData* FirstPointData = CastChecked<const UPCGPointData>(FirstData);
+		const UPCGPointData* SecondPointData = CastChecked<const UPCGPointData>(SecondData);
+
+		if (!SpatialBasicsAreIdentical(FirstPointData, SecondPointData))
+		{
+			return false;
+		}
 
-		const TArray<FPCGPoint>& FirstPoints = Cast<const UPCGPointData>(FirstPointData)->GetPoints();
-		const TArray<FPCGPoint>& SecondPoints = Cast<const UPCGPointData>(SecondPointData)->GetPoints();
+		const TArray<FPCGPoint>& FirstPoints = FirstPointData->GetPoints();
+		const TArray<FPCGPoint>& SecondPoints = SecondPointData->GetPoints();
 
 		if (FirstPoints.Num() != SecondPoints.Num())
 		{
@@ -227,7 +349,119 @@ namespace PCGDeterminismTests
 		return true;
 	}
 
-	bool ComparisonIsUnimplemented(const UPCGData* FirstPointData, const UPCGData* SecondPointData)
+	bool VolumeDataIsIdentical(const UPCGData* FirstData, const UPCGData* SecondData)
+	{
+		const UPCGVolumeData* FirstVolumeData = CastChecked<const UPCGVolumeData>(FirstData);
+		const UPCGVolumeData* SecondVolumeData = CastChecked<const UPCGVolumeData>(SecondData);
+
+		return FirstVolumeData->VoxelSize == SecondVolumeData->VoxelSize && SpatialBasicsAreIdentical(FirstVolumeData, SecondVolumeData);
+	}
+
+	bool SurfaceDataIsIdentical(const UPCGData* FirstData, const UPCGData* SecondData)
+	{
+		const UPCGSurfaceData* FirstSurfaceData = CastChecked<const UPCGSurfaceData>(FirstData);
+		const UPCGSurfaceData* SecondSurfaceData = CastChecked<const UPCGSurfaceData>(SecondData);
+
+		// TODO: Implement Surface Data comparison as needed in the future
+		UE_LOG(LogPCG, Warning, TEXT("Surface comparison not fully implemented."));
+
+		return SpatialBasicsAreIdentical(FirstSurfaceData, SecondSurfaceData);
+	}
+
+	bool PolyLineDataIsIdentical(const UPCGData* FirstData, const UPCGData* SecondData)
+	{
+		const UPCGPolyLineData* FirstPolyLineData = CastChecked<const UPCGPolyLineData>(FirstData);
+		const UPCGPolyLineData* SecondPolyLineData = CastChecked<const UPCGPolyLineData>(SecondData);
+
+		if (!SpatialBasicsAreIdentical(FirstPolyLineData, SecondPolyLineData))
+		{
+			return false;
+		}
+
+		int NumSegments = FirstPolyLineData->GetNumSegments();
+		for (int32 I = 0; I < NumSegments; ++I)
+		{
+			// TODO: Needs more robust checking for straight line vs spline tangents, etc
+			if (FirstPolyLineData->GetSegmentLength(I) != SecondPolyLineData->GetSegmentLength(I) ||
+				!FirstPolyLineData->GetTransformAtDistance(I, 0.0).Equals(SecondPolyLineData->GetTransformAtDistance(I, 0.0)))
+			{
+				return false;
+			}
+		}
+
+		return true;
+	}
+
+	bool PrimitiveDataIsIdentical(const UPCGData* FirstData, const UPCGData* SecondData)
+	{
+		const UPCGPrimitiveData* FirstPrimitiveData = CastChecked<const UPCGPrimitiveData>(FirstData);
+		const UPCGPrimitiveData* SecondPrimitiveData = CastChecked<const UPCGPrimitiveData>(SecondData);
+
+		if (FirstPrimitiveData->VoxelSize != SecondPrimitiveData->VoxelSize || !SpatialBasicsAreIdentical(FirstPrimitiveData, SecondPrimitiveData))
+		{
+			return false;
+		}
+
+		// TODO: Compare the ToPointData, which will require the Context
+		return false;
+	}
+
+	bool DifferenceDataIsIdentical(const UPCGData* FirstData, const UPCGData* SecondData)
+	{
+		const UPCGDifferenceData* FirstDifferenceData = CastChecked<const UPCGDifferenceData>(FirstData);
+		const UPCGDifferenceData* SecondDifferenceData = CastChecked<const UPCGDifferenceData>(SecondData);
+
+		if (!SpatialBasicsAreIdentical(FirstDifferenceData, SecondDifferenceData))
+		{
+			return false;
+		}
+
+		// At this point, bounds has already been checked for equality
+		const FBox SampleBounds = FirstDifferenceData->GetBounds();
+		const FVector SampleExtent = SampleBounds.GetExtent();
+
+		FPCGPoint FirstPoint;
+		FPCGPoint SecondPoint;
+		FVector StepInterval = SampleExtent * 2.0 / FMath::Max(NumSamplingStepsPerDimension, 1);
+		FVector StartingOffset = SampleBounds.Min + StepInterval * 0.5;
+
+		// Sample points across the 3D volume
+		for (FVector::FReal X = StartingOffset.X; X < SampleBounds.Max.X; X += StepInterval.X)
+		{
+			for (FVector::FReal Y = StartingOffset.Y; Y < SampleBounds.Max.Y; Y += StepInterval.Y)
+			{
+				for (FVector::FReal Z = StartingOffset.Z; Z < SampleBounds.Max.Z; Z += StepInterval.Z)
+				{
+					FTransform PointTransform(FVector(X, Y, Z));
+					bool bFirstPointWasSampled = FirstDifferenceData->SamplePoint(PointTransform, SampleBounds, FirstPoint, nullptr);
+					bool bSecondPointWasSampled = SecondDifferenceData->SamplePoint(PointTransform, SampleBounds, SecondPoint, nullptr);
+
+					if (bFirstPointWasSampled != bSecondPointWasSampled)
+					{
+						return false;
+					}
+
+					// Only compare if both points were sampled
+					if (bFirstPointWasSampled && bSecondPointWasSampled && !PCGTestsCommon::PointsAreIdentical(FirstPoint, SecondPoint))
+					{
+						return false;
+					}
+				}
+			}
+		}
+
+		return true;
+	}
+
+	bool SpatialBasicsAreIdentical(const UPCGSpatialData* FirstSpatialData, const UPCGSpatialData* SecondSpatialData)
+	{
+		return (FirstSpatialData->GetDataType() == SecondSpatialData->GetDataType() &&
+			FirstSpatialData->GetDimension() == SecondSpatialData->GetDimension() &&
+			FirstSpatialData->GetBounds() == SecondSpatialData->GetBounds() &&
+			FirstSpatialData->GetStrictBounds() == SecondSpatialData->GetStrictBounds());
+	}
+
+	bool ComparisonIsUnimplemented(const UPCGData* FirstData, const UPCGData* SecondData)
 	{
 		return false;
 	}
@@ -240,7 +474,7 @@ namespace PCGDeterminismTests
 			return true;
 		}
 
-		// TODO: Data types that don't need to be compared
+		// Data types that don't need to be compared
 		if (DataType == EPCGDataType::None || DataType == EPCGDataType::Other || DataType == EPCGDataType::Settings)
 		{
 			return false;
@@ -275,7 +509,6 @@ namespace PCGDeterminismTests
 
 		switch (DataType)
 		{
-			// TODO: Many more data type comparisons to be added
 		case EPCGDataType::Spatial:
 			return SpatialDataIsIdentical;
 		default:
@@ -289,8 +522,8 @@ namespace PCGDeterminismTests
 		FPCGElementPtr FirstElement = FirstTestData.Settings->GetElement();
 		FPCGElementPtr SecondElement = SecondTestData.Settings->GetElement();
 
-		TUniquePtr<FPCGContext> FirstContext(FirstElement->Initialize(FirstTestData.InputData, FirstTestData.TestComponent, nullptr));
-		TUniquePtr<FPCGContext> SecondContext(SecondElement->Initialize(SecondTestData.InputData, SecondTestData.TestComponent, nullptr));
+		TUniquePtr<FPCGContext> FirstContext(FirstElement->Initialize(FirstTestData.InputData, FirstTestData.TestPCGComponent, nullptr));
+		TUniquePtr<FPCGContext> SecondContext(SecondElement->Initialize(SecondTestData.InputData, SecondTestData.TestPCGComponent, nullptr));
 
 		FirstContext->NumAvailableTasks = 1;
 		SecondContext->NumAvailableTasks = 1;
@@ -306,4 +539,9 @@ namespace PCGDeterminismTests
 
 		return PCGDeterminismTests::DataCollectionsAreIdentical(FirstContext->OutputData, SecondContext->OutputData);
 	}
+
+	bool ExecutionIsDeterministicSameData(FTestData& TestData)
+	{
+		return ExecutionIsDeterministic(TestData, TestData);
+	}
 }
\ No newline at end of file
diff --git a/Engine/Plugins/Experimental/PCG/Source/PCG/Private/Tests/Determinism/PCGDifferenceDeterminismTest.cpp b/Engine/Plugins/Experimental/PCG/Source/PCG/Private/Tests/Determinism/PCGDifferenceDeterminismTest.cpp
new file mode 100644
--- /dev/null
+++ b/Engine/Plugins/Experimental/PCG/Source/PCG/Private/Tests/Determinism/PCGDifferenceDeterminismTest.cpp
@@ -0,0 +1,110 @@
+// Copyright Epic Games, Inc. All Rights Reserved.
+
+#if WITH_EDITOR
+
+#include "Elements/PCGDifferenceElement.h"
+#include "Tests/Determinism/PCGDeterminismTestsCommon.h"
+
+IMPLEMENT_CUSTOM_SIMPLE_AUTOMATION_TEST(FPCGDifferenceDeterminismSingleSameDataTest, FPCGTestBaseClass, "pcg.tests.Difference.Determinism.SingleSameData", PCGTestsCommon::TestFlags)
+IMPLEMENT_CUSTOM_SIMPLE_AUTOMATION_TEST(FPCGDifferenceDeterminismSingleIdenticalDataTest, FPCGTestBaseClass, "pcg.tests.Difference.Determinism.SingleMultipleData", PCGTestsCommon::TestFlags)
+IMPLEMENT_CUSTOM_SIMPLE_AUTOMATION_TEST(FPCGDifferenceDeterminismMultipleSameDataTest, FPCGTestBaseClass, "pcg.tests.Difference.Determinism.MultipleSameData", PCGTestsCommon::TestFlags)
+IMPLEMENT_CUSTOM_SIMPLE_AUTOMATION_TEST(FPCGDifferenceDeterminismMultipleIdenticalDataTest, FPCGTestBaseClass, "pcg.tests.Difference.Determinism.MultipleIdenticalData", PCGTestsCommon::TestFlags)
+// TODO: Temporarily disabled while more testing into this is conducted
+//IMPLEMENT_CUSTOM_SIMPLE_AUTOMATION_TEST(FPCGDifferenceDeterminismOrderIndependenceTest, FPCGTestBaseClass, "pcg.tests.Difference.Determinism.OrderIndependence", PCGTestsCommon::TestFlags)
+
+namespace
+{
+	void DifferenceTestBase(PCGDeterminismTests::FTestData& TestData)
+	{
+		PCGDeterminismTests::GenerateSettings<UPCGDifferenceSettings>(TestData);
+		// Source
+		PCGDeterminismTests::AddVolumeInputData(TestData.InputData, FVector::ZeroVector, FVector::OneVector * 2000.f, FVector::OneVector * 200.f, PCGDifferenceConstants::SourceLabel);
+
+		// Difference 
+		PCGDeterminismTests::AddVolumeInputData(TestData.InputData, FVector::ZeroVector, FVector::OneVector * 500.f, FVector::OneVector * 200.f, PCGDifferenceConstants::DifferencesLabel);
+	}
+
+	void DifferenceTestMultiple(PCGDeterminismTests::FTestData& TestData)
+	{
+		DifferenceTestBase(TestData);
+
+		// Randomized Sources
+		PCGDeterminismTests::AddRandomizedVolumeInputData(TestData, PCGDifferenceConstants::SourceLabel);
+		PCGDeterminismTests::AddRandomizedMultiplePointInputData(TestData, 20, PCGDifferenceConstants::SourceLabel);
+
+		// Randomized Differences
+		PCGDeterminismTests::AddRandomizedMultiplePointInputData(TestData, 20, PCGDifferenceConstants::DifferencesLabel);
+	}
+
+}
+
+bool FPCGDifferenceDeterminismSingleSameDataTest::RunTest(const FString& Parameters)
+{
+	// Test single same data
+	PCGDeterminismTests::FTestData TestData(PCGDeterminismTests::DefaultSeed);
+
+	DifferenceTestBase(TestData);
+
+	TestTrue("Same single input and settings, same output", PCGDeterminismTests::ExecutionIsDeterministicSameData(TestData));
+
+	return true;
+}
+
+bool FPCGDifferenceDeterminismSingleIdenticalDataTest::RunTest(const FString& Parameters)
+{
+	// Test single identical data
+	PCGDeterminismTests::FTestData FirstTestData(PCGDeterminismTests::DefaultSeed);
+	PCGDeterminismTests::FTestData SecondTestData(PCGDeterminismTests::DefaultSeed);
+
+	DifferenceTestBase(FirstTestData);
+	DifferenceTestBase(SecondTestData);
+
+	TestTrue("Identical single input and settings, same output", PCGDeterminismTests::ExecutionIsDeterministic(FirstTestData, SecondTestData));
+
+	return true;
+}
+
+bool FPCGDifferenceDeterminismMultipleSameDataTest::RunTest(const FString& Parameters)
+{
+	// Test multiple same data
+	PCGDeterminismTests::FTestData TestData(PCGDeterminismTests::DefaultSeed);
+
+	DifferenceTestMultiple(TestData);
+
+	TestTrue("Identical multiple input, same output", PCGDeterminismTests::ExecutionIsDeterministicSameData(TestData));
+
+	return true;
+}
+
+bool FPCGDifferenceDeterminismMultipleIdenticalDataTest::RunTest(const FString& Parameters)
+{
+	// Test multiple identical data
+	PCGDeterminismTests::FTestData FirstTestData(PCGDeterminismTests::DefaultSeed);
+	PCGDeterminismTests::FTestData SecondTestData(PCGDeterminismTests::DefaultSeed);
+
+	DifferenceTestMultiple(FirstTestData);
+	DifferenceTestMultiple(SecondTestData);
+
+	TestTrue("Identical single input and settings, same output", PCGDeterminismTests::ExecutionIsDeterministic(FirstTestData, SecondTestData));
+
+	return true;
+}
+
+// TODO: Temporarily disabled while more testing into this is conducted
+//bool FPCGDifferenceDeterminismOrderIndependenceTest::RunTest(const FString& Parameters)
+//{
+//	// Test multiple identical data
+//	PCGDeterminismTests::FTestData FirstTestData(PCGDeterminismTests::DefaultSeed);
+//	PCGDeterminismTests::FTestData SecondTestData(PCGDeterminismTests::DefaultSeed);
+//
+//	DifferenceTestMultiple(FirstTestData);
+//	DifferenceTestMultiple(SecondTestData);
+//
+//	PCGDeterminismTests::ShuffleInputOrder(SecondTestData);
+//
+//	TestTrue("Shuffled input order, same output", PCGDeterminismTests::ExecutionIsDeterministic(FirstTestData, SecondTestData));
+//
+//	return true;
+//}
+
+#endif
\ No newline at end of file
diff --git a/Engine/Plugins/Experimental/PCG/Source/PCG/Private/Tests/Determinism/PCGVolumeSamplerDeterminismTest.cpp b/Engine/Plugins/Experimental/PCG/Source/PCG/Private/Tests/Determinism/PCGVolumeSamplerDeterminismTest.cpp
--- a/Engine/Plugins/Experimental/PCG/Source/PCG/Private/Tests/Determinism/PCGVolumeSamplerDeterminismTest.cpp
+++ b/Engine/Plugins/Experimental/PCG/Source/PCG/Private/Tests/Determinism/PCGVolumeSamplerDeterminismTest.cpp
@@ -1,88 +1,80 @@
 // Copyright Epic Games, Inc. All Rights Reserved.
 
-#include "PCGComponent.h"
-#include "PCGGraph.h"
-#include "Data/PCGSpatialData.h"
-#include "Data/PCGPointData.h"
+#if WITH_EDITOR
+
 #include "Elements/PCGVolumeSampler.h"
-#include "Tests/PCGTestsCommon.h"
 #include "Tests/Determinism/PCGDeterminismTestsCommon.h"
 
-IMPLEMENT_CUSTOM_SIMPLE_AUTOMATION_TEST(FPCGVolumeSamplerDeterminismSingleSameDataTest, FPCGTestBaseClass, "pcg.determinism.VolumeSampler.SingleSameData", PCGTestsCommon::TestFlags)
-IMPLEMENT_CUSTOM_SIMPLE_AUTOMATION_TEST(FPCGVolumeSamplerDeterminismSingleIdenticalDataTest, FPCGTestBaseClass, "pcg.determinism.VolumeSampler.SingleIdenticalData", PCGTestsCommon::TestFlags)
-IMPLEMENT_CUSTOM_SIMPLE_AUTOMATION_TEST(FPCGVolumeSamplerDeterminismMultipleSameDataTest, FPCGTestBaseClass, "pcg.determinism.VolumeSampler.MultipleSameData", PCGTestsCommon::TestFlags)
-IMPLEMENT_CUSTOM_SIMPLE_AUTOMATION_TEST(FPCGVolumeSamplerDeterminismMultipleIdenticalDataTest, FPCGTestBaseClass, "pcg.determinism.VolumeSampler.MultipleIdenticalData", PCGTestsCommon::TestFlags)
-IMPLEMENT_CUSTOM_SIMPLE_AUTOMATION_TEST(FPCGVolumeSamplerDeterminismOrderIndependenceTest, FPCGTestBaseClass, "pcg.determinism.VolumeSampler.OrderIndependence", PCGTestsCommon::TestFlags)
+IMPLEMENT_CUSTOM_SIMPLE_AUTOMATION_TEST(FPCGVolumeSamplerDeterminismSingleSameDataTest, FPCGTestBaseClass, "pcg.tests.VolumeSampler.Determinism.SingleSameData", PCGTestsCommon::TestFlags)
+IMPLEMENT_CUSTOM_SIMPLE_AUTOMATION_TEST(FPCGVolumeSamplerDeterminismSingleIdenticalDataTest, FPCGTestBaseClass, "pcg.tests.VolumeSampler.Determinism.SingleIdenticalData", PCGTestsCommon::TestFlags)
+IMPLEMENT_CUSTOM_SIMPLE_AUTOMATION_TEST(FPCGVolumeSamplerDeterminismMultipleSameDataTest, FPCGTestBaseClass, "pcg.tests.VolumeSampler.Determinism.MultipleSameData", PCGTestsCommon::TestFlags)
+IMPLEMENT_CUSTOM_SIMPLE_AUTOMATION_TEST(FPCGVolumeSamplerDeterminismMultipleIdenticalDataTest, FPCGTestBaseClass, "pcg.tests.VolumeSampler.Determinism.MultipleIdenticalData", PCGTestsCommon::TestFlags)
+IMPLEMENT_CUSTOM_SIMPLE_AUTOMATION_TEST(FPCGVolumeSamplerDeterminismOrderIndependenceTest, FPCGTestBaseClass, "pcg.tests.VolumeSampler.Determinism.OrderIndependence", PCGTestsCommon::TestFlags)
 
 namespace
 {
-	void VolumeSettingsDelegate(UPCGVolumeSamplerSettings* Settings, FRandomStream& RandomStream)
+	void RandomizeVolumeSettingsVoxelSize(PCGDeterminismTests::FTestData& TestData)
 	{
-		Settings->VoxelSize = FVector::OneVector * 200.f + RandomStream.VRand() * 100.f;
+		CastChecked<UPCGVolumeSamplerSettings>(TestData.Settings)->VoxelSize = FVector::OneVector * 200.f + TestData.RandomStream.VRand() * 100.f;
 	}
 }
 
 bool FPCGVolumeSamplerDeterminismSingleSameDataTest::RunTest(const FString& Parameters)
 {
-#if WITH_EDITOR
 	// Test single same data
-	PCGDeterminismTests::FTestData FirstTestData(PCGDeterminismTests::DefaultSeed);
+	PCGDeterminismTests::FTestData TestData(PCGDeterminismTests::DefaultSeed);
 
-	PCGDeterminismTests::GenerateSettings<UPCGVolumeSamplerSettings>(FirstTestData, VolumeSettingsDelegate);
-	PCGDeterminismTests::AddRandomizedVolumeInputData(FirstTestData);
+	PCGDeterminismTests::GenerateSettings<UPCGVolumeSamplerSettings>(TestData, RandomizeVolumeSettingsVoxelSize);
+	PCGDeterminismTests::AddRandomizedVolumeInputData(TestData);
 
-	TestTrue("Same single input and settings, same output", PCGDeterminismTests::ExecutionIsDeterministic(FirstTestData, FirstTestData));
+	TestTrue("Same single input and settings, same output", PCGDeterminismTests::ExecutionIsDeterministicSameData(TestData));
 
-#endif
 	return true;
 }
 
 bool FPCGVolumeSamplerDeterminismSingleIdenticalDataTest::RunTest(const FString& Parameters)
 {
-#if WITH_EDITOR
 	// Test single identical data
 	PCGDeterminismTests::FTestData FirstTestData(PCGDeterminismTests::DefaultSeed);
 	PCGDeterminismTests::FTestData SecondTestData(PCGDeterminismTests::DefaultSeed);
 
-	PCGDeterminismTests::GenerateSettings<UPCGVolumeSamplerSettings>(FirstTestData, VolumeSettingsDelegate);
-	PCGDeterminismTests::GenerateSettings<UPCGVolumeSamplerSettings>(SecondTestData, VolumeSettingsDelegate);
+	PCGDeterminismTests::GenerateSettings<UPCGVolumeSamplerSettings>(FirstTestData, RandomizeVolumeSettingsVoxelSize);
+	PCGDeterminismTests::GenerateSettings<UPCGVolumeSamplerSettings>(SecondTestData, RandomizeVolumeSettingsVoxelSize);
 
 	PCGDeterminismTests::AddRandomizedVolumeInputData(FirstTestData);
 	PCGDeterminismTests::AddRandomizedVolumeInputData(SecondTestData);
 
 	TestTrue("Identical single input and settings, same output", PCGDeterminismTests::ExecutionIsDeterministic(FirstTestData, SecondTestData));
-#endif
+
 	return true;
 }
 
 bool FPCGVolumeSamplerDeterminismMultipleSameDataTest::RunTest(const FString& Parameters)
 {
-#if WITH_EDITOR
-	PCGDeterminismTests::FTestData FirstTestData(PCGDeterminismTests::DefaultSeed);
-
 	// Test multiple same data
-	PCGDeterminismTests::GenerateSettings<UPCGVolumeSamplerSettings>(FirstTestData, VolumeSettingsDelegate);
+	PCGDeterminismTests::FTestData TestData(PCGDeterminismTests::DefaultSeed);
+
+	PCGDeterminismTests::GenerateSettings<UPCGVolumeSamplerSettings>(TestData, RandomizeVolumeSettingsVoxelSize);
 
 	// Add many random volumes
 	for (int32 I = 0; I < 10; ++I)
 	{
-		PCGDeterminismTests::AddRandomizedVolumeInputData(FirstTestData);
+		PCGDeterminismTests::AddRandomizedVolumeInputData(TestData);
 	}
 
-	TestTrue("Same multiple input, same output", PCGDeterminismTests::ExecutionIsDeterministic(FirstTestData, FirstTestData));
-#endif
+	TestTrue("Same multiple input, same output", PCGDeterminismTests::ExecutionIsDeterministicSameData(TestData));
+
 	return true;
 }
 
 bool FPCGVolumeSamplerDeterminismMultipleIdenticalDataTest::RunTest(const FString& Parameters)
 {
-#if WITH_EDITOR
 	// Test multiple identical data
 	PCGDeterminismTests::FTestData FirstTestData(PCGDeterminismTests::DefaultSeed);
 	PCGDeterminismTests::FTestData SecondTestData(PCGDeterminismTests::DefaultSeed);
 
-	PCGDeterminismTests::GenerateSettings<UPCGVolumeSamplerSettings>(FirstTestData, VolumeSettingsDelegate);
-	PCGDeterminismTests::GenerateSettings<UPCGVolumeSamplerSettings>(SecondTestData, VolumeSettingsDelegate);
+	PCGDeterminismTests::GenerateSettings<UPCGVolumeSamplerSettings>(FirstTestData, RandomizeVolumeSettingsVoxelSize);
+	PCGDeterminismTests::GenerateSettings<UPCGVolumeSamplerSettings>(SecondTestData, RandomizeVolumeSettingsVoxelSize);
 
 	// Add many random volumes
 	for (int32 I = 0; I < 10; ++I)
@@ -92,19 +84,18 @@ bool FPCGVolumeSamplerDeterminismMultipleIdenticalDataTest::RunTest(const FStrin
 	}
 
 	TestTrue("Identical multiple input, same output", PCGDeterminismTests::ExecutionIsDeterministic(FirstTestData, SecondTestData));
-#endif
+
 	return true;
 }
 
 bool FPCGVolumeSamplerDeterminismOrderIndependenceTest::RunTest(const FString& Parameters)
 {
-#if WITH_EDITOR
 	// Test shuffled input data
 	PCGDeterminismTests::FTestData FirstTestData(PCGDeterminismTests::DefaultSeed);
 	PCGDeterminismTests::FTestData SecondTestData(PCGDeterminismTests::DefaultSeed);
 
-	PCGDeterminismTests::GenerateSettings<UPCGVolumeSamplerSettings>(FirstTestData, VolumeSettingsDelegate);
-	PCGDeterminismTests::GenerateSettings<UPCGVolumeSamplerSettings>(SecondTestData, VolumeSettingsDelegate);
+	PCGDeterminismTests::GenerateSettings<UPCGVolumeSamplerSettings>(FirstTestData, RandomizeVolumeSettingsVoxelSize);
+	PCGDeterminismTests::GenerateSettings<UPCGVolumeSamplerSettings>(SecondTestData, RandomizeVolumeSettingsVoxelSize);
 
 	// Add many random volumes
 	for (int32 I = 0; I < 10; ++I)
@@ -116,6 +107,8 @@ bool FPCGVolumeSamplerDeterminismOrderIndependenceTest::RunTest(const FString& P
 	PCGDeterminismTests::ShuffleInputOrder(SecondTestData);
 
 	TestTrue("Shuffled input order, same output", PCGDeterminismTests::ExecutionIsDeterministic(FirstTestData, SecondTestData));
-#endif
+
 	return true;
 }
+
+#endif
\ No newline at end of file
diff --git a/Engine/Plugins/Experimental/PCG/Source/PCG/Private/Tests/Elements/PCGCollapseElementTest.cpp b/Engine/Plugins/Experimental/PCG/Source/PCG/Private/Tests/Elements/PCGCollapseElementTest.cpp
--- a/Engine/Plugins/Experimental/PCG/Source/PCG/Private/Tests/Elements/PCGCollapseElementTest.cpp
+++ b/Engine/Plugins/Experimental/PCG/Source/PCG/Private/Tests/Elements/PCGCollapseElementTest.cpp
@@ -3,7 +3,7 @@
 #include "Tests/PCGTestsCommon.h"
 #include "Elements/PCGCollapseElement.h"
 
-IMPLEMENT_CUSTOM_SIMPLE_AUTOMATION_TEST(FPCGCollapseElementSmokeTest, FPCGTestBaseClass, "pcg.tests.Collapse.Smoke", PCGTestsCommon::TestFlags)
+IMPLEMENT_CUSTOM_SIMPLE_AUTOMATION_TEST(FPCGCollapseElementSmokeTest, FPCGTestBaseClass, "pcg.Collapse.Smoke", PCGTestsCommon::TestFlags)
 
 bool FPCGCollapseElementSmokeTest::RunTest(const FString& Parameters)
 {
diff --git a/Engine/Plugins/Experimental/PCG/Source/PCG/Private/Tests/Elements/PCGStaticMeshSpawnerTest.cpp b/Engine/Plugins/Experimental/PCG/Source/PCG/Private/Tests/Elements/PCGStaticMeshSpawnerTest.cpp
--- a/Engine/Plugins/Experimental/PCG/Source/PCG/Private/Tests/Elements/PCGStaticMeshSpawnerTest.cpp
+++ b/Engine/Plugins/Experimental/PCG/Source/PCG/Private/Tests/Elements/PCGStaticMeshSpawnerTest.cpp
@@ -62,7 +62,7 @@ void TestMeshSelectorByAttribute(
 
 	// initialize and execute the StaticMeshSpawner
 	FPCGElementPtr StaticMeshSpawner = Settings->GetElement();
-	FPCGContext* Context = StaticMeshSpawner->Initialize(TestData.InputData, TestData.TestComponent, nullptr);
+	FPCGContext* Context = StaticMeshSpawner->Initialize(TestData.InputData, TestData.TestPCGComponent, nullptr);
 
 	while (!StaticMeshSpawner->Execute(Context))
 	{}
@@ -91,10 +91,10 @@ void TestMeshSelectorByAttribute(
 		ValidInstanceCount += Entry.EndIndex - Entry.StartIndex;
 	}
 
-	check(TestData.TestComponent);
+	check(TestData.TestPCGComponent);
 
 	TArray<UInstancedStaticMeshComponent*> ISMCs;
-	TestData.TestComponent->ForEachManagedResource([&ISMCs](UPCGManagedResource* InResource)
+	TestData.TestPCGComponent->ForEachManagedResource([&ISMCs](UPCGManagedResource* InResource)
 	{
 		if (UPCGManagedISMComponent* Resource = Cast<UPCGManagedISMComponent>(InResource))
 		{
@@ -151,8 +151,8 @@ void TestMeshSelectorByAttribute(
 		}
 	}
 
-	TestData.TestComponent->bGenerated = true;
-	TestData.TestComponent->Cleanup();
+	TestData.TestPCGComponent->bGenerated = true;
+	TestData.TestPCGComponent->Cleanup();
 	
 	Test->TestTrue("Valid total instance count", TotalInstanceCount == ValidInstanceCount);
 }
@@ -171,7 +171,7 @@ void TestMeshSelectorWeighted(
 
 	// initialize and execute the StaticMeshSpawner
 	FPCGElementPtr StaticMeshSpawner = Settings->GetElement();
-	FPCGContext* Context = StaticMeshSpawner->Initialize(TestData.InputData, TestData.TestComponent, nullptr);
+	FPCGContext* Context = StaticMeshSpawner->Initialize(TestData.InputData, TestData.TestPCGComponent, nullptr);
 
 	while (!StaticMeshSpawner->Execute(Context))
 	{}
@@ -182,10 +182,10 @@ void TestMeshSelectorWeighted(
 		TotalWeight += Entry.Weight;
 	}
 
-	check(TestData.TestComponent);
+	check(TestData.TestPCGComponent);
 
 	TArray<UInstancedStaticMeshComponent*> ISMCs;
-	TestData.TestComponent->ForEachManagedResource([&ISMCs](UPCGManagedResource* InResource)
+	TestData.TestPCGComponent->ForEachManagedResource([&ISMCs](UPCGManagedResource* InResource)
 	{
 		if (UPCGManagedISMComponent* Resource = Cast<UPCGManagedISMComponent>(InResource))
 		{
@@ -220,8 +220,8 @@ void TestMeshSelectorWeighted(
 		Test->TestTrue("Valid instance count per mesh", InstanceCount >= TargetCount - ErrorBound && InstanceCount <= TargetCount + ErrorBound);
 	}
 
-	TestData.TestComponent->bGenerated = true;
-	TestData.TestComponent->Cleanup();
+	TestData.TestPCGComponent->bGenerated = true;
+	TestData.TestPCGComponent->Cleanup();
 	
 	Test->TestTrue("Valid total instance count", TotalInstanceCount == PointCount);
 }
@@ -243,16 +243,16 @@ void TestMeshSelectorWeightedByCategory(
 
 	// initialize and execute the StaticMeshSpawner
 	FPCGElementPtr StaticMeshSpawner = Settings->GetElement();
-	FPCGContext* Context = StaticMeshSpawner->Initialize(TestData.InputData, TestData.TestComponent, nullptr);
+	FPCGContext* Context = StaticMeshSpawner->Initialize(TestData.InputData, TestData.TestPCGComponent, nullptr);
 
 	while (!StaticMeshSpawner->Execute(Context))
 	{
 	}
 
-	check(TestData.TestComponent);
+	check(TestData.TestPCGComponent);
 
 	TArray<UInstancedStaticMeshComponent*> ISMCs;
-	TestData.TestComponent->ForEachManagedResource([&ISMCs](UPCGManagedResource* InResource)
+	TestData.TestPCGComponent->ForEachManagedResource([&ISMCs](UPCGManagedResource* InResource)
 		{
 			if (UPCGManagedISMComponent* Resource = Cast<UPCGManagedISMComponent>(InResource))
 			{
@@ -285,8 +285,8 @@ void TestMeshSelectorWeightedByCategory(
 		Test->TestTrue("Valid instance count per mesh", InstanceCount >= Entry->ExpectedCount - ErrorBound && InstanceCount <= Entry->ExpectedCount + ErrorBound);
 	}
 
-	TestData.TestComponent->bGenerated = true;
-	TestData.TestComponent->Cleanup();
+	TestData.TestPCGComponent->bGenerated = true;
+	TestData.TestPCGComponent->Cleanup();
 
 	Test->TestTrue("Valid total instance count", TotalInstanceCount == PointCount);
 }
diff --git a/Engine/Plugins/Experimental/PCG/Source/PCG/Public/Elements/PCGDifferenceElement.h b/Engine/Plugins/Experimental/PCG/Source/PCG/Public/Elements/PCGDifferenceElement.h
--- a/Engine/Plugins/Experimental/PCG/Source/PCG/Public/Elements/PCGDifferenceElement.h
+++ b/Engine/Plugins/Experimental/PCG/Source/PCG/Public/Elements/PCGDifferenceElement.h
@@ -9,6 +9,12 @@
 
 #include "PCGDifferenceElement.generated.h"
 
+namespace PCGDifferenceConstants
+{
+	const FName SourceLabel = TEXT("Source");
+	const FName DifferencesLabel = TEXT("Differences");
+}
+
 UCLASS(BlueprintType, ClassGroup = (Procedural))
 class PCG_API UPCGDifferenceSettings : public UPCGSettings
 {
diff --git a/Engine/Plugins/Experimental/PCG/Source/PCG/Public/Tests/Determinism/PCGDeterminismTestsCommon.h b/Engine/Plugins/Experimental/PCG/Source/PCG/Public/Tests/Determinism/PCGDeterminismTestsCommon.h
--- a/Engine/Plugins/Experimental/PCG/Source/PCG/Public/Tests/Determinism/PCGDeterminismTestsCommon.h
+++ b/Engine/Plugins/Experimental/PCG/Source/PCG/Public/Tests/Determinism/PCGDeterminismTestsCommon.h
@@ -4,15 +4,23 @@
 
 #include "PCGCommon.h"
 #include "PCGSettings.h"
+#include "Tests/PCGTestsCommon.h"
 
 class AActor;
 class UPCGComponent;
 class UPCGData;
+class UPCGSpatialData;
+class UPrimitiveComponent;
+class USplineComponent;
 struct FPCGDataCollection;
 
 namespace PCGDeterminismTests
 {
-	constexpr int32 DefaultSeed = 42;
+	constexpr static int32 DefaultSeed = 42;
+	constexpr static int32 DefaultNumPointsToGenerate = 1;
+	constexpr static int32 DefaultNumPolyLinePointsToGenerate = 6;
+
+	const FName DefaultTestPinName = TEXT("Test");
 
 	struct FTestData
 	{
@@ -21,29 +29,60 @@ namespace PCGDeterminismTests
 		void Reset();
 
 		AActor* TestActor;
-		UPCGComponent* TestComponent;
+		UPCGComponent* TestPCGComponent;
 		FPCGDataCollection InputData;
 		UPCGSettings* Settings;
 		int32 Seed;
 		FRandomStream RandomStream;
 	};
 
+	/** Helper for adding PointData with a single Point to an InputData */
+	void AddSinglePointInputData(FPCGDataCollection& InputData, const FVector& Location, const FName& PinName = DefaultTestPinName);
+	/** Helper for adding PointData with multiple Points to an InputData */
+	void AddMultiplePointsInputData(FPCGDataCollection& InputData, const TArray<FPCGPoint>& Points, const FName& PinName = DefaultTestPinName);
 	/** Helper for adding VolumeData to an InputData */
-	void AddVolumeInputData(FPCGDataCollection& InputData, const FVector& Location, const FVector& HalfSize, const FVector& VoxelSize);
-
-	/** Adds generic Volume Spatial Data at the origin */
-	void AddGenericVolumeInputData(FTestData& TestData);
-	/** Adds randomized Volume Spatial Data at the origin */
-	void AddRandomizedVolumeInputData(FTestData& TestData);
+	void AddVolumeInputData(FPCGDataCollection& InputData, const FVector& Location, const FVector& HalfSize, const FVector& VoxelSize, const FName& PinName = DefaultTestPinName);
+	/** Helper for adding PolyLineData to an InputData */
+	void AddPolyLineInputData(FPCGDataCollection& InputData, USplineComponent* SplineComponent, const FName& PinName = DefaultTestPinName);
+	/** Helper for adding PrimitiveData to an InputData */
+	void AddPrimitiveInputData(FPCGDataCollection& InputData, UPrimitiveComponent* SplineComponent, const FVector& VoxelSize, const FName& PinName = DefaultTestPinName);
+
+	/** Adds randomized PointData with a single Point */
+	void AddRandomizedSinglePointInputData(FTestData& TestData, int32 PointNum = DefaultNumPointsToGenerate, const FName& PinName = DefaultTestPinName);
+	/** Adds randomized PointData with multiple Points */
+	void AddRandomizedMultiplePointInputData(FTestData& TestData, int32 PointNum = DefaultNumPointsToGenerate, const FName& PinName = DefaultTestPinName);
+	/** Adds randomized Volume Spatial Data */
+	void AddRandomizedVolumeInputData(FTestData& TestData, const FName& PinName = DefaultTestPinName);
+	/** Adds randomized Surface Spatial Data */
+	void AddRandomizedSurfaceInputData(FTestData& TestData, const FName& PinName = DefaultTestPinName);
+	/** Adds randomized PolyLine Spatial Data */
+	void AddRandomizedPolyLineInputData(FTestData& TestData, int32 PointNum = DefaultNumPolyLinePointsToGenerate, const FName& PinName = DefaultTestPinName);
+	/** Adds randomized Primitive Spatial Data */
+	void AddRandomizedPrimitiveInputData(FTestData& TestData, const FName& PinName = DefaultTestPinName);
 
 	/** Validates whether two DataCollection objects are identical */
 	bool DataCollectionsAreIdentical(const FPCGDataCollection& FirstCollection, const FPCGDataCollection& SecondCollection);
+
 	/** Validates whether two SpatialData objects are identical */
-	bool SpatialDataIsIdentical(const UPCGData* FirstSpatialData, const UPCGData* SecondSpatialData);
+	bool SpatialDataIsIdentical(const UPCGData* FirstData, const UPCGData* SecondData);
 	/** Validates whether two PointData objects are identical */
-	bool PointDataIsIdentical(const UPCGData* FirstPointData, const UPCGData* SecondPointData);
+	bool PointDataIsIdentical(const UPCGData* FirstData, const UPCGData* SecondData);
+	/** Validates whether two VolumeData objects are identical */
+	bool VolumeDataIsIdentical(const UPCGData* FirstData, const UPCGData* SecondData);
+	/** Validates whether two SurfaceData objects are identical */
+	bool SurfaceDataIsIdentical(const UPCGData* FirstData, const UPCGData* SecondData);
+	/** Validates whether two PolyLineData objects are identical */
+	bool PolyLineDataIsIdentical(const UPCGData* FirstData, const UPCGData* SecondData);
+	/** Validates whether two PrimitiveData objects are identical */
+	bool PrimitiveDataIsIdentical(const UPCGData* FirstData, const UPCGData* SecondData);
+	/** Validates whether two DifferenceData objects are identical */
+	bool DifferenceDataIsIdentical(const UPCGData* FirstData, const UPCGData* SecondData);
+
+	/** Validates the basics of SpatialData are identical */
+	bool SpatialBasicsAreIdentical(const UPCGSpatialData* FirstSpatialData, const UPCGSpatialData* SecondSpatialData);
+
 	/** A catch function for unimplemented comparisons */
-	bool ComparisonIsUnimplemented(const UPCGData* FirstPointData, const UPCGData* SecondPointData);
+	bool ComparisonIsUnimplemented(const UPCGData* FirstData, const UPCGData* SecondData);
 
 	/** Determines if a PCG DataType is a type that should be compared */
 	bool DataIsComparable(EPCGDataType DataType);
@@ -56,20 +95,32 @@ namespace PCGDeterminismTests
 
 	/** Execute the elements for each valid input and compare if all the outputs are identical */
 	bool ExecutionIsDeterministic(FTestData& FirstTestData, FTestData& SecondTestData);
+	/** Execute the same element twice compare if all the outputs are identical */
+	bool ExecutionIsDeterministicSameData(FTestData& TestData);
 
 	/** Generates settings based upon a UPCGSettings subclass */
 	template<typename SettingsType>
-	void GenerateSettings(FTestData& TestData, TFunction<void(SettingsType*, FRandomStream&)> ExtraSettingsDelegate = nullptr)
+	void GenerateSettings(FTestData& TestData, TFunction<void(FTestData& TestData)> ExtraSettingsDelegate = nullptr)
 	{
 		TestData.Settings = NewObject<SettingsType>();
 		check(TestData.Settings);
 		TestData.Settings->Seed = TestData.Seed;
 
 		TestData.InputData.TaggedData.Emplace_GetRef().Data = TestData.Settings;
+		TestData.InputData.TaggedData.Last().Pin = FName(TEXT("Settings"));
 
 		if (ExtraSettingsDelegate)
 		{
-			ExtraSettingsDelegate(Cast<SettingsType>(TestData.Settings), TestData.RandomStream);
+			ExtraSettingsDelegate(TestData);
 		}
 	}
+
+	/** Validates whether both UPCGData can be cast to a specified subclass */
+	template<typename DataType>
+	bool BothDataCastsToDataType(const UPCGData* FirstData, const UPCGData* SecondData)
+	{
+		check(FirstData && SecondData);
+
+		return (Cast<const DataType>(FirstData) != nullptr && Cast<const DataType>(SecondData) != nullptr);
+	}
 }
\ No newline at end of file
