Tentative fix for non-unity build

diff --git a/Engine/Plugins/Experimental/PCG/Source/PCG/Private/PCGSettings.cpp b/Engine/Plugins/Experimental/PCG/Source/PCG/Private/PCGSettings.cpp
--- a/Engine/Plugins/Experimental/PCG/Source/PCG/Private/PCGSettings.cpp
+++ b/Engine/Plugins/Experimental/PCG/Source/PCG/Private/PCGSettings.cpp
@@ -2,6 +2,7 @@
 
 #include "PCGSettings.h"
 #include "PCGNode.h"
+#include "PCGSubsystem.h"
 #include "Serialization/ArchiveObjectCrc32.h"
 
 #if WITH_EDITOR
diff --git a/Engine/Plugins/Experimental/PCG/Source/PCG/Public/Elements/PCGStaticMeshSpawner.h b/Engine/Plugins/Experimental/PCG/Source/PCG/Public/Elements/PCGStaticMeshSpawner.h
--- a/Engine/Plugins/Experimental/PCG/Source/PCG/Public/Elements/PCGStaticMeshSpawner.h
+++ b/Engine/Plugins/Experimental/PCG/Source/PCG/Public/Elements/PCGStaticMeshSpawner.h
@@ -7,6 +7,7 @@
 #include "MeshSelectors/PCGMeshSelectorBase.h"
 
 #include "Engine/CollisionProfile.h"
+#include "Templates/SubclassOf.h"
 
 #include "PCGStaticMeshSpawner.generated.h"
 
diff --git a/Engine/Plugins/Experimental/PCG/Source/PCGEditor/Public/PCGEditorSettings.h b/Engine/Plugins/Experimental/PCG/Source/PCGEditor/Public/PCGEditorSettings.h
--- a/Engine/Plugins/Experimental/PCG/Source/PCGEditor/Public/PCGEditorSettings.h
+++ b/Engine/Plugins/Experimental/PCG/Source/PCGEditor/Public/PCGEditorSettings.h
@@ -3,6 +3,7 @@
 #pragma once
 
 #include "CoreMinimal.h"
+#include "Templates/SubclassOf.h"
 #include "PCGSettings.h"
 
 #include "PCGEditorSettings.generated.h"
