MeshSelector Material Overrides

diff --git a/Engine/Plugins/Experimental/PCG/Source/PCG/Private/Elements/PCGStaticMeshSpawner.cpp b/Engine/Plugins/Experimental/PCG/Source/PCG/Private/Elements/PCGStaticMeshSpawner.cpp
--- a/Engine/Plugins/Experimental/PCG/Source/PCG/Private/Elements/PCGStaticMeshSpawner.cpp
+++ b/Engine/Plugins/Experimental/PCG/Source/PCG/Private/Elements/PCGStaticMeshSpawner.cpp
@@ -37,7 +37,7 @@ bool FPCGStaticMeshSpawnerElement::ExecuteInternal(FPCGContext* Context) const
 	}
 
 	// perform mesh selection
-	TMap<TSoftObjectPtr<UStaticMesh>, FPCGMeshInstanceList> MeshInstances;
+	TArray<FPCGMeshInstanceList> MeshInstances;
 	TArray<FPCGTaggedData> Inputs = Context->InputData.GetInputs();
 
 	for (const FPCGTaggedData& Input : Inputs)
@@ -69,23 +69,20 @@ bool FPCGStaticMeshSpawnerElement::ExecuteInternal(FPCGContext* Context) const
 	return true;
 }
 
-void FPCGStaticMeshSpawnerElement::SpawnStaticMeshInstances(FPCGContext* Context, const TMap<TSoftObjectPtr<UStaticMesh>, FPCGMeshInstanceList>& MeshInstances, AActor* TargetActor) const
+void FPCGStaticMeshSpawnerElement::SpawnStaticMeshInstances(FPCGContext* Context, const TArray<FPCGMeshInstanceList>& MeshInstances, AActor* TargetActor) const
 {
 	// Populate the (H)ISM from the previously prepared entries
 	TRACE_CPUPROFILER_EVENT_SCOPE(FPCGStaticMeshSpawnerElement::Execute::PopulateISMs);
 
-	for (const TPair<TSoftObjectPtr<UStaticMesh>, FPCGMeshInstanceList>& Elem : MeshInstances)
+	for (const FPCGMeshInstanceList& InstanceList : MeshInstances)
 	{
-		const TSoftObjectPtr<UStaticMesh> Mesh = Elem.Key;
-		const FPCGMeshInstanceList& InstanceList = Elem.Value;
-
 		if (InstanceList.Instances.Num() == 0)
 		{
 			continue;
 		}
 
 		// Todo: we could likely pre-load these meshes asynchronously in the settings
-		UStaticMesh* LoadedMesh = Mesh.LoadSynchronous();
+		UStaticMesh* LoadedMesh = InstanceList.Mesh.LoadSynchronous();
 
 		if (!LoadedMesh)
 		{
@@ -94,11 +91,17 @@ void FPCGStaticMeshSpawnerElement::SpawnStaticMeshInstances(FPCGContext* Context
 
 		FPCGISMCBuilderParameters Params;
 		Params.Mesh = LoadedMesh;
+
 		if (InstanceList.bOverrideCollisionProfile)
 		{
 			Params.CollisionProfile = InstanceList.CollisionProfile.Name;
 		}
 
+		if (InstanceList.bOverrideMaterials)
+		{
+			Params.MaterialOverrides = InstanceList.MaterialOverrides;
+		}
+
 		UInstancedStaticMeshComponent* ISMC = UPCGActorHelpers::GetOrCreateISMC(TargetActor, Context->SourceComponent, Params);
 
 		// TODO: add scaling
@@ -107,7 +110,7 @@ void FPCGStaticMeshSpawnerElement::SpawnStaticMeshInstances(FPCGContext* Context
 		ISMC->AddInstances(InstanceList.Instances, false, true);
 		ISMC->UpdateBounds();
 
-		PCGE_LOG(Verbose, "Added %d instances of %s on actor %s", InstanceList.Instances.Num(), *Mesh->GetFName().ToString(), *TargetActor->GetFName().ToString());
+		PCGE_LOG(Verbose, "Added %d instances of %s on actor %s", InstanceList.Instances.Num(), *InstanceList.Mesh->GetFName().ToString(), *TargetActor->GetFName().ToString());
 	}
 }
 
diff --git a/Engine/Plugins/Experimental/PCG/Source/PCG/Private/MeshSelectors/PCGMeshSelectorBase.cpp b/Engine/Plugins/Experimental/PCG/Source/PCG/Private/MeshSelectors/PCGMeshSelectorBase.cpp
--- a/Engine/Plugins/Experimental/PCG/Source/PCG/Private/MeshSelectors/PCGMeshSelectorBase.cpp
+++ b/Engine/Plugins/Experimental/PCG/Source/PCG/Private/MeshSelectors/PCGMeshSelectorBase.cpp
@@ -8,7 +8,43 @@ void UPCGMeshSelectorBase::SelectInstances_Implementation(
 	FPCGContext& Context, 
 	const UPCGStaticMeshSpawnerSettings* Settings, 
 	const UPCGSpatialData* InSpatialData,
-	TMap<TSoftObjectPtr<UStaticMesh>, FPCGMeshInstanceList>& OutMeshInstances) const
+	TArray<FPCGMeshInstanceList>& OutMeshInstances) const
 {
 	PCGE_LOG_C(Error, &Context, "Invalid use of abstract MeshSelectorBase class. Please use an already existing class or implement the CreateMeshInstanceData method");
 }
+
+bool UPCGMeshSelectorBase::FindOrAddInstanceList(
+	TArray<FPCGMeshInstanceList>& OutInstanceLists,
+	const TSoftObjectPtr<UStaticMesh>& Mesh,
+	bool bOverrideCollisionProfile,
+	const FCollisionProfileName& CollisionProfile,
+	bool bOverrideMaterials,
+	const TArray<UMaterialInterface*>& MaterialOverrides,
+	int32& OutIndex) const
+{
+	for (int Index = 0; Index < OutInstanceLists.Num(); ++Index)
+	{
+		if (OutInstanceLists[Index].Mesh != Mesh || OutInstanceLists[Index].bOverrideCollisionProfile != bOverrideCollisionProfile || OutInstanceLists[Index].bOverrideMaterials != bOverrideMaterials)
+		{
+			continue;
+		}
+
+		if (OutInstanceLists[Index].bOverrideCollisionProfile && !OutInstanceLists[Index].CollisionProfile.Name.IsEqual(CollisionProfile.Name))
+		{
+			continue;
+		}
+
+		if (OutInstanceLists[Index].bOverrideMaterials && OutInstanceLists[Index].MaterialOverrides != MaterialOverrides)
+		{
+			continue;
+		}
+
+		OutIndex = Index;
+
+		return false;
+	}
+
+	OutIndex = OutInstanceLists.Emplace(Mesh, bOverrideCollisionProfile, CollisionProfile, bOverrideMaterials, MaterialOverrides);
+
+	return true;
+}
diff --git a/Engine/Plugins/Experimental/PCG/Source/PCG/Private/MeshSelectors/PCGMeshSelectorByAttribute.cpp b/Engine/Plugins/Experimental/PCG/Source/PCG/Private/MeshSelectors/PCGMeshSelectorByAttribute.cpp
--- a/Engine/Plugins/Experimental/PCG/Source/PCG/Private/MeshSelectors/PCGMeshSelectorByAttribute.cpp
+++ b/Engine/Plugins/Experimental/PCG/Source/PCG/Private/MeshSelectors/PCGMeshSelectorByAttribute.cpp
@@ -13,7 +13,7 @@ void UPCGMeshSelectorByAttribute::SelectInstances_Implementation(
 	FPCGContext& Context, 
 	const UPCGStaticMeshSpawnerSettings* Settings, 
 	const UPCGSpatialData* InSpatialData, 
-	TMap<TSoftObjectPtr<UStaticMesh>, FPCGMeshInstanceList>& OutMeshInstances) const
+	TArray<FPCGMeshInstanceList>& OutMeshInstances) const
 {
 	const UPCGPointData* PointData = InSpatialData->ToPointData(&Context);
 
@@ -50,7 +50,7 @@ void UPCGMeshSelectorByAttribute::SelectInstances_Implementation(
 
 	TRACE_CPUPROFILER_EVENT_SCOPE(FPCGStaticMeshSpawnerElement::Execute::SelectEntries);
 
-	// for every point, add an instance of the stored Mesh
+	// Assign points to entries
 	for (const FPCGPoint& Point : PointData->GetPoints()) 
 	{
 		if (Point.Density <= 0.0f)
@@ -81,17 +81,8 @@ void UPCGMeshSelectorByAttribute::SelectInstances_Implementation(
 			continue;
 		}
 
-		// insert the Mesh if it does not exist already
-		FPCGMeshInstanceList* InstanceList = OutMeshInstances.Find(Mesh);
-		if (!InstanceList)
-		{
-			FPCGMeshInstanceList& NewInstanceList = OutMeshInstances.Emplace(Mesh);
-			NewInstanceList.bOverrideCollisionProfile = bOverrideCollisionProfile;
-			NewInstanceList.CollisionProfile = CollisionProfile;
-			
-			InstanceList = &NewInstanceList;
-		}
-
-		InstanceList->Instances.Emplace(Point.Transform);
+		int32 Index = INDEX_NONE;
+		FindOrAddInstanceList(OutMeshInstances, Mesh, bOverrideCollisionProfile, CollisionProfile, bOverrideMaterials, MaterialOverrides, Index);
+		OutMeshInstances[Index].Instances.Emplace(Point.Transform);
 	}
 }
diff --git a/Engine/Plugins/Experimental/PCG/Source/PCG/Private/MeshSelectors/PCGMeshSelectorWeighted.cpp b/Engine/Plugins/Experimental/PCG/Source/PCG/Private/MeshSelectors/PCGMeshSelectorWeighted.cpp
--- a/Engine/Plugins/Experimental/PCG/Source/PCG/Private/MeshSelectors/PCGMeshSelectorWeighted.cpp
+++ b/Engine/Plugins/Experimental/PCG/Source/PCG/Private/MeshSelectors/PCGMeshSelectorWeighted.cpp
@@ -5,6 +5,10 @@
 #include "PCGHelpers.h"
 #include "PCGPoint.h"
 #include "Elements/PCGStaticMeshSpawner.h"
+#include "Metadata/PCGMetadata.h"
+#include "Metadata/PCGMetadataAttribute.h"
+#include "Metadata/PCGMetadataAttributeTraits.h"
+#include "Metadata/PCGMetadataAttributeTpl.h"
 
 #include "Math/RandomStream.h"
 
@@ -12,9 +16,8 @@ void UPCGMeshSelectorWeighted::SelectInstances_Implementation(
 	FPCGContext& Context, 
 	const UPCGStaticMeshSpawnerSettings* Settings, 
 	const UPCGSpatialData* InSpatialData,
-	TMap<TSoftObjectPtr<UStaticMesh>, FPCGMeshInstanceList>& OutMeshInstances) const
+	TArray<FPCGMeshInstanceList>& OutMeshInstances) const
 {
-	TArray<TSoftObjectPtr<UStaticMesh>> Meshes;
 	TArray<int> CumulativeWeights;
 
 	int TotalWeight = 0;
@@ -26,16 +29,11 @@ void UPCGMeshSelectorWeighted::SelectInstances_Implementation(
 			PCGE_LOG_C(Verbose, &Context, "Entry found with weight <= 0");
 			continue;
 		}
-	
-		// TODO: handle collisions between different collision profiles
-		FPCGMeshInstanceList InstanceList;
-		InstanceList.bOverrideCollisionProfile = Entry.bOverrideCollisionProfile;
-		InstanceList.CollisionProfile = Entry.CollisionProfile;
-		OutMeshInstances.Emplace(Entry.Mesh, InstanceList);
 
+		int32 Index = INDEX_NONE;
+		FindOrAddInstanceList(OutMeshInstances, Entry.Mesh, Entry.bOverrideCollisionProfile, Entry.CollisionProfile, Entry.bOverrideMaterials, Entry.MaterialOverrides, Index);
 		TotalWeight += Entry.Weight;
 		CumulativeWeights.Add(TotalWeight);
-		Meshes.Add(Entry.Mesh);
 	}
 
 	if (TotalWeight <= 0)
@@ -66,14 +64,14 @@ void UPCGMeshSelectorWeighted::SelectInstances_Implementation(
 			int RandomWeightedPick = RandomSource.RandRange(0, TotalWeight - 1);
 
 			int RandomPick = 0;
-			while (RandomPick < Meshes.Num() && CumulativeWeights[RandomPick] <= RandomWeightedPick)
+			while (RandomPick < OutMeshInstances.Num() && CumulativeWeights[RandomPick] <= RandomWeightedPick)
 			{
 				++RandomPick;
 			}
 
-			if (RandomPick < Meshes.Num())
+			if (RandomPick < OutMeshInstances.Num())
 			{
-				OutMeshInstances[Meshes[RandomPick]].Instances.Emplace(Point.Transform);
+				OutMeshInstances[RandomPick].Instances.Emplace(Point.Transform);
 			}
 		}
 	}
diff --git a/Engine/Plugins/Experimental/PCG/Source/PCG/Private/MeshSelectors/PCGMeshSelectorWeightedByCategory.cpp b/Engine/Plugins/Experimental/PCG/Source/PCG/Private/MeshSelectors/PCGMeshSelectorWeightedByCategory.cpp
--- a/Engine/Plugins/Experimental/PCG/Source/PCG/Private/MeshSelectors/PCGMeshSelectorWeightedByCategory.cpp
+++ b/Engine/Plugins/Experimental/PCG/Source/PCG/Private/MeshSelectors/PCGMeshSelectorWeightedByCategory.cpp
@@ -12,17 +12,17 @@
 
 #include "Math/RandomStream.h"
 
-struct FPCGMeshesAndWeights
+struct FPCGInstancesAndWeights
 {
-	TArray<TSoftObjectPtr<UStaticMesh>> Meshes;
+	TArray<int> InstanceListIndices;
 	TArray<int> CumulativeWeights;
 };
 
 void UPCGMeshSelectorWeightedByCategory::SelectInstances_Implementation(
-	FPCGContext& Context, 
-	const UPCGStaticMeshSpawnerSettings* Settings, 
+	FPCGContext& Context,
+	const UPCGStaticMeshSpawnerSettings* Settings,
 	const UPCGSpatialData* InSpatialData,
-	TMap<TSoftObjectPtr<UStaticMesh>, FPCGMeshInstanceList>& OutMeshInstances) const
+	TArray<FPCGMeshInstanceList>& OutMeshInstances) const
 {
 	const UPCGPointData* PointData = InSpatialData->ToPointData(&Context);
 
@@ -57,7 +57,7 @@ void UPCGMeshSelectorWeightedByCategory::SelectInstances_Implementation(
 	const FPCGMetadataAttribute<FString>* Attribute = static_cast<const FPCGMetadataAttribute<FString>*>(AttributeBase);
 
 	// maps a CategoryEntry ValueKey to the meshes and precomputed weight data 
-	TMap<PCGMetadataValueKey, FPCGMeshesAndWeights> CategoryEntryToMeshesAndWeights;
+	TMap<PCGMetadataValueKey, FPCGInstancesAndWeights> CategoryEntryToInstancesAndWeights;
 
 	// unmarked points will fallback to the MeshEntries associated with the DefaultValueKey
 	PCGMetadataValueKey DefaultValueKey = PCGDefaultValueKey;
@@ -78,9 +78,9 @@ void UPCGMeshSelectorWeightedByCategory::SelectInstances_Implementation(
 			continue;
 		}
 
-		FPCGMeshesAndWeights* MeshesAndWeights = CategoryEntryToMeshesAndWeights.Find(ValueKey);
+		FPCGInstancesAndWeights* InstancesAndWeights = CategoryEntryToInstancesAndWeights.Find(ValueKey);
 
-		if (MeshesAndWeights)
+		if (InstancesAndWeights)
 		{
 			PCGE_LOG_C(Warning, &Context, "Duplicate entry found in category %s. Subsequent entries are ignored.", *Entry.CategoryEntry);
 			continue;
@@ -98,7 +98,7 @@ void UPCGMeshSelectorWeightedByCategory::SelectInstances_Implementation(
 			}
 		}
 
-		MeshesAndWeights = &CategoryEntryToMeshesAndWeights.Add(ValueKey, FPCGMeshesAndWeights());
+		InstancesAndWeights = &CategoryEntryToInstancesAndWeights.Add(ValueKey, FPCGInstancesAndWeights());
 
 		int TotalWeight = 0;
 		for (const FPCGMeshSelectorWeightedEntry& WeightedEntry : Entry.WeightedMeshEntries)
@@ -108,21 +108,20 @@ void UPCGMeshSelectorWeightedByCategory::SelectInstances_Implementation(
 				PCGE_LOG_C(Verbose, &Context, "Entry found with weight <= 0 in category %s", *Entry.CategoryEntry);
 				continue;
 			}
-
-			// create or overwrite the entry for this mesh
-			FPCGMeshInstanceList& InstanceList = OutMeshInstances.Emplace(WeightedEntry.Mesh);
-			InstanceList.bOverrideCollisionProfile = WeightedEntry.bOverrideCollisionProfile;
-			InstanceList.CollisionProfile = WeightedEntry.CollisionProfile;
+			
+			int32 Index = INDEX_NONE;
+			FindOrAddInstanceList(OutMeshInstances, WeightedEntry.Mesh, WeightedEntry.bOverrideCollisionProfile, WeightedEntry.CollisionProfile, WeightedEntry.bOverrideMaterials, WeightedEntry.MaterialOverrides, Index);
 
 			// precompute the weights
 			TotalWeight += WeightedEntry.Weight;
-			MeshesAndWeights->CumulativeWeights.Add(TotalWeight);
-			MeshesAndWeights->Meshes.Add(WeightedEntry.Mesh);
+			InstancesAndWeights->CumulativeWeights.Add(TotalWeight);
+			InstancesAndWeights->InstanceListIndices.Add(Index);
 		}
 	}
 
 	TRACE_CPUPROFILER_EVENT_SCOPE(FPCGStaticMeshSpawnerElement::Execute::SelectEntries);
 
+	// Assign points to entries
 	for (const FPCGPoint& Point : PointData->GetPoints())
 	{
 		if (Point.Density <= 0.0f)
@@ -133,13 +132,13 @@ void UPCGMeshSelectorWeightedByCategory::SelectInstances_Implementation(
 		PCGMetadataValueKey ValueKey = Attribute->GetValueKey(Point.MetadataEntry);
 
 		// if no mesh list was processed for this attribute value, fallback to the default mesh list
-		FPCGMeshesAndWeights* MeshesAndWeights = CategoryEntryToMeshesAndWeights.Find(ValueKey);
-		if (!MeshesAndWeights)
+		FPCGInstancesAndWeights* InstancesAndWeights = CategoryEntryToInstancesAndWeights.Find(ValueKey);
+		if (!InstancesAndWeights)
 		{
 			if (DefaultValueKey != PCGDefaultValueKey)
 			{
-				MeshesAndWeights = CategoryEntryToMeshesAndWeights.Find(DefaultValueKey);
-				check(MeshesAndWeights);
+				InstancesAndWeights = CategoryEntryToInstancesAndWeights.Find(DefaultValueKey);
+				check(InstancesAndWeights);
 			}
 			else
 			{
@@ -147,20 +146,21 @@ void UPCGMeshSelectorWeightedByCategory::SelectInstances_Implementation(
 			}
 		}
 
-		const int TotalWeight = MeshesAndWeights->CumulativeWeights.Last();
+		const int TotalWeight = InstancesAndWeights->CumulativeWeights.Last();
 
 		FRandomStream RandomSource(PCGHelpers::ComputeSeed(Point.Seed, Settings->Seed));
 		int RandomWeightedPick = RandomSource.RandRange(0, TotalWeight - 1);
 
 		int RandomPick = 0;
-		while (RandomPick < MeshesAndWeights->Meshes.Num() && MeshesAndWeights->CumulativeWeights[RandomPick] <= RandomWeightedPick)
+		while (RandomPick < InstancesAndWeights->InstanceListIndices.Num() && InstancesAndWeights->CumulativeWeights[RandomPick] <= RandomWeightedPick)
 		{
 			++RandomPick;
 		}
 
-		if (RandomPick < MeshesAndWeights->Meshes.Num())
+		if (RandomPick < InstancesAndWeights->InstanceListIndices.Num())
 		{
-			OutMeshInstances[MeshesAndWeights->Meshes[RandomPick]].Instances.Emplace(Point.Transform);
+			const int32 Index = InstancesAndWeights->InstanceListIndices[RandomPick];
+			OutMeshInstances[Index].Instances.Emplace(Point.Transform);
 		}
 	}
 }
diff --git a/Engine/Plugins/Experimental/PCG/Source/PCG/Public/Elements/PCGStaticMeshSpawner.h b/Engine/Plugins/Experimental/PCG/Source/PCG/Public/Elements/PCGStaticMeshSpawner.h
--- a/Engine/Plugins/Experimental/PCG/Source/PCG/Public/Elements/PCGStaticMeshSpawner.h
+++ b/Engine/Plugins/Experimental/PCG/Source/PCG/Public/Elements/PCGStaticMeshSpawner.h
@@ -84,6 +84,6 @@ public:
 
 protected:
 	virtual bool ExecuteInternal(FPCGContext* Context) const override;	
-	void SpawnStaticMeshInstances(FPCGContext* Context, const TMap<TSoftObjectPtr<UStaticMesh>, FPCGMeshInstanceList>& MeshInstances, AActor* TargetActor) const;
+	void SpawnStaticMeshInstances(FPCGContext* Context, const TArray<FPCGMeshInstanceList>& MeshInstances, AActor* TargetActor) const;
 };
 
diff --git a/Engine/Plugins/Experimental/PCG/Source/PCG/Public/MeshSelectors/PCGMeshSelectorBase.h b/Engine/Plugins/Experimental/PCG/Source/PCG/Public/MeshSelectors/PCGMeshSelectorBase.h
--- a/Engine/Plugins/Experimental/PCG/Source/PCG/Public/MeshSelectors/PCGMeshSelectorBase.h
+++ b/Engine/Plugins/Experimental/PCG/Source/PCG/Public/MeshSelectors/PCGMeshSelectorBase.h
@@ -12,20 +12,36 @@
 #include "PCGMeshSelectorBase.generated.h"
 
 class UPCGStaticMeshSpawnerSettings;
+class UMaterialInterface;
 
 USTRUCT(BlueprintType)
 struct FPCGMeshInstanceList
 {
 	GENERATED_BODY()
 
+	FPCGMeshInstanceList() = default;
+	
+	FPCGMeshInstanceList(const TSoftObjectPtr<UStaticMesh>& InMesh, bool bInOverrideCollisionProfile, const FCollisionProfileName& InCollisionProfile, bool bInOverrideMaterials, const TArray<UMaterialInterface*>& InMaterialOverrides)
+		: Mesh(InMesh), bOverrideCollisionProfile(bInOverrideCollisionProfile), CollisionProfile(InCollisionProfile), bOverrideMaterials(bInOverrideMaterials), MaterialOverrides(InMaterialOverrides)
+	{}
+
 	UPROPERTY(BlueprintReadWrite, EditAnywhere, Category = Settings)
-    TArray<FTransform> Instances;
+	TSoftObjectPtr<UStaticMesh> Mesh;
 
 	UPROPERTY(BlueprintReadWrite, EditAnywhere, Category = Settings)
 	bool bOverrideCollisionProfile = false;
 
 	UPROPERTY(BlueprintReadWrite, EditAnywhere, Category = Settings)
 	FCollisionProfileName CollisionProfile;
+
+	UPROPERTY(BlueprintReadWrite, EditAnywhere, Category = Settings)
+	bool bOverrideMaterials = false;
+
+	UPROPERTY(BlueprintReadWrite, EditAnywhere, Category = Settings)
+	TArray<UMaterialInterface*> MaterialOverrides;
+
+	UPROPERTY(BlueprintReadWrite, EditAnywhere, Category = Settings)
+    TArray<FTransform> Instances;
 };
 
 UCLASS(Abstract, BlueprintType, Blueprintable, ClassGroup = (Procedural))
@@ -35,15 +51,26 @@ class PCG_API UPCGMeshSelectorBase : public UObject
 
 public:
 	UFUNCTION(BlueprintNativeEvent, Category = MeshSelection)
-		void SelectInstances(
+	void SelectInstances(
 		FPCGContext& Context,
 		const UPCGStaticMeshSpawnerSettings* Settings,
 		const UPCGSpatialData* InSpatialData,
-		TMap<TSoftObjectPtr<UStaticMesh>, FPCGMeshInstanceList>& OutMeshInstances) const;
+		TArray<FPCGMeshInstanceList>& OutMeshInstances) const;
 
 	virtual void SelectInstances_Implementation(
 		FPCGContext& Context,
 		const UPCGStaticMeshSpawnerSettings* Settings,
 		const UPCGSpatialData* InSpatialData,
-		TMap<TSoftObjectPtr<UStaticMesh>, FPCGMeshInstanceList>& OutMeshInstances) const;
+		TArray<FPCGMeshInstanceList>& OutMeshInstances) const;
+
+	/** Searches OutInstanceLists for an InstanceList matching the given parameters. If nothing is found, creates a new InstanceList and adds to OutInstanceLists. Returns true if added. */
+	UFUNCTION(BlueprintCallable, Category = MeshSelection)
+	bool FindOrAddInstanceList(
+		TArray<FPCGMeshInstanceList>& OutInstanceLists,
+		const TSoftObjectPtr<UStaticMesh>& Mesh,
+		bool bOverrideCollisionProfile,
+		const FCollisionProfileName& CollisionProfile,
+		bool bOverrideMaterials,
+		const TArray<UMaterialInterface*>& MaterialOverrides,
+		int32& OutIndex) const;
 };
diff --git a/Engine/Plugins/Experimental/PCG/Source/PCG/Public/MeshSelectors/PCGMeshSelectorByAttribute.h b/Engine/Plugins/Experimental/PCG/Source/PCG/Public/MeshSelectors/PCGMeshSelectorByAttribute.h
--- a/Engine/Plugins/Experimental/PCG/Source/PCG/Public/MeshSelectors/PCGMeshSelectorByAttribute.h
+++ b/Engine/Plugins/Experimental/PCG/Source/PCG/Public/MeshSelectors/PCGMeshSelectorByAttribute.h
@@ -16,7 +16,7 @@ public:
 		UPARAM(ref) FPCGContext& Context,
 		const UPCGStaticMeshSpawnerSettings* Settings,
 		const UPCGSpatialData* InSpatialData,
-		TMap<TSoftObjectPtr<UStaticMesh>, FPCGMeshInstanceList>& OutMeshInstances) const override;
+		TArray<FPCGMeshInstanceList>& OutMeshInstances) const override;
 
 public:
 	UPROPERTY(BlueprintReadWrite, EditAnywhere, Category = Settings)
@@ -27,5 +27,11 @@ public:
 
 	UPROPERTY(BlueprintReadWrite, EditAnywhere, Category = Settings)
 	FCollisionProfileName CollisionProfile;
+
+	UPROPERTY(BlueprintReadWrite, EditAnywhere, Category = Settings)
+	bool bOverrideMaterials = false;
+
+	UPROPERTY(BlueprintReadWrite, EditAnywhere, Category = Settings)
+	TArray<UMaterialInterface*> MaterialOverrides;
 };
 
diff --git a/Engine/Plugins/Experimental/PCG/Source/PCG/Public/MeshSelectors/PCGMeshSelectorWeighted.h b/Engine/Plugins/Experimental/PCG/Source/PCG/Public/MeshSelectors/PCGMeshSelectorWeighted.h
--- a/Engine/Plugins/Experimental/PCG/Source/PCG/Public/MeshSelectors/PCGMeshSelectorWeighted.h
+++ b/Engine/Plugins/Experimental/PCG/Source/PCG/Public/MeshSelectors/PCGMeshSelectorWeighted.h
@@ -20,11 +20,17 @@ struct PCG_API FPCGMeshSelectorWeightedEntry
 	UPROPERTY(BlueprintReadWrite, EditAnywhere, Category = Settings)
 	TSoftObjectPtr<UStaticMesh> Mesh;
 
+	UPROPERTY(BlueprintReadWrite, EditAnywhere, Category = Settings)
+	bool bOverrideCollisionProfile = false;
+
 	UPROPERTY(BlueprintReadWrite, EditAnywhere, Category = Settings)
 	FCollisionProfileName CollisionProfile;
 
 	UPROPERTY(BlueprintReadWrite, EditAnywhere, Category = Settings)
-	bool bOverrideCollisionProfile = false;
+	bool bOverrideMaterials = false;
+
+	UPROPERTY(BlueprintReadWrite, EditAnywhere, Category = Settings)
+	TArray<UMaterialInterface*> MaterialOverrides;
 
 	UPROPERTY(BlueprintReadWrite, EditAnywhere, Category = Settings, meta = (ClampMin = "0"))
 	int Weight = 1;
@@ -40,7 +46,7 @@ public:
 		UPARAM(ref) FPCGContext& Context, 
 		const UPCGStaticMeshSpawnerSettings* Settings, 
 		const UPCGSpatialData* InSpatialData,
-		TMap<TSoftObjectPtr<UStaticMesh>, FPCGMeshInstanceList>& OutMeshInstances) const override;
+		TArray<FPCGMeshInstanceList>& OutMeshInstances) const override;
 
 public:
 	UPROPERTY(BlueprintReadWrite, EditAnywhere, Category = Settings)
diff --git a/Engine/Plugins/Experimental/PCG/Source/PCG/Public/MeshSelectors/PCGMeshSelectorWeightedByCategory.h b/Engine/Plugins/Experimental/PCG/Source/PCG/Public/MeshSelectors/PCGMeshSelectorWeightedByCategory.h
--- a/Engine/Plugins/Experimental/PCG/Source/PCG/Public/MeshSelectors/PCGMeshSelectorWeightedByCategory.h
+++ b/Engine/Plugins/Experimental/PCG/Source/PCG/Public/MeshSelectors/PCGMeshSelectorWeightedByCategory.h
@@ -38,7 +38,7 @@ public:
 		UPARAM(ref) FPCGContext& Context,
 		const UPCGStaticMeshSpawnerSettings* Settings,
 		const UPCGSpatialData* InSpatialData,
-		TMap<TSoftObjectPtr<UStaticMesh>, FPCGMeshInstanceList>& OutMeshInstances) const override;
+		TArray<FPCGMeshInstanceList>& OutMeshInstances) const override;
 
 public:
 	UPROPERTY(BlueprintReadWrite, EditAnywhere, Category = Settings)
