Add support for texture tiling/rotation

diff --git a/Engine/Plugins/Experimental/PCG/Source/PCG/Private/Data/PCGTextureData.cpp b/Engine/Plugins/Experimental/PCG/Source/PCG/Private/Data/PCGTextureData.cpp
--- a/Engine/Plugins/Experimental/PCG/Source/PCG/Private/Data/PCGTextureData.cpp
+++ b/Engine/Plugins/Experimental/PCG/Source/PCG/Private/Data/PCGTextureData.cpp
@@ -9,7 +9,7 @@
 namespace PCGTextureSampling
 {
 	template<typename ValueType>
-	ValueType Sample(const FVector2D& InPosition, const FBox2D& InSurface, int32 Width, int32 Height, TFunctionRef<ValueType(int32 Index)> SamplingFunction)
+	ValueType Sample(const FVector2D& InPosition, const FBox2D& InSurface, const UPCGBaseTextureData* InTextureData, int32 Width, int32 Height, TFunctionRef<ValueType(int32 Index)> SamplingFunction)
 	{
 		check(Width > 0 && Height > 0);
 		if (Width <= 0 || Height <= 0)
@@ -17,9 +17,41 @@ namespace PCGTextureSampling
 			return ValueType{};
 		}
 
-		// TODO: There seems to be a bias issue here, as the bounds size are not in the same space as the texels.
-		// Implementation note: Supports only stretch fit
-		FVector2D Pos = (InPosition - InSurface.Min) * FVector2D(Width, Height) / InSurface.GetSize();
+		check(InTextureData);
+
+		FVector2D Pos = FVector2D::ZeroVector;
+		if (InTextureData->bStretchToFit)
+		{
+			// TODO: There seems to be a bias issue here, as the bounds size are not in the same space as the texels.
+			Pos = (InPosition - InSurface.Min) * FVector2D(Width, Height) / InSurface.GetSize();
+		}
+		else
+		{
+			const float TextureSizeX = InTextureData->TexelSize * Width;
+			const float TextureSizeY = InTextureData->TexelSize * Height;
+
+			const FRotator Rotation = FRotator(0.f, -InTextureData->Rotation, 0.f);
+			const FVector Translation = FVector(-InTextureData->XOffset * TextureSizeX, -InTextureData->YOffset * TextureSizeY, 0.f);
+			const FTransform Transform(Rotation, Translation);
+			const FVector2D SamplePosition = FVector2D(Transform.TransformPosition(FVector(InPosition, 0.f)));
+
+			float X = FMath::Clamp(static_cast<float>(FMath::Fmod(SamplePosition.X, TextureSizeX)) / TextureSizeX, -1.f, 1.f);
+			if (X < 0.f)
+			{
+				X = 1.f + X;
+			}
+
+			float Y = FMath::Clamp(static_cast<float>(FMath::Fmod(SamplePosition.Y, TextureSizeY)) / TextureSizeY, -1.f, 1.f);
+			if (Y < 0.f)
+			{
+				Y = 1.f + Y;
+			}
+
+			X *= Width;
+			Y *= Height;
+
+			Pos = FVector2D(X, Y);
+		}
 
 		// TODO: this isn't super robust, if that becomes an issue
 		int32 X0 = FMath::FloorToInt(Pos.X);
@@ -104,7 +136,7 @@ bool UPCGBaseTextureData::SamplePoint(const FTransform& InTransform, const FBox&
 	FVector2D Position2D(PointPositionInLocalSpace.X, PointPositionInLocalSpace.Y);
 	FBox2D Surface(FVector2D(-1.0f, -1.0f), FVector2D(1.0f, 1.0f));
 
-	FLinearColor Color = PCGTextureSampling::Sample<FLinearColor>(Position2D, Surface, Width, Height, [this](int32 Index) { return ColorData[Index]; });
+	FLinearColor Color = PCGTextureSampling::Sample<FLinearColor>(Position2D, Surface, this, Width, Height, [this](int32 Index) { return ColorData[Index]; });
 	OutPoint.Color = Color;
 	OutPoint.Density = ((DensityFunction == EPCGTextureDensityFunction::Ignore) ? 1.0f : PCGTextureSampling::SampleFloatChannel(Color, ColorChannel));
 
@@ -191,12 +223,10 @@ void UPCGTextureData::Initialize(UTexture2D* InTexture, const FTransform& InTran
 			Width = Texture->GetSizeX();
 			Height = Texture->GetSizeY();
 
-			UTexture2D* TempTexture2D = DuplicateObject<UTexture2D>(InTexture, GetTransientPackage());
-			TempTexture2D->CompressionSettings = TextureCompressionSettings::TC_VectorDisplacementmap;
-			TempTexture2D->MipGenSettings = TextureMipGenSettings::TMGS_NoMipmaps;
-			TempTexture2D->SRGB = false;
-			TempTexture2D->UpdateResource();
-			const FColor* FormattedImageData = static_cast<const FColor*>(TempTexture2D->GetPlatformData()->Mips[0].BulkData.LockReadOnly());
+			// TODO: previously this code created a duplicate of InTexture and read the ColorData from that source 
+			// however, there was a problem with the duplicate not having texture data in Mips[0], so we've removed the duplication for now
+			// original code can be found in ProceduralTextureReader.h
+			const FColor* FormattedImageData = static_cast<const FColor*>(InTexture->GetPlatformData()->Mips[0].BulkData.LockReadOnly());
 
 			ColorData.SetNum(Width * Height);
 			for (int32 D = 0; D < Width * Height; ++D)
@@ -204,7 +234,7 @@ void UPCGTextureData::Initialize(UTexture2D* InTexture, const FTransform& InTran
 				ColorData[D] = FormattedImageData[D].ReinterpretAsLinear();
 			}
 
-			TempTexture2D->GetPlatformData()->Mips[0].BulkData.Unlock();
+			InTexture->GetPlatformData()->Mips[0].BulkData.Unlock();
 		}
 #endif
 	}
diff --git a/Engine/Plugins/Experimental/PCG/Source/PCG/Private/Tests/Data/PCGTextureDataTest.cpp b/Engine/Plugins/Experimental/PCG/Source/PCG/Private/Tests/Data/PCGTextureDataTest.cpp
new file mode 100644
--- /dev/null
+++ b/Engine/Plugins/Experimental/PCG/Source/PCG/Private/Tests/Data/PCGTextureDataTest.cpp
@@ -0,0 +1,119 @@
+// Copyright Epic Games, Inc. All Rights Reserved.
+
+#include "Tests/PCGTestsCommon.h"
+
+#include "Data/PCGTextureData.h"
+#include "Data/PCGRenderTargetData.h"
+
+#include "Engine/Texture.h"
+#include "Engine/Texture2D.h"
+
+#if WITH_EDITOR
+
+IMPLEMENT_CUSTOM_SIMPLE_AUTOMATION_TEST(FPCGTextureDataOffsetTilingRotation, FPCGTestBaseClass, "pcg.tests.Texture.OffsetTilingRotation", PCGTestsCommon::TestFlags)
+
+bool FPCGTextureDataOffsetTilingRotation::RunTest(const FString& Parameters)
+{
+	const int32 TextureSize = 128;
+	const int32 WhitePixelX = 50;
+	const int32 WhitePixelY = 70;
+
+	TArray<FColor> Pixels;
+	Pixels.Init(FColor::Black, TextureSize * TextureSize);
+	Pixels[WhitePixelY * TextureSize + WhitePixelX] = FColor::White;
+
+	UTexture2D* Texture2D = UTexture2D::CreateTransient(TextureSize, TextureSize, EPixelFormat::PF_R8G8B8A8);
+	Texture2D->CompressionSettings = TextureCompressionSettings::TC_VectorDisplacementmap;
+	Texture2D->SRGB = 0;
+	Texture2D->MipGenSettings = TMGS_NoMipmaps;
+	Texture2D->UpdateResource();
+
+	void* RawTextureData = Texture2D->GetPlatformData()->Mips[0].BulkData.Lock(LOCK_READ_WRITE);
+
+	FMemory::Memcpy(RawTextureData, Pixels.GetData(), Pixels.Num() * sizeof(FColor));
+
+	Texture2D->GetPlatformData()->Mips[0].BulkData.Unlock();
+	Texture2D->UpdateResource();
+
+	UPCGTextureData* TextureData = NewObject<UPCGTextureData>();
+	TextureData->Initialize(Texture2D, FTransform());
+
+	TextureData->bStretchToFit = false;
+
+	FRandomStream RandomStream;
+	FPCGPoint OutPoint;
+
+	const float ScaledPixelX = static_cast<float>(WhitePixelX) / TextureSize;
+	const float ScaledPixelY = static_cast<float>(WhitePixelY) / TextureSize;
+	const float TexelScalar = 1.f / 3.f;
+
+	for (int TexelSizeFactor = 0; TexelSizeFactor <= 3; ++TexelSizeFactor)
+	{
+		const float TexelSize = FMath::Pow(10.f, TexelSizeFactor);
+		const float ScaledTexelSize = TexelSize * TexelScalar;
+
+		TextureData->TexelSize = TexelSize;
+		TextureData->Rotation = 0.f;
+	
+		// sampling with no offset
+		TextureData->XOffset = 0.f;
+		TextureData->YOffset = 0.f;
+		TextureData->SamplePoint(FTransform(), FBox(), OutPoint, nullptr);
+		TestEqual("Valid color sampled for no offset at (0, 0)", OutPoint.Color, static_cast<FVector>(FColor::Black));
+
+		// sampling at position, with no offset
+		TextureData->XOffset = 0.f;
+		TextureData->YOffset = 0.f;
+		TextureData->SamplePoint(FTransform(FVector(ScaledTexelSize, ScaledTexelSize, 0.f)), FBox(), OutPoint, nullptr);
+		TestEqual("Valid color sampled for no offset at position", OutPoint.Color, static_cast<FVector>(FColor::Black));
+	
+		// sampling with offset
+		TextureData->XOffset = 1 - ScaledPixelX;
+		TextureData->YOffset = 1 - ScaledPixelY;
+		TextureData->SamplePoint(FTransform(), FBox(), OutPoint, nullptr);
+		TestEqual("Valid color sampled for offset from (0, 0)", OutPoint.Color, static_cast<FVector>(FColor::White));
+
+		// sampling at position, with offset
+		TextureData->XOffset = (TexelScalar - WhitePixelX) / TextureSize;
+		TextureData->YOffset = (TexelScalar - WhitePixelY) / TextureSize;
+		TextureData->SamplePoint(FTransform(FVector(ScaledTexelSize, ScaledTexelSize, 0.f)), FBox(), OutPoint, nullptr);
+		TestEqual("Valid color sampled for offset from position", OutPoint.Color, static_cast<FVector>(FColor::White));
+	
+		for (float Rotation = 0.f; Rotation < 360.f; Rotation += 10.f)
+		{
+			const float Theta = FMath::DegreesToRadians(Rotation);
+			const float CosTheta = FMath::Cos(Theta);
+			const float SinTheta = FMath::Sin(Theta);
+			const float X = ((WhitePixelX * CosTheta) - (WhitePixelY * SinTheta)) * TexelSize;
+			const float Y = ((WhitePixelY * CosTheta) + (WhitePixelX * SinTheta)) * TexelSize;
+
+			FTransform RotatedTransform(FVector(X, Y, 0.f));
+
+			// sampling with rotation at black position
+			TextureData->Rotation = Rotation;
+			TextureData->XOffset = 0.f; 
+			TextureData->YOffset = 0.f; 
+			TextureData->SamplePoint(FTransform(FVector(ScaledTexelSize, ScaledTexelSize, 0.f)), FBox(), OutPoint, nullptr);
+			TestEqual("Valid color sampled for non-rotated position", OutPoint.Color, static_cast<FVector>(FColor::Black));
+	
+			// sampling with rotation at white position
+			TextureData->Rotation = Rotation;
+			TextureData->XOffset = 0.f; 
+			TextureData->YOffset = 0.f; 
+			TextureData->SamplePoint(RotatedTransform, FBox(), OutPoint, nullptr);
+			TestEqual("Valid color sampled for rotated position", OutPoint.Color, static_cast<FVector>(FColor::White));
+			
+			RotatedTransform.SetLocation(FVector(X / 2, Y / 2, 0.f));
+	
+			// sampling with rotation and offset
+			TextureData->Rotation = Rotation;
+			TextureData->XOffset = -ScaledPixelX / 2;
+			TextureData->YOffset = -ScaledPixelY / 2;
+			TextureData->SamplePoint(RotatedTransform, FBox(), OutPoint, nullptr);
+			TestEqual("Valid color sampled for offset and rotation", OutPoint.Color, static_cast<FVector>(FColor::White));
+		}
+	}
+
+	return true;
+}
+#endif
diff --git a/Engine/Plugins/Experimental/PCG/Source/PCG/Public/Data/PCGTextureData.h b/Engine/Plugins/Experimental/PCG/Source/PCG/Public/Data/PCGTextureData.h
--- a/Engine/Plugins/Experimental/PCG/Source/PCG/Public/Data/PCGTextureData.h
+++ b/Engine/Plugins/Experimental/PCG/Source/PCG/Public/Data/PCGTextureData.h
@@ -53,6 +53,26 @@ public:
 	UPROPERTY(BlueprintReadWrite, EditAnywhere, Category = Settings)
 	EPCGTextureColorChannel ColorChannel = EPCGTextureColorChannel::Alpha;
 
+	/** Whether to tile the source or to stretch it to fit target area. */
+	UPROPERTY(BlueprintReadWrite, EditAnywhere, Category = Settings)
+	bool bStretchToFit = true;
+
+	/** The size of one texel in cm. */
+	UPROPERTY(BlueprintReadWrite, EditAnywhere, Category = Settings, meta = (UIMin = "1.0", ClampMin = "1.0", EditCondition = "!bStretchToFit"))
+	float TexelSize = 50;
+
+	/** X Offset factor to apply when sampling texture (Factor of texture dimension). */
+	UPROPERTY(BlueprintReadWrite, EditAnywhere, Category = Settings, meta = (UIMin = "0.0", ClampMin = "0.0", UIMax = "1.0", ClampMax = "1.0", DisplayName = "X Offset", EditCondition = "!bStretchToFit"))
+	float XOffset = 0;
+
+	/** Y Offset factor to apply when sampling texture (Factor of texture dimension). */
+	UPROPERTY(BlueprintReadWrite, EditAnywhere, Category = Settings, meta = (UIMin = "0.0", ClampMin = "0.0", UIMax = "1.0", ClampMax = "1.0", DisplayName = "Y Offset", EditCondition = "!bStretchToFit"))
+	float YOffset = 0;
+
+	/** Rotation to apply when sampling texture. */
+	UPROPERTY(BlueprintReadWrite, EditAnywhere, Category = Settings, meta = (UIMin = 0, ClampMin = 0, UIMax = 359, ClampMax = 359, Units = deg, EditCondition = "!bStretchToFit"))
+	float Rotation = 0;
+
 protected:
 	UPROPERTY()
 	TArray<FLinearColor> ColorData;
