Tentative fix non-unity build PCGDeterminismTestsCommon.cpp

#rnx

diff --git a/Engine/Plugins/Experimental/PCG/Source/PCG/Private/Tests/Determinism/PCGDeterminismTestsCommon.cpp b/Engine/Plugins/Experimental/PCG/Source/PCG/Private/Tests/Determinism/PCGDeterminismTestsCommon.cpp
--- a/Engine/Plugins/Experimental/PCG/Source/PCG/Private/Tests/Determinism/PCGDeterminismTestsCommon.cpp
+++ b/Engine/Plugins/Experimental/PCG/Source/PCG/Private/Tests/Determinism/PCGDeterminismTestsCommon.cpp
@@ -16,6 +16,7 @@
 
 #include "GameFramework/Actor.h"
 #include "Components/SplineComponent.h"
+#include "Components/StaticMeshComponent.h"
 
 #if WITH_EDITOR
 #include "Editor.h"
