Add optional output to static mesh spawner

diff --git a/Engine/Plugins/Experimental/PCG/Source/PCG/Private/Elements/PCGStaticMeshSpawner.cpp b/Engine/Plugins/Experimental/PCG/Source/PCG/Private/Elements/PCGStaticMeshSpawner.cpp
--- a/Engine/Plugins/Experimental/PCG/Source/PCG/Private/Elements/PCGStaticMeshSpawner.cpp
+++ b/Engine/Plugins/Experimental/PCG/Source/PCG/Private/Elements/PCGStaticMeshSpawner.cpp
@@ -2,6 +2,7 @@
 
 #include "Elements/PCGStaticMeshSpawner.h"
 
+#include "PCGCommon.h"
 #include "PCGHelpers.h"
 #include "Data/PCGPointData.h"
 #include "Data/PCGSpatialData.h"
@@ -39,6 +40,12 @@ bool FPCGStaticMeshSpawnerElement::ExecuteInternal(FPCGContext* Context) const
 	// perform mesh selection
 	TArray<FPCGMeshInstanceList> MeshInstances;
 	TArray<FPCGTaggedData> Inputs = Context->InputData.GetInputs();
+	TArray<FPCGTaggedData>& Outputs = Context->OutputData.TaggedData;
+
+	// Forward any non-input data
+	Outputs.Append(Context->InputData.GetAllSettings());
+
+	const bool bOutputPinConnected = Context->Node && Context->Node->IsOutputPinConnected(PCGPinConstants::DefaultOutputLabel);
 
 	for (const FPCGTaggedData& Input : Inputs)
 	{
@@ -58,7 +65,27 @@ bool FPCGStaticMeshSpawnerElement::ExecuteInternal(FPCGContext* Context) const
 			continue;
 		}
 
-		Settings->MeshSelectorInstance->SelectInstances(*Context, Settings, SpatialData, MeshInstances);
+		UPCGPointData* OutputPointData = nullptr;
+
+		if (bOutputPinConnected || Settings->bForceConnectOutput)
+		{
+			FPCGTaggedData& Output = Outputs.Add_GetRef(Input); 
+			
+			OutputPointData = NewObject<UPCGPointData>();
+			OutputPointData->InitializeFromData(SpatialData);
+
+			if (OutputPointData->Metadata->HasAttribute(Settings->OutAttributeName))
+			{
+				OutputPointData->Metadata->DeleteAttribute(Settings->OutAttributeName);
+				PCGE_LOG(Verbose, "Metadata attribute %s is being overwritten in the output data", *Settings->OutAttributeName.ToString());
+			}
+
+			OutputPointData->Metadata->CreateStringAttribute(Settings->OutAttributeName, FName(NAME_None).ToString(), false);
+
+			Output.Data = OutputPointData;
+		}
+
+		Settings->MeshSelectorInstance->SelectInstances(*Context, Settings, SpatialData, MeshInstances, OutputPointData);
 
 		// Spawn a static mesh for each instance
 		SpawnStaticMeshInstances(Context, MeshInstances, TargetActor);
diff --git a/Engine/Plugins/Experimental/PCG/Source/PCG/Private/MeshSelectors/PCGMeshSelectorBase.cpp b/Engine/Plugins/Experimental/PCG/Source/PCG/Private/MeshSelectors/PCGMeshSelectorBase.cpp
--- a/Engine/Plugins/Experimental/PCG/Source/PCG/Private/MeshSelectors/PCGMeshSelectorBase.cpp
+++ b/Engine/Plugins/Experimental/PCG/Source/PCG/Private/MeshSelectors/PCGMeshSelectorBase.cpp
@@ -8,7 +8,8 @@ void UPCGMeshSelectorBase::SelectInstances_Implementation(
 	FPCGContext& Context, 
 	const UPCGStaticMeshSpawnerSettings* Settings, 
 	const UPCGSpatialData* InSpatialData,
-	TArray<FPCGMeshInstanceList>& OutMeshInstances) const
+	TArray<FPCGMeshInstanceList>& OutMeshInstances,
+	UPCGPointData* OutPointData) const
 {
 	PCGE_LOG_C(Error, &Context, "Invalid use of abstract MeshSelectorBase class. Please use an already existing class or implement the CreateMeshInstanceData method");
 }
diff --git a/Engine/Plugins/Experimental/PCG/Source/PCG/Private/MeshSelectors/PCGMeshSelectorByAttribute.cpp b/Engine/Plugins/Experimental/PCG/Source/PCG/Private/MeshSelectors/PCGMeshSelectorByAttribute.cpp
--- a/Engine/Plugins/Experimental/PCG/Source/PCG/Private/MeshSelectors/PCGMeshSelectorByAttribute.cpp
+++ b/Engine/Plugins/Experimental/PCG/Source/PCG/Private/MeshSelectors/PCGMeshSelectorByAttribute.cpp
@@ -13,7 +13,8 @@ void UPCGMeshSelectorByAttribute::SelectInstances_Implementation(
 	FPCGContext& Context, 
 	const UPCGStaticMeshSpawnerSettings* Settings, 
 	const UPCGSpatialData* InSpatialData, 
-	TArray<FPCGMeshInstanceList>& OutMeshInstances) const
+	TArray<FPCGMeshInstanceList>& OutMeshInstances,
+	UPCGPointData* OutPointData) const
 {
 	const UPCGPointData* PointData = InSpatialData->ToPointData(&Context);
 
@@ -48,6 +49,14 @@ void UPCGMeshSelectorByAttribute::SelectInstances_Implementation(
 
 	TMap<PCGMetadataValueKey, TSoftObjectPtr<UStaticMesh>> ValueKeyToMesh;
 
+	// ByAttribute takes in SoftObjectPaths per point in the metadata, so we can pass those directly into the outgoing pin if it exists
+	if (OutPointData)
+	{
+		OutPointData->SetPoints(PointData->GetPoints()); 
+		OutPointData->Metadata->DeleteAttribute(Settings->OutAttributeName);
+		OutPointData->Metadata->CopyAttribute(PointData->Metadata, AttributeName, Settings->OutAttributeName);
+	}
+
 	TRACE_CPUPROFILER_EVENT_SCOPE(FPCGStaticMeshSpawnerElement::Execute::SelectEntries);
 
 	// Assign points to entries
diff --git a/Engine/Plugins/Experimental/PCG/Source/PCG/Private/MeshSelectors/PCGMeshSelectorWeighted.cpp b/Engine/Plugins/Experimental/PCG/Source/PCG/Private/MeshSelectors/PCGMeshSelectorWeighted.cpp
--- a/Engine/Plugins/Experimental/PCG/Source/PCG/Private/MeshSelectors/PCGMeshSelectorWeighted.cpp
+++ b/Engine/Plugins/Experimental/PCG/Source/PCG/Private/MeshSelectors/PCGMeshSelectorWeighted.cpp
@@ -10,13 +10,15 @@
 #include "Metadata/PCGMetadataAttributeTraits.h"
 #include "Metadata/PCGMetadataAttributeTpl.h"
 
+#include "Engine/StaticMesh.h"
 #include "Math/RandomStream.h"
 
 void UPCGMeshSelectorWeighted::SelectInstances_Implementation(
 	FPCGContext& Context, 
 	const UPCGStaticMeshSpawnerSettings* Settings, 
 	const UPCGSpatialData* InSpatialData,
-	TArray<FPCGMeshInstanceList>& OutMeshInstances) const
+	TArray<FPCGMeshInstanceList>& OutMeshInstances,
+	UPCGPointData* OutPointData) const
 {
 	TArray<int> CumulativeWeights;
 
@@ -49,6 +51,35 @@ void UPCGMeshSelectorWeighted::SelectInstances_Implementation(
 		return;
 	}
 
+	TArray<FPCGPoint>* OutPoints = nullptr;
+	FPCGMetadataAttribute<FString>* OutAttribute = nullptr;
+	TMap<TSoftObjectPtr<UStaticMesh>, PCGMetadataValueKey> MeshToValueKey;
+
+	if (OutPointData)
+	{
+		check(OutPointData->Metadata);
+
+		if (!OutPointData->Metadata->HasAttribute(Settings->OutAttributeName)) 
+		{
+			PCGE_LOG_C(Error, &Context, "Out attribute %s is not in the metadata", *Settings->OutAttributeName.ToString());
+		}
+
+		FPCGMetadataAttributeBase* OutAttributeBase = OutPointData->Metadata->GetMutableAttribute(Settings->OutAttributeName);
+
+		if (OutAttributeBase)
+		{
+			if (OutAttributeBase->GetTypeId() == PCG::Private::MetadataTypes<FString>::Id)
+			{
+				OutAttribute = static_cast<FPCGMetadataAttribute<FString>*>(OutAttributeBase);
+				OutPoints = &OutPointData->GetMutablePoints();
+			}
+			else
+			{
+				PCGE_LOG_C(Error, &Context, "Out attribute is not of valid type FString");
+			}
+		}
+	}
+
 	// Assign points to entries
 	{
 		TRACE_CPUPROFILER_EVENT_SCOPE(FPCGStaticMeshSpawnerElement::Execute::SelectEntries);
@@ -71,7 +102,26 @@ void UPCGMeshSelectorWeighted::SelectInstances_Implementation(
 
 			if (RandomPick < OutMeshInstances.Num())
 			{
-				OutMeshInstances[RandomPick].Instances.Emplace(Point.Transform);
+				FPCGMeshInstanceList& InstanceList = OutMeshInstances[RandomPick];
+				InstanceList.Instances.Emplace(Point.Transform);
+
+				const TSoftObjectPtr<UStaticMesh>& Mesh = InstanceList.Mesh;
+
+				if (OutPointData && OutAttribute)
+				{
+					PCGMetadataValueKey* OutValueKey = MeshToValueKey.Find(Mesh);
+					if(!OutValueKey)
+					{
+						PCGMetadataValueKey ValueKey = OutAttribute->AddValue(Mesh.ToSoftObjectPath().ToString());
+						OutValueKey = &MeshToValueKey.Add(Mesh, ValueKey);
+					}
+					
+					check(OutValueKey);
+					
+					FPCGPoint& OutPoint = OutPoints->Add_GetRef(Point);
+					OutPointData->Metadata->InitializeOnSet(OutPoint.MetadataEntry);
+					OutAttribute->SetValueFromValueKey(OutPoint.MetadataEntry, *OutValueKey);
+				}
 			}
 		}
 	}
diff --git a/Engine/Plugins/Experimental/PCG/Source/PCG/Private/MeshSelectors/PCGMeshSelectorWeightedByCategory.cpp b/Engine/Plugins/Experimental/PCG/Source/PCG/Private/MeshSelectors/PCGMeshSelectorWeightedByCategory.cpp
--- a/Engine/Plugins/Experimental/PCG/Source/PCG/Private/MeshSelectors/PCGMeshSelectorWeightedByCategory.cpp
+++ b/Engine/Plugins/Experimental/PCG/Source/PCG/Private/MeshSelectors/PCGMeshSelectorWeightedByCategory.cpp
@@ -11,6 +11,7 @@
 #include "Elements/PCGStaticMeshSpawner.h"
 
 #include "Math/RandomStream.h"
+#include "Engine/StaticMesh.h"
 
 struct FPCGInstancesAndWeights
 {
@@ -22,7 +23,8 @@ void UPCGMeshSelectorWeightedByCategory::SelectInstances_Implementation(
 	FPCGContext& Context,
 	const UPCGStaticMeshSpawnerSettings* Settings,
 	const UPCGSpatialData* InSpatialData,
-	TArray<FPCGMeshInstanceList>& OutMeshInstances) const
+	TArray<FPCGMeshInstanceList>& OutMeshInstances,
+	UPCGPointData* OutPointData) const
 {
 	const UPCGPointData* PointData = InSpatialData->ToPointData(&Context);
 
@@ -117,6 +119,41 @@ void UPCGMeshSelectorWeightedByCategory::SelectInstances_Implementation(
 			InstancesAndWeights->CumulativeWeights.Add(TotalWeight);
 			InstancesAndWeights->InstanceListIndices.Add(Index);
 		}
+
+		// if no weighted entries were collected, discard this InstancesAndWeights
+		if (InstancesAndWeights->CumulativeWeights.Num() == 0)
+		{
+			CategoryEntryToInstancesAndWeights.Remove(ValueKey);
+		}
+	}
+
+	TArray<FPCGPoint>* OutPoints = nullptr;
+	FPCGMetadataAttribute<FString>* OutAttribute = nullptr;
+	TMap<TSoftObjectPtr<UStaticMesh>, PCGMetadataValueKey> MeshToValueKey;
+
+	if (OutPointData)
+	{
+		check(OutPointData->Metadata);
+
+		if (!OutPointData->Metadata->HasAttribute(Settings->OutAttributeName)) 
+		{
+			PCGE_LOG_C(Error, &Context, "Out attribute %s is not in the metadata", *Settings->OutAttributeName.ToString());
+		}
+
+		FPCGMetadataAttributeBase* OutAttributeBase = OutPointData->Metadata->GetMutableAttribute(Settings->OutAttributeName);
+
+		if (OutAttributeBase)
+		{
+			if (OutAttributeBase->GetTypeId() == PCG::Private::MetadataTypes<FString>::Id)
+			{
+				OutAttribute = static_cast<FPCGMetadataAttribute<FString>*>(OutAttributeBase);
+				OutPoints = &OutPointData->GetMutablePoints();
+			}
+			else
+			{
+				PCGE_LOG_C(Error, &Context, "Out attribute is not of valid type FString");
+			}
+		}
 	}
 
 	TRACE_CPUPROFILER_EVENT_SCOPE(FPCGStaticMeshSpawnerElement::Execute::SelectEntries);
@@ -160,7 +197,26 @@ void UPCGMeshSelectorWeightedByCategory::SelectInstances_Implementation(
 		if (RandomPick < InstancesAndWeights->InstanceListIndices.Num())
 		{
 			const int32 Index = InstancesAndWeights->InstanceListIndices[RandomPick];
-			OutMeshInstances[Index].Instances.Emplace(Point.Transform);
+			FPCGMeshInstanceList& InstanceList = OutMeshInstances[Index];
+			InstanceList.Instances.Emplace(Point.Transform);
+
+			const TSoftObjectPtr<UStaticMesh>& Mesh = InstanceList.Mesh;
+
+			if (OutPointData && OutAttribute)
+			{
+				PCGMetadataValueKey* OutValueKey = MeshToValueKey.Find(Mesh);
+				if(!OutValueKey)
+				{
+					PCGMetadataValueKey TempValueKey = OutAttribute->AddValue(Mesh.ToSoftObjectPath().ToString());
+					OutValueKey = &MeshToValueKey.Add(Mesh, TempValueKey);
+				}
+				
+				check(OutValueKey);
+				
+				FPCGPoint& OutPoint = OutPoints->Add_GetRef(Point);
+				OutPointData->Metadata->InitializeOnSet(OutPoint.MetadataEntry);
+				OutAttribute->SetValueFromValueKey(OutPoint.MetadataEntry, *OutValueKey);
+			}
 		}
 	}
 }
diff --git a/Engine/Plugins/Experimental/PCG/Source/PCG/Private/Tests/Elements/PCGStaticMeshSpawnerTest.cpp b/Engine/Plugins/Experimental/PCG/Source/PCG/Private/Tests/Elements/PCGStaticMeshSpawnerTest.cpp
--- a/Engine/Plugins/Experimental/PCG/Source/PCG/Private/Tests/Elements/PCGStaticMeshSpawnerTest.cpp
+++ b/Engine/Plugins/Experimental/PCG/Source/PCG/Private/Tests/Elements/PCGStaticMeshSpawnerTest.cpp
@@ -17,6 +17,11 @@
 #include "MeshSelectors/PCGMeshSelectorByAttribute.h"
 #include "MeshSelectors/PCGMeshSelectorWeightedByCategory.h"
 
+#include "Metadata/PCGMetadata.h"
+#include "Metadata/PCGMetadataAttribute.h"
+#include "Metadata/PCGMetadataAttributeTraits.h"
+#include "Metadata/PCGMetadataAttributeTpl.h"
+
 #include "Components/InstancedStaticMeshComponent.h"
 #include "Engine/StaticMesh.h"
 
@@ -48,18 +53,102 @@ namespace
 	};
 }
 
+/** Finds the output UPCGPointData and builds a mapping from world-space position to each FPCGPoint */
+void GetOutputDataAndBuildPositionMap(FPCGTestBaseClass* Test, const TArray<FPCGTaggedData>& OutputData, TMap<FVector, FPCGPoint>& OutLocationToPoint, UPCGPointData*& OutPointData)
+{
+	check(Test);
+
+	if (Test->TestEqual("Valid number of data elements in output", OutputData.Num(), 2))
+	{
+		TObjectPtr<UPCGSettings> OutSettings = Cast<UPCGSettings>(OutputData[0].Data);
+		Test->TestTrue("Valid output settings", OutSettings != nullptr);
+
+		OutPointData = Cast<UPCGPointData>(OutputData[1].Data);
+	}
+
+	if (Test->TestNotNull("Valid output data", OutPointData))
+	{
+		const TArray<FPCGPoint>& OutPoints = OutPointData->GetPoints();
+
+		for (const FPCGPoint& OutPoint : OutPoints)
+		{
+			OutLocationToPoint.Emplace(OutPoint.Transform.GetLocation(), OutPoint);
+		}
+	}
+}
+
+/** Iterates over instances of the ISMC to ensure each instance matches the softobjectpath metadata of the corresponding FPCGPoint */
+void ValidateOutPointData(
+	FPCGTestBaseClass* Test, 
+	const UInstancedStaticMeshComponent* ISMC, 
+	const TMap<FVector, FPCGPoint>& LocationToPoint, 
+	const FName& OutputAttributeName, 
+	const UPCGPointData* OutPointData) 
+{
+	check(Test);
+	check(ISMC);
+	check(OutPointData && OutPointData->Metadata);
+
+	const FPCGMetadataAttributeBase* AttributeBase = OutPointData->Metadata->GetConstAttribute(OutputAttributeName);
+	if (!Test->TestNotNull("Metadata contains out attribute", AttributeBase)) 
+	{ 
+		return; 
+	}
+
+	const bool bIsMetadataTypeValid = AttributeBase->GetTypeId() == PCG::Private::MetadataTypes<FString>::Id;
+	if (!Test->TestTrue("Output attribute is a valid type", bIsMetadataTypeValid)) 
+	{ 
+		return; 
+	}
+
+	const FPCGMetadataAttribute<FString>* OutAttribute = static_cast<const FPCGMetadataAttribute<FString>*>(AttributeBase);
+
+	const FString& MeshPath = FSoftObjectPath(ISMC->GetStaticMesh()).ToString();
+	const PCGMetadataValueKey MeshValueKey = OutAttribute->FindValue(MeshPath);
+
+	if (MeshValueKey == PCGDefaultValueKey)
+	{
+		if (!Test->TestEqual("Mesh value exists in output attribute", MeshPath, OutAttribute->GetValue(MeshValueKey)))
+		{
+			return;
+		}
+	}
+
+	for (int InstanceIndex = 0; InstanceIndex < ISMC->GetInstanceCount(); ++InstanceIndex)
+	{
+		FTransform Transform;
+		ISMC->GetInstanceTransform(InstanceIndex, Transform, /*bWorldSpace=*/true);
+
+		const FPCGPoint* OutPoint = LocationToPoint.Find(Transform.GetLocation());
+
+		if (Test->TestNotNull("Output point exists", OutPoint) && OutAttribute)
+		{
+			const PCGMetadataValueKey ValueKey = OutAttribute->GetValueKey(OutPoint->MetadataEntry);
+			Test->TestEqual("Valid mesh stored in output point metadata", ValueKey, MeshValueKey);
+		}
+	}
+}
+
 void TestMeshSelectorByAttribute(
 	FPCGStaticMeshSpawnerByAttributeTest* Test, 
 	const PCGDeterminismTests::FTestData& TestData, 
 	const FName& AttributeName, 
-	const TArray<FPCGByAttributeValidationData>& ValidationDataEntries)
+	const TArray<FPCGByAttributeValidationData>& ValidationDataEntries,
+	bool bValidateOutput=false)
 {
+	check(Test);
+
 	// set the MeshSelector
 	UPCGStaticMeshSpawnerSettings* Settings = CastChecked<UPCGStaticMeshSpawnerSettings>(TestData.Settings);
 	Settings->SetMeshSelectorType(UPCGMeshSelectorByAttribute::StaticClass());
 	UPCGMeshSelectorByAttribute* MeshSelector = CastChecked<UPCGMeshSelectorByAttribute>(Settings->MeshSelectorInstance);
 	MeshSelector->AttributeName = AttributeName;
 
+	if (bValidateOutput)
+	{
+		Settings->bForceConnectOutput = true;
+	}
+
 	// initialize and execute the StaticMeshSpawner
 	FPCGElementPtr StaticMeshSpawner = Settings->GetElement();
 	FPCGContext* Context = StaticMeshSpawner->Initialize(TestData.InputData, TestData.TestPCGComponent, nullptr);
@@ -67,6 +156,14 @@ void TestMeshSelectorByAttribute(
 	while (!StaticMeshSpawner->Execute(Context))
 	{}
 
+	TMap<FVector, FPCGPoint> LocationToPoint;
+	UPCGPointData* OutPointData = nullptr;
+
+	if (bValidateOutput)
+	{
+		GetOutputDataAndBuildPositionMap(Test, Context->OutputData.TaggedData, LocationToPoint, OutPointData);
+	}
+
 	struct FPCGByAttributeValidIndexList
 	{
 		TArray<TPair<int, int>> ValidIndices;
@@ -117,58 +214,69 @@ void TestMeshSelectorByAttribute(
 			return Entry.Mesh == StaticMesh;
 			});
 
-		if (!Entry)
-		{
-			Test->TestTrue("Validate instanced mesh exists in MeshEntries", false);
-			continue;
-		}
-
-		const FPCGByAttributeValidIndexList* IndexList = MapToIndexList.Find(Entry->Mesh);
-
-		Test->TestTrue("Valid instance count per mesh", InstanceCount == IndexList->Count);
-
-		for (int InstanceIndex = 0; InstanceIndex < InstanceCount; InstanceIndex++)
+		if (Test->TestNotNull("Validate instanced mesh exists in MeshEntries", Entry))
 		{
-			FTransform Transform;
-			ISMC->GetInstanceTransform(InstanceIndex, Transform);
-			FVector Location = Transform.GetLocation();
+			const FPCGByAttributeValidIndexList* IndexList = MapToIndexList.Find(Entry->Mesh);
+			check(IndexList);
 
-			bool bFound = false;
+			Test->TestEqual("Valid instance count per mesh", InstanceCount, IndexList->Count);
 
-			for (int i = 0; i < IndexList->ValidIndices.Num(); i++)
+			for (int InstanceIndex = 0; InstanceIndex < InstanceCount; InstanceIndex++)
 			{
-				const int StartIndex = IndexList->ValidIndices[i].Key;
-				const int EndIndex = IndexList->ValidIndices[i].Value;
+				FTransform Transform;
+				ISMC->GetInstanceTransform(InstanceIndex, Transform);
+				FVector Location = Transform.GetLocation();
+
+				bool bFound = false;
 
-				if (Location.X >= StartIndex && Location.X < EndIndex)
+				for (int i = 0; i < IndexList->ValidIndices.Num(); i++)
 				{
-					bFound = true;
-					break;
+					const int StartIndex = IndexList->ValidIndices[i].Key;
+					const int EndIndex = IndexList->ValidIndices[i].Value;
+
+					if (Location.X >= StartIndex && Location.X < EndIndex)
+					{
+						bFound = true;
+						break;
+					}
 				}
+
+				Test->TestTrue("Valid transform locations", bFound);
 			}
 
-			Test->TestTrue("Valid transform locations", bFound);
+			if (bValidateOutput)
+			{
+				ValidateOutPointData(Test, ISMC, LocationToPoint, Settings->OutAttributeName, OutPointData);
+			}
 		}
 	}
 
 	TestData.TestPCGComponent->bGenerated = true;
 	TestData.TestPCGComponent->Cleanup();
 	
-	Test->TestTrue("Valid total instance count", TotalInstanceCount == ValidInstanceCount);
+	Test->TestEqual("Valid total instance count", TotalInstanceCount, ValidInstanceCount);
 }
 
 void TestMeshSelectorWeighted(
 	FPCGStaticMeshSpawnerWeightedTest* Test, 
 	const PCGDeterminismTests::FTestData& TestData, 
 	const TArray<FPCGMeshSelectorWeightedEntry>& Entries, 
-	int PointCount)
+	int PointCount,
+	bool bValidateOutput=false)
 {
+	check(Test);
+
 	// set the MeshSelector
 	UPCGStaticMeshSpawnerSettings* Settings = CastChecked<UPCGStaticMeshSpawnerSettings>(TestData.Settings);
 	Settings->SetMeshSelectorType(UPCGMeshSelectorWeighted::StaticClass());
 	UPCGMeshSelectorWeighted* MeshSelector = CastChecked<UPCGMeshSelectorWeighted>(Settings->MeshSelectorInstance);
 	MeshSelector->MeshEntries = Entries;
 
+	if (bValidateOutput)
+	{
+		Settings->bForceConnectOutput = true;
+	}
+
 	// initialize and execute the StaticMeshSpawner
 	FPCGElementPtr StaticMeshSpawner = Settings->GetElement();
 	FPCGContext* Context = StaticMeshSpawner->Initialize(TestData.InputData, TestData.TestPCGComponent, nullptr);
@@ -176,6 +284,14 @@ void TestMeshSelectorWeighted(
 	while (!StaticMeshSpawner->Execute(Context))
 	{}
 
+	TMap<FVector, FPCGPoint> LocationToPoint;
+	UPCGPointData* OutPointData = nullptr;
+
+	if (bValidateOutput)
+	{
+		GetOutputDataAndBuildPositionMap(Test, Context->OutputData.TaggedData, LocationToPoint, OutPointData);
+	}
+
 	int TotalWeight = 0;
 	for (const FPCGMeshSelectorWeightedEntry& Entry : Entries)
 	{
@@ -208,22 +324,26 @@ void TestMeshSelectorWeighted(
 			return Entry.Mesh == StaticMesh;
 			});
 
-		if (!Entry)
+		if (Test->TestNotNull("Validate instanced mesh exists in MeshEntries", Entry))
 		{
-			Test->TestTrue("Validate instanced mesh exists in MeshEntries", false);
-			continue;
-		}
+			check(TotalWeight > 0);
+
+			// validate correct InstanceCount (within a range) for the weight of this mesh
+			const int TargetCount = PointCount * Entry->Weight / TotalWeight;
+			const float ErrorBound = FMath::Max(TargetCount * 0.1f, 10); // 10% error boundary 
+			Test->TestTrue("Valid instance count per mesh", InstanceCount >= TargetCount - ErrorBound && InstanceCount <= TargetCount + ErrorBound);
 
-		// validate correct InstanceCount (within a range) for the weight of this mesh
-		const int TargetCount = PointCount * Entry->Weight / TotalWeight;
-		const float ErrorBound = FMath::Max(TargetCount * 0.1f, 10); // 10% error boundary 
-		Test->TestTrue("Valid instance count per mesh", InstanceCount >= TargetCount - ErrorBound && InstanceCount <= TargetCount + ErrorBound);
+			if (bValidateOutput)
+			{
+				ValidateOutPointData(Test, ISMC, LocationToPoint, Settings->OutAttributeName, OutPointData);
+			}
+		}
 	}
 
 	TestData.TestPCGComponent->bGenerated = true;
 	TestData.TestPCGComponent->Cleanup();
 	
-	Test->TestTrue("Valid total instance count", TotalInstanceCount == PointCount);
+	Test->TestEqual("Valid total instance count", TotalInstanceCount, PointCount);
 }
 
 void TestMeshSelectorWeightedByCategory(
@@ -232,8 +352,11 @@ void TestMeshSelectorWeightedByCategory(
 	const FName& AttributeName, 
 	const TArray<FPCGWeightedByCategoryEntryList>& Entries, 
 	const TArray<FPCGWeightedByCategoryValidationData>& ValidationDataEntries,
-	int PointCount)
+	int PointCount,
+	bool bValidateOutput=false)
 {
+	check(Test);
+
 	// set the MeshSelector
 	UPCGStaticMeshSpawnerSettings* Settings = CastChecked<UPCGStaticMeshSpawnerSettings>(TestData.Settings);
 	Settings->SetMeshSelectorType(UPCGMeshSelectorWeightedByCategory::StaticClass());
@@ -241,6 +364,11 @@ void TestMeshSelectorWeightedByCategory(
 	MeshSelector->CategoryAttribute = AttributeName;
 	MeshSelector->Entries = Entries;
 
+	if (bValidateOutput)
+	{
+		Settings->bForceConnectOutput = true;
+	}
+
 	// initialize and execute the StaticMeshSpawner
 	FPCGElementPtr StaticMeshSpawner = Settings->GetElement();
 	FPCGContext* Context = StaticMeshSpawner->Initialize(TestData.InputData, TestData.TestPCGComponent, nullptr);
@@ -249,6 +377,14 @@ void TestMeshSelectorWeightedByCategory(
 	{
 	}
 
+	TMap<FVector, FPCGPoint> LocationToPoint;
+	UPCGPointData* OutPointData = nullptr;
+
+	if (bValidateOutput)
+	{
+		GetOutputDataAndBuildPositionMap(Test, Context->OutputData.TaggedData, LocationToPoint, OutPointData);
+	}
+
 	check(TestData.TestPCGComponent);
 
 	TArray<UInstancedStaticMeshComponent*> ISMCs;
@@ -275,20 +411,22 @@ void TestMeshSelectorWeightedByCategory(
 			return Entry.Mesh == StaticMesh;
 			});
 
-		if (!Entry)
+		if (Test->TestNotNull("Validate instanced mesh exists in entries", Entry))
 		{
-			Test->TestTrue("Validate instanced mesh exists in entries", false);
-			continue;
-		}
+			const float ErrorBound = FMath::Max(Entry->ExpectedCount * 0.1f, 10); // 10% error boundary 
+			Test->TestTrue("Valid instance count per mesh", InstanceCount >= Entry->ExpectedCount - ErrorBound && InstanceCount <= Entry->ExpectedCount + ErrorBound);
 
-		const float ErrorBound = FMath::Max(Entry->ExpectedCount * 0.1f, 10); // 10% error boundary 
-		Test->TestTrue("Valid instance count per mesh", InstanceCount >= Entry->ExpectedCount - ErrorBound && InstanceCount <= Entry->ExpectedCount + ErrorBound);
+			if (bValidateOutput)
+			{
+				ValidateOutPointData(Test, ISMC, LocationToPoint, Settings->OutAttributeName, OutPointData);
+			}
+		}
 	}
 
 	TestData.TestPCGComponent->bGenerated = true;
 	TestData.TestPCGComponent->Cleanup();
 
-	Test->TestTrue("Valid total instance count", TotalInstanceCount == PointCount);
+	Test->TestEqual("Valid total instance count", TotalInstanceCount, PointCount);
 }
 
 bool FPCGStaticMeshSpawnerByAttributeTest::RunTest(const FString& Parameters)
@@ -314,9 +452,9 @@ bool FPCGStaticMeshSpawnerByAttributeTest::RunTest(const FString& Parameters)
 	// compose the test data
 	TObjectPtr<UPCGPointData> PointData = PCGTestsCommon::CreateEmptyPointData();
 	PointData->TargetActor = TestData.TestActor;
-	PointData->Metadata->CreateStringAttribute(RockAttribute, ConePath.GetAssetPathString(), false);
-	PointData->Metadata->CreateStringAttribute(TreeAttribute, ConePath.GetAssetPathString(), false);
-	PointData->Metadata->CreateStringAttribute(BushAttribute, ConePath.GetAssetPathString(), false);
+	PointData->Metadata->CreateStringAttribute(RockAttribute, ConePath.ToString(), false);
+	PointData->Metadata->CreateStringAttribute(TreeAttribute, ConePath.ToString(), false);
+	PointData->Metadata->CreateStringAttribute(BushAttribute, ConePath.ToString(), false);
 
 	// first 20 points are Rock, next 50 are Tree, next 30 are Bush, for a total of 100 points
 	const int PointCount = 100;
@@ -346,7 +484,7 @@ bool FPCGStaticMeshSpawnerByAttributeTest::RunTest(const FString& Parameters)
 			ObjectPath = (i < 85) ? CubePath : CylinderPath; // half are CubePath, half are CylinderPath
 		}
 		
-		UPCGMetadataAccessorHelpers::SetStringAttribute(Points[i], PointData->Metadata, AttributeName, ObjectPath.GetAssetPathString());
+		UPCGMetadataAccessorHelpers::SetStringAttribute(Points[i], PointData->Metadata, AttributeName, ObjectPath.ToString());
 	}
 	
 	FPCGTaggedData& TaggedDataPoints = TestData.InputData.TaggedData.Emplace_GetRef(FPCGTaggedData());
@@ -370,13 +508,13 @@ bool FPCGStaticMeshSpawnerByAttributeTest::RunTest(const FString& Parameters)
 	ValidationDataEntries.Emplace(ConeMesh, 70, 100);
 	TestMeshSelectorByAttribute(this, TestData, TreeAttribute, ValidationDataEntries);
 
-	// Test 3
+	// Test 3 - also tests optional output data
 	// should have ConeMesh at Loc.x = [0, 70), CubeMesh at Loc.x = [70, 85), CylinderMesh at Loc.x = [85, 100)
 	ValidationDataEntries.Empty();
 	ValidationDataEntries.Emplace(ConeMesh, 0, 70);
 	ValidationDataEntries.Emplace(CubeMesh, 70, 85);
 	ValidationDataEntries.Emplace(CylinderMesh, 85, 100);
-	TestMeshSelectorByAttribute(this, TestData, BushAttribute, ValidationDataEntries);
+	TestMeshSelectorByAttribute(this, TestData, BushAttribute, ValidationDataEntries, /*bValidateOutput=*/true);
 
 	return true;
 }
@@ -435,11 +573,17 @@ bool FPCGStaticMeshSpawnerWeightedTest::RunTest(const FString& Parameters)
 	Entries[2].Weight = 1;
 	TestMeshSelectorWeighted(this, TestData, Entries, PointCount);
 
-	// Test 5: null weight
-	Entries[0].Weight = 1;
+	// Test 5: zero total weight
+	Entries[0].Weight = 0;
 	Entries[1].Weight = 0;
-	Entries[2].Weight = 1;
-	TestMeshSelectorWeighted(this, TestData, Entries, PointCount);
+	Entries[2].Weight = 0;
+	TestMeshSelectorWeighted(this, TestData, Entries, 0);
+
+	// Test 6: validate optional output data
+	Entries[0].Weight = 1;
+	Entries[1].Weight = 10;
+	Entries[2].Weight = 100;
+	TestMeshSelectorWeighted(this, TestData, Entries, PointCount, /*bValidateOutput=*/true);
 
 	return true;
 }
@@ -514,9 +658,9 @@ bool FPCGStaticMeshSpawnerWeightedByCategoryTest::RunTest(const FString& Paramet
 
 	// Test 1 - basic weighted functionality
 	{
-		const int CubeWeight = 1.f;
-		const int SphereWeight = 1.f;
-		const int CylinderWeight = 1.f;
+		const int CubeWeight = 1;
+		const int SphereWeight = 1;
+		const int CylinderWeight = 1;
 
 		const int IgneousCategoryTotalWeight = 3;
 
@@ -544,14 +688,14 @@ bool FPCGStaticMeshSpawnerWeightedByCategoryTest::RunTest(const FString& Paramet
 
 	// Test 2 - selection based on rocktype
 	{
-		const int IgneousCubeWeight = 1.f;
-		const int IgneousSphereWeight = 1.f;
+		const int IgneousCubeWeight = 1;
+		const int IgneousSphereWeight = 1;
 
-		const int SedimentaryCubeWeight = 10.f;
-		const int SedimentaryCylinderWeight = 1.f;
+		const int SedimentaryCubeWeight = 10;
+		const int SedimentaryCylinderWeight = 1;
 
-		const int MetamorphicSphereWeight = 1.f;
-		const int MetamorphicCylinderWeight = 10.f;
+		const int MetamorphicSphereWeight = 1;
+		const int MetamorphicCylinderWeight = 10;
 
 		const int IgneousCategoryTotalWeight = 2;
 		const int SedimentaryCategoryTotalWeight = 11;
@@ -590,10 +734,10 @@ bool FPCGStaticMeshSpawnerWeightedByCategoryTest::RunTest(const FString& Paramet
 
 	// Test 3 - empty weighted entries
 	{
-		const int CubeWeight = 1.f;
-		const int IgneousSphereWeight = 1.f;
-		const int MetamorphicSphereWeight = 1.f;
-		const int CylinderWeight = 10.f;
+		const int CubeWeight = 1;
+		const int IgneousSphereWeight = 1;
+		const int MetamorphicSphereWeight = 1;
+		const int CylinderWeight = 10;
 
 		const int IgneousCategoryTotalWeight = 2;
 		const int MetamorphicCategoryTotalWeight = 11;
@@ -629,8 +773,8 @@ bool FPCGStaticMeshSpawnerWeightedByCategoryTest::RunTest(const FString& Paramet
 
 	// Test 4 - duplicate category
 	{
-		const int CubeWeight = 1.f;
-		const int SphereWeight = 1.f;
+		const int CubeWeight = 1;
+		const int SphereWeight = 1;
 
 		const int IgneousCategoryTotalWeight = 1;
 		const int SedimentaryCategoryTotalWeight = 1;
@@ -666,8 +810,8 @@ bool FPCGStaticMeshSpawnerWeightedByCategoryTest::RunTest(const FString& Paramet
 
 	// Test 5 - invalid rock type
 	{
-		const int CubeWeight = 1.f;
-		const int SphereWeight = 1.f;
+		const int CubeWeight = 1;
+		const int SphereWeight = 1;
 		const int CategoryTotalWeight = 2;
 		const int TotalSize = IgneousCategorySize;
 
@@ -694,14 +838,14 @@ bool FPCGStaticMeshSpawnerWeightedByCategoryTest::RunTest(const FString& Paramet
 
 	// Test 6 - default entry list + duplicate default ignored
 	{
-		const int IgneousCubeWeight = 1.f;
-		const int IgneousSphereWeight = 1.f;
+		const int IgneousCubeWeight = 1;
+		const int IgneousSphereWeight = 1;
 
-		const int SedimentaryCubeWeight = 10.f;
-		const int SedimentaryCylinderWeight = 1.f;
+		const int SedimentaryCubeWeight = 10;
+		const int SedimentaryCylinderWeight = 1;
 
-		const int MetamorphicSphereWeight = 1.f;
-		const int MetamorphicCylinderWeight = 10.f;
+		const int MetamorphicSphereWeight = 1;
+		const int MetamorphicCylinderWeight = 10;
 
 		const int IgneousCategoryTotalWeight = 2;
 		const int SedimentaryCategoryTotalWeight = 11;
@@ -756,5 +900,89 @@ bool FPCGStaticMeshSpawnerWeightedByCategoryTest::RunTest(const FString& Paramet
 		TestMeshSelectorWeightedByCategory(this, TestData, AttributeName, Entries, ValidationData, TotalSize);
 	}
 
+	// Test 7 - zero total weight for a category
+	{
+		const int CubeWeight = 0;
+		const int SphereWeight = 0;
+
+		const int CylinderWeight = 1;
+		const int ConeWeight = 1;
+
+		const int SedimentaryCategoryTotalWeight = 2;
+
+		const int TotalSize = SedimentaryCategorySize;
+
+		TArray<FPCGMeshSelectorWeightedEntry> IgneousEntries;
+		IgneousEntries.Emplace(CubeMesh, CubeWeight);
+		IgneousEntries.Emplace(SphereMesh, SphereWeight);
+
+		TArray<FPCGMeshSelectorWeightedEntry> SedimentaryEntries;
+		SedimentaryEntries.Emplace(CylinderMesh, CylinderWeight);
+		SedimentaryEntries.Emplace(ConeMesh, ConeWeight);
+
+		TArray<FPCGWeightedByCategoryEntryList> Entries;
+		Entries.Emplace(IgneousRockType, IgneousEntries);
+		Entries.Emplace(SedimentaryRockType, SedimentaryEntries);
+
+		const int CubeCount = 0;
+		const int SphereCount = 0;
+		const int CylinderCount = (SedimentaryCategorySize * CylinderWeight / SedimentaryCategoryTotalWeight);
+		const int ConeCount = (SedimentaryCategorySize * ConeWeight / SedimentaryCategoryTotalWeight);
+
+		TArray<FPCGWeightedByCategoryValidationData> ValidationData;
+		ValidationData.Emplace(CubeMesh, CubeCount);
+		ValidationData.Emplace(SphereMesh, SphereCount);
+		ValidationData.Emplace(CylinderMesh, CylinderCount);
+		ValidationData.Emplace(ConeMesh, ConeCount);
+
+		TestMeshSelectorWeightedByCategory(this, TestData, AttributeName, Entries, ValidationData, TotalSize);
+	}
+
+	// Test 8 - validate optional output data
+	{
+		const int IgneousCubeWeight = 1;
+		const int IgneousSphereWeight = 1;
+
+		const int SedimentaryCubeWeight = 10;
+		const int SedimentaryCylinderWeight = 1;
+
+		const int MetamorphicSphereWeight = 1;
+		const int MetamorphicCylinderWeight = 10;
+
+		const int IgneousCategoryTotalWeight = 2;
+		const int SedimentaryCategoryTotalWeight = 11;
+		const int MetamorphicCategoryTotalWeight = 11;
+
+		const int TotalSize = IgneousCategorySize + SedimentaryCategorySize + MetamorphicCategorySize;
+
+		TArray<FPCGMeshSelectorWeightedEntry> IgneousEntries;
+		IgneousEntries.Emplace(CubeMesh, IgneousCubeWeight);
+		IgneousEntries.Emplace(SphereMesh, IgneousSphereWeight);
+
+		TArray<FPCGMeshSelectorWeightedEntry> SedimentaryEntries;
+		SedimentaryEntries.Emplace(CubeMesh, SedimentaryCubeWeight);
+		SedimentaryEntries.Emplace(CylinderMesh, SedimentaryCylinderWeight);
+
+		TArray<FPCGMeshSelectorWeightedEntry> MetamorphicEntries;
+		MetamorphicEntries.Emplace(SphereMesh, MetamorphicSphereWeight);
+		MetamorphicEntries.Emplace(CylinderMesh, MetamorphicCylinderWeight);
+
+		TArray<FPCGWeightedByCategoryEntryList> Entries;
+		Entries.Emplace(IgneousRockType, IgneousEntries);
+		Entries.Emplace(SedimentaryRockType, SedimentaryEntries);
+		Entries.Emplace(MetamorphicRockType, MetamorphicEntries);
+
+		const int CubeCount = (IgneousCategorySize * IgneousCubeWeight / IgneousCategoryTotalWeight) + (SedimentaryCategorySize * SedimentaryCubeWeight / SedimentaryCategoryTotalWeight);
+		const int SphereCount = (IgneousCategorySize * IgneousSphereWeight / IgneousCategoryTotalWeight) + (MetamorphicCategorySize * MetamorphicSphereWeight / MetamorphicCategoryTotalWeight);
+		const int CylinderCount = (SedimentaryCategorySize * SedimentaryCylinderWeight / SedimentaryCategoryTotalWeight) + (MetamorphicCategorySize * MetamorphicCylinderWeight / MetamorphicCategoryTotalWeight);
+
+		TArray<FPCGWeightedByCategoryValidationData> ValidationData;
+		ValidationData.Emplace(CubeMesh, CubeCount);
+		ValidationData.Emplace(SphereMesh, SphereCount);
+		ValidationData.Emplace(CylinderMesh, CylinderCount);
+
+		TestMeshSelectorWeightedByCategory(this, TestData, AttributeName, Entries, ValidationData, TotalSize, /*bValidateOutput=*/true);
+	}
+
 	return true;
 }
diff --git a/Engine/Plugins/Experimental/PCG/Source/PCG/Public/Elements/PCGStaticMeshSpawner.h b/Engine/Plugins/Experimental/PCG/Source/PCG/Public/Elements/PCGStaticMeshSpawner.h
--- a/Engine/Plugins/Experimental/PCG/Source/PCG/Public/Elements/PCGStaticMeshSpawner.h
+++ b/Engine/Plugins/Experimental/PCG/Source/PCG/Public/Elements/PCGStaticMeshSpawner.h
@@ -45,7 +45,7 @@ public:
 	virtual EPCGSettingsType GetType() const override { return EPCGSettingsType::Spawner; }
 #endif
 
-	virtual TArray<FPCGPinProperties> OutputPinProperties() const override { return TArray<FPCGPinProperties>(); }
+	virtual TArray<FPCGPinProperties> OutputPinProperties() const override { return Super::DefaultPointOutputPinProperties(); }
 
 protected:
 	virtual FPCGElementPtr CreateElement() const override;
@@ -69,9 +69,14 @@ public:
 	UPROPERTY(BlueprintReadOnly, VisibleAnywhere, Instanced, Category = Settings)
 	TObjectPtr<UPCGMeshSelectorBase> MeshSelectorInstance;
 
+	UPROPERTY(BlueprintReadOnly, VisibleAnywhere, Category = Settings)
+	FName OutAttributeName = NAME_None;
+
 	UPROPERTY(meta=(DeprecatedProperty, DeprecationMessage="Use MeshSelectorType and MeshSelectorInstance instead."))
 	TArray<FPCGStaticMeshSpawnerEntry> Meshes_DEPRECATED;
 
+	bool bForceConnectOutput = false;
+
 protected:
 	void RefreshMeshSelector(); 
 };
diff --git a/Engine/Plugins/Experimental/PCG/Source/PCG/Public/MeshSelectors/PCGMeshSelectorBase.h b/Engine/Plugins/Experimental/PCG/Source/PCG/Public/MeshSelectors/PCGMeshSelectorBase.h
--- a/Engine/Plugins/Experimental/PCG/Source/PCG/Public/MeshSelectors/PCGMeshSelectorBase.h
+++ b/Engine/Plugins/Experimental/PCG/Source/PCG/Public/MeshSelectors/PCGMeshSelectorBase.h
@@ -55,13 +55,15 @@ public:
 		FPCGContext& Context,
 		const UPCGStaticMeshSpawnerSettings* Settings,
 		const UPCGSpatialData* InSpatialData,
-		TArray<FPCGMeshInstanceList>& OutMeshInstances) const;
+		TArray<FPCGMeshInstanceList>& OutMeshInstances,
+		UPCGPointData* OutPointData) const;
 
 	virtual void SelectInstances_Implementation(
 		FPCGContext& Context,
 		const UPCGStaticMeshSpawnerSettings* Settings,
 		const UPCGSpatialData* InSpatialData,
-		TArray<FPCGMeshInstanceList>& OutMeshInstances) const;
+		TArray<FPCGMeshInstanceList>& OutMeshInstances,
+		UPCGPointData* OutPointData) const;
 
 	/** Searches OutInstanceLists for an InstanceList matching the given parameters. If nothing is found, creates a new InstanceList and adds to OutInstanceLists. Returns true if added. */
 	UFUNCTION(BlueprintCallable, Category = MeshSelection)
diff --git a/Engine/Plugins/Experimental/PCG/Source/PCG/Public/MeshSelectors/PCGMeshSelectorByAttribute.h b/Engine/Plugins/Experimental/PCG/Source/PCG/Public/MeshSelectors/PCGMeshSelectorByAttribute.h
--- a/Engine/Plugins/Experimental/PCG/Source/PCG/Public/MeshSelectors/PCGMeshSelectorByAttribute.h
+++ b/Engine/Plugins/Experimental/PCG/Source/PCG/Public/MeshSelectors/PCGMeshSelectorByAttribute.h
@@ -16,7 +16,8 @@ public:
 		UPARAM(ref) FPCGContext& Context,
 		const UPCGStaticMeshSpawnerSettings* Settings,
 		const UPCGSpatialData* InSpatialData,
-		TArray<FPCGMeshInstanceList>& OutMeshInstances) const override;
+		TArray<FPCGMeshInstanceList>& OutMeshInstances,
+		UPCGPointData* OutPointData) const override;
 
 public:
 	UPROPERTY(BlueprintReadWrite, EditAnywhere, Category = Settings)
diff --git a/Engine/Plugins/Experimental/PCG/Source/PCG/Public/MeshSelectors/PCGMeshSelectorWeighted.h b/Engine/Plugins/Experimental/PCG/Source/PCG/Public/MeshSelectors/PCGMeshSelectorWeighted.h
--- a/Engine/Plugins/Experimental/PCG/Source/PCG/Public/MeshSelectors/PCGMeshSelectorWeighted.h
+++ b/Engine/Plugins/Experimental/PCG/Source/PCG/Public/MeshSelectors/PCGMeshSelectorWeighted.h
@@ -46,7 +46,8 @@ public:
 		UPARAM(ref) FPCGContext& Context, 
 		const UPCGStaticMeshSpawnerSettings* Settings, 
 		const UPCGSpatialData* InSpatialData,
-		TArray<FPCGMeshInstanceList>& OutMeshInstances) const override;
+		TArray<FPCGMeshInstanceList>& OutMeshInstances,
+		UPCGPointData* OutPointData) const override;
 
 public:
 	UPROPERTY(BlueprintReadWrite, EditAnywhere, Category = Settings)
diff --git a/Engine/Plugins/Experimental/PCG/Source/PCG/Public/MeshSelectors/PCGMeshSelectorWeightedByCategory.h b/Engine/Plugins/Experimental/PCG/Source/PCG/Public/MeshSelectors/PCGMeshSelectorWeightedByCategory.h
--- a/Engine/Plugins/Experimental/PCG/Source/PCG/Public/MeshSelectors/PCGMeshSelectorWeightedByCategory.h
+++ b/Engine/Plugins/Experimental/PCG/Source/PCG/Public/MeshSelectors/PCGMeshSelectorWeightedByCategory.h
@@ -38,7 +38,8 @@ public:
 		UPARAM(ref) FPCGContext& Context,
 		const UPCGStaticMeshSpawnerSettings* Settings,
 		const UPCGSpatialData* InSpatialData,
-		TArray<FPCGMeshInstanceList>& OutMeshInstances) const override;
+		TArray<FPCGMeshInstanceList>& OutMeshInstances,
+		UPCGPointData* OutPointData) const override;
 
 public:
 	UPROPERTY(BlueprintReadWrite, EditAnywhere, Category = Settings)
