PCG:
*Removed Node Type from the graph nodes and enabled node renaming for native pcg nodes.

diff --git a/Engine/Plugins/Experimental/PCG/Source/PCGEditor/Private/PCGEditorGraph.cpp b/Engine/Plugins/Experimental/PCG/Source/PCGEditor/Private/PCGEditorGraph.cpp
--- a/Engine/Plugins/Experimental/PCG/Source/PCGEditor/Private/PCGEditorGraph.cpp
+++ b/Engine/Plugins/Experimental/PCG/Source/PCGEditor/Private/PCGEditorGraph.cpp
@@ -22,14 +22,14 @@ void UPCGEditorGraph::InitFromNodeGraph(UPCGGraph* InPCGGraph)
 	UPCGNode* InputNode = PCGGraph->GetInputNode();
 	FGraphNodeCreator<UPCGEditorGraphNodeInput> InputNodeCreator(*this);
 	UPCGEditorGraphNodeInput* InputGraphNode = InputNodeCreator.CreateNode(bSelectNewNode);
-	InputGraphNode->Construct(InputNode, EPCGEditorGraphNodeType::Input);
+	InputGraphNode->Construct(InputNode);
 	InputNodeCreator.Finalize();
 	NodeLookup.Add(InputNode, InputGraphNode);
 
 	UPCGNode* OutputNode = PCGGraph->GetOutputNode();
 	FGraphNodeCreator<UPCGEditorGraphNodeOutput> OutputNodeCreator(*this);
 	UPCGEditorGraphNodeOutput* OutputGraphNode = OutputNodeCreator.CreateNode(bSelectNewNode);
-	OutputGraphNode->Construct(OutputNode, EPCGEditorGraphNodeType::Output);
+	OutputGraphNode->Construct(OutputNode);
 	OutputNodeCreator.Finalize();
 	NodeLookup.Add(OutputNode, OutputGraphNode);
 
@@ -37,7 +37,7 @@ void UPCGEditorGraph::InitFromNodeGraph(UPCGGraph* InPCGGraph)
 	{
 		FGraphNodeCreator<UPCGEditorGraphNode> NodeCreator(*this);
 		UPCGEditorGraphNode* GraphNode = NodeCreator.CreateNode(bSelectNewNode);
-		GraphNode->Construct(PCGNode, EPCGEditorGraphNodeType::Settings);
+		GraphNode->Construct(PCGNode);
 		NodeCreator.Finalize();
 		NodeLookup.Add(PCGNode, GraphNode);
 	}
diff --git a/Engine/Plugins/Experimental/PCG/Source/PCGEditor/Private/PCGEditorGraphNode.cpp b/Engine/Plugins/Experimental/PCG/Source/PCGEditor/Private/PCGEditorGraphNode.cpp
--- a/Engine/Plugins/Experimental/PCG/Source/PCGEditor/Private/PCGEditorGraphNode.cpp
+++ b/Engine/Plugins/Experimental/PCG/Source/PCGEditor/Private/PCGEditorGraphNode.cpp
@@ -13,6 +13,12 @@
 
 #define LOCTEXT_NAMESPACE "PCGEditorGraphNode"
 
+UPCGEditorGraphNode::UPCGEditorGraphNode(const FObjectInitializer& ObjectInitializer)
+	: Super(ObjectInitializer)
+{
+	bCanRenameNode = true;
+}
+
 FText UPCGEditorGraphNode::GetNodeTitle(ENodeTitleType::Type TitleType) const
 {
 	if (PCGNode)
diff --git a/Engine/Plugins/Experimental/PCG/Source/PCGEditor/Private/PCGEditorGraphNode.h b/Engine/Plugins/Experimental/PCG/Source/PCGEditor/Private/PCGEditorGraphNode.h
--- a/Engine/Plugins/Experimental/PCG/Source/PCGEditor/Private/PCGEditorGraphNode.h
+++ b/Engine/Plugins/Experimental/PCG/Source/PCGEditor/Private/PCGEditorGraphNode.h
@@ -16,6 +16,8 @@ class UPCGEditorGraphNode : public UPCGEditorGraphNodeBase
 	GENERATED_BODY()
 
 public:
+	UPCGEditorGraphNode(const FObjectInitializer& ObjectInitializer);
+	
 	// ~Begin UEdGraphNode interface
 	virtual FText GetNodeTitle(ENodeTitleType::Type TitleType) const override;
 	virtual void GetNodeContextMenuActions(UToolMenu* Menu, class UGraphNodeContextMenuContext* Context) const override;
diff --git a/Engine/Plugins/Experimental/PCG/Source/PCGEditor/Private/PCGEditorGraphNodeBase.cpp b/Engine/Plugins/Experimental/PCG/Source/PCGEditor/Private/PCGEditorGraphNodeBase.cpp
--- a/Engine/Plugins/Experimental/PCG/Source/PCGEditor/Private/PCGEditorGraphNodeBase.cpp
+++ b/Engine/Plugins/Experimental/PCG/Source/PCGEditor/Private/PCGEditorGraphNodeBase.cpp
@@ -22,6 +22,19 @@
 
 #define LOCTEXT_NAMESPACE "PCGEditorGraphNodeBase"
 
+void UPCGEditorGraphNodeBase::Construct(UPCGNode* InPCGNode)
+{
+	check(InPCGNode);
+	PCGNode = InPCGNode;
+	InPCGNode->OnNodeChangedDelegate.AddUObject(this, &UPCGEditorGraphNodeBase::OnNodeChanged);
+
+	NodePosX = InPCGNode->PositionX;
+	NodePosY = InPCGNode->PositionY;
+	NodeComment = InPCGNode->NodeComment;
+	bCommentBubblePinned = InPCGNode->bCommentBubblePinned;
+	bCommentBubbleVisible = InPCGNode->bCommentBubbleVisible;
+}
+
 void UPCGEditorGraphNodeBase::BeginDestroy()
 {
 	if (PCGNode)
@@ -49,23 +62,6 @@ void UPCGEditorGraphNodeBase::PostTransacted(const FTransactionObjectEvent& Tran
 	}
 }
 
-void UPCGEditorGraphNodeBase::Construct(UPCGNode* InPCGNode, EPCGEditorGraphNodeType InNodeType)
-{
-	check(InPCGNode);
-	PCGNode = InPCGNode;
-	InPCGNode->OnNodeChangedDelegate.AddUObject(this, &UPCGEditorGraphNodeBase::OnNodeChanged);
-
-	NodePosX = InPCGNode->PositionX;
-	NodePosY = InPCGNode->PositionY;
-	NodeComment = InPCGNode->NodeComment;
-	bCommentBubblePinned = InPCGNode->bCommentBubblePinned;
-	bCommentBubbleVisible = InPCGNode->bCommentBubbleVisible;
-	
-	NodeType = InNodeType;
-
-	bCanRenameNode = (InNodeType == EPCGEditorGraphNodeType::Settings);
-}
-
 void UPCGEditorGraphNodeBase::GetNodeContextMenuActions(UToolMenu* Menu, class UGraphNodeContextMenuContext* Context) const
 {
 	if (!Context->Node)
diff --git a/Engine/Plugins/Experimental/PCG/Source/PCGEditor/Private/PCGEditorGraphNodeBase.h b/Engine/Plugins/Experimental/PCG/Source/PCGEditor/Private/PCGEditorGraphNodeBase.h
--- a/Engine/Plugins/Experimental/PCG/Source/PCGEditor/Private/PCGEditorGraphNodeBase.h
+++ b/Engine/Plugins/Experimental/PCG/Source/PCGEditor/Private/PCGEditorGraphNodeBase.h
@@ -15,21 +15,13 @@ class UPCGNode;
 class UPCGPin;
 class UToolMenu;
 
-UENUM()
-enum class EPCGEditorGraphNodeType : uint8
-{
-	Input,
-	Output,
-	Settings
-};
-
 UCLASS()
 class UPCGEditorGraphNodeBase : public UEdGraphNode
 {
 	GENERATED_BODY()
 
 public:
-	void Construct(UPCGNode* InPCGNode, EPCGEditorGraphNodeType InNodeType);
+	void Construct(UPCGNode* InPCGNode);
 
 	// ~Begin UObject interface
 	virtual void BeginDestroy() override;
@@ -73,8 +65,5 @@ protected:
 	UPROPERTY()
 	TObjectPtr<UPCGNode> PCGNode = nullptr;
 
- 	UPROPERTY()
- 	EPCGEditorGraphNodeType NodeType = EPCGEditorGraphNodeType::Settings;
-
 	bool bDisableReconstructFromNode = false;
 };
diff --git a/Engine/Plugins/Experimental/PCG/Source/PCGEditor/Private/PCGEditorGraphSchemaActions.cpp b/Engine/Plugins/Experimental/PCG/Source/PCGEditor/Private/PCGEditorGraphSchemaActions.cpp
--- a/Engine/Plugins/Experimental/PCG/Source/PCGEditor/Private/PCGEditorGraphSchemaActions.cpp
+++ b/Engine/Plugins/Experimental/PCG/Source/PCGEditor/Private/PCGEditorGraphSchemaActions.cpp
@@ -43,7 +43,7 @@ UEdGraphNode* FPCGEditorGraphSchemaAction_NewNativeElement::PerformAction(UEdGra
 
 	FGraphNodeCreator<UPCGEditorGraphNode> NodeCreator(*EditorGraph);
 	UPCGEditorGraphNode* NewNode = NodeCreator.CreateUserInvokedNode(bSelectNewNode);
-	NewNode->Construct(NewPCGNode, EPCGEditorGraphNodeType::Settings);
+	NewNode->Construct(NewPCGNode);
 	NewNode->NodePosX = Location.X;
 	NewNode->NodePosY = Location.Y;
 	NodeCreator.Finalize();
@@ -90,7 +90,7 @@ UEdGraphNode* FPCGEditorGraphSchemaAction_NewBlueprintElement::PerformAction(UEd
 
 	FGraphNodeCreator<UPCGEditorGraphNode> NodeCreator(*EditorGraph);
 	UPCGEditorGraphNode* NewNode = NodeCreator.CreateUserInvokedNode(bSelectNewNode);
-	NewNode->Construct(NewPCGNode, EPCGEditorGraphNodeType::Settings);
+	NewNode->Construct(NewPCGNode);
 	NewNode->NodePosX = Location.X;
 	NewNode->NodePosY = Location.Y;
 	NodeCreator.Finalize();
@@ -141,7 +141,7 @@ UEdGraphNode* FPCGEditorGraphSchemaAction_NewSubgraphElement::PerformAction(UEdG
 
 	FGraphNodeCreator<UPCGEditorGraphNode> NodeCreator(*EditorGraph);
 	UPCGEditorGraphNode* NewNode = NodeCreator.CreateUserInvokedNode(bSelectNewNode);
-	NewNode->Construct(NewPCGNode, EPCGEditorGraphNodeType::Settings);
+	NewNode->Construct(NewPCGNode);
 	NewNode->NodePosX = Location.X;
 	NewNode->NodePosY = Location.Y;
 	NodeCreator.Finalize();
