Fix Static Analysis warning in PCGStaticMeshSpawnerTest.cpp

#rnx

diff --git a/Engine/Plugins/Experimental/PCG/Source/PCG/Private/Tests/Elements/PCGStaticMeshSpawnerTest.cpp b/Engine/Plugins/Experimental/PCG/Source/PCG/Private/Tests/Elements/PCGStaticMeshSpawnerTest.cpp
--- a/Engine/Plugins/Experimental/PCG/Source/PCG/Private/Tests/Elements/PCGStaticMeshSpawnerTest.cpp
+++ b/Engine/Plugins/Experimental/PCG/Source/PCG/Private/Tests/Elements/PCGStaticMeshSpawnerTest.cpp
@@ -102,6 +102,10 @@ void ValidateOutPointData(
 	}
 
 	const FPCGMetadataAttribute<FString>* OutAttribute = static_cast<const FPCGMetadataAttribute<FString>*>(AttributeBase);
+	if (!Test->TestNotNull("Out attribute casted to valid type", OutAttribute)) 
+	{ 
+		return; 
+	}
 
 	const FString& MeshPath = FSoftObjectPath(ISMC->GetStaticMesh()).ToString();
 	const PCGMetadataValueKey MeshValueKey = OutAttribute->FindValue(MeshPath);
