PCG Graph UI addition for testing determinism on any given node.

Inside the PCG Graph Editor:
+ 'Run Determinism Tests' command via any node's right click context
+ 'Alt+D' keyboard shortcut for same functionality
+ Determinism ListView Tab Widget for Determinism Test Results

diff --git a/Engine/Plugins/Experimental/PCG/Source/PCG/Private/Tests/Determinism/PCGDeterminismTestsCommon.cpp b/Engine/Plugins/Experimental/PCG/Source/PCG/Private/Tests/Determinism/PCGDeterminismTestsCommon.cpp
--- a/Engine/Plugins/Experimental/PCG/Source/PCG/Private/Tests/Determinism/PCGDeterminismTestsCommon.cpp
+++ b/Engine/Plugins/Experimental/PCG/Source/PCG/Private/Tests/Determinism/PCGDeterminismTestsCommon.cpp
@@ -69,6 +69,67 @@ namespace PCGDeterminismTests
 		Settings = nullptr;
 	}
 
+	PCG_API void RunDeterminismTests(const UPCGNode* PCGNode, FPCGDeterminismResult& Result)
+	{
+		check(PCGNode);
+
+		FTestData TestData(DefaultSeed);
+		TestData.Settings = PCGNode->DefaultSettings;
+
+		// Generate random input data
+		AddInputDataBasedOnPins(TestData, PCGNode, Result);
+		Result.bIsDeterministic = ExecutionIsDeterministicSameData(TestData, PCGNode);
+
+		// TODO: Add other 4 determinism tests, ex. Order Independence, etc
+	}
+
+	void AddInputDataBasedOnPins(FTestData& TestData, const UPCGNode* PCGNode, FPCGDeterminismResult& Result)
+	{
+		check(PCGNode);
+
+		for (UPCGPin* InputPin : PCGNode->GetInputPins())
+		{
+			check(InputPin);
+
+			switch (InputPin->Properties.AllowedTypes)
+			{
+			case EPCGDataType::Point:
+				AddRandomizedMultiplePointInputData(TestData);
+				Result.DataTestedString = TEXT("Point");
+				break;
+			case EPCGDataType::Volume:
+				AddRandomizedVolumeInputData(TestData);
+				Result.DataTestedString = TEXT("Volume");
+				break;
+			case EPCGDataType::PolyLine:
+				AddRandomizedPolyLineInputData(TestData);
+				Result.DataTestedString = TEXT("PolyLine");
+				break;
+			case EPCGDataType::Primitive:
+				AddRandomizedPrimitiveInputData(TestData);
+				Result.DataTestedString = TEXT("Primitive");
+				break;
+			case EPCGDataType::Surface:
+				AddRandomizedSurfaceInputData(TestData);
+				Result.DataTestedString = TEXT("Surface");
+				break;
+			case EPCGDataType::Any:
+			case EPCGDataType::Spatial:
+				AddRandomizedMultiplePointInputData(TestData, 1000);
+				AddRandomizedVolumeInputData(TestData);
+				AddRandomizedPolyLineInputData(TestData);
+				AddRandomizedPrimitiveInputData(TestData);
+				AddRandomizedSurfaceInputData(TestData);
+				Result.DataTestedString = TEXT("Spatial");
+				break;
+			default:
+				Result.DataTestedString = TEXT("Unknown");
+				Result.AdditionalDetailString = TEXT("Unknown data Input Pin");
+				break;
+			}
+		}
+	}
+
 	void AddSinglePointInputData(FPCGDataCollection& InputData, const FVector& Location, const FName& PinName)
 	{
 		UPCGPointData* PointData = PCGTestsCommon::CreatePointData(Location);
@@ -214,7 +275,7 @@ namespace PCGDeterminismTests
 		// Find comparable data from first collection
 		for (int32 I = 0; I < FirstCollection.TaggedData.Num(); ++I)
 		{
-			if (DataIsComparable(FirstCollection.TaggedData[I].Data->GetDataType()))
+			if (DataIsComparable(FirstCollection.TaggedData[I].Data))
 			{
 				FirstCollectionComparableIndices.Emplace(I);
 			}
@@ -224,7 +285,7 @@ namespace PCGDeterminismTests
 		for (int32 I = 0; I < SecondCollection.TaggedData.Num(); ++I)
 		{
 			// Check if its a type we care about comparing
-			if (DataIsComparable(SecondCollection.TaggedData[I].Data->GetDataType()))
+			if (DataIsComparable(SecondCollection.TaggedData[I].Data))
 			{
 				SecondCollectionComparableIndices.Emplace(I);
 			}
@@ -467,7 +528,7 @@ namespace PCGDeterminismTests
 		return false;
 	}
 
-	bool DataIsComparable(EPCGDataType DataType)
+	bool DataTypeIsComparable(EPCGDataType DataType)
 	{
 		// Comparable data types
 		if (DataType == EPCGDataType::Spatial)
@@ -485,6 +546,11 @@ namespace PCGDeterminismTests
 		return false;
 	}
 
+	bool DataIsComparable(const UPCGData* Data)
+	{
+		return Data && DataTypeIsComparable(Data->GetDataType());
+	}
+
 	void ShuffleInputOrder(FTestData& TestData)
 	{
 		const int32 LastIndex = TestData.InputData.TaggedData.Num() - 1;
@@ -501,7 +567,7 @@ namespace PCGDeterminismTests
 
 	TFunction<bool(const UPCGData*, const UPCGData*)> GetCompareFunction(EPCGDataType DataType)
 	{
-		if (!DataIsComparable(DataType))
+		if (!DataTypeIsComparable(DataType))
 		{
 			// Should never reach here
 			UE_LOG(LogPCG, Warning, TEXT("Attempting to compare incomparable data."));
@@ -518,13 +584,13 @@ namespace PCGDeterminismTests
 		}
 	}
 
-	bool ExecutionIsDeterministic(FTestData& FirstTestData, FTestData& SecondTestData)
+	bool ExecutionIsDeterministic(FTestData& FirstTestData, FTestData& SecondTestData, const UPCGNode* PCGNode)
 	{
 		FPCGElementPtr FirstElement = FirstTestData.Settings->GetElement();
 		FPCGElementPtr SecondElement = SecondTestData.Settings->GetElement();
 
-		TUniquePtr<FPCGContext> FirstContext(FirstElement->Initialize(FirstTestData.InputData, FirstTestData.TestPCGComponent, nullptr));
-		TUniquePtr<FPCGContext> SecondContext(SecondElement->Initialize(SecondTestData.InputData, SecondTestData.TestPCGComponent, nullptr));
+		TUniquePtr<FPCGContext> FirstContext(FirstElement->Initialize(FirstTestData.InputData, FirstTestData.TestPCGComponent, PCGNode));
+		TUniquePtr<FPCGContext> SecondContext(SecondElement->Initialize(SecondTestData.InputData, SecondTestData.TestPCGComponent, PCGNode));
 
 		FirstContext->NumAvailableTasks = 1;
 		SecondContext->NumAvailableTasks = 1;
@@ -541,8 +607,8 @@ namespace PCGDeterminismTests
 		return PCGDeterminismTests::DataCollectionsAreIdentical(FirstContext->OutputData, SecondContext->OutputData);
 	}
 
-	bool ExecutionIsDeterministicSameData(FTestData& TestData)
+	bool ExecutionIsDeterministicSameData(FTestData& TestData, const UPCGNode* PCGNode)
 	{
-		return ExecutionIsDeterministic(TestData, TestData);
+		return ExecutionIsDeterministic(TestData, TestData, PCGNode);
 	}
 }
\ No newline at end of file
diff --git a/Engine/Plugins/Experimental/PCG/Source/PCG/Public/Tests/Determinism/PCGDeterminismTestsCommon.h b/Engine/Plugins/Experimental/PCG/Source/PCG/Public/Tests/Determinism/PCGDeterminismTestsCommon.h
--- a/Engine/Plugins/Experimental/PCG/Source/PCG/Public/Tests/Determinism/PCGDeterminismTestsCommon.h
+++ b/Engine/Plugins/Experimental/PCG/Source/PCG/Public/Tests/Determinism/PCGDeterminismTestsCommon.h
@@ -36,6 +36,22 @@ namespace PCGDeterminismTests
 		FRandomStream RandomStream;
 	};
 
+	struct FPCGDeterminismResult
+	{
+		int32 Index;
+		bool bIsDeterministic;
+		FName NodeTitle;
+		FString NodeNameString;
+		FString DataTestedString;
+		FString AdditionalDetailString;
+	};
+
+	/** Validates if a PCGNode is deterministic, updating the passed in result accordingly */
+	PCG_API void RunDeterminismTests(const UPCGNode* InPCGNode, FPCGDeterminismResult& OutResult);
+
+	/** Adds input data to test data, based on input pins' allowed types */
+	void AddInputDataBasedOnPins(FTestData& TestData, const UPCGNode* InPCGNode, FPCGDeterminismResult& OutResult);
+
 	/** Helper for adding PointData with a single Point to an InputData */
 	void AddSinglePointInputData(FPCGDataCollection& InputData, const FVector& Location, const FName& PinName = DefaultTestPinName);
 	/** Helper for adding PointData with multiple Points to an InputData */
@@ -85,7 +101,9 @@ namespace PCGDeterminismTests
 	bool ComparisonIsUnimplemented(const UPCGData* FirstData, const UPCGData* SecondData);
 
 	/** Determines if a PCG DataType is a type that should be compared */
-	bool DataIsComparable(EPCGDataType DataType);
+	bool DataTypeIsComparable(EPCGDataType DataType);
+	/** Validates PCGData and if its DataType is comparable */
+	bool DataIsComparable(const UPCGData* Data);
 
 	/** Randomizes the order of InputData */
 	void ShuffleInputOrder(FTestData& TestData);
@@ -94,9 +112,9 @@ namespace PCGDeterminismTests
 	TFunction<bool(const UPCGData*, const UPCGData*)> GetCompareFunction(EPCGDataType DataType);
 
 	/** Execute the elements for each valid input and compare if all the outputs are identical */
-	bool ExecutionIsDeterministic(FTestData& FirstTestData, FTestData& SecondTestData);
+	bool ExecutionIsDeterministic(FTestData& FirstTestData, FTestData& SecondTestData, const UPCGNode* PCGNode = nullptr);
 	/** Execute the same element twice compare if all the outputs are identical */
-	bool ExecutionIsDeterministicSameData(FTestData& TestData);
+	bool ExecutionIsDeterministicSameData(FTestData& TestData, const UPCGNode* PCGNode = nullptr);
 
 	/** Generates settings based upon a UPCGSettings subclass */
 	template<typename SettingsType>
diff --git a/Engine/Plugins/Experimental/PCG/Source/PCGEditor/Private/PCGEditor.cpp b/Engine/Plugins/Experimental/PCG/Source/PCGEditor/Private/PCGEditor.cpp
--- a/Engine/Plugins/Experimental/PCG/Source/PCGEditor/Private/PCGEditor.cpp
+++ b/Engine/Plugins/Experimental/PCG/Source/PCGEditor/Private/PCGEditor.cpp
@@ -19,8 +19,10 @@
 #include "PCGEditorGraphSchemaActions.h"
 #include "PCGEditorSettings.h"
 #include "PCGEditorUtils.h"
+#include "SPCGEditorGraphDeterminism.h"
 #include "SPCGEditorGraphFind.h"
 #include "SPCGEditorGraphNodePalette.h"
+#include "Tests/Determinism/PCGDeterminismTestsCommon.h"
 
 #include "EdGraphUtilities.h"
 #include "Editor.h"
@@ -59,6 +61,7 @@ namespace FPCGEditor_private
 	const FName AttributesID = FName(TEXT("Attributes"));
 	const FName ViewportID = FName(TEXT("Viewport"));
 	const FName FindID = FName(TEXT("Find"));
+	const FName DeterminismID = FName(TEXT("Determinism"));
 }
 
 void FPCGEditor::Initialize(const EToolkitMode::Type InMode, const TSharedPtr<class IToolkitHost>& InToolkitHost, UPCGGraph* InPCGGraph)
@@ -80,6 +83,7 @@ void FPCGEditor::Initialize(const EToolkitMode::Type InMode, const TSharedPtr<cl
 	GraphEditorWidget = CreateGraphEditorWidget();
 	PaletteWidget = CreatePaletteWidget();
 	FindWidget = CreateFindWidget();
+	DeterminismWidget = CreateDeterminismWidget();
 
 	BindCommands();
 
@@ -89,7 +93,6 @@ void FPCGEditor::Initialize(const EToolkitMode::Type InMode, const TSharedPtr<cl
 			FTabManager::NewPrimaryArea()->SetOrientation(Orient_Horizontal)
 			->Split
 			(
-				
 				FTabManager::NewSplitter()->SetOrientation(Orient_Vertical)
 				->SetSizeCoefficient(0.10f)
 				->Split
@@ -124,6 +127,7 @@ void FPCGEditor::Initialize(const EToolkitMode::Type InMode, const TSharedPtr<cl
 					->SetSizeCoefficient(0.28)
 					->SetHideTabWell(true)
 					->AddTab(FPCGEditor_private::AttributesID, ETabState::OpenedTab)
+					->AddTab(FPCGEditor_private::DeterminismID, ETabState::ClosedTab)
 					->AddTab(FPCGEditor_private::FindID, ETabState::ClosedTab)
 				)
 			)
@@ -185,6 +189,10 @@ void FPCGEditor::RegisterTabSpawners(const TSharedRef<FTabManager>& InTabManager
 	InTabManager->RegisterTabSpawner(FPCGEditor_private::FindID, FOnSpawnTab::CreateSP(this, &FPCGEditor::SpawnTab_Find))
 		.SetDisplayName(LOCTEXT("FindTab", "Find"))
 		.SetGroup(WorkspaceMenuCategoryRef);
+
+	InTabManager->RegisterTabSpawner(FPCGEditor_private::DeterminismID, FOnSpawnTab::CreateSP(this, &FPCGEditor::SpawnTab_Determinism))
+		.SetDisplayName(LOCTEXT("DeterminismTab", "Determinism"))
+		.SetGroup(WorkspaceMenuCategoryRef);
 }
 
 void FPCGEditor::UnregisterTabSpawners(const TSharedRef<class FTabManager>& InTabManager)
@@ -194,6 +202,8 @@ void FPCGEditor::UnregisterTabSpawners(const TSharedRef<class FTabManager>& InTa
 	InTabManager->UnregisterTabSpawner(FPCGEditor_private::PaletteID);
 	InTabManager->UnregisterTabSpawner(FPCGEditor_private::AttributesID);
 	InTabManager->UnregisterTabSpawner(FPCGEditor_private::ViewportID);
+	InTabManager->UnregisterTabSpawner(FPCGEditor_private::FindID);
+	InTabManager->UnregisterTabSpawner(FPCGEditor_private::DeterminismID);
 
 	FAssetEditorToolkit::UnregisterTabSpawners(InTabManager);
 }
@@ -254,6 +264,11 @@ void FPCGEditor::BindCommands()
 		PCGEditorCommands.CollapseNodes,
 		FExecuteAction::CreateSP(this, &FPCGEditor::OnCollapseNodesInSubgraph),
 		FCanExecuteAction::CreateSP(this, &FPCGEditor::CanCollapseNodesInSubgraph));
+
+	GraphEditorCommands->MapAction(
+		PCGEditorCommands.RunDeterminismTest,
+		FExecuteAction::CreateSP(this, &FPCGEditor::OnDeterminismTests),
+		FCanExecuteAction::CreateSP(this, &FPCGEditor::CanRunDeterminismTests));
 }
 
 void FPCGEditor::OnFind()
@@ -265,6 +280,64 @@ void FPCGEditor::OnFind()
 	}
 }
 
+bool FPCGEditor::CanRunDeterminismTests() const
+{
+	check(GraphEditorWidget.IsValid());
+
+	bool bValidObjectFound = false;
+	for (const UObject* Object : GraphEditorWidget->GetSelectedNodes())
+	{
+		if (Cast<const UPCGEditorGraphNodeBase>(Object) && !Cast<const UPCGEditorGraphNodeInput>(Object) && !Cast<const UPCGEditorGraphNodeOutput>(Object))
+		{
+			bValidObjectFound = true;
+			break;
+		}
+	}
+
+	return bValidObjectFound;
+}
+
+void FPCGEditor::OnDeterminismTests()
+{
+	check(GraphEditorWidget.IsValid());
+
+	if (!DeterminismWidget.IsValid() || !DeterminismWidget->IsContructed())
+	{
+		return;
+	}
+
+	DeterminismWidget->Clear();
+
+	int32 Index = 1;
+	for (UObject* Object : GraphEditorWidget->GetSelectedNodes())
+	{
+		check(Object);
+
+		if (!Object->IsA<UPCGEditorGraphNodeInput>() && !Object->IsA<UPCGEditorGraphNodeOutput>())
+		{
+			if (const UPCGEditorGraphNodeBase* PCGEditorGraphNode = Cast<const UPCGEditorGraphNodeBase>(Object))
+			{
+				const UPCGNode* PCGNode = PCGEditorGraphNode->GetPCGNode();
+				check(PCGNode);
+
+				TSharedPtr<PCGDeterminismTests::FPCGDeterminismResult> NodeResult = MakeShared<PCGDeterminismTests::FPCGDeterminismResult>();
+				NodeResult->Index = Index++;
+				NodeResult->NodeTitle = PCGNode->GetNodeTitle();
+				NodeResult->NodeNameString = PCGNode->GetName();
+				PCGDeterminismTests::RunDeterminismTests(PCGNode, *NodeResult);
+
+				DeterminismWidget->AddItem(NodeResult);
+			}
+		}
+	}
+
+	// Give focus to the Determinism Output Tab
+	if (TabManager.IsValid())
+	{
+		TabManager->TryInvokeTab(FPCGEditor_private::DeterminismID);
+	}
+}
+
 bool FPCGEditor::CanCollapseNodesInSubgraph()
 {
 	bool HasPCGNode = false;
@@ -1098,6 +1171,11 @@ TSharedRef<SPCGEditorGraphFind> FPCGEditor::CreateFindWidget()
 	return SNew(SPCGEditorGraphFind, SharedThis(this));
 }
 
+TSharedRef<SPCGEditorGraphDeterminismListView> FPCGEditor::CreateDeterminismWidget()
+{
+	return SNew(SPCGEditorGraphDeterminismListView, SharedThis(this));
+}
+
 void FPCGEditor::OnSelectedNodesChanged(const TSet<UObject*>& NewSelection)
 {
 	TArray<TWeakObjectPtr<UObject>> SelectedObjects;
@@ -1251,4 +1329,14 @@ TSharedRef<SDockTab> FPCGEditor::SpawnTab_Find(const FSpawnTabArgs& Args)
 		];
 }
 
+TSharedRef<SDockTab> FPCGEditor::SpawnTab_Determinism(const FSpawnTabArgs& Args)
+{
+	return SNew(SDockTab)
+		.Label(LOCTEXT("PCGDeterminismTitle", "Determinism"))
+		.TabColorScale(GetTabColorScale())
+		[
+			DeterminismWidget.ToSharedRef()
+		];
+}
+
 #undef LOCTEXT_NAMESPACE
diff --git a/Engine/Plugins/Experimental/PCG/Source/PCGEditor/Private/PCGEditor.h b/Engine/Plugins/Experimental/PCG/Source/PCGEditor/Private/PCGEditor.h
--- a/Engine/Plugins/Experimental/PCG/Source/PCGEditor/Private/PCGEditor.h
+++ b/Engine/Plugins/Experimental/PCG/Source/PCGEditor/Private/PCGEditor.h
@@ -9,6 +9,7 @@
 class FUICommandList;
 class IDetailsView;
 class SGraphEditor;
+class SPCGEditorGraphDeterminismListView;
 class SPCGEditorGraphFind;
 class SPCGEditorGraphNodePalette;
 class UPCGEditorGraph;
@@ -58,6 +59,11 @@ private:
 	/** Bring up the find tab */
 	void OnFind();
 
+	/** Can determinism be tested on this node */
+	bool CanRunDeterminismTests() const;
+	/** Summon the Determinism tab */
+	void OnDeterminismTests();
+
 	/** Select every node in the graph */
 	void SelectAllNodes();
 	/** Whether we can select every node */
@@ -115,6 +121,9 @@ private:
 	/** Create new find widget */
 	TSharedRef<SPCGEditorGraphFind> CreateFindWidget();
 
+	/** Create a new determinism tab widget */
+	TSharedRef<SPCGEditorGraphDeterminismListView> CreateDeterminismWidget();
+
 	/** Called when the selection changes in the GraphEditor */
 	void OnSelectedNodesChanged(const TSet<UObject*>& NewSelection);
 
@@ -144,11 +153,13 @@ private:
 	TSharedRef<SDockTab> SpawnTab_Attributes(const FSpawnTabArgs& Args);
 	TSharedRef<SDockTab> SpawnTab_Viewport(const FSpawnTabArgs& Args);
 	TSharedRef<SDockTab> SpawnTab_Find(const FSpawnTabArgs& Args);
+	TSharedRef<SDockTab> SpawnTab_Determinism(const FSpawnTabArgs& Args);
 
 	TSharedPtr<SGraphEditor> GraphEditorWidget;
 	TSharedPtr<IDetailsView> PropertyDetailsWidget;
 	TSharedPtr<SPCGEditorGraphNodePalette> PaletteWidget;
 	TSharedPtr<SPCGEditorGraphFind> FindWidget;
+	TSharedPtr<SPCGEditorGraphDeterminismListView> DeterminismWidget;
 
 	TSharedPtr<FUICommandList> GraphEditorCommands;
 
diff --git a/Engine/Plugins/Experimental/PCG/Source/PCGEditor/Private/PCGEditorCommands.cpp b/Engine/Plugins/Experimental/PCG/Source/PCGEditor/Private/PCGEditorCommands.cpp
--- a/Engine/Plugins/Experimental/PCG/Source/PCGEditor/Private/PCGEditorCommands.cpp
+++ b/Engine/Plugins/Experimental/PCG/Source/PCGEditor/Private/PCGEditorCommands.cpp
@@ -20,6 +20,7 @@ void FPCGEditorCommands::RegisterCommands()
 {
 	UI_COMMAND(Find, "Find", "Finds PCG nodes and comments in the current graph", EUserInterfaceActionType::Button, FInputChord(EModifierKey::Control, EKeys::F));
 	UI_COMMAND(CollapseNodes, "Collapse into Subgraph", "Collapse selected nodes into a separate PCGGraph asset.", EUserInterfaceActionType::Button, FInputChord(EModifierKey::Alt, EKeys::J));
+	UI_COMMAND(RunDeterminismTest, "Run Determinism Test", "Evaluate the current node for determinism.", EUserInterfaceActionType::Button, FInputChord(EModifierKey::Alt, EKeys::D));
 }
 
 #undef LOCTEXT_NAMESPACE
diff --git a/Engine/Plugins/Experimental/PCG/Source/PCGEditor/Private/PCGEditorCommands.h b/Engine/Plugins/Experimental/PCG/Source/PCGEditor/Private/PCGEditorCommands.h
--- a/Engine/Plugins/Experimental/PCG/Source/PCGEditor/Private/PCGEditorCommands.h
+++ b/Engine/Plugins/Experimental/PCG/Source/PCGEditor/Private/PCGEditorCommands.h
@@ -16,4 +16,5 @@ public:
 
 	TSharedPtr<FUICommandInfo> Find;
 	TSharedPtr<FUICommandInfo> CollapseNodes;
+	TSharedPtr<FUICommandInfo> RunDeterminismTest;
 };
diff --git a/Engine/Plugins/Experimental/PCG/Source/PCGEditor/Private/PCGEditorGraphNodeBase.cpp b/Engine/Plugins/Experimental/PCG/Source/PCGEditor/Private/PCGEditorGraphNodeBase.cpp
--- a/Engine/Plugins/Experimental/PCG/Source/PCGEditor/Private/PCGEditorGraphNodeBase.cpp
+++ b/Engine/Plugins/Experimental/PCG/Source/PCGEditor/Private/PCGEditorGraphNodeBase.cpp
@@ -111,6 +111,13 @@ void UPCGEditorGraphNodeBase::GetNodeContextMenuActions(UToolMenu* Menu, class U
 			LOCTEXT("MultiCommentDesc", "Create Comment from Selection"),
 			LOCTEXT("CommentToolTip", "Create a resizable comment box around selection."));
 	}
+
+	{
+		FToolMenuSection& Section = Menu->AddSection("EdGraphSchemaDeterminism", LOCTEXT("DeterminismHeader", "Determinism"));
+		Section.AddMenuEntry(FPCGEditorCommands::Get().RunDeterminismTest,
+			LOCTEXT("Determinism_RunTest", "Validate Determinism on Selection"),
+			LOCTEXT("Determinism_RunTestToolTip", "Run a test to validate the selected nodes for determinism."));
+	}
 }
 
 void UPCGEditorGraphNodeBase::AutowireNewNode(UEdGraphPin* FromPin)
diff --git a/Engine/Plugins/Experimental/PCG/Source/PCGEditor/Private/SPCGEditorGraphDeterminism.cpp b/Engine/Plugins/Experimental/PCG/Source/PCGEditor/Private/SPCGEditorGraphDeterminism.cpp
new file mode 100644
--- /dev/null
+++ b/Engine/Plugins/Experimental/PCG/Source/PCGEditor/Private/SPCGEditorGraphDeterminism.cpp
@@ -0,0 +1,148 @@
+// Copyright Epic Games, Inc. All Rights Reserved.
+
+#include "SPCGEditorGraphDeterminism.h"
+
+#include "PCGEditor.h"
+
+#define LOCTEXT_NAMESPACE "SPCGEditorGraphDeterminismListView"
+
+namespace
+{
+	const FName NAME_Index(TEXT("IndexColumn"));
+	const FName NAME_NodeTitle(TEXT("NodeTitleColumn"));
+	const FName NAME_NodeName(TEXT("NodeNameColumn"));
+	const FName NAME_Data(TEXT("DataColumn"));
+	const FName NAME_Result(TEXT("ResultColumn"));
+	const FName NAME_AdditionalDetails(TEXT("AdditionalDetailsColumn"));
+
+	const FText TEXT_Index(LOCTEXT("IndexLabel", ""));
+	const FText TEXT_NodeTitle(LOCTEXT("NodeTitleLabel", "Title"));
+	const FText TEXT_NodeName(LOCTEXT("NodeNameLabel", "Name"));
+	const FText TEXT_Data(LOCTEXT("DataLabel", "Input Data"));
+	const FText TEXT_Result(LOCTEXT("ResultLabel", "Deterministic"));
+	const FText TEXT_AdditionalDetails(LOCTEXT("AdditionalDetailsLabel", "Additional Details"));
+
+	const FText TEXT_Yes(LOCTEXT("Yes", "Yes"));
+	const FText TEXT_No(LOCTEXT("No", "No"));
+}
+
+void SPCGEditorGraphDeterminismRow::Construct(const FArguments& InArgs, const TSharedRef<STableViewBase>& InOwnerTableView, const FPCGDeterminismResultPtr& Item)
+{
+	CurrentItem = Item;
+	SMultiColumnTableRow<FPCGDeterminismResultPtr>::Construct(SMultiColumnTableRow::FArguments(), InOwnerTableView);
+}
+
+TSharedRef<SWidget> SPCGEditorGraphDeterminismRow::GenerateWidgetForColumn(const FName& ColumnId)
+{
+	check(CurrentItem);
+
+	FText CellText = LOCTEXT("UnknownColumn", "Unknown");
+
+	if (ColumnId == NAME_Index)
+	{
+		CellText = FText::FromString(FString::FromInt(CurrentItem->Index));
+	}
+	else if (ColumnId == NAME_NodeTitle)
+	{
+		CellText = FText::FromString(CurrentItem->NodeTitle.ToString());
+	}
+	else if (ColumnId == NAME_NodeName)
+	{
+		CellText = FText::FromString(CurrentItem->NodeNameString);
+	}
+	else if (ColumnId == NAME_Data)
+	{
+		CellText = FText::FromString(CurrentItem->DataTestedString);
+	}
+	else if (ColumnId == NAME_Result)
+	{
+		if (CurrentItem->bIsDeterministic)
+		{
+			return SNew(STextBlock)
+					.Text(TEXT_Yes)
+					.ColorAndOpacity(FColor::Green);
+		}
+		else
+		{
+			return SNew(STextBlock)
+					.Text(TEXT_No)
+					.ColorAndOpacity(FColor::Red);
+		}
+	}
+	else if (ColumnId == NAME_AdditionalDetails)
+	{
+		CellText = FText::FromString(CurrentItem->AdditionalDetailString);
+	}
+
+	return SNew(STextBlock).Text(CellText);
+}
+
+void SPCGEditorGraphDeterminismListView::Construct(const FArguments& InArgs, TWeakPtr<FPCGEditor> InPCGEditor)
+{
+	check(InPCGEditor.IsValid() && !bIsConstructed);
+	PCGEditorPtr = InPCGEditor;
+
+	SAssignNew(ListView, SListView<FPCGDeterminismResultPtr>)
+		.ListItemsSource(&ListViewItems)
+		.ItemHeight(36)
+		.OnGenerateRow(this, &SPCGEditorGraphDeterminismListView::OnGenerateRow)
+		.HeaderRow
+		(
+			SNew(SHeaderRow)
+			+ SHeaderRow::Column(NAME_Index)
+			.DefaultLabel(TEXT_Index)
+			.ManualWidth(25)
+			.HAlignCell(HAlign_Center)
+			+ SHeaderRow::Column(NAME_NodeTitle)
+			.DefaultLabel(TEXT_NodeTitle)
+			.ManualWidth(160)
+			.HAlignCell(HAlign_Left)
+			+ SHeaderRow::Column(NAME_NodeName)
+			.DefaultLabel(TEXT_NodeName)
+			.ManualWidth(160)
+			.HAlignCell(HAlign_Left)
+			+ SHeaderRow::Column(NAME_Data)
+			.DefaultLabel(TEXT_Data)
+			.ManualWidth(70)
+			.HAlignCell(HAlign_Center)
+			+ SHeaderRow::Column(NAME_Result)
+			.DefaultLabel(TEXT_Result)
+			.ManualWidth(90)
+			.HAlignCell(HAlign_Center)
+			+ SHeaderRow::Column(NAME_AdditionalDetails)
+			.DefaultLabel(TEXT_AdditionalDetails)
+			.HAlignCell(HAlign_Left)
+		);
+
+	ChildSlot
+		[
+			ListView->AsShared()
+		];
+
+	bIsConstructed = true;
+}
+
+void SPCGEditorGraphDeterminismListView::AddItem(const FPCGDeterminismResultPtr Item)
+{
+	check(Item.IsValid());
+	ListViewItems.Add(Item);
+	ListView->RequestListRefresh();
+}
+
+void SPCGEditorGraphDeterminismListView::Clear()
+{
+	ListViewItems.Empty();
+	ListView->RequestListRefresh();
+}
+
+bool SPCGEditorGraphDeterminismListView::IsContructed() const
+{
+	return bIsConstructed;
+}
+
+TSharedRef<ITableRow> SPCGEditorGraphDeterminismListView::OnGenerateRow(const FPCGDeterminismResultPtr Item, const TSharedRef<STableViewBase>& OwnerTable)
+{
+	return SNew(SPCGEditorGraphDeterminismRow, OwnerTable, Item);
+}
+
+#undef LOCTEXT_NAMESPACE
\ No newline at end of file
diff --git a/Engine/Plugins/Experimental/PCG/Source/PCGEditor/Private/SPCGEditorGraphDeterminism.h b/Engine/Plugins/Experimental/PCG/Source/PCGEditor/Private/SPCGEditorGraphDeterminism.h
new file mode 100644
--- /dev/null
+++ b/Engine/Plugins/Experimental/PCG/Source/PCGEditor/Private/SPCGEditorGraphDeterminism.h
@@ -0,0 +1,52 @@
+// Copyright Epic Games, Inc. All Rights Reserved.
+
+#pragma once
+
+#include "Tests/Determinism/PCGDeterminismTestsCommon.h"
+#include "Widgets/Views/SListView.h"
+
+class FPCGEditor;
+
+typedef TSharedPtr<PCGDeterminismTests::FPCGDeterminismResult> FPCGDeterminismResultPtr;
+
+class SPCGEditorGraphDeterminismRow : public SMultiColumnTableRow<FPCGDeterminismResultPtr>
+{
+	SLATE_BEGIN_ARGS(SPCGEditorGraphDeterminismRow) {}
+	SLATE_END_ARGS()
+
+	/** Construct a row of the ListView */
+	void Construct(const FArguments& InArgs, const TSharedRef<STableViewBase>& InOwnerTableView, const FPCGDeterminismResultPtr& Item);
+
+	/** Generates a column, given the column's ID */
+	virtual TSharedRef<SWidget> GenerateWidgetForColumn(const FName& ColumnId) override;
+
+protected:
+	FPCGDeterminismResultPtr CurrentItem;
+};
+
+class SPCGEditorGraphDeterminismListView : public SCompoundWidget
+{
+public:
+	SLATE_BEGIN_ARGS(SPCGEditorGraphDeterminismListView) { }
+	SLATE_END_ARGS()
+
+	/** Construct the ListView */
+	void Construct(const FArguments& InArgs, TWeakPtr<FPCGEditor> InPCGEditor);
+	/** Add an item to the ListView */
+	void AddItem(const FPCGDeterminismResultPtr Item);
+	/** Clear all items from the ListView */
+	void Clear();
+	/** Validates if the ListView has been constructed */
+	bool IsContructed() const;
+
+private:
+	TSharedRef<ITableRow> OnGenerateRow(const FPCGDeterminismResultPtr Item, const TSharedRef<STableViewBase>& OwnerTable);
+
+	/** Weak pointer to the PCGEditor */
+	TWeakPtr<FPCGEditor> PCGEditorPtr;
+
+	TSharedPtr<SListView<FPCGDeterminismResultPtr>> ListView;
+	TArray<FPCGDeterminismResultPtr> ListViewItems;
+
+	bool bIsConstructed = false;
+};
\ No newline at end of file
