Added other callback for position change (when driven by blueprint)
Made the generate/cleanup calls multicast so it works across replicated actors
Made the spawn actor skip replicated actors when the owning component isn't the authority - we can review later as in the eventuality that there is a PCG component on these, they would need to be called from the server.

diff --git a/Engine/Plugins/Experimental/PCG/Source/PCG/Private/Elements/PCGSpawnActor.cpp b/Engine/Plugins/Experimental/PCG/Source/PCG/Private/Elements/PCGSpawnActor.cpp
--- a/Engine/Plugins/Experimental/PCG/Source/PCG/Private/Elements/PCGSpawnActor.cpp
+++ b/Engine/Plugins/Experimental/PCG/Source/PCG/Private/Elements/PCGSpawnActor.cpp
@@ -103,6 +103,9 @@ bool FPCGSpawnActorElement::ExecuteInternal(FPCGContext* Context) const
 			continue;
 		}
 
+		const bool bHasAuthority = !Context->SourceComponent || (Context->SourceComponent->GetOwner() && Context->SourceComponent->GetOwner()->HasAuthority());
+		const bool bSpawnedActorsRequireAuthority = Settings->TemplateActorClass->GetDefaultObject<AActor>()->GetIsReplicated();
+
 		// First, create target instance transforms
 		const UPCGPointData* PointData = SpatialData->ToPointData(Context);
 
@@ -169,7 +172,7 @@ bool FPCGSpawnActorElement::ExecuteInternal(FPCGContext* Context) const
 						Instances.Add(Point.Transform);
 					}
 				}
-				else if (Settings->Option != EPCGSpawnActorOption::CollapseActors)
+				else if (Settings->Option != EPCGSpawnActorOption::CollapseActors && (bHasAuthority || !bSpawnedActorsRequireAuthority))
 				{
 					FActorSpawnParameters SpawnParams;
 					SpawnParams.Owner = TargetActor;
diff --git a/Engine/Plugins/Experimental/PCG/Source/PCG/Private/Grid/PCGPartitionActor.cpp b/Engine/Plugins/Experimental/PCG/Source/PCG/Private/Grid/PCGPartitionActor.cpp
--- a/Engine/Plugins/Experimental/PCG/Source/PCG/Private/Grid/PCGPartitionActor.cpp
+++ b/Engine/Plugins/Experimental/PCG/Source/PCG/Private/Grid/PCGPartitionActor.cpp
@@ -149,7 +149,7 @@ bool APCGPartitionActor::RemoveGraphInstance(UPCGComponent* OriginalComponent)
 	LocalToOriginalMap.Remove(LocalComponent);
 
 	// TODO Add option to not cleanup?
-	LocalComponent->Cleanup(/*bRemoveComponents=*/true);
+	LocalComponent->CleanupLocal(/*bRemoveComponents=*/true);
 	LocalComponent->DestroyComponent();
 
 	return OriginalToLocalMap.IsEmpty();
@@ -182,7 +182,7 @@ bool APCGPartitionActor::CleanupDeadGraphInstances()
 
 		if (DeadInstance)
 		{
-			DeadInstance->Cleanup(/*bRemoveComponents=*/true);
+			DeadInstance->CleanupLocal(/*bRemoveComponents=*/true);
 			DeadInstance->DestroyComponent();
 		}
 	}
diff --git a/Engine/Plugins/Experimental/PCG/Source/PCG/Private/PCGComponent.cpp b/Engine/Plugins/Experimental/PCG/Source/PCG/Private/PCGComponent.cpp
--- a/Engine/Plugins/Experimental/PCG/Source/PCG/Private/PCGComponent.cpp
+++ b/Engine/Plugins/Experimental/PCG/Source/PCG/Private/PCGComponent.cpp
@@ -180,10 +180,17 @@ void UPCGComponent::Generate()
 	FScopedTransaction Transaction(LOCTEXT("PCGGenerate", "Execute generation on PCG component"));
 #endif
 
-	Generate(/*bForce=*/PCGComponent::bSaveOnCleanupAndGenerate);
+	GenerateLocal(/*bForce=*/PCGComponent::bSaveOnCleanupAndGenerate);
 }
 
-void UPCGComponent::Generate(bool bForce)
+void UPCGComponent::Generate_Implementation(bool bForce)
+{
+	// Force component activation so it's easier to control by BP.
+	bActivated = true;
+	GenerateLocal(bForce);
+}
+
+void UPCGComponent::GenerateLocal(bool bForce)
 {
 #if WITH_EDITOR
 	if (bIsGenerating)
@@ -324,10 +331,15 @@ void UPCGComponent::Cleanup()
 	FScopedTransaction Transaction(LOCTEXT("PCGCleanup", "Clean up PCG component"));
 #endif
 
-	Cleanup(/*bRemoveComponents=*/true, /*bSave=*/PCGComponent::bSaveOnCleanupAndGenerate);
+	CleanupLocal(/*bRemoveComponents=*/true, /*bSave=*/PCGComponent::bSaveOnCleanupAndGenerate);
 }
 
-void UPCGComponent::Cleanup(bool bRemoveComponents, bool bSave)
+void UPCGComponent::Cleanup_Implementation(bool bRemoveComponents, bool bSave)
+{
+	CleanupLocal(bRemoveComponents, bSave);
+}
+
+void UPCGComponent::CleanupLocal(bool bRemoveComponents, bool bSave)
 {
 	if (!bGenerated || !GetSubsystem())
 	{
@@ -408,7 +420,7 @@ void UPCGComponent::BeginPlay()
 
 	if(bActivated && !bGenerated && !IsPartitioned())
 	{
-		Generate(/*bForce=*/false);
+		GenerateLocal(/*bForce=*/false);
 		bRuntimeGenerated = true;
 	}
 }
@@ -626,7 +638,7 @@ void UPCGComponent::PreEditUndo()
 	if (bGenerated)
 	{
 		// Cleanup so managed resources are cleaned in all cases
-		Cleanup(/*bRemoveComponents=*/true, /*bSave=*/PCGComponent::bSaveOnCleanupAndGenerate);
+		CleanupLocal(/*bRemoveComponents=*/true, /*bSave=*/PCGComponent::bSaveOnCleanupAndGenerate);
 		// Put back generated flag to its original value so it is captured properly
 		bGenerated = true;
 	}	
@@ -658,12 +670,14 @@ void UPCGComponent::PostEditUndo()
 void UPCGComponent::SetupActorCallbacks()
 {
 	GEngine->OnActorMoved().AddUObject(this, &UPCGComponent::OnActorMoved);
+	GEngine->OnComponentTransformChanged().AddUObject(this, &UPCGComponent::OnComponentTransformChanged);
 	FCoreUObjectDelegates::OnObjectPropertyChanged.AddUObject(this, &UPCGComponent::OnObjectPropertyChanged);
 }
 
 void UPCGComponent::TeardownActorCallbacks()
 {
 	FCoreUObjectDelegates::OnObjectPropertyChanged.RemoveAll(this);
+	GEngine->OnComponentTransformChanged().RemoveAll(this);
 	GEngine->OnActorMoved().RemoveAll(this);
 }
 
@@ -827,6 +841,17 @@ void UPCGComponent::OnActorMoved(AActor* InActor)
 	}
 }
 
+void UPCGComponent::OnComponentTransformChanged(USceneComponent* InComponent, ETeleportType InTeleport)
+{
+	if (!InComponent || InComponent->GetOwner() != GetOwner())
+	{
+		return;
+	}
+
+	DirtyGenerated(EPCGComponentDirtyFlag::Actor);
+	Refresh();
+}
+
 void UPCGComponent::UpdateTrackedLandscape(bool bBoundsCheck)
 {
 	TeardownLandscapeTracking();
@@ -1105,7 +1130,7 @@ void UPCGComponent::Refresh()
 		else
 		{
 			bool bWasGenerated = bGenerated;
-			Cleanup(/*bRemoveComponents=*/false);
+			CleanupLocal(/*bRemoveComponents=*/false);
 			bGenerated = bWasGenerated;
 		}
 	}
@@ -1113,7 +1138,7 @@ void UPCGComponent::Refresh()
 	{
 		if (bGenerated && bRegenerateInEditor)
 		{
-			Generate(/*bForce=*/false);
+			GenerateLocal(/*bForce=*/false);
 		}
 		else if (IsPartitioned() && GetSubsystem())
 		{
diff --git a/Engine/Plugins/Experimental/PCG/Source/PCG/Public/PCGComponent.h b/Engine/Plugins/Experimental/PCG/Source/PCG/Public/PCGComponent.h
--- a/Engine/Plugins/Experimental/PCG/Source/PCG/Public/PCGComponent.h
+++ b/Engine/Plugins/Experimental/PCG/Source/PCG/Public/PCGComponent.h
@@ -83,10 +83,17 @@ public:
 	void Generate();
 	void Cleanup();
 
-	UFUNCTION(BlueprintCallable, Category = PCG)
+	/** Starts generation from a local (vs. remote) standpoint. Will not be replicated. */
+	void GenerateLocal(bool bForce);
+	/** Cleans up the generation from a local (vs. remote) standpoint. Will not be replicated. */
+	void CleanupLocal(bool bRemoveComponents, bool bSave = false);
+
+	/** Networked generation call that also activates the component as needed */
+	UFUNCTION(BlueprintCallable, NetMulticast, Reliable, Category = PCG)
 	void Generate(bool bForce);
 
-	UFUNCTION(BlueprintCallable, Category = PCG)
+	/** Networked cleanup call */
+	UFUNCTION(BlueprintCallable, NetMulticast, Reliable, Category = PCG)
 	void Cleanup(bool bRemoveComponents, bool bSave = false);
 
 	UPROPERTY(BlueprintReadWrite, EditAnywhere, Category = Settings)
@@ -185,6 +192,7 @@ private:
 	void OnActorAdded(AActor* InActor);
 	void OnActorDeleted(AActor* InActor);
 	void OnActorMoved(AActor* InActor);
+	void OnComponentTransformChanged(USceneComponent* InComponent, ETeleportType InTeleport);
 	void OnObjectPropertyChanged(UObject* InObject, FPropertyChangedEvent& InEvent);
 	bool ActorHasExcludedTag(AActor* InActor) const;
 	bool ActorIsTracked(AActor* InActor) const;
