Add a "Clear PCG link" option in the PCG Component

diff --git a/Engine/Plugins/Experimental/PCG/Source/PCG/Private/Helpers/PCGActorHelpers.cpp b/Engine/Plugins/Experimental/PCG/Source/PCG/Private/Helpers/PCGActorHelpers.cpp
--- a/Engine/Plugins/Experimental/PCG/Source/PCG/Private/Helpers/PCGActorHelpers.cpp
+++ b/Engine/Plugins/Experimental/PCG/Source/PCG/Private/Helpers/PCGActorHelpers.cpp
@@ -292,7 +292,7 @@ void UPCGActorHelpers::GetActorClassDefaultComponents(const TSubclassOf<AActor>&
 	OutComponents = MoveTemp(ResultComponents);
 }
 
-void UPCGActorHelpers::ForEachActorInLevel(TObjectPtr<ULevel> Level, TSubclassOf<AActor> ActorClass, TFunctionRef<void(AActor*)> Callback)
+void UPCGActorHelpers::ForEachActorInLevel(ULevel* Level, TSubclassOf<AActor> ActorClass, TFunctionRef<void(AActor*)> Callback)
 {
 	if (!Level)
 	{
@@ -307,3 +307,50 @@ void UPCGActorHelpers::ForEachActorInLevel(TObjectPtr<ULevel> Level, TSubclassOf
 		}
 	}
 }
+
+AActor* UPCGActorHelpers::SpawnDefaultActor(UWorld* World, TSubclassOf<AActor> ActorClass, FName BaseName, const FTransform& Transform, AActor* Parent)
+{
+	if (!World || !ActorClass)
+	{
+		return nullptr;
+	}
+
+	FActorSpawnParameters SpawnParams;
+	SpawnParams.Name = MakeUniqueObjectName(World, ActorClass, BaseName);
+	SpawnParams.Owner = Parent;
+	AActor* NewActor = World->SpawnActor(*ActorClass, &Transform, SpawnParams);
+	
+	if (!NewActor)
+	{
+		return nullptr;
+	}
+
+#if WITH_EDITOR
+	NewActor->SetActorLabel(SpawnParams.Name.ToString());
+#endif // WITH_EDITOR
+
+	USceneComponent* RootComponent = NewActor->GetRootComponent();
+	if (!RootComponent)
+	{
+		RootComponent = NewObject<USceneComponent>(NewActor, USceneComponent::GetDefaultSceneRootVariableName(), RF_Transactional);
+		RootComponent->SetWorldTransform(Transform);
+
+		NewActor->SetRootComponent(RootComponent);
+		NewActor->AddInstanceComponent(RootComponent);
+
+		RootComponent->RegisterComponent();
+	}
+
+	RootComponent->Mobility = EComponentMobility::Static;
+
+#if WITH_EDITOR
+	RootComponent->bVisualizeComponent = true;
+#endif // WITH_EDITOR
+
+	if (Parent)
+	{
+		NewActor->AttachToActor(Parent, FAttachmentTransformRules::KeepWorldTransform);
+	}
+
+	return NewActor;
+}
diff --git a/Engine/Plugins/Experimental/PCG/Source/PCG/Private/PCGComponent.cpp b/Engine/Plugins/Experimental/PCG/Source/PCG/Private/PCGComponent.cpp
--- a/Engine/Plugins/Experimental/PCG/Source/PCG/Private/PCGComponent.cpp
+++ b/Engine/Plugins/Experimental/PCG/Source/PCG/Private/PCGComponent.cpp
@@ -367,6 +367,95 @@ void UPCGComponent::CleanupLocal(bool bRemoveComponents, bool bSave)
 #endif
 }
 
+AActor* UPCGComponent::ClearPCGLink(UClass* TemplateActor)
+{
+	if (!bGenerated || !GetOwner() || !GetWorld())
+	{
+		return nullptr;
+	}
+
+	// TODO: Perhaps remove this part if we want to do it in the PCG Graph.
+#if WITH_EDITOR
+	if (bIsGenerating)
+	{
+		return nullptr;
+	}
+#endif // WITH_EDITOR
+
+	UWorld* World = GetWorld();
+
+	// First create a new actor that will be the new owner of all the resources
+	AActor* NewActor = UPCGActorHelpers::SpawnDefaultActor(World, TemplateActor ? TemplateActor : AActor::StaticClass(), TEXT("PCGStamp"), GetOwner()->GetTransform());
+
+	// Then move all resources linked to this component to this actor
+	bool bHasMovedResources = MoveResourcesToNewActor(NewActor, /*bCreateChild=*/false);
+
+	// And finally, if we are partitioned, we need to do the same for all PCGActors, in Editor only.
+	if (IsPartitioned())
+	{
+#if WITH_EDITOR
+		if (UPCGSubsystem* Subsystem = GetSubsystem())
+		{
+			Subsystem->ClearPCGLink(this, LastGeneratedBounds, NewActor);
+		}
+#endif // WITH_EDITOR
+	}
+	else
+	{
+		if (bHasMovedResources)
+		{
+			Cleanup(true);
+		}
+		else
+		{
+			World->DestroyActor(NewActor);
+			NewActor = nullptr;
+		}
+	}
+
+	return NewActor;
+}
+
+bool UPCGComponent::MoveResourcesToNewActor(AActor* InNewActor, bool bCreateChild)
+{
+	check(InNewActor);
+	AActor* NewActor = InNewActor;
+
+	bool bHasMovedResources = false;
+
+	Modify();
+
+	if (bCreateChild)
+	{
+		NewActor = UPCGActorHelpers::SpawnDefaultActor(GetWorld(), NewActor->GetClass(), TEXT("PCGStampChild"), GetOwner()->GetTransform(), NewActor);
+		check(NewActor);
+	}
+
+	// Trying to move all resources for now. Perhaps in the future we won't want that.
+	{
+		FScopeLock ResourcesLock(&GeneratedResourcesLock);
+		for (TObjectPtr<UPCGManagedResource>& GeneratedResource : GeneratedResources)
+		{
+			check(GeneratedResource);
+			GeneratedResource->MoveResourceToNewActor(NewActor);
+			TSet<TSoftObjectPtr<AActor>> Dummy;
+			GeneratedResource->ReleaseIfUnused(Dummy);
+			bHasMovedResources = true;
+		}
+
+		GeneratedResources.Empty();
+	}
+
+	if (!bHasMovedResources && bCreateChild)
+	{
+		// There was no resource moved, delete the newly spawned actor.
+		GetWorld()->DestroyActor(NewActor);
+		return false;
+	}
+
+	return bHasMovedResources;
+}
+
 void UPCGComponent::CleanupInternal(bool bRemoveComponents)
 {
 	TSet<TSoftObjectPtr<AActor>> ActorsToDelete;
diff --git a/Engine/Plugins/Experimental/PCG/Source/PCG/Private/PCGManagedResource.cpp b/Engine/Plugins/Experimental/PCG/Source/PCG/Private/PCGManagedResource.cpp
--- a/Engine/Plugins/Experimental/PCG/Source/PCG/Private/PCGManagedResource.cpp
+++ b/Engine/Plugins/Experimental/PCG/Source/PCG/Private/PCGManagedResource.cpp
@@ -4,6 +4,7 @@
 #include "PCGComponent.h"
 
 #include "Components/InstancedStaticMeshComponent.h"
+#include "Components/SceneComponent.h"
 
 // By default, if it is not a hard release, we mark the resource unused.
 bool UPCGManagedResource::Release(bool bHardRelease, TSet<TSoftObjectPtr<AActor>>& /*OutActorsToDelete*/)
@@ -72,6 +73,26 @@ bool UPCGManagedActors::ReleaseIfUnused(TSet<TSoftObjectPtr<AActor>>& OutActorsT
 	return Super::ReleaseIfUnused(OutActorsToDelete) || GeneratedActors.IsEmpty();
 }
 
+bool UPCGManagedActors::MoveResourceToNewActor(AActor* NewActor)
+{
+	check(NewActor);
+
+	for (TSoftObjectPtr<AActor>& Actor : GeneratedActors)
+	{
+		if (!Actor.IsValid())
+		{
+			continue;
+		}
+
+		Actor->DetachFromActor(FDetachmentTransformRules::KeepWorldTransform);
+		Actor->AttachToActor(NewActor, FAttachmentTransformRules::KeepWorldTransform);
+	}
+
+	GeneratedActors.Empty();
+
+	return true;
+}
+
 void UPCGManagedComponent::PostEditImport()
 {
 	Super::PostEditImport();
@@ -136,6 +157,47 @@ bool UPCGManagedComponent::ReleaseIfUnused(TSet<TSoftObjectPtr<AActor>>& OutActo
 	return Super::ReleaseIfUnused(OutActorsToDelete) || !GeneratedComponent.IsValid();
 }
 
+bool UPCGManagedComponent::MoveResourceToNewActor(AActor* NewActor)
+{
+	check(NewActor);
+
+	if (!GeneratedComponent.IsValid())
+	{
+		return false;
+	}
+
+	TObjectPtr<AActor> OldOwner = GeneratedComponent->GetOwner();
+	check(OldOwner);
+
+	bool bDetached = false;
+	bool bAttached = false;
+
+	// Check if it is a scene component, and if so, use its method to attach/detach to root component
+	if (TObjectPtr<USceneComponent> GeneratedSceneComponent = Cast<USceneComponent>(GeneratedComponent.Get()))
+	{
+		GeneratedSceneComponent->DetachFromComponent(FDetachmentTransformRules::KeepWorldTransform);
+		bDetached = true;
+		bAttached = GeneratedSceneComponent->AttachToComponent(NewActor->GetRootComponent(), FAttachmentTransformRules::KeepWorldTransform);
+	}
+
+	// Otherwise use the default one.
+	if (!bAttached)
+	{
+		if (!bDetached)
+		{
+			OldOwner->RemoveInstanceComponent(GeneratedComponent.Get());
+		}
+
+		NewActor->AddInstanceComponent(GeneratedComponent.Get());
+	}
+
+	GeneratedComponent->Rename(nullptr, NewActor);
+
+	GeneratedComponent.Reset();
+
+	return true;
+}
+
 void UPCGManagedComponent::MarkAsUsed()
 {
 	if (!bIsMarkedUnused)
diff --git a/Engine/Plugins/Experimental/PCG/Source/PCG/Private/PCGSubsystem.cpp b/Engine/Plugins/Experimental/PCG/Source/PCG/Private/PCGSubsystem.cpp
--- a/Engine/Plugins/Experimental/PCG/Source/PCG/Private/PCGSubsystem.cpp
+++ b/Engine/Plugins/Experimental/PCG/Source/PCG/Private/PCGSubsystem.cpp
@@ -517,6 +517,72 @@ void UPCGSubsystem::CleanupPartitionActors(const FBox& InBounds)
 	PCGSubsystem::ForEachIntersectingCell(GraphExecutor, GetWorld(), InBounds, /*bCreateActor=*/false, /*bLoadCell=*/false, /*bSave=*/false, ScheduleTask);
 }
 
+void UPCGSubsystem::ClearPCGLink(UPCGComponent* InComponent, const FBox& InBounds, AActor* InNewActor)
+{
+	TWeakObjectPtr<AActor> NewActorPtr(InNewActor);
+	TWeakObjectPtr<UPCGComponent> ComponentPtr(InComponent);
+
+	auto ScheduleTask = [this, NewActorPtr, ComponentPtr](APCGPartitionActor* PCGActor, const FBox& InIntersectedBounds, const TArray<FPCGTaskId>& TaskDependencies) {
+		auto MoveTask = [this, NewActorPtr, ComponentPtr, PCGActor]() {
+			check(NewActorPtr.IsValid() && ComponentPtr.IsValid() && PCGActor != nullptr);
+
+			if (TObjectPtr<UPCGComponent> LocalComponent = PCGActor->GetLocalComponent(ComponentPtr.Get()))
+			{
+				LocalComponent->MoveResourcesToNewActor(NewActorPtr.Get(), /*bCreateChild=*/true);
+			}
+
+			return true;
+		};
+
+		return GraphExecutor->ScheduleGeneric(MoveTask, TaskDependencies);
+	};
+
+	FPCGTaskId TaskId = PCGSubsystem::ForEachIntersectingCell(GraphExecutor, GetWorld(), InBounds, /*bCreateActor=*/false, /*bLoadCell=*/false, /*bSave=*/false, ScheduleTask);
+
+	// Verify if the NewActor has some components attached to its root or attached actors. If not, destroy it.
+	// Return false if the new actor is not valid or destroyed.
+	auto VerifyAndDestroyNewActor = [this, NewActorPtr]() {
+		check(NewActorPtr.IsValid());
+
+		USceneComponent* RootComponent = NewActorPtr->GetRootComponent();
+		check(RootComponent);
+
+		AActor* NewActor = NewActorPtr.Get();
+
+		TArray<AActor*> AttachedActors;
+		NewActor->GetAttachedActors(AttachedActors);
+
+		if (RootComponent->GetNumChildrenComponents() == 0 && AttachedActors.IsEmpty())
+		{
+			GetWorld()->DestroyActor(NewActor);
+			return false;
+		}
+
+		return true;
+	};
+
+	if (TaskId != InvalidPCGTaskId)
+	{
+		auto CleanupTask = [this, ComponentPtr, VerifyAndDestroyNewActor]() {
+
+			// If the new actor is valid, clean up the original component.
+			if (VerifyAndDestroyNewActor())
+			{
+				check(ComponentPtr.IsValid());
+				ComponentPtr->Cleanup(/*bRemoveComponents=*/true);
+			}
+
+			return true;
+		};
+
+		GraphExecutor->ScheduleGeneric(CleanupTask, { TaskId });
+	}
+	else
+	{
+		VerifyAndDestroyNewActor();
+	}
+}
+
 void UPCGSubsystem::DeletePartitionActors()
 {
 	TRACE_CPUPROFILER_EVENT_SCOPE(UPCGSubsystem::DeletePartitionActors);
diff --git a/Engine/Plugins/Experimental/PCG/Source/PCG/Public/Helpers/PCGActorHelpers.h b/Engine/Plugins/Experimental/PCG/Source/PCG/Public/Helpers/PCGActorHelpers.h
--- a/Engine/Plugins/Experimental/PCG/Source/PCG/Public/Helpers/PCGActorHelpers.h
+++ b/Engine/Plugins/Experimental/PCG/Source/PCG/Public/Helpers/PCGActorHelpers.h
@@ -19,6 +19,7 @@ class UPCGComponent;
 class UMaterialInterface;
 class UActorComponent;
 class ULevel;
+class UWorld;
 
 struct FPCGISMCBuilderParameters
 {
@@ -46,10 +47,33 @@ public:
 	static void GetActorClassDefaultComponents(const TSubclassOf<AActor>& ActorClass, TArray<UActorComponent*>& OutComponents, const TSubclassOf<UActorComponent>& InComponentClass = TSubclassOf<UActorComponent>());
 
 	template <typename T, typename = typename std::enable_if_t<std::is_base_of_v<AActor, T>>>
-	static void ForEachActorInLevel(TObjectPtr<ULevel> Level, TFunctionRef<void(AActor*)> Callback)
+	static void ForEachActorInLevel(ULevel* Level, TFunctionRef<void(AActor*)> Callback)
 	{
 		return ForEachActorInLevel(Level, T::StaticClass(), Callback);
 	}
 
-	static void ForEachActorInLevel(TObjectPtr<ULevel> Level, TSubclassOf<AActor> ActorClass, TFunctionRef<void(AActor*)> Callback);
+	static void ForEachActorInLevel(ULevel* Level, TSubclassOf<AActor> ActorClass, TFunctionRef<void(AActor*)> Callback);
+
+	/**
+	* Spawn a new actor of type T and attach it to the parent (if not null)
+	* @param World The world
+	* @param BaseName Base name for the actor, will have a unique name
+	* @param Transform The transform for the new actor
+	* @param Parent Optional parent to attach to.
+	*/
+	template <typename T = AActor, typename = typename std::enable_if_t<std::is_base_of_v<AActor, T>>>
+	static AActor* SpawnDefaultActor(UWorld* World, FName BaseName, const FTransform& Transform, AActor* Parent = nullptr)
+	{
+		return SpawnDefaultActor(World, T::StaticClass(), BaseName, Transform, Parent);
+	}
+
+	/**
+	* Spawn a new actor and attach it to the parent (if not null)
+	* @param World The world
+	* @param ActorClass Class of the actor to spawn
+	* @param BaseName Base name for the actor, will have a unique name
+	* @param Transform The transform for the new actor
+	* @param Parent Optional parent to attach to.
+	*/
+	static AActor* SpawnDefaultActor(UWorld* World, TSubclassOf<AActor> ActorClass, FName BaseName, const FTransform& Transform, AActor* Parent = nullptr);
 };
\ No newline at end of file
diff --git a/Engine/Plugins/Experimental/PCG/Source/PCG/Public/PCGComponent.h b/Engine/Plugins/Experimental/PCG/Source/PCG/Public/PCGComponent.h
--- a/Engine/Plugins/Experimental/PCG/Source/PCG/Public/PCGComponent.h
+++ b/Engine/Plugins/Experimental/PCG/Source/PCG/Public/PCGComponent.h
@@ -15,6 +15,7 @@ class UPCGData;
 class UPCGSubsystem;
 class ALandscapeProxy;
 class FLandscapeProxyComponentDataChangedParams;
+class UClass;
 
 #if WITH_EDITOR
 	DECLARE_MULTICAST_DELEGATE_OneParam(FOnPCGGraphGenerated, UPCGComponent*);
@@ -96,6 +97,10 @@ public:
 	UFUNCTION(BlueprintCallable, NetMulticast, Reliable, Category = PCG)
 	void Cleanup(bool bRemoveComponents, bool bSave = false);
 
+	/** Move all generated resources under a new actor, following a template (AActor if not provided), clearing all link to this PCG component. Returns the new actor.*/
+	UFUNCTION(BlueprintCallable, Category = PCG)
+	AActor* ClearPCGLink(UClass* TemplateActor = nullptr);
+
 	UPROPERTY(BlueprintReadWrite, EditAnywhere, Category = Settings)
 	EPCGComponentInput InputType = EPCGComponentInput::Actor;
 
@@ -171,6 +176,7 @@ private:
 	void PostProcessGraph(const FBox& InNewBounds, bool bInGenerated);
 	void OnProcessGraphAborted();
 	void CleanupUnusedManagedResources();
+	bool MoveResourcesToNewActor(AActor* InNewActor, bool bCreateChild);
 
 	bool GetActorsFromTags(const TSet<FName>& InTags, TSet<TWeakObjectPtr<AActor>>& OutActors, bool bCullAgainstLocalBounds);
 
diff --git a/Engine/Plugins/Experimental/PCG/Source/PCG/Public/PCGManagedResource.h b/Engine/Plugins/Experimental/PCG/Source/PCG/Public/PCGManagedResource.h
--- a/Engine/Plugins/Experimental/PCG/Source/PCG/Public/PCGManagedResource.h
+++ b/Engine/Plugins/Experimental/PCG/Source/PCG/Public/PCGManagedResource.h
@@ -28,6 +28,9 @@ public:
 
 	virtual void MarkAsUsed() { bIsMarkedUnused = false; }
 
+	/** Move the given resource to a new actor. Return true if it has succeeded */
+	virtual bool MoveResourceToNewActor(AActor* NewActor) { return false; };
+
 protected:
 	UPROPERTY(Transient, VisibleAnywhere, Category = GeneratedData)
 	bool bIsMarkedUnused = false;
@@ -46,6 +49,7 @@ public:
 	//~Begin UPCGManagedResource interface
 	virtual bool Release(bool bHardRelease, TSet<TSoftObjectPtr<AActor>>& OutActorsToDelete) override;
 	virtual bool ReleaseIfUnused(TSet<TSoftObjectPtr<AActor>>& OutActorsToDelete) override;
+	virtual bool MoveResourceToNewActor(AActor* NewActor) override;
 	//~End UPCGManagedResource interface
 
 	UPROPERTY(BlueprintReadWrite, EditAnywhere, Category = GeneratedData)
@@ -65,6 +69,7 @@ public:
 	//~Begin UPCGManagedResource interface
 	virtual bool Release(bool bHardRelease, TSet<TSoftObjectPtr<AActor>>& OutActorsToDelete) override;
 	virtual bool ReleaseIfUnused(TSet<TSoftObjectPtr<AActor>>& OutActorsToDelete) override;
+	virtual bool MoveResourceToNewActor(AActor* NewActor) override;
 	//~End UPCGManagedResource interface
 
 	virtual void ResetComponent() { check(0); }
diff --git a/Engine/Plugins/Experimental/PCG/Source/PCG/Public/PCGSubsystem.h b/Engine/Plugins/Experimental/PCG/Source/PCG/Public/PCGSubsystem.h
--- a/Engine/Plugins/Experimental/PCG/Source/PCG/Public/PCGSubsystem.h
+++ b/Engine/Plugins/Experimental/PCG/Source/PCG/Public/PCGSubsystem.h
@@ -85,6 +85,9 @@ public:
 	/** Flushes the graph cache completely, use only for debugging */
 	void FlushCache();
 
+	/** Move all resources from sub actors to a new actor */
+	void ClearPCGLink(UPCGComponent* InComponent, const FBox& InBounds, AActor* InNewActor);
+
 private:
 	FPCGTaskId DelayProcessGraph(UPCGComponent* Component, bool bGenerate, bool bSave, bool bUseEmptyNewBounds);
 	FPCGTaskId ProcessGraph(UPCGComponent* Component, const FBox& InPreviousBounds, const FBox& InNewBounds, bool bGenerate, bool bSave);
diff --git a/Engine/Plugins/Experimental/PCG/Source/PCGEditor/Private/PCGComponentDetails.cpp b/Engine/Plugins/Experimental/PCG/Source/PCGEditor/Private/PCGComponentDetails.cpp
--- a/Engine/Plugins/Experimental/PCG/Source/PCGEditor/Private/PCGComponentDetails.cpp
+++ b/Engine/Plugins/Experimental/PCG/Source/PCGEditor/Private/PCGComponentDetails.cpp
@@ -78,6 +78,18 @@ void FPCGComponentDetails::CustomizeDetails(IDetailLayoutBuilder& DetailBuilder)
 					.Text(LOCTEXT("CleanupButton", "Cleanup"))
 				]
 			]
+			+ SHorizontalBox::Slot()
+			.AutoWidth()
+			.Padding(4.0f)
+			[
+				SNew(SButton)
+				.OnClicked(this, &FPCGComponentDetails::OnClearPCGLinkClicked)
+				[
+					SNew(STextBlock)
+					.Font(IDetailLayoutBuilder::GetDetailFont())
+					.Text(LOCTEXT("ClearPCGLinkButton", "Clear PCG Link"))
+				]
+			]
 		];
 
 	// Attach to generated delegate on the selected components
@@ -117,6 +129,19 @@ FReply FPCGComponentDetails::OnGenerateClicked()
 	return FReply::Handled();
 }
 
+FReply FPCGComponentDetails::OnClearPCGLinkClicked()
+{
+	for (TWeakObjectPtr<UPCGComponent>& Component : SelectedComponents)
+	{
+		if (Component.IsValid())
+		{
+			Component.Get()->ClearPCGLink();
+		}
+	}
+
+	return FReply::Handled();
+}
+
 void FPCGComponentDetails::OnGraphChanged(UPCGComponent* InComponent)
 {
 	if (!InComponent || !InComponent->GetOwner())
diff --git a/Engine/Plugins/Experimental/PCG/Source/PCGEditor/Public/PCGComponentDetails.h b/Engine/Plugins/Experimental/PCG/Source/PCGEditor/Public/PCGComponentDetails.h
--- a/Engine/Plugins/Experimental/PCG/Source/PCGEditor/Public/PCGComponentDetails.h
+++ b/Engine/Plugins/Experimental/PCG/Source/PCGEditor/Public/PCGComponentDetails.h
@@ -21,6 +21,7 @@ public:
 private:
 	FReply OnGenerateClicked();
 	FReply OnCleanupClicked();
+	FReply OnClearPCGLinkClicked();
 	void OnGraphChanged(UPCGComponent* InComponent);
 
 	TArray<TWeakObjectPtr<UPCGComponent>> SelectedComponents;
