PCG:
*Added element type filter to node palette.

diff --git a/Engine/Plugins/Experimental/PCG/Source/PCGEditor/PCGEditor.Build.cs b/Engine/Plugins/Experimental/PCG/Source/PCGEditor/PCGEditor.Build.cs
--- a/Engine/Plugins/Experimental/PCG/Source/PCGEditor/PCGEditor.Build.cs
+++ b/Engine/Plugins/Experimental/PCG/Source/PCGEditor/PCGEditor.Build.cs
@@ -31,6 +31,7 @@ namespace UnrealBuildTool.Rules
 					"EditorScriptingUtilities",
 					"EditorStyle",
 					"EditorSubsystem",
+					"EditorWidgets",
 					"GraphEditor",
 					"InputCore",
 					"Kismet",
diff --git a/Engine/Plugins/Experimental/PCG/Source/PCGEditor/Private/PCGEditorCommon.h b/Engine/Plugins/Experimental/PCG/Source/PCGEditor/Private/PCGEditorCommon.h
--- a/Engine/Plugins/Experimental/PCG/Source/PCGEditor/Private/PCGEditorCommon.h
+++ b/Engine/Plugins/Experimental/PCG/Source/PCGEditor/Private/PCGEditorCommon.h
@@ -2,6 +2,19 @@
 
 #pragma once
 
+#include "PCGEditorCommon.generated.h"
+
+UENUM()
+enum class EPCGElementType : uint8
+{
+	Native = 1 << 0,
+	Blueprint = 1 << 1,
+	Subgraph = 1 << 2,
+	Other = 1 << 3,
+	All = (1 << 4) - 1
+};
+ENUM_CLASS_FLAGS(EPCGElementType);
+
 namespace FPCGEditorCommon
 {
 	const FString ContextIdentifier = TEXT("PCGEditorContext");
diff --git a/Engine/Plugins/Experimental/PCG/Source/PCGEditor/Private/PCGEditorGraphSchema.cpp b/Engine/Plugins/Experimental/PCG/Source/PCGEditor/Private/PCGEditorGraphSchema.cpp
--- a/Engine/Plugins/Experimental/PCG/Source/PCGEditor/Private/PCGEditorGraphSchema.cpp
+++ b/Engine/Plugins/Experimental/PCG/Source/PCGEditor/Private/PCGEditorGraphSchema.cpp
@@ -20,12 +20,24 @@
 #define LOCTEXT_NAMESPACE "PCGEditorGraphSchema"
 
 
-void UPCGEditorGraphSchema::GetPaletteActions(FGraphActionMenuBuilder& ActionMenuBuilder) const
+void UPCGEditorGraphSchema::GetPaletteActions(FGraphActionMenuBuilder& ActionMenuBuilder, const EPCGElementType InPCGElementTypeFilter) const
 {
-	GetNativeElementActions(ActionMenuBuilder);
-	GetSubgraphElementActions(ActionMenuBuilder);
-	GetBlueprintElementActions(ActionMenuBuilder);
-	GetExtraElementActions(ActionMenuBuilder);
+	if (!!(InPCGElementTypeFilter & EPCGElementType::Native))
+	{
+		GetNativeElementActions(ActionMenuBuilder);
+	}
+	if (!!(InPCGElementTypeFilter & EPCGElementType::Subgraph))
+	{
+		GetSubgraphElementActions(ActionMenuBuilder);
+	}
+	if (!!(InPCGElementTypeFilter & EPCGElementType::Blueprint))
+	{
+		GetBlueprintElementActions(ActionMenuBuilder);
+	}
+	if (!!(InPCGElementTypeFilter & EPCGElementType::Other))
+	{
+		GetExtraElementActions(ActionMenuBuilder);
+	}
 }
 
 void UPCGEditorGraphSchema::GetGraphContextActions(FGraphContextMenuBuilder& ContextMenuBuilder) const
diff --git a/Engine/Plugins/Experimental/PCG/Source/PCGEditor/Private/PCGEditorGraphSchema.h b/Engine/Plugins/Experimental/PCG/Source/PCGEditor/Private/PCGEditorGraphSchema.h
--- a/Engine/Plugins/Experimental/PCG/Source/PCGEditor/Private/PCGEditorGraphSchema.h
+++ b/Engine/Plugins/Experimental/PCG/Source/PCGEditor/Private/PCGEditorGraphSchema.h
@@ -4,6 +4,7 @@
 
 #include "EdGraph/EdGraphSchema.h"
 #include "ConnectionDrawingPolicy.h"
+#include "PCGEditorCommon.h"
 
 #include "PCGEditorGraphSchema.generated.h"
 
@@ -15,7 +16,7 @@ class UPCGEditorGraphSchema : public UEdGraphSchema
 	GENERATED_BODY()
 
 public:
-	void GetPaletteActions(FGraphActionMenuBuilder& ActionMenuBuilder) const;
+	void GetPaletteActions(FGraphActionMenuBuilder& ActionMenuBuilder, const EPCGElementType InPCGElementTypeFilter) const;
 
 	//~ Begin EdGraphSchema Interface
 	virtual void GetGraphContextActions(FGraphContextMenuBuilder& ContextMenuBuilder) const;
diff --git a/Engine/Plugins/Experimental/PCG/Source/PCGEditor/Private/SPCGEditorGraphNodePalette.cpp b/Engine/Plugins/Experimental/PCG/Source/PCGEditor/Private/SPCGEditorGraphNodePalette.cpp
--- a/Engine/Plugins/Experimental/PCG/Source/PCGEditor/Private/SPCGEditorGraphNodePalette.cpp
+++ b/Engine/Plugins/Experimental/PCG/Source/PCGEditor/Private/SPCGEditorGraphNodePalette.cpp
@@ -3,14 +3,14 @@
 #include "SPCGEditorGraphNodePalette.h"
 
 #include "Elements/PCGExecuteBlueprint.h"
-#include "PCGEditor.h"
 #include "PCGEditorGraphSchema.h"
+#include "PCGEditorUtils.h"
 #include "PCGGraph.h"
 
 #include "AssetRegistry/AssetRegistryModule.h"
+#include "EditorWidgets/Public/SEnumCombo.h"
 #include "Modules/ModuleManager.h"
 #include "SGraphActionMenu.h"
-#include "PCGEditorUtils.h"
 
 #define LOCTEXT_NAMESPACE "SPCGEditorGraphNodePalette"
 
@@ -35,17 +35,44 @@ FText SPCGEditorGraphNodePaletteItem::GetItemTooltip() const
 	return ActionPtr.Pin()->GetTooltipDescription();
 }
 
-
-
 void SPCGEditorGraphNodePalette::Construct(const FArguments& InArgs)
 {
+	const UEnum* PCGElementTypeEnum = StaticEnum<EPCGElementType>();
+	
 	this->ChildSlot
 	[
-		SAssignNew(GraphActionMenu, SGraphActionMenu)
-		.OnActionDragged(this, &SPCGEditorGraphNodePalette::OnActionDragged)
-		.OnCreateWidgetForAction(this, &SPCGEditorGraphNodePalette::OnCreateWidgetForAction)
-		.OnCollectAllActions(this, &SPCGEditorGraphNodePalette::CollectAllActions)
-		.AutoExpandActionMenu(true)
+		SNew(SVerticalBox)
+		+SVerticalBox::Slot()
+		.AutoHeight()
+		[
+			SNew(SHorizontalBox)
+			+SHorizontalBox::Slot()
+			.AutoWidth()
+			.VAlign(VAlign_Center)
+			.Padding(2.0f, 0.0f)
+			[
+				SNew(STextBlock)
+				.Text(LOCTEXT("TypeTextBlock", "Type:"))
+			]
+			+SHorizontalBox::Slot()
+			.VAlign(VAlign_Center)
+			.FillWidth(1.0f)
+			[
+				SNew(SEnumComboBox, PCGElementTypeEnum)
+				.ContentPadding(FMargin(4, 0))
+				.OnEnumSelectionChanged(this, &SPCGEditorGraphNodePalette::OnTypeSelectionChanged)
+				.CurrentValue(this, &SPCGEditorGraphNodePalette::GetTypeValue)
+			]
+		]
+		+SVerticalBox::Slot()
+		.FillHeight(1.0f)
+		[
+			SAssignNew(GraphActionMenu, SGraphActionMenu)
+			.OnActionDragged(this, &SPCGEditorGraphNodePalette::OnActionDragged)
+			.OnCreateWidgetForAction(this, &SPCGEditorGraphNodePalette::OnCreateWidgetForAction)
+			.OnCollectAllActions(this, &SPCGEditorGraphNodePalette::CollectAllActions)
+			.AutoExpandActionMenu(true)
+		]
 	];
 
 	FAssetRegistryModule& AssetRegistryModule = FModuleManager::GetModuleChecked<FAssetRegistryModule>(TEXT("AssetRegistry"));
@@ -80,7 +107,7 @@ void SPCGEditorGraphNodePalette::CollectAllActions(FGraphActionListBuilderBase&
 	const UPCGEditorGraphSchema* PCGSchema = GetDefault<UPCGEditorGraphSchema>();
 
 	FGraphActionMenuBuilder ActionMenuBuilder;
-	PCGSchema->GetPaletteActions(ActionMenuBuilder);
+	PCGSchema->GetPaletteActions(ActionMenuBuilder, ElementType);
 	OutAllActions.Append(ActionMenuBuilder);
 }
 
@@ -97,5 +124,16 @@ void SPCGEditorGraphNodePalette::OnAssetRenamed(const FAssetData& InAssetData, c
 	OnAssetChanged(InAssetData);
 }
 
+void SPCGEditorGraphNodePalette::OnTypeSelectionChanged(int32 InValue, ESelectInfo::Type /*SelectInfo*/)
+{
+	ElementType = static_cast<EPCGElementType>(InValue);
+	RefreshActionsList(true);
+}
+
+int32 SPCGEditorGraphNodePalette::GetTypeValue() const
+{
+	return static_cast<int32>(ElementType);
+}
+
 #undef LOCTEXT_NAMESPACE
- 
\ No newline at end of file
+ 
diff --git a/Engine/Plugins/Experimental/PCG/Source/PCGEditor/Private/SPCGEditorGraphNodePalette.h b/Engine/Plugins/Experimental/PCG/Source/PCGEditor/Private/SPCGEditorGraphNodePalette.h
--- a/Engine/Plugins/Experimental/PCG/Source/PCGEditor/Private/SPCGEditorGraphNodePalette.h
+++ b/Engine/Plugins/Experimental/PCG/Source/PCGEditor/Private/SPCGEditorGraphNodePalette.h
@@ -2,7 +2,8 @@
 
 #pragma once
 
-#include "CoreMinimal.h"
+#include "PCGEditorCommon.h"
+
 #include "SGraphPalette.h"
 
 class SPCGEditorGraphNodePaletteItem : public SGraphPaletteItem
@@ -38,4 +39,9 @@ protected:
 private:
 	void OnAssetChanged(const FAssetData& InAssetData);
 	void OnAssetRenamed(const FAssetData& InAssetData, const FString& InNewAssetName);
+	void OnTypeSelectionChanged(int32, ESelectInfo::Type SelectInfo);
+	
+	int32 GetTypeValue() const;
+	
+	EPCGElementType ElementType = EPCGElementType::All;
 };
