Update to the PCG Determinism Test Framework.
+ Defaults created and implemented for consistency across testing environments.
+ Continuous vs Discrete test cases for Difference element determinism test.
* Other determinism tests and framework updated to make use of public defaults
* Previous tests re-tested with new defaults

#rnx

diff --git a/Engine/Plugins/Experimental/PCG/Source/PCG/Private/Tests/Determinism/PCGCollapseDeterminismTest.cpp b/Engine/Plugins/Experimental/PCG/Source/PCG/Private/Tests/Determinism/PCGCollapseDeterminismTest.cpp
--- a/Engine/Plugins/Experimental/PCG/Source/PCG/Private/Tests/Determinism/PCGCollapseDeterminismTest.cpp
+++ b/Engine/Plugins/Experimental/PCG/Source/PCG/Private/Tests/Determinism/PCGCollapseDeterminismTest.cpp
@@ -27,13 +27,13 @@ namespace
 
 	void CreateNewRandomizedSpatialDataSet(TArray<PCGDeterminismTests::FTestData>& TestDataSet,
 		int32 NumSpatialSetsToAdd = 1,
-		int32 NumPointsToCreatePerSet = PCGDeterminismTests::DefaultNumPointsToGenerate,
-		int32 NumPolyLinePointsToCreatePerSet = PCGDeterminismTests::DefaultNumPolyLinePointsToGenerate)
+		int32 NumPointsToCreatePerSet = PCGDeterminismTests::Defaults::NumPointsToGenerate,
+		int32 NumPolyLinePointsToCreatePerSet = PCGDeterminismTests::Defaults::NumPolyLinePointsToGenerate)
 	{
 		check(NumSpatialSetsToAdd > 0 && NumPointsToCreatePerSet > 0 && NumPolyLinePointsToCreatePerSet > 0);
 
 		TestDataSet.Empty();
-		TestDataSet.Init(PCGDeterminismTests::FTestData(PCGDeterminismTests::DefaultSeed), NumTestableSpatialDataTypes);
+		TestDataSet.Init(PCGDeterminismTests::FTestData(PCGDeterminismTests::Defaults::Seed), NumTestableSpatialDataTypes);
 
 		for (PCGDeterminismTests::FTestData& Data : TestDataSet)
 		{
diff --git a/Engine/Plugins/Experimental/PCG/Source/PCG/Private/Tests/Determinism/PCGDeterminismTestsCommon.cpp b/Engine/Plugins/Experimental/PCG/Source/PCG/Private/Tests/Determinism/PCGDeterminismTestsCommon.cpp
--- a/Engine/Plugins/Experimental/PCG/Source/PCG/Private/Tests/Determinism/PCGDeterminismTestsCommon.cpp
+++ b/Engine/Plugins/Experimental/PCG/Source/PCG/Private/Tests/Determinism/PCGDeterminismTestsCommon.cpp
@@ -7,6 +7,7 @@
 #include "PCGHelpers.h"
 #include "PCGPoint.h"
 #include "Data/PCGDifferenceData.h"
+#include "Data/PCGIntersectionData.h"
 #include "Data/PCGPointData.h"
 #include "Data/PCGPrimitiveData.h"
 #include "Data/PCGSplineData.h"
@@ -24,11 +25,6 @@
 
 namespace PCGDeterminismTests
 {
-	constexpr static FVector::FReal SmallDistance = 50.0;
-	constexpr static FVector::FReal MediumDistance = 200.0;
-	constexpr static FVector::FReal LargeDistance = 500.0;
-	constexpr static int32 NumSamplingStepsPerDimension = 50;
-
 	FTestData::FTestData(int32 RandomSeed) :
 		Seed(RandomSeed),
 		RandomStream(Seed)
@@ -73,7 +69,7 @@ namespace PCGDeterminismTests
 	{
 		check(PCGNode);
 
-		FTestData TestData(DefaultSeed);
+		FTestData TestData(Defaults::Seed);
 		TestData.Settings = PCGNode->DefaultSettings;
 
 		// Generate random input data
@@ -185,7 +181,7 @@ namespace PCGDeterminismTests
 		check(PointNum > 0);
 		for (int32 I = 0; I < PointNum; ++I)
 		{
-			AddSinglePointInputData(TestData.InputData, TestData.RandomStream.VRand() * LargeDistance, PinName);
+			AddSinglePointInputData(TestData.InputData, TestData.RandomStream.VRand() * Defaults::LargeDistance, PinName);
 		}
 	}
 
@@ -197,7 +193,7 @@ namespace PCGDeterminismTests
 		Points.SetNumUninitialized(PointNum);
 		for (int32 I = 0; I < PointNum; ++I)
 		{
-			FVector NewLocation = TestData.RandomStream.VRand() * LargeDistance;
+			FVector NewLocation = TestData.RandomStream.VRand() * Defaults::LargeDistance;
 			FTransform NewTransform(FRotator::ZeroRotator, NewLocation, FVector::OneVector * TestData.RandomStream.FRandRange(0.5, 1.5));
 			int PointSeed = PCGHelpers::ComputeSeed(static_cast<int>(NewLocation.X), static_cast<int>(NewLocation.Y), static_cast<int>(NewLocation.Z));
 			Points[I] = FPCGPoint(NewTransform, 1.f, PCGHelpers::ComputeSeed(PointSeed, TestData.Seed));
@@ -208,9 +204,11 @@ namespace PCGDeterminismTests
 
 	void AddRandomizedVolumeInputData(FTestData& TestData, const FName& PinName)
 	{
-		AddVolumeInputData(TestData.InputData, TestData.RandomStream.VRand() * LargeDistance,
-			FVector::OneVector * LargeDistance + TestData.RandomStream.VRand() * MediumDistance,
-			FVector::OneVector * MediumDistance + TestData.RandomStream.VRand() * 0.5f * MediumDistance, PinName);
+		AddVolumeInputData(TestData.InputData,
+			TestData.RandomStream.VRand() * Defaults::MediumDistance,
+			Defaults::MediumVector + TestData.RandomStream.VRand() * 0.5f * Defaults::MediumDistance,
+			Defaults::SmallVector + TestData.RandomStream.VRand() * 0.5f * Defaults::SmallDistance,
+			PinName);
 	}
 
 	void AddRandomizedSurfaceInputData(FTestData& TestData, const FName& PinName)
@@ -231,7 +229,7 @@ namespace PCGDeterminismTests
 		check(PointNum > 1);
 		for (int32 I = 0; I < PointNum; ++I)
 		{
-			TestSplineComponent->AddSplinePoint(TestData.RandomStream.VRand() * LargeDistance, ESplineCoordinateSpace::Type::World, false);
+			TestSplineComponent->AddSplinePoint(TestData.RandomStream.VRand() * Defaults::LargeDistance, ESplineCoordinateSpace::Type::World, false);
 			TestSplineComponent->AddRelativeRotation(FRotator(
 				TestData.RandomStream.FRandRange(-90.0, 90.0),
 				TestData.RandomStream.FRandRange(-90.0, 90.0),
@@ -255,11 +253,11 @@ namespace PCGDeterminismTests
 
 		TestPrimitiveComponent->SetWorldTransform(FTransform(
 			FRotator(TestData.RandomStream.FRandRange(0.0, 90.0), TestData.RandomStream.FRandRange(0.0, 90.0), TestData.RandomStream.FRandRange(0.0, 90.0)),
-			TestData.RandomStream.VRand() * LargeDistance,
+			TestData.RandomStream.VRand() * Defaults::LargeDistance,
 			FVector::OneVector * TestData.RandomStream.FRandRange(0.5, 1.5)));
 
 		// TODO: Probably more varieties to add in the future
-		AddPrimitiveInputData(TestData.InputData, TestPrimitiveComponent, FVector::OneVector * MediumDistance + TestData.RandomStream.VRand() * 0.5f * MediumDistance, PinName);
+		AddPrimitiveInputData(TestData.InputData, TestPrimitiveComponent, Defaults::MediumVector + TestData.RandomStream.VRand() * 0.5f * Defaults::MediumDistance, PinName);
 	}
 
 	bool DataCollectionsAreIdentical(const FPCGDataCollection& FirstCollection, const FPCGDataCollection& SecondCollection)
@@ -353,9 +351,9 @@ namespace PCGDeterminismTests
 		{
 			return PrimitiveDataIsIdentical(FirstData, SecondData);
 		}
-		else if (BothDataCastsToDataType<const UPCGDifferenceData>(FirstData, SecondData))
+		else if (BothDataCastsToDataType<const UPCGSpatialData>(FirstData, SecondData))
 		{
-			return DifferenceDataIsIdentical(FirstData, SecondData);
+			return SampledSpatialDataIsIdentical(Cast<const UPCGSpatialData>(FirstData), Cast<const UPCGSpatialData>(SecondData));
 		}
 
 		return false;
@@ -468,23 +466,20 @@ namespace PCGDeterminismTests
 		return false;
 	}
 
-	bool DifferenceDataIsIdentical(const UPCGData* FirstData, const UPCGData* SecondData)
+	bool SampledSpatialDataIsIdentical(const UPCGSpatialData* FirstSpatialData, const UPCGSpatialData* SecondSpatialData)
 	{
-		const UPCGDifferenceData* FirstDifferenceData = CastChecked<const UPCGDifferenceData>(FirstData);
-		const UPCGDifferenceData* SecondDifferenceData = CastChecked<const UPCGDifferenceData>(SecondData);
-
-		if (!SpatialBasicsAreIdentical(FirstDifferenceData, SecondDifferenceData))
+		if (!SpatialBasicsAreIdentical(FirstSpatialData, SecondSpatialData))
 		{
 			return false;
 		}
 
 		// At this point, bounds has already been checked for equality
-		const FBox SampleBounds = FirstDifferenceData->GetBounds();
+		const FBox SampleBounds = Defaults::TestingVolume;
 		const FVector SampleExtent = SampleBounds.GetExtent();
 
 		FPCGPoint FirstPoint;
 		FPCGPoint SecondPoint;
-		FVector StepInterval = SampleExtent * 2.0 / FMath::Max(NumSamplingStepsPerDimension, 1);
+		FVector StepInterval = SampleExtent * 2.0 / FMath::Max(Defaults::NumSamplingStepsPerDimension, 1);
 		FVector StartingOffset = SampleBounds.Min + StepInterval * 0.5;
 
 		// Sample points across the 3D volume
@@ -495,8 +490,8 @@ namespace PCGDeterminismTests
 				for (FVector::FReal Z = StartingOffset.Z; Z < SampleBounds.Max.Z; Z += StepInterval.Z)
 				{
 					FTransform PointTransform(FVector(X, Y, Z));
-					bool bFirstPointWasSampled = FirstDifferenceData->SamplePoint(PointTransform, SampleBounds, FirstPoint, nullptr);
-					bool bSecondPointWasSampled = SecondDifferenceData->SamplePoint(PointTransform, SampleBounds, SecondPoint, nullptr);
+					bool bFirstPointWasSampled = FirstSpatialData->SamplePoint(PointTransform, SampleBounds, FirstPoint, nullptr);
+					bool bSecondPointWasSampled = SecondSpatialData->SamplePoint(PointTransform, SampleBounds, SecondPoint, nullptr);
 
 					if (bFirstPointWasSampled != bSecondPointWasSampled)
 					{
diff --git a/Engine/Plugins/Experimental/PCG/Source/PCG/Private/Tests/Determinism/PCGDifferenceDeterminismTest.cpp b/Engine/Plugins/Experimental/PCG/Source/PCG/Private/Tests/Determinism/PCGDifferenceDeterminismTest.cpp
--- a/Engine/Plugins/Experimental/PCG/Source/PCG/Private/Tests/Determinism/PCGDifferenceDeterminismTest.cpp
+++ b/Engine/Plugins/Experimental/PCG/Source/PCG/Private/Tests/Determinism/PCGDifferenceDeterminismTest.cpp
@@ -9,43 +9,83 @@ IMPLEMENT_CUSTOM_SIMPLE_AUTOMATION_TEST(FPCGDifferenceDeterminismSingleSameDataT
 IMPLEMENT_CUSTOM_SIMPLE_AUTOMATION_TEST(FPCGDifferenceDeterminismSingleIdenticalDataTest, FPCGTestBaseClass, "pcg.tests.Difference.Determinism.SingleMultipleData", PCGTestsCommon::TestFlags)
 IMPLEMENT_CUSTOM_SIMPLE_AUTOMATION_TEST(FPCGDifferenceDeterminismMultipleSameDataTest, FPCGTestBaseClass, "pcg.tests.Difference.Determinism.MultipleSameData", PCGTestsCommon::TestFlags)
 IMPLEMENT_CUSTOM_SIMPLE_AUTOMATION_TEST(FPCGDifferenceDeterminismMultipleIdenticalDataTest, FPCGTestBaseClass, "pcg.tests.Difference.Determinism.MultipleIdenticalData", PCGTestsCommon::TestFlags)
-// TODO: Temporarily disabled while more testing into this is conducted
-//IMPLEMENT_CUSTOM_SIMPLE_AUTOMATION_TEST(FPCGDifferenceDeterminismOrderIndependenceTest, FPCGTestBaseClass, "pcg.tests.Difference.Determinism.OrderIndependence", PCGTestsCommon::TestFlags)
+IMPLEMENT_CUSTOM_SIMPLE_AUTOMATION_TEST(FPCGDifferenceDeterminismOrderIndependenceTest, FPCGTestBaseClass, "pcg.tests.Difference.Determinism.OrderIndependence", PCGTestsCommon::TestFlags)
 
 namespace
 {
-	void DifferenceTestBase(PCGDeterminismTests::FTestData& TestData)
+	void SetSettingsModeInferred(PCGDeterminismTests::FTestData& TestData)
 	{
-		PCGDeterminismTests::GenerateSettings<UPCGDifferenceSettings>(TestData);
+		UPCGDifferenceSettings* Settings = CastChecked<UPCGDifferenceSettings>(TestData.Settings);
+		Settings->Mode = EPCGDifferenceMode::Inferred;
+	}
+
+	void SetSettingsModeContinuous(PCGDeterminismTests::FTestData& TestData)
+	{
+		UPCGDifferenceSettings* Settings = CastChecked<UPCGDifferenceSettings>(TestData.Settings);
+		Settings->Mode = EPCGDifferenceMode::Continuous;
+	}
+
+	void SetSettingsModeDiscrete(PCGDeterminismTests::FTestData& TestData)
+	{
+		UPCGDifferenceSettings* Settings = CastChecked<UPCGDifferenceSettings>(TestData.Settings);
+		Settings->Mode = EPCGDifferenceMode::Discrete;
+	}
+
+	void DifferenceTestBase(PCGDeterminismTests::FTestData& TestData, EPCGDifferenceMode DifferenceMode)
+	{
+		TFunction<void(PCGDeterminismTests::FTestData& TestData)> AdditionalSettingsDelegate;
+
+		switch (DifferenceMode)
+		{
+		case EPCGDifferenceMode::Inferred: // To PointData if points in source and difference inputs
+			AdditionalSettingsDelegate = SetSettingsModeInferred;
+			break;
+		case EPCGDifferenceMode::Continuous: // Stays DifferenceData
+			AdditionalSettingsDelegate = SetSettingsModeContinuous;
+			break;
+		case EPCGDifferenceMode::Discrete: // To PointData
+			AdditionalSettingsDelegate = SetSettingsModeDiscrete;
+			break;
+		default:
+			break;
+		}
+
+		PCGDeterminismTests::GenerateSettings<UPCGDifferenceSettings>(TestData, AdditionalSettingsDelegate);
 		// Source
-		PCGDeterminismTests::AddVolumeInputData(TestData.InputData, FVector::ZeroVector, FVector::OneVector * 2000.f, FVector::OneVector * 200.f, PCGDifferenceConstants::SourceLabel);
+		PCGDeterminismTests::AddVolumeInputData(TestData.InputData, FVector::ZeroVector, PCGDeterminismTests::Defaults::LargeVector, PCGDeterminismTests::Defaults::SmallVector, PCGDifferenceConstants::SourceLabel);
 
 		// Difference 
-		PCGDeterminismTests::AddVolumeInputData(TestData.InputData, FVector::ZeroVector, FVector::OneVector * 500.f, FVector::OneVector * 200.f, PCGDifferenceConstants::DifferencesLabel);
+		PCGDeterminismTests::AddVolumeInputData(TestData.InputData, FVector::ZeroVector, PCGDeterminismTests::Defaults::MediumVector, PCGDeterminismTests::Defaults::SmallVector, PCGDifferenceConstants::DifferencesLabel);
 	}
 
-	void DifferenceTestMultiple(PCGDeterminismTests::FTestData& TestData)
+	void DifferenceTestMultiple(PCGDeterminismTests::FTestData& TestData, EPCGDifferenceMode DifferenceMode)
 	{
-		DifferenceTestBase(TestData);
+		DifferenceTestBase(TestData, DifferenceMode);
 
 		// Randomized Sources
 		PCGDeterminismTests::AddRandomizedVolumeInputData(TestData, PCGDifferenceConstants::SourceLabel);
 		PCGDeterminismTests::AddRandomizedMultiplePointInputData(TestData, 20, PCGDifferenceConstants::SourceLabel);
 
 		// Randomized Differences
+		PCGDeterminismTests::AddRandomizedVolumeInputData(TestData, PCGDifferenceConstants::DifferencesLabel);
 		PCGDeterminismTests::AddRandomizedMultiplePointInputData(TestData, 20, PCGDifferenceConstants::DifferencesLabel);
 	}
-
 }
 
 bool FPCGDifferenceDeterminismSingleSameDataTest::RunTest(const FString& Parameters)
 {
 	// Test single same data
-	PCGDeterminismTests::FTestData TestData(PCGDeterminismTests::DefaultSeed);
+	PCGDeterminismTests::FTestData TestDataInferred(PCGDeterminismTests::Defaults::Seed);
+	PCGDeterminismTests::FTestData TestDataContinuous(PCGDeterminismTests::Defaults::Seed);
+	PCGDeterminismTests::FTestData TestDataDiscrete(PCGDeterminismTests::Defaults::Seed);
 
-	DifferenceTestBase(TestData);
+	DifferenceTestBase(TestDataInferred, EPCGDifferenceMode::Inferred);
+	DifferenceTestBase(TestDataContinuous, EPCGDifferenceMode::Continuous);
+	DifferenceTestBase(TestDataDiscrete, EPCGDifferenceMode::Discrete);
 
-	TestTrue("Same single input and settings, same output", PCGDeterminismTests::ExecutionIsDeterministicSameData(TestData));
+	TestTrue("Same single input, same output, 'Inferred' mode", PCGDeterminismTests::ExecutionIsDeterministicSameData(TestDataInferred));
+	TestTrue("Same single input, same output, 'Continuous' mode", PCGDeterminismTests::ExecutionIsDeterministicSameData(TestDataContinuous));
+	TestTrue("Same single input, same output, 'Discrete' mode", PCGDeterminismTests::ExecutionIsDeterministicSameData(TestDataDiscrete));
 
 	return true;
 }
@@ -53,13 +93,23 @@ bool FPCGDifferenceDeterminismSingleSameDataTest::RunTest(const FString& Paramet
 bool FPCGDifferenceDeterminismSingleIdenticalDataTest::RunTest(const FString& Parameters)
 {
 	// Test single identical data
-	PCGDeterminismTests::FTestData FirstTestData(PCGDeterminismTests::DefaultSeed);
-	PCGDeterminismTests::FTestData SecondTestData(PCGDeterminismTests::DefaultSeed);
-
-	DifferenceTestBase(FirstTestData);
-	DifferenceTestBase(SecondTestData);
-
-	TestTrue("Identical single input and settings, same output", PCGDeterminismTests::ExecutionIsDeterministic(FirstTestData, SecondTestData));
+	PCGDeterminismTests::FTestData FirstTestDataInferred(PCGDeterminismTests::Defaults::Seed);
+	PCGDeterminismTests::FTestData FirstTestDataContinuous(PCGDeterminismTests::Defaults::Seed);
+	PCGDeterminismTests::FTestData FirstTestDataDiscrete(PCGDeterminismTests::Defaults::Seed);
+	PCGDeterminismTests::FTestData SecondTestDataInferred(PCGDeterminismTests::Defaults::Seed);
+	PCGDeterminismTests::FTestData SecondTestDataContinuous(PCGDeterminismTests::Defaults::Seed);
+	PCGDeterminismTests::FTestData SecondTestDataDiscrete(PCGDeterminismTests::Defaults::Seed);
+
+	DifferenceTestBase(FirstTestDataInferred, EPCGDifferenceMode::Inferred);
+	DifferenceTestBase(FirstTestDataContinuous, EPCGDifferenceMode::Continuous);
+	DifferenceTestBase(FirstTestDataDiscrete, EPCGDifferenceMode::Discrete);
+	DifferenceTestBase(SecondTestDataInferred, EPCGDifferenceMode::Inferred);
+	DifferenceTestBase(SecondTestDataContinuous, EPCGDifferenceMode::Continuous);
+	DifferenceTestBase(SecondTestDataDiscrete, EPCGDifferenceMode::Discrete);
+
+	TestTrue("Identical single input, same output, 'Inferred' mode", PCGDeterminismTests::ExecutionIsDeterministic(FirstTestDataInferred, SecondTestDataInferred));
+	TestTrue("Identical single input, same output, 'Continuous' mode", PCGDeterminismTests::ExecutionIsDeterministic(FirstTestDataContinuous, SecondTestDataContinuous));
+	TestTrue("Identical single input, same output, 'Discrete' mode", PCGDeterminismTests::ExecutionIsDeterministic(FirstTestDataDiscrete, SecondTestDataDiscrete));
 
 	return true;
 }
@@ -67,11 +117,17 @@ bool FPCGDifferenceDeterminismSingleIdenticalDataTest::RunTest(const FString& Pa
 bool FPCGDifferenceDeterminismMultipleSameDataTest::RunTest(const FString& Parameters)
 {
 	// Test multiple same data
-	PCGDeterminismTests::FTestData TestData(PCGDeterminismTests::DefaultSeed);
+	PCGDeterminismTests::FTestData TestDataInferred(PCGDeterminismTests::Defaults::Seed);
+	PCGDeterminismTests::FTestData TestDataContinuous(PCGDeterminismTests::Defaults::Seed);
+	PCGDeterminismTests::FTestData TestDataDiscrete(PCGDeterminismTests::Defaults::Seed);
 
-	DifferenceTestMultiple(TestData);
+	DifferenceTestMultiple(TestDataInferred, EPCGDifferenceMode::Inferred);
+	DifferenceTestMultiple(TestDataContinuous, EPCGDifferenceMode::Continuous);
+	DifferenceTestMultiple(TestDataDiscrete, EPCGDifferenceMode::Discrete);
 
-	TestTrue("Identical multiple input, same output", PCGDeterminismTests::ExecutionIsDeterministicSameData(TestData));
+	TestTrue("Same single input, same output, 'Inferred' mode", PCGDeterminismTests::ExecutionIsDeterministicSameData(TestDataInferred));
+	TestTrue("Same single input, same output, 'Continuous' mode", PCGDeterminismTests::ExecutionIsDeterministicSameData(TestDataContinuous));
+	TestTrue("Same single input, same output, 'Discrete' mode", PCGDeterminismTests::ExecutionIsDeterministicSameData(TestDataDiscrete));
 
 	return true;
 }
@@ -79,32 +135,53 @@ bool FPCGDifferenceDeterminismMultipleSameDataTest::RunTest(const FString& Param
 bool FPCGDifferenceDeterminismMultipleIdenticalDataTest::RunTest(const FString& Parameters)
 {
 	// Test multiple identical data
-	PCGDeterminismTests::FTestData FirstTestData(PCGDeterminismTests::DefaultSeed);
-	PCGDeterminismTests::FTestData SecondTestData(PCGDeterminismTests::DefaultSeed);
+	PCGDeterminismTests::FTestData FirstTestDataInferred(PCGDeterminismTests::Defaults::Seed);
+	PCGDeterminismTests::FTestData FirstTestDataContinuous(PCGDeterminismTests::Defaults::Seed);
+	PCGDeterminismTests::FTestData FirstTestDataDiscrete(PCGDeterminismTests::Defaults::Seed);
+	PCGDeterminismTests::FTestData SecondTestDataInferred(PCGDeterminismTests::Defaults::Seed);
+	PCGDeterminismTests::FTestData SecondTestDataContinuous(PCGDeterminismTests::Defaults::Seed);
+	PCGDeterminismTests::FTestData SecondTestDataDiscrete(PCGDeterminismTests::Defaults::Seed);
+
+	DifferenceTestMultiple(FirstTestDataInferred, EPCGDifferenceMode::Inferred);
+	DifferenceTestMultiple(FirstTestDataContinuous, EPCGDifferenceMode::Continuous);
+	DifferenceTestMultiple(FirstTestDataDiscrete, EPCGDifferenceMode::Discrete);
+	DifferenceTestMultiple(SecondTestDataInferred, EPCGDifferenceMode::Inferred);
+	DifferenceTestMultiple(SecondTestDataContinuous, EPCGDifferenceMode::Continuous);
+	DifferenceTestMultiple(SecondTestDataDiscrete, EPCGDifferenceMode::Discrete);
+
+	TestTrue("Identical single input, same output, 'Inferred' mode", PCGDeterminismTests::ExecutionIsDeterministic(FirstTestDataInferred, SecondTestDataInferred));
+	TestTrue("Identical single input, same output, 'Continuous' mode", PCGDeterminismTests::ExecutionIsDeterministic(FirstTestDataContinuous, SecondTestDataContinuous));
+	TestTrue("Identical single input, same output, 'Discrete' mode", PCGDeterminismTests::ExecutionIsDeterministic(FirstTestDataDiscrete, SecondTestDataDiscrete));
 
-	DifferenceTestMultiple(FirstTestData);
-	DifferenceTestMultiple(SecondTestData);
+	return true;
+}
 
-	TestTrue("Identical single input and settings, same output", PCGDeterminismTests::ExecutionIsDeterministic(FirstTestData, SecondTestData));
+bool FPCGDifferenceDeterminismOrderIndependenceTest::RunTest(const FString& Parameters)
+{
+	// Test multiple identical data, shuffled
+	PCGDeterminismTests::FTestData FirstTestDataInferred(PCGDeterminismTests::Defaults::Seed);
+	PCGDeterminismTests::FTestData FirstTestDataContinuous(PCGDeterminismTests::Defaults::Seed);
+	PCGDeterminismTests::FTestData FirstTestDataDiscrete(PCGDeterminismTests::Defaults::Seed);
+	PCGDeterminismTests::FTestData SecondTestDataInferred(PCGDeterminismTests::Defaults::Seed);
+	PCGDeterminismTests::FTestData SecondTestDataContinuous(PCGDeterminismTests::Defaults::Seed);
+	PCGDeterminismTests::FTestData SecondTestDataDiscrete(PCGDeterminismTests::Defaults::Seed);
+
+	DifferenceTestMultiple(FirstTestDataInferred, EPCGDifferenceMode::Inferred);
+	DifferenceTestMultiple(FirstTestDataContinuous, EPCGDifferenceMode::Continuous);
+	DifferenceTestMultiple(FirstTestDataDiscrete, EPCGDifferenceMode::Discrete);
+	DifferenceTestMultiple(SecondTestDataInferred, EPCGDifferenceMode::Inferred);
+	DifferenceTestMultiple(SecondTestDataContinuous, EPCGDifferenceMode::Continuous);
+	DifferenceTestMultiple(SecondTestDataDiscrete, EPCGDifferenceMode::Discrete);
+
+	PCGDeterminismTests::ShuffleInputOrder(SecondTestDataInferred);
+	PCGDeterminismTests::ShuffleInputOrder(SecondTestDataContinuous);
+	PCGDeterminismTests::ShuffleInputOrder(SecondTestDataDiscrete);
+
+	TestTrue("Shuffled input order, same output, 'Inferred' mode", PCGDeterminismTests::ExecutionIsDeterministic(FirstTestDataInferred, SecondTestDataInferred));
+	TestTrue("Shuffled input order, same output, 'Continuous' mode", PCGDeterminismTests::ExecutionIsDeterministic(FirstTestDataContinuous, SecondTestDataContinuous));
+	TestTrue("Shuffled input order, same output, 'Discrete' mode", PCGDeterminismTests::ExecutionIsDeterministic(FirstTestDataDiscrete, SecondTestDataDiscrete));
 
 	return true;
 }
 
-// TODO: Temporarily disabled while more testing into this is conducted
-//bool FPCGDifferenceDeterminismOrderIndependenceTest::RunTest(const FString& Parameters)
-//{
-//	// Test multiple identical data
-//	PCGDeterminismTests::FTestData FirstTestData(PCGDeterminismTests::DefaultSeed);
-//	PCGDeterminismTests::FTestData SecondTestData(PCGDeterminismTests::DefaultSeed);
-//
-//	DifferenceTestMultiple(FirstTestData);
-//	DifferenceTestMultiple(SecondTestData);
-//
-//	PCGDeterminismTests::ShuffleInputOrder(SecondTestData);
-//
-//	TestTrue("Shuffled input order, same output", PCGDeterminismTests::ExecutionIsDeterministic(FirstTestData, SecondTestData));
-//
-//	return true;
-//}
-
 #endif
\ No newline at end of file
diff --git a/Engine/Plugins/Experimental/PCG/Source/PCG/Private/Tests/Determinism/PCGVolumeSamplerDeterminismTest.cpp b/Engine/Plugins/Experimental/PCG/Source/PCG/Private/Tests/Determinism/PCGVolumeSamplerDeterminismTest.cpp
--- a/Engine/Plugins/Experimental/PCG/Source/PCG/Private/Tests/Determinism/PCGVolumeSamplerDeterminismTest.cpp
+++ b/Engine/Plugins/Experimental/PCG/Source/PCG/Private/Tests/Determinism/PCGVolumeSamplerDeterminismTest.cpp
@@ -13,16 +13,19 @@ IMPLEMENT_CUSTOM_SIMPLE_AUTOMATION_TEST(FPCGVolumeSamplerDeterminismOrderIndepen
 
 namespace
 {
+	constexpr static int32 NumInputVolumesToAdd = 10;
+
 	void RandomizeVolumeSettingsVoxelSize(PCGDeterminismTests::FTestData& TestData)
 	{
-		CastChecked<UPCGVolumeSamplerSettings>(TestData.Settings)->VoxelSize = FVector::OneVector * 200.f + TestData.RandomStream.VRand() * 100.f;
+		UPCGVolumeSamplerSettings* Settings = CastChecked<UPCGVolumeSamplerSettings>(TestData.Settings);
+		Settings->VoxelSize = PCGDeterminismTests::Defaults::MediumVector + TestData.RandomStream.VRand() * 0.5 * PCGDeterminismTests::Defaults::MediumDistance;
 	}
 }
 
 bool FPCGVolumeSamplerDeterminismSingleSameDataTest::RunTest(const FString& Parameters)
 {
 	// Test single same data
-	PCGDeterminismTests::FTestData TestData(PCGDeterminismTests::DefaultSeed);
+	PCGDeterminismTests::FTestData TestData(PCGDeterminismTests::Defaults::Seed);
 
 	PCGDeterminismTests::GenerateSettings<UPCGVolumeSamplerSettings>(TestData, RandomizeVolumeSettingsVoxelSize);
 	PCGDeterminismTests::AddRandomizedVolumeInputData(TestData);
@@ -35,8 +38,8 @@ bool FPCGVolumeSamplerDeterminismSingleSameDataTest::RunTest(const FString& Para
 bool FPCGVolumeSamplerDeterminismSingleIdenticalDataTest::RunTest(const FString& Parameters)
 {
 	// Test single identical data
-	PCGDeterminismTests::FTestData FirstTestData(PCGDeterminismTests::DefaultSeed);
-	PCGDeterminismTests::FTestData SecondTestData(PCGDeterminismTests::DefaultSeed);
+	PCGDeterminismTests::FTestData FirstTestData(PCGDeterminismTests::Defaults::Seed);
+	PCGDeterminismTests::FTestData SecondTestData(PCGDeterminismTests::Defaults::Seed);
 
 	PCGDeterminismTests::GenerateSettings<UPCGVolumeSamplerSettings>(FirstTestData, RandomizeVolumeSettingsVoxelSize);
 	PCGDeterminismTests::GenerateSettings<UPCGVolumeSamplerSettings>(SecondTestData, RandomizeVolumeSettingsVoxelSize);
@@ -52,12 +55,12 @@ bool FPCGVolumeSamplerDeterminismSingleIdenticalDataTest::RunTest(const FString&
 bool FPCGVolumeSamplerDeterminismMultipleSameDataTest::RunTest(const FString& Parameters)
 {
 	// Test multiple same data
-	PCGDeterminismTests::FTestData TestData(PCGDeterminismTests::DefaultSeed);
+	PCGDeterminismTests::FTestData TestData(PCGDeterminismTests::Defaults::Seed);
 
 	PCGDeterminismTests::GenerateSettings<UPCGVolumeSamplerSettings>(TestData, RandomizeVolumeSettingsVoxelSize);
 
 	// Add many random volumes
-	for (int32 I = 0; I < 10; ++I)
+	for (int32 I = 0; I < NumInputVolumesToAdd; ++I)
 	{
 		PCGDeterminismTests::AddRandomizedVolumeInputData(TestData);
 	}
@@ -70,14 +73,14 @@ bool FPCGVolumeSamplerDeterminismMultipleSameDataTest::RunTest(const FString& Pa
 bool FPCGVolumeSamplerDeterminismMultipleIdenticalDataTest::RunTest(const FString& Parameters)
 {
 	// Test multiple identical data
-	PCGDeterminismTests::FTestData FirstTestData(PCGDeterminismTests::DefaultSeed);
-	PCGDeterminismTests::FTestData SecondTestData(PCGDeterminismTests::DefaultSeed);
+	PCGDeterminismTests::FTestData FirstTestData(PCGDeterminismTests::Defaults::Seed);
+	PCGDeterminismTests::FTestData SecondTestData(PCGDeterminismTests::Defaults::Seed);
 
 	PCGDeterminismTests::GenerateSettings<UPCGVolumeSamplerSettings>(FirstTestData, RandomizeVolumeSettingsVoxelSize);
 	PCGDeterminismTests::GenerateSettings<UPCGVolumeSamplerSettings>(SecondTestData, RandomizeVolumeSettingsVoxelSize);
 
 	// Add many random volumes
-	for (int32 I = 0; I < 10; ++I)
+	for (int32 I = 0; I < NumInputVolumesToAdd; ++I)
 	{
 		PCGDeterminismTests::AddRandomizedVolumeInputData(FirstTestData);
 		PCGDeterminismTests::AddRandomizedVolumeInputData(SecondTestData);
@@ -91,14 +94,14 @@ bool FPCGVolumeSamplerDeterminismMultipleIdenticalDataTest::RunTest(const FStrin
 bool FPCGVolumeSamplerDeterminismOrderIndependenceTest::RunTest(const FString& Parameters)
 {
 	// Test shuffled input data
-	PCGDeterminismTests::FTestData FirstTestData(PCGDeterminismTests::DefaultSeed);
-	PCGDeterminismTests::FTestData SecondTestData(PCGDeterminismTests::DefaultSeed);
+	PCGDeterminismTests::FTestData FirstTestData(PCGDeterminismTests::Defaults::Seed);
+	PCGDeterminismTests::FTestData SecondTestData(PCGDeterminismTests::Defaults::Seed);
 
 	PCGDeterminismTests::GenerateSettings<UPCGVolumeSamplerSettings>(FirstTestData, RandomizeVolumeSettingsVoxelSize);
 	PCGDeterminismTests::GenerateSettings<UPCGVolumeSamplerSettings>(SecondTestData, RandomizeVolumeSettingsVoxelSize);
 
 	// Add many random volumes
-	for (int32 I = 0; I < 10; ++I)
+	for (int32 I = 0; I < NumInputVolumesToAdd; ++I)
 	{
 		PCGDeterminismTests::AddRandomizedVolumeInputData(FirstTestData);
 		PCGDeterminismTests::AddRandomizedVolumeInputData(SecondTestData);
diff --git a/Engine/Plugins/Experimental/PCG/Source/PCG/Private/Tests/Elements/PCGStaticMeshSpawnerTest.cpp b/Engine/Plugins/Experimental/PCG/Source/PCG/Private/Tests/Elements/PCGStaticMeshSpawnerTest.cpp
--- a/Engine/Plugins/Experimental/PCG/Source/PCG/Private/Tests/Elements/PCGStaticMeshSpawnerTest.cpp
+++ b/Engine/Plugins/Experimental/PCG/Source/PCG/Private/Tests/Elements/PCGStaticMeshSpawnerTest.cpp
@@ -430,7 +430,7 @@ void TestMeshSelectorWeightedByCategory(
 
 bool FPCGStaticMeshSpawnerByAttributeTest::RunTest(const FString& Parameters)
 {
-	PCGDeterminismTests::FTestData TestData(PCGDeterminismTests::DefaultSeed);
+	PCGDeterminismTests::FTestData TestData(PCGDeterminismTests::Defaults::Seed);
 	PCGDeterminismTests::GenerateSettings<UPCGStaticMeshSpawnerSettings>(TestData);
 
 	// create test attribute data
@@ -520,7 +520,7 @@ bool FPCGStaticMeshSpawnerByAttributeTest::RunTest(const FString& Parameters)
 
 bool FPCGStaticMeshSpawnerWeightedTest::RunTest(const FString& Parameters)
 {
-	PCGDeterminismTests::FTestData TestData(PCGDeterminismTests::DefaultSeed);
+	PCGDeterminismTests::FTestData TestData(PCGDeterminismTests::Defaults::Seed);
 	PCGDeterminismTests::GenerateSettings<UPCGStaticMeshSpawnerSettings>(TestData);
 
 	// create test meshes
@@ -589,7 +589,7 @@ bool FPCGStaticMeshSpawnerWeightedTest::RunTest(const FString& Parameters)
 
 bool FPCGStaticMeshSpawnerWeightedByCategoryTest::RunTest(const FString& Parameters)
 {
-	PCGDeterminismTests::FTestData TestData(PCGDeterminismTests::DefaultSeed);
+	PCGDeterminismTests::FTestData TestData(PCGDeterminismTests::Defaults::Seed);
 	PCGDeterminismTests::GenerateSettings<UPCGStaticMeshSpawnerSettings>(TestData);
 
 	// create test meshes
diff --git a/Engine/Plugins/Experimental/PCG/Source/PCG/Public/Tests/Determinism/PCGDeterminismTestsCommon.h b/Engine/Plugins/Experimental/PCG/Source/PCG/Public/Tests/Determinism/PCGDeterminismTestsCommon.h
--- a/Engine/Plugins/Experimental/PCG/Source/PCG/Public/Tests/Determinism/PCGDeterminismTestsCommon.h
+++ b/Engine/Plugins/Experimental/PCG/Source/PCG/Public/Tests/Determinism/PCGDeterminismTestsCommon.h
@@ -16,16 +16,29 @@ struct FPCGDataCollection;
 
 namespace PCGDeterminismTests
 {
-	constexpr static int32 DefaultSeed = 42;
-	constexpr static int32 DefaultNumPointsToGenerate = 1;
-	constexpr static int32 DefaultNumPolyLinePointsToGenerate = 6;
-
-	const FName DefaultTestPinName = TEXT("Test");
+	namespace Defaults
+	{
+		constexpr static int32 Seed = 42;
+		constexpr static int32 NumPointsToGenerate = 1;
+		constexpr static int32 NumPolyLinePointsToGenerate = 6;
+		constexpr static int32 NumSamplingStepsPerDimension = 100;
+
+		constexpr static FVector::FReal SmallDistance = 50.0;
+		constexpr static FVector::FReal MediumDistance = 200.0;
+		constexpr static FVector::FReal LargeDistance = 500.0;
+
+		const FVector SmallVector = FVector::OneVector * SmallDistance;
+		const FVector MediumVector = FVector::OneVector * MediumDistance;
+		const FVector LargeVector = FVector::OneVector * LargeDistance;
+
+		const FName TestPinName = TEXT("Test");
+		/** A little bigger than the typical largest volumes */
+		const FBox TestingVolume = FBox(-1.2 * LargeVector, 1.2 * LargeVector);
+	}
 
 	struct FTestData
 	{
-		explicit FTestData(int32 Seed = DefaultSeed);
-
+		explicit FTestData(int32 Seed = Defaults::Seed);
 		void Reset();
 
 		AActor* TestActor;
@@ -53,28 +66,28 @@ namespace PCGDeterminismTests
 	void AddInputDataBasedOnPins(FTestData& TestData, const UPCGNode* InPCGNode, FPCGDeterminismResult& OutResult);
 
 	/** Helper for adding PointData with a single Point to an InputData */
-	void AddSinglePointInputData(FPCGDataCollection& InputData, const FVector& Location, const FName& PinName = DefaultTestPinName);
+	void AddSinglePointInputData(FPCGDataCollection& InputData, const FVector& Location, const FName& PinName = Defaults::TestPinName);
 	/** Helper for adding PointData with multiple Points to an InputData */
-	void AddMultiplePointsInputData(FPCGDataCollection& InputData, const TArray<FPCGPoint>& Points, const FName& PinName = DefaultTestPinName);
+	void AddMultiplePointsInputData(FPCGDataCollection& InputData, const TArray<FPCGPoint>& Points, const FName& PinName = Defaults::TestPinName);
 	/** Helper for adding VolumeData to an InputData */
-	void AddVolumeInputData(FPCGDataCollection& InputData, const FVector& Location, const FVector& HalfSize, const FVector& VoxelSize, const FName& PinName = DefaultTestPinName);
+	void AddVolumeInputData(FPCGDataCollection& InputData, const FVector& Location, const FVector& HalfSize, const FVector& VoxelSize, const FName& PinName = Defaults::TestPinName);
 	/** Helper for adding PolyLineData to an InputData */
-	void AddPolyLineInputData(FPCGDataCollection& InputData, USplineComponent* SplineComponent, const FName& PinName = DefaultTestPinName);
+	void AddPolyLineInputData(FPCGDataCollection& InputData, USplineComponent* SplineComponent, const FName& PinName = Defaults::TestPinName);
 	/** Helper for adding PrimitiveData to an InputData */
-	void AddPrimitiveInputData(FPCGDataCollection& InputData, UPrimitiveComponent* SplineComponent, const FVector& VoxelSize, const FName& PinName = DefaultTestPinName);
+	void AddPrimitiveInputData(FPCGDataCollection& InputData, UPrimitiveComponent* SplineComponent, const FVector& VoxelSize, const FName& PinName = Defaults::TestPinName);
 
 	/** Adds randomized PointData with a single Point */
-	void AddRandomizedSinglePointInputData(FTestData& TestData, int32 PointNum = DefaultNumPointsToGenerate, const FName& PinName = DefaultTestPinName);
+	void AddRandomizedSinglePointInputData(FTestData& TestData, int32 PointNum = Defaults::NumPointsToGenerate, const FName& PinName = Defaults::TestPinName);
 	/** Adds randomized PointData with multiple Points */
-	void AddRandomizedMultiplePointInputData(FTestData& TestData, int32 PointNum = DefaultNumPointsToGenerate, const FName& PinName = DefaultTestPinName);
+	void AddRandomizedMultiplePointInputData(FTestData& TestData, int32 PointNum = Defaults::NumPointsToGenerate, const FName& PinName = Defaults::TestPinName);
 	/** Adds randomized Volume Spatial Data */
-	void AddRandomizedVolumeInputData(FTestData& TestData, const FName& PinName = DefaultTestPinName);
+	void AddRandomizedVolumeInputData(FTestData& TestData, const FName& PinName = Defaults::TestPinName);
 	/** Adds randomized Surface Spatial Data */
-	void AddRandomizedSurfaceInputData(FTestData& TestData, const FName& PinName = DefaultTestPinName);
+	void AddRandomizedSurfaceInputData(FTestData& TestData, const FName& PinName = Defaults::TestPinName);
 	/** Adds randomized PolyLine Spatial Data */
-	void AddRandomizedPolyLineInputData(FTestData& TestData, int32 PointNum = DefaultNumPolyLinePointsToGenerate, const FName& PinName = DefaultTestPinName);
+	void AddRandomizedPolyLineInputData(FTestData& TestData, int32 PointNum = Defaults::NumPolyLinePointsToGenerate, const FName& PinName = Defaults::TestPinName);
 	/** Adds randomized Primitive Spatial Data */
-	void AddRandomizedPrimitiveInputData(FTestData& TestData, const FName& PinName = DefaultTestPinName);
+	void AddRandomizedPrimitiveInputData(FTestData& TestData, const FName& PinName = Defaults::TestPinName);
 
 	/** Validates whether two DataCollection objects are identical */
 	bool DataCollectionsAreIdentical(const FPCGDataCollection& FirstCollection, const FPCGDataCollection& SecondCollection);
@@ -91,8 +104,9 @@ namespace PCGDeterminismTests
 	bool PolyLineDataIsIdentical(const UPCGData* FirstData, const UPCGData* SecondData);
 	/** Validates whether two PrimitiveData objects are identical */
 	bool PrimitiveDataIsIdentical(const UPCGData* FirstData, const UPCGData* SecondData);
-	/** Validates whether two DifferenceData objects are identical */
-	bool DifferenceDataIsIdentical(const UPCGData* FirstData, const UPCGData* SecondData);
+	
+	/** Validates whether two SpatialData objects are identical via Point Sampling */
+	bool SampledSpatialDataIsIdentical(const UPCGSpatialData* FirstSpatialData, const UPCGSpatialData* SecondSpatialData);
 
 	/** Validates the basics of SpatialData are identical */
 	bool SpatialBasicsAreIdentical(const UPCGSpatialData* FirstSpatialData, const UPCGSpatialData* SecondSpatialData);
