Added FRotator to the data types supported in the PCG metadata

diff --git a/Engine/Plugins/Experimental/PCG/Source/PCG/Private/Helpers/PCGSettingsHelpers.cpp b/Engine/Plugins/Experimental/PCG/Source/PCG/Private/Helpers/PCGSettingsHelpers.cpp
--- a/Engine/Plugins/Experimental/PCG/Source/PCG/Private/Helpers/PCGSettingsHelpers.cpp
+++ b/Engine/Plugins/Experimental/PCG/Source/PCG/Private/Helpers/PCGSettingsHelpers.cpp
@@ -156,7 +156,7 @@ namespace PCGSettingsHelpers
 			else if (PropertyTypeName == NAME_Quat)
 			{
 				FQuat Value = FQuat::Identity;
-				if (GetParamValue<FQuat>(Value, MatchingAttribute))
+				if (GetParamValue<FQuat, FRotator>(Value, MatchingAttribute))
 				{
 					*reinterpret_cast<FQuat*>(PropertyValue) = Value;
 				}
@@ -173,7 +173,18 @@ namespace PCGSettingsHelpers
 					bTypeError = true;
 				}
 			}
-			//else if(PropertyTypeName == NAME_Rotator)
+			else if (PropertyTypeName == NAME_Rotator)
+			{
+				FRotator Value = FRotator::ZeroRotator;
+				if (GetParamValue<FRotator, FQuat>(Value, MatchingAttribute))
+				{
+					*reinterpret_cast<FRotator*>(PropertyValue) = Value;
+				}
+				else
+				{
+					bTypeError = true;
+				}
+			}
 			//else if (PropertyTypeName == NAME_Color)
 			//else if (PropertyTypeName == NAME_LinearColor)
 			//else if soft object path to something?
diff --git a/Engine/Plugins/Experimental/PCG/Source/PCG/Private/Metadata/PCGMetadata.cpp b/Engine/Plugins/Experimental/PCG/Source/PCG/Private/Metadata/PCGMetadata.cpp
--- a/Engine/Plugins/Experimental/PCG/Source/PCG/Private/Metadata/PCGMetadata.cpp
+++ b/Engine/Plugins/Experimental/PCG/Source/PCG/Private/Metadata/PCGMetadata.cpp
@@ -316,6 +316,11 @@ void UPCGMetadata::CreateVector4Attribute(FName AttributeName, FVector4 DefaultV
 	CreateAttribute<FVector4>(AttributeName, DefaultValue, bAllowsInterpolation, bOverrideParent);
 }
 
+void UPCGMetadata::CreateRotatorAttribute(FName AttributeName, FRotator DefaultValue, bool bAllowsInterpolation, bool bOverrideParent)
+{
+	CreateAttribute<FRotator>(AttributeName, DefaultValue, bAllowsInterpolation, bOverrideParent);
+}
+
 void UPCGMetadata::CreateQuatAttribute(FName AttributeName, FQuat DefaultValue, bool bAllowsInterpolation, bool bOverrideParent)
 {
 	CreateAttribute<FQuat>(AttributeName, DefaultValue, bAllowsInterpolation, bOverrideParent);
diff --git a/Engine/Plugins/Experimental/PCG/Source/PCG/Private/Metadata/PCGMetadataAccessor.cpp b/Engine/Plugins/Experimental/PCG/Source/PCG/Private/Metadata/PCGMetadataAccessor.cpp
--- a/Engine/Plugins/Experimental/PCG/Source/PCG/Private/Metadata/PCGMetadataAccessor.cpp
+++ b/Engine/Plugins/Experimental/PCG/Source/PCG/Private/Metadata/PCGMetadataAccessor.cpp
@@ -129,6 +129,16 @@ void UPCGMetadataAccessorHelpers::SetVector4AttributeByMetadataKey(int64& Key, U
 	SetAttribute(Key, Metadata, AttributeName, Value);
 }
 
+FRotator UPCGMetadataAccessorHelpers::GetRotatorAttributeByMetadataKey(int64 Key, const UPCGMetadata* Metadata, FName AttributeName)
+{
+	return GetAttribute<FRotator>(Key, Metadata, AttributeName);
+}
+
+void UPCGMetadataAccessorHelpers::SetRotatorAttributeByMetadataKey(int64& Key, UPCGMetadata* Metadata, FName AttributeName, const FRotator& Value)
+{
+	SetAttribute(Key, Metadata, AttributeName, Value);
+}
+
 FQuat UPCGMetadataAccessorHelpers::GetQuatAttributeByMetadataKey(int64 Key, const UPCGMetadata* Metadata, FName AttributeName)
 {
 	return GetAttribute<FQuat>(Key, Metadata, AttributeName);
@@ -248,6 +258,16 @@ void UPCGMetadataAccessorHelpers::SetVector4Attribute(FPCGPoint& Point, UPCGMeta
 	SetAttribute(Point.MetadataEntry, Metadata, AttributeName, Value);
 }
 
+FRotator UPCGMetadataAccessorHelpers::GetRotatorAttribute(const FPCGPoint& Point, const UPCGMetadata* Metadata, FName AttributeName)
+{
+	return GetAttribute<FRotator>(Point.MetadataEntry, Metadata, AttributeName);
+}
+
+void UPCGMetadataAccessorHelpers::SetRotatorAttribute(FPCGPoint& Point, UPCGMetadata* Metadata, FName AttributeName, const FRotator& Value)
+{
+	SetAttribute(Point.MetadataEntry, Metadata, AttributeName, Value);
+}
+
 FQuat UPCGMetadataAccessorHelpers::GetQuatAttribute(const FPCGPoint& Point, const UPCGMetadata* Metadata, FName AttributeName)
 {
 	return GetAttribute<FQuat>(Point.MetadataEntry, Metadata, AttributeName);
diff --git a/Engine/Plugins/Experimental/PCG/Source/PCG/Public/Metadata/PCGMetadata.h b/Engine/Plugins/Experimental/PCG/Source/PCG/Public/Metadata/PCGMetadata.h
--- a/Engine/Plugins/Experimental/PCG/Source/PCG/Public/Metadata/PCGMetadata.h
+++ b/Engine/Plugins/Experimental/PCG/Source/PCG/Public/Metadata/PCGMetadata.h
@@ -74,6 +74,9 @@ public:
 	UFUNCTION(BlueprintCallable, Category = "PCG|Metadata", meta=(DisplayName = "Create Vector4 Attribute"))
 	void CreateVector4Attribute(FName AttributeName, FVector4 DefaultValue, bool bAllowsInterpolation, bool bOverrideParent = true);
 
+	UFUNCTION(BlueprintCallable, Category = "PCG|Metadata")
+	void CreateRotatorAttribute(FName AttributeName, FRotator DefaultValue, bool bAllowsInterpolation, bool bOverrideParent = true);
+
 	UFUNCTION(BlueprintCallable, Category = "PCG|Metadata")
 	void CreateQuatAttribute(FName AttributeName, FQuat DefaultValue, bool bAllowsInterpolation, bool bOverrideParent = true);
 
diff --git a/Engine/Plugins/Experimental/PCG/Source/PCG/Public/Metadata/PCGMetadataAccessor.h b/Engine/Plugins/Experimental/PCG/Source/PCG/Public/Metadata/PCGMetadataAccessor.h
--- a/Engine/Plugins/Experimental/PCG/Source/PCG/Public/Metadata/PCGMetadataAccessor.h
+++ b/Engine/Plugins/Experimental/PCG/Source/PCG/Public/Metadata/PCGMetadataAccessor.h
@@ -42,6 +42,12 @@ public:
 	UFUNCTION(BlueprintCallable, Category = "PCG|Metadata")
 	static void SetVector4AttributeByMetadataKey(UPARAM(ref) int64& Key, UPCGMetadata* Metadata, FName AttributeName, const FVector4& Value);
 
+	UFUNCTION(BlueprintCallable, Category = "PCG|Metadata")
+	static FRotator GetRotatorAttributeByMetadataKey(int64 Key, const UPCGMetadata* Metadata, FName AttributeName);
+
+	UFUNCTION(BlueprintCallable, Category = "PCG|Metadata")
+	static void SetRotatorAttributeByMetadataKey(UPARAM(ref) int64& Key, UPCGMetadata* Metadata, FName AttributeName, const FRotator& Value);
+
 	UFUNCTION(BlueprintCallable, Category = "PCG|Metadata")
 	static FQuat GetQuatAttributeByMetadataKey(int64 Key, const UPCGMetadata* Metadata, FName AttributeName);
 
@@ -99,6 +105,12 @@ public:
 	UFUNCTION(BlueprintCallable, Category = "PCG|Metadata", meta = (ScriptMethod))
 	static void SetVector4Attribute(UPARAM(ref) FPCGPoint& Point, UPCGMetadata* Metadata, FName AttributeName, const FVector4& Value);
 
+	UFUNCTION(BlueprintCallable, Category = "PCG|Metadata", meta = (ScriptMethod))
+	static FRotator GetRotatorAttribute(const FPCGPoint& Point, const UPCGMetadata* Metadata, FName AttributeName);
+
+	UFUNCTION(BlueprintCallable, Category = "PCG|Metadata", meta = (ScriptMethod))
+	static void SetRotatorAttribute(UPARAM(ref) FPCGPoint& Point, UPCGMetadata* Metadata, FName AttributeName, const FRotator& Value);
+
 	UFUNCTION(BlueprintCallable, Category = "PCG|Metadata", meta = (ScriptMethod))
 	static FQuat GetQuatAttribute(const FPCGPoint& Point, const UPCGMetadata* Metadata, FName AttributeName);
 
diff --git a/Engine/Plugins/Experimental/PCG/Source/PCG/Public/Metadata/PCGMetadataAttributeTraits.h b/Engine/Plugins/Experimental/PCG/Source/PCG/Public/Metadata/PCGMetadataAttributeTraits.h
--- a/Engine/Plugins/Experimental/PCG/Source/PCG/Public/Metadata/PCGMetadataAttributeTraits.h
+++ b/Engine/Plugins/Experimental/PCG/Source/PCG/Public/Metadata/PCGMetadataAttributeTraits.h
@@ -19,6 +19,7 @@ enum class EPCGMetadataTypes : uint8
 	Transform,
 	String,
 	Boolean,
+	Rotator,
 	Unknown = 255
 };
 
@@ -44,6 +45,7 @@ namespace PCG
 		PCGMetadataGenerateDataTypes(FTransform, Transform);
 		PCGMetadataGenerateDataTypes(FString, String);
 		PCGMetadataGenerateDataTypes(bool, Boolean);
+		PCGMetadataGenerateDataTypes(FRotator, Rotator);
 
 #undef PCGMetadataGenerateDataTypes
 
@@ -257,6 +259,53 @@ namespace PCG
 			}
 		};
 
+		// Rotator
+		template<>
+		struct MetadataTraits<FRotator>
+		{
+			enum { CompressData = false };
+			enum { CanMinMax = false };
+			enum { CanSubAdd = true };
+			enum { CanMulDiv = true };
+			enum { CanInterpolate = true };
+
+			static bool Equal(const FRotator& A, const FRotator& B)
+			{
+				return A == B;
+			}
+
+			static FRotator Add(const FRotator& A, const FRotator& B)
+			{
+				return A + B;
+			}
+
+			static FRotator Sub(const FRotator& A, const FRotator& B)
+			{
+				return A - B;
+			}
+
+			static FRotator Mul(const FRotator& A, const FRotator& B)
+			{
+				return A + B;
+			}
+
+			static FRotator Div(const FRotator& A, const FRotator& B)
+			{
+				return A - B;
+			}
+
+			static FRotator WeightedSum(const FRotator& A, const FRotator& B, float Weight)
+			{
+				// TODO review this, should we use TCustomLerp<UE::Math::TRotator<T>> ?
+				return A + (B * Weight);
+			}
+
+			static FRotator ZeroValue()
+			{
+				return FRotator::ZeroRotator;
+			}
+		};
+
 		// Transform
 		template<>
 		struct MetadataTraits<FTransform>
