PCG:
*Removed option to open PCG graphs in the old asset editor.

diff --git a/Engine/Plugins/Experimental/PCG/Source/PCGEditor/Private/AssetTypeActions/PCGGraphAssetTypeActions.cpp b/Engine/Plugins/Experimental/PCG/Source/PCGEditor/Private/AssetTypeActions/PCGGraphAssetTypeActions.cpp
--- a/Engine/Plugins/Experimental/PCG/Source/PCGEditor/Private/AssetTypeActions/PCGGraphAssetTypeActions.cpp
+++ b/Engine/Plugins/Experimental/PCG/Source/PCGEditor/Private/AssetTypeActions/PCGGraphAssetTypeActions.cpp
@@ -4,10 +4,6 @@
 #include "PCGGraph.h"
 #include "PCGEditor.h"
 
-#include "HAL/IConsoleManager.h"
-
-static TAutoConsoleVariable<bool> CVarPCGUseGraphEditor(TEXT("pcg.UseGraphEditor"), true, TEXT("Wheter to use use the new graph editor or not."));
-
 FText FPCGGraphAssetTypeActions::GetName() const
 {
 	return NSLOCTEXT("AssetTypeActions", "PCGGraphAssetTypeActions", "PCG Graph");
@@ -20,19 +16,12 @@ UClass* FPCGGraphAssetTypeActions::GetSupportedClass() const
 
 void FPCGGraphAssetTypeActions::OpenAssetEditor(const TArray<UObject*>& InObjects, TSharedPtr<class IToolkitHost> EditWithinLevelEditor)
 {
-	if (CVarPCGUseGraphEditor.GetValueOnAnyThread())
+	for (UObject* Object : InObjects)
 	{
-		for (UObject* Object : InObjects)
+		if (UPCGGraph* PCGGraph = Cast<UPCGGraph>(Object))
 		{
-			if (UPCGGraph* PCGGraph = Cast<UPCGGraph>(Object))
-			{
-				TSharedRef<FPCGEditor> PCGEditor(new FPCGEditor());
-				PCGEditor->Initialize(EToolkitMode::Standalone, EditWithinLevelEditor, PCGGraph);
-			}
+			const TSharedRef<FPCGEditor> PCGEditor = MakeShared<FPCGEditor>();
+			PCGEditor->Initialize(EToolkitMode::Standalone, EditWithinLevelEditor, PCGGraph);
 		}
 	}
-	else
-	{
-		FPCGCommonAssetTypeActions::OpenAssetEditor(InObjects, EditWithinLevelEditor);
-	}
-}
\ No newline at end of file
+}
