PCG:
*Removed unused viewport tab from PCG editor

diff --git a/Engine/Plugins/Experimental/PCG/Source/PCGEditor/Private/PCGEditor.cpp b/Engine/Plugins/Experimental/PCG/Source/PCGEditor/Private/PCGEditor.cpp
--- a/Engine/Plugins/Experimental/PCG/Source/PCGEditor/Private/PCGEditor.cpp
+++ b/Engine/Plugins/Experimental/PCG/Source/PCGEditor/Private/PCGEditor.cpp
@@ -59,7 +59,6 @@ namespace FPCGEditor_private
 	const FName PropertyDetailsID = FName(TEXT("PropertyDetails"));
 	const FName PaletteID = FName(TEXT("Palette"));
 	const FName AttributesID = FName(TEXT("Attributes"));
-	const FName ViewportID = FName(TEXT("Viewport"));
 	const FName FindID = FName(TEXT("Find"));
 	const FName DeterminismID = FName(TEXT("Determinism"));
 }
@@ -92,23 +91,11 @@ void FPCGEditor::Initialize(const EToolkitMode::Type InMode, const TSharedPtr<cl
 		(
 			FTabManager::NewPrimaryArea()->SetOrientation(Orient_Horizontal)
 			->Split
-			(
-				FTabManager::NewSplitter()->SetOrientation(Orient_Vertical)
+			(				
+				FTabManager::NewStack()
 				->SetSizeCoefficient(0.10f)
-				->Split
-				(
-					FTabManager::NewStack()
-					->SetSizeCoefficient(0.16)
-					->SetHideTabWell(true)
-					->AddTab(FPCGEditor_private::ViewportID, ETabState::OpenedTab)
-				)
-				->Split
-				(
-					FTabManager::NewStack()
-					->SetSizeCoefficient(0.84)
-					->SetHideTabWell(true)
-					->AddTab(FPCGEditor_private::PaletteID, ETabState::OpenedTab)
-				)
+				->SetHideTabWell(true)
+				->AddTab(FPCGEditor_private::PaletteID, ETabState::OpenedTab)
 			)
 			->Split
 			(
@@ -182,10 +169,6 @@ void FPCGEditor::RegisterTabSpawners(const TSharedRef<FTabManager>& InTabManager
 		.SetDisplayName(LOCTEXT("AttributesTab", "Attributes"))
 		.SetGroup(WorkspaceMenuCategoryRef);
 
-	InTabManager->RegisterTabSpawner(FPCGEditor_private::ViewportID, FOnSpawnTab::CreateSP(this, &FPCGEditor::SpawnTab_Viewport))
-		.SetDisplayName(LOCTEXT("ViewportTab", "Viewport"))
-		.SetGroup(WorkspaceMenuCategoryRef);
-
 	InTabManager->RegisterTabSpawner(FPCGEditor_private::FindID, FOnSpawnTab::CreateSP(this, &FPCGEditor::SpawnTab_Find))
 		.SetDisplayName(LOCTEXT("FindTab", "Find"))
 		.SetGroup(WorkspaceMenuCategoryRef);
@@ -201,7 +184,6 @@ void FPCGEditor::UnregisterTabSpawners(const TSharedRef<class FTabManager>& InTa
 	InTabManager->UnregisterTabSpawner(FPCGEditor_private::PropertyDetailsID);
 	InTabManager->UnregisterTabSpawner(FPCGEditor_private::PaletteID);
 	InTabManager->UnregisterTabSpawner(FPCGEditor_private::AttributesID);
-	InTabManager->UnregisterTabSpawner(FPCGEditor_private::ViewportID);
 	InTabManager->UnregisterTabSpawner(FPCGEditor_private::FindID);
 	InTabManager->UnregisterTabSpawner(FPCGEditor_private::DeterminismID);
 
@@ -1309,16 +1291,6 @@ TSharedRef<SDockTab> FPCGEditor::SpawnTab_Attributes(const FSpawnTabArgs& Args)
 		];
 }
 
-TSharedRef<SDockTab> FPCGEditor::SpawnTab_Viewport(const FSpawnTabArgs& Args)
-{
-	return SNew(SDockTab)
-		.Label(LOCTEXT("PCGViewportTitle", "Viewport"))
-		.TabColorScale(GetTabColorScale())
-		[
-			SNullWidget::NullWidget
-		];
-}
-
 TSharedRef<SDockTab> FPCGEditor::SpawnTab_Find(const FSpawnTabArgs& Args)
 {
 	return SNew(SDockTab)
diff --git a/Engine/Plugins/Experimental/PCG/Source/PCGEditor/Private/PCGEditor.h b/Engine/Plugins/Experimental/PCG/Source/PCGEditor/Private/PCGEditor.h
--- a/Engine/Plugins/Experimental/PCG/Source/PCGEditor/Private/PCGEditor.h
+++ b/Engine/Plugins/Experimental/PCG/Source/PCGEditor/Private/PCGEditor.h
@@ -151,7 +151,6 @@ private:
 	TSharedRef<SDockTab> SpawnTab_PropertyDetails(const FSpawnTabArgs& Args);
 	TSharedRef<SDockTab> SpawnTab_Palette(const FSpawnTabArgs& Args);
 	TSharedRef<SDockTab> SpawnTab_Attributes(const FSpawnTabArgs& Args);
-	TSharedRef<SDockTab> SpawnTab_Viewport(const FSpawnTabArgs& Args);
 	TSharedRef<SDockTab> SpawnTab_Find(const FSpawnTabArgs& Args);
 	TSharedRef<SDockTab> SpawnTab_Determinism(const FSpawnTabArgs& Args);
 
@@ -165,4 +164,4 @@ private:
 
 	UPCGGraph* PCGGraphBeingEdited = nullptr;
 	UPCGEditorGraph* PCGEditorGraph = nullptr;
-};
\ No newline at end of file
+};
