+ Added intersection element determinism tests

diff --git a/Engine/Plugins/Experimental/PCG/Source/PCG/Private/Tests/Determinism/PCGDeterminismTestsCommon.cpp b/Engine/Plugins/Experimental/PCG/Source/PCG/Private/Tests/Determinism/PCGDeterminismTestsCommon.cpp
--- a/Engine/Plugins/Experimental/PCG/Source/PCG/Private/Tests/Determinism/PCGDeterminismTestsCommon.cpp
+++ b/Engine/Plugins/Experimental/PCG/Source/PCG/Private/Tests/Determinism/PCGDeterminismTestsCommon.cpp
@@ -6,8 +6,6 @@
 #include "PCGGraph.h"
 #include "PCGHelpers.h"
 #include "PCGPoint.h"
-#include "Data/PCGDifferenceData.h"
-#include "Data/PCGIntersectionData.h"
 #include "Data/PCGPointData.h"
 #include "Data/PCGPrimitiveData.h"
 #include "Data/PCGSplineData.h"
@@ -205,7 +203,7 @@ namespace PCGDeterminismTests
 	void AddRandomizedVolumeInputData(FTestData& TestData, const FName& PinName)
 	{
 		AddVolumeInputData(TestData.InputData,
-			TestData.RandomStream.VRand() * Defaults::MediumDistance,
+			Defaults::MediumDistance * TestData.RandomStream.VRand(),
 			Defaults::MediumVector + TestData.RandomStream.VRand() * 0.5f * Defaults::MediumDistance,
 			Defaults::SmallVector + TestData.RandomStream.VRand() * 0.5f * Defaults::SmallDistance,
 			PinName);
@@ -244,7 +242,7 @@ namespace PCGDeterminismTests
 	{
 		check(TestData.TestActor);
 		UPrimitiveComponent* TestPrimitiveComponent = Cast<UPrimitiveComponent>(TestData.TestActor->GetComponentByClass(UPrimitiveComponent::StaticClass()));
-		
+
 		if (TestPrimitiveComponent == nullptr)
 		{
 			// TODO: If it reaches here, this will break, as it has no bounds. Please suggest the best way to give it some bounds?
@@ -599,7 +597,7 @@ namespace PCGDeterminismTests
 		{
 		}
 
-		return PCGDeterminismTests::DataCollectionsAreIdentical(FirstContext->OutputData, SecondContext->OutputData);
+		return DataCollectionsAreIdentical(FirstContext->OutputData, SecondContext->OutputData);
 	}
 
 	bool ExecutionIsDeterministicSameData(FTestData& TestData, const UPCGNode* PCGNode)
diff --git a/Engine/Plugins/Experimental/PCG/Source/PCG/Private/Tests/Determinism/PCGIntersectionDeterminismTest.cpp b/Engine/Plugins/Experimental/PCG/Source/PCG/Private/Tests/Determinism/PCGIntersectionDeterminismTest.cpp
new file mode 100644
--- /dev/null
+++ b/Engine/Plugins/Experimental/PCG/Source/PCG/Private/Tests/Determinism/PCGIntersectionDeterminismTest.cpp
@@ -0,0 +1,91 @@
+// Copyright Epic Games, Inc. All Rights Reserved.
+
+#if WITH_EDITOR
+
+#include "Elements/PCGIntersectionElement.h"
+#include "Tests/Determinism/PCGDeterminismTestsCommon.h"
+
+IMPLEMENT_CUSTOM_SIMPLE_AUTOMATION_TEST(FPCGIntersectionDeterminismSingleSameDataTest, FPCGTestBaseClass, "pcg.tests.Intersection.Determinism.SingleSameData", PCGTestsCommon::TestFlags)
+IMPLEMENT_CUSTOM_SIMPLE_AUTOMATION_TEST(FPCGIntersectionDeterminismSingleIdenticalDataTest, FPCGTestBaseClass, "pcg.tests.Intersection.Determinism.SingleMultipleData", PCGTestsCommon::TestFlags)
+IMPLEMENT_CUSTOM_SIMPLE_AUTOMATION_TEST(FPCGIntersectionDeterminismMultipleSameDataTest, FPCGTestBaseClass, "pcg.tests.Intersection.Determinism.MultipleSameData", PCGTestsCommon::TestFlags)
+IMPLEMENT_CUSTOM_SIMPLE_AUTOMATION_TEST(FPCGIntersectionDeterminismMultipleIdenticalDataTest, FPCGTestBaseClass, "pcg.tests.Intersection.Determinism.MultipleIdenticalData", PCGTestsCommon::TestFlags)
+IMPLEMENT_CUSTOM_SIMPLE_AUTOMATION_TEST(FPCGIntersectionDeterminismOrderIndependenceTest, FPCGTestBaseClass, "pcg.tests.Intersection.Determinism.OrderIndependence", PCGTestsCommon::TestFlags)
+
+namespace
+{
+	void IntersectionTestBase(PCGDeterminismTests::FTestData& TestData)
+	{
+		PCGDeterminismTests::GenerateSettings<UPCGIntersectionSettings>(TestData);
+		// Source Volumes
+		PCGDeterminismTests::AddVolumeInputData(TestData.InputData, PCGDeterminismTests::Defaults::SmallVector, PCGDeterminismTests::Defaults::MediumVector, PCGDeterminismTests::Defaults::MediumVector);
+		PCGDeterminismTests::AddVolumeInputData(TestData.InputData, PCGDeterminismTests::Defaults::SmallVector * -1.f, PCGDeterminismTests::Defaults::MediumVector, PCGDeterminismTests::Defaults::MediumVector);
+	}
+
+	void IntersectionTestMultiple(PCGDeterminismTests::FTestData& TestData)
+	{
+		IntersectionTestBase(TestData);
+
+		// Randomized Sources
+		AddRandomizedVolumeInputData(TestData);
+	}
+}
+
+bool FPCGIntersectionDeterminismSingleSameDataTest::RunTest(const FString& Parameters)
+{
+	// Test single same data
+	PCGDeterminismTests::FTestData TestData(PCGDeterminismTests::Defaults::Seed);
+
+	IntersectionTestBase(TestData);
+
+	return TestTrue(TEXT("Same single input and settings, same output"), ExecutionIsDeterministicSameData(TestData));
+}
+
+bool FPCGIntersectionDeterminismSingleIdenticalDataTest::RunTest(const FString& Parameters)
+{
+	// Test single identical data
+	PCGDeterminismTests::FTestData FirstTestData(PCGDeterminismTests::Defaults::Seed);
+	PCGDeterminismTests::FTestData SecondTestData(PCGDeterminismTests::Defaults::Seed);
+
+	IntersectionTestBase(FirstTestData);
+	IntersectionTestBase(SecondTestData);
+
+	return TestTrue(TEXT("Identical single input and settings, same output"), ExecutionIsDeterministic(FirstTestData, SecondTestData));
+}
+
+bool FPCGIntersectionDeterminismMultipleSameDataTest::RunTest(const FString& Parameters)
+{
+	// Test multiple same data
+	PCGDeterminismTests::FTestData TestData(PCGDeterminismTests::Defaults::Seed);
+
+	IntersectionTestMultiple(TestData);
+
+	return TestTrue(TEXT("Identical multiple input, same output"), ExecutionIsDeterministicSameData(TestData));
+}
+
+bool FPCGIntersectionDeterminismMultipleIdenticalDataTest::RunTest(const FString& Parameters)
+{
+	// Test multiple identical data
+	PCGDeterminismTests::FTestData FirstTestData(PCGDeterminismTests::Defaults::Seed);
+	PCGDeterminismTests::FTestData SecondTestData(PCGDeterminismTests::Defaults::Seed);
+
+	IntersectionTestMultiple(FirstTestData);
+	IntersectionTestMultiple(SecondTestData);
+
+	return TestTrue(TEXT("Identical single input and settings, same output"), ExecutionIsDeterministic(FirstTestData, SecondTestData));
+}
+
+bool FPCGIntersectionDeterminismOrderIndependenceTest::RunTest(const FString& Parameters)
+{
+	// Test multiple identical, shuffled data
+	PCGDeterminismTests::FTestData FirstTestData(PCGDeterminismTests::Defaults::Seed);
+	PCGDeterminismTests::FTestData SecondTestData(PCGDeterminismTests::Defaults::Seed);
+
+	IntersectionTestMultiple(FirstTestData);
+	IntersectionTestMultiple(SecondTestData);
+
+	ShuffleInputOrder(SecondTestData);
+
+	return TestTrue(TEXT("Shuffled input order, same output"), ExecutionIsDeterministic(FirstTestData, SecondTestData));
+}
+
+#endif
\ No newline at end of file
diff --git a/Engine/Plugins/Experimental/PCG/Source/PCG/Public/Tests/Determinism/PCGDeterminismTestsCommon.h b/Engine/Plugins/Experimental/PCG/Source/PCG/Public/Tests/Determinism/PCGDeterminismTestsCommon.h
--- a/Engine/Plugins/Experimental/PCG/Source/PCG/Public/Tests/Determinism/PCGDeterminismTestsCommon.h
+++ b/Engine/Plugins/Experimental/PCG/Source/PCG/Public/Tests/Determinism/PCGDeterminismTestsCommon.h
@@ -104,7 +104,6 @@ namespace PCGDeterminismTests
 	bool PolyLineDataIsIdentical(const UPCGData* FirstData, const UPCGData* SecondData);
 	/** Validates whether two PrimitiveData objects are identical */
 	bool PrimitiveDataIsIdentical(const UPCGData* FirstData, const UPCGData* SecondData);
-	
 	/** Validates whether two SpatialData objects are identical via Point Sampling */
 	bool SampledSpatialDataIsIdentical(const UPCGSpatialData* FirstSpatialData, const UPCGSpatialData* SecondSpatialData);
 
@@ -132,7 +131,7 @@ namespace PCGDeterminismTests
 
 	/** Generates settings based upon a UPCGSettings subclass */
 	template<typename SettingsType>
-	void GenerateSettings(FTestData& TestData, TFunction<void(FTestData& TestData)> ExtraSettingsDelegate = nullptr)
+	void GenerateSettings(FTestData& TestData, TFunction<void(FTestData&)> ExtraSettingsDelegate = nullptr)
 	{
 		TestData.Settings = NewObject<SettingsType>();
 		check(TestData.Settings);
