Add generic property-to-param element

diff --git a/Engine/Plugins/Experimental/PCG/Source/PCG/Private/Elements/PCGPropertyToParamData.cpp b/Engine/Plugins/Experimental/PCG/Source/PCG/Private/Elements/PCGPropertyToParamData.cpp
new file mode 100644
--- /dev/null
+++ b/Engine/Plugins/Experimental/PCG/Source/PCG/Private/Elements/PCGPropertyToParamData.cpp
@@ -0,0 +1,270 @@
+// Copyright Epic Games, Inc. All Rights Reserved.
+
+#include "Elements/PCGPropertyToParamData.h"
+
+#include "PCGCommon.h"
+#include "PCGContext.h"
+#include "PCGComponent.h"
+#include "PCGParamData.h"
+#include "Helpers/PCGBlueprintHelpers.h"
+#include "Helpers/PCGActorHelpers.h"
+#include "Metadata/PCGMetadata.h"
+#include "Metadata/PCGMetadataAttributeTraits.h"
+
+#include "EngineUtils.h"
+#include "Kismet/GameplayStatics.h"
+
+namespace PCGPropertyToParamDataHelpers
+{
+	// Need to pass a pointer of pointer to the found actor. The lambda will capture this pointer and modify its value when an actor is found.
+	TFunction<bool(AActor*)> GetFilteringFunction(const UPCGPropertyToParamDataSettings* InSettings, AActor*& InFoundActor)
+	{
+		check(InSettings);
+
+		switch (InSettings->ActorSelection)
+		{
+		case EPCGActorSelection::ByTag:
+			return [ActorSelectionTag = InSettings->ActorSelectionTag, &InFoundActor](AActor* Actor) -> bool
+				{
+					if (Actor->ActorHasTag(ActorSelectionTag))
+					{
+						InFoundActor = Actor;
+						return false;
+					}
+
+					return true;
+				};
+
+		case EPCGActorSelection::ByName:
+			return [ActorSelectionName = InSettings->ActorSelectionName, &InFoundActor](AActor* Actor) -> bool
+				{
+					if (Actor->GetName() == ActorSelectionName)
+					{
+						InFoundActor = Actor;
+						return false;
+					}
+
+					return true;
+				};
+
+		case EPCGActorSelection::ByClass:
+			return [ActorSelectionClass = InSettings->ActorSelectionClass, &InFoundActor](AActor* Actor) -> bool
+				{
+					if (Actor->IsA(ActorSelectionClass))
+					{
+						InFoundActor = Actor;
+						return false;
+					}
+
+					return true;
+				};
+
+		default:
+			break;
+		}
+
+		return [](AActor* Actor) -> bool { return false; };
+	}
+
+	AActor* FindActor(FPCGContext& InContext)
+	{
+		TRACE_CPUPROFILER_EVENT_SCOPE(PCGPropertyToParamDataHelpers::FindActor);
+
+		AActor* FoundActor = nullptr;
+
+		if (!InContext.SourceComponent)
+		{
+			return FoundActor;
+		}
+
+		UPCGComponent* OriginalComponent = UPCGBlueprintHelpers::GetOriginalComponent(InContext);
+		check(OriginalComponent);
+
+		UWorld* World = OriginalComponent->GetWorld();
+		if (!World)
+		{
+			return FoundActor;
+		}
+
+		const UPCGPropertyToParamDataSettings* Settings = InContext.GetInputSettings<UPCGPropertyToParamDataSettings>();
+		check(Settings);
+
+		// Early out if we have not the information necessary
+		if ((Settings->ActorSelection == EPCGActorSelection::ByTag && Settings->ActorSelectionTag == NAME_None) ||
+			(Settings->ActorSelection == EPCGActorSelection::ByName && Settings->ActorSelectionName.IsEmpty()) ||
+			(Settings->ActorSelection == EPCGActorSelection::ByClass && !Settings->ActorSelectionClass))
+		{
+			return FoundActor;
+		}
+
+		// We pass FoundActor ref, that will be captured by the FilteringFunction
+		// It will modify the FoundActor pointer to the found actor, if found.
+		TFunction<bool(AActor*)> FilteringFunction = PCGPropertyToParamDataHelpers::GetFilteringFunction(Settings, FoundActor);
+
+		// In case of iterating over all actors in the world, call our filtering function and get out.
+		if (Settings->ActorFilter == EPCGActorFilter::AllWorldActors)
+		{
+			UPCGActorHelpers::ForEachActorInWorld<AActor>(World, FilteringFunction);
+
+			// FoundActor is set by the FilteringFunction (captured)
+			return FoundActor;
+		}
+
+		// Otherwise, gather all the actors we need to check
+		TArray<AActor*> ActorsToCheck;
+		switch (Settings->ActorFilter)
+		{
+		case EPCGActorFilter::Self:
+			if (AActor* Owner = OriginalComponent->GetOwner())
+			{
+				ActorsToCheck.Add(Owner);
+			}
+			break;
+
+		case EPCGActorFilter::Parent:
+			if (AActor* Owner = OriginalComponent->GetOwner())
+			{
+				if (AActor* Parent = Owner->GetParentActor())
+				{
+					ActorsToCheck.Add(Parent);
+				}
+				else
+				{
+					// If there is no parent, set the owner as the parent.
+					ActorsToCheck.Add(Owner);
+				}
+			}
+			break;
+
+		case EPCGActorFilter::Root:
+		{
+			AActor* Current = OriginalComponent->GetOwner();
+			while (Current != nullptr)
+			{
+				AActor* Parent = Current->GetParentActor();
+				if (Parent == nullptr)
+				{
+					ActorsToCheck.Add(Current);
+					break;
+				}
+				Current = Parent;
+			}
+
+			break;
+		}
+
+		//case EPCGActorFilter::TrackedActors:
+			//	//TODO
+			//	break;
+
+		default:
+			break;
+		}
+
+		if (Settings->bIncludeChildren)
+		{
+			int32 InitialCount = ActorsToCheck.Num();
+			for (int32 i = 0; i < InitialCount; ++i)
+			{
+				ActorsToCheck[i]->GetAttachedActors(ActorsToCheck, /*bResetArray=*/ false, /*bRecursivelyIncludeAttachedActors=*/ true);
+			}
+		}
+
+		for (AActor* Actor : ActorsToCheck)
+		{
+			// FoundActor is set by the FilteringFunction (captured)
+			if (!FilteringFunction(Actor))
+			{
+				break;
+			}
+		}
+
+		return FoundActor;
+	}
+}
+
+TArray<FPCGPinProperties> UPCGPropertyToParamDataSettings::OutputPinProperties() const
+{
+	TArray<FPCGPinProperties> PinProperties;
+	PinProperties.Emplace(PCGPinConstants::DefaultOutputLabel, EPCGDataType::Param);
+
+	return PinProperties;
+}
+
+FPCGElementPtr UPCGPropertyToParamDataSettings::CreateElement() const
+{
+	return MakeShared<FPCGPropertyToParamDataElement>();
+}
+
+
+bool FPCGPropertyToParamDataElement::ExecuteInternal(FPCGContext* Context) const
+{
+	TRACE_CPUPROFILER_EVENT_SCOPE(FPCGPropertyToParamDataElement::Execute);
+
+	check(Context);
+
+	const UPCGPropertyToParamDataSettings* Settings = Context->GetInputSettings<UPCGPropertyToParamDataSettings>();
+	check(Settings);
+
+	// Early out if arguments are not specified
+	if (Settings->PropertyName == NAME_None || (Settings->bSelectComponent && !Settings->ComponentClass))
+	{
+		PCGE_LOG(Error, "Some parameters are missing, abort.");
+		return true;
+	}
+
+	// If we have no output connected, nothing to do
+	if (!Context->Node || !Context->Node->IsOutputPinConnected(PCGPinConstants::DefaultOutputLabel))
+	{
+		PCGE_LOG(Verbose, "Node is not connected, nothing to do");
+		return true;
+	}
+
+	TArray<FPCGTaggedData>& Outputs = Context->OutputData.TaggedData;
+
+	// First find the actor depending on the selection
+	AActor* FoundActor = PCGPropertyToParamDataHelpers::FindActor(*Context);
+
+	if (!FoundActor)
+	{
+		PCGE_LOG(Error, "No matching actor was found.");
+		return true;
+	}
+
+	// From there, we either check the actor, or the component attached to it.
+	UObject* ObjectToInspect = FoundActor;
+	if (Settings->bSelectComponent)
+	{
+		ObjectToInspect = FoundActor->GetComponentByClass(Settings->ComponentClass);
+		if (!ObjectToInspect)
+		{
+			PCGE_LOG(Error, "Component doesn't exist in the found actor.");
+			return true;
+		}
+	}
+
+	// Try to get the property
+	FProperty* Property = FindFProperty<FProperty>(ObjectToInspect->GetClass(), Settings->PropertyName);
+	if (!Property)
+	{
+		PCGE_LOG(Error, "Property doesn't exist in the found actor.");
+		return true;
+	}
+
+	// From there, we should be able to create the data.
+	UPCGParamData* ParamData = NewObject<UPCGParamData>();
+	UPCGMetadata* Metadata = ParamData->MutableMetadata();
+	check(Metadata);
+	PCGMetadataEntryKey EntryKey = Metadata->AddEntry();
+
+	if (!Metadata->SetAttributeFromProperty(Settings->OutputAttributeName, EntryKey, ObjectToInspect, Property, /*bCreate=*/ true))
+	{
+		PCGE_LOG(Error, "Error while creating an attribute. Either the property type is not supported by PCG or attribute creation failed.");
+		return true;
+	}
+
+	FPCGTaggedData& Output = Outputs.Emplace_GetRef();
+	Output.Data = ParamData;
+
+	return true;
+}
\ No newline at end of file
diff --git a/Engine/Plugins/Experimental/PCG/Source/PCG/Private/Elements/PCGSpawnActor.cpp b/Engine/Plugins/Experimental/PCG/Source/PCG/Private/Elements/PCGSpawnActor.cpp
--- a/Engine/Plugins/Experimental/PCG/Source/PCG/Private/Elements/PCGSpawnActor.cpp
+++ b/Engine/Plugins/Experimental/PCG/Source/PCG/Private/Elements/PCGSpawnActor.cpp
@@ -1,12 +1,15 @@
 // Copyright Epic Games, Inc. All Rights Reserved.
 
 #include "Elements/PCGSpawnActor.h"
-#include "Data/PCGPointData.h"
-#include "Helpers/PCGActorHelpers.h"
+
 #include "PCGComponent.h"
 #include "PCGHelpers.h"
 #include "PCGManagedResource.h"
 
+#include "Data/PCGPointData.h"
+#include "Helpers/PCGActorHelpers.h"
+#include "Helpers/PCGSettingsHelpers.h"
+
 #include "Components/InstancedStaticMeshComponent.h"
 #include "GameFramework/Actor.h"
 #include "Engine/World.h"
@@ -72,10 +75,29 @@ bool FPCGSpawnActorElement::ExecuteInternal(FPCGContext* Context) const
 	const UPCGSpawnActorSettings* Settings = Context->GetInputSettings<UPCGSpawnActorSettings>();
 	check(Settings);
 
+	// TODO: Temporary override to verify that soft path can also be overridden
+	// Should probably be handled somewhere else (helpers?).
+	TSubclassOf<AActor> TemplateActorClass = Settings->TemplateActorClass;
+
+	if (UPCGParamData* PCGParams = Context->InputData.GetParams())
+	{
+		FString TemplateActorClassPath = TemplateActorClass ? TemplateActorClass->GetPathName() : "";
+		FString OverriddenTemplateActorClassPath = PCGSettingsHelpers::GetValue(GET_MEMBER_NAME_CHECKED(UPCGSpawnActorSettings, TemplateActorClass), TemplateActorClassPath, PCGParams);
+
+		if (OverriddenTemplateActorClassPath != TemplateActorClassPath)
+		{
+			UClass* OverriddenTemplateActorClass = FSoftClassPath(OverriddenTemplateActorClassPath).ResolveClass();
+			if (OverriddenTemplateActorClass && OverriddenTemplateActorClass->IsChildOf<AActor>())
+			{
+				TemplateActorClass = OverriddenTemplateActorClass;
+			}
+		}
+	}
+
 	// Early out
-	if(!Settings->TemplateActorClass || Settings->TemplateActorClass->HasAnyClassFlags(CLASS_Abstract))
+	if(!TemplateActorClass || TemplateActorClass->HasAnyClassFlags(CLASS_Abstract))
 	{
-		PCGE_LOG(Error, "Invalid template actor class (%s)", Settings->TemplateActorClass ? *Settings->TemplateActorClass->GetFName().ToString() : TEXT("None"));
+		PCGE_LOG(Error, "Invalid template actor class (%s)", TemplateActorClass ? *TemplateActorClass->GetFName().ToString() : TEXT("None"));
 		return true;
 	}
 
@@ -104,7 +126,7 @@ bool FPCGSpawnActorElement::ExecuteInternal(FPCGContext* Context) const
 		}
 
 		const bool bHasAuthority = !Context->SourceComponent || (Context->SourceComponent->GetOwner() && Context->SourceComponent->GetOwner()->HasAuthority());
-		const bool bSpawnedActorsRequireAuthority = Settings->TemplateActorClass->GetDefaultObject<AActor>()->GetIsReplicated();
+		const bool bSpawnedActorsRequireAuthority = TemplateActorClass->GetDefaultObject<AActor>()->GetIsReplicated();
 
 		// First, create target instance transforms
 		const UPCGPointData* PointData = SpatialData->ToPointData(Context);
@@ -132,7 +154,7 @@ bool FPCGSpawnActorElement::ExecuteInternal(FPCGContext* Context) const
 			if (Settings->Option == EPCGSpawnActorOption::CollapseActors)
 			{
 				TArray<UActorComponent*> Components;
-				UPCGActorHelpers::GetActorClassDefaultComponents(Settings->TemplateActorClass, Components, UStaticMeshComponent::StaticClass());
+				UPCGActorHelpers::GetActorClassDefaultComponents(TemplateActorClass, Components, UStaticMeshComponent::StaticClass());
 				UStaticMeshComponent* FirstSMC = nullptr;
 				UStaticMesh* Mesh = nullptr;
 
@@ -182,7 +204,7 @@ bool FPCGSpawnActorElement::ExecuteInternal(FPCGContext* Context) const
 
 					for (const FPCGPoint& Point : Points)
 					{
-						AActor* GeneratedActor = TargetActor->GetWorld()->SpawnActor(Settings->TemplateActorClass, &Point.Transform, SpawnParams);
+						AActor* GeneratedActor = TargetActor->GetWorld()->SpawnActor(TemplateActorClass, &Point.Transform, SpawnParams);
 						GeneratedActor->Tags.Add(PCGHelpers::DefaultPCGActorTag);
 						GeneratedActor->AttachToActor(TargetActor, FAttachmentTransformRules::KeepWorldTransform);
 
diff --git a/Engine/Plugins/Experimental/PCG/Source/PCG/Private/Helpers/PCGActorHelpers.cpp b/Engine/Plugins/Experimental/PCG/Source/PCG/Private/Helpers/PCGActorHelpers.cpp
--- a/Engine/Plugins/Experimental/PCG/Source/PCG/Private/Helpers/PCGActorHelpers.cpp
+++ b/Engine/Plugins/Experimental/PCG/Source/PCG/Private/Helpers/PCGActorHelpers.cpp
@@ -6,6 +6,7 @@
 
 #include "PCGComponent.h"
 #include "PCGManagedResource.h"
+#include "EngineUtils.h"
 #include "Components/InstancedStaticMeshComponent.h"
 #include "Components/HierarchicalInstancedStaticMeshComponent.h"
 #include "Engine/InheritableComponentHandler.h"
@@ -292,7 +293,7 @@ void UPCGActorHelpers::GetActorClassDefaultComponents(const TSubclassOf<AActor>&
 	OutComponents = MoveTemp(ResultComponents);
 }
 
-void UPCGActorHelpers::ForEachActorInLevel(ULevel* Level, TSubclassOf<AActor> ActorClass, TFunctionRef<void(AActor*)> Callback)
+void UPCGActorHelpers::ForEachActorInLevel(ULevel* Level, TSubclassOf<AActor> ActorClass, TFunctionRef<bool(AActor*)> Callback)
 {
 	if (!Level)
 	{
@@ -303,7 +304,29 @@ void UPCGActorHelpers::ForEachActorInLevel(ULevel* Level, TSubclassOf<AActor> Ac
 	{
 		if (Actor && Actor->IsA(ActorClass))
 		{
-			Callback(Actor);
+			if (!Callback(Actor))
+			{
+				break;
+			}
+		}
+	}
+}
+
+void UPCGActorHelpers::ForEachActorInWorld(UWorld* World, TSubclassOf<AActor> ActorClass, TFunctionRef<bool(AActor*)> Callback)
+{
+	if (!World)
+	{
+		return;
+	}
+
+	for (TActorIterator<AActor> It(World, ActorClass); It; ++It)
+	{
+		if (AActor* Actor = *It)
+		{
+			if (!Callback(Actor))
+			{
+				break;
+			}
 		}
 	}
 }
diff --git a/Engine/Plugins/Experimental/PCG/Source/PCG/Private/Metadata/PCGMetadata.cpp b/Engine/Plugins/Experimental/PCG/Source/PCG/Private/Metadata/PCGMetadata.cpp
--- a/Engine/Plugins/Experimental/PCG/Source/PCG/Private/Metadata/PCGMetadata.cpp
+++ b/Engine/Plugins/Experimental/PCG/Source/PCG/Private/Metadata/PCGMetadata.cpp
@@ -5,6 +5,7 @@
 #include "PCGModule.h"
 
 #include "Misc/ScopeRWLock.h"
+#include "Helpers/PCGSettingsHelpers.h"
 
 void UPCGMetadata::Serialize(FArchive& Ar)
 {
@@ -336,6 +337,51 @@ void UPCGMetadata::CreateStringAttribute(FName AttributeName, FString DefaultVal
 	CreateAttribute<FString>(AttributeName, DefaultValue, bAllowsInterpolation, bOverrideParent);
 }
 
+bool UPCGMetadata::SetAttributeFromProperty(FName AttributeName, PCGMetadataEntryKey& EntryKey, const UObject* Object, const FProperty* InProperty, bool bCreate)
+{
+	if (!InProperty || !Object)
+	{
+		return false;
+	}
+
+	// Check if an attribute already exists or not if we ask to create a new one
+	if (!bCreate && !HasAttribute(AttributeName))
+	{
+		return false;
+	}
+
+	auto CreateAttributeAndSet = [&AttributeName, this, bCreate, &EntryKey](auto&& PropertyValue) -> bool
+	{
+		using PropertyType = std::remove_const_t<std::remove_reference_t<decltype(PropertyValue)>>;
+
+		FPCGMetadataAttributeBase* BaseAttribute = GetMutableAttribute(AttributeName);
+
+		if (!BaseAttribute && bCreate)
+		{
+			// Interpolation is disabled and no parent override.
+			BaseAttribute = CreateAttribute<PropertyType>(AttributeName, PropertyValue, false, false);
+		}
+
+		if (!BaseAttribute)
+		{
+			return false;
+		}
+
+		// Check that the property match the attribute type!
+		if (PCG::Private::MetadataTypes<PropertyType>::Id != BaseAttribute->GetTypeId())
+		{
+			return false;
+		}
+
+		FPCGMetadataAttribute<PropertyType>* Attribute = static_cast<FPCGMetadataAttribute<PropertyType>*>(BaseAttribute);
+		Attribute->SetValue(EntryKey, PropertyValue);
+
+		return true;
+	};
+
+	return PCGSettingsHelpers::GetPropertyValueWithCallback(Object, InProperty, CreateAttributeAndSet);
+}
+
 void UPCGMetadata::CopyExistingAttribute(FName AttributeToCopy, FName NewAttributeName, bool bKeepParent)
 {
 	CopyAttribute(AttributeToCopy, NewAttributeName, bKeepParent, /*bCopyEntries=*/true, /*bCopyValues=*/true);
diff --git a/Engine/Plugins/Experimental/PCG/Source/PCG/Private/Metadata/PCGMetadataAccessor.cpp b/Engine/Plugins/Experimental/PCG/Source/PCG/Private/Metadata/PCGMetadataAccessor.cpp
--- a/Engine/Plugins/Experimental/PCG/Source/PCG/Private/Metadata/PCGMetadataAccessor.cpp
+++ b/Engine/Plugins/Experimental/PCG/Source/PCG/Private/Metadata/PCGMetadataAccessor.cpp
@@ -169,6 +169,22 @@ void UPCGMetadataAccessorHelpers::SetStringAttributeByMetadataKey(int64& Key, UP
 	SetAttribute(Key, Metadata, AttributeName, Value);
 }
 
+bool UPCGMetadataAccessorHelpers::SetAttributeFromPropertyByMetadataKey(int64& Key, UPCGMetadata* Metadata, FName AttributeName, const UObject* Object, FName PropertyName)
+{
+	if (!Object)
+	{
+		return false;
+	}
+
+	const FProperty* Property = FindFProperty<FProperty>(Object->GetClass(), PropertyName);
+	if (!Property)
+	{
+		return false;
+	}
+
+	return Metadata->SetAttributeFromProperty(AttributeName, Key, Object, Property, /*bCreate=*/ true);
+}
+
 /** Point-based implementations */
 void UPCGMetadataAccessorHelpers::CopyPoint(const FPCGPoint& InPoint, FPCGPoint& OutPoint, bool bCopyMetadata, const UPCGMetadata* InMetadata, UPCGMetadata* OutMetadata)
 {
diff --git a/Engine/Plugins/Experimental/PCG/Source/PCG/Private/PCGSubsystem.cpp b/Engine/Plugins/Experimental/PCG/Source/PCG/Private/PCGSubsystem.cpp
--- a/Engine/Plugins/Experimental/PCG/Source/PCG/Private/PCGSubsystem.cpp
+++ b/Engine/Plugins/Experimental/PCG/Source/PCG/Private/PCGSubsystem.cpp
@@ -602,7 +602,7 @@ void UPCGSubsystem::DeletePartitionActors()
 		return;
 	}
 
-	auto GatherAndDestroyLoadedActors = [&PackagesToCleanup, &PackagesToDeleteFromSCC, World, &ActorsNotSafeToBeDeleted](AActor* Actor)
+	auto GatherAndDestroyLoadedActors = [&PackagesToCleanup, &PackagesToDeleteFromSCC, World, &ActorsNotSafeToBeDeleted](AActor* Actor) -> bool
 	{
 		// Make sure that this actor was not flagged to not be deleted, or not safe for deletion
 		TObjectPtr<APCGPartitionActor> PartitionActor = CastChecked<APCGPartitionActor>(Actor);
@@ -615,7 +615,10 @@ void UPCGSubsystem::DeletePartitionActors()
 			// Also reset the last generated bounds to indicate to the component to re-create its PartitionActors when it generates.
 			for (TObjectPtr<UPCGComponent> PCGComponent : PartitionActor->GetAllOriginalPCGComponents())
 			{
-				PCGComponent->ResetLastGeneratedBounds();
+				if (PCGComponent)
+				{
+					PCGComponent->ResetLastGeneratedBounds();
+				}
 			}
 
 			if (UPackage* ExternalPackage = PartitionActor->GetExternalPackage())
@@ -625,6 +628,8 @@ void UPCGSubsystem::DeletePartitionActors()
 
 			World->DestroyActor(PartitionActor);
 		}
+
+		return true;
 	};
 
 	auto GatherAndDestroyActors = [&PackagesToCleanup, &PackagesToDeleteFromSCC, World, &ActorsNotSafeToBeDeleted, &GatherAndDestroyLoadedActors](const FWorldPartitionActorDesc* ActorDesc) {
diff --git a/Engine/Plugins/Experimental/PCG/Source/PCG/Private/PCGWorldActor.cpp b/Engine/Plugins/Experimental/PCG/Source/PCG/Private/PCGWorldActor.cpp
--- a/Engine/Plugins/Experimental/PCG/Source/PCG/Private/PCGWorldActor.cpp
+++ b/Engine/Plugins/Experimental/PCG/Source/PCG/Private/PCGWorldActor.cpp
@@ -85,14 +85,14 @@ void APCGWorldActor::OnPartitionGridSizeChanged()
 	
 	bool bAllSafeToDelete = true;
 
-	auto AddPartitionComponentAndCheckIfSafeToDelete = [&AllPartitionedComponents, &bAllSafeToDelete](AActor* Actor)
+	auto AddPartitionComponentAndCheckIfSafeToDelete = [&AllPartitionedComponents, &bAllSafeToDelete](AActor* Actor) -> bool
 	{
 		TObjectPtr<APCGPartitionActor> PartitionActor = CastChecked<APCGPartitionActor>(Actor);
 
 		if (!PartitionActor->IsSafeForDeletion())
 		{
 			bAllSafeToDelete = false;
-			return;
+			return true;
 		}
 
 		for (UPCGComponent* PCGComponent : PartitionActor->GetAllOriginalPCGComponents())
@@ -102,6 +102,8 @@ void APCGWorldActor::OnPartitionGridSizeChanged()
 				AllPartitionedComponents.Add(PCGComponent);
 			}
 		}
+
+		return true;
 	};
 
 	UPCGActorHelpers::ForEachActorInLevel<APCGPartitionActor>(Level, AddPartitionComponentAndCheckIfSafeToDelete);
diff --git a/Engine/Plugins/Experimental/PCG/Source/PCG/Public/Elements/PCGPropertyToParamData.h b/Engine/Plugins/Experimental/PCG/Source/PCG/Public/Elements/PCGPropertyToParamData.h
new file mode 100644
--- /dev/null
+++ b/Engine/Plugins/Experimental/PCG/Source/PCG/Public/Elements/PCGPropertyToParamData.h
@@ -0,0 +1,95 @@
+// Copyright Epic Games, Inc. All Rights Reserved.
+
+#pragma once
+
+#include "PCGElement.h"
+#include "PCGSettings.h"
+
+#include "Templates/SubclassOf.h"
+
+#include "PCGPropertyToParamData.generated.h"
+
+class AActor;
+class UActorComponent;
+
+UENUM()
+enum class EPCGActorSelection : uint8
+{
+	ByTag,
+	ByName,
+	ByClass
+};
+
+UENUM()
+enum class EPCGActorFilter : uint8
+{
+	Self,
+	Parent,
+	Root,
+	AllWorldActors
+	// TODO
+	// TrackedActors
+};
+
+UCLASS(BlueprintType, ClassGroup = (Procedural))
+class PCG_API UPCGPropertyToParamDataSettings : public UPCGSettings
+{
+	GENERATED_BODY()
+
+public:
+	//~Begin UPCGSettings interface
+#if WITH_EDITOR
+	virtual FName GetDefaultNodeName() const override { return FName(TEXT("PropertyToParamDataNode")); }
+	virtual EPCGSettingsType GetType() const override { return EPCGSettingsType::Metadata; }
+#endif
+
+	virtual TArray<FPCGPinProperties> OutputPinProperties() const override;
+	virtual TArray<FPCGPinProperties> InputPinProperties() const override { return TArray<FPCGPinProperties>(); }
+
+protected:
+	virtual FPCGElementPtr CreateElement() const override;
+	//~End UPCGSettings interface
+
+public:
+	UPROPERTY(BlueprintReadWrite, EditAnywhere, Category = Settings)
+	EPCGActorSelection ActorSelection;
+
+	UPROPERTY(BlueprintReadWrite, EditAnywhere, Category = Settings, meta = (EditCondition = "ActorSelection==EPCGActorSelection::ByTag", EditConditionHides))
+	FName ActorSelectionTag;
+
+	UPROPERTY(BlueprintReadWrite, EditAnywhere, Category = Settings, meta = (EditCondition = "ActorSelection==EPCGActorSelection::ByName", EditConditionHides))
+	FString ActorSelectionName;
+
+	UPROPERTY(BlueprintReadWrite, EditAnywhere, Category = Settings, meta = (EditCondition = "ActorSelection==EPCGActorSelection::ByClass", EditConditionHides))
+	TSubclassOf<AActor> ActorSelectionClass;
+
+	UPROPERTY(BlueprintReadWrite, EditAnywhere, Category = Settings)
+	EPCGActorFilter ActorFilter = EPCGActorFilter::Self;
+
+	UPROPERTY(BlueprintReadWrite, EditAnywhere, Category = Settings, meta = (EditCondition = "ActorFilter!=EPCGActorFilter::AllWorldActors", EditConditionHides))
+	bool bIncludeChildren = false;
+
+	UPROPERTY(BlueprintReadWrite, EditAnywhere, Category = Settings)
+	bool bSelectComponent = false;
+
+	UPROPERTY(BlueprintReadWrite, EditAnywhere, Category = Settings, meta = (EditCondition = "bSelectComponent", EditConditionHides))
+	TSubclassOf<UActorComponent> ComponentClass;
+
+	UPROPERTY(BlueprintReadWrite, EditAnywhere, Category = Settings)
+	FName PropertyName = NAME_None;
+
+	UPROPERTY(BlueprintReadWrite, EditAnywhere, Category = Settings)
+	FName OutputAttributeName = NAME_None;
+
+	UPROPERTY(BlueprintReadWrite, EditAnywhere, Category = Settings)
+	bool bAlwaysRequeryActors = true;
+};
+
+class FPCGPropertyToParamDataElement : public FSimplePCGElement
+{
+public:
+	virtual bool CanExecuteOnlyOnMainThread(const UPCGSettings* InSettings) const override { return true; }
+	virtual bool IsCacheable(const UPCGSettings* InSettings) const override { return !CastChecked<UPCGPropertyToParamDataSettings>(InSettings)->bAlwaysRequeryActors; }
+protected:
+	virtual bool ExecuteInternal(FPCGContext* Context) const;
+};
\ No newline at end of file
diff --git a/Engine/Plugins/Experimental/PCG/Source/PCG/Public/Helpers/PCGActorHelpers.h b/Engine/Plugins/Experimental/PCG/Source/PCG/Public/Helpers/PCGActorHelpers.h
--- a/Engine/Plugins/Experimental/PCG/Source/PCG/Public/Helpers/PCGActorHelpers.h
+++ b/Engine/Plugins/Experimental/PCG/Source/PCG/Public/Helpers/PCGActorHelpers.h
@@ -47,12 +47,32 @@ public:
 	static void GetActorClassDefaultComponents(const TSubclassOf<AActor>& ActorClass, TArray<UActorComponent*>& OutComponents, const TSubclassOf<UActorComponent>& InComponentClass = TSubclassOf<UActorComponent>());
 
 	template <typename T, typename = typename std::enable_if_t<std::is_base_of_v<AActor, T>>>
-	static void ForEachActorInLevel(ULevel* Level, TFunctionRef<void(AActor*)> Callback)
+	inline static void ForEachActorInLevel(ULevel* Level, TFunctionRef<bool(AActor*)> Callback)
 	{
 		return ForEachActorInLevel(Level, T::StaticClass(), Callback);
 	}
 
-	static void ForEachActorInLevel(ULevel* Level, TSubclassOf<AActor> ActorClass, TFunctionRef<void(AActor*)> Callback);
+	/**
+	* Iterate over all actors in the level, from the given class, and pass them to a callback
+	* @param Level The level
+	* @param ActorClass class of AActor to pass to the callback
+	* @param Callback Function to call with the found actor. Needs to return a bool, to indicate if it needs to continue (true = yes)
+	*/
+	static void ForEachActorInLevel(ULevel* Level, TSubclassOf<AActor> ActorClass, TFunctionRef<bool(AActor*)> Callback);
+
+	template <typename T, typename = typename std::enable_if_t<std::is_base_of_v<AActor, T>>>
+	inline static void ForEachActorInWorld(UWorld* World, TFunctionRef<bool(AActor*)> Callback)
+	{
+		return ForEachActorInWorld(World, T::StaticClass(), Callback);
+	}
+
+	/**
+	* Iterate over all actors in the world, from the given class, and pass them to a callback
+	* @param World The world
+	* @param ActorClass class of AActor to pass to the callback
+	* @param Callback Function to call with the found actor. Needs to return a bool, to indicate if it needs to continue (true = yes)
+	*/
+	static void ForEachActorInWorld(UWorld* World, TSubclassOf<AActor> ActorClass, TFunctionRef<bool(AActor*)> Callback);
 
 	/**
 	* Spawn a new actor of type T and attach it to the parent (if not null)
@@ -62,7 +82,7 @@ public:
 	* @param Parent Optional parent to attach to.
 	*/
 	template <typename T = AActor, typename = typename std::enable_if_t<std::is_base_of_v<AActor, T>>>
-	static AActor* SpawnDefaultActor(UWorld* World, FName BaseName, const FTransform& Transform, AActor* Parent = nullptr)
+	inline static AActor* SpawnDefaultActor(UWorld* World, FName BaseName, const FTransform& Transform, AActor* Parent = nullptr)
 	{
 		return SpawnDefaultActor(World, T::StaticClass(), BaseName, Transform, Parent);
 	}
diff --git a/Engine/Plugins/Experimental/PCG/Source/PCG/Public/Helpers/PCGSettingsHelpers.h b/Engine/Plugins/Experimental/PCG/Source/PCG/Public/Helpers/PCGSettingsHelpers.h
--- a/Engine/Plugins/Experimental/PCG/Source/PCG/Public/Helpers/PCGSettingsHelpers.h
+++ b/Engine/Plugins/Experimental/PCG/Source/PCG/Public/Helpers/PCGSettingsHelpers.h
@@ -88,6 +88,97 @@ namespace PCGSettingsHelpers
 
 	/** Sets data from the params to a given property, matched on a name basis */
 	void SetValue(UPCGParamData* Params, UObject* Object, FProperty* Property);
+
+	/**
+	* Validate that the InProperty is a PCGData supported type and call InFunc with the value of this property (with the right type).
+	* The function returns the result of InFunc (or default value of the return type of InFunc if it failed).
+	*/
+	template <typename Func>
+	inline decltype(auto) GetPropertyValueWithCallback(const UObject* InObject, const FProperty* InProperty, Func InFunc)
+	{
+		// Double property is supported by PCG, we use a dummy double to deduce the return type of InFunc.
+		using ReturnType = decltype(InFunc(0.0));
+
+		if (!InObject || !InProperty)
+		{
+			return ReturnType();
+		}
+
+		const void* PropertyAddressData = InProperty->ContainerPtrToValuePtr<void>(InObject);
+
+		if (const FNumericProperty* NumericProperty = CastField<FNumericProperty>(InProperty))
+		{
+			if (NumericProperty->IsFloatingPoint())
+			{
+				return InFunc(NumericProperty->GetFloatingPointPropertyValue(PropertyAddressData));
+			}
+			else if (NumericProperty->IsInteger())
+			{
+				return InFunc(NumericProperty->GetSignedIntPropertyValue(PropertyAddressData));
+			}
+		}
+		else if (const FBoolProperty* BoolProperty = CastField<FBoolProperty>(InProperty))
+		{
+			return InFunc(BoolProperty->GetPropertyValue(PropertyAddressData));
+		}
+		else if (const FStrProperty* StringProperty = CastField<FStrProperty>(InProperty))
+		{
+			return InFunc(StringProperty->GetPropertyValue(PropertyAddressData));
+		}
+		else if (const FNameProperty* NameProperty = CastField<FNameProperty>(InProperty))
+		{
+			return InFunc(NameProperty->GetPropertyValue(PropertyAddressData));
+		}
+		else if (const FEnumProperty* EnumProperty = CastField<FEnumProperty>(InProperty))
+		{
+			return InFunc(EnumProperty->GetUnderlyingProperty()->GetSignedIntPropertyValue(PropertyAddressData));
+		}
+		else if (const FStructProperty* StructProperty = CastField<FStructProperty>(InProperty))
+		{
+			if (StructProperty->Struct == TBaseStructure<FVector>::Get())
+			{
+				return InFunc(*reinterpret_cast<const FVector*>(PropertyAddressData));
+			}
+			else if (StructProperty->Struct == TBaseStructure<FVector4>::Get())
+			{
+				return InFunc(*reinterpret_cast<const FVector4*>(PropertyAddressData));
+			}
+			else if (StructProperty->Struct == TBaseStructure<FQuat>::Get())
+			{
+				return InFunc(*reinterpret_cast<const FQuat*>(PropertyAddressData));
+			}
+			else if (StructProperty->Struct == TBaseStructure<FTransform>::Get())
+			{
+				return InFunc(*reinterpret_cast<const FTransform*>(PropertyAddressData));
+			}
+			else if (StructProperty->Struct == TBaseStructure<FRotator>::Get())
+			{
+				return InFunc(*reinterpret_cast<const FRotator*>(PropertyAddressData));
+			}
+			else if (StructProperty->Struct == TBaseStructure<FSoftObjectPath>::Get())
+			{
+				// Soft object path are transformed to strings
+				return InFunc(reinterpret_cast<const FSoftObjectPath*>(PropertyAddressData)->ToString());
+			}
+			else if (StructProperty->Struct == TBaseStructure<FSoftClassPath>::Get())
+			{
+				// Soft class path are transformed to strings
+				return InFunc(reinterpret_cast<const FSoftClassPath*>(PropertyAddressData)->ToString());
+			}
+			//else if (StructProperty->Struct == TBaseStructure<FColor>::Get())
+			//else if (StructProperty->Struct == TBaseStructure<FLinearColor>::Get())
+		}
+		else if (const FObjectPropertyBase* ObjectProperty = CastField<FObjectPropertyBase>(InProperty))
+		{
+			if (const UObject* Object = ObjectProperty->GetObjectPropertyValue(PropertyAddressData))
+			{
+				// Object are transformed into their soft path name (as a string attribute)
+				return InFunc(Object->GetPathName());
+			}
+		}
+
+		return ReturnType();
+	}
 }
 
 #define PCG_GET_OVERRIDEN_VALUE(Settings, Variable, Params) PCGSettingsHelpers::GetValue(GET_MEMBER_NAME_CHECKED(TRemovePointer<TRemoveConst<decltype(Settings)>::Type>::Type, Variable), (Settings)->Variable, Params);
\ No newline at end of file
diff --git a/Engine/Plugins/Experimental/PCG/Source/PCG/Public/Metadata/PCGMetadata.h b/Engine/Plugins/Experimental/PCG/Source/PCG/Public/Metadata/PCGMetadata.h
--- a/Engine/Plugins/Experimental/PCG/Source/PCG/Public/Metadata/PCGMetadata.h
+++ b/Engine/Plugins/Experimental/PCG/Source/PCG/Public/Metadata/PCGMetadata.h
@@ -86,6 +86,16 @@ public:
 	UFUNCTION(BlueprintCallable, Category = "PCG|Metadata")
 	void CreateStringAttribute(FName AttributeName, FString DefaultValue, bool bAllowsInterpolation, bool bOverrideParent = true);
 
+	/** Set an attribute given a property and its value.
+	* @param AttributeName: Target attribute to set the property's value to
+	* @param EntryKey: Metadata entry key to set the value to
+	* @param Object: Object to get the property value from
+	* @param Property: The property to set from
+	* @param bCreate: If true and the attribute doesn't exists, it will create an attribute based on the property type
+	* @returns true if the attribute creation (if required) and the value set succeeded
+	*/
+	bool SetAttributeFromProperty(FName AttributeName, PCGMetadataEntryKey& EntryKey, const UObject* Object, const FProperty* Property, bool bCreate);
+
 	/** Get attributes */
 	FPCGMetadataAttributeBase* GetMutableAttribute(FName AttributeName);
 	const FPCGMetadataAttributeBase* GetConstAttribute(FName AttributeName) const;
diff --git a/Engine/Plugins/Experimental/PCG/Source/PCG/Public/Metadata/PCGMetadataAccessor.h b/Engine/Plugins/Experimental/PCG/Source/PCG/Public/Metadata/PCGMetadataAccessor.h
--- a/Engine/Plugins/Experimental/PCG/Source/PCG/Public/Metadata/PCGMetadataAccessor.h
+++ b/Engine/Plugins/Experimental/PCG/Source/PCG/Public/Metadata/PCGMetadataAccessor.h
@@ -66,6 +66,9 @@ public:
 	UFUNCTION(BlueprintCallable, Category = "PCG|Metadata")
 	static void SetStringAttributeByMetadataKey(UPARAM(ref) int64& Key, UPCGMetadata* Metadata, FName AttributeName, const FString& Value);
 
+	UFUNCTION(BlueprintCallable, Category = "PCG|Metadata")
+	static bool SetAttributeFromPropertyByMetadataKey(UPARAM(ref) int64& Key, UPCGMetadata* Metadata, FName AttributeName, const UObject* Object, FName PropertyName);
+
 	UFUNCTION(BlueprintCallable, Category = "PCG|Metadata")
 	static bool HasAttributeSetByMetadataKey(int64 Key, const UPCGMetadata* Metadata, FName AttributeName);
 
diff --git a/Engine/Plugins/Experimental/PCG/Source/PCG/Public/Metadata/PCGMetadataAttributeTraits.h b/Engine/Plugins/Experimental/PCG/Source/PCG/Public/Metadata/PCGMetadataAttributeTraits.h
--- a/Engine/Plugins/Experimental/PCG/Source/PCG/Public/Metadata/PCGMetadataAttributeTraits.h
+++ b/Engine/Plugins/Experimental/PCG/Source/PCG/Public/Metadata/PCGMetadataAttributeTraits.h
@@ -20,6 +20,7 @@ enum class EPCGMetadataTypes : uint8
 	String,
 	Boolean,
 	Rotator,
+	Name,
 	Unknown = 255
 };
 
@@ -46,6 +47,7 @@ namespace PCG
 		PCGMetadataGenerateDataTypes(FString, String);
 		PCGMetadataGenerateDataTypes(bool, Boolean);
 		PCGMetadataGenerateDataTypes(FRotator, Rotator);
+		PCGMetadataGenerateDataTypes(FName, Name);
 
 #undef PCGMetadataGenerateDataTypes
 
@@ -383,5 +385,20 @@ namespace PCG
 				return A == B;
 			}
 		};
+
+		template<>
+		struct MetadataTraits<FName>
+		{
+			enum { CompressData = true };
+			enum { CanMinMax = false };
+			enum { CanSubAdd = false };
+			enum { CanMulDiv = false };
+			enum { CanInterpolate = false };
+
+			static bool Equal(const FName& A, const FName& B)
+			{
+				return A == B;
+			}
+		};
 	}
 }
\ No newline at end of file
