Add PCGVolume details panel

diff --git a/Engine/Plugins/Experimental/PCG/Source/PCG/Private/PCGManagedResource.cpp b/Engine/Plugins/Experimental/PCG/Source/PCG/Private/PCGManagedResource.cpp
--- a/Engine/Plugins/Experimental/PCG/Source/PCG/Private/PCGManagedResource.cpp
+++ b/Engine/Plugins/Experimental/PCG/Source/PCG/Private/PCGManagedResource.cpp
@@ -172,6 +172,9 @@ bool UPCGManagedComponent::MoveResourceToNewActor(AActor* NewActor)
 	bool bDetached = false;
 	bool bAttached = false;
 
+	// Need to change owner first to avoid that the PCG Component will react to this component changes.
+	GeneratedComponent->Rename(nullptr, NewActor);
+
 	// Check if it is a scene component, and if so, use its method to attach/detach to root component
 	if (TObjectPtr<USceneComponent> GeneratedSceneComponent = Cast<USceneComponent>(GeneratedComponent.Get()))
 	{
@@ -191,8 +194,6 @@ bool UPCGManagedComponent::MoveResourceToNewActor(AActor* NewActor)
 		NewActor->AddInstanceComponent(GeneratedComponent.Get());
 	}
 
-	GeneratedComponent->Rename(nullptr, NewActor);
-
 	GeneratedComponent.Reset();
 
 	return true;
diff --git a/Engine/Plugins/Experimental/PCG/Source/PCGEditor/Private/PCGComponentDetails.cpp b/Engine/Plugins/Experimental/PCG/Source/PCGEditor/Private/PCGComponentDetails.cpp
--- a/Engine/Plugins/Experimental/PCG/Source/PCGEditor/Private/PCGComponentDetails.cpp
+++ b/Engine/Plugins/Experimental/PCG/Source/PCGEditor/Private/PCGComponentDetails.cpp
@@ -20,15 +20,9 @@ TSharedRef<IDetailCustomization> FPCGComponentDetails::MakeInstance()
 	return MakeShareable(new FPCGComponentDetails());
 }
 
-void FPCGComponentDetails::CustomizeDetails(IDetailLayoutBuilder& DetailBuilder)
+void FPCGComponentDetails::GatherPCGComponentsFromSelection(const TArray<TWeakObjectPtr<UObject>>& InObjectSelected)
 {
-	const FName PCGCategoryName("PCG");
-	IDetailCategoryBuilder& PCGCategory = DetailBuilder.EditCategory(PCGCategoryName);
-
-	TArray<TWeakObjectPtr<UObject>> ObjectsBeingCustomized;
-	DetailBuilder.GetObjectsBeingCustomized(ObjectsBeingCustomized);
-
-	for (TWeakObjectPtr<UObject>& Object : ObjectsBeingCustomized)
+	for (const TWeakObjectPtr<UObject>& Object : InObjectSelected)
 	{
 		UPCGComponent* Component = Cast<UPCGComponent>(Object.Get());
 		if (ensure(Component))
@@ -36,16 +30,30 @@ void FPCGComponentDetails::CustomizeDetails(IDetailLayoutBuilder& DetailBuilder)
 			SelectedComponents.Add(Component);
 		}
 	}
+}
 
-	TArray<TSharedRef<IPropertyHandle>> AllProperties;
-	bool bSimpleProperties = true;
-	bool bAdvancedProperties = false;
-	// Add all properties in the category in order
-	PCGCategory.GetDefaultProperties(AllProperties, bSimpleProperties, bAdvancedProperties);
+void FPCGComponentDetails::CustomizeDetails(IDetailLayoutBuilder& DetailBuilder)
+{
+	const FName PCGCategoryName("PCG");
+	IDetailCategoryBuilder& PCGCategory = DetailBuilder.EditCategory(PCGCategoryName);
+
+	TArray<TWeakObjectPtr<UObject>> ObjectsBeingCustomized;
+	DetailBuilder.GetObjectsBeingCustomized(ObjectsBeingCustomized);
 
-	for (auto& Property : AllProperties)
+	GatherPCGComponentsFromSelection(ObjectsBeingCustomized);
+
+	if (AddDefaultProperties())
 	{
-		PCGCategory.AddProperty(Property);
+		TArray<TSharedRef<IPropertyHandle>> AllProperties;
+		bool bSimpleProperties = true;
+		bool bAdvancedProperties = false;
+		// Add all properties in the category in order
+		PCGCategory.GetDefaultProperties(AllProperties, bSimpleProperties, bAdvancedProperties);
+
+		for (auto& Property : AllProperties)
+		{
+			PCGCategory.AddProperty(Property);
+		}
 	}
 
 	FDetailWidgetRow& NewRow = PCGCategory.AddCustomRow(FText::GetEmpty());
diff --git a/Engine/Plugins/Experimental/PCG/Source/PCGEditor/Private/PCGEditorModule.cpp b/Engine/Plugins/Experimental/PCG/Source/PCGEditor/Private/PCGEditorModule.cpp
--- a/Engine/Plugins/Experimental/PCG/Source/PCGEditor/Private/PCGEditorModule.cpp
+++ b/Engine/Plugins/Experimental/PCG/Source/PCGEditor/Private/PCGEditorModule.cpp
@@ -10,6 +10,7 @@
 #include "PCGEditorSettings.h"
 #include "PCGGraphDetails.h"
 #include "PCGSubsystem.h"
+#include "PCGVolumeDetails.h"
 #include "PCGVolumeFactory.h"
 
 #include "EdGraphUtilities.h"
@@ -70,6 +71,7 @@ void FPCGEditorModule::RegisterDetailsCustomizations()
 	FPropertyEditorModule& PropertyEditor = FModuleManager::LoadModuleChecked<FPropertyEditorModule>("PropertyEditor");
 	PropertyEditor.RegisterCustomClassLayout("PCGComponent", FOnGetDetailCustomizationInstance::CreateStatic(&FPCGComponentDetails::MakeInstance));
 	PropertyEditor.RegisterCustomClassLayout("PCGGraph", FOnGetDetailCustomizationInstance::CreateStatic(&FPCGGraphDetails::MakeInstance));
+	PropertyEditor.RegisterCustomClassLayout("PCGVolume", FOnGetDetailCustomizationInstance::CreateStatic(&FPCGVolumeDetails::MakeInstance));
 }
 
 void FPCGEditorModule::UnregisterDetailsCustomizations()
@@ -79,6 +81,7 @@ void FPCGEditorModule::UnregisterDetailsCustomizations()
 		FPropertyEditorModule& PropertyModule = FModuleManager::LoadModuleChecked<FPropertyEditorModule>("PropertyEditor");
 		PropertyModule.UnregisterCustomClassLayout("PCGComponent");
 		PropertyModule.UnregisterCustomClassLayout("PCGGraph");
+		PropertyModule.UnregisterCustomClassLayout("PCGVolume");
 		PropertyModule.NotifyCustomizationModuleChanged();
 	}
 }
diff --git a/Engine/Plugins/Experimental/PCG/Source/PCGEditor/Private/PCGVolumeDetails.cpp b/Engine/Plugins/Experimental/PCG/Source/PCGEditor/Private/PCGVolumeDetails.cpp
new file mode 100644
--- /dev/null
+++ b/Engine/Plugins/Experimental/PCG/Source/PCGEditor/Private/PCGVolumeDetails.cpp
@@ -0,0 +1,30 @@
+// Copyright Epic Games, Inc. All Rights Reserved.
+
+#include "PCGVolumeDetails.h"
+
+#include "PCGComponent.h"
+#include "PCGVolume.h"
+
+TSharedRef<IDetailCustomization> FPCGVolumeDetails::MakeInstance()
+{
+	return MakeShareable(new FPCGVolumeDetails());
+}
+
+void FPCGVolumeDetails::GatherPCGComponentsFromSelection(const TArray<TWeakObjectPtr<UObject>>& InObjectSelected)
+{
+	for (const TWeakObjectPtr<UObject>& Object : InObjectSelected)
+	{
+		if (APCGVolume* Volume = Cast<APCGVolume>(Object))
+		{
+			TInlineComponentArray<UPCGComponent*, 1> Components;
+			Volume->GetComponents<UPCGComponent>(Components);
+			for (UPCGComponent* Component : Components)
+			{
+				if (Component)
+				{
+					SelectedComponents.Add(Component);
+				}
+			}
+		}
+	}
+}
\ No newline at end of file
diff --git a/Engine/Plugins/Experimental/PCG/Source/PCGEditor/Public/PCGComponentDetails.h b/Engine/Plugins/Experimental/PCG/Source/PCGEditor/Public/PCGComponentDetails.h
--- a/Engine/Plugins/Experimental/PCG/Source/PCGEditor/Public/PCGComponentDetails.h
+++ b/Engine/Plugins/Experimental/PCG/Source/PCGEditor/Public/PCGComponentDetails.h
@@ -18,11 +18,16 @@ public:
 	virtual void PendingDelete() override;
 	/** ~End IDetailCustomization interface */
 
+protected:
+	virtual void GatherPCGComponentsFromSelection(const TArray<TWeakObjectPtr<UObject>>& InObjectSelected);
+	virtual bool AddDefaultProperties() const { return true; }
+
 private:
 	FReply OnGenerateClicked();
 	FReply OnCleanupClicked();
 	FReply OnClearPCGLinkClicked();
 	void OnGraphChanged(UPCGComponent* InComponent);
 
+protected:
 	TArray<TWeakObjectPtr<UPCGComponent>> SelectedComponents;
 };
\ No newline at end of file
diff --git a/Engine/Plugins/Experimental/PCG/Source/PCGEditor/Public/PCGVolumeDetails.h b/Engine/Plugins/Experimental/PCG/Source/PCGEditor/Public/PCGVolumeDetails.h
new file mode 100644
--- /dev/null
+++ b/Engine/Plugins/Experimental/PCG/Source/PCGEditor/Public/PCGVolumeDetails.h
@@ -0,0 +1,17 @@
+// Copyright Epic Games, Inc. All Rights Reserved.
+
+#pragma once
+
+#include "PCGComponentDetails.h"
+
+class FPCGVolumeDetails : public FPCGComponentDetails
+{
+public:
+	static TSharedRef<IDetailCustomization> MakeInstance();
+
+protected:
+	/** ~Begin FPCGComponentDetails interface */
+	virtual void GatherPCGComponentsFromSelection(const TArray<TWeakObjectPtr<UObject>>& InObjectSelected) override;
+	virtual bool AddDefaultProperties() const override { return false; }
+	/** ~End FPCGComponentDetails interface */
+};
\ No newline at end of file
