Adding multiple, dynamic determinism tests to the PCG Editor Graph with the ability to add custom tests:
+ Tests are individualized now
+ Basic tests can be generated through the framework
* All 6 basic tests will run by default on nodes

KNOWN ISSUES:
Some Spatial nodes, especially Difference, take a long time to complete. Will be adding chrono time per test in a future CL

diff --git a/Engine/Plugins/Experimental/PCG/Source/PCG/Private/Tests/Determinism/PCGDeterminismTestsCommon.cpp b/Engine/Plugins/Experimental/PCG/Source/PCG/Private/Tests/Determinism/PCGDeterminismTestsCommon.cpp
--- a/Engine/Plugins/Experimental/PCG/Source/PCG/Private/Tests/Determinism/PCGDeterminismTestsCommon.cpp
+++ b/Engine/Plugins/Experimental/PCG/Source/PCG/Private/Tests/Determinism/PCGDeterminismTestsCommon.cpp
@@ -21,9 +21,12 @@
 #include "Editor.h"
 #endif
 
+#define LOCTEXT_NAMESPACE "PCGDeterminism"
+
 namespace PCGDeterminismTests
 {
-	FTestData::FTestData(int32 RandomSeed) :
+	FTestData::FTestData(int32 RandomSeed, UPCGSettings* DefaultSettings) :
+		Settings(DefaultSettings),
 		Seed(RandomSeed),
 		RandomStream(Seed)
 	{
@@ -63,21 +66,119 @@ namespace PCGDeterminismTests
 		Settings = nullptr;
 	}
 
-	PCG_API void RunDeterminismTests(const UPCGNode* PCGNode, FPCGDeterminismResult& Result)
+	bool LogInvalidTest(const UPCGNode* InPCGNode, int32 Seed, EPCGDataType& OutDataTypesTested, TArray<FString>& OutAdditionalDetails)
 	{
-		check(PCGNode);
+		UE_LOG(LogPCG, Warning, TEXT("Attempting to run an invalid determinism test"));
+		OutAdditionalDetails.Add(TEXT("Invalid test"));
+		return false;
+	}
 
-		FTestData TestData(Defaults::Seed);
-		TestData.Settings = PCGNode->DefaultSettings;
+	PCG_API void RunDeterminismTest(const UPCGNode* InPCGNode, FNodeTestResult& OutResult, const FNodeTestInfo& TestToRun)
+	{
+		if (!InPCGNode || !InPCGNode->DefaultSettings)
+		{
+			OutResult.DataTypesTested = EPCGDataType::None;
+			OutResult.bFlagRaised = true;
+			OutResult.AdditionalDetails.Add(TEXT("Invalid Node or Settings"));
+			return;
+		}
+
+		bool bTestWasSuccessful = TestToRun.TestDelegate(InPCGNode, InPCGNode->DefaultSettings->Seed, OutResult.DataTypesTested, OutResult.AdditionalDetails);
+		OutResult.TestResults.FindOrAdd(FName(TestToRun.TestLabel.ToString()), bTestWasSuccessful);
+		OutResult.bFlagRaised = OutResult.bFlagRaised && !bTestWasSuccessful;
+	}
+
+	PCG_API void RetrieveBasicTests(TArray<FNodeTestInfo>& OutBasicTests)
+	{
+		OutBasicTests.Emplace(LOCTEXT("SingleSameData", "Single Same Data"), RunSingleSameDataTest);
+		OutBasicTests.Emplace(LOCTEXT("SingleIdenticalData", "Single Identical Data"), RunSingleIdenticalDataTest);
+		OutBasicTests.Emplace(LOCTEXT("MultipleSameData", "Multiple Same Data"), RunMultipleSameDataTest);
+		OutBasicTests.Emplace(LOCTEXT("MultipleIdenticalData", "Multiple Identical Data"), RunMultipleIdenticalDataTest);
+		OutBasicTests.Emplace(LOCTEXT("DataCollectionOrderIndependence", "Shuffle Data Indices"),RunDataCollectionOrderIndependenceTest);
+		OutBasicTests.Emplace(LOCTEXT("AllDataOrderIndependence", "Shuffle Internal Data"), RunAllDataOrderIndependenceTest);
+	}
+
+	bool RunSingleSameDataTest(const UPCGNode* InPCGNode, int32 Seed, EPCGDataType& OutDataTypesTested, TArray<FString>& OutAdditionalDetails)
+	{
+		FTestData TestData(Seed, InPCGNode->DefaultSettings);
 
 		// Generate random input data
-		AddInputDataBasedOnPins(TestData, PCGNode, Result);
-		Result.bIsDeterministic = ExecutionIsDeterministicSameData(TestData, PCGNode);
+		AddInputDataBasedOnPins(TestData, InPCGNode, OutDataTypesTested, OutAdditionalDetails);
+
+		return ExecutionIsDeterministicSameData(TestData, InPCGNode);
+	}
+
+	bool RunSingleIdenticalDataTest(const UPCGNode* InPCGNode, int32 Seed, EPCGDataType& OutDataTypesTested, TArray<FString>& OutAdditionalDetails)
+	{
+		FTestData FirstTestData(Seed, InPCGNode->DefaultSettings);
+		FTestData SecondTestData(Seed, InPCGNode->DefaultSettings);
+
+		AddInputDataBasedOnPins(FirstTestData, InPCGNode, OutDataTypesTested, OutAdditionalDetails);
+		AddInputDataBasedOnPins(SecondTestData, InPCGNode, OutDataTypesTested, OutAdditionalDetails);
+
+		return ExecutionIsDeterministic(FirstTestData, SecondTestData, InPCGNode);
+	}
+
+	bool RunMultipleSameDataTest(const UPCGNode* InPCGNode, int32 Seed, EPCGDataType& OutDataTypesTested, TArray<FString>& OutAdditionalDetails)
+	{
+		FTestData TestData(Seed, InPCGNode->DefaultSettings);
+
+		for (int32 I = 0; I < Defaults::NumMultipleTestDataSets; ++I)
+		{
+			AddInputDataBasedOnPins(TestData, InPCGNode, OutDataTypesTested, OutAdditionalDetails);
+		}
+
+		return ExecutionIsDeterministicSameData(TestData, InPCGNode);
+	}
+
+	bool RunMultipleIdenticalDataTest(const UPCGNode* InPCGNode, int32 Seed, EPCGDataType& OutDataTypesTested, TArray<FString>& OutAdditionalDetails)
+	{
+		FTestData FirstTestData(Seed, InPCGNode->DefaultSettings);
+		FTestData SecondTestData(Seed, InPCGNode->DefaultSettings);
+
+		for (int32 I = 0; I < Defaults::NumMultipleTestDataSets; ++I)
+		{
+			AddInputDataBasedOnPins(FirstTestData, InPCGNode, OutDataTypesTested, OutAdditionalDetails);
+			AddInputDataBasedOnPins(SecondTestData, InPCGNode, OutDataTypesTested, OutAdditionalDetails);
+		}
 
-		// TODO: Add other 4 determinism tests, ex. Order Independence, etc
+		return ExecutionIsDeterministic(FirstTestData, SecondTestData, InPCGNode);
 	}
 
-	void AddInputDataBasedOnPins(FTestData& TestData, const UPCGNode* PCGNode, FPCGDeterminismResult& Result)
+	bool RunDataCollectionOrderIndependenceTest(const UPCGNode* InPCGNode, int32 Seed, EPCGDataType& OutDataTypesTested, TArray<FString>& OutAdditionalDetails)
+	{
+		FTestData FirstTestData(Seed, InPCGNode->DefaultSettings);
+		FTestData SecondTestData(Seed, InPCGNode->DefaultSettings);
+
+		for (int32 I = 0; I < Defaults::NumMultipleTestDataSets; ++I)
+		{
+			AddInputDataBasedOnPins(FirstTestData, InPCGNode, OutDataTypesTested, OutAdditionalDetails);
+			AddInputDataBasedOnPins(SecondTestData, InPCGNode, OutDataTypesTested, OutAdditionalDetails);
+		}
+
+		ShuffleInputOrder(SecondTestData);
+
+		return ExecutionIsDeterministic(FirstTestData, SecondTestData, InPCGNode);
+	}
+
+	bool RunAllDataOrderIndependenceTest(const UPCGNode* InPCGNode, int32 Seed, EPCGDataType& OutDataTypesTested, TArray<FString>& OutAdditionalDetails)
+	{
+		FTestData FirstTestData(Seed, InPCGNode->DefaultSettings);
+		FTestData SecondTestData(Seed, InPCGNode->DefaultSettings);
+
+		for (int32 I = 0; I < Defaults::NumMultipleTestDataSets; ++I)
+		{
+			AddInputDataBasedOnPins(FirstTestData, InPCGNode, OutDataTypesTested, OutAdditionalDetails);
+			AddInputDataBasedOnPins(SecondTestData, InPCGNode, OutDataTypesTested, OutAdditionalDetails);
+		}
+
+		ShuffleInputOrder(SecondTestData);
+		ShuffleAllInternalData(SecondTestData);
+
+		return ExecutionIsDeterministic(FirstTestData, SecondTestData, InPCGNode);
+	}
+
+	void AddInputDataBasedOnPins(FTestData& TestData, const UPCGNode* PCGNode, EPCGDataType& OutDataTypesTested, TArray<FString>& OutAdditionalDetails)
 	{
 		check(PCGNode);
 
@@ -89,36 +190,36 @@ namespace PCGDeterminismTests
 			{
 			case EPCGDataType::Point:
 				AddRandomizedMultiplePointInputData(TestData);
-				Result.DataTestedString = TEXT("Point");
+				OutDataTypesTested |= EPCGDataType::Point;
 				break;
 			case EPCGDataType::Volume:
 				AddRandomizedVolumeInputData(TestData);
-				Result.DataTestedString = TEXT("Volume");
+				OutDataTypesTested |= EPCGDataType::Volume;
 				break;
 			case EPCGDataType::PolyLine:
 				AddRandomizedPolyLineInputData(TestData);
-				Result.DataTestedString = TEXT("PolyLine");
+				OutDataTypesTested |= EPCGDataType::PolyLine;
 				break;
 			case EPCGDataType::Primitive:
 				AddRandomizedPrimitiveInputData(TestData);
-				Result.DataTestedString = TEXT("Primitive");
+				OutDataTypesTested |= EPCGDataType::Primitive;
 				break;
 			case EPCGDataType::Surface:
 				AddRandomizedSurfaceInputData(TestData);
-				Result.DataTestedString = TEXT("Surface");
+				OutDataTypesTested |= EPCGDataType::Surface;
 				break;
 			case EPCGDataType::Any:
 			case EPCGDataType::Spatial:
-				AddRandomizedMultiplePointInputData(TestData, 1000);
+				AddRandomizedMultiplePointInputData(TestData, Defaults::NumTestPointsToGenerate);
 				AddRandomizedVolumeInputData(TestData);
-				AddRandomizedPolyLineInputData(TestData);
+				AddRandomizedPolyLineInputData(TestData, Defaults::NumTestPolyLinePointsToGenerate);
 				AddRandomizedPrimitiveInputData(TestData);
 				AddRandomizedSurfaceInputData(TestData);
-				Result.DataTestedString = TEXT("Spatial");
+				OutDataTypesTested |= EPCGDataType::Spatial;
 				break;
 			default:
-				Result.DataTestedString = TEXT("Unknown");
-				Result.AdditionalDetailString = TEXT("Unknown data Input Pin");
+				OutDataTypesTested = EPCGDataType::None;
+				OutAdditionalDetails.Add(TEXT("Unknown InputPin data type to test"));
 				break;
 			}
 		}
@@ -203,7 +304,7 @@ namespace PCGDeterminismTests
 	void AddRandomizedVolumeInputData(FTestData& TestData, const FName& PinName)
 	{
 		AddVolumeInputData(TestData.InputData,
-			Defaults::MediumDistance * TestData.RandomStream.VRand(),
+			TestData.RandomStream.VRand() * Defaults::MediumDistance,
 			Defaults::MediumVector + TestData.RandomStream.VRand() * 0.5f * Defaults::MediumDistance,
 			Defaults::SmallVector + TestData.RandomStream.VRand() * 0.5f * Defaults::SmallDistance,
 			PinName);
@@ -306,6 +407,11 @@ namespace PCGDeterminismTests
 				// Only compare if they are the same type and same pin label
 				if (FirstDataType == SecondDataType && FirstCollection.TaggedData[I].Pin == SecondCollection.TaggedData[I].Pin)
 				{
+					// TODO: Compare in cascading order:
+					// Order independent
+					// Order consistent
+					// Deterministic but not having a predictable order
+					// Non-deterministic
 					if (GetCompareFunction(FirstDataType)(FirstCollection.TaggedData[I].Data, SecondCollection.TaggedData[J].Data))
 					{
 						MatchedIndices.Emplace(J);
@@ -524,7 +630,7 @@ namespace PCGDeterminismTests
 	bool DataTypeIsComparable(EPCGDataType DataType)
 	{
 		// Comparable data types
-		if (DataType == EPCGDataType::Spatial)
+		if (!!(DataType & EPCGDataType::Spatial))
 		{
 			return true;
 		}
@@ -544,16 +650,25 @@ namespace PCGDeterminismTests
 		return Data && DataTypeIsComparable(Data->GetDataType());
 	}
 
+	bool DataCanBeShuffled(const UPCGData* Data)
+	{
+		// TODO: Revisit for other data types that can be shuffled
+		return Data && Data->IsA<UPCGPointData>();
+	}
+
 	void ShuffleInputOrder(FTestData& TestData)
 	{
-		const int32 LastIndex = TestData.InputData.TaggedData.Num() - 1;
-		for (int32 I = 0; I <= LastIndex; ++I)
-		{
-			int32 Index = TestData.RandomStream.RandRange(I, LastIndex);
+		ShuffleArray<FPCGTaggedData>(TestData.InputData.TaggedData, TestData.RandomStream);
+	}
 
-			if (I != Index)
+	void ShuffleAllInternalData(FTestData& TestData)
+	{
+		for (FPCGTaggedData& TaggedData : TestData.InputData.TaggedData)
+		{
+			if (DataCanBeShuffled(TaggedData.Data))
 			{
-				TestData.InputData.TaggedData.Swap(I, Index);
+				UPCGPointData* PointData = CastChecked<UPCGPointData>(TaggedData.Data);
+				ShuffleArray<FPCGPoint>(PointData->GetMutablePoints(), TestData.RandomStream);
 			}
 		}
 	}
@@ -577,7 +692,7 @@ namespace PCGDeterminismTests
 		}
 	}
 
-	bool ExecutionIsDeterministic(FTestData& FirstTestData, FTestData& SecondTestData, const UPCGNode* PCGNode)
+	bool ExecutionIsDeterministic(const FTestData& FirstTestData, const FTestData& SecondTestData, const UPCGNode* PCGNode)
 	{
 		FPCGElementPtr FirstElement = FirstTestData.Settings->GetElement();
 		FPCGElementPtr SecondElement = SecondTestData.Settings->GetElement();
@@ -604,4 +719,6 @@ namespace PCGDeterminismTests
 	{
 		return ExecutionIsDeterministic(TestData, TestData, PCGNode);
 	}
-}
\ No newline at end of file
+}
+
+#undef LOCTEXT_NAMESPACE
\ No newline at end of file
diff --git a/Engine/Plugins/Experimental/PCG/Source/PCG/Public/Tests/Determinism/PCGDeterminismTestsCommon.h b/Engine/Plugins/Experimental/PCG/Source/PCG/Public/Tests/Determinism/PCGDeterminismTestsCommon.h
--- a/Engine/Plugins/Experimental/PCG/Source/PCG/Public/Tests/Determinism/PCGDeterminismTestsCommon.h
+++ b/Engine/Plugins/Experimental/PCG/Source/PCG/Public/Tests/Determinism/PCGDeterminismTestsCommon.h
@@ -20,8 +20,11 @@ namespace PCGDeterminismTests
 	{
 		constexpr static int32 Seed = 42;
 		constexpr static int32 NumPointsToGenerate = 1;
+		constexpr static int32 NumTestPointsToGenerate = 1000;
 		constexpr static int32 NumPolyLinePointsToGenerate = 6;
+		constexpr static int32 NumTestPolyLinePointsToGenerate = 12;
 		constexpr static int32 NumSamplingStepsPerDimension = 100;
+		constexpr static int32 NumMultipleTestDataSets = 2;
 
 		constexpr static FVector::FReal SmallDistance = 50.0;
 		constexpr static FVector::FReal MediumDistance = 200.0;
@@ -38,7 +41,8 @@ namespace PCGDeterminismTests
 
 	struct FTestData
 	{
-		explicit FTestData(int32 Seed = Defaults::Seed);
+		explicit FTestData(int32 Seed = Defaults::Seed, UPCGSettings* DefaultSettings = nullptr);
+
 		void Reset();
 
 		AActor* TestActor;
@@ -49,21 +53,56 @@ namespace PCGDeterminismTests
 		FRandomStream RandomStream;
 	};
 
-	struct FPCGDeterminismResult
+	struct FNodeTestResult
+	{
+		int32 Index = -1;
+		EPCGDataType DataTypesTested = EPCGDataType::None;
+		FName NodeTitle = TEXT("Untitled");
+		FString NodeNameString = TEXT("Unnamed");
+		TMap<FName, bool> TestResults;
+		TArray<FString> AdditionalDetails;
+		bool bFlagRaised = false;
+		// TODO: Chrono Duration of how long the tests took
+	};
+
+	/** A default delegate to report an unset test */
+	bool LogInvalidTest(const UPCGNode* InPCGNode, int32 Seed, EPCGDataType& OutDataTypesTested, TArray<FString>& OutAdditionalDetails);
+
+	typedef TFunction<bool(const UPCGNode* InPCGNode, int32 Seed, EPCGDataType& OutDataTypesTested, TArray<FString>& OutAdditionalDetails)> TestFunction;
+
+	struct FNodeTestInfo
 	{
-		int32 Index;
-		bool bIsDeterministic;
-		FName NodeTitle;
-		FString NodeNameString;
-		FString DataTestedString;
-		FString AdditionalDetailString;
+		FNodeTestInfo(FText Label, TestFunction Delegate, float LabelWidth = 140.f) :
+			TestLabel(Label),
+			TestDelegate(Delegate),
+			TestLabelWidth(LabelWidth) {}
+
+		FText TestLabel = NSLOCTEXT("PCGDeterminism", "UnnamedTest", "Unnamed Test");
+		TestFunction TestDelegate = LogInvalidTest;
+		float TestLabelWidth = 140.f;
 	};
 
-	/** Validates if a PCGNode is deterministic, updating the passed in result accordingly */
-	PCG_API void RunDeterminismTests(const UPCGNode* InPCGNode, FPCGDeterminismResult& OutResult);
+	/** Validates if a PCGNode is deterministic */
+	PCG_API void RunDeterminismTest(const UPCGNode* InPCGNode, FNodeTestResult& OutResult, const FNodeTestInfo& TestToRun);
+
+	/** Adds the basic set of determinism tests to the passed in array */
+	PCG_API void RetrieveBasicTests(TArray<FNodeTestInfo>& OutBasicTests);
+
+	/** Validates node determinism against the same single test data */
+	bool RunSingleSameDataTest(const UPCGNode* InPCGNode, int32 Seed, EPCGDataType& OutDataTypesTested, TArray<FString>& OutAdditionalDetails);
+	/** Validates node determinism against two identical single test data */
+	bool RunSingleIdenticalDataTest(const UPCGNode* InPCGNode, int32 Seed, EPCGDataType& OutDataTypesTested, TArray<FString>& OutAdditionalDetails);
+	/** Validates node determinism with the same multiple sets of test data */
+	bool RunMultipleSameDataTest(const UPCGNode* InPCGNode, int32 Seed, EPCGDataType& OutDataTypesTested, TArray<FString>& OutAdditionalDetails);
+	/** Validates node determinism with two identical multiple sets of test data */
+	bool RunMultipleIdenticalDataTest(const UPCGNode* InPCGNode, int32 Seed, EPCGDataType& OutDataTypesTested, TArray<FString>& OutAdditionalDetails);
+	/** Validates node determinism with multiple sets of test data, shuffling the order of the second set's data collection */
+	bool RunDataCollectionOrderIndependenceTest(const UPCGNode* InPCGNode, int32 Seed, EPCGDataType& OutDataTypesTested, TArray<FString>& OutAdditionalDetails);
+	/** Validates node determinism with multiple sets of test data, shuffling all internal data */
+	bool RunAllDataOrderIndependenceTest(const UPCGNode* InPCGNode, int32 Seed, EPCGDataType& OutDataTypesTested, TArray<FString>& OutAdditionalDetails);
 
 	/** Adds input data to test data, based on input pins' allowed types */
-	void AddInputDataBasedOnPins(FTestData& TestData, const UPCGNode* InPCGNode, FPCGDeterminismResult& OutResult);
+	void AddInputDataBasedOnPins(FTestData& TestData, const UPCGNode* InPCGNode, EPCGDataType& OutDataTypesTested, TArray<FString>& OutAdditionalDetails);
 
 	/** Helper for adding PointData with a single Point to an InputData */
 	void AddSinglePointInputData(FPCGDataCollection& InputData, const FVector& Location, const FName& PinName = Defaults::TestPinName);
@@ -117,24 +156,30 @@ namespace PCGDeterminismTests
 	bool DataTypeIsComparable(EPCGDataType DataType);
 	/** Validates PCGData and if its DataType is comparable */
 	bool DataIsComparable(const UPCGData* Data);
+	/** Validates that the data contains a shuffle-able array of data */
+	bool DataCanBeShuffled(const UPCGData* Data);
 
-	/** Randomizes the order of InputData */
+	/** Randomizes the order of InputData collections */
 	void ShuffleInputOrder(FTestData& TestData);
+	/** Randomizes the order of all internal data */
+	void ShuffleAllInternalData(FTestData& TestData);
 
 	/** Gets a comparison function to compare two of a specific DataType */
 	TFunction<bool(const UPCGData*, const UPCGData*)> GetCompareFunction(EPCGDataType DataType);
 
 	/** Execute the elements for each valid input and compare if all the outputs are identical */
-	bool ExecutionIsDeterministic(FTestData& FirstTestData, FTestData& SecondTestData, const UPCGNode* PCGNode = nullptr);
+	bool ExecutionIsDeterministic(const FTestData& FirstTestData, const FTestData& SecondTestData, const UPCGNode* PCGNode = nullptr);
 	/** Execute the same element twice compare if all the outputs are identical */
 	bool ExecutionIsDeterministicSameData(FTestData& TestData, const UPCGNode* PCGNode = nullptr);
 
 	/** Generates settings based upon a UPCGSettings subclass */
 	template<typename SettingsType>
-	void GenerateSettings(FTestData& TestData, TFunction<void(FTestData&)> ExtraSettingsDelegate = nullptr)
+	SettingsType* GenerateSettings(FTestData& TestData, TFunction<void(FTestData&)> ExtraSettingsDelegate = nullptr)
 	{
-		TestData.Settings = NewObject<SettingsType>();
-		check(TestData.Settings);
+		SettingsType* TypedSettings = NewObject<SettingsType>();
+		check(TypedSettings);
+
+		TestData.Settings = TypedSettings;
 		TestData.Settings->Seed = TestData.Seed;
 
 		TestData.InputData.TaggedData.Emplace_GetRef().Data = TestData.Settings;
@@ -144,6 +189,8 @@ namespace PCGDeterminismTests
 		{
 			ExtraSettingsDelegate(TestData);
 		}
+
+		return TypedSettings;
 	}
 
 	/** Validates whether both UPCGData can be cast to a specified subclass */
@@ -154,4 +201,19 @@ namespace PCGDeterminismTests
 
 		return (Cast<const DataType>(FirstData) != nullptr && Cast<const DataType>(SecondData) != nullptr);
 	}
+
+	template<typename DataType>
+	void ShuffleArray(TArray<DataType>& Array, FRandomStream& RandomStream)
+	{
+		const int32 LastIndex = Array.Num() - 1;
+		for (int32 I = 0; I <= LastIndex; ++I)
+		{
+			int32 Index = RandomStream.RandRange(I, LastIndex);
+
+			if (I != Index)
+			{
+				Array.Swap(I, Index);
+			}
+		}
+	}
 }
\ No newline at end of file
diff --git a/Engine/Plugins/Experimental/PCG/Source/PCGEditor/Private/PCGEditor.cpp b/Engine/Plugins/Experimental/PCG/Source/PCGEditor/Private/PCGEditor.cpp
--- a/Engine/Plugins/Experimental/PCG/Source/PCGEditor/Private/PCGEditor.cpp
+++ b/Engine/Plugins/Experimental/PCG/Source/PCGEditor/Private/PCGEditor.cpp
@@ -283,7 +283,7 @@ void FPCGEditor::OnDeterminismTests()
 {
 	check(GraphEditorWidget.IsValid());
 
-	if (!DeterminismWidget.IsValid() || !DeterminismWidget->IsContructed())
+	if (!DeterminismWidget.IsValid() || !DeterminismWidget->WidgetIsConstructed())
 	{
 		return;
 	}
@@ -302,11 +302,19 @@ void FPCGEditor::OnDeterminismTests()
 				const UPCGNode* PCGNode = PCGEditorGraphNode->GetPCGNode();
 				check(PCGNode);
 
-				TSharedPtr<PCGDeterminismTests::FPCGDeterminismResult> NodeResult = MakeShared<PCGDeterminismTests::FPCGDeterminismResult>();
+				TSharedPtr<PCGDeterminismTests::FNodeTestResult> NodeResult = MakeShared<PCGDeterminismTests::FNodeTestResult>();
+				check(NodeResult.IsValid());
 				NodeResult->Index = Index++;
 				NodeResult->NodeTitle = PCGNode->GetNodeTitle();
 				NodeResult->NodeNameString = PCGNode->GetName();
-				PCGDeterminismTests::RunDeterminismTests(PCGNode, *NodeResult);
+
+				TArray<PCGDeterminismTests::FNodeTestInfo> BasicTests;
+				PCGDeterminismTests::RetrieveBasicTests(BasicTests);
+
+				for (PCGDeterminismTests::FNodeTestInfo TestInfo : BasicTests)
+				{
+					PCGDeterminismTests::RunDeterminismTest(PCGNode, *NodeResult, TestInfo);
+				}
 
 				DeterminismWidget->AddItem(NodeResult);
 			}
@@ -1155,7 +1163,17 @@ TSharedRef<SPCGEditorGraphFind> FPCGEditor::CreateFindWidget()
 
 TSharedRef<SPCGEditorGraphDeterminismListView> FPCGEditor::CreateDeterminismWidget()
 {
-	return SNew(SPCGEditorGraphDeterminismListView, SharedThis(this));
+	TArray<PCGDeterminismTests::FNodeTestInfo> BasicTests;
+	PCGDeterminismTests::RetrieveBasicTests(BasicTests);
+
+	TArray<FTestColumnInfo> TestColumnInfo;
+	for (const PCGDeterminismTests::FNodeTestInfo& TestInfo : BasicTests)
+	{
+		// Build the columns based on the tests to run
+		TestColumnInfo.Emplace(FName(TestInfo.TestLabel.ToString()), TestInfo.TestLabel, TestInfo.TestLabelWidth, HAlign_Center);
+	}
+
+	return SNew(SPCGEditorGraphDeterminismListView, SharedThis(this), TestColumnInfo);
 }
 
 void FPCGEditor::OnSelectedNodesChanged(const TSet<UObject*>& NewSelection)
diff --git a/Engine/Plugins/Experimental/PCG/Source/PCGEditor/Private/SPCGEditorGraphDeterminism.cpp b/Engine/Plugins/Experimental/PCG/Source/PCGEditor/Private/SPCGEditorGraphDeterminism.cpp
--- a/Engine/Plugins/Experimental/PCG/Source/PCGEditor/Private/SPCGEditorGraphDeterminism.cpp
+++ b/Engine/Plugins/Experimental/PCG/Source/PCGEditor/Private/SPCGEditorGraphDeterminism.cpp
@@ -4,32 +4,35 @@
 
 #include "PCGEditor.h"
 
-#define LOCTEXT_NAMESPACE "SPCGEditorGraphDeterminismListView"
+#define LOCTEXT_NAMESPACE "PCGDeterminism"
 
 namespace
 {
-	const FName NAME_Index(TEXT("IndexColumn"));
-	const FName NAME_NodeTitle(TEXT("NodeTitleColumn"));
-	const FName NAME_NodeName(TEXT("NodeNameColumn"));
-	const FName NAME_Data(TEXT("DataColumn"));
-	const FName NAME_Result(TEXT("ResultColumn"));
-	const FName NAME_AdditionalDetails(TEXT("AdditionalDetailsColumn"));
-
-	const FText TEXT_Index(LOCTEXT("IndexLabel", ""));
-	const FText TEXT_NodeTitle(LOCTEXT("NodeTitleLabel", "Title"));
-	const FText TEXT_NodeName(LOCTEXT("NodeNameLabel", "Name"));
-	const FText TEXT_Data(LOCTEXT("DataLabel", "Input Data"));
-	const FText TEXT_Result(LOCTEXT("ResultLabel", "Deterministic"));
-	const FText TEXT_AdditionalDetails(LOCTEXT("AdditionalDetailsLabel", "Additional Details"));
+	const FName NAME_Index(TEXT("Index_ColumnID"));
+	const FName NAME_NodeTitle(TEXT("NodeTitle_ColumnID"));
+	const FName NAME_NodeName(TEXT("NodeName_ColumnID"));
+	const FName NAME_DataTypesTested(TEXT("DataTypesTested_ColumnID"));
+	const FName NAME_AdditionalDetails(TEXT("AdditionalDetails_ColumnID"));
+
+	const FText TEXT_Index(LOCTEXT("Index_Label", ""));
+	const FText TEXT_NodeTitle(LOCTEXT("NodeTitle_Label", "Title"));
+	const FText TEXT_NodeName(LOCTEXT("NodeName_Label", "Name"));
+	const FText TEXT_DataTypesTested(LOCTEXT("DataTypesTested_Label", "Input Data"));
+	const FText TEXT_AdditionalDetails(LOCTEXT("AdditionalDetails_Label", "Additional Details"));
 
 	const FText TEXT_Yes(LOCTEXT("Yes", "Yes"));
 	const FText TEXT_No(LOCTEXT("No", "No"));
+
+	constexpr static float SmallManualWidth = 25.f;
+	constexpr static float MediumManualWidth = 70.f;
+	constexpr static float LargeManualWidth = 160.f;
+	constexpr static float ListViewRowHeight = 36.f;
 }
 
-void SPCGEditorGraphDeterminismRow::Construct(const FArguments& InArgs, const TSharedRef<STableViewBase>& InOwnerTableView, const FPCGDeterminismResultPtr& Item)
+void SPCGEditorGraphDeterminismRow::Construct(const FArguments& InArgs, const TSharedRef<STableViewBase>& InOwnerTableView, const FPCGNodeTestResultPtr& Item)
 {
 	CurrentItem = Item;
-	SMultiColumnTableRow<FPCGDeterminismResultPtr>::Construct(SMultiColumnTableRow::FArguments(), InOwnerTableView);
+	SMultiColumnTableRow<FPCGNodeTestResultPtr>::Construct(SMultiColumnTableRow::FArguments(), InOwnerTableView);
 }
 
 TSharedRef<SWidget> SPCGEditorGraphDeterminismRow::GenerateWidgetForColumn(const FName& ColumnId)
@@ -38,109 +41,122 @@ TSharedRef<SWidget> SPCGEditorGraphDeterminismRow::GenerateWidgetForColumn(const
 
 	FText CellText = LOCTEXT("UnknownColumn", "Unknown");
 
+	auto ReturnColorCodedResultBlock = [this](const FText& CellText)
+	{
+		return SNew(STextBlock)
+			.Text(CellText)
+			.ColorAndOpacity(CurrentItem->bFlagRaised ? FColor::Red : FColor::Green);
+	};
+
+	// Permanent columns
 	if (ColumnId == NAME_Index)
 	{
-		CellText = FText::FromString(FString::FromInt(CurrentItem->Index));
+		return ReturnColorCodedResultBlock(FText::FromString(FString::FromInt(CurrentItem->Index)));
 	}
 	else if (ColumnId == NAME_NodeTitle)
 	{
-		CellText = FText::FromString(CurrentItem->NodeTitle.ToString());
+		return ReturnColorCodedResultBlock(FText::FromName(CurrentItem->NodeTitle));
 	}
 	else if (ColumnId == NAME_NodeName)
 	{
 		CellText = FText::FromString(CurrentItem->NodeNameString);
 	}
-	else if (ColumnId == NAME_Data)
+	else if (ColumnId == NAME_DataTypesTested)
 	{
-		CellText = FText::FromString(CurrentItem->DataTestedString);
+		FString DataTypesTestedString = *UEnum::GetValueAsString(CurrentItem->DataTypesTested);
+		DataTypesTestedString.RemoveFromStart("EPCGDataType::");
+		CellText = FText::FromString(DataTypesTestedString);
 	}
-	else if (ColumnId == NAME_Result)
+	else if (ColumnId == NAME_AdditionalDetails)
 	{
-		if (CurrentItem->bIsDeterministic)
-		{
-			return SNew(STextBlock)
-					.Text(TEXT_Yes)
-					.ColorAndOpacity(FColor::Green);
-		}
-		else
-		{
-			return SNew(STextBlock)
-					.Text(TEXT_No)
-					.ColorAndOpacity(FColor::Red);
-		}
+		FString FullDetails = FString::Join(CurrentItem->AdditionalDetails, TEXT(", "));
+		CellText = FText::FromString(FullDetails);
 	}
-	else if (ColumnId == NAME_AdditionalDetails)
+	// Test columns
+	else if (bool* bResultPtr = CurrentItem->TestResults.Find(ColumnId))
 	{
-		CellText = FText::FromString(CurrentItem->AdditionalDetailString);
+		bool bTestWasSuccessful = bResultPtr && *bResultPtr;
+		return SNew(STextBlock)
+			.Text(bTestWasSuccessful ? TEXT_Yes : TEXT_No)
+			.ColorAndOpacity(bTestWasSuccessful ? FColor::Green : FColor::Red);
 	}
 
-	return SNew(STextBlock).Text(CellText);
+	return SNew(STextBlock)
+		.Text(CellText);
 }
 
-void SPCGEditorGraphDeterminismListView::Construct(const FArguments& InArgs, TWeakPtr<FPCGEditor> InPCGEditor)
+void SPCGEditorGraphDeterminismListView::Construct(const FArguments& InArgs, TWeakPtr<FPCGEditor> InPCGEditor, const TArray<FTestColumnInfo>& InTestColumns)
 {
 	check(InPCGEditor.IsValid() && !bIsConstructed);
 	PCGEditorPtr = InPCGEditor;
 
-	SAssignNew(ListView, SListView<FPCGDeterminismResultPtr>)
+	TSharedRef<SHeaderRow> GeneratedHeaderRow = SNew(SHeaderRow);
+
+	TArray<FTestColumnInfo> TestColumnInfo;
+	// Permanent columns
+	TestColumnInfo.Emplace(NAME_Index, TEXT_Index, SmallManualWidth, HAlign_Center);
+	TestColumnInfo.Emplace(NAME_NodeTitle, TEXT_NodeTitle, LargeManualWidth, HAlign_Left);
+	TestColumnInfo.Emplace(NAME_NodeName, TEXT_NodeName, LargeManualWidth, HAlign_Left);
+	TestColumnInfo.Emplace(NAME_DataTypesTested, TEXT_DataTypesTested, MediumManualWidth, HAlign_Center);
+
+	// Test columns
+	TestColumnInfo.Append(InTestColumns);
+
+	// Final details column
+	TestColumnInfo.Emplace(NAME_AdditionalDetails, TEXT_AdditionalDetails, 0.f, HAlign_Left);
+
+	// Build column arguments for test columns dynamically
+	for (const FTestColumnInfo& ColumnInfo : TestColumnInfo)
+	{
+		SHeaderRow::FColumn::FArguments Arguments;
+		Arguments.ColumnId(ColumnInfo.ColumnID);
+		Arguments.DefaultLabel(ColumnInfo.ColumnLabel);
+		if (ColumnInfo.Width > 0.f)
+		{
+			Arguments.ManualWidth(ColumnInfo.Width);
+		}
+		Arguments.HAlignCell(ColumnInfo.HAlign);
+		GeneratedHeaderRow->AddColumn(Arguments);
+	}
+
+	SAssignNew(ListView, SListView<FPCGNodeTestResultPtr>)
 		.ListItemsSource(&ListViewItems)
-		.ItemHeight(36)
+		.ItemHeight(ListViewRowHeight)
 		.OnGenerateRow(this, &SPCGEditorGraphDeterminismListView::OnGenerateRow)
-		.HeaderRow
-		(
-			SNew(SHeaderRow)
-			+ SHeaderRow::Column(NAME_Index)
-			.DefaultLabel(TEXT_Index)
-			.ManualWidth(25)
-			.HAlignCell(HAlign_Center)
-			+ SHeaderRow::Column(NAME_NodeTitle)
-			.DefaultLabel(TEXT_NodeTitle)
-			.ManualWidth(160)
-			.HAlignCell(HAlign_Left)
-			+ SHeaderRow::Column(NAME_NodeName)
-			.DefaultLabel(TEXT_NodeName)
-			.ManualWidth(160)
-			.HAlignCell(HAlign_Left)
-			+ SHeaderRow::Column(NAME_Data)
-			.DefaultLabel(TEXT_Data)
-			.ManualWidth(70)
-			.HAlignCell(HAlign_Center)
-			+ SHeaderRow::Column(NAME_Result)
-			.DefaultLabel(TEXT_Result)
-			.ManualWidth(90)
-			.HAlignCell(HAlign_Center)
-			+ SHeaderRow::Column(NAME_AdditionalDetails)
-			.DefaultLabel(TEXT_AdditionalDetails)
-			.HAlignCell(HAlign_Left)
-		);
+		.HeaderRow(GeneratedHeaderRow);
 
 	ChildSlot
-		[
-			ListView->AsShared()
-		];
+	[
+		ListView->AsShared()
+	];
 
 	bIsConstructed = true;
 }
 
-void SPCGEditorGraphDeterminismListView::AddItem(const FPCGDeterminismResultPtr Item)
+void SPCGEditorGraphDeterminismListView::AddItem(const FPCGNodeTestResultPtr& Item)
 {
 	check(Item.IsValid());
-	ListViewItems.Add(Item);
-	ListView->RequestListRefresh();
+	ListViewItems.Emplace(Item);
+	Refresh();
 }
 
 void SPCGEditorGraphDeterminismListView::Clear()
 {
 	ListViewItems.Empty();
+	Refresh();
+}
+
+void SPCGEditorGraphDeterminismListView::Refresh()
+{
 	ListView->RequestListRefresh();
 }
 
-bool SPCGEditorGraphDeterminismListView::IsContructed() const
+bool SPCGEditorGraphDeterminismListView::WidgetIsConstructed() const
 {
 	return bIsConstructed;
 }
 
-TSharedRef<ITableRow> SPCGEditorGraphDeterminismListView::OnGenerateRow(const FPCGDeterminismResultPtr Item, const TSharedRef<STableViewBase>& OwnerTable)
+TSharedRef<ITableRow> SPCGEditorGraphDeterminismListView::OnGenerateRow(const FPCGNodeTestResultPtr Item, const TSharedRef<STableViewBase>& OwnerTable) const
 {
 	return SNew(SPCGEditorGraphDeterminismRow, OwnerTable, Item);
 }
diff --git a/Engine/Plugins/Experimental/PCG/Source/PCGEditor/Private/SPCGEditorGraphDeterminism.h b/Engine/Plugins/Experimental/PCG/Source/PCGEditor/Private/SPCGEditorGraphDeterminism.h
--- a/Engine/Plugins/Experimental/PCG/Source/PCGEditor/Private/SPCGEditorGraphDeterminism.h
+++ b/Engine/Plugins/Experimental/PCG/Source/PCGEditor/Private/SPCGEditorGraphDeterminism.h
@@ -7,46 +7,62 @@
 
 class FPCGEditor;
 
-typedef TSharedPtr<PCGDeterminismTests::FPCGDeterminismResult> FPCGDeterminismResultPtr;
+typedef TSharedPtr<PCGDeterminismTests::FNodeTestResult> FPCGNodeTestResultPtr;
 
-class SPCGEditorGraphDeterminismRow : public SMultiColumnTableRow<FPCGDeterminismResultPtr>
+struct FTestColumnInfo
+{
+	FTestColumnInfo(FName ColumnID, FText ColumnLabel, float Width = 0.f, EHorizontalAlignment HAlign = HAlign_Left) :
+		ColumnID(ColumnID),
+		ColumnLabel(ColumnLabel),
+		Width(Width),
+		HAlign(HAlign) {}
+
+	FName ColumnID = TEXT("UnnamedColumn_ID");
+	FText ColumnLabel = NSLOCTEXT("PCGDeterminism", "Unnamed_Column", "Unnamed Column");
+	float Width = 0.f;
+	EHorizontalAlignment HAlign = HAlign_Left;
+};
+
+class SPCGEditorGraphDeterminismRow final : public SMultiColumnTableRow<FPCGNodeTestResultPtr>
 {
 	SLATE_BEGIN_ARGS(SPCGEditorGraphDeterminismRow) {}
 	SLATE_END_ARGS()
 
 	/** Construct a row of the ListView */
-	void Construct(const FArguments& InArgs, const TSharedRef<STableViewBase>& InOwnerTableView, const FPCGDeterminismResultPtr& Item);
+	void Construct(const FArguments& InArgs, const TSharedRef<STableViewBase>& InOwnerTableView, const FPCGNodeTestResultPtr& Item);
 
 	/** Generates a column, given the column's ID */
 	virtual TSharedRef<SWidget> GenerateWidgetForColumn(const FName& ColumnId) override;
 
 protected:
-	FPCGDeterminismResultPtr CurrentItem;
+	FPCGNodeTestResultPtr CurrentItem;
 };
 
-class SPCGEditorGraphDeterminismListView : public SCompoundWidget
+class SPCGEditorGraphDeterminismListView final : public SCompoundWidget
 {
 public:
-	SLATE_BEGIN_ARGS(SPCGEditorGraphDeterminismListView) { }
+	SLATE_BEGIN_ARGS(SPCGEditorGraphDeterminismListView) {}
 	SLATE_END_ARGS()
 
 	/** Construct the ListView */
-	void Construct(const FArguments& InArgs, TWeakPtr<FPCGEditor> InPCGEditor);
+	void Construct(const FArguments& InArgs, TWeakPtr<FPCGEditor> InPCGEditor, const TArray<FTestColumnInfo>& InTestColumn);
 	/** Add an item to the ListView */
-	void AddItem(const FPCGDeterminismResultPtr Item);
+	void AddItem(const FPCGNodeTestResultPtr& Item);
 	/** Clear all items from the ListView */
 	void Clear();
+	/** Refreshes the ListView */
+	void Refresh();
 	/** Validates if the ListView has been constructed */
-	bool IsContructed() const;
+	bool WidgetIsConstructed() const;
 
 private:
-	TSharedRef<ITableRow> OnGenerateRow(const FPCGDeterminismResultPtr Item, const TSharedRef<STableViewBase>& OwnerTable);
+	/** Generate the row widget */
+	TSharedRef<ITableRow> OnGenerateRow(const FPCGNodeTestResultPtr Item, const TSharedRef<STableViewBase>& OwnerTable) const;
 
-	/** Weak pointer to the PCGEditor */
 	TWeakPtr<FPCGEditor> PCGEditorPtr;
 
-	TSharedPtr<SListView<FPCGDeterminismResultPtr>> ListView;
-	TArray<FPCGDeterminismResultPtr> ListViewItems;
+	TSharedPtr<SListView<FPCGNodeTestResultPtr>> ListView;
+	TArray<FPCGNodeTestResultPtr> ListViewItems;
 
 	bool bIsConstructed = false;
 };
\ No newline at end of file
