Fix crash sometime when clearing PCG World Actors

diff --git a/Engine/Plugins/Experimental/PCG/Source/PCG/Private/Grid/PCGPartitionActor.cpp b/Engine/Plugins/Experimental/PCG/Source/PCG/Private/Grid/PCGPartitionActor.cpp
--- a/Engine/Plugins/Experimental/PCG/Source/PCG/Private/Grid/PCGPartitionActor.cpp
+++ b/Engine/Plugins/Experimental/PCG/Source/PCG/Private/Grid/PCGPartitionActor.cpp
@@ -24,6 +24,11 @@ void APCGPartitionActor::PostLoad()
 	{
 		PCGGridSize = APCGWorldActor::DefaultPartitionGridSize;
 	}
+
+#if WITH_EDITOR
+	// Make sure that we don't track objects that do not exist anymore
+	CleanupDeadGraphInstances();
+#endif // WITH_EDITOR
 }
 
 void APCGPartitionActor::BeginPlay()
@@ -203,7 +208,7 @@ bool APCGPartitionActor::IsSafeForDeletion() const
 	ensure(IsInGameThread());
 	for (TObjectPtr<UPCGComponent> PCGComponent : GetAllOriginalPCGComponents())
 	{
-		if (PCGComponent->IsGenerating())
+		if (PCGComponent && PCGComponent->IsGenerating())
 		{
 			return false;
 		}
