Native TransformPoints PCG Element

diff --git a/Engine/Plugins/Experimental/PCG/Source/PCG/Private/Elements/PCGTransformPoints.cpp b/Engine/Plugins/Experimental/PCG/Source/PCG/Private/Elements/PCGTransformPoints.cpp
new file mode 100644
--- /dev/null
+++ b/Engine/Plugins/Experimental/PCG/Source/PCG/Private/Elements/PCGTransformPoints.cpp
@@ -0,0 +1,155 @@
+// Copyright Epic Games, Inc. All Rights Reserved.
+
+#include "Elements/PCGTransformPoints.h"
+
+#include "PCGHelpers.h"
+#include "Data/PCGPointData.h"
+#include "Data/PCGSpatialData.h"
+#include "Helpers/PCGAsync.h"
+#include "Helpers/PCGSettingsHelpers.h"
+
+TArray<FPCGPinProperties> UPCGTransformPointsSettings::InputPinProperties() const
+{
+	TArray<FPCGPinProperties> PinProperties;
+	PinProperties.Emplace(TEXT("Params"), EPCGDataType::Param);
+	PinProperties.Emplace(PCGPinConstants::DefaultInputLabel, EPCGDataType::Spatial);
+
+	return PinProperties;
+}
+
+FPCGElementPtr UPCGTransformPointsSettings::CreateElement() const
+{
+	return MakeShared<FPCGTransformPointsElement>();
+}
+
+bool FPCGTransformPointsElement::ExecuteInternal(FPCGContext* Context) const
+{
+	TRACE_CPUPROFILER_EVENT_SCOPE(FPCGTransformPointsElement::Execute);
+
+	const UPCGTransformPointsSettings* Settings = Context->GetInputSettings<UPCGTransformPointsSettings>();
+	check(Settings);
+
+	const TArray<FPCGTaggedData> Inputs = Context->InputData.GetInputs();
+	TArray<FPCGTaggedData>& Outputs = Context->OutputData.TaggedData;
+
+	UPCGParamData* Params = Context->InputData.GetParams();
+	const FVector OffsetMin = PCG_GET_OVERRIDEN_VALUE(Settings, OffsetMin, Params);
+	const FVector OffsetMax = PCG_GET_OVERRIDEN_VALUE(Settings, OffsetMax, Params);
+	const bool bAbsoluteOffset = PCG_GET_OVERRIDEN_VALUE(Settings, bAbsoluteOffset, Params);
+	const FRotator RotationMin = PCG_GET_OVERRIDEN_VALUE(Settings, RotationMin, Params);
+	const FRotator RotationMax = PCG_GET_OVERRIDEN_VALUE(Settings, RotationMax, Params);
+	const bool bAbsoluteRotation = PCG_GET_OVERRIDEN_VALUE(Settings, bAbsoluteRotation, Params);
+	const FVector ScaleMin = PCG_GET_OVERRIDEN_VALUE(Settings, ScaleMin, Params);
+	const FVector ScaleMax = PCG_GET_OVERRIDEN_VALUE(Settings, ScaleMax, Params);
+	const bool bAbsoluteScale = PCG_GET_OVERRIDEN_VALUE(Settings, bAbsoluteScale, Params);
+	const bool bUniformScale = PCG_GET_OVERRIDEN_VALUE(Settings, bUniformScale, Params);
+	const bool bRecomputeSeed = PCG_GET_OVERRIDEN_VALUE(Settings, bRecomputeSeed, Params);
+
+	for (const FPCGTaggedData& Input : Inputs)
+	{
+		TRACE_CPUPROFILER_EVENT_SCOPE(FPCGTransformPointsElement::Execute::InputLoop);
+
+		FPCGTaggedData& Output = Outputs.Add_GetRef(Input);
+
+		if (!Input.Data || Cast<UPCGSpatialData>(Input.Data) == nullptr)
+		{
+			PCGE_LOG(Error, "Invalid input data");
+			continue;
+		}
+
+		const UPCGSpatialData* SpatialData = Cast<UPCGSpatialData>(Input.Data);
+
+		if (!SpatialData)
+		{
+			PCGE_LOG(Error, "Unable to get SpatialData from input");
+			continue;
+		}
+
+		const UPCGPointData* PointData = SpatialData->ToPointData(Context);
+
+		if (!PointData)
+		{
+			PCGE_LOG(Error, "Unable to get PointData from input");
+			continue;
+		}
+
+		const TArray<FPCGPoint>& Points = PointData->GetPoints();
+
+		UPCGPointData* OutPointData = NewObject<UPCGPointData>();
+		OutPointData->InitializeFromData(PointData);
+		TArray<FPCGPoint>& OutPoints = OutPointData->GetMutablePoints();
+		Output.Data = OutPointData;
+
+		// loop on points
+		FPCGAsync::AsyncPointProcessing(Context, Points.Num(), OutPoints, [&](int32 Index, FPCGPoint& OutPoint)
+		{
+			const FPCGPoint& InPoint = Points[Index];
+			OutPoint = InPoint;
+
+			FRandomStream RandomSource(PCGHelpers::ComputeSeed(Settings->Seed, InPoint.Seed));
+
+			const float OffsetX = RandomSource.FRandRange(OffsetMin.X, OffsetMax.X);
+			const float OffsetY = RandomSource.FRandRange(OffsetMin.Y, OffsetMax.Y);
+			const float OffsetZ = RandomSource.FRandRange(OffsetMin.Z, OffsetMax.Z);
+			const FVector RandomOffset(OffsetX, OffsetY, OffsetZ);
+
+			const float RotationX = RandomSource.FRandRange(RotationMin.Pitch, RotationMax.Pitch);
+			const float RotationY = RandomSource.FRandRange(RotationMin.Yaw, RotationMax.Yaw);
+			const float RotationZ = RandomSource.FRandRange(RotationMin.Roll, RotationMax.Roll);
+			const FQuat RandomRotation(FRotator(RotationX, RotationY, RotationZ).Quaternion());
+
+			FVector RandomScale;
+			if (bUniformScale)
+			{
+				RandomScale = FVector(RandomSource.FRandRange(ScaleMin.X, ScaleMax.X));
+			}
+			else
+			{
+				RandomScale.X = RandomSource.FRandRange(ScaleMin.X, ScaleMax.X);
+				RandomScale.Y = RandomSource.FRandRange(ScaleMin.Y, ScaleMax.Y);
+				RandomScale.Z = RandomSource.FRandRange(ScaleMin.Z, ScaleMax.Z);
+			}
+
+			if (bAbsoluteOffset)
+			{
+				OutPoint.Transform.SetLocation(InPoint.Transform.GetLocation() + RandomOffset); 
+			}
+			else
+			{
+				const FTransform RotatedTransform(InPoint.Transform.GetRotation());
+				OutPoint.Transform.SetLocation(InPoint.Transform.GetLocation() + RotatedTransform.TransformPosition(RandomOffset)); 
+			}
+
+			if (bAbsoluteRotation)
+			{
+				OutPoint.Transform.SetRotation(RandomRotation);
+			}
+			else
+			{
+				OutPoint.Transform.SetRotation(InPoint.Transform.GetRotation() * RandomRotation);
+			}
+
+			if (bAbsoluteScale)
+			{
+				OutPoint.Transform.SetScale3D(RandomScale);
+			}
+			else
+			{
+				OutPoint.Transform.SetScale3D(InPoint.Transform.GetScale3D() * RandomScale);
+			}
+
+			if (bRecomputeSeed)
+			{
+				const FVector& Position = OutPoint.Transform.GetLocation();
+				OutPoint.Seed = PCGHelpers::ComputeSeed((int)Position.X, (int)Position.Y, (int)Position.Z);
+			}
+
+			return true;
+		});
+	}
+
+	// Forward any non-input data
+	Outputs.Append(Context->InputData.GetAllSettings());
+
+	return true;
+}
diff --git a/Engine/Plugins/Experimental/PCG/Source/PCG/Private/Tests/Elements/PCGTransformPointsTest.cpp b/Engine/Plugins/Experimental/PCG/Source/PCG/Private/Tests/Elements/PCGTransformPointsTest.cpp
new file mode 100644
--- /dev/null
+++ b/Engine/Plugins/Experimental/PCG/Source/PCG/Private/Tests/Elements/PCGTransformPointsTest.cpp
@@ -0,0 +1,195 @@
+// Copyright Epic Games, Inc. All Rights Reserved.
+
+#include "Tests/PCGTestsCommon.h"
+#include "Tests/Determinism/PCGDeterminismTestsCommon.h"
+
+#include "PCGComponent.h"
+#include "PCGData.h"
+#include "PCGHelpers.h"
+
+#include "Data/PCGPointData.h"
+#include "Elements/PCGTransformPoints.h"
+
+IMPLEMENT_CUSTOM_SIMPLE_AUTOMATION_TEST(FPCGTransformPointsTest, FPCGTestBaseClass, "pcg.tests.TransformPoints.Basic", PCGTestsCommon::TestFlags)
+
+bool FPCGTransformPointsTest::RunTest(const FString& Parameters)
+{
+	PCGDeterminismTests::FTestData TestData(PCGDeterminismTests::Defaults::Seed);
+	PCGDeterminismTests::GenerateSettings<UPCGTransformPointsSettings>(TestData);
+	UPCGTransformPointsSettings* Settings = CastChecked<UPCGTransformPointsSettings>(TestData.Settings);
+	FPCGElementPtr TransformPointsElement = TestData.Settings->GetElement();
+
+	TObjectPtr<UPCGPointData> EmptyData = PCGTestsCommon::CreateEmptyPointData();
+	TObjectPtr<UPCGPointData> SimpleData = PCGTestsCommon::CreateEmptyPointData();
+	TObjectPtr<UPCGPointData> ComplexData = PCGTestsCommon::CreateEmptyPointData();
+
+	// SimplePoints are empty transforms, ComplexPoints have unique location, rotation, and scale
+	TArray<FPCGPoint>& SimplePoints = SimpleData->GetMutablePoints();
+	TArray<FPCGPoint>& ComplexPoints = ComplexData->GetMutablePoints();
+
+	FRandomStream RandomSource(PCGHelpers::ComputeSeed(TestData.Seed));
+	const int PointCount = 100;
+	for (int I = 0; I < PointCount; ++I)
+	{
+		SimplePoints.Emplace(FTransform(), 1, I);
+
+		const FRotator RandomRotation = FRotator(RandomSource.FRandRange(0.f, 360.f));
+		const FVector RandomLocation = RandomSource.VRand();
+		const FVector RandomScale = RandomSource.VRand().GetAbs();
+		ComplexPoints.Emplace(FTransform(RandomRotation, RandomLocation, RandomScale), 1, I);
+	}
+
+	FPCGTaggedData& EmptyTaggedData = TestData.InputData.TaggedData.Emplace_GetRef(FPCGTaggedData());
+	EmptyTaggedData.Data = EmptyData;
+
+	FPCGTaggedData& SimpleTaggedData = TestData.InputData.TaggedData.Emplace_GetRef(FPCGTaggedData());
+	SimpleTaggedData.Data = SimpleData;
+
+	FPCGTaggedData& ComplexTaggedData = TestData.InputData.TaggedData.Emplace_GetRef(FPCGTaggedData());
+	ComplexTaggedData.Data = ComplexData;
+
+	// Basic parameters we will use for all tests
+	Settings->OffsetMin = FVector(-1.f);
+	Settings->OffsetMax = FVector(1.f);
+	Settings->RotationMin = FRotator(-1.f);
+	Settings->RotationMax = FRotator(1.f);
+	Settings->ScaleMin = FVector(1.f);
+	Settings->ScaleMax = FVector(2.f);
+
+	auto ValidateTransformPoints = [this, &TestData, TransformPointsElement, Settings]()
+	{
+		TUniquePtr<FPCGContext> Context = MakeUnique<FPCGContext>(*TransformPointsElement->Initialize(TestData.InputData, TestData.TestPCGComponent, nullptr));
+		Context->NumAvailableTasks = 1;
+
+		while (!TransformPointsElement->Execute(Context.Get()))
+		{}
+
+		const TArray<FPCGTaggedData>& Inputs = Context->InputData.GetInputs();
+		const TArray<FPCGTaggedData>& Outputs = Context->OutputData.GetInputs();
+
+		if (!TestEqual("Valid number of outputs", Inputs.Num(), Outputs.Num()))
+		{
+			return;
+		}
+
+		for (int DataIndex = 0; DataIndex < Inputs.Num(); ++DataIndex)
+		{
+			const FPCGTaggedData& Input = Inputs[DataIndex];
+			const FPCGTaggedData& Output = Outputs[DataIndex];
+
+			check(Input.Data);
+			
+			const UPCGSpatialData* InSpatialData = Cast<UPCGSpatialData>(Input.Data);
+			check(InSpatialData);
+
+			const UPCGPointData* InPointData = InSpatialData->ToPointData(Context.Get());
+			check(InPointData);
+
+			if (!TestTrue("Valid output data", Output.Data != nullptr))
+			{
+				continue;
+			}
+
+			const UPCGSpatialData* OutSpatialData = Cast<UPCGSpatialData>(Output.Data);
+
+			if (!TestNotNull("Valid ouptut SpatialData", OutSpatialData))
+			{
+				continue;
+			}
+
+			const UPCGPointData* OutPointData = OutSpatialData->ToPointData(Context.Get());
+
+			if (!TestNotNull("Valid output PointData", OutPointData))
+			{
+				continue;
+			}
+
+			const TArray<FPCGPoint>& InPoints = InPointData->GetPoints();
+			const TArray<FPCGPoint>& OutPoints = OutPointData->GetPoints();
+
+			if (!TestEqual("Input and output point counts match", InPoints.Num(), OutPoints.Num()))
+			{ 
+				continue;
+			}
+
+			for (int PointIndex = 0; PointIndex < InPoints.Num(); ++PointIndex)
+			{
+				const FPCGPoint& InPoint = InPoints[PointIndex];
+				const FPCGPoint& OutPoint = OutPoints[PointIndex];
+
+				FPCGPoint RepositionedPoint = InPoint;
+				RepositionedPoint.Transform = OutPoint.Transform;
+				TestTrue("RepositionedPoint and OutPoint are identical", PCGTestsCommon::PointsAreIdentical(RepositionedPoint, OutPoint));
+
+				// Validate transform is within range of original transform
+				const FTransform& InTransform = InPoint.Transform;
+				const FTransform& OutTransform = OutPoint.Transform;
+
+				FTransform AbsoluteTransform = OutTransform;
+				AbsoluteTransform.SetLocation(OutTransform.GetLocation() - InTransform.GetLocation());
+
+				FTransform RelativeTransform = OutTransform.GetRelativeTransform(InTransform);
+				RelativeTransform.SetLocation(RelativeTransform.GetLocation() * InTransform.GetScale3D());
+
+				const FTransform& TransformForLocation = (Settings->bAbsoluteOffset ? AbsoluteTransform : RelativeTransform);
+				const FTransform& TransformForRotation = (Settings->bAbsoluteRotation ? AbsoluteTransform : RelativeTransform);
+				const FTransform& TransformForScale = (Settings->bAbsoluteScale ? AbsoluteTransform : RelativeTransform);
+
+				FVector OriginalOffset = TransformForLocation.GetLocation();
+				FRotator OriginalRotation = TransformForRotation.GetRotation().Rotator();
+				FVector OriginalScale = TransformForScale.GetScale3D();
+
+				TestTrue("Valid location", OriginalOffset.X >= Settings->OffsetMin.X && OriginalOffset.X <= Settings->OffsetMax.X
+					&& OriginalOffset.Y >= Settings->OffsetMin.Y && OriginalOffset.Y <= Settings->OffsetMax.Y
+					&& OriginalOffset.Z >= Settings->OffsetMin.Z && OriginalOffset.Z <= Settings->OffsetMax.Z);
+
+				TestTrue("Valid rotation", OriginalRotation.Pitch >= Settings->RotationMin.Pitch && OriginalRotation.Pitch <= Settings->RotationMax.Pitch
+					&& OriginalRotation.Yaw >= Settings->RotationMin.Yaw && OriginalRotation.Yaw <= Settings->RotationMax.Yaw
+					&& OriginalRotation.Roll >= Settings->RotationMin.Roll && OriginalRotation.Roll <= Settings->RotationMax.Roll);
+
+				if (Settings->bUniformScale)
+				{
+					TestTrue("Valid scale", OriginalScale.X >= Settings->ScaleMin.X && OriginalScale.X <= Settings->ScaleMax.X
+						&& OriginalScale.Y >= Settings->ScaleMin.X && OriginalScale.Y <= Settings->ScaleMax.X
+						&& OriginalScale.Z >= Settings->ScaleMin.X && OriginalScale.Z <= Settings->ScaleMax.X);
+				}
+				else
+				{
+					TestTrue("Valid scale", OriginalScale.X >= Settings->ScaleMin.X && OriginalScale.X <= Settings->ScaleMax.X
+						&& OriginalScale.Y >= Settings->ScaleMin.Y && OriginalScale.Y <= Settings->ScaleMax.Y
+						&& OriginalScale.Z >= Settings->ScaleMin.Z && OriginalScale.Z <= Settings->ScaleMax.Z);
+				}
+
+				if (Settings->bRecomputeSeed)
+				{
+					const FVector& Position = OutTransform.GetLocation();
+					TestEqual("Valid seed", OutPoint.Seed, PCGHelpers::ComputeSeed((int)Position.X, (int)Position.Y, (int)Position.Z));
+				}
+				else
+				{
+					TestEqual("Valid seed", OutPoint.Seed, InPoint.Seed);
+				}
+			}
+		}
+	};
+
+	// Test 1 - absolute transformations
+	Settings->bAbsoluteOffset = true;
+	Settings->bAbsoluteRotation = true;
+	Settings->bAbsoluteScale = true;
+	ValidateTransformPoints();
+
+	// Test 2 - relative transformations
+	Settings->bAbsoluteOffset = false;
+	Settings->bAbsoluteRotation = false;
+	Settings->bAbsoluteScale = false;
+	ValidateTransformPoints();
+
+	// Test 3 - uniform scale and recompute seed
+	Settings->bUniformScale = true;
+	Settings->bRecomputeSeed = true;
+	ValidateTransformPoints();
+
+	return true;
+
+}
diff --git a/Engine/Plugins/Experimental/PCG/Source/PCG/Public/Elements/PCGTransformPoints.h b/Engine/Plugins/Experimental/PCG/Source/PCG/Public/Elements/PCGTransformPoints.h
new file mode 100644
--- /dev/null
+++ b/Engine/Plugins/Experimental/PCG/Source/PCG/Public/Elements/PCGTransformPoints.h
@@ -0,0 +1,73 @@
+// Copyright Epic Games, Inc. All Rights Reserved.
+
+#pragma once
+
+#include "CoreMinimal.h"
+#include "PCGSettings.h"
+
+#include "PCGTransformPoints.generated.h"
+
+UCLASS(BlueprintType, ClassGroup = (Procedural))
+class PCG_API UPCGTransformPointsSettings : public UPCGSettings
+{
+	GENERATED_BODY()
+
+public:
+	//~Begin UPCGSettings interface
+#if WITH_EDITOR
+	virtual FName GetDefaultNodeName() const override { return FName(TEXT("TransformPointsNode")); }
+	virtual EPCGSettingsType GetType() const override { return EPCGSettingsType::Spatial; }
+#endif
+
+	virtual TArray<FPCGPinProperties> InputPinProperties() const override;
+	virtual TArray<FPCGPinProperties> OutputPinProperties() const override { return Super::DefaultPointOutputPinProperties(); }
+
+protected:
+	virtual FPCGElementPtr CreateElement() const override;
+	//~End UPCGSettings interface
+
+public:
+	UPROPERTY(BlueprintReadWrite, EditAnywhere, Category = Settings)
+	FVector OffsetMin = FVector::Zero();
+
+	UPROPERTY(BlueprintReadWrite, EditAnywhere, Category = Settings)
+	FVector OffsetMax = FVector::Zero();
+
+	/** Set offset in world space */
+	UPROPERTY(BlueprintReadWrite, EditAnywhere, Category = Settings)
+	bool bAbsoluteOffset = false;
+
+	UPROPERTY(BlueprintReadWrite, EditAnywhere, Category = Settings)
+	FRotator RotationMin = FRotator::ZeroRotator;
+
+	UPROPERTY(BlueprintReadWrite, EditAnywhere, Category = Settings)
+	FRotator RotationMax = FRotator::ZeroRotator;
+
+	/** Set rotation directly instead of additively */
+	UPROPERTY(BlueprintReadWrite, EditAnywhere, Category = Settings)
+	bool bAbsoluteRotation = false;
+
+	UPROPERTY(BlueprintReadWrite, EditAnywhere, Category = Settings, meta = (AllowPreserveRatio))
+	FVector ScaleMin = FVector::One();
+
+	UPROPERTY(BlueprintReadWrite, EditAnywhere, Category = Settings, meta = (AllowPreserveRatio))
+	FVector ScaleMax = FVector::One();
+
+	/** Set scale directly instead of multiplicatively */
+	UPROPERTY(BlueprintReadWrite, EditAnywhere, Category = Settings)
+	bool bAbsoluteScale = false;
+
+	/** Scale uniformly on each axis. Uses the X component of ScaleMin and ScaleMax. */
+	UPROPERTY(BlueprintReadWrite, EditAnywhere, Category = Settings)
+	bool bUniformScale = true;
+
+	/** Recompute the seed for each new point using its new location */
+	UPROPERTY(BlueprintReadWrite, EditAnywhere, Category = Settings)
+	bool bRecomputeSeed = false;
+};
+
+class FPCGTransformPointsElement : public FSimplePCGElement
+{
+protected:
+	virtual bool ExecuteInternal(FPCGContext* Context) const override;
+};
