Native CopyPoints PCG Element

diff --git a/Engine/Plugins/Experimental/PCG/Source/PCG/Private/Elements/PCGCopyPoints.cpp b/Engine/Plugins/Experimental/PCG/Source/PCG/Private/Elements/PCGCopyPoints.cpp
new file mode 100644
--- /dev/null
+++ b/Engine/Plugins/Experimental/PCG/Source/PCG/Private/Elements/PCGCopyPoints.cpp
@@ -0,0 +1,212 @@
+// Copyright Epic Games, Inc. All Rights Reserved.
+
+#include "Elements/PCGCopyPoints.h"
+
+#include "PCGHelpers.h"
+#include "Data/PCGPointData.h"
+#include "Data/PCGSpatialData.h"
+#include "Helpers/PCGAsync.h"
+#include "Helpers/PCGSettingsHelpers.h"
+
+TArray<FPCGPinProperties> UPCGCopyPointsSettings::InputPinProperties() const
+{
+	TArray<FPCGPinProperties> PinProperties;
+	PinProperties.Emplace(PCGCopyPointsConstants::ParamsLabel, EPCGDataType::Param);
+	PinProperties.Emplace(PCGCopyPointsConstants::SourcePointsLabel, EPCGDataType::Spatial, /*bAllowMultipleConnections=*/false);
+	PinProperties.Emplace(PCGCopyPointsConstants::TargetPointsLabel, EPCGDataType::Spatial, /*bAllowMultipleConnections=*/false);
+
+	return PinProperties;
+}
+
+FPCGElementPtr UPCGCopyPointsSettings::CreateElement() const
+{
+	return MakeShared<FPCGCopyPointsElement>();
+}
+
+bool FPCGCopyPointsElement::ExecuteInternal(FPCGContext* Context) const
+{
+	TRACE_CPUPROFILER_EVENT_SCOPE(FPCGCopyPointsElement::Execute);
+
+	const UPCGCopyPointsSettings* Settings = Context->GetInputSettings<UPCGCopyPointsSettings>();
+	check(Settings);
+
+	UPCGParamData* Params = Context->InputData.GetParams();
+	const EPCGCopyPointsInheritanceMode RotationInheritance = PCG_GET_OVERRIDEN_VALUE(Settings, RotationInheritance, Params);
+	const EPCGCopyPointsInheritanceMode ScaleInheritance = PCG_GET_OVERRIDEN_VALUE(Settings, ScaleInheritance, Params);
+	const EPCGCopyPointsInheritanceMode ColorInheritance = PCG_GET_OVERRIDEN_VALUE(Settings, ColorInheritance, Params);
+	const EPCGCopyPointsInheritanceMode SeedInheritance = PCG_GET_OVERRIDEN_VALUE(Settings, SeedInheritance, Params);
+	const EPCGCopyPointsMetadataInheritanceMode AttributeInheritance = PCG_GET_OVERRIDEN_VALUE(Settings, AttributeInheritance, Params);
+
+	const TArray<FPCGTaggedData> Sources = Context->InputData.GetInputsByPin(PCGCopyPointsConstants::SourcePointsLabel);
+	const TArray<FPCGTaggedData> Targets = Context->InputData.GetInputsByPin(PCGCopyPointsConstants::TargetPointsLabel);
+
+	TArray<FPCGTaggedData>& Outputs = Context->OutputData.TaggedData;
+
+	if (Sources.Num() != 1 || Targets.Num() != 1)
+	{
+		PCGE_LOG(Error, "Invalid number of inputs (should be exactly one source and one target)");
+		return true;
+	}
+	
+	const FPCGTaggedData& Source = Sources[0];
+	const FPCGTaggedData& Target = Targets[0];
+
+	FPCGTaggedData& Output = Outputs.Add_GetRef(Source);
+
+	if (!Source.Data || !Target.Data) 
+	{
+		PCGE_LOG(Error, "Invalid input data");
+		return true;
+	}
+
+	const UPCGSpatialData* SourceSpatialData = Cast<UPCGSpatialData>(Source.Data);
+	const UPCGSpatialData* TargetSpatialData = Cast<UPCGSpatialData>(Target.Data);
+
+	if (!SourceSpatialData || !TargetSpatialData)
+	{
+		PCGE_LOG(Error, "Unable to get SpatialData from input");
+		return true;
+	}
+
+	const UPCGPointData* SourcePointData = SourceSpatialData->ToPointData(Context);
+	const UPCGPointData* TargetPointData = TargetSpatialData->ToPointData(Context);
+
+	if (!SourcePointData || !TargetPointData)
+	{
+		PCGE_LOG(Error, "Unable to get PointData from input");
+		return true;
+	}
+
+	const TArray<FPCGPoint>& SourcePoints = SourcePointData->GetPoints();
+	const TArray<FPCGPoint>& TargetPoints = TargetPointData->GetPoints();
+
+	UPCGPointData* OutPointData = NewObject<UPCGPointData>();
+	TArray<FPCGPoint>& OutPoints = OutPointData->GetMutablePoints();
+	Output.Data = OutPointData;
+
+	// RootMetadata will be parent to the ouptut metadata, while NonRootMetadata will carry attributes from the input not selected for inheritance
+	const UPCGMetadata* RootMetadata = nullptr;
+	const UPCGMetadata* NonRootMetadata = nullptr;
+	if (Settings->AttributeInheritance == EPCGCopyPointsMetadataInheritanceMode::Source)
+	{
+		OutPointData->InitializeFromData(SourcePointData);
+		RootMetadata = SourcePointData->Metadata;
+		NonRootMetadata = TargetPointData->Metadata;
+	}
+	else // if (Settings->AttributeInheritance == EPCGCopyPointsMetadataInheritanceMode::Target)
+	{
+		OutPointData->InitializeFromData(TargetPointData);
+		RootMetadata = TargetPointData->Metadata;
+		NonRootMetadata = SourcePointData->Metadata;
+	}
+
+	check(OutPointData->Metadata && NonRootMetadata);
+
+	TArray<FName> AttributeNames;
+	TArray<EPCGMetadataTypes> AttributeTypes;
+	NonRootMetadata->GetAttributes(AttributeNames, AttributeTypes);
+
+	TArray<FPCGMetadataAttributeBase*> NonRootAttributes;
+	for (const FName& AttributeName : AttributeNames)
+	{
+		if (!OutPointData->Metadata->HasAttribute(AttributeName))
+		{
+			const FPCGMetadataAttributeBase* Attribute = NonRootMetadata->GetConstAttribute(AttributeName);
+			FPCGMetadataAttributeBase* NewAttribute = OutPointData->Metadata->CopyAttribute(Attribute, AttributeName, /*bKeepRoot=*/false, /*bCopyEntries=*/false, /*bCopyValues=*/true);
+			NonRootAttributes.Add(NewAttribute);
+		}
+	}
+
+	FPCGAsync::AsyncPointProcessing(Context, SourcePoints.Num() * TargetPoints.Num(), OutPoints, [&](int32 Index, FPCGPoint& OutPoint)
+	{
+		const FPCGPoint& SourcePoint = SourcePoints[Index / TargetPoints.Num()];
+		const FPCGPoint& TargetPoint = TargetPoints[Index % TargetPoints.Num()];
+
+		OutPoint = SourcePoint;
+
+		// Set Rotation, Scale, and Color based on inheritance mode
+		if (RotationInheritance == EPCGCopyPointsInheritanceMode::Relative)
+		{
+			OutPoint.Transform.SetRotation(TargetPoint.Transform.GetRotation() * SourcePoint.Transform.GetRotation());
+		}
+		else if (RotationInheritance == EPCGCopyPointsInheritanceMode::Source)
+		{
+			OutPoint.Transform.SetRotation(SourcePoint.Transform.GetRotation());
+		}
+		else // if (RotationInheritance == EPCGCopyPointsInheritanceMode::Target)
+		{
+			OutPoint.Transform.SetRotation(TargetPoint.Transform.GetRotation());
+		}
+
+		if (ScaleInheritance == EPCGCopyPointsInheritanceMode::Relative)
+		{
+			OutPoint.Transform.SetScale3D(SourcePoint.Transform.GetScale3D() * TargetPoint.Transform.GetScale3D());
+		}
+		else if (ScaleInheritance == EPCGCopyPointsInheritanceMode::Source)
+		{ 
+			OutPoint.Transform.SetScale3D(SourcePoint.Transform.GetScale3D());
+		}
+		else // if (ScaleInheritance == EPCGCopyPointsInheritanceMode::Target)
+		{
+			OutPoint.Transform.SetScale3D(TargetPoint.Transform.GetScale3D());
+		}
+
+		if (ColorInheritance == EPCGCopyPointsInheritanceMode::Relative)
+		{
+			OutPoint.Color = SourcePoint.Color * TargetPoint.Color;
+		}
+		else if (ColorInheritance == EPCGCopyPointsInheritanceMode::Source)
+		{ 
+			OutPoint.Color = SourcePoint.Color;
+		}
+		else // if (ColorInheritance == EPCGCopyPointsInheritanceMode::Target)
+		{ 
+			OutPoint.Color = TargetPoint.Color;
+		}
+
+		const FVector Location = TargetPoint.Transform.TransformPosition(SourcePoint.Transform.GetLocation());
+		OutPoint.Transform.SetLocation(Location);
+
+		// Set seed based on inheritance mode
+		if (ColorInheritance == EPCGCopyPointsInheritanceMode::Relative)
+		{
+			const FVector& Position = OutPoint.Transform.GetLocation();
+			OutPoint.Seed = PCGHelpers::ComputeSeed((int)Position.X, (int)Position.Y, (int)Position.Z);
+		}
+		else if (ColorInheritance == EPCGCopyPointsInheritanceMode::Target)
+		{
+			OutPoint.Seed = TargetPoint.Seed;
+		}
+
+		const FPCGPoint* RootPoint = nullptr;
+		const FPCGPoint* NonRootPoint = nullptr;
+		if (Settings->AttributeInheritance == EPCGCopyPointsMetadataInheritanceMode::Source)
+		{
+			RootPoint = &SourcePoint;
+			NonRootPoint = &TargetPoint;
+		}
+		else // if (Settings->AttributeInheritance == EPCGCopyPointsMetadataInheritanceMode::Target)
+		{
+			RootPoint = &TargetPoint;
+			NonRootPoint = &SourcePoint;
+		}
+
+		OutPoint.MetadataEntry = OutPointData->Metadata->AddEntry(RootPoint->MetadataEntry);
+
+		// Copy EntryToValue key mappings from NonRootAttributes
+		for (FPCGMetadataAttributeBase* NonRootAttribute : NonRootAttributes)
+		{
+			const FPCGMetadataAttributeBase* Attribute = NonRootMetadata->GetConstAttribute(NonRootAttribute->Name);
+			check(Attribute);
+			const PCGMetadataValueKey ValueKey = Attribute->GetValueKey(NonRootPoint->MetadataEntry);
+			NonRootAttribute->SetValueFromValueKey(OutPoint.MetadataEntry, ValueKey);
+		}
+
+		return true;
+	});
+
+	// Forward any non-input data
+	Outputs.Append(Context->InputData.GetAllSettings());
+
+	return true;
+}
diff --git a/Engine/Plugins/Experimental/PCG/Source/PCG/Private/Tests/Elements/PCGCopyPointsTest.cpp b/Engine/Plugins/Experimental/PCG/Source/PCG/Private/Tests/Elements/PCGCopyPointsTest.cpp
new file mode 100644
--- /dev/null
+++ b/Engine/Plugins/Experimental/PCG/Source/PCG/Private/Tests/Elements/PCGCopyPointsTest.cpp
@@ -0,0 +1,321 @@
+// Copyright Epic Games, Inc. All Rights Reserved.
+
+#include "Tests/Determinism/PCGDeterminismTestsCommon.h"
+#include "Tests/PCGTestsCommon.h"
+
+#include "PCGComponent.h"
+#include "PCGData.h"
+#include "PCGHelpers.h"
+
+#include "Data/PCGPointData.h"
+#include "Elements/PCGCopyPoints.h"
+#include "Metadata/PCGMetadata.h"
+#include "Metadata/PCGMetadataAttribute.h"
+#include "Metadata/PCGMetadataAttributeTraits.h"
+#include "Metadata/PCGMetadataAttributeTpl.h"
+
+IMPLEMENT_CUSTOM_SIMPLE_AUTOMATION_TEST(FPCGCopyPointsTest, FPCGTestBaseClass, "pcg.tests.CopyPoints.Basic", PCGTestsCommon::TestFlags)
+
+bool FPCGCopyPointsTest::RunTest(const FString& Parameters)
+{
+	PCGDeterminismTests::FTestData TestData(PCGDeterminismTests::Defaults::Seed);
+	PCGDeterminismTests::GenerateSettings<UPCGCopyPointsSettings>(TestData);
+	UPCGCopyPointsSettings* Settings = CastChecked<UPCGCopyPointsSettings>(TestData.Settings);
+	FPCGElementPtr CopyPointsElement = TestData.Settings->GetElement();
+
+	TObjectPtr<UPCGPointData> EmptyData = PCGTestsCommon::CreateEmptyPointData();
+	TObjectPtr<UPCGPointData> SourceData = PCGTestsCommon::CreateRandomPointData(10, TestData.Seed);
+	TObjectPtr<UPCGPointData> TargetData = PCGTestsCommon::CreateRandomPointData(4, TestData.Seed+1);
+
+	const FName LengthName = TEXT("Length");
+	const FName WidthName = TEXT("Width");
+	const FName HeightName = TEXT("Height");
+
+	const bool bAllowsInterpolation = false;
+	const bool bOverrideParent = false;
+
+	SourceData->Metadata->CreateFloatAttribute(LengthName, 0, bAllowsInterpolation, bOverrideParent);
+	SourceData->Metadata->CreateFloatAttribute(WidthName, 0, bAllowsInterpolation, bOverrideParent);
+	TargetData->Metadata->CreateFloatAttribute(LengthName, 1, bAllowsInterpolation, bOverrideParent);
+	TargetData->Metadata->CreateFloatAttribute(HeightName, 1, bAllowsInterpolation, bOverrideParent);
+
+	TArray<FPCGPoint>& SourcePoints = SourceData->GetMutablePoints();
+	TArray<FPCGPoint>& TargetPoints = TargetData->GetMutablePoints();
+
+	FRandomStream RandomSource(TestData.Seed);
+	for (int I = 0; I < SourcePoints.Num(); ++I)
+	{
+		// first half of the points have a non-default Length attribute value, the rest have a non-default width attribute value
+	 	UPCGMetadataAccessorHelpers::SetFloatAttribute(SourcePoints[I], SourceData->Metadata, (I < SourcePoints.Num() / 2) ? LengthName : WidthName, RandomSource.FRand());
+	}
+
+	for (int I = 0; I < TargetPoints.Num(); ++I)
+	{
+		// first quarter of the points have a non-default Length attribute value, the rest have a non-default Height attribute value
+		UPCGMetadataAccessorHelpers::SetFloatAttribute(TargetPoints[I], TargetData->Metadata, (I < SourcePoints.Num() / 4) ? LengthName : HeightName, RandomSource.FRand());
+	}
+
+	FPCGTaggedData& SourceTaggedData = TestData.InputData.TaggedData.Emplace_GetRef(FPCGTaggedData());
+	SourceTaggedData.Data = SourceData;
+	SourceTaggedData.Pin = PCGCopyPointsConstants::SourcePointsLabel;
+
+	FPCGTaggedData& TargetTaggedData = TestData.InputData.TaggedData.Emplace_GetRef(FPCGTaggedData());
+	TargetTaggedData.Data = TargetData;
+	TargetTaggedData.Pin = PCGCopyPointsConstants::TargetPointsLabel;
+
+	// Test only supports float attributes for simplicity
+	auto ValidateCopyPointsPoints = [this, &TestData, CopyPointsElement, Settings]()
+	{
+		TUniquePtr<FPCGContext> Context = MakeUnique<FPCGContext>(*CopyPointsElement->Initialize(TestData.InputData, TestData.TestPCGComponent, nullptr));
+		Context->NumAvailableTasks = 1;
+
+		while (!CopyPointsElement->Execute(Context.Get()))
+		{}
+
+		const TArray<FPCGTaggedData> Sources = Context->InputData.GetInputsByPin(PCGCopyPointsConstants::SourcePointsLabel);
+		const TArray<FPCGTaggedData> Targets = Context->InputData.GetInputsByPin(PCGCopyPointsConstants::TargetPointsLabel);
+		const TArray<FPCGTaggedData>& Outputs = Context->OutputData.GetInputs();
+
+		if (!TestTrue("Valid number of inputs", Sources.Num() == 1 && Targets.Num() == 1))
+		{
+			return;
+		}
+
+		if (!TestTrue("Valid number of outputs", Sources.Num() == 1))
+		{
+			return;
+		}
+
+		const FPCGTaggedData& Source = Sources[0];
+		const FPCGTaggedData& Target = Targets[0];
+		const FPCGTaggedData& Output = Outputs[0];
+
+		check(Source.Data && Target.Data);
+
+		const UPCGSpatialData* SourceSpatialData = Cast<UPCGSpatialData>(Source.Data);
+		const UPCGSpatialData* TargetSpatialData = Cast<UPCGSpatialData>(Target.Data);
+
+		check(SourceSpatialData && TargetSpatialData);
+
+		const UPCGPointData* SourcePointData = SourceSpatialData->ToPointData(Context.Get());
+		const UPCGPointData* TargetPointData = TargetSpatialData->ToPointData(Context.Get());
+
+		check(SourcePointData && TargetPointData);
+
+		if (!TestTrue("Valid output data", Output.Data != nullptr))
+		{
+			return;
+		}
+
+		const UPCGSpatialData* OutSpatialData = Cast<UPCGSpatialData>(Output.Data);
+
+		if (!TestNotNull("Valid output SpatialData", OutSpatialData))
+		{
+			return;
+		}
+
+		const UPCGPointData* OutPointData = OutSpatialData->ToPointData(Context.Get());
+
+		if (!TestNotNull("Valid output PointData", OutPointData))
+		{
+			return;
+		}
+
+		const TArray<FPCGPoint>& SourcePoints = SourcePointData->GetPoints();
+		const TArray<FPCGPoint>& TargetPoints = TargetPointData->GetPoints();
+		const TArray<FPCGPoint>& OutPoints = OutPointData->GetPoints();
+
+		if (!TestEqual("Valid number of output points", SourcePoints.Num() * TargetPoints.Num(), OutPoints.Num()))
+		{ 
+			return;
+		}
+
+		const UPCGMetadata* RootMetadata = nullptr;
+		const UPCGMetadata* NonRootMetadata = nullptr;
+		EPCGCopyPointsMetadataInheritanceMode RootMode;
+		EPCGCopyPointsMetadataInheritanceMode NonRootMode;
+
+		if (Settings->AttributeInheritance == EPCGCopyPointsMetadataInheritanceMode::Source)
+		{
+			RootMetadata = SourcePointData->Metadata;
+			NonRootMetadata = TargetPointData->Metadata;
+			RootMode = EPCGCopyPointsMetadataInheritanceMode::Source;
+			NonRootMode = EPCGCopyPointsMetadataInheritanceMode::Target;
+		}
+		else if (Settings->AttributeInheritance == EPCGCopyPointsMetadataInheritanceMode::Target)
+		{ 
+			RootMetadata = TargetPointData->Metadata;
+			NonRootMetadata = SourcePointData->Metadata;
+			RootMode = EPCGCopyPointsMetadataInheritanceMode::Target;
+			NonRootMode = EPCGCopyPointsMetadataInheritanceMode::Source;
+		}
+
+		TArray<const FPCGMetadataAttribute<float>*> InheritedAttributes;
+		TArray<EPCGCopyPointsMetadataInheritanceMode> InheritedAttributeModes;
+		TArray<const FPCGMetadataAttribute<float>*> OutAttributes;
+
+		TArray<FName> AttributeNames;
+		TArray<EPCGMetadataTypes> AttributeTypes;
+		RootMetadata->GetAttributes(AttributeNames, AttributeTypes);
+
+		for (const FName& AttributeName : AttributeNames)
+		{
+			if (TestTrue("Contains attribute found in source", OutPointData->Metadata->HasAttribute(AttributeName)))
+			{
+				const FPCGMetadataAttribute<float>* InheritedAttribute = static_cast<const FPCGMetadataAttribute<float>*>(RootMetadata->GetConstAttribute(AttributeName));
+				const FPCGMetadataAttribute<float>* OutAttribute = static_cast<const FPCGMetadataAttribute<float>*>(OutPointData->Metadata->GetConstAttribute(AttributeName));
+
+				InheritedAttributes.Add(InheritedAttribute);
+				InheritedAttributeModes.Add(RootMode);
+				OutAttributes.Add(OutAttribute);
+			}
+		}
+
+		NonRootMetadata->GetAttributes(AttributeNames, AttributeTypes);
+
+		for (const FName& AttributeName : AttributeNames)
+		{
+			if (!RootMetadata->HasAttribute(AttributeName))
+			{
+				if (TestTrue("Contains attribute found in target", OutPointData->Metadata->HasAttribute(AttributeName)))
+				{
+					const FPCGMetadataAttribute<float>* InheritedAttribute = static_cast<const FPCGMetadataAttribute<float>*>(NonRootMetadata->GetConstAttribute(AttributeName));
+					const FPCGMetadataAttribute<float>* OutAttribute = static_cast<const FPCGMetadataAttribute<float>*>(OutPointData->Metadata->GetConstAttribute(AttributeName));
+
+					InheritedAttributes.Add(InheritedAttribute);
+					InheritedAttributeModes.Add(NonRootMode);
+					OutAttributes.Add(OutAttribute);
+				}
+			}
+		}
+
+		check(InheritedAttributeModes.Num() == InheritedAttributeModes.Num() && InheritedAttributeModes.Num() == OutAttributes.Num())
+
+		for (int PointIndex = 0; PointIndex < OutPoints.Num(); ++PointIndex)
+		{
+			const FPCGPoint& SourcePoint = SourcePoints[PointIndex / TargetPoints.Num()];
+			const FPCGPoint& TargetPoint = TargetPoints[PointIndex % TargetPoints.Num()];
+			const FPCGPoint& OutPoint = OutPoints[PointIndex];
+
+			TestEqual("SourcePoint and OutPoint have same Density", SourcePoint.Density, OutPoint.Density);
+			TestEqual("SourcePoint and OutPoint have same BoundsMin", SourcePoint.BoundsMin, OutPoint.BoundsMin);
+			TestEqual("SourcePoint and OutPoint have same BoundsMax", SourcePoint.BoundsMax, OutPoint.BoundsMax);
+			TestEqual("SourcePoint and OutPoint have same Steepness", SourcePoint.Steepness, OutPoint.Steepness);
+
+			// Validate transform inheritance
+			const FTransform& SourceTransform = SourcePoint.Transform;
+			const FTransform& TargetTransform = TargetPoint.Transform;
+			const FTransform& OutTransform = OutPoint.Transform;
+
+			const FVector Location = TargetPoint.Transform.TransformPosition(SourcePoint.Transform.GetLocation());
+			FQuat Rotation;
+			FVector Scale;
+			FVector4 Color;
+			int32 Seed;
+
+			if (Settings->RotationInheritance == EPCGCopyPointsInheritanceMode::Relative)
+			{
+				Rotation = TargetTransform.GetRotation() * SourceTransform.GetRotation();
+			}
+			else if (Settings->RotationInheritance == EPCGCopyPointsInheritanceMode::Source)
+			{
+				Rotation = SourceTransform.GetRotation();
+			}
+			else // if (Settings->RotationInheritance == EPCGCopyPointsInheritanceMode::Target)
+			{
+				Rotation = TargetTransform.GetRotation();
+			}
+
+			if (Settings->ScaleInheritance == EPCGCopyPointsInheritanceMode::Relative)
+			{
+				Scale = SourceTransform.GetScale3D() * TargetTransform.GetScale3D();
+			}
+			else if (Settings->ScaleInheritance == EPCGCopyPointsInheritanceMode::Source)
+			{
+				Scale = SourceTransform.GetScale3D();
+			}
+			else // if (Settings->ScaleInheritance == EPCGCopyPointsInheritanceMode::Target)
+			{
+				Scale = TargetTransform.GetScale3D();
+			}
+
+			if (Settings->ColorInheritance == EPCGCopyPointsInheritanceMode::Relative)
+			{
+				Color = SourcePoint.Color * TargetPoint.Color;
+			}
+			else if (Settings->ColorInheritance == EPCGCopyPointsInheritanceMode::Source)
+			{
+				Color = SourcePoint.Color;
+			}
+			else // if (Settings->ColorInheritance == EPCGCopyPointsInheritanceMode::Target)
+			{
+				Color = TargetPoint.Color;
+			}
+
+			if (Settings->ColorInheritance == EPCGCopyPointsInheritanceMode::Relative)
+			{
+				Seed = PCGHelpers::ComputeSeed((int)Location.X, (int)Location.Y, (int)Location.Z);
+			}
+			else if (Settings->ColorInheritance == EPCGCopyPointsInheritanceMode::Source)
+			{
+				Seed = SourcePoint.Seed;
+			}
+			else // if (Settings->ColorInheritance == EPCGCopyPointsInheritanceMode::Target)
+			{
+				Seed = TargetPoint.Seed;
+			}
+
+			TestEqual("Valid rotation", Rotation, OutTransform.GetRotation());
+			TestEqual("Valid scale", Scale, OutTransform.GetScale3D());
+			TestEqual("Valid color", Color, OutPoint.Color);
+			TestEqual("Valid location", Location, OutTransform.GetLocation());
+			TestEqual("Valid location", Seed, OutPoint.Seed);		
+
+			// Validate point value keys
+			for (int AttributeIndex = 0; AttributeIndex < InheritedAttributes.Num(); ++AttributeIndex)
+			{
+				const FPCGMetadataAttribute<float>* InheritedAttribute = InheritedAttributes[AttributeIndex];
+				const FPCGMetadataAttribute<float>* OutAttribute = OutAttributes[AttributeIndex];
+				check(InheritedAttribute && OutAttribute);
+
+				const PCGMetadataEntryKey EntryKey = (InheritedAttributeModes[AttributeIndex] == EPCGCopyPointsMetadataInheritanceMode::Source) ? SourcePoint.MetadataEntry : TargetPoint.MetadataEntry;
+				TestEqual("Valid metadata value", InheritedAttribute->GetValueKey(EntryKey), OutAttribute->GetValueKey(OutPoint.MetadataEntry));
+			}
+		}
+	};
+
+	// Test 1 - inherit from transform multiplication
+	Settings->RotationInheritance = EPCGCopyPointsInheritanceMode::Relative;
+	Settings->ScaleInheritance = EPCGCopyPointsInheritanceMode::Relative;
+	Settings->ColorInheritance = EPCGCopyPointsInheritanceMode::Relative;
+	Settings->SeedInheritance = EPCGCopyPointsInheritanceMode::Relative;
+	ValidateCopyPointsPoints();
+
+	// Test 2 - inherit from source points
+	Settings->RotationInheritance = EPCGCopyPointsInheritanceMode::Source;
+	Settings->ScaleInheritance = EPCGCopyPointsInheritanceMode::Source;
+	Settings->ColorInheritance = EPCGCopyPointsInheritanceMode::Source;
+	Settings->SeedInheritance = EPCGCopyPointsInheritanceMode::Source;
+	Settings->AttributeInheritance = EPCGCopyPointsMetadataInheritanceMode::Source;
+	ValidateCopyPointsPoints();
+
+	// Test 3 - inherit from target points
+	Settings->RotationInheritance = EPCGCopyPointsInheritanceMode::Target;
+	Settings->ScaleInheritance = EPCGCopyPointsInheritanceMode::Target;
+	Settings->ColorInheritance = EPCGCopyPointsInheritanceMode::Target;
+	Settings->SeedInheritance = EPCGCopyPointsInheritanceMode::Target;
+	Settings->AttributeInheritance = EPCGCopyPointsMetadataInheritanceMode::Target;
+	ValidateCopyPointsPoints();
+
+	// Test 4 - empty source data
+	SourceTaggedData.Data = EmptyData;
+	ValidateCopyPointsPoints();
+	SourceTaggedData.Data = SourceData;
+
+	// Test 5 - empty target data
+	TargetTaggedData.Data = EmptyData;
+	ValidateCopyPointsPoints();
+	TargetTaggedData.Data = TargetData;
+
+	return true;
+}
diff --git a/Engine/Plugins/Experimental/PCG/Source/PCG/Private/Tests/PCGTestsCommon.cpp b/Engine/Plugins/Experimental/PCG/Source/PCG/Private/Tests/PCGTestsCommon.cpp
--- a/Engine/Plugins/Experimental/PCG/Source/PCG/Private/Tests/PCGTestsCommon.cpp
+++ b/Engine/Plugins/Experimental/PCG/Source/PCG/Private/Tests/PCGTestsCommon.cpp
@@ -2,6 +2,7 @@
 
 #include "Tests/PCGTestsCommon.h"
 
+#include "PCGHelpers.h"
 #include "PCGParamData.h"
 #include "PCGSettings.h"
 #include "Data/PCGPointData.h"
@@ -47,6 +48,30 @@ namespace PCGTestsCommon
 		return SinglePointData;
 	}
 
+	/** Creates a point data with PointCount many points, and randomizes the Transform and Color */
+	UPCGPointData* CreateRandomPointData(int32 PointCount, int32 Seed)
+	{
+		TObjectPtr<UPCGPointData> PointData = PCGTestsCommon::CreateEmptyPointData();
+		TArray<FPCGPoint>& Points = PointData->GetMutablePoints();
+
+		Points.Reserve(PointCount);
+	
+		FRandomStream RandomSource(Seed);
+		for (int I = 0; I < PointCount; ++I)
+		{
+			FQuat Rotation(FRotator(RandomSource.FRandRange(0.f, 360.f)).Quaternion()); 
+			FVector Scale(RandomSource.VRand());
+			FVector Location(RandomSource.VRand());
+
+			FPCGPoint& Point = Points.Emplace_GetRef(FTransform(Rotation, Location, Scale), 1.f, I);
+			Point.Color = RandomSource.VRand();
+			Point.Density = 1.f;
+			Point.Seed = I;
+		}
+
+		return PointData;
+	}
+
 	UPCGPolyLineData* CreatePolyLineData()
 	{
 		// TODO: spline, landscape spline
diff --git a/Engine/Plugins/Experimental/PCG/Source/PCG/Public/Elements/PCGCopyPoints.h b/Engine/Plugins/Experimental/PCG/Source/PCG/Public/Elements/PCGCopyPoints.h
new file mode 100644
--- /dev/null
+++ b/Engine/Plugins/Experimental/PCG/Source/PCG/Public/Elements/PCGCopyPoints.h
@@ -0,0 +1,77 @@
+// Copyright Epic Games, Inc. All Rights Reserved.
+
+#pragma once
+
+#include "CoreMinimal.h"
+#include "PCGSettings.h"
+
+#include "PCGCopyPoints.generated.h"
+
+namespace PCGCopyPointsConstants
+{
+	const FName ParamsLabel = TEXT("Params");
+	const FName SourcePointsLabel = TEXT("Source");
+	const FName TargetPointsLabel = TEXT("Target");
+}
+
+UENUM()
+enum class EPCGCopyPointsInheritanceMode : uint8
+{
+	Relative,
+	Source,
+	Target
+};
+
+UENUM()
+enum class EPCGCopyPointsMetadataInheritanceMode : uint8
+{
+	Source,
+	Target
+};
+
+UCLASS(BlueprintType, ClassGroup = (Procedural))
+class PCG_API UPCGCopyPointsSettings : public UPCGSettings
+{
+	GENERATED_BODY()
+
+public:
+	//~Begin UPCGSettings interface
+#if WITH_EDITOR
+	virtual FName GetDefaultNodeName() const override { return FName(TEXT("CopyPointsNode")); }
+	virtual EPCGSettingsType GetType() const override { return EPCGSettingsType::Sampler; }
+#endif
+
+	virtual TArray<FPCGPinProperties> InputPinProperties() const override;
+	virtual TArray<FPCGPinProperties> OutputPinProperties() const override { return Super::DefaultPointOutputPinProperties(); }
+
+protected:
+	virtual FPCGElementPtr CreateElement() const override;
+	//~End UPCGSettings interface
+
+public:
+	/** The method used to determine output point rotation */
+	UPROPERTY(BlueprintReadWrite, EditAnywhere, Category = Settings)
+	EPCGCopyPointsInheritanceMode RotationInheritance = EPCGCopyPointsInheritanceMode::Relative;
+
+	/** The method used to determine output point scale */
+	UPROPERTY(BlueprintReadWrite, EditAnywhere, Category = Settings)
+	EPCGCopyPointsInheritanceMode ScaleInheritance = EPCGCopyPointsInheritanceMode::Relative;
+
+	/** The method used to determine output point color */
+	UPROPERTY(BlueprintReadWrite, EditAnywhere, Category = Settings)
+	EPCGCopyPointsInheritanceMode ColorInheritance = EPCGCopyPointsInheritanceMode::Relative;
+
+	/** The method used to determine output seed values. Relative recomputes the seed from the new location. */
+	UPROPERTY(BlueprintReadWrite, EditAnywhere, Category = Settings)
+	EPCGCopyPointsInheritanceMode SeedInheritance = EPCGCopyPointsInheritanceMode::Relative;
+
+	/** The method used to determine output data attributes */
+	UPROPERTY(BlueprintReadWrite, EditAnywhere, Category = Settings)
+	EPCGCopyPointsMetadataInheritanceMode AttributeInheritance = EPCGCopyPointsMetadataInheritanceMode::Source;
+};
+
+class FPCGCopyPointsElement : public FSimplePCGElement
+{
+protected:
+	virtual bool ExecuteInternal(FPCGContext* Context) const override;
+};
diff --git a/Engine/Plugins/Experimental/PCG/Source/PCG/Public/Metadata/PCGMetadata.h b/Engine/Plugins/Experimental/PCG/Source/PCG/Public/Metadata/PCGMetadata.h
--- a/Engine/Plugins/Experimental/PCG/Source/PCG/Public/Metadata/PCGMetadata.h
+++ b/Engine/Plugins/Experimental/PCG/Source/PCG/Public/Metadata/PCGMetadata.h
@@ -56,6 +56,9 @@ public:
 	UFUNCTION(BlueprintCallable, Category = "PCG|Metadata|Advanced")
 	void CopyAttribute(const UPCGMetadata* InOther, FName AttributeToCopy, FName NewAttributeName);
 
+	/** Copies another attribute, with options to keep its parent and copy entries/values */
+	FPCGMetadataAttributeBase* CopyAttribute(const FPCGMetadataAttributeBase* OriginalAttribute, FName NewAttributeName, bool bKeepParent, bool bCopyEntries, bool bCopyValues);
+
 	/** Returns this metadata's parent */
 	const UPCGMetadata* GetParent() const { return Parent.Get(); }
 	const UPCGMetadata* GetRoot() const;
@@ -177,7 +180,6 @@ protected:
 	FPCGMetadataAttributeBase* CreateAttribute(FName AttributeName, const T& DefaultValue, bool bAllowsInterpolation, bool bOverrideParent);
 
 	FPCGMetadataAttributeBase* CopyAttribute(FName AttributeToCopy, FName NewAttributeName, bool bKeepParent, bool bCopyEntries, bool bCopyValues);
-	FPCGMetadataAttributeBase* CopyAttribute(const FPCGMetadataAttributeBase* OriginalAttribute, FName NewAttributeName, bool bKeepParent, bool bCopyEntries, bool bCopyValues);
 
 	bool ParentHasAttribute(FName AttributeName) const;
 	int64 GetItemCountForChild() const;
diff --git a/Engine/Plugins/Experimental/PCG/Source/PCG/Public/Tests/PCGTestsCommon.h b/Engine/Plugins/Experimental/PCG/Source/PCG/Public/Tests/PCGTestsCommon.h
--- a/Engine/Plugins/Experimental/PCG/Source/PCG/Public/Tests/PCGTestsCommon.h
+++ b/Engine/Plugins/Experimental/PCG/Source/PCG/Public/Tests/PCGTestsCommon.h
@@ -30,6 +30,9 @@ namespace PCGTestsCommon
 	/** Creates a PointData with a single point at the provided location */
 	UPCGPointData* CreatePointData(const FVector& InLocation);
 
+	/** Creates a PointData with PointCount many points, and randomizes the Transform and Color */
+	UPCGPointData* CreateRandomPointData(int32 PointCount, int32 Seed);
+
 	UPCGPolyLineData* CreatePolyLineData();
 	UPCGSurfaceData* CreateSurfaceData();
 	UPCGVolumeData* CreateVolumeData(const FBox& InBounds = FBox::BuildAABB(FVector::ZeroVector, FVector::OneVector * 100));
