Fixed issue with how we return the world pointer in PCG blueprint elements.
Fixed issue with the resources not being cleaned up properly at runtime since they were running in the Post-Graph which was editor-only before.

#rnx

diff --git a/Engine/Plugins/Experimental/PCG/Source/PCG/Private/Elements/PCGExecuteBlueprint.cpp b/Engine/Plugins/Experimental/PCG/Source/PCG/Private/Elements/PCGExecuteBlueprint.cpp
--- a/Engine/Plugins/Experimental/PCG/Source/PCG/Private/Elements/PCGExecuteBlueprint.cpp
+++ b/Engine/Plugins/Experimental/PCG/Source/PCG/Private/Elements/PCGExecuteBlueprint.cpp
@@ -1,15 +1,16 @@
 // Copyright Epic Games, Inc. All Rights Reserved.
 
 #include "Elements/PCGExecuteBlueprint.h"
+#include "PCGComponent.h"
 #include "PCGHelpers.h"
 #include "Data/PCGPointData.h"
 #include "Data/PCGSpatialData.h"
 #include "Helpers/PCGAsync.h"
 #include "Helpers/PCGSettingsHelpers.h"
 
-#if WITH_EDITOR
 #include "Engine/World.h"
 
+#if WITH_EDITOR
 namespace PCGBlueprintHelper
 {
 	TSet<TObjectPtr<UObject>> GetDataDependencies(UPCGBlueprintElement* InElement)
@@ -29,7 +30,7 @@ UWorld* UPCGBlueprintElement::GetWorld() const
 #if WITH_EDITOR
 	return GWorld;
 #else
-	return nullptr;
+	return InstanceWorld ? InstanceWorld : Super::GetWorld();
 #endif
 }
 
@@ -621,6 +622,9 @@ FPCGContext* FPCGExecuteBlueprintElement::Initialize(const FPCGDataCollection& I
 	if (Settings && Settings->BlueprintElementInstance)
 	{
 		Context->BlueprintElementInstance = CastChecked<UPCGBlueprintElement>(StaticDuplicateObject(Settings->BlueprintElementInstance, GetTransientPackage(), FName()));
+#if !WITH_EDITOR
+		Context->BlueprintElementInstance->SetInstanceWorld(SourceComponent ? SourceComponent->GetOwner()->GetWorld() : nullptr);
+#endif
 	}
 	else
 	{
diff --git a/Engine/Plugins/Experimental/PCG/Source/PCG/Private/Graph/PCGGraphCompiler.cpp b/Engine/Plugins/Experimental/PCG/Source/PCG/Private/Graph/PCGGraphCompiler.cpp
--- a/Engine/Plugins/Experimental/PCG/Source/PCG/Private/Graph/PCGGraphCompiler.cpp
+++ b/Engine/Plugins/Experimental/PCG/Source/PCG/Private/Graph/PCGGraphCompiler.cpp
@@ -221,7 +221,6 @@ TArray<FPCGGraphTask> FPCGGraphCompiler::GetCompiledTasks(UPCGGraph* InGraph, bo
 {
 	TArray<FPCGGraphTask> CompiledTasks;
 
-#if WITH_EDITOR
 	if (bIsTopGraph)
 	{
 		// Always try to compile
@@ -236,7 +235,6 @@ TArray<FPCGGraphTask> FPCGGraphCompiler::GetCompiledTasks(UPCGGraph* InGraph, bo
 		GraphToTaskMapLock.ReadUnlock();
 	}
 	else
-#endif
 	{
 		// Always try to compile
 		Compile(InGraph);
@@ -266,7 +264,6 @@ void FPCGGraphCompiler::OffsetNodeIds(TArray<FPCGGraphTask>& Tasks, FPCGTaskId O
 	}
 }
 
-#if WITH_EDITOR
 void FPCGGraphCompiler::CompileTopGraph(UPCGGraph* InGraph)
 {
 	GraphToTaskMapLock.ReadLock();
@@ -338,6 +335,7 @@ void FPCGGraphCompiler::CompileTopGraph(UPCGGraph* InGraph)
 	GraphToTaskMapLock.WriteUnlock();
 }
 
+#if WITH_EDITOR
 void FPCGGraphCompiler::NotifyGraphChanged(UPCGGraph* InGraph)
 {
 	if (InGraph)
@@ -356,9 +354,7 @@ void FPCGGraphCompiler::RemoveFromCacheRecursive(UPCGGraph* InGraph)
 {
 	GraphToTaskMapLock.WriteLock();
 	GraphToTaskMap.Remove(InGraph);
-#if WITH_EDITOR
 	TopGraphToTaskMap.Remove(InGraph);
-#endif
 	GraphToTaskMapLock.WriteUnlock();
 
 	GraphDependenciesLock.Lock();
diff --git a/Engine/Plugins/Experimental/PCG/Source/PCG/Private/Graph/PCGGraphExecutor.cpp b/Engine/Plugins/Experimental/PCG/Source/PCG/Private/Graph/PCGGraphExecutor.cpp
--- a/Engine/Plugins/Experimental/PCG/Source/PCG/Private/Graph/PCGGraphExecutor.cpp
+++ b/Engine/Plugins/Experimental/PCG/Source/PCG/Private/Graph/PCGGraphExecutor.cpp
@@ -83,7 +83,6 @@ FPCGTaskId FPCGGraphExecutor::Schedule(UPCGGraph* Graph, UPCGComponent* SourceCo
 		NextTaskId += ScheduledTask.Tasks.Num();
 		ScheduledId = NextTaskId - 1; // This is true because the last task is from the output node or is the post-execute task
 
-#if WITH_EDITOR
 		// Push task (not data) dependencies on the pre-execute task
 		// Note must be done after the offset ids, otherwise we'll break the dependencies
 		check(ScheduledTask.Tasks.Num() >= 2 && ScheduledTask.Tasks[ScheduledTask.Tasks.Num() - 2].Node == nullptr);
@@ -91,7 +90,6 @@ FPCGTaskId FPCGGraphExecutor::Schedule(UPCGGraph* Graph, UPCGComponent* SourceCo
 		{
 			ScheduledTask.Tasks[ScheduledTask.Tasks.Num() - 2].Inputs.Emplace(ExternalDependency, nullptr, nullptr);
 		}
-#endif
 
 		ScheduleLock.Unlock();
 	}
@@ -563,7 +561,7 @@ void FPCGGraphExecutor::NotifyGraphChanged(UPCGGraph* InGraph)
 		GraphCompiler->NotifyGraphChanged(InGraph);
 	}
 }
-#endif
+#endif // WITH_EDITOR
 
 bool FPCGFetchInputElement::ExecuteInternal(FPCGContext* Context) const
 {
diff --git a/Engine/Plugins/Experimental/PCG/Source/PCG/Private/PCGComponent.cpp b/Engine/Plugins/Experimental/PCG/Source/PCG/Private/PCGComponent.cpp
--- a/Engine/Plugins/Experimental/PCG/Source/PCG/Private/PCGComponent.cpp
+++ b/Engine/Plugins/Experimental/PCG/Source/PCG/Private/PCGComponent.cpp
@@ -185,8 +185,6 @@ void UPCGComponent::Generate()
 
 void UPCGComponent::Generate_Implementation(bool bForce)
 {
-	// Force component activation so it's easier to control by BP.
-	bActivated = true;
 	GenerateLocal(bForce);
 }
 
@@ -199,6 +197,13 @@ void UPCGComponent::GenerateLocal(bool bForce)
 	}
 #endif
 
+	// Force component activation so it's easier to control by BP.
+	if (!bActivated)
+	{
+		Modify();
+		bActivated = true;
+	}
+
 	FPCGTaskId TaskId = GenerateInternal(bForce, {});
 
 #if WITH_EDITOR
@@ -1669,13 +1674,11 @@ FBox UPCGComponent::GetGridBounds(AActor* Actor) const
 			}
 		}
 	}
-	// TODO: this might not work in non-editor builds
-#if WITH_EDITOR
+	// TODO: verify this works as expected in non-editor builds
 	else if (ALandscapeProxy* LandscapeActor = Cast<ALandscape>(Actor))
 	{
 		Bounds = PCGHelpers::GetLandscapeBounds(LandscapeActor);
 	}
-#endif
 	else
 	{
 		Bounds = PCGHelpers::GetActorBounds(Actor);
diff --git a/Engine/Plugins/Experimental/PCG/Source/PCG/Private/PCGSubsystem.cpp b/Engine/Plugins/Experimental/PCG/Source/PCG/Private/PCGSubsystem.cpp
--- a/Engine/Plugins/Experimental/PCG/Source/PCG/Private/PCGSubsystem.cpp
+++ b/Engine/Plugins/Experimental/PCG/Source/PCG/Private/PCGSubsystem.cpp
@@ -103,7 +103,6 @@ FPCGTaskId UPCGSubsystem::ScheduleComponent(UPCGComponent* PCGComponent, const T
 
 	if (ProcessTaskId != InvalidPCGTaskId)
 	{
-#if WITH_EDITOR
 		TWeakObjectPtr<UPCGComponent> ComponentPtr(PCGComponent);
 
 		return GraphExecutor->ScheduleGeneric([ComponentPtr]() {
@@ -115,9 +114,6 @@ FPCGTaskId UPCGSubsystem::ScheduleComponent(UPCGComponent* PCGComponent, const T
 
 			return true;
 			}, { ProcessTaskId });
-#else
-		return ProcessTaskId;
-#endif
 	}
 	else
 	{
diff --git a/Engine/Plugins/Experimental/PCG/Source/PCG/Public/Elements/PCGExecuteBlueprint.h b/Engine/Plugins/Experimental/PCG/Source/PCG/Public/Elements/PCGExecuteBlueprint.h
--- a/Engine/Plugins/Experimental/PCG/Source/PCG/Public/Elements/PCGExecuteBlueprint.h
+++ b/Engine/Plugins/Experimental/PCG/Source/PCG/Public/Elements/PCGExecuteBlueprint.h
@@ -96,6 +96,10 @@ public:
 	/** Needed to be able to call certain blueprint functions */
 	virtual UWorld* GetWorld() const override;
 
+#if !WITH_EDITOR
+	void SetInstanceWorld(UWorld* World) { InstanceWorld = World; }
+#endif
+
 #if WITH_EDITOR
 	FOnPCGBlueprintChanged OnBlueprintChangedDelegate;
 #endif
@@ -140,6 +144,10 @@ protected:
 	void OnDependencyChanged(UObject* Object, FPropertyChangedEvent& PropertyChangedEvent);
 	TSet<TObjectPtr<UObject>> DataDependencies;
 #endif
+
+#if !WITH_EDITORONLY_DATA
+	UWorld* InstanceWorld = nullptr;
+#endif
 };
 
 UCLASS(BlueprintType, ClassGroup = (Procedural))
diff --git a/Engine/Plugins/Experimental/PCG/Source/PCG/Public/PCGComponent.h b/Engine/Plugins/Experimental/PCG/Source/PCG/Public/PCGComponent.h
--- a/Engine/Plugins/Experimental/PCG/Source/PCG/Public/PCGComponent.h
+++ b/Engine/Plugins/Experimental/PCG/Source/PCG/Public/PCGComponent.h
@@ -85,8 +85,11 @@ public:
 	void Cleanup();
 
 	/** Starts generation from a local (vs. remote) standpoint. Will not be replicated. */
+	UFUNCTION(BlueprintCallable, Category = PCG)
 	void GenerateLocal(bool bForce);
+
 	/** Cleans up the generation from a local (vs. remote) standpoint. Will not be replicated. */
+	UFUNCTION(BlueprintCallable, Category = PCG)
 	void CleanupLocal(bool bRemoveComponents, bool bSave = false);
 
 	/** Networked generation call that also activates the component as needed */
