PCG:
*Added ability to debug pcg components and inspect various node data in an attribute list view.

diff --git a/Engine/Plugins/Experimental/PCG/Source/PCG/Private/Graph/PCGGraphExecutor.cpp b/Engine/Plugins/Experimental/PCG/Source/PCG/Private/Graph/PCGGraphExecutor.cpp
--- a/Engine/Plugins/Experimental/PCG/Source/PCG/Private/Graph/PCGGraphExecutor.cpp
+++ b/Engine/Plugins/Experimental/PCG/Source/PCG/Private/Graph/PCGGraphExecutor.cpp
@@ -236,6 +236,14 @@ void FPCGGraphExecutor::Execute()
 					StoreResults(Task.NodeId, CachedOutput);
 					ReadyTasks.RemoveAtSwap(ReadyTaskIndex);
 					bSomeTaskEndedInCurrentLoop = true;
+#if WITH_EDITOR
+					UPCGComponent* SourceComponent = Task.SourceComponent;
+					if (SourceComponent && SourceComponent->IsInspecting())
+					{
+						SourceComponent->StoreInspectionData(Task.Node, CachedOutput);
+					}
+#endif
+					
 					continue;
 				}
 
@@ -342,6 +350,12 @@ void FPCGGraphExecutor::Execute()
 			// Execute debug display code as needed - done here because it needs to be done on the main thread
 			// Additional note: this needs to be executed before the StoreResults since debugging might cancel further tasks
 			ActiveTask.Element->DebugDisplay(ActiveTask.Context.Get());
+
+			UPCGComponent* SourceComponent = ActiveTask.Context->SourceComponent;
+			if (SourceComponent && SourceComponent->IsInspecting())
+			{
+				SourceComponent->StoreInspectionData(ActiveTask.Context->Node, ActiveTask.Context->OutputData);
+			}
 #endif
 
 			// Store output in data map
diff --git a/Engine/Plugins/Experimental/PCG/Source/PCG/Private/PCGComponent.cpp b/Engine/Plugins/Experimental/PCG/Source/PCG/Private/PCGComponent.cpp
--- a/Engine/Plugins/Experimental/PCG/Source/PCG/Private/PCGComponent.cpp
+++ b/Engine/Plugins/Experimental/PCG/Source/PCG/Private/PCGComponent.cpp
@@ -1129,6 +1129,10 @@ void UPCGComponent::OnGraphChanged(UPCGGraph* InGraph, bool bIsStructural, bool
 		{
 			Refresh();
 		}
+
+#if WITH_EDITOR
+		InspectionCache.Empty();	
+#endif
 	}
 }
 
@@ -1195,6 +1199,34 @@ void UPCGComponent::ResetLastGeneratedBounds()
 	LastGeneratedBounds = FBox(EForceInit::ForceInit);
 }
 
+#if WITH_EDITOR
+void UPCGComponent::DisableInspection()
+{
+	bIsInspecting = false;
+	InspectionCache.Empty();
+};
+
+void UPCGComponent::StoreInspectionData(const UPCGNode* InNode, const FPCGDataCollection& InInspectionData)
+{
+	if (!InNode)
+	{
+		return;
+	}
+
+	if (GetGraph() != InNode->GetGraph())
+	{
+		return;
+	}
+
+	InspectionCache.Add(InNode, InInspectionData);
+}
+
+const FPCGDataCollection* UPCGComponent::GetInspectionData(const UPCGNode* InNode) const
+{
+	return InspectionCache.Find(InNode);
+}
+#endif
+
 void UPCGComponent::Refresh()
 {
 #if WITH_EDITOR
@@ -1874,4 +1906,4 @@ bool UPCGComponent::GraphUsesLandscapePin() const
 
 #endif // WITH_EDITOR
 
-#undef LOCTEXT_NAMESPACE
\ No newline at end of file
+#undef LOCTEXT_NAMESPACE
diff --git a/Engine/Plugins/Experimental/PCG/Source/PCG/Public/PCGComponent.h b/Engine/Plugins/Experimental/PCG/Source/PCG/Public/PCGComponent.h
--- a/Engine/Plugins/Experimental/PCG/Source/PCG/Public/PCGComponent.h
+++ b/Engine/Plugins/Experimental/PCG/Source/PCG/Public/PCGComponent.h
@@ -151,6 +151,13 @@ public:
 
 	/** Return if we are currently generating the graph for this component */
 	bool IsGenerating() const { return bIsGenerating; }
+
+	/** Functions for managing the node inspection cache */
+	bool IsInspecting() const { return bIsInspecting; }
+	void EnableInspection() { bIsInspecting = true; };
+	void DisableInspection();
+	void StoreInspectionData(const UPCGNode* InNode, const FPCGDataCollection& InInspectionData);
+	const FPCGDataCollection* GetInspectionData(const UPCGNode* InNode) const;
 #endif
 
 	bool IsPartitioned() const;
@@ -256,6 +263,7 @@ private:
 
 #if WITH_EDITOR
 	bool bIsGenerating = false;
+	bool bIsInspecting = false;
 	FBox LastGeneratedBoundsPriorToUndo = FBox(EForceInit::ForceInit);
 	FPCGTagToSettingsMap CachedTrackedTagsToSettings;
 
@@ -278,7 +286,10 @@ private:
 	TMap<TWeakObjectPtr<AActor>, TSet<FName>> CachedTrackedActorToTags;
 	TMap<TWeakObjectPtr<AActor>, TSet<TObjectPtr<UObject>>> CachedTrackedActorToDependencies;
 	bool bActorToTagsMapPopulated = false;
+
+	UPROPERTY(Transient)
+	TMap<const UPCGNode*, FPCGDataCollection> InspectionCache;
 #endif
 
 	FCriticalSection GeneratedResourcesLock;
-};
\ No newline at end of file
+};
diff --git a/Engine/Plugins/Experimental/PCG/Source/PCGEditor/PCGEditor.Build.cs b/Engine/Plugins/Experimental/PCG/Source/PCGEditor/PCGEditor.Build.cs
--- a/Engine/Plugins/Experimental/PCG/Source/PCGEditor/PCGEditor.Build.cs
+++ b/Engine/Plugins/Experimental/PCG/Source/PCGEditor/PCGEditor.Build.cs
@@ -36,6 +36,7 @@ namespace UnrealBuildTool.Rules
 					"InputCore",
 					"Kismet",
 					"PCG",
+					"PropertyEditor",
 					"Slate",
 					"SlateCore",
 					"SourceControl",
diff --git a/Engine/Plugins/Experimental/PCG/Source/PCGEditor/Private/PCGEditor.cpp b/Engine/Plugins/Experimental/PCG/Source/PCGEditor/Private/PCGEditor.cpp
--- a/Engine/Plugins/Experimental/PCG/Source/PCGEditor/Private/PCGEditor.cpp
+++ b/Engine/Plugins/Experimental/PCG/Source/PCGEditor/Private/PCGEditor.cpp
@@ -2,12 +2,14 @@
 
 #include "PCGEditor.h"
 
+#include "PCGComponent.h"
 #include "PCGEdge.h"
 #include "PCGGraph.h"
 #include "PCGGraphFactory.h"
 #include "PCGInputOutputSettings.h"
 #include "PCGPin.h"
 #include "PCGSubgraph.h"
+#include "Tests/Determinism/PCGDeterminismTestsCommon.h"
 
 #include "PCGEditorCommands.h"
 #include "PCGEditorCommon.h"
@@ -17,39 +19,41 @@
 #include "PCGEditorGraphNodeOutput.h"
 #include "PCGEditorGraphSchema.h"
 #include "PCGEditorGraphSchemaActions.h"
+#include "PCGEditorMenuContext.h"
 #include "PCGEditorSettings.h"
 #include "PCGEditorUtils.h"
+#include "SPCGEditorGraphAttributeListView.h"
+#include "SPCGEditorGraphDebugObjectWidget.h"
 #include "SPCGEditorGraphDeterminism.h"
 #include "SPCGEditorGraphFind.h"
 #include "SPCGEditorGraphNodePalette.h"
-#include "Tests/Determinism/PCGDeterminismTestsCommon.h"
 
+#include "AssetToolsModule.h"
 #include "EdGraphUtilities.h"
 #include "Editor.h"
+#include "Editor/UnrealEdEngine.h"
+#include "EditorAssetLibrary.h"
 #include "Framework/Commands/GenericCommands.h"
 #include "Framework/Commands/UIAction.h"
 #include "Framework/Commands/UICommandList.h"
 #include "Framework/Docking/TabManager.h"
 #include "Framework/MultiBox/MultiBoxBuilder.h"
 #include "Framework/Notifications/NotificationManager.h"
+#include "GraphEditAction.h"
 #include "GraphEditor.h"
 #include "GraphEditorActions.h"
 #include "HAL/PlatformApplicationMisc.h"
 #include "IDetailsView.h"
 #include "Modules/ModuleManager.h"
+#include "Preferences/UnrealEdOptions.h"
 #include "PropertyEditorModule.h"
 #include "SNodePanel.h"
 #include "ScopedTransaction.h"
-#include "ToolMenus.h"
-#include "Widgets/Docking/SDockTab.h"
 #include "SourceCodeNavigation.h"
+#include "ToolMenus.h"
 #include "UnrealEdGlobals.h"
-#include "Editor/UnrealEdEngine.h"
-#include "Preferences/UnrealEdOptions.h"
+#include "Widgets/Docking/SDockTab.h"
 #include "Widgets/Notifications/SNotificationList.h"
-#include "GraphEditAction.h"
-#include "AssetToolsModule.h"
-#include "EditorAssetLibrary.h"
 
 #define LOCTEXT_NAMESPACE "PCGGraphEditor"
 
@@ -82,9 +86,11 @@ void FPCGEditor::Initialize(const EToolkitMode::Type InMode, const TSharedPtr<cl
 	GraphEditorWidget = CreateGraphEditorWidget();
 	PaletteWidget = CreatePaletteWidget();
 	FindWidget = CreateFindWidget();
+	AttributesWidget = CreateAttributesWidget();
 	DeterminismWidget = CreateDeterminismWidget();
 
 	BindCommands();
+	RegisterToolbar();
 
 	const TSharedRef<FTabManager::FLayout> StandaloneDefaultLayout = FTabManager::NewLayout("Standalone_PCGGraphEditor_Layout_v0.4")
 		->AddArea
@@ -137,6 +143,31 @@ UPCGEditorGraph* FPCGEditor::GetPCGEditorGraph()
 	return PCGEditorGraph;
 }
 
+void FPCGEditor::SetPCGComponentBeingDebugged(UPCGComponent* InPCGComponent)
+{
+	if (PCGComponentBeingDebugged != InPCGComponent)
+	{
+		PCGComponentBeingDebugged = InPCGComponent;
+
+		OnDebugObjectChangedDelegate.Broadcast(PCGComponentBeingDebugged); 
+
+		if (PCGComponentBeingDebugged)
+		{
+			// We need to force generation so that we create debug data
+			PCGComponentBeingDebugged->GenerateLocal(/*bForce=*/true);
+		}
+	}
+}
+
+void FPCGEditor::SetPCGNodeBeingInspected(UPCGNode* InPCGNode)
+{
+	if (GetPCGNodeBeingInspected() != InPCGNode)
+	{
+		PCGNodeBeingInspected = InPCGNode;
+		OnInspectedNodeChangedDelegate.Broadcast(PCGNodeBeingInspected);	
+	}
+}
+
 void FPCGEditor::JumpToNode(const UEdGraphNode* InNode)
 {
 	if (GraphEditorWidget.IsValid())
@@ -234,6 +265,35 @@ FString FPCGEditor::GetWorldCentricTabPrefix() const
 	return LOCTEXT("WorldCentricTabPrefix", "PCG ").ToString();
 }
 
+void FPCGEditor::RegisterToolbar() const
+{
+	UToolMenus* ToolMenus = UToolMenus::Get();
+	FName ParentName;
+	const FName MenuName = GetToolMenuToolbarName(ParentName);
+
+	UToolMenu* ToolBar = ToolMenus->RegisterMenu(MenuName, ParentName, EMultiBoxType::ToolBar);
+
+	const FPCGEditorCommands& PCGEditorCommands = FPCGEditorCommands::Get();
+	const FToolMenuInsert InsertAfterAssetSection("Asset", EToolMenuInsertType::After);
+	{
+		FToolMenuSection& Section = ToolBar->AddSection("PCGToolbar", TAttribute<FText>(), InsertAfterAssetSection);
+		Section.AddEntry(FToolMenuEntry::InitToolBarButton(
+ 			PCGEditorCommands.Find,
+ 			TAttribute<FText>(),
+ 			TAttribute<FText>(),
+ 			FSlateIcon(FAppStyle::GetAppStyleSetName(), "BlueprintEditor.FindInBlueprint"))); // TODO, use own application style?
+		Section.AddSeparator(NAME_None);
+		Section.AddDynamicEntry("Debugging", FNewToolMenuSectionDelegate::CreateLambda([](FToolMenuSection& InSection)
+		{
+			const UPCGEditorMenuContext* Context = InSection.FindContext<UPCGEditorMenuContext>();
+			if (Context && Context->PCGEditor.IsValid())
+			{
+				InSection.AddEntry(FToolMenuEntry::InitWidget("SelectedDebugObjectWidget", SNew(SPCGEditorGraphDebugObjectWidget, Context->PCGEditor.Pin()), FText::GetEmpty()));	
+			}
+		}));
+	}
+}
+
 void FPCGEditor::BindCommands()
 {
 	const FPCGEditorCommands& PCGEditorCommands = FPCGEditorCommands::Get();
@@ -247,6 +307,14 @@ void FPCGEditor::BindCommands()
 		FExecuteAction::CreateSP(this, &FPCGEditor::OnCollapseNodesInSubgraph),
 		FCanExecuteAction::CreateSP(this, &FPCGEditor::CanCollapseNodesInSubgraph));
 
+	GraphEditorCommands->MapAction(
+		PCGEditorCommands.StartInspectNode,
+		FExecuteAction::CreateSP(this, &FPCGEditor::OnStartInspectNode));
+
+	GraphEditorCommands->MapAction(
+		PCGEditorCommands.StopInspectNode,
+		FExecuteAction::CreateSP(this, &FPCGEditor::OnStopInspectNode));
+
 	GraphEditorCommands->MapAction(
 		PCGEditorCommands.RunDeterminismTest,
 		FExecuteAction::CreateSP(this, &FPCGEditor::OnDeterminismTests),
@@ -328,7 +396,7 @@ void FPCGEditor::OnDeterminismTests()
 	}
 }
 
-bool FPCGEditor::CanCollapseNodesInSubgraph()
+bool FPCGEditor::CanCollapseNodesInSubgraph() const
 {
 	bool HasPCGNode = false;
 
@@ -704,6 +772,50 @@ void FPCGEditor::OnCollapseNodesInSubgraph()
 	NewPCGGraph->NotifyGraphChanged(true);
 }
 
+void FPCGEditor::OnStartInspectNode()
+{
+	if (!GraphEditorWidget.IsValid())
+	{
+		return;
+	}
+
+	UEdGraphNode* GraphNode = GraphEditorWidget->GetSingleSelectedNode();
+	if (!GraphNode)
+	{
+		return;
+	}
+
+	UPCGEditorGraphNodeBase* PCGGraphNodeBase = Cast<UPCGEditorGraphNodeBase>(GraphNode);
+	if (!PCGGraphNodeBase)
+	{
+		return;
+	}
+	
+	if (PCGGraphNodeBase == PCGGraphNodeBeingInspected)
+	{
+		return;
+	}
+	
+	if (PCGGraphNodeBeingInspected)
+	{
+		PCGGraphNodeBeingInspected->SetInspected(false);
+	}	
+
+	PCGGraphNodeBeingInspected = PCGGraphNodeBase;
+	PCGGraphNodeBeingInspected->SetInspected(true);
+	SetPCGNodeBeingInspected(PCGGraphNodeBase->GetPCGNode());
+}
+
+void FPCGEditor::OnStopInspectNode()
+{
+	if (PCGGraphNodeBeingInspected)
+	{
+		PCGGraphNodeBeingInspected->SetInspected(false);
+		PCGGraphNodeBeingInspected = nullptr;
+		SetPCGNodeBeingInspected(nullptr);
+	}
+}
+
 void FPCGEditor::SelectAllNodes()
 {
 	if (GraphEditorWidget.IsValid())
@@ -1151,6 +1263,15 @@ void FPCGEditor::OnClose()
 	FAssetEditorToolkit::OnClose();
 }
 
+void FPCGEditor::InitToolMenuContext(FToolMenuContext& MenuContext)
+{
+	FAssetEditorToolkit::InitToolMenuContext(MenuContext);
+
+	UPCGEditorMenuContext* Context = NewObject<UPCGEditorMenuContext>();
+	Context->PCGEditor = SharedThis(this);
+	MenuContext.AddObject(Context);
+}
+
 TSharedRef<SPCGEditorGraphNodePalette> FPCGEditor::CreatePaletteWidget()
 {
 	return SNew(SPCGEditorGraphNodePalette);
@@ -1161,6 +1282,11 @@ TSharedRef<SPCGEditorGraphFind> FPCGEditor::CreateFindWidget()
 	return SNew(SPCGEditorGraphFind, SharedThis(this));
 }
 
+TSharedRef<SPCGEditorGraphAttributeListView> FPCGEditor::CreateAttributesWidget()
+{
+	return SNew(SPCGEditorGraphAttributeListView, SharedThis(this));
+}
+
 TSharedRef<SPCGEditorGraphDeterminismListView> FPCGEditor::CreateDeterminismWidget()
 {
 	TArray<PCGDeterminismTests::FNodeTestInfo> BasicTests;
@@ -1305,7 +1431,7 @@ TSharedRef<SDockTab> FPCGEditor::SpawnTab_Attributes(const FSpawnTabArgs& Args)
 		.Label(LOCTEXT("PCGAttributesTitle", "Attributes"))
 		.TabColorScale(GetTabColorScale())
 		[
-			SNullWidget::NullWidget
+			AttributesWidget.ToSharedRef()
 		];
 }
 
diff --git a/Engine/Plugins/Experimental/PCG/Source/PCGEditor/Private/PCGEditor.h b/Engine/Plugins/Experimental/PCG/Source/PCGEditor/Private/PCGEditor.h
--- a/Engine/Plugins/Experimental/PCG/Source/PCGEditor/Private/PCGEditor.h
+++ b/Engine/Plugins/Experimental/PCG/Source/PCGEditor/Private/PCGEditor.h
@@ -9,11 +9,18 @@
 class FUICommandList;
 class IDetailsView;
 class SGraphEditor;
+class SPCGEditorGraphAttributeListView;
 class SPCGEditorGraphDeterminismListView;
 class SPCGEditorGraphFind;
 class SPCGEditorGraphNodePalette;
+class UPCGComponent;
 class UPCGEditorGraph;
+class UPCGEditorGraphNodeBase;
 class UPCGGraph;
+class UPCGNode;
+
+DECLARE_MULTICAST_DELEGATE_OneParam(FOnDebugObjectChanged, UPCGComponent*);
+DECLARE_MULTICAST_DELEGATE_OneParam(FOnInspectedNodeChanged, UPCGNode*);
 
 class FPCGEditor : public FAssetEditorToolkit, public FSelfRegisteringEditorUndoClient
 {
@@ -24,6 +31,18 @@ public:
 	/** Get the PCG graph being edited */
 	UPCGEditorGraph* GetPCGEditorGraph();
 
+	/** Sets the PCG component we want to debug */
+	void SetPCGComponentBeingDebugged(UPCGComponent* InPCGComponent);
+	
+	/** Gets the PCG component we are debugging */
+	UPCGComponent* GetPCGComponentBeingDebugged() const { return PCGComponentBeingDebugged; }
+
+	/** Sets the PCG node we want to inspect */
+	void SetPCGNodeBeingInspected(UPCGNode* InPCGNode);
+	
+	/** Gets the PCG node we are inspecting */
+	UPCGNode* GetPCGNodeBeingInspected() const { return PCGNodeBeingInspected; }
+
 	/** Focus the graph view on a specific node */
 	void JumpToNode(const UEdGraphNode* InNode);
 
@@ -44,7 +63,11 @@ public:
 	virtual FLinearColor GetWorldCentricTabColorScale() const override;
 	virtual FString GetWorldCentricTabPrefix() const override;
 	virtual void OnClose() override;
+	virtual void InitToolMenuContext(FToolMenuContext& MenuContext) override;
 	// ~End FAssetEditorToolkit interface
+	
+	FOnDebugObjectChanged OnDebugObjectChangedDelegate;
+	FOnInspectedNodeChanged OnInspectedNodeChangedDelegate;
 
 protected:
 	// ~Begin FAssetEditorToolkit interface
@@ -53,12 +76,20 @@ protected:
 	// ~End FAssetEditorToolkit interface
 
 private:
+	/** Register PCG specific toolbar for the editor */
+	void RegisterToolbar() const;
+
 	/** Bind commands to delegates */
 	void BindCommands();
 
 	/** Bring up the find tab */
 	void OnFind();
 
+	/** Start inspecting the current selected node */
+	void OnStartInspectNode();
+	/** Stop inspecting the current inspected node */
+	void OnStopInspectNode();
+
 	/** Can determinism be tested on this node */
 	bool CanRunDeterminismTests() const;
 	/** Summon the Determinism tab */
@@ -99,7 +130,7 @@ private:
 	/** Collapse the currently selected nodes in a subgraph */
 	void OnCollapseNodesInSubgraph();
 	/** Whether we can collapse nodes in a subgraph */
-	bool CanCollapseNodesInSubgraph();
+	bool CanCollapseNodesInSubgraph() const;
 
 	void OnAlignTop();
 	void OnAlignMiddle();
@@ -121,6 +152,9 @@ private:
 	/** Create new find widget */
 	TSharedRef<SPCGEditorGraphFind> CreateFindWidget();
 
+	/** Create new attributes widget */
+	TSharedRef<SPCGEditorGraphAttributeListView> CreateAttributesWidget();
+
 	/** Create a new determinism tab widget */
 	TSharedRef<SPCGEditorGraphDeterminismListView> CreateDeterminismWidget();
 
@@ -158,10 +192,15 @@ private:
 	TSharedPtr<IDetailsView> PropertyDetailsWidget;
 	TSharedPtr<SPCGEditorGraphNodePalette> PaletteWidget;
 	TSharedPtr<SPCGEditorGraphFind> FindWidget;
+	TSharedPtr<SPCGEditorGraphAttributeListView> AttributesWidget;
 	TSharedPtr<SPCGEditorGraphDeterminismListView> DeterminismWidget;
 
 	TSharedPtr<FUICommandList> GraphEditorCommands;
 
 	UPCGGraph* PCGGraphBeingEdited = nullptr;
 	UPCGEditorGraph* PCGEditorGraph = nullptr;
+
+	UPCGComponent* PCGComponentBeingDebugged = nullptr;
+	UPCGNode* PCGNodeBeingInspected = nullptr;
+	UPCGEditorGraphNodeBase* PCGGraphNodeBeingInspected = nullptr;
 };
diff --git a/Engine/Plugins/Experimental/PCG/Source/PCGEditor/Private/PCGEditorCommands.cpp b/Engine/Plugins/Experimental/PCG/Source/PCGEditor/Private/PCGEditorCommands.cpp
--- a/Engine/Plugins/Experimental/PCG/Source/PCGEditor/Private/PCGEditorCommands.cpp
+++ b/Engine/Plugins/Experimental/PCG/Source/PCGEditor/Private/PCGEditorCommands.cpp
@@ -18,9 +18,11 @@ FPCGEditorCommands::FPCGEditorCommands()
 
 void FPCGEditorCommands::RegisterCommands()
 {
-	UI_COMMAND(Find, "Find", "Finds PCG nodes and comments in the current graph", EUserInterfaceActionType::Button, FInputChord(EModifierKey::Control, EKeys::F));
 	UI_COMMAND(CollapseNodes, "Collapse into Subgraph", "Collapse selected nodes into a separate PCGGraph asset.", EUserInterfaceActionType::Button, FInputChord(EModifierKey::Alt, EKeys::J));
+	UI_COMMAND(Find, "Find", "Finds PCG nodes and comments in the current graph", EUserInterfaceActionType::Button, FInputChord(EModifierKey::Control, EKeys::F));
 	UI_COMMAND(RunDeterminismTest, "Run Determinism Test", "Evaluate the current node for determinism.", EUserInterfaceActionType::Button, FInputChord(EModifierKey::Alt, EKeys::D));
+	UI_COMMAND(StartInspectNode, "Start Inspecting Node", "Start Inspecting Node", EUserInterfaceActionType::Button, FInputChord(EKeys::I));
+	UI_COMMAND(StopInspectNode, "Stop Inspecting Node", "Stop Inspecting Node", EUserInterfaceActionType::Button, FInputChord(EModifierKey::Shift, EKeys::I));
 }
 
 #undef LOCTEXT_NAMESPACE
diff --git a/Engine/Plugins/Experimental/PCG/Source/PCGEditor/Private/PCGEditorCommands.h b/Engine/Plugins/Experimental/PCG/Source/PCGEditor/Private/PCGEditorCommands.h
--- a/Engine/Plugins/Experimental/PCG/Source/PCGEditor/Private/PCGEditorCommands.h
+++ b/Engine/Plugins/Experimental/PCG/Source/PCGEditor/Private/PCGEditorCommands.h
@@ -14,7 +14,9 @@ public:
 	virtual void RegisterCommands() override;
 	// ~End TCommands<> interface
 
-	TSharedPtr<FUICommandInfo> Find;
 	TSharedPtr<FUICommandInfo> CollapseNodes;
+	TSharedPtr<FUICommandInfo> Find;
 	TSharedPtr<FUICommandInfo> RunDeterminismTest;
+	TSharedPtr<FUICommandInfo> StartInspectNode;
+	TSharedPtr<FUICommandInfo> StopInspectNode;
 };
diff --git a/Engine/Plugins/Experimental/PCG/Source/PCGEditor/Private/PCGEditorGraphNodeBase.cpp b/Engine/Plugins/Experimental/PCG/Source/PCGEditor/Private/PCGEditorGraphNodeBase.cpp
--- a/Engine/Plugins/Experimental/PCG/Source/PCGEditor/Private/PCGEditorGraphNodeBase.cpp
+++ b/Engine/Plugins/Experimental/PCG/Source/PCGEditor/Private/PCGEditorGraphNodeBase.cpp
@@ -2,6 +2,7 @@
 
 #include "PCGEditorGraphNodeBase.h"
 
+#include "PCGEditorCommands.h"
 #include "PCGEditorCommon.h"
 #include "PCGEditorGraph.h"
 #include "PCGEditorGraphSchema.h"
@@ -71,6 +72,7 @@ void UPCGEditorGraphNodeBase::GetNodeContextMenuActions(UToolMenu* Menu, class U
 
 	{
 		FToolMenuSection& Section = Menu->AddSection("EdGraphSchemaNodeActions", LOCTEXT("NodeActionsHeader", "Node Actions"));
+		Section.AddMenuEntry(bIsInspected ? FPCGEditorCommands::Get().StopInspectNode : FPCGEditorCommands::Get().StartInspectNode);
 		Section.AddMenuEntry(FGraphEditorCommands::Get().BreakNodeLinks);
 		Section.AddMenuEntry(FPCGEditorCommands::Get().CollapseNodes);
 	}
@@ -290,7 +292,11 @@ FLinearColor UPCGEditorGraphNodeBase::GetNodeTitleColor() const
 {
 	if (PCGNode)
 	{
-		if (PCGNode->NodeTitleColor != FLinearColor::White)
+		if (bIsInspected)
+		{
+			return GetDefault<UPCGEditorSettings>()->InspectNodeColor;
+		}
+		else if (PCGNode->NodeTitleColor != FLinearColor::White)
 		{
 			return PCGNode->NodeTitleColor;
 		}
diff --git a/Engine/Plugins/Experimental/PCG/Source/PCGEditor/Private/PCGEditorGraphNodeBase.h b/Engine/Plugins/Experimental/PCG/Source/PCGEditor/Private/PCGEditorGraphNodeBase.h
--- a/Engine/Plugins/Experimental/PCG/Source/PCGEditor/Private/PCGEditorGraphNodeBase.h
+++ b/Engine/Plugins/Experimental/PCG/Source/PCGEditor/Private/PCGEditorGraphNodeBase.h
@@ -47,6 +47,7 @@ public:
 	const UPCGNode* GetPCGNode() const { return PCGNode; }
 	void PostCopy();
 	void PostPaste();
+	void SetInspected(bool InIsInspecting) { bIsInspected = InIsInspecting; }
 
 	DECLARE_DELEGATE(FOnPCGEditorGraphNodeChanged);
 	FOnPCGEditorGraphNodeChanged OnNodeChangedDelegate;
@@ -66,4 +67,5 @@ protected:
 	TObjectPtr<UPCGNode> PCGNode = nullptr;
 
 	bool bDisableReconstructFromNode = false;
+	bool bIsInspected = false;
 };
diff --git a/Engine/Plugins/Experimental/PCG/Source/PCGEditor/Private/PCGEditorMenuContext.h b/Engine/Plugins/Experimental/PCG/Source/PCGEditor/Private/PCGEditorMenuContext.h
new file mode 100644
--- /dev/null
+++ b/Engine/Plugins/Experimental/PCG/Source/PCGEditor/Private/PCGEditorMenuContext.h
@@ -0,0 +1,18 @@
+﻿// Copyright Epic Games, Inc. All Rights Reserved.
+
+#pragma once
+
+#include "CoreMinimal.h"
+
+#include "PCGEditorMenuContext.generated.h"
+
+class FPCGEditor;
+
+UCLASS()
+class UPCGEditorMenuContext : public UObject
+{
+	GENERATED_BODY()
+
+public:
+	TWeakPtr<FPCGEditor> PCGEditor;
+};
diff --git a/Engine/Plugins/Experimental/PCG/Source/PCGEditor/Private/PCGEditorSettings.cpp b/Engine/Plugins/Experimental/PCG/Source/PCGEditor/Private/PCGEditorSettings.cpp
--- a/Engine/Plugins/Experimental/PCG/Source/PCGEditor/Private/PCGEditorSettings.cpp
+++ b/Engine/Plugins/Experimental/PCG/Source/PCGEditor/Private/PCGEditorSettings.cpp
@@ -9,6 +9,7 @@ UPCGEditorSettings::UPCGEditorSettings(const FObjectInitializer& ObjectInitializ
 	: Super(ObjectInitializer)
 {
 	DefaultNodeColor = FLinearColor(0.4f, 0.62f, 1.0f);
+	InspectNodeColor = FLinearColor(0.0f, 0.0f, 1.0f);
 	InputOutputNodeColor = FLinearColor(1.0f, 0.0f, 0.0f);
 	SetOperationNodeColor = FLinearColor(0.8f, 0.2f, 0.8f);
 	DensityOperationNodeColor = FLinearColor(0.6f, 1.0f, 0.6f);
diff --git a/Engine/Plugins/Experimental/PCG/Source/PCGEditor/Private/SPCGEditorGraphAttributeListView.cpp b/Engine/Plugins/Experimental/PCG/Source/PCGEditor/Private/SPCGEditorGraphAttributeListView.cpp
new file mode 100644
--- /dev/null
+++ b/Engine/Plugins/Experimental/PCG/Source/PCGEditor/Private/SPCGEditorGraphAttributeListView.cpp
@@ -0,0 +1,480 @@
+// Copyright Epic Games, Inc. All Rights Reserved.
+
+#include "SPCGEditorGraphAttributeListView.h"
+
+#include "Data/PCGPointData.h"
+#include "Data/PCGSpatialData.h"
+#include "PCGComponent.h"
+#include "PCGData.h"
+#include "PCGEditor.h"
+#include "PCGGraph.h"
+#include "PCGNode.h"
+
+#include "Styling/AppStyle.h"
+#include "Widgets/Input/SSearchBox.h"
+#include "Widgets/Layout/SScrollBox.h"
+#include "Widgets/SBoxPanel.h"
+#include "Widgets/SWidget.h"
+
+#define LOCTEXT_NAMESPACE "SPCGEditorGraphAttributeListView"
+
+namespace PCGEditorGraphAttributeListView
+{
+	/** Names of the columns in the attribute list */
+	const FName NAME_IndexColumn = FName(TEXT("IndexColumn"));
+	const FName NAME_PointPositionX = FName(TEXT("PointPositionX"));
+	const FName NAME_PointPositionY = FName(TEXT("PointPositionY"));
+	const FName NAME_PointPositionZ = FName(TEXT("PointPositionZ"));
+	const FName NAME_PointRotationX = FName(TEXT("PointRotationX"));
+	const FName NAME_PointRotationY = FName(TEXT("PointRotationY"));
+	const FName NAME_PointRotationZ = FName(TEXT("PointRotationZ"));
+	const FName NAME_PointScaleX = FName(TEXT("PointScaleX"));
+	const FName NAME_PointScaleY = FName(TEXT("PointScaleY"));
+	const FName NAME_PointScaleZ = FName(TEXT("PointScaleZ"));
+	const FName NAME_PointBoundsMinX = FName(TEXT("PointBoundsMinX"));
+	const FName NAME_PointBoundsMinY = FName(TEXT("PointBoundsMinY"));
+	const FName NAME_PointBoundsMinZ = FName(TEXT("PointBoundsMinZ"));
+	const FName NAME_PointBoundsMaxX = FName(TEXT("PointBoundsMaxX"));
+	const FName NAME_PointBoundsMaxY = FName(TEXT("PointBoundsMaxY"));
+	const FName NAME_PointBoundsMaxZ = FName(TEXT("PointBoundsMaxZ"));
+	const FName NAME_PointColorR = FName(TEXT("PointColorR"));
+	const FName NAME_PointColorG = FName(TEXT("PointColorG"));
+	const FName NAME_PointColorB = FName(TEXT("PointColorB"));
+	const FName NAME_PointColorA = FName(TEXT("PointColorA"));
+	const FName NAME_PointDensity = FName(TEXT("PointDensity"));
+	const FName NAME_PointSteepness = FName(TEXT("PointSteepness"));
+	const FName NAME_PointSeed = FName(TEXT("PointSeed"));
+
+	/** Labels of the columns */
+	const FText TEXT_IndexLabel = LOCTEXT("IndexLabel", "Index");
+	const FText TEXT_PointPositionLabelX = LOCTEXT("PointPositionLabelX", "PositionX");
+	const FText TEXT_PointPositionLabelY = LOCTEXT("PointPositionLabelY", "PositionY");
+	const FText TEXT_PointPositionLabelZ = LOCTEXT("PointPositionLabelZ", "PositionZ");
+	const FText TEXT_PointRotationLabelX = LOCTEXT("PointRotationLabelX", "RotationX");
+	const FText TEXT_PointRotationLabelY = LOCTEXT("PointRotationLabelY", "RotationY");
+	const FText TEXT_PointRotationLabelZ = LOCTEXT("PointRotationLabelZ", "RotationZ");
+	const FText TEXT_PointScaleLabelX = LOCTEXT("PointScaleLabelX", "ScaleX");
+	const FText TEXT_PointScaleLabelY = LOCTEXT("PointScaleLabelY", "ScaleY");
+	const FText TEXT_PointScaleLabelZ = LOCTEXT("PointScaleLabelZ", "ScaleZ");
+	const FText TEXT_PointBoundsLabelMinX = LOCTEXT("PointBoundsMinX", "BoundsMinX");
+	const FText TEXT_PointBoundsLabelMinY = LOCTEXT("PointBoundsMinY", "BoundsMinY");
+	const FText TEXT_PointBoundsLabelMinZ = LOCTEXT("PointBoundsMinZ", "BoundsMinZ");
+	const FText TEXT_PointBoundsLabelMaxX = LOCTEXT("PointBoundsMaxX", "BoundsMaxX");
+	const FText TEXT_PointBoundsLabelMaxY = LOCTEXT("PointBoundsMaxY", "BoundsMaxY");
+	const FText TEXT_PointBoundsLabelMaxZ = LOCTEXT("PointBoundsMaxZ", "BoundsMaxZ");
+	const FText TEXT_PointColorLabelR = LOCTEXT("PointColorR", "ColorR");
+	const FText TEXT_PointColorLabelG = LOCTEXT("PointColorG", "ColorG");
+	const FText TEXT_PointColorLabelB = LOCTEXT("PointColorB", "ColorB");
+	const FText TEXT_PointColorLabelA = LOCTEXT("PointColorA", "ColorA");
+	const FText TEXT_PointDensityLabel = LOCTEXT("PointDensityLabel", "Density");
+	const FText TEXT_PointSteepnessLabel = LOCTEXT("PointSteepnessLabel", "Steepness");
+	const FText TEXT_PointSeedLabel = LOCTEXT("PointSeedLabel", "Seed");
+}
+
+void SPCGListViewItemRow::Construct(const FArguments& InArgs, const TSharedRef<STableViewBase>& InOwnerTableView, const PCGListviewItemPtr& Item)
+{
+	InternalItem = Item;
+
+	SMultiColumnTableRow<PCGListviewItemPtr>::Construct(
+		SMultiColumnTableRow::FArguments()
+		.Style(FAppStyle::Get(), "DataTableEditor.CellListViewRow"),
+		InOwnerTableView);
+}
+
+TSharedRef<SWidget> SPCGListViewItemRow::GenerateWidgetForColumn(const FName& ColumnId)
+{
+	FText ColumnData = LOCTEXT("ColumnError", "Unrecognized Column");
+	if (const FPCGPoint* PCGPoint = InternalItem->PCGPoint)
+	{
+		const FTransform& Transform = PCGPoint->Transform;
+		if (ColumnId == PCGEditorGraphAttributeListView::NAME_IndexColumn)
+		{
+			ColumnData = FText::FromString(FString::FromInt(InternalItem->Index));
+		}
+		else if (ColumnId == PCGEditorGraphAttributeListView::NAME_PointPositionX)
+		{
+			const FVector& Position = Transform.GetLocation();
+			ColumnData = FText::AsNumber(Position.X);
+		}
+		else if (ColumnId == PCGEditorGraphAttributeListView::NAME_PointPositionY)
+		{
+			const FVector& Position = Transform.GetLocation();
+			ColumnData = FText::AsNumber(Position.Y);
+		}
+		else if (ColumnId == PCGEditorGraphAttributeListView::NAME_PointPositionZ)
+		{
+			const FVector& Position = Transform.GetLocation();
+			ColumnData = FText::AsNumber(Position.Z);
+		}
+		else if (ColumnId == PCGEditorGraphAttributeListView::NAME_PointRotationX)
+		{
+			const FRotator& Rotation = Transform.Rotator();
+			ColumnData = FText::AsNumber(Rotation.Roll);
+		}
+		else if (ColumnId == PCGEditorGraphAttributeListView::NAME_PointRotationY)
+		{
+			const FRotator& Rotation = Transform.Rotator();
+			ColumnData = FText::AsNumber(Rotation.Pitch);
+		}
+		else if (ColumnId == PCGEditorGraphAttributeListView::NAME_PointRotationZ)
+		{
+			const FRotator& Rotation = Transform.Rotator();
+			ColumnData = FText::AsNumber(Rotation.Yaw);
+		}
+		else if (ColumnId == PCGEditorGraphAttributeListView::NAME_PointScaleX)
+		{
+			const FVector& Scale = Transform.GetScale3D();
+			ColumnData = FText::AsNumber(Scale.X);
+		}
+		else if (ColumnId == PCGEditorGraphAttributeListView::NAME_PointScaleY)
+		{
+			const FVector& Scale = Transform.GetScale3D();
+			ColumnData = FText::AsNumber(Scale.Y);
+		}
+		else if (ColumnId == PCGEditorGraphAttributeListView::NAME_PointScaleZ)
+		{
+			const FVector& Scale = Transform.GetScale3D();
+			ColumnData = FText::AsNumber(Scale.Z);
+		}
+		else if (ColumnId == PCGEditorGraphAttributeListView::NAME_PointBoundsMinX)
+		{
+			ColumnData = FText::AsNumber(PCGPoint->BoundsMin.X);
+		}
+		else if (ColumnId == PCGEditorGraphAttributeListView::NAME_PointBoundsMinY)
+		{
+			ColumnData = FText::AsNumber(PCGPoint->BoundsMin.Y);
+		}
+		else if (ColumnId == PCGEditorGraphAttributeListView::NAME_PointBoundsMinZ)
+		{
+			ColumnData = FText::AsNumber(PCGPoint->BoundsMin.Z);
+		}
+		else if (ColumnId == PCGEditorGraphAttributeListView::NAME_PointBoundsMaxX)
+		{
+			ColumnData = FText::AsNumber(PCGPoint->BoundsMax.X);
+		}
+		else if (ColumnId == PCGEditorGraphAttributeListView::NAME_PointBoundsMaxY)
+		{
+			ColumnData = FText::AsNumber(PCGPoint->BoundsMax.Y);
+		}
+		else if (ColumnId == PCGEditorGraphAttributeListView::NAME_PointBoundsMaxZ)
+		{
+			ColumnData = FText::AsNumber(PCGPoint->BoundsMax.Z);
+		}
+		else if (ColumnId == PCGEditorGraphAttributeListView::NAME_PointColorR)
+		{
+			ColumnData = FText::AsNumber(PCGPoint->Color.X);
+		}
+		else if (ColumnId == PCGEditorGraphAttributeListView::NAME_PointColorG)
+		{
+			ColumnData = FText::AsNumber(PCGPoint->Color.Y);
+		}
+		else if (ColumnId == PCGEditorGraphAttributeListView::NAME_PointColorB)
+		{
+			ColumnData = FText::AsNumber(PCGPoint->Color.Z);
+		}
+		else if (ColumnId == PCGEditorGraphAttributeListView::NAME_PointColorA)
+		{
+			ColumnData = FText::AsNumber(PCGPoint->Color.W);
+		}		
+		else if (ColumnId == PCGEditorGraphAttributeListView::NAME_PointDensity)
+		{
+			const float Density = PCGPoint->Density;
+			ColumnData = FText::AsNumber(Density);
+		}
+		else if (ColumnId == PCGEditorGraphAttributeListView::NAME_PointSteepness)
+		{
+			const float Steepness = PCGPoint->Steepness;
+			ColumnData = FText::AsNumber(Steepness);
+		}
+		else if (ColumnId == PCGEditorGraphAttributeListView::NAME_PointSeed)
+		{
+			const int32 Seed = PCGPoint->Seed;
+			ColumnData = FText::AsNumber(Seed);
+		}
+	}
+
+	return SNew(STextBlock).Text(ColumnData);
+}
+
+SPCGEditorGraphAttributeListView::~SPCGEditorGraphAttributeListView()
+{
+	if (PCGEditorPtr.IsValid())
+	{
+		PCGEditorPtr.Pin()->OnDebugObjectChangedDelegate.RemoveAll(this);
+		PCGEditorPtr.Pin()->OnInspectedNodeChangedDelegate.RemoveAll(this);
+	}
+}
+
+void SPCGEditorGraphAttributeListView::Construct(const FArguments& InArgs, TSharedPtr<FPCGEditor> InPCGEditor)
+{
+	PCGEditorPtr = InPCGEditor;
+
+	PCGEditorPtr.Pin()->OnDebugObjectChangedDelegate.AddSP(this, &SPCGEditorGraphAttributeListView::OnDebugObjectChanged);
+	PCGEditorPtr.Pin()->OnInspectedNodeChangedDelegate.AddSP(this, &SPCGEditorGraphAttributeListView::OnInspectedNodeChanged);
+
+	ListViewHeader = CreateHeaderRowWidget();
+
+	const TSharedRef<SScrollBar> HorizontalScrollBar = SNew(SScrollBar)
+		.Orientation(Orient_Horizontal)
+		.Thickness(FVector2D(12.0f, 12.0f));
+
+	const TSharedRef<SScrollBar> VerticalScrollBar = SNew(SScrollBar)
+		.Orientation(Orient_Vertical)
+		.Thickness(FVector2D(12.0f, 12.0f));
+	
+	SAssignNew(ListView, SListView<PCGListviewItemPtr>)
+		.ListItemsSource(&ListViewItems)
+		.HeaderRow(ListViewHeader)
+		.OnGenerateRow(this, &SPCGEditorGraphAttributeListView::OnGenerateRow)
+		.AllowOverscroll(EAllowOverscroll::No)
+		.ExternalScrollbar(VerticalScrollBar)
+		.ConsumeMouseWheel(EConsumeMouseWheel::Always);
+
+	this->ChildSlot
+	[
+		SNew(SVerticalBox)
+		+SVerticalBox::Slot()
+		.FillHeight(1.0f)
+		[
+			SNew(SHorizontalBox)
+			+SHorizontalBox::Slot()
+			.FillWidth(1.0f)
+			[
+				SNew(SScrollBox)
+				.Orientation(Orient_Horizontal)
+				.ExternalScrollbar(HorizontalScrollBar)
+				+SScrollBox::Slot()
+				[
+					ListView->AsShared()
+				]
+			]
+			+SHorizontalBox::Slot()
+			.AutoWidth()
+			[
+				VerticalScrollBar
+			]
+		]
+		+SVerticalBox::Slot()
+		.AutoHeight()
+		[
+			SNew(SHorizontalBox)
+			+SHorizontalBox::Slot()
+			[
+				HorizontalScrollBar
+			]
+		]
+	];
+}
+
+TSharedRef<SHeaderRow> SPCGEditorGraphAttributeListView::CreateHeaderRowWidget() const
+{
+	return SNew(SHeaderRow)
+			.ResizeMode(ESplitterResizeMode::FixedPosition)
+			.CanSelectGeneratedColumn(true)
+			+ SHeaderRow::Column(PCGEditorGraphAttributeListView::NAME_IndexColumn)
+			.DefaultLabel(PCGEditorGraphAttributeListView::TEXT_IndexLabel)
+			.HAlignHeader(HAlign_Center)
+			.HAlignCell(HAlign_Right)
+			.ManualWidth(44)
+			+ SHeaderRow::Column(PCGEditorGraphAttributeListView::NAME_PointPositionX)
+			.DefaultLabel(PCGEditorGraphAttributeListView::TEXT_PointPositionLabelX)
+			.HAlignHeader(HAlign_Center)
+			.HAlignCell(HAlign_Right)
+			.ManualWidth(94)
+			+ SHeaderRow::Column(PCGEditorGraphAttributeListView::NAME_PointPositionY)
+			.DefaultLabel(PCGEditorGraphAttributeListView::TEXT_PointPositionLabelY)
+			.HAlignHeader(HAlign_Center)
+			.HAlignCell(HAlign_Right)
+			.ManualWidth(94)
+			+ SHeaderRow::Column(PCGEditorGraphAttributeListView::NAME_PointPositionZ)
+			.DefaultLabel(PCGEditorGraphAttributeListView::TEXT_PointPositionLabelZ)
+			.HAlignHeader(HAlign_Center)
+			.HAlignCell(HAlign_Right)
+			.ManualWidth(94)
+			+ SHeaderRow::Column(PCGEditorGraphAttributeListView::NAME_PointRotationX)
+			.DefaultLabel(PCGEditorGraphAttributeListView::TEXT_PointRotationLabelX)
+			.HAlignHeader(HAlign_Center)
+			.HAlignCell(HAlign_Right)
+			.ManualWidth(68)
+			+ SHeaderRow::Column(PCGEditorGraphAttributeListView::NAME_PointRotationY)
+			.DefaultLabel(PCGEditorGraphAttributeListView::TEXT_PointRotationLabelY)
+			.HAlignHeader(HAlign_Center)
+			.HAlignCell(HAlign_Right)
+			.ManualWidth(68)
+			+ SHeaderRow::Column(PCGEditorGraphAttributeListView::NAME_PointRotationZ)
+			.DefaultLabel(PCGEditorGraphAttributeListView::TEXT_PointRotationLabelZ)
+			.HAlignHeader(HAlign_Center)
+			.HAlignCell(HAlign_Right)
+			.ManualWidth(68)
+			+ SHeaderRow::Column(PCGEditorGraphAttributeListView::NAME_PointScaleX)
+			.DefaultLabel(PCGEditorGraphAttributeListView::TEXT_PointScaleLabelX)
+			.HAlignHeader(HAlign_Center)
+			.HAlignCell(HAlign_Right)
+			.ManualWidth(50)
+			+ SHeaderRow::Column(PCGEditorGraphAttributeListView::NAME_PointScaleY)
+			.DefaultLabel(PCGEditorGraphAttributeListView::TEXT_PointScaleLabelY)
+			.HAlignHeader(HAlign_Center)
+			.HAlignCell(HAlign_Right)
+			.ManualWidth(50)
+			+ SHeaderRow::Column(PCGEditorGraphAttributeListView::NAME_PointScaleZ)
+			.DefaultLabel(PCGEditorGraphAttributeListView::TEXT_PointScaleLabelZ)
+			.HAlignHeader(HAlign_Center)
+			.HAlignCell(HAlign_Right)
+			.ManualWidth(50)
+			+ SHeaderRow::Column(PCGEditorGraphAttributeListView::NAME_PointBoundsMinX)
+			.DefaultLabel(PCGEditorGraphAttributeListView::TEXT_PointBoundsLabelMinX)
+			.HAlignHeader(HAlign_Center)
+			.HAlignCell(HAlign_Right)
+			.ManualWidth(80)
+			+ SHeaderRow::Column(PCGEditorGraphAttributeListView::NAME_PointBoundsMinY)
+			.DefaultLabel(PCGEditorGraphAttributeListView::TEXT_PointBoundsLabelMinY)
+			.HAlignHeader(HAlign_Center)
+			.HAlignCell(HAlign_Right)
+			.ManualWidth(80)
+			+ SHeaderRow::Column(PCGEditorGraphAttributeListView::NAME_PointBoundsMinZ)
+			.DefaultLabel(PCGEditorGraphAttributeListView::TEXT_PointBoundsLabelMinZ)
+			.HAlignHeader(HAlign_Center)
+			.HAlignCell(HAlign_Right)
+			.ManualWidth(80)
+			+ SHeaderRow::Column(PCGEditorGraphAttributeListView::NAME_PointBoundsMaxX)
+			.DefaultLabel(PCGEditorGraphAttributeListView::TEXT_PointBoundsLabelMaxX)
+			.HAlignHeader(HAlign_Center)
+			.HAlignCell(HAlign_Right)
+			.ManualWidth(88)
+			+ SHeaderRow::Column(PCGEditorGraphAttributeListView::NAME_PointBoundsMaxY)
+			.DefaultLabel(PCGEditorGraphAttributeListView::TEXT_PointBoundsLabelMaxY)
+			.HAlignHeader(HAlign_Center)
+			.HAlignCell(HAlign_Right)
+			.ManualWidth(88)
+			+ SHeaderRow::Column(PCGEditorGraphAttributeListView::NAME_PointBoundsMaxZ)
+			.DefaultLabel(PCGEditorGraphAttributeListView::TEXT_PointBoundsLabelMaxZ)
+			.HAlignHeader(HAlign_Center)
+			.HAlignCell(HAlign_Right)
+			.ManualWidth(88)
+			+ SHeaderRow::Column(PCGEditorGraphAttributeListView::NAME_PointColorR)
+			.DefaultLabel(PCGEditorGraphAttributeListView::TEXT_PointColorLabelR)
+			.HAlignHeader(HAlign_Center)
+			.HAlignCell(HAlign_Right)
+			.ManualWidth(50)
+			+ SHeaderRow::Column(PCGEditorGraphAttributeListView::NAME_PointColorG)
+			.DefaultLabel(PCGEditorGraphAttributeListView::TEXT_PointColorLabelG)
+			.HAlignHeader(HAlign_Center)
+			.HAlignCell(HAlign_Right)
+			.ManualWidth(50)
+			+ SHeaderRow::Column(PCGEditorGraphAttributeListView::NAME_PointColorB)
+			.DefaultLabel(PCGEditorGraphAttributeListView::TEXT_PointColorLabelB)
+			.HAlignHeader(HAlign_Center)
+			.HAlignCell(HAlign_Right)
+			.ManualWidth(50)
+			+ SHeaderRow::Column(PCGEditorGraphAttributeListView::NAME_PointColorA)
+			.DefaultLabel(PCGEditorGraphAttributeListView::TEXT_PointColorLabelA)
+			.HAlignHeader(HAlign_Center)
+			.HAlignCell(HAlign_Right)
+			.ManualWidth(50)
+			+ SHeaderRow::Column(PCGEditorGraphAttributeListView::NAME_PointDensity)
+			.DefaultLabel(PCGEditorGraphAttributeListView::TEXT_PointDensityLabel)
+			.HAlignHeader(HAlign_Center)
+			.HAlignCell(HAlign_Right)
+			.ManualWidth(54)
+			+ SHeaderRow::Column(PCGEditorGraphAttributeListView::NAME_PointSteepness)
+			.DefaultLabel(PCGEditorGraphAttributeListView::TEXT_PointSteepnessLabel)
+			.HAlignHeader(HAlign_Center)
+			.HAlignCell(HAlign_Right)
+			.ManualWidth(73)
+			+ SHeaderRow::Column(PCGEditorGraphAttributeListView::NAME_PointSeed)
+			.DefaultLabel(PCGEditorGraphAttributeListView::TEXT_PointSeedLabel)
+			.HAlignHeader(HAlign_Center)
+			.HAlignCell(HAlign_Right)
+			.ManualWidth(88);
+}
+
+void SPCGEditorGraphAttributeListView::OnDebugObjectChanged(UPCGComponent* InPCGComponent)
+{
+	if (PCGComponent)
+	{
+		PCGComponent->OnPCGGraphGeneratedDelegate.RemoveAll(this);
+		PCGComponent->OnPCGGraphCleanedDelegate.RemoveAll(this);
+		PCGComponent->DisableInspection();
+	}
+	
+	PCGComponent = InPCGComponent;
+
+	if (PCGComponent)
+	{
+		PCGComponent->EnableInspection();
+		PCGComponent->OnPCGGraphGeneratedDelegate.AddSP(this, &SPCGEditorGraphAttributeListView::RebuildAttributeList);
+		PCGComponent->OnPCGGraphCleanedDelegate.AddSP(this, &SPCGEditorGraphAttributeListView::RebuildAttributeList);	
+	}
+
+	RebuildAttributeList();
+}
+
+void SPCGEditorGraphAttributeListView::OnInspectedNodeChanged(UPCGNode* /*InPCGNode*/)
+{	
+	RebuildAttributeList();
+}
+
+void SPCGEditorGraphAttributeListView::RebuildAttributeList(UPCGComponent* /*InPCGComponent*/)
+{
+	RebuildAttributeList();
+}
+
+void SPCGEditorGraphAttributeListView::RebuildAttributeList()
+{
+	ListViewItems.Empty();
+	ListView->RequestListRefresh();
+
+	const TSharedPtr<FPCGEditor> PCGEditor = PCGEditorPtr.Pin();
+	if (!PCGEditor.IsValid())
+	{
+		return;
+	}
+
+	if (!PCGComponent)
+	{
+		return;
+	}
+
+	const UPCGNode* PCGNode = PCGEditor->GetPCGNodeBeingInspected();
+	if (!PCGNode)
+	{
+		return;
+	}
+	
+	const FPCGDataCollection* InspectionData = PCGComponent->GetInspectionData(PCGNode);
+	if (!InspectionData)
+	{
+		return;
+	}
+	
+	if (InspectionData->TaggedData.Num() > 0)
+	{
+		const FPCGTaggedData& TaggedData = InspectionData->TaggedData[0]; // TODO: Add dropdown to select tagged data.
+		const UPCGData* PCGData = TaggedData.Data;
+		if (const UPCGSpatialData* PCGSpatialData = Cast<UPCGSpatialData>(PCGData))
+		{
+			if (const UPCGPointData* PCGPointData = PCGSpatialData->ToPointData())
+			{
+				const TArray<FPCGPoint>& PCGPoints = PCGPointData->GetPoints();
+
+				ListViewItems.Reserve(PCGPoints.Num());
+				for (int32 PointIndex = 0; PointIndex < PCGPoints.Num(); PointIndex++)
+				{
+					const FPCGPoint& PCGPoint = PCGPoints[PointIndex];
+					// TODO: Investigate swapping out the shared ptr's for better performance on huge data sets
+					PCGListviewItemPtr ListViewItem = MakeShared<FPCGListViewItem>();
+					ListViewItem->Index = PointIndex;
+					ListViewItem->PCGPoint = &PCGPoint;
+					ListViewItems.Add(ListViewItem);
+				}
+			}
+		}
+	}
+}
+
+TSharedRef<ITableRow> SPCGEditorGraphAttributeListView::OnGenerateRow(PCGListviewItemPtr Item, const TSharedRef<STableViewBase>& OwnerTable) const
+{
+	return SNew(SPCGListViewItemRow, OwnerTable, Item);
+}
+
+#undef LOCTEXT_NAMESPACE
diff --git a/Engine/Plugins/Experimental/PCG/Source/PCGEditor/Private/SPCGEditorGraphAttributeListView.h b/Engine/Plugins/Experimental/PCG/Source/PCGEditor/Private/SPCGEditorGraphAttributeListView.h
new file mode 100644
--- /dev/null
+++ b/Engine/Plugins/Experimental/PCG/Source/PCGEditor/Private/SPCGEditorGraphAttributeListView.h
@@ -0,0 +1,67 @@
+// Copyright Epic Games, Inc. All Rights Reserved.
+
+#pragma once
+
+#include "CoreMinimal.h"
+#include "Widgets/SCompoundWidget.h"
+#include "Widgets/Views/SHeaderRow.h"
+#include "Widgets/Views/SListView.h"
+
+class FPCGEditor;
+class UPCGComponent;
+class UPCGNode;
+struct FPCGPoint;
+
+struct FPCGListViewItem
+{
+	int32 Index = INDEX_NONE;
+	const FPCGPoint* PCGPoint = nullptr;
+};
+
+typedef TSharedPtr<FPCGListViewItem> PCGListviewItemPtr;
+
+class SPCGListViewItemRow : public SMultiColumnTableRow<PCGListviewItemPtr>
+{
+public:
+	SLATE_BEGIN_ARGS(SPCGListViewItemRow) {}
+	SLATE_END_ARGS()
+
+	void Construct(const FArguments& InArgs, const TSharedRef<STableViewBase>& InOwnerTableView, const PCGListviewItemPtr& Item);
+
+	virtual TSharedRef<SWidget> GenerateWidgetForColumn(const FName& ColumnId) override;
+
+protected:
+	PCGListviewItemPtr InternalItem;
+};
+
+class SPCGEditorGraphAttributeListView : public SCompoundWidget
+{
+public:
+	SLATE_BEGIN_ARGS(SPCGEditorGraphAttributeListView) { }
+	SLATE_END_ARGS()
+
+	virtual ~SPCGEditorGraphAttributeListView();
+
+	void Construct(const FArguments& InArgs, TSharedPtr<FPCGEditor> InPCGEditor);
+
+private:
+	TSharedRef<SHeaderRow> CreateHeaderRowWidget() const;
+	
+	void OnDebugObjectChanged(UPCGComponent* InPCGComponent);
+	void OnInspectedNodeChanged(UPCGNode* InPCGNode);
+
+	void RebuildAttributeList(UPCGComponent* InPCGComponent);
+	void RebuildAttributeList();
+
+	TSharedRef<ITableRow> OnGenerateRow(PCGListviewItemPtr Item, const TSharedRef<STableViewBase>& OwnerTable) const;
+
+	/** Pointer back to the PCG editor that owns us */
+	TWeakPtr<FPCGEditor> PCGEditorPtr;
+
+	/** Cached PCGComponent being viewed */
+	UPCGComponent* PCGComponent = nullptr;
+
+	TSharedPtr<SHeaderRow> ListViewHeader;
+	TSharedPtr<SListView<PCGListviewItemPtr>> ListView;
+	TArray<PCGListviewItemPtr> ListViewItems;
+};
diff --git a/Engine/Plugins/Experimental/PCG/Source/PCGEditor/Private/SPCGEditorGraphDebugObjectWidget.cpp b/Engine/Plugins/Experimental/PCG/Source/PCGEditor/Private/SPCGEditorGraphDebugObjectWidget.cpp
new file mode 100644
--- /dev/null
+++ b/Engine/Plugins/Experimental/PCG/Source/PCGEditor/Private/SPCGEditorGraphDebugObjectWidget.cpp
@@ -0,0 +1,180 @@
+// Copyright Epic Games, Inc. All Rights Reserved.
+
+#include "SPCGEditorGraphDebugObjectWidget.h"
+
+#include "PCGComponent.h"
+#include "PCGEditor.h"
+#include "PCGEditorGraph.h"
+#include "PCGGraph.h"
+
+#include "Editor.h"
+#include "Editor/UnrealEdEngine.h"
+#include "PropertyCustomizationHelpers.h"
+#include "UnrealEdGlobals.h"
+#include "Widgets/Input/STextComboBox.h"
+#include "Widgets/Text/STextBlock.h"
+
+#define LOCTEXT_NAMESPACE "PCGEditorGraphDebugObjectWidget"
+
+namespace PCGEditorGraphDebugObjectWidget
+{
+	const FString NoObjectString = TEXT("No debug object selected");
+	const FString SelectionString = TEXT(" (selected)");
+	const FString SeparatorString = TEXT(" / ");
+}
+
+void SPCGEditorGraphDebugObjectWidget::Construct(const FArguments& InArgs, TSharedPtr<FPCGEditor> InPCGEditor)
+{
+	PCGEditorPtr = InPCGEditor;
+
+	DebugObjects.Add(MakeShared<FPCGEditorGraphDebugObjectInstance>(PCGEditorGraphDebugObjectWidget::NoObjectString));
+
+	const TSharedRef<SWidget> BrowseButton = PropertyCustomizationHelpers::MakeBrowseButton(
+		FSimpleDelegate::CreateSP(this, &SPCGEditorGraphDebugObjectWidget::SelectedDebugObject_OnClicked),
+		LOCTEXT("DebugSelectActor", "Select this Actor in level"),
+		TAttribute<bool>::Create(TAttribute<bool>::FGetter::CreateSP(this, &SPCGEditorGraphDebugObjectWidget::IsSelectDebugObjectButtonEnabled))
+	);
+	
+	DebugObjectsComboBox = SNew(SComboBox<TSharedPtr<FPCGEditorGraphDebugObjectInstance>>)
+		.OptionsSource(&DebugObjects)
+		.InitiallySelectedItem(DebugObjects[0])
+		.OnComboBoxOpening(this, &SPCGEditorGraphDebugObjectWidget::OnComboBoxOpening)
+		.OnGenerateWidget(this, &SPCGEditorGraphDebugObjectWidget::OnGenerateWidget)
+		.OnSelectionChanged(this, &SPCGEditorGraphDebugObjectWidget::OnSelectionChanged)
+		.Content()
+		[
+			SNew(STextBlock)
+			.Text(this, &SPCGEditorGraphDebugObjectWidget::GetSelectedDebugObjectText)
+		];
+
+	this->ChildSlot
+	[
+		SNew(SHorizontalBox)
+		+SHorizontalBox::Slot()
+		.AutoWidth()
+		[
+			DebugObjectsComboBox.ToSharedRef()
+		]
+		+SHorizontalBox::Slot()
+		.AutoWidth()
+		.VAlign(VAlign_Center)
+		.Padding(4.0f)
+		[
+			BrowseButton
+		]
+	];
+}
+
+void SPCGEditorGraphDebugObjectWidget::OnComboBoxOpening()
+{
+	DebugObjects.Empty();
+
+	if (!PCGEditorPtr.IsValid())
+	{
+		return;
+	}
+
+	UPCGEditorGraph* PCGEditorGraph = PCGEditorPtr.Pin()->GetPCGEditorGraph();
+	if (!PCGEditorGraph)
+	{
+		return;
+	}
+
+	const UPCGGraph* PCGGraph = PCGEditorGraph->GetPCGGraph();
+	if (!PCGGraph)
+	{
+		return;
+	}
+
+	DebugObjects.Add(MakeShared<FPCGEditorGraphDebugObjectInstance>(PCGEditorGraphDebugObjectWidget::NoObjectString));
+	DebugObjectsComboBox->SetSelectedItem(DebugObjects[0]);
+
+	TArray<UObject*> PCGComponents;
+	GetObjectsOfClass(UPCGComponent::StaticClass(), PCGComponents, true);
+	for (UObject* PCGComponentObject : PCGComponents)
+	{
+		UPCGComponent* PCGComponent = Cast<UPCGComponent>(PCGComponentObject);
+		if (!PCGComponent)
+		{
+			continue;
+		}
+		
+		const AActor* Actor = PCGComponent->GetOwner();
+		if (!Actor)
+		{
+			continue;
+		}
+
+		const UPCGGraph* PCGComponentGraph = PCGComponent->GetGraph();
+		if (!PCGComponentGraph)
+		{
+			continue;
+		}
+
+		if (PCGComponentGraph == PCGGraph)
+		{
+			FString ActorNameOrLabel = Actor->GetActorNameOrLabel();
+			if (Actor->IsSelected())
+			{
+				ActorNameOrLabel.Append(PCGEditorGraphDebugObjectWidget::SelectionString);
+			}
+			
+			FString ComponentName = PCGComponent->GetFName().ToString();
+			if (PCGComponent->IsSelected())
+			{
+				ComponentName.Append(PCGEditorGraphDebugObjectWidget::SelectionString);
+			}
+
+			const FString Label = ActorNameOrLabel + PCGEditorGraphDebugObjectWidget::SeparatorString + ComponentName;
+			DebugObjects.Add(MakeShared<FPCGEditorGraphDebugObjectInstance>(PCGComponent, Label));
+		}
+	}
+}
+
+void SPCGEditorGraphDebugObjectWidget::OnSelectionChanged(TSharedPtr<FPCGEditorGraphDebugObjectInstance> NewSelection, ESelectInfo::Type SelectInfo) const
+{
+	if (NewSelection.IsValid())
+	{
+		UPCGComponent* PCGComponent = NewSelection->PCGComponent.Get();
+		PCGEditorPtr.Pin()->SetPCGComponentBeingDebugged(PCGComponent);
+	}
+}
+
+TSharedRef<SWidget> SPCGEditorGraphDebugObjectWidget::OnGenerateWidget(TSharedPtr<FPCGEditorGraphDebugObjectInstance> InDebugObjectInstance) const
+{
+	const FString ItemString = InDebugObjectInstance->Label;
+
+	return SNew(STextBlock)
+		.Text(FText::FromString(*ItemString));
+}
+
+FText SPCGEditorGraphDebugObjectWidget::GetSelectedDebugObjectText() const
+{
+	if (const TSharedPtr<FPCGEditorGraphDebugObjectInstance> SelectedItem = DebugObjectsComboBox->GetSelectedItem())
+	{
+		return SelectedItem->GetDebugObjectText();
+	}
+
+	return FText::GetEmpty();
+}
+
+void SPCGEditorGraphDebugObjectWidget::SelectedDebugObject_OnClicked() const
+{
+	if (UPCGComponent* PCGComponent = PCGEditorPtr.Pin()->GetPCGComponentBeingDebugged())
+	{
+		if (AActor* Actor = PCGComponent->GetOwner())
+		{
+			GEditor->SelectNone(/*bNoteSelectionChange=*/false, /*bDeselectBSPSurfs=*/true, /*WarnAboutManyActors=*/false);
+			GEditor->SelectActor(Actor, /*bInSelected=*/true, /*bNotify=*/true, /*bSelectEvenIfHidden=*/true);
+			GUnrealEd->Exec(Actor->GetWorld(), TEXT("CAMERA ALIGN ACTIVEVIEWPORTONLY"));
+			GEditor->SelectComponent(PCGComponent, /*bInSelected=*/true, /*bNotify=*/true, /*bSelectEvenIfHidden=*/true);
+		}
+	}
+}
+
+bool SPCGEditorGraphDebugObjectWidget::IsSelectDebugObjectButtonEnabled() const
+{
+	return PCGEditorPtr.IsValid() && (PCGEditorPtr.Pin()->GetPCGComponentBeingDebugged() != nullptr);
+}
+
+#undef LOCTEXT_NAMESPACE
diff --git a/Engine/Plugins/Experimental/PCG/Source/PCGEditor/Private/SPCGEditorGraphDebugObjectWidget.h b/Engine/Plugins/Experimental/PCG/Source/PCGEditor/Private/SPCGEditorGraphDebugObjectWidget.h
new file mode 100644
--- /dev/null
+++ b/Engine/Plugins/Experimental/PCG/Source/PCGEditor/Private/SPCGEditorGraphDebugObjectWidget.h
@@ -0,0 +1,58 @@
+// Copyright Epic Games, Inc. All Rights Reserved.
+
+#pragma once
+
+#include "CoreMinimal.h"
+#include "Widgets/SCompoundWidget.h"
+#include "Widgets/Input/STextComboBox.h"
+
+class FPCGEditor;
+class UPCGComponent;
+
+struct FPCGEditorGraphDebugObjectInstance
+{
+	FPCGEditorGraphDebugObjectInstance(const FString& InLabel)
+		: PCGComponent(nullptr)
+		, Label(InLabel)
+	{
+	}
+	
+	FPCGEditorGraphDebugObjectInstance(TWeakObjectPtr<UPCGComponent> InPCGComponent, const FString& InLabel)
+		: PCGComponent(InPCGComponent)
+		, Label(InLabel)
+	{
+	}
+
+	FText GetDebugObjectText() const
+	{
+		return FText::FromString(Label);
+	}
+
+	TWeakObjectPtr<UPCGComponent> PCGComponent;
+	FString Label;
+};
+
+
+class SPCGEditorGraphDebugObjectWidget: public SCompoundWidget
+{
+public:
+	SLATE_BEGIN_ARGS(SPCGEditorGraphDebugObjectWidget) { }
+	SLATE_END_ARGS()
+
+	void Construct(const FArguments& InArgs, TSharedPtr<FPCGEditor> InPCGEditor);
+
+private:
+	void OnComboBoxOpening();
+	void OnSelectionChanged(TSharedPtr<FPCGEditorGraphDebugObjectInstance> NewSelection, ESelectInfo::Type SelectInfo) const;
+	TSharedRef<SWidget> OnGenerateWidget(TSharedPtr<FPCGEditorGraphDebugObjectInstance> InDebugObjectInstance) const;
+
+	FText GetSelectedDebugObjectText() const;
+	void SelectedDebugObject_OnClicked() const;
+	bool IsSelectDebugObjectButtonEnabled() const;
+
+	/** Pointer back to the PCG editor that owns us */
+	TWeakPtr<FPCGEditor> PCGEditorPtr;
+
+	TArray<TSharedPtr<FPCGEditorGraphDebugObjectInstance>> DebugObjects;
+	TSharedPtr<SComboBox<TSharedPtr<FPCGEditorGraphDebugObjectInstance>>> DebugObjectsComboBox;
+};
diff --git a/Engine/Plugins/Experimental/PCG/Source/PCGEditor/Public/PCGEditorSettings.h b/Engine/Plugins/Experimental/PCG/Source/PCGEditor/Public/PCGEditorSettings.h
--- a/Engine/Plugins/Experimental/PCG/Source/PCGEditor/Public/PCGEditorSettings.h
+++ b/Engine/Plugins/Experimental/PCG/Source/PCGEditor/Public/PCGEditorSettings.h
@@ -25,6 +25,10 @@ public:
 	UPROPERTY(EditAnywhere, config, Category = Node, meta = (HideAlphaChannel))
 	FLinearColor DefaultNodeColor;
 
+	/** Inspect node color */
+	UPROPERTY(EditAnywhere, config, Category = Node, meta = (HideAlphaChannel))
+	FLinearColor InspectNodeColor;
+
 	/** Color used for input & output nodes */
 	UPROPERTY(EditAnywhere, config, Category = Node, meta = (HideAlphaChannel))
 	FLinearColor InputOutputNodeColor;
