Fix StaticAnalysis Warning in PCGCopyPointsTest.cpp

diff --git a/Engine/Plugins/Experimental/PCG/Source/PCG/Private/Tests/Elements/PCGCopyPointsTest.cpp b/Engine/Plugins/Experimental/PCG/Source/PCG/Private/Tests/Elements/PCGCopyPointsTest.cpp
--- a/Engine/Plugins/Experimental/PCG/Source/PCG/Private/Tests/Elements/PCGCopyPointsTest.cpp
+++ b/Engine/Plugins/Experimental/PCG/Source/PCG/Private/Tests/Elements/PCGCopyPointsTest.cpp
@@ -142,14 +142,19 @@ bool FPCGCopyPointsTest::RunTest(const FString& Parameters)
 			RootMode = EPCGCopyPointsMetadataInheritanceMode::Source;
 			NonRootMode = EPCGCopyPointsMetadataInheritanceMode::Target;
 		}
-		else if (Settings->AttributeInheritance == EPCGCopyPointsMetadataInheritanceMode::Target)
-		{ 
+		else // if (Settings->AttributeInheritance == EPCGCopyPointsMetadataInheritanceMode::Target)
+		{
 			RootMetadata = TargetPointData->Metadata;
 			NonRootMetadata = SourcePointData->Metadata;
 			RootMode = EPCGCopyPointsMetadataInheritanceMode::Target;
 			NonRootMode = EPCGCopyPointsMetadataInheritanceMode::Source;
 		}
 
+		if (!TestTrue("Valid input metadata", RootMetadata && NonRootMetadata))
+		{
+			return;
+		}
+
 		TArray<const FPCGMetadataAttribute<float>*> InheritedAttributes;
 		TArray<EPCGCopyPointsMetadataInheritanceMode> InheritedAttributeModes;
 		TArray<const FPCGMetadataAttribute<float>*> OutAttributes;
