Use loaded bounds in PIE instead of relying on WP actor desc.
Disabled explicit garbage collection in PIE/non-editor builds from PCG.

#rnx

diff --git a/Engine/Plugins/Experimental/PCG/Source/PCG/Private/Graph/PCGGraphExecutor.cpp b/Engine/Plugins/Experimental/PCG/Source/PCG/Private/Graph/PCGGraphExecutor.cpp
--- a/Engine/Plugins/Experimental/PCG/Source/PCG/Private/Graph/PCGGraphExecutor.cpp
+++ b/Engine/Plugins/Experimental/PCG/Source/PCG/Private/Graph/PCGGraphExecutor.cpp
@@ -561,11 +561,12 @@ void FPCGGraphExecutor::ReleaseUnusedActors()
 	ActorsToRelease.Reset();
 	ActorsListLock.Unlock();
 
-	// TODO: maybe consider other checks
-	if (bRunGC)
+#if WITH_EDITOR
+	if (bRunGC && !GIsPlayInEditorWorld)
 	{
 		CollectGarbage(RF_NoFlags, true);
 	}
+#endif
 }
 
 void FPCGGraphExecutor::NotifyGraphChanged(UPCGGraph* InGraph)
diff --git a/Engine/Plugins/Experimental/PCG/Source/PCG/Private/PCGContext.cpp b/Engine/Plugins/Experimental/PCG/Source/PCG/Private/PCGContext.cpp
--- a/Engine/Plugins/Experimental/PCG/Source/PCG/Private/PCGContext.cpp
+++ b/Engine/Plugins/Experimental/PCG/Source/PCG/Private/PCGContext.cpp
@@ -30,5 +30,5 @@ FString FPCGContext::GetTaskName() const
 
 FString FPCGContext::GetComponentName() const
 {
-	return SourceComponent ? SourceComponent->GetOwner()->GetFName().ToString() : TEXT("Non-PCG Component");
+	return SourceComponent && SourceComponent->GetOwner() ? SourceComponent->GetOwner()->GetFName().ToString() : TEXT("Non-PCG Component");
 }
\ No newline at end of file
diff --git a/Engine/Plugins/Experimental/PCG/Source/PCG/Private/PCGHelpers.cpp b/Engine/Plugins/Experimental/PCG/Source/PCG/Private/PCGHelpers.cpp
--- a/Engine/Plugins/Experimental/PCG/Source/PCG/Private/PCGHelpers.cpp
+++ b/Engine/Plugins/Experimental/PCG/Source/PCG/Private/PCGHelpers.cpp
@@ -82,10 +82,15 @@ namespace PCGHelpers
 		if (ALandscape* Landscape = Cast<ALandscape>(InLandscape))
 		{
 #if WITH_EDITOR
-			return Landscape->GetCompleteBounds();
-#else
-			return Landscape->GetLoadedBounds();
+			if (!GIsPlayInEditorWorld)
+			{
+				return Landscape->GetCompleteBounds();
+			}
+			else
 #endif
+			{
+				return Landscape->GetLoadedBounds();
+			}
 		}
 		else
 		{
