Add Tests for PCGPropertyToParamData

diff --git a/Engine/Plugins/Experimental/PCG/Source/PCG/Private/Elements/PCGPropertyToParamData.cpp b/Engine/Plugins/Experimental/PCG/Source/PCG/Private/Elements/PCGPropertyToParamData.cpp
--- a/Engine/Plugins/Experimental/PCG/Source/PCG/Private/Elements/PCGPropertyToParamData.cpp
+++ b/Engine/Plugins/Experimental/PCG/Source/PCG/Private/Elements/PCGPropertyToParamData.cpp
@@ -38,7 +38,7 @@ namespace PCGPropertyToParamDataHelpers
 		case EPCGActorSelection::ByName:
 			return [ActorSelectionName = InSettings->ActorSelectionName, &InFoundActor](AActor* Actor) -> bool
 				{
-					if (Actor->GetName() == ActorSelectionName)
+					if (Actor->GetFName().IsEqual(ActorSelectionName, ENameCase::IgnoreCase, /*bCompareNumber=*/ false))
 					{
 						InFoundActor = Actor;
 						return false;
@@ -91,7 +91,7 @@ namespace PCGPropertyToParamDataHelpers
 
 		// Early out if we have not the information necessary
 		if ((Settings->ActorSelection == EPCGActorSelection::ByTag && Settings->ActorSelectionTag == NAME_None) ||
-			(Settings->ActorSelection == EPCGActorSelection::ByName && Settings->ActorSelectionName.IsEmpty()) ||
+			(Settings->ActorSelection == EPCGActorSelection::ByName && Settings->ActorSelectionName == NAME_None) ||
 			(Settings->ActorSelection == EPCGActorSelection::ByClass && !Settings->ActorSelectionClass))
 		{
 			return FoundActor;
@@ -220,8 +220,6 @@ bool FPCGPropertyToParamDataElement::ExecuteInternal(FPCGContext* Context) const
 		return true;
 	}
 
-	TArray<FPCGTaggedData>& Outputs = Context->OutputData.TaggedData;
-
 	// First find the actor depending on the selection
 	AActor* FoundActor = PCGPropertyToParamDataHelpers::FindActor(*Context);
 
@@ -263,6 +261,7 @@ bool FPCGPropertyToParamDataElement::ExecuteInternal(FPCGContext* Context) const
 		return true;
 	}
 
+	TArray<FPCGTaggedData>& Outputs = Context->OutputData.TaggedData;
 	FPCGTaggedData& Output = Outputs.Emplace_GetRef();
 	Output.Data = ParamData;
 
diff --git a/Engine/Plugins/Experimental/PCG/Source/PCG/Private/Elements/PCGSpawnActor.cpp b/Engine/Plugins/Experimental/PCG/Source/PCG/Private/Elements/PCGSpawnActor.cpp
--- a/Engine/Plugins/Experimental/PCG/Source/PCG/Private/Elements/PCGSpawnActor.cpp
+++ b/Engine/Plugins/Experimental/PCG/Source/PCG/Private/Elements/PCGSpawnActor.cpp
@@ -75,29 +75,10 @@ bool FPCGSpawnActorElement::ExecuteInternal(FPCGContext* Context) const
 	const UPCGSpawnActorSettings* Settings = Context->GetInputSettings<UPCGSpawnActorSettings>();
 	check(Settings);
 
-	// TODO: Temporary override to verify that soft path can also be overridden
-	// Should probably be handled somewhere else (helpers?).
-	TSubclassOf<AActor> TemplateActorClass = Settings->TemplateActorClass;
-
-	if (UPCGParamData* PCGParams = Context->InputData.GetParams())
-	{
-		FString TemplateActorClassPath = TemplateActorClass ? TemplateActorClass->GetPathName() : "";
-		FString OverriddenTemplateActorClassPath = PCGSettingsHelpers::GetValue(GET_MEMBER_NAME_CHECKED(UPCGSpawnActorSettings, TemplateActorClass), TemplateActorClassPath, PCGParams);
-
-		if (OverriddenTemplateActorClassPath != TemplateActorClassPath)
-		{
-			UClass* OverriddenTemplateActorClass = FSoftClassPath(OverriddenTemplateActorClassPath).ResolveClass();
-			if (OverriddenTemplateActorClass && OverriddenTemplateActorClass->IsChildOf<AActor>())
-			{
-				TemplateActorClass = OverriddenTemplateActorClass;
-			}
-		}
-	}
-
 	// Early out
-	if(!TemplateActorClass || TemplateActorClass->HasAnyClassFlags(CLASS_Abstract))
+	if(!Settings->TemplateActorClass || Settings->TemplateActorClass->HasAnyClassFlags(CLASS_Abstract))
 	{
-		PCGE_LOG(Error, "Invalid template actor class (%s)", TemplateActorClass ? *TemplateActorClass->GetFName().ToString() : TEXT("None"));
+		PCGE_LOG(Error, "Invalid template actor class (%s)", Settings->TemplateActorClass ? *Settings->TemplateActorClass->GetFName().ToString() : TEXT("None"));
 		return true;
 	}
 
@@ -126,7 +107,7 @@ bool FPCGSpawnActorElement::ExecuteInternal(FPCGContext* Context) const
 		}
 
 		const bool bHasAuthority = !Context->SourceComponent || (Context->SourceComponent->GetOwner() && Context->SourceComponent->GetOwner()->HasAuthority());
-		const bool bSpawnedActorsRequireAuthority = TemplateActorClass->GetDefaultObject<AActor>()->GetIsReplicated();
+		const bool bSpawnedActorsRequireAuthority = Settings->TemplateActorClass->GetDefaultObject<AActor>()->GetIsReplicated();
 
 		// First, create target instance transforms
 		const UPCGPointData* PointData = SpatialData->ToPointData(Context);
@@ -154,7 +135,7 @@ bool FPCGSpawnActorElement::ExecuteInternal(FPCGContext* Context) const
 			if (Settings->Option == EPCGSpawnActorOption::CollapseActors)
 			{
 				TArray<UActorComponent*> Components;
-				UPCGActorHelpers::GetActorClassDefaultComponents(TemplateActorClass, Components, UStaticMeshComponent::StaticClass());
+				UPCGActorHelpers::GetActorClassDefaultComponents(Settings->TemplateActorClass, Components, UStaticMeshComponent::StaticClass());
 				UStaticMeshComponent* FirstSMC = nullptr;
 				UStaticMesh* Mesh = nullptr;
 
@@ -204,7 +185,7 @@ bool FPCGSpawnActorElement::ExecuteInternal(FPCGContext* Context) const
 
 					for (const FPCGPoint& Point : Points)
 					{
-						AActor* GeneratedActor = TargetActor->GetWorld()->SpawnActor(TemplateActorClass, &Point.Transform, SpawnParams);
+						AActor* GeneratedActor = TargetActor->GetWorld()->SpawnActor(Settings->TemplateActorClass, &Point.Transform, SpawnParams);
 						GeneratedActor->Tags.Add(PCGHelpers::DefaultPCGActorTag);
 						GeneratedActor->AttachToActor(TargetActor, FAttachmentTransformRules::KeepWorldTransform);
 
diff --git a/Engine/Plugins/Experimental/PCG/Source/PCG/Private/Tests/Determinism/PCGCollapseDeterminismTest.cpp b/Engine/Plugins/Experimental/PCG/Source/PCG/Private/Tests/Determinism/PCGCollapseDeterminismTest.cpp
--- a/Engine/Plugins/Experimental/PCG/Source/PCG/Private/Tests/Determinism/PCGCollapseDeterminismTest.cpp
+++ b/Engine/Plugins/Experimental/PCG/Source/PCG/Private/Tests/Determinism/PCGCollapseDeterminismTest.cpp
@@ -25,7 +25,7 @@ namespace
 		EPCGDataType::Primitive
 	};
 
-	void CreateNewRandomizedSpatialDataSet(TArray<PCGDeterminismTests::FTestData>& TestDataSet,
+	void CreateNewRandomizedSpatialDataSet(TArray<PCGTestsCommon::FTestData>& TestDataSet,
 		int32 NumSpatialSetsToAdd = 1,
 		int32 NumPointsToCreatePerSet = PCGDeterminismTests::Defaults::NumPointsToGenerate,
 		int32 NumPolyLinePointsToCreatePerSet = PCGDeterminismTests::Defaults::NumPolyLinePointsToGenerate)
@@ -33,9 +33,9 @@ namespace
 		check(NumSpatialSetsToAdd > 0 && NumPointsToCreatePerSet > 0 && NumPolyLinePointsToCreatePerSet > 0);
 
 		TestDataSet.Empty();
-		TestDataSet.Init(PCGDeterminismTests::FTestData(PCGDeterminismTests::Defaults::Seed), NumTestableSpatialDataTypes);
+		TestDataSet.Init(PCGTestsCommon::FTestData(PCGDeterminismTests::Defaults::Seed), NumTestableSpatialDataTypes);
 
-		for (PCGDeterminismTests::FTestData& Data : TestDataSet)
+		for (PCGTestsCommon::FTestData& Data : TestDataSet)
 		{
 			PCGDeterminismTests::GenerateSettings<UPCGCollapseSettings>(Data);
 		}
@@ -54,7 +54,7 @@ namespace
 bool FPCGCollapseDeterminismSingleSameDataTest::RunTest(const FString& Parameters)
 {
 	// Test single same data
-	TArray<PCGDeterminismTests::FTestData> TestDataSet;
+	TArray<PCGTestsCommon::FTestData> TestDataSet;
 	CreateNewRandomizedSpatialDataSet(TestDataSet);
 
 	for (int32 I = 0; I < TestDataSet.Num(); ++I)
@@ -69,8 +69,8 @@ bool FPCGCollapseDeterminismSingleSameDataTest::RunTest(const FString& Parameter
 bool FPCGCollapseDeterminismSingleIdenticalDataTest::RunTest(const FString& Parameters)
 {
 	// Test single identical data
-	TArray<PCGDeterminismTests::FTestData> FirstTestDataSet;
-	TArray<PCGDeterminismTests::FTestData> SecondTestDataSet;
+	TArray<PCGTestsCommon::FTestData> FirstTestDataSet;
+	TArray<PCGTestsCommon::FTestData> SecondTestDataSet;
 	CreateNewRandomizedSpatialDataSet(FirstTestDataSet);
 	CreateNewRandomizedSpatialDataSet(SecondTestDataSet);
 
@@ -86,7 +86,7 @@ bool FPCGCollapseDeterminismSingleIdenticalDataTest::RunTest(const FString& Para
 bool FPCGCollapseDeterminismMultipleSameDataTest::RunTest(const FString& Parameters)
 {
 	// Test multiple same data
-	TArray<PCGDeterminismTests::FTestData> TestDataSet;
+	TArray<PCGTestsCommon::FTestData> TestDataSet;
 	CreateNewRandomizedSpatialDataSet(TestDataSet, NumMultipleSpatialSetsToAdd, NumPointsToAddDuringMultipleCases);
 
 	for (int32 I = 0; I < TestDataSet.Num(); ++I)
@@ -101,8 +101,8 @@ bool FPCGCollapseDeterminismMultipleSameDataTest::RunTest(const FString& Paramet
 bool FPCGCollapseDeterminismMultipleIdenticalDataTest::RunTest(const FString& Parameters)
 {
 	// Test multiple identical data
-	TArray<PCGDeterminismTests::FTestData> FirstTestDataSet;
-	TArray<PCGDeterminismTests::FTestData> SecondTestDataSet;
+	TArray<PCGTestsCommon::FTestData> FirstTestDataSet;
+	TArray<PCGTestsCommon::FTestData> SecondTestDataSet;
 	CreateNewRandomizedSpatialDataSet(FirstTestDataSet, NumMultipleSpatialSetsToAdd);
 	CreateNewRandomizedSpatialDataSet(SecondTestDataSet, NumMultipleSpatialSetsToAdd);
 
@@ -118,8 +118,8 @@ bool FPCGCollapseDeterminismMultipleIdenticalDataTest::RunTest(const FString& Pa
 bool FPCGCollapseDeterminismOrderIndependenceTest::RunTest(const FString& Parameters)
 {
 	// Test shuffled input data
-	TArray<PCGDeterminismTests::FTestData> FirstTestDataSet;
-	TArray<PCGDeterminismTests::FTestData> SecondTestDataSet;
+	TArray<PCGTestsCommon::FTestData> FirstTestDataSet;
+	TArray<PCGTestsCommon::FTestData> SecondTestDataSet;
 	CreateNewRandomizedSpatialDataSet(FirstTestDataSet, NumMultipleSpatialSetsToAdd);
 	CreateNewRandomizedSpatialDataSet(SecondTestDataSet, NumMultipleSpatialSetsToAdd);
 
diff --git a/Engine/Plugins/Experimental/PCG/Source/PCG/Private/Tests/Determinism/PCGDeterminismTestsCommon.cpp b/Engine/Plugins/Experimental/PCG/Source/PCG/Private/Tests/Determinism/PCGDeterminismTestsCommon.cpp
--- a/Engine/Plugins/Experimental/PCG/Source/PCG/Private/Tests/Determinism/PCGDeterminismTestsCommon.cpp
+++ b/Engine/Plugins/Experimental/PCG/Source/PCG/Private/Tests/Determinism/PCGDeterminismTestsCommon.cpp
@@ -17,55 +17,10 @@
 #include "Components/SplineComponent.h"
 #include "Components/StaticMeshComponent.h"
 
-#if WITH_EDITOR
-#include "Editor.h"
-#endif
-
 #define LOCTEXT_NAMESPACE "PCGDeterminism"
 
 namespace PCGDeterminismTests
 {
-	FTestData::FTestData(int32 RandomSeed, UPCGSettings* DefaultSettings) :
-		Settings(DefaultSettings),
-		Seed(RandomSeed),
-		RandomStream(Seed)
-	{
-#if WITH_EDITOR
-		check(GEditor);
-		UWorld* EditorWorld = GEditor->GetEditorWorldContext().World();
-		check(EditorWorld);
-
-		// No getting the level dirty
-		FActorSpawnParameters TransientActorParameters;
-		TransientActorParameters.bHideFromSceneOutliner = true;
-		TransientActorParameters.bTemporaryEditorActor = true;
-		TransientActorParameters.ObjectFlags = RF_Transient;
-		TestActor = EditorWorld->SpawnActor<AActor>(AActor::StaticClass(), TransientActorParameters);
-		check(TestActor);
-
-		TestPCGComponent = NewObject<UPCGComponent>(TestActor, FName(TEXT("Test PCG Component")), RF_Transient);
-		check(TestPCGComponent);
-		TestActor->AddInstanceComponent(TestPCGComponent);
-		TestPCGComponent->RegisterComponent();
-
-		UPCGGraph* TestGraph = NewObject<UPCGGraph>(TestPCGComponent, FName(TEXT("Test PCG Graph")), RF_Transient);
-		check(TestGraph);
-		TestPCGComponent->SetGraph(TestGraph);
-#else
-		TestActor = nullptr;
-		TestPCGComponent = nullptr;
-		Settings = nullptr;
-#endif
-	}
-
-	void FTestData::Reset()
-	{
-		// Clear all the data
-		RandomStream.Reset();
-		InputData.TaggedData.Empty();
-		Settings = nullptr;
-	}
-
 	bool LogInvalidTest(const UPCGNode* InPCGNode, int32 Seed, EPCGDataType& OutDataTypesTested, TArray<FString>& OutAdditionalDetails)
 	{
 		UE_LOG(LogPCG, Warning, TEXT("Attempting to run an invalid determinism test"));
@@ -100,7 +55,7 @@ namespace PCGDeterminismTests
 
 	bool RunSingleSameDataTest(const UPCGNode* InPCGNode, int32 Seed, EPCGDataType& OutDataTypesTested, TArray<FString>& OutAdditionalDetails)
 	{
-		FTestData TestData(Seed, InPCGNode->DefaultSettings);
+		PCGTestsCommon::FTestData TestData(Seed, InPCGNode->DefaultSettings);
 
 		// Generate random input data
 		AddInputDataBasedOnPins(TestData, InPCGNode, OutDataTypesTested, OutAdditionalDetails);
@@ -110,8 +65,8 @@ namespace PCGDeterminismTests
 
 	bool RunSingleIdenticalDataTest(const UPCGNode* InPCGNode, int32 Seed, EPCGDataType& OutDataTypesTested, TArray<FString>& OutAdditionalDetails)
 	{
-		FTestData FirstTestData(Seed, InPCGNode->DefaultSettings);
-		FTestData SecondTestData(Seed, InPCGNode->DefaultSettings);
+		PCGTestsCommon::FTestData FirstTestData(Seed, InPCGNode->DefaultSettings);
+		PCGTestsCommon::FTestData SecondTestData(Seed, InPCGNode->DefaultSettings);
 
 		AddInputDataBasedOnPins(FirstTestData, InPCGNode, OutDataTypesTested, OutAdditionalDetails);
 		AddInputDataBasedOnPins(SecondTestData, InPCGNode, OutDataTypesTested, OutAdditionalDetails);
@@ -121,7 +76,7 @@ namespace PCGDeterminismTests
 
 	bool RunMultipleSameDataTest(const UPCGNode* InPCGNode, int32 Seed, EPCGDataType& OutDataTypesTested, TArray<FString>& OutAdditionalDetails)
 	{
-		FTestData TestData(Seed, InPCGNode->DefaultSettings);
+		PCGTestsCommon::FTestData TestData(Seed, InPCGNode->DefaultSettings);
 
 		for (int32 I = 0; I < Defaults::NumMultipleTestDataSets; ++I)
 		{
@@ -133,8 +88,8 @@ namespace PCGDeterminismTests
 
 	bool RunMultipleIdenticalDataTest(const UPCGNode* InPCGNode, int32 Seed, EPCGDataType& OutDataTypesTested, TArray<FString>& OutAdditionalDetails)
 	{
-		FTestData FirstTestData(Seed, InPCGNode->DefaultSettings);
-		FTestData SecondTestData(Seed, InPCGNode->DefaultSettings);
+		PCGTestsCommon::FTestData FirstTestData(Seed, InPCGNode->DefaultSettings);
+		PCGTestsCommon::FTestData SecondTestData(Seed, InPCGNode->DefaultSettings);
 
 		for (int32 I = 0; I < Defaults::NumMultipleTestDataSets; ++I)
 		{
@@ -147,8 +102,8 @@ namespace PCGDeterminismTests
 
 	bool RunDataCollectionOrderIndependenceTest(const UPCGNode* InPCGNode, int32 Seed, EPCGDataType& OutDataTypesTested, TArray<FString>& OutAdditionalDetails)
 	{
-		FTestData FirstTestData(Seed, InPCGNode->DefaultSettings);
-		FTestData SecondTestData(Seed, InPCGNode->DefaultSettings);
+		PCGTestsCommon::FTestData FirstTestData(Seed, InPCGNode->DefaultSettings);
+		PCGTestsCommon::FTestData SecondTestData(Seed, InPCGNode->DefaultSettings);
 
 		for (int32 I = 0; I < Defaults::NumMultipleTestDataSets; ++I)
 		{
@@ -163,8 +118,8 @@ namespace PCGDeterminismTests
 
 	bool RunAllDataOrderIndependenceTest(const UPCGNode* InPCGNode, int32 Seed, EPCGDataType& OutDataTypesTested, TArray<FString>& OutAdditionalDetails)
 	{
-		FTestData FirstTestData(Seed, InPCGNode->DefaultSettings);
-		FTestData SecondTestData(Seed, InPCGNode->DefaultSettings);
+		PCGTestsCommon::FTestData FirstTestData(Seed, InPCGNode->DefaultSettings);
+		PCGTestsCommon::FTestData SecondTestData(Seed, InPCGNode->DefaultSettings);
 
 		for (int32 I = 0; I < Defaults::NumMultipleTestDataSets; ++I)
 		{
@@ -178,7 +133,7 @@ namespace PCGDeterminismTests
 		return ExecutionIsDeterministic(FirstTestData, SecondTestData, InPCGNode);
 	}
 
-	void AddInputDataBasedOnPins(FTestData& TestData, const UPCGNode* PCGNode, EPCGDataType& OutDataTypesTested, TArray<FString>& OutAdditionalDetails)
+	void AddInputDataBasedOnPins(PCGTestsCommon::FTestData& TestData, const UPCGNode* PCGNode, EPCGDataType& OutDataTypesTested, TArray<FString>& OutAdditionalDetails)
 	{
 		check(PCGNode);
 
@@ -275,7 +230,7 @@ namespace PCGDeterminismTests
 		TaggedData.Pin = PinName;
 	}
 
-	void AddRandomizedSinglePointInputData(FTestData& TestData, int32 PointNum, const FName& PinName)
+	void AddRandomizedSinglePointInputData(PCGTestsCommon::FTestData& TestData, int32 PointNum, const FName& PinName)
 	{
 		check(PointNum > 0);
 		for (int32 I = 0; I < PointNum; ++I)
@@ -284,7 +239,7 @@ namespace PCGDeterminismTests
 		}
 	}
 
-	void AddRandomizedMultiplePointInputData(FTestData& TestData, int32 PointNum, const FName& PinName)
+	void AddRandomizedMultiplePointInputData(PCGTestsCommon::FTestData& TestData, int32 PointNum, const FName& PinName)
 	{
 		check(PointNum > 0);
 
@@ -301,7 +256,7 @@ namespace PCGDeterminismTests
 		AddMultiplePointsInputData(TestData.InputData, Points, PinName);
 	}
 
-	void AddRandomizedVolumeInputData(FTestData& TestData, const FName& PinName)
+	void AddRandomizedVolumeInputData(PCGTestsCommon::FTestData& TestData, const FName& PinName)
 	{
 		AddVolumeInputData(TestData.InputData,
 			TestData.RandomStream.VRand() * Defaults::MediumDistance,
@@ -310,12 +265,12 @@ namespace PCGDeterminismTests
 			PinName);
 	}
 
-	void AddRandomizedSurfaceInputData(FTestData& TestData, const FName& PinName)
+	void AddRandomizedSurfaceInputData(PCGTestsCommon::FTestData& TestData, const FName& PinName)
 	{
 		// TODO: PCG doesn't currently generate Surface data; function remains for future scalability
 	}
 
-	void AddRandomizedPolyLineInputData(FTestData& TestData, int32 PointNum, const FName& PinName)
+	void AddRandomizedPolyLineInputData(PCGTestsCommon::FTestData& TestData, int32 PointNum, const FName& PinName)
 	{
 		check(TestData.TestActor);
 		USplineComponent* TestSplineComponent = Cast<USplineComponent>(TestData.TestActor->GetComponentByClass(USplineComponent::StaticClass()));
@@ -339,7 +294,7 @@ namespace PCGDeterminismTests
 		AddPolyLineInputData(TestData.InputData, TestSplineComponent, PinName);
 	}
 
-	void AddRandomizedPrimitiveInputData(FTestData& TestData, const FName& PinName)
+	void AddRandomizedPrimitiveInputData(PCGTestsCommon::FTestData& TestData, const FName& PinName)
 	{
 		check(TestData.TestActor);
 		UPrimitiveComponent* TestPrimitiveComponent = Cast<UPrimitiveComponent>(TestData.TestActor->GetComponentByClass(UPrimitiveComponent::StaticClass()));
@@ -656,12 +611,12 @@ namespace PCGDeterminismTests
 		return Data && Data->IsA<UPCGPointData>();
 	}
 
-	void ShuffleInputOrder(FTestData& TestData)
+	void ShuffleInputOrder(PCGTestsCommon::FTestData& TestData)
 	{
 		ShuffleArray<FPCGTaggedData>(TestData.InputData.TaggedData, TestData.RandomStream);
 	}
 
-	void ShuffleAllInternalData(FTestData& TestData)
+	void ShuffleAllInternalData(PCGTestsCommon::FTestData& TestData)
 	{
 		for (FPCGTaggedData& TaggedData : TestData.InputData.TaggedData)
 		{
@@ -692,7 +647,7 @@ namespace PCGDeterminismTests
 		}
 	}
 
-	bool ExecutionIsDeterministic(const FTestData& FirstTestData, const FTestData& SecondTestData, const UPCGNode* PCGNode)
+	bool ExecutionIsDeterministic(const PCGTestsCommon::FTestData& FirstTestData, const PCGTestsCommon::FTestData& SecondTestData, const UPCGNode* PCGNode)
 	{
 		FPCGElementPtr FirstElement = FirstTestData.Settings->GetElement();
 		FPCGElementPtr SecondElement = SecondTestData.Settings->GetElement();
@@ -715,7 +670,7 @@ namespace PCGDeterminismTests
 		return DataCollectionsAreIdentical(FirstContext->OutputData, SecondContext->OutputData);
 	}
 
-	bool ExecutionIsDeterministicSameData(FTestData& TestData, const UPCGNode* PCGNode)
+	bool ExecutionIsDeterministicSameData(PCGTestsCommon::FTestData& TestData, const UPCGNode* PCGNode)
 	{
 		return ExecutionIsDeterministic(TestData, TestData, PCGNode);
 	}
diff --git a/Engine/Plugins/Experimental/PCG/Source/PCG/Private/Tests/Determinism/PCGDifferenceDeterminismTest.cpp b/Engine/Plugins/Experimental/PCG/Source/PCG/Private/Tests/Determinism/PCGDifferenceDeterminismTest.cpp
--- a/Engine/Plugins/Experimental/PCG/Source/PCG/Private/Tests/Determinism/PCGDifferenceDeterminismTest.cpp
+++ b/Engine/Plugins/Experimental/PCG/Source/PCG/Private/Tests/Determinism/PCGDifferenceDeterminismTest.cpp
@@ -13,27 +13,27 @@ IMPLEMENT_CUSTOM_SIMPLE_AUTOMATION_TEST(FPCGDifferenceDeterminismOrderIndependen
 
 namespace
 {
-	void SetSettingsModeInferred(PCGDeterminismTests::FTestData& TestData)
+	void SetSettingsModeInferred(PCGTestsCommon::FTestData& TestData)
 	{
 		UPCGDifferenceSettings* Settings = CastChecked<UPCGDifferenceSettings>(TestData.Settings);
 		Settings->Mode = EPCGDifferenceMode::Inferred;
 	}
 
-	void SetSettingsModeContinuous(PCGDeterminismTests::FTestData& TestData)
+	void SetSettingsModeContinuous(PCGTestsCommon::FTestData& TestData)
 	{
 		UPCGDifferenceSettings* Settings = CastChecked<UPCGDifferenceSettings>(TestData.Settings);
 		Settings->Mode = EPCGDifferenceMode::Continuous;
 	}
 
-	void SetSettingsModeDiscrete(PCGDeterminismTests::FTestData& TestData)
+	void SetSettingsModeDiscrete(PCGTestsCommon::FTestData& TestData)
 	{
 		UPCGDifferenceSettings* Settings = CastChecked<UPCGDifferenceSettings>(TestData.Settings);
 		Settings->Mode = EPCGDifferenceMode::Discrete;
 	}
 
-	void DifferenceTestBase(PCGDeterminismTests::FTestData& TestData, EPCGDifferenceMode DifferenceMode)
+	void DifferenceTestBase(PCGTestsCommon::FTestData& TestData, EPCGDifferenceMode DifferenceMode)
 	{
-		TFunction<void(PCGDeterminismTests::FTestData& TestData)> AdditionalSettingsDelegate;
+		TFunction<void(PCGTestsCommon::FTestData& TestData)> AdditionalSettingsDelegate;
 
 		switch (DifferenceMode)
 		{
@@ -58,7 +58,7 @@ namespace
 		PCGDeterminismTests::AddVolumeInputData(TestData.InputData, FVector::ZeroVector, PCGDeterminismTests::Defaults::MediumVector, PCGDeterminismTests::Defaults::SmallVector, PCGDifferenceConstants::DifferencesLabel);
 	}
 
-	void DifferenceTestMultiple(PCGDeterminismTests::FTestData& TestData, EPCGDifferenceMode DifferenceMode)
+	void DifferenceTestMultiple(PCGTestsCommon::FTestData& TestData, EPCGDifferenceMode DifferenceMode)
 	{
 		DifferenceTestBase(TestData, DifferenceMode);
 
@@ -75,9 +75,9 @@ namespace
 bool FPCGDifferenceDeterminismSingleSameDataTest::RunTest(const FString& Parameters)
 {
 	// Test single same data
-	PCGDeterminismTests::FTestData TestDataInferred(PCGDeterminismTests::Defaults::Seed);
-	PCGDeterminismTests::FTestData TestDataContinuous(PCGDeterminismTests::Defaults::Seed);
-	PCGDeterminismTests::FTestData TestDataDiscrete(PCGDeterminismTests::Defaults::Seed);
+	PCGTestsCommon::FTestData TestDataInferred(PCGDeterminismTests::Defaults::Seed);
+	PCGTestsCommon::FTestData TestDataContinuous(PCGDeterminismTests::Defaults::Seed);
+	PCGTestsCommon::FTestData TestDataDiscrete(PCGDeterminismTests::Defaults::Seed);
 
 	DifferenceTestBase(TestDataInferred, EPCGDifferenceMode::Inferred);
 	DifferenceTestBase(TestDataContinuous, EPCGDifferenceMode::Continuous);
@@ -93,12 +93,12 @@ bool FPCGDifferenceDeterminismSingleSameDataTest::RunTest(const FString& Paramet
 bool FPCGDifferenceDeterminismSingleIdenticalDataTest::RunTest(const FString& Parameters)
 {
 	// Test single identical data
-	PCGDeterminismTests::FTestData FirstTestDataInferred(PCGDeterminismTests::Defaults::Seed);
-	PCGDeterminismTests::FTestData FirstTestDataContinuous(PCGDeterminismTests::Defaults::Seed);
-	PCGDeterminismTests::FTestData FirstTestDataDiscrete(PCGDeterminismTests::Defaults::Seed);
-	PCGDeterminismTests::FTestData SecondTestDataInferred(PCGDeterminismTests::Defaults::Seed);
-	PCGDeterminismTests::FTestData SecondTestDataContinuous(PCGDeterminismTests::Defaults::Seed);
-	PCGDeterminismTests::FTestData SecondTestDataDiscrete(PCGDeterminismTests::Defaults::Seed);
+	PCGTestsCommon::FTestData FirstTestDataInferred(PCGDeterminismTests::Defaults::Seed);
+	PCGTestsCommon::FTestData FirstTestDataContinuous(PCGDeterminismTests::Defaults::Seed);
+	PCGTestsCommon::FTestData FirstTestDataDiscrete(PCGDeterminismTests::Defaults::Seed);
+	PCGTestsCommon::FTestData SecondTestDataInferred(PCGDeterminismTests::Defaults::Seed);
+	PCGTestsCommon::FTestData SecondTestDataContinuous(PCGDeterminismTests::Defaults::Seed);
+	PCGTestsCommon::FTestData SecondTestDataDiscrete(PCGDeterminismTests::Defaults::Seed);
 
 	DifferenceTestBase(FirstTestDataInferred, EPCGDifferenceMode::Inferred);
 	DifferenceTestBase(FirstTestDataContinuous, EPCGDifferenceMode::Continuous);
@@ -117,9 +117,9 @@ bool FPCGDifferenceDeterminismSingleIdenticalDataTest::RunTest(const FString& Pa
 bool FPCGDifferenceDeterminismMultipleSameDataTest::RunTest(const FString& Parameters)
 {
 	// Test multiple same data
-	PCGDeterminismTests::FTestData TestDataInferred(PCGDeterminismTests::Defaults::Seed);
-	PCGDeterminismTests::FTestData TestDataContinuous(PCGDeterminismTests::Defaults::Seed);
-	PCGDeterminismTests::FTestData TestDataDiscrete(PCGDeterminismTests::Defaults::Seed);
+	PCGTestsCommon::FTestData TestDataInferred(PCGDeterminismTests::Defaults::Seed);
+	PCGTestsCommon::FTestData TestDataContinuous(PCGDeterminismTests::Defaults::Seed);
+	PCGTestsCommon::FTestData TestDataDiscrete(PCGDeterminismTests::Defaults::Seed);
 
 	DifferenceTestMultiple(TestDataInferred, EPCGDifferenceMode::Inferred);
 	DifferenceTestMultiple(TestDataContinuous, EPCGDifferenceMode::Continuous);
@@ -135,12 +135,12 @@ bool FPCGDifferenceDeterminismMultipleSameDataTest::RunTest(const FString& Param
 bool FPCGDifferenceDeterminismMultipleIdenticalDataTest::RunTest(const FString& Parameters)
 {
 	// Test multiple identical data
-	PCGDeterminismTests::FTestData FirstTestDataInferred(PCGDeterminismTests::Defaults::Seed);
-	PCGDeterminismTests::FTestData FirstTestDataContinuous(PCGDeterminismTests::Defaults::Seed);
-	PCGDeterminismTests::FTestData FirstTestDataDiscrete(PCGDeterminismTests::Defaults::Seed);
-	PCGDeterminismTests::FTestData SecondTestDataInferred(PCGDeterminismTests::Defaults::Seed);
-	PCGDeterminismTests::FTestData SecondTestDataContinuous(PCGDeterminismTests::Defaults::Seed);
-	PCGDeterminismTests::FTestData SecondTestDataDiscrete(PCGDeterminismTests::Defaults::Seed);
+	PCGTestsCommon::FTestData FirstTestDataInferred(PCGDeterminismTests::Defaults::Seed);
+	PCGTestsCommon::FTestData FirstTestDataContinuous(PCGDeterminismTests::Defaults::Seed);
+	PCGTestsCommon::FTestData FirstTestDataDiscrete(PCGDeterminismTests::Defaults::Seed);
+	PCGTestsCommon::FTestData SecondTestDataInferred(PCGDeterminismTests::Defaults::Seed);
+	PCGTestsCommon::FTestData SecondTestDataContinuous(PCGDeterminismTests::Defaults::Seed);
+	PCGTestsCommon::FTestData SecondTestDataDiscrete(PCGDeterminismTests::Defaults::Seed);
 
 	DifferenceTestMultiple(FirstTestDataInferred, EPCGDifferenceMode::Inferred);
 	DifferenceTestMultiple(FirstTestDataContinuous, EPCGDifferenceMode::Continuous);
@@ -159,12 +159,12 @@ bool FPCGDifferenceDeterminismMultipleIdenticalDataTest::RunTest(const FString&
 bool FPCGDifferenceDeterminismOrderIndependenceTest::RunTest(const FString& Parameters)
 {
 	// Test multiple identical data, shuffled
-	PCGDeterminismTests::FTestData FirstTestDataInferred(PCGDeterminismTests::Defaults::Seed);
-	PCGDeterminismTests::FTestData FirstTestDataContinuous(PCGDeterminismTests::Defaults::Seed);
-	PCGDeterminismTests::FTestData FirstTestDataDiscrete(PCGDeterminismTests::Defaults::Seed);
-	PCGDeterminismTests::FTestData SecondTestDataInferred(PCGDeterminismTests::Defaults::Seed);
-	PCGDeterminismTests::FTestData SecondTestDataContinuous(PCGDeterminismTests::Defaults::Seed);
-	PCGDeterminismTests::FTestData SecondTestDataDiscrete(PCGDeterminismTests::Defaults::Seed);
+	PCGTestsCommon::FTestData FirstTestDataInferred(PCGDeterminismTests::Defaults::Seed);
+	PCGTestsCommon::FTestData FirstTestDataContinuous(PCGDeterminismTests::Defaults::Seed);
+	PCGTestsCommon::FTestData FirstTestDataDiscrete(PCGDeterminismTests::Defaults::Seed);
+	PCGTestsCommon::FTestData SecondTestDataInferred(PCGDeterminismTests::Defaults::Seed);
+	PCGTestsCommon::FTestData SecondTestDataContinuous(PCGDeterminismTests::Defaults::Seed);
+	PCGTestsCommon::FTestData SecondTestDataDiscrete(PCGDeterminismTests::Defaults::Seed);
 
 	DifferenceTestMultiple(FirstTestDataInferred, EPCGDifferenceMode::Inferred);
 	DifferenceTestMultiple(FirstTestDataContinuous, EPCGDifferenceMode::Continuous);
diff --git a/Engine/Plugins/Experimental/PCG/Source/PCG/Private/Tests/Determinism/PCGIntersectionDeterminismTest.cpp b/Engine/Plugins/Experimental/PCG/Source/PCG/Private/Tests/Determinism/PCGIntersectionDeterminismTest.cpp
--- a/Engine/Plugins/Experimental/PCG/Source/PCG/Private/Tests/Determinism/PCGIntersectionDeterminismTest.cpp
+++ b/Engine/Plugins/Experimental/PCG/Source/PCG/Private/Tests/Determinism/PCGIntersectionDeterminismTest.cpp
@@ -13,7 +13,7 @@ IMPLEMENT_CUSTOM_SIMPLE_AUTOMATION_TEST(FPCGIntersectionDeterminismOrderIndepend
 
 namespace
 {
-	void IntersectionTestBase(PCGDeterminismTests::FTestData& TestData)
+	void IntersectionTestBase(PCGTestsCommon::FTestData& TestData)
 	{
 		PCGDeterminismTests::GenerateSettings<UPCGIntersectionSettings>(TestData);
 		// Source Volumes
@@ -21,71 +21,71 @@ namespace
 		PCGDeterminismTests::AddVolumeInputData(TestData.InputData, PCGDeterminismTests::Defaults::SmallVector * -1.f, PCGDeterminismTests::Defaults::MediumVector, PCGDeterminismTests::Defaults::MediumVector);
 	}
 
-	void IntersectionTestMultiple(PCGDeterminismTests::FTestData& TestData)
+	void IntersectionTestMultiple(PCGTestsCommon::FTestData& TestData)
 	{
 		IntersectionTestBase(TestData);
 
 		// Randomized Sources
-		AddRandomizedVolumeInputData(TestData);
+		PCGDeterminismTests::AddRandomizedVolumeInputData(TestData);
 	}
 }
 
 bool FPCGIntersectionDeterminismSingleSameDataTest::RunTest(const FString& Parameters)
 {
 	// Test single same data
-	PCGDeterminismTests::FTestData TestData(PCGDeterminismTests::Defaults::Seed);
+	PCGTestsCommon::FTestData TestData(PCGDeterminismTests::Defaults::Seed);
 
 	IntersectionTestBase(TestData);
 
-	return TestTrue(TEXT("Same single input and settings, same output"), ExecutionIsDeterministicSameData(TestData));
+	return TestTrue(TEXT("Same single input and settings, same output"), PCGDeterminismTests::ExecutionIsDeterministicSameData(TestData));
 }
 
 bool FPCGIntersectionDeterminismSingleIdenticalDataTest::RunTest(const FString& Parameters)
 {
 	// Test single identical data
-	PCGDeterminismTests::FTestData FirstTestData(PCGDeterminismTests::Defaults::Seed);
-	PCGDeterminismTests::FTestData SecondTestData(PCGDeterminismTests::Defaults::Seed);
+	PCGTestsCommon::FTestData FirstTestData(PCGDeterminismTests::Defaults::Seed);
+	PCGTestsCommon::FTestData SecondTestData(PCGDeterminismTests::Defaults::Seed);
 
 	IntersectionTestBase(FirstTestData);
 	IntersectionTestBase(SecondTestData);
 
-	return TestTrue(TEXT("Identical single input and settings, same output"), ExecutionIsDeterministic(FirstTestData, SecondTestData));
+	return TestTrue(TEXT("Identical single input and settings, same output"), PCGDeterminismTests::ExecutionIsDeterministic(FirstTestData, SecondTestData));
 }
 
 bool FPCGIntersectionDeterminismMultipleSameDataTest::RunTest(const FString& Parameters)
 {
 	// Test multiple same data
-	PCGDeterminismTests::FTestData TestData(PCGDeterminismTests::Defaults::Seed);
+	PCGTestsCommon::FTestData TestData(PCGDeterminismTests::Defaults::Seed);
 
 	IntersectionTestMultiple(TestData);
 
-	return TestTrue(TEXT("Identical multiple input, same output"), ExecutionIsDeterministicSameData(TestData));
+	return TestTrue(TEXT("Identical multiple input, same output"), PCGDeterminismTests::ExecutionIsDeterministicSameData(TestData));
 }
 
 bool FPCGIntersectionDeterminismMultipleIdenticalDataTest::RunTest(const FString& Parameters)
 {
 	// Test multiple identical data
-	PCGDeterminismTests::FTestData FirstTestData(PCGDeterminismTests::Defaults::Seed);
-	PCGDeterminismTests::FTestData SecondTestData(PCGDeterminismTests::Defaults::Seed);
+	PCGTestsCommon::FTestData FirstTestData(PCGDeterminismTests::Defaults::Seed);
+	PCGTestsCommon::FTestData SecondTestData(PCGDeterminismTests::Defaults::Seed);
 
 	IntersectionTestMultiple(FirstTestData);
 	IntersectionTestMultiple(SecondTestData);
 
-	return TestTrue(TEXT("Identical single input and settings, same output"), ExecutionIsDeterministic(FirstTestData, SecondTestData));
+	return TestTrue(TEXT("Identical single input and settings, same output"), PCGDeterminismTests::ExecutionIsDeterministic(FirstTestData, SecondTestData));
 }
 
 bool FPCGIntersectionDeterminismOrderIndependenceTest::RunTest(const FString& Parameters)
 {
 	// Test multiple identical, shuffled data
-	PCGDeterminismTests::FTestData FirstTestData(PCGDeterminismTests::Defaults::Seed);
-	PCGDeterminismTests::FTestData SecondTestData(PCGDeterminismTests::Defaults::Seed);
+	PCGTestsCommon::FTestData FirstTestData(PCGDeterminismTests::Defaults::Seed);
+	PCGTestsCommon::FTestData SecondTestData(PCGDeterminismTests::Defaults::Seed);
 
 	IntersectionTestMultiple(FirstTestData);
 	IntersectionTestMultiple(SecondTestData);
 
-	ShuffleInputOrder(SecondTestData);
+	PCGDeterminismTests::ShuffleInputOrder(SecondTestData);
 
-	return TestTrue(TEXT("Shuffled input order, same output"), ExecutionIsDeterministic(FirstTestData, SecondTestData));
+	return TestTrue(TEXT("Shuffled input order, same output"), PCGDeterminismTests::ExecutionIsDeterministic(FirstTestData, SecondTestData));
 }
 
 #endif
\ No newline at end of file
diff --git a/Engine/Plugins/Experimental/PCG/Source/PCG/Private/Tests/Determinism/PCGVolumeSamplerDeterminismTest.cpp b/Engine/Plugins/Experimental/PCG/Source/PCG/Private/Tests/Determinism/PCGVolumeSamplerDeterminismTest.cpp
--- a/Engine/Plugins/Experimental/PCG/Source/PCG/Private/Tests/Determinism/PCGVolumeSamplerDeterminismTest.cpp
+++ b/Engine/Plugins/Experimental/PCG/Source/PCG/Private/Tests/Determinism/PCGVolumeSamplerDeterminismTest.cpp
@@ -15,7 +15,7 @@ namespace
 {
 	constexpr static int32 NumInputVolumesToAdd = 10;
 
-	void RandomizeVolumeSettingsVoxelSize(PCGDeterminismTests::FTestData& TestData)
+	void RandomizeVolumeSettingsVoxelSize(PCGTestsCommon::FTestData& TestData)
 	{
 		UPCGVolumeSamplerSettings* Settings = CastChecked<UPCGVolumeSamplerSettings>(TestData.Settings);
 		Settings->VoxelSize = PCGDeterminismTests::Defaults::MediumVector + TestData.RandomStream.VRand() * 0.5 * PCGDeterminismTests::Defaults::MediumDistance;
@@ -25,7 +25,7 @@ namespace
 bool FPCGVolumeSamplerDeterminismSingleSameDataTest::RunTest(const FString& Parameters)
 {
 	// Test single same data
-	PCGDeterminismTests::FTestData TestData(PCGDeterminismTests::Defaults::Seed);
+	PCGTestsCommon::FTestData TestData(PCGDeterminismTests::Defaults::Seed);
 
 	PCGDeterminismTests::GenerateSettings<UPCGVolumeSamplerSettings>(TestData, RandomizeVolumeSettingsVoxelSize);
 	PCGDeterminismTests::AddRandomizedVolumeInputData(TestData);
@@ -38,8 +38,8 @@ bool FPCGVolumeSamplerDeterminismSingleSameDataTest::RunTest(const FString& Para
 bool FPCGVolumeSamplerDeterminismSingleIdenticalDataTest::RunTest(const FString& Parameters)
 {
 	// Test single identical data
-	PCGDeterminismTests::FTestData FirstTestData(PCGDeterminismTests::Defaults::Seed);
-	PCGDeterminismTests::FTestData SecondTestData(PCGDeterminismTests::Defaults::Seed);
+	PCGTestsCommon::FTestData FirstTestData(PCGDeterminismTests::Defaults::Seed);
+	PCGTestsCommon::FTestData SecondTestData(PCGDeterminismTests::Defaults::Seed);
 
 	PCGDeterminismTests::GenerateSettings<UPCGVolumeSamplerSettings>(FirstTestData, RandomizeVolumeSettingsVoxelSize);
 	PCGDeterminismTests::GenerateSettings<UPCGVolumeSamplerSettings>(SecondTestData, RandomizeVolumeSettingsVoxelSize);
@@ -55,7 +55,7 @@ bool FPCGVolumeSamplerDeterminismSingleIdenticalDataTest::RunTest(const FString&
 bool FPCGVolumeSamplerDeterminismMultipleSameDataTest::RunTest(const FString& Parameters)
 {
 	// Test multiple same data
-	PCGDeterminismTests::FTestData TestData(PCGDeterminismTests::Defaults::Seed);
+	PCGTestsCommon::FTestData TestData(PCGDeterminismTests::Defaults::Seed);
 
 	PCGDeterminismTests::GenerateSettings<UPCGVolumeSamplerSettings>(TestData, RandomizeVolumeSettingsVoxelSize);
 
@@ -73,8 +73,8 @@ bool FPCGVolumeSamplerDeterminismMultipleSameDataTest::RunTest(const FString& Pa
 bool FPCGVolumeSamplerDeterminismMultipleIdenticalDataTest::RunTest(const FString& Parameters)
 {
 	// Test multiple identical data
-	PCGDeterminismTests::FTestData FirstTestData(PCGDeterminismTests::Defaults::Seed);
-	PCGDeterminismTests::FTestData SecondTestData(PCGDeterminismTests::Defaults::Seed);
+	PCGTestsCommon::FTestData FirstTestData(PCGDeterminismTests::Defaults::Seed);
+	PCGTestsCommon::FTestData SecondTestData(PCGDeterminismTests::Defaults::Seed);
 
 	PCGDeterminismTests::GenerateSettings<UPCGVolumeSamplerSettings>(FirstTestData, RandomizeVolumeSettingsVoxelSize);
 	PCGDeterminismTests::GenerateSettings<UPCGVolumeSamplerSettings>(SecondTestData, RandomizeVolumeSettingsVoxelSize);
@@ -94,8 +94,8 @@ bool FPCGVolumeSamplerDeterminismMultipleIdenticalDataTest::RunTest(const FStrin
 bool FPCGVolumeSamplerDeterminismOrderIndependenceTest::RunTest(const FString& Parameters)
 {
 	// Test shuffled input data
-	PCGDeterminismTests::FTestData FirstTestData(PCGDeterminismTests::Defaults::Seed);
-	PCGDeterminismTests::FTestData SecondTestData(PCGDeterminismTests::Defaults::Seed);
+	PCGTestsCommon::FTestData FirstTestData(PCGDeterminismTests::Defaults::Seed);
+	PCGTestsCommon::FTestData SecondTestData(PCGDeterminismTests::Defaults::Seed);
 
 	PCGDeterminismTests::GenerateSettings<UPCGVolumeSamplerSettings>(FirstTestData, RandomizeVolumeSettingsVoxelSize);
 	PCGDeterminismTests::GenerateSettings<UPCGVolumeSamplerSettings>(SecondTestData, RandomizeVolumeSettingsVoxelSize);
diff --git a/Engine/Plugins/Experimental/PCG/Source/PCG/Private/Tests/Elements/PCGCopyPointsTest.cpp b/Engine/Plugins/Experimental/PCG/Source/PCG/Private/Tests/Elements/PCGCopyPointsTest.cpp
--- a/Engine/Plugins/Experimental/PCG/Source/PCG/Private/Tests/Elements/PCGCopyPointsTest.cpp
+++ b/Engine/Plugins/Experimental/PCG/Source/PCG/Private/Tests/Elements/PCGCopyPointsTest.cpp
@@ -18,7 +18,7 @@ IMPLEMENT_CUSTOM_SIMPLE_AUTOMATION_TEST(FPCGCopyPointsTest, FPCGTestBaseClass, "
 
 bool FPCGCopyPointsTest::RunTest(const FString& Parameters)
 {
-	PCGDeterminismTests::FTestData TestData(PCGDeterminismTests::Defaults::Seed);
+	PCGTestsCommon::FTestData TestData(PCGDeterminismTests::Defaults::Seed);
 	PCGDeterminismTests::GenerateSettings<UPCGCopyPointsSettings>(TestData);
 	UPCGCopyPointsSettings* Settings = CastChecked<UPCGCopyPointsSettings>(TestData.Settings);
 	FPCGElementPtr CopyPointsElement = TestData.Settings->GetElement();
diff --git a/Engine/Plugins/Experimental/PCG/Source/PCG/Private/Tests/Elements/PCGPropertyToParamDataElementTest.cpp b/Engine/Plugins/Experimental/PCG/Source/PCG/Private/Tests/Elements/PCGPropertyToParamDataElementTest.cpp
new file mode 100644
--- /dev/null
+++ b/Engine/Plugins/Experimental/PCG/Source/PCG/Private/Tests/Elements/PCGPropertyToParamDataElementTest.cpp
@@ -0,0 +1,328 @@
+// Copyright Epic Games, Inc. All Rights Reserved.
+
+#include "Tests/Elements/PCGPropertyToParamDataElementTest.h"
+#include "PCGComponent.h"
+#include "PCGGraph.h"
+#include "PCGParamData.h"
+#include "PCGVolume.h"
+#include "Elements/PCGPropertyToParamData.h"
+#include "Metadata/PCGMetadataAttribute.h"
+#include "Metadata/PCGMetadataAttributeTpl.h"
+#include "Metadata/PCGMetadataAttributeTraits.h"
+
+// To run: automation runtest pcg.tests.PropertyToParamData
+
+#if WITH_EDITOR
+
+IMPLEMENT_CUSTOM_SIMPLE_AUTOMATION_TEST(FPCGPropertyToParamDataPropertyTypeTest, FPCGTestBaseClass, "pcg.tests.PropertyToParamData.PropertyType", PCGTestsCommon::TestFlags)
+IMPLEMENT_CUSTOM_SIMPLE_AUTOMATION_TEST(FPCGPropertyToParamDataActorFindTest, FPCGTestBaseClass, "pcg.tests.PropertyToParamData.ActorFind", PCGTestsCommon::TestFlags)
+
+template <typename AttributeType>
+bool VerifyAttributeValue(FPCGTestBaseClass* TestInstance, PCGTestsCommon::FTestData& TestData, FName PropertyName, AttributeType&& ExpectedValue, bool bShouldFail, FString ExtraTestWhat)
+{
+	// Use universal reference to either pass an rvalue or lvalue. Use this type below to know the raw type behind.
+	using RawAttributeType = std::remove_const_t<std::remove_reference_t<AttributeType>>;
+
+	// Use TestData settings to set the property name
+	UPCGPropertyToParamDataSettings* Settings = Cast<UPCGPropertyToParamDataSettings>(TestData.Settings);
+	if (!TestInstance->TestNotNull(TEXT("CastToUPCGPropertyToParamDataSettings"), Settings))
+	{
+		return false;
+	}
+
+	Settings->PropertyName = PropertyName;
+	Settings->OutputAttributeName = Settings->PropertyName;
+
+	// Add 2 nodes, PropertyToParamDataNode and a Trivial node (just there for the connection)
+	UPCGNode* TestNode = TestData.TestPCGComponent->GetGraph()->AddNode(Settings);
+
+	UPCGTrivialSettings* TrivialSettings = NewObject<UPCGTrivialSettings>();
+	UPCGNode* TrivialNode = TestData.TestPCGComponent->GetGraph()->AddNode(TrivialSettings);
+
+	TestNode->GetOutputPins()[0]->AddEdgeTo(TrivialNode->GetInputPins()[0]);
+
+	FPCGElementPtr Element = Settings->GetElement();
+	TUniquePtr<FPCGContext> Context(Element->Initialize(FPCGDataCollection(), TestData.TestPCGComponent, TestNode));
+	Context->NumAvailableTasks = 1;
+
+	// Run the element, can throw an error.
+	while (!Element->Execute(Context.Get()))
+	{};
+
+	auto FormatWithPropertyName = [&ExtraTestWhat, PropertyName](const FString& InText) -> FString
+	{
+		return FString::Format(TEXT("{0}_{1}_{2}"), { ExtraTestWhat, PropertyName.ToString(), InText });
+	};
+
+	bool bSuccess = true;
+
+	if (bShouldFail)
+	{
+		// If it should fail, it will have no output.
+		bSuccess = TestInstance->TestEqual(FormatWithPropertyName("NumOutput"), Context->OutputData.GetAllParams().Num(), 0);
+	}
+	else
+	{
+		// If it should succeed, we check the number of output and verify that the type and value matches the expected one.
+		bSuccess = TestInstance->TestEqual(FormatWithPropertyName("NumOutput"), Context->OutputData.GetAllParams().Num(), 1);
+
+		if (bSuccess)
+		{
+			UPCGParamData* ParamData = Cast<UPCGParamData>(Context->OutputData.GetAllParams()[0].Data);
+			const FPCGMetadataAttributeBase* Attribute = ParamData->ConstMetadata()->GetConstAttribute(Settings->PropertyName);
+
+			bSuccess = TestInstance->TestEqual(FormatWithPropertyName("AttributeType"), Attribute->GetTypeId(), PCG::Private::MetadataTypes<RawAttributeType>::Id);
+
+			if (bSuccess)
+			{
+				RawAttributeType Value = static_cast<const FPCGMetadataAttribute<RawAttributeType>*>(Attribute)->GetValueFromItemKey(0);
+				TestInstance->TestEqual(FormatWithPropertyName("AttributeValue"), Value, ExpectedValue);
+			}
+		}
+	}
+
+	// Cleanup behind ourselves.
+	TestData.TestPCGComponent->GetGraph()->RemoveNode(TrivialNode);
+	TestData.TestPCGComponent->GetGraph()->RemoveNode(TestNode);
+
+	return bSuccess;
+}
+
+bool FPCGPropertyToParamDataPropertyTypeTest::RunTest(const FString& Parameters)
+{
+	bool bSuccess = true;
+
+	UPCGPropertyToParamDataSettings* Settings = NewObject<UPCGPropertyToParamDataSettings>();
+	Settings->ActorSelection = EPCGActorSelection::ByClass;
+	Settings->ActorSelectionClass = APCGUnitTestDummyActor::StaticClass();
+	Settings->ActorFilter = EPCGActorFilter::Self;
+
+	static constexpr int32 Seed = 42;
+	const FString ExtraTestWhat = "PropertyToParamDataPropertyTypeTest";
+
+	PCGTestsCommon::FTestData TestData(Seed, Settings, APCGUnitTestDummyActor::StaticClass());
+
+	// Set all properties
+	const FName NameValue = TEXT("HelloWorld");
+	const FString StringValue = TEXT("HelloWorld");
+
+	const FVector VectorValue{ 1.0, 2.0, 3.0 };
+	const FVector4 Vector4Value{ 1.0, 2.0, 3.0, 4.0 };
+	const FRotator RotatorValue{ 45.0, 45.0, 45.0 };
+	const FQuat QuatValue = RotatorValue.Quaternion();
+	const FTransform TransformValue{ QuatValue, VectorValue, VectorValue };
+
+	APCGVolume* PCGVolume = NewObject<APCGVolume>();
+
+	const FSoftObjectPath SoftObjectPathValue{ PCGVolume };
+	const FSoftClassPath SoftClassPathValue{ APCGVolume::StaticClass() };
+
+	const FVector2D Vector2Value = { 1.0, 2.0 };
+
+	APCGUnitTestDummyActor* Actor = Cast<APCGUnitTestDummyActor>(TestData.TestActor);
+	Actor->IntProperty = 42;
+	Actor->Int64Property = 42ll;
+	Actor->FloatProperty = 1.0f;
+	Actor->DoubleProperty = 1.0;
+	Actor->BoolProperty = true;
+	Actor->NameProperty = NameValue;
+	Actor->StringProperty = StringValue;
+	Actor->EnumProperty = EPCGUnitTestDummyEnum::Three;
+	Actor->VectorProperty = VectorValue;
+	Actor->Vector4Property = Vector4Value;
+	Actor->RotatorProperty = RotatorValue;
+	Actor->QuatProperty = QuatValue;
+	Actor->TransformProperty = TransformValue;
+	Actor->SoftObjectPathProperty = SoftObjectPathValue;
+	Actor->SoftClassPathProperty = SoftClassPathValue;
+	Actor->ClassProperty = APCGVolume::StaticClass();
+	Actor->ObjectProperty = PCGVolume;
+	Actor->Vector2Property = Vector2Value;
+
+	// Basic properties
+	bSuccess &= VerifyAttributeValue(this, TestData, GET_MEMBER_NAME_CHECKED(APCGUnitTestDummyActor, IntProperty), 42ll, false, ExtraTestWhat);
+	bSuccess &= VerifyAttributeValue(this, TestData, GET_MEMBER_NAME_CHECKED(APCGUnitTestDummyActor, Int64Property), 42ll, false, ExtraTestWhat);
+	bSuccess &= VerifyAttributeValue(this, TestData, GET_MEMBER_NAME_CHECKED(APCGUnitTestDummyActor, FloatProperty), 1.0, false, ExtraTestWhat);
+	bSuccess &= VerifyAttributeValue(this, TestData, GET_MEMBER_NAME_CHECKED(APCGUnitTestDummyActor, DoubleProperty), 1.0, false, ExtraTestWhat);
+	bSuccess &= VerifyAttributeValue(this, TestData, GET_MEMBER_NAME_CHECKED(APCGUnitTestDummyActor, BoolProperty), true, false, ExtraTestWhat);
+
+	// String/Name Properties
+	bSuccess &= VerifyAttributeValue(this, TestData, GET_MEMBER_NAME_CHECKED(APCGUnitTestDummyActor, NameProperty), NameValue, false, ExtraTestWhat);
+	bSuccess &= VerifyAttributeValue(this, TestData, GET_MEMBER_NAME_CHECKED(APCGUnitTestDummyActor, StringProperty), StringValue, false, ExtraTestWhat);
+
+	// Enum Property
+	bSuccess &= VerifyAttributeValue(this, TestData, GET_MEMBER_NAME_CHECKED(APCGUnitTestDummyActor, EnumProperty), (int64)EPCGUnitTestDummyEnum::Three, false, ExtraTestWhat);
+
+	// Struct Properties
+	bSuccess &= VerifyAttributeValue(this, TestData, GET_MEMBER_NAME_CHECKED(APCGUnitTestDummyActor, VectorProperty), VectorValue, false, ExtraTestWhat);
+	bSuccess &= VerifyAttributeValue(this, TestData, GET_MEMBER_NAME_CHECKED(APCGUnitTestDummyActor, Vector4Property), Vector4Value, false, ExtraTestWhat);
+	bSuccess &= VerifyAttributeValue(this, TestData, GET_MEMBER_NAME_CHECKED(APCGUnitTestDummyActor, RotatorProperty), RotatorValue, false, ExtraTestWhat);
+	bSuccess &= VerifyAttributeValue(this, TestData, GET_MEMBER_NAME_CHECKED(APCGUnitTestDummyActor, QuatProperty), QuatValue, false, ExtraTestWhat);
+	bSuccess &= VerifyAttributeValue(this, TestData, GET_MEMBER_NAME_CHECKED(APCGUnitTestDummyActor, TransformProperty), TransformValue, false, ExtraTestWhat);
+	bSuccess &= VerifyAttributeValue(this, TestData, GET_MEMBER_NAME_CHECKED(APCGUnitTestDummyActor, SoftObjectPathProperty), SoftObjectPathValue.ToString(), false, ExtraTestWhat);
+	bSuccess &= VerifyAttributeValue(this, TestData, GET_MEMBER_NAME_CHECKED(APCGUnitTestDummyActor, SoftClassPathProperty), SoftClassPathValue.ToString(), false, ExtraTestWhat);
+
+	// Objects properties
+	bSuccess &= VerifyAttributeValue(this, TestData, GET_MEMBER_NAME_CHECKED(APCGUnitTestDummyActor, ClassProperty), APCGVolume::StaticClass()->GetPathName(), false, ExtraTestWhat);
+	bSuccess &= VerifyAttributeValue(this, TestData, GET_MEMBER_NAME_CHECKED(APCGUnitTestDummyActor, ObjectProperty), PCGVolume->GetPathName(), false, ExtraTestWhat);
+
+	// Unsupported property
+	AddExpectedError(TEXT("Error while creating an attribute. Either the property type is not supported by PCG or attribute creation failed."), EAutomationExpectedErrorFlags::Contains, 1);
+	bSuccess &= VerifyAttributeValue(this, TestData, GET_MEMBER_NAME_CHECKED(APCGUnitTestDummyActor, Vector2Property), Vector2Value, true, ExtraTestWhat);
+
+	// Unknown property
+	AddExpectedError(TEXT("Property doesn't exist in the found actor."), EAutomationExpectedErrorFlags::Contains, 1);
+	bSuccess &= VerifyAttributeValue(this, TestData, TEXT("DummyMissingProperty"), 42, true, ExtraTestWhat);
+
+	// Missing property
+	AddExpectedError(TEXT("Some parameters are missing, abort."), EAutomationExpectedErrorFlags::Contains, 1);
+	bSuccess &= VerifyAttributeValue(this, TestData, NAME_None, 42, true, ExtraTestWhat);
+
+	return bSuccess;
+}
+
+bool FPCGPropertyToParamDataActorFindTest::RunTest(const FString& Parameters)
+{
+	bool bSuccess = true;
+
+	static constexpr int32 Seed = 42;
+	static const FName Tag = TEXT("PCGUnitTestDummyMyTag");
+	static const FString Name = TEXT("MyPCGUnitTestDummyActor");
+
+	UPCGPropertyToParamDataSettings* Settings = NewObject<UPCGPropertyToParamDataSettings>();
+	Settings->ActorSelectionClass = APCGUnitTestDummyActor::StaticClass();
+	Settings->ActorSelectionTag = Tag;
+	Settings->ActorSelectionName = APCGUnitTestDummyActor::StaticClass()->GetFName();
+	Settings->ComponentClass = UPCGUnitTestDummyComponent::StaticClass();
+
+	// Self by class
+	{
+		Settings->ActorSelection = EPCGActorSelection::ByClass;
+		Settings->ActorFilter = EPCGActorFilter::Self;
+
+		PCGTestsCommon::FTestData TestData(Seed, Settings, APCGUnitTestDummyActor::StaticClass());
+		Cast<APCGUnitTestDummyActor>(TestData.TestActor)->IntProperty = 42;
+
+		bSuccess &= VerifyAttributeValue(this, TestData, GET_MEMBER_NAME_CHECKED(APCGUnitTestDummyActor, IntProperty), 42ll, false, "PropertyToParamDataActorFindTest_Self_Class");
+	}
+
+	// Self by tag
+	{
+		Settings->ActorSelection = EPCGActorSelection::ByTag;
+		Settings->ActorFilter = EPCGActorFilter::Self;
+
+		PCGTestsCommon::FTestData TestData(Seed, Settings, APCGUnitTestDummyActor::StaticClass());
+		TestData.TestActor->Tags.Add(Tag);
+		Cast<APCGUnitTestDummyActor>(TestData.TestActor)->IntProperty = 42;
+
+		bSuccess &= VerifyAttributeValue(this, TestData, GET_MEMBER_NAME_CHECKED(APCGUnitTestDummyActor, IntProperty), 42ll, false, "PropertyToParamDataActorFindTest_Self_Tag");
+	}
+
+	// Self by name
+	{
+		Settings->ActorSelection = EPCGActorSelection::ByName;
+		Settings->ActorFilter = EPCGActorFilter::Self;
+
+		PCGTestsCommon::FTestData TestData(Seed, Settings, APCGUnitTestDummyActor::StaticClass());
+		Cast<APCGUnitTestDummyActor>(TestData.TestActor)->IntProperty = 42;
+
+		bSuccess &= VerifyAttributeValue(this, TestData, GET_MEMBER_NAME_CHECKED(APCGUnitTestDummyActor, IntProperty), 42ll, false, "PropertyToParamDataActorFindTest_Self_Name");
+	}
+
+	// TODO: Need a good way to spawn actors with parenting relation between them
+	//// Parent by class
+	//{
+	//	Settings->ActorSelection = EPCGActorSelection::ByClass;
+	//	Settings->ActorFilter = EPCGActorFilter::Parent;
+
+	//	PCGTestsCommon::FTestData TestData(Seed, Settings, AActor::StaticClass());
+	//	TestData.AddActor(APCGUnitTestDummyActor::StaticClass(), true);
+	//	Cast<APCGUnitTestDummyActor>(TestData.TestActor)->IntProperty = 42;
+
+	//	bSuccess &= VerifyAttributeValue(this, TestData, GET_MEMBER_NAME_CHECKED(APCGUnitTestDummyActor, IntProperty), 42ll, false, "PropertyToParamDataActorFindTest_Parent_Class");
+	//}
+
+	//// Root by tag
+	//{
+	//	Settings->ActorSelection = EPCGActorSelection::ByTag;
+	//	Settings->ActorFilter = EPCGActorFilter::Root;
+
+	//	PCGTestsCommon::FTestData TestData(Seed, Settings, AActor::StaticClass());
+	//	TestData.AddActor(AActor::StaticClass(), true);
+	//	TestData.AddActor(APCGUnitTestDummyActor::StaticClass(), true);
+	//	TestData.TestActor->Tags.Add(Tag);
+	//	Cast<APCGUnitTestDummyActor>(TestData.TestActor)->IntProperty = 42;
+
+	//	bSuccess &= VerifyAttributeValue(this, TestData, GET_MEMBER_NAME_CHECKED(APCGUnitTestDummyActor, IntProperty), 42ll, false, "PropertyToParamDataActorFindTest_Root_Tag");
+	//}
+
+	//// Root by tag including children
+	//{
+	//	Settings->ActorSelection = EPCGActorSelection::ByTag;
+	//	Settings->ActorFilter = EPCGActorFilter::Root;
+	//	Settings->bIncludeChildren = true;
+
+	//	PCGTestsCommon::FTestData TestData(Seed, Settings, AActor::StaticClass());
+	//	TestData.AddActor(APCGUnitTestDummyActor::StaticClass(), true);
+	//	TestData.TestActor->Tags.Add(Tag);
+	//	Cast<APCGUnitTestDummyActor>(TestData.TestActor)->IntProperty = 42;
+
+	//	TestData.AddActor(AActor::StaticClass(), true);
+
+	//	bSuccess &= VerifyAttributeValue(this, TestData, GET_MEMBER_NAME_CHECKED(APCGUnitTestDummyActor, IntProperty), 42ll, false, "PropertyToParamDataActorFindTest_Root_Tag_Children");
+	//	Settings->bIncludeChildren = false;
+	//}
+
+	// Self by class including component
+	{
+		Settings->ActorSelection = EPCGActorSelection::ByTag;
+		Settings->ActorFilter = EPCGActorFilter::Self;
+		Settings->bSelectComponent = true;
+		
+		PCGTestsCommon::FTestData TestData(Seed, Settings, AActor::StaticClass());
+		TestData.TestActor->Tags.Add(Tag);
+		UPCGUnitTestDummyComponent* Component = Cast< UPCGUnitTestDummyComponent>(TestData.TestActor->AddComponentByClass(Settings->ComponentClass, false, FTransform::Identity, false));
+		Component->IntProperty = 42;
+
+		bSuccess &= VerifyAttributeValue(this, TestData, GET_MEMBER_NAME_CHECKED(UPCGUnitTestDummyComponent, IntProperty), 42ll, false, "PropertyToParamDataActorFindTest_Self_Tag_Component");
+		Settings->bSelectComponent = false;
+	}
+
+	// World by class
+	{
+		Settings->ActorSelection = EPCGActorSelection::ByClass;
+		Settings->ActorFilter = EPCGActorFilter::AllWorldActors;
+
+		PCGTestsCommon::FTestData TestData(Seed, Settings, APCGUnitTestDummyActor::StaticClass());
+		Cast<APCGUnitTestDummyActor>(TestData.TestActor)->IntProperty = 42;
+
+		bSuccess &= VerifyAttributeValue(this, TestData, GET_MEMBER_NAME_CHECKED(APCGUnitTestDummyActor, IntProperty), 42ll, false, "PropertyToParamDataActorFindTest_World_Class");
+	}
+
+	// World by tag
+	{
+		Settings->ActorSelection = EPCGActorSelection::ByTag;
+		Settings->ActorFilter = EPCGActorFilter::AllWorldActors;
+
+		PCGTestsCommon::FTestData TestData(Seed, Settings, APCGUnitTestDummyActor::StaticClass());
+		TestData.TestActor->Tags.Add(Tag);
+		Cast<APCGUnitTestDummyActor>(TestData.TestActor)->IntProperty = 42;
+
+		bSuccess &= VerifyAttributeValue(this, TestData, GET_MEMBER_NAME_CHECKED(APCGUnitTestDummyActor, IntProperty), 42ll, false, "PropertyToParamDataActorFindTest_World_Tag");
+	}
+
+	// World by name
+	{
+		Settings->ActorSelection = EPCGActorSelection::ByName;
+		Settings->ActorFilter = EPCGActorFilter::AllWorldActors;
+
+		PCGTestsCommon::FTestData TestData(Seed, Settings, APCGUnitTestDummyActor::StaticClass());
+		Cast<APCGUnitTestDummyActor>(TestData.TestActor)->IntProperty = 42;
+
+		bSuccess &= VerifyAttributeValue(this, TestData, GET_MEMBER_NAME_CHECKED(APCGUnitTestDummyActor, IntProperty), 42ll, false, "PropertyToParamDataActorFindTest_World_Name");
+	}
+
+	return bSuccess;
+}
+
+#endif // WITH_EDITOR
\ No newline at end of file
diff --git a/Engine/Plugins/Experimental/PCG/Source/PCG/Private/Tests/Elements/PCGPropertyToParamDataElementTest.h b/Engine/Plugins/Experimental/PCG/Source/PCG/Private/Tests/Elements/PCGPropertyToParamDataElementTest.h
new file mode 100644
--- /dev/null
+++ b/Engine/Plugins/Experimental/PCG/Source/PCG/Private/Tests/Elements/PCGPropertyToParamDataElementTest.h
@@ -0,0 +1,88 @@
+// Copyright Epic Games, Inc. All Rights Reserved.
+#pragma once
+
+#include "CoreMinimal.h"
+#include "GameFramework/Actor.h"
+
+#include "Tests/PCGTestsCommon.h"
+
+#include "PCGPropertyToParamDataElementTest.generated.h"
+
+UENUM()
+enum class EPCGUnitTestDummyEnum : int64
+{
+	One,
+	Two,
+	Three
+};
+
+UCLASS(MinimalAPI, NotBlueprintable, NotPlaceable, NotBlueprintType, Transient, HideDropdown)
+class APCGUnitTestDummyActor : public AActor
+{
+	GENERATED_BODY()
+
+public:
+	UPROPERTY()
+	int IntProperty = 0;
+
+	UPROPERTY()
+	float FloatProperty = 0.0f;
+
+	UPROPERTY()
+	int64 Int64Property = 0;
+
+	UPROPERTY()
+	double DoubleProperty = 0.0;
+
+	UPROPERTY()
+	bool BoolProperty = false;
+
+	UPROPERTY()
+	FName NameProperty = NAME_None;
+
+	UPROPERTY()
+	FString StringProperty = "";
+
+	UPROPERTY()
+	EPCGUnitTestDummyEnum EnumProperty = EPCGUnitTestDummyEnum::One;
+
+	UPROPERTY()
+	FVector VectorProperty;
+
+	UPROPERTY()
+	FVector4 Vector4Property;
+
+	UPROPERTY()
+	FTransform TransformProperty;
+
+	UPROPERTY()
+	FRotator RotatorProperty;
+
+	UPROPERTY()
+	FQuat QuatProperty;
+
+	UPROPERTY()
+	FSoftObjectPath SoftObjectPathProperty;
+
+	UPROPERTY()
+	FSoftClassPath SoftClassPathProperty;
+
+	UPROPERTY()
+	TSubclassOf<AActor> ClassProperty = AActor::StaticClass();
+
+	UPROPERTY()
+	UObject* ObjectProperty = nullptr;
+
+	UPROPERTY()
+	FVector2D Vector2Property;
+};
+
+UCLASS(MinimalAPI, NotBlueprintable, NotPlaceable, NotBlueprintType, Transient, HideDropdown)
+class UPCGUnitTestDummyComponent : public UActorComponent
+{
+	GENERATED_BODY()
+
+public:
+	UPROPERTY()
+	int IntProperty = 0;
+};
\ No newline at end of file
diff --git a/Engine/Plugins/Experimental/PCG/Source/PCG/Private/Tests/Elements/PCGStaticMeshSpawnerTest.cpp b/Engine/Plugins/Experimental/PCG/Source/PCG/Private/Tests/Elements/PCGStaticMeshSpawnerTest.cpp
--- a/Engine/Plugins/Experimental/PCG/Source/PCG/Private/Tests/Elements/PCGStaticMeshSpawnerTest.cpp
+++ b/Engine/Plugins/Experimental/PCG/Source/PCG/Private/Tests/Elements/PCGStaticMeshSpawnerTest.cpp
@@ -130,7 +130,7 @@ void ValidateOutPointData(
 
 void TestMeshSelectorByAttribute(
 	FPCGStaticMeshSpawnerByAttributeTest* Test, 
-	const PCGDeterminismTests::FTestData& TestData, 
+	const PCGTestsCommon::FTestData& TestData, 
 	const FName& AttributeName, 
 	const TArray<FPCGByAttributeValidationData>& ValidationDataEntries,
 	bool bValidateOutput=false)
@@ -258,7 +258,7 @@ void TestMeshSelectorByAttribute(
 
 void TestMeshSelectorWeighted(
 	FPCGStaticMeshSpawnerWeightedTest* Test, 
-	const PCGDeterminismTests::FTestData& TestData, 
+	const PCGTestsCommon::FTestData& TestData, 
 	const TArray<FPCGMeshSelectorWeightedEntry>& Entries, 
 	int PointCount,
 	bool bValidateOutput=false)
@@ -347,7 +347,7 @@ void TestMeshSelectorWeighted(
 
 void TestMeshSelectorWeightedByCategory(
 	FPCGStaticMeshSpawnerWeightedByCategoryTest* Test, 
-	const PCGDeterminismTests::FTestData& TestData, 
+	const PCGTestsCommon::FTestData& TestData, 
 	const FName& AttributeName, 
 	const TArray<FPCGWeightedByCategoryEntryList>& Entries, 
 	const TArray<FPCGWeightedByCategoryValidationData>& ValidationDataEntries,
@@ -430,7 +430,7 @@ void TestMeshSelectorWeightedByCategory(
 
 bool FPCGStaticMeshSpawnerByAttributeTest::RunTest(const FString& Parameters)
 {
-	PCGDeterminismTests::FTestData TestData(PCGDeterminismTests::Defaults::Seed);
+	PCGTestsCommon::FTestData TestData(PCGDeterminismTests::Defaults::Seed);
 	PCGDeterminismTests::GenerateSettings<UPCGStaticMeshSpawnerSettings>(TestData);
 
 	// create test attribute data
@@ -520,7 +520,7 @@ bool FPCGStaticMeshSpawnerByAttributeTest::RunTest(const FString& Parameters)
 
 bool FPCGStaticMeshSpawnerWeightedTest::RunTest(const FString& Parameters)
 {
-	PCGDeterminismTests::FTestData TestData(PCGDeterminismTests::Defaults::Seed);
+	PCGTestsCommon::FTestData TestData(PCGDeterminismTests::Defaults::Seed);
 	PCGDeterminismTests::GenerateSettings<UPCGStaticMeshSpawnerSettings>(TestData);
 
 	// create test meshes
@@ -589,7 +589,7 @@ bool FPCGStaticMeshSpawnerWeightedTest::RunTest(const FString& Parameters)
 
 bool FPCGStaticMeshSpawnerWeightedByCategoryTest::RunTest(const FString& Parameters)
 {
-	PCGDeterminismTests::FTestData TestData(PCGDeterminismTests::Defaults::Seed);
+	PCGTestsCommon::FTestData TestData(PCGDeterminismTests::Defaults::Seed);
 	PCGDeterminismTests::GenerateSettings<UPCGStaticMeshSpawnerSettings>(TestData);
 
 	// create test meshes
diff --git a/Engine/Plugins/Experimental/PCG/Source/PCG/Private/Tests/Elements/PCGTransformPointsTest.cpp b/Engine/Plugins/Experimental/PCG/Source/PCG/Private/Tests/Elements/PCGTransformPointsTest.cpp
--- a/Engine/Plugins/Experimental/PCG/Source/PCG/Private/Tests/Elements/PCGTransformPointsTest.cpp
+++ b/Engine/Plugins/Experimental/PCG/Source/PCG/Private/Tests/Elements/PCGTransformPointsTest.cpp
@@ -14,7 +14,7 @@ IMPLEMENT_CUSTOM_SIMPLE_AUTOMATION_TEST(FPCGTransformPointsTest, FPCGTestBaseCla
 
 bool FPCGTransformPointsTest::RunTest(const FString& Parameters)
 {
-	PCGDeterminismTests::FTestData TestData(PCGDeterminismTests::Defaults::Seed);
+	PCGTestsCommon::FTestData TestData(PCGDeterminismTests::Defaults::Seed);
 	PCGDeterminismTests::GenerateSettings<UPCGTransformPointsSettings>(TestData);
 	UPCGTransformPointsSettings* Settings = CastChecked<UPCGTransformPointsSettings>(TestData.Settings);
 	FPCGElementPtr TransformPointsElement = TestData.Settings->GetElement();
diff --git a/Engine/Plugins/Experimental/PCG/Source/PCG/Private/Tests/PCGTestsCommon.cpp b/Engine/Plugins/Experimental/PCG/Source/PCG/Private/Tests/PCGTestsCommon.cpp
--- a/Engine/Plugins/Experimental/PCG/Source/PCG/Private/Tests/PCGTestsCommon.cpp
+++ b/Engine/Plugins/Experimental/PCG/Source/PCG/Private/Tests/PCGTestsCommon.cpp
@@ -2,7 +2,8 @@
 
 #include "Tests/PCGTestsCommon.h"
 
-#include "PCGHelpers.h"
+#include "PCGComponent.h"
+#include "PCGGraph.h"
 #include "PCGParamData.h"
 #include "PCGSettings.h"
 #include "Data/PCGPointData.h"
@@ -11,10 +12,70 @@
 #include "Data/PCGSurfaceData.h"
 #include "Data/PCGVolumeData.h"
 
-#include "GameFramework/Actor.h"
+#if WITH_EDITOR
+#include "Editor.h"
+#include "Engine/World.h"
+#endif
 
 namespace PCGTestsCommon
 {
+	FTestData::FTestData(int32 RandomSeed, UPCGSettings* DefaultSettings, TSubclassOf<AActor> ActorClass)
+		: Settings(DefaultSettings)
+		, Seed(RandomSeed)
+		, RandomStream(Seed)
+	{
+#if WITH_EDITOR
+		check(GEditor);
+		UWorld* EditorWorld = GEditor->GetEditorWorldContext().World();
+		check(EditorWorld);
+
+		// No getting the level dirty
+		FActorSpawnParameters TransientActorParameters;
+		TransientActorParameters.bHideFromSceneOutliner = true;
+		TransientActorParameters.bTemporaryEditorActor = true;
+		TransientActorParameters.ObjectFlags = RF_Transient;
+		TestActor = EditorWorld->SpawnActor<AActor>(ActorClass, TransientActorParameters);
+		check(TestActor);
+
+		TestPCGComponent = NewObject<UPCGComponent>(TestActor, FName(TEXT("Test PCG Component")), RF_Transient);
+		check(TestPCGComponent);
+		TestActor->AddInstanceComponent(TestPCGComponent);
+		TestPCGComponent->RegisterComponent();
+
+		UPCGGraph* TestGraph = NewObject<UPCGGraph>(TestPCGComponent, FName(TEXT("Test PCG Graph")), RF_Transient);
+		check(TestGraph);
+		TestPCGComponent->SetGraph(TestGraph);
+#else
+		TestActor = nullptr;
+		TestPCGComponent = nullptr;
+		Settings = nullptr;
+#endif
+	}
+
+	FTestData::~FTestData()
+	{
+#if WITH_EDITOR
+		if (GEditor)
+		{
+			if (UWorld* EditorWorld = GEditor->GetEditorWorldContext().World())
+			{
+				if (TestActor)
+				{
+					EditorWorld->DestroyActor(TestActor);
+				}
+			}
+		}
+#endif // WITH_EDITOR
+	}
+
+	void FTestData::Reset()
+	{
+		// Clear all the data
+		RandomStream.Reset();
+		InputData.TaggedData.Empty();
+		Settings = nullptr;
+	}
+
 	AActor* CreateTemporaryActor()
 	{
 		return NewObject<AActor>();
diff --git a/Engine/Plugins/Experimental/PCG/Source/PCG/Public/Elements/PCGPropertyToParamData.h b/Engine/Plugins/Experimental/PCG/Source/PCG/Public/Elements/PCGPropertyToParamData.h
--- a/Engine/Plugins/Experimental/PCG/Source/PCG/Public/Elements/PCGPropertyToParamData.h
+++ b/Engine/Plugins/Experimental/PCG/Source/PCG/Public/Elements/PCGPropertyToParamData.h
@@ -58,7 +58,7 @@ public:
 	FName ActorSelectionTag;
 
 	UPROPERTY(BlueprintReadWrite, EditAnywhere, Category = Settings, meta = (EditCondition = "ActorSelection==EPCGActorSelection::ByName", EditConditionHides))
-	FString ActorSelectionName;
+	FName ActorSelectionName;
 
 	UPROPERTY(BlueprintReadWrite, EditAnywhere, Category = Settings, meta = (EditCondition = "ActorSelection==EPCGActorSelection::ByClass", EditConditionHides))
 	TSubclassOf<AActor> ActorSelectionClass;
diff --git a/Engine/Plugins/Experimental/PCG/Source/PCG/Public/Tests/Determinism/PCGDeterminismTestsCommon.h b/Engine/Plugins/Experimental/PCG/Source/PCG/Public/Tests/Determinism/PCGDeterminismTestsCommon.h
--- a/Engine/Plugins/Experimental/PCG/Source/PCG/Public/Tests/Determinism/PCGDeterminismTestsCommon.h
+++ b/Engine/Plugins/Experimental/PCG/Source/PCG/Public/Tests/Determinism/PCGDeterminismTestsCommon.h
@@ -39,20 +39,6 @@ namespace PCGDeterminismTests
 		const FBox TestingVolume = FBox(-1.2 * LargeVector, 1.2 * LargeVector);
 	}
 
-	struct FTestData
-	{
-		explicit FTestData(int32 Seed = Defaults::Seed, UPCGSettings* DefaultSettings = nullptr);
-
-		void Reset();
-
-		AActor* TestActor;
-		UPCGComponent* TestPCGComponent;
-		FPCGDataCollection InputData;
-		UPCGSettings* Settings;
-		int32 Seed;
-		FRandomStream RandomStream;
-	};
-
 	struct FNodeTestResult
 	{
 		int32 Index = -1;
@@ -102,7 +88,7 @@ namespace PCGDeterminismTests
 	bool RunAllDataOrderIndependenceTest(const UPCGNode* InPCGNode, int32 Seed, EPCGDataType& OutDataTypesTested, TArray<FString>& OutAdditionalDetails);
 
 	/** Adds input data to test data, based on input pins' allowed types */
-	void AddInputDataBasedOnPins(FTestData& TestData, const UPCGNode* InPCGNode, EPCGDataType& OutDataTypesTested, TArray<FString>& OutAdditionalDetails);
+	void AddInputDataBasedOnPins(PCGTestsCommon::FTestData& TestData, const UPCGNode* InPCGNode, EPCGDataType& OutDataTypesTested, TArray<FString>& OutAdditionalDetails);
 
 	/** Helper for adding PointData with a single Point to an InputData */
 	void AddSinglePointInputData(FPCGDataCollection& InputData, const FVector& Location, const FName& PinName = Defaults::TestPinName);
@@ -116,17 +102,17 @@ namespace PCGDeterminismTests
 	void AddPrimitiveInputData(FPCGDataCollection& InputData, UPrimitiveComponent* SplineComponent, const FVector& VoxelSize, const FName& PinName = Defaults::TestPinName);
 
 	/** Adds randomized PointData with a single Point */
-	void AddRandomizedSinglePointInputData(FTestData& TestData, int32 PointNum = Defaults::NumPointsToGenerate, const FName& PinName = Defaults::TestPinName);
+	void AddRandomizedSinglePointInputData(PCGTestsCommon::FTestData& TestData, int32 PointNum = Defaults::NumPointsToGenerate, const FName& PinName = Defaults::TestPinName);
 	/** Adds randomized PointData with multiple Points */
-	void AddRandomizedMultiplePointInputData(FTestData& TestData, int32 PointNum = Defaults::NumPointsToGenerate, const FName& PinName = Defaults::TestPinName);
+	void AddRandomizedMultiplePointInputData(PCGTestsCommon::FTestData& TestData, int32 PointNum = Defaults::NumPointsToGenerate, const FName& PinName = Defaults::TestPinName);
 	/** Adds randomized Volume Spatial Data */
-	void AddRandomizedVolumeInputData(FTestData& TestData, const FName& PinName = Defaults::TestPinName);
+	void AddRandomizedVolumeInputData(PCGTestsCommon::FTestData& TestData, const FName& PinName = Defaults::TestPinName);
 	/** Adds randomized Surface Spatial Data */
-	void AddRandomizedSurfaceInputData(FTestData& TestData, const FName& PinName = Defaults::TestPinName);
+	void AddRandomizedSurfaceInputData(PCGTestsCommon::FTestData& TestData, const FName& PinName = Defaults::TestPinName);
 	/** Adds randomized PolyLine Spatial Data */
-	void AddRandomizedPolyLineInputData(FTestData& TestData, int32 PointNum = Defaults::NumPolyLinePointsToGenerate, const FName& PinName = Defaults::TestPinName);
+	void AddRandomizedPolyLineInputData(PCGTestsCommon::FTestData& TestData, int32 PointNum = Defaults::NumPolyLinePointsToGenerate, const FName& PinName = Defaults::TestPinName);
 	/** Adds randomized Primitive Spatial Data */
-	void AddRandomizedPrimitiveInputData(FTestData& TestData, const FName& PinName = Defaults::TestPinName);
+	void AddRandomizedPrimitiveInputData(PCGTestsCommon::FTestData& TestData, const FName& PinName = Defaults::TestPinName);
 
 	/** Validates whether two DataCollection objects are identical */
 	bool DataCollectionsAreIdentical(const FPCGDataCollection& FirstCollection, const FPCGDataCollection& SecondCollection);
@@ -160,21 +146,21 @@ namespace PCGDeterminismTests
 	bool DataCanBeShuffled(const UPCGData* Data);
 
 	/** Randomizes the order of InputData collections */
-	void ShuffleInputOrder(FTestData& TestData);
+	void ShuffleInputOrder(PCGTestsCommon::FTestData& TestData);
 	/** Randomizes the order of all internal data */
-	void ShuffleAllInternalData(FTestData& TestData);
+	void ShuffleAllInternalData(PCGTestsCommon::FTestData& TestData);
 
 	/** Gets a comparison function to compare two of a specific DataType */
 	TFunction<bool(const UPCGData*, const UPCGData*)> GetCompareFunction(EPCGDataType DataType);
 
 	/** Execute the elements for each valid input and compare if all the outputs are identical */
-	bool ExecutionIsDeterministic(const FTestData& FirstTestData, const FTestData& SecondTestData, const UPCGNode* PCGNode = nullptr);
+	bool ExecutionIsDeterministic(const PCGTestsCommon::FTestData& FirstTestData, const PCGTestsCommon::FTestData& SecondTestData, const UPCGNode* PCGNode = nullptr);
 	/** Execute the same element twice compare if all the outputs are identical */
-	bool ExecutionIsDeterministicSameData(FTestData& TestData, const UPCGNode* PCGNode = nullptr);
+	bool ExecutionIsDeterministicSameData(PCGTestsCommon::FTestData& TestData, const UPCGNode* PCGNode = nullptr);
 
 	/** Generates settings based upon a UPCGSettings subclass */
 	template<typename SettingsType>
-	SettingsType* GenerateSettings(FTestData& TestData, TFunction<void(FTestData&)> ExtraSettingsDelegate = nullptr)
+	SettingsType* GenerateSettings(PCGTestsCommon::FTestData& TestData, TFunction<void(PCGTestsCommon::FTestData&)> ExtraSettingsDelegate = nullptr)
 	{
 		SettingsType* TypedSettings = NewObject<SettingsType>();
 		check(TypedSettings);
diff --git a/Engine/Plugins/Experimental/PCG/Source/PCG/Public/Tests/PCGTestsCommon.h b/Engine/Plugins/Experimental/PCG/Source/PCG/Public/Tests/PCGTestsCommon.h
--- a/Engine/Plugins/Experimental/PCG/Source/PCG/Public/Tests/PCGTestsCommon.h
+++ b/Engine/Plugins/Experimental/PCG/Source/PCG/Public/Tests/PCGTestsCommon.h
@@ -2,18 +2,20 @@
 
 #pragma once
 
+#include "PCGData.h"
+
+#include "GameFramework/Actor.h"
 #include "Misc/AutomationTest.h"
 #include "Tests/AutomationCommon.h"
 
-class AActor;
+class UPCGComponent;
 class UPCGParamData;
 class UPCGPointData;
 class UPCGPolyLineData;
 class UPCGPrimitiveData;
+class UPCGSettings;
 class UPCGSurfaceData;
 class UPCGVolumeData;
-class UPCGSettings;
-struct FPCGDataCollection;
 struct FPCGPinProperties;
 struct FPCGPoint;
 
@@ -21,6 +23,21 @@ namespace PCGTestsCommon
 {
 	static const int TestFlags = EAutomationTestFlags::EditorContext | EAutomationTestFlags::EngineFilter;
 
+	struct FTestData
+	{
+		FTestData(int32 Seed, UPCGSettings* DefaultSettings = nullptr, TSubclassOf<AActor> ActorClass = AActor::StaticClass());
+		~FTestData();
+
+		void Reset();
+
+		AActor* TestActor;
+		UPCGComponent* TestPCGComponent;
+		FPCGDataCollection InputData;
+		UPCGSettings* Settings;
+		int32 Seed;
+		FRandomStream RandomStream;
+	};
+
 	AActor* CreateTemporaryActor();
 	UPCGPointData* CreateEmptyPointData();
 	UPCGParamData* CreateEmptyParamData();
