Implement a native DensityRemap PCG Element

diff --git a/Engine/Plugins/Experimental/PCG/Source/PCG/Private/Elements/PCGDensityRemapElement.cpp b/Engine/Plugins/Experimental/PCG/Source/PCG/Private/Elements/PCGDensityRemapElement.cpp
new file mode 100644
--- /dev/null
+++ b/Engine/Plugins/Experimental/PCG/Source/PCG/Private/Elements/PCGDensityRemapElement.cpp
@@ -0,0 +1,104 @@
+// Copyright Epic Games, Inc. All Rights Reserved.
+
+#include "Elements/PCGDensityRemapElement.h"
+
+#include "PCGHelpers.h"
+#include "Data/PCGPointData.h"
+#include "Helpers/PCGAsync.h"
+#include "Helpers/PCGSettingsHelpers.h"
+
+#include "Math/RandomStream.h"
+
+FPCGElementPtr UPCGDensityRemapSettings::CreateElement() const
+{
+	return MakeShared<FPCGDensityRemapElement>();
+}
+
+bool FPCGDensityRemapElement::ExecuteInternal(FPCGContext* Context) const
+{
+	TRACE_CPUPROFILER_EVENT_SCOPE(FPCGDensityRemapElement::Execute);
+
+	const UPCGDensityRemapSettings* Settings = Context->GetInputSettings<UPCGDensityRemapSettings>();
+	check(Settings);
+
+	TArray<FPCGTaggedData> Inputs = Context->InputData.GetInputs();
+	TArray<FPCGTaggedData>& Outputs = Context->OutputData.TaggedData;
+	UPCGParamData* Params = Context->InputData.GetParams();
+
+	const float InRangeMin = PCG_GET_OVERRIDEN_VALUE(Settings, InRangeMin, Params);
+	const float InRangeMax = PCG_GET_OVERRIDEN_VALUE(Settings, InRangeMax, Params);
+	const float OutRangeMin = PCG_GET_OVERRIDEN_VALUE(Settings, OutRangeMin, Params);
+	const float OutRangeMax = PCG_GET_OVERRIDEN_VALUE(Settings, OutRangeMax, Params);
+	const bool bExcludeValuesOutsideInputRange = PCG_GET_OVERRIDEN_VALUE(Settings, bExcludeValuesOutsideInputRange, Params);
+
+	// used to determine if a density value lies between TrueMin and TrueMax
+	const float InRangeTrueMin = FMath::Min(InRangeMin, InRangeMax);
+	const float InRangeTrueMax = FMath::Max(InRangeMin, InRangeMax);
+
+	const float InRangeDifference = InRangeMax - InRangeMin;
+	const float OutRangeDifference = OutRangeMax - OutRangeMin;
+
+	float Slope;
+	float Intercept;
+
+	// When InRange is a point leave the Slope at 0 so that Density = Intercept
+	if (InRangeDifference == 0)
+	{
+		Slope = 0;
+		Intercept = (OutRangeMin + OutRangeMax) / 2.f;
+	}
+	else
+	{
+		Slope = OutRangeDifference / InRangeDifference;
+		Intercept = OutRangeMin;
+	}
+
+	// TODO: embarassingly parallel loop
+	for (const FPCGTaggedData& Input : Inputs)
+	{
+		TRACE_CPUPROFILER_EVENT_SCOPE(FPCGDensityRemapElement::Execute::InputLoop);
+		FPCGTaggedData& Output = Outputs.Add_GetRef(Input);
+
+		const UPCGSpatialData* SpatialData = Cast<UPCGSpatialData>(Input.Data);
+
+		if (!SpatialData)
+		{
+			PCGE_LOG(Error, "Unable to get SpatialData from input");
+			continue;
+		}
+
+		const UPCGPointData* PointData = SpatialData->ToPointData(Context);
+
+		if (!PointData)
+		{
+			PCGE_LOG(Error, "Unable to get PointData from input");
+			continue;
+		}
+
+		const TArray<FPCGPoint>& Points = PointData->GetPoints();
+
+		UPCGPointData* SampledData = NewObject<UPCGPointData>();
+		SampledData->InitializeFromData(PointData);
+		TArray<FPCGPoint>& SampledPoints = SampledData->GetMutablePoints();
+		Output.Data = SampledData;
+
+		FPCGAsync::AsyncPointProcessing(Context, Points.Num(), SampledPoints, [&](int32 Index, FPCGPoint& OutPoint)
+		{
+			OutPoint = Points[Index];
+			const float SourceDensity = Points[Index].Density;
+
+			if (!bExcludeValuesOutsideInputRange || (SourceDensity >= InRangeTrueMin && SourceDensity <= InRangeTrueMax))
+			{
+				const float UnclampedDensity = Slope * (SourceDensity - InRangeMin) + Intercept;
+				OutPoint.Density = FMath::Clamp(UnclampedDensity, 0.f, 1.f);
+			}
+
+			return true;
+		});
+	}
+
+	// Forward any non-input data
+	Outputs.Append(Context->InputData.GetAllSettings());
+
+	return true;
+}
diff --git a/Engine/Plugins/Experimental/PCG/Source/PCG/Private/Tests/Elements/PCGDensityRemapTest.cpp b/Engine/Plugins/Experimental/PCG/Source/PCG/Private/Tests/Elements/PCGDensityRemapTest.cpp
new file mode 100644
--- /dev/null
+++ b/Engine/Plugins/Experimental/PCG/Source/PCG/Private/Tests/Elements/PCGDensityRemapTest.cpp
@@ -0,0 +1,175 @@
+// Copyright Epic Games, Inc. All Rights Reserved.
+
+#include "Tests/Determinism/PCGDeterminismTestsCommon.h"
+#include "Tests/PCGTestsCommon.h"
+
+#include "PCGComponent.h"
+#include "PCGData.h"
+#include "PCGHelpers.h"
+
+#include "Data/PCGPointData.h"
+#include "Elements/PCGDensityRemapElement.h"
+
+IMPLEMENT_CUSTOM_SIMPLE_AUTOMATION_TEST(FPCGDensityRemapTest, FPCGTestBaseClass, "pcg.tests.DensityRemap.Basic", PCGTestsCommon::TestFlags)
+
+bool FPCGDensityRemapTest::RunTest(const FString& Parameters)
+{
+	PCGDeterminismTests::FTestData TestData(PCGDeterminismTests::Defaults::Seed);
+	PCGDeterminismTests::GenerateSettings<UPCGDensityRemapSettings>(TestData);
+	UPCGDensityRemapSettings* Settings = CastChecked<UPCGDensityRemapSettings>(TestData.Settings);
+	FPCGElementPtr DensityRemapElement = TestData.Settings->GetElement();
+
+	TObjectPtr<UPCGPointData> PointData = PCGTestsCommon::CreateEmptyPointData();
+	TArray<FPCGPoint>& Points = PointData->GetMutablePoints();
+
+	FRandomStream RandomSource(TestData.Seed);
+	const int PointCount = 6;
+	for (int I = 0; I < PointCount; ++I)
+	{
+		FPCGPoint& Point = Points.Emplace_GetRef(FTransform(), 1, I);
+		Point.Density = static_cast<float>(I) / (PointCount - 1);
+	}
+
+	FPCGTaggedData& TaggedData = TestData.InputData.TaggedData.Emplace_GetRef(FPCGTaggedData());
+	TaggedData.Data = PointData;
+
+	auto ValidateDensityRemap = [this, &TestData, DensityRemapElement, Settings](TArray<float> CorrectDensities) -> bool
+	{
+		TUniquePtr<FPCGContext> Context = MakeUnique<FPCGContext>(*DensityRemapElement->Initialize(TestData.InputData, TestData.TestPCGComponent, nullptr));
+		Context->NumAvailableTasks = 1;
+
+		while (!DensityRemapElement->Execute(Context.Get()))
+		{}
+
+		const TArray<FPCGTaggedData>& Inputs = Context->InputData.GetInputs();
+		const TArray<FPCGTaggedData>& Outputs = Context->OutputData.GetInputs();
+		
+		if (!TestEqual("Valid number of outputs", Inputs.Num(), Outputs.Num()))
+		{
+			return false;
+		}
+	
+		bool bTestPassed = true;
+
+		for (int DataIndex = 0; DataIndex < Inputs.Num(); ++DataIndex)
+		{
+			const FPCGTaggedData& Input = Inputs[DataIndex];
+			const FPCGTaggedData& Output = Outputs[DataIndex];
+
+			const UPCGSpatialData* InSpatialData = Cast<UPCGSpatialData>(Input.Data);
+			check(InSpatialData);
+
+			const UPCGPointData* InPointData = InSpatialData->ToPointData(Context.Get());
+			check(InPointData);
+
+			const UPCGSpatialData* OutSpatialData = Cast<UPCGSpatialData>(Output.Data);
+
+			if (!TestNotNull("Valid output SpatialData", OutSpatialData))
+			{
+				bTestPassed = false;
+				continue;
+			}
+
+			const UPCGPointData* OutPointData = OutSpatialData->ToPointData(Context.Get());
+
+			if (!TestNotNull("Valid output PointData", OutPointData))
+			{
+				bTestPassed = false;
+				continue;
+			}
+
+			const TArray<FPCGPoint>& InPoints = InPointData->GetPoints();
+			const TArray<FPCGPoint>& OutPoints = OutPointData->GetPoints();
+
+			if (!TestEqual("Input and output point counts match", InPoints.Num(), OutPoints.Num()))
+			{ 
+				bTestPassed = false;
+				continue;
+			}
+
+			for (int PointIndex = 0; PointIndex < InPoints.Num(); ++PointIndex)
+			{
+				bTestPassed &= TestEqual("Correct density", OutPoints[PointIndex].Density, CorrectDensities[PointIndex]);
+			}
+		}
+
+		return bTestPassed;
+	};
+
+	bool bTestPassed = true;
+
+	// Test [0-1] -> [0-1]
+	{
+		Settings->InRangeMin = 0.f;
+		Settings->InRangeMax = 1.f;
+		Settings->OutRangeMin = 0.f;
+		Settings->OutRangeMax = 1.f;
+		bTestPassed &= TestTrue("Input and Output are identical when InRange and OutRange are identical", ValidateDensityRemap({ 0.f, 0.2f, 0.4f, 0.6f, 0.8f, 1.f }));
+	}
+
+	// Test [0-0.4] -> [0-1]
+	{
+		Settings->InRangeMin = 0.f;
+		Settings->InRangeMax = 0.4f;
+		Settings->OutRangeMin = 0.f;
+		Settings->OutRangeMax = 1.f;
+		bTestPassed &= TestTrue("Valid densities for partial InRange", ValidateDensityRemap({ 0.f, 0.5f, 1.f, 0.6f, 0.8f, 1.f }));
+
+		Settings->InRangeMin = 0.4f;
+		Settings->InRangeMax = 0.f;
+		Settings->OutRangeMin = 1.f;
+		Settings->OutRangeMax = 0.f;
+		bTestPassed &= TestTrue("Inverting ranges does not effect output", ValidateDensityRemap({ 0.f, 0.5f, 1.f, 0.6f, 0.8f, 1.f }));
+	}
+
+	// Test [0-0.4] -> [0.5-1]
+	{
+		Settings->InRangeMin = 0.4f;
+		Settings->InRangeMax = 1.f;
+		Settings->OutRangeMin = 0.5f;
+		Settings->OutRangeMax = 1.f;
+		bTestPassed &= TestTrue("Valid densities for partial OutRange", ValidateDensityRemap({ 0.f, 0.2f, 3.f / 6.f, 4.f / 6.f, 5.f / 6.f, 1.f }));
+
+		Settings->InRangeMin = 1.f;
+		Settings->InRangeMax = 0.4f;
+		Settings->OutRangeMin = 1.f;
+		Settings->OutRangeMax = 0.5f;
+		bTestPassed &= TestTrue("Inverting ranges does not effect output", ValidateDensityRemap({ 0.f, 0.2f, 3.f / 6.f, 4.f / 6.f, 5.f / 6.f, 1.f }));
+	}
+
+	// Test disabling Range Exclusion
+	{
+		Settings->InRangeMin = 0.4f;
+		Settings->InRangeMax = 1.f;
+		Settings->OutRangeMin = 0.5f;
+		Settings->OutRangeMax = 1.f;
+		Settings->bExcludeValuesOutsideInputRange = false;
+
+		bTestPassed &= TestTrue("All values are remapped when Range Exclusion is disabled", ValidateDensityRemap({ 1.f / 6.f, 2.f / 6.f, 3.f / 6.f, 4.f / 6.f, 5.f / 6.f, 1.f }));
+
+		Settings->bExcludeValuesOutsideInputRange = true;
+	}
+
+	// Test Point to Point, Point to Range, and Range to Point
+	{
+		Settings->InRangeMin = 0.2f;
+		Settings->InRangeMax = 0.2f;
+		Settings->OutRangeMin = 0.f;
+		Settings->OutRangeMax = 1.f;
+		bTestPassed &= TestTrue("Point input to Range output", ValidateDensityRemap({ 0.f, 0.5f, 0.4f, 0.6f, 0.8f, 1.f }));
+
+		Settings->InRangeMin = 0.2f;
+		Settings->InRangeMax = 0.2f;
+		Settings->OutRangeMin = 0.5f;
+		Settings->OutRangeMax = 0.5f;
+		bTestPassed &= TestTrue("Point input to Point output", ValidateDensityRemap({ 0.f, 0.5f, 0.4f, 0.6f, 0.8f, 1.f }));
+
+		Settings->InRangeMin = 0.2f;
+		Settings->InRangeMax = 1.f;
+		Settings->OutRangeMin = 0.5f;
+		Settings->OutRangeMax = 0.5f;
+		bTestPassed &= TestTrue("Range input to Point output", ValidateDensityRemap({ 0.f, 0.5f, 0.5f, 0.5f, 0.5f, 0.5f }));
+	}
+
+	return bTestPassed;
+}
diff --git a/Engine/Plugins/Experimental/PCG/Source/PCG/Public/Elements/PCGDensityRemapElement.h b/Engine/Plugins/Experimental/PCG/Source/PCG/Public/Elements/PCGDensityRemapElement.h
new file mode 100644
--- /dev/null
+++ b/Engine/Plugins/Experimental/PCG/Source/PCG/Public/Elements/PCGDensityRemapElement.h
@@ -0,0 +1,54 @@
+// Copyright Epic Games, Inc. All Rights Reserved.
+
+#pragma once
+
+#include "CoreMinimal.h"
+#include "PCGElement.h"
+#include "PCGNode.h"
+#include "PCGSettings.h"
+
+#include "PCGDensityRemapElement.generated.h"
+
+UCLASS(BlueprintType, ClassGroup = (Procedural))
+class PCG_API UPCGDensityRemapSettings : public UPCGSettings
+{
+	GENERATED_BODY()
+
+public:
+	// ~Begin UPCGSettings interface
+#if WITH_EDITOR
+	virtual FName GetDefaultNodeName() const override { return FName(TEXT("DensityRemapNode")); }
+	virtual EPCGSettingsType GetType() const override { return EPCGSettingsType::Density; }
+#endif
+
+	virtual TArray<FPCGPinProperties> OutputPinProperties() const override { return Super::DefaultPointOutputPinProperties(); }
+
+protected:
+	virtual FPCGElementPtr CreateElement() const override;
+	//~End UPCGSettings interface
+
+public:
+	/** If InRangeMin = InRangeMax, then that density value is mapped to the average of OutRangeMin and OutRangeMax */
+	UPROPERTY(BlueprintReadWrite, EditAnywhere, Category = Settings, meta = (ClampMin = "0", ClampMax = "1"))
+	float InRangeMin = 0.f;
+
+	/** If InRangeMin = InRangeMax, then that density value is mapped to the average of OutRangeMin and OutRangeMax */
+	UPROPERTY(BlueprintReadWrite, EditAnywhere, Category = Settings, meta = (ClampMin = "0", ClampMax = "1"))
+	float InRangeMax = 1.f;
+
+	UPROPERTY(BlueprintReadWrite, EditAnywhere, Category = Settings, meta = (ClampMin = "0", ClampMax = "1"))
+	float OutRangeMin = 0.f;
+
+	UPROPERTY(BlueprintReadWrite, EditAnywhere, Category = Settings, meta = (ClampMin = "0", ClampMax = "1"))
+	float OutRangeMax = 1.f;
+
+	/** Density values outside of the input range will be unaffected by the remapping */
+	UPROPERTY(BlueprintReadWrite, EditAnywhere, Category = Settings)
+	bool bExcludeValuesOutsideInputRange = true;
+};
+
+class FPCGDensityRemapElement : public FSimplePCGElement
+{
+protected:
+	virtual bool ExecuteInternal(FPCGContext* Context) const;
+};
