Implement a native DensityNoise PCG Element

diff --git a/Engine/Plugins/Experimental/PCG/Source/PCG/Private/Elements/PCGDensityNoise.cpp b/Engine/Plugins/Experimental/PCG/Source/PCG/Private/Elements/PCGDensityNoise.cpp
new file mode 100644
--- /dev/null
+++ b/Engine/Plugins/Experimental/PCG/Source/PCG/Private/Elements/PCGDensityNoise.cpp
@@ -0,0 +1,105 @@
+// Copyright Epic Games, Inc. All Rights Reserved.
+
+#include "Elements/PCGDensityNoise.h"
+
+#include "PCGHelpers.h"
+#include "Data/PCGPointData.h"
+#include "Helpers/PCGAsync.h"
+#include "Helpers/PCGSettingsHelpers.h"
+
+#include "Math/RandomStream.h"
+
+FPCGElementPtr UPCGDensityNoiseSettings::CreateElement() const
+{
+	return MakeShared<FPCGDensityNoiseElement>();
+}
+
+bool FPCGDensityNoiseElement::ExecuteInternal(FPCGContext* Context) const
+{
+	TRACE_CPUPROFILER_EVENT_SCOPE(FPCGDensityNoiseElement::Execute);
+
+	const UPCGDensityNoiseSettings* Settings = Context->GetInputSettings<UPCGDensityNoiseSettings>();
+	check(Settings);
+
+	TArray<FPCGTaggedData> Inputs = Context->InputData.GetInputs();
+	TArray<FPCGTaggedData>& Outputs = Context->OutputData.TaggedData;
+	UPCGParamData* Params = Context->InputData.GetParams();
+
+	const EPCGDensityNoiseMode DensityMode = PCG_GET_OVERRIDEN_VALUE(Settings, DensityMode, Params);
+	const float DensityNoiseMin = PCG_GET_OVERRIDEN_VALUE(Settings, DensityNoiseMin, Params);
+	const float DensityNoiseMax = PCG_GET_OVERRIDEN_VALUE(Settings, DensityNoiseMax, Params);
+	const bool bInvertSourceDensity = PCG_GET_OVERRIDEN_VALUE(Settings, bInvertSourceDensity, Params);
+
+	// TODO: embarassingly parallel loop
+	for (const FPCGTaggedData& Input : Inputs)
+	{
+		TRACE_CPUPROFILER_EVENT_SCOPE(FPCGDensityNoiseElement::Execute::InputLoop);
+		FPCGTaggedData& Output = Outputs.Add_GetRef(Input);
+
+		const UPCGSpatialData* SpatialData = Cast<UPCGSpatialData>(Input.Data);
+
+		if (!SpatialData)
+		{
+			PCGE_LOG(Error, "Unable to get SpatialData from input");
+			continue;
+		}
+
+		const UPCGPointData* PointData = SpatialData->ToPointData(Context);
+
+		if (!PointData)
+		{
+			PCGE_LOG(Error, "Unable to get PointData from input");
+			continue;
+		}
+
+		const TArray<FPCGPoint>& Points = PointData->GetPoints();
+
+		UPCGPointData* SampledData = NewObject<UPCGPointData>();
+		SampledData->InitializeFromData(PointData);
+		TArray<FPCGPoint>& SampledPoints = SampledData->GetMutablePoints();
+		Output.Data = SampledData;
+
+		FPCGAsync::AsyncPointProcessing(Context, Points.Num(), SampledPoints, [&Points, Settings, DensityMode, DensityNoiseMin, DensityNoiseMax, bInvertSourceDensity](int32 Index, FPCGPoint& OutPoint)
+		{
+			OutPoint = Points[Index];
+			FRandomStream RandomSource(PCGHelpers::ComputeSeed(Settings->Seed, OutPoint.Seed));
+
+			const float DensityNoise = RandomSource.FRandRange(DensityNoiseMin, DensityNoiseMax);
+
+			// This inversion was previously calculated as OutPoint.Density *= 1.f - FMath::Abs(OutPoint.Density * 2.f - 1.f);
+			const float SourceDensity = bInvertSourceDensity ? (1.0f - OutPoint.Density) : OutPoint.Density;
+
+			float UnclampedDensity = 0;
+
+			if (DensityMode == EPCGDensityNoiseMode::Minimum)
+			{
+				UnclampedDensity = FMath::Min(SourceDensity, DensityNoise);
+			}
+			else if (DensityMode == EPCGDensityNoiseMode::Maximum)
+			{
+				UnclampedDensity = FMath::Max(SourceDensity, DensityNoise);
+			}
+			else if (DensityMode == EPCGDensityNoiseMode::Add)
+			{
+				UnclampedDensity = SourceDensity + DensityNoise;
+			}
+			else if (DensityMode == EPCGDensityNoiseMode::Multiply)
+			{
+				UnclampedDensity = SourceDensity * DensityNoise;
+			}
+			else //if (DensityMode == EPCGDensityNoiseMode::Set)
+			{
+				UnclampedDensity = DensityNoise;
+			}
+
+			OutPoint.Density = FMath::Clamp(UnclampedDensity, 0.f, 1.f);
+
+			return true;
+		});
+	}
+
+	// Forward any non-input data
+	Outputs.Append(Context->InputData.GetAllSettings());
+
+	return true;
+}
diff --git a/Engine/Plugins/Experimental/PCG/Source/PCG/Private/Tests/Elements/PCGDensityNoiseTest.cpp b/Engine/Plugins/Experimental/PCG/Source/PCG/Private/Tests/Elements/PCGDensityNoiseTest.cpp
new file mode 100644
--- /dev/null
+++ b/Engine/Plugins/Experimental/PCG/Source/PCG/Private/Tests/Elements/PCGDensityNoiseTest.cpp
@@ -0,0 +1,234 @@
+// Copyright Epic Games, Inc. All Rights Reserved.
+
+#include "Tests/Determinism/PCGDeterminismTestsCommon.h"
+#include "Tests/PCGTestsCommon.h"
+
+#include "PCGComponent.h"
+#include "PCGData.h"
+#include "PCGHelpers.h"
+
+#include "Data/PCGPointData.h"
+#include "Elements/PCGDensityNoise.h"
+
+IMPLEMENT_CUSTOM_SIMPLE_AUTOMATION_TEST(FPCGDensityNoiseTest, FPCGTestBaseClass, "pcg.tests.DensityNoise.Basic", PCGTestsCommon::TestFlags)
+
+bool FPCGDensityNoiseTest::RunTest(const FString& Parameters)
+{
+	PCGDeterminismTests::FTestData TestData(PCGDeterminismTests::Defaults::Seed);
+	PCGDeterminismTests::GenerateSettings<UPCGDensityNoiseSettings>(TestData);
+	UPCGDensityNoiseSettings* Settings = CastChecked<UPCGDensityNoiseSettings>(TestData.Settings);
+	FPCGElementPtr DensityNoiseElement = TestData.Settings->GetElement();
+
+	TObjectPtr<UPCGPointData> PointData = PCGTestsCommon::CreateEmptyPointData();
+	TArray<FPCGPoint>& Points = PointData->GetMutablePoints();
+
+	FRandomStream RandomSource(TestData.Seed);
+	const int PointCount = 5;
+	for (int I = 0; I < PointCount; ++I)
+	{
+		FPCGPoint& Point = Points.Emplace_GetRef(FTransform(), 1, I);
+		Point.Density = RandomSource.GetFraction();
+	}
+
+	FPCGTaggedData& TaggedData = TestData.InputData.TaggedData.Emplace_GetRef(FPCGTaggedData());
+	TaggedData.Data = PointData;
+
+	auto ValidateDensityNoise = [this, &TestData, DensityNoiseElement, Settings](TArray<float> ExpectedOutput) -> bool
+	{
+		TUniquePtr<FPCGContext> Context = MakeUnique<FPCGContext>(*DensityNoiseElement->Initialize(TestData.InputData, TestData.TestPCGComponent, nullptr));
+		Context->NumAvailableTasks = 1;
+
+		while (!DensityNoiseElement->Execute(Context.Get()))
+		{}
+
+		const TArray<FPCGTaggedData>& Inputs = Context->InputData.GetInputs();
+		const TArray<FPCGTaggedData>& Outputs = Context->OutputData.GetInputs();
+
+		if (!TestEqual("Valid number of outputs", Inputs.Num(), Outputs.Num()))
+		{
+			return false;
+		}
+
+		bool bTestPassed = true;
+
+		for (int DataIndex = 0; DataIndex < Inputs.Num(); ++DataIndex)
+		{
+			const FPCGTaggedData& Input = Inputs[DataIndex];
+			const FPCGTaggedData& Output = Outputs[DataIndex];
+
+			const UPCGSpatialData* InSpatialData = Cast<UPCGSpatialData>(Input.Data);
+			check(InSpatialData);
+
+			const UPCGPointData* InPointData = InSpatialData->ToPointData(Context.Get());
+			check(InPointData);
+
+			const UPCGSpatialData* OutSpatialData = Cast<UPCGSpatialData>(Output.Data);
+
+			if (!TestNotNull("Valid output SpatialData", OutSpatialData))
+			{
+				bTestPassed = false;
+				continue;
+			}
+
+			const UPCGPointData* OutPointData = OutSpatialData->ToPointData(Context.Get());
+
+			if (!TestNotNull("Valid output PointData", OutPointData))
+			{
+				bTestPassed = false;
+				continue;
+			}
+
+			const TArray<FPCGPoint>& InPoints = InPointData->GetPoints();
+			const TArray<FPCGPoint>& OutPoints = OutPointData->GetPoints();
+
+			if (!TestEqual("Input and output point counts match", InPoints.Num(), OutPoints.Num()))
+			{ 
+				bTestPassed = false;
+				continue;
+			}
+
+			for (int PointIndex = 0; PointIndex < ExpectedOutput.Num(); ++PointIndex)
+			{
+				bTestPassed &= TestEqual("Correct density", OutPoints[PointIndex].Density, ExpectedOutput[PointIndex]);
+			}
+		}
+
+		return bTestPassed;
+	};
+
+	bool bTestPassed = true;
+
+	auto ValidateDensityNoiseForAllDensityModes = [this, &bTestPassed, ValidateDensityNoise, Settings]()
+	{
+		Settings->DensityMode = EPCGDensityNoiseMode::Set;
+		bTestPassed &= ValidateDensityNoise({});
+
+		Settings->DensityMode = EPCGDensityNoiseMode::Minimum;
+		bTestPassed &= ValidateDensityNoise({});
+
+		Settings->DensityMode = EPCGDensityNoiseMode::Maximum;
+		bTestPassed &= ValidateDensityNoise({});
+
+		Settings->DensityMode = EPCGDensityNoiseMode::Add;
+		bTestPassed &= ValidateDensityNoise({});
+
+		Settings->DensityMode = EPCGDensityNoiseMode::Multiply;
+		bTestPassed &= ValidateDensityNoise({});
+	};
+
+	// Test [0-1]
+	{
+		Settings->DensityNoiseMin = 0.f;
+		Settings->DensityNoiseMax = 1.f;
+		Settings->bInvertSourceDensity = false;
+		ValidateDensityNoiseForAllDensityModes();
+
+		Settings->bInvertSourceDensity = true;
+		ValidateDensityNoiseForAllDensityModes();
+	}
+
+	// Test [0-0.5]
+	{
+		Settings->DensityNoiseMin = 0.f;
+		Settings->DensityNoiseMax = 0.5f;
+		Settings->bInvertSourceDensity = false;
+		ValidateDensityNoiseForAllDensityModes();
+		
+		Settings->bInvertSourceDensity = true;
+		ValidateDensityNoiseForAllDensityModes();
+	}
+
+	// Test [0.5-1]
+	{
+		Settings->DensityNoiseMin = 0.5f;
+		Settings->DensityNoiseMax = 1.f;
+		Settings->bInvertSourceDensity = false;
+		ValidateDensityNoiseForAllDensityModes();
+
+		Settings->bInvertSourceDensity = true;
+		ValidateDensityNoiseForAllDensityModes();
+	}
+
+	// Test [1-0]
+	{
+		Settings->DensityNoiseMin = 1.f;
+		Settings->DensityNoiseMax = 0.f;
+		Settings->bInvertSourceDensity = false;
+		ValidateDensityNoiseForAllDensityModes();
+
+		Settings->bInvertSourceDensity = true;
+		ValidateDensityNoiseForAllDensityModes();
+	}
+
+	// Test [0.5-0]
+	{
+		Settings->DensityNoiseMin = 0.5f;
+		Settings->DensityNoiseMax = 0.f;
+		Settings->bInvertSourceDensity = false;
+		ValidateDensityNoiseForAllDensityModes();
+		
+		Settings->bInvertSourceDensity = true;
+		ValidateDensityNoiseForAllDensityModes();
+	}
+
+	// Test [1-0.5]
+	{
+		Settings->DensityNoiseMin = 1.f;
+		Settings->DensityNoiseMax = 0.5f;
+		Settings->bInvertSourceDensity = false;
+		ValidateDensityNoiseForAllDensityModes();
+
+		Settings->bInvertSourceDensity = true;
+		ValidateDensityNoiseForAllDensityModes();
+	}
+
+	Settings->bInvertSourceDensity = false;
+
+	// Test expected values for Set
+	{
+		Settings->DensityNoiseMin = 0.5f;
+		Settings->DensityNoiseMax = 0.5f;
+		Settings->DensityMode = EPCGDensityNoiseMode::Set;
+		bTestPassed &= ValidateDensityNoise({ 0.5f, 0.5f, 0.5f, 0.5f, 0.5f });
+	}
+
+	// Test expected values for Minimum
+	{
+		Settings->DensityNoiseMin = 0.f;
+		Settings->DensityNoiseMax = 0.f;
+		Settings->DensityMode = EPCGDensityNoiseMode::Minimum;
+		bTestPassed &= ValidateDensityNoise({ 0.0f, 0.0f, 0.0f, 0.0f, 0.0f });
+	}
+
+	// Test expected values for Maximum
+	{
+		Settings->DensityNoiseMin = 1.f;
+		Settings->DensityNoiseMax = 1.f;
+		Settings->DensityMode = EPCGDensityNoiseMode::Maximum;
+		bTestPassed &= ValidateDensityNoise({ 1.f, 1.f, 1.f, 1.f, 1.f });
+	}
+
+	// Test expected values for Add
+	{
+		Settings->DensityNoiseMin = 0.f;
+		Settings->DensityNoiseMax = 0.f;
+		Settings->DensityMode = EPCGDensityNoiseMode::Add;
+		bTestPassed &= ValidateDensityNoise({ 
+			Points[0].Density,
+			Points[1].Density,
+			Points[2].Density,
+			Points[3].Density,
+			Points[4].Density,
+		});
+	}
+
+	// Test expected values for Multiply
+	{
+		Settings->DensityNoiseMin = 0.f;
+		Settings->DensityNoiseMax = 0.f;
+		Settings->DensityMode = EPCGDensityNoiseMode::Multiply;
+		bTestPassed &= ValidateDensityNoise({ 0.0f, 0.0f, 0.0f, 0.0f, 0.0f });
+	}
+
+	return bTestPassed;
+}
diff --git a/Engine/Plugins/Experimental/PCG/Source/PCG/Public/Elements/PCGDensityNoise.h b/Engine/Plugins/Experimental/PCG/Source/PCG/Public/Elements/PCGDensityNoise.h
new file mode 100644
--- /dev/null
+++ b/Engine/Plugins/Experimental/PCG/Source/PCG/Public/Elements/PCGDensityNoise.h
@@ -0,0 +1,60 @@
+// Copyright Epic Games, Inc. All Rights Reserved.
+
+#pragma once
+
+#include "CoreMinimal.h"
+#include "PCGElement.h"
+#include "PCGNode.h"
+#include "PCGSettings.h"
+
+#include "PCGDensityNoise.generated.h"
+
+UENUM()
+enum class EPCGDensityNoiseMode : uint8
+{
+	Set,
+	Minimum,
+	Maximum,
+	Add,
+	Multiply
+};
+
+UCLASS(BlueprintType, ClassGroup = (Procedural))
+class PCG_API UPCGDensityNoiseSettings : public UPCGSettings
+{
+	GENERATED_BODY()
+
+public:
+	// ~Begin UPCGSettings interface
+#if WITH_EDITOR
+	virtual FName GetDefaultNodeName() const override { return FName(TEXT("DensityNoiseNode")); }
+	virtual EPCGSettingsType GetType() const override { return EPCGSettingsType::Density; }
+#endif
+
+	virtual TArray<FPCGPinProperties> OutputPinProperties() const override { return Super::DefaultPointOutputPinProperties(); }
+
+protected:
+	virtual FPCGElementPtr CreateElement() const override;
+	//~End UPCGSettings interface
+
+public:
+	/** Density = (OriginalDensity op DensityNoise), DensityNoise in [DensityNoiseMin, DensityNoiseMax] */
+	UPROPERTY(BlueprintReadWrite, EditAnywhere, Category = Settings)
+	EPCGDensityNoiseMode DensityMode = EPCGDensityNoiseMode::Set;
+
+	UPROPERTY(BlueprintReadWrite, EditAnywhere, Category = Settings)
+	float DensityNoiseMin = 0.f;
+
+	UPROPERTY(BlueprintReadWrite, EditAnywhere, Category = Settings)
+	float DensityNoiseMax = 1.f;
+
+	/** Density = 1 - Density before applying the DensityMode operation */
+	UPROPERTY(BlueprintReadWrite, EditAnywhere, Category = Settings)
+	bool bInvertSourceDensity = false;
+};
+
+class FPCGDensityNoiseElement : public FSimplePCGElement
+{
+protected:
+	virtual bool ExecuteInternal(FPCGContext* Context) const;
+};
diff --git a/Engine/Plugins/Experimental/PCG/Source/PCG/Public/Elements/PCGDensityRemap.h b/Engine/Plugins/Experimental/PCG/Source/PCG/Public/Elements/PCGDensityRemap.h
--- a/Engine/Plugins/Experimental/PCG/Source/PCG/Public/Elements/PCGDensityRemap.h
+++ b/Engine/Plugins/Experimental/PCG/Source/PCG/Public/Elements/PCGDensityRemap.h
@@ -9,7 +9,7 @@
 
 #include "PCGDensityRemap.generated.h"
 
-UCLASS(BlueprintType, ClassGroup = (Procedural))
+UCLASS(BlueprintType, ClassGroup = (Procedural), Hidden, meta = (Deprecated = "5.1", DeprecationMessage = "Use DensityNoise instead."))
 class PCG_API UPCGLinearDensityRemapSettings : public UPCGSettings
 {
 	GENERATED_BODY()
