Fix compilation errors in PCGDensityNoiseTest and PCGDensityRemapTest

#rnx

diff --git a/Engine/Plugins/Experimental/PCG/Source/PCG/Private/Tests/Elements/PCGDensityNoiseTest.cpp b/Engine/Plugins/Experimental/PCG/Source/PCG/Private/Tests/Elements/PCGDensityNoiseTest.cpp
--- a/Engine/Plugins/Experimental/PCG/Source/PCG/Private/Tests/Elements/PCGDensityNoiseTest.cpp
+++ b/Engine/Plugins/Experimental/PCG/Source/PCG/Private/Tests/Elements/PCGDensityNoiseTest.cpp
@@ -14,7 +14,7 @@ IMPLEMENT_CUSTOM_SIMPLE_AUTOMATION_TEST(FPCGDensityNoiseTest, FPCGTestBaseClass,
 
 bool FPCGDensityNoiseTest::RunTest(const FString& Parameters)
 {
-	PCGDeterminismTests::FTestData TestData(PCGDeterminismTests::Defaults::Seed);
+	PCGTestsCommon::FTestData TestData(PCGDeterminismTests::Defaults::Seed);
 	PCGDeterminismTests::GenerateSettings<UPCGDensityNoiseSettings>(TestData);
 	UPCGDensityNoiseSettings* Settings = CastChecked<UPCGDensityNoiseSettings>(TestData.Settings);
 	FPCGElementPtr DensityNoiseElement = TestData.Settings->GetElement();
diff --git a/Engine/Plugins/Experimental/PCG/Source/PCG/Private/Tests/Elements/PCGDensityRemapTest.cpp b/Engine/Plugins/Experimental/PCG/Source/PCG/Private/Tests/Elements/PCGDensityRemapTest.cpp
--- a/Engine/Plugins/Experimental/PCG/Source/PCG/Private/Tests/Elements/PCGDensityRemapTest.cpp
+++ b/Engine/Plugins/Experimental/PCG/Source/PCG/Private/Tests/Elements/PCGDensityRemapTest.cpp
@@ -14,7 +14,7 @@ IMPLEMENT_CUSTOM_SIMPLE_AUTOMATION_TEST(FPCGDensityRemapTest, FPCGTestBaseClass,
 
 bool FPCGDensityRemapTest::RunTest(const FString& Parameters)
 {
-	PCGDeterminismTests::FTestData TestData(PCGDeterminismTests::Defaults::Seed);
+	PCGTestsCommon::FTestData TestData(PCGDeterminismTests::Defaults::Seed);
 	PCGDeterminismTests::GenerateSettings<UPCGDensityRemapSettings>(TestData);
 	UPCGDensityRemapSettings* Settings = CastChecked<UPCGDensityRemapSettings>(TestData.Settings);
 	FPCGElementPtr DensityRemapElement = TestData.Settings->GetElement();
