Support for Generate on Load and Generate on Demand for PCG components

* Also add a bounding box draw in the viewport for selected PCGActorPartition
* And few tweaks in the detail panel (with greyed out options for better clarity)

diff --git a/Engine/Plugins/Experimental/PCG/Source/PCG/Private/Grid/PCGPartitionActor.cpp b/Engine/Plugins/Experimental/PCG/Source/PCG/Private/Grid/PCGPartitionActor.cpp
--- a/Engine/Plugins/Experimental/PCG/Source/PCG/Private/Grid/PCGPartitionActor.cpp
+++ b/Engine/Plugins/Experimental/PCG/Source/PCG/Private/Grid/PCGPartitionActor.cpp
@@ -6,6 +6,7 @@
 #include "PCGWorldActor.h"
 
 #include "Landscape.h"
+#include "Components/BoxComponent.h"
 
 constexpr uint32 InvalidPCGGridSizeValue = 0u;
 
@@ -13,6 +14,17 @@ APCGPartitionActor::APCGPartitionActor(const FObjectInitializer& ObjectInitializ
 	: Super(ObjectInitializer)
 {
 	PCGGridSize = InvalidPCGGridSizeValue;
+
+#if WITH_EDITOR
+	// Setup bounds component
+	BoundsComponent = CreateDefaultSubobject<UBoxComponent>(TEXT("BoundsComponent"));
+	BoundsComponent->SetCollisionObjectType(ECC_WorldStatic);
+	BoundsComponent->SetCollisionResponseToAllChannels(ECR_Ignore);
+	BoundsComponent->SetCollisionEnabled(ECollisionEnabled::NoCollision);
+	BoundsComponent->SetGenerateOverlapEvents(false);
+	BoundsComponent->SetupAttachment(GetRootComponent());
+	BoundsComponent->bDrawOnlyIfSelected = true;
+#endif // WITH_EDITOR
 }
 
 void APCGPartitionActor::PostLoad()
@@ -28,11 +40,25 @@ void APCGPartitionActor::PostLoad()
 #if WITH_EDITOR
 	// Make sure that we don't track objects that do not exist anymore
 	CleanupDeadGraphInstances();
+
+	if (BoundsComponent)
+	{
+		BoundsComponent->SetBoxExtent(GetFixedBounds().GetExtent());
+	}
 #endif // WITH_EDITOR
 }
 
 void APCGPartitionActor::BeginPlay()
 {
+	// Pass through all the pcg components, and make sure we match generation trigger for the local component if it is not overriden
+	for (auto& It : OriginalToLocalMap)
+	{
+		if (It.Key && It.Value && !It.Value->bGenerationTriggerLocalOverride)
+		{
+			It.Value->GenerationTrigger = It.Key->GenerationTrigger;
+		}
+	}
+
 	Super::BeginPlay();
 
 	// Register cell to the PCG grid
@@ -201,6 +227,13 @@ bool APCGPartitionActor::CleanupDeadGraphInstances()
 void APCGPartitionActor::PostCreation()
 {
 	PCGGridSize = GridSize;
+
+#if WITH_EDITOR
+	if (BoundsComponent)
+	{
+		BoundsComponent->SetBoxExtent(GetFixedBounds().GetExtent());
+	}
+#endif // WITH_EDITOR
 }
 
 bool APCGPartitionActor::IsSafeForDeletion() const
diff --git a/Engine/Plugins/Experimental/PCG/Source/PCG/Private/Grid/PCGPartitionActor.h b/Engine/Plugins/Experimental/PCG/Source/PCG/Private/Grid/PCGPartitionActor.h
--- a/Engine/Plugins/Experimental/PCG/Source/PCG/Private/Grid/PCGPartitionActor.h
+++ b/Engine/Plugins/Experimental/PCG/Source/PCG/Private/Grid/PCGPartitionActor.h
@@ -8,6 +8,7 @@
 #include "PCGPartitionActor.generated.h"
 
 class UPCGComponent;
+class UBoxComponent;
 
 /** 
 * The APCGPartitionActor actor is used to store grid cell data
@@ -82,4 +83,10 @@ private:
 
 	UPROPERTY(VisibleAnywhere, Category = WorldPartition)
 	uint32 PCGGridSize;
+
+#if WITH_EDITORONLY_DATA
+	/** Box component to draw the Partition actor bounds in the Editor viewport */
+	UPROPERTY(Transient)
+	TObjectPtr<UBoxComponent> BoundsComponent;
+#endif // WITH_EDITORONLY_DATA
 };
diff --git a/Engine/Plugins/Experimental/PCG/Source/PCG/Private/PCGComponent.cpp b/Engine/Plugins/Experimental/PCG/Source/PCG/Private/PCGComponent.cpp
--- a/Engine/Plugins/Experimental/PCG/Source/PCG/Private/PCGComponent.cpp
+++ b/Engine/Plugins/Experimental/PCG/Source/PCG/Private/PCGComponent.cpp
@@ -159,6 +159,8 @@ void UPCGComponent::SetPropertiesFromOriginal(const UPCGComponent* Original)
 	Seed = Original->Seed;
 	SetGraph(Original->Graph);
 
+	GenerationTrigger = Original->GenerationTrigger;
+
 #if WITH_EDITOR
 	// Note that while we dirty here, we won't trigger a refresh since we don't have the required context
 	if (bIsDirty)
@@ -512,7 +514,7 @@ void UPCGComponent::BeginPlay()
 {
 	Super::BeginPlay();
 
-	if(bActivated && !bGenerated && !IsPartitioned())
+	if(bActivated && !bGenerated && !IsPartitioned() && GenerationTrigger == EPCGComponentGenerationTrigger::GenerateOnLoad)
 	{
 		GenerateLocal(/*bForce=*/false);
 		bRuntimeGenerated = true;
@@ -525,6 +527,7 @@ void UPCGComponent::OnComponentCreated()
 
 #if WITH_EDITOR
 	SetupActorCallbacks();
+	UpdateIsLocalComponent();
 #endif
 }
 
@@ -600,6 +603,8 @@ void UPCGComponent::PostLoad()
 	{
 		Graph->OnGraphChangedDelegate.AddUObject(this, &UPCGComponent::OnGraphChanged);
 	}
+	
+	UpdateIsLocalComponent();
 #endif
 }
 
@@ -1904,6 +1909,14 @@ bool UPCGComponent::GraphUsesLandscapePin() const
 	return Graph && Graph->GetInputNode()->IsOutputPinConnected(PCGInputOutputConstants::DefaultLandscapeLabel);
 }
 
+void UPCGComponent::UpdateIsLocalComponent()
+{
+	if (GetOwner() && GetOwner()->IsA<APCGPartitionActor>())
+	{
+		bIsComponentLocal = true;
+	}
+}
+
 #endif // WITH_EDITOR
 
 #undef LOCTEXT_NAMESPACE
diff --git a/Engine/Plugins/Experimental/PCG/Source/PCG/Public/PCGComponent.h b/Engine/Plugins/Experimental/PCG/Source/PCG/Public/PCGComponent.h
--- a/Engine/Plugins/Experimental/PCG/Source/PCG/Public/PCGComponent.h
+++ b/Engine/Plugins/Experimental/PCG/Source/PCG/Public/PCGComponent.h
@@ -32,6 +32,13 @@ enum class EPCGComponentInput : uint8
 	EPCGComponentInput_MAX
 };
 
+UENUM(Blueprintable)
+enum class EPCGComponentGenerationTrigger : uint8
+{
+	GenerateOnLoad,
+	GenerateOnDemand
+};
+
 UENUM(meta = (Bitflags))
 enum class EPCGComponentDirtyFlag : uint8
 {
@@ -121,9 +128,15 @@ public:
 	UPROPERTY(BlueprintReadWrite, EditAnywhere, Category = Properties)
 	bool bActivated = true;
 
-	UPROPERTY(BlueprintReadWrite, EditAnywhere, AdvancedDisplay, Category = Properties)
+	UPROPERTY(BlueprintReadWrite, EditAnywhere, AdvancedDisplay, Category = Properties, meta = (EditCondition = "!bIsComponentLocal"))
 	bool bIsPartitioned = true;
 
+	UPROPERTY(BlueprintReadOnly, EditAnywhere, Category = Properties, AdvancedDisplay, meta = (EditCondition = "bIsComponentLocal", EditConditionHides))
+	bool bGenerationTriggerLocalOverride = false;
+
+	UPROPERTY(BlueprintReadOnly, EditAnywhere, Category = Properties, AdvancedDisplay, meta = (EditCondition = "!bIsComponentLocal || bGenerationTriggerLocalOverride", EditConditionHides))
+	EPCGComponentGenerationTrigger GenerationTrigger = EPCGComponentGenerationTrigger::GenerateOnLoad;
+
 	/** Flag to indicate whether this component has run in the editor. Note that for partitionable actors, this will always be false. */
 	UPROPERTY(BlueprintReadOnly, VisibleAnywhere, AdvancedDisplay, Category = Properties)
 	bool bGenerated = false;
@@ -132,6 +145,9 @@ public:
 	bool bRuntimeGenerated = false;
 
 #if WITH_EDITORONLY_DATA
+	UPROPERTY(VisibleAnywhere, Transient, Category = Properties, meta = (EditCondition = false, EditConditionHides))
+	bool bIsComponentLocal = false;
+
 	UPROPERTY(BlueprintReadWrite, EditAnywhere, AdvancedDisplay, Category = Properties, meta = (DisplayName = "Regenerate PCG volume in editor"))
 	bool bRegenerateInEditor = true;
 
@@ -271,6 +287,7 @@ private:
 	void TeardownLandscapeTracking();
 	void UpdateTrackedLandscape(bool bBoundsCheck = true);
 	void OnLandscapeChanged(ALandscapeProxy* Landscape, const FLandscapeProxyComponentDataChangedParams& ChangeParams);
+	void UpdateIsLocalComponent();
 #endif
 
 #if WITH_EDITORONLY_DATA
