Static analyzer fix.

diff --git a/Engine/Plugins/Experimental/PCG/Source/PCG/Private/Tests/Elements/PCGCopyPointsTest.cpp b/Engine/Plugins/Experimental/PCG/Source/PCG/Private/Tests/Elements/PCGCopyPointsTest.cpp
--- a/Engine/Plugins/Experimental/PCG/Source/PCG/Private/Tests/Elements/PCGCopyPointsTest.cpp
+++ b/Engine/Plugins/Experimental/PCG/Source/PCG/Private/Tests/Elements/PCGCopyPointsTest.cpp
@@ -155,6 +155,8 @@ bool FPCGCopyPointsTest::RunTest(const FString& Parameters)
 			return;
 		}
 
+		check(RootMetadata); // Static analyzer doesn't understand semantics of TestTrue, but does understand check().
+
 		TArray<const FPCGMetadataAttribute<float>*> InheritedAttributes;
 		TArray<EPCGCopyPointsMetadataInheritanceMode> InheritedAttributeModes;
 		TArray<const FPCGMetadataAttribute<float>*> OutAttributes;
