Another static analysis fix.

diff --git a/Engine/Plugins/Experimental/PCG/Source/PCG/Private/Tests/Elements/PCGCopyPointsTest.cpp b/Engine/Plugins/Experimental/PCG/Source/PCG/Private/Tests/Elements/PCGCopyPointsTest.cpp
--- a/Engine/Plugins/Experimental/PCG/Source/PCG/Private/Tests/Elements/PCGCopyPointsTest.cpp
+++ b/Engine/Plugins/Experimental/PCG/Source/PCG/Private/Tests/Elements/PCGCopyPointsTest.cpp
@@ -155,7 +155,8 @@ bool FPCGCopyPointsTest::RunTest(const FString& Parameters)
 			return;
 		}
 
-		check(RootMetadata); // Static analyzer doesn't understand semantics of TestTrue, but does understand check().
+		CA_ASSUME(RootMetadata); // Static analyzer doesn't understand semantics of TestTrue.
+		CA_ASSUME(NonRootMetadata);
 
 		TArray<const FPCGMetadataAttribute<float>*> InheritedAttributes;
 		TArray<EPCGCopyPointsMetadataInheritanceMode> InheritedAttributeModes;
