Removed OnComponentTransfromUpdated since it causes aliasing issues on play, causes potential refresh issues in the editor and is not a more long term solution

#rnx

diff --git a/Engine/Plugins/Experimental/PCG/Source/PCG/Private/PCGComponent.cpp b/Engine/Plugins/Experimental/PCG/Source/PCG/Private/PCGComponent.cpp
--- a/Engine/Plugins/Experimental/PCG/Source/PCG/Private/PCGComponent.cpp
+++ b/Engine/Plugins/Experimental/PCG/Source/PCG/Private/PCGComponent.cpp
@@ -769,14 +769,12 @@ void UPCGComponent::PostEditUndo()
 void UPCGComponent::SetupActorCallbacks()
 {
 	GEngine->OnActorMoved().AddUObject(this, &UPCGComponent::OnActorMoved);
-	GEngine->OnComponentTransformChanged().AddUObject(this, &UPCGComponent::OnComponentTransformChanged);
 	FCoreUObjectDelegates::OnObjectPropertyChanged.AddUObject(this, &UPCGComponent::OnObjectPropertyChanged);
 }
 
 void UPCGComponent::TeardownActorCallbacks()
 {
 	FCoreUObjectDelegates::OnObjectPropertyChanged.RemoveAll(this);
-	GEngine->OnComponentTransformChanged().RemoveAll(this);
 	GEngine->OnActorMoved().RemoveAll(this);
 }
 
@@ -940,17 +938,6 @@ void UPCGComponent::OnActorMoved(AActor* InActor)
 	}
 }
 
-void UPCGComponent::OnComponentTransformChanged(USceneComponent* InComponent, ETeleportType InTeleport)
-{
-	if (!InComponent || InComponent->GetOwner() != GetOwner())
-	{
-		return;
-	}
-
-	DirtyGenerated(EPCGComponentDirtyFlag::Actor);
-	Refresh();
-}
-
 void UPCGComponent::UpdateTrackedLandscape(bool bBoundsCheck)
 {
 	TeardownLandscapeTracking();
diff --git a/Engine/Plugins/Experimental/PCG/Source/PCG/Public/PCGComponent.h b/Engine/Plugins/Experimental/PCG/Source/PCG/Public/PCGComponent.h
--- a/Engine/Plugins/Experimental/PCG/Source/PCG/Public/PCGComponent.h
+++ b/Engine/Plugins/Experimental/PCG/Source/PCG/Public/PCGComponent.h
@@ -224,7 +224,6 @@ private:
 	void OnActorAdded(AActor* InActor);
 	void OnActorDeleted(AActor* InActor);
 	void OnActorMoved(AActor* InActor);
-	void OnComponentTransformChanged(USceneComponent* InComponent, ETeleportType InTeleport);
 	void OnObjectPropertyChanged(UObject* InObject, FPropertyChangedEvent& InEvent);
 	bool ActorHasExcludedTag(AActor* InActor) const;
 	bool ActorIsTracked(AActor* InActor) const;
