Correctly fail TransformPointsTest and CopyPointsTest when test conditions fail

#rnx

diff --git a/Engine/Plugins/Experimental/PCG/Source/PCG/Private/Tests/Elements/PCGCopyPointsTest.cpp b/Engine/Plugins/Experimental/PCG/Source/PCG/Private/Tests/Elements/PCGCopyPointsTest.cpp
--- a/Engine/Plugins/Experimental/PCG/Source/PCG/Private/Tests/Elements/PCGCopyPointsTest.cpp
+++ b/Engine/Plugins/Experimental/PCG/Source/PCG/Private/Tests/Elements/PCGCopyPointsTest.cpp
@@ -14,6 +14,8 @@
 #include "Metadata/PCGMetadataAttributeTraits.h"
 #include "Metadata/PCGMetadataAttributeTpl.h"
 
+#if WITH_EDITOR
+
 IMPLEMENT_CUSTOM_SIMPLE_AUTOMATION_TEST(FPCGCopyPointsTest, FPCGTestBaseClass, "pcg.tests.CopyPoints.Basic", PCGTestsCommon::TestFlags)
 
 bool FPCGCopyPointsTest::RunTest(const FString& Parameters)
@@ -64,7 +66,7 @@ bool FPCGCopyPointsTest::RunTest(const FString& Parameters)
 	TargetTaggedData.Pin = PCGCopyPointsConstants::TargetPointsLabel;
 
 	// Test only supports float attributes for simplicity
-	auto ValidateCopyPointsPoints = [this, &TestData, CopyPointsElement, Settings]()
+	auto ValidateCopyPointsPoints = [this, &TestData, CopyPointsElement, Settings]() -> bool
 	{
 		TUniquePtr<FPCGContext> Context = MakeUnique<FPCGContext>(*CopyPointsElement->Initialize(TestData.InputData, TestData.TestPCGComponent, nullptr));
 		Context->NumAvailableTasks = 1;
@@ -78,12 +80,12 @@ bool FPCGCopyPointsTest::RunTest(const FString& Parameters)
 
 		if (!TestTrue("Valid number of inputs", Sources.Num() == 1 && Targets.Num() == 1))
 		{
-			return;
+			return false;
 		}
 
 		if (!TestTrue("Valid number of outputs", Sources.Num() == 1))
 		{
-			return;
+			return false;
 		}
 
 		const FPCGTaggedData& Source = Sources[0];
@@ -104,21 +106,21 @@ bool FPCGCopyPointsTest::RunTest(const FString& Parameters)
 
 		if (!TestTrue("Valid output data", Output.Data != nullptr))
 		{
-			return;
+			return false;
 		}
 
 		const UPCGSpatialData* OutSpatialData = Cast<UPCGSpatialData>(Output.Data);
 
 		if (!TestNotNull("Valid output SpatialData", OutSpatialData))
 		{
-			return;
+			return false;
 		}
 
 		const UPCGPointData* OutPointData = OutSpatialData->ToPointData(Context.Get());
 
 		if (!TestNotNull("Valid output PointData", OutPointData))
 		{
-			return;
+			return false;
 		}
 
 		const TArray<FPCGPoint>& SourcePoints = SourcePointData->GetPoints();
@@ -127,7 +129,7 @@ bool FPCGCopyPointsTest::RunTest(const FString& Parameters)
 
 		if (!TestEqual("Valid number of output points", SourcePoints.Num() * TargetPoints.Num(), OutPoints.Num()))
 		{ 
-			return;
+			return false;
 		}
 
 		const UPCGMetadata* RootMetadata = nullptr;
@@ -152,11 +154,13 @@ bool FPCGCopyPointsTest::RunTest(const FString& Parameters)
 
 		if (!TestTrue("Valid input metadata", RootMetadata && NonRootMetadata))
 		{
-			return;
+			return false;
 		}
 
 		CA_ASSUME(RootMetadata); // Static analyzer doesn't understand semantics of TestTrue.
 		CA_ASSUME(NonRootMetadata);
+		
+		bool bTestPassed = true;
 
 		TArray<const FPCGMetadataAttribute<float>*> InheritedAttributes;
 		TArray<EPCGCopyPointsMetadataInheritanceMode> InheritedAttributeModes;
@@ -177,6 +181,10 @@ bool FPCGCopyPointsTest::RunTest(const FString& Parameters)
 				InheritedAttributeModes.Add(RootMode);
 				OutAttributes.Add(OutAttribute);
 			}
+			else
+			{
+				bTestPassed = false;
+			}
 		}
 
 		NonRootMetadata->GetAttributes(AttributeNames, AttributeTypes);
@@ -194,6 +202,10 @@ bool FPCGCopyPointsTest::RunTest(const FString& Parameters)
 					InheritedAttributeModes.Add(NonRootMode);
 					OutAttributes.Add(OutAttribute);
 				}
+				else
+				{
+					bTestPassed = false;
+				}
 			}
 		}
 
@@ -205,10 +217,10 @@ bool FPCGCopyPointsTest::RunTest(const FString& Parameters)
 			const FPCGPoint& TargetPoint = TargetPoints[PointIndex % TargetPoints.Num()];
 			const FPCGPoint& OutPoint = OutPoints[PointIndex];
 
-			TestEqual("SourcePoint and OutPoint have same Density", SourcePoint.Density, OutPoint.Density);
-			TestEqual("SourcePoint and OutPoint have same BoundsMin", SourcePoint.BoundsMin, OutPoint.BoundsMin);
-			TestEqual("SourcePoint and OutPoint have same BoundsMax", SourcePoint.BoundsMax, OutPoint.BoundsMax);
-			TestEqual("SourcePoint and OutPoint have same Steepness", SourcePoint.Steepness, OutPoint.Steepness);
+			bTestPassed &= TestEqual("SourcePoint and OutPoint have same Density", SourcePoint.Density, OutPoint.Density);
+			bTestPassed &= TestEqual("SourcePoint and OutPoint have same BoundsMin", SourcePoint.BoundsMin, OutPoint.BoundsMin);
+			bTestPassed &= TestEqual("SourcePoint and OutPoint have same BoundsMax", SourcePoint.BoundsMax, OutPoint.BoundsMax);
+			bTestPassed &= TestEqual("SourcePoint and OutPoint have same Steepness", SourcePoint.Steepness, OutPoint.Steepness);
 
 			// Validate transform inheritance
 			const FTransform& SourceTransform = SourcePoint.Transform;
@@ -273,11 +285,11 @@ bool FPCGCopyPointsTest::RunTest(const FString& Parameters)
 				Seed = TargetPoint.Seed;
 			}
 
-			TestEqual("Valid rotation", Rotation, OutTransform.GetRotation());
-			TestEqual("Valid scale", Scale, OutTransform.GetScale3D());
-			TestEqual("Valid color", Color, OutPoint.Color);
-			TestEqual("Valid location", Location, OutTransform.GetLocation());
-			TestEqual("Valid location", Seed, OutPoint.Seed);		
+			bTestPassed &= TestEqual("Valid rotation", Rotation, OutTransform.GetRotation());
+			bTestPassed &= TestEqual("Valid scale", Scale, OutTransform.GetScale3D());
+			bTestPassed &= TestEqual("Valid color", Color, OutPoint.Color);
+			bTestPassed &= TestEqual("Valid location", Location, OutTransform.GetLocation());
+			bTestPassed &= TestEqual("Valid location", Seed, OutPoint.Seed);		
 
 			// Validate point value keys
 			for (int AttributeIndex = 0; AttributeIndex < InheritedAttributes.Num(); ++AttributeIndex)
@@ -287,17 +299,21 @@ bool FPCGCopyPointsTest::RunTest(const FString& Parameters)
 				check(InheritedAttribute && OutAttribute);
 
 				const PCGMetadataEntryKey EntryKey = (InheritedAttributeModes[AttributeIndex] == EPCGCopyPointsMetadataInheritanceMode::Source) ? SourcePoint.MetadataEntry : TargetPoint.MetadataEntry;
-				TestEqual("Valid metadata value", InheritedAttribute->GetValueKey(EntryKey), OutAttribute->GetValueKey(OutPoint.MetadataEntry));
+				bTestPassed &= TestEqual("Valid metadata value", InheritedAttribute->GetValueKey(EntryKey), OutAttribute->GetValueKey(OutPoint.MetadataEntry));
 			}
 		}
+
+		return bTestPassed;
 	};
 
+	bool bTestPassed = true;
+
 	// Test 1 - inherit from transform multiplication
 	Settings->RotationInheritance = EPCGCopyPointsInheritanceMode::Relative;
 	Settings->ScaleInheritance = EPCGCopyPointsInheritanceMode::Relative;
 	Settings->ColorInheritance = EPCGCopyPointsInheritanceMode::Relative;
 	Settings->SeedInheritance = EPCGCopyPointsInheritanceMode::Relative;
-	ValidateCopyPointsPoints();
+	bTestPassed &= ValidateCopyPointsPoints();
 
 	// Test 2 - inherit from source points
 	Settings->RotationInheritance = EPCGCopyPointsInheritanceMode::Source;
@@ -305,7 +321,7 @@ bool FPCGCopyPointsTest::RunTest(const FString& Parameters)
 	Settings->ColorInheritance = EPCGCopyPointsInheritanceMode::Source;
 	Settings->SeedInheritance = EPCGCopyPointsInheritanceMode::Source;
 	Settings->AttributeInheritance = EPCGCopyPointsMetadataInheritanceMode::Source;
-	ValidateCopyPointsPoints();
+	bTestPassed &= ValidateCopyPointsPoints();
 
 	// Test 3 - inherit from target points
 	Settings->RotationInheritance = EPCGCopyPointsInheritanceMode::Target;
@@ -313,17 +329,19 @@ bool FPCGCopyPointsTest::RunTest(const FString& Parameters)
 	Settings->ColorInheritance = EPCGCopyPointsInheritanceMode::Target;
 	Settings->SeedInheritance = EPCGCopyPointsInheritanceMode::Target;
 	Settings->AttributeInheritance = EPCGCopyPointsMetadataInheritanceMode::Target;
-	ValidateCopyPointsPoints();
+	bTestPassed &= ValidateCopyPointsPoints();
 
 	// Test 4 - empty source data
 	SourceTaggedData.Data = EmptyData;
-	ValidateCopyPointsPoints();
+	bTestPassed &= ValidateCopyPointsPoints();
 	SourceTaggedData.Data = SourceData;
 
 	// Test 5 - empty target data
 	TargetTaggedData.Data = EmptyData;
-	ValidateCopyPointsPoints();
+	bTestPassed &= ValidateCopyPointsPoints();
 	TargetTaggedData.Data = TargetData;
 
-	return true;
+	return bTestPassed;
 }
+
+#endif // WITH_EDITOR
diff --git a/Engine/Plugins/Experimental/PCG/Source/PCG/Private/Tests/Elements/PCGTransformPointsTest.cpp b/Engine/Plugins/Experimental/PCG/Source/PCG/Private/Tests/Elements/PCGTransformPointsTest.cpp
--- a/Engine/Plugins/Experimental/PCG/Source/PCG/Private/Tests/Elements/PCGTransformPointsTest.cpp
+++ b/Engine/Plugins/Experimental/PCG/Source/PCG/Private/Tests/Elements/PCGTransformPointsTest.cpp
@@ -10,6 +10,8 @@
 #include "Data/PCGPointData.h"
 #include "Elements/PCGTransformPoints.h"
 
+#if WITH_EDITOR
+
 IMPLEMENT_CUSTOM_SIMPLE_AUTOMATION_TEST(FPCGTransformPointsTest, FPCGTestBaseClass, "pcg.tests.TransformPoints.Basic", PCGTestsCommon::TestFlags)
 
 bool FPCGTransformPointsTest::RunTest(const FString& Parameters)
@@ -56,7 +58,7 @@ bool FPCGTransformPointsTest::RunTest(const FString& Parameters)
 	Settings->ScaleMin = FVector(1.f);
 	Settings->ScaleMax = FVector(2.f);
 
-	auto ValidateTransformPoints = [this, &TestData, TransformPointsElement, Settings]()
+	auto ValidateTransformPoints = [this, &TestData, TransformPointsElement, Settings]() -> bool
 	{
 		TUniquePtr<FPCGContext> Context = MakeUnique<FPCGContext>(*TransformPointsElement->Initialize(TestData.InputData, TestData.TestPCGComponent, nullptr));
 		Context->NumAvailableTasks = 1;
@@ -69,9 +71,11 @@ bool FPCGTransformPointsTest::RunTest(const FString& Parameters)
 
 		if (!TestEqual("Valid number of outputs", Inputs.Num(), Outputs.Num()))
 		{
-			return;
+			return false;
 		}
 
+		bool bTestPassed = true;
+
 		for (int DataIndex = 0; DataIndex < Inputs.Num(); ++DataIndex)
 		{
 			const FPCGTaggedData& Input = Inputs[DataIndex];
@@ -87,6 +91,7 @@ bool FPCGTransformPointsTest::RunTest(const FString& Parameters)
 
 			if (!TestTrue("Valid output data", Output.Data != nullptr))
 			{
+				bTestPassed = false;
 				continue;
 			}
 
@@ -94,6 +99,7 @@ bool FPCGTransformPointsTest::RunTest(const FString& Parameters)
 
 			if (!TestNotNull("Valid ouptut SpatialData", OutSpatialData))
 			{
+				bTestPassed = false;
 				continue;
 			}
 
@@ -101,6 +107,7 @@ bool FPCGTransformPointsTest::RunTest(const FString& Parameters)
 
 			if (!TestNotNull("Valid output PointData", OutPointData))
 			{
+				bTestPassed = false;
 				continue;
 			}
 
@@ -109,6 +116,7 @@ bool FPCGTransformPointsTest::RunTest(const FString& Parameters)
 
 			if (!TestEqual("Input and output point counts match", InPoints.Num(), OutPoints.Num()))
 			{ 
+				bTestPassed = false;
 				continue;
 			}
 
@@ -119,7 +127,7 @@ bool FPCGTransformPointsTest::RunTest(const FString& Parameters)
 
 				FPCGPoint RepositionedPoint = InPoint;
 				RepositionedPoint.Transform = OutPoint.Transform;
-				TestTrue("RepositionedPoint and OutPoint are identical", PCGTestsCommon::PointsAreIdentical(RepositionedPoint, OutPoint));
+				bTestPassed &= TestTrue("RepositionedPoint and OutPoint are identical", PCGTestsCommon::PointsAreIdentical(RepositionedPoint, OutPoint));
 
 				// Validate transform is within range of original transform
 				const FTransform& InTransform = InPoint.Transform;
@@ -139,23 +147,23 @@ bool FPCGTransformPointsTest::RunTest(const FString& Parameters)
 				FRotator OriginalRotation = TransformForRotation.GetRotation().Rotator();
 				FVector OriginalScale = TransformForScale.GetScale3D();
 
-				TestTrue("Valid location", OriginalOffset.X >= Settings->OffsetMin.X && OriginalOffset.X <= Settings->OffsetMax.X
+				bTestPassed &= TestTrue("Valid location", OriginalOffset.X >= Settings->OffsetMin.X && OriginalOffset.X <= Settings->OffsetMax.X
 					&& OriginalOffset.Y >= Settings->OffsetMin.Y && OriginalOffset.Y <= Settings->OffsetMax.Y
 					&& OriginalOffset.Z >= Settings->OffsetMin.Z && OriginalOffset.Z <= Settings->OffsetMax.Z);
 
-				TestTrue("Valid rotation", OriginalRotation.Pitch >= Settings->RotationMin.Pitch && OriginalRotation.Pitch <= Settings->RotationMax.Pitch
+				bTestPassed &= TestTrue("Valid rotation", OriginalRotation.Pitch >= Settings->RotationMin.Pitch && OriginalRotation.Pitch <= Settings->RotationMax.Pitch
 					&& OriginalRotation.Yaw >= Settings->RotationMin.Yaw && OriginalRotation.Yaw <= Settings->RotationMax.Yaw
 					&& OriginalRotation.Roll >= Settings->RotationMin.Roll && OriginalRotation.Roll <= Settings->RotationMax.Roll);
 
 				if (Settings->bUniformScale)
 				{
-					TestTrue("Valid scale", OriginalScale.X >= Settings->ScaleMin.X && OriginalScale.X <= Settings->ScaleMax.X
+					bTestPassed &= TestTrue("Valid scale", OriginalScale.X >= Settings->ScaleMin.X && OriginalScale.X <= Settings->ScaleMax.X
 						&& OriginalScale.Y >= Settings->ScaleMin.X && OriginalScale.Y <= Settings->ScaleMax.X
 						&& OriginalScale.Z >= Settings->ScaleMin.X && OriginalScale.Z <= Settings->ScaleMax.X);
 				}
 				else
 				{
-					TestTrue("Valid scale", OriginalScale.X >= Settings->ScaleMin.X && OriginalScale.X <= Settings->ScaleMax.X
+					bTestPassed &= TestTrue("Valid scale", OriginalScale.X >= Settings->ScaleMin.X && OriginalScale.X <= Settings->ScaleMax.X
 						&& OriginalScale.Y >= Settings->ScaleMin.Y && OriginalScale.Y <= Settings->ScaleMax.Y
 						&& OriginalScale.Z >= Settings->ScaleMin.Z && OriginalScale.Z <= Settings->ScaleMax.Z);
 				}
@@ -163,33 +171,38 @@ bool FPCGTransformPointsTest::RunTest(const FString& Parameters)
 				if (Settings->bRecomputeSeed)
 				{
 					const FVector& Position = OutTransform.GetLocation();
-					TestEqual("Valid seed", OutPoint.Seed, PCGHelpers::ComputeSeed((int)Position.X, (int)Position.Y, (int)Position.Z));
+					bTestPassed &= TestEqual("Valid seed", OutPoint.Seed, PCGHelpers::ComputeSeed((int)Position.X, (int)Position.Y, (int)Position.Z));
 				}
 				else
 				{
-					TestEqual("Valid seed", OutPoint.Seed, InPoint.Seed);
+					bTestPassed &= TestEqual("Valid seed", OutPoint.Seed, InPoint.Seed);
 				}
 			}
 		}
+
+		return bTestPassed;
 	};
 
+	bool bTestPassed = true;
+
 	// Test 1 - absolute transformations
 	Settings->bAbsoluteOffset = true;
 	Settings->bAbsoluteRotation = true;
 	Settings->bAbsoluteScale = true;
-	ValidateTransformPoints();
+	bTestPassed &= ValidateTransformPoints();
 
 	// Test 2 - relative transformations
 	Settings->bAbsoluteOffset = false;
 	Settings->bAbsoluteRotation = false;
 	Settings->bAbsoluteScale = false;
-	ValidateTransformPoints();
+	bTestPassed &= ValidateTransformPoints();
 
 	// Test 3 - uniform scale and recompute seed
 	Settings->bUniformScale = true;
 	Settings->bRecomputeSeed = true;
-	ValidateTransformPoints();
-
-	return true;
+	bTestPassed &= ValidateTransformPoints();
 
+	return bTestPassed;
 }
+
+#endif // WITH_EDITOR
