[PCG] Enable Geneate and cleanup of Partition actors to be triggered at runtime using the original component.

diff --git a/Engine/Plugins/Experimental/PCG/Source/PCG/Private/Grid/PCGPartitionActor.cpp b/Engine/Plugins/Experimental/PCG/Source/PCG/Private/Grid/PCGPartitionActor.cpp
--- a/Engine/Plugins/Experimental/PCG/Source/PCG/Private/Grid/PCGPartitionActor.cpp
+++ b/Engine/Plugins/Experimental/PCG/Source/PCG/Private/Grid/PCGPartitionActor.cpp
@@ -50,15 +50,27 @@ void APCGPartitionActor::PostLoad()
 
 void APCGPartitionActor::BeginPlay()
 {
-	// Pass through all the pcg components, and make sure we match generation trigger for the local component if it is not overriden
+	// Pass through all the pcg components, to verify if we need to generate them
 	for (auto& It : OriginalToLocalMap)
 	{
-		if (It.Key && It.Value && !It.Value->bGenerationTriggerLocalOverride)
+		if (It.Key && It.Value)
 		{
-			It.Value->GenerationTrigger = It.Key->GenerationTrigger;
+			// If we have an original component that is generated (or generating), this one is automatically generated => GenerateOnLoad
+			if (It.Key->bGenerated || It.Key->IsGenerating())
+			{
+				It.Value->GenerationTrigger = EPCGComponentGenerationTrigger::GenerateOnLoad;
+			}
+			// Otherwise, make them match
+			else
+			{
+				It.Value->GenerationTrigger = It.Key->GenerationTrigger;
+			}
 		}
 	}
 
+	// Make the Partition actor register itself to the PCG Component
+	RegisterActorToOriginalComponents();
+
 	Super::BeginPlay();
 
 	// Register cell to the PCG grid
@@ -70,6 +82,9 @@ void APCGPartitionActor::EndPlay(const EEndPlayReason::Type EndPlayReason)
 	// Unregister each cell to the PCG grid
 	// TODO
 
+	// Make the Partition actor unregister itself to the PCG Component
+	UnregisterActorToOriginalComponents();
+
 	Super::EndPlay(EndPlayReason);
 }
 
@@ -115,6 +130,28 @@ UPCGComponent* APCGPartitionActor::GetOriginalComponent(const UPCGComponent* Loc
 	return OriginalComponent ? (*OriginalComponent).Get() : nullptr;
 }
 
+void APCGPartitionActor::RegisterActorToOriginalComponents() const
+{
+	for (auto& It : OriginalToLocalMap)
+	{
+		if (It.Key && It.Value)
+		{
+			It.Key->AddPCGPartitionActor(this);
+		}
+	}
+}
+
+void APCGPartitionActor::UnregisterActorToOriginalComponents() const
+{
+	for (auto& It : OriginalToLocalMap)
+	{
+		if (It.Key && It.Value)
+		{
+			It.Key->RemovePCGPartitionActor(this);
+		}
+	}
+}
+
 #if WITH_EDITOR
 FBox APCGPartitionActor::GetStreamingBounds() const
 {
diff --git a/Engine/Plugins/Experimental/PCG/Source/PCG/Private/Grid/PCGPartitionActor.h b/Engine/Plugins/Experimental/PCG/Source/PCG/Private/Grid/PCGPartitionActor.h
--- a/Engine/Plugins/Experimental/PCG/Source/PCG/Private/Grid/PCGPartitionActor.h
+++ b/Engine/Plugins/Experimental/PCG/Source/PCG/Private/Grid/PCGPartitionActor.h
@@ -89,4 +89,7 @@ private:
 	UPROPERTY(Transient)
 	TObjectPtr<UBoxComponent> BoundsComponent;
 #endif // WITH_EDITORONLY_DATA
+
+	void RegisterActorToOriginalComponents() const;
+	void UnregisterActorToOriginalComponents() const;
 };
diff --git a/Engine/Plugins/Experimental/PCG/Source/PCG/Private/PCGComponent.cpp b/Engine/Plugins/Experimental/PCG/Source/PCG/Private/PCGComponent.cpp
--- a/Engine/Plugins/Experimental/PCG/Source/PCG/Private/PCGComponent.cpp
+++ b/Engine/Plugins/Experimental/PCG/Source/PCG/Private/PCGComponent.cpp
@@ -98,29 +98,30 @@ void UPCGComponent::ForEachManagedResource(TFunctionRef<void(UPCGManagedResource
 	}
 }
 
-bool UPCGComponent::ShouldGenerate(bool bForce) const
+bool UPCGComponent::ShouldGenerate(bool bForce, EPCGComponentGenerationTrigger RequestedGenerationTrigger) const
 {
 	if (!bActivated || !Graph || !GetSubsystem())
 	{
 		return false;
 	}
 
-	if (IsPartitioned())
-	{
 #if WITH_EDITOR
-		return !GIsPlayInEditorWorld;
-#else
-		return false;
-#endif
-	}
-	else
+	// Always run Generate if we are in editor and partitioned since the original component doesn't know the state of the local one.
+	if (IsPartitioned() && !GIsPlayInEditorWorld)
 	{
-		return (!bGenerated || 
+		return true;
+	}
+#endif
+
+	// A request is invalid only if it was requested "GenerateOnLoad", but it is "GenerateOnDemand"
+	// Meaning that all "GenerateOnDemand" requests are always valid, and "GenerateOnLoad" request is only valid if we want a "GenerateOnLoad" trigger.
+	bool bValidRequest = !(RequestedGenerationTrigger == EPCGComponentGenerationTrigger::GenerateOnLoad && GenerationTrigger == EPCGComponentGenerationTrigger::GenerateOnDemand);
+
+	return ((!bGenerated && bValidRequest) ||
 #if WITH_EDITOR
 			bDirtyGenerated || 
 #endif
 			bForce);
-	}
 }
 
 void UPCGComponent::SetPropertiesFromOriginal(const UPCGComponent* Original)
@@ -173,12 +174,12 @@ void UPCGComponent::SetPropertiesFromOriginal(const UPCGComponent* Original)
 
 void UPCGComponent::Generate()
 {
-#if WITH_EDITOR
 	if (bIsGenerating)
 	{
 		return;
 	}
 
+#if WITH_EDITOR
 	FScopedTransaction Transaction(LOCTEXT("PCGGenerate", "Execute generation on PCG component"));
 #endif
 
@@ -192,12 +193,15 @@ void UPCGComponent::Generate_Implementation(bool bForce)
 
 void UPCGComponent::GenerateLocal(bool bForce)
 {
-#if WITH_EDITOR
+	GenerateLocal(bForce, EPCGComponentGenerationTrigger::GenerateOnDemand);
+}
+
+void UPCGComponent::GenerateLocal(bool bForce, EPCGComponentGenerationTrigger RequestedGenerationTrigger)
+{
 	if (bIsGenerating)
 	{
 		return;
 	}
-#endif
 
 	// Force component activation so it's easier to control by BP.
 	if (!bActivated)
@@ -206,21 +210,19 @@ void UPCGComponent::GenerateLocal(bool bForce)
 		bActivated = true;
 	}
 
-	FPCGTaskId TaskId = GenerateInternal(bForce, {});
+	FPCGTaskId TaskId = GenerateInternal(bForce, RequestedGenerationTrigger, {});
 
-#if WITH_EDITOR
 	if (TaskId != InvalidPCGTaskId)
 	{
 		bIsGenerating = true;
 	}
-#endif
 }
 
-FPCGTaskId UPCGComponent::GenerateInternal(bool bForce, const TArray<FPCGTaskId>& TaskDependencies)
+FPCGTaskId UPCGComponent::GenerateInternal(bool bForce, EPCGComponentGenerationTrigger RequestedGenerationTrigger, const TArray<FPCGTaskId>& TaskDependencies)
 {
 	FPCGTaskId TaskId = InvalidPCGTaskId;
 
-	if (!ShouldGenerate(bForce))
+	if (!ShouldGenerate(bForce, RequestedGenerationTrigger))
 	{
 		return InvalidPCGTaskId;
 	}
@@ -238,8 +240,34 @@ FPCGTaskId UPCGComponent::GenerateInternal(bool bForce, const TArray<FPCGTaskId>
 	if (IsPartitioned())
 	{
 #if WITH_EDITOR
-		TaskId = GetSubsystem()->DelayGenerateGraph(this, /*bSave=*/bForce);
+		if (!GIsPlayInEditorWorld)
+		{
+			TaskId = GetSubsystem()->DelayGenerateGraph(this, /*bSave=*/bForce);
+		}
+		else
 #endif
+		{
+			// If we don't have valid bounds, just cleanup
+			const FBox NewBounds = GetGridBounds();
+			if (!NewBounds.IsValid)
+			{
+				CleanupLocal(/*bRemoveComponents=*/false);
+				return InvalidPCGTaskId;
+			}
+
+			// Otherwise, ask for generation on all the partition actors registered.
+			TArray<FPCGTaskId> TaskIds = GetSubsystem()->ScheduleMultipleComponent(this, PartitionActors, TaskDependencies);
+
+			// Finally, create a task to call PostProcessGraph.
+			if (!TaskIds.IsEmpty())
+			{
+				return GetSubsystem()->ScheduleGeneric([this, NewBounds]() { PostProcessGraph(NewBounds, true); return true; }, TaskIds);
+			}
+			else
+			{
+				return InvalidPCGTaskId;
+			}
+		}
 	}
 	else
 	{
@@ -250,10 +278,12 @@ FPCGTaskId UPCGComponent::GenerateInternal(bool bForce, const TArray<FPCGTaskId>
 		}
 
 		const FBox NewBounds = GetGridBounds();
-		if (NewBounds.IsValid)
+		if (!NewBounds.IsValid)
 		{
-			TaskId = GetSubsystem()->ScheduleComponent(this, TaskDependencies);
+			return InvalidPCGTaskId;
 		}
+
+		TaskId = GetSubsystem()->ScheduleComponent(this, TaskDependencies);
 	}
 
 	return TaskId;
@@ -297,6 +327,21 @@ bool UPCGComponent::GetActorsFromTags(const TSet<FName>& InTags, TSet<TWeakObjec
 	return bHasValidTag;
 }
 
+void UPCGComponent::AddPCGPartitionActor(const APCGPartitionActor* Actor)
+{
+	PartitionActors.Add(Actor);
+}
+
+void UPCGComponent::RemovePCGPartitionActor(const APCGPartitionActor* Actor)
+{
+	PartitionActors.Remove(Actor);
+}
+
+void UPCGComponent::ClearPCGPartitionActors()
+{
+	PartitionActors.Empty();
+}
+
 void UPCGComponent::PostProcessGraph(const FBox& InNewBounds, bool bInGenerated)
 {
 	LastGeneratedBounds = InNewBounds;
@@ -307,9 +352,10 @@ void UPCGComponent::PostProcessGraph(const FBox& InNewBounds, bool bInGenerated)
 
 		bGenerated = true;
 
+		bIsGenerating = false;
+
 #if WITH_EDITOR
 		bDirtyGenerated = false;
-		bIsGenerating = false;
 		OnPCGGraphGeneratedDelegate.Broadcast(this);
 #endif
 	}
@@ -320,10 +366,10 @@ void UPCGComponent::OnProcessGraphAborted()
 	UE_LOG(LogPCG, Warning, TEXT("Process Graph was called but aborted, check for errors in log if you expected a result."));
 
 	bGenerated = false;
+	bIsGenerating = false;
 
 #if WITH_EDITOR
 	bDirtyGenerated = false;
-	bIsGenerating = false;
 #endif
 }
 
@@ -360,9 +406,14 @@ void UPCGComponent::CleanupLocal(bool bRemoveComponents, bool bSave)
 		{
 			Modify();
 			GetSubsystem()->CleanupGraph(this, LastGeneratedBounds, bRemoveComponents, bSave);
-			bGenerated = false;
 		}
+		else
 #endif
+		{
+			GetSubsystem()->ScheduleMultipleCleanup(this, PartitionActors, bRemoveComponents, {});
+		}
+
+		bGenerated = false;
 	}
 	else
 	{
@@ -382,12 +433,10 @@ AActor* UPCGComponent::ClearPCGLink(UClass* TemplateActor)
 	}
 
 	// TODO: Perhaps remove this part if we want to do it in the PCG Graph.
-#if WITH_EDITOR
 	if (bIsGenerating)
 	{
 		return nullptr;
 	}
-#endif // WITH_EDITOR
 
 	UWorld* World = GetWorld();
 
@@ -514,13 +563,34 @@ void UPCGComponent::BeginPlay()
 {
 	Super::BeginPlay();
 
-	if(bActivated && !bGenerated && !IsPartitioned() && GenerationTrigger == EPCGComponentGenerationTrigger::GenerateOnLoad)
+	if(bActivated && !bGenerated && GenerationTrigger == EPCGComponentGenerationTrigger::GenerateOnLoad)
 	{
-		GenerateLocal(/*bForce=*/false);
-		bRuntimeGenerated = true;
+		if (IsPartitioned())
+		{
+			// If we are partitioned, the responsibility of the generation is to the partition actors.
+			// but we still need to know that we are currently generated (even if the state is held by the partition actors)
+			// TODO: Will be cleaner when we have dynamic association.
+			const FBox NewBounds = GetGridBounds();
+			if (NewBounds.IsValid)
+			{
+				PostProcessGraph(NewBounds, true);
+			}
+		}
+		else
+		{
+			GenerateLocal(/*bForce=*/false, EPCGComponentGenerationTrigger::GenerateOnLoad);
+			bRuntimeGenerated = true;
+		}
 	}
 }
 
+void UPCGComponent::EndPlay(const EEndPlayReason::Type EndPlayReason)
+{
+	Super::EndPlay(EndPlayReason);
+
+	ClearPCGPartitionActors();
+}
+
 void UPCGComponent::OnComponentCreated()
 {
 	Super::OnComponentCreated();
@@ -664,19 +734,39 @@ void UPCGComponent::PostEditChangeProperty(FPropertyChangedEvent& PropertyChange
 			if (bActivated)
 			{
 				bool bIsNowPartitioned = bIsPartitioned;
-				bIsPartitioned = !bIsPartitioned;
-
-				// First, we'll cleanup
-				bActivated = false;
-				Refresh();
-
-				// Then invalidate the previous bounds to force actor creation (as if we moved the volume)
-				// and do a normal refresh
-				bActivated = true;
-				bIsPartitioned = bIsNowPartitioned;
-				ResetLastGeneratedBounds();
-				DirtyGenerated();
-				Refresh();
+				if (bGenerated)
+				{
+					bIsPartitioned = !bIsPartitioned;
+
+					// First, we'll cleanup
+					bActivated = false;
+					Refresh();
+
+					// Then invalidate the previous bounds to force actor creation (as if we moved the volume)
+					// and do a normal refresh
+					bActivated = true;
+					bIsPartitioned = bIsNowPartitioned;
+					ResetLastGeneratedBounds();
+					DirtyGenerated();
+					Refresh();
+				}
+				else
+				{
+					// We need the component to be partitioned if we use the subsystem.
+					bIsPartitioned = true;
+
+					if (bIsNowPartitioned)
+					{
+						GetSubsystem()->DelayPartitionGraph(this);
+					}
+					else
+					{
+						GetSubsystem()->DelayUnpartitionGraph(this);
+
+					}
+
+					bIsPartitioned = bIsNowPartitioned;
+				}
 			}
 		}
 		else
diff --git a/Engine/Plugins/Experimental/PCG/Source/PCG/Private/PCGSubsystem.cpp b/Engine/Plugins/Experimental/PCG/Source/PCG/Private/PCGSubsystem.cpp
--- a/Engine/Plugins/Experimental/PCG/Source/PCG/Private/PCGSubsystem.cpp
+++ b/Engine/Plugins/Experimental/PCG/Source/PCG/Private/PCGSubsystem.cpp
@@ -154,6 +154,40 @@ FPCGTaskId UPCGSubsystem::ScheduleGeneric(TFunction<bool()> InOperation, const T
 	return GraphExecutor->ScheduleGeneric(InOperation, TaskDependencies);
 }
 
+TArray<FPCGTaskId> UPCGSubsystem::ScheduleMultipleComponent(UPCGComponent* OriginalComponent, TSet<TSoftObjectPtr<APCGPartitionActor>>& PartitionActors, const TArray<FPCGTaskId>& Dependencies)
+{
+	TArray<FPCGTaskId> TaskIds;
+	for (TSoftObjectPtr<APCGPartitionActor>& PartitionActorPtr : PartitionActors)
+	{
+		if (APCGPartitionActor* PartitionActor = PartitionActorPtr.Get())
+		{
+			if (UPCGComponent* LocalComponent = PartitionActor->GetLocalComponent(OriginalComponent))
+			{
+				TaskIds.Add(ScheduleComponent(LocalComponent, Dependencies));
+			}
+		}
+	}
+
+	return TaskIds;
+}
+
+TArray<FPCGTaskId> UPCGSubsystem::ScheduleMultipleCleanup(UPCGComponent* OriginalComponent, TSet<TSoftObjectPtr<APCGPartitionActor>>& PartitionActors, bool bRemoveComponents, const TArray<FPCGTaskId>& Dependencies)
+{
+	TArray<FPCGTaskId> TaskIds;
+	for (TSoftObjectPtr<APCGPartitionActor>& PartitionActorPtr : PartitionActors)
+	{
+		if (APCGPartitionActor* PartitionActor = PartitionActorPtr.Get())
+		{
+			if (UPCGComponent* LocalComponent = PartitionActor->GetLocalComponent(OriginalComponent))
+			{
+				TaskIds.Add(ScheduleCleanup(LocalComponent, bRemoveComponents, Dependencies));
+			}
+		}
+	}
+
+	return TaskIds;
+}
+
 bool UPCGSubsystem::GetOutputData(FPCGTaskId TaskId, FPCGDataCollection& OutData)
 {
 	check(GraphExecutor);
@@ -360,7 +394,7 @@ FPCGTaskId UPCGSubsystem::ProcessGraph(UPCGComponent* Component, const FBox& InP
 					return InvalidPCGTaskId;
 				}
 
-				return LocalComponent->GenerateInternal(/*bForce=*/false, TaskDependencies);
+				return LocalComponent->GenerateInternal(/*bForce=*/false, EPCGComponentGenerationTrigger::GenerateOnDemand, TaskDependencies);
 			}
 			else
 			{
@@ -442,29 +476,42 @@ FPCGTaskId UPCGSubsystem::ProcessGraph(UPCGComponent* Component, const FBox& InP
 	}
 }
 
+FPCGTaskId UPCGSubsystem::ScheduleCleanup(UPCGComponent* PCGComponent, bool bRemoveComponents, const TArray<FPCGTaskId>& Dependencies)
+{
+	TWeakObjectPtr<UPCGComponent> ComponentPtr(PCGComponent);
+
+	auto CleanupTask = [ComponentPtr, bRemoveComponents]() {
+		check(ComponentPtr.Get() != nullptr);
+		if (UPCGComponent* Component = ComponentPtr.Get())
+		{
+			Component->CleanupInternal(bRemoveComponents);
+		}
+
+		return true;
+	};
+
+	return GraphExecutor->ScheduleGeneric(CleanupTask, Dependencies);
+}
+
 void UPCGSubsystem::CleanupGraph(UPCGComponent* Component, const FBox& InBounds, bool bRemoveComponents, bool bSave)
 {
 	TWeakObjectPtr<UPCGComponent> ComponentPtr(Component);
 
 	auto ScheduleTask = [this, ComponentPtr, bRemoveComponents](APCGPartitionActor* PCGActor, const FBox& InIntersectedBounds, const TArray<FPCGTaskId>& TaskDependencies) {
-		auto CleanupTask = [ComponentPtr, PCGActor, bRemoveComponents]() {
-			check(ComponentPtr.Get() != nullptr && PCGActor != nullptr);
-			UPCGComponent* Component = ComponentPtr.Get();
-
-			if (!PCGActor || !Component)
-			{
-				return true;
-			}
+		UPCGComponent* Component = ComponentPtr.Get();
+		check(Component != nullptr && PCGActor != nullptr);
 
-			if (UPCGComponent* LocalComponent = PCGActor->GetLocalComponent(Component))
-			{
-				LocalComponent->CleanupInternal(bRemoveComponents);
-			}
+		if (!PCGActor || !Component)
+		{
+			return InvalidPCGTaskId;
+		}
 
-			return true;
-		};
+		if (UPCGComponent* LocalComponent = PCGActor->GetLocalComponent(Component))
+		{
+			return ScheduleCleanup(LocalComponent, bRemoveComponents, TaskDependencies);
+		}
 
-		return GraphExecutor->ScheduleGeneric(CleanupTask, TaskDependencies);
+		return InvalidPCGTaskId;
 	};
 
 	PCGSubsystem::ForEachIntersectingCell(GraphExecutor, Component->GetWorld(), InBounds, /*bCreateActor=*/false, /*bLoadCell=*/false, bSave, ScheduleTask);
diff --git a/Engine/Plugins/Experimental/PCG/Source/PCG/Public/PCGComponent.h b/Engine/Plugins/Experimental/PCG/Source/PCG/Public/PCGComponent.h
--- a/Engine/Plugins/Experimental/PCG/Source/PCG/Public/PCGComponent.h
+++ b/Engine/Plugins/Experimental/PCG/Source/PCG/Public/PCGComponent.h
@@ -8,6 +8,7 @@
 
 #include "PCGComponent.generated.h"
 
+class APCGPartitionActor;
 class UPCGComponent;
 class UPCGGraph;
 class UPCGManagedResource;
@@ -68,6 +69,7 @@ public:
 
 	//~Begin UActorComponent Interface
 	virtual void BeginPlay() override;
+	virtual void EndPlay(const EEndPlayReason::Type EndPlayReason) override;
 	virtual void OnComponentCreated() override;
 	virtual void OnComponentDestroyed(bool bDestroyingHierarchy) override;
 	//~End UActorComponent Interface
@@ -131,10 +133,7 @@ public:
 	UPROPERTY(BlueprintReadWrite, EditAnywhere, AdvancedDisplay, Category = Properties, meta = (EditCondition = "!bIsComponentLocal"))
 	bool bIsPartitioned = true;
 
-	UPROPERTY(BlueprintReadOnly, EditAnywhere, Category = Properties, AdvancedDisplay, meta = (EditCondition = "bIsComponentLocal", EditConditionHides))
-	bool bGenerationTriggerLocalOverride = false;
-
-	UPROPERTY(BlueprintReadOnly, EditAnywhere, Category = Properties, AdvancedDisplay, meta = (EditCondition = "!bIsComponentLocal || bGenerationTriggerLocalOverride", EditConditionHides))
+	UPROPERTY(BlueprintReadOnly, EditAnywhere, Category = Properties, AdvancedDisplay, meta = (EditCondition = "!bIsComponentLocal", EditConditionHides))
 	EPCGComponentGenerationTrigger GenerationTrigger = EPCGComponentGenerationTrigger::GenerateOnLoad;
 
 	/** Flag to indicate whether this component has run in the editor. Note that for partitionable actors, this will always be false. */
@@ -158,6 +157,9 @@ public:
 	FOnPCGGraphCleaned OnPCGGraphCleanedDelegate;
 #endif
 
+	/** Return if we are currently generating the graph for this component */
+	bool IsGenerating() const { return bIsGenerating; }
+
 #if WITH_EDITOR
 	void Refresh();
 	void DirtyGenerated(EPCGComponentDirtyFlag DataToDirtyFlag = EPCGComponentDirtyFlag::None);
@@ -165,9 +167,6 @@ public:
 	/** Reset last generated bounds to force PCGPartitionActor creation on next refresh */
 	void ResetLastGeneratedBounds();
 
-	/** Return if we are currently generating the graph for this component */
-	bool IsGenerating() const { return bIsGenerating; }
-
 	/** Functions for managing the node inspection cache */
 	bool IsInspecting() const { return bIsInspecting; }
 	void EnableInspection() { bIsInspecting = true; };
@@ -183,6 +182,10 @@ public:
 
 	UPCGSubsystem* GetSubsystem() const;
 
+	void AddPCGPartitionActor(const APCGPartitionActor* Actor);
+	void RemovePCGPartitionActor(const APCGPartitionActor* Actor);
+	void ClearPCGPartitionActors();
+
 protected:
 	UPROPERTY(BlueprintReadWrite, EditAnywhere, Category = Properties)
 	TObjectPtr<UPCGGraph> Graph;
@@ -195,8 +198,9 @@ private:
 	UPCGData* CreateLandscapePCGData();
 	void UpdatePCGExclusionData();
 
-	bool ShouldGenerate(bool bForce = false) const;
-	FPCGTaskId GenerateInternal(bool bForce, const TArray<FPCGTaskId>& Dependencies);
+	bool ShouldGenerate(bool bForce, EPCGComponentGenerationTrigger RequestedGenerationTrigger) const;
+	void GenerateLocal(bool bForce, EPCGComponentGenerationTrigger RequestedGenerationTrigger);
+	FPCGTaskId GenerateInternal(bool bForce, EPCGComponentGenerationTrigger RequestedGenerationTrigger, const TArray<FPCGTaskId>& Dependencies);
 	void CleanupInternal(bool bRemoveComponents);
 	void CleanupInternal(bool bRemoveComponents, TSet<TSoftObjectPtr<AActor>>& OutActorsToDelete);
 	void PostProcessGraph(const FBox& InNewBounds, bool bInGenerated);
@@ -276,8 +280,9 @@ private:
 	UPROPERTY()
 	FBox LastGeneratedBounds = FBox(EForceInit::ForceInit);
 
-#if WITH_EDITOR
 	bool bIsGenerating = false;
+
+#if WITH_EDITOR
 	bool bIsInspecting = false;
 	FBox LastGeneratedBoundsPriorToUndo = FBox(EForceInit::ForceInit);
 	FPCGTagToSettingsMap CachedTrackedTagsToSettings;
@@ -307,5 +312,9 @@ private:
 	TMap<const UPCGNode*, FPCGDataCollection> InspectionCache;
 #endif
 
+	// TODO: Will need to be cleaned up after we have dynamic association between PartitionActors and PCGComponents.
+	UPROPERTY(Transient)
+	TSet<TSoftObjectPtr<APCGPartitionActor>> PartitionActors;
+
 	FCriticalSection GeneratedResourcesLock;
 };
diff --git a/Engine/Plugins/Experimental/PCG/Source/PCG/Public/PCGSubsystem.h b/Engine/Plugins/Experimental/PCG/Source/PCG/Public/PCGSubsystem.h
--- a/Engine/Plugins/Experimental/PCG/Source/PCG/Public/PCGSubsystem.h
+++ b/Engine/Plugins/Experimental/PCG/Source/PCG/Public/PCGSubsystem.h
@@ -50,6 +50,15 @@ public:
 	// Schedule graph (owner -> graph)
 	FPCGTaskId ScheduleComponent(UPCGComponent* PCGComponent, const TArray<FPCGTaskId>& Dependencies);
 
+	// Schedule cleanup (owner -> graph)
+	FPCGTaskId ScheduleCleanup(UPCGComponent* PCGComponent, bool bRemoveComponents, const TArray<FPCGTaskId>& Dependencies);
+
+	// Schedule multiple graphs
+	TArray<FPCGTaskId> ScheduleMultipleComponent(UPCGComponent* OriginalComponent, TSet<TSoftObjectPtr<APCGPartitionActor>>& PartitionActors, const TArray<FPCGTaskId>& Dependencies);
+
+	// Schedule multiple cleanups
+	TArray<FPCGTaskId> ScheduleMultipleCleanup(UPCGComponent* OriginalComponent, TSet<TSoftObjectPtr<APCGPartitionActor>>& PartitionActors, bool bRemoveComponents, const TArray<FPCGTaskId>& Dependencies);
+
 	// Schedule graph (used internally for dynamic subgraph execution)
 	FPCGTaskId ScheduleGraph(UPCGGraph* Graph, UPCGComponent* SourceComponent, FPCGElementPtr InputElement, const TArray<FPCGTaskId>& Dependencies);
 
