PCG:
*Added data selection combobox to attribute widget to be able to select between output data on the inspected node.

diff --git a/Engine/Plugins/Experimental/PCG/Source/PCGEditor/Private/SPCGEditorGraphAttributeListView.cpp b/Engine/Plugins/Experimental/PCG/Source/PCGEditor/Private/SPCGEditorGraphAttributeListView.cpp
--- a/Engine/Plugins/Experimental/PCG/Source/PCGEditor/Private/SPCGEditorGraphAttributeListView.cpp
+++ b/Engine/Plugins/Experimental/PCG/Source/PCGEditor/Private/SPCGEditorGraphAttributeListView.cpp
@@ -20,6 +20,8 @@
 
 namespace PCGEditorGraphAttributeListView
 {
+	const FText NoDataAvailableText = LOCTEXT("NoDataAvailableText", "No data available");
+	
 	/** Names of the columns in the attribute list */
 	const FName NAME_IndexColumn = FName(TEXT("IndexColumn"));
 	const FName NAME_PointPositionX = FName(TEXT("PointPositionX"));
@@ -230,10 +232,24 @@ void SPCGEditorGraphAttributeListView::Construct(const FArguments& InArgs, TShar
 		.ExternalScrollbar(VerticalScrollBar)
 		.ConsumeMouseWheel(EConsumeMouseWheel::Always);
 
+	SAssignNew(DataComboBox, SComboBox<TSharedPtr<FName>>)
+		.OptionsSource(&DataComboBoxItems)
+		.OnGenerateWidget(this, &SPCGEditorGraphAttributeListView::OnGenerateDataWidget)
+		.OnSelectionChanged(this, &SPCGEditorGraphAttributeListView::OnSelectionChanged)
+		[
+			SNew(STextBlock)
+			.Text(this, &SPCGEditorGraphAttributeListView::OnGenerateSelectedDataText)
+		];
+	
 	this->ChildSlot
 	[
 		SNew(SVerticalBox)
 		+SVerticalBox::Slot()
+		.AutoHeight()
+		[
+			DataComboBox->AsShared()
+		]
+		+SVerticalBox::Slot()
 		.FillHeight(1.0f)
 		[
 			SNew(SHorizontalBox)
@@ -402,24 +418,29 @@ void SPCGEditorGraphAttributeListView::OnDebugObjectChanged(UPCGComponent* InPCG
 	if (PCGComponent)
 	{
 		PCGComponent->EnableInspection();
-		PCGComponent->OnPCGGraphGeneratedDelegate.AddSP(this, &SPCGEditorGraphAttributeListView::RebuildAttributeList);
-		PCGComponent->OnPCGGraphCleanedDelegate.AddSP(this, &SPCGEditorGraphAttributeListView::RebuildAttributeList);	
+		PCGComponent->OnPCGGraphGeneratedDelegate.AddSP(this, &SPCGEditorGraphAttributeListView::OnGenerateUpdated);
+		PCGComponent->OnPCGGraphCleanedDelegate.AddSP(this, &SPCGEditorGraphAttributeListView::OnGenerateUpdated);	
+	}
+	else
+	{
+		RefreshDataComboBox();
+		RefreshAttributeList();
 	}
-
-	RebuildAttributeList();
 }
 
 void SPCGEditorGraphAttributeListView::OnInspectedNodeChanged(UPCGNode* /*InPCGNode*/)
-{	
-	RebuildAttributeList();
+{
+	RefreshDataComboBox();
+	RefreshAttributeList();
 }
 
-void SPCGEditorGraphAttributeListView::RebuildAttributeList(UPCGComponent* /*InPCGComponent*/)
+void SPCGEditorGraphAttributeListView::OnGenerateUpdated(UPCGComponent* /*InPCGComponent*/)
 {
-	RebuildAttributeList();
+	RefreshDataComboBox();
+	RefreshAttributeList();
 }
 
-void SPCGEditorGraphAttributeListView::RebuildAttributeList()
+void SPCGEditorGraphAttributeListView::RefreshAttributeList()
 {
 	ListViewItems.Empty();
 	ListView->RequestListRefresh();
@@ -446,32 +467,108 @@ void SPCGEditorGraphAttributeListView::RebuildAttributeList()
 	{
 		return;
 	}
+
+	const int32 DataIndex = GetSelectedDataIndex();
+	if (!InspectionData->TaggedData.IsValidIndex(DataIndex))
+	{
+		return;
+	}
 	
-	if (InspectionData->TaggedData.Num() > 0)
+	const FPCGTaggedData& TaggedData = InspectionData->TaggedData[DataIndex];
+	const UPCGData* PCGData = TaggedData.Data;
+	if (const UPCGSpatialData* PCGSpatialData = Cast<UPCGSpatialData>(PCGData))
 	{
-		const FPCGTaggedData& TaggedData = InspectionData->TaggedData[0]; // TODO: Add dropdown to select tagged data.
-		const UPCGData* PCGData = TaggedData.Data;
-		if (const UPCGSpatialData* PCGSpatialData = Cast<UPCGSpatialData>(PCGData))
+		if (const UPCGPointData* PCGPointData = PCGSpatialData->ToPointData())
 		{
-			if (const UPCGPointData* PCGPointData = PCGSpatialData->ToPointData())
+			const TArray<FPCGPoint>& PCGPoints = PCGPointData->GetPoints();
+
+			ListViewItems.Reserve(PCGPoints.Num());
+			for (int32 PointIndex = 0; PointIndex < PCGPoints.Num(); PointIndex++)
 			{
-				const TArray<FPCGPoint>& PCGPoints = PCGPointData->GetPoints();
-
-				ListViewItems.Reserve(PCGPoints.Num());
-				for (int32 PointIndex = 0; PointIndex < PCGPoints.Num(); PointIndex++)
-				{
-					const FPCGPoint& PCGPoint = PCGPoints[PointIndex];
-					// TODO: Investigate swapping out the shared ptr's for better performance on huge data sets
-					PCGListviewItemPtr ListViewItem = MakeShared<FPCGListViewItem>();
-					ListViewItem->Index = PointIndex;
-					ListViewItem->PCGPoint = &PCGPoint;
-					ListViewItems.Add(ListViewItem);
-				}
+				const FPCGPoint& PCGPoint = PCGPoints[PointIndex];
+				// TODO: Investigate swapping out the shared ptr's for better performance on huge data sets
+				PCGListviewItemPtr ListViewItem = MakeShared<FPCGListViewItem>();
+				ListViewItem->Index = PointIndex;
+				ListViewItem->PCGPoint = &PCGPoint;
+				ListViewItems.Add(ListViewItem);
 			}
 		}
 	}
 }
 
+void SPCGEditorGraphAttributeListView::RefreshDataComboBox()
+{
+	DataComboBoxItems.Empty();
+	DataComboBox->ClearSelection();
+	DataComboBox->RefreshOptions();
+
+	if (!PCGComponent)
+	{
+		return;
+	}
+
+	const TSharedPtr<FPCGEditor> PCGEditor = PCGEditorPtr.Pin();
+	if (!PCGEditor.IsValid())
+	{
+		return;
+	}
+
+	const UPCGNode* PCGNode = PCGEditor->GetPCGNodeBeingInspected();
+	if (!PCGNode)
+	{
+		return;
+	}
+	
+	const FPCGDataCollection* InspectionData = PCGComponent->GetInspectionData(PCGNode);
+	if (!InspectionData)
+	{
+		return;
+	}
+
+	for (const FPCGTaggedData& TaggedData : InspectionData->TaggedData)
+	{
+		DataComboBoxItems.Add(MakeShared<FName>(TaggedData.Pin));
+	}
+
+	if (DataComboBoxItems.Num() > 0)
+	{
+		DataComboBox->SetSelectedItem(DataComboBoxItems[0]);
+	}
+}
+
+TSharedRef<SWidget> SPCGEditorGraphAttributeListView::OnGenerateDataWidget(TSharedPtr<FName> InItem) const
+{
+	return SNew(STextBlock).Text(FText::FromName(InItem.IsValid() ? *InItem : NAME_None));
+}
+
+void SPCGEditorGraphAttributeListView::OnSelectionChanged(TSharedPtr<FName> Item, ESelectInfo::Type SelectInfo)
+{
+	RefreshAttributeList();
+}
+
+FText SPCGEditorGraphAttributeListView::OnGenerateSelectedDataText() const
+{
+	if (const TSharedPtr<FName> SelectedDataName = DataComboBox->GetSelectedItem())
+	{
+		return FText::FromName(*SelectedDataName);
+	}
+	else
+	{
+		return PCGEditorGraphAttributeListView::NoDataAvailableText;
+	}
+}
+
+int32 SPCGEditorGraphAttributeListView::GetSelectedDataIndex() const
+{
+	int32 Index = INDEX_NONE;
+	if (const TSharedPtr<FName> SelectedItem = DataComboBox->GetSelectedItem())
+	{
+		DataComboBoxItems.Find(SelectedItem, Index);
+	}
+
+	return Index;
+}
+
 TSharedRef<ITableRow> SPCGEditorGraphAttributeListView::OnGenerateRow(PCGListviewItemPtr Item, const TSharedRef<STableViewBase>& OwnerTable) const
 {
 	return SNew(SPCGListViewItemRow, OwnerTable, Item);
diff --git a/Engine/Plugins/Experimental/PCG/Source/PCGEditor/Private/SPCGEditorGraphAttributeListView.h b/Engine/Plugins/Experimental/PCG/Source/PCGEditor/Private/SPCGEditorGraphAttributeListView.h
--- a/Engine/Plugins/Experimental/PCG/Source/PCGEditor/Private/SPCGEditorGraphAttributeListView.h
+++ b/Engine/Plugins/Experimental/PCG/Source/PCGEditor/Private/SPCGEditorGraphAttributeListView.h
@@ -3,6 +3,7 @@
 #pragma once
 
 #include "CoreMinimal.h"
+#include "Widgets/Input/SComboBox.h"
 #include "Widgets/SCompoundWidget.h"
 #include "Widgets/Views/SHeaderRow.h"
 #include "Widgets/Views/SListView.h"
@@ -50,8 +51,15 @@ private:
 	void OnDebugObjectChanged(UPCGComponent* InPCGComponent);
 	void OnInspectedNodeChanged(UPCGNode* InPCGNode);
 
-	void RebuildAttributeList(UPCGComponent* InPCGComponent);
-	void RebuildAttributeList();
+	void OnGenerateUpdated(UPCGComponent* InPCGComponent);
+	
+	void RefreshAttributeList();
+	void RefreshDataComboBox();
+
+	TSharedRef<SWidget> OnGenerateDataWidget(TSharedPtr<FName> InItem) const;
+	void OnSelectionChanged(TSharedPtr<FName> Item, ESelectInfo::Type SelectInfo);
+	FText OnGenerateSelectedDataText() const;
+	int32 GetSelectedDataIndex() const;
 
 	TSharedRef<ITableRow> OnGenerateRow(PCGListviewItemPtr Item, const TSharedRef<STableViewBase>& OwnerTable) const;
 
@@ -64,4 +72,7 @@ private:
 	TSharedPtr<SHeaderRow> ListViewHeader;
 	TSharedPtr<SListView<PCGListviewItemPtr>> ListView;
 	TArray<PCGListviewItemPtr> ListViewItems;
+
+	TSharedPtr<SComboBox<TSharedPtr<FName>>> DataComboBox;
+	TArray<TSharedPtr<FName>> DataComboBoxItems;
 };
