Fix PCGStaticMeshSpawnerTest in OpenWorld

#rnx

diff --git a/Engine/Plugins/Experimental/PCG/Source/PCG/Private/Tests/PCGTestsCommon.cpp b/Engine/Plugins/Experimental/PCG/Source/PCG/Private/Tests/PCGTestsCommon.cpp
--- a/Engine/Plugins/Experimental/PCG/Source/PCG/Private/Tests/PCGTestsCommon.cpp
+++ b/Engine/Plugins/Experimental/PCG/Source/PCG/Private/Tests/PCGTestsCommon.cpp
@@ -39,6 +39,8 @@ namespace PCGTestsCommon
 
 		TestPCGComponent = NewObject<UPCGComponent>(TestActor, FName(TEXT("Test PCG Component")), RF_Transient);
 		check(TestPCGComponent);
+		// By default PCG components for tests will be non-partitioned
+		TestPCGComponent->bIsPartitioned = false;
 		TestActor->AddInstanceComponent(TestPCGComponent);
 		TestPCGComponent->RegisterComponent();
 
