SetGraph in runtime PCG

diff --git a/Engine/Plugins/Experimental/PCG/Source/PCG/Private/PCGComponent.cpp b/Engine/Plugins/Experimental/PCG/Source/PCG/Private/PCGComponent.cpp
--- a/Engine/Plugins/Experimental/PCG/Source/PCG/Private/PCGComponent.cpp
+++ b/Engine/Plugins/Experimental/PCG/Source/PCG/Private/PCGComponent.cpp
@@ -54,7 +54,12 @@ bool UPCGComponent::IsPartitioned() const
 	return bIsPartitioned && CanPartition();
 }
 
-void UPCGComponent::SetGraph(UPCGGraph* InGraph)
+void UPCGComponent::SetGraph_Implementation(UPCGGraph* InGraph)
+{
+	SetGraphLocal(InGraph);
+}
+
+void UPCGComponent::SetGraphLocal(UPCGGraph* InGraph)
 {
 	if(Graph == InGraph)
 	{
@@ -75,9 +80,9 @@ void UPCGComponent::SetGraph(UPCGGraph* InGraph)
 	{
 		Graph->OnGraphChangedDelegate.AddUObject(this, &UPCGComponent::OnGraphChanged);
 	}
-
-	OnGraphChanged(Graph, true, false);
 #endif
+
+	OnGraphChanged(Graph, true, true);
 }
 
 void UPCGComponent::AddToManagedResources(UPCGManagedResource* InResource)
@@ -158,7 +163,7 @@ void UPCGComponent::SetPropertiesFromOriginal(const UPCGComponent* Original)
 
 	InputType = NewInputType;
 	Seed = Original->Seed;
-	SetGraph(Original->Graph);
+	SetGraphLocal(Original->Graph);
 
 	GenerationTrigger = Original->GenerationTrigger;
 
@@ -1197,9 +1202,15 @@ void UPCGComponent::OnGraphChanged(UPCGGraph* InGraph, bool bIsStructural)
 
 void UPCGComponent::OnGraphChanged(UPCGGraph* InGraph, bool bIsStructural, bool bShouldRefresh)
 {
-	if (InGraph == Graph)
+	if (InGraph != Graph)
+	{
+		return;
+	}
+
+#if WITH_EDITOR
+	// In editor, since we've changed the graph, we might have changed the tracked actor tags as well
+	if (!GIsPlayInEditorWorld)
 	{
-		// Since we've changed the graph, we might have changed the tracked actor tags as well
 		TeardownTrackingCallbacks();
 		SetupTrackingCallbacks();
 		RefreshTrackingData();
@@ -1207,14 +1218,29 @@ void UPCGComponent::OnGraphChanged(UPCGGraph* InGraph, bool bIsStructural, bool
 		UpdateTrackedLandscape();
 
 		DirtyGenerated();
-		if (bShouldRefresh)
+		if (InGraph && bShouldRefresh)
 		{
 			Refresh();
 		}
+		else if (!InGraph)
+		{
+			// With no graph, we clean up
+			CleanupLocal(/*bRemoveComponents=*/true, /*bSave=*/ false);
+		}
 
-#if WITH_EDITOR
-		InspectionCache.Empty();	
+		InspectionCache.Empty();
+		return;
+	}
 #endif
+
+	// Otherwise, if we are in PIE or runtime, force generate if we have a graph (and were generated). Or cleanup if we have no graph
+	if (InGraph && bGenerated)
+	{
+		GenerateLocal(/*bForce=*/true);
+	}
+	else if (!InGraph)
+	{
+		CleanupLocal(/*bRemoveComponents=*/true, /*bSave=*/ false);
 	}
 }
 
diff --git a/Engine/Plugins/Experimental/PCG/Source/PCG/Private/PCGSubsystem.cpp b/Engine/Plugins/Experimental/PCG/Source/PCG/Private/PCGSubsystem.cpp
--- a/Engine/Plugins/Experimental/PCG/Source/PCG/Private/PCGSubsystem.cpp
+++ b/Engine/Plugins/Experimental/PCG/Source/PCG/Private/PCGSubsystem.cpp
@@ -180,6 +180,9 @@ TArray<FPCGTaskId> UPCGSubsystem::ScheduleMultipleComponent(UPCGComponent* Origi
 		{
 			if (UPCGComponent* LocalComponent = PartitionActor->GetLocalComponent(OriginalComponent))
 			{
+				// Ensure that the PCG actor match our original
+				LocalComponent->SetPropertiesFromOriginal(OriginalComponent);
+
 				TaskIds.Add(ScheduleComponent(LocalComponent, Dependencies));
 			}
 		}
diff --git a/Engine/Plugins/Experimental/PCG/Source/PCG/Private/Tests/PCGTestsCommon.cpp b/Engine/Plugins/Experimental/PCG/Source/PCG/Private/Tests/PCGTestsCommon.cpp
--- a/Engine/Plugins/Experimental/PCG/Source/PCG/Private/Tests/PCGTestsCommon.cpp
+++ b/Engine/Plugins/Experimental/PCG/Source/PCG/Private/Tests/PCGTestsCommon.cpp
@@ -46,7 +46,7 @@ namespace PCGTestsCommon
 
 		UPCGGraph* TestGraph = NewObject<UPCGGraph>(TestPCGComponent, FName(TEXT("Test PCG Graph")), RF_Transient);
 		check(TestGraph);
-		TestPCGComponent->SetGraph(TestGraph);
+		TestPCGComponent->SetGraphLocal(TestGraph);
 #else
 		TestActor = nullptr;
 		TestPCGComponent = nullptr;
diff --git a/Engine/Plugins/Experimental/PCG/Source/PCG/Public/PCGComponent.h b/Engine/Plugins/Experimental/PCG/Source/PCG/Public/PCGComponent.h
--- a/Engine/Plugins/Experimental/PCG/Source/PCG/Public/PCGComponent.h
+++ b/Engine/Plugins/Experimental/PCG/Source/PCG/Public/PCGComponent.h
@@ -84,6 +84,9 @@ public:
 	bool CanPartition() const;
 
 	UPCGGraph* GetGraph() const { return Graph; }
+	void SetGraphLocal(UPCGGraph* InGraph);
+
+	UFUNCTION(BlueprintCallable, NetMulticast, Reliable, Category = PCG)
 	void SetGraph(UPCGGraph* InGraph);
 
 	void AddToManagedResources(UPCGManagedResource* InResource);
@@ -210,15 +213,15 @@ private:
 
 	bool GetActorsFromTags(const TSet<FName>& InTags, TSet<TWeakObjectPtr<AActor>>& OutActors, bool bCullAgainstLocalBounds);
 
+	void OnGraphChanged(UPCGGraph* InGraph, bool bIsStructural, bool bShouldRefresh);
+	void OnGraphChanged(UPCGGraph* InGraph, bool bIsStructural);
+
 #if WITH_EDITOR
 	virtual void PreEditChange(FProperty* PropertyAboutToChange) override;
 	virtual void PostEditChangeProperty(FPropertyChangedEvent& PropertyChangedEvent) override;
 	virtual void PreEditUndo() override;
 	virtual void PostEditUndo() override;
 
-	void OnGraphChanged(UPCGGraph* InGraph, bool bIsStructural, bool bShouldRefresh);
-	void OnGraphChanged(UPCGGraph* InGraph, bool bIsStructural);
-
 	void SetupActorCallbacks();
 	void TeardownActorCallbacks();
 	void SetupTrackingCallbacks();
