Add data collection "cancel on empty", useful for LOD-like processing

diff --git a/Engine/Plugins/Experimental/PCG/Source/PCG/Private/Graph/PCGGraphExecutor.cpp b/Engine/Plugins/Experimental/PCG/Source/PCG/Private/Graph/PCGGraphExecutor.cpp
--- a/Engine/Plugins/Experimental/PCG/Source/PCG/Private/Graph/PCGGraphExecutor.cpp
+++ b/Engine/Plugins/Experimental/PCG/Source/PCG/Private/Graph/PCGGraphExecutor.cpp
@@ -475,6 +475,11 @@ void FPCGGraphExecutor::BuildTaskInput(const FPCGGraphTask& Task, FPCGDataCollec
 			TaskInput.TaggedData.Append(InputCollection.TaggedData);
 		}
 
+		if (TaskInput.TaggedData.Num() == TaggedDataOffset && InputCollection.bCancelExecutionOnEmpty)
+		{
+			TaskInput.bCancelExecution = true;
+		}
+
 		// Apply labelling on data; technically, we should ensure that we do this only for pass-through nodes,
 		// Otherwise we could also null out the label on the input...
 		if (Input.OutPin)
diff --git a/Engine/Plugins/Experimental/PCG/Source/PCG/Public/PCGData.h b/Engine/Plugins/Experimental/PCG/Source/PCG/Public/PCGData.h
--- a/Engine/Plugins/Experimental/PCG/Source/PCG/Public/PCGData.h
+++ b/Engine/Plugins/Experimental/PCG/Source/PCG/Public/PCGData.h
@@ -82,6 +82,9 @@ struct PCG_API FPCGDataCollection
 	UPROPERTY(BlueprintReadWrite, EditAnywhere, Category = Data)
 	TArray<FPCGTaggedData> TaggedData;
 
+	UPROPERTY(BlueprintReadWrite, EditAnywhere, Category = Data)
+	bool bCancelExecutionOnEmpty = false;
+
 	/** This flag is used to cancel further computation or for the debug/isolate feature */
 	bool bCancelExecution = false;
 };
