Fix PCGPropertyToParamDataElementTest to cut invalid branches that could make it not compile

#rnx

diff --git a/Engine/Plugins/Experimental/PCG/Source/PCG/Private/Tests/Elements/PCGPropertyToParamDataElementTest.cpp b/Engine/Plugins/Experimental/PCG/Source/PCG/Private/Tests/Elements/PCGPropertyToParamDataElementTest.cpp
--- a/Engine/Plugins/Experimental/PCG/Source/PCG/Private/Tests/Elements/PCGPropertyToParamDataElementTest.cpp
+++ b/Engine/Plugins/Experimental/PCG/Source/PCG/Private/Tests/Elements/PCGPropertyToParamDataElementTest.cpp
@@ -17,8 +17,10 @@
 IMPLEMENT_CUSTOM_SIMPLE_AUTOMATION_TEST(FPCGPropertyToParamDataPropertyTypeTest, FPCGTestBaseClass, "pcg.tests.PropertyToParamData.PropertyType", PCGTestsCommon::TestFlags)
 IMPLEMENT_CUSTOM_SIMPLE_AUTOMATION_TEST(FPCGPropertyToParamDataActorFindTest, FPCGTestBaseClass, "pcg.tests.PropertyToParamData.ActorFind", PCGTestsCommon::TestFlags)
 
-template <typename AttributeType>
-bool VerifyAttributeValue(FPCGTestBaseClass* TestInstance, PCGTestsCommon::FTestData& TestData, FName PropertyName, AttributeType&& ExpectedValue, bool bShouldFail, FString ExtraTestWhat)
+// bShouldFail is a constexpr bool used to cut the "valid" branch
+// Because invalid types in expected failed tests could not compile.
+template <typename AttributeType, bool bShouldFail>
+bool VerifyAttributeValue(FPCGTestBaseClass* TestInstance, PCGTestsCommon::FTestData& TestData, FName PropertyName, AttributeType&& ExpectedValue, FString ExtraTestWhat)
 {
 	// Use universal reference to either pass an rvalue or lvalue. Use this type below to know the raw type behind.
 	using RawAttributeType = std::remove_const_t<std::remove_reference_t<AttributeType>>;
@@ -56,7 +58,7 @@ bool VerifyAttributeValue(FPCGTestBaseClass* TestInstance, PCGTestsCommon::FTest
 
 	bool bSuccess = true;
 
-	if (bShouldFail)
+	if constexpr (bShouldFail)
 	{
 		// If it should fail, it will have no output.
 		bSuccess = TestInstance->TestEqual(FormatWithPropertyName("NumOutput"), Context->OutputData.GetAllParams().Num(), 0);
@@ -88,6 +90,19 @@ bool VerifyAttributeValue(FPCGTestBaseClass* TestInstance, PCGTestsCommon::FTest
 	return bSuccess;
 }
 
+// Use aliases for convenience
+template <typename AttributeType>
+bool VerifyAttributeValueValid(FPCGTestBaseClass* TestInstance, PCGTestsCommon::FTestData& TestData, FName PropertyName, AttributeType&& ExpectedValue, FString ExtraTestWhat)
+{
+	return VerifyAttributeValue<AttributeType, false>(TestInstance, TestData, PropertyName, std::forward<AttributeType&&>(ExpectedValue), ExtraTestWhat);
+}
+
+template <typename AttributeType>
+bool VerifyAttributeValueInvalid(FPCGTestBaseClass* TestInstance, PCGTestsCommon::FTestData& TestData, FName PropertyName, AttributeType&& ExpectedValue, FString ExtraTestWhat)
+{
+	return VerifyAttributeValue<AttributeType, true>(TestInstance, TestData, PropertyName, std::forward<AttributeType&&>(ExpectedValue), ExtraTestWhat);
+}
+
 bool FPCGPropertyToParamDataPropertyTypeTest::RunTest(const FString& Parameters)
 {
 	bool bSuccess = true;
@@ -118,6 +133,7 @@ bool FPCGPropertyToParamDataPropertyTypeTest::RunTest(const FString& Parameters)
 	const FSoftClassPath SoftClassPathValue{ APCGVolume::StaticClass() };
 
 	const FVector2D Vector2Value = { 1.0, 2.0 };
+	const FColor ColorValue = FColor::Cyan;
 
 	APCGUnitTestDummyActor* Actor = Cast<APCGUnitTestDummyActor>(TestData.TestActor);
 	Actor->IntProperty = 42;
@@ -138,45 +154,47 @@ bool FPCGPropertyToParamDataPropertyTypeTest::RunTest(const FString& Parameters)
 	Actor->ClassProperty = APCGVolume::StaticClass();
 	Actor->ObjectProperty = PCGVolume;
 	Actor->Vector2Property = Vector2Value;
+	Actor->ColorProperty = ColorValue;
 
 	// Basic properties
-	bSuccess &= VerifyAttributeValue(this, TestData, GET_MEMBER_NAME_CHECKED(APCGUnitTestDummyActor, IntProperty), 42ll, false, ExtraTestWhat);
-	bSuccess &= VerifyAttributeValue(this, TestData, GET_MEMBER_NAME_CHECKED(APCGUnitTestDummyActor, Int64Property), 42ll, false, ExtraTestWhat);
-	bSuccess &= VerifyAttributeValue(this, TestData, GET_MEMBER_NAME_CHECKED(APCGUnitTestDummyActor, FloatProperty), 1.0, false, ExtraTestWhat);
-	bSuccess &= VerifyAttributeValue(this, TestData, GET_MEMBER_NAME_CHECKED(APCGUnitTestDummyActor, DoubleProperty), 1.0, false, ExtraTestWhat);
-	bSuccess &= VerifyAttributeValue(this, TestData, GET_MEMBER_NAME_CHECKED(APCGUnitTestDummyActor, BoolProperty), true, false, ExtraTestWhat);
+	bSuccess &= VerifyAttributeValueValid(this, TestData, GET_MEMBER_NAME_CHECKED(APCGUnitTestDummyActor, IntProperty), 42ll, ExtraTestWhat);
+	bSuccess &= VerifyAttributeValueValid(this, TestData, GET_MEMBER_NAME_CHECKED(APCGUnitTestDummyActor, Int64Property), 42ll, ExtraTestWhat);
+	bSuccess &= VerifyAttributeValueValid(this, TestData, GET_MEMBER_NAME_CHECKED(APCGUnitTestDummyActor, FloatProperty), 1.0, ExtraTestWhat);
+	bSuccess &= VerifyAttributeValueValid(this, TestData, GET_MEMBER_NAME_CHECKED(APCGUnitTestDummyActor, DoubleProperty), 1.0, ExtraTestWhat);
+	bSuccess &= VerifyAttributeValueValid(this, TestData, GET_MEMBER_NAME_CHECKED(APCGUnitTestDummyActor, BoolProperty), true, ExtraTestWhat);
 
 	// String/Name Properties
-	bSuccess &= VerifyAttributeValue(this, TestData, GET_MEMBER_NAME_CHECKED(APCGUnitTestDummyActor, NameProperty), NameValue, false, ExtraTestWhat);
-	bSuccess &= VerifyAttributeValue(this, TestData, GET_MEMBER_NAME_CHECKED(APCGUnitTestDummyActor, StringProperty), StringValue, false, ExtraTestWhat);
+	bSuccess &= VerifyAttributeValueValid(this, TestData, GET_MEMBER_NAME_CHECKED(APCGUnitTestDummyActor, NameProperty), NameValue, ExtraTestWhat);
+	bSuccess &= VerifyAttributeValueValid(this, TestData, GET_MEMBER_NAME_CHECKED(APCGUnitTestDummyActor, StringProperty), StringValue, ExtraTestWhat);
 
 	// Enum Property
-	bSuccess &= VerifyAttributeValue(this, TestData, GET_MEMBER_NAME_CHECKED(APCGUnitTestDummyActor, EnumProperty), (int64)EPCGUnitTestDummyEnum::Three, false, ExtraTestWhat);
+	bSuccess &= VerifyAttributeValueValid(this, TestData, GET_MEMBER_NAME_CHECKED(APCGUnitTestDummyActor, EnumProperty), (int64)EPCGUnitTestDummyEnum::Three, ExtraTestWhat);
 
 	// Struct Properties
-	bSuccess &= VerifyAttributeValue(this, TestData, GET_MEMBER_NAME_CHECKED(APCGUnitTestDummyActor, VectorProperty), VectorValue, false, ExtraTestWhat);
-	bSuccess &= VerifyAttributeValue(this, TestData, GET_MEMBER_NAME_CHECKED(APCGUnitTestDummyActor, Vector4Property), Vector4Value, false, ExtraTestWhat);
-	bSuccess &= VerifyAttributeValue(this, TestData, GET_MEMBER_NAME_CHECKED(APCGUnitTestDummyActor, RotatorProperty), RotatorValue, false, ExtraTestWhat);
-	bSuccess &= VerifyAttributeValue(this, TestData, GET_MEMBER_NAME_CHECKED(APCGUnitTestDummyActor, QuatProperty), QuatValue, false, ExtraTestWhat);
-	bSuccess &= VerifyAttributeValue(this, TestData, GET_MEMBER_NAME_CHECKED(APCGUnitTestDummyActor, TransformProperty), TransformValue, false, ExtraTestWhat);
-	bSuccess &= VerifyAttributeValue(this, TestData, GET_MEMBER_NAME_CHECKED(APCGUnitTestDummyActor, SoftObjectPathProperty), SoftObjectPathValue.ToString(), false, ExtraTestWhat);
-	bSuccess &= VerifyAttributeValue(this, TestData, GET_MEMBER_NAME_CHECKED(APCGUnitTestDummyActor, SoftClassPathProperty), SoftClassPathValue.ToString(), false, ExtraTestWhat);
+	bSuccess &= VerifyAttributeValueValid(this, TestData, GET_MEMBER_NAME_CHECKED(APCGUnitTestDummyActor, VectorProperty), VectorValue, ExtraTestWhat);
+	bSuccess &= VerifyAttributeValueValid(this, TestData, GET_MEMBER_NAME_CHECKED(APCGUnitTestDummyActor, Vector4Property), Vector4Value, ExtraTestWhat);
+	bSuccess &= VerifyAttributeValueValid(this, TestData, GET_MEMBER_NAME_CHECKED(APCGUnitTestDummyActor, RotatorProperty), RotatorValue, ExtraTestWhat);
+	bSuccess &= VerifyAttributeValueValid(this, TestData, GET_MEMBER_NAME_CHECKED(APCGUnitTestDummyActor, QuatProperty), QuatValue, ExtraTestWhat);
+	bSuccess &= VerifyAttributeValueValid(this, TestData, GET_MEMBER_NAME_CHECKED(APCGUnitTestDummyActor, TransformProperty), TransformValue, ExtraTestWhat);
+	bSuccess &= VerifyAttributeValueValid(this, TestData, GET_MEMBER_NAME_CHECKED(APCGUnitTestDummyActor, SoftObjectPathProperty), SoftObjectPathValue.ToString(), ExtraTestWhat);
+	bSuccess &= VerifyAttributeValueValid(this, TestData, GET_MEMBER_NAME_CHECKED(APCGUnitTestDummyActor, SoftClassPathProperty), SoftClassPathValue.ToString(), ExtraTestWhat);
 
 	// Objects properties
-	bSuccess &= VerifyAttributeValue(this, TestData, GET_MEMBER_NAME_CHECKED(APCGUnitTestDummyActor, ClassProperty), APCGVolume::StaticClass()->GetPathName(), false, ExtraTestWhat);
-	bSuccess &= VerifyAttributeValue(this, TestData, GET_MEMBER_NAME_CHECKED(APCGUnitTestDummyActor, ObjectProperty), PCGVolume->GetPathName(), false, ExtraTestWhat);
+	bSuccess &= VerifyAttributeValueValid(this, TestData, GET_MEMBER_NAME_CHECKED(APCGUnitTestDummyActor, ClassProperty), APCGVolume::StaticClass()->GetPathName(), ExtraTestWhat);
+	bSuccess &= VerifyAttributeValueValid(this, TestData, GET_MEMBER_NAME_CHECKED(APCGUnitTestDummyActor, ObjectProperty), PCGVolume->GetPathName(), ExtraTestWhat);
 
-	// Unsupported property
-	AddExpectedError(TEXT("Error while creating an attribute. Either the property type is not supported by PCG or attribute creation failed."), EAutomationExpectedErrorFlags::Contains, 1);
-	bSuccess &= VerifyAttributeValue(this, TestData, GET_MEMBER_NAME_CHECKED(APCGUnitTestDummyActor, Vector2Property), Vector2Value, true, ExtraTestWhat);
+	// Unsupported properties
+	AddExpectedError(TEXT("Error while creating an attribute. Either the property type is not supported by PCG or attribute creation failed."), EAutomationExpectedErrorFlags::Contains, 2);
+	bSuccess &= VerifyAttributeValueInvalid(this, TestData, GET_MEMBER_NAME_CHECKED(APCGUnitTestDummyActor, Vector2Property), Vector2Value, ExtraTestWhat);
+	bSuccess &= VerifyAttributeValueInvalid(this, TestData, GET_MEMBER_NAME_CHECKED(APCGUnitTestDummyActor, ColorProperty), ColorValue, ExtraTestWhat);
 
 	// Unknown property
 	AddExpectedError(TEXT("Property doesn't exist in the found actor."), EAutomationExpectedErrorFlags::Contains, 1);
-	bSuccess &= VerifyAttributeValue(this, TestData, TEXT("DummyMissingProperty"), 42, true, ExtraTestWhat);
+	bSuccess &= VerifyAttributeValueInvalid(this, TestData, TEXT("DummyMissingProperty"), 42, ExtraTestWhat);
 
 	// Missing property
 	AddExpectedError(TEXT("Some parameters are missing, abort."), EAutomationExpectedErrorFlags::Contains, 1);
-	bSuccess &= VerifyAttributeValue(this, TestData, NAME_None, 42, true, ExtraTestWhat);
+	bSuccess &= VerifyAttributeValueInvalid(this, TestData, NAME_None, 42, ExtraTestWhat);
 
 	return bSuccess;
 }
@@ -203,7 +221,7 @@ bool FPCGPropertyToParamDataActorFindTest::RunTest(const FString& Parameters)
 		PCGTestsCommon::FTestData TestData(Seed, Settings, APCGUnitTestDummyActor::StaticClass());
 		Cast<APCGUnitTestDummyActor>(TestData.TestActor)->IntProperty = 42;
 
-		bSuccess &= VerifyAttributeValue(this, TestData, GET_MEMBER_NAME_CHECKED(APCGUnitTestDummyActor, IntProperty), 42ll, false, "PropertyToParamDataActorFindTest_Self_Class");
+		bSuccess &= VerifyAttributeValueValid(this, TestData, GET_MEMBER_NAME_CHECKED(APCGUnitTestDummyActor, IntProperty), 42ll, "PropertyToParamDataActorFindTest_Self_Class");
 	}
 
 	// Self by tag
@@ -215,7 +233,7 @@ bool FPCGPropertyToParamDataActorFindTest::RunTest(const FString& Parameters)
 		TestData.TestActor->Tags.Add(Tag);
 		Cast<APCGUnitTestDummyActor>(TestData.TestActor)->IntProperty = 42;
 
-		bSuccess &= VerifyAttributeValue(this, TestData, GET_MEMBER_NAME_CHECKED(APCGUnitTestDummyActor, IntProperty), 42ll, false, "PropertyToParamDataActorFindTest_Self_Tag");
+		bSuccess &= VerifyAttributeValueValid(this, TestData, GET_MEMBER_NAME_CHECKED(APCGUnitTestDummyActor, IntProperty), 42ll, "PropertyToParamDataActorFindTest_Self_Tag");
 	}
 
 	// Self by name
@@ -226,7 +244,7 @@ bool FPCGPropertyToParamDataActorFindTest::RunTest(const FString& Parameters)
 		PCGTestsCommon::FTestData TestData(Seed, Settings, APCGUnitTestDummyActor::StaticClass());
 		Cast<APCGUnitTestDummyActor>(TestData.TestActor)->IntProperty = 42;
 
-		bSuccess &= VerifyAttributeValue(this, TestData, GET_MEMBER_NAME_CHECKED(APCGUnitTestDummyActor, IntProperty), 42ll, false, "PropertyToParamDataActorFindTest_Self_Name");
+		bSuccess &= VerifyAttributeValueValid(this, TestData, GET_MEMBER_NAME_CHECKED(APCGUnitTestDummyActor, IntProperty), 42ll, "PropertyToParamDataActorFindTest_Self_Name");
 	}
 
 	// TODO: Need a good way to spawn actors with parenting relation between them
@@ -239,7 +257,7 @@ bool FPCGPropertyToParamDataActorFindTest::RunTest(const FString& Parameters)
 	//	TestData.AddActor(APCGUnitTestDummyActor::StaticClass(), true);
 	//	Cast<APCGUnitTestDummyActor>(TestData.TestActor)->IntProperty = 42;
 
-	//	bSuccess &= VerifyAttributeValue(this, TestData, GET_MEMBER_NAME_CHECKED(APCGUnitTestDummyActor, IntProperty), 42ll, false, "PropertyToParamDataActorFindTest_Parent_Class");
+	//	bSuccess &= VerifyAttributeValueValid(this, TestData, GET_MEMBER_NAME_CHECKED(APCGUnitTestDummyActor, IntProperty), 42ll, "PropertyToParamDataActorFindTest_Parent_Class");
 	//}
 
 	//// Root by tag
@@ -253,7 +271,7 @@ bool FPCGPropertyToParamDataActorFindTest::RunTest(const FString& Parameters)
 	//	TestData.TestActor->Tags.Add(Tag);
 	//	Cast<APCGUnitTestDummyActor>(TestData.TestActor)->IntProperty = 42;
 
-	//	bSuccess &= VerifyAttributeValue(this, TestData, GET_MEMBER_NAME_CHECKED(APCGUnitTestDummyActor, IntProperty), 42ll, false, "PropertyToParamDataActorFindTest_Root_Tag");
+	//	bSuccess &= VerifyAttributeValueValid(this, TestData, GET_MEMBER_NAME_CHECKED(APCGUnitTestDummyActor, IntProperty), 42ll, "PropertyToParamDataActorFindTest_Root_Tag");
 	//}
 
 	//// Root by tag including children
@@ -269,7 +287,7 @@ bool FPCGPropertyToParamDataActorFindTest::RunTest(const FString& Parameters)
 
 	//	TestData.AddActor(AActor::StaticClass(), true);
 
-	//	bSuccess &= VerifyAttributeValue(this, TestData, GET_MEMBER_NAME_CHECKED(APCGUnitTestDummyActor, IntProperty), 42ll, false, "PropertyToParamDataActorFindTest_Root_Tag_Children");
+	//	bSuccess &= VerifyAttributeValueValid(this, TestData, GET_MEMBER_NAME_CHECKED(APCGUnitTestDummyActor, IntProperty), 42ll, "PropertyToParamDataActorFindTest_Root_Tag_Children");
 	//	Settings->bIncludeChildren = false;
 	//}
 
@@ -284,7 +302,7 @@ bool FPCGPropertyToParamDataActorFindTest::RunTest(const FString& Parameters)
 		UPCGUnitTestDummyComponent* Component = Cast< UPCGUnitTestDummyComponent>(TestData.TestActor->AddComponentByClass(Settings->ComponentClass, false, FTransform::Identity, false));
 		Component->IntProperty = 42;
 
-		bSuccess &= VerifyAttributeValue(this, TestData, GET_MEMBER_NAME_CHECKED(UPCGUnitTestDummyComponent, IntProperty), 42ll, false, "PropertyToParamDataActorFindTest_Self_Tag_Component");
+		bSuccess &= VerifyAttributeValueValid(this, TestData, GET_MEMBER_NAME_CHECKED(UPCGUnitTestDummyComponent, IntProperty), 42ll, "PropertyToParamDataActorFindTest_Self_Tag_Component");
 		Settings->bSelectComponent = false;
 	}
 
@@ -296,7 +314,7 @@ bool FPCGPropertyToParamDataActorFindTest::RunTest(const FString& Parameters)
 		PCGTestsCommon::FTestData TestData(Seed, Settings, APCGUnitTestDummyActor::StaticClass());
 		Cast<APCGUnitTestDummyActor>(TestData.TestActor)->IntProperty = 42;
 
-		bSuccess &= VerifyAttributeValue(this, TestData, GET_MEMBER_NAME_CHECKED(APCGUnitTestDummyActor, IntProperty), 42ll, false, "PropertyToParamDataActorFindTest_World_Class");
+		bSuccess &= VerifyAttributeValueValid(this, TestData, GET_MEMBER_NAME_CHECKED(APCGUnitTestDummyActor, IntProperty), 42ll, "PropertyToParamDataActorFindTest_World_Class");
 	}
 
 	// World by tag
@@ -308,7 +326,7 @@ bool FPCGPropertyToParamDataActorFindTest::RunTest(const FString& Parameters)
 		TestData.TestActor->Tags.Add(Tag);
 		Cast<APCGUnitTestDummyActor>(TestData.TestActor)->IntProperty = 42;
 
-		bSuccess &= VerifyAttributeValue(this, TestData, GET_MEMBER_NAME_CHECKED(APCGUnitTestDummyActor, IntProperty), 42ll, false, "PropertyToParamDataActorFindTest_World_Tag");
+		bSuccess &= VerifyAttributeValueValid(this, TestData, GET_MEMBER_NAME_CHECKED(APCGUnitTestDummyActor, IntProperty), 42ll, "PropertyToParamDataActorFindTest_World_Tag");
 	}
 
 	// World by name
@@ -319,7 +337,7 @@ bool FPCGPropertyToParamDataActorFindTest::RunTest(const FString& Parameters)
 		PCGTestsCommon::FTestData TestData(Seed, Settings, APCGUnitTestDummyActor::StaticClass());
 		Cast<APCGUnitTestDummyActor>(TestData.TestActor)->IntProperty = 42;
 
-		bSuccess &= VerifyAttributeValue(this, TestData, GET_MEMBER_NAME_CHECKED(APCGUnitTestDummyActor, IntProperty), 42ll, false, "PropertyToParamDataActorFindTest_World_Name");
+		bSuccess &= VerifyAttributeValueValid(this, TestData, GET_MEMBER_NAME_CHECKED(APCGUnitTestDummyActor, IntProperty), 42ll, "PropertyToParamDataActorFindTest_World_Name");
 	}
 
 	return bSuccess;
diff --git a/Engine/Plugins/Experimental/PCG/Source/PCG/Private/Tests/Elements/PCGPropertyToParamDataElementTest.h b/Engine/Plugins/Experimental/PCG/Source/PCG/Private/Tests/Elements/PCGPropertyToParamDataElementTest.h
--- a/Engine/Plugins/Experimental/PCG/Source/PCG/Private/Tests/Elements/PCGPropertyToParamDataElementTest.h
+++ b/Engine/Plugins/Experimental/PCG/Source/PCG/Private/Tests/Elements/PCGPropertyToParamDataElementTest.h
@@ -75,6 +75,9 @@ public:
 
 	UPROPERTY()
 	FVector2D Vector2Property;
+
+	UPROPERTY()
+	FColor ColorProperty;
 };
 
 UCLASS(MinimalAPI, NotBlueprintable, NotPlaceable, NotBlueprintType, Transient, HideDropdown)
