Fixed issue with nanite meshes not using ISM as expected

diff --git a/Engine/Plugins/Experimental/PCG/Source/PCG/Private/Helpers/PCGActorHelpers.cpp b/Engine/Plugins/Experimental/PCG/Source/PCG/Private/Helpers/PCGActorHelpers.cpp
--- a/Engine/Plugins/Experimental/PCG/Source/PCG/Private/Helpers/PCGActorHelpers.cpp
+++ b/Engine/Plugins/Experimental/PCG/Source/PCG/Private/Helpers/PCGActorHelpers.cpp
@@ -102,7 +102,14 @@ UInstancedStaticMeshComponent* UPCGActorHelpers::GetOrCreateISMC(AActor* InTarge
 	// TODO: add hism/ism switch or better yet, use a template component
 	UInstancedStaticMeshComponent* ISMC = nullptr;
 
-	if (InMesh->HasValidNaniteData())
+	// Done as in InstancedStaticMesh.cpp
+#if WITH_EDITOR
+	const bool bMeshHasNaniteData = InMesh->NaniteSettings.bEnabled;
+#else
+	const bool bMeshHasNaniteData = InMesh->GetRenderData()->NaniteResources.PageStreamingStates.Num() > 0;
+#endif
+
+	if (bMeshHasNaniteData)
 	{
 		ISMC = NewObject<UInstancedStaticMeshComponent>(InTargetActor);
 	}
