PCG:
*Added point metadata support for attribute list view.

diff --git a/Engine/Plugins/Experimental/PCG/Source/PCG/Public/Metadata/PCGMetadata.h b/Engine/Plugins/Experimental/PCG/Source/PCG/Public/Metadata/PCGMetadata.h
--- a/Engine/Plugins/Experimental/PCG/Source/PCG/Public/Metadata/PCGMetadata.h
+++ b/Engine/Plugins/Experimental/PCG/Source/PCG/Public/Metadata/PCGMetadata.h
@@ -3,8 +3,9 @@
 #pragma once
 
 #include "CoreMinimal.h"
-#include "PCGMetadataCommon.h"
 #include "PCGMetadataAccessor.h"
+#include "PCGMetadataAttributeTraits.h"
+#include "PCGMetadataCommon.h"
 
 #include "PCGMetadata.generated.h"
 
@@ -203,4 +204,4 @@ protected:
 
 	mutable FRWLock AttributeLock;
 	mutable FRWLock ItemLock;
-};
\ No newline at end of file
+};
diff --git a/Engine/Plugins/Experimental/PCG/Source/PCGEditor/Private/SPCGEditorGraphAttributeListView.cpp b/Engine/Plugins/Experimental/PCG/Source/PCGEditor/Private/SPCGEditorGraphAttributeListView.cpp
--- a/Engine/Plugins/Experimental/PCG/Source/PCGEditor/Private/SPCGEditorGraphAttributeListView.cpp
+++ b/Engine/Plugins/Experimental/PCG/Source/PCGEditor/Private/SPCGEditorGraphAttributeListView.cpp
@@ -4,6 +4,8 @@
 
 #include "Data/PCGPointData.h"
 #include "Data/PCGSpatialData.h"
+#include "Metadata/PCGMetadataAttribute.h"
+#include "Metadata/PCGMetadataAttributeTpl.h"
 #include "PCGComponent.h"
 #include "PCGData.h"
 #include "PCGEditor.h"
@@ -193,6 +195,136 @@ TSharedRef<SWidget> SPCGListViewItemRow::GenerateWidgetForColumn(const FName& Co
 			const int32 Seed = PCGPoint->Seed;
 			ColumnData = FText::AsNumber(Seed);
 		}
+
+		if (const UPCGMetadata* PCGMetadata = InternalItem->PCGMetadata)
+		{
+			if (const FPCGMetadataInfo* MetadataInfo = InternalItem->MetadataInfos->Find(ColumnId))
+			{
+				if (const FPCGMetadataAttributeBase* AttributeBase = PCGMetadata->GetConstAttribute((*MetadataInfo).MetadataId)) // todo investigate deadlock caused by read lock
+				{
+					if (AttributeBase->GetTypeId() == PCG::Private::MetadataTypes<float>::Id)
+					{
+						const float MetaFloat = static_cast<const FPCGMetadataAttribute<float>*>(AttributeBase)->GetValueFromItemKey(PCGPoint->MetadataEntry);
+						ColumnData = FText::AsNumber(MetaFloat);
+					}
+					else if (AttributeBase->GetTypeId() == PCG::Private::MetadataTypes<double>::Id)
+					{
+						const double MetaDouble = static_cast<const FPCGMetadataAttribute<double>*>(AttributeBase)->GetValueFromItemKey(PCGPoint->MetadataEntry);
+						ColumnData = FText::AsNumber(MetaDouble);
+					}
+					else if (AttributeBase->GetTypeId() == PCG::Private::MetadataTypes<bool>::Id)
+					{
+						const bool bMetaBool = static_cast<const FPCGMetadataAttribute<bool>*>(AttributeBase)->GetValueFromItemKey(PCGPoint->MetadataEntry);
+						ColumnData = bMetaBool ? FText::FromString(TEXT("true")) : FText::FromString(TEXT("false"));
+					}
+					else if (AttributeBase->GetTypeId() == PCG::Private::MetadataTypes<FVector>::Id)
+					{
+						const FVector MetaVector = static_cast<const FPCGMetadataAttribute<FVector>*>(AttributeBase)->GetValueFromItemKey(PCGPoint->MetadataEntry);
+						ColumnData = FText::AsNumber(MetaVector[(*MetadataInfo).Index]);
+					}
+					else if (AttributeBase->GetTypeId() == PCG::Private::MetadataTypes<FVector4>::Id)
+					{
+						const FVector4 MetaVector4 = static_cast<const FPCGMetadataAttribute<FVector4>*>(AttributeBase)->GetValueFromItemKey(PCGPoint->MetadataEntry);
+						ColumnData = FText::AsNumber(MetaVector4[(*MetadataInfo).Index]);
+					}
+					else if (AttributeBase->GetTypeId() == PCG::Private::MetadataTypes<int32>::Id)
+					{
+						const int32 MetaInt32 = static_cast<const FPCGMetadataAttribute<int32>*>(AttributeBase)->GetValueFromItemKey(PCGPoint->MetadataEntry);
+						ColumnData = FText::AsNumber(MetaInt32);
+					}
+					else if (AttributeBase->GetTypeId() == PCG::Private::MetadataTypes<int64>::Id)
+					{
+						const int64 MetaInt64 = static_cast<const FPCGMetadataAttribute<int64>*>(AttributeBase)->GetValueFromItemKey(PCGPoint->MetadataEntry);
+						ColumnData = FText::AsNumber(MetaInt64);
+					}
+					else if (AttributeBase->GetTypeId() == PCG::Private::MetadataTypes<FString>::Id)
+					{
+						const FString MetaString = static_cast<const FPCGMetadataAttribute<FString>*>(AttributeBase)->GetValueFromItemKey(PCGPoint->MetadataEntry);
+						ColumnData = FText::FromString(MetaString);
+					}
+					else if (AttributeBase->GetTypeId() == PCG::Private::MetadataTypes<FName>::Id)
+					{
+						const FName MetaName = static_cast<const FPCGMetadataAttribute<FName>*>(AttributeBase)->GetValueFromItemKey(PCGPoint->MetadataEntry);
+						ColumnData = FText::FromName(MetaName);
+					}
+					else if (AttributeBase->GetTypeId() == PCG::Private::MetadataTypes<FQuat>::Id)
+					{
+						const FQuat MetaQuat = static_cast<const FPCGMetadataAttribute<FQuat>*>(AttributeBase)->GetValueFromItemKey(PCGPoint->MetadataEntry);
+						float QuatValue = 0.0f;
+						if ((*MetadataInfo).Index == 0)
+						{
+							QuatValue = MetaQuat.X;
+						}
+						else if ((*MetadataInfo).Index == 1)
+						{
+							QuatValue = MetaQuat.Y;
+						}
+						else if ((*MetadataInfo).Index == 2)
+						{
+							QuatValue = MetaQuat.Z;
+						}
+						else if ((*MetadataInfo).Index == 3)
+						{
+							QuatValue = MetaQuat.W;
+						}
+						
+						ColumnData = FText::AsNumber(QuatValue);
+					}
+					else if (AttributeBase->GetTypeId() == PCG::Private::MetadataTypes<FRotator>::Id)
+					{
+						const FRotator MetaRotator = static_cast<const FPCGMetadataAttribute<FRotator>*>(AttributeBase)->GetValueFromItemKey(PCGPoint->MetadataEntry);
+						float RotatorValue = 0.0f;
+						if ((*MetadataInfo).Index == 0)
+						{
+							RotatorValue = MetaRotator.Roll;
+						}
+						else if ((*MetadataInfo).Index == 1)
+						{
+							RotatorValue = MetaRotator.Pitch;
+						}
+						else if ((*MetadataInfo).Index == 2)
+						{
+							RotatorValue = MetaRotator.Yaw;
+						}
+						ColumnData = FText::AsNumber(RotatorValue);
+					}
+					else if (AttributeBase->GetTypeId() == PCG::Private::MetadataTypes<FTransform>::Id)
+					{
+						const FTransform& MetaTransform = static_cast<const FPCGMetadataAttribute<FTransform>*>(AttributeBase)->GetValueFromItemKey(PCGPoint->MetadataEntry);
+						const int8 ComponentIndex = (*MetadataInfo).Index / 3;
+						const int8 ValueIndex = (*MetadataInfo).Index % 3;
+						if (ComponentIndex == 0)
+						{
+							const FVector& MetaVector = MetaTransform.GetLocation();
+							ColumnData = FText::AsNumber(MetaVector[ValueIndex]);
+						}
+						else if(ComponentIndex == 1)
+						{
+							const FRotator& MetaRotator = MetaTransform.Rotator();
+							float RotatorValue = 0.0f;
+							if (ValueIndex == 0)
+							{
+								RotatorValue = MetaRotator.Roll;
+							}
+							else if (ValueIndex == 1)
+							{
+								RotatorValue = MetaRotator.Pitch;
+							}
+							else if (ValueIndex == 2)
+							{
+								RotatorValue = MetaRotator.Yaw;
+							}
+							ColumnData = FText::AsNumber(RotatorValue);
+						}
+						else if(ComponentIndex == 2)
+						{
+							const FVector& MetaVector = MetaTransform.GetScale3D();
+							ColumnData = FText::AsNumber(MetaVector[ValueIndex]);
+						}
+					}
+				}
+			}
+		}
 	}
 
 	return SNew(STextBlock).Text(ColumnData);
@@ -443,7 +575,10 @@ void SPCGEditorGraphAttributeListView::OnGenerateUpdated(UPCGComponent* /*InPCGC
 void SPCGEditorGraphAttributeListView::RefreshAttributeList()
 {
 	ListViewItems.Empty();
-	ListView->RequestListRefresh();
+	ListView->SetListItemsSource(EmptyList); // TODO: Investigate if it's possible to avoid this hack that's used to force refresh the list view. Without it the AddColumn will try to create widgets for old data.
+
+	RemoveMetadataColumns();
+	MetadataInfos.Empty();
 
 	const TSharedPtr<FPCGEditor> PCGEditor = PCGEditorPtr.Pin();
 	if (!PCGEditor.IsValid())
@@ -476,10 +611,69 @@ void SPCGEditorGraphAttributeListView::RefreshAttributeList()
 	
 	const FPCGTaggedData& TaggedData = InspectionData->TaggedData[DataIndex];
 	const UPCGData* PCGData = TaggedData.Data;
+
 	if (const UPCGSpatialData* PCGSpatialData = Cast<UPCGSpatialData>(PCGData))
 	{
 		if (const UPCGPointData* PCGPointData = PCGSpatialData->ToPointData())
 		{
+			const UPCGMetadata* PCGMetadata = PCGPointData->ConstMetadata(); // TODO: add support for paramdata as well.
+		
+			TArray<FName> AttributeNames;
+			TArray<EPCGMetadataTypes> AttributeTypes;
+			PCGMetadata->GetAttributes(AttributeNames, AttributeTypes);
+
+			for (int32 i = 0; i < AttributeNames.Num(); i++)
+			{
+				const FName& AttributeName = AttributeNames[i];
+				const EPCGMetadataTypes& AttributeType = AttributeTypes[i];
+
+				switch (AttributeType)
+				{
+				case EPCGMetadataTypes::Float:
+				case EPCGMetadataTypes::Double:
+				case EPCGMetadataTypes::Integer32:
+				case EPCGMetadataTypes::Integer64:
+				case EPCGMetadataTypes::Boolean:
+				case EPCGMetadataTypes::String:
+				case EPCGMetadataTypes::Name:
+					{
+						AddMetadataColumn(AttributeName);
+						break;
+					}
+				case EPCGMetadataTypes::Vector:
+					{
+						AddMetadataColumn(AttributeName, 0, TEXT("_X"));
+						AddMetadataColumn(AttributeName, 1, TEXT("_Y"));
+						AddMetadataColumn(AttributeName, 2, TEXT("_Z"));
+						break;
+					}
+				case EPCGMetadataTypes::Vector4:
+				case EPCGMetadataTypes::Quaternion:
+					{
+						AddMetadataColumn(AttributeName, 0, TEXT("_X"));
+						AddMetadataColumn(AttributeName, 1, TEXT("_Y"));
+						AddMetadataColumn(AttributeName, 2, TEXT("_Z"));
+						AddMetadataColumn(AttributeName, 3, TEXT("_W"));
+						break;
+					}
+				case EPCGMetadataTypes::Transform:
+					{
+						AddMetadataColumn(AttributeName, 0, TEXT("_tX"));
+						AddMetadataColumn(AttributeName, 1, TEXT("_tY"));
+						AddMetadataColumn(AttributeName, 2, TEXT("_tZ"));
+						AddMetadataColumn(AttributeName, 3, TEXT("_rX"));
+						AddMetadataColumn(AttributeName, 4, TEXT("_rY"));
+						AddMetadataColumn(AttributeName, 5, TEXT("_rZ"));
+						AddMetadataColumn(AttributeName, 6, TEXT("_sX"));
+						AddMetadataColumn(AttributeName, 7, TEXT("_sY"));
+						AddMetadataColumn(AttributeName, 8, TEXT("_sZ"));
+						break;
+					}
+				default:
+					break;
+				}
+			}
+
 			const TArray<FPCGPoint>& PCGPoints = PCGPointData->GetPoints();
 
 			ListViewItems.Reserve(PCGPoints.Num());
@@ -490,10 +684,14 @@ void SPCGEditorGraphAttributeListView::RefreshAttributeList()
 				PCGListviewItemPtr ListViewItem = MakeShared<FPCGListViewItem>();
 				ListViewItem->Index = PointIndex;
 				ListViewItem->PCGPoint = &PCGPoint;
+				ListViewItem->PCGMetadata = PCGMetadata;
+				ListViewItem->MetadataInfos = &MetadataInfos;
 				ListViewItems.Add(ListViewItem);
 			}
 		}
 	}
+
+	ListView->SetListItemsSource(ListViewItems);
 }
 
 void SPCGEditorGraphAttributeListView::RefreshDataComboBox()
@@ -543,7 +741,10 @@ TSharedRef<SWidget> SPCGEditorGraphAttributeListView::OnGenerateDataWidget(TShar
 
 void SPCGEditorGraphAttributeListView::OnSelectionChanged(TSharedPtr<FName> Item, ESelectInfo::Type SelectInfo)
 {
-	RefreshAttributeList();
+	if (SelectInfo != ESelectInfo::Direct)
+	{
+		RefreshAttributeList();
+	}
 }
 
 FText SPCGEditorGraphAttributeListView::OnGenerateSelectedDataText() const
@@ -574,4 +775,38 @@ TSharedRef<ITableRow> SPCGEditorGraphAttributeListView::OnGenerateRow(PCGListvie
 	return SNew(SPCGListViewItemRow, OwnerTable, Item);
 }
 
+void SPCGEditorGraphAttributeListView::AddMetadataColumn(const FName& InColumnId, const int8 InValueIndex,  const TCHAR* PostFix)
+{
+	FString ColumnIdString = InColumnId.ToString();
+
+	if (PostFix)
+	{
+		ColumnIdString.Append(PostFix);
+	}
+
+	const FName ColumnId(ColumnIdString);
+	
+	FPCGMetadataInfo MetadataInfo;
+	MetadataInfo.MetadataId = InColumnId;
+	MetadataInfo.Index = InValueIndex;
+	MetadataInfos.Add(ColumnId, MetadataInfo);
+
+	SHeaderRow::FColumn::FArguments ColumnArguments;
+	ColumnArguments.ColumnId(ColumnId);
+	ColumnArguments.DefaultLabel(FText::FromName(ColumnId));
+	ColumnArguments.HAlignHeader(EHorizontalAlignment::HAlign_Center);
+	ColumnArguments.HAlignCell(EHorizontalAlignment::HAlign_Right);
+	ColumnArguments.FillWidth(1.0f);
+	ListViewHeader->AddColumn(ColumnArguments);
+}
+
+void SPCGEditorGraphAttributeListView::RemoveMetadataColumns()
+{
+	for (auto MetadataInfo : MetadataInfos)
+	{
+		const FName& MetadataKey = MetadataInfo.Key;
+		ListViewHeader->RemoveColumn(MetadataKey);
+	}
+}
+
 #undef LOCTEXT_NAMESPACE
diff --git a/Engine/Plugins/Experimental/PCG/Source/PCGEditor/Private/SPCGEditorGraphAttributeListView.h b/Engine/Plugins/Experimental/PCG/Source/PCGEditor/Private/SPCGEditorGraphAttributeListView.h
--- a/Engine/Plugins/Experimental/PCG/Source/PCGEditor/Private/SPCGEditorGraphAttributeListView.h
+++ b/Engine/Plugins/Experimental/PCG/Source/PCGEditor/Private/SPCGEditorGraphAttributeListView.h
@@ -3,6 +3,7 @@
 #pragma once
 
 #include "CoreMinimal.h"
+#include "Metadata/PCGMetadata.h"
 #include "Widgets/Input/SComboBox.h"
 #include "Widgets/SCompoundWidget.h"
 #include "Widgets/Views/SHeaderRow.h"
@@ -13,10 +14,18 @@ class UPCGComponent;
 class UPCGNode;
 struct FPCGPoint;
 
+struct FPCGMetadataInfo
+{
+	FName MetadataId = NAME_None;
+	int8 Index = 0;
+};
+
 struct FPCGListViewItem
 {
 	int32 Index = INDEX_NONE;
 	const FPCGPoint* PCGPoint = nullptr;
+	const UPCGMetadata* PCGMetadata = nullptr;
+	const TMap<FName, FPCGMetadataInfo>* MetadataInfos = nullptr;
 };
 
 typedef TSharedPtr<FPCGListViewItem> PCGListviewItemPtr;
@@ -63,6 +72,9 @@ private:
 
 	TSharedRef<ITableRow> OnGenerateRow(PCGListviewItemPtr Item, const TSharedRef<STableViewBase>& OwnerTable) const;
 
+	void AddMetadataColumn(const FName& InColumnId, const int8 InValueIndex = INDEX_NONE, const TCHAR* PostFix = nullptr);
+	void RemoveMetadataColumns();
+
 	/** Pointer back to the PCG editor that owns us */
 	TWeakPtr<FPCGEditor> PCGEditorPtr;
 
@@ -72,7 +84,11 @@ private:
 	TSharedPtr<SHeaderRow> ListViewHeader;
 	TSharedPtr<SListView<PCGListviewItemPtr>> ListView;
 	TArray<PCGListviewItemPtr> ListViewItems;
+	/** Empty list to force refresh the ListView when regenerating */
+	TArray<PCGListviewItemPtr> EmptyList;
 
 	TSharedPtr<SComboBox<TSharedPtr<FName>>> DataComboBox;
 	TArray<TSharedPtr<FName>> DataComboBoxItems;
+
+	TMap<FName, FPCGMetadataInfo> MetadataInfos;
 };
