Added support for grayscale texture data

#rnx

diff --git a/Engine/Plugins/Experimental/PCG/Source/PCG/Private/Data/PCGTextureData.cpp b/Engine/Plugins/Experimental/PCG/Source/PCG/Private/Data/PCGTextureData.cpp
--- a/Engine/Plugins/Experimental/PCG/Source/PCG/Private/Data/PCGTextureData.cpp
+++ b/Engine/Plugins/Experimental/PCG/Source/PCG/Private/Data/PCGTextureData.cpp
@@ -187,8 +187,7 @@ const UPCGPointData* UPCGBaseTextureData::CreatePointData(FPCGContext* Context)
 				PCGHelpers::ComputeSeed(X, Y));
 
 			const FVector TransformScale = Transform.GetScale3D();
-			// Note: divided by 4 here because the scale is doubled before, and the extents represent half a pixel
-			OutPoint.SetExtents(FVector(TransformScale.X * XScale / 4.0, TransformScale.Y * YScale / 4.0, 1.0));
+			OutPoint.SetExtents(FVector(TransformScale.X * XScale / 2.0, TransformScale.Y * YScale / 2.0, 1.0));
 			OutPoint.Color = ColorData[X + Y * Width];
 
 			return true;
@@ -220,21 +219,41 @@ void UPCGTextureData::Initialize(UTexture2D* InTexture, const FTransform& InTran
 		if (Texture->GetPlatformData()->Mips.Num() > 0)
 		{
 			TRACE_CPUPROFILER_EVENT_SCOPE(UPCGTextureData::Initialize::ReadData);
-			Width = Texture->GetSizeX();
-			Height = Texture->GetSizeY();
 
-			// TODO: previously this code created a duplicate of InTexture and read the ColorData from that source 
-			// however, there was a problem with the duplicate not having texture data in Mips[0], so we've removed the duplication for now
-			// original code can be found in ProceduralTextureReader.h
-			const FColor* FormattedImageData = static_cast<const FColor*>(InTexture->GetPlatformData()->Mips[0].BulkData.LockReadOnly());
-
-			ColorData.SetNum(Width * Height);
-			for (int32 D = 0; D < Width * Height; ++D)
+			if (Texture->GetPlatformData()->PixelFormat == PF_R8G8B8A8 || Texture->GetPlatformData()->PixelFormat == PF_G8)
 			{
-				ColorData[D] = FormattedImageData[D].ReinterpretAsLinear();
+				Width = Texture->GetSizeX();
+				Height = Texture->GetSizeY();
+				const int32 PixelCount = Width * Height;
+				ColorData.SetNum(PixelCount);
+
+				// TODO: previously this code created a duplicate of InTexture and read the ColorData from that source 
+				// however, there was a problem with the duplicate not having texture data in Mips[0], so we've removed the duplication for now
+				// original code can be found in ProceduralTextureReader.h
+				const uint8_t* BulkData = reinterpret_cast<const uint8_t*>(InTexture->GetPlatformData()->Mips[0].BulkData.LockReadOnly());
+
+				if (Texture->GetPlatformData()->PixelFormat == PF_R8G8B8A8)
+				{
+					const FColor* FormattedImageData = reinterpret_cast<const FColor*>(BulkData);
+					for (int32 D = 0; D < PixelCount; ++D)
+					{
+						ColorData[D] = FormattedImageData[D].ReinterpretAsLinear();
+					}
+				}
+				else if (Texture->GetPlatformData()->PixelFormat == PF_G8)
+				{
+					for (int32 D = 0; D < PixelCount; ++D)
+					{
+						ColorData[D] = FColor(BulkData[D], BulkData[D], BulkData[D]).ReinterpretAsLinear();
+					}
+				}
+
+				InTexture->GetPlatformData()->Mips[0].BulkData.Unlock();
+			}
+			else
+			{
+				UE_LOG(LogPCG, Error, TEXT("PCGTextureData does not support the format of %s"), *Texture->GetFName().ToString());
 			}
-
-			InTexture->GetPlatformData()->Mips[0].BulkData.Unlock();
 		}
 #endif
 	}
