PCG:
*Added Commands to change node execution mode for selection. The commands are exposed in the node rightclick menu.
*Added visuals to execution mode and changed the node inspection visual to a orange icon.

diff --git a/Engine/Plugins/Experimental/PCG/Source/PCGEditor/Private/PCGEditor.cpp b/Engine/Plugins/Experimental/PCG/Source/PCGEditor/Private/PCGEditor.cpp
--- a/Engine/Plugins/Experimental/PCG/Source/PCGEditor/Private/PCGEditor.cpp
+++ b/Engine/Plugins/Experimental/PCG/Source/PCGEditor/Private/PCGEditor.cpp
@@ -319,6 +319,30 @@ void FPCGEditor::BindCommands()
 		PCGEditorCommands.RunDeterminismTest,
 		FExecuteAction::CreateSP(this, &FPCGEditor::OnDeterminismTests),
 		FCanExecuteAction::CreateSP(this, &FPCGEditor::CanRunDeterminismTests));
+
+	GraphEditorCommands->MapAction(
+		PCGEditorCommands.ExecutionModeEnabled,
+		FExecuteAction::CreateSP(this, &FPCGEditor::OnSetExecutionMode, EPCGSettingsExecutionMode::Enabled),
+		FCanExecuteAction(),
+		FIsActionChecked::CreateSP(this, &FPCGEditor::IsExecutionModeActive, EPCGSettingsExecutionMode::Enabled));
+	
+	GraphEditorCommands->MapAction(
+		PCGEditorCommands.ExecutionModeDebug,
+		FExecuteAction::CreateSP(this, &FPCGEditor::OnSetExecutionMode, EPCGSettingsExecutionMode::Debug),
+		FCanExecuteAction(),
+		FIsActionChecked::CreateSP(this, &FPCGEditor::IsExecutionModeActive, EPCGSettingsExecutionMode::Debug));
+	
+	GraphEditorCommands->MapAction(
+		PCGEditorCommands.ExecutionModeDisabled,
+		FExecuteAction::CreateSP(this, &FPCGEditor::OnSetExecutionMode, EPCGSettingsExecutionMode::Disabled),
+		FCanExecuteAction(),
+		FIsActionChecked::CreateSP(this, &FPCGEditor::IsExecutionModeActive, EPCGSettingsExecutionMode::Disabled));
+	
+	GraphEditorCommands->MapAction(
+		PCGEditorCommands.ExecutionModeIsolated,
+		FExecuteAction::CreateSP(this, &FPCGEditor::OnSetExecutionMode, EPCGSettingsExecutionMode::Isolated),
+		FCanExecuteAction(),
+		FIsActionChecked::CreateSP(this, &FPCGEditor::IsExecutionModeActive, EPCGSettingsExecutionMode::Isolated));
 }
 
 void FPCGEditor::OnFind()
@@ -816,6 +840,77 @@ void FPCGEditor::OnStopInspectNode()
 	}
 }
 
+bool FPCGEditor::IsExecutionModeActive(EPCGSettingsExecutionMode InExecutionMode) const
+{
+	if (GraphEditorWidget.IsValid())
+	{
+		for (const UObject* Object : GraphEditorWidget->GetSelectedNodes())
+		{
+			const UPCGEditorGraphNodeBase* PCGEditorGraphNode = Cast<const UPCGEditorGraphNodeBase>(Object);
+			if (!PCGEditorGraphNode)
+			{
+				continue;
+			}
+			
+			const UPCGNode* PCGNode = PCGEditorGraphNode->GetPCGNode();
+			if (!PCGNode)
+			{
+				continue;
+			}
+
+			const UPCGSettings* PCGSettings = PCGNode->DefaultSettings; 
+			if (!PCGSettings)
+			{
+				continue;				
+			}
+			
+			if (PCGSettings->ExecutionMode == InExecutionMode)
+			{
+				return true;
+			}
+		}
+	}
+	
+	return false;
+}
+
+void FPCGEditor::OnSetExecutionMode(EPCGSettingsExecutionMode InExecutionMode)
+{
+	// todo should perhaps make a utility function to get all selected nodes.
+	if (GraphEditorWidget.IsValid())
+	{
+		const FScopedTransaction Transaction(*FPCGEditorCommon::ContextIdentifier, LOCTEXT("PCGEditorExecutionModeTransactionMessage", "PCG Editor: ExecutionMode"), nullptr);
+		
+		for (UObject* Object : GraphEditorWidget->GetSelectedNodes())
+		{
+			UPCGEditorGraphNodeBase* PCGEditorGraphNode = Cast<UPCGEditorGraphNodeBase>(Object);
+			if (!PCGEditorGraphNode)
+			{
+				continue;
+			}
+
+			UPCGNode* PCGNode = PCGEditorGraphNode->GetPCGNode();
+			if (!PCGNode)
+			{
+				continue;
+			}
+
+			UPCGSettings* PCGSettings = PCGNode->DefaultSettings;
+			if (!PCGSettings)
+			{
+				continue;
+			}
+
+			if (PCGSettings->ExecutionMode != InExecutionMode)
+			{
+				PCGSettings->Modify();
+				PCGSettings->ExecutionMode = InExecutionMode;
+				PCGSettings->OnSettingsChangedDelegate.Broadcast(PCGSettings, EPCGChangeType::Settings);
+			}
+		}
+	}
+}
+
 void FPCGEditor::SelectAllNodes()
 {
 	if (GraphEditorWidget.IsValid())
diff --git a/Engine/Plugins/Experimental/PCG/Source/PCGEditor/Private/PCGEditor.h b/Engine/Plugins/Experimental/PCG/Source/PCGEditor/Private/PCGEditor.h
--- a/Engine/Plugins/Experimental/PCG/Source/PCGEditor/Private/PCGEditor.h
+++ b/Engine/Plugins/Experimental/PCG/Source/PCGEditor/Private/PCGEditor.h
@@ -3,8 +3,10 @@
 #pragma once
 
 #include "PCGEditorModule.h"
-#include "Toolkits/AssetEditorToolkit.h"
+#include "PCGSettings.h"
+
 #include "EditorUndoClient.h"
+#include "Toolkits/AssetEditorToolkit.h"
 
 class FUICommandList;
 class IDetailsView;
@@ -95,6 +97,12 @@ private:
 	/** Summon the Determinism tab */
 	void OnDeterminismTests();
 
+	/** Whether or not an execution mode is active for the selected ndoes */
+	bool IsExecutionModeActive(EPCGSettingsExecutionMode InExecutionMode) const;
+		
+	/** Set execution mode for selected nodes */
+	void OnSetExecutionMode(EPCGSettingsExecutionMode InExecutionMode);
+	
 	/** Select every node in the graph */
 	void SelectAllNodes();
 	/** Whether we can select every node */
diff --git a/Engine/Plugins/Experimental/PCG/Source/PCGEditor/Private/PCGEditorCommands.cpp b/Engine/Plugins/Experimental/PCG/Source/PCGEditor/Private/PCGEditorCommands.cpp
--- a/Engine/Plugins/Experimental/PCG/Source/PCGEditor/Private/PCGEditorCommands.cpp
+++ b/Engine/Plugins/Experimental/PCG/Source/PCGEditor/Private/PCGEditorCommands.cpp
@@ -23,6 +23,10 @@ void FPCGEditorCommands::RegisterCommands()
 	UI_COMMAND(RunDeterminismTest, "Run Determinism Test", "Evaluate the current node for determinism.", EUserInterfaceActionType::Button, FInputChord(EModifierKey::Alt, EKeys::D));
 	UI_COMMAND(StartInspectNode, "Start Inspecting Node", "Start Inspecting Node", EUserInterfaceActionType::Button, FInputChord(EKeys::I));
 	UI_COMMAND(StopInspectNode, "Stop Inspecting Node", "Stop Inspecting Node", EUserInterfaceActionType::Button, FInputChord(EModifierKey::Shift, EKeys::I));
+	UI_COMMAND(ExecutionModeEnabled, "Execution Mode Enabled", "Set the execution mode for this node to Enabled.", EUserInterfaceActionType::RadioButton, FInputChord());
+	UI_COMMAND(ExecutionModeDebug, "Execution Mode Debug", "Set the execution mode for this node to Debug.", EUserInterfaceActionType::RadioButton, FInputChord());
+	UI_COMMAND(ExecutionModeIsolated, "Execution Mode Isolated", "Set the execution mode for this node to Isolated.", EUserInterfaceActionType::RadioButton, FInputChord());
+	UI_COMMAND(ExecutionModeDisabled, "Execution Mode Disabled", "Set the execution mode for this node to Disabled.", EUserInterfaceActionType::RadioButton, FInputChord());
 }
 
 #undef LOCTEXT_NAMESPACE
diff --git a/Engine/Plugins/Experimental/PCG/Source/PCGEditor/Private/PCGEditorCommands.h b/Engine/Plugins/Experimental/PCG/Source/PCGEditor/Private/PCGEditorCommands.h
--- a/Engine/Plugins/Experimental/PCG/Source/PCGEditor/Private/PCGEditorCommands.h
+++ b/Engine/Plugins/Experimental/PCG/Source/PCGEditor/Private/PCGEditorCommands.h
@@ -19,4 +19,8 @@ public:
 	TSharedPtr<FUICommandInfo> RunDeterminismTest;
 	TSharedPtr<FUICommandInfo> StartInspectNode;
 	TSharedPtr<FUICommandInfo> StopInspectNode;
+	TSharedPtr<FUICommandInfo> ExecutionModeEnabled;
+	TSharedPtr<FUICommandInfo> ExecutionModeDebug;
+	TSharedPtr<FUICommandInfo> ExecutionModeIsolated;
+	TSharedPtr<FUICommandInfo> ExecutionModeDisabled;
 };
diff --git a/Engine/Plugins/Experimental/PCG/Source/PCGEditor/Private/PCGEditorGraphNodeBase.cpp b/Engine/Plugins/Experimental/PCG/Source/PCGEditor/Private/PCGEditorGraphNodeBase.cpp
--- a/Engine/Plugins/Experimental/PCG/Source/PCGEditor/Private/PCGEditorGraphNodeBase.cpp
+++ b/Engine/Plugins/Experimental/PCG/Source/PCGEditor/Private/PCGEditorGraphNodeBase.cpp
@@ -75,6 +75,15 @@ void UPCGEditorGraphNodeBase::GetNodeContextMenuActions(UToolMenu* Menu, class U
 		Section.AddMenuEntry(bIsInspected ? FPCGEditorCommands::Get().StopInspectNode : FPCGEditorCommands::Get().StartInspectNode);
 		Section.AddMenuEntry(FGraphEditorCommands::Get().BreakNodeLinks);
 		Section.AddMenuEntry(FPCGEditorCommands::Get().CollapseNodes);
+
+		Section.AddSubMenu("ExecutionMode", LOCTEXT("ExecutionModeHeader", "Execution Mode"), FText(), FNewToolMenuDelegate::CreateLambda([](UToolMenu* ExecutionModeMenu)
+		{
+			FToolMenuSection& SubSection = ExecutionModeMenu->AddSection("ExecutionMode");
+			SubSection.AddMenuEntry(FPCGEditorCommands::Get().ExecutionModeEnabled, LOCTEXT("ExecutionModeEnabledLabel", "Enabled"));
+			SubSection.AddMenuEntry(FPCGEditorCommands::Get().ExecutionModeDebug, LOCTEXT("ExecutionModeDebugLabel", "Debug"));
+			SubSection.AddMenuEntry(FPCGEditorCommands::Get().ExecutionModeDisabled, LOCTEXT("ExecutionModeDisabledLabel", "Disabled"));
+			SubSection.AddMenuEntry(FPCGEditorCommands::Get().ExecutionModeIsolated, LOCTEXT("ExecutionModeIsolatedLabel", "Isolated"));
+		}));
 	}
 
 	{
@@ -230,6 +239,18 @@ void UPCGEditorGraphNodeBase::OnNodeChanged(UPCGNode* InNode, EPCGChangeType Cha
 {
 	if (InNode == PCGNode)
 	{
+		if (!!(ChangeType & EPCGChangeType::Settings))
+		{
+			if (const UPCGSettings* PCGSettings = InNode->DefaultSettings)
+			{
+				const ENodeEnabledState NewEnabledState = (PCGSettings->ExecutionMode == EPCGSettingsExecutionMode::Disabled) ? ENodeEnabledState::Disabled : ENodeEnabledState::Enabled;
+				if (NewEnabledState != GetDesiredEnabledState())
+				{
+					SetEnabledState(NewEnabledState);
+				}
+			}
+		}
+		
 		ReconstructNode();
 	}
 }
@@ -292,15 +313,17 @@ FLinearColor UPCGEditorGraphNodeBase::GetNodeTitleColor() const
 {
 	if (PCGNode)
 	{
-		if (bIsInspected)
+		const UPCGSettings* PCGSettings = PCGNode->DefaultSettings;
+		if (PCGSettings && PCGSettings->ExecutionMode == EPCGSettingsExecutionMode::Isolated)
 		{
-			return GetDefault<UPCGEditorSettings>()->InspectNodeColor;
+			return GetDefault<UPCGEditorSettings>()->IsolatedNodeColor;
 		}
-		else if (PCGNode->NodeTitleColor != FLinearColor::White)
+
+		if (PCGNode->NodeTitleColor != FLinearColor::White)
 		{
 			return PCGNode->NodeTitleColor;
 		}
-		else if (PCGNode->DefaultSettings)
+		else if (PCGSettings)
 		{
 			FLinearColor SettingsColor = PCGNode->DefaultSettings->GetNodeTitleColor();
 			if (SettingsColor == FLinearColor::White)
@@ -318,6 +341,20 @@ FLinearColor UPCGEditorGraphNodeBase::GetNodeTitleColor() const
 	return GetDefault<UPCGEditorSettings>()->DefaultNodeColor;
 }
 
+FLinearColor UPCGEditorGraphNodeBase::GetNodeBodyTintColor() const
+{
+	if (PCGNode)
+	{
+		const UPCGSettings* PCGSettings = PCGNode->DefaultSettings;
+		if (PCGSettings && PCGSettings->ExecutionMode == EPCGSettingsExecutionMode::Isolated)
+		{
+			return GetDefault<UPCGEditorSettings>()->IsolatedNodeColor;
+		}
+	}
+
+	return Super::GetNodeBodyTintColor();
+}
+
 FEdGraphPinType UPCGEditorGraphNodeBase::GetPinType(const UPCGPin* InPin)
 {
 	FEdGraphPinType EdPinType;
diff --git a/Engine/Plugins/Experimental/PCG/Source/PCGEditor/Private/PCGEditorGraphNodeBase.h b/Engine/Plugins/Experimental/PCG/Source/PCGEditor/Private/PCGEditorGraphNodeBase.h
--- a/Engine/Plugins/Experimental/PCG/Source/PCGEditor/Private/PCGEditorGraphNodeBase.h
+++ b/Engine/Plugins/Experimental/PCG/Source/PCGEditor/Private/PCGEditorGraphNodeBase.h
@@ -34,7 +34,8 @@ public:
 	virtual void PrepareForCopying() override;
 	virtual bool CanCreateUnderSpecifiedSchema(const UEdGraphSchema* Schema) const override;
 	virtual void ReconstructNode() override;
-	virtual FLinearColor GetNodeTitleColor() const;
+	virtual FLinearColor GetNodeTitleColor() const override;
+	virtual FLinearColor GetNodeBodyTintColor() const override;
 	virtual void PostPasteNode() override;
 	virtual FText GetTooltipText() const override;
 	virtual void GetPinHoverText(const UEdGraphPin& Pin, FString& HoverTextOut) const override;
@@ -48,6 +49,7 @@ public:
 	void PostCopy();
 	void PostPaste();
 	void SetInspected(bool InIsInspecting) { bIsInspected = InIsInspecting; }
+	bool GetInspected() const { return bIsInspected; }
 
 	DECLARE_DELEGATE(FOnPCGEditorGraphNodeChanged);
 	FOnPCGEditorGraphNodeChanged OnNodeChangedDelegate;
diff --git a/Engine/Plugins/Experimental/PCG/Source/PCGEditor/Private/PCGEditorModule.cpp b/Engine/Plugins/Experimental/PCG/Source/PCGEditor/Private/PCGEditorModule.cpp
--- a/Engine/Plugins/Experimental/PCG/Source/PCGEditor/Private/PCGEditorModule.cpp
+++ b/Engine/Plugins/Experimental/PCG/Source/PCGEditor/Private/PCGEditorModule.cpp
@@ -8,6 +8,7 @@
 #include "PCGEditorCommands.h"
 #include "PCGEditorGraphNodeFactory.h"
 #include "PCGEditorSettings.h"
+#include "PCGEditorStyle.h"
 #include "PCGGraphDetails.h"
 #include "PCGSubsystem.h"
 #include "PCGVolumeDetails.h"
@@ -34,6 +35,7 @@ void FPCGEditorModule::StartupModule()
 	RegisterSettings();
 
 	FPCGEditorCommands::Register();
+	FPCGEditorStyle::Register();
 
 	GraphNodeFactory = MakeShareable(new FPCGEditorGraphNodeFactory());
 	FEdGraphUtilities::RegisterVisualNodeFactory(GraphNodeFactory);
@@ -52,6 +54,7 @@ void FPCGEditorModule::ShutdownModule()
 	UnregisterMenuExtensions();
 
 	FPCGEditorCommands::Unregister();
+	FPCGEditorStyle::Unregister();
 
 	FEdGraphUtilities::UnregisterVisualNodeFactory(GraphNodeFactory);
 
@@ -195,4 +198,4 @@ IMPLEMENT_MODULE(FPCGEditorModule, PCGEditor);
 
 DEFINE_LOG_CATEGORY(LogPCGEditor);
 
-#undef LOCTEXT_NAMESPACE
\ No newline at end of file
+#undef LOCTEXT_NAMESPACE
diff --git a/Engine/Plugins/Experimental/PCG/Source/PCGEditor/Private/PCGEditorSettings.cpp b/Engine/Plugins/Experimental/PCG/Source/PCGEditor/Private/PCGEditorSettings.cpp
--- a/Engine/Plugins/Experimental/PCG/Source/PCGEditor/Private/PCGEditorSettings.cpp
+++ b/Engine/Plugins/Experimental/PCG/Source/PCGEditor/Private/PCGEditorSettings.cpp
@@ -9,7 +9,7 @@ UPCGEditorSettings::UPCGEditorSettings(const FObjectInitializer& ObjectInitializ
 	: Super(ObjectInitializer)
 {
 	DefaultNodeColor = FLinearColor(0.4f, 0.62f, 1.0f);
-	InspectNodeColor = FLinearColor(0.0f, 0.0f, 1.0f);
+	IsolatedNodeColor = FLinearColor(0.0f, 0.0f, 1.0f);
 	InputOutputNodeColor = FLinearColor(1.0f, 0.0f, 0.0f);
 	SetOperationNodeColor = FLinearColor(0.8f, 0.2f, 0.8f);
 	DensityOperationNodeColor = FLinearColor(0.6f, 1.0f, 0.6f);
diff --git a/Engine/Plugins/Experimental/PCG/Source/PCGEditor/Private/PCGEditorStyle.cpp b/Engine/Plugins/Experimental/PCG/Source/PCGEditor/Private/PCGEditorStyle.cpp
new file mode 100644
--- /dev/null
+++ b/Engine/Plugins/Experimental/PCG/Source/PCGEditor/Private/PCGEditorStyle.cpp
@@ -0,0 +1,34 @@
+// Copyright Epic Games, Inc. All Rights Reserved.
+
+#include "PCGEditorStyle.h"
+
+#include "Misc/Paths.h"
+#include "Styling/SlateStyleMacros.h"
+#include "Styling/SlateStyleRegistry.h"
+#include "Styling/StyleColors.h"
+
+void FPCGEditorStyle::Register()
+{
+	FSlateStyleRegistry::RegisterSlateStyle(Get());
+}
+
+void FPCGEditorStyle::Unregister()
+{
+	FSlateStyleRegistry::UnRegisterSlateStyle(Get());
+}
+
+FPCGEditorStyle::FPCGEditorStyle() : FSlateStyleSet("PCGEditorStyle")
+{
+    static const FVector2D Icon20x20(20.0f, 20.0f);
+	
+	SetCoreContentRoot(FPaths::EngineContentDir() / TEXT("Editor/Slate"));
+
+	Set("PCG.NodeOverlay.Debug", new CORE_IMAGE_BRUSH_SVG("Starship/Blueprints/Breakpoint_Valid", Icon20x20, FSlateColor(FColor::Cyan)));
+	Set("PCG.NodeOverlay.Inspect", new CORE_IMAGE_BRUSH_SVG("Starship/Blueprints/Breakpoint_Valid", Icon20x20, FSlateColor(FColor::Orange)));
+}
+
+const FPCGEditorStyle& FPCGEditorStyle::Get()
+{
+	static FPCGEditorStyle Instance;
+	return Instance;
+}
diff --git a/Engine/Plugins/Experimental/PCG/Source/PCGEditor/Private/PCGEditorStyle.h b/Engine/Plugins/Experimental/PCG/Source/PCGEditor/Private/PCGEditorStyle.h
new file mode 100644
--- /dev/null
+++ b/Engine/Plugins/Experimental/PCG/Source/PCGEditor/Private/PCGEditorStyle.h
@@ -0,0 +1,17 @@
+// Copyright Epic Games, Inc. All Rights Reserved.
+
+#pragma once
+
+#include "Styling/SlateStyle.h"
+
+class FPCGEditorStyle : public FSlateStyleSet
+{
+public:
+	static void Register();
+	static void Unregister();
+
+	static const FPCGEditorStyle& Get();
+
+private:	
+	FPCGEditorStyle();
+};
diff --git a/Engine/Plugins/Experimental/PCG/Source/PCGEditor/Private/SPCGEditorGraphNode.cpp b/Engine/Plugins/Experimental/PCG/Source/PCGEditor/Private/SPCGEditorGraphNode.cpp
--- a/Engine/Plugins/Experimental/PCG/Source/PCGEditor/Private/SPCGEditorGraphNode.cpp
+++ b/Engine/Plugins/Experimental/PCG/Source/PCGEditor/Private/SPCGEditorGraphNode.cpp
@@ -3,10 +3,13 @@
 #include "SPCGEditorGraphNode.h"
 
 #include "PCGEditorGraphNodeBase.h"
+#include "PCGEditorStyle.h"
 #include "PCGNode.h"
 #include "PCGPin.h"
+#include "PCGSettings.h"
 
 #include "SGraphPin.h"
+#include "Styling/SlateBrush.h"
 
 void SPCGEditorGraphNode::Construct(const FArguments& InArgs, UPCGEditorGraphNodeBase* InNode)
 {
@@ -24,7 +27,7 @@ void SPCGEditorGraphNode::Construct(const FArguments& InArgs, UPCGEditorGraphNod
 void SPCGEditorGraphNode::AddPin(const TSharedRef<SGraphPin>& PinToAdd)
 {
 	check(PCGEditorGraphNode);
-	UPCGNode* PCGNode = PCGEditorGraphNode ? PCGEditorGraphNode->GetPCGNode() : nullptr;
+	UPCGNode* PCGNode = PCGEditorGraphNode->GetPCGNode();
 	// Implementation note: we do not distinguish single/multiple pins on the output since that is not relevant
 	if (PCGNode && PinToAdd->GetPinObj())
 	{
@@ -40,7 +43,41 @@ void SPCGEditorGraphNode::AddPin(const TSharedRef<SGraphPin>& PinToAdd)
 	SGraphNode::AddPin(PinToAdd);
 }
 
+void SPCGEditorGraphNode::GetOverlayBrushes(bool bSelected, const FVector2D WidgetSize, TArray<FOverlayBrushInfo>& Brushes) const
+{
+	check(PCGEditorGraphNode);
+	
+	const FSlateBrush* DebugBrush = FPCGEditorStyle::Get().GetBrush(TEXT("PCG.NodeOverlay.Debug"));
+	const FSlateBrush* InspectBrush = FPCGEditorStyle::Get().GetBrush(TEXT("PCG.NodeOverlay.Inspect"));
+	
+	const FVector2D HalfDebugBrushSize = DebugBrush->GetImageSize() / 2.0;
+	const FVector2D HalfInspectBrushSize = InspectBrush->GetImageSize() / 2.0;
+	
+	FVector2D OverlayOffset(0.0, 0.0);
+
+	if (const UPCGNode* PCGNode = PCGEditorGraphNode->GetPCGNode())
+	{
+		if (PCGNode->DefaultSettings && PCGNode->DefaultSettings->ExecutionMode == EPCGSettingsExecutionMode::Debug)
+		{
+			FOverlayBrushInfo BrushInfo;
+			BrushInfo.Brush = DebugBrush;
+			BrushInfo.OverlayOffset = OverlayOffset - HalfDebugBrushSize;
+			Brushes.Add(BrushInfo);
+			
+			OverlayOffset.Y += HalfDebugBrushSize.Y + HalfInspectBrushSize.Y;
+		}
+	}
+
+	if (PCGEditorGraphNode->GetInspected())
+	{
+		FOverlayBrushInfo BrushInfo;
+		BrushInfo.Brush = InspectBrush;
+		BrushInfo.OverlayOffset = OverlayOffset - HalfInspectBrushSize;
+		Brushes.Add(BrushInfo);	
+	}
+}
+
 void SPCGEditorGraphNode::OnNodeChanged()
 {
 	UpdateGraphNode();
-}
\ No newline at end of file
+}
diff --git a/Engine/Plugins/Experimental/PCG/Source/PCGEditor/Private/SPCGEditorGraphNode.h b/Engine/Plugins/Experimental/PCG/Source/PCGEditor/Private/SPCGEditorGraphNode.h
--- a/Engine/Plugins/Experimental/PCG/Source/PCGEditor/Private/SPCGEditorGraphNode.h
+++ b/Engine/Plugins/Experimental/PCG/Source/PCGEditor/Private/SPCGEditorGraphNode.h
@@ -20,6 +20,10 @@ public:
 	virtual void AddPin(const TSharedRef<SGraphPin>& PinToAdd) override;
 	//~ End SGraphNode Interface
 
+	//~ Begin SNodePanel::SNode Interface
+	virtual void GetOverlayBrushes(bool bSelected, const FVector2D WidgetSize, TArray<FOverlayBrushInfo>& Brushes) const override;
+	//~ End SNodePanel::SNode Interface
+
 protected:
 	void OnNodeChanged();
 
diff --git a/Engine/Plugins/Experimental/PCG/Source/PCGEditor/Public/PCGEditorSettings.h b/Engine/Plugins/Experimental/PCG/Source/PCGEditor/Public/PCGEditorSettings.h
--- a/Engine/Plugins/Experimental/PCG/Source/PCGEditor/Public/PCGEditorSettings.h
+++ b/Engine/Plugins/Experimental/PCG/Source/PCGEditor/Public/PCGEditorSettings.h
@@ -25,9 +25,9 @@ public:
 	UPROPERTY(EditAnywhere, config, Category = Node, meta = (HideAlphaChannel))
 	FLinearColor DefaultNodeColor;
 
-	/** Inspect node color */
+	/** Isolated node color */
 	UPROPERTY(EditAnywhere, config, Category = Node, meta = (HideAlphaChannel))
-	FLinearColor InspectNodeColor;
+	FLinearColor IsolatedNodeColor;
 
 	/** Color used for input & output nodes */
 	UPROPERTY(EditAnywhere, config, Category = Node, meta = (HideAlphaChannel))
@@ -96,4 +96,4 @@ public:
 	/** Specify if we want to jump to definition in case of double click on native PCG Nodes */
 	UPROPERTY(EditAnywhere, config, Category = Workflow)
 	bool bEnableNavigateToNativeNodes;
-};
\ No newline at end of file
+};
