Improved PCG gc test in PIE

diff --git a/Engine/Plugins/Experimental/PCG/Source/PCG/Private/Graph/PCGGraphExecutor.cpp b/Engine/Plugins/Experimental/PCG/Source/PCG/Private/Graph/PCGGraphExecutor.cpp
--- a/Engine/Plugins/Experimental/PCG/Source/PCG/Private/Graph/PCGGraphExecutor.cpp
+++ b/Engine/Plugins/Experimental/PCG/Source/PCG/Private/Graph/PCGGraphExecutor.cpp
@@ -10,6 +10,10 @@
 
 #include "Async/Async.h"
 
+#if WITH_EDITOR
+#include "Editor.h"
+#endif
+
 // World partition support for in-editor workflows needs these includes
 #if WITH_EDITOR
 #include "WorldPartition/WorldPartition.h"
@@ -567,7 +571,7 @@ void FPCGGraphExecutor::ReleaseUnusedActors()
 	ActorsListLock.Unlock();
 
 #if WITH_EDITOR
-	if (bRunGC && !GIsPlayInEditorWorld)
+	if (bRunGC && !((GEditor && GEditor->PlayWorld) || GIsPlayInEditorWorld))
 	{
 		CollectGarbage(RF_NoFlags, true);
 	}
