Add dynamic mapping between PCGComponents and PCGPartitionActors + few fixes

Fixes:
* IsRuntimeOrPIE helper function added and used instead of GIsPlayInEditorWorld
* Unused resources are now inconditionnaly cleaned up after any generation (succeeded or failed)

diff --git a/Engine/Plugins/Experimental/PCG/Source/PCG/Private/Graph/PCGGraphExecutor.cpp b/Engine/Plugins/Experimental/PCG/Source/PCG/Private/Graph/PCGGraphExecutor.cpp
--- a/Engine/Plugins/Experimental/PCG/Source/PCG/Private/Graph/PCGGraphExecutor.cpp
+++ b/Engine/Plugins/Experimental/PCG/Source/PCG/Private/Graph/PCGGraphExecutor.cpp
@@ -1,12 +1,13 @@
 // Copyright Epic Games, Inc. All Rights Reserved.
 
 #include "PCGGraphExecutor.h"
-#include "PCGGraph.h"
+#include "PCGComponent.h"
 #include "PCGData.h"
-#include "PCGSubgraph.h"
+#include "PCGHelpers.h"
+#include "PCGGraph.h"
 #include "PCGGraphCompiler.h"
-#include "PCGComponent.h"
 #include "PCGInputOutputSettings.h"
+#include "PCGSubgraph.h"
 
 #include "Async/Async.h"
 
@@ -571,7 +572,7 @@ void FPCGGraphExecutor::ReleaseUnusedActors()
 	ActorsListLock.Unlock();
 
 #if WITH_EDITOR
-	if (bRunGC && !((GEditor && GEditor->PlayWorld) || GIsPlayInEditorWorld))
+	if (bRunGC && !PCGHelpers::IsRuntimeOrPIE())
 	{
 		CollectGarbage(RF_NoFlags, true);
 	}
diff --git a/Engine/Plugins/Experimental/PCG/Source/PCG/Private/Grid/PCGComponentOctree.cpp b/Engine/Plugins/Experimental/PCG/Source/PCG/Private/Grid/PCGComponentOctree.cpp
new file mode 100644
--- /dev/null
+++ b/Engine/Plugins/Experimental/PCG/Source/PCG/Private/Grid/PCGComponentOctree.cpp
@@ -0,0 +1,21 @@
+// Copyright Epic Games, Inc. All Rights Reserved.
+
+#include "Grid/PCGComponentOctree.h"
+#include "PCGComponent.h"
+
+FPCGComponentRef::FPCGComponentRef(UPCGComponent* InComponent, const FPCGComponentOctreeIDSharedRef& InIdShared)
+	: IdShared(InIdShared)
+{
+	check(InComponent);
+
+	Component = InComponent;
+
+	UpdateBounds();
+}
+
+void FPCGComponentRef::UpdateBounds()
+{
+	check(Component);
+
+	Bounds = Component->GetGridBounds();
+}
diff --git a/Engine/Plugins/Experimental/PCG/Source/PCG/Private/Grid/PCGPartitionActor.cpp b/Engine/Plugins/Experimental/PCG/Source/PCG/Private/Grid/PCGPartitionActor.cpp
--- a/Engine/Plugins/Experimental/PCG/Source/PCG/Private/Grid/PCGPartitionActor.cpp
+++ b/Engine/Plugins/Experimental/PCG/Source/PCG/Private/Grid/PCGPartitionActor.cpp
@@ -3,7 +3,9 @@
 #include "PCGPartitionActor.h"
 #include "PCGComponent.h"
 #include "PCGHelpers.h"
+#include "PCGSubsystem.h"
 #include "PCGWorldActor.h"
+#include "Helpers/PCGActorHelpers.h"
 
 #include "Landscape.h"
 #include "Components/BoxComponent.h"
@@ -37,10 +39,10 @@
 		PCGGridSize = APCGWorldActor::DefaultPartitionGridSize;
 	}
 
-#if WITH_EDITOR
 	// Make sure that we don't track objects that do not exist anymore
 	CleanupDeadGraphInstances();
 
+#if WITH_EDITOR
 	if (BoundsComponent)
 	{
 		BoundsComponent->SetBoxExtent(GetFixedBounds().GetExtent());
@@ -68,22 +70,27 @@
 		}
 	}
 
-	// Make the Partition actor register itself to the PCG Component
-	RegisterActorToOriginalComponents();
+	// Make the Partition actor register itself to the PCG Subsystem
+	UWorld* World = GetWorld();
+	check(World);
 
+	UPCGSubsystem* Subsystem = World->GetSubsystem<UPCGSubsystem>();
+	check(Subsystem);
+
+	Subsystem->RegisterPartitionActor(this);
+
 	Super::BeginPlay();
-
-	// Register cell to the PCG grid
-	// TODO
 }
 
 void APCGPartitionActor::EndPlay(const EEndPlayReason::Type EndPlayReason)
 {
-	// Unregister each cell to the PCG grid
-	// TODO
+	UWorld* World = GetWorld();
+	check(World);
 
-	// Make the Partition actor unregister itself to the PCG Component
-	UnregisterActorToOriginalComponents();
+	UPCGSubsystem* Subsystem = World->GetSubsystem<UPCGSubsystem>();
+	check(Subsystem);
 
+	Subsystem->UnregisterPartitionActor(this);
+
 	Super::EndPlay(EndPlayReason);
 }
@@ -108,6 +115,12 @@
 	return FBox(Center - HalfGridSize, Center + HalfGridSize);
 }
 
+FIntVector APCGPartitionActor::GetGridCoord() const
+{
+	const FVector Center = GetActorLocation();
+	return UPCGActorHelpers::GetCellCoord(Center, PCGGridSize);
+}
+
 void APCGPartitionActor::GetActorBounds(bool bOnlyCollidingComponents, FVector& Origin, FVector& BoxExtent, bool bIncludeFromChildActors) const
 {
 	Super::GetActorBounds(bOnlyCollidingComponents, Origin, BoxExtent, bIncludeFromChildActors);
@@ -130,44 +143,42 @@
 	return OriginalComponent ? (*OriginalComponent).Get() : nullptr;
 }
 
-void APCGPartitionActor::RegisterActorToOriginalComponents() const
+bool APCGPartitionActor::CleanupDeadGraphInstances()
 {
-	for (auto& It : OriginalToLocalMap)
+	TSet<TObjectPtr<UPCGComponent>> DeadLocalInstances;
+
+	// Note: since we might end up with multiple nulls in the original to local map
+	// it might not be very stable to use it; we'll use the 
+	for (const auto& LocalToOriginalItem : LocalToOriginalMap)
 	{
-		if (It.Key && It.Value)
+		if (!LocalToOriginalItem.Value)
 		{
-			It.Key->AddPCGPartitionActor(this);
+			DeadLocalInstances.Add(LocalToOriginalItem.Key);
 		}
 	}
+
+	if (DeadLocalInstances.Num() == 0)
+	{
+		return OriginalToLocalMap.IsEmpty();
+	}
+
+	Modify();
+
+	for (const TObjectPtr<UPCGComponent>& DeadInstance : DeadLocalInstances)
+	{
+		LocalToOriginalMap.Remove(DeadInstance);
+
+		if (DeadInstance)
+		{
+			DeadInstance->CleanupLocal(/*bRemoveComponents=*/true);
+			DeadInstance->DestroyComponent();
+		}
+	}
+
+	// Remove all dead entries
+	OriginalToLocalMap.Remove(nullptr);
+
+	return OriginalToLocalMap.IsEmpty();
 }
 
-void APCGPartitionActor::UnregisterActorToOriginalComponents() const
-{
-	for (auto& It : OriginalToLocalMap)
-	{
-		if (It.Key && It.Value)
-		{
-			It.Key->RemovePCGPartitionActor(this);
-		}
-	}
-}
-
-#if WITH_EDITOR
-FBox APCGPartitionActor::GetStreamingBounds() const
-{
-	return Super::GetStreamingBounds() + GetFixedBounds();
-}
-
-AActor* APCGPartitionActor::GetSceneOutlinerParent() const
-{
-	if (APCGWorldActor* PCGActor = PCGHelpers::GetPCGWorldActor(GetWorld()))
-	{
-		return PCGActor;
-	}
-	else
-	{
-		return Super::GetSceneOutlinerParent();
-	}	
-}
-
 void APCGPartitionActor::AddGraphInstance(UPCGComponent* OriginalComponent)
@@ -223,42 +234,22 @@
 	return OriginalToLocalMap.IsEmpty();
 }
 
-bool APCGPartitionActor::CleanupDeadGraphInstances()
+#if WITH_EDITOR
+FBox APCGPartitionActor::GetStreamingBounds() const
 {
-	TSet<TObjectPtr<UPCGComponent>> DeadLocalInstances;
-
-	// Note: since we might end up with multiple nulls in the original to local map
-	// it might not be very stable to use it; we'll use the 
-	for (const auto& LocalToOriginalItem : LocalToOriginalMap)
-	{
-		if (!LocalToOriginalItem.Value)
-		{
-			DeadLocalInstances.Add(LocalToOriginalItem.Key);
-		}
-	}
-
-	if (DeadLocalInstances.Num() == 0)
-	{
-		return OriginalToLocalMap.IsEmpty();
-	}
-
-	Modify();
-
-	for (const TObjectPtr<UPCGComponent>& DeadInstance : DeadLocalInstances)
-	{
-		LocalToOriginalMap.Remove(DeadInstance);
-
-		if (DeadInstance)
-		{
-			DeadInstance->CleanupLocal(/*bRemoveComponents=*/true);
-			DeadInstance->DestroyComponent();
-		}
-	}
-
-	// Remove all dead entries
-	OriginalToLocalMap.Remove(nullptr);
-
-	return OriginalToLocalMap.IsEmpty();
+	return Super::GetStreamingBounds() + GetFixedBounds();
 }
 
+AActor* APCGPartitionActor::GetSceneOutlinerParent() const
+{
+	if (APCGWorldActor* PCGActor = PCGHelpers::GetPCGWorldActor(GetWorld()))
+	{
+		return PCGActor;
+	}
+	else
+	{
+		return Super::GetSceneOutlinerParent();
+	}	
+}
+
 void APCGPartitionActor::PostCreation()
diff --git a/Engine/Plugins/Experimental/PCG/Source/PCG/Private/Grid/PCGPartitionActor.h b/Engine/Plugins/Experimental/PCG/Source/PCG/Private/Grid/PCGPartitionActor.h
--- a/Engine/Plugins/Experimental/PCG/Source/PCG/Private/Grid/PCGPartitionActor.h
+++ b/Engine/Plugins/Experimental/PCG/Source/PCG/Private/Grid/PCGPartitionActor.h
@@ -44,11 +44,13 @@ public:
 #endif
 
 	FBox GetFixedBounds() const;
-#if WITH_EDITOR
+	FIntVector GetGridCoord() const;
+
 	void AddGraphInstance(UPCGComponent* OriginalComponent);
 	bool RemoveGraphInstance(UPCGComponent* OriginalComponent);
 	bool CleanupDeadGraphInstances();
 
+#if WITH_EDITOR
 	/** To be called after the creation of a new actor to copy the GridSize property (Editor only) into the PCGGridSize property */
 	void PostCreation();
 
@@ -89,7 +91,4 @@ private:
 	UPROPERTY(Transient)
 	TObjectPtr<UBoxComponent> BoundsComponent;
 #endif // WITH_EDITORONLY_DATA
-
-	void RegisterActorToOriginalComponents() const;
-	void UnregisterActorToOriginalComponents() const;
 };
diff --git a/Engine/Plugins/Experimental/PCG/Source/PCG/Private/Helpers/PCGActorHelpers.cpp b/Engine/Plugins/Experimental/PCG/Source/PCG/Private/Helpers/PCGActorHelpers.cpp
--- a/Engine/Plugins/Experimental/PCG/Source/PCG/Private/Helpers/PCGActorHelpers.cpp
+++ b/Engine/Plugins/Experimental/PCG/Source/PCG/Private/Helpers/PCGActorHelpers.cpp
@@ -384,3 +384,16 @@ AActor* UPCGActorHelpers::SpawnDefaultActor(UWorld* World, TSubclassOf<AActor> A
 
 	return NewActor;
 }
+
+FIntVector UPCGActorHelpers::GetCellCoord(FVector InPosition, int InGridSize)
+{
+	check(InGridSize > 0);
+
+	FVector Temp = InPosition / InGridSize;
+
+	return FIntVector(
+		FMath::FloorToInt(Temp.X),
+		FMath::FloorToInt(Temp.Y),
+		FMath::FloorToInt(Temp.Z)
+	);
+}
diff --git a/Engine/Plugins/Experimental/PCG/Source/PCG/Private/PCGComponent.cpp b/Engine/Plugins/Experimental/PCG/Source/PCG/Private/PCGComponent.cpp
--- a/Engine/Plugins/Experimental/PCG/Source/PCG/Private/PCGComponent.cpp
+++ b/Engine/Plugins/Experimental/PCG/Source/PCG/Private/PCGComponent.cpp
@@ -112,7 +112,7 @@ bool UPCGComponent::ShouldGenerate(bool bForce, EPCGComponentGenerationTrigger R
 
 #if WITH_EDITOR
 	// Always run Generate if we are in editor and partitioned since the original component doesn't know the state of the local one.
-	if (IsPartitioned() && !GIsPlayInEditorWorld)
+	if (IsPartitioned() && !PCGHelpers::IsRuntimeOrPIE())
 	{
 		return true;
 	}
@@ -221,12 +221,14 @@ void UPCGComponent::GenerateLocal(bool bForce, EPCGComponentGenerationTrigger Re
 	{
 		bIsGenerating = true;
 	}
+	else
+	{
+		OnProcessGraphAborted();
+	}
 }
 
 FPCGTaskId UPCGComponent::GenerateInternal(bool bForce, EPCGComponentGenerationTrigger RequestedGenerationTrigger, const TArray<FPCGTaskId>& TaskDependencies)
 {
-	FPCGTaskId TaskId = InvalidPCGTaskId;
-
 	if (!ShouldGenerate(bForce, RequestedGenerationTrigger))
 	{
 		return InvalidPCGTaskId;
@@ -242,56 +244,19 @@ FPCGTaskId UPCGComponent::GenerateInternal(bool bForce, EPCGComponentGenerationT
 
 	Modify();
 
-	if (IsPartitioned())
+	// Immediate operation: cleanup beforehand
+	if (bGenerated)
 	{
-#if WITH_EDITOR
-		if (!GIsPlayInEditorWorld)
-		{
-			TaskId = GetSubsystem()->DelayGenerateGraph(this, /*bSave=*/bForce);
-		}
-		else
-#endif
-		{
-			// If we don't have valid bounds, just cleanup
-			const FBox NewBounds = GetGridBounds();
-			if (!NewBounds.IsValid)
-			{
-				CleanupLocal(/*bRemoveComponents=*/false);
-				return InvalidPCGTaskId;
-			}
-
-			// Otherwise, ask for generation on all the partition actors registered.
-			TArray<FPCGTaskId> TaskIds = GetSubsystem()->ScheduleMultipleComponent(this, PartitionActors, TaskDependencies);
-
-			// Finally, create a task to call PostProcessGraph.
-			if (!TaskIds.IsEmpty())
-			{
-				return GetSubsystem()->ScheduleGeneric([this, NewBounds]() { PostProcessGraph(NewBounds, true); return true; }, TaskIds);
-			}
-			else
-			{
-				return InvalidPCGTaskId;
-			}
-		}
+		CleanupInternal(/*bRemoveComponents=*/false);
 	}
-	else
-	{
-		// Immediate operation: cleanup beforehand
-		if (bGenerated)
-		{
-			CleanupInternal(/*bRemoveComponents=*/false);
-		}
-
-		const FBox NewBounds = GetGridBounds();
-		if (!NewBounds.IsValid)
-		{
-			return InvalidPCGTaskId;
-		}
 
-		TaskId = GetSubsystem()->ScheduleComponent(this, TaskDependencies);
+	const FBox NewBounds = GetGridBounds();
+	if (!NewBounds.IsValid)
+	{
+		return InvalidPCGTaskId;
 	}
 
-	return TaskId;
+	return GetSubsystem()->ScheduleComponent(this, /*bSave=*/false, TaskDependencies);
 }
 
 bool UPCGComponent::GetActorsFromTags(const TSet<FName>& InTags, TSet<TWeakObjectPtr<AActor>>& OutActors, bool bCullAgainstLocalBounds)
@@ -332,29 +297,14 @@ bool UPCGComponent::GetActorsFromTags(const TSet<FName>& InTags, TSet<TWeakObjec
 	return bHasValidTag;
 }
 
-void UPCGComponent::AddPCGPartitionActor(const APCGPartitionActor* Actor)
-{
-	PartitionActors.Add(Actor);
-}
-
-void UPCGComponent::RemovePCGPartitionActor(const APCGPartitionActor* Actor)
-{
-	PartitionActors.Remove(Actor);
-}
-
-void UPCGComponent::ClearPCGPartitionActors()
-{
-	PartitionActors.Empty();
-}
-
 void UPCGComponent::PostProcessGraph(const FBox& InNewBounds, bool bInGenerated)
 {
 	LastGeneratedBounds = InNewBounds;
 
+	CleanupUnusedManagedResources();
+
 	if (bInGenerated)
 	{
-		CleanupUnusedManagedResources();
-
 		bGenerated = true;
 
 		bIsGenerating = false;
@@ -370,6 +320,8 @@ void UPCGComponent::OnProcessGraphAborted()
 {
 	UE_LOG(LogPCG, Warning, TEXT("Process Graph was called but aborted, check for errors in log if you expected a result."));
 
+	CleanupUnusedManagedResources();
+
 	bGenerated = false;
 	bIsGenerating = false;
 
@@ -404,26 +356,14 @@ void UPCGComponent::CleanupLocal(bool bRemoveComponents, bool bSave)
 		return;
 	}
 
+	Modify();
+
+	GetSubsystem()->ScheduleCleanup(this, bRemoveComponents, bSave, {});
+
 	if (IsPartitioned())
 	{
-#if WITH_EDITOR
-		if (!GIsPlayInEditorWorld)
-		{
-			Modify();
-			GetSubsystem()->CleanupGraph(this, LastGeneratedBounds, bRemoveComponents, bSave);
-		}
-		else
-#endif
-		{
-			GetSubsystem()->ScheduleMultipleCleanup(this, PartitionActors, bRemoveComponents, {});
-		}
-
 		bGenerated = false;
 	}
-	else
-	{
-		CleanupInternal(bRemoveComponents);
-	}
 
 #if WITH_EDITOR
 	OnPCGGraphCleanedDelegate.Broadcast(this);
@@ -568,6 +508,12 @@ void UPCGComponent::BeginPlay()
 {
 	Super::BeginPlay();
 
+	// First if it is partitioned, register itself to the PCGSubsystem, to map the component to all its corresponding PartitionActors
+	if (IsPartitioned() && GetSubsystem())
+	{
+		GetSubsystem()->RegisterPCGComponent(this);
+	}
+
 	if(bActivated && !bGenerated && GenerationTrigger == EPCGComponentGenerationTrigger::GenerateOnLoad)
 	{
 		if (IsPartitioned())
@@ -591,9 +537,14 @@ void UPCGComponent::BeginPlay()
 
 void UPCGComponent::EndPlay(const EEndPlayReason::Type EndPlayReason)
 {
-	Super::EndPlay(EndPlayReason);
+	// Always try to unregister itself, if it doesn't exist, it will early out. 
+	// Just making sure that we don't left some resources registered while dead.
+	if (GetSubsystem())
+	{
+		GetSubsystem()->UnregisterPCGComponent(this);
+	}
 
-	ClearPCGPartitionActors();
+	Super::EndPlay(EndPlayReason);
 }
 
 void UPCGComponent::OnComponentCreated()
@@ -713,7 +664,7 @@ void UPCGComponent::OnGraphChanged(UPCGGraph* InGraph, bool bIsStructural, bool
 
 #if WITH_EDITOR
 	// In editor, since we've changed the graph, we might have changed the tracked actor tags as well
-	if (!GIsPlayInEditorWorld)
+	if (!PCGHelpers::IsRuntimeOrPIE())
 	{
 		TeardownTrackingCallbacks();
 		SetupTrackingCallbacks();
diff --git a/Engine/Plugins/Experimental/PCG/Source/PCG/Private/PCGHelpers.cpp b/Engine/Plugins/Experimental/PCG/Source/PCG/Private/PCGHelpers.cpp
--- a/Engine/Plugins/Experimental/PCG/Source/PCG/Private/PCGHelpers.cpp
+++ b/Engine/Plugins/Experimental/PCG/Source/PCG/Private/PCGHelpers.cpp
@@ -10,6 +10,10 @@
 #include "UObject/UObjectIterator.h"
 #include "WorldPartition/WorldPartition.h"
 
+#if WITH_EDITOR
+#include "Editor.h"
+#endif
+
 namespace PCGHelpers
 {
 	int ComputeSeed(int A)
@@ -75,6 +79,15 @@ namespace PCGHelpers
 		return Box;
 	}
 
+	bool IsRuntimeOrPIE()
+	{
+#if WITH_EDITOR
+		return (GEditor && GEditor->PlayWorld) || GIsPlayInEditorWorld;
+#else
+		return true;
+#endif // WITH_EDITOR
+	}
+
 	FBox GetLandscapeBounds(ALandscapeProxy* InLandscape)
 	{
 		check(InLandscape);
@@ -82,7 +95,7 @@ namespace PCGHelpers
 		if (ALandscape* Landscape = Cast<ALandscape>(InLandscape))
 		{
 #if WITH_EDITOR
-			if (!GIsPlayInEditorWorld)
+			if (!IsRuntimeOrPIE())
 			{
 				return Landscape->GetCompleteBounds();
 			}
diff --git a/Engine/Plugins/Experimental/PCG/Source/PCG/Private/PCGHelpers.h b/Engine/Plugins/Experimental/PCG/Source/PCG/Private/PCGHelpers.h
--- a/Engine/Plugins/Experimental/PCG/Source/PCG/Private/PCGHelpers.h
+++ b/Engine/Plugins/Experimental/PCG/Source/PCG/Private/PCGHelpers.h
@@ -33,6 +33,8 @@ namespace PCGHelpers
 	ALandscape* GetLandscape(UWorld* InWorld, const FBox& InActorBounds);
 	ALandscape* GetAnyLandscape(UWorld* InWorld);
 
+	bool IsRuntimeOrPIE();
+
 #if WITH_EDITOR
 	APCGWorldActor* GetPCGWorldActor(UWorld* InWorld);
 
diff --git a/Engine/Plugins/Experimental/PCG/Source/PCG/Private/PCGSubsystem.cpp b/Engine/Plugins/Experimental/PCG/Source/PCG/Private/PCGSubsystem.cpp
--- a/Engine/Plugins/Experimental/PCG/Source/PCG/Private/PCGSubsystem.cpp
+++ b/Engine/Plugins/Experimental/PCG/Source/PCG/Private/PCGSubsystem.cpp
@@ -2,6 +2,7 @@
 
 #include "PCGSubsystem.h"
 #include "PCGComponent.h"
+#include "PCGHelpers.h"
 #include "PCGWorldActor.h"
 #include "Graph/PCGGraphExecutor.h"
 #include "Grid/PCGPartitionActor.h"
@@ -9,6 +10,7 @@
 
 #include "Engine/World.h"
 #include "HAL/IConsoleManager.h"
+#include "Math/GenericOctree.h"
 
 #if WITH_EDITOR
 #include "Editor.h"
@@ -58,6 +60,11 @@ void UPCGSubsystem::PostInitialize()
 	GraphExecutor = new FPCGGraphExecutor(this);
 	// gather things.. ?
 	// TODO	
+
+	// TODO: For now we set our octree to be 2km wide, but it would be perhaps better to
+	// scale it to the size of our world.
+	constexpr FVector::FReal OctreeExtent = 200000; // 2km
+	PCGComponentOctree = FPCGComponentOctree(FVector::ZeroVector, OctreeExtent);
 }
 
 void UPCGSubsystem::Tick(float DeltaSeconds)
@@ -95,13 +102,43 @@ void UPCGSubsystem::UnregisterPCGWorldActor(APCGWorldActor* InActor)
 	}
 }
 
-FPCGTaskId UPCGSubsystem::ScheduleComponent(UPCGComponent* PCGComponent, const TArray<FPCGTaskId>& Dependencies)
+FPCGTaskId UPCGSubsystem::ScheduleComponent(UPCGComponent* PCGComponent, bool bSave, const TArray<FPCGTaskId>& Dependencies)
 {
 	check(GraphExecutor);
 
-	FPCGTaskId ProcessTaskId = PCGComponent ? GraphExecutor->Schedule(PCGComponent, Dependencies) : InvalidPCGTaskId;
+	if (!PCGComponent)
+	{
+		return InvalidPCGTaskId;
+	}
+
+	// TODO: Merge both path
+#if WITH_EDITOR
+	if (PCGComponent->IsPartitioned() && !PCGHelpers::IsRuntimeOrPIE())
+	{
+		return DelayGenerateGraph(PCGComponent, bSave);
+	}
+#endif // WITH_EDITOR
+
+	TArray<FPCGTaskId> AllTasks;
+
+	// If the component is partitioned, we will forward the calls to its registered PCG Partition actors
+	if (PCGComponent->IsPartitioned())
+	{
+		{
+			// TODO: Might be more interesting to copy the set and release the lock.
+			FReadScopeLock ReadLock(ComponentToPartitionActorsMapLock);
+			if (TSet<TObjectPtr<APCGPartitionActor>>* PartitionActorsPtr = ComponentToPartitionActorsMap.Find(PCGComponent))
+			{
+				AllTasks = ScheduleMultipleComponent(PCGComponent, *PartitionActorsPtr, Dependencies);
+			}
+		}
+	}
+	else
+	{
+		AllTasks.Add(GraphExecutor->Schedule(PCGComponent, Dependencies));
+	}
 
-	if (ProcessTaskId != InvalidPCGTaskId)
+	if (!AllTasks.IsEmpty() && AllTasks.Last() != InvalidPCGTaskId)
 	{
 		TWeakObjectPtr<UPCGComponent> ComponentPtr(PCGComponent);
 
@@ -113,7 +150,7 @@ FPCGTaskId UPCGSubsystem::ScheduleComponent(UPCGComponent* PCGComponent, const T
 			}
 
 			return true;
-			}, { ProcessTaskId });
+			}, AllTasks);
 	}
 	else
 	{
@@ -126,21 +163,53 @@ FPCGTaskId UPCGSubsystem::ScheduleComponent(UPCGComponent* PCGComponent, const T
 	}
 }
 
-FPCGTaskId UPCGSubsystem::ScheduleCleanup(UPCGComponent* PCGComponent, bool bRemoveComponents, const TArray<FPCGTaskId>& Dependencies)
+FPCGTaskId UPCGSubsystem::ScheduleCleanup(UPCGComponent* PCGComponent, bool bRemoveComponents, bool bSave, const TArray<FPCGTaskId>& Dependencies)
 {
-	TWeakObjectPtr<UPCGComponent> ComponentPtr(PCGComponent);
+	if (!PCGComponent)
+	{
+		return InvalidPCGTaskId;
+	}
 
-	auto CleanupTask = [ComponentPtr, bRemoveComponents]() {
-		check(ComponentPtr.Get() != nullptr);
-		if (UPCGComponent* Component = ComponentPtr.Get())
+	// TODO: Merge both path
+#if WITH_EDITOR
+	if (PCGComponent->IsPartitioned() && !PCGHelpers::IsRuntimeOrPIE())
+	{
+		// Immediate operation
+		CleanupGraph(PCGComponent, PCGComponent->LastGeneratedBounds, bRemoveComponents, bSave);
+		return InvalidPCGTaskId;
+	}
+#endif // WITH_EDITOR
+
+	TArray<FPCGTaskId> AllTasks;
+
+	// If the component is partitioned, we will forward the calls to its registered PCG Partition actors
+	if (PCGComponent->IsPartitioned())
+	{
 		{
-			Component->CleanupInternal(bRemoveComponents);
+			// TODO: Might be more interesting to copy the set and release the lock.
+			FReadScopeLock ReadLock(ComponentToPartitionActorsMapLock);
+			TSet<TObjectPtr<APCGPartitionActor>>* PartitionActorsPtr = ComponentToPartitionActorsMap.Find(PCGComponent);
+
+			if (PartitionActorsPtr)
+			{
+				AllTasks = ScheduleMultipleCleanup(PCGComponent, *PartitionActorsPtr, bRemoveComponents, Dependencies);
+			}
 		}
+	}
+	else
+	{
+		// In non partitioned mode, do it immediately
+		// TODO: Should it be a parameter?
+		PCGComponent->CleanupInternal(bRemoveComponents);
+	}
 
-		return true;
-	};
+	if (AllTasks.IsEmpty())
+	{
+		return InvalidPCGTaskId;
+	}
 
-	return GraphExecutor->ScheduleGeneric(CleanupTask, Dependencies);
+	// Add a dummy task to wait after all of the tasks
+	return GraphExecutor->ScheduleGeneric([]() { return true; }, AllTasks);
 }
 
 FPCGTaskId UPCGSubsystem::ScheduleGraph(UPCGGraph* Graph, UPCGComponent* SourceComponent, FPCGElementPtr InputElement, const TArray<FPCGTaskId>& Dependencies)
@@ -171,19 +240,28 @@ FPCGTaskId UPCGSubsystem::ScheduleGeneric(TFunction<bool()> InOperation, const T
 	return GraphExecutor->ScheduleGeneric(InOperation, TaskDependencies);
 }
 
-TArray<FPCGTaskId> UPCGSubsystem::ScheduleMultipleComponent(UPCGComponent* OriginalComponent, TSet<TSoftObjectPtr<APCGPartitionActor>>& PartitionActors, const TArray<FPCGTaskId>& Dependencies)
+TArray<FPCGTaskId> UPCGSubsystem::ScheduleMultipleComponent(UPCGComponent* OriginalComponent, TSet<TObjectPtr<APCGPartitionActor>>& PartitionActors, const TArray<FPCGTaskId>& Dependencies)
 {
 	TArray<FPCGTaskId> TaskIds;
-	for (TSoftObjectPtr<APCGPartitionActor>& PartitionActorPtr : PartitionActors)
+	for (APCGPartitionActor* PartitionActor : PartitionActors)
 	{
-		if (APCGPartitionActor* PartitionActor = PartitionActorPtr.Get())
+		if (PartitionActor)
 		{
 			if (UPCGComponent* LocalComponent = PartitionActor->GetLocalComponent(OriginalComponent))
 			{
-				// Ensure that the PCG actor match our original
-				LocalComponent->SetPropertiesFromOriginal(OriginalComponent);
+				// Add check to avoid infinite loop
+				if (ensure(!LocalComponent->IsPartitioned()))
+				{
+					// Ensure that the PCG actor match our original
+					LocalComponent->SetPropertiesFromOriginal(OriginalComponent);
 
-				TaskIds.Add(ScheduleComponent(LocalComponent, Dependencies));
+					FPCGTaskId LocalTask = ScheduleComponent(LocalComponent, /*bSave=*/ false, Dependencies);
+
+					if (LocalTask != InvalidPCGTaskId)
+					{
+						TaskIds.Add(LocalTask);
+					}
+				}
 			}
 		}
 	}
@@ -191,16 +269,31 @@ TArray<FPCGTaskId> UPCGSubsystem::ScheduleMultipleComponent(UPCGComponent* Origi
 	return TaskIds;
 }
 
-TArray<FPCGTaskId> UPCGSubsystem::ScheduleMultipleCleanup(UPCGComponent* OriginalComponent, TSet<TSoftObjectPtr<APCGPartitionActor>>& PartitionActors, bool bRemoveComponents, const TArray<FPCGTaskId>& Dependencies)
+TArray<FPCGTaskId> UPCGSubsystem::ScheduleMultipleCleanup(UPCGComponent* OriginalComponent, TSet<TObjectPtr<APCGPartitionActor>>& PartitionActors, bool bRemoveComponents, const TArray<FPCGTaskId>& Dependencies)
 {
 	TArray<FPCGTaskId> TaskIds;
-	for (TSoftObjectPtr<APCGPartitionActor>& PartitionActorPtr : PartitionActors)
+	for (APCGPartitionActor* PartitionActor : PartitionActors)
 	{
-		if (APCGPartitionActor* PartitionActor = PartitionActorPtr.Get())
+		if (PartitionActor)
 		{
 			if (UPCGComponent* LocalComponent = PartitionActor->GetLocalComponent(OriginalComponent))
 			{
-				TaskIds.Add(ScheduleCleanup(LocalComponent, bRemoveComponents, Dependencies));
+				// Add check to avoid infinite loop
+				if (ensure(!LocalComponent->IsPartitioned()))
+				{
+					TWeakObjectPtr<UPCGComponent> LocalComponentPtr(LocalComponent);
+					auto CleanupTask = [this, LocalComponentPtr, bRemoveComponents]()
+					{
+						if (UPCGComponent* Component = LocalComponentPtr.Get())
+						{
+							ScheduleCleanup(Component, bRemoveComponents, /*bSave=*/false, {});
+						}
+
+						return true;
+					};
+
+					TaskIds.Add(GraphExecutor->ScheduleGeneric(CleanupTask, Dependencies));
+				}
 			}
 		}
 	}
@@ -214,6 +307,243 @@ bool UPCGSubsystem::GetOutputData(FPCGTaskId TaskId, FPCGDataCollection& OutData
 	return GraphExecutor->GetOutputData(TaskId, OutData);
 }
 
+void UPCGSubsystem::RegisterPCGComponent(UPCGComponent* InComponent)
+{
+	// Just make sure that we don't register components that are from the PCG Partition Actor
+	check(InComponent && InComponent->GetOwner() && !InComponent->GetOwner()->IsA<APCGPartitionActor>());
+
+	// Check also that the bounds are valid. If not early out.
+	if (!InComponent->GetGridBounds().IsValid)
+	{
+		UE_LOG(LogPCG, Error, TEXT("[RegisterPCGComponent] Component has invalid bounds, not registered."));
+		return;
+	}
+
+	FBox Bounds(EForceInit::ForceInit);
+
+	{
+		FWriteScopeLock WriteLock(PCGVolumeOctreeLock);
+
+		// We should not register a component twice.
+		check(!ComponentToIdMap.Contains(InComponent));
+
+		FPCGComponentOctreeIDSharedRef IdShared = MakeShared<FPCGComponentOctreeID>();
+		FPCGComponentRef ComponentRef(InComponent, IdShared);
+		Bounds = ComponentRef.Bounds.GetBox();
+		check(Bounds.IsValid);
+		PCGComponentOctree.AddElement(ComponentRef);
+
+		// Store the shared ptr, because if we add/remove components in the octree, the id might change.
+		// We need to make sure that we always have the latest id for the given component.
+		ComponentToIdMap.Add(InComponent, MoveTemp(IdShared));
+	}
+
+	// After adding, do the mapping
+	{
+		FWriteScopeLock WriteLock(ComponentToPartitionActorsMapLock);
+
+		TSet<TObjectPtr<APCGPartitionActor>>& PartitionActors = ComponentToPartitionActorsMap.Emplace(InComponent);
+		FindAllIntersectingPartitionActors(Bounds, [&PartitionActors, InComponent](APCGPartitionActor* Actor)
+			{
+				PartitionActors.Add(Actor);
+				Actor->AddGraphInstance(InComponent);
+			});
+	}
+}
+
+void UPCGSubsystem::UpdatePCGComponentBounds(UPCGComponent* InComponent)
+{
+	check(InComponent);
+	FBox Bounds;
+
+	// Remove it and add it again to the octree
+	{
+		FWriteScopeLock WriteLock(PCGVolumeOctreeLock);
+
+		FPCGComponentOctreeIDSharedRef* ElementIdPtr = ComponentToIdMap.Find(InComponent);
+
+		if (!ElementIdPtr)
+		{
+			return;
+		}
+
+		FPCGComponentRef ComponentRef = PCGComponentOctree.GetElementById((*ElementIdPtr)->Id);
+		PCGComponentOctree.RemoveElement((*ElementIdPtr)->Id);
+
+		ComponentRef.UpdateBounds();
+		Bounds = ComponentRef.Bounds.GetBox();
+
+		PCGComponentOctree.AddElement(ComponentRef);
+	}
+
+	// Redo the mapping between components and PCG PartitionActors
+	{
+		FWriteScopeLock WriteLock(ComponentToPartitionActorsMapLock);
+		TSet<TObjectPtr<APCGPartitionActor>>* PartitionActorsPtr = ComponentToPartitionActorsMap.Find(InComponent);
+		if (ensure(PartitionActorsPtr))
+		{
+			PartitionActorsPtr->Empty();
+			FindAllIntersectingPartitionActors(Bounds, [PartitionActorsPtr, InComponent](APCGPartitionActor* Actor)
+				{
+					PartitionActorsPtr->Add(Actor);
+					Actor->AddGraphInstance(InComponent);
+				});
+		}
+	}
+
+}
+
+void UPCGSubsystem::UnregisterPCGComponent(UPCGComponent* InComponent)
+{
+	check(InComponent);
+
+	// First verification we have this component registered
+	{
+		FReadScopeLock ReadLock(PCGVolumeOctreeLock);
+
+		FPCGComponentOctreeIDSharedRef* ElementIdPtr = ComponentToIdMap.Find(InComponent);
+
+		if (!ElementIdPtr)
+		{
+			return;
+		}
+	}
+
+	// If so, lock in write and retry to find it (to avoid removing it twice)
+	{
+		FWriteScopeLock WriteLock(PCGVolumeOctreeLock);
+
+		FPCGComponentOctreeIDSharedRef* ElementIdPtr = ComponentToIdMap.Find(InComponent);
+
+		if (!ElementIdPtr)
+		{
+			return;
+		}
+
+		PCGComponentOctree.RemoveElement((*ElementIdPtr)->Id);
+		ComponentToIdMap.Remove(InComponent);
+	}
+
+	{
+		FReadScopeLock ReadLock(ComponentToPartitionActorsMapLock);
+		TSet<TObjectPtr<APCGPartitionActor>>* PartitionActorsPtr = ComponentToPartitionActorsMap.Find(InComponent);
+		if (ensure(PartitionActorsPtr))
+		{
+			for (APCGPartitionActor* Actor : *PartitionActorsPtr)
+			{
+				Actor->RemoveGraphInstance(InComponent);
+			}
+		}
+	}
+
+	{
+		FWriteScopeLock WriteLock(ComponentToPartitionActorsMapLock);
+		ComponentToPartitionActorsMap.Remove(InComponent);
+	}
+}
+
+void UPCGSubsystem::FindAllIntersectingComponents(const FBoxCenterAndExtent& InBounds, TFunction<void(UPCGComponent*)> InFunc) const
+{
+	FReadScopeLock ReadLock(PCGVolumeOctreeLock);
+
+	PCGComponentOctree.FindElementsWithBoundsTest(InBounds, [&InFunc](const FPCGComponentRef& ComponentRef)
+		{
+			InFunc(ComponentRef.Component);
+		});
+}
+
+void UPCGSubsystem::RegisterPartitionActor(APCGPartitionActor* Actor)
+{
+	check(Actor);
+
+	FIntVector GridCoord = Actor->GetGridCoord();
+	{
+		FWriteScopeLock WriteLock(PartitionActorsMapLock);
+
+		if (PartitionActorsMap.Contains(GridCoord))
+		{
+			return;
+		}
+
+		PartitionActorsMap.Add(GridCoord, Actor);
+	}
+
+	// And then register itself to all the components that intersect with it
+	{
+		FWriteScopeLock WriteLock(ComponentToPartitionActorsMapLock);
+		FindAllIntersectingComponents(FBoxCenterAndExtent(Actor->GetFixedBounds()), [this, Actor](UPCGComponent* Component)
+			{
+				TSet<TObjectPtr<APCGPartitionActor>>* PartitionActorsPtr = ComponentToPartitionActorsMap.Find(Component);
+				if (ensure(PartitionActorsPtr))
+				{
+					Actor->AddGraphInstance(Component);
+					PartitionActorsPtr->Add(Actor);
+				}
+			});
+	}
+}
+
+void UPCGSubsystem::UnregisterPartitionActor(APCGPartitionActor* Actor)
+{
+	check(Actor);
+
+	FIntVector GridCoord = Actor->GetGridCoord();
+
+	{
+		FWriteScopeLock WriteLock(PartitionActorsMapLock);
+		PartitionActorsMap.Remove(GridCoord);
+	}
+
+	// And then unregister itself to all the components that intersect with it
+	{
+		FWriteScopeLock WriteLock(ComponentToPartitionActorsMapLock);
+		FindAllIntersectingComponents(FBoxCenterAndExtent(Actor->GetFixedBounds()), [this, Actor](UPCGComponent* Component)
+			{
+				TSet<TObjectPtr<APCGPartitionActor>>* PartitionActorsPtr = ComponentToPartitionActorsMap.Find(Component);
+				if (ensure(PartitionActorsPtr))
+				{
+					Actor->RemoveGraphInstance(Component);
+					PartitionActorsPtr->Remove(Actor);
+				}
+			});
+	}
+}
+
+void UPCGSubsystem::FindAllIntersectingPartitionActors(const FBox& InBounds, TFunction<void(APCGPartitionActor*)> InFunc) const
+{
+	// No World actor nor any PA registered, just early out. Same for invalid bounds.
+	if (!PCGWorldActor || PartitionActorsMap.IsEmpty() || !InBounds.IsValid)
+	{
+		return;
+	}
+
+	const uint32 GridSize = PCGWorldActor->PartitionGridSize;
+	const FIntVector MinCellCoords = UPCGActorHelpers::GetCellCoord(InBounds.Min, GridSize);
+	const FIntVector MaxCellCoords = UPCGActorHelpers::GetCellCoord(InBounds.Max, GridSize);
+
+	{
+		FReadScopeLock ReadLock(PartitionActorsMapLock);
+
+		for (int32 z = MinCellCoords.Z; z <= MaxCellCoords.Z; z++)
+		{
+			for (int32 y = MinCellCoords.Y; y <= MaxCellCoords.Y; y++)
+			{
+				for (int32 x = MinCellCoords.X; x <= MaxCellCoords.X; x++)
+				{
+					FIntVector CellCoords(x, y, z);
+					if (const TObjectPtr<APCGPartitionActor>* ActorPtr = PartitionActorsMap.Find(CellCoords))
+					{
+						if (APCGPartitionActor* Actor = ActorPtr->Get())
+						{
+							InFunc(Actor);
+						}
+					}
+				}
+			}
+		}
+	}
+}
+
 #if WITH_EDITOR
 
 namespace PCGSubsystem
@@ -500,7 +830,7 @@ void UPCGSubsystem::CleanupGraph(UPCGComponent* Component, const FBox& InBounds,
 {
 	TWeakObjectPtr<UPCGComponent> ComponentPtr(Component);
 
-	auto ScheduleTask = [this, ComponentPtr, bRemoveComponents](APCGPartitionActor* PCGActor, const FBox& InIntersectedBounds, const TArray<FPCGTaskId>& TaskDependencies) {
+	auto ScheduleTask = [this, ComponentPtr, bRemoveComponents, bSave](APCGPartitionActor* PCGActor, const FBox& InIntersectedBounds, const TArray<FPCGTaskId>& TaskDependencies) {
 		UPCGComponent* Component = ComponentPtr.Get();
 		check(Component != nullptr && PCGActor != nullptr);
 
@@ -511,7 +841,11 @@ void UPCGSubsystem::CleanupGraph(UPCGComponent* Component, const FBox& InBounds,
 
 		if (UPCGComponent* LocalComponent = PCGActor->GetLocalComponent(Component))
 		{
-			return ScheduleCleanup(LocalComponent, bRemoveComponents, TaskDependencies);
+			// Ensure to avoid infinite loop
+			if (ensure(!LocalComponent->IsPartitioned()))
+			{
+				return ScheduleCleanup(LocalComponent, bRemoveComponents, bSave, TaskDependencies);
+			}
 		}
 
 		return InvalidPCGTaskId;
@@ -752,6 +1086,13 @@ void UPCGSubsystem::FlushCache()
 	}
 }
 
+void UPCGSubsystem::ResetPartitionActorsMap()
+{
+	PartitionActorsMapLock.WriteLock();
+	PartitionActorsMap.Empty();
+	PartitionActorsMapLock.WriteUnlock();
+}
+
 ///////////////////////////////////////////////////////////////////////////////////////////
 void UPCGSubsystem::DelayPartitionGraph(UPCGComponent* Component)
 {
diff --git a/Engine/Plugins/Experimental/PCG/Source/PCG/Public/Grid/PCGComponentOctree.h b/Engine/Plugins/Experimental/PCG/Source/PCG/Public/Grid/PCGComponentOctree.h
new file mode 100644
--- /dev/null
+++ b/Engine/Plugins/Experimental/PCG/Source/PCG/Public/Grid/PCGComponentOctree.h
@@ -0,0 +1,59 @@
+// Copyright Epic Games, Inc. All Rights Reserved.
+
+#pragma once
+
+#include "Math/GenericOctree.h"
+#include "Math/GenericOctreePublic.h"
+#include "Templates/SharedPointer.h"
+#include "UObject/ObjectPtr.h"
+
+class UPCGComponent;
+
+struct PCG_API FPCGComponentOctreeID : public TSharedFromThis<FPCGComponentOctreeID, ESPMode::ThreadSafe>
+{
+	FOctreeElementId2 Id;
+};
+
+typedef TSharedRef<struct FPCGComponentOctreeID, ESPMode::ThreadSafe> FPCGComponentOctreeIDSharedRef;
+
+struct PCG_API FPCGComponentRef
+{
+	FPCGComponentRef(UPCGComponent* InComponent, const FPCGComponentOctreeIDSharedRef& InIdShared);
+
+	void UpdateBounds();
+
+	FPCGComponentOctreeIDSharedRef IdShared;
+	TObjectPtr<UPCGComponent> Component;
+	FBoxSphereBounds Bounds;
+};
+
+struct PCG_API FPCGComponentRefSemantics
+{
+	enum { MaxElementsPerLeaf = 16 };
+	enum { MinInclusiveElementsPerNode = 7 };
+	enum { MaxNodeDepth = 12 };
+
+	typedef TInlineAllocator<MaxElementsPerLeaf> ElementAllocator;
+
+	FORCEINLINE static const FBoxSphereBounds& GetBoundingBox(const FPCGComponentRef& InVolume)
+	{
+		return InVolume.Bounds;
+	}
+
+	FORCEINLINE static const bool AreElementsEqual(const FPCGComponentRef& A, const FPCGComponentRef& B)
+	{
+		return A.Component == B.Component;
+	}
+
+	FORCEINLINE static void ApplyOffset(FPCGComponentRef& InVolume, const FVector& Offset)
+	{
+		InVolume.Bounds.Origin += Offset;
+	}
+
+	FORCEINLINE static void SetElementId(const FPCGComponentRef& Element, FOctreeElementId2 OctreeElementID)
+	{
+		Element.IdShared->Id = OctreeElementID;
+	}
+};
+
+typedef TOctree2<FPCGComponentRef, FPCGComponentRefSemantics> FPCGComponentOctree;
diff --git a/Engine/Plugins/Experimental/PCG/Source/PCG/Public/Helpers/PCGActorHelpers.h b/Engine/Plugins/Experimental/PCG/Source/PCG/Public/Helpers/PCGActorHelpers.h
--- a/Engine/Plugins/Experimental/PCG/Source/PCG/Public/Helpers/PCGActorHelpers.h
+++ b/Engine/Plugins/Experimental/PCG/Source/PCG/Public/Helpers/PCGActorHelpers.h
@@ -96,4 +96,9 @@ public:
 	* @param Parent Optional parent to attach to.
 	*/
 	static AActor* SpawnDefaultActor(UWorld* World, TSubclassOf<AActor> ActorClass, FName BaseName, const FTransform& Transform, AActor* Parent = nullptr);
+
+	/**
+	 * Return the grid cell coordinates on the PCG partition grid given a position and the grid size.
+	 */
+	static FIntVector GetCellCoord(FVector InPosition, int InGridSize);
 };
\ No newline at end of file
diff --git a/Engine/Plugins/Experimental/PCG/Source/PCG/Public/PCGComponent.h b/Engine/Plugins/Experimental/PCG/Source/PCG/Public/PCGComponent.h
--- a/Engine/Plugins/Experimental/PCG/Source/PCG/Public/PCGComponent.h
+++ b/Engine/Plugins/Experimental/PCG/Source/PCG/Public/PCGComponent.h
@@ -185,9 +185,7 @@ public:
 
 	UPCGSubsystem* GetSubsystem() const;
 
-	void AddPCGPartitionActor(const APCGPartitionActor* Actor);
-	void RemovePCGPartitionActor(const APCGPartitionActor* Actor);
-	void ClearPCGPartitionActors();
+	FBox GetGridBounds() const;
 
 protected:
 	UPROPERTY(BlueprintReadWrite, EditAnywhere, Category = Properties)
@@ -250,7 +248,6 @@ private:
 	bool GraphUsesLandscapePin() const;
 #endif
 
-	FBox GetGridBounds() const;
 	FBox GetGridBounds(AActor* InActor) const;
 
 	UPROPERTY(Transient, NonPIEDuplicateTransient)
@@ -315,9 +312,5 @@ private:
 	TMap<const UPCGNode*, FPCGDataCollection> InspectionCache;
 #endif
 
-	// TODO: Will need to be cleaned up after we have dynamic association between PartitionActors and PCGComponents.
-	UPROPERTY(Transient)
-	TSet<TSoftObjectPtr<APCGPartitionActor>> PartitionActors;
-
 	FCriticalSection GeneratedResourcesLock;
 };
diff --git a/Engine/Plugins/Experimental/PCG/Source/PCG/Public/PCGSubsystem.h b/Engine/Plugins/Experimental/PCG/Source/PCG/Public/PCGSubsystem.h
--- a/Engine/Plugins/Experimental/PCG/Source/PCG/Public/PCGSubsystem.h
+++ b/Engine/Plugins/Experimental/PCG/Source/PCG/Public/PCGSubsystem.h
@@ -6,6 +6,8 @@
 
 #include "PCGCommon.h"
 #include "PCGComponent.h"
+#include "PCGVolume.h"
+#include "Grid/PCGComponentOctree.h"
 
 #include "PCGSubsystem.generated.h"
 
@@ -18,6 +20,7 @@ struct FPCGDataCollection;
 class IPCGElement;
 typedef TSharedPtr<IPCGElement, ESPMode::ThreadSafe> FPCGElementPtr;
 
+
 /**
 * UPCGSubsystem
 */
@@ -48,16 +51,10 @@ public:
 	void UnregisterPCGWorldActor(APCGWorldActor* InActor);
 
 	// Schedule graph (owner -> graph)
-	FPCGTaskId ScheduleComponent(UPCGComponent* PCGComponent, const TArray<FPCGTaskId>& Dependencies);
-
-	// Schedule cleanup (owner -> graph)
-	FPCGTaskId ScheduleCleanup(UPCGComponent* PCGComponent, bool bRemoveComponents, const TArray<FPCGTaskId>& Dependencies);
+	FPCGTaskId ScheduleComponent(UPCGComponent* PCGComponent, bool bSave, const TArray<FPCGTaskId>& Dependencies);
 
-	// Schedule multiple graphs
-	TArray<FPCGTaskId> ScheduleMultipleComponent(UPCGComponent* OriginalComponent, TSet<TSoftObjectPtr<APCGPartitionActor>>& PartitionActors, const TArray<FPCGTaskId>& Dependencies);
-
-	// Schedule multiple cleanups
-	TArray<FPCGTaskId> ScheduleMultipleCleanup(UPCGComponent* OriginalComponent, TSet<TSoftObjectPtr<APCGPartitionActor>>& PartitionActors, bool bRemoveComponents, const TArray<FPCGTaskId>& Dependencies);
+	/** Schedule cleanup(owner->graph). Note that in non-partitioned mode, cleanup is immediate. */
+	FPCGTaskId ScheduleCleanup(UPCGComponent* PCGComponent, bool bRemoveComponents, bool bSave, const TArray<FPCGTaskId>& Dependencies);
 
 	// Schedule graph (used internally for dynamic subgraph execution)
 	FPCGTaskId ScheduleGraph(UPCGGraph* Graph, UPCGComponent* SourceComponent, FPCGElementPtr InputElement, const TArray<FPCGTaskId>& Dependencies);
@@ -68,6 +65,21 @@ public:
 	/** Gets the output data for a given task */
 	bool GetOutputData(FPCGTaskId InTaskId, FPCGDataCollection& OutData);
 
+	/** Register a new PCG Component, will be added to the octree. Returns an id that needs to be kept for update and removal. Thread safe */
+	void RegisterPCGComponent(UPCGComponent* InComponent);
+
+	/** Update a PCG Component, if it has changed its transform. Thread safe */
+	void UpdatePCGComponentBounds(UPCGComponent* InComponent);
+
+	/** Unregister a PCG Component, will be removed from the octree. Thread safe */
+	void UnregisterPCGComponent(UPCGComponent* InComponent);
+
+	/** Register a new Partition actor, will be added to a map and will query all intersecting volume to bind to them. Thread safe */
+	void RegisterPartitionActor(APCGPartitionActor* InActor);
+
+	/** Unregister a Partition actor, will be removed from the map and remove itself to all intersecting volumes. Thread safe */
+	void UnregisterPartitionActor(APCGPartitionActor* InActor);
+
 #if WITH_EDITOR
 public:
 	/** Operations to schedule a later operation, which enables to delay bounds querying */
@@ -97,12 +109,37 @@ public:
 	/** Move all resources from sub actors to a new actor */
 	void ClearPCGLink(UPCGComponent* InComponent, const FBox& InBounds, AActor* InNewActor);
 
+	/** If the partition grid size change, call this to empty the Partition actors map */
+	void ResetPartitionActorsMap();
+
 private:
 	FPCGTaskId DelayProcessGraph(UPCGComponent* Component, bool bGenerate, bool bSave, bool bUseEmptyNewBounds);
 	FPCGTaskId ProcessGraph(UPCGComponent* Component, const FBox& InPreviousBounds, const FBox& InNewBounds, bool bGenerate, bool bSave);
 #endif // WITH_EDITOR
+
+	// Schedule multiple graphs
+	TArray<FPCGTaskId> ScheduleMultipleComponent(UPCGComponent* OriginalComponent, TSet<TObjectPtr<APCGPartitionActor>>& PartitionActors, const TArray<FPCGTaskId>& Dependencies);
+
+	// Schedule multiple cleanups
+	TArray<FPCGTaskId> ScheduleMultipleCleanup(UPCGComponent* OriginalComponent, TSet<TObjectPtr<APCGPartitionActor>>& PartitionActors, bool bRemoveComponents, const TArray<FPCGTaskId>& Dependencies);
+
+	/** Iterate other all the components which bounds intersect the box in param and call a callback. Thread safe */
+	void FindAllIntersectingComponents(const FBoxCenterAndExtent& InBounds, TFunction<void(UPCGComponent*)> InFunc) const;
+
+	/** Iterate other all the int coordinates given a box and call a callback. Thread safe */
+	void FindAllIntersectingPartitionActors(const FBox& InBounds, TFunction<void(APCGPartitionActor*)> InFunc) const;
 	
 private:
 	APCGWorldActor* PCGWorldActor = nullptr;
 	FPCGGraphExecutor* GraphExecutor = nullptr;
+
+	FPCGComponentOctree PCGComponentOctree;
+	TMap<TObjectPtr<const UPCGComponent>, FPCGComponentOctreeIDSharedRef> ComponentToIdMap;
+	mutable FRWLock PCGVolumeOctreeLock;
+
+	TMap<FIntVector, TObjectPtr<APCGPartitionActor>> PartitionActorsMap;
+	mutable FRWLock PartitionActorsMapLock;
+
+	TMap<TObjectPtr<const UPCGComponent>, TSet<TObjectPtr<APCGPartitionActor>>> ComponentToPartitionActorsMap;
+	mutable FRWLock ComponentToPartitionActorsMapLock;
 };
\ No newline at end of file
