PCG MetadataBreakVector Element

diff --git a/Engine/Plugins/Experimental/PCG/Source/PCG/Private/Elements/PCGMetadataBreakVector.cpp b/Engine/Plugins/Experimental/PCG/Source/PCG/Private/Elements/PCGMetadataBreakVector.cpp
new file mode 100644
--- /dev/null
+++ b/Engine/Plugins/Experimental/PCG/Source/PCG/Private/Elements/PCGMetadataBreakVector.cpp
@@ -0,0 +1,344 @@
+// Copyright Epic Games, Inc. All Rights Reserved.
+
+#include "Elements/PCGMetadataBreakVector.h"
+
+#include "PCGParamData.h"
+#include "Data/PCGSpatialData.h"
+#include "Helpers/PCGSettingsHelpers.h"
+#include "Metadata/PCGMetadata.h"
+#include "Metadata/PCGMetadataAttribute.h"
+#include "Metadata/PCGMetadataAttributeTpl.h"
+
+namespace PCGMetadataBreakVectorElement {
+	void DuplicateTaggedData(const FPCGTaggedData& InTaggedData, FPCGTaggedData& OutTaggedData, UPCGMetadata*& OutMetadata)
+	{
+		if (const UPCGSpatialData* SpatialInput = Cast<const UPCGSpatialData>(InTaggedData.Data))
+		{
+			UPCGSpatialData* NewSpatialData = Cast<UPCGSpatialData>(StaticDuplicateObject(SpatialInput, const_cast<UPCGSpatialData*>(SpatialInput), FName()));
+			NewSpatialData->InitializeFromData(SpatialInput);
+			OutTaggedData.Data = NewSpatialData;
+
+			OutMetadata = NewSpatialData->Metadata;
+		}
+		else if (const UPCGParamData* ParamsInput = Cast<const UPCGParamData>(InTaggedData.Data))
+		{
+			UPCGParamData* NewParamData = Cast<UPCGParamData>(StaticDuplicateObject(ParamsInput, const_cast<UPCGParamData*>(ParamsInput), FName()));
+			OutTaggedData.Data = NewParamData;
+
+			NewParamData->Metadata->Initialize(ParamsInput->Metadata);
+			OutMetadata = NewParamData->Metadata;
+		}
+	}
+
+	/** Creates a new double attribute, or clears the attribute if it already exists and is  a 'double' type */
+	void ClearOrCreateDoubleAttribute(UPCGMetadata* Metadata, const FName& DestinationAttribute, double DefaultValue, FPCGMetadataAttribute<double>*& OutAttribute)
+	{
+		if (!Metadata)
+		{
+			UE_LOG(LogPCG, Error, TEXT("Failed to create metadata"));
+			return;
+		}
+
+		if (Metadata->HasAttribute(DestinationAttribute))
+		{
+			UE_LOG(LogPCG, Warning, TEXT("Attribute %s already exists and has been overwritten"), *DestinationAttribute.ToString());
+
+			if (Metadata->GetConstAttribute(DestinationAttribute)->GetTypeId() == PCG::Private::MetadataTypes<double>::Id)
+			{
+				// TODO: If DestinationAttribute is an existing double attribute, it will retain its old DefaultValue instead of using the given DefaultValue
+				Metadata->ClearAttribute(DestinationAttribute);
+			}
+			else
+			{
+				Metadata->DeleteAttribute(DestinationAttribute);
+				Metadata->CreateDoubleAttribute(DestinationAttribute, DefaultValue, /*bAllowsInterpolation=*/true, /*bOverrideParent=*/false);
+			}
+		}
+		else
+		{
+			Metadata->CreateDoubleAttribute(DestinationAttribute, DefaultValue, /*bAllowsInterpolation=*/true, /*bOverrideParent=*/false);
+		}
+
+		OutAttribute = static_cast<FPCGMetadataAttribute<double>*>(Metadata->GetMutableAttribute(DestinationAttribute));
+	}
+
+	/** Copies the entry to value key relationship stored in the given Metadata, including its parents */
+	void CopyEntryToValueKeyMap(const UPCGMetadata* MetadataToCopy, const FPCGMetadataAttributeBase* AttributeToCopy, FPCGMetadataAttributeBase* OutAttribute)
+	{
+		if (!OutAttribute)
+		{
+			UE_LOG(LogPCG, Error, TEXT("Failed to create output attribute"));
+			return;
+		}
+
+		const PCGMetadataEntryKey EntryKeyCount = MetadataToCopy->GetItemCountForChild();
+		for (PCGMetadataEntryKey EntryKey = 0; EntryKey < EntryKeyCount; ++EntryKey)
+		{
+			const PCGMetadataValueKey ValueKey = AttributeToCopy->GetValueKey(EntryKey);
+			OutAttribute->SetValueFromValueKey(EntryKey, ValueKey);
+		}
+	}
+}
+
+TArray<FPCGPinProperties> UPCGMetadataBreakVectorSettings::InputPinProperties() const
+{
+	TArray<FPCGPinProperties> PinProperties;
+	PinProperties.Emplace(PCGMetadataBreakVectorConstants::ParamsLabel, EPCGDataType::Param);
+	PinProperties.Emplace(PCGMetadataBreakVectorConstants::SourceLabel, EPCGDataType::Any);
+
+	return PinProperties;
+}
+
+TArray<FPCGPinProperties> UPCGMetadataBreakVectorSettings::OutputPinProperties() const
+{
+	TArray<FPCGPinProperties> PinProperties;
+	PinProperties.Emplace(PCGMetadataBreakVectorConstants::XLabel, EPCGDataType::Any);
+	PinProperties.Emplace(PCGMetadataBreakVectorConstants::YLabel, EPCGDataType::Any);
+	PinProperties.Emplace(PCGMetadataBreakVectorConstants::ZLabel, EPCGDataType::Any);
+	PinProperties.Emplace(PCGMetadataBreakVectorConstants::WLabel, EPCGDataType::Any);
+
+	return PinProperties;
+}
+
+FPCGElementPtr UPCGMetadataBreakVectorSettings::CreateElement() const
+{
+	return MakeShared<FPCGMetadataBreakVectorElement>();
+}
+
+bool FPCGMetadataBreakVectorElement::ExecuteInternal(FPCGContext* Context) const
+{
+	TRACE_CPUPROFILER_EVENT_SCOPE(FPCGMetadataBreakVectorElement::Execute);
+
+	const UPCGMetadataBreakVectorSettings* Settings = Context->GetInputSettings<UPCGMetadataBreakVectorSettings>();
+	check(Settings);
+
+	const TArray<FPCGTaggedData>& Inputs = Context->InputData.TaggedData;
+	TArray<FPCGTaggedData>& Outputs = Context->OutputData.TaggedData;
+
+	const TArray<FPCGTaggedData>& ParamData = Context->InputData.GetParamsByPin(PCGMetadataBreakVectorConstants::ParamsLabel);
+	UPCGParamData* Params = nullptr;
+
+	// TODO: this only captures the first Param
+	for (const FPCGTaggedData& TaggedDatum : ParamData)
+	{
+		if (UPCGParamData* ParamDatum = Cast<UPCGParamData>(TaggedDatum.Data))
+		{
+			Params = ParamDatum;
+			break;
+		}
+	}
+
+	const FName SourceAttributeName = PCG_GET_OVERRIDEN_VALUE(Settings, SourceAttributeName, Params);
+
+	// These exist to facilitate unit testing of the MetadataBreakVector element in isolation, when it otherwise would not have any pin connections to generate output
+	bool bXPinConnected = Settings->bForceConnectX;
+	bool bYPinConnected = Settings->bForceConnectY;
+	bool bZPinConnected = Settings->bForceConnectZ;
+	bool bWPinConnected = Settings->bForceConnectW;
+
+	if (Context->Node)
+	{
+		bXPinConnected |= Context->Node->IsOutputPinConnected(PCGMetadataBreakVectorConstants::XLabel);
+		bYPinConnected |= Context->Node->IsOutputPinConnected(PCGMetadataBreakVectorConstants::YLabel);
+		bZPinConnected |= Context->Node->IsOutputPinConnected(PCGMetadataBreakVectorConstants::ZLabel);
+		bWPinConnected |= Context->Node->IsOutputPinConnected(PCGMetadataBreakVectorConstants::WLabel);
+	}
+
+	for (const FPCGTaggedData& Input : Inputs)
+	{
+		if (!Input.Data)
+		{
+			PCGE_LOG(Error, "Unable to get data from input");
+			continue;
+		}
+
+		UPCGMetadata* SourceMetadata = nullptr;
+		if (const UPCGSpatialData* SpatialInput = Cast<const UPCGSpatialData>(Input.Data))
+		{
+			SourceMetadata = SpatialInput->Metadata;
+		}
+		else if (const UPCGParamData* ParamsInput = Cast<const UPCGParamData>(Input.Data))
+		{
+			SourceMetadata = ParamsInput->Metadata;
+		}
+		else // If the Data type does not handle Metadata, forward it into the valid pins
+		{
+			if (bXPinConnected)
+			{
+				FPCGTaggedData& Data = Outputs.Add_GetRef(Input);
+				Data.Pin = PCGMetadataBreakVectorConstants::XLabel;
+			}
+
+			if (bYPinConnected)
+			{
+				FPCGTaggedData& Data = Outputs.Add_GetRef(Input);
+				Data.Pin = PCGMetadataBreakVectorConstants::YLabel;
+			}
+
+			if (bZPinConnected)
+			{
+				FPCGTaggedData& Data = Outputs.Add_GetRef(Input);
+				Data.Pin = PCGMetadataBreakVectorConstants::ZLabel;
+			}
+
+			if (bWPinConnected)
+			{
+				FPCGTaggedData& Data = Outputs.Add_GetRef(Input);
+				Data.Pin = PCGMetadataBreakVectorConstants::WLabel;
+			}
+
+			continue;
+		}
+
+		if (!SourceMetadata)
+		{
+			PCGE_LOG(Warning, "Invalid metadata");
+			continue;
+		}
+
+		const FName LocalSourceAttributeName = (SourceAttributeName == NAME_None) ? SourceMetadata->GetSingleAttributeNameOrNone() : SourceAttributeName;
+		FName DestinationAttributeForX = NAME_None;
+		FName DestinationAttributeForY = NAME_None;
+		FName DestinationAttributeForZ = NAME_None;
+		FName DestinationAttributeForW = NAME_None;
+
+		if (LocalSourceAttributeName != NAME_None)
+		{
+			DestinationAttributeForX = FName(LocalSourceAttributeName.ToString() + ".X");
+			DestinationAttributeForY = FName(LocalSourceAttributeName.ToString() + ".Y");
+			DestinationAttributeForZ = FName(LocalSourceAttributeName.ToString() + ".Z");
+			DestinationAttributeForW = FName(LocalSourceAttributeName.ToString() + ".W");
+		}
+
+		const FPCGMetadataAttributeBase* SourceAttribute = SourceMetadata->GetConstAttribute(LocalSourceAttributeName);
+		if (!SourceAttribute)
+		{
+			PCGE_LOG(Error, "Attribute %s does not exist", *LocalSourceAttributeName.ToString());
+			continue;
+		}
+
+		if (SourceAttribute->GetTypeId() != PCG::Private::MetadataTypes<FVector>::Id 
+			&& SourceAttribute->GetTypeId() != PCG::Private::MetadataTypes<FVector4>::Id
+			&& SourceAttribute->GetTypeId() != PCG::Private::MetadataTypes<FRotator>::Id)
+		{
+			PCGE_LOG(Error, "Attribute %s is not a breakable type", *LocalSourceAttributeName.ToString());
+			continue;
+		}
+
+		FVector4 DefaultValue;
+
+		if (SourceAttribute->GetTypeId() == PCG::Private::MetadataTypes<FVector>::Id)
+		{
+			DefaultValue = static_cast<const FPCGMetadataAttribute<FVector>*>(SourceAttribute)->GetValue(PCGDefaultValueKey);
+		}
+		else if (SourceAttribute->GetTypeId() == PCG::Private::MetadataTypes<FVector4>::Id)
+		{
+			DefaultValue = static_cast<const FPCGMetadataAttribute<FVector4>*>(SourceAttribute)->GetValue(PCGDefaultValueKey);
+		}
+		else //if (SourceAttribute->GetTypeId() == PCG::Private::MetadataTypes<FRotator>::Id)
+		{
+			const FRotator Rotator = static_cast<const FPCGMetadataAttribute<FRotator>*>(SourceAttribute)->GetValue(PCGDefaultValueKey);
+			DefaultValue.X = Rotator.Roll;
+			DefaultValue.Y = Rotator.Pitch;
+			DefaultValue.Z = Rotator.Yaw;
+		}
+
+		FPCGMetadataAttribute<double>* AttributeX = nullptr;
+		FPCGMetadataAttribute<double>* AttributeY = nullptr;
+		FPCGMetadataAttribute<double>* AttributeZ = nullptr;
+		FPCGMetadataAttribute<double>* AttributeW = nullptr;
+
+		if (bXPinConnected)
+		{
+			FPCGTaggedData& OutputX = Outputs.Add_GetRef(Input);
+			OutputX.Pin = PCGMetadataBreakVectorConstants::XLabel;
+
+			UPCGMetadata* OutMetadata = nullptr;
+
+			PCGMetadataBreakVectorElement::DuplicateTaggedData(Input, OutputX, OutMetadata);
+			PCGMetadataBreakVectorElement::ClearOrCreateDoubleAttribute(OutMetadata, DestinationAttributeForX, DefaultValue.X, AttributeX);
+			PCGMetadataBreakVectorElement::CopyEntryToValueKeyMap(SourceMetadata, SourceAttribute, AttributeX);
+		}
+
+		if (bYPinConnected)
+		{
+			FPCGTaggedData& OutputY = Outputs.Add_GetRef(Input);
+			OutputY.Pin = PCGMetadataBreakVectorConstants::YLabel;
+
+			UPCGMetadata* OutMetadata = nullptr;
+
+			PCGMetadataBreakVectorElement::DuplicateTaggedData(Input, OutputY, OutMetadata);
+			PCGMetadataBreakVectorElement::ClearOrCreateDoubleAttribute(OutMetadata, DestinationAttributeForY, DefaultValue.Y, AttributeY);
+			PCGMetadataBreakVectorElement::CopyEntryToValueKeyMap(SourceMetadata, SourceAttribute, AttributeY);
+		}
+
+		if (bZPinConnected)
+		{
+			FPCGTaggedData& OutputZ = Outputs.Add_GetRef(Input);
+			OutputZ.Pin = PCGMetadataBreakVectorConstants::ZLabel;
+
+			UPCGMetadata* OutMetadata = nullptr;
+
+			PCGMetadataBreakVectorElement::DuplicateTaggedData(Input, OutputZ, OutMetadata);
+			PCGMetadataBreakVectorElement::ClearOrCreateDoubleAttribute(OutMetadata, DestinationAttributeForZ, DefaultValue.Z, AttributeZ);
+			PCGMetadataBreakVectorElement::CopyEntryToValueKeyMap(SourceMetadata, SourceAttribute, AttributeZ);
+		}
+
+		if (bWPinConnected && SourceAttribute->GetTypeId() == PCG::Private::MetadataTypes<FVector4>::Id)
+		{
+			FPCGTaggedData& OutputW = Outputs.Add_GetRef(Input);
+			OutputW.Pin = PCGMetadataBreakVectorConstants::WLabel;
+
+			UPCGMetadata* OutMetadata = nullptr;
+
+			PCGMetadataBreakVectorElement::DuplicateTaggedData(Input, OutputW, OutMetadata);
+			PCGMetadataBreakVectorElement::ClearOrCreateDoubleAttribute(OutMetadata, DestinationAttributeForW, DefaultValue.W, AttributeW);
+			PCGMetadataBreakVectorElement::CopyEntryToValueKeyMap(SourceMetadata, SourceAttribute, AttributeW);
+		}
+
+		// Copy all Value pairs from the parent hierarchy into our new attributes. Assumes that adding values in order of the parent hierarchy will result in identical ValueKey->Value mappings.
+		const PCGMetadataValueKey NumValueKeys = SourceAttribute->GetValueKeyOffsetForChild();
+		for (PCGMetadataValueKey ValueKey = 0; ValueKey < NumValueKeys; ++ValueKey)
+		{
+			FVector4 Value;
+
+			if (SourceAttribute->GetTypeId() == PCG::Private::MetadataTypes<FVector>::Id)
+			{
+				Value = static_cast<const FPCGMetadataAttribute<FVector>*>(SourceAttribute)->GetValue(ValueKey);
+			}
+			else if (SourceAttribute->GetTypeId() == PCG::Private::MetadataTypes<FVector4>::Id)
+			{
+				Value = static_cast<const FPCGMetadataAttribute<FVector4>*>(SourceAttribute)->GetValue(ValueKey);
+			}
+			else //if (SourceAttribute->GetTypeId() == PCG::Private::MetadataTypes<FRotator>::Id)
+			{
+				const FRotator Rotator = static_cast<const FPCGMetadataAttribute<FRotator>*>(SourceAttribute)->GetValue(ValueKey);
+				Value.X = Rotator.Roll;
+				Value.Y = Rotator.Pitch;
+				Value.Z = Rotator.Yaw;
+			}
+
+			if (AttributeX)
+			{
+				AttributeX->AddValue(Value.X);
+			}
+
+			if (AttributeY)
+			{
+				AttributeY->AddValue(Value.Y);
+			}
+
+			if (AttributeZ)
+			{
+				AttributeZ->AddValue(Value.Z);
+			}
+
+			if (AttributeW)
+			{
+				AttributeW->AddValue(Value.W);
+			}
+		}
+	}
+
+	return true;
+}
diff --git a/Engine/Plugins/Experimental/PCG/Source/PCG/Private/Metadata/PCGMetadata.cpp b/Engine/Plugins/Experimental/PCG/Source/PCG/Private/Metadata/PCGMetadata.cpp
--- a/Engine/Plugins/Experimental/PCG/Source/PCG/Private/Metadata/PCGMetadata.cpp
+++ b/Engine/Plugins/Experimental/PCG/Source/PCG/Private/Metadata/PCGMetadata.cpp
@@ -307,6 +307,11 @@ void UPCGMetadata::CreateFloatAttribute(FName AttributeName, float DefaultValue,
 	CreateAttribute<float>(AttributeName, DefaultValue, bAllowsInterpolation, bOverrideParent);
 }
 
+void UPCGMetadata::CreateDoubleAttribute(FName AttributeName, double DefaultValue, bool bAllowsInterpolation, bool bOverrideParent)
+{
+	CreateAttribute<double>(AttributeName, DefaultValue, bAllowsInterpolation, bOverrideParent);
+}
+
 void UPCGMetadata::CreateVectorAttribute(FName AttributeName, FVector DefaultValue, bool bAllowsInterpolation, bool bOverrideParent)
 {
 	CreateAttribute<FVector>(AttributeName, DefaultValue, bAllowsInterpolation, bOverrideParent);
diff --git a/Engine/Plugins/Experimental/PCG/Source/PCG/Private/Metadata/PCGMetadataAccessor.cpp b/Engine/Plugins/Experimental/PCG/Source/PCG/Private/Metadata/PCGMetadataAccessor.cpp
--- a/Engine/Plugins/Experimental/PCG/Source/PCG/Private/Metadata/PCGMetadataAccessor.cpp
+++ b/Engine/Plugins/Experimental/PCG/Source/PCG/Private/Metadata/PCGMetadataAccessor.cpp
@@ -109,6 +109,16 @@ void UPCGMetadataAccessorHelpers::SetFloatAttributeByMetadataKey(int64& Key, UPC
 	SetAttribute(Key, Metadata, AttributeName, Value);
 }
 
+double UPCGMetadataAccessorHelpers::GetDoubleAttributeByMetadataKey(int64 Key, const UPCGMetadata* Metadata, FName AttributeName)
+{
+	return GetAttribute<double>(Key, Metadata, AttributeName);
+}
+
+void UPCGMetadataAccessorHelpers::SetDoubleAttributeByMetadataKey(int64& Key, UPCGMetadata* Metadata, FName AttributeName, double Value)
+{
+	SetAttribute(Key, Metadata, AttributeName, Value);
+}
+
 FVector UPCGMetadataAccessorHelpers::GetVectorAttributeByMetadataKey(int64 Key, const UPCGMetadata* Metadata, FName AttributeName)
 {
 	return GetAttribute<FVector>(Key, Metadata, AttributeName);
@@ -254,6 +264,16 @@ void UPCGMetadataAccessorHelpers::SetFloatAttribute(FPCGPoint& Point, UPCGMetada
 	SetAttribute(Point.MetadataEntry, Metadata, AttributeName, Value);
 }
 
+double UPCGMetadataAccessorHelpers::GetDoubleAttribute(const FPCGPoint& Point, const UPCGMetadata* Metadata, FName AttributeName)
+{
+	return GetAttribute<double>(Point.MetadataEntry, Metadata, AttributeName);
+}
+
+void UPCGMetadataAccessorHelpers::SetDoubleAttribute(FPCGPoint& Point, UPCGMetadata* Metadata, FName AttributeName, double Value)
+{
+	SetAttribute(Point.MetadataEntry, Metadata, AttributeName, Value);
+}
+
 FVector UPCGMetadataAccessorHelpers::GetVectorAttribute(const FPCGPoint& Point, const UPCGMetadata* Metadata, FName AttributeName)
 {
 	return GetAttribute<FVector>(Point.MetadataEntry, Metadata, AttributeName);
diff --git a/Engine/Plugins/Experimental/PCG/Source/PCG/Private/Tests/Elements/PCGMetadataBreakVectorTest.cpp b/Engine/Plugins/Experimental/PCG/Source/PCG/Private/Tests/Elements/PCGMetadataBreakVectorTest.cpp
new file mode 100644
--- /dev/null
+++ b/Engine/Plugins/Experimental/PCG/Source/PCG/Private/Tests/Elements/PCGMetadataBreakVectorTest.cpp
@@ -0,0 +1,318 @@
+// Copyright Epic Games, Inc. All Rights Reserved.
+
+#include "Tests/Determinism/PCGDeterminismTestsCommon.h"
+#include "Tests/PCGTestsCommon.h"
+
+#include "PCGComponent.h"
+#include "PCGData.h"
+#include "PCGParamData.h"
+
+#include "Data/PCGPointData.h"
+#include "Elements/PCGMetadataBreakVector.h"
+#include "Metadata/PCGMetadata.h"
+#include "Metadata/PCGMetadataAttribute.h"
+#include "Metadata/PCGMetadataAttributeTpl.h"
+
+#if WITH_EDITOR
+
+IMPLEMENT_CUSTOM_SIMPLE_AUTOMATION_TEST(FPCGMetadataBreakVectorTest, FPCGTestBaseClass, "pcg.tests.Metadata.BreakVector", PCGTestsCommon::TestFlags)
+
+namespace
+{
+	enum class EPCGComponentToCheck : uint8
+	{
+		X,
+		Z,
+		Y,
+		W
+	};
+}
+
+bool FPCGMetadataBreakVectorTest::RunTest(const FString& Parameters)
+{
+	PCGTestsCommon::FTestData TestData(PCGDeterminismTests::Defaults::Seed);
+	PCGDeterminismTests::GenerateSettings<UPCGMetadataBreakVectorSettings>(TestData);
+	UPCGMetadataBreakVectorSettings* Settings = CastChecked<UPCGMetadataBreakVectorSettings>(TestData.Settings);
+	FPCGElementPtr MetadataBreakVectorElement = TestData.Settings->GetElement();
+
+	FPCGTaggedData& SpatialTaggedData = TestData.InputData.TaggedData.Emplace_GetRef(FPCGTaggedData());
+	TObjectPtr<UPCGPointData> PointData = PCGTestsCommon::CreateRandomPointData(10, TestData.Seed);
+	SpatialTaggedData.Data = PointData;
+
+	FPCGTaggedData& ParamTaggedData = TestData.InputData.TaggedData.Emplace_GetRef(FPCGTaggedData());
+	TObjectPtr<UPCGParamData> ParamData = PCGTestsCommon::CreateEmptyParamData();
+	ParamTaggedData.Data = ParamData;
+
+	const FName Vec3Attribute = TEXT("Vec3");
+	const FName Vec4Attribute = TEXT("Vec4");
+	const FName RotatorAttribute = TEXT("Rotator");
+	const FName FloatAttribute = TEXT("Float");
+
+	const bool bAllowsInterpolation = false;
+	const bool bOverrideParent = false;
+
+	PointData->Metadata->CreateVectorAttribute(Vec3Attribute, FVector::Zero(), bAllowsInterpolation, bOverrideParent);
+	PointData->Metadata->CreateVector4Attribute(Vec4Attribute, FVector4::Zero(), bAllowsInterpolation, bOverrideParent);
+	PointData->Metadata->CreateRotatorAttribute(RotatorAttribute, FRotator::ZeroRotator, bAllowsInterpolation, bOverrideParent);
+	PointData->Metadata->CreateFloatAttribute(FloatAttribute, 0, bAllowsInterpolation, bOverrideParent);
+
+	TArray<FPCGPoint>& SourcePoints = PointData->GetMutablePoints();
+
+	// Randomize attribute values, and leave the second half to default value
+	FRandomStream RandomSource(TestData.Seed);
+	for (int I = 0; I < SourcePoints.Num() / 2; ++I)
+	{
+		UPCGMetadataAccessorHelpers::SetVectorAttribute(SourcePoints[I], PointData->Metadata, Vec3Attribute, RandomSource.VRand());
+		UPCGMetadataAccessorHelpers::SetVector4Attribute(SourcePoints[I], PointData->Metadata, Vec4Attribute, FVector4(RandomSource.VRand(), RandomSource.FRand()));
+		UPCGMetadataAccessorHelpers::SetRotatorAttribute(SourcePoints[I], PointData->Metadata, RotatorAttribute, FRotator(RandomSource.FRand(), RandomSource.FRand(), RandomSource.FRand()));
+		UPCGMetadataAccessorHelpers::SetFloatAttribute(SourcePoints[I], PointData->Metadata, FloatAttribute, RandomSource.FRand());
+	}
+
+	ParamData->Metadata->InitializeAsCopy(PointData->Metadata);
+
+	auto ValidateMetadataBreakVector = [this, &TestData, MetadataBreakVectorElement, Settings]() -> bool
+	{
+		TUniquePtr<FPCGContext> Context = MakeUnique<FPCGContext>(*MetadataBreakVectorElement->Initialize(TestData.InputData, TestData.TestPCGComponent, nullptr));
+		Context->NumAvailableTasks = 1;
+
+		while (!MetadataBreakVectorElement->Execute(Context.Get()))
+		{}
+
+		const TArray<FPCGTaggedData>& Inputs = Context->InputData.TaggedData;
+
+		const TArray<FPCGTaggedData> OutputsX = Context->OutputData.GetInputsByPin(PCGMetadataBreakVectorConstants::XLabel);
+		const TArray<FPCGTaggedData> OutputsY = Context->OutputData.GetInputsByPin(PCGMetadataBreakVectorConstants::YLabel);
+		const TArray<FPCGTaggedData> OutputsZ = Context->OutputData.GetInputsByPin(PCGMetadataBreakVectorConstants::ZLabel);
+		const TArray<FPCGTaggedData> OutputsW = Context->OutputData.GetInputsByPin(PCGMetadataBreakVectorConstants::WLabel);
+
+		const FName DestinationAttributeForX = FName(Settings->SourceAttributeName.ToString() + ".X");
+		const FName DestinationAttributeForY = FName(Settings->SourceAttributeName.ToString() + ".Y");
+		const FName DestinationAttributeForZ = FName(Settings->SourceAttributeName.ToString() + ".Z");
+		const FName DestinationAttributeForW = FName(Settings->SourceAttributeName.ToString() + ".W");
+
+		bool bTestPassed = true;
+
+		for (int DataIndex = 0; DataIndex < Inputs.Num(); ++DataIndex)
+		{
+			const FPCGTaggedData& Input = Inputs[DataIndex];
+
+			UPCGMetadata* SourceMetadata = nullptr; 
+			if (const UPCGSpatialData* SpatialInput = Cast<const UPCGSpatialData>(Input.Data))
+			{
+				SourceMetadata = SpatialInput->Metadata;
+			}
+			else if (const UPCGParamData* ParamsInput = Cast<const UPCGParamData>(Input.Data))
+			{
+				SourceMetadata = ParamsInput->Metadata;
+			}
+			else
+			{
+				continue;
+			}
+
+			check(SourceMetadata);
+
+			const FPCGMetadataAttributeBase* SourceAttributeBase = SourceMetadata->GetConstAttribute(Settings->SourceAttributeName);
+			check(SourceAttributeBase);
+
+			auto ValidateComponentOutput = [&](const FPCGTaggedData& Output, const FName& OutAttributeName, EPCGComponentToCheck ComponentToCheck)
+			{
+				if (!TestTrue("Valid output data", Output.Data != nullptr))
+				{
+					bTestPassed = false;
+					return;
+				}
+
+				UPCGMetadata* OutMetadata = nullptr;
+				if (const UPCGSpatialData* SpatialInput = Cast<const UPCGSpatialData>(Output.Data))
+				{
+					OutMetadata = SpatialInput->Metadata;
+				}
+				else if (const UPCGParamData* ParamsInput = Cast<const UPCGParamData>(Output.Data))
+				{
+					OutMetadata = ParamsInput->Metadata;
+				}
+				else
+				{
+					return;
+				}
+
+				if (!TestNotNull("Valid output metadata", OutMetadata))
+				{
+					bTestPassed = false;
+					return;
+				}
+
+				const FPCGMetadataAttributeBase* OutAttributeBase = OutMetadata->GetConstAttribute(OutAttributeName);
+
+				if (!TestNotNull("Valid output attribute", OutAttributeBase))
+				{
+					bTestPassed = false;
+					return;
+				}
+
+				if (!TestEqual("Output attribute is a valid type (double)", OutAttributeBase->GetTypeId(), PCG::Private::MetadataTypes<double>::Id))
+				{
+					bTestPassed = false;
+					return;
+				}
+
+				const FPCGMetadataAttribute<double>* OutAttribute = static_cast<const FPCGMetadataAttribute<double>*>(OutAttributeBase);
+				check(OutAttribute);
+
+				const PCGMetadataEntryKey InEntryKeyCount = SourceMetadata->GetItemCountForChild();
+				const PCGMetadataEntryKey OutEntryKeyCount = OutMetadata->GetItemCountForChild();
+
+				if (!TestEqual("Identical EntryKey counts", InEntryKeyCount, OutEntryKeyCount))
+				{
+					bTestPassed = false;
+					return;
+				}
+
+				const PCGMetadataValueKey InValueKeyCount = SourceAttributeBase->GetValueKeyOffsetForChild();
+				const PCGMetadataValueKey OutValueKeyCount = OutAttribute->GetValueKeyOffsetForChild();
+
+				if (!TestEqual("Identical ValueKey counts", InValueKeyCount, OutValueKeyCount))
+				{
+					bTestPassed = false;
+					return;
+				}
+
+				for (PCGMetadataEntryKey EntryKey = 0; EntryKey < InEntryKeyCount; ++EntryKey)
+				{
+					const PCGMetadataValueKey InValueKey = SourceAttributeBase->GetValueKey(EntryKey);
+					const PCGMetadataValueKey OutValueKey = OutAttribute->GetValueKey(EntryKey);
+
+					bTestPassed &= TestEqual("Identical value keys", InValueKey, OutValueKey);
+
+					FVector4 SourceValue;
+					if (SourceAttributeBase->GetTypeId() == PCG::Private::MetadataTypes<FVector>::Id)
+					{
+						SourceValue = static_cast<const FPCGMetadataAttribute<FVector>*>(SourceAttributeBase)->GetValue(InValueKey);
+					}
+					else if (SourceAttributeBase->GetTypeId() == PCG::Private::MetadataTypes<FVector4>::Id)
+					{
+						SourceValue = static_cast<const FPCGMetadataAttribute<FVector4>*>(SourceAttributeBase)->GetValue(InValueKey);
+					}
+					else //if (SourceAttributeBase->GetTypeId() == PCG::Private::MetadataTypes<FRotator>::Id)
+					{
+						const FRotator Rotator = static_cast<const FPCGMetadataAttribute<FRotator>*>(SourceAttributeBase)->GetValue(InValueKey);
+						SourceValue.X = Rotator.Roll;
+						SourceValue.Y = Rotator.Pitch;
+						SourceValue.Z = Rotator.Yaw;
+					}
+
+					const double OutValue = OutAttribute->GetValue(OutValueKey);
+
+					if (ComponentToCheck == EPCGComponentToCheck::X)
+					{
+						bTestPassed &= TestEqual("Identical values", SourceValue.X, OutValue);
+					}
+					else if (ComponentToCheck == EPCGComponentToCheck::Y)
+					{
+						bTestPassed &= TestEqual("Identical values", SourceValue.Y, OutValue);
+					}
+					else if (ComponentToCheck == EPCGComponentToCheck::Z)
+					{
+						bTestPassed &= TestEqual("Identical values", SourceValue.Z, OutValue);
+					}
+					else if (ComponentToCheck == EPCGComponentToCheck::W)
+					{
+						bTestPassed &= TestEqual("Identical values", SourceValue.W, OutValue);
+					}
+				}
+			};
+
+			if (Settings->bForceConnectX)
+			{
+				if (TestTrue("Appropriate number of outputs generated", OutputsX.Num() > DataIndex))
+				{
+					ValidateComponentOutput(OutputsX[DataIndex], DestinationAttributeForX, EPCGComponentToCheck::X);
+				}
+				else
+				{
+					bTestPassed = false;
+				}
+			}
+
+			if (Settings->bForceConnectY)
+			{
+				if (TestTrue("Appropriate number of outputs generated", OutputsY.Num() > DataIndex))
+				{
+					ValidateComponentOutput(OutputsY[DataIndex], DestinationAttributeForY, EPCGComponentToCheck::Y);
+				}
+				else
+				{
+					bTestPassed = false;
+				}
+			}
+
+			if (Settings->bForceConnectZ)
+			{
+				if (TestTrue("Appropriate number of outputs generated", OutputsZ.Num() > DataIndex))
+				{
+					ValidateComponentOutput(OutputsZ[DataIndex], DestinationAttributeForZ, EPCGComponentToCheck::Z);
+				}
+				else
+				{
+					bTestPassed = false;
+				}
+			}
+
+			if (Settings->bForceConnectW && SourceAttributeBase->GetTypeId() == PCG::Private::MetadataTypes<FVector4>::Id)
+			{
+				if (TestTrue("Appropriate number of outputs generated", OutputsW.Num() > DataIndex))
+				{
+					ValidateComponentOutput(OutputsW[DataIndex], DestinationAttributeForW, EPCGComponentToCheck::W);
+				}
+				else
+				{
+					bTestPassed = false;
+				}
+			}
+		}
+
+		return bTestPassed;
+	};
+
+	bool bTestPassed = true;
+
+	Settings->bForceConnectX = true;
+	Settings->bForceConnectY = true;
+	Settings->bForceConnectZ = true;
+
+	// Test FVector source attribute
+	{
+		Settings->SourceAttributeName = Vec3Attribute;
+		bTestPassed &= ValidateMetadataBreakVector();
+	}
+
+	// Test FVector4 source attribute
+	{
+		Settings->SourceAttributeName = Vec4Attribute;
+		Settings->bForceConnectW = true;
+		bTestPassed &= ValidateMetadataBreakVector();
+		Settings->bForceConnectW = false;
+	}
+
+	// Test FRotator source attribute
+	{
+		Settings->SourceAttributeName = RotatorAttribute;
+		bTestPassed &= ValidateMetadataBreakVector();
+	}
+
+	// Test invalid source attribute type
+	{
+		Settings->SourceAttributeName = FloatAttribute;
+		Settings->bForceConnectX = false;
+		Settings->bForceConnectY = false;
+		Settings->bForceConnectZ = false;
+
+		AddExpectedError(TEXT("Attribute ") + FloatAttribute.ToString() + TEXT(" is not a breakable type"), EAutomationExpectedErrorFlags::Contains, 2);
+		bTestPassed &= ValidateMetadataBreakVector();
+	}
+
+	return bTestPassed;
+}
+
+#endif // WITH_EDITOR
diff --git a/Engine/Plugins/Experimental/PCG/Source/PCG/Public/Elements/PCGMetadataBreakVector.h b/Engine/Plugins/Experimental/PCG/Source/PCG/Public/Elements/PCGMetadataBreakVector.h
new file mode 100644
--- /dev/null
+++ b/Engine/Plugins/Experimental/PCG/Source/PCG/Public/Elements/PCGMetadataBreakVector.h
@@ -0,0 +1,59 @@
+// Copyright Epic Games, Inc. All Rights Reserved.
+
+#pragma once
+
+#include "PCGElement.h"
+#include "PCGSettings.h"
+
+#include "PCGMetadataBreakVector.generated.h"
+
+namespace PCGMetadataBreakVectorConstants
+{
+	const FName ParamsLabel = TEXT("Params");
+	const FName SourceLabel = TEXT("Source");
+	const FName XLabel = TEXT("X");
+	const FName YLabel = TEXT("Y");
+	const FName ZLabel = TEXT("Z");
+	const FName WLabel = TEXT("W");
+}
+
+UCLASS(BlueprintType, ClassGroup = (Procedural))
+class PCG_API UPCGMetadataBreakVectorSettings : public UPCGSettings
+{
+	GENERATED_BODY()
+
+public:
+	//~Begin UPCGSettings interface
+#if WITH_EDITOR
+	virtual FName GetDefaultNodeName() const override { return FName(TEXT("MetadataBreakVectorNode")); }
+	virtual EPCGSettingsType GetType() const override { return EPCGSettingsType::Metadata; }
+#endif
+
+	virtual TArray<FPCGPinProperties> InputPinProperties() const override;
+	virtual TArray<FPCGPinProperties> OutputPinProperties() const override;
+
+protected:
+	virtual FPCGElementPtr CreateElement() const override;
+	//~End UPCGSettings interface
+
+public:
+	UPROPERTY(BlueprintReadWrite, EditAnywhere, Category = Settings)
+	FName SourceAttributeName = NAME_None;
+
+	// These exist to facilitate unit testing of the MetadataBreakVector element in isolation, when it otherwise would not have any pin connections to generate output
+	bool bForceConnectX = false;
+	bool bForceConnectY = false;
+	bool bForceConnectZ = false;
+	bool bForceConnectW = false;
+};
+
+class FPCGMetadataBreakVectorElement : public FSimplePCGElement
+{
+public:
+	// MetadataBreakVector relies on StaticDuplicateObject, so we cannot run outside of the main thread
+	virtual bool CanExecuteOnlyOnMainThread(const UPCGSettings* InSettings) const override { return true; }
+	virtual bool IsCacheable(const UPCGSettings* InSettings) const override { return false; }
+
+protected:
+	virtual bool ExecuteInternal(FPCGContext* Context) const override;
+};
diff --git a/Engine/Plugins/Experimental/PCG/Source/PCG/Public/Metadata/PCGMetadata.h b/Engine/Plugins/Experimental/PCG/Source/PCG/Public/Metadata/PCGMetadata.h
--- a/Engine/Plugins/Experimental/PCG/Source/PCG/Public/Metadata/PCGMetadata.h
+++ b/Engine/Plugins/Experimental/PCG/Source/PCG/Public/Metadata/PCGMetadata.h
@@ -72,6 +72,9 @@ public:
 	UFUNCTION(BlueprintCallable, Category = "PCG|Metadata")
 	void CreateFloatAttribute(FName AttributeName, float DefaultValue, bool bAllowsInterpolation, bool bOverrideParent = true);
 
+	UFUNCTION(BlueprintCallable, Category = "PCG|Metadata")
+	void CreateDoubleAttribute(FName AttributeName, double DefaultValue, bool bAllowsInterpolation, bool bOverrideParent = true);
+
 	UFUNCTION(BlueprintCallable, Category = "PCG|Metadata")
 	void CreateVectorAttribute(FName AttributeName, FVector DefaultValue, bool bAllowsInterpolation, bool bOverrideParent = true);
 
@@ -176,6 +179,8 @@ public:
 	UFUNCTION(BlueprintCallable, Category = "PCG|Metadata")
 	void AccumulatePointWeightedAttributes(const FPCGPoint& InPoint, const UPCGMetadata* InMetadata, float Weight, bool bSetNonInterpolableAttributes, UPARAM(ref) FPCGPoint& OutPoint);
 
+	int64 GetItemCountForChild() const;
+
 protected:
 	template<typename T>
 	FPCGMetadataAttributeBase* CreateAttribute(FName AttributeName, const T& DefaultValue, bool bAllowsInterpolation, bool bOverrideParent);
@@ -183,7 +188,6 @@ protected:
 	FPCGMetadataAttributeBase* CopyAttribute(FName AttributeToCopy, FName NewAttributeName, bool bKeepParent, bool bCopyEntries, bool bCopyValues);
 
 	bool ParentHasAttribute(FName AttributeName) const;
-	int64 GetItemCountForChild() const;
 
 	void AddAttributeInternal(FName AttributeName, FPCGMetadataAttributeBase* Attribute);
 	void RemoveAttributeInternal(FName AttributeName);
diff --git a/Engine/Plugins/Experimental/PCG/Source/PCG/Public/Metadata/PCGMetadataAccessor.h b/Engine/Plugins/Experimental/PCG/Source/PCG/Public/Metadata/PCGMetadataAccessor.h
--- a/Engine/Plugins/Experimental/PCG/Source/PCG/Public/Metadata/PCGMetadataAccessor.h
+++ b/Engine/Plugins/Experimental/PCG/Source/PCG/Public/Metadata/PCGMetadataAccessor.h
@@ -30,6 +30,12 @@ public:
 	UFUNCTION(BlueprintCallable, Category = "PCG|Metadata")
 	static void SetFloatAttributeByMetadataKey(UPARAM(ref) int64& Key, UPCGMetadata* Metadata, FName AttributeName, float Value);
 
+	UFUNCTION(BlueprintCallable, Category = "PCG|Metadata")
+	static double GetDoubleAttributeByMetadataKey(int64 Key, const UPCGMetadata* Metadata, FName AttributeName);
+
+	UFUNCTION(BlueprintCallable, Category = "PCG|Metadata")
+	static void SetDoubleAttributeByMetadataKey(UPARAM(ref) int64& Key, UPCGMetadata* Metadata, FName AttributeName, double Value);
+
 	UFUNCTION(BlueprintCallable, Category = "PCG|Metadata")
 	static FVector GetVectorAttributeByMetadataKey(int64 Key, const UPCGMetadata* Metadata, FName AttributeName);
 
@@ -96,6 +102,12 @@ public:
 	UFUNCTION(BlueprintCallable, Category = "PCG|Metadata", meta = (ScriptMethod))
 	static void SetFloatAttribute(UPARAM(ref) FPCGPoint& Point, UPCGMetadata* Metadata, FName AttributeName, float Value);
 
+	UFUNCTION(BlueprintCallable, Category = "PCG|Metadata", meta = (ScriptMethod))
+	static double GetDoubleAttribute(const FPCGPoint& Point, const UPCGMetadata* Metadata, FName AttributeName);
+
+	UFUNCTION(BlueprintCallable, Category = "PCG|Metadata", meta = (ScriptMethod))
+	static void SetDoubleAttribute(UPARAM(ref) FPCGPoint& Point, UPCGMetadata* Metadata, FName AttributeName, double Value);
+
 	UFUNCTION(BlueprintCallable, Category = "PCG|Metadata", meta = (ScriptMethod))
 	static FVector GetVectorAttribute(const FPCGPoint& Point, const UPCGMetadata* Metadata, FName AttributeName);
 
