Added basic landscape layer read support. First version, will have more changes later based on needs or to improve the workflow.

diff --git a/Engine/Plugins/Experimental/PCG/Source/PCG/Private/Data/PCGLandscapeData.cpp b/Engine/Plugins/Experimental/PCG/Source/PCG/Private/Data/PCGLandscapeData.cpp
--- a/Engine/Plugins/Experimental/PCG/Source/PCG/Private/Data/PCGLandscapeData.cpp
+++ b/Engine/Plugins/Experimental/PCG/Source/PCG/Private/Data/PCGLandscapeData.cpp
@@ -3,19 +3,16 @@
 #include "Data/PCGLandscapeData.h"
 
 #include "PCGHelpers.h"
+#include "PCGSubsystem.h"
 #include "Data/PCGPointData.h"
 #include "Data/PCGPolyLineData.h"
 #include "Data/PCGSpatialData.h"
+#include "Grid/PCGLandscapeCache.h"
 #include "Helpers/PCGAsync.h"
 
 #include "Landscape.h"
 #include "LandscapeEdit.h"
 
-FPCGLandscapeDataPoint::FPCGLandscapeDataPoint(int InX, int InY, float InHeight)
-	: X(InX), Y(InY), Height(InHeight)
-{
-}
-
 void UPCGLandscapeData::Initialize(ALandscapeProxy* InLandscape, const FBox& InBounds)
 {
 	check(InLandscape);
@@ -24,6 +21,27 @@ void UPCGLandscapeData::Initialize(ALandscapeProxy* InLandscape, const FBox& InB
 	Bounds = InBounds;
 
 	Transform = Landscape->GetActorTransform();
+
+	// TODO: find a better way to do this - maybe there should be a prototype metadata in the landscape cache
+	if (Landscape->GetWorld())
+	{
+		if (UPCGSubsystem* Subsystem = Landscape->GetWorld()->GetSubsystem<UPCGSubsystem>())
+		{
+			if (FPCGLandscapeCache* LandscapeCache = Subsystem->GetLandscapeCache())
+			{
+				const TArray<FName> Layers = Subsystem->GetLandscapeCache()->GetLayerNames(Landscape.Get());
+
+				for (const FName& Layer : Layers)
+				{
+					Metadata->CreateFloatAttribute(Layer, 0.0f, /*bAllowInterpolation=*/true);
+				}
+			}
+			else
+			{
+				UE_LOG(LogPCG, Error, TEXT("Landscape is unable to access the landscape cache"));
+			}
+		}
+	}
 }
 
 FBox UPCGLandscapeData::GetBounds() const
@@ -40,17 +58,72 @@ FBox UPCGLandscapeData::GetStrictBounds() const
 
 bool UPCGLandscapeData::SamplePoint(const FTransform& InTransform, const FBox& InBounds, FPCGPoint& OutPoint, UPCGMetadata* OutMetadata) const
 {
-	// TODO support metadata / additional parameters
-	// TODO support bounds tests
-	// TODO: maybe we should consider the bounds on this data? wasn't done previously
-	TOptional<float> HeightAtVertex = Landscape->GetHeightAtLocation(InTransform.GetLocation());
-	if (HeightAtVertex.IsSet())
+	if (!Landscape || !Landscape->GetWorld())
 	{
-		OutPoint.Transform = InTransform;
-		const FVector LocationOnLandscape = FVector(InTransform.GetLocation().X, InTransform.GetLocation().Y, HeightAtVertex.GetValue());
-		OutPoint.Transform.SetLocation(LocationOnLandscape);
-		OutPoint.SetLocalBounds(InBounds); // TODO: unless we do bounds testing should set Z to 0.
-		OutPoint.Density = 1.0f;
+		return false;
+	}
+
+	UPCGSubsystem* Subsystem = Landscape->GetWorld()->GetSubsystem<UPCGSubsystem>();
+
+	if (!Subsystem || !Subsystem->GetLandscapeCache())
+	{
+		UE_LOG(LogPCG, Error, TEXT("PCG Subsystem or landscape cache are not initialized"));
+		return false;
+	}
+
+	// Compute landscape-centric coordinates from the transform
+	// TODO: support bounds tests/interpolation
+	ULandscapeInfo* LandscapeInfo = Landscape->GetLandscapeInfo();
+
+	if (!LandscapeInfo)
+	{
+		return false;
+	}
+
+	const FTransform LocalTransform = InTransform * Transform.Inverse();
+	const FVector LocalPoint = LocalTransform.GetLocation();
+
+	const FIntPoint ComponentMapKey(FMath::FloorToInt(LocalPoint.X / LandscapeInfo->ComponentSizeQuads), FMath::FloorToInt(LocalPoint.Y / LandscapeInfo->ComponentSizeQuads));
+
+	if (ULandscapeComponent* LandscapeComponent = LandscapeInfo->XYtoComponentMap.FindRef(ComponentMapKey))
+	{
+		check(Subsystem->GetLandscapeCache());
+		const FPCGLandscapeCacheEntry* LandscapeCacheEntry = Subsystem->GetLandscapeCache()->GetCacheEntry(LandscapeComponent, ComponentMapKey);
+
+		if (!LandscapeCacheEntry)
+		{
+			return false;
+		}
+
+		// Compute the 4 points indices
+		const FVector2D ComponentLocalPoint(LocalPoint.X - ComponentMapKey.X * LandscapeInfo->ComponentSizeQuads, LocalPoint.Y - ComponentMapKey.Y * LandscapeInfo->ComponentSizeQuads);
+		const int32 X0Y0 = FMath::FloorToInt(ComponentLocalPoint.X) + FMath::FloorToInt(ComponentLocalPoint.Y) * (LandscapeInfo->ComponentSizeQuads + 1);
+		const int32 X1Y0 = X0Y0 + 1;
+		const int32 X0Y1 = X0Y0 + (LandscapeInfo->ComponentSizeQuads + 1);
+		const int32 X1Y1 = X0Y1 + 1;
+
+		FPCGPoint PX0Y0;
+		FPCGPoint PX1Y0;
+		FPCGPoint PX0Y1;
+		FPCGPoint PX1Y1;
+
+		LandscapeCacheEntry->GetPoint(X0Y0, PX0Y0, OutMetadata);
+		LandscapeCacheEntry->GetPoint(X1Y0, PX1Y0, OutMetadata);
+		LandscapeCacheEntry->GetPoint(X0Y1, PX0Y1, OutMetadata);
+		LandscapeCacheEntry->GetPoint(X1Y1, PX1Y1, OutMetadata);
+
+
+		PCGPointHelpers::Bilerp(
+			PX0Y0,
+			PX1Y0,
+			PX0Y1,
+			PX1Y1,
+			OutMetadata,
+			OutPoint,
+			OutMetadata,
+			FMath::Fractional(LocalPoint.X),
+			FMath::Fractional(LocalPoint.Y));
+
 		return true;
 	}
 	else
@@ -63,6 +136,28 @@ const UPCGPointData* UPCGLandscapeData::CreatePointData(FPCGContext* Context, co
 {
 	TRACE_CPUPROFILER_EVENT_SCOPE(UPCGLandscapeData::CreatePointData);
 
+	if (!Landscape || !Landscape->GetWorld())
+	{
+		return nullptr;
+	}
+
+	UPCGSubsystem* Subsystem = Landscape->GetWorld()->GetSubsystem<UPCGSubsystem>();
+
+	if (!Subsystem || !Subsystem->GetLandscapeCache())
+	{
+		UE_LOG(LogPCG, Error, TEXT("PCG Subsystem or landscape cache are not initialized"));
+		return nullptr;
+	}
+
+	ULandscapeInfo* LandscapeInfo = Landscape->GetLandscapeInfo();
+
+	if (!LandscapeInfo)
+	{
+		return nullptr;
+	}
+
+	const int32 ComponentSizeQuads = LandscapeInfo->ComponentSizeQuads;
+
 	UPCGPointData* Data = NewObject<UPCGPointData>(const_cast<UPCGLandscapeData*>(this));
 	Data->InitializeFromData(this);
 	TArray<FPCGPoint>& Points = Data->GetMutablePoints();
@@ -73,45 +168,76 @@ const UPCGPointData* UPCGLandscapeData::CreatePointData(FPCGContext* Context, co
 		EffectiveBounds = Bounds.Overlap(InBounds);
 	}
 
+	// Early out
+	if (!EffectiveBounds.IsValid)
+	{
+		return Data;
+	}
+
 	// TODO: add offset to nearest edge, will have an impact if the grid size doesn't match the landscape size
 	const FVector MinPt = Transform.InverseTransformPosition(EffectiveBounds.Min);
 	const FVector MaxPt = Transform.InverseTransformPosition(EffectiveBounds.Max);
 
-	const int32 MinX = FMath::FloorToInt(MinPt.X);
+	// Note: the MaxX/Y here are inclusive, hence the floor & the +1 in the sizes
+	const int32 MinX = FMath::CeilToInt(MinPt.X);
 	const int32 MaxX = FMath::FloorToInt(MaxPt.X);
-	const int32 MinY = FMath::FloorToInt(MinPt.Y);
+	const int32 MinY = FMath::CeilToInt(MinPt.Y);
 	const int32 MaxY = FMath::FloorToInt(MaxPt.Y);
 
-	int32 NumIterations = (MaxX - MinX) * (MaxY - MinY);
+	//Early out if the bounds do not overlap any landscape vertices
+	if (MaxX < MinX || MaxY < MinY)
+	{
+		return Data;
+	}
 
-	FPCGAsync::AsyncPointProcessing(Context, NumIterations, Points, [this, &EffectiveBounds, MinX, MaxX, MinY](int32 Index, FPCGPoint& OutPoint)
+	const int64 PointCountUpperBound = (1 + MaxX - MinX) * (1 + MaxY - MinY);
+	if (PointCountUpperBound > 0)
 	{
-		const int X = MinX + (Index % (MaxX - MinX));
-		const int Y = MinY + (Index / (MaxX - MinX));
-		const bool bPlaneCase = EffectiveBounds.Min.Z == EffectiveBounds.Max.Z;
-
-		FVector VertexLocation = Transform.TransformPosition(FVector(X, Y, 0));
-		TOptional<float> HeightAtVertex = Landscape->GetHeightAtLocation(VertexLocation);
-		if (HeightAtVertex.IsSet() &&
-			HeightAtVertex.GetValue() >= EffectiveBounds.Min.Z &&
-			(bPlaneCase ? HeightAtVertex.GetValue() <= EffectiveBounds.Max.Z : HeightAtVertex.GetValue() < EffectiveBounds.Max.Z))
-		{
-			VertexLocation.Z = HeightAtVertex.GetValue();
+		Points.Reserve(PointCountUpperBound);
+	}
 
-			OutPoint = FPCGPoint(FTransform(VertexLocation),
-				1.0f,
-				PCGHelpers::ComputeSeed(X, Y));
-			OutPoint.SetExtents(Transform.GetScale3D() / 2.0);
+	const int32 MinComponentX = MinX / ComponentSizeQuads;
+	const int32 MaxComponentX = MaxX / ComponentSizeQuads;
+	const int32 MinComponentY = MinY / ComponentSizeQuads;
+	const int32 MaxComponentY = MaxY / ComponentSizeQuads;
 
-			return true;
-		}
-		else
+	for (int32 ComponentX = MinComponentX; ComponentX <= MaxComponentX; ++ComponentX)
+	{
+		for (int32 ComponentY = MinComponentY; ComponentY <= MaxComponentY; ++ComponentY)
 		{
-			return false;
+			FIntPoint ComponentMapKey(ComponentX, ComponentY);
+			if (ULandscapeComponent* LandscapeComponent = LandscapeInfo->XYtoComponentMap.FindRef(ComponentMapKey))
+			{
+				const FPCGLandscapeCacheEntry* LandscapeCacheEntry = Subsystem->GetLandscapeCache()->GetCacheEntry(LandscapeComponent, ComponentMapKey);
+
+				if (!LandscapeCacheEntry)
+				{
+					continue;
+				}
+
+				// Rebase our bounds in the component referential
+				const int32 LocalMinX = FMath::Clamp(MinX - ComponentMapKey.X * ComponentSizeQuads, 0, ComponentSizeQuads - 1);
+				const int32 LocalMaxX = FMath::Clamp(MaxX - ComponentMapKey.X * ComponentSizeQuads, 0, ComponentSizeQuads - 1);
+
+				const int32 LocalMinY = FMath::Clamp(MinY - ComponentMapKey.Y * ComponentSizeQuads, 0, ComponentSizeQuads - 1);
+				const int32 LocalMaxY = FMath::Clamp(MaxY - ComponentMapKey.Y * ComponentSizeQuads, 0, ComponentSizeQuads - 1);
+
+				// We can't really copy data from the component points wholesale because the component points have an additional boundary point.
+				// TODO: consider optimizing this, though it will impact the Sample then
+				for (int32 LocalX = LocalMinX; LocalX <= LocalMaxX; ++LocalX)
+				{
+					for (int32 LocalY = LocalMinY; LocalY <= LocalMaxY; ++LocalY)
+					{
+						FPCGPoint& Point = Points.Emplace_GetRef();
+						LandscapeCacheEntry->GetPoint(LocalX + LocalY * (ComponentSizeQuads + 1), Point, Data->Metadata);
+					}
+				}
+			}
 		}
-	});
+	}
 
-	UE_LOG(LogPCG, Verbose, TEXT("Landscape %s extracted %d of %d potential points"), *Landscape->GetFName().ToString(), Points.Num(), (MaxX - MinX) * (MaxY - MinY));
+	check(Points.Num() <= PointCountUpperBound);
+	UE_LOG(LogPCG, Verbose, TEXT("Landscape %s extracted %d of %d potential points"), *Landscape->GetFName().ToString(), Points.Num(), PointCountUpperBound);
 
 	return Data;
 }
\ No newline at end of file
diff --git a/Engine/Plugins/Experimental/PCG/Source/PCG/Private/Data/PCGPointData.cpp b/Engine/Plugins/Experimental/PCG/Source/PCG/Private/Data/PCGPointData.cpp
--- a/Engine/Plugins/Experimental/PCG/Source/PCG/Private/Data/PCGPointData.cpp
+++ b/Engine/Plugins/Experimental/PCG/Source/PCG/Private/Data/PCGPointData.cpp
@@ -89,7 +89,73 @@ namespace PCGPointHelpers
 			return Q + BlendQuat;
 		else
 			return Q - BlendQuat;
-	}	
+	}
+
+	void Lerp(const FPCGPoint& A, const FPCGPoint& B, float Ratio, const UPCGMetadata* SourceMetadata, FPCGPoint& OutPoint, UPCGMetadata* OutMetadata)
+	{
+		check(Ratio >= 0 && Ratio <= 1.0f);
+		// TODO: this might be incorrect. See UKismetMathLibrary::TLerp instead
+		OutPoint.Transform = FTransform(
+			FMath::Lerp(A.Transform.GetRotation(), B.Transform.GetRotation(), Ratio),
+			FMath::Lerp(A.Transform.GetLocation(), B.Transform.GetLocation(), Ratio),
+			FMath::Lerp(A.Transform.GetScale3D(), B.Transform.GetScale3D(), Ratio));
+		OutPoint.Density = FMath::Lerp(A.Density, B.Density, Ratio);
+		OutPoint.BoundsMin = FMath::Lerp(A.BoundsMin, B.BoundsMin, Ratio);
+		OutPoint.BoundsMax = FMath::Lerp(A.BoundsMax, B.BoundsMax, Ratio);
+		OutPoint.Color = FMath::Lerp(A.Color, B.Color, Ratio);
+		OutPoint.Steepness = FMath::Lerp(A.Steepness, B.Steepness, Ratio);
+
+		if (OutMetadata)
+		{
+			UPCGMetadataAccessorHelpers::InitializeMetadataWithParent(OutPoint, OutMetadata, ((Ratio <= 0.5f) ? A : B), SourceMetadata);
+			OutMetadata->ResetPointWeightedAttributes(OutPoint);
+			OutMetadata->AccumulatePointWeightedAttributes(A, SourceMetadata, Ratio, (Ratio <= 0.5f), OutPoint);
+			OutMetadata->AccumulatePointWeightedAttributes(B, SourceMetadata, 1.0f-Ratio, (Ratio > 0.5f), OutPoint);
+		}
+	}
+
+	void Bilerp(const FPCGPoint& X0Y0, const FPCGPoint& X1Y0, const FPCGPoint& X0Y1, const FPCGPoint& X1Y1, const UPCGMetadata* SourceMetadata, FPCGPoint& OutPoint, UPCGMetadata* OutMetadata, float XFactor, float YFactor)
+	{
+		const FPCGPoint* SelectedPoint = nullptr;
+
+		// First, check preexisting points as interpolation might be costly
+		if (XFactor < KINDA_SMALL_NUMBER && YFactor < KINDA_SMALL_NUMBER)
+		{
+			SelectedPoint = &X0Y0;
+		}
+		else if (XFactor > 1.0f - KINDA_SMALL_NUMBER && YFactor < KINDA_SMALL_NUMBER)
+		{
+			SelectedPoint = &X1Y0;
+		}
+		else if (XFactor < KINDA_SMALL_NUMBER && YFactor > 1.0f - KINDA_SMALL_NUMBER)
+		{
+			SelectedPoint = &X0Y1;
+		}
+		else if (XFactor > 1.0f - KINDA_SMALL_NUMBER && YFactor > 1.0f - KINDA_SMALL_NUMBER)
+		{
+			SelectedPoint = &X1Y1;
+		}
+
+		if (SelectedPoint)
+		{
+			OutPoint = *SelectedPoint;
+			if (OutMetadata)
+			{
+				OutMetadata->SetPointAttributes(*SelectedPoint, SourceMetadata, OutPoint);
+			}
+		}
+		else
+		{
+			// Interpolate X0Y0-X1Y0 and X0Y1-X1Y1 using XFactor
+			FPCGPoint Y0Lerp;
+			FPCGPoint Y1Lerp;
+
+			Lerp(X0Y0, X1Y0, XFactor, SourceMetadata, Y0Lerp, OutMetadata);
+			Lerp(X0Y1, X1Y1, XFactor, SourceMetadata, Y1Lerp, OutMetadata);
+			// Interpolate between the two points using YFactor
+			Lerp(Y0Lerp, Y1Lerp, YFactor, SourceMetadata, OutPoint, OutMetadata);
+		}
+	}
 }
 
 FPCGPointRef::FPCGPointRef(const FPCGPoint& InPoint)
@@ -265,7 +331,7 @@ bool UPCGPointData::SamplePoint(const FTransform& InTransform, const FBox& InBou
 		const FPCGPoint& SourcePoint = *Contribution.Key;
 		const float Weight = Contribution.Value / SumContributions;
 
-		WeightedPosition += SourcePoint.Transform.GetLocation();
+		WeightedPosition += SourcePoint.Transform.GetLocation() * Weight;
 		WeightedQuat = PCGPointHelpers::AddQuatWithWeight(WeightedQuat, SourcePoint.Transform.GetRotation(), Weight);
 		WeightedScale += SourcePoint.Transform.GetScale3D() * Weight;
 
diff --git a/Engine/Plugins/Experimental/PCG/Source/PCG/Private/Grid/PCGLandscapeCache.cpp b/Engine/Plugins/Experimental/PCG/Source/PCG/Private/Grid/PCGLandscapeCache.cpp
new file mode 100644
--- /dev/null
+++ b/Engine/Plugins/Experimental/PCG/Source/PCG/Private/Grid/PCGLandscapeCache.cpp
@@ -0,0 +1,334 @@
+// Copyright Epic Games, Inc. All Rights Reserved.
+
+#include "Grid/PCGLandscapeCache.h"
+
+#include "PCGSubsystem.h"
+#include "Data/PCGPointData.h"
+#include "Helpers/PCGBlueprintHelpers.h"
+#include "Metadata/PCGMetadata.h"
+#include "Metadata/PCGMetadataAttribute.h"
+#include "Metadata/PCGMetadataAttributeTpl.h"
+#include "Metadata/PCGMetadataAttributeTraits.h"
+
+#include "UObject/WeakObjectPtr.h"
+
+#include "LandscapeHeightfieldCollisionComponent.h"
+#include "LandscapeComponent.h"
+#include "LandscapeInfo.h"
+#include "LandscapeInfoMap.h"
+#include "LandscapeProxy.h"
+#include "LandscapeDataAccess.h"
+#include "Landscape.h"
+#include "Kismet/GameplayStatics.h"
+
+#if WITH_EDITOR
+void FPCGLandscapeCacheEntry::BuildCacheData(ULandscapeInfo* LandscapeInfo, ULandscapeComponent* InComponent, UObject* Owner)
+{
+	check(Owner && !Component.Get() && InComponent && PositionsAndNormals.Num() == 0);
+	Component = InComponent;
+
+	ALandscape* Landscape = Component->GetLandscapeActor();
+	if (!Landscape)
+	{
+		return;
+	}
+
+	// Get landscape default layer (heightmap/tangents/normal)
+	{
+		FLandscapeComponentDataInterface CDI(InComponent, 0);
+		FVector WorldPos;
+		FVector WorldTangentX;
+		FVector WorldTangentY;
+		FVector WorldTangentZ;
+
+		PointHalfSize = InComponent->GetComponentTransform().GetScale3D() * 0.5;
+
+		// The component has an extra vertex on the edge, for interpolation purposes
+		const int32 ComponentSizeQuads = InComponent->ComponentSizeQuads + 1;
+		const int32 NumVertices = FMath::Square(ComponentSizeQuads);
+
+		PositionsAndNormals.Reserve(2 * NumVertices);
+		for (int32 Index = 0; Index < NumVertices; ++Index)
+		{
+			CDI.GetWorldPositionTangents(Index, WorldPos, WorldTangentX, WorldTangentY, WorldTangentZ);
+			PositionsAndNormals.Add(WorldPos);
+			PositionsAndNormals.Add(WorldTangentZ);
+		}
+	}
+	
+	// Get other layers, push data into metadata attributes
+	TArray<uint8> LayerCache;
+	for (const FLandscapeInfoLayerSettings& Layer : LandscapeInfo->Layers)
+	{
+		ULandscapeLayerInfoObject* LayerInfo = Layer.LayerInfoObj;
+		if (!LayerInfo)
+		{
+			continue;
+		}
+
+		if (FLandscapeLayer* EditLayer = Landscape->GetLayer(Landscape->GetLayerIndex(Layer.LayerName)))
+		{
+			FScopedSetLandscapeEditingLayer Scope(Landscape, EditLayer->Guid);
+			FLandscapeComponentDataInterface CDI(InComponent, 0, /*bUseEditingLayer=*/true);
+			if (CDI.GetWeightmapTextureData(LayerInfo, LayerCache, /*bUseEditingLayer=*/true))
+			{
+				FPCGLandscapeCacheLayer& PCGLayer = LayerData.Emplace_GetRef();
+				PCGLayer.Name = Layer.LayerName;
+				PCGLayer.Data = LayerCache;
+			}
+		}
+
+		LayerCache.Reset();
+	}
+}
+#endif // WITH_EDITOR
+
+bool FPCGLandscapeCacheEntry::GetPoint(int32 PointIndex, FPCGPoint& OutPoint, UPCGMetadata* OutMetadata) const
+{
+	if (PointIndex < 0 || 2 * PointIndex >= PositionsAndNormals.Num())
+	{
+		return false;
+	}
+
+	const FVector& Position = PositionsAndNormals[2 * PointIndex];
+	const FVector& Normal = PositionsAndNormals[2 * PointIndex + 1];
+	
+	FVector TangentX;
+	FVector TangentY;
+	TangentX = FVector(Normal.Z, 0.f, -Normal.X);
+	TangentY = Normal ^ TangentX;
+
+	OutPoint.Transform = FTransform(TangentX, TangentY, Normal, Position);
+	OutPoint.BoundsMin = -PointHalfSize;
+	OutPoint.BoundsMax = PointHalfSize;
+	OutPoint.Seed = UPCGBlueprintHelpers::ComputeSeedFromPosition(Position);
+	
+	if (OutMetadata && !LayerData.IsEmpty())
+	{
+		OutPoint.MetadataEntry = OutMetadata->AddEntry();
+
+		for (const FPCGLandscapeCacheLayer& Layer : LayerData)
+		{
+			if (FPCGMetadataAttributeBase* Attribute = OutMetadata->GetMutableAttribute(Layer.Name))
+			{
+				check(Attribute->GetTypeId() == PCG::Private::MetadataTypes<float>::Id);
+				static_cast<FPCGMetadataAttribute<float>*>(Attribute)->SetValue(OutPoint.MetadataEntry, (float)Layer.Data[PointIndex] / 255.0f);
+			}
+		}
+	}
+
+	return true;
+}
+
+FPCGLandscapeCache::FPCGLandscapeCache(UObject* InOwner)
+	: Owner(InOwner)
+{
+	check(Owner);
+#if WITH_EDITOR
+	SetupLandscapeCallbacks();
+	CacheLayerNames();
+#endif
+}
+
+void FPCGLandscapeCache::SetOwner(UObject* InOwner)
+{
+	check(!Owner || Owner == InOwner || InOwner == nullptr);
+	Owner = InOwner;
+#if WITH_EDITOR
+	if (InOwner)
+	{
+		SetupLandscapeCallbacks();
+	}
+	else
+	{
+		TeardownLandscapeCallbacks();
+	}
+
+	CacheLayerNames();
+#endif
+}
+
+FPCGLandscapeCache::~FPCGLandscapeCache()
+{
+#if WITH_EDITOR
+	TeardownLandscapeCallbacks();
+#endif
+}
+
+void FPCGLandscapeCache::PrimeCache()
+{
+#if WITH_EDITOR
+	if (!Owner)
+	{
+		return;
+	}
+
+	UWorld* World = Owner->GetWorld();
+	if (!World)
+	{
+		return;
+	}
+
+	for (auto It = ULandscapeInfoMap::GetLandscapeInfoMap(World).Map.CreateIterator(); It; ++It)
+	{
+		ULandscapeInfo* LandscapeInfo = It.Value();
+		if (IsValid(LandscapeInfo))
+		{
+			// Build per-component information
+			LandscapeInfo->ForAllLandscapeProxies([this, LandscapeInfo](const ALandscapeProxy* LandscapeProxy)
+			{
+				for (ULandscapeComponent* LandscapeComponent : LandscapeProxy->LandscapeComponents)
+				{
+					if (!LandscapeComponent)
+					{
+						continue;
+					}
+
+					FIntPoint ComponentKey(LandscapeComponent->SectionBaseX / LandscapeComponent->ComponentSizeQuads, LandscapeComponent->SectionBaseY / LandscapeComponent->ComponentSizeQuads);
+
+					if (!CachedData.Contains(ComponentKey))
+					{
+						CachedData.Add(ComponentKey).BuildCacheData(LandscapeInfo, LandscapeComponent, Owner);
+					}
+				}
+			});
+		}
+	}	
+#endif
+}
+
+void FPCGLandscapeCache::ClearCache()
+{
+	CachedData.Reset();
+}
+
+const FPCGLandscapeCacheEntry* FPCGLandscapeCache::GetCacheEntry(ULandscapeComponent* LandscapeComponent, const FIntPoint& ComponentKey)
+{
+	const FPCGLandscapeCacheEntry* FoundEntry = nullptr;
+
+	{
+#if WITH_EDITOR
+		FReadScopeLock ScopeLock(CacheLock);
+#endif
+		FoundEntry = CachedData.Find(ComponentKey);
+	}
+
+#if WITH_EDITOR
+	if (!FoundEntry && LandscapeComponent && LandscapeComponent->GetLandscapeInfo())
+	{
+		check(LandscapeComponent->SectionBaseX / LandscapeComponent->ComponentSizeQuads == ComponentKey.X && LandscapeComponent->SectionBaseY / LandscapeComponent->ComponentSizeQuads == ComponentKey.Y);
+		// Create entry
+		FPCGLandscapeCacheEntry NewEntry;
+		NewEntry.BuildCacheData(LandscapeComponent->GetLandscapeInfo(), LandscapeComponent, Owner);
+
+		// Try to store it
+		{
+			FWriteScopeLock ScopeLock(CacheLock);
+			CachedData.FindOrAdd(ComponentKey, NewEntry);
+			FoundEntry = CachedData.Find(ComponentKey);
+		}
+	}
+#endif
+	
+	return FoundEntry;
+}
+
+TArray<FName> FPCGLandscapeCache::GetLayerNames(ALandscapeProxy* Landscape)
+{
+#if WITH_EDITOR
+	FReadScopeLock ScopeLock(CacheLock);
+#endif
+	return CachedLayerNames.Array();
+}
+
+#if WITH_EDITOR
+void FPCGLandscapeCache::SetupLandscapeCallbacks()
+{
+	// Remove previous callbacks, if any
+	TeardownLandscapeCallbacks();
+
+	if (!Owner)
+	{
+		return;
+	}
+
+	UWorld* World = Owner->GetWorld();
+	if (!World)
+	{
+		return;
+	}
+
+	// Gather landspace actors
+	TArray<AActor*> FoundLandscapes;
+	UGameplayStatics::GetAllActorsOfClass(World, ALandscapeProxy::StaticClass(), FoundLandscapes);
+
+	for (AActor* FoundLandscape : FoundLandscapes)
+	{
+		ALandscapeProxy* Landscape = Cast<ALandscapeProxy>(FoundLandscape);
+		check(Landscape);
+
+		Landscapes.Add(Landscape);
+		Landscape->OnComponentDataChanged.AddRaw(this, &FPCGLandscapeCache::OnLandscapeChanged);
+	}
+}
+
+void FPCGLandscapeCache::TeardownLandscapeCallbacks()
+{
+	for (TWeakObjectPtr<ALandscapeProxy> LandscapeWeakPtr : Landscapes)
+	{
+		if (ALandscapeProxy* Landscape = LandscapeWeakPtr.Get())
+		{
+			Landscape->OnComponentDataChanged.RemoveAll(this);
+		}
+	}
+}
+
+void FPCGLandscapeCache::OnLandscapeChanged(ALandscapeProxy* Landscape, const FLandscapeProxyComponentDataChangedParams& ChangeParams)
+{
+	if (!Landscapes.Contains(Landscape))
+	{
+		return;
+	}
+
+	CacheLock.WriteLock();
+
+	// Just remove these from the cache, they'll be added back on demand
+	ChangeParams.ForEachComponent([this](const ULandscapeComponent* LandscapeComponent)
+	{
+		if (LandscapeComponent)
+		{
+			FIntPoint ComponentKey(LandscapeComponent->SectionBaseX / LandscapeComponent->ComponentSizeQuads, LandscapeComponent->SectionBaseY / LandscapeComponent->ComponentSizeQuads);
+			CachedData.Remove(ComponentKey);
+		}
+	});
+
+	CacheLayerNames(Landscape);
+
+	CacheLock.WriteUnlock();
+}
+
+void FPCGLandscapeCache::CacheLayerNames()
+{
+	for (TWeakObjectPtr<ALandscapeProxy> Landscape : Landscapes)
+	{
+		if (Landscape.Get())
+		{
+			CacheLayerNames(Landscape.Get());
+		}
+	}
+}
+
+void FPCGLandscapeCache::CacheLayerNames(ALandscapeProxy* Landscape)
+{
+	check(Landscape);
+
+	if (ULandscapeInfo* LandscapeInfo = Landscape->GetLandscapeInfo())
+	{
+		for (const FLandscapeInfoLayerSettings& Layer : LandscapeInfo->Layers)
+		{
+			CachedLayerNames.Add(Layer.LayerName);
+		}
+	}
+}
+
+#endif // WITH_EDITOR
\ No newline at end of file
diff --git a/Engine/Plugins/Experimental/PCG/Source/PCG/Private/Metadata/PCGMetadata.cpp b/Engine/Plugins/Experimental/PCG/Source/PCG/Private/Metadata/PCGMetadata.cpp
--- a/Engine/Plugins/Experimental/PCG/Source/PCG/Private/Metadata/PCGMetadata.cpp
+++ b/Engine/Plugins/Experimental/PCG/Source/PCG/Private/Metadata/PCGMetadata.cpp
@@ -7,9 +7,54 @@
 #include "Misc/ScopeRWLock.h"
 #include "Helpers/PCGSettingsHelpers.h"
 
-void UPCGMetadata::Serialize(FArchive& Ar)
+void UPCGMetadata::Serialize(FArchive& InArchive)
 {
-	// TODO
+	Super::Serialize(InArchive);
+
+	int32 NumAttributes = (InArchive.IsLoading() ? 0 : Attributes.Num());
+
+	InArchive << NumAttributes;
+
+	if (InArchive.IsLoading())
+	{
+		for (int32 AttributeIndex = 0; AttributeIndex < NumAttributes; ++AttributeIndex)
+		{
+			FName AttributeName = NAME_None;
+			InArchive << AttributeName;
+
+			int32 AttributeTypeId = 0;
+			InArchive << AttributeTypeId;
+
+			FPCGMetadataAttributeBase* SerializedAttribute = PCGMetadataAttribute::AllocateEmptyAttributeFromType(static_cast<int16>(AttributeTypeId));
+			if (SerializedAttribute)
+			{
+				SerializedAttribute->Name = AttributeName;
+				SerializedAttribute->Serialize(this, InArchive);
+				Attributes.Add(AttributeName, SerializedAttribute);
+			}
+		}
+	}
+	else
+	{
+		for (auto& AttributePair : Attributes)
+		{
+			InArchive << AttributePair.Key;
+			
+			int32 AttributeTypeId = AttributePair.Value->GetTypeId();
+			InArchive << AttributeTypeId;
+
+			AttributePair.Value->Serialize(this, InArchive);
+		}
+	}
+
+	InArchive << ParentKeys;
+
+	// Finally, initialize non-serialized members
+	if (InArchive.IsLoading())
+	{
+		NextAttributeId = Attributes.Num();
+		ItemKeyOffset = (Parent ? Parent->GetItemCountForChild() : 0);
+	}
 }
 
 void UPCGMetadata::BeginDestroy()
@@ -247,12 +292,36 @@ const FPCGMetadataAttributeBase* UPCGMetadata::GetConstAttribute(FName Attribute
 	return Attribute;
 }
 
+const FPCGMetadataAttributeBase* UPCGMetadata::GetConstAttributeById(int32 InAttributeId) const
+{
+	const FPCGMetadataAttributeBase* Attribute = nullptr;
+
+	AttributeLock.ReadLock();
+	for (const auto& AttributePair : Attributes)
+	{
+		if (AttributePair.Value && AttributePair.Value->AttributeId == InAttributeId)
+		{
+			Attribute = AttributePair.Value;
+			break;
+		}
+	}
+	AttributeLock.ReadUnlock();
+
+	return Attribute;
+}
+
 bool UPCGMetadata::HasAttribute(FName AttributeName) const
 {
 	FReadScopeLock ScopeLock(AttributeLock);
 	return Attributes.Contains(AttributeName);
 }
 
+int32 UPCGMetadata::GetAttributeCount() const
+{
+	FReadScopeLock ScopeLock(AttributeLock);
+	return Attributes.Num();
+}
+
 void UPCGMetadata::GetAttributes(TArray<FName>& AttributeNames, TArray<EPCGMetadataTypes>& AttributeTypes) const
 {
 	AttributeNames.Reset();
@@ -424,7 +493,7 @@ FPCGMetadataAttributeBase* UPCGMetadata::CopyAttribute(const FPCGMetadataAttribu
 	FPCGMetadataAttributeBase* NewAttribute = OriginalAttribute->Copy(NewAttributeName, this, bKeepParent, bCopyEntries, bCopyValues);
 
 	AttributeLock.WriteLock();
-	NewAttribute->AttributeId = ++NextAttributeId;
+	NewAttribute->AttributeId = NextAttributeId++;
 	AddAttributeInternal(NewAttributeName, NewAttribute);
 	AttributeLock.WriteUnlock();
 
@@ -697,7 +766,7 @@ void UPCGMetadata::SetAttributes(PCGMetadataEntryKey InKey, const UPCGMetadata*
 
 void UPCGMetadata::SetPointAttributes(const TArrayView<const FPCGPoint>& InPoints, const UPCGMetadata* InMetadata, const TArrayView<FPCGPoint>& OutPoints)
 {
-	if (!InMetadata)
+	if (!InMetadata || InMetadata->GetAttributeCount() == 0)
 	{
 		return;
 	}
diff --git a/Engine/Plugins/Experimental/PCG/Source/PCG/Private/Metadata/PCGMetadataAttribute.cpp b/Engine/Plugins/Experimental/PCG/Source/PCG/Private/Metadata/PCGMetadataAttribute.cpp
--- a/Engine/Plugins/Experimental/PCG/Source/PCG/Private/Metadata/PCGMetadataAttribute.cpp
+++ b/Engine/Plugins/Experimental/PCG/Source/PCG/Private/Metadata/PCGMetadataAttribute.cpp
@@ -10,11 +10,29 @@ FPCGMetadataAttributeBase::FPCGMetadataAttributeBase(UPCGMetadata* InMetadata, F
 {
 }
 
+void FPCGMetadataAttributeBase::Serialize(UPCGMetadata* InMetadata, FArchive& InArchive)
+{
+	InArchive << EntryToValueKeyMap;
+	Metadata = InMetadata;
+
+	int32 ParentAttributeId = (Parent ? Parent->AttributeId : -1);
+	InArchive << ParentAttributeId;
+
+	if (InArchive.IsLoading() && ParentAttributeId >= 0 && Metadata->GetParent())
+	{
+		Parent = Metadata->GetParent()->GetConstAttributeById(ParentAttributeId);
+	}
+
+	//Type id should already be known by then, so no need to serialize it
+	InArchive << Name;
+	InArchive << AttributeId;
+}
+
 void FPCGMetadataAttributeBase::SetValueFromValueKey(PCGMetadataEntryKey EntryKey, PCGMetadataValueKey ValueKey)
 {
 	check(EntryKey != PCGInvalidEntryKey);
 	FWriteScopeLock ScopeLock(EntryMapLock);
-	EntryToValueMap.FindOrAdd(EntryKey) = ValueKey;
+	EntryToValueKeyMap.FindOrAdd(EntryKey) = ValueKey;
 }
 
 PCGMetadataValueKey FPCGMetadataAttributeBase::GetValueKey(PCGMetadataEntryKey EntryKey) const
@@ -28,7 +46,7 @@ PCGMetadataValueKey FPCGMetadataAttributeBase::GetValueKey(PCGMetadataEntryKey E
 	bool bFoundKey = false;
 
 	EntryMapLock.ReadLock();
-	if (const PCGMetadataValueKey* FoundLocalKey = EntryToValueMap.Find(EntryKey))
+	if (const PCGMetadataValueKey* FoundLocalKey = EntryToValueKeyMap.Find(EntryKey))
 	{
 		ValueKey = *FoundLocalKey;
 		bFoundKey = true;
@@ -52,5 +70,5 @@ bool FPCGMetadataAttributeBase::HasNonDefaultValue(PCGMetadataEntryKey EntryKey)
 
 void FPCGMetadataAttributeBase::ClearEntries()
 {
-	EntryToValueMap.Reset();
+	EntryToValueKeyMap.Reset();
 }
\ No newline at end of file
diff --git a/Engine/Plugins/Experimental/PCG/Source/PCG/Private/PCGSubsystem.cpp b/Engine/Plugins/Experimental/PCG/Source/PCG/Private/PCGSubsystem.cpp
--- a/Engine/Plugins/Experimental/PCG/Source/PCG/Private/PCGSubsystem.cpp
+++ b/Engine/Plugins/Experimental/PCG/Source/PCG/Private/PCGSubsystem.cpp
@@ -58,8 +58,19 @@ void UPCGSubsystem::PostInitialize()
 	// Initialize graph executor
 	check(!GraphExecutor);
 	GraphExecutor = new FPCGGraphExecutor(this);
-	// gather things.. ?
-	// TODO	
+	
+	// Gather world pcg actor if it exists
+	if (!PCGWorldActor)
+	{
+		if (UWorld* World = GetWorld())
+		{
+			UPCGActorHelpers::ForEachActorInWorld<APCGWorldActor>(World, [this](AActor* InActor)
+			{
+				PCGWorldActor = Cast<APCGWorldActor>(InActor);
+				return PCGWorldActor == nullptr;
+			});
+		}
+	}
 
 	// TODO: For now we set our octree to be 2km wide, but it would be perhaps better to
 	// scale it to the size of our world.
@@ -80,13 +91,31 @@ APCGWorldActor* UPCGSubsystem::GetPCGWorldActor()
 #if WITH_EDITOR
 	if (!PCGWorldActor)
 	{
+		PCGWorldActorLock.Lock();
+		if (!PCGWorldActor)
+		{
 		PCGWorldActor = APCGWorldActor::CreatePCGWorldActor(GetWorld());
 	}
+		PCGWorldActorLock.Unlock();
+	}
 #endif
 
 	return PCGWorldActor;
 }
 
+#if WITH_EDITOR
+void UPCGSubsystem::DestroyPCGWorldActor()
+{
+	if (PCGWorldActor)
+	{
+		PCGWorldActorLock.Lock();
+		PCGWorldActor->Destroy();
+		PCGWorldActor = nullptr;
+		PCGWorldActorLock.Unlock();
+	}
+}
+#endif
+
 void UPCGSubsystem::RegisterPCGWorldActor(APCGWorldActor* InActor)
 {
 	check(!PCGWorldActor || PCGWorldActor == InActor);
@@ -95,13 +124,18 @@ void UPCGSubsystem::RegisterPCGWorldActor(APCGWorldActor* InActor)
 
 void UPCGSubsystem::UnregisterPCGWorldActor(APCGWorldActor* InActor)
 {
-	check(PCGWorldActor == InActor);
 	if (PCGWorldActor == InActor)
 	{
 		PCGWorldActor = nullptr;
 	}
 }
 
+FPCGLandscapeCache* UPCGSubsystem::GetLandscapeCache()
+{
+	APCGWorldActor* LandscapeCacheOwner = GetPCGWorldActor();
+	return LandscapeCacheOwner ? &LandscapeCacheOwner->LandscapeCache : nullptr;
+}
+
 FPCGTaskId UPCGSubsystem::ScheduleComponent(UPCGComponent* PCGComponent, bool bSave, const TArray<FPCGTaskId>& Dependencies)
 {
 	check(GraphExecutor);
@@ -185,7 +219,7 @@ FPCGTaskId UPCGSubsystem::ScheduleCleanup(UPCGComponent* PCGComponent, bool bRem
 	// If the component is partitioned, we will forward the calls to its registered PCG Partition actors
 	if (PCGComponent->IsPartitioned())
 	{
-		{
+{
 			// TODO: Might be more interesting to copy the set and release the lock.
 			FReadScopeLock ReadLock(ComponentToPartitionActorsMapLock);
 			TSet<TObjectPtr<APCGPartitionActor>>* PartitionActorsPtr = ComponentToPartitionActorsMap.Find(PCGComponent);
@@ -197,11 +231,11 @@ FPCGTaskId UPCGSubsystem::ScheduleCleanup(UPCGComponent* PCGComponent, bool bRem
 		}
 	}
 	else
-	{
+		{
 		// In non partitioned mode, do it immediately
 		// TODO: Should it be a parameter?
 		PCGComponent->CleanupInternal(bRemoveComponents);
-	}
+		}
 
 	if (AllTasks.IsEmpty())
 	{
@@ -252,8 +286,8 @@ TArray<FPCGTaskId> UPCGSubsystem::ScheduleMultipleComponent(UPCGComponent* Origi
 				// Add check to avoid infinite loop
 				if (ensure(!LocalComponent->IsPartitioned()))
 				{
-					// Ensure that the PCG actor match our original
-					LocalComponent->SetPropertiesFromOriginal(OriginalComponent);
+				// Ensure that the PCG actor match our original
+				LocalComponent->SetPropertiesFromOriginal(OriginalComponent);
 
 					FPCGTaskId LocalTask = ScheduleComponent(LocalComponent, /*bSave=*/ false, Dependencies);
 
@@ -1086,6 +1120,27 @@ void UPCGSubsystem::FlushCache()
 	}
 }
 
+void UPCGSubsystem::BuildLandscapeCache()
+{
+	if (FPCGLandscapeCache* LandscapeCache = GetLandscapeCache())
+	{
+		PCGWorldActor->Modify();
+		LandscapeCache->PrimeCache();
+	}
+	else
+	{
+		UE_LOG(LogPCG, Error, TEXT("Unable to build landscape cache because either the world is null or there is no PCG world actor"));
+	}
+}
+
+void UPCGSubsystem::ClearLandscapeCache()
+{
+	if (FPCGLandscapeCache* LandscapeCache = GetLandscapeCache())
+	{
+		LandscapeCache->ClearCache();
+	}
+}
+
 void UPCGSubsystem::ResetPartitionActorsMap()
 {
 	PartitionActorsMapLock.WriteLock();
diff --git a/Engine/Plugins/Experimental/PCG/Source/PCG/Private/PCGWorldActor.cpp b/Engine/Plugins/Experimental/PCG/Source/PCG/Private/PCGWorldActor.cpp
--- a/Engine/Plugins/Experimental/PCG/Source/PCG/Private/PCGWorldActor.cpp
+++ b/Engine/Plugins/Experimental/PCG/Source/PCG/Private/PCGWorldActor.cpp
@@ -12,18 +12,37 @@
 APCGWorldActor::APCGWorldActor(const FObjectInitializer& ObjectInitializer)
 	: Super(ObjectInitializer)
 {
-	bIsEditorOnlyActor = true;
 #if WITH_EDITORONLY_DATA
 	bIsSpatiallyLoaded = false;
 	bDefaultOutlinerExpansionState = false;
 #endif
 
 	PartitionGridSize = DefaultPartitionGridSize;
+	LandscapeCache.SetOwner(this);
 }
 
+void APCGWorldActor::Serialize(FArchive& InArchive)
+{
+	if (InArchive.IsSaving() && !InArchive.IsCooking() && !(GetWorld() && GetWorld()->GetWorldPartition()))
+	{
+		LandscapeCache.ClearCache();
+	}
+
+	Super::Serialize(InArchive);
+}
+
+#if WITH_EDITOR
+void APCGWorldActor::BeginCacheForCookedPlatformData(const ITargetPlatform* TargetPlatform)
+{
+	Super::BeginCacheForCookedPlatformData(TargetPlatform);
+	LandscapeCache.PrimeCache();
+}
+#endif
+
 void APCGWorldActor::PostLoad()
 {
 	Super::PostLoad();
+	LandscapeCache.SetOwner(this);
 	RegisterToSubsystem();
 }
 
diff --git a/Engine/Plugins/Experimental/PCG/Source/PCG/Public/Data/PCGLandscapeData.h b/Engine/Plugins/Experimental/PCG/Source/PCG/Public/Data/PCGLandscapeData.h
--- a/Engine/Plugins/Experimental/PCG/Source/PCG/Public/Data/PCGLandscapeData.h
+++ b/Engine/Plugins/Experimental/PCG/Source/PCG/Public/Data/PCGLandscapeData.h
@@ -4,30 +4,10 @@
 
 #include "CoreMinimal.h"
 #include "PCGSurfaceData.h"
-#include "Elements/PCGSurfaceSampler.h"
 
 #include "PCGLandscapeData.generated.h"
 
 class ALandscapeProxy;
-class UPCGPolyLineData;
-
-USTRUCT()
-struct PCG_API FPCGLandscapeDataPoint
-{
-	GENERATED_BODY()
-
-	FPCGLandscapeDataPoint() = default;
-	FPCGLandscapeDataPoint(int InX, int InY, float InHeight);
-
-	UPROPERTY()
-	int X = 0;
-
-	UPROPERTY()
-	int Y = 0;
-
-	UPROPERTY()
-	float Height = 0;
-};
 
 UCLASS(BlueprintType, ClassGroup = (Procedural))
 class PCG_API UPCGLandscapeData : public UPCGSurfaceData
@@ -58,13 +38,7 @@ public:
 	UPROPERTY(BlueprintReadWrite, EditAnywhere, Category = SourceData)
 	TSoftObjectPtr<ALandscapeProxy> Landscape;
 
-	// TODO: add on property changed to clear cached data
-	//UPROPERTY(BlueprintReadWrite, EditAnywhere)
-	//TObjectPtr<UPCGSurfaceSamplerSettings> ImplicitSamplerSettings;
 protected:
-	UPROPERTY()
-	TArray<FPCGLandscapeDataPoint> LandscapePoints;
-
 	UPROPERTY()
 	FBox Bounds = FBox(EForceInit::ForceInit);
 };
diff --git a/Engine/Plugins/Experimental/PCG/Source/PCG/Public/Data/PCGPointData.h b/Engine/Plugins/Experimental/PCG/Source/PCG/Public/Data/PCGPointData.h
--- a/Engine/Plugins/Experimental/PCG/Source/PCG/Public/Data/PCGPointData.h
+++ b/Engine/Plugins/Experimental/PCG/Source/PCG/Public/Data/PCGPointData.h
@@ -12,6 +12,12 @@
 
 class AActor;
 
+namespace PCGPointHelpers
+{
+	void Lerp(const FPCGPoint& A, const FPCGPoint& B, float Ratio, const UPCGMetadata* SourceMetadata, FPCGPoint& OutPoint, UPCGMetadata* OutMetadata);
+	void Bilerp(const FPCGPoint& X0Y0, const FPCGPoint& X1Y0, const FPCGPoint& X0Y1, const FPCGPoint& X1Y1, const UPCGMetadata* SourceMetadata, FPCGPoint& OutPoint, UPCGMetadata* OutMetadata, float XFactor, float YFactor);
+}
+
 struct PCG_API FPCGPointRef
 {
 	FPCGPointRef(const FPCGPoint& InPoint);
diff --git a/Engine/Plugins/Experimental/PCG/Source/PCG/Public/Data/PCGSpatialData.h b/Engine/Plugins/Experimental/PCG/Source/PCG/Public/Data/PCGSpatialData.h
--- a/Engine/Plugins/Experimental/PCG/Source/PCG/Public/Data/PCGSpatialData.h
+++ b/Engine/Plugins/Experimental/PCG/Source/PCG/Public/Data/PCGSpatialData.h
@@ -105,8 +105,8 @@ public:
 	UFUNCTION(BlueprintCallable, Category = SpatialData)
 	void InitializeFromData(const UPCGSpatialData* InSource, const UPCGMetadata* InMetadataParentOverride = nullptr, bool bInheritMetadata = true);
 
-	UPROPERTY(BlueprintReadWrite, EditAnywhere, Category = Data)
-	AActor* TargetActor = nullptr;
+	UPROPERTY(Transient, BlueprintReadWrite, EditAnywhere, Category = Data)
+	TObjectPtr<AActor> TargetActor = nullptr;
 
 #if WITH_EDITORONLY_DATA
 	UPROPERTY(Transient, BlueprintReadWrite, EditAnywhere, Category = Debug)
diff --git a/Engine/Plugins/Experimental/PCG/Source/PCG/Public/Grid/PCGLandscapeCache.h b/Engine/Plugins/Experimental/PCG/Source/PCG/Public/Grid/PCGLandscapeCache.h
new file mode 100644
--- /dev/null
+++ b/Engine/Plugins/Experimental/PCG/Source/PCG/Public/Grid/PCGLandscapeCache.h
@@ -0,0 +1,104 @@
+// Copyright Epic Games, Inc. All Rights Reserved.
+
+#pragma once
+
+#include "PCGLandscapeCache.generated.h"
+
+class ALandscapeProxy;
+class FLandscapeProxyComponentDataChangedParams;
+class ULandscapeComponent;
+class ULandscapeInfo;
+class UPCGPointData;
+class UPCGMetadata;
+struct FPCGPoint;
+
+USTRUCT()
+struct FPCGLandscapeCacheLayer
+{
+	GENERATED_BODY()
+
+	UPROPERTY(VisibleAnywhere, Category = "Cache")
+	FName Name;
+
+	UPROPERTY()
+	TArray<uint8> Data;
+};
+
+USTRUCT()
+struct FPCGLandscapeCacheEntry
+{
+	GENERATED_BODY()
+public:
+#if WITH_EDITOR
+	void BuildCacheData(ULandscapeInfo* LandscapeInfo, ULandscapeComponent* InComponent, UObject* Owner);
+#endif
+
+	bool GetPoint(int32 PointIndex, FPCGPoint& OutPoint, UPCGMetadata* OutMetadata) const;
+
+	// TODO: this isn't really needed
+	UPROPERTY()
+	TWeakObjectPtr<const ULandscapeComponent> Component = nullptr;
+
+	UPROPERTY()
+	TArray<FVector> PositionsAndNormals;
+
+	UPROPERTY()
+	TArray<FPCGLandscapeCacheLayer> LayerData;
+
+	UPROPERTY()
+	FVector PointHalfSize;
+};
+
+USTRUCT()
+struct FPCGLandscapeCache
+{
+	GENERATED_BODY()
+public:
+	FPCGLandscapeCache() = default;
+	~FPCGLandscapeCache();
+	FPCGLandscapeCache(UObject* InOwner);
+
+	void SetOwner(UObject* InOwner);
+	void PrimeCache();
+	void ClearCache();
+
+	const FPCGLandscapeCacheEntry* GetCacheEntry(ULandscapeComponent* LandscapeComponent, const FIntPoint& ComponentKey);
+	TArray<FName> GetLayerNames(ALandscapeProxy* Landscape);
+
+private:
+#if WITH_EDITOR
+	void SetupLandscapeCallbacks();
+	void TeardownLandscapeCallbacks();
+	void OnLandscapeChanged(ALandscapeProxy* Landscape, const FLandscapeProxyComponentDataChangedParams& ChangeParams);
+	void CacheLayerNames(ALandscapeProxy* Landscape);
+	void CacheLayerNames();
+#endif
+
+	// TODO: need to have an indirection for multiple landscape support
+	UPROPERTY(VisibleAnywhere, Category = "Cache")
+	TMap<FIntPoint, FPCGLandscapeCacheEntry> CachedData;
+
+	//TODO: separate by landscape
+	UPROPERTY(VisibleAnywhere, Category = "Cache")
+	TSet<FName> CachedLayerNames;
+
+	// Transient by design
+	UObject* Owner = nullptr;
+
+#if WITH_EDITOR
+	TSet<TWeakObjectPtr<ALandscapeProxy>> Landscapes;
+#endif
+
+#if WITH_EDITOR
+	FRWLock CacheLock;
+#endif
+};
+
+template<>
+struct TStructOpsTypeTraits<FPCGLandscapeCache> : public TStructOpsTypeTraitsBase2< FPCGLandscapeCache>
+{
+	enum
+	{
+		WithCopy = false
+	};
+};
\ No newline at end of file
diff --git a/Engine/Plugins/Experimental/PCG/Source/PCG/Public/Metadata/PCGMetadata.h b/Engine/Plugins/Experimental/PCG/Source/PCG/Public/Metadata/PCGMetadata.h
--- a/Engine/Plugins/Experimental/PCG/Source/PCG/Public/Metadata/PCGMetadata.h
+++ b/Engine/Plugins/Experimental/PCG/Source/PCG/Public/Metadata/PCGMetadata.h
@@ -30,7 +30,7 @@ class PCG_API UPCGMetadata : public UObject
 public:
 
 	//~ Begin UObject interface
-	virtual void Serialize(FArchive& Ar) override;
+	virtual void Serialize(FArchive& InArchive) override;
 	virtual void BeginDestroy() override;
 	//~ End UObject interface
 
@@ -106,7 +106,9 @@ public:
 	/** Get attributes */
 	FPCGMetadataAttributeBase* GetMutableAttribute(FName AttributeName);
 	const FPCGMetadataAttributeBase* GetConstAttribute(FName AttributeName) const;
+	const FPCGMetadataAttributeBase* GetConstAttributeById(int32 AttributeId) const;
 	bool HasAttribute(FName AttributeName) const;
+	int32 GetAttributeCount() const;
 
 	UFUNCTION(BlueprintCallable, Category = "PCG|Metadata")
 	void GetAttributes(TArray<FName>& AttributeNames, TArray<EPCGMetadataTypes>& AttributeTypes) const;
diff --git a/Engine/Plugins/Experimental/PCG/Source/PCG/Public/Metadata/PCGMetadataAttribute.h b/Engine/Plugins/Experimental/PCG/Source/PCG/Public/Metadata/PCGMetadataAttribute.h
--- a/Engine/Plugins/Experimental/PCG/Source/PCG/Public/Metadata/PCGMetadataAttribute.h
+++ b/Engine/Plugins/Experimental/PCG/Source/PCG/Public/Metadata/PCGMetadataAttribute.h
@@ -16,6 +16,7 @@
 	FPCGMetadataAttributeBase(UPCGMetadata* InMetadata, FName InName, const FPCGMetadataAttributeBase* InParent, bool bInAllowsInterpolation);
 
 	virtual ~FPCGMetadataAttributeBase() = default;
+	virtual void Serialize(UPCGMetadata* InMetadata, FArchive& InArchive);
 
 	const UPCGMetadata* GetMetadata() const { return Metadata; }
 	int16 GetTypeId() const { return TypeId; }
@@ -37,7 +38,7 @@
 	bool AllowsInterpolation() const { return bAllowsInterpolation; }
 
 protected:
-	TMap<PCGMetadataEntryKey, PCGMetadataValueKey> EntryToValueMap;
+	TMap<PCGMetadataEntryKey, PCGMetadataValueKey> EntryToValueKeyMap;
 	mutable FRWLock EntryMapLock;
 
 	TObjectPtr<UPCGMetadata> Metadata = nullptr;
diff --git a/Engine/Plugins/Experimental/PCG/Source/PCG/Public/Metadata/PCGMetadataAttributeTpl.h b/Engine/Plugins/Experimental/PCG/Source/PCG/Public/Metadata/PCGMetadataAttributeTpl.h
--- a/Engine/Plugins/Experimental/PCG/Source/PCG/Public/Metadata/PCGMetadataAttributeTpl.h
+++ b/Engine/Plugins/Experimental/PCG/Source/PCG/Public/Metadata/PCGMetadataAttributeTpl.h
@@ -22,6 +22,26 @@ public:
 		}
 	}
 
+	// This constructor is used only during serialization
+	FPCGMetadataAttribute()
+	{
+		TypeId = PCG::Private::MetadataTypes<T>::Id;
+	}
+
+	virtual void Serialize(UPCGMetadata* InMetadata, FArchive& InArchive) override
+	{
+		FPCGMetadataAttributeBase::Serialize(InMetadata, InArchive);
+
+		InArchive << Values;
+		InArchive << DefaultValue;
+		
+		// Initialize non-serialized members
+		if (InArchive.IsLoading())
+		{
+			ValueKeyOffset = GetParent() ? GetParent()->GetValueKeyOffsetForChild() : 0;
+		}
+	}
+
 	const FPCGMetadataAttribute* GetParent() const { return static_cast<const FPCGMetadataAttribute*>(Parent); }
 
 	virtual FPCGMetadataAttributeBase* Copy(FName NewName, UPCGMetadata* InMetadata, bool bKeepParent, bool bCopyEntries = true, bool bCopyValues = true) const override
@@ -33,7 +53,7 @@ public:
 		if (bCopyEntries)
 		{
 			EntryMapLock.ReadLock();
-			AttributeCopy->EntryToValueMap = EntryToValueMap;
+			AttributeCopy->EntryToValueKeyMap = EntryToValueKeyMap;
 			EntryMapLock.ReadUnlock();
 		}
 
@@ -373,6 +393,35 @@ protected:
 protected:
 	mutable FRWLock ValueLock;
 	TArray<T> Values;
-	T DefaultValue;
+	T DefaultValue = T{};
 	PCGMetadataValueKey ValueKeyOffset = 0;
-};
\ No newline at end of file
+};
+
+namespace PCGMetadataAttribute
+{
+	inline FPCGMetadataAttributeBase* AllocateEmptyAttributeFromType(int16 TypeId)
+	{
+#define AllocatePCGMetadataAttributeOnType(Type) case PCG::Private::MetadataTypes<Type>::Id : { return new FPCGMetadataAttribute<Type>(); } break;
+
+		switch (TypeId)
+		{
+			AllocatePCGMetadataAttributeOnType(float);
+			AllocatePCGMetadataAttributeOnType(double);
+			AllocatePCGMetadataAttributeOnType(int32);
+			AllocatePCGMetadataAttributeOnType(int64);
+			AllocatePCGMetadataAttributeOnType(FVector);
+			AllocatePCGMetadataAttributeOnType(FVector4);
+			AllocatePCGMetadataAttributeOnType(FQuat);
+			AllocatePCGMetadataAttributeOnType(FTransform);
+			AllocatePCGMetadataAttributeOnType(FString);
+			AllocatePCGMetadataAttributeOnType(bool);
+			AllocatePCGMetadataAttributeOnType(FRotator);
+			AllocatePCGMetadataAttributeOnType(FName);
+
+		default:
+			return nullptr;
+		}
+
+#undef AllocatePCGMetadataAttributeOnType
+	}
+}
\ No newline at end of file
diff --git a/Engine/Plugins/Experimental/PCG/Source/PCG/Public/PCGSubsystem.h b/Engine/Plugins/Experimental/PCG/Source/PCG/Public/PCGSubsystem.h
--- a/Engine/Plugins/Experimental/PCG/Source/PCG/Public/PCGSubsystem.h
+++ b/Engine/Plugins/Experimental/PCG/Source/PCG/Public/PCGSubsystem.h
@@ -16,6 +16,7 @@ class APCGPartitionActor;
 class APCGWorldActor;
 class UPCGGraph;
 struct FPCGDataCollection;
+struct FPCGLandscapeCache;
 
 class IPCGElement;
 typedef TSharedPtr<IPCGElement, ESPMode::ThreadSafe> FPCGElementPtr;
@@ -47,9 +48,14 @@ public:
 	//~ End FTickableGameObject
 
 	APCGWorldActor* GetPCGWorldActor();
+#if WITH_EDITOR
+	void DestroyPCGWorldActor();
+#endif
 	void RegisterPCGWorldActor(APCGWorldActor* InActor);
 	void UnregisterPCGWorldActor(APCGWorldActor* InActor);
 
+	FPCGLandscapeCache* GetLandscapeCache();
+
 	// Schedule graph (owner -> graph)
 	FPCGTaskId ScheduleComponent(UPCGComponent* PCGComponent, bool bSave, const TArray<FPCGTaskId>& Dependencies);
 
@@ -112,11 +118,17 @@ public:
 	/** If the partition grid size change, call this to empty the Partition actors map */
 	void ResetPartitionActorsMap();
 
+	/** Builds the landscape data cache */
+	void BuildLandscapeCache();
+
+	/** Clears the landscape data cache */
+	void ClearLandscapeCache();
+
 private:
 	FPCGTaskId DelayProcessGraph(UPCGComponent* Component, bool bGenerate, bool bSave, bool bUseEmptyNewBounds);
 	FPCGTaskId ProcessGraph(UPCGComponent* Component, const FBox& InPreviousBounds, const FBox& InNewBounds, bool bGenerate, bool bSave);
 #endif // WITH_EDITOR
-
+	
 	// Schedule multiple graphs
 	TArray<FPCGTaskId> ScheduleMultipleComponent(UPCGComponent* OriginalComponent, TSet<TObjectPtr<APCGPartitionActor>>& PartitionActors, const TArray<FPCGTaskId>& Dependencies);
 
@@ -133,6 +145,10 @@ private:
 	APCGWorldActor* PCGWorldActor = nullptr;
 	FPCGGraphExecutor* GraphExecutor = nullptr;
 
+#if WITH_EDITOR
+	FCriticalSection PCGWorldActorLock;
+#endif
+
 	FPCGComponentOctree PCGComponentOctree;
 	TMap<TObjectPtr<const UPCGComponent>, FPCGComponentOctreeIDSharedRef> ComponentToIdMap;
 	mutable FRWLock PCGVolumeOctreeLock;
diff --git a/Engine/Plugins/Experimental/PCG/Source/PCG/Public/PCGWorldActor.h b/Engine/Plugins/Experimental/PCG/Source/PCG/Public/PCGWorldActor.h
--- a/Engine/Plugins/Experimental/PCG/Source/PCG/Public/PCGWorldActor.h
+++ b/Engine/Plugins/Experimental/PCG/Source/PCG/Public/PCGWorldActor.h
@@ -5,6 +5,8 @@
 #include "CoreMinimal.h"
 #include "GameFramework/Actor.h"
 
+#include "Grid/PCGLandscapeCache.h"
+
 #include "PCGWorldActor.generated.h"
 
 UCLASS(MinimalAPI, NotBlueprintable, NotPlaceable)
@@ -16,12 +18,14 @@ public:
 	APCGWorldActor(const FObjectInitializer& ObjectInitializer);
 
 	//~Begin AActor Interface
+	virtual void Serialize(FArchive& InArchive) override;
 	virtual void PostLoad() override;
 	virtual void BeginDestroy() override;
 
 #if WITH_EDITOR	
 	virtual bool CanChangeIsSpatiallyLoadedFlag() const { return false; }
 	virtual bool IsUserManaged() const override { return false; }
+	virtual void BeginCacheForCookedPlatformData(const ITargetPlatform* TargetPlatform);
 	//~End AActor Interface
 
 	static APCGWorldActor* CreatePCGWorldActor(UWorld* InWorld);
@@ -39,6 +43,9 @@ public:
 	UPROPERTY(config, EditAnywhere, Category = WorldPartition)
 	uint32 PartitionGridSize;
 
+	UPROPERTY(VisibleAnywhere, Category = CachedData)
+	FPCGLandscapeCache LandscapeCache;
+
 private:
 	void RegisterToSubsystem();
 	void UnregisterFromSubsystem();
diff --git a/Engine/Plugins/Experimental/PCG/Source/PCGEditor/Private/PCGEditorModule.cpp b/Engine/Plugins/Experimental/PCG/Source/PCGEditor/Private/PCGEditorModule.cpp
--- a/Engine/Plugins/Experimental/PCG/Source/PCGEditor/Private/PCGEditorModule.cpp
+++ b/Engine/Plugins/Experimental/PCG/Source/PCGEditor/Private/PCGEditorModule.cpp
@@ -13,6 +13,7 @@
 #include "PCGSubsystem.h"
 #include "PCGVolumeDetails.h"
 #include "PCGVolumeFactory.h"
+#include "PCGWorldActor.h"
 
 #include "EdGraphUtilities.h"
 #include "Framework/MultiBox/MultiBoxBuilder.h"
@@ -173,6 +174,63 @@ void FPCGEditorModule::PopulateMenuActions(FMenuBuilder& MenuBuilder)
 				}
 			})),
 		NAME_None);
+
+	MenuBuilder.AddMenuEntry(
+		LOCTEXT("DeletePCGWorldActor", "Delete PCG World Actor"),
+		LOCTEXT("DeletePCGWorldActor_Tooltip", "Deletes the PCG World Actor"),
+		FSlateIcon(),
+		FUIAction(
+			FExecuteAction::CreateLambda([]() {
+				if (GEditor)
+				{
+					if (UWorld* World = GEditor->GetEditorWorldContext().World())
+					{
+						if (APCGWorldActor* PCGWorldActor = World->GetSubsystem<UPCGSubsystem>()->GetPCGWorldActor())
+						{
+							World->GetSubsystem<UPCGSubsystem>()->DestroyPCGWorldActor();
+						}
+					}
+				}
+			})),
+		NAME_None);
+
+	MenuBuilder.AddMenuEntry(
+		LOCTEXT("BuildLandscapeCache", "Build Landscape Data Cache"),
+		LOCTEXT("BuildLandscapeCache_Tooltip", "Caches the landscape data in the PCG World Actor"),
+		FSlateIcon(),
+		FUIAction(
+			FExecuteAction::CreateLambda([]() {
+				if (GEditor)
+				{
+					if (UWorld* World = GEditor->GetEditorWorldContext().World())
+					{
+						if (UPCGSubsystem* Subsystem = World->GetSubsystem<UPCGSubsystem>())
+						{
+							Subsystem->BuildLandscapeCache();
+						}
+					}
+				}
+			})),
+		NAME_None);
+
+	MenuBuilder.AddMenuEntry(
+		LOCTEXT("ClearLandscapeCache", "Clear Landscape Data Cache"),
+		LOCTEXT("ClearLandscapeCache_Tooltip", "Clears the landscape data cache in the PCG World Actor"),
+		FSlateIcon(),
+		FUIAction(
+			FExecuteAction::CreateLambda([]() {
+				if (GEditor)
+				{
+					if (UWorld* World = GEditor->GetEditorWorldContext().World())
+					{
+						if (UPCGSubsystem* Subsystem = World->GetSubsystem<UPCGSubsystem>())
+						{
+							Subsystem->ClearLandscapeCache();
+						}
+					}
+				}
+				})),
+		NAME_None);
 }
 
 void FPCGEditorModule::RegisterSettings()
