Fixed node enabled state issue on load

#rnx

diff --git a/Engine/Plugins/Experimental/PCG/Source/PCGEditor/Private/PCGEditorGraphNodeBase.cpp b/Engine/Plugins/Experimental/PCG/Source/PCGEditor/Private/PCGEditorGraphNodeBase.cpp
--- a/Engine/Plugins/Experimental/PCG/Source/PCGEditor/Private/PCGEditorGraphNodeBase.cpp
+++ b/Engine/Plugins/Experimental/PCG/Source/PCGEditor/Private/PCGEditorGraphNodeBase.cpp
@@ -34,6 +34,12 @@ void UPCGEditorGraphNodeBase::Construct(UPCGNode* InPCGNode)
 	NodeComment = InPCGNode->NodeComment;
 	bCommentBubblePinned = InPCGNode->bCommentBubblePinned;
 	bCommentBubbleVisible = InPCGNode->bCommentBubbleVisible;
+
+	if (const UPCGSettings* PCGSettings = InPCGNode->DefaultSettings)
+	{
+		const ENodeEnabledState NewEnabledState = (PCGSettings->ExecutionMode == EPCGSettingsExecutionMode::Disabled) ? ENodeEnabledState::Disabled : ENodeEnabledState::Enabled;
+		SetEnabledState(NewEnabledState);
+	}
 }
 
 void UPCGEditorGraphNodeBase::BeginDestroy()
