Improved async dispatch (current thread always executes part of the work now)
Improved the GetPoint in the landscape data by caching some commonly used pointers
Improved behavior on the Bilerp on full point data

diff --git a/Engine/Plugins/Experimental/PCG/Source/PCG/Private/Data/PCGLandscapeData.cpp b/Engine/Plugins/Experimental/PCG/Source/PCG/Private/Data/PCGLandscapeData.cpp
--- a/Engine/Plugins/Experimental/PCG/Source/PCG/Private/Data/PCGLandscapeData.cpp
+++ b/Engine/Plugins/Experimental/PCG/Source/PCG/Private/Data/PCGLandscapeData.cpp
@@ -22,26 +22,24 @@ void UPCGLandscapeData::Initialize(ALandscapeProxy* InLandscape, const FBox& InB
 
 	Transform = Landscape->GetActorTransform();
 
+	// Cache additional data
+	LandscapeInfo = Landscape ? Landscape->GetLandscapeInfo() : nullptr;
+	LandscapeCache = Landscape && Landscape->GetWorld() && Landscape->GetWorld()->GetSubsystem<UPCGSubsystem>() ? Landscape->GetWorld()->GetSubsystem<UPCGSubsystem>()->GetLandscapeCache() : nullptr;
+
 	// TODO: find a better way to do this - maybe there should be a prototype metadata in the landscape cache
-	if (Landscape->GetWorld())
+	if (LandscapeCache)
 	{
-		if (UPCGSubsystem* Subsystem = Landscape->GetWorld()->GetSubsystem<UPCGSubsystem>())
-		{
-			if (FPCGLandscapeCache* LandscapeCache = Subsystem->GetLandscapeCache())
-			{
-				const TArray<FName> Layers = Subsystem->GetLandscapeCache()->GetLayerNames(Landscape.Get());
+		const TArray<FName> Layers = LandscapeCache->GetLayerNames(Landscape.Get());
 
-				for (const FName& Layer : Layers)
-				{
-					Metadata->CreateFloatAttribute(Layer, 0.0f, /*bAllowInterpolation=*/true);
-				}
-			}
-			else
-			{
-				UE_LOG(LogPCG, Error, TEXT("Landscape is unable to access the landscape cache"));
-			}
+		for (const FName& Layer : Layers)
+		{
+			Metadata->CreateFloatAttribute(Layer, 0.0f, /*bAllowInterpolation=*/true);
 		}
 	}
+	else
+	{
+		UE_LOG(LogPCG, Error, TEXT("Landscape is unable to access the landscape cache"));
+	}
 }
 
 FBox UPCGLandscapeData::GetBounds() const
@@ -58,24 +56,8 @@ FBox UPCGLandscapeData::GetStrictBounds() const
 
 bool UPCGLandscapeData::SamplePoint(const FTransform& InTransform, const FBox& InBounds, FPCGPoint& OutPoint, UPCGMetadata* OutMetadata) const
 {
-	if (!Landscape || !Landscape->GetWorld())
-	{
-		return false;
-	}
-
-	UPCGSubsystem* Subsystem = Landscape->GetWorld()->GetSubsystem<UPCGSubsystem>();
-
-	if (!Subsystem || !Subsystem->GetLandscapeCache())
-	{
-		UE_LOG(LogPCG, Error, TEXT("PCG Subsystem or landscape cache are not initialized"));
-		return false;
-	}
-
-	// Compute landscape-centric coordinates from the transform
-	// TODO: support bounds tests/interpolation
-	ULandscapeInfo* LandscapeInfo = Landscape->GetLandscapeInfo();
-
-	if (!LandscapeInfo)
+	//TRACE_CPUPROFILER_EVENT_SCOPE(UPCGLandscapeData::SamplePoint);
+	if (!LandscapeInfo || !LandscapeCache)
 	{
 		return false;
 	}
@@ -87,8 +69,7 @@ bool UPCGLandscapeData::SamplePoint(const FTransform& InTransform, const FBox& I
 
 	if (ULandscapeComponent* LandscapeComponent = LandscapeInfo->XYtoComponentMap.FindRef(ComponentMapKey))
 	{
-		check(Subsystem->GetLandscapeCache());
-		const FPCGLandscapeCacheEntry* LandscapeCacheEntry = Subsystem->GetLandscapeCache()->GetCacheEntry(LandscapeComponent, ComponentMapKey);
+		const FPCGLandscapeCacheEntry* LandscapeCacheEntry = LandscapeCache->GetCacheEntry(LandscapeComponent, ComponentMapKey);
 
 		if (!LandscapeCacheEntry)
 		{
@@ -102,27 +83,48 @@ bool UPCGLandscapeData::SamplePoint(const FTransform& InTransform, const FBox& I
 		const int32 X0Y1 = X0Y0 + (LandscapeInfo->ComponentSizeQuads + 1);
 		const int32 X1Y1 = X0Y1 + 1;
 
-		FPCGPoint PX0Y0;
-		FPCGPoint PX1Y0;
-		FPCGPoint PX0Y1;
-		FPCGPoint PX1Y1;
+		const float FractionalX = FMath::Fractional(LocalPoint.X);
+		const float FractionalY = FMath::Fractional(LocalPoint.Y);
 
-		LandscapeCacheEntry->GetPoint(X0Y0, PX0Y0, OutMetadata);
-		LandscapeCacheEntry->GetPoint(X1Y0, PX1Y0, OutMetadata);
-		LandscapeCacheEntry->GetPoint(X0Y1, PX0Y1, OutMetadata);
-		LandscapeCacheEntry->GetPoint(X1Y1, PX1Y1, OutMetadata);
-
-
-		PCGPointHelpers::Bilerp(
-			PX0Y0,
-			PX1Y0,
-			PX0Y1,
-			PX1Y1,
-			OutMetadata,
-			OutPoint,
-			OutMetadata,
-			FMath::Fractional(LocalPoint.X),
-			FMath::Fractional(LocalPoint.Y));
+		if (FractionalX < KINDA_SMALL_NUMBER && FractionalY < KINDA_SMALL_NUMBER)
+		{
+			LandscapeCacheEntry->GetPoint(X0Y0, OutPoint, OutMetadata);
+		}
+		else if (FractionalX < KINDA_SMALL_NUMBER && FractionalY > 1.0f - KINDA_SMALL_NUMBER)
+		{
+			LandscapeCacheEntry->GetPoint(X1Y0, OutPoint, OutMetadata);
+		}
+		else if (FractionalX > 1.0f - KINDA_SMALL_NUMBER && FractionalY < KINDA_SMALL_NUMBER)
+		{
+			LandscapeCacheEntry->GetPoint(X0Y1, OutPoint, OutMetadata);
+		}
+		else if (FractionalX > 1.0f - KINDA_SMALL_NUMBER && FractionalY < 1.0f - KINDA_SMALL_NUMBER)
+		{
+			LandscapeCacheEntry->GetPoint(X1Y1, OutPoint, OutMetadata);
+		}
+		else
+		{
+			FPCGPoint PX0Y0;
+			FPCGPoint PX1Y0;
+			FPCGPoint PX0Y1;
+			FPCGPoint PX1Y1;
+
+			LandscapeCacheEntry->GetPoint(X0Y0, PX0Y0, OutMetadata);
+			LandscapeCacheEntry->GetPoint(X1Y0, PX1Y0, OutMetadata);
+			LandscapeCacheEntry->GetPoint(X0Y1, PX0Y1, OutMetadata);
+			LandscapeCacheEntry->GetPoint(X1Y1, PX1Y1, OutMetadata);
+
+			PCGPointHelpers::Bilerp(
+				PX0Y0,
+				PX1Y0,
+				PX0Y1,
+				PX1Y1,
+				OutMetadata,
+				OutPoint,
+				OutMetadata,
+				FractionalX,
+				FractionalY);
+		}
 
 		return true;
 	}
@@ -136,23 +138,9 @@ const UPCGPointData* UPCGLandscapeData::CreatePointData(FPCGContext* Context, co
 {
 	TRACE_CPUPROFILER_EVENT_SCOPE(UPCGLandscapeData::CreatePointData);
 
-	if (!Landscape || !Landscape->GetWorld())
-	{
-		return nullptr;
-	}
-
-	UPCGSubsystem* Subsystem = Landscape->GetWorld()->GetSubsystem<UPCGSubsystem>();
-
-	if (!Subsystem || !Subsystem->GetLandscapeCache())
-	{
-		UE_LOG(LogPCG, Error, TEXT("PCG Subsystem or landscape cache are not initialized"));
-		return nullptr;
-	}
-
-	ULandscapeInfo* LandscapeInfo = Landscape->GetLandscapeInfo();
-
-	if (!LandscapeInfo)
+	if (!LandscapeInfo || !LandscapeCache)
 	{
+		UE_LOG(LogPCG, Error, TEXT("PCG Landscape cache or Landscape info are not initialized"));
 		return nullptr;
 	}
 
@@ -208,7 +196,7 @@ const UPCGPointData* UPCGLandscapeData::CreatePointData(FPCGContext* Context, co
 			FIntPoint ComponentMapKey(ComponentX, ComponentY);
 			if (ULandscapeComponent* LandscapeComponent = LandscapeInfo->XYtoComponentMap.FindRef(ComponentMapKey))
 			{
-				const FPCGLandscapeCacheEntry* LandscapeCacheEntry = Subsystem->GetLandscapeCache()->GetCacheEntry(LandscapeComponent, ComponentMapKey);
+				const FPCGLandscapeCacheEntry* LandscapeCacheEntry = LandscapeCache->GetCacheEntry(LandscapeComponent, ComponentMapKey);
 
 				if (!LandscapeCacheEntry)
 				{
diff --git a/Engine/Plugins/Experimental/PCG/Source/PCG/Private/Data/PCGPointData.cpp b/Engine/Plugins/Experimental/PCG/Source/PCG/Private/Data/PCGPointData.cpp
--- a/Engine/Plugins/Experimental/PCG/Source/PCG/Private/Data/PCGPointData.cpp
+++ b/Engine/Plugins/Experimental/PCG/Source/PCG/Private/Data/PCGPointData.cpp
@@ -105,7 +105,7 @@ namespace PCGPointHelpers
 		OutPoint.Color = FMath::Lerp(A.Color, B.Color, Ratio);
 		OutPoint.Steepness = FMath::Lerp(A.Steepness, B.Steepness, Ratio);
 
-		if (OutMetadata)
+		if (OutMetadata && SourceMetadata && SourceMetadata->GetAttributeCount() > 0)
 		{
 			UPCGMetadataAccessorHelpers::InitializeMetadataWithParent(OutPoint, OutMetadata, ((Ratio <= 0.5f) ? A : B), SourceMetadata);
 			OutMetadata->ResetPointWeightedAttributes(OutPoint);
@@ -114,48 +114,80 @@ namespace PCGPointHelpers
 		}
 	}
 
-	void Bilerp(const FPCGPoint& X0Y0, const FPCGPoint& X1Y0, const FPCGPoint& X0Y1, const FPCGPoint& X1Y1, const UPCGMetadata* SourceMetadata, FPCGPoint& OutPoint, UPCGMetadata* OutMetadata, float XFactor, float YFactor)
+	void BilerpWithSnapping(const FPCGPoint& X0Y0, const FPCGPoint& X1Y0, const FPCGPoint& X0Y1, const FPCGPoint& X1Y1, const UPCGMetadata* SourceMetadata, FPCGPoint& OutPoint, UPCGMetadata* OutMetadata, float XFactor, float YFactor)
 	{
-		const FPCGPoint* SelectedPoint = nullptr;
+		const bool bIsOnLeftEdge = (XFactor < KINDA_SMALL_NUMBER);
+		const bool bIsOnRightEdge = (XFactor > 1.0f - KINDA_SMALL_NUMBER);
+		const bool bIsOnTopEdge = (YFactor < KINDA_SMALL_NUMBER);
+		const bool bIsOnBottomEdge = (YFactor > 1.0f - KINDA_SMALL_NUMBER);
 
-		// First, check preexisting points as interpolation might be costly
-		if (XFactor < KINDA_SMALL_NUMBER && YFactor < KINDA_SMALL_NUMBER)
-		{
-			SelectedPoint = &X0Y0;
-		}
-		else if (XFactor > 1.0f - KINDA_SMALL_NUMBER && YFactor < KINDA_SMALL_NUMBER)
+		auto CopyPoint = [&OutPoint, &OutMetadata, &SourceMetadata](const FPCGPoint& PointToCopy)
 		{
-			SelectedPoint = &X1Y0;
-		}
-		else if (XFactor < KINDA_SMALL_NUMBER && YFactor > 1.0f - KINDA_SMALL_NUMBER)
-		{
-			SelectedPoint = &X0Y1;
-		}
-		else if (XFactor > 1.0f - KINDA_SMALL_NUMBER && YFactor > 1.0f - KINDA_SMALL_NUMBER)
-		{
-			SelectedPoint = &X1Y1;
-		}
+			OutPoint = PointToCopy;
+			if (OutMetadata)
+			{
+				OutMetadata->SetPointAttributes(PointToCopy, SourceMetadata, OutPoint);
+			}
+		};
 
-		if (SelectedPoint)
+		if (bIsOnLeftEdge || bIsOnRightEdge || bIsOnTopEdge || bIsOnBottomEdge)
 		{
-			OutPoint = *SelectedPoint;
-			if (OutMetadata)
+			if (bIsOnLeftEdge)
+			{
+				if (bIsOnTopEdge)
+				{
+					CopyPoint(X0Y0);
+				}
+				else if (bIsOnBottomEdge)
+				{
+					CopyPoint(X0Y1);
+				}
+				else
+				{
+					Lerp(X0Y0, X0Y1, YFactor, SourceMetadata, OutPoint, OutMetadata);
+				}
+			}
+			else if (bIsOnRightEdge)
 			{
-				OutMetadata->SetPointAttributes(*SelectedPoint, SourceMetadata, OutPoint);
+				if (bIsOnTopEdge)
+				{
+					CopyPoint(X1Y0);
+				}
+				else if (bIsOnBottomEdge)
+				{
+					CopyPoint(X1Y1);
+				}
+				else
+				{
+					Lerp(X1Y0, X1Y1, YFactor, SourceMetadata, OutPoint, OutMetadata);
+				}
+			}
+			else if (bIsOnTopEdge)
+			{
+				Lerp(X0Y0, X1Y0, XFactor, SourceMetadata, OutPoint, OutMetadata);
+			}
+			else // bIsOnBottomEdge
+			{
+				Lerp(X0Y1, X1Y1, XFactor, SourceMetadata, OutPoint, OutMetadata);
 			}
 		}
 		else
 		{
-			// Interpolate X0Y0-X1Y0 and X0Y1-X1Y1 using XFactor
-			FPCGPoint Y0Lerp;
-			FPCGPoint Y1Lerp;
-
-			Lerp(X0Y0, X1Y0, XFactor, SourceMetadata, Y0Lerp, OutMetadata);
-			Lerp(X0Y1, X1Y1, XFactor, SourceMetadata, Y1Lerp, OutMetadata);
-			// Interpolate between the two points using YFactor
-			Lerp(Y0Lerp, Y1Lerp, YFactor, SourceMetadata, OutPoint, OutMetadata);
+			Bilerp(X0Y0, X1Y0, X0Y1, X1Y1, SourceMetadata, OutPoint, OutMetadata, XFactor, YFactor);
 		}
 	}
+
+	void Bilerp(const FPCGPoint& X0Y0, const FPCGPoint& X1Y0, const FPCGPoint& X0Y1, const FPCGPoint& X1Y1, const UPCGMetadata* SourceMetadata, FPCGPoint& OutPoint, UPCGMetadata* OutMetadata, float XFactor, float YFactor)
+	{
+		// Interpolate X0Y0-X1Y0 and X0Y1-X1Y1 using XFactor
+		FPCGPoint Y0Lerp;
+		FPCGPoint Y1Lerp;
+
+		Lerp(X0Y0, X1Y0, XFactor, SourceMetadata, Y0Lerp, OutMetadata);
+		Lerp(X0Y1, X1Y1, XFactor, SourceMetadata, Y1Lerp, OutMetadata);
+		// Interpolate between the two points using YFactor
+		Lerp(Y0Lerp, Y1Lerp, YFactor, SourceMetadata, OutPoint, OutMetadata);
+	}
 }
 
 FPCGPointRef::FPCGPointRef(const FPCGPoint& InPoint)
diff --git a/Engine/Plugins/Experimental/PCG/Source/PCG/Private/Grid/PCGLandscapeCache.cpp b/Engine/Plugins/Experimental/PCG/Source/PCG/Private/Grid/PCGLandscapeCache.cpp
--- a/Engine/Plugins/Experimental/PCG/Source/PCG/Private/Grid/PCGLandscapeCache.cpp
+++ b/Engine/Plugins/Experimental/PCG/Source/PCG/Private/Grid/PCGLandscapeCache.cpp
@@ -24,6 +24,7 @@
 #if WITH_EDITOR
 void FPCGLandscapeCacheEntry::BuildCacheData(ULandscapeInfo* LandscapeInfo, ULandscapeComponent* InComponent, UObject* Owner)
 {
+	TRACE_CPUPROFILER_EVENT_SCOPE(FPCGLandscapeCacheEntry::BuildCacheData);
 	check(Owner && !Component.Get() && InComponent && PositionsAndNormals.Num() == 0);
 	Component = InComponent;
 
@@ -85,6 +86,7 @@ void FPCGLandscapeCacheEntry::BuildCacheData(ULandscapeInfo* LandscapeInfo, ULan
 
 bool FPCGLandscapeCacheEntry::GetPoint(int32 PointIndex, FPCGPoint& OutPoint, UPCGMetadata* OutMetadata) const
 {
+	//TRACE_CPUPROFILER_EVENT_SCOPE(FPCGLandscapeCacheEntry::GetPoint);
 	if (PointIndex < 0 || 2 * PointIndex >= PositionsAndNormals.Num())
 	{
 		return false;
diff --git a/Engine/Plugins/Experimental/PCG/Source/PCG/Private/Helpers/PCGAsync.cpp b/Engine/Plugins/Experimental/PCG/Source/PCG/Private/Helpers/PCGAsync.cpp
--- a/Engine/Plugins/Experimental/PCG/Source/PCG/Private/Helpers/PCGAsync.cpp
+++ b/Engine/Plugins/Experimental/PCG/Source/PCG/Private/Helpers/PCGAsync.cpp
@@ -18,68 +18,89 @@ namespace FPCGAsync
 	{
 		TRACE_CPUPROFILER_EVENT_SCOPE(IPCGElement::AsyncPointProcessing);
 		check(NumAvailableTasks > 0 && MinIterationsPerTask > 0 && NumIterations >= 0);
+		if (NumAvailableTasks <= 0 || MinIterationsPerTask <= 0 || NumIterations <= 0)
+		{
+			return;
+		}
+
 		// Get number of available threads from the context
 		const int32 NumTasks = FMath::Min(NumAvailableTasks, FMath::Max(1, NumIterations / MinIterationsPerTask));
 		const int32 IterationsPerTask = NumIterations / NumTasks;
-
+		const int32 NumFutures = NumTasks - 1;
+		 
 		{
 			TRACE_CPUPROFILER_EVENT_SCOPE(IPCGElement::AsyncPointProcessing::AllocatingArray);
 			// Pre-reserve the out points array
 			OutPoints.SetNum(NumIterations);
 		}
 
+		auto IterationInnerLoop = [&PointFunc, &OutPoints](int32 StartIndex, int32 EndIndex) -> int32
+		{
+			TRACE_CPUPROFILER_EVENT_SCOPE(IPCGElement::AsyncPointProcessing::InnerLoop);
+			int32 NumPointsWritten = 0;
+
+			for (int32 Index = StartIndex; Index < EndIndex; ++Index)
+			{
+				if (PointFunc(Index, OutPoints[StartIndex + NumPointsWritten]))
+				{
+					++NumPointsWritten;
+				}
+			}
+
+			return NumPointsWritten;
+		};
+
 		// Setup [current, last, nb points] data per dispatch
-		// Execute
 		TArray<TFuture<int32>> AsyncTasks;
-		AsyncTasks.Reserve(NumTasks);
+		AsyncTasks.Reserve(NumFutures);
 
 		// Launch the async tasks
-		for (int32 TaskIndex = 0; TaskIndex < NumTasks; ++TaskIndex)
+		for (int32 TaskIndex = 0; TaskIndex < NumFutures; ++TaskIndex)
 		{
 			const int32 StartIndex = TaskIndex * IterationsPerTask;
-			const int32 EndIndex = (TaskIndex == NumTasks - 1) ? NumIterations : (StartIndex + IterationsPerTask);
+			const int32 EndIndex = StartIndex + IterationsPerTask;
 
-			AsyncTasks.Emplace(Async(EAsyncExecution::ThreadPool, [&PointFunc, StartIndex, EndIndex, &OutPoints]()
+			AsyncTasks.Emplace(Async(EAsyncExecution::ThreadPool, [&IterationInnerLoop, StartIndex, EndIndex]()
 			{
-				TRACE_CPUPROFILER_EVENT_SCOPE(IPCGElement::AsyncPointProcessing::InnerLoop);
-				int32 NumPointsWritten = 0;
-
-				for (int32 Index = StartIndex; Index < EndIndex; ++Index)
-				{
-					if (PointFunc(Index, OutPoints[StartIndex + NumPointsWritten]))
-					{
-						++NumPointsWritten;
-					}
-				}
-
-				return NumPointsWritten;
+				return IterationInnerLoop(StartIndex, EndIndex);
 			}));
 		}
 
+		// Execute last batch locally
+		int32 NumPointsWrittenOnThisThread = IterationInnerLoop(NumFutures * IterationsPerTask, NumIterations);
+
 		// Wait/Gather results & collapse points
 		{
 			TRACE_CPUPROFILER_EVENT_SCOPE(IPCGElement::AsyncPointProcessing::WaitAndCollapseArray);
-			int RangeIndex = 0;
-			for (int32 AsyncIndex = 0; AsyncIndex < AsyncTasks.Num(); ++AsyncIndex)
-			{
-				const int32 StartPointsIndex = AsyncIndex * IterationsPerTask;
-
-				TFuture<int32>& AsyncTask = AsyncTasks[AsyncIndex];
-				AsyncTask.Wait();
-				int32 NumberOfPointsAdded = AsyncTask.Get();
 
+			auto CollapsePoints = [&OutPoints](int32 RangeIndex, int32 StartPointsIndex, int32 NumPointsToCollapse)
+			{
 				// Move points from [StartsPointIndex, StartsPointIndex + NumberPointsAdded] to [RangeIndex, RangeIndex + NumberPointsAdded]
 				if (StartPointsIndex != RangeIndex)
 				{
-					for (int32 MoveIndex = 0; MoveIndex < NumberOfPointsAdded; ++MoveIndex)
+					for (int32 MoveIndex = 0; MoveIndex < NumPointsToCollapse; ++MoveIndex)
 					{
 						OutPoints[RangeIndex + MoveIndex] = MoveTemp(OutPoints[StartPointsIndex + MoveIndex]);
 					}
 				}
 
-				RangeIndex += NumberOfPointsAdded;
+				return RangeIndex + NumPointsToCollapse;
+			};
+
+			int RangeIndex = 0;
+			for (int32 AsyncIndex = 0; AsyncIndex < AsyncTasks.Num(); ++AsyncIndex)
+			{
+				const int32 StartPointsIndex = AsyncIndex * IterationsPerTask;
+
+				TFuture<int32>& AsyncTask = AsyncTasks[AsyncIndex];
+				AsyncTask.Wait();
+				const int32 NumberOfPointsAdded = AsyncTask.Get();
+				RangeIndex = CollapsePoints(RangeIndex, StartPointsIndex, NumberOfPointsAdded);
 			}
 
+			// Finally, add current thread points
+			RangeIndex = CollapsePoints(RangeIndex, NumFutures * IterationsPerTask, NumPointsWrittenOnThisThread);
+
 			OutPoints.SetNum(RangeIndex);
 		}
 	}
@@ -94,9 +115,15 @@ namespace FPCGAsync
 	{
 		TRACE_CPUPROFILER_EVENT_SCOPE(IPCGElement::AsyncPointFilterProcessing);
 		check(NumAvailableTasks > 0 && MinIterationsPerTask > 0 && NumIterations >= 0);
+		if (NumAvailableTasks <= 0 || MinIterationsPerTask <= 0 || NumIterations <= 0)
+		{
+			return;
+		}
+
 		// Get number of available threads from the context
 		const int32 NumTasks = FMath::Min(NumAvailableTasks, FMath::Max(1, NumIterations / MinIterationsPerTask));
 		const int32 IterationsPerTask = NumIterations / NumTasks;
+		const int32 NumFutures = NumTasks - 1;
 
 		{
 			TRACE_CPUPROFILER_EVENT_SCOPE(IPCGElement::AsyncPointProcessing::AllocatingArray);
@@ -105,52 +132,54 @@ namespace FPCGAsync
 			OutFilterPoints.SetNum(NumIterations);
 		}
 
+		auto IterationInnerLoop = [&PointFunc, &InFilterPoints, &OutFilterPoints](int32 StartIndex, int32 EndIndex) -> TPair<int32, int32>
+		{
+			TRACE_CPUPROFILER_EVENT_SCOPE(IPCGElement::AsyncPointProcessing::InnerLoop);
+			int32 NumPointsInWritten = 0;
+			int32 NumPointsOutWritten = 0;
+
+			for (int32 Index = StartIndex; Index < EndIndex; ++Index)
+			{
+				if (PointFunc(Index, InFilterPoints[StartIndex + NumPointsInWritten], OutFilterPoints[StartIndex + NumPointsOutWritten]))
+				{
+					++NumPointsInWritten;
+				}
+				else
+				{
+					++NumPointsOutWritten;
+				}
+			}
+
+			return TPair<int32, int32>(NumPointsInWritten, NumPointsOutWritten);
+		};
+
 		// Setup [current, last, nb points] data per dispatch
-		// Execute
 		TArray<TFuture<TPair<int32, int32>>> AsyncTasks;
-		AsyncTasks.Reserve(NumTasks);
+		AsyncTasks.Reserve(NumFutures);
 
 		// Launch the async tasks
-		for (int32 TaskIndex = 0; TaskIndex < NumTasks; ++TaskIndex)
+		for (int32 TaskIndex = 0; TaskIndex < NumFutures; ++TaskIndex)
 		{
 			const int32 StartIndex = TaskIndex * IterationsPerTask;
-			const int32 EndIndex = (TaskIndex == NumTasks - 1) ? NumIterations : (StartIndex + IterationsPerTask);
-
-			AsyncTasks.Emplace(Async(EAsyncExecution::ThreadPool, [&PointFunc, StartIndex, EndIndex, &InFilterPoints, &OutFilterPoints]()
-				{
-					TRACE_CPUPROFILER_EVENT_SCOPE(IPCGElement::AsyncPointProcessing::InnerLoop);
-					int32 NumPointsInWritten = 0;
-					int32 NumPointsOutWritten = 0;
-
-					for (int32 Index = StartIndex; Index < EndIndex; ++Index)
-					{
-						if (PointFunc(Index, InFilterPoints[StartIndex + NumPointsInWritten], OutFilterPoints[StartIndex + NumPointsOutWritten]))
-						{
-							++NumPointsInWritten;
-						}
-						else
-						{
-							++NumPointsOutWritten;
-						}
-					}
+			const int32 EndIndex = StartIndex + IterationsPerTask;
 
-					return TPair<int32, int32>(NumPointsInWritten, NumPointsOutWritten);
-				}));
+			AsyncTasks.Emplace(Async(EAsyncExecution::ThreadPool, [&IterationInnerLoop, StartIndex, EndIndex]() -> TPair<int32, int32>
+			{
+				return IterationInnerLoop(StartIndex, EndIndex);
+			}));
 		}
 
+		// Launch remainder on current thread
+		TPair<int32, int32> NumPointsWrittenOnThisThread = IterationInnerLoop(NumFutures * IterationsPerTask, NumIterations);
+
 		// Wait/Gather results & collapse points
 		{
 			TRACE_CPUPROFILER_EVENT_SCOPE(IPCGElement::AsyncPointFilterProcessing::WaitAndCollapseArray);
-			int InFilterRangeIndex = 0;
-			int OutFilterRangeIndex = 0;
 
-			for (int32 AsyncIndex = 0; AsyncIndex < AsyncTasks.Num(); ++AsyncIndex)
+			auto CollapsePoints = [&InFilterPoints, &OutFilterPoints](const TPair<int32, int32>& RangeIndices, int32 StartPointsIndex, const TPair<int32, int32>& NumberOfPointsAdded) -> TPair<int32, int32>
 			{
-				const int32 StartPointsIndex = AsyncIndex * IterationsPerTask;
-
-				TFuture<TPair<int32, int32>>& AsyncTask = AsyncTasks[AsyncIndex];
-				AsyncTask.Wait();
-				TPair<int32, int32> NumberOfPointsAdded = AsyncTask.Get();
+				int32 InFilterRangeIndex = RangeIndices.Key;
+				int32 OutFilterRangeIndex = RangeIndices.Value;
 
 				// Move in-filter points
 				{
@@ -181,10 +210,27 @@ namespace FPCGAsync
 
 					OutFilterRangeIndex += NumOutFilterPoints;
 				}
+
+				return TPair<int32, int32>(InFilterRangeIndex, OutFilterRangeIndex);
+			};
+
+			TPair<int32, int32> RangeIndices = TPair<int32, int32>(0, 0);
+			for (int32 AsyncIndex = 0; AsyncIndex < AsyncTasks.Num(); ++AsyncIndex)
+			{
+				const int32 StartPointsIndex = AsyncIndex * IterationsPerTask;
+
+				TFuture<TPair<int32, int32>>& AsyncTask = AsyncTasks[AsyncIndex];
+				AsyncTask.Wait();
+				TPair<int32, int32> NumberOfPointsAdded = AsyncTask.Get();
+
+				RangeIndices = CollapsePoints(RangeIndices, StartPointsIndex, NumberOfPointsAdded);
 			}
 
-			InFilterPoints.SetNum(InFilterRangeIndex);
-			OutFilterPoints.SetNum(OutFilterRangeIndex);
+			// Finally, add this thread results
+			RangeIndices = CollapsePoints(RangeIndices, NumFutures * IterationsPerTask, NumPointsWrittenOnThisThread);
+
+			InFilterPoints.SetNum(RangeIndices.Key);
+			OutFilterPoints.SetNum(RangeIndices.Value);
 		}
 	}
 
@@ -198,33 +244,46 @@ namespace FPCGAsync
 	{
 		TRACE_CPUPROFILER_EVENT_SCOPE(IPCGElement::AsyncMultiPointProcessing);
 		check(NumAvailableTasks > 0 && MinIterationsPerTask > 0 && NumIterations >= 0);
+		if (NumAvailableTasks <= 0 || MinIterationsPerTask <= 0 || NumIterations <= 0)
+		{
+			return;
+		}
+
 		// Get number of available threads from the context
 		const int32 NumTasks = FMath::Min(NumAvailableTasks, FMath::Max(1, NumIterations / MinIterationsPerTask));
 		const int32 IterationsPerTask = NumIterations / NumTasks;
+		const int32 NumFutures = NumTasks - 1;
+
+		auto IterationInnerLoop = [&PointFunc](int32 StartIndex, int32 EndIndex) -> TArray<FPCGPoint>
+		{
+			TRACE_CPUPROFILER_EVENT_SCOPE(IPCGElement::AsyncMultiPointProcessing::InnerLoop);
+			TArray<FPCGPoint> OutPoints;
+
+			for (int32 Index = StartIndex; Index < EndIndex; ++Index)
+			{
+				OutPoints.Append(PointFunc(Index));
+			}
+
+			return OutPoints;
+		};
 
 		TArray<TFuture<TArray<FPCGPoint>>> AsyncTasks;
-		AsyncTasks.Reserve(NumTasks);
+		AsyncTasks.Reserve(NumFutures);
 
 		// Launch the async tasks
-		for (int32 TaskIndex = 0; TaskIndex < NumTasks; ++TaskIndex)
+		for (int32 TaskIndex = 0; TaskIndex < NumFutures; ++TaskIndex)
 		{
 			const int32 StartIndex = TaskIndex * IterationsPerTask;
-			const int32 EndIndex = (TaskIndex == NumTasks - 1) ? NumIterations : (StartIndex + IterationsPerTask);
+			const int32 EndIndex = StartIndex + IterationsPerTask;
 
-			AsyncTasks.Emplace(Async(EAsyncExecution::ThreadPool, [&PointFunc, StartIndex, EndIndex]()
+			AsyncTasks.Emplace(Async(EAsyncExecution::ThreadPool, [&IterationInnerLoop, StartIndex, EndIndex]() -> TArray<FPCGPoint>
 			{
-				TRACE_CPUPROFILER_EVENT_SCOPE(IPCGElement::AsyncMultiPointProcessing::InnerLoop);
-				TArray<FPCGPoint> OutPoints;
-
-				for (int32 Index = StartIndex; Index < EndIndex; ++Index)
-				{
-					OutPoints.Append(PointFunc(Index));
-				}
-
-				return OutPoints;
+				return IterationInnerLoop(StartIndex, EndIndex);
 			}));
 		}
 
+		TArray<FPCGPoint> PointsFromThisThread = IterationInnerLoop(NumFutures * IterationsPerTask, NumIterations);
+
 		// Wait/Gather results & collapse points
 		{
 			TRACE_CPUPROFILER_EVENT_SCOPE(IPCGElement::AsyncMultiPointProcessing::WaitAndCollapseArray);
@@ -233,6 +292,8 @@ namespace FPCGAsync
 				AsyncTask.Wait();
 				OutPoints.Append(AsyncTask.Get());
 			}
+
+			OutPoints.Append(PointsFromThisThread);
 		}
 	}
 }
\ No newline at end of file
diff --git a/Engine/Plugins/Experimental/PCG/Source/PCG/Public/Data/PCGLandscapeData.h b/Engine/Plugins/Experimental/PCG/Source/PCG/Public/Data/PCGLandscapeData.h
--- a/Engine/Plugins/Experimental/PCG/Source/PCG/Public/Data/PCGLandscapeData.h
+++ b/Engine/Plugins/Experimental/PCG/Source/PCG/Public/Data/PCGLandscapeData.h
@@ -8,6 +8,8 @@
 #include "PCGLandscapeData.generated.h"
 
 class ALandscapeProxy;
+class ULandscapeInfo;
+struct FPCGLandscapeCache;
 
 UCLASS(BlueprintType, ClassGroup = (Procedural))
 class PCG_API UPCGLandscapeData : public UPCGSurfaceData
@@ -41,4 +43,9 @@ public:
 protected:
 	UPROPERTY()
 	FBox Bounds = FBox(EForceInit::ForceInit);
+
+private:
+	// Transient data
+	ULandscapeInfo* LandscapeInfo = nullptr;
+	FPCGLandscapeCache* LandscapeCache = nullptr;
 };
diff --git a/Engine/Plugins/Experimental/PCG/Source/PCG/Public/Data/PCGPointData.h b/Engine/Plugins/Experimental/PCG/Source/PCG/Public/Data/PCGPointData.h
--- a/Engine/Plugins/Experimental/PCG/Source/PCG/Public/Data/PCGPointData.h
+++ b/Engine/Plugins/Experimental/PCG/Source/PCG/Public/Data/PCGPointData.h
@@ -16,6 +16,7 @@ namespace PCGPointHelpers
 {
 	void Lerp(const FPCGPoint& A, const FPCGPoint& B, float Ratio, const UPCGMetadata* SourceMetadata, FPCGPoint& OutPoint, UPCGMetadata* OutMetadata);
 	void Bilerp(const FPCGPoint& X0Y0, const FPCGPoint& X1Y0, const FPCGPoint& X0Y1, const FPCGPoint& X1Y1, const UPCGMetadata* SourceMetadata, FPCGPoint& OutPoint, UPCGMetadata* OutMetadata, float XFactor, float YFactor);
+	void BilerpWithSnapping(const FPCGPoint& X0Y0, const FPCGPoint& X1Y0, const FPCGPoint& X0Y1, const FPCGPoint& X1Y1, const UPCGMetadata* SourceMetadata, FPCGPoint& OutPoint, UPCGMetadata* OutMetadata, float XFactor, float YFactor);
 }
 
 struct PCG_API FPCGPointRef
