Implements a more efficient bilerp method for the pcg landscape cache

diff --git a/Engine/Plugins/Experimental/PCG/Source/PCG/Private/Data/PCGLandscapeData.cpp b/Engine/Plugins/Experimental/PCG/Source/PCG/Private/Data/PCGLandscapeData.cpp
--- a/Engine/Plugins/Experimental/PCG/Source/PCG/Private/Data/PCGLandscapeData.cpp
+++ b/Engine/Plugins/Experimental/PCG/Source/PCG/Private/Data/PCGLandscapeData.cpp
@@ -79,54 +79,11 @@ bool UPCGLandscapeData::SamplePoint(const FTransform& InTransform, const FBox& I
 		// Compute the 4 points indices
 		const FVector2D ComponentLocalPoint(LocalPoint.X - ComponentMapKey.X * LandscapeInfo->ComponentSizeQuads, LocalPoint.Y - ComponentMapKey.Y * LandscapeInfo->ComponentSizeQuads);
 		const int32 X0Y0 = FMath::FloorToInt(ComponentLocalPoint.X) + FMath::FloorToInt(ComponentLocalPoint.Y) * (LandscapeInfo->ComponentSizeQuads + 1);
-		const int32 X1Y0 = X0Y0 + 1;
-		const int32 X0Y1 = X0Y0 + (LandscapeInfo->ComponentSizeQuads + 1);
-		const int32 X1Y1 = X0Y1 + 1;
 
 		const float FractionalX = FMath::Fractional(LocalPoint.X);
 		const float FractionalY = FMath::Fractional(LocalPoint.Y);
 
-		if (FractionalX < KINDA_SMALL_NUMBER && FractionalY < KINDA_SMALL_NUMBER)
-		{
-			LandscapeCacheEntry->GetPoint(X0Y0, OutPoint, OutMetadata);
-		}
-		else if (FractionalX < KINDA_SMALL_NUMBER && FractionalY > 1.0f - KINDA_SMALL_NUMBER)
-		{
-			LandscapeCacheEntry->GetPoint(X1Y0, OutPoint, OutMetadata);
-		}
-		else if (FractionalX > 1.0f - KINDA_SMALL_NUMBER && FractionalY < KINDA_SMALL_NUMBER)
-		{
-			LandscapeCacheEntry->GetPoint(X0Y1, OutPoint, OutMetadata);
-		}
-		else if (FractionalX > 1.0f - KINDA_SMALL_NUMBER && FractionalY < 1.0f - KINDA_SMALL_NUMBER)
-		{
-			LandscapeCacheEntry->GetPoint(X1Y1, OutPoint, OutMetadata);
-		}
-		else
-		{
-			FPCGPoint PX0Y0;
-			FPCGPoint PX1Y0;
-			FPCGPoint PX0Y1;
-			FPCGPoint PX1Y1;
-
-			LandscapeCacheEntry->GetPoint(X0Y0, PX0Y0, OutMetadata);
-			LandscapeCacheEntry->GetPoint(X1Y0, PX1Y0, OutMetadata);
-			LandscapeCacheEntry->GetPoint(X0Y1, PX0Y1, OutMetadata);
-			LandscapeCacheEntry->GetPoint(X1Y1, PX1Y1, OutMetadata);
-
-			PCGPointHelpers::Bilerp(
-				PX0Y0,
-				PX1Y0,
-				PX0Y1,
-				PX1Y1,
-				OutMetadata,
-				OutPoint,
-				OutMetadata,
-				FractionalX,
-				FractionalY);
-		}
-
-		return true;
+		return LandscapeCacheEntry->GetInterpolatedPoint(X0Y0, LandscapeInfo->ComponentSizeQuads + 1, FractionalX, FractionalY, OutPoint, OutMetadata);
 	}
 	else
 	{
diff --git a/Engine/Plugins/Experimental/PCG/Source/PCG/Private/Grid/PCGLandscapeCache.cpp b/Engine/Plugins/Experimental/PCG/Source/PCG/Private/Grid/PCGLandscapeCache.cpp
--- a/Engine/Plugins/Experimental/PCG/Source/PCG/Private/Grid/PCGLandscapeCache.cpp
+++ b/Engine/Plugins/Experimental/PCG/Source/PCG/Private/Grid/PCGLandscapeCache.cpp
@@ -122,6 +122,69 @@ bool FPCGLandscapeCacheEntry::GetPoint(int32 PointIndex, FPCGPoint& OutPoint, UP
 	return true;
 }
 
+bool FPCGLandscapeCacheEntry::GetInterpolatedPoint(int32 BasePointIndex, int32 PointLineStride, float XFactor, float YFactor, FPCGPoint& OutPoint, UPCGMetadata* OutMetadata) const
+{
+	if (BasePointIndex < 0 || 2 * BasePointIndex >= PositionsAndNormals.Num())
+	{
+		return false;
+	}
+
+	const int32 X0Y0 = BasePointIndex;
+	const int32 X1Y0 = X0Y0 + 1;
+	const int32 X0Y1 = X0Y0 + PointLineStride;
+	const int32 X1Y1 = X0Y1 + 1;
+
+	const FVector& PositionX0Y0 = PositionsAndNormals[2 * X0Y0];
+	const FVector& NormalX0Y0 = PositionsAndNormals[2 * X0Y0 + 1];
+	const FVector& PositionX1Y0 = PositionsAndNormals[2 * X1Y0];
+	const FVector& NormalX1Y0 = PositionsAndNormals[2 * X1Y0 + 1];
+	const FVector& PositionX0Y1 = PositionsAndNormals[2 * X0Y1];
+	const FVector& NormalX0Y1 = PositionsAndNormals[2 * X0Y1 + 1];
+	const FVector& PositionX1Y1 = PositionsAndNormals[2 * X1Y1];
+	const FVector& NormalX1Y1 = PositionsAndNormals[2 * X1Y1 + 1];
+
+	FVector LerpPositionY0 = FMath::Lerp(PositionX0Y0, PositionX1Y0, XFactor);
+	FVector LerpPositionY1 = FMath::Lerp(PositionX0Y1, PositionX1Y1, XFactor);
+	FVector Position = FMath::Lerp(LerpPositionY0, LerpPositionY1, YFactor);
+
+	FVector LerpNormalY0 = FMath::Lerp(NormalX0Y0, NormalX1Y0, XFactor);
+	FVector LerpNormalY1 = FMath::Lerp(NormalX0Y1, NormalX1Y1, XFactor);
+	FVector Normal = FMath::Lerp(LerpNormalY0, LerpNormalY1, YFactor);
+
+	// TODO: we could preserve normal length by lerping this too?
+	Normal.Normalize();
+
+	FVector TangentX;
+	FVector TangentY;
+	TangentX = FVector(Normal.Z, 0.f, -Normal.X);
+	TangentY = Normal ^ TangentX;
+
+	OutPoint.Transform = FTransform(TangentX, TangentY, Normal, Position);
+	OutPoint.BoundsMin = -PointHalfSize;
+	OutPoint.BoundsMax = PointHalfSize;
+	OutPoint.Seed = UPCGBlueprintHelpers::ComputeSeedFromPosition(Position);
+
+	if (OutMetadata && !LayerData.IsEmpty())
+	{
+		OutPoint.MetadataEntry = OutMetadata->AddEntry();
+
+		for (const FPCGLandscapeCacheLayer& Layer : LayerData)
+		{
+			if (FPCGMetadataAttributeBase* Attribute = OutMetadata->GetMutableAttribute(Layer.Name))
+			{
+				check(Attribute->GetTypeId() == PCG::Private::MetadataTypes<float>::Id);
+				float Y0Data = FMath::Lerp((float)Layer.Data[X0Y0] / 255.0f, (float)Layer.Data[X1Y0] / 255.0f, XFactor);
+				float Y1Data = FMath::Lerp((float)Layer.Data[X0Y1] / 255.0f, (float)Layer.Data[X1Y1] / 255.0f, XFactor);
+				float Data = FMath::Lerp(Y0Data, Y1Data, YFactor);
+
+				static_cast<FPCGMetadataAttribute<float>*>(Attribute)->SetValue(OutPoint.MetadataEntry, Data);
+			}
+		}
+	}
+
+	return true;
+}
+
 FPCGLandscapeCache::FPCGLandscapeCache(UObject* InOwner)
 	: Owner(InOwner)
 {
diff --git a/Engine/Plugins/Experimental/PCG/Source/PCG/Public/Grid/PCGLandscapeCache.h b/Engine/Plugins/Experimental/PCG/Source/PCG/Public/Grid/PCGLandscapeCache.h
--- a/Engine/Plugins/Experimental/PCG/Source/PCG/Public/Grid/PCGLandscapeCache.h
+++ b/Engine/Plugins/Experimental/PCG/Source/PCG/Public/Grid/PCGLandscapeCache.h
@@ -34,6 +34,7 @@ public:
 #endif
 
 	bool GetPoint(int32 PointIndex, FPCGPoint& OutPoint, UPCGMetadata* OutMetadata) const;
+	bool GetInterpolatedPoint(int32 PointIndex, int32 PointLineStride, float XFactor, float YFactor, FPCGPoint& OutPoint, UPCGMetadata* OutMetadata) const;
 
 	// TODO: this isn't really needed
 	UPROPERTY()
