Support for 2D PCG partition grid

diff --git a/Engine/Plugins/Experimental/PCG/Source/PCG/Private/Grid/PCGPartitionActor.cpp b/Engine/Plugins/Experimental/PCG/Source/PCG/Private/Grid/PCGPartitionActor.cpp
--- a/Engine/Plugins/Experimental/PCG/Source/PCG/Private/Grid/PCGPartitionActor.cpp
+++ b/Engine/Plugins/Experimental/PCG/Source/PCG/Private/Grid/PCGPartitionActor.cpp
@@ -112,13 +112,22 @@ FBox APCGPartitionActor::GetFixedBounds() const
 {
 	const FVector Center = GetActorLocation();
 	const FVector::FReal HalfGridSize = PCGGridSize / 2.0;
-	return FBox(Center - HalfGridSize, Center + HalfGridSize);
+
+	FVector Extent(HalfGridSize, HalfGridSize, HalfGridSize);
+
+	// In case of 2D grid, it's like the actor has infinite bounds on the Z axis
+	if (bUse2DGrid)
+	{
+		Extent.Z = HALF_WORLD_MAX1;
+	}
+
+	return FBox(Center - Extent, Center + Extent);
 }
 
 FIntVector APCGPartitionActor::GetGridCoord() const
 {
 	const FVector Center = GetActorLocation();
-	return UPCGActorHelpers::GetCellCoord(Center, PCGGridSize);
+	return UPCGActorHelpers::GetCellCoord(Center, PCGGridSize, bUse2DGrid);
 }
 
 void APCGPartitionActor::GetActorBounds(bool bOnlyCollidingComponents, FVector& Origin, FVector& BoxExtent, bool bIncludeFromChildActors) const
@@ -256,8 +265,19 @@ void APCGPartitionActor::PostCreation()
 {
 	PCGGridSize = GridSize;
 
+	// Put in cache if we use the 2D grid or not.
+	if (APCGWorldActor* PCGActor = PCGHelpers::GetPCGWorldActor(GetWorld()))
+	{
+		bUse2DGrid = PCGActor->bUse2DGrid;
+	}
+	else
+	{
+		bUse2DGrid = false;
+	}
+
 #if WITH_EDITOR
-	if (BoundsComponent)
+	// Since we have infinite bounds in 2D, we just disable the bounds
+	if (BoundsComponent && !bUse2DGrid)
 	{
 		BoundsComponent->SetBoxExtent(GetFixedBounds().GetExtent());
 	}
diff --git a/Engine/Plugins/Experimental/PCG/Source/PCG/Private/Grid/PCGPartitionActor.h b/Engine/Plugins/Experimental/PCG/Source/PCG/Private/Grid/PCGPartitionActor.h
--- a/Engine/Plugins/Experimental/PCG/Source/PCG/Private/Grid/PCGPartitionActor.h
+++ b/Engine/Plugins/Experimental/PCG/Source/PCG/Private/Grid/PCGPartitionActor.h
@@ -46,6 +46,9 @@ public:
 	FBox GetFixedBounds() const;
 	FIntVector GetGridCoord() const;
 
+	uint32 GetGridSize() const { return PCGGridSize; }
+	bool IsUsing2DGrid() const { return bUse2DGrid; }
+
 	void AddGraphInstance(UPCGComponent* OriginalComponent);
 	bool RemoveGraphInstance(UPCGComponent* OriginalComponent);
 	bool CleanupDeadGraphInstances();
@@ -86,6 +89,9 @@ private:
 	UPROPERTY(VisibleAnywhere, Category = WorldPartition)
 	uint32 PCGGridSize;
 
+	UPROPERTY(VisibleAnywhere, Category = WorldPartition)
+	bool bUse2DGrid;
+
 #if WITH_EDITORONLY_DATA
 	/** Box component to draw the Partition actor bounds in the Editor viewport */
 	UPROPERTY(Transient)
diff --git a/Engine/Plugins/Experimental/PCG/Source/PCG/Private/Helpers/PCGActorHelpers.cpp b/Engine/Plugins/Experimental/PCG/Source/PCG/Private/Helpers/PCGActorHelpers.cpp
--- a/Engine/Plugins/Experimental/PCG/Source/PCG/Private/Helpers/PCGActorHelpers.cpp
+++ b/Engine/Plugins/Experimental/PCG/Source/PCG/Private/Helpers/PCGActorHelpers.cpp
@@ -385,15 +385,16 @@ AActor* UPCGActorHelpers::SpawnDefaultActor(UWorld* World, TSubclassOf<AActor> A
 	return NewActor;
 }
 
-FIntVector UPCGActorHelpers::GetCellCoord(FVector InPosition, int InGridSize)
+FIntVector UPCGActorHelpers::GetCellCoord(FVector InPosition, int InGridSize, bool bUse2DGrid)
 {
 	check(InGridSize > 0);
 
 	FVector Temp = InPosition / InGridSize;
 
+	// In case of 2D grid, Z coordinate is always 0
 	return FIntVector(
 		FMath::FloorToInt(Temp.X),
 		FMath::FloorToInt(Temp.Y),
-		FMath::FloorToInt(Temp.Z)
+		bUse2DGrid ? 0 : FMath::FloorToInt(Temp.Z)
 	);
 }
diff --git a/Engine/Plugins/Experimental/PCG/Source/PCG/Private/PCGComponent.cpp b/Engine/Plugins/Experimental/PCG/Source/PCG/Private/PCGComponent.cpp
--- a/Engine/Plugins/Experimental/PCG/Source/PCG/Private/PCGComponent.cpp
+++ b/Engine/Plugins/Experimental/PCG/Source/PCG/Private/PCGComponent.cpp
@@ -221,10 +221,6 @@ void UPCGComponent::GenerateLocal(bool bForce, EPCGComponentGenerationTrigger Re
 	{
 		bIsGenerating = true;
 	}
-	else
-	{
-		OnProcessGraphAborted();
-	}
 }
 
 FPCGTaskId UPCGComponent::GenerateInternal(bool bForce, EPCGComponentGenerationTrigger RequestedGenerationTrigger, const TArray<FPCGTaskId>& TaskDependencies)
@@ -253,10 +249,11 @@ FPCGTaskId UPCGComponent::GenerateInternal(bool bForce, EPCGComponentGenerationT
 	const FBox NewBounds = GetGridBounds();
 	if (!NewBounds.IsValid)
 	{
+		OnProcessGraphAborted();
 		return InvalidPCGTaskId;
 	}
 
-	return GetSubsystem()->ScheduleComponent(this, /*bSave=*/false, TaskDependencies);
+	return GetSubsystem()->ScheduleComponent(this, /*bSave=*/bForce, TaskDependencies);
 }
 
 bool UPCGComponent::GetActorsFromTags(const TSet<FName>& InTags, TSet<TWeakObjectPtr<AActor>>& OutActors, bool bCullAgainstLocalBounds)
diff --git a/Engine/Plugins/Experimental/PCG/Source/PCG/Private/PCGSubsystem.cpp b/Engine/Plugins/Experimental/PCG/Source/PCG/Private/PCGSubsystem.cpp
--- a/Engine/Plugins/Experimental/PCG/Source/PCG/Private/PCGSubsystem.cpp
+++ b/Engine/Plugins/Experimental/PCG/Source/PCG/Private/PCGSubsystem.cpp
@@ -89,7 +89,7 @@ void UPCGSubsystem::Tick(float DeltaSeconds)
 APCGWorldActor* UPCGSubsystem::GetPCGWorldActor()
 {
 #if WITH_EDITOR
-	if (!PCGWorldActor)
+	if (!PCGWorldActor && !PCGHelpers::IsRuntimeOrPIE())
 	{
 		PCGWorldActorLock.Lock();
 		if (!PCGWorldActor)
@@ -545,19 +545,26 @@ void UPCGSubsystem::UnregisterPartitionActor(APCGPartitionActor* Actor)
 
 void UPCGSubsystem::FindAllIntersectingPartitionActors(const FBox& InBounds, TFunction<void(APCGPartitionActor*)> InFunc) const
 {
-	// No World actor nor any PA registered, just early out. Same for invalid bounds.
-	if (!PCGWorldActor || PartitionActorsMap.IsEmpty() || !InBounds.IsValid)
+
+	// No PCGWorldActor just early out. Same for invalid bounds.
+	if (!PCGWorldActor || !InBounds.IsValid)
 	{
 		return;
 	}
 
 	const uint32 GridSize = PCGWorldActor->PartitionGridSize;
-	const FIntVector MinCellCoords = UPCGActorHelpers::GetCellCoord(InBounds.Min, GridSize);
-	const FIntVector MaxCellCoords = UPCGActorHelpers::GetCellCoord(InBounds.Max, GridSize);
+	const bool bUse2DGrid = PCGWorldActor->bUse2DGrid;
+	FIntVector MinCellCoords = UPCGActorHelpers::GetCellCoord(InBounds.Min, GridSize, bUse2DGrid);
+	FIntVector MaxCellCoords = UPCGActorHelpers::GetCellCoord(InBounds.Max, GridSize, bUse2DGrid);
 
 	{
 		FReadScopeLock ReadLock(PartitionActorsMapLock);
 
+		if (PartitionActorsMap.IsEmpty())
+		{
+			return;
+		}
+
 		for (int32 z = MinCellCoords.Z; z <= MaxCellCoords.Z; z++)
 		{
 			for (int32 y = MinCellCoords.Y; y <= MaxCellCoords.Y; y++)
@@ -590,12 +597,29 @@ namespace PCGSubsystem
 			return InvalidPCGTaskId;
 		}
 
+		APCGWorldActor* PCGWorldActor = PCGHelpers::GetPCGWorldActor(World);
+
+		check(PCGWorldActor);
+
+		// In case of 2D grid, we are clamping our bounds in Z to be within 0 and PartitionGridSize
+		// By doing so, WP will tie all the partition actors to [0, PartitionGridSize[ interval, generating a "2D grid" instead of 3D.
+		FBox ModifiedInBounds = InBounds;
+		if (PCGWorldActor->bUse2DGrid)
+		{
+			FVector MinBounds = InBounds.Min;
+			FVector MaxBounds = InBounds.Max;
+
+			MinBounds.Z = 0;
+			MaxBounds.Z = FMath::Max<FVector::FReal>(0.0, PCGWorldActor->PartitionGridSize - 1.0); // -1 to be just below the PartitionGridSize
+			ModifiedInBounds = FBox(MinBounds, MaxBounds);
+		}
+
 		TArray<FPCGTaskId> CellTasks;
 
-		auto CellLambda = [&CellTasks, GraphExecutor, World, bCreateActor, bLoadCell, bSaveActors, InBounds, &InOperation](const UActorPartitionSubsystem::FCellCoord& CellCoord, const FBox& CellBounds)
+		auto CellLambda = [&CellTasks, GraphExecutor, World, bCreateActor, bLoadCell, bSaveActors, ModifiedInBounds, &InOperation](const UActorPartitionSubsystem::FCellCoord& CellCoord, const FBox& CellBounds)
 		{
 			UActorPartitionSubsystem* PartitionSubsystem = UWorld::GetSubsystem<UActorPartitionSubsystem>(World);
-			FBox IntersectedBounds = InBounds.Overlap(CellBounds);
+			FBox IntersectedBounds = ModifiedInBounds.Overlap(CellBounds);
 
 			if (IntersectedBounds.IsValid)
 			{
@@ -717,7 +741,7 @@ namespace PCGSubsystem
 
 		// TODO: accumulate last phases of every cell + run a GC at the end?
 		//  or add some mechanism on the graph executor side to run GC every now and then
-		FActorPartitionGridHelper::ForEachIntersectingCell(APCGPartitionActor::StaticClass(), InBounds, World->PersistentLevel, CellLambda);
+		FActorPartitionGridHelper::ForEachIntersectingCell(APCGPartitionActor::StaticClass(), ModifiedInBounds, World->PersistentLevel, CellLambda);
 
 		// Finally, create a dummy generic task to wait on all cells
 		if (!CellTasks.IsEmpty())
@@ -732,100 +756,90 @@ namespace PCGSubsystem
 
 } // end namepsace PCGSubsystem
 
-FPCGTaskId UPCGSubsystem::ProcessGraph(UPCGComponent* Component, const FBox& InPreviousBounds, const FBox& InNewBounds, bool bGenerate, bool bSave)
+FPCGTaskId UPCGSubsystem::ProcessGraph(UPCGComponent* Component, const FBox& InPreviousBounds, const FBox& InNewBounds, EOperation InOperation, bool bSave)
 {
-	check(Component && (InPreviousBounds.IsValid || InNewBounds.IsValid));
+	check(Component);
 	TWeakObjectPtr<UPCGComponent> ComponentPtr(Component);
 
 	// TODO: optimal implementation would find the difference between the previous bounds and the new bounds
 	// and process these only. This is esp. important because of the CreateActor parameter.
-	auto ScheduleTask = [this, ComponentPtr, InPreviousBounds, InNewBounds, bGenerate](APCGPartitionActor* PCGActor, const FBox& InBounds, const TArray<FPCGTaskId>& TaskDependencies){
-		auto UnpartitionTask = [ComponentPtr, PCGActor]() {
+	auto ScheduleTask = [this, ComponentPtr, InOperation](APCGPartitionActor* PCGActor, const FBox& InBounds, const TArray<FPCGTaskId>& TaskDependencies){
+		TWeakObjectPtr<APCGPartitionActor> PCGActorPtr(PCGActor);
+
+		auto UnpartitionTask = [ComponentPtr, PCGActorPtr]() {
 			// TODO: PCG actors that become empty could be deleted, but we also need to keep track
 			// of packages that would need to be deleted from SCC.
-			if (UPCGComponent* Component = ComponentPtr.Get())
-			{
-				PCGActor->RemoveGraphInstance(Component);
-			}
-			else
+			if (APCGPartitionActor* PCGActor = PCGActorPtr.Get())
 			{
+				if (UPCGComponent* Component = ComponentPtr.Get())
+				{
+					PCGActor->RemoveGraphInstance(Component);
+				}
+				else
+				{
 #if WITH_EDITOR
-				PCGActor->CleanupDeadGraphInstances();
+					PCGActor->CleanupDeadGraphInstances();
 #endif
+				}
 			}
 			return true;
 		};
 
-		auto PartitionTask = [ComponentPtr, PCGActor]() {
-			if (UPCGComponent* Component = ComponentPtr.Get())
+		auto PartitionTask = [ComponentPtr, PCGActorPtr]() {
+			if (APCGPartitionActor* PCGActor = PCGActorPtr.Get())
 			{
-				PCGActor->AddGraphInstance(Component);
-			}			
+				if (UPCGComponent* Component = ComponentPtr.Get())
+				{
+					PCGActor->AddGraphInstance(Component);
+				}
+			}
+
 			return true;
 		};
 
-		auto ScheduleGraph = [this, PCGActor, ComponentPtr, &TaskDependencies]() {
-			if (UPCGComponent* Component = ComponentPtr.Get())
+		auto ScheduleGraph = [this, PCGActorPtr, ComponentPtr, &TaskDependencies]() {
+			if (APCGPartitionActor* PCGActor = PCGActorPtr.Get())
 			{
-				// Ensure that the PCG actor has a matching local component.
-				// This is done immediately, but technically we could add it as a task
-				PCGActor->AddGraphInstance(Component);
+				if (UPCGComponent* Component = ComponentPtr.Get())
+				{
+					// Ensure that the PCG actor has a matching local component.
+					// This is done immediately, but technically we could add it as a task
+					PCGActor->AddGraphInstance(Component);
 
-				UPCGComponent* LocalComponent = PCGActor->GetLocalComponent(Component);
+					UPCGComponent* LocalComponent = PCGActor->GetLocalComponent(Component);
+
+					if (!LocalComponent)
+					{
+						return InvalidPCGTaskId;
+					}
 
-				if (!LocalComponent)
+					return LocalComponent->GenerateInternal(/*bForce=*/false, EPCGComponentGenerationTrigger::GenerateOnDemand, TaskDependencies);
+				}
+				else
 				{
+					UE_LOG(LogPCG, Error, TEXT("[ProcessGraph] PCG Component on PCG Actor is null"));
 					return InvalidPCGTaskId;
 				}
-
-				return LocalComponent->GenerateInternal(/*bForce=*/false, EPCGComponentGenerationTrigger::GenerateOnDemand, TaskDependencies);
-			}
-			else
-			{
-				UE_LOG(LogPCG, Error, TEXT("[ProcessGraph] PCG Component on PCG Actor is null"));
-				return InvalidPCGTaskId;
 			}
+			
+			return InvalidPCGTaskId;
 		};
 
-		const bool bIsInOldBounds = InPreviousBounds.IsValid && InBounds.Intersect(InPreviousBounds);
-		const bool bIsInNewBounds = InNewBounds.IsValid && InBounds.Intersect(InNewBounds);
-
-		// 4 cases here:
-		// old & new -> generate only
-		// old & !new -> unpartition
-		// !old & new -> partition & generate
-		// !old & !new -> nothing to do
-		if (bIsInOldBounds && !bIsInNewBounds)
+		switch (InOperation)
 		{
+		case EOperation::Unpartition:
 			return GraphExecutor->ScheduleGeneric(UnpartitionTask, TaskDependencies);
-		}
-		else if (!bIsInOldBounds && bIsInNewBounds)
-		{
-			// Here, we should try to schedule the local component, but it doesn't exist yet.
-			// what to do? in practice, it's not a problem to not execute the Partition in a separate process,
-			// since there won't be any dependencies there in this case; but, we should execute it in a task
-			// if it's the only thing we do, otherwise we'll have issues with the other operations (save, etc.)
-			if (bGenerate)
-			{
-				return ScheduleGraph();
-			}
-			else
-			{
-				return GraphExecutor->ScheduleGeneric(PartitionTask, TaskDependencies);
-			}
-		}
-		else if (bIsInOldBounds && bIsInNewBounds && bGenerate)
-		{
+		case EOperation::Partition:
+			return GraphExecutor->ScheduleGeneric(PartitionTask, TaskDependencies);
+		case EOperation::Generate:
+		default:
 			return ScheduleGraph();
 		}
-		else
-		{
-			return InvalidPCGTaskId;
-		}
 	};
 
 	FBox UnionBounds = InPreviousBounds + InNewBounds;
-	const bool bCreateActors = (!UnionBounds.Equals(InPreviousBounds));
+	const bool bGenerate = InOperation == EOperation::Generate;
+	const bool bCreateActors = InOperation != EOperation::Unpartition;
 	const bool bLoadCell = (bGenerate && bSave); // TODO: review this
 	const bool bSaveActors = (bSave);
 
@@ -1151,30 +1165,30 @@ void UPCGSubsystem::ResetPartitionActorsMap()
 ///////////////////////////////////////////////////////////////////////////////////////////
 void UPCGSubsystem::DelayPartitionGraph(UPCGComponent* Component)
 {
-	DelayProcessGraph(Component, /*bGenerate=*/false, /*bSave=*/false, /*bUseEmptyNewBounds=*/false);
+	DelayProcessGraph(Component, EOperation::Partition, /*bSave=*/false);
 }
 
 void UPCGSubsystem::DelayUnpartitionGraph(UPCGComponent* Component)
 {
-	DelayProcessGraph(Component, /*bGenerate=*/false, /*bSave=*/false, /*bUseEmptyNewBounds=*/true);
+	DelayProcessGraph(Component, EOperation::Unpartition, /*bSave=*/false);
 }
 
 FPCGTaskId UPCGSubsystem::DelayGenerateGraph(UPCGComponent* Component, bool bSave)
 {
-	return DelayProcessGraph(Component, /*bGenerate=*/true, bSave, /*bUseEmptyNewBounds=*/false);
+	return DelayProcessGraph(Component, EOperation::Generate, bSave);
 }
 
-FPCGTaskId UPCGSubsystem::DelayProcessGraph(UPCGComponent* Component, bool bGenerate, bool bSave, bool bUseEmptyNewBounds)
+FPCGTaskId UPCGSubsystem::DelayProcessGraph(UPCGComponent* Component, EOperation InOperation, bool bSave)
 {
 	check(Component && Component->IsPartitioned());
 	TWeakObjectPtr<UPCGComponent> ComponentPtr(Component);
 
-	auto ExecuteProcessGraph = [this, ComponentPtr, bGenerate, bSave, bUseEmptyNewBounds]() {
+	auto ExecuteProcessGraph = [this, ComponentPtr, InOperation, bSave]() {
 		if (UPCGComponent* Component = ComponentPtr.Get())
 		{
 			const FBox PreviousBounds = Component->LastGeneratedBounds;
-			FBox NewBounds = bUseEmptyNewBounds ? FBox(EForceInit::ForceInit) : Component->GetGridBounds();
-			/*FPCGTaskId GraphTaskId = */ ProcessGraph(Component, PreviousBounds, NewBounds, bGenerate, bSave);
+			const FBox NewBounds = Component->GetGridBounds();
+			/*FPCGTaskId GraphTaskId = */ ProcessGraph(Component, PreviousBounds, NewBounds, InOperation, bSave);
 		}
 
 		return true;
diff --git a/Engine/Plugins/Experimental/PCG/Source/PCG/Private/PCGWorldActor.cpp b/Engine/Plugins/Experimental/PCG/Source/PCG/Private/PCGWorldActor.cpp
--- a/Engine/Plugins/Experimental/PCG/Source/PCG/Private/PCGWorldActor.cpp
+++ b/Engine/Plugins/Experimental/PCG/Source/PCG/Private/PCGWorldActor.cpp
@@ -99,7 +99,7 @@ void APCGWorldActor::OnPartitionGridSizeChanged()
 		return;
 	}
 
-	// First gather all the PCG components that linked to PCGPartitionActor and generated
+	// First gather all the PCG components that linked to PCGPartitionActor
 	TSet<TObjectPtr<UPCGComponent>> AllPartitionedComponents;
 	
 	bool bAllSafeToDelete = true;
@@ -116,7 +116,7 @@ void APCGWorldActor::OnPartitionGridSizeChanged()
 
 		for (UPCGComponent* PCGComponent : PartitionActor->GetAllOriginalPCGComponents())
 		{
-			if (PCGComponent && PCGComponent->bGenerated)
+			if (PCGComponent)
 			{
 				AllPartitionedComponents.Add(PCGComponent);
 			}
@@ -138,7 +138,7 @@ void APCGWorldActor::OnPartitionGridSizeChanged()
 	// Then delete all PCGPartitionActors
 	PCGSubsystem->DeletePartitionActors();
 
-	// And finally, refresh all components that were generated.
+	// And finally, refresh all components
 	for (TObjectPtr<UPCGComponent> PCGComponent : AllPartitionedComponents)
 	{
 		PCGComponent->DirtyGenerated();
@@ -148,7 +148,8 @@ void APCGWorldActor::OnPartitionGridSizeChanged()
 
 void APCGWorldActor::PostEditChangeProperty(FPropertyChangedEvent& PropertyChangedEvent)
 {
-	if (PropertyChangedEvent.GetPropertyName() == GET_MEMBER_NAME_CHECKED(APCGWorldActor, PartitionGridSize))
+	if (PropertyChangedEvent.GetPropertyName() == GET_MEMBER_NAME_CHECKED(APCGWorldActor, PartitionGridSize) 
+		|| PropertyChangedEvent.GetPropertyName() == GET_MEMBER_NAME_CHECKED(APCGWorldActor, bUse2DGrid))
 	{
 		OnPartitionGridSizeChanged();
 	}
diff --git a/Engine/Plugins/Experimental/PCG/Source/PCG/Public/Helpers/PCGActorHelpers.h b/Engine/Plugins/Experimental/PCG/Source/PCG/Public/Helpers/PCGActorHelpers.h
--- a/Engine/Plugins/Experimental/PCG/Source/PCG/Public/Helpers/PCGActorHelpers.h
+++ b/Engine/Plugins/Experimental/PCG/Source/PCG/Public/Helpers/PCGActorHelpers.h
@@ -100,5 +100,5 @@ public:
 	/**
 	 * Return the grid cell coordinates on the PCG partition grid given a position and the grid size.
 	 */
-	static FIntVector GetCellCoord(FVector InPosition, int InGridSize);
+	static FIntVector GetCellCoord(FVector InPosition, int InGridSize, bool bUse2DGrid);
 };
\ No newline at end of file
diff --git a/Engine/Plugins/Experimental/PCG/Source/PCG/Public/PCGSubsystem.h b/Engine/Plugins/Experimental/PCG/Source/PCG/Public/PCGSubsystem.h
--- a/Engine/Plugins/Experimental/PCG/Source/PCG/Public/PCGSubsystem.h
+++ b/Engine/Plugins/Experimental/PCG/Source/PCG/Public/PCGSubsystem.h
@@ -125,10 +125,18 @@ public:
 	void ClearLandscapeCache();
 
 private:
-	FPCGTaskId DelayProcessGraph(UPCGComponent* Component, bool bGenerate, bool bSave, bool bUseEmptyNewBounds);
-	FPCGTaskId ProcessGraph(UPCGComponent* Component, const FBox& InPreviousBounds, const FBox& InNewBounds, bool bGenerate, bool bSave);
+	enum class EOperation : uint32
+	{
+		Partition,
+		Unpartition,
+		Generate
+	};
+
+	FPCGTaskId DelayProcessGraph(UPCGComponent* Component, EOperation InOperation, bool bSave);
+	FPCGTaskId ProcessGraph(UPCGComponent* Component, const FBox& InPreviousBounds, const FBox& InNewBounds, EOperation InOperation, bool bSave);
 #endif // WITH_EDITOR
 	
+private:
 	// Schedule multiple graphs
 	TArray<FPCGTaskId> ScheduleMultipleComponent(UPCGComponent* OriginalComponent, TSet<TObjectPtr<APCGPartitionActor>>& PartitionActors, const TArray<FPCGTaskId>& Dependencies);
 
diff --git a/Engine/Plugins/Experimental/PCG/Source/PCG/Public/PCGWorldActor.h b/Engine/Plugins/Experimental/PCG/Source/PCG/Public/PCGWorldActor.h
--- a/Engine/Plugins/Experimental/PCG/Source/PCG/Public/PCGWorldActor.h
+++ b/Engine/Plugins/Experimental/PCG/Source/PCG/Public/PCGWorldActor.h
@@ -46,6 +46,10 @@ public:
 	UPROPERTY(VisibleAnywhere, Category = CachedData)
 	FPCGLandscapeCache LandscapeCache;
 
+	/** Disable creation of Partition Actors on the Z axis. Can improve performances if 3D partitioning is not needed. */
+	UPROPERTY(config, EditAnywhere, Category = WorldPartition)
+	bool bUse2DGrid = false;
+
 private:
 	void RegisterToSubsystem();
 	void UnregisterFromSubsystem();
