Fixed guid/naming issue with basic layer parsing. Might be revisited for edit layers.
Fixed issues in attributes merge due to presence of default attributes.

#rnx

diff --git a/Engine/Plugins/Experimental/PCG/Source/PCG/Private/Data/PCGDifferenceData.cpp b/Engine/Plugins/Experimental/PCG/Source/PCG/Private/Data/PCGDifferenceData.cpp
--- a/Engine/Plugins/Experimental/PCG/Source/PCG/Private/Data/PCGDifferenceData.cpp
+++ b/Engine/Plugins/Experimental/PCG/Source/PCG/Private/Data/PCGDifferenceData.cpp
@@ -31,6 +31,9 @@ void UPCGDifferenceData::Initialize(const UPCGSpatialData* InData)
 	check(InData);
 	Source = InData;
 	TargetActor = InData->TargetActor;
+
+	check(Metadata);
+	Metadata->Initialize(Source->Metadata);
 }
 
 void UPCGDifferenceData::AddDifference(const UPCGSpatialData* InDifference)
@@ -123,7 +126,7 @@ bool UPCGDifferenceData::SamplePoint(const FTransform& InTransform, const FBox&
 		// Color?
 		if (OutMetadata && OutPoint.Density > 0)
 		{
-			OutMetadata->MergePointAttributes(PointFromSource, PointFromDiff, OutPoint, EPCGMetadataOp::Sub);
+			OutMetadata->MergePointAttributesSubset(PointFromSource, OutMetadata, Source->Metadata, PointFromDiff, OutMetadata, Difference->Metadata, OutPoint, EPCGMetadataOp::Sub);
 		}
 
 		return OutPoint.Density > 0;
@@ -180,7 +183,7 @@ const UPCGPointData* UPCGDifferenceData::CreatePointData(FPCGContext* Context) c
 
 			if (Data->Metadata && OutPoint.Density > 0)
 			{
-				Data->Metadata->MergePointAttributes(Point, SourcePointData->Metadata, PointFromDiff, Data->Metadata, OutPoint, EPCGMetadataOp::Sub);
+				Data->Metadata->MergePointAttributesSubset(Point, SourcePointData->Metadata, SourcePointData->Metadata, PointFromDiff, Data->Metadata, Difference->Metadata, OutPoint, EPCGMetadataOp::Sub);
 			}
 
 #if WITH_EDITORONLY_DATA
diff --git a/Engine/Plugins/Experimental/PCG/Source/PCG/Private/Data/PCGIntersectionData.cpp b/Engine/Plugins/Experimental/PCG/Source/PCG/Private/Data/PCGIntersectionData.cpp
--- a/Engine/Plugins/Experimental/PCG/Source/PCG/Private/Data/PCGIntersectionData.cpp
+++ b/Engine/Plugins/Experimental/PCG/Source/PCG/Private/Data/PCGIntersectionData.cpp
@@ -30,6 +30,19 @@ void UPCGIntersectionData::Initialize(const UPCGSpatialData* InA, const UPCGSpat
 
 	CachedBounds = PCGHelpers::OverlapBounds(A->GetBounds(), B->GetBounds());
 	CachedStrictBounds = PCGHelpers::OverlapBounds(A->GetStrictBounds(), B->GetStrictBounds());
+
+	check(Metadata);
+	// Note: this should behave the same way as the ToPointData
+	if (A->GetDimension() <= B->GetDimension())
+	{
+		Metadata->Initialize(A->Metadata);
+		Metadata->AddAttributes(B->Metadata);
+	}
+	else
+	{
+		Metadata->Initialize(B->Metadata);
+		Metadata->AddAttributes(A->Metadata);
+	}
 }
 
 int UPCGIntersectionData::GetDimension() const
@@ -75,7 +88,7 @@ bool UPCGIntersectionData::SamplePoint(const FTransform& InTransform, const FBox
 
 	if (OutMetadata)
 	{
-		OutMetadata->MergePointAttributes(PointFromX, PointFromY, OutPoint, EPCGMetadataOp::Min);
+		OutMetadata->MergePointAttributesSubset(PointFromX, OutMetadata, X->Metadata, PointFromY, OutMetadata, Y->Metadata, OutPoint, EPCGMetadataOp::Min);
 	}
 
 	return true;
@@ -146,7 +159,7 @@ UPCGPointData* UPCGIntersectionData::CreateAndFilterPointData(FPCGContext* Conte
 
 		if (Data->Metadata)
 		{
-			Data->Metadata->MergePointAttributes(Point, SourcePointData->Metadata, PointFromY, Data->Metadata, OutPoint, EPCGMetadataOp::Min);
+			Data->Metadata->MergePointAttributesSubset(Point, SourcePointData->Metadata, SourcePointData->Metadata, PointFromY, Data->Metadata, Y->Metadata, OutPoint, EPCGMetadataOp::Min);
 		}
 
 		return true;
diff --git a/Engine/Plugins/Experimental/PCG/Source/PCG/Private/Data/PCGProjectionData.cpp b/Engine/Plugins/Experimental/PCG/Source/PCG/Private/Data/PCGProjectionData.cpp
--- a/Engine/Plugins/Experimental/PCG/Source/PCG/Private/Data/PCGProjectionData.cpp
+++ b/Engine/Plugins/Experimental/PCG/Source/PCG/Private/Data/PCGProjectionData.cpp
@@ -94,7 +94,7 @@ bool UPCGProjectionData::SamplePoint(const FTransform& InTransform, const FBox&
 	if (OutMetadata)
 	{
 		//METADATA TODO Review op
-		OutMetadata->MergePointAttributes(PointFromSource, PointFromTarget, OutPoint, EPCGMetadataOp::Max);
+		OutMetadata->MergePointAttributesSubset(PointFromSource, OutMetadata, Source->Metadata, PointFromTarget, OutMetadata, Target->Metadata, OutPoint, EPCGMetadataOp::Max);
 	}
 
 	return true;
@@ -115,6 +115,8 @@ const UPCGPointData* UPCGProjectionData::CreatePointData(FPCGContext* Context) c
 
 	UPCGPointData* PointData = NewObject<UPCGPointData>(const_cast<UPCGProjectionData*>(this));
 	PointData->InitializeFromData(this, SourcePointData->Metadata);
+	check(PointData->Metadata);
+	PointData->Metadata->AddAttributes(Target->Metadata);
 
 	TArray<FPCGPoint>& Points = PointData->GetMutablePoints();
 
@@ -142,7 +144,7 @@ const UPCGPointData* UPCGProjectionData::CreatePointData(FPCGContext* Context) c
 		if (PointData->Metadata)
 		{
 			//METADATA TODO review op
-			PointData->Metadata->MergePointAttributes(SourcePoint, SourcePointData->Metadata, PointFromTarget, PointData->Metadata, OutPoint, EPCGMetadataOp::Max);
+			PointData->Metadata->MergePointAttributesSubset(SourcePoint, SourcePointData->Metadata, SourcePointData->Metadata, PointFromTarget, PointData->Metadata, Target->Metadata, OutPoint, EPCGMetadataOp::Max);
 		}
 
 		return true;
diff --git a/Engine/Plugins/Experimental/PCG/Source/PCG/Private/Data/PCGSplineData.cpp b/Engine/Plugins/Experimental/PCG/Source/PCG/Private/Data/PCGSplineData.cpp
--- a/Engine/Plugins/Experimental/PCG/Source/PCG/Private/Data/PCGSplineData.cpp
+++ b/Engine/Plugins/Experimental/PCG/Source/PCG/Private/Data/PCGSplineData.cpp
@@ -220,7 +220,7 @@ bool UPCGSplineProjectionData::SamplePoint(const FTransform& InTransform, const
 			if (OutMetadata)
 			{
 				//TODO review op
-				OutMetadata->MergePointAttributes(SplinePoint, SurfacePoint, OutPoint, EPCGMetadataOp::Max);
+				OutMetadata->MergePointAttributesSubset(SplinePoint, OutMetadata, GetSpline()->Metadata, SurfacePoint, OutMetadata, GetSurface()->Metadata, OutPoint, EPCGMetadataOp::Max);
 			}
 
 			return true;
diff --git a/Engine/Plugins/Experimental/PCG/Source/PCG/Private/Data/PCGUnionData.cpp b/Engine/Plugins/Experimental/PCG/Source/PCG/Private/Data/PCGUnionData.cpp
--- a/Engine/Plugins/Experimental/PCG/Source/PCG/Private/Data/PCGUnionData.cpp
+++ b/Engine/Plugins/Experimental/PCG/Source/PCG/Private/Data/PCGUnionData.cpp
@@ -40,6 +40,7 @@ void UPCGUnionData::Initialize(const UPCGSpatialData* InA, const UPCGSpatialData
 void UPCGUnionData::AddData(const UPCGSpatialData* InData)
 {
 	check(InData);
+	check(Metadata);
 
 	Data.Add(InData);
 
@@ -49,12 +50,14 @@ void UPCGUnionData::AddData(const UPCGSpatialData* InData)
 		CachedBounds = InData->GetBounds();
 		CachedStrictBounds = InData->GetStrictBounds();
 		CachedDimension = InData->GetDimension();
+		Metadata->Initialize(InData->Metadata);
 	}
 	else
 	{
 		CachedBounds += InData->GetBounds();
 		CachedStrictBounds = PCGHelpers::OverlapBounds(CachedStrictBounds, InData->GetStrictBounds());
 		CachedDimension = FMath::Max(CachedDimension, InData->GetDimension());
+		Metadata->AddAttributes(InData->Metadata);
 	}
 
 	if (!FirstNonTrivialTransformData && InData->HasNonTrivialTransform())
@@ -128,7 +131,7 @@ bool UPCGUnionData::SamplePoint(const FTransform& InTransform, const FBox& InBou
 				// Merge properties into OutPoint
 				if (OutMetadata)
 				{
-					OutMetadata->MergePointAttributes(OutPoint, PointInData, OutPoint, EPCGMetadataOp::Max);
+					OutMetadata->MergePointAttributesSubset(OutPoint, OutMetadata, OutMetadata, PointInData, OutMetadata, Data[DataIndex]->Metadata, OutPoint, EPCGMetadataOp::Max);
 				}
 			}
 			
@@ -294,7 +297,7 @@ void UPCGUnionData::CreateSequentialPointData(FPCGContext* Context, UPCGPointDat
 
 					if (PointData->Metadata)
 					{
-						PointData->Metadata->MergePointAttributes(OutPoint, PointInData, OutPoint, EPCGMetadataOp::Max);
+						PointData->Metadata->MergePointAttributesSubset(OutPoint, PointData->Metadata, PointData->Metadata, PointInData, PointData->Metadata, Data[FollowingDataIndex]->Metadata, OutPoint, EPCGMetadataOp::Max);
 					}
 					else if (OutPoint.Density >= 1.0f)
 					{
diff --git a/Engine/Plugins/Experimental/PCG/Source/PCG/Private/Grid/PCGLandscapeCache.cpp b/Engine/Plugins/Experimental/PCG/Source/PCG/Private/Grid/PCGLandscapeCache.cpp
--- a/Engine/Plugins/Experimental/PCG/Source/PCG/Private/Grid/PCGLandscapeCache.cpp
+++ b/Engine/Plugins/Experimental/PCG/Source/PCG/Private/Grid/PCGLandscapeCache.cpp
@@ -67,16 +67,12 @@ void FPCGLandscapeCacheEntry::BuildCacheData(ULandscapeInfo* LandscapeInfo, ULan
 			continue;
 		}
 
-		if (FLandscapeLayer* EditLayer = Landscape->GetLayer(Landscape->GetLayerIndex(Layer.LayerName)))
+		FLandscapeComponentDataInterface CDI(InComponent, 0, /*bUseEditingLayer=*/true);
+		if (CDI.GetWeightmapTextureData(LayerInfo, LayerCache, /*bUseEditingLayer=*/true))
 		{
-			FScopedSetLandscapeEditingLayer Scope(Landscape, EditLayer->Guid);
-			FLandscapeComponentDataInterface CDI(InComponent, 0, /*bUseEditingLayer=*/true);
-			if (CDI.GetWeightmapTextureData(LayerInfo, LayerCache, /*bUseEditingLayer=*/true))
-			{
-				FPCGLandscapeCacheLayer& PCGLayer = LayerData.Emplace_GetRef();
-				PCGLayer.Name = Layer.LayerName;
-				PCGLayer.Data = LayerCache;
-			}
+			FPCGLandscapeCacheLayer& PCGLayer = LayerData.Emplace_GetRef();
+			PCGLayer.Name = Layer.LayerName;
+			PCGLayer.Data = LayerCache;
 		}
 
 		LayerCache.Reset();
@@ -258,7 +254,9 @@ void FPCGLandscapeCache::PrimeCache()
 				}
 			});
 		}
-	}	
+	}
+
+	CacheLayerNames();
 #endif
 }
 
@@ -329,8 +327,7 @@ void FPCGLandscapeCache::SetupLandscapeCallbacks()
 
 	for (AActor* FoundLandscape : FoundLandscapes)
 	{
-		ALandscapeProxy* Landscape = Cast<ALandscapeProxy>(FoundLandscape);
-		check(Landscape);
+		ALandscapeProxy* Landscape = CastChecked<ALandscapeProxy>(FoundLandscape);
 
 		Landscapes.Add(Landscape);
 		Landscape->OnComponentDataChanged.AddRaw(this, &FPCGLandscapeCache::OnLandscapeChanged);
@@ -374,6 +371,8 @@ void FPCGLandscapeCache::OnLandscapeChanged(ALandscapeProxy* Landscape, const FL
 
 void FPCGLandscapeCache::CacheLayerNames()
 {
+	CachedLayerNames.Reset();
+
 	for (TWeakObjectPtr<ALandscapeProxy> Landscape : Landscapes)
 	{
 		if (Landscape.Get())
@@ -391,6 +390,12 @@ void FPCGLandscapeCache::CacheLayerNames(ALandscapeProxy* Landscape)
 	{
 		for (const FLandscapeInfoLayerSettings& Layer : LandscapeInfo->Layers)
 		{
+			ULandscapeLayerInfoObject* LayerInfo = Layer.LayerInfoObj;
+			if (!LayerInfo)
+			{
+				continue;
+			}
+
 			CachedLayerNames.Add(Layer.LayerName);
 		}
 	}
diff --git a/Engine/Plugins/Experimental/PCG/Source/PCG/Private/Metadata/PCGMetadata.cpp b/Engine/Plugins/Experimental/PCG/Source/PCG/Private/Metadata/PCGMetadata.cpp
--- a/Engine/Plugins/Experimental/PCG/Source/PCG/Private/Metadata/PCGMetadata.cpp
+++ b/Engine/Plugins/Experimental/PCG/Source/PCG/Private/Metadata/PCGMetadata.cpp
@@ -633,7 +633,17 @@ void UPCGMetadata::MergePointAttributes(const FPCGPoint& InPointA, const FPCGPoi
 	MergeAttributes(InPointA.MetadataEntry, this, InPointB.MetadataEntry, this, OutPoint.MetadataEntry, Op);
 }
 
+void UPCGMetadata::MergePointAttributesSubset(const FPCGPoint& InPointA, const UPCGMetadata* InMetadataA, const UPCGMetadata* InMetadataSubsetA, const FPCGPoint& InPointB, const UPCGMetadata* InMetadataB, const UPCGMetadata* InMetadataSubsetB, FPCGPoint& OutPoint, EPCGMetadataOp Op)
+{
+	MergeAttributesSubset(InPointA.MetadataEntry, InMetadataA, InMetadataSubsetA, InPointB.MetadataEntry, InMetadataB, InMetadataSubsetB, OutPoint.MetadataEntry, Op);
+}
+
 void UPCGMetadata::MergeAttributes(PCGMetadataEntryKey InKeyA, const UPCGMetadata* InMetadataA, PCGMetadataEntryKey InKeyB, const UPCGMetadata* InMetadataB, PCGMetadataEntryKey& OutKey, EPCGMetadataOp Op)
+{
+	MergeAttributesSubset(InKeyA, InMetadataA, InMetadataA, InKeyB, InMetadataB, InMetadataB, OutKey, Op);
+}
+
+void UPCGMetadata::MergeAttributesSubset(PCGMetadataEntryKey InKeyA, const UPCGMetadata* InMetadataA, const UPCGMetadata* InMetadataSubsetA, PCGMetadataEntryKey InKeyB, const UPCGMetadata* InMetadataB, const UPCGMetadata* InMetadataSubsetB, PCGMetadataEntryKey& OutKey, EPCGMetadataOp Op)
 {
 	// Early out: nothing to do if both input metadata are null / points have no assigned metadata
 	if (!InMetadataA && !InMetadataB)
@@ -651,7 +661,11 @@ void UPCGMetadata::MergeAttributes(PCGMetadataEntryKey InKeyA, const UPCGMetadat
 		FPCGMetadataAttributeBase* Attribute = AttributePair.Value;
 
 		// Get attribute from A
-		const FPCGMetadataAttributeBase* AttributeA = InMetadataA ? InMetadataA->GetConstAttribute(AttributeName) : nullptr;
+		const FPCGMetadataAttributeBase* AttributeA = nullptr;
+		if (InMetadataA && ((InMetadataA == InMetadataSubsetA) || (InMetadataSubsetA && InMetadataSubsetA->HasAttribute(AttributeName))))
+		{
+			AttributeA = InMetadataA->GetConstAttribute(AttributeName);
+		}
 
 		if (AttributeA && AttributeA->GetTypeId() != Attribute->GetTypeId())
 		{
@@ -660,7 +674,11 @@ void UPCGMetadata::MergeAttributes(PCGMetadataEntryKey InKeyA, const UPCGMetadat
 		}
 
 		// Get attribute from B
-		const FPCGMetadataAttributeBase* AttributeB = InMetadataB ? InMetadataB->GetConstAttribute(AttributeName) : nullptr;
+		const FPCGMetadataAttributeBase* AttributeB = nullptr;
+		if (InMetadataB && ((InMetadataB == InMetadataSubsetB) || (InMetadataSubsetB && InMetadataSubsetB->HasAttribute(AttributeName))))
+		{
+			AttributeB = InMetadataB->GetConstAttribute(AttributeName);
+		}
 
 		if (AttributeB && AttributeB->GetTypeId() != Attribute->GetTypeId())
 		{
diff --git a/Engine/Plugins/Experimental/PCG/Source/PCG/Public/Metadata/PCGMetadata.h b/Engine/Plugins/Experimental/PCG/Source/PCG/Public/Metadata/PCGMetadata.h
--- a/Engine/Plugins/Experimental/PCG/Source/PCG/Public/Metadata/PCGMetadata.h
+++ b/Engine/Plugins/Experimental/PCG/Source/PCG/Public/Metadata/PCGMetadata.h
@@ -144,6 +144,7 @@ public:
 
 	/** Attributes operations */
 	void MergeAttributes(PCGMetadataEntryKey InKeyA, const UPCGMetadata* InMetadataA, PCGMetadataEntryKey InKeyB, const UPCGMetadata* InMetadataB, PCGMetadataEntryKey& OutKey, EPCGMetadataOp Op);
+	void MergeAttributesSubset(PCGMetadataEntryKey InKeyA, const UPCGMetadata* InMetadataA, const UPCGMetadata* InMetadataSubetA, PCGMetadataEntryKey InKeyB, const UPCGMetadata* InMetadataB, const UPCGMetadata* InMetadataSubsetb, PCGMetadataEntryKey& OutKey, EPCGMetadataOp Op);
 	void MergeAttributes(PCGMetadataEntryKey InKeyA, PCGMetadataEntryKey InKeyB, PCGMetadataEntryKey& OutKey, EPCGMetadataOp Op);
 
 	void ResetWeightedAttributes(PCGMetadataEntryKey& OutKey);
@@ -154,6 +155,7 @@ public:
 
 	/** Attributes operations - shorthand for points */
 	void MergePointAttributes(const FPCGPoint& InPointA, const FPCGPoint& InPointB, FPCGPoint& OutPoint, EPCGMetadataOp Op);
+	void MergePointAttributesSubset(const FPCGPoint& InPointA, const UPCGMetadata* InMetadataA, const UPCGMetadata* InMetadataSubetA, const FPCGPoint& InPointB, const UPCGMetadata* InMetadataB, const UPCGMetadata* InMetadataSubsetb, FPCGPoint& OutPoint, EPCGMetadataOp Op);
 	void SetPointAttributes(const TArrayView<const FPCGPoint>& InPoints, const UPCGMetadata* InMetadata, const TArrayView<FPCGPoint>& OutPoints);
 
 	/** Blueprint-friend versions */
