Adds a multi-tier evaluation of order independence to the PCG Determinism tests:
+ Order independent
+ Order consistent
+ Orthogonal
+ Non-deterministic

diff --git a/Engine/Plugins/Experimental/PCG/Source/PCG/Private/Tests/Determinism/PCGDeterminismTestsCommon.cpp b/Engine/Plugins/Experimental/PCG/Source/PCG/Private/Tests/Determinism/PCGDeterminismTestsCommon.cpp
--- a/Engine/Plugins/Experimental/PCG/Source/PCG/Private/Tests/Determinism/PCGDeterminismTestsCommon.cpp
+++ b/Engine/Plugins/Experimental/PCG/Source/PCG/Private/Tests/Determinism/PCGDeterminismTestsCommon.cpp
@@ -1,30 +1,36 @@
 // Copyright Epic Games, Inc. All Rights Reserved.
-
 #include "Tests/Determinism/PCGDeterminismTestsCommon.h"
 
 #include "PCGComponent.h"
-#include "PCGGraph.h"
+#include "PCGData.h"
 #include "PCGHelpers.h"
+#include "PCGNode.h"
+#include "PCGPin.h"
 #include "PCGPoint.h"
 #include "Data/PCGPointData.h"
 #include "Data/PCGPrimitiveData.h"
 #include "Data/PCGSplineData.h"
 #include "Data/PCGSurfaceData.h"
 #include "Data/PCGVolumeData.h"
-#include "Tests/PCGTestsCommon.h"
 
-#include "GameFramework/Actor.h"
 #include "Components/SplineComponent.h"
 #include "Components/StaticMeshComponent.h"
+#include "GameFramework/Actor.h"
 
 #define LOCTEXT_NAMESPACE "PCGDeterminism"
 
+namespace
+{
+	const FText TEXT_BasicTest = LOCTEXT("BasicTests", "Basic Tests");
+	const FText TEXT_OrderIndependenceTest = LOCTEXT("OrderIndependence", "Order Independence");
+}
+
 namespace PCGDeterminismTests
 {
-	bool LogInvalidTest(const UPCGNode* InPCGNode, int32 Seed, EPCGDataType& OutDataTypesTested, TArray<FString>& OutAdditionalDetails)
+	bool LogInvalidTest(const UPCGNode* InPCGNode, FNodeTestResult& OutResult)
 	{
 		UE_LOG(LogPCG, Warning, TEXT("Attempting to run an invalid determinism test"));
-		OutAdditionalDetails.Add(TEXT("Invalid test"));
+		OutResult.AdditionalDetails.Add(TEXT("Invalid test"));
 		return false;
 	}
 
@@ -38,145 +44,238 @@ namespace PCGDeterminismTests
 			return;
 		}
 
-		bool bTestWasSuccessful = TestToRun.TestDelegate(InPCGNode, InPCGNode->DefaultSettings->Seed, OutResult.DataTypesTested, OutResult.AdditionalDetails);
-		OutResult.TestResults.FindOrAdd(FName(TestToRun.TestLabel.ToString()), bTestWasSuccessful);
-		OutResult.bFlagRaised = OutResult.bFlagRaised && !bTestWasSuccessful;
+		OutResult.Seed = InPCGNode->DefaultSettings->Seed;
+		bool bTestWasSuccessful = TestToRun.TestDelegate(InPCGNode, TestToRun.TestName, OutResult);
+		OutResult.bFlagRaised |= !bTestWasSuccessful;
+
+		// If the test did not set its own result, base it off the success/fail
+		if (!OutResult.TestResults.Find(TestToRun.TestName))
+		{
+			UpdateTestResults(TestToRun.TestName,
+				OutResult,
+				bTestWasSuccessful ? EDeterminismLevel::Basic : EDeterminismLevel::NoDeterminism);
+		}
 	}
 
 	PCG_API void RetrieveBasicTests(TArray<FNodeTestInfo>& OutBasicTests)
 	{
-		OutBasicTests.Emplace(LOCTEXT("SingleSameData", "Single Same Data"), RunSingleSameDataTest);
-		OutBasicTests.Emplace(LOCTEXT("SingleIdenticalData", "Single Identical Data"), RunSingleIdenticalDataTest);
-		OutBasicTests.Emplace(LOCTEXT("MultipleSameData", "Multiple Same Data"), RunMultipleSameDataTest);
-		OutBasicTests.Emplace(LOCTEXT("MultipleIdenticalData", "Multiple Identical Data"), RunMultipleIdenticalDataTest);
-		OutBasicTests.Emplace(LOCTEXT("DataCollectionOrderIndependence", "Shuffle Data Indices"),RunDataCollectionOrderIndependenceTest);
-		OutBasicTests.Emplace(LOCTEXT("AllDataOrderIndependence", "Shuffle Internal Data"), RunAllDataOrderIndependenceTest);
+		OutBasicTests.Emplace(TEXT_BasicTest, RunBasicTestSuite, 80.f);
+		OutBasicTests.Emplace(TEXT_OrderIndependenceTest, RunOrderIndependenceSuite, 125.f);
 	}
 
-	bool RunSingleSameDataTest(const UPCGNode* InPCGNode, int32 Seed, EPCGDataType& OutDataTypesTested, TArray<FString>& OutAdditionalDetails)
+	bool RunBasicTestSuite(const UPCGNode* InPCGNode, const FName& TestName, FNodeTestResult& OutResult)
 	{
-		PCGTestsCommon::FTestData TestData(Seed, InPCGNode->DefaultSettings);
+		// Get Pins
+		const TArray<TObjectPtr<UPCGPin>>& InputPins = InPCGNode->GetInputPins();
 
-		// Generate random input data
-		AddInputDataBasedOnPins(TestData, InPCGNode, OutDataTypesTested, OutAdditionalDetails);
+		// Base Options Per Pin
+		FNodeAndOptions NodeAndOptions(InPCGNode, OutResult.Seed, /*bMultipleOptionsPerPin=*/true);
+		RetrieveBaseOptionsPerPin(NodeAndOptions.BaseOptionsByPin, InputPins, OutResult.DataTypesTested, Defaults::NumTestInputsPerPin);
 
-		return ExecutionIsDeterministicSameData(TestData, InPCGNode);
+		return RunBasicSelfTest(NodeAndOptions) && RunBasicCopiedSelfTest(NodeAndOptions);
 	}
 
-	bool RunSingleIdenticalDataTest(const UPCGNode* InPCGNode, int32 Seed, EPCGDataType& OutDataTypesTested, TArray<FString>& OutAdditionalDetails)
+	bool RunOrderIndependenceSuite(const UPCGNode* InPCGNode, const FName& TestName, FNodeTestResult& OutResult)
 	{
-		PCGTestsCommon::FTestData FirstTestData(Seed, InPCGNode->DefaultSettings);
-		PCGTestsCommon::FTestData SecondTestData(Seed, InPCGNode->DefaultSettings);
+		// Start with not deterministic
+		UpdateTestResults(TestName, OutResult, EDeterminismLevel::NoDeterminism);
 
-		AddInputDataBasedOnPins(FirstTestData, InPCGNode, OutDataTypesTested, OutAdditionalDetails);
-		AddInputDataBasedOnPins(SecondTestData, InPCGNode, OutDataTypesTested, OutAdditionalDetails);
+		// Get Pins
+		const TArray<TObjectPtr<UPCGPin>>& InputPins = InPCGNode->GetInputPins();
 
-		return ExecutionIsDeterministic(FirstTestData, SecondTestData, InPCGNode);
-	}
-
-	bool RunMultipleSameDataTest(const UPCGNode* InPCGNode, int32 Seed, EPCGDataType& OutDataTypesTested, TArray<FString>& OutAdditionalDetails)
-	{
-		PCGTestsCommon::FTestData TestData(Seed, InPCGNode->DefaultSettings);
+		// Base Options Per Pin
+		FNodeAndOptions NodeAndOptions(InPCGNode, OutResult.Seed, /*bMultipleOptionsPerPin=*/true);
+		// TODO: Multiple inputs per pin might lead to redundant tests. Filter these out
+		RetrieveBaseOptionsPerPin(NodeAndOptions.BaseOptionsByPin, InputPins, OutResult.DataTypesTested, Defaults::NumTestInputsPerPin);
 
-		for (int32 I = 0; I < Defaults::NumMultipleTestDataSets; ++I)
+		// No inputs, so just test determinism against itself
+		if (InputPins.Num() == 0 && RunBasicCopiedSelfTest(NodeAndOptions))
 		{
-			AddInputDataBasedOnPins(TestData, InPCGNode, OutDataTypesTested, OutAdditionalDetails);
+			UpdateTestResults(TestName, OutResult, EDeterminismLevel::OrderIndependent);
+			return true;
 		}
 
-		return ExecutionIsDeterministicSameData(TestData, InPCGNode);
+		EDeterminismLevel MaxLevel = GetHighestDeterminismLevel(NodeAndOptions);
+		UpdateTestResults(TestName, OutResult, MaxLevel);
+
+		return MaxLevel != EDeterminismLevel::None;
 	}
 
-	bool RunMultipleIdenticalDataTest(const UPCGNode* InPCGNode, int32 Seed, EPCGDataType& OutDataTypesTested, TArray<FString>& OutAdditionalDetails)
+	bool RunBasicSelfTest(const FNodeAndOptions& NodeAndOptions)
 	{
-		PCGTestsCommon::FTestData FirstTestData(Seed, InPCGNode->DefaultSettings);
-		PCGTestsCommon::FTestData SecondTestData(Seed, InPCGNode->DefaultSettings);
+		const UPCGNode* PCGNode = NodeAndOptions.PCGNode;
+		const TArray<TArray<EPCGDataType>>& BaseOptionsByPin = NodeAndOptions.BaseOptionsByPin;
+
+		PCGTestsCommon::FTestData TestData(NodeAndOptions.Seed, PCGNode->DefaultSettings);
 
-		for (int32 I = 0; I < Defaults::NumMultipleTestDataSets; ++I)
+		// For all permutations
+		const int32 NumPermutations = GetNumPermutations(BaseOptionsByPin);
+		for (int32 PermutationIndex = 0; PermutationIndex < NumPermutations; ++PermutationIndex)
 		{
-			AddInputDataBasedOnPins(FirstTestData, InPCGNode, OutDataTypesTested, OutAdditionalDetails);
-			AddInputDataBasedOnPins(SecondTestData, InPCGNode, OutDataTypesTested, OutAdditionalDetails);
+			// Reset test data
+			TestData.Reset(PCGNode->DefaultSettings);
+
+			// Add an input for each pin
+			for (int32 PinIndex = 0; PinIndex < BaseOptionsByPin.Num(); ++PinIndex)
+			{
+				EPCGDataType DataType = GetPermutation(PermutationIndex, PinIndex, BaseOptionsByPin);
+				AddRandomizedInputData(TestData, DataType);
+			}
+
+			if (!ExecutionIsDeterministicSameData(TestData, PCGNode))
+			{
+				return false;
+			}
 		}
 
-		return ExecutionIsDeterministic(FirstTestData, SecondTestData, InPCGNode);
+		return true;
 	}
 
-	bool RunDataCollectionOrderIndependenceTest(const UPCGNode* InPCGNode, int32 Seed, EPCGDataType& OutDataTypesTested, TArray<FString>& OutAdditionalDetails)
+	bool RunBasicCopiedSelfTest(const FNodeAndOptions& NodeAndOptions)
 	{
-		PCGTestsCommon::FTestData FirstTestData(Seed, InPCGNode->DefaultSettings);
-		PCGTestsCommon::FTestData SecondTestData(Seed, InPCGNode->DefaultSettings);
+		const UPCGNode* PCGNode = NodeAndOptions.PCGNode;
+		const TArray<TArray<EPCGDataType>>& BaseOptionsByPin = NodeAndOptions.BaseOptionsByPin;
 
-		for (int32 I = 0; I < Defaults::NumMultipleTestDataSets; ++I)
+		PCGTestsCommon::FTestData FirstTestData(NodeAndOptions.Seed, PCGNode->DefaultSettings);
+		PCGTestsCommon::FTestData SecondTestData(NodeAndOptions.Seed, PCGNode->DefaultSettings);
+
+		const int32 NumPermutations = GetNumPermutations(BaseOptionsByPin);
+		// For all permutations
+		for (int32 PermutationIndex = 0; PermutationIndex < NumPermutations; ++PermutationIndex)
 		{
-			AddInputDataBasedOnPins(FirstTestData, InPCGNode, OutDataTypesTested, OutAdditionalDetails);
-			AddInputDataBasedOnPins(SecondTestData, InPCGNode, OutDataTypesTested, OutAdditionalDetails);
-		}
+			// Reset test data
+			FirstTestData.Reset(PCGNode->DefaultSettings);
+			SecondTestData.Reset(PCGNode->DefaultSettings);
+
+			// Add an input for each pin
+			for (int32 PinIndex = 0; PinIndex < BaseOptionsByPin.Num(); ++PinIndex)
+			{
+				EPCGDataType DataType = GetPermutation(PermutationIndex, PinIndex, BaseOptionsByPin);
+				AddRandomizedInputData(FirstTestData, DataType);
+				AddRandomizedInputData(SecondTestData, DataType);
+			}
 
-		ShuffleInputOrder(SecondTestData);
+			if (!ExecutionIsDeterministic(FirstTestData, SecondTestData, PCGNode))
+			{
+				return false;
+			}
+		}
 
-		return ExecutionIsDeterministic(FirstTestData, SecondTestData, InPCGNode);
+		return true;
 	}
 
-	bool RunAllDataOrderIndependenceTest(const UPCGNode* InPCGNode, int32 Seed, EPCGDataType& OutDataTypesTested, TArray<FString>& OutAdditionalDetails)
+	EDeterminismLevel GetHighestDeterminismLevel(const FNodeAndOptions& NodeAndOptions,
+		int32 NumInputsPerPin,
+		EDeterminismLevel MaxLevel)
 	{
-		PCGTestsCommon::FTestData FirstTestData(Seed, InPCGNode->DefaultSettings);
-		PCGTestsCommon::FTestData SecondTestData(Seed, InPCGNode->DefaultSettings);
+		const UPCGNode* PCGNode = NodeAndOptions.PCGNode;
+		const TArray<TArray<EPCGDataType>>& BaseOptionsByPin = NodeAndOptions.BaseOptionsByPin;
+
+		PCGTestsCommon::FTestData FirstTestData(NodeAndOptions.Seed, PCGNode->DefaultSettings);
+		PCGTestsCommon::FTestData SecondTestData(NodeAndOptions.Seed, PCGNode->DefaultSettings);
 
-		for (int32 I = 0; I < Defaults::NumMultipleTestDataSets; ++I)
+		int32 NumInputs = NodeAndOptions.bMultipleOptionsPerPin ? NumInputsPerPin : 1;
+		const int32 NumPermutations = GetNumPermutations(BaseOptionsByPin);
+		// For each rotation of the inputs
+		for (int32 RotationIndex = 0; RotationIndex < BaseOptionsByPin.Num(); ++RotationIndex)
 		{
-			AddInputDataBasedOnPins(FirstTestData, InPCGNode, OutDataTypesTested, OutAdditionalDetails);
-			AddInputDataBasedOnPins(SecondTestData, InPCGNode, OutDataTypesTested, OutAdditionalDetails);
-		}
+			// For all permutations
+			for (int32 PermutationIndex = 0; PermutationIndex < NumPermutations;)
+			{
+				// Reset test data
+				FirstTestData.Reset(PCGNode->DefaultSettings);
+				SecondTestData.Reset(PCGNode->DefaultSettings);
+
+				// Add an input for each pin
+				for (int32 PinIndex = 0; PinIndex < BaseOptionsByPin.Num(); ++PinIndex)
+				{
+					for (int32 I = 0; I < NumInputs; ++I)
+					{
+						EPCGDataType DataType = GetPermutation(PermutationIndex++, PinIndex, BaseOptionsByPin);
+						AddRandomizedInputData(FirstTestData, DataType);
+						AddRandomizedInputData(SecondTestData, DataType);
+					}
+				}
 
-		ShuffleInputOrder(SecondTestData);
-		ShuffleAllInternalData(SecondTestData);
+				ShiftInputOrder(SecondTestData, RotationIndex + 1);
 
-		return ExecutionIsDeterministic(FirstTestData, SecondTestData, InPCGNode);
+				ExecuteWithTestData(FirstTestData, PCGNode);
+				ExecuteWithTestData(SecondTestData, PCGNode);
+
+				// Always execute at the minimum level. If that level fails, drop a level and start over
+				switch (MaxLevel)
+				{
+				case EDeterminismLevel::OrderIndependent:
+					if (DataCollectionsAreIdentical(FirstTestData.OutputData, SecondTestData.OutputData))
+					{
+						break;
+					}
+					MaxLevel = EDeterminismLevel::OrderConsistent;
+					PermutationIndex = 0;
+					// Falls through
+				case EDeterminismLevel::OrderConsistent:
+					if (DataCollectionsAreConsistent(FirstTestData.OutputData, SecondTestData.OutputData, NumInputs))
+					{
+						// With only 1 input, orthogonal == consistent
+						if (BaseOptionsByPin.Num() < 2)
+						{
+							if (DataCollectionsContainSameData(FirstTestData.OutputData, SecondTestData.OutputData))
+							{
+								break;
+							}
+						}
+						break;
+					}
+					MaxLevel = EDeterminismLevel::OrderOrthogonal;
+					PermutationIndex = 0;
+					// Falls through
+				case EDeterminismLevel::OrderOrthogonal:
+					if (DataCollectionsContainSameData(FirstTestData.OutputData, SecondTestData.OutputData))
+					{
+						break;
+					}
+					// Falls through
+				default:
+					return EDeterminismLevel::NoDeterminism;
+				}
+			}
+		}
+
+		return MaxLevel;
 	}
 
-	void AddInputDataBasedOnPins(PCGTestsCommon::FTestData& TestData, const UPCGNode* PCGNode, EPCGDataType& OutDataTypesTested, TArray<FString>& OutAdditionalDetails)
+	void AddRandomizedInputData(PCGTestsCommon::FTestData& TestData, EPCGDataType DataType, const FName& PinName)
 	{
-		check(PCGNode);
-
-		for (UPCGPin* InputPin : PCGNode->GetInputPins())
+		switch (DataType)
 		{
-			check(InputPin);
-
-			switch (InputPin->Properties.AllowedTypes)
-			{
-			case EPCGDataType::Point:
-				AddRandomizedMultiplePointInputData(TestData);
-				OutDataTypesTested |= EPCGDataType::Point;
-				break;
-			case EPCGDataType::Volume:
-				AddRandomizedVolumeInputData(TestData);
-				OutDataTypesTested |= EPCGDataType::Volume;
-				break;
-			case EPCGDataType::PolyLine:
-				AddRandomizedPolyLineInputData(TestData);
-				OutDataTypesTested |= EPCGDataType::PolyLine;
-				break;
-			case EPCGDataType::Primitive:
-				AddRandomizedPrimitiveInputData(TestData);
-				OutDataTypesTested |= EPCGDataType::Primitive;
-				break;
-			case EPCGDataType::Surface:
-				AddRandomizedSurfaceInputData(TestData);
-				OutDataTypesTested |= EPCGDataType::Surface;
-				break;
-			case EPCGDataType::Any:
-			case EPCGDataType::Spatial:
-				AddRandomizedMultiplePointInputData(TestData, Defaults::NumTestPointsToGenerate);
-				AddRandomizedVolumeInputData(TestData);
-				AddRandomizedPolyLineInputData(TestData, Defaults::NumTestPolyLinePointsToGenerate);
-				AddRandomizedPrimitiveInputData(TestData);
-				AddRandomizedSurfaceInputData(TestData);
-				OutDataTypesTested |= EPCGDataType::Spatial;
-				break;
-			default:
-				OutDataTypesTested = EPCGDataType::None;
-				OutAdditionalDetails.Add(TEXT("Unknown InputPin data type to test"));
-				break;
-			}
+		case EPCGDataType::Point:
+			AddRandomizedMultiplePointInputData(TestData, Defaults::NumTestPointsToGenerate, PinName);
+			return;
+		case EPCGDataType::Volume:
+			AddRandomizedVolumeInputData(TestData, PinName);
+			return;
+		case EPCGDataType::PolyLine:
+			AddRandomizedPolyLineInputData(TestData, Defaults::NumTestPolyLinePointsToGenerate, PinName);
+			return;
+		case EPCGDataType::Primitive:
+			AddRandomizedPrimitiveInputData(TestData, PinName);
+			return;
+		case EPCGDataType::Surface:
+			AddRandomizedSurfaceInputData(TestData, PinName);
+			return;
+		case EPCGDataType::Landscape:
+			AddRandomizedLandscapeInputData(TestData, PinName);
+			return;
+		// Just return to represent "empty" data
+		case EPCGDataType::None:
+			return;
+		case EPCGDataType::Any:
+		case EPCGDataType::Spatial:
+			UE_LOG(LogPCG, Error, TEXT("Specific EPCGDataType required to add to test data"));
+			break;
+		default:
+			UE_LOG(LogPCG, Error, TEXT("Unknown PCGDataType to add to test data: %s"), *UEnum::GetValueAsString(DataType));
+			break;
 		}
 	}
 
@@ -230,6 +329,11 @@ namespace PCGDeterminismTests
 		TaggedData.Pin = PinName;
 	}
 
+	void AddLandscapeInputData(FPCGDataCollection& InputData)
+	{
+		// TODO: Implement this in the future as needed
+	}
+
 	void AddRandomizedSinglePointInputData(PCGTestsCommon::FTestData& TestData, int32 PointNum, const FName& PinName)
 	{
 		check(PointNum > 0);
@@ -314,6 +418,11 @@ namespace PCGDeterminismTests
 		AddPrimitiveInputData(TestData.InputData, TestPrimitiveComponent, Defaults::MediumVector + TestData.RandomStream.VRand() * 0.5f * Defaults::MediumDistance, PinName);
 	}
 
+	void AddRandomizedLandscapeInputData(PCGTestsCommon::FTestData& TestData, const FName& PinName)
+	{
+		// TODO: Implement this in the future as needed
+	}
+
 	bool DataCollectionsAreIdentical(const FPCGDataCollection& FirstCollection, const FPCGDataCollection& SecondCollection)
 	{
 		if (FirstCollection.TaggedData.Num() != SecondCollection.TaggedData.Num())
@@ -321,6 +430,90 @@ namespace PCGDeterminismTests
 			return false;
 		}
 
+		// Find comparable data from first collection
+		for (int32 I = 0; I < FirstCollection.TaggedData.Num(); ++I)
+		{
+			EPCGDataType FirstDataType = FirstCollection.TaggedData[I].Data->GetDataType();
+			EPCGDataType SecondDataType = SecondCollection.TaggedData[I].Data->GetDataType();
+
+			// Only compare if they are the same type and same pin label
+			if (FirstDataType != SecondDataType || FirstCollection.TaggedData[I].Pin != SecondCollection.TaggedData[I].Pin)
+			{
+				return false;
+			}
+
+			auto CompareFunction = GetDataCompareFunction(FirstDataType, EDeterminismLevel::OrderIndependent);
+			if (!CompareFunction(FirstCollection.TaggedData[I].Data, SecondCollection.TaggedData[I].Data))
+			{
+				return false;
+			}
+		}
+
+		return true;
+	}
+
+	bool DataCollectionsAreConsistent(const FPCGDataCollection& FirstCollection, const FPCGDataCollection& SecondCollection, int32 NumInputs)
+	{
+		TArray<int32> IndexOffsets;
+		if (!DataCollectionsMatch(FirstCollection, SecondCollection, IndexOffsets))
+		{
+			return false;
+		}
+
+		// Short circuit null set
+		if (IndexOffsets.Num() == 0)
+		{
+			return true;
+		}
+
+		// All of the offset correlations should zero sum, if all output elements are matching
+		int32 Sum = 0;
+		// Data collection elements should be consistent in their groupings based on the ordering of the input data.
+		// The offset of each grouping should be the same, so when it switches, it indicates a different grouping. There
+		// should only be at most one grouping per input
+		int32 NumSwitches = 0;
+
+		int32 LastOffset = IndexOffsets[0];
+		for (int32 Offset : IndexOffsets)
+		{
+			if (Offset != LastOffset)
+			{
+				// Number of switches must be less than number of inputs
+				if (++NumSwitches >= NumInputs)
+				{
+					return false;
+				}
+			}
+			Sum += Offset;
+		}
+
+		// Should be zero summed
+		return Sum == 0;
+	}
+
+	bool DataCollectionsContainSameData(const FPCGDataCollection& FirstCollection, const FPCGDataCollection& SecondCollection)
+	{
+		TArray<int32> NullArray;
+		return DataCollectionsMatch(FirstCollection, SecondCollection, NullArray);
+	}
+
+	bool DataCollectionsMatch(const FPCGDataCollection& FirstCollection,
+		const FPCGDataCollection& SecondCollection,
+		TArray<int32>& OutIndexOffsets)
+	{
+		OutIndexOffsets.Empty();
+
+		if (FirstCollection.TaggedData.Num() != SecondCollection.TaggedData.Num())
+		{
+			return false;
+		}
+
+		// Short circuit if there's nothing to compare
+		if (FirstCollection.TaggedData.Num() == 0)
+		{
+			return true;
+		}
+
 		TArray<int32> FirstCollectionComparableIndices;
 		TArray<int32> SecondCollectionComparableIndices;
 
@@ -343,9 +536,22 @@ namespace PCGDeterminismTests
 			}
 		}
 
-		TArray<int32> MatchedIndices;
+		if (FirstCollectionComparableIndices.Num() != SecondCollectionComparableIndices.Num())
+		{
+			return false;
+		}
+
+		// Nothing to compare, so its a null match
+		if (FirstCollectionComparableIndices.Num() == 0)
+		{
+			return true;
+		}
+
+		// Offsets of one data's index from the other collection's index
+		TSet<int32> MatchedIndices;
+		MatchedIndices.Reserve(FirstCollectionComparableIndices.Num());
+		OutIndexOffsets.SetNumUninitialized(FirstCollectionComparableIndices.Num());
 
-		// Find matches for each TaggedData array
 		for (int32 I : FirstCollectionComparableIndices)
 		{
 			bool bFound = false;
@@ -362,14 +568,11 @@ namespace PCGDeterminismTests
 				// Only compare if they are the same type and same pin label
 				if (FirstDataType == SecondDataType && FirstCollection.TaggedData[I].Pin == SecondCollection.TaggedData[I].Pin)
 				{
-					// TODO: Compare in cascading order:
-					// Order independent
-					// Order consistent
-					// Deterministic but not having a predictable order
-					// Non-deterministic
-					if (GetCompareFunction(FirstDataType)(FirstCollection.TaggedData[I].Data, SecondCollection.TaggedData[J].Data))
+					auto CompareFunction = GetDataCollectionCompareFunction(FirstDataType);
+					if (CompareFunction(FirstCollection.TaggedData[I].Data, SecondCollection.TaggedData[J].Data))
 					{
 						MatchedIndices.Emplace(J);
+						OutIndexOffsets[I] = J - I;
 						bFound = true;
 						break;
 					}
@@ -378,6 +581,7 @@ namespace PCGDeterminismTests
 
 			if (!bFound)
 			{
+				OutIndexOffsets.Empty();
 				return false;
 			}
 		}
@@ -385,12 +589,25 @@ namespace PCGDeterminismTests
 		return true;
 	}
 
+	bool InternalDataMatches(const UPCGData* FirstData, const UPCGData* SecondData, TArray<int32>& OutIndexOffsets)
+	{
+		check(FirstData->GetDataType() == SecondData->GetDataType());
+
+		// Filter output differences
+		auto CompareFunction = GetDataCompareFunction(FirstData->GetDataType(), EDeterminismLevel::OrderConsistent);
+		return CompareFunction(FirstData, SecondData);
+	}
+
 	bool SpatialDataIsIdentical(const UPCGData* FirstData, const UPCGData* SecondData)
 	{
-		check(FirstData);
-		check(SecondData);
+		const UPCGSpatialData* FirstSpatialData = CastChecked<const UPCGSpatialData>(FirstData);
+		const UPCGSpatialData* SecondSpatialData = CastChecked<const UPCGSpatialData>(SecondData);
 
-		if (BothDataCastsToDataType<const UPCGPointData>(FirstData, SecondData))
+		if (!SpatialBasicsAreIdentical(FirstSpatialData, SecondSpatialData))
+		{
+			return false;
+		}
+		else if (BothDataCastsToDataType<const UPCGPointData>(FirstData, SecondData))
 		{
 			return PointDataIsIdentical(FirstData, SecondData);
 		}
@@ -410,6 +627,7 @@ namespace PCGDeterminismTests
 		{
 			return PrimitiveDataIsIdentical(FirstData, SecondData);
 		}
+		// Default to sampling
 		else if (BothDataCastsToDataType<const UPCGSpatialData>(FirstData, SecondData))
 		{
 			return SampledSpatialDataIsIdentical(Cast<const UPCGSpatialData>(FirstData), Cast<const UPCGSpatialData>(SecondData));
@@ -436,30 +654,9 @@ namespace PCGDeterminismTests
 			return false;
 		}
 
-		// Match points based on indices
-		TSet<int32> RemainingIndices;
-		for (int32 i = 0; i < SecondPoints.Num(); ++i)
-		{
-			RemainingIndices.Add(i);
-		}
-
-		for (int32 i = 0; i < FirstPoints.Num(); ++i)
+		for (int32 I = 0; I < FirstPoints.Num(); ++I)
 		{
-			bool bFound = false;
-
-			// TODO: Maybe use Octree to prune. Is the saved computation worth the complexity of more complex index tracking?
-			for (int32 j : RemainingIndices)
-			{
-				if (PCGTestsCommon::PointsAreIdentical(FirstPoints[i], SecondPoints[j]))
-				{
-					RemainingIndices.Remove(j);
-					bFound = true;
-					break;
-				}
-			}
-
-			// Couldn't find the matching point
-			if (!bFound)
+			if (!PCGTestsCommon::PointsAreIdentical(FirstPoints[I], SecondPoints[I]))
 			{
 				return false;
 			}
@@ -473,18 +670,15 @@ namespace PCGDeterminismTests
 		const UPCGVolumeData* FirstVolumeData = CastChecked<const UPCGVolumeData>(FirstData);
 		const UPCGVolumeData* SecondVolumeData = CastChecked<const UPCGVolumeData>(SecondData);
 
-		return FirstVolumeData->VoxelSize == SecondVolumeData->VoxelSize && SpatialBasicsAreIdentical(FirstVolumeData, SecondVolumeData);
+		return FirstVolumeData->VoxelSize == SecondVolumeData->VoxelSize;
 	}
 
 	bool SurfaceDataIsIdentical(const UPCGData* FirstData, const UPCGData* SecondData)
 	{
-		const UPCGSurfaceData* FirstSurfaceData = CastChecked<const UPCGSurfaceData>(FirstData);
-		const UPCGSurfaceData* SecondSurfaceData = CastChecked<const UPCGSurfaceData>(SecondData);
-
 		// TODO: Implement Surface Data comparison as needed in the future
 		UE_LOG(LogPCG, Warning, TEXT("Surface comparison not fully implemented."));
 
-		return SpatialBasicsAreIdentical(FirstSurfaceData, SecondSurfaceData);
+		return false;
 	}
 
 	bool PolyLineDataIsIdentical(const UPCGData* FirstData, const UPCGData* SecondData)
@@ -492,11 +686,6 @@ namespace PCGDeterminismTests
 		const UPCGPolyLineData* FirstPolyLineData = CastChecked<const UPCGPolyLineData>(FirstData);
 		const UPCGPolyLineData* SecondPolyLineData = CastChecked<const UPCGPolyLineData>(SecondData);
 
-		if (!SpatialBasicsAreIdentical(FirstPolyLineData, SecondPolyLineData))
-		{
-			return false;
-		}
-
 		int NumSegments = FirstPolyLineData->GetNumSegments();
 		for (int32 I = 0; I < NumSegments; ++I)
 		{
@@ -516,22 +705,17 @@ namespace PCGDeterminismTests
 		const UPCGPrimitiveData* FirstPrimitiveData = CastChecked<const UPCGPrimitiveData>(FirstData);
 		const UPCGPrimitiveData* SecondPrimitiveData = CastChecked<const UPCGPrimitiveData>(SecondData);
 
-		if (FirstPrimitiveData->VoxelSize != SecondPrimitiveData->VoxelSize || !SpatialBasicsAreIdentical(FirstPrimitiveData, SecondPrimitiveData))
+		if (FirstPrimitiveData->VoxelSize != SecondPrimitiveData->VoxelSize)
 		{
 			return false;
 		}
 
 		// TODO: Compare the ToPointData, which will require the Context
-		return false;
+		return true;
 	}
 
 	bool SampledSpatialDataIsIdentical(const UPCGSpatialData* FirstSpatialData, const UPCGSpatialData* SecondSpatialData)
 	{
-		if (!SpatialBasicsAreIdentical(FirstSpatialData, SecondSpatialData))
-		{
-			return false;
-		}
-
 		// At this point, bounds has already been checked for equality
 		const FBox SampleBounds = Defaults::TestingVolume;
 		const FVector SampleExtent = SampleBounds.GetExtent();
@@ -569,12 +753,146 @@ namespace PCGDeterminismTests
 		return true;
 	}
 
-	bool SpatialBasicsAreIdentical(const UPCGSpatialData* FirstSpatialData, const UPCGSpatialData* SecondSpatialData)
+	bool SpatialBasicsAreIdentical(const UPCGData* FirstData, const UPCGData* SecondData)
 	{
-		return (FirstSpatialData->GetDataType() == SecondSpatialData->GetDataType() &&
+		const UPCGSpatialData* FirstSpatialData = CastChecked<const UPCGSpatialData>(FirstData);
+		const UPCGSpatialData* SecondSpatialData = CastChecked<const UPCGSpatialData>(SecondData);
+
+		return FirstSpatialData->GetDataType() == SecondSpatialData->GetDataType() &&
 			FirstSpatialData->GetDimension() == SecondSpatialData->GetDimension() &&
 			FirstSpatialData->GetBounds() == SecondSpatialData->GetBounds() &&
-			FirstSpatialData->GetStrictBounds() == SecondSpatialData->GetStrictBounds());
+			FirstSpatialData->GetStrictBounds() == SecondSpatialData->GetStrictBounds();
+	}
+
+	bool SpatialDataIsConsistent(const UPCGData* FirstData, const UPCGData* SecondData)
+	{
+		const UPCGSpatialData* FirstSpatialData = CastChecked<const UPCGSpatialData>(FirstData);
+		const UPCGSpatialData* SecondSpatialData = CastChecked<const UPCGSpatialData>(SecondData);
+
+		// Even for consistent, these should be identical
+		if (!SpatialBasicsAreIdentical(FirstSpatialData, SecondSpatialData))
+		{
+			return false;
+		}
+
+		if (BothDataCastsToDataType<const UPCGPointData>(FirstSpatialData, SecondSpatialData))
+		{
+			return PointDataIsConsistent(FirstSpatialData, SecondSpatialData);
+		}
+
+		// TODO: Implement other data types as they come
+		return SpatialDataIsIdentical(FirstData, SecondData);
+	}
+
+	bool PointDataIsConsistent(const UPCGData* FirstData, const UPCGData* SecondData)
+	{
+		const UPCGPointData* FirstPointData = CastChecked<const UPCGPointData>(FirstData);
+		const UPCGPointData* SecondPointData = CastChecked<const UPCGPointData>(SecondData);
+
+		const TArray<FPCGPoint>& FirstPoints = FirstPointData->GetPoints();
+		const TArray<FPCGPoint>& SecondPoints = SecondPointData->GetPoints();
+
+		if (FirstPoints.Num() != SecondPoints.Num())
+		{
+			return false;
+		}
+
+		// Short circuit, because a null set is deterministic
+		if (FirstPoints.Num() == 0)
+		{
+			return true;
+		}
+
+		// TODO: To be used in the future for series comparisons
+		TArray<int32> IndexOffsets;
+		IndexOffsets.SetNum(FirstPoints.Num());
+
+		for (int32 I = 0; I < FirstPoints.Num(); ++I)
+		{
+			bool bFound = false;
+
+			for (int32 J = 0; J < SecondPoints.Num(); ++J)
+			{
+				if (PCGTestsCommon::PointsAreIdentical(FirstPoints[I], SecondPoints[J]))
+				{
+					IndexOffsets[I] = J - I;
+					bFound = true;
+					break;
+				}
+			}
+
+			if (!bFound)
+			{
+				return false;
+			}
+		}
+
+		return true;
+	}
+
+	bool SpatialDataIsOrthogonal(const UPCGData* FirstData, const UPCGData* SecondData)
+	{
+		const UPCGSpatialData* FirstSpatialData = CastChecked<const UPCGSpatialData>(FirstData);
+		const UPCGSpatialData* SecondSpatialData = CastChecked<const UPCGSpatialData>(SecondData);
+
+		if (!SpatialBasicsAreIdentical(FirstSpatialData, SecondSpatialData))
+		{
+			return false;
+		}
+
+		if (BothDataCastsToDataType<const UPCGPointData>(FirstSpatialData, SecondSpatialData))
+		{
+			return PointDataIsOrthogonal(FirstSpatialData, SecondSpatialData);
+		}
+
+		// TODO: Implement other data types as they come
+		return SpatialDataIsIdentical(FirstData, SecondData);
+	}
+
+	bool PointDataIsOrthogonal(const UPCGData* FirstData, const UPCGData* SecondData)
+	{
+		const UPCGPointData* FirstPointData = CastChecked<const UPCGPointData>(FirstData);
+		const UPCGPointData* SecondPointData = CastChecked<const UPCGPointData>(SecondData);
+
+		const TArray<FPCGPoint>& FirstPoints = FirstPointData->GetPoints();
+		const TArray<FPCGPoint>& SecondPoints = SecondPointData->GetPoints();
+
+		if (FirstPoints.Num() != SecondPoints.Num())
+		{
+			return false;
+		}
+
+		// Match points based on indices
+		TSet<int32> RemainingIndices;
+		RemainingIndices.Reserve(SecondPoints.Num());
+
+		for (int32 I = 0; I < SecondPoints.Num(); ++I)
+		{
+			RemainingIndices.Add(I);
+		}
+
+		for (int32 I = 0; I < FirstPoints.Num(); ++I)
+		{
+			bool bFound = false;
+
+			for (int32 J : RemainingIndices)
+			{
+				if (PCGTestsCommon::PointsAreIdentical(FirstPoints[I], SecondPoints[J]))
+				{
+					RemainingIndices.Remove(J);
+					bFound = true;
+					break;
+				}
+			}
+
+			// Couldn't find the matching point
+			if (!bFound)
+			{
+				return false;
+			}
+		}
+
+		return true;
 	}
 
 	bool ComparisonIsUnimplemented(const UPCGData* FirstData, const UPCGData* SecondData)
@@ -582,9 +900,15 @@ namespace PCGDeterminismTests
 		return false;
 	}
 
+	bool ConsistencyComparisonIsUnimplemented(const UPCGData* FirstData, const UPCGData* SecondData, TArray<int32>& OutOutIndexOffsets)
+	{
+		return ComparisonIsUnimplemented(FirstData, SecondData);
+	}
+
 	bool DataTypeIsComparable(EPCGDataType DataType)
 	{
 		// Comparable data types
+		// TODO: Add metadata/param data as comparable in the future
 		if (!!(DataType & EPCGDataType::Spatial))
 		{
 			return true;
@@ -616,6 +940,11 @@ namespace PCGDeterminismTests
 		ShuffleArray<FPCGTaggedData>(TestData.InputData.TaggedData, TestData.RandomStream);
 	}
 
+	void ShuffleOutputOrder(PCGTestsCommon::FTestData& TestData)
+	{
+		ShuffleArray<FPCGTaggedData>(TestData.OutputData.TaggedData, TestData.RandomStream);
+	}
+
 	void ShuffleAllInternalData(PCGTestsCommon::FTestData& TestData)
 	{
 		for (FPCGTaggedData& TaggedData : TestData.InputData.TaggedData)
@@ -628,32 +957,155 @@ namespace PCGDeterminismTests
 		}
 	}
 
-	TFunction<bool(const UPCGData*, const UPCGData*)> GetCompareFunction(EPCGDataType DataType)
+	void ShiftInputOrder(PCGTestsCommon::FTestData& TestData, int32 NumShifts)
+	{
+		ShiftArrayElements<FPCGTaggedData>(TestData.InputData.TaggedData, NumShifts);
+	}
+
+	TArray<EPCGDataType> FilterTestableDataTypes(EPCGDataType AllowedTypes, int32 NumMultipleInputs)
+	{
+		TArray<EPCGDataType> RemainingDataTypes;
+		if (!(EPCGDataType::Spatial & AllowedTypes))
+		{
+			// TODO: Expand this to other data types
+			UE_LOG(LogPCG, Error, TEXT("Only spatial data types are currently supported"));
+			return RemainingDataTypes;
+		}
+
+		// Add every possible combination of testable data types
+		for (EPCGDataType DataType : TestableDataTypes)
+		{
+			if (DataType == EPCGDataType::None || !!(AllowedTypes & DataType))
+			{
+				for (int32 I = 0; I < NumMultipleInputs; ++I)
+				{
+					RemainingDataTypes.Emplace(DataType);
+				}
+			}
+		}
+
+		return RemainingDataTypes;
+	}
+
+	void RetrieveBaseOptionsPerPin(TArray<TArray<EPCGDataType>>& BaseOptionsArray,
+		const TArray<TObjectPtr<UPCGPin>>& InputPins,
+		EPCGDataType& OutDataTypesTested,
+		int32 NumMultipleInputs)
+	{
+		BaseOptionsArray.Empty();
+
+		for (int32 I = 0; I < InputPins.Num(); ++I)
+		{
+			EPCGDataType AllowedDataTypes = InputPins[I]->Properties.AllowedTypes;
+			if (!DataTypeIsComparable(AllowedDataTypes))
+			{
+				continue;
+			}
+
+			TArray<EPCGDataType> TestableDataTypesForPin = FilterTestableDataTypes(AllowedDataTypes, NumMultipleInputs);
+			if (!TestableDataTypesForPin.IsEmpty())
+			{
+				BaseOptionsArray.Emplace(TestableDataTypesForPin);
+				OutDataTypesTested |= AllowedDataTypes;
+			}
+		}
+	}
+
+	int32 GetNumPermutations(const TArray<TArray<EPCGDataType>>& BaseOptionsArray)
+	{
+		int32 NumPermutations = 1;
+
+		for (int32 I = 0; I < BaseOptionsArray.Num(); ++I)
+		{
+			NumPermutations *= FMath::Max(BaseOptionsArray[I].Num(), 1);
+		}
+
+		return NumPermutations;
+	}
+
+	EPCGDataType GetPermutation(int32 PermutationIteration, int32 PinIndex, const TArray<TArray<EPCGDataType>>& BaseOptionsPerPin)
+	{
+		check(PermutationIteration >= 0);
+		check(PinIndex >= 0);
+		check(PinIndex < BaseOptionsPerPin[PinIndex].Num());
+
+		int32 CardinalityProduct = 1;
+		for (int I = 0; I < PinIndex; ++I)
+		{
+			CardinalityProduct *= BaseOptionsPerPin[I].Num();
+		}
+		check(CardinalityProduct > 0);
+
+		int32 PermutationIndex = (PermutationIteration / CardinalityProduct) % BaseOptionsPerPin[PinIndex].Num();
+		return BaseOptionsPerPin[PinIndex][PermutationIndex];
+	}
+
+	void UpdateTestResults(FName TestName, FNodeTestResult& OutResult, EDeterminismLevel DeterminismLevel)
+	{
+		EDeterminismLevel& Result = OutResult.TestResults.FindOrAdd(TestName);
+		Result = DeterminismLevel;
+	}
+
+	TFunction<bool(const UPCGData*, const UPCGData*)> GetDataCompareFunction(EPCGDataType DataType, EDeterminismLevel DeterminismLevel)
 	{
-		if (!DataTypeIsComparable(DataType))
+		// TODO: Currently only spatial data is comparable. Other types to be added.
+		if (!DataTypeIsComparable(DataType) || DataType != EPCGDataType::Spatial)
 		{
 			// Should never reach here
 			UE_LOG(LogPCG, Warning, TEXT("Attempting to compare incomparable data."));
 			return ComparisonIsUnimplemented;
 		}
 
-		switch (DataType)
+		switch (DeterminismLevel)
 		{
-		case EPCGDataType::Spatial:
+		case EDeterminismLevel::OrderOrthogonal:
+			return SpatialDataIsOrthogonal;
+		case EDeterminismLevel::OrderConsistent:
+			return SpatialDataIsConsistent;
+		case EDeterminismLevel::OrderIndependent:
 			return SpatialDataIsIdentical;
 		default:
-			UE_LOG(LogPCG, Warning, TEXT("Comparable PCG DataType has no 'comparison function'."));
+			UE_LOG(LogPCG, Warning, TEXT("Comparable PCG DataType has no 'data comparison function'."));
+			return ComparisonIsUnimplemented;
+		}
+	}
+
+	TFunction<bool(const UPCGData*, const UPCGData*)> GetDataCollectionCompareFunction(EPCGDataType DataType)
+	{
+		// TODO: Currently only spatial data is comparable. Other types to be added.
+		if (!DataTypeIsComparable(DataType) || !(DataType & EPCGDataType::Spatial))
+		{
+			// Should never reach here
+			UE_LOG(LogPCG, Warning, TEXT("Attempting to compare incomparable data."));
 			return ComparisonIsUnimplemented;
 		}
+
+		return SpatialBasicsAreIdentical;
 	}
 
-	bool ExecutionIsDeterministic(const PCGTestsCommon::FTestData& FirstTestData, const PCGTestsCommon::FTestData& SecondTestData, const UPCGNode* PCGNode)
+	void ExecuteWithTestData(PCGTestsCommon::FTestData& TestData, const UPCGNode* PCGNode)
 	{
-		FPCGElementPtr FirstElement = FirstTestData.Settings->GetElement();
-		FPCGElementPtr SecondElement = SecondTestData.Settings->GetElement();
+		FPCGElementPtr Element = TestData.Settings->GetElement();
 
-		TUniquePtr<FPCGContext> FirstContext(FirstElement->Initialize(FirstTestData.InputData, FirstTestData.TestPCGComponent, PCGNode));
-		TUniquePtr<FPCGContext> SecondContext(SecondElement->Initialize(SecondTestData.InputData, SecondTestData.TestPCGComponent, PCGNode));
+		TUniquePtr<FPCGContext> Context(Element->Initialize(TestData.InputData, TestData.TestPCGComponent, PCGNode));
+
+		Context->NumAvailableTasks = 1;
+
+		// Execute both elements until complete
+		while (!Element->Execute(Context.Get()))
+		{
+		}
+
+		TestData.OutputData = MoveTemp(Context->OutputData);
+	}
+
+	void ExecuteWithSameTestData(const PCGTestsCommon::FTestData& TestData, const UPCGNode* PCGNode, FPCGDataCollection& OutFirstOutputData, FPCGDataCollection& OutSecondOutputData)
+	{
+		FPCGElementPtr FirstElement = TestData.Settings->GetElement();
+		FPCGElementPtr SecondElement = TestData.Settings->GetElement();
+
+		TUniquePtr<FPCGContext> FirstContext(FirstElement->Initialize(TestData.InputData, TestData.TestPCGComponent, PCGNode));
+		TUniquePtr<FPCGContext> SecondContext(SecondElement->Initialize(TestData.InputData, TestData.TestPCGComponent, PCGNode));
 
 		FirstContext->NumAvailableTasks = 1;
 		SecondContext->NumAvailableTasks = 1;
@@ -663,17 +1115,60 @@ namespace PCGDeterminismTests
 		{
 		}
 
+		// Execute both elements until complete
 		while (!SecondElement->Execute(SecondContext.Get()))
 		{
 		}
 
-		return DataCollectionsAreIdentical(FirstContext->OutputData, SecondContext->OutputData);
+		OutFirstOutputData = FirstContext->OutputData;
+		OutSecondOutputData = SecondContext->OutputData;
 	}
 
-	bool ExecutionIsDeterministicSameData(PCGTestsCommon::FTestData& TestData, const UPCGNode* PCGNode)
+	void ExecuteWithSameTestDataSameElement(const PCGTestsCommon::FTestData& TestData, const UPCGNode* PCGNode, FPCGDataCollection& OutFirstOutputData, FPCGDataCollection& OutSecondOutputData)
 	{
-		return ExecutionIsDeterministic(TestData, TestData, PCGNode);
+		FPCGElementPtr Element = TestData.Settings->GetElement();
+
+		TUniquePtr<FPCGContext> FirstContext(Element->Initialize(TestData.InputData, TestData.TestPCGComponent, PCGNode));
+		TUniquePtr<FPCGContext> SecondContext(Element->Initialize(TestData.InputData, TestData.TestPCGComponent, PCGNode));
+
+		FirstContext->NumAvailableTasks = 1;
+		SecondContext->NumAvailableTasks = 1;
+
+		// Execute both elements until complete
+		while (!Element->Execute(FirstContext.Get()))
+		{
+		}
+
+		// Execute both elements until complete
+		while (!Element->Execute(SecondContext.Get()))
+		{
+		}
+
+		OutFirstOutputData = FirstContext->OutputData;
+		OutSecondOutputData = SecondContext->OutputData;
+	}
+
+	bool ExecutionIsDeterministic(PCGTestsCommon::FTestData& FirstTestData, PCGTestsCommon::FTestData& SecondTestData, const UPCGNode* PCGNode)
+	{
+		ExecuteWithTestData(FirstTestData, PCGNode);
+		ExecuteWithTestData(SecondTestData, PCGNode);
+
+		return DataCollectionsContainSameData(FirstTestData.OutputData, SecondTestData.OutputData);
+	}
+
+	bool ExecutionIsDeterministicSameData(const PCGTestsCommon::FTestData& TestData, const UPCGNode* PCGNode)
+	{
+		FPCGDataCollection FirstOutput;
+		FPCGDataCollection SecondOutput;
+		ExecuteWithSameTestData(TestData, PCGNode, FirstOutput, SecondOutput);
+
+		FPCGDataCollection ThirdOutput;
+		FPCGDataCollection FourthOutput;
+		ExecuteWithSameTestDataSameElement(TestData, PCGNode, FirstOutput, SecondOutput);
+
+		// Same data, so it ought to be exact
+		return DataCollectionsAreIdentical(FirstOutput, SecondOutput) && DataCollectionsAreIdentical(ThirdOutput, FourthOutput);
 	}
 }
 
-#undef LOCTEXT_NAMESPACE
\ No newline at end of file
+#undef LOCTEXT_NAMESPACE
diff --git a/Engine/Plugins/Experimental/PCG/Source/PCG/Private/Tests/PCGTestsCommon.cpp b/Engine/Plugins/Experimental/PCG/Source/PCG/Private/Tests/PCGTestsCommon.cpp
--- a/Engine/Plugins/Experimental/PCG/Source/PCG/Private/Tests/PCGTestsCommon.cpp
+++ b/Engine/Plugins/Experimental/PCG/Source/PCG/Private/Tests/PCGTestsCommon.cpp
@@ -70,12 +70,13 @@ namespace PCGTestsCommon
 #endif // WITH_EDITOR
 	}
 
-	void FTestData::Reset()
+	void FTestData::Reset(UPCGSettings* InSettings)
 	{
 		// Clear all the data
 		RandomStream.Reset();
 		InputData.TaggedData.Empty();
-		Settings = nullptr;
+		OutputData.TaggedData.Empty();
+		Settings = InSettings;
 	}
 
 	AActor* CreateTemporaryActor()
diff --git a/Engine/Plugins/Experimental/PCG/Source/PCG/Public/Tests/Determinism/PCGDeterminismTestsCommon.h b/Engine/Plugins/Experimental/PCG/Source/PCG/Public/Tests/Determinism/PCGDeterminismTestsCommon.h
--- a/Engine/Plugins/Experimental/PCG/Source/PCG/Public/Tests/Determinism/PCGDeterminismTestsCommon.h
+++ b/Engine/Plugins/Experimental/PCG/Source/PCG/Public/Tests/Determinism/PCGDeterminismTestsCommon.h
@@ -6,7 +6,6 @@
 #include "PCGSettings.h"
 #include "Tests/PCGTestsCommon.h"
 
-class AActor;
 class UPCGComponent;
 class UPCGData;
 class UPCGSpatialData;
@@ -20,11 +19,11 @@ namespace PCGDeterminismTests
 	{
 		constexpr static int32 Seed = 42;
 		constexpr static int32 NumPointsToGenerate = 1;
-		constexpr static int32 NumTestPointsToGenerate = 1000;
+		constexpr static int32 NumTestPointsToGenerate = 100;
 		constexpr static int32 NumPolyLinePointsToGenerate = 6;
-		constexpr static int32 NumTestPolyLinePointsToGenerate = 12;
+		constexpr static int32 NumTestPolyLinePointsToGenerate = 6;
 		constexpr static int32 NumSamplingStepsPerDimension = 100;
-		constexpr static int32 NumMultipleTestDataSets = 2;
+		constexpr static int32 NumTestInputsPerPin = 2;
 
 		constexpr static FVector::FReal SmallDistance = 50.0;
 		constexpr static FVector::FReal MediumDistance = 200.0;
@@ -39,56 +38,99 @@ namespace PCGDeterminismTests
 		const FBox TestingVolume = FBox(-1.2 * LargeVector, 1.2 * LargeVector);
 	}
 
+	constexpr static EPCGDataType TestableDataTypes[6] =
+	{
+		EPCGDataType::None,
+		EPCGDataType::Point,
+		EPCGDataType::Volume,
+		EPCGDataType::PolyLine,
+		EPCGDataType::Primitive,
+		EPCGDataType::Landscape
+	};
+
+	// This will include multiple values of different meanings, but we use an enum to facilitate data passing
+	UENUM()
+	enum class EDeterminismLevel : uint8
+	{
+		None = 0u,
+		NoDeterminism = None,
+		Basic,
+		OrderOrthogonal,
+		OrderConsistent,
+		OrderIndependent,
+		Deterministic = OrderIndependent
+	};
+
 	struct FNodeTestResult
 	{
 		int32 Index = -1;
-		EPCGDataType DataTypesTested = EPCGDataType::None;
 		FName NodeTitle = TEXT("Untitled");
 		FString NodeNameString = TEXT("Unnamed");
-		TMap<FName, bool> TestResults;
+		// TODO: Add the seed to the UI widget
+		int32 Seed = -1;
+		EPCGDataType DataTypesTested = EPCGDataType::None;
+		TMap<FName, EDeterminismLevel> TestResults;
 		TArray<FString> AdditionalDetails;
 		bool bFlagRaised = false;
 		// TODO: Chrono Duration of how long the tests took
 	};
 
 	/** A default delegate to report an unset test */
-	bool LogInvalidTest(const UPCGNode* InPCGNode, int32 Seed, EPCGDataType& OutDataTypesTested, TArray<FString>& OutAdditionalDetails);
+	bool LogInvalidTest(const UPCGNode* InPCGNode, const FName& TestName, FNodeTestResult& OutResult);
 
-	typedef TFunction<bool(const UPCGNode* InPCGNode, int32 Seed, EPCGDataType& OutDataTypesTested, TArray<FString>& OutAdditionalDetails)> TestFunction;
+	typedef TFunction<bool(const UPCGNode* InPCGNode, const FName& TestName, FNodeTestResult& OutResult)> TestFunction;
 
 	struct FNodeTestInfo
 	{
 		FNodeTestInfo(FText Label, TestFunction Delegate, float LabelWidth = 140.f) :
 			TestLabel(Label),
+			TestName(Label.ToString()),
 			TestDelegate(Delegate),
 			TestLabelWidth(LabelWidth) {}
 
 		FText TestLabel = NSLOCTEXT("PCGDeterminism", "UnnamedTest", "Unnamed Test");
+		const FName TestName;
 		TestFunction TestDelegate = LogInvalidTest;
 		float TestLabelWidth = 140.f;
 	};
 
+	struct FNodeAndOptions
+	{
+		explicit FNodeAndOptions(const UPCGNode* PCGNode, const int32 Seed, bool bMultipleOptionsPerPin) :
+			PCGNode(PCGNode),
+			Seed(Seed),
+			bMultipleOptionsPerPin(bMultipleOptionsPerPin) {}
+
+		const UPCGNode* PCGNode;
+		const int32 Seed;
+		const bool bMultipleOptionsPerPin;
+		TArray<TArray<EPCGDataType>> BaseOptionsByPin;
+	};
+
 	/** Validates if a PCGNode is deterministic */
 	PCG_API void RunDeterminismTest(const UPCGNode* InPCGNode, FNodeTestResult& OutResult, const FNodeTestInfo& TestToRun);
 
 	/** Adds the basic set of determinism tests to the passed in array */
 	PCG_API void RetrieveBasicTests(TArray<FNodeTestInfo>& OutBasicTests);
 
-	/** Validates node determinism against the same single test data */
-	bool RunSingleSameDataTest(const UPCGNode* InPCGNode, int32 Seed, EPCGDataType& OutDataTypesTested, TArray<FString>& OutAdditionalDetails);
-	/** Validates node determinism against two identical single test data */
-	bool RunSingleIdenticalDataTest(const UPCGNode* InPCGNode, int32 Seed, EPCGDataType& OutDataTypesTested, TArray<FString>& OutAdditionalDetails);
-	/** Validates node determinism with the same multiple sets of test data */
-	bool RunMultipleSameDataTest(const UPCGNode* InPCGNode, int32 Seed, EPCGDataType& OutDataTypesTested, TArray<FString>& OutAdditionalDetails);
-	/** Validates node determinism with two identical multiple sets of test data */
-	bool RunMultipleIdenticalDataTest(const UPCGNode* InPCGNode, int32 Seed, EPCGDataType& OutDataTypesTested, TArray<FString>& OutAdditionalDetails);
-	/** Validates node determinism with multiple sets of test data, shuffling the order of the second set's data collection */
-	bool RunDataCollectionOrderIndependenceTest(const UPCGNode* InPCGNode, int32 Seed, EPCGDataType& OutDataTypesTested, TArray<FString>& OutAdditionalDetails);
-	/** Validates node determinism with multiple sets of test data, shuffling all internal data */
-	bool RunAllDataOrderIndependenceTest(const UPCGNode* InPCGNode, int32 Seed, EPCGDataType& OutDataTypesTested, TArray<FString>& OutAdditionalDetails);
+	/** Validates all the generic determinism tests for any given node */
+	bool RunBasicTestSuite(const UPCGNode* InPCGNode, const FName& TestName, FNodeTestResult& OutResult);
+
+	/** Validates the various levels of order independence for any given node */
+	bool RunOrderIndependenceSuite(const UPCGNode* InPCGNode, const FName& TestName, FNodeTestResult& OutResult);
+
+	/** Validates minimal node determinism against the same single test data */
+	bool RunBasicSelfTest(const FNodeAndOptions& NodeAndOptions);
+	/** Validates minimal node determinism against two identical single test data */
+	bool RunBasicCopiedSelfTest(const FNodeAndOptions& NodeAndOptions);
+
+	/** Conducts tests on all permutations and determines the highest level of determinism */
+	EDeterminismLevel GetHighestDeterminismLevel(const FNodeAndOptions& NodeAndOptions,
+		int32 NumInputsPerPin = Defaults::NumTestInputsPerPin,
+		EDeterminismLevel MaxLevel = EDeterminismLevel::OrderIndependent);
 
 	/** Adds input data to test data, based on input pins' allowed types */
-	void AddInputDataBasedOnPins(PCGTestsCommon::FTestData& TestData, const UPCGNode* InPCGNode, EPCGDataType& OutDataTypesTested, TArray<FString>& OutAdditionalDetails);
+	void AddRandomizedInputData(PCGTestsCommon::FTestData& TestData, EPCGDataType DataType, const FName& PinName = Defaults::TestPinName);
 
 	/** Helper for adding PointData with a single Point to an InputData */
 	void AddSinglePointInputData(FPCGDataCollection& InputData, const FVector& Location, const FName& PinName = Defaults::TestPinName);
@@ -100,6 +142,8 @@ namespace PCGDeterminismTests
 	void AddPolyLineInputData(FPCGDataCollection& InputData, USplineComponent* SplineComponent, const FName& PinName = Defaults::TestPinName);
 	/** Helper for adding PrimitiveData to an InputData */
 	void AddPrimitiveInputData(FPCGDataCollection& InputData, UPrimitiveComponent* SplineComponent, const FVector& VoxelSize, const FName& PinName = Defaults::TestPinName);
+	/** Helper for adding LandscapeData to an InputData */
+	void AddLandscapeInputData(FPCGDataCollection& InputData);
 
 	/** Adds randomized PointData with a single Point */
 	void AddRandomizedSinglePointInputData(PCGTestsCommon::FTestData& TestData, int32 PointNum = Defaults::NumPointsToGenerate, const FName& PinName = Defaults::TestPinName);
@@ -113,9 +157,21 @@ namespace PCGDeterminismTests
 	void AddRandomizedPolyLineInputData(PCGTestsCommon::FTestData& TestData, int32 PointNum = Defaults::NumPolyLinePointsToGenerate, const FName& PinName = Defaults::TestPinName);
 	/** Adds randomized Primitive Spatial Data */
 	void AddRandomizedPrimitiveInputData(PCGTestsCommon::FTestData& TestData, const FName& PinName = Defaults::TestPinName);
+	/** Adds randomized Landscape Spatial Data */
+	void AddRandomizedLandscapeInputData(PCGTestsCommon::FTestData& TestData, const FName& PinName = Defaults::TestPinName);
 
-	/** Validates whether two DataCollection objects are identical */
+	/** Validates whether two DataCollection objects are exactly identical */
 	bool DataCollectionsAreIdentical(const FPCGDataCollection& FirstCollection, const FPCGDataCollection& SecondCollection);
+	/** Validates whether two DataCollection objects have data in order relative to their inputs */
+	bool DataCollectionsAreConsistent(const FPCGDataCollection& FirstCollection, const FPCGDataCollection& SecondCollection, int32 NumInputs);
+	/** Validates whether two DataCollection objects contain all the same data */
+	bool DataCollectionsContainSameData(const FPCGDataCollection& FirstCollection, const FPCGDataCollection& SecondCollection);
+	/** Validates the data collections match and delivers matched indices and offsets */
+	bool DataCollectionsMatch(const FPCGDataCollection& FirstCollection,
+		const FPCGDataCollection& SecondCollection,
+		TArray<int32>& OutIndexOffsets);
+	/** Validates whether the internal data within the two DataCollection objects are consistent */
+	bool InternalDataMatches(const UPCGData* FirstData, const UPCGData* SecondData, TArray<int32>& OutIndexOffsets);
 
 	/** Validates whether two SpatialData objects are identical */
 	bool SpatialDataIsIdentical(const UPCGData* FirstData, const UPCGData* SecondData);
@@ -131,9 +187,18 @@ namespace PCGDeterminismTests
 	bool PrimitiveDataIsIdentical(const UPCGData* FirstData, const UPCGData* SecondData);
 	/** Validates whether two SpatialData objects are identical via Point Sampling */
 	bool SampledSpatialDataIsIdentical(const UPCGSpatialData* FirstSpatialData, const UPCGSpatialData* SecondSpatialData);
-
 	/** Validates the basics of SpatialData are identical */
-	bool SpatialBasicsAreIdentical(const UPCGSpatialData* FirstSpatialData, const UPCGSpatialData* SecondSpatialData);
+	bool SpatialBasicsAreIdentical(const UPCGData* FirstData, const UPCGData* SecondData);
+
+	/** Validates whether two SpatialData objects are consistent */
+	bool SpatialDataIsConsistent(const UPCGData* FirstData, const UPCGData* SecondData);
+	/** Validates whether two PointData objects are consistent */
+	bool PointDataIsConsistent(const UPCGData* FirstData, const UPCGData* SecondData);
+
+	/** Validates whether two SpatialData objects contain the same spatial data */
+	bool SpatialDataIsOrthogonal(const UPCGData* FirstData, const UPCGData* SecondData);
+	/** Validates whether two PointData objects contain the same points */
+	bool PointDataIsOrthogonal(const UPCGData* FirstData, const UPCGData* SecondData);
 
 	/** A catch function for unimplemented comparisons */
 	bool ComparisonIsUnimplemented(const UPCGData* FirstData, const UPCGData* SecondData);
@@ -147,16 +212,43 @@ namespace PCGDeterminismTests
 
 	/** Randomizes the order of InputData collections */
 	void ShuffleInputOrder(PCGTestsCommon::FTestData& TestData);
+	/** Randomizes the order of OutputData collections */
+	void ShuffleOutputOrder(PCGTestsCommon::FTestData& TestData);
 	/** Randomizes the order of all internal data */
 	void ShuffleAllInternalData(PCGTestsCommon::FTestData& TestData);
-
-	/** Gets a comparison function to compare two of a specific DataType */
-	TFunction<bool(const UPCGData*, const UPCGData*)> GetCompareFunction(EPCGDataType DataType);
-
-	/** Execute the elements for each valid input and compare if all the outputs are identical */
-	bool ExecutionIsDeterministic(const PCGTestsCommon::FTestData& FirstTestData, const PCGTestsCommon::FTestData& SecondTestData, const UPCGNode* PCGNode = nullptr);
-	/** Execute the same element twice compare if all the outputs are identical */
-	bool ExecutionIsDeterministicSameData(PCGTestsCommon::FTestData& TestData, const UPCGNode* PCGNode = nullptr);
+	/** Shifts the order of InputData collections */
+	void ShiftInputOrder(PCGTestsCommon::FTestData& TestData, int32 NumShifts = 1);
+
+	/** Gets an array of the union between an InputPin's allowed types and testable types */
+	TArray<EPCGDataType> FilterTestableDataTypes(EPCGDataType AllowedDataTypes, int32 NumMultipleInputs = 1);
+	/** Given input pins, updates an array with a permutation base of options */
+	void RetrieveBaseOptionsPerPin(TArray<TArray<EPCGDataType>>& InBaseOptionsArray,
+		const TArray<TObjectPtr<UPCGPin>>& InputPins,
+		EPCGDataType& OutDataTypesTested,
+		int32 NumMultipleInputs = 1);
+	/** Gets the number of permutations from a base set of options */
+	int32 GetNumPermutations(const TArray<TArray<EPCGDataType>>& BaseOptionsArray);
+	/** Gets the index of a pin's permutation based on the permutation iteration */
+	EPCGDataType GetPermutation(int32 PermutationIteration, int32 PinIndex, const TArray<TArray<EPCGDataType>>& BaseOptionsPerPin);
+
+	/** Helper to update a tests' outgoing results */
+	void UpdateTestResults(FName TestName, FNodeTestResult& OutResult, EDeterminismLevel DeterminismLevel);
+
+	/** Gets a comparison function to compare two data objects */
+	TFunction<bool(const UPCGData*, const UPCGData*)> GetDataCompareFunction(EPCGDataType DataType, EDeterminismLevel DeterminismLevel);
+	/** Gets a comparison function to compare two data collection entry basics */
+	TFunction<bool(const UPCGData*, const UPCGData*)> GetDataCollectionCompareFunction(EPCGDataType DataType);
+
+	/** Executes the element with given test data */
+	void ExecuteWithTestData(PCGTestsCommon::FTestData& TestData, const UPCGNode* PCGNode);
+	/** Executes the element with given test data against itself */
+	void ExecuteWithSameTestData(const PCGTestsCommon::FTestData& TestData, const UPCGNode* PCGNode, FPCGDataCollection& OutFirstOutputData, FPCGDataCollection& OutSecondOutputData);
+	/** Executes the element with given test data against itself, with the same element */
+	void ExecuteWithSameTestDataSameElement(const PCGTestsCommon::FTestData& TestData, const UPCGNode* PCGNode, FPCGDataCollection& OutFirstOutputData, FPCGDataCollection& OutSecondOutputData);
+	/** Execute the elements for each valid input and compare if all the outputs are at least orthogonally deterministic */
+	bool ExecutionIsDeterministic(PCGTestsCommon::FTestData& FirstTestData, PCGTestsCommon::FTestData& SecondTestData, const UPCGNode* PCGNode = nullptr);
+	/** Execute the same element twice compare if all the outputs are at least orthogonally deterministic */
+	bool ExecutionIsDeterministicSameData(const PCGTestsCommon::FTestData& TestData, const UPCGNode* PCGNode = nullptr);
 
 	/** Generates settings based upon a UPCGSettings subclass */
 	template<typename SettingsType>
@@ -185,7 +277,7 @@ namespace PCGDeterminismTests
 	{
 		check(FirstData && SecondData);
 
-		return (Cast<const DataType>(FirstData) != nullptr && Cast<const DataType>(SecondData) != nullptr);
+		return Cast<const DataType>(FirstData) != nullptr && Cast<const DataType>(SecondData) != nullptr;
 	}
 
 	template<typename DataType>
@@ -202,4 +294,35 @@ namespace PCGDeterminismTests
 			}
 		}
 	}
+
+	template<typename DataType>
+	void ShiftArrayElements(TArray<DataType>& Array, int32 NumShifts = 1)
+	{
+		if (Array.Num() < 2)
+		{
+			return;
+		}
+
+		int32 Count = Array.Num();
+		NumShifts %= Count;
+		if (NumShifts < 0)
+		{
+			NumShifts += Count;
+		}
+
+		TArray<DataType> TempArray;
+		TempArray.SetNum(Count);
+
+		for (int32 I = 0; I < NumShifts; ++I)
+		{
+			TempArray[I] = Array[I + Count - NumShifts];
+		}
+
+		for (int32 I = NumShifts; I < Count; ++I)
+		{
+			TempArray[I] = Array[I - NumShifts];
+		}
+
+		Array = MoveTemp(TempArray);
+	}
 }
\ No newline at end of file
diff --git a/Engine/Plugins/Experimental/PCG/Source/PCG/Public/Tests/PCGTestsCommon.h b/Engine/Plugins/Experimental/PCG/Source/PCG/Public/Tests/PCGTestsCommon.h
--- a/Engine/Plugins/Experimental/PCG/Source/PCG/Public/Tests/PCGTestsCommon.h
+++ b/Engine/Plugins/Experimental/PCG/Source/PCG/Public/Tests/PCGTestsCommon.h
@@ -21,18 +21,19 @@ struct FPCGPoint;
 
 namespace PCGTestsCommon
 {
-	static const int TestFlags = EAutomationTestFlags::EditorContext | EAutomationTestFlags::EngineFilter;
+	constexpr int TestFlags = EAutomationTestFlags::EditorContext | EAutomationTestFlags::EngineFilter;
 
 	struct FTestData
 	{
 		FTestData(int32 Seed, UPCGSettings* DefaultSettings = nullptr, TSubclassOf<AActor> ActorClass = AActor::StaticClass());
 		~FTestData();
 
-		void Reset();
+		void Reset(UPCGSettings* InSettings = nullptr);
 
 		AActor* TestActor;
 		UPCGComponent* TestPCGComponent;
 		FPCGDataCollection InputData;
+		FPCGDataCollection OutputData;
 		UPCGSettings* Settings;
 		int32 Seed;
 		FRandomStream RandomStream;
diff --git a/Engine/Plugins/Experimental/PCG/Source/PCGEditor/Private/SPCGEditorGraphDeterminism.cpp b/Engine/Plugins/Experimental/PCG/Source/PCGEditor/Private/SPCGEditorGraphDeterminism.cpp
--- a/Engine/Plugins/Experimental/PCG/Source/PCGEditor/Private/SPCGEditorGraphDeterminism.cpp
+++ b/Engine/Plugins/Experimental/PCG/Source/PCGEditor/Private/SPCGEditorGraphDeterminism.cpp
@@ -20,13 +20,16 @@ namespace
 	const FText TEXT_DataTypesTested(LOCTEXT("DataTypesTested_Label", "Input Data"));
 	const FText TEXT_AdditionalDetails(LOCTEXT("AdditionalDetails_Label", "Additional Details"));
 
-	const FText TEXT_Yes(LOCTEXT("Yes", "Yes"));
-	const FText TEXT_No(LOCTEXT("No", "No"));
-
-	constexpr static float SmallManualWidth = 25.f;
-	constexpr static float MediumManualWidth = 70.f;
-	constexpr static float LargeManualWidth = 160.f;
-	constexpr static float ListViewRowHeight = 36.f;
+	const FText TEXT_NotDeterministic(LOCTEXT("NotDeterministic", "Fail"));
+	const FText TEXT_Consistent(LOCTEXT("OrderConsistent", "Order Consistent"));
+	const FText TEXT_Independent(LOCTEXT("OrderIndependent", "Order Independent"));
+	const FText TEXT_Orthogonal(LOCTEXT("OrderOrthogonal", "Order Orthogonal"));
+	const FText TEXT_Basic(LOCTEXT("BasicDeterminism", "Pass"));
+
+	constexpr float SmallManualWidth = 25.f;
+	constexpr float MediumManualWidth = 70.f;
+	constexpr float LargeManualWidth = 160.f;
+	constexpr float ListViewRowHeight = 36.f;
 }
 
 void SPCGEditorGraphDeterminismRow::Construct(const FArguments& InArgs, const TSharedRef<STableViewBase>& InOwnerTableView, const FPCGNodeTestResultPtr& Item)
@@ -41,10 +44,10 @@ TSharedRef<SWidget> SPCGEditorGraphDeterminismRow::GenerateWidgetForColumn(const
 
 	FText CellText = LOCTEXT("UnknownColumn", "Unknown");
 
-	auto ReturnColorCodedResultBlock = [this](const FText& CellText)
+	auto ReturnColorCodedResultBlock = [this](const FText& OutCellText)
 	{
 		return SNew(STextBlock)
-			.Text(CellText)
+			.Text(OutCellText)
 			.ColorAndOpacity(CurrentItem->bFlagRaised ? FColor::Red : FColor::Green);
 	};
 
@@ -72,13 +75,34 @@ TSharedRef<SWidget> SPCGEditorGraphDeterminismRow::GenerateWidgetForColumn(const
 		FString FullDetails = FString::Join(CurrentItem->AdditionalDetails, TEXT(", "));
 		CellText = FText::FromString(FullDetails);
 	}
+
 	// Test columns
-	else if (bool* bResultPtr = CurrentItem->TestResults.Find(ColumnId))
+	if (PCGDeterminismTests::EDeterminismLevel* DeterminismLevel = CurrentItem->TestResults.Find(ColumnId))
 	{
-		bool bTestWasSuccessful = bResultPtr && *bResultPtr;
-		return SNew(STextBlock)
-			.Text(bTestWasSuccessful ? TEXT_Yes : TEXT_No)
-			.ColorAndOpacity(bTestWasSuccessful ? FColor::Green : FColor::Red);
+		switch (*DeterminismLevel)
+		{
+		case PCGDeterminismTests::EDeterminismLevel::OrderOrthogonal:
+			return SNew(STextBlock)
+				.Text(TEXT_Orthogonal)
+				.ColorAndOpacity(FColor::Orange);
+		case PCGDeterminismTests::EDeterminismLevel::OrderConsistent:
+			return SNew(STextBlock)
+				.Text(TEXT_Consistent)
+				.ColorAndOpacity(FColor::Yellow);
+		case PCGDeterminismTests::EDeterminismLevel::OrderIndependent:
+			return SNew(STextBlock)
+				.Text(TEXT_Independent)
+				.ColorAndOpacity(FColor::Green);
+		case PCGDeterminismTests::EDeterminismLevel::Basic:
+			return SNew(STextBlock)
+				.Text(TEXT_Basic)
+				.ColorAndOpacity(FColor::Turquoise);
+		case PCGDeterminismTests::EDeterminismLevel::NoDeterminism:
+		default:
+			return SNew(STextBlock)
+				.Text(TEXT_NotDeterministic)
+				.ColorAndOpacity(FColor::Red);
+		}
 	}
 
 	return SNew(STextBlock)
