Fixed landscape spline bounds that were missing a transform

#rnx

diff --git a/Engine/Plugins/Experimental/PCG/Source/PCG/Private/Data/PCGLandscapeSplineData.cpp b/Engine/Plugins/Experimental/PCG/Source/PCG/Private/Data/PCGLandscapeSplineData.cpp
--- a/Engine/Plugins/Experimental/PCG/Source/PCG/Private/Data/PCGLandscapeSplineData.cpp
+++ b/Engine/Plugins/Experimental/PCG/Source/PCG/Private/Data/PCGLandscapeSplineData.cpp
@@ -148,6 +148,11 @@ FBox UPCGLandscapeSplineData::GetBounds() const
 		Bounds += Segment->GetBounds();
 	}
 
+	if (Bounds.IsValid)
+	{
+		Bounds = Bounds.TransformBy(Spline->GetComponentToWorld());
+	}
+
 	return Bounds;
 }
 
