Added landscape height only pin + fix non-initialized value

diff --git a/Engine/Plugins/Experimental/PCG/Source/PCG/Private/Data/PCGLandscapeData.cpp b/Engine/Plugins/Experimental/PCG/Source/PCG/Private/Data/PCGLandscapeData.cpp
--- a/Engine/Plugins/Experimental/PCG/Source/PCG/Private/Data/PCGLandscapeData.cpp
+++ b/Engine/Plugins/Experimental/PCG/Source/PCG/Private/Data/PCGLandscapeData.cpp
@@ -13,12 +13,13 @@
 #include "Landscape.h"
 #include "LandscapeEdit.h"
 
-void UPCGLandscapeData::Initialize(ALandscapeProxy* InLandscape, const FBox& InBounds)
+void UPCGLandscapeData::Initialize(ALandscapeProxy* InLandscape, const FBox& InBounds, bool bInHeightOnly)
 {
 	check(InLandscape);
 	Landscape = InLandscape;
 	TargetActor = InLandscape;
 	Bounds = InBounds;
+	bHeightOnly = bInHeightOnly;
 
 	Transform = Landscape->GetActorTransform();
 
@@ -29,11 +30,14 @@ void UPCGLandscapeData::Initialize(ALandscapeProxy* InLandscape, const FBox& InB
 	// TODO: find a better way to do this - maybe there should be a prototype metadata in the landscape cache
 	if (LandscapeCache)
 	{
-		const TArray<FName> Layers = LandscapeCache->GetLayerNames(Landscape.Get());
-
-		for (const FName& Layer : Layers)
+		if (!bHeightOnly)
 		{
-			Metadata->CreateFloatAttribute(Layer, 0.0f, /*bAllowInterpolation=*/true);
+			const TArray<FName> Layers = LandscapeCache->GetLayerNames(Landscape.Get());
+
+			for (const FName& Layer : Layers)
+			{
+				Metadata->CreateFloatAttribute(Layer, 0.0f, /*bAllowInterpolation=*/true);
+			}
 		}
 	}
 	else
@@ -76,14 +80,18 @@ bool UPCGLandscapeData::SamplePoint(const FTransform& InTransform, const FBox& I
 			return false;
 		}
 
-		// Compute the 4 points indices
 		const FVector2D ComponentLocalPoint(LocalPoint.X - ComponentMapKey.X * LandscapeInfo->ComponentSizeQuads, LocalPoint.Y - ComponentMapKey.Y * LandscapeInfo->ComponentSizeQuads);
-		const int32 X0Y0 = FMath::FloorToInt(ComponentLocalPoint.X) + FMath::FloorToInt(ComponentLocalPoint.Y) * (LandscapeInfo->ComponentSizeQuads + 1);
 
-		const float FractionalX = FMath::Fractional(LocalPoint.X);
-		const float FractionalY = FMath::Fractional(LocalPoint.Y);
+		if (bHeightOnly)
+		{
+			LandscapeCacheEntry->GetInterpolatedPointHeightOnly(ComponentLocalPoint, OutPoint);
+		}
+		else
+		{
+			LandscapeCacheEntry->GetInterpolatedPoint(ComponentLocalPoint, OutPoint, OutMetadata);
+		}
 
-		return LandscapeCacheEntry->GetInterpolatedPoint(X0Y0, LandscapeInfo->ComponentSizeQuads + 1, FractionalX, FractionalY, OutPoint, OutMetadata);
+		return true;
 	}
 	else
 	{
@@ -173,8 +181,17 @@ const UPCGPointData* UPCGLandscapeData::CreatePointData(FPCGContext* Context, co
 				{
 					for (int32 LocalY = LocalMinY; LocalY <= LocalMaxY; ++LocalY)
 					{
+						const int32 PointIndex = LocalX + LocalY * (ComponentSizeQuads + 1);
+
 						FPCGPoint& Point = Points.Emplace_GetRef();
-						LandscapeCacheEntry->GetPoint(LocalX + LocalY * (ComponentSizeQuads + 1), Point, Data->Metadata);
+						if (bHeightOnly)
+						{
+							LandscapeCacheEntry->GetPointHeightOnly(PointIndex, Point);
+						}
+						else
+						{
+							LandscapeCacheEntry->GetPoint(PointIndex, Point, Data->Metadata);
+						}
 					}
 				}
 			}
diff --git a/Engine/Plugins/Experimental/PCG/Source/PCG/Private/Graph/PCGGraphExecutor.cpp b/Engine/Plugins/Experimental/PCG/Source/PCG/Private/Graph/PCGGraphExecutor.cpp
--- a/Engine/Plugins/Experimental/PCG/Source/PCG/Private/Graph/PCGGraphExecutor.cpp
+++ b/Engine/Plugins/Experimental/PCG/Source/PCG/Private/Graph/PCGGraphExecutor.cpp
@@ -639,6 +639,16 @@ bool FPCGFetchInputElement::ExecuteInternal(FPCGContext* Context) const
 		}
 	}
 
+	if (Context->Node->IsOutputPinConnected(PCGInputOutputConstants::DefaultLandscapeHeightLabel))
+	{
+		if (UPCGData* LandscapeHeightPCGData = Component->GetLandscapeHeightPCGData())
+		{
+			FPCGTaggedData& TaggedData = Context->OutputData.TaggedData.Emplace_GetRef();
+			TaggedData.Data = LandscapeHeightPCGData;
+			TaggedData.Pin = PCGInputOutputConstants::DefaultLandscapeHeightLabel;
+		}
+	}
+
 	if (Context->Node->IsOutputPinConnected(PCGInputOutputConstants::DefaultOriginalActorLabel))
 	{
 		if (UPCGData* OriginalActorPCGData = Component->GetOriginalActorPCGData())
diff --git a/Engine/Plugins/Experimental/PCG/Source/PCG/Private/Grid/PCGLandscapeCache.cpp b/Engine/Plugins/Experimental/PCG/Source/PCG/Private/Grid/PCGLandscapeCache.cpp
--- a/Engine/Plugins/Experimental/PCG/Source/PCG/Private/Grid/PCGLandscapeCache.cpp
+++ b/Engine/Plugins/Experimental/PCG/Source/PCG/Private/Grid/PCGLandscapeCache.cpp
@@ -48,6 +48,8 @@ void FPCGLandscapeCacheEntry::BuildCacheData(ULandscapeInfo* LandscapeInfo, ULan
 		const int32 ComponentSizeQuads = InComponent->ComponentSizeQuads + 1;
 		const int32 NumVertices = FMath::Square(ComponentSizeQuads);
 
+		Stride = ComponentSizeQuads;
+
 		PositionsAndNormals.Reserve(2 * NumVertices);
 		for (int32 Index = 0; Index < NumVertices; ++Index)
 		{
@@ -80,13 +82,10 @@ void FPCGLandscapeCacheEntry::BuildCacheData(ULandscapeInfo* LandscapeInfo, ULan
 }
 #endif // WITH_EDITOR
 
-bool FPCGLandscapeCacheEntry::GetPoint(int32 PointIndex, FPCGPoint& OutPoint, UPCGMetadata* OutMetadata) const
+void FPCGLandscapeCacheEntry::GetPoint(int32 PointIndex, FPCGPoint& OutPoint, UPCGMetadata* OutMetadata) const
 {
 	//TRACE_CPUPROFILER_EVENT_SCOPE(FPCGLandscapeCacheEntry::GetPoint);
-	if (PointIndex < 0 || 2 * PointIndex >= PositionsAndNormals.Num())
-	{
-		return false;
-	}
+	check(PointIndex >= 0 && 2 * PointIndex < PositionsAndNormals.Num());
 
 	const FVector& Position = PositionsAndNormals[2 * PointIndex];
 	const FVector& Normal = PositionsAndNormals[2 * PointIndex + 1];
@@ -99,7 +98,7 @@ bool FPCGLandscapeCacheEntry::GetPoint(int32 PointIndex, FPCGPoint& OutPoint, UP
 	OutPoint.Transform = FTransform(TangentX, TangentY, Normal, Position);
 	OutPoint.BoundsMin = -PointHalfSize;
 	OutPoint.BoundsMax = PointHalfSize;
-	OutPoint.Seed = UPCGBlueprintHelpers::ComputeSeedFromPosition(Position);
+	OutPoint.Seed = 1 + PointIndex;
 	
 	if (OutMetadata && !LayerData.IsEmpty())
 	{
@@ -114,22 +113,33 @@ bool FPCGLandscapeCacheEntry::GetPoint(int32 PointIndex, FPCGPoint& OutPoint, UP
 			}
 		}
 	}
-
-	return true;
 }
 
-bool FPCGLandscapeCacheEntry::GetInterpolatedPoint(int32 BasePointIndex, int32 PointLineStride, float XFactor, float YFactor, FPCGPoint& OutPoint, UPCGMetadata* OutMetadata) const
+void FPCGLandscapeCacheEntry::GetPointHeightOnly(int32 PointIndex, FPCGPoint& OutPoint) const
 {
-	if (BasePointIndex < 0 || 2 * BasePointIndex >= PositionsAndNormals.Num())
-	{
-		return false;
-	}
+	//TRACE_CPUPROFILER_EVENT_SCOPE(FPCGLandscapeCacheEntry::GetPointHeightOnly);
+	check(PointIndex >= 0 && 2 * PointIndex < PositionsAndNormals.Num());
+
+	const FVector& Position = PositionsAndNormals[2 * PointIndex];
+	
+	OutPoint.Transform = FTransform(Position);
+	OutPoint.BoundsMin = -PointHalfSize;
+	OutPoint.BoundsMax = PointHalfSize;
+	OutPoint.Seed = 1 + PointIndex;
+}
 
-	const int32 X0Y0 = BasePointIndex;
+void FPCGLandscapeCacheEntry::GetInterpolatedPoint(const FVector2D& LocalPoint, FPCGPoint& OutPoint, UPCGMetadata* OutMetadata) const
+{
+	const int32 X0Y0 = FMath::FloorToInt(LocalPoint.X) + FMath::FloorToInt(LocalPoint.Y) * Stride;
 	const int32 X1Y0 = X0Y0 + 1;
-	const int32 X0Y1 = X0Y0 + PointLineStride;
+	const int32 X0Y1 = X0Y0 + Stride;
 	const int32 X1Y1 = X0Y1 + 1;
 
+	const float XFactor = FMath::Fractional(LocalPoint.X);
+	const float YFactor = FMath::Fractional(LocalPoint.Y);
+
+	check(X0Y0 >= 0 && 2 * X1Y1 < PositionsAndNormals.Num());
+
 	const FVector& PositionX0Y0 = PositionsAndNormals[2 * X0Y0];
 	const FVector& NormalX0Y0 = PositionsAndNormals[2 * X0Y0 + 1];
 	const FVector& PositionX1Y0 = PositionsAndNormals[2 * X1Y0];
@@ -139,12 +149,12 @@ bool FPCGLandscapeCacheEntry::GetInterpolatedPoint(int32 BasePointIndex, int32 P
 	const FVector& PositionX1Y1 = PositionsAndNormals[2 * X1Y1];
 	const FVector& NormalX1Y1 = PositionsAndNormals[2 * X1Y1 + 1];
 
-	FVector LerpPositionY0 = FMath::Lerp(PositionX0Y0, PositionX1Y0, XFactor);
-	FVector LerpPositionY1 = FMath::Lerp(PositionX0Y1, PositionX1Y1, XFactor);
-	FVector Position = FMath::Lerp(LerpPositionY0, LerpPositionY1, YFactor);
+	const FVector LerpPositionY0 = FMath::Lerp(PositionX0Y0, PositionX1Y0, XFactor);
+	const FVector LerpPositionY1 = FMath::Lerp(PositionX0Y1, PositionX1Y1, XFactor);
+	const FVector Position = FMath::Lerp(LerpPositionY0, LerpPositionY1, YFactor);
 
-	FVector LerpNormalY0 = FMath::Lerp(NormalX0Y0, NormalX1Y0, XFactor);
-	FVector LerpNormalY1 = FMath::Lerp(NormalX0Y1, NormalX1Y1, XFactor);
+	const FVector LerpNormalY0 = FMath::Lerp(NormalX0Y0, NormalX1Y0, XFactor);
+	const FVector LerpNormalY1 = FMath::Lerp(NormalX0Y1, NormalX1Y1, XFactor);
 	FVector Normal = FMath::Lerp(LerpNormalY0, LerpNormalY1, YFactor);
 
 	// TODO: we could preserve normal length by lerping this too?
@@ -177,8 +187,33 @@ bool FPCGLandscapeCacheEntry::GetInterpolatedPoint(int32 BasePointIndex, int32 P
 			}
 		}
 	}
+}
+
+void FPCGLandscapeCacheEntry::GetInterpolatedPointHeightOnly(const FVector2D& LocalPoint, FPCGPoint& OutPoint) const
+{
+	const int32 X0Y0 = FMath::FloorToInt(LocalPoint.X) + FMath::FloorToInt(LocalPoint.Y) * Stride;
+	const int32 X1Y0 = X0Y0 + 1;
+	const int32 X0Y1 = X0Y0 + Stride;
+	const int32 X1Y1 = X0Y1 + 1;
+
+	const float XFactor = FMath::Fractional(LocalPoint.X);
+	const float YFactor = FMath::Fractional(LocalPoint.Y);
+
+	check(X0Y0 >= 0 && 2 * X1Y1 < PositionsAndNormals.Num());
+
+	const FVector& PositionX0Y0 = PositionsAndNormals[2 * X0Y0];
+	const FVector& PositionX1Y0 = PositionsAndNormals[2 * X1Y0];
+	const FVector& PositionX0Y1 = PositionsAndNormals[2 * X0Y1];
+	const FVector& PositionX1Y1 = PositionsAndNormals[2 * X1Y1];
 
-	return true;
+	const FVector LerpPositionY0 = FMath::Lerp(PositionX0Y0, PositionX1Y0, XFactor);
+	const FVector LerpPositionY1 = FMath::Lerp(PositionX0Y1, PositionX1Y1, XFactor);
+	const FVector Position = FMath::Lerp(LerpPositionY0, LerpPositionY1, YFactor);
+
+	OutPoint.Transform = FTransform(Position);
+	OutPoint.BoundsMin = -PointHalfSize;
+	OutPoint.BoundsMax = PointHalfSize;
+	OutPoint.Seed = UPCGBlueprintHelpers::ComputeSeedFromPosition(Position);
 }
 
 FPCGLandscapeCache::FPCGLandscapeCache(UObject* InOwner)
diff --git a/Engine/Plugins/Experimental/PCG/Source/PCG/Private/PCGComponent.cpp b/Engine/Plugins/Experimental/PCG/Source/PCG/Private/PCGComponent.cpp
--- a/Engine/Plugins/Experimental/PCG/Source/PCG/Private/PCGComponent.cpp
+++ b/Engine/Plugins/Experimental/PCG/Source/PCG/Private/PCGComponent.cpp
@@ -1204,6 +1204,7 @@ void UPCGComponent::DirtyGenerated(EPCGComponentDirtyFlag DirtyFlag)
 		if (Cast<ALandscapeProxy>(GetOwner()))
 		{
 			CachedLandscapeData = nullptr;
+			CachedLandscapeHeightData = nullptr;
 		}
 
 		CachedInputData = nullptr;
@@ -1213,6 +1214,7 @@ void UPCGComponent::DirtyGenerated(EPCGComponentDirtyFlag DirtyFlag)
 	if (!!(DirtyFlag & EPCGComponentDirtyFlag::Landscape))
 	{
 		CachedLandscapeData = nullptr;
+		CachedLandscapeHeightData = nullptr;
 		if (InputType == EPCGComponentInput::Landscape)
 		{
 			CachedInputData = nullptr;
@@ -1364,12 +1366,22 @@ UPCGData* UPCGComponent::GetLandscapePCGData()
 {
 	if (!CachedLandscapeData)
 	{
-		CachedLandscapeData = CreateLandscapePCGData();
+		CachedLandscapeData = CreateLandscapePCGData(/*bHeightOnly=*/false);
 	}
 
 	return CachedLandscapeData;
 }
 
+UPCGData* UPCGComponent::GetLandscapeHeightPCGData()
+{
+	if (!CachedLandscapeHeightData)
+	{
+		CachedLandscapeHeightData = CreateLandscapePCGData(/*bHeightOnly=*/true);
+	}
+
+	return CachedLandscapeHeightData;
+}
+
 UPCGData* UPCGComponent::GetOriginalActorPCGData()
 {
 	if (APCGPartitionActor* PartitionActor = Cast<APCGPartitionActor>(GetOwner()))
@@ -1490,7 +1502,7 @@ UPCGData* UPCGComponent::CreateActorPCGData(AActor* Actor)
 	else if (ALandscapeProxy* Landscape = Cast<ALandscapeProxy>(Actor))
 	{
 		UPCGLandscapeData* Data = NewObject<UPCGLandscapeData>(this);
-		Data->Initialize(Landscape, GetGridBounds(Actor));
+		Data->Initialize(Landscape, GetGridBounds(Actor), /*bHeightOnly=*/false);
 
 		return Data;
 	}
@@ -1632,7 +1644,7 @@ UPCGData* UPCGComponent::CreatePCGData()
 	return Difference ? Difference : InputData;
 }
 
-UPCGData* UPCGComponent::CreateLandscapePCGData()
+UPCGData* UPCGComponent::CreateLandscapePCGData(bool bHeightOnly)
 {
 	TRACE_CPUPROFILER_EVENT_SCOPE(UPCGComponent::CreateLandscapePCGData);
 	AActor* Actor = GetOwner();
@@ -1674,7 +1686,7 @@ UPCGData* UPCGComponent::CreateLandscapePCGData()
 
 	// TODO: we're creating separate landscape data instances here so we can do some tweaks on it (such as storing the right target actor) but this probably should change
 	UPCGLandscapeData* LandscapeData = NewObject<UPCGLandscapeData>(this);
-	LandscapeData->Initialize(Landscape, GetGridBounds(Landscape));
+	LandscapeData->Initialize(Landscape, GetGridBounds(Landscape), bHeightOnly);
 	// Need to override target actor for this one, not the landscape
 	LandscapeData->TargetActor = Actor;
 
diff --git a/Engine/Plugins/Experimental/PCG/Source/PCG/Private/PCGInputOutputSettings.cpp b/Engine/Plugins/Experimental/PCG/Source/PCG/Private/PCGInputOutputSettings.cpp
--- a/Engine/Plugins/Experimental/PCG/Source/PCG/Private/PCGInputOutputSettings.cpp
+++ b/Engine/Plugins/Experimental/PCG/Source/PCG/Private/PCGInputOutputSettings.cpp
@@ -20,6 +20,7 @@ UPCGGraphInputOutputSettings::UPCGGraphInputOutputSettings(const FObjectInitiali
 	StaticAdvancedInLabels.Add(PCGInputOutputConstants::DefaultActorLabel);
 	StaticAdvancedInLabels.Add(PCGInputOutputConstants::DefaultOriginalActorLabel);
 	StaticAdvancedInLabels.Add(PCGInputOutputConstants::DefaultLandscapeLabel);
+	StaticAdvancedInLabels.Add(PCGInputOutputConstants::DefaultLandscapeHeightLabel);
 	StaticAdvancedInLabels.Add(PCGInputOutputConstants::DefaultExcludedActorsLabel);
 	
 	StaticOutLabels.Add(PCGPinConstants::DefaultOutputLabel);
diff --git a/Engine/Plugins/Experimental/PCG/Source/PCG/Public/Data/PCGLandscapeData.h b/Engine/Plugins/Experimental/PCG/Source/PCG/Public/Data/PCGLandscapeData.h
--- a/Engine/Plugins/Experimental/PCG/Source/PCG/Public/Data/PCGLandscapeData.h
+++ b/Engine/Plugins/Experimental/PCG/Source/PCG/Public/Data/PCGLandscapeData.h
@@ -17,7 +17,7 @@ class PCG_API UPCGLandscapeData : public UPCGSurfaceData
 	GENERATED_BODY()
 
 public:
-	void Initialize(ALandscapeProxy* InLandscape, const FBox& InBounds);
+	void Initialize(ALandscapeProxy* InLandscape, const FBox& InBounds, bool bInHeightOnly);
 
 	// ~Begin UPCGData interface
 	virtual EPCGDataType GetDataType() const override { return EPCGDataType::Landscape | Super::GetDataType(); }
@@ -44,6 +44,9 @@ protected:
 	UPROPERTY()
 	FBox Bounds = FBox(EForceInit::ForceInit);
 
+	UPROPERTY()
+	bool bHeightOnly = false;
+
 private:
 	// Transient data
 	ULandscapeInfo* LandscapeInfo = nullptr;
diff --git a/Engine/Plugins/Experimental/PCG/Source/PCG/Public/Grid/PCGLandscapeCache.h b/Engine/Plugins/Experimental/PCG/Source/PCG/Public/Grid/PCGLandscapeCache.h
--- a/Engine/Plugins/Experimental/PCG/Source/PCG/Public/Grid/PCGLandscapeCache.h
+++ b/Engine/Plugins/Experimental/PCG/Source/PCG/Public/Grid/PCGLandscapeCache.h
@@ -33,8 +33,10 @@ public:
 	void BuildCacheData(ULandscapeInfo* LandscapeInfo, ULandscapeComponent* InComponent, UObject* Owner);
 #endif
 
-	bool GetPoint(int32 PointIndex, FPCGPoint& OutPoint, UPCGMetadata* OutMetadata) const;
-	bool GetInterpolatedPoint(int32 PointIndex, int32 PointLineStride, float XFactor, float YFactor, FPCGPoint& OutPoint, UPCGMetadata* OutMetadata) const;
+	void GetPoint(int32 PointIndex, FPCGPoint& OutPoint, UPCGMetadata* OutMetadata) const;
+	void GetPointHeightOnly(int32 PointIndex, FPCGPoint& OutPoint) const;
+	void GetInterpolatedPoint(const FVector2D& LocalPoint, FPCGPoint& OutPoint, UPCGMetadata* OutMetadata) const;
+	void GetInterpolatedPointHeightOnly(const FVector2D& LocalPoint, FPCGPoint& OutPoint) const;
 
 	// TODO: this isn't really needed
 	UPROPERTY()
@@ -47,7 +49,10 @@ public:
 	TArray<FPCGLandscapeCacheLayer> LayerData;
 
 	UPROPERTY()
-	FVector PointHalfSize;
+	FVector PointHalfSize = FVector::One();
+
+	UPROPERTY()
+	int32 Stride = 0;
 };
 
 USTRUCT()
diff --git a/Engine/Plugins/Experimental/PCG/Source/PCG/Public/PCGComponent.h b/Engine/Plugins/Experimental/PCG/Source/PCG/Public/PCGComponent.h
--- a/Engine/Plugins/Experimental/PCG/Source/PCG/Public/PCGComponent.h
+++ b/Engine/Plugins/Experimental/PCG/Source/PCG/Public/PCGComponent.h
@@ -78,6 +78,7 @@ public:
 	UPCGData* GetInputPCGData();
 	UPCGData* GetActorPCGData();
 	UPCGData* GetLandscapePCGData();
+	UPCGData* GetLandscapeHeightPCGData();
 	UPCGData* GetOriginalActorPCGData();
 	TArray<UPCGData*> GetPCGExclusionData();
 
@@ -196,7 +197,7 @@ private:
 	UPCGData* CreateInputPCGData();
 	UPCGData* CreateActorPCGData();
 	UPCGData* CreateActorPCGData(AActor* Actor);
-	UPCGData* CreateLandscapePCGData();
+	UPCGData* CreateLandscapePCGData(bool bHeightOnly);
 	void UpdatePCGExclusionData();
 
 	bool ShouldGenerate(bool bForce, EPCGComponentGenerationTrigger RequestedGenerationTrigger) const;
@@ -262,6 +263,9 @@ private:
 	UPROPERTY(Transient, NonPIEDuplicateTransient)
 	UPCGData* CachedLandscapeData = nullptr;
 
+	UPROPERTY(Transient, NonPIEDuplicateTransient)
+	UPCGData* CachedLandscapeHeightData = nullptr;
+
 	UPROPERTY(Transient, NonPIEDuplicateTransient)
 	TMap<AActor*, UPCGData*> CachedExclusionData;
 
diff --git a/Engine/Plugins/Experimental/PCG/Source/PCG/Public/PCGInputOutputSettings.h b/Engine/Plugins/Experimental/PCG/Source/PCG/Public/PCGInputOutputSettings.h
--- a/Engine/Plugins/Experimental/PCG/Source/PCG/Public/PCGInputOutputSettings.h
+++ b/Engine/Plugins/Experimental/PCG/Source/PCG/Public/PCGInputOutputSettings.h
@@ -13,6 +13,7 @@ namespace PCGInputOutputConstants
 	const FName DefaultActorLabel = TEXT("Actor");
 	const FName DefaultOriginalActorLabel = TEXT("OriginalActor");
 	const FName DefaultLandscapeLabel = TEXT("Landscape");
+	const FName DefaultLandscapeHeightLabel = TEXT("Landscape Height");
 	const FName DefaultExcludedActorsLabel = TEXT("ExcludedActors");
 }
 
