Fixed caching issue with PCG BP elements

#rnx

diff --git a/Engine/Plugins/Experimental/PCG/Source/PCG/Private/Elements/PCGExecuteBlueprint.cpp b/Engine/Plugins/Experimental/PCG/Source/PCG/Private/Elements/PCGExecuteBlueprint.cpp
--- a/Engine/Plugins/Experimental/PCG/Source/PCG/Private/Elements/PCGExecuteBlueprint.cpp
+++ b/Engine/Plugins/Experimental/PCG/Source/PCG/Private/Elements/PCGExecuteBlueprint.cpp
@@ -639,7 +639,7 @@ bool FPCGExecuteBlueprintElement::IsCacheable(const UPCGSettings* InSettings) co
 	const UPCGBlueprintSettings* BPSettings = Cast<const UPCGBlueprintSettings>(InSettings);
 	if (BPSettings && BPSettings->BlueprintElementInstance)
 	{
-		return BPSettings->BlueprintElementInstance->bCreatesArtifacts;
+		return !BPSettings->BlueprintElementInstance->bCreatesArtifacts;
 	}
 	else
 	{
