Add support to move point metadata to ISM custom data

diff --git a/Engine/Plugins/Experimental/PCG/Source/PCG/Private/Elements/PCGStaticMeshSpawner.cpp b/Engine/Plugins/Experimental/PCG/Source/PCG/Private/Elements/PCGStaticMeshSpawner.cpp
--- a/Engine/Plugins/Experimental/PCG/Source/PCG/Private/Elements/PCGStaticMeshSpawner.cpp
+++ b/Engine/Plugins/Experimental/PCG/Source/PCG/Private/Elements/PCGStaticMeshSpawner.cpp
@@ -7,6 +7,8 @@
 #include "Data/PCGPointData.h"
 #include "Data/PCGSpatialData.h"
 #include "Helpers/PCGActorHelpers.h"
+
+#include "InstancePackers/PCGInstancePackerBase.h"
 #include "MeshSelectors/PCGMeshSelectorBase.h"
 #include "MeshSelectors/PCGMeshSelectorWeighted.h"
 
@@ -41,6 +43,8 @@ bool FPCGStaticMeshSpawnerElement::ExecuteInternal(FPCGContext* Context) const
 	TArray<FPCGMeshInstanceList> MeshInstances;
 	TArray<FPCGTaggedData> Inputs = Context->InputData.GetInputs();
 	TArray<FPCGTaggedData>& Outputs = Context->OutputData.TaggedData;
+	
+	FPCGPackedCustomData PackedCustomData;
 
 	// Forward any non-input data
 	Outputs.Append(Context->InputData.GetAllSettings());
@@ -80,15 +84,25 @@ bool FPCGStaticMeshSpawnerElement::ExecuteInternal(FPCGContext* Context) const
 				PCGE_LOG(Verbose, "Metadata attribute %s is being overwritten in the output data", *Settings->OutAttributeName.ToString());
 			}
 
-			OutputPointData->Metadata->CreateStringAttribute(Settings->OutAttributeName, FName(NAME_None).ToString(), false);
+			OutputPointData->Metadata->CreateStringAttribute(Settings->OutAttributeName, FName(NAME_None).ToString(), /*bAllowsInterpolation=*/false);
 
 			Output.Data = OutputPointData;
 		}
 
 		Settings->MeshSelectorInstance->SelectInstances(*Context, Settings, SpatialData, MeshInstances, OutputPointData);
 
-		// Spawn a static mesh for each instance
-		SpawnStaticMeshInstances(Context, MeshInstances, TargetActor);
+		for (const FPCGMeshInstanceList& InstanceList : MeshInstances)
+		{
+			if (Settings->InstancePackerInstance)
+			{
+				PackedCustomData.CustomData.Reset();
+				PackedCustomData.NumCustomDataFloats = 0;
+
+				Settings->InstancePackerInstance->PackInstances(*Context, SpatialData, InstanceList, PackedCustomData);
+			}
+
+			SpawnStaticMeshInstances(Context, InstanceList, TargetActor, PackedCustomData);
+		}
 
 		MeshInstances.Reset();
 	}
@@ -96,49 +110,74 @@ bool FPCGStaticMeshSpawnerElement::ExecuteInternal(FPCGContext* Context) const
 	return true;
 }
 
-void FPCGStaticMeshSpawnerElement::SpawnStaticMeshInstances(FPCGContext* Context, const TArray<FPCGMeshInstanceList>& MeshInstances, AActor* TargetActor) const
+void FPCGStaticMeshSpawnerElement::SpawnStaticMeshInstances(FPCGContext* Context, const FPCGMeshInstanceList& InstanceList, AActor* TargetActor, const FPCGPackedCustomData& PackedCustomData) const
 {
 	// Populate the (H)ISM from the previously prepared entries
 	TRACE_CPUPROFILER_EVENT_SCOPE(FPCGStaticMeshSpawnerElement::Execute::PopulateISMs);
 
-	for (const FPCGMeshInstanceList& InstanceList : MeshInstances)
+	if (InstanceList.Instances.Num() == 0)
 	{
-		if (InstanceList.Instances.Num() == 0)
-		{
-			continue;
-		}
+		return;
+	}
 
-		// Todo: we could likely pre-load these meshes asynchronously in the settings
-		UStaticMesh* LoadedMesh = InstanceList.Mesh.LoadSynchronous();
+	// Todo: we could likely pre-load these meshes asynchronously in the settings
+	UStaticMesh* LoadedMesh = InstanceList.Mesh.LoadSynchronous();
 
-		if (!LoadedMesh)
-		{
-			continue;
-		}
+	if (!LoadedMesh)
+	{
+		return;
+	}
 
-		FPCGISMCBuilderParameters Params;
-		Params.Mesh = LoadedMesh;
+	FPCGISMCBuilderParameters Params;
+	Params.Mesh = LoadedMesh;
 
-		if (InstanceList.bOverrideCollisionProfile)
-		{
-			Params.CollisionProfile = InstanceList.CollisionProfile.Name;
-		}
+	if (InstanceList.bOverrideCollisionProfile)
+	{
+		Params.CollisionProfile = InstanceList.CollisionProfile.Name;
+	}
 
-		if (InstanceList.bOverrideMaterials)
-		{
-			Params.MaterialOverrides = InstanceList.MaterialOverrides;
-		}
+	if (InstanceList.bOverrideMaterials)
+	{
+		Params.MaterialOverrides = InstanceList.MaterialOverrides;
+	}
+
+	Params.NumCustomDataFloats = PackedCustomData.NumCustomDataFloats;
+
+	UInstancedStaticMeshComponent* ISMC = UPCGActorHelpers::GetOrCreateISMC(TargetActor, Context->SourceComponent, Params);
+
+	const int32 PreExistingInstanceCount = ISMC->GetInstanceCount();
+	const int32 NewInstanceCount = InstanceList.Instances.Num();
+	const int32 NumCustomDataFloats = PackedCustomData.NumCustomDataFloats;
 
-		UInstancedStaticMeshComponent* ISMC = UPCGActorHelpers::GetOrCreateISMC(TargetActor, Context->SourceComponent, Params);
+	check((ISMC->NumCustomDataFloats == 0 && PreExistingInstanceCount == 0) || ISMC->NumCustomDataFloats == NumCustomDataFloats);
+	ISMC->NumCustomDataFloats = NumCustomDataFloats;
 
-		// TODO: add scaling
-		// TODO: document these arguments
-		ISMC->NumCustomDataFloats = 0;
-		ISMC->AddInstances(InstanceList.Instances, false, true);
-		ISMC->UpdateBounds();
+	// The index in ISMC PerInstanceSMCustomData where we should pick up to begin inserting new floats
+	const int32 PreviousCustomDataOffset = PreExistingInstanceCount * NumCustomDataFloats;
 
-		PCGE_LOG(Verbose, "Added %d instances of %s on actor %s", InstanceList.Instances.Num(), *InstanceList.Mesh->GetFName().ToString(), *TargetActor->GetFName().ToString());
+	// Populate the ISM instances
+	TArray<FTransform> Instances;
+	Instances.Reserve(NewInstanceCount);
+	for (int32 InstanceIndex = 0; InstanceIndex < NewInstanceCount; ++InstanceIndex)
+	{
+		Instances.Emplace(InstanceList.Instances[InstanceIndex].Transform);
 	}
+
+	ISMC->AddInstances(Instances, /*bShouldReturnIndices=*/false, /*bWorldSpace=*/true);
+
+	// Copy new CustomData into the ISMC PerInstanceSMCustomData
+	if (NumCustomDataFloats > 0)
+	{
+		check(PreviousCustomDataOffset + PackedCustomData.CustomData.Num() == ISMC->PerInstanceSMCustomData.Num());
+		FMemory::Memcpy(&ISMC->PerInstanceSMCustomData[PreviousCustomDataOffset], &PackedCustomData.CustomData[0], PackedCustomData.CustomData.Num() * sizeof(float));
+
+		// Force recreation of the render data when proxy is created
+		ISMC->InstanceUpdateCmdBuffer.NumEdits++;
+	}
+
+	ISMC->UpdateBounds();
+
+	PCGE_LOG(Verbose, "Added %d instances of %s on actor %s", InstanceList.Instances.Num(), *InstanceList.Mesh->GetFName().ToString(), *TargetActor->GetFName().ToString());
 }
 
 void UPCGStaticMeshSpawnerSettings::PostLoad()
@@ -167,14 +206,28 @@ void UPCGStaticMeshSpawnerSettings::PostLoad()
 	{
 		RefreshMeshSelector();
 	}
+
+	if (!InstancePackerInstance)
+	{
+		RefreshInstancePacker();
+	}
 }
 
 #if WITH_EDITOR
 void UPCGStaticMeshSpawnerSettings::PostEditChangeProperty(FPropertyChangedEvent& PropertyChangedEvent) 
 {
-	if (PropertyChangedEvent.Property && PropertyChangedEvent.Property->GetFName() == GET_MEMBER_NAME_CHECKED(UPCGStaticMeshSpawnerSettings, MeshSelectorType))
+	if (PropertyChangedEvent.Property)
 	{
-		RefreshMeshSelector();
+		const FName& PropertyName = PropertyChangedEvent.Property->GetFName();
+
+		if (PropertyName == GET_MEMBER_NAME_CHECKED(UPCGStaticMeshSpawnerSettings, MeshSelectorType))
+		{
+			RefreshMeshSelector();
+		} 
+		else if (PropertyName == GET_MEMBER_NAME_CHECKED(UPCGStaticMeshSpawnerSettings, InstancePackerType))
+		{
+			RefreshInstancePacker();
+		}
 	}
 
 	Super::PostEditChangeProperty(PropertyChangedEvent);
@@ -194,6 +247,19 @@ void UPCGStaticMeshSpawnerSettings::SetMeshSelectorType(TSubclassOf<UPCGMeshSele
 	}
 }
 
+void UPCGStaticMeshSpawnerSettings::SetInstancePackerType(TSubclassOf<UPCGInstancePackerBase> InInstancePackerType) 
+{
+	if (!InstancePackerInstance || InInstancePackerType != InstancePackerType)
+	{
+		if (InInstancePackerType != InstancePackerType)
+		{
+			InstancePackerType = InInstancePackerType;
+		}
+		
+		RefreshInstancePacker();
+	}
+}
+
 void UPCGStaticMeshSpawnerSettings::RefreshMeshSelector()
 {
 	if (MeshSelectorType)
@@ -205,3 +271,15 @@ void UPCGStaticMeshSpawnerSettings::RefreshMeshSelector()
 		MeshSelectorInstance = nullptr;
 	}
 }
+
+void UPCGStaticMeshSpawnerSettings::RefreshInstancePacker()
+{
+	if (InstancePackerType)
+	{
+		InstancePackerInstance = NewObject<UPCGInstancePackerBase>(this, InstancePackerType);
+	}
+	else
+	{
+		InstancePackerInstance = nullptr;
+	}
+}
diff --git a/Engine/Plugins/Experimental/PCG/Source/PCG/Private/Helpers/PCGActorHelpers.cpp b/Engine/Plugins/Experimental/PCG/Source/PCG/Private/Helpers/PCGActorHelpers.cpp
--- a/Engine/Plugins/Experimental/PCG/Source/PCG/Private/Helpers/PCGActorHelpers.cpp
+++ b/Engine/Plugins/Experimental/PCG/Source/PCG/Private/Helpers/PCGActorHelpers.cpp
@@ -29,6 +29,7 @@ UInstancedStaticMeshComponent* UPCGActorHelpers::GetOrCreateISMC(AActor* InTarge
 {
 	UStaticMesh* InMesh = InParams.Mesh;
 	const TArray<UMaterialInterface*>& InMaterials = InParams.MaterialOverrides;
+	const int32 InNumCustomDataFloats = InParams.NumCustomDataFloats;
 
 	check(InTargetActor != nullptr && InMesh != nullptr);
 	check(InSourceComponent);
@@ -53,7 +54,8 @@ UInstancedStaticMeshComponent* UPCGActorHelpers::GetOrCreateISMC(AActor* InTarge
 
 		if (ISMC && 
 			ISMC->GetStaticMesh() == InMesh &&
-			ISMC->ComponentTags.Contains(InSourceComponent->GetFName()))
+			ISMC->ComponentTags.Contains(InSourceComponent->GetFName()) &&
+			ISMC->NumCustomDataFloats == InNumCustomDataFloats)
 		{
 			// If materials are provided, we'll make sure they match to the already set materials.
 			// If not provided, we'll make sure that the current materials aren't overriden
diff --git a/Engine/Plugins/Experimental/PCG/Source/PCG/Private/InstancePackers/PCGInstancePackerBase.cpp b/Engine/Plugins/Experimental/PCG/Source/PCG/Private/InstancePackers/PCGInstancePackerBase.cpp
new file mode 100644
--- /dev/null
+++ b/Engine/Plugins/Experimental/PCG/Source/PCG/Private/InstancePackers/PCGInstancePackerBase.cpp
@@ -0,0 +1,144 @@
+// Copyright Epic Games, Inc. All Rights Reserved.
+
+#include "InstancePackers/PCGInstancePackerBase.h"
+
+#include "Metadata/PCGMetadata.h"
+#include "Metadata/PCGMetadataAttribute.h"
+#include "Metadata/PCGMetadataAttributeTraits.h"
+#include "Metadata/PCGMetadataAttributeTpl.h"
+
+void UPCGInstancePackerBase::PackInstances_Implementation(FPCGContext& Context, const UPCGSpatialData* InSpatialData, const FPCGMeshInstanceList& InstanceList, FPCGPackedCustomData& OutPackedCustomData) const
+{
+	PCGE_LOG_C(Error, &Context, "Invalid use of abstract InstancePackerBase class. Please use an already existing class or implement the PackInstances method");
+}
+
+bool UPCGInstancePackerBase::AddTypeToPacking(int TypeId, FPCGPackedCustomData& OutPackedCustomData) const
+{
+	if (TypeId == PCG::Private::MetadataTypes<float>::Id)
+	{
+		OutPackedCustomData.NumCustomDataFloats += 1;
+	}
+	else if (TypeId == PCG::Private::MetadataTypes<double>::Id)
+	{
+		OutPackedCustomData.NumCustomDataFloats += 1;
+	}
+	else if (TypeId == PCG::Private::MetadataTypes<int32>::Id)
+	{
+		OutPackedCustomData.NumCustomDataFloats += 1;
+	}
+	else if (TypeId == PCG::Private::MetadataTypes<int64>::Id)
+	{
+		OutPackedCustomData.NumCustomDataFloats += 1;
+	}
+	else if (TypeId == PCG::Private::MetadataTypes<FRotator>::Id)
+	{
+		OutPackedCustomData.NumCustomDataFloats += 3;
+	}
+	else if (TypeId == PCG::Private::MetadataTypes<FVector>::Id)
+	{
+		OutPackedCustomData.NumCustomDataFloats += 3;
+	}
+	else if (TypeId == PCG::Private::MetadataTypes<FVector4>::Id)
+	{
+		OutPackedCustomData.NumCustomDataFloats += 4;
+	}
+	else
+	{
+		return false;
+	}
+
+	return true;
+}
+
+void UPCGInstancePackerBase::PackCustomDataFromAttributes(const FPCGMeshInstanceList& InstanceList, const UPCGMetadata* Metadata, const TArray<FName>& AttributeNames, FPCGPackedCustomData& OutPackedCustomData) const
+{
+	if (!Metadata)
+	{
+		UE_LOG(LogPCG, Error, TEXT("Invalid metadata"));
+		return;
+	}
+
+	TArray<const FPCGMetadataAttributeBase*> Attributes;
+	for (const FName& AttributeName : AttributeNames)
+	{
+		Attributes.Add(Metadata->GetConstAttribute(AttributeName));
+	}
+
+	PackCustomDataFromAttributes(InstanceList, Attributes, OutPackedCustomData);
+}
+
+void UPCGInstancePackerBase::PackCustomDataFromAttributes(const FPCGMeshInstanceList& InstanceList, const TArray<const FPCGMetadataAttributeBase*>& Attributes, FPCGPackedCustomData& OutPackedCustomData) const
+{
+	for (const FPCGPoint& Point : InstanceList.Instances)
+	{
+		for (const FPCGMetadataAttributeBase* AttributeBase : Attributes)
+		{
+			check(AttributeBase);
+
+			if (AttributeBase->GetTypeId() == PCG::Private::MetadataTypes<float>::Id)
+			{
+				const FPCGMetadataAttribute<float>* Attribute = static_cast<const FPCGMetadataAttribute<float>*>(AttributeBase);
+				check(Attribute);
+
+				const float Value = Attribute->GetValueFromItemKey(Point.MetadataEntry);
+				OutPackedCustomData.CustomData.Add(Value);
+			}
+			else if (AttributeBase->GetTypeId() == PCG::Private::MetadataTypes<double>::Id)
+			{
+				const FPCGMetadataAttribute<double>* Attribute = static_cast<const FPCGMetadataAttribute<double>*>(AttributeBase);
+				check(Attribute);
+
+				const double Value = Attribute->GetValueFromItemKey(Point.MetadataEntry);
+				OutPackedCustomData.CustomData.Add(Value);
+			}
+			else if (AttributeBase->GetTypeId() == PCG::Private::MetadataTypes<int32>::Id)
+			{
+				const FPCGMetadataAttribute<int32>* Attribute = static_cast<const FPCGMetadataAttribute<int32>*>(AttributeBase);
+				check(Attribute);
+
+				const float Value = static_cast<float>(Attribute->GetValueFromItemKey(Point.MetadataEntry));
+				OutPackedCustomData.CustomData.Add(Value);
+			}
+			else if (AttributeBase->GetTypeId() == PCG::Private::MetadataTypes<int64>::Id)
+			{
+				const FPCGMetadataAttribute<int64>* Attribute = static_cast<const FPCGMetadataAttribute<int64>*>(AttributeBase);
+				check(Attribute);
+
+				const float Value = static_cast<float>(Attribute->GetValueFromItemKey(Point.MetadataEntry));
+				OutPackedCustomData.CustomData.Add(Value);
+			}
+			else if (AttributeBase->GetTypeId() == PCG::Private::MetadataTypes<FRotator>::Id)
+			{
+				const FPCGMetadataAttribute<FRotator>* Attribute = static_cast<const FPCGMetadataAttribute<FRotator>*>(AttributeBase);
+				check(Attribute);
+
+				const FRotator Value = Attribute->GetValueFromItemKey(Point.MetadataEntry);
+				OutPackedCustomData.CustomData.Add(Value.Roll);
+				OutPackedCustomData.CustomData.Add(Value.Pitch);
+				OutPackedCustomData.CustomData.Add(Value.Yaw);
+			}
+			else if (AttributeBase->GetTypeId() == PCG::Private::MetadataTypes<FVector>::Id)
+			{
+				const FPCGMetadataAttribute<FVector>* Attribute = static_cast<const FPCGMetadataAttribute<FVector>*>(AttributeBase);
+				check(Attribute);
+
+				const FVector Value = Attribute->GetValueFromItemKey(Point.MetadataEntry);
+				OutPackedCustomData.CustomData.Add(Value.X);
+				OutPackedCustomData.CustomData.Add(Value.Y);
+				OutPackedCustomData.CustomData.Add(Value.Z);
+			}
+			else if (AttributeBase->GetTypeId() == PCG::Private::MetadataTypes<FVector4>::Id)
+			{
+				const FPCGMetadataAttribute<FVector4>* Attribute = static_cast<const FPCGMetadataAttribute<FVector4>*>(AttributeBase);
+				check(Attribute);
+
+				const FVector4 Value = Attribute->GetValueFromItemKey(Point.MetadataEntry);
+				OutPackedCustomData.CustomData.Add(Value.X);
+				OutPackedCustomData.CustomData.Add(Value.Y);
+				OutPackedCustomData.CustomData.Add(Value.Z);
+				OutPackedCustomData.CustomData.Add(Value.W);
+			}
+		}
+	}
+}
+
diff --git a/Engine/Plugins/Experimental/PCG/Source/PCG/Private/InstancePackers/PCGInstancePackerByAttribute.cpp b/Engine/Plugins/Experimental/PCG/Source/PCG/Private/InstancePackers/PCGInstancePackerByAttribute.cpp
new file mode 100644
--- /dev/null
+++ b/Engine/Plugins/Experimental/PCG/Source/PCG/Private/InstancePackers/PCGInstancePackerByAttribute.cpp
@@ -0,0 +1,38 @@
+// Copyright Epic Games, Inc. All Rights Reserved.
+
+#include "InstancePackers/PCGInstancePackerByAttribute.h"
+#include "Metadata/PCGMetadataAttribute.h"
+
+void UPCGInstancePackerByAttribute::PackInstances_Implementation(FPCGContext& Context, const UPCGSpatialData* InSpatialData, const FPCGMeshInstanceList& InstanceList, FPCGPackedCustomData& OutPackedCustomData) const
+{
+	if (!InSpatialData || !InSpatialData->Metadata)
+	{
+		PCGE_LOG_C(Error, &Context, "Invalid input data");
+		return;
+	}
+
+	TArray<const FPCGMetadataAttributeBase*> SelectedAttributes;
+
+	// Find Attributes by name and calculate NumCustomDataFloats
+	for (const FName& AttributeName : AttributeNames)
+	{
+		if (!InSpatialData->Metadata->HasAttribute(AttributeName)) 
+		{
+			PCGE_LOG_C(Warning, &Context, "Attribute %s is not in the metadata", *AttributeName.ToString());
+			continue;
+		}
+
+		const FPCGMetadataAttributeBase* AttributeBase = InSpatialData->Metadata->GetConstAttribute(AttributeName);
+		check(AttributeBase);
+
+		if (!AddTypeToPacking(AttributeBase->GetTypeId(), OutPackedCustomData))
+		{
+			PCGE_LOG_C(Warning, &Context, "Attribute name %s is not a valid type", *AttributeName.ToString());
+			continue;
+		}
+
+		SelectedAttributes.Add(AttributeBase);
+	}
+
+	PackCustomDataFromAttributes(InstanceList, SelectedAttributes, OutPackedCustomData);
+ }
diff --git a/Engine/Plugins/Experimental/PCG/Source/PCG/Private/InstancePackers/PCGInstancePackerByRegex.cpp b/Engine/Plugins/Experimental/PCG/Source/PCG/Private/InstancePackers/PCGInstancePackerByRegex.cpp
new file mode 100644
--- /dev/null
+++ b/Engine/Plugins/Experimental/PCG/Source/PCG/Private/InstancePackers/PCGInstancePackerByRegex.cpp
@@ -0,0 +1,59 @@
+// Copyright Epic Games, Inc. All Rights Reserved.
+
+#include "InstancePackers/PCGInstancePackerByRegex.h"
+#include "Metadata/PCGMetadataAttribute.h"
+#include "Internationalization/Regex.h"
+
+void UPCGInstancePackerByRegex::PackInstances_Implementation(FPCGContext& Context, const UPCGSpatialData* InSpatialData, const FPCGMeshInstanceList& InstanceList, FPCGPackedCustomData& OutPackedCustomData) const
+{
+	if (!InSpatialData || !InSpatialData->Metadata)
+	{
+		PCGE_LOG_C(Error, &Context, "Invalid input data");
+		return;
+	}
+	
+	TArray<FName> AttributeNames;
+	TArray<EPCGMetadataTypes> AttributeTypes;
+	InSpatialData->Metadata->GetAttributes(AttributeNames, AttributeTypes);
+
+	TArray<const FPCGMetadataAttributeBase*> SelectedAttributes;
+	TSet<PCGMetadataAttributeKey> FoundAttributes;
+
+	for (const FString& Regex : RegexPatterns)
+	{
+		const FRegexPattern Pattern(Regex);
+	
+		// Match our Regex Pattern against every Attribute
+		for (int Index = 0; Index < AttributeNames.Num(); ++Index)
+		{
+			const FName& AttributeName = AttributeNames[Index];
+			FRegexMatcher RegexMatcher(Pattern, AttributeName.ToString());
+	
+			if (!RegexMatcher.FindNext()) 
+			{
+				continue;
+			}
+	
+			const FPCGMetadataAttributeBase* AttributeBase = InSpatialData->Metadata->GetConstAttribute(AttributeName);
+			check(AttributeBase);
+
+			// Avoid adding the same attribute multiple times if it is captured by different regex patterns
+			if (FoundAttributes.Contains(AttributeBase->AttributeId))
+			{
+				continue;
+			}
+
+			FoundAttributes.Add(AttributeBase->AttributeId);
+	
+			if (!AddTypeToPacking(AttributeBase->GetTypeId(), OutPackedCustomData))
+			{
+				PCGE_LOG_C(Warning, &Context, "Attribute name %s is not a valid type", *AttributeName.ToString());
+				continue;
+			}
+	
+			SelectedAttributes.Add(AttributeBase);
+		}
+	}
+	
+	PackCustomDataFromAttributes(InstanceList, SelectedAttributes, OutPackedCustomData);
+}
diff --git a/Engine/Plugins/Experimental/PCG/Source/PCG/Private/MeshSelectors/PCGMeshSelectorBase.cpp b/Engine/Plugins/Experimental/PCG/Source/PCG/Private/MeshSelectors/PCGMeshSelectorBase.cpp
--- a/Engine/Plugins/Experimental/PCG/Source/PCG/Private/MeshSelectors/PCGMeshSelectorBase.cpp
+++ b/Engine/Plugins/Experimental/PCG/Source/PCG/Private/MeshSelectors/PCGMeshSelectorBase.cpp
@@ -11,7 +11,7 @@ void UPCGMeshSelectorBase::SelectInstances_Implementation(
 	TArray<FPCGMeshInstanceList>& OutMeshInstances,
 	UPCGPointData* OutPointData) const
 {
-	PCGE_LOG_C(Error, &Context, "Invalid use of abstract MeshSelectorBase class. Please use an already existing class or implement the CreateMeshInstanceData method");
+	PCGE_LOG_C(Error, &Context, "Invalid use of abstract MeshSelectorBase class. Please use an already existing class or implement the SelectInstances method");
 }
 
 bool UPCGMeshSelectorBase::FindOrAddInstanceList(
diff --git a/Engine/Plugins/Experimental/PCG/Source/PCG/Private/MeshSelectors/PCGMeshSelectorByAttribute.cpp b/Engine/Plugins/Experimental/PCG/Source/PCG/Private/MeshSelectors/PCGMeshSelectorByAttribute.cpp
--- a/Engine/Plugins/Experimental/PCG/Source/PCG/Private/MeshSelectors/PCGMeshSelectorByAttribute.cpp
+++ b/Engine/Plugins/Experimental/PCG/Source/PCG/Private/MeshSelectors/PCGMeshSelectorByAttribute.cpp
@@ -92,6 +92,6 @@ void UPCGMeshSelectorByAttribute::SelectInstances_Implementation(
 
 		int32 Index = INDEX_NONE;
 		FindOrAddInstanceList(OutMeshInstances, Mesh, bOverrideCollisionProfile, CollisionProfile, bOverrideMaterials, MaterialOverrides, Index);
-		OutMeshInstances[Index].Instances.Emplace(Point.Transform);
+		OutMeshInstances[Index].Instances.Emplace(Point);
 	}
 }
diff --git a/Engine/Plugins/Experimental/PCG/Source/PCG/Private/MeshSelectors/PCGMeshSelectorWeighted.cpp b/Engine/Plugins/Experimental/PCG/Source/PCG/Private/MeshSelectors/PCGMeshSelectorWeighted.cpp
--- a/Engine/Plugins/Experimental/PCG/Source/PCG/Private/MeshSelectors/PCGMeshSelectorWeighted.cpp
+++ b/Engine/Plugins/Experimental/PCG/Source/PCG/Private/MeshSelectors/PCGMeshSelectorWeighted.cpp
@@ -103,7 +103,7 @@ void UPCGMeshSelectorWeighted::SelectInstances_Implementation(
 			if (RandomPick < OutMeshInstances.Num())
 			{
 				FPCGMeshInstanceList& InstanceList = OutMeshInstances[RandomPick];
-				InstanceList.Instances.Emplace(Point.Transform);
+				InstanceList.Instances.Emplace(Point);
 
 				const TSoftObjectPtr<UStaticMesh>& Mesh = InstanceList.Mesh;
 
diff --git a/Engine/Plugins/Experimental/PCG/Source/PCG/Private/MeshSelectors/PCGMeshSelectorWeightedByCategory.cpp b/Engine/Plugins/Experimental/PCG/Source/PCG/Private/MeshSelectors/PCGMeshSelectorWeightedByCategory.cpp
--- a/Engine/Plugins/Experimental/PCG/Source/PCG/Private/MeshSelectors/PCGMeshSelectorWeightedByCategory.cpp
+++ b/Engine/Plugins/Experimental/PCG/Source/PCG/Private/MeshSelectors/PCGMeshSelectorWeightedByCategory.cpp
@@ -198,7 +198,7 @@ void UPCGMeshSelectorWeightedByCategory::SelectInstances_Implementation(
 		{
 			const int32 Index = InstancesAndWeights->InstanceListIndices[RandomPick];
 			FPCGMeshInstanceList& InstanceList = OutMeshInstances[Index];
-			InstanceList.Instances.Emplace(Point.Transform);
+			InstanceList.Instances.Emplace(Point);
 
 			const TSoftObjectPtr<UStaticMesh>& Mesh = InstanceList.Mesh;
 
diff --git a/Engine/Plugins/Experimental/PCG/Source/PCG/Private/Tests/Elements/PCGInstancePackerTest.cpp b/Engine/Plugins/Experimental/PCG/Source/PCG/Private/Tests/Elements/PCGInstancePackerTest.cpp
new file mode 100644
--- /dev/null
+++ b/Engine/Plugins/Experimental/PCG/Source/PCG/Private/Tests/Elements/PCGInstancePackerTest.cpp
@@ -0,0 +1,413 @@
+// Copyright Epic Games, Inc. All Rights Reserved.
+
+#include "Tests/PCGTestsCommon.h"
+#include "Tests/Determinism/PCGDeterminismTestsCommon.h"
+
+#include "PCGComponent.h"
+#include "PCGManagedResource.h"
+#include "Data/PCGPointData.h"
+#include "Elements/PCGStaticMeshSpawner.h"
+#include "InstancePackers/PCGInstancePackerBase.h"
+#include "InstancePackers/PCGInstancePackerByAttribute.h"
+#include "InstancePackers/PCGInstancePackerByRegex.h"
+#include "MeshSelectors/PCGMeshSelectorWeighted.h"
+#include "Metadata/PCGMetadata.h"
+#include "Metadata/PCGMetadataAttribute.h"
+#include "Metadata/PCGMetadataAttributeTpl.h"
+#include "Metadata/PCGMetadataAttributeTraits.h"
+
+#include "Components/InstancedStaticMeshComponent.h"
+#include "Engine/StaticMesh.h"
+
+#if WITH_EDITOR
+
+IMPLEMENT_CUSTOM_SIMPLE_AUTOMATION_TEST(FPCGStaticMeshSpawnerInstancePackerByAttributeTest, FPCGTestBaseClass, "pcg.tests.StaticMeshSpawner.InstancePacker.ByAttribute", PCGTestsCommon::TestFlags)
+IMPLEMENT_CUSTOM_SIMPLE_AUTOMATION_TEST(FPCGStaticMeshSpawnerInstancePackerByRegexTest, FPCGTestBaseClass, "pcg.tests.StaticMeshSpawner.InstancePacker.ByRegex", PCGTestsCommon::TestFlags)
+
+namespace
+{
+	struct FPCGTestCustomData
+	{
+		TArray<float> CustomData;
+		int NumFloats;
+	};
+
+	TArray<float> LocalPackFloats(TArray<FPCGTestCustomData> InFloatArrays, int NumPoints)
+	{
+		TArray<float> OutPackedFloats;
+
+		for (int InstanceIndex = 0; InstanceIndex < NumPoints; ++InstanceIndex)
+		{
+			for (const FPCGTestCustomData& FloatArray : InFloatArrays) {
+				const int StartingIndex = InstanceIndex * FloatArray.NumFloats;
+				const int EndingIndex = (InstanceIndex + 1) * FloatArray.NumFloats;
+
+				OutPackedFloats.Append(MakeArrayView(&(FloatArray.CustomData[StartingIndex]), EndingIndex - StartingIndex));
+			}
+		}
+
+		return OutPackedFloats;
+	}
+
+	bool ValidateInstancePacker(
+		FPCGTestBaseClass* Test,
+		const PCGTestsCommon::FTestData& TestData,
+		const FPCGElementPtr& StaticMeshSpawnerElement,
+		const UPCGStaticMeshSpawnerSettings* Settings,
+		TArray<float> ExpectedCustomData, 
+		int NumCustomDataFloats)
+	{
+		TUniquePtr<FPCGContext> Context = MakeUnique<FPCGContext>(*StaticMeshSpawnerElement->Initialize(TestData.InputData, TestData.TestPCGComponent, nullptr));
+		Context->NumAvailableTasks = 1;
+
+		while (!StaticMeshSpawnerElement->Execute(Context.Get()))
+		{}
+
+		bool bTestPassed = true;
+
+		TArray<UInstancedStaticMeshComponent*> ISMCs;
+		TestData.TestPCGComponent->ForEachManagedResource([&ISMCs](UPCGManagedResource* InResource)
+		{
+			if (UPCGManagedISMComponent* Resource = Cast<UPCGManagedISMComponent>(InResource))
+			{
+				ISMCs.Add(Resource->GetComponent());
+			}
+		});
+
+		for (UInstancedStaticMeshComponent* ISMC : ISMCs)
+		{
+			if (!Test->TestEqual("Valid NumCustomDataFloats in ISMC", ISMC->NumCustomDataFloats, NumCustomDataFloats))
+			{
+				bTestPassed = false;
+				continue;
+			}
+
+			const int32 NumInstances = ISMC->GetInstanceCount();
+			const int32 NumCustomData = ISMC->PerInstanceSMCustomData.Num();
+
+			if (!Test->TestEqual("ISMC CustomData count matches expected CustomData count", NumCustomData, ExpectedCustomData.Num()))
+			{
+				bTestPassed = false;
+				continue;
+			}
+
+			bTestPassed &= Test->TestEqual("Packed data matches expected data", ISMC->PerInstanceSMCustomData, ExpectedCustomData);
+		}
+
+		TestData.TestPCGComponent->bGenerated = true;
+		TestData.TestPCGComponent->Cleanup();
+
+		return bTestPassed;
+	};
+}
+
+bool FPCGStaticMeshSpawnerInstancePackerByAttributeTest::RunTest(const FString& Parameters)
+{
+	PCGTestsCommon::FTestData TestData(PCGDeterminismTests::Defaults::Seed);
+	PCGDeterminismTests::GenerateSettings<UPCGStaticMeshSpawnerSettings>(TestData);
+	UPCGStaticMeshSpawnerSettings* Settings = CastChecked<UPCGStaticMeshSpawnerSettings>(TestData.Settings);
+	FPCGElementPtr StaticMeshSpawnerElement = TestData.Settings->GetElement();
+
+	FPCGTaggedData& SourceTaggedData = TestData.InputData.TaggedData.Emplace_GetRef(FPCGTaggedData());
+	TObjectPtr<UPCGPointData> SourceData = PCGTestsCommon::CreateRandomPointData(5, TestData.Seed);
+	SourceTaggedData.Data = SourceData;
+	SourceData->TargetActor = TestData.TestActor;
+
+	const FName FloatName = TEXT("Float");
+	const FName DoubleName = TEXT("Double");
+	const FName IntName = TEXT("Int");
+	const FName VecName = TEXT("Vec");
+	const FName Vec4Name = TEXT("Vec4");
+	const FName RotatorName = TEXT("Rotator");
+
+	const bool bAllowsInterpolation = false;
+	const bool bOverrideParent = false;
+
+	SourceData->Metadata->CreateFloatAttribute(FloatName, 0, bAllowsInterpolation, bOverrideParent);
+	SourceData->Metadata->CreateDoubleAttribute(DoubleName, 0, bAllowsInterpolation, bOverrideParent);
+	SourceData->Metadata->CreateInteger64Attribute(IntName, 0, bAllowsInterpolation, bOverrideParent);
+	SourceData->Metadata->CreateVectorAttribute(VecName, FVector::Zero(), bAllowsInterpolation, bOverrideParent);
+	SourceData->Metadata->CreateVector4Attribute(Vec4Name, FVector4::Zero(), bAllowsInterpolation, bOverrideParent);
+	SourceData->Metadata->CreateRotatorAttribute(RotatorName, FRotator::ZeroRotator, bAllowsInterpolation, bOverrideParent);
+
+	TArray<FPCGPoint>& SourcePoints = SourceData->GetMutablePoints();
+
+	TArray<float> FloatValues;
+	TArray<float> DoubleValues;
+	TArray<float> IntValues;
+	TArray<float> VecValues;
+	TArray<float> Vec4Values;
+	TArray<float> RotatorValues;
+
+	FRandomStream RandomSource(TestData.Seed);
+	int NumPoints = 5;
+	for (int I = 0; I < NumPoints; ++I)
+	{
+		const float FloatValue = (I + 0.f) / NumPoints;
+		const double DoubleValue = FloatValue + 1;
+		const int64 IntValue = I;
+		const FVector VecValue = FVector(1, 2, 3) * I;
+		const FVector4 Vec4Value = -FVector(1, 2, 3) * I;
+		const FRotator RotatorValue = FRotator(15, 30, 45) * I;
+
+		FloatValues.Add(FloatValue);
+		DoubleValues.Add(DoubleValue);
+		IntValues.Add(IntValue);
+		VecValues.Add(VecValue.X);
+		VecValues.Add(VecValue.Y);
+		VecValues.Add(VecValue.Z);
+		Vec4Values.Add(Vec4Value.X);
+		Vec4Values.Add(Vec4Value.Y);
+		Vec4Values.Add(Vec4Value.Z);
+		Vec4Values.Add(Vec4Value.W);
+		RotatorValues.Add(RotatorValue.Roll);
+		RotatorValues.Add(RotatorValue.Pitch);
+		RotatorValues.Add(RotatorValue.Yaw);
+
+		UPCGMetadataAccessorHelpers::SetFloatAttribute(SourcePoints[I], SourceData->Metadata, FloatName, FloatValue);
+		UPCGMetadataAccessorHelpers::SetDoubleAttribute(SourcePoints[I], SourceData->Metadata, DoubleName, DoubleValue);
+		UPCGMetadataAccessorHelpers::SetInteger64Attribute(SourcePoints[I], SourceData->Metadata, IntName, IntValue);
+		UPCGMetadataAccessorHelpers::SetVectorAttribute(SourcePoints[I], SourceData->Metadata, VecName, VecValue);
+		UPCGMetadataAccessorHelpers::SetVector4Attribute(SourcePoints[I], SourceData->Metadata, Vec4Name, Vec4Value);
+		UPCGMetadataAccessorHelpers::SetRotatorAttribute(SourcePoints[I], SourceData->Metadata, RotatorName, RotatorValue);
+	}
+
+	Settings->SetInstancePackerType(UPCGInstancePackerByAttribute::StaticClass());
+	UPCGInstancePackerByAttribute* InstancePacker = CastChecked<UPCGInstancePackerByAttribute>(Settings->InstancePackerInstance);
+	UPCGMeshSelectorWeighted* MeshSelector = CastChecked<UPCGMeshSelectorWeighted>(Settings->MeshSelectorInstance);
+
+	const FString CubePath = TEXT("StaticMesh'/Engine/BasicShapes/Cube.Cube'");
+	const UStaticMesh* CubeMesh = Cast<UStaticMesh>(StaticLoadObject(UStaticMesh::StaticClass(), NULL, *CubePath));
+	MeshSelector->MeshEntries.Add(FPCGMeshSelectorWeightedEntry(CubeMesh, 1));
+
+	bool bTestPassed = true;
+
+	// No attributes
+	{
+		InstancePacker->AttributeNames = {};
+		bTestPassed &= ValidateInstancePacker(this, TestData, StaticMeshSpawnerElement, Settings, {}, 0);
+	}
+
+	// Float Attribute
+	{
+		InstancePacker->AttributeNames = { FloatName };
+		bTestPassed &= ValidateInstancePacker(this, TestData, StaticMeshSpawnerElement, Settings, FloatValues, 1);
+	}
+
+	// Double Attribute
+	{
+		InstancePacker->AttributeNames = { DoubleName };
+		bTestPassed &= ValidateInstancePacker(this, TestData, StaticMeshSpawnerElement, Settings, DoubleValues, 1);
+	}
+
+	// Int Attribute
+	{
+		InstancePacker->AttributeNames = { IntName };
+		bTestPassed &= ValidateInstancePacker(this, TestData, StaticMeshSpawnerElement, Settings, IntValues, 1);
+	}
+
+	// Vector Attribute
+	{
+		InstancePacker->AttributeNames = { VecName };
+		bTestPassed &= ValidateInstancePacker(this, TestData, StaticMeshSpawnerElement, Settings, VecValues, 3);
+	}
+
+	// Vector4 Attribute
+	{
+		InstancePacker->AttributeNames = { Vec4Name };
+		bTestPassed &= ValidateInstancePacker(this, TestData, StaticMeshSpawnerElement, Settings, Vec4Values, 4);
+	}
+
+	// Rotator Attribute
+	{
+		InstancePacker->AttributeNames = { RotatorName };
+		bTestPassed &= ValidateInstancePacker(this, TestData, StaticMeshSpawnerElement, Settings, RotatorValues, 3);
+	}
+
+	// Float + Vector
+	{
+		TArray<float> ExpectedCustomData = LocalPackFloats({ { FloatValues, 1 }, { VecValues, 3 } }, NumPoints);
+
+		InstancePacker->AttributeNames = { FloatName, VecName };
+		bTestPassed &= ValidateInstancePacker(this, TestData, StaticMeshSpawnerElement, Settings, ExpectedCustomData, 4);
+	}
+
+	// Vector + Float
+	{
+		TArray<float> ExpectedCustomData = LocalPackFloats({ { VecValues, 3 }, { FloatValues, 1 } }, NumPoints);
+
+		InstancePacker->AttributeNames = { VecName, FloatName };
+		bTestPassed &= ValidateInstancePacker(this, TestData, StaticMeshSpawnerElement, Settings, ExpectedCustomData, 4);
+	}
+
+	// Vector + Float + Vector
+	{
+		TArray<float> ExpectedCustomData = LocalPackFloats({ { VecValues, 3 }, { FloatValues, 1 }, { VecValues, 3 } }, NumPoints);
+
+		InstancePacker->AttributeNames = { VecName, FloatName, VecName };
+		bTestPassed &= ValidateInstancePacker(this, TestData, StaticMeshSpawnerElement, Settings, ExpectedCustomData, 7);
+	}
+
+	// Vector + Float + Vector + Vector4
+	{
+		TArray<float> ExpectedCustomData = LocalPackFloats({ { VecValues, 3 }, { FloatValues, 1 }, { VecValues, 3 }, { Vec4Values, 4 } }, NumPoints);
+
+		InstancePacker->AttributeNames = { VecName, FloatName, VecName, Vec4Name };
+		bTestPassed &= ValidateInstancePacker(this, TestData, StaticMeshSpawnerElement, Settings, ExpectedCustomData, 11);
+	}
+
+	return bTestPassed;
+}
+
+bool FPCGStaticMeshSpawnerInstancePackerByRegexTest::RunTest(const FString& Parameters)
+{
+	PCGTestsCommon::FTestData TestData(PCGDeterminismTests::Defaults::Seed);
+	PCGDeterminismTests::GenerateSettings<UPCGStaticMeshSpawnerSettings>(TestData);
+	UPCGStaticMeshSpawnerSettings* Settings = CastChecked<UPCGStaticMeshSpawnerSettings>(TestData.Settings);
+	FPCGElementPtr StaticMeshSpawnerElement = TestData.Settings->GetElement();
+
+	FPCGTaggedData& SourceTaggedData = TestData.InputData.TaggedData.Emplace_GetRef(FPCGTaggedData());
+	TObjectPtr<UPCGPointData> SourceData = PCGTestsCommon::CreateRandomPointData(5, TestData.Seed);
+	SourceTaggedData.Data = SourceData;
+	SourceData->TargetActor = TestData.TestActor;
+
+	const FName NameA1 = TEXT("A1"); // Float
+	const FName NameA2 = TEXT("A2"); // Double
+	const FName NameA3 = TEXT("A3"); // Int
+	const FName NameB1 = TEXT("B1"); // FVector
+	const FName NameB2 = TEXT("B2"); // FVector4
+
+	const bool bAllowsInterpolation = false;
+	const bool bOverrideParent = false;
+
+	SourceData->Metadata->CreateFloatAttribute(NameA1, 0, bAllowsInterpolation, bOverrideParent);
+	SourceData->Metadata->CreateDoubleAttribute(NameA2, 0, bAllowsInterpolation, bOverrideParent);
+	SourceData->Metadata->CreateInteger64Attribute(NameA3, 0, bAllowsInterpolation, bOverrideParent);
+	SourceData->Metadata->CreateVectorAttribute(NameB1, FVector::Zero(), bAllowsInterpolation, bOverrideParent);
+	SourceData->Metadata->CreateVector4Attribute(NameB2, FVector4::Zero(), bAllowsInterpolation, bOverrideParent);
+
+	TArray<FPCGPoint>& SourcePoints = SourceData->GetMutablePoints();
+
+	TArray<float> FloatValues;
+	TArray<float> DoubleValues;
+	TArray<float> IntValues;
+	TArray<float> VecValues;
+	TArray<float> Vec4Values;
+
+	FRandomStream RandomSource(TestData.Seed);
+	int NumPoints = 5;
+	for (int I = 0; I < NumPoints; ++I)
+	{
+		const float FloatValue = (I + 0.f) / NumPoints;
+		const double DoubleValue = FloatValue + 1;
+		const int64 IntValue = I;
+		const FVector VecValue = FVector(1, 2, 3) * I;
+		const FVector4 Vec4Value = -FVector(1, 2, 3) * I;
+
+		FloatValues.Add(FloatValue);
+		DoubleValues.Add(DoubleValue);
+		IntValues.Add(IntValue);
+		VecValues.Add(VecValue.X);
+		VecValues.Add(VecValue.Y);
+		VecValues.Add(VecValue.Z);
+		Vec4Values.Add(Vec4Value.X);
+		Vec4Values.Add(Vec4Value.Y);
+		Vec4Values.Add(Vec4Value.Z);
+		Vec4Values.Add(Vec4Value.W);
+
+		UPCGMetadataAccessorHelpers::SetFloatAttribute(SourcePoints[I], SourceData->Metadata, NameA1, FloatValue);
+		UPCGMetadataAccessorHelpers::SetDoubleAttribute(SourcePoints[I], SourceData->Metadata, NameA2, DoubleValue);
+		UPCGMetadataAccessorHelpers::SetInteger64Attribute(SourcePoints[I], SourceData->Metadata, NameA3, IntValue);
+		UPCGMetadataAccessorHelpers::SetVectorAttribute(SourcePoints[I], SourceData->Metadata, NameB1, VecValue);
+		UPCGMetadataAccessorHelpers::SetVector4Attribute(SourcePoints[I], SourceData->Metadata, NameB2, Vec4Value);
+	}
+
+	Settings->SetInstancePackerType(UPCGInstancePackerByRegex::StaticClass());
+	UPCGInstancePackerByRegex* InstancePacker = CastChecked<UPCGInstancePackerByRegex>(Settings->InstancePackerInstance);
+	UPCGMeshSelectorWeighted* MeshSelector = CastChecked<UPCGMeshSelectorWeighted>(Settings->MeshSelectorInstance);
+
+	const FString CubePath = TEXT("StaticMesh'/Engine/BasicShapes/Cube.Cube'");
+	const UStaticMesh* CubeMesh = Cast<UStaticMesh>(StaticLoadObject(UStaticMesh::StaticClass(), NULL, *CubePath));
+	MeshSelector->MeshEntries.Add(FPCGMeshSelectorWeightedEntry(CubeMesh, 1));
+
+	bool bTestPassed = true;
+
+	// No attributes
+	{
+		InstancePacker->RegexPatterns = {};
+		bTestPassed &= ValidateInstancePacker(this, TestData, StaticMeshSpawnerElement, Settings, {}, 0);
+	}
+
+	{
+		InstancePacker->RegexPatterns = { TEXT("A1") };
+		bTestPassed &= ValidateInstancePacker(this, TestData, StaticMeshSpawnerElement, Settings, FloatValues, 1);
+	}
+
+	{
+		TArray<float> ExpectedCustomData = LocalPackFloats({ { FloatValues, 1 }, { Vec4Values, 4 } }, NumPoints);
+
+		InstancePacker->RegexPatterns = { TEXT("A1"), TEXT("B2") };
+		bTestPassed &= ValidateInstancePacker(this, TestData, StaticMeshSpawnerElement, Settings, ExpectedCustomData, 5);
+	}
+
+	{
+		TArray<float> ExpectedCustomData = LocalPackFloats({ { FloatValues, 1 }, { DoubleValues, 1 }, { IntValues, 1 } }, NumPoints);
+
+		InstancePacker->RegexPatterns = { TEXT("A.") };
+		bTestPassed &= ValidateInstancePacker(this, TestData, StaticMeshSpawnerElement, Settings, ExpectedCustomData, 3);
+	}
+
+	{
+		TArray<float> ExpectedCustomData = LocalPackFloats({ { VecValues, 3 }, { Vec4Values, 4 } }, NumPoints);
+
+		InstancePacker->RegexPatterns = { TEXT("B.") };
+		bTestPassed &= ValidateInstancePacker(this, TestData, StaticMeshSpawnerElement, Settings, ExpectedCustomData, 7);
+	}
+
+	{
+		TArray<float> ExpectedCustomData = LocalPackFloats({ { FloatValues, 1 }, { DoubleValues, 1 }, { IntValues, 1 }, { VecValues, 3 }, { Vec4Values, 4 } }, NumPoints);
+
+		InstancePacker->RegexPatterns = { TEXT("A."), TEXT("B.") };
+		bTestPassed &= ValidateInstancePacker(this, TestData, StaticMeshSpawnerElement, Settings, ExpectedCustomData, 10);
+	}
+
+	{
+		TArray<float> ExpectedCustomData = LocalPackFloats({ { VecValues, 3 }, { Vec4Values, 4 }, { FloatValues, 1 }, { DoubleValues, 1 }, { IntValues, 1 } }, NumPoints);
+
+		InstancePacker->RegexPatterns = { TEXT("B."), TEXT("A.") };
+		bTestPassed &= ValidateInstancePacker(this, TestData, StaticMeshSpawnerElement, Settings, ExpectedCustomData, 10);
+	}
+
+	{
+		TArray<float> ExpectedCustomData = LocalPackFloats({ { FloatValues, 1 }, { VecValues, 3 } }, NumPoints);
+
+		InstancePacker->RegexPatterns = { TEXT(".1") };
+		bTestPassed &= ValidateInstancePacker(this, TestData, StaticMeshSpawnerElement, Settings, ExpectedCustomData, 4);
+	}
+
+	{
+		TArray<float> ExpectedCustomData = LocalPackFloats({ { FloatValues, 1 }, { VecValues, 3 } }, NumPoints);
+
+		InstancePacker->RegexPatterns = { TEXT(".1"), TEXT(".1") };
+		bTestPassed &= ValidateInstancePacker(this, TestData, StaticMeshSpawnerElement, Settings, ExpectedCustomData, 4);
+	}
+
+	{
+		TArray<float> ExpectedCustomData = LocalPackFloats({ { FloatValues, 1 }, { VecValues, 3 }, { DoubleValues, 1 }, { Vec4Values, 4 } }, NumPoints);
+
+		InstancePacker->RegexPatterns = { TEXT(".1"), TEXT(".2") };
+		bTestPassed &= ValidateInstancePacker(this, TestData, StaticMeshSpawnerElement, Settings, ExpectedCustomData, 9);
+	}
+
+	// Capture all attributes
+	{
+		TArray<float> ExpectedCustomData = LocalPackFloats({ { FloatValues, 1 }, { DoubleValues, 1 }, { IntValues, 1 }, { VecValues, 3 }, { Vec4Values, 4 } }, NumPoints);
+
+		InstancePacker->RegexPatterns = { TEXT("[A-Z][0-9]") };
+		bTestPassed &= ValidateInstancePacker(this, TestData, StaticMeshSpawnerElement, Settings, ExpectedCustomData, 10);
+	}
+
+	return bTestPassed;
+}
+
+#endif // WITH_EDITOR
diff --git a/Engine/Plugins/Experimental/PCG/Source/PCG/Public/Elements/PCGStaticMeshSpawner.h b/Engine/Plugins/Experimental/PCG/Source/PCG/Public/Elements/PCGStaticMeshSpawner.h
--- a/Engine/Plugins/Experimental/PCG/Source/PCG/Public/Elements/PCGStaticMeshSpawner.h
+++ b/Engine/Plugins/Experimental/PCG/Source/PCG/Public/Elements/PCGStaticMeshSpawner.h
@@ -4,6 +4,7 @@
 
 #include "CoreMinimal.h"
 #include "PCGSettings.h"
+#include "InstancePackers/PCGInstancePackerBase.h"
 #include "MeshSelectors/PCGMeshSelectorBase.h"
 
 #include "Engine/CollisionProfile.h"
@@ -62,14 +63,26 @@ public:
 	UFUNCTION(BlueprintCallable, Category = Settings)
 	void SetMeshSelectorType(TSubclassOf<UPCGMeshSelectorBase> InMeshSelectorType);
 
+	UFUNCTION(BlueprintCallable, Category = Settings)
+	void SetInstancePackerType(TSubclassOf<UPCGInstancePackerBase> InInstancePackerType);
+
 public:
+	/** Defines the method of mesh selection per input data */
 	UPROPERTY(BlueprintReadOnly, EditAnywhere, NoClear, Category = Settings)
 	TSubclassOf<UPCGMeshSelectorBase> MeshSelectorType;
 
 	UPROPERTY(BlueprintReadOnly, VisibleAnywhere, Instanced, Category = Settings)
 	TObjectPtr<UPCGMeshSelectorBase> MeshSelectorInstance;
 
-	UPROPERTY(BlueprintReadOnly, VisibleAnywhere, Category = Settings)
+	/** Defines the method of custom data packing for spawned (H)ISMCs */
+	UPROPERTY(BlueprintReadOnly, EditAnywhere, Category = Settings)
+	TSubclassOf<UPCGInstancePackerBase> InstancePackerType;
+
+	UPROPERTY(BlueprintReadOnly, VisibleAnywhere, Instanced, Category = Settings)
+	TObjectPtr<UPCGInstancePackerBase> InstancePackerInstance = nullptr;
+
+	/** Attribute name to store mesh SoftObjectPaths inside if the output pin is connected. Note: Will overwrite existing data if the attribute name already exists. */
+	UPROPERTY(BlueprintReadWrite, EditAnywhere, Category = Settings)
 	FName OutAttributeName = NAME_None;
 
 	UPROPERTY(meta=(DeprecatedProperty, DeprecationMessage="Use MeshSelectorType and MeshSelectorInstance instead."))
@@ -78,7 +91,8 @@ public:
 	bool bForceConnectOutput = false;
 
 protected:
-	void RefreshMeshSelector(); 
+	void RefreshMeshSelector();
+	void RefreshInstancePacker();
 };
 
 class FPCGStaticMeshSpawnerElement : public FSimplePCGElement
@@ -89,6 +103,6 @@ public:
 
 protected:
 	virtual bool ExecuteInternal(FPCGContext* Context) const override;	
-	void SpawnStaticMeshInstances(FPCGContext* Context, const TArray<FPCGMeshInstanceList>& MeshInstances, AActor* TargetActor) const;
+	void SpawnStaticMeshInstances(FPCGContext* Context, const FPCGMeshInstanceList& InstanceList, AActor* TargetActor, const FPCGPackedCustomData& PackedCustomData) const;
 };
 
diff --git a/Engine/Plugins/Experimental/PCG/Source/PCG/Public/Helpers/PCGActorHelpers.h b/Engine/Plugins/Experimental/PCG/Source/PCG/Public/Helpers/PCGActorHelpers.h
--- a/Engine/Plugins/Experimental/PCG/Source/PCG/Public/Helpers/PCGActorHelpers.h
+++ b/Engine/Plugins/Experimental/PCG/Source/PCG/Public/Helpers/PCGActorHelpers.h
@@ -27,6 +27,7 @@ struct FPCGISMCBuilderParameters
 	TArray<UMaterialInterface*> MaterialOverrides;
 	EComponentMobility::Type Mobility = EComponentMobility::Static;
 	FName CollisionProfile = TEXT("Default");
+	int32 NumCustomDataFloats = 0;
 };
 
 UCLASS(BlueprintType)
diff --git a/Engine/Plugins/Experimental/PCG/Source/PCG/Public/InstancePackers/PCGInstancePackerBase.h b/Engine/Plugins/Experimental/PCG/Source/PCG/Public/InstancePackers/PCGInstancePackerBase.h
new file mode 100644
--- /dev/null
+++ b/Engine/Plugins/Experimental/PCG/Source/PCG/Public/InstancePackers/PCGInstancePackerBase.h
@@ -0,0 +1,48 @@
+// Copyright Epic Games, Inc. All Rights Reserved.
+
+#pragma once
+
+#include "CoreMinimal.h"
+#include "PCGElement.h"
+#include "MeshSelectors/PCGMeshSelectorBase.h"
+
+#include "PCGInstancePackerBase.generated.h"
+
+class UPCGSpatialData;
+class FPCGMetadataAttributeBase;
+
+USTRUCT(BlueprintType)
+struct FPCGPackedCustomData
+{
+	GENERATED_BODY()
+
+	UPROPERTY(BlueprintReadWrite, EditAnywhere, Category = Settings)
+	int NumCustomDataFloats = 0;
+
+	UPROPERTY(BlueprintReadWrite, EditAnywhere, Category = Settings)
+    TArray<float> CustomData;
+};
+
+UCLASS(Abstract, BlueprintType, Blueprintable, ClassGroup = (Procedural))
+class PCG_API UPCGInstancePackerBase : public UObject 
+{
+	GENERATED_BODY()
+
+public:
+	/** Defines the strategy for (H)ISM custom float data packing */
+	UFUNCTION(BlueprintNativeEvent, Category = InstancePacking)
+	void PackInstances(FPCGContext& Context, const UPCGSpatialData* InSpatialData, const FPCGMeshInstanceList& InstanceList, FPCGPackedCustomData& OutPackedCustomData) const;
+
+	virtual void PackInstances_Implementation(FPCGContext& Context, const UPCGSpatialData* InSpatialData, const FPCGMeshInstanceList& InstanceList, FPCGPackedCustomData& OutPackedCustomData) const;
+
+	/** Interprets Metadata TypeId and increments OutPackedCustomData.NumCustomDataFloats appropriately. Returns false if the type could not be interpreted. */
+	UFUNCTION(BlueprintCallable, Category = InstancePacking)
+	bool AddTypeToPacking(int TypeId, FPCGPackedCustomData& OutPackedCustomData) const;
+
+	/** Build a PackedCustomData by processing each attribute in order for each point in the InstanceList */
+	UFUNCTION(BlueprintCallable, Category = InstancePacking) 
+	void PackCustomDataFromAttributes(const FPCGMeshInstanceList& InstanceList, const UPCGMetadata* Metadata, const TArray<FName>& AttributeNames, FPCGPackedCustomData& OutPackedCustomData) const;
+
+	/** Build a PackedCustomData by processing each attribute in order for each point in the InstanceList */
+	void PackCustomDataFromAttributes(const FPCGMeshInstanceList& InstanceList, const TArray<const FPCGMetadataAttributeBase*>& Attributes, FPCGPackedCustomData& OutPackedCustomData) const;
+};
diff --git a/Engine/Plugins/Experimental/PCG/Source/PCG/Public/InstancePackers/PCGInstancePackerByAttribute.h b/Engine/Plugins/Experimental/PCG/Source/PCG/Public/InstancePackers/PCGInstancePackerByAttribute.h
new file mode 100644
--- /dev/null
+++ b/Engine/Plugins/Experimental/PCG/Source/PCG/Public/InstancePackers/PCGInstancePackerByAttribute.h
@@ -0,0 +1,20 @@
+// Copyright Epic Games, Inc. All Rights Reserved.
+
+#pragma once
+
+#include "PCGInstancePackerBase.h"
+
+#include "PCGInstancePackerByAttribute.generated.h"
+
+UCLASS(BlueprintType, ClassGroup = (Procedural))
+class PCG_API UPCGInstancePackerByAttribute : public UPCGInstancePackerBase 
+{
+	GENERATED_BODY()
+
+public:
+	virtual void PackInstances_Implementation(UPARAM(ref) FPCGContext& Context, const UPCGSpatialData* InSpatialData, UPARAM(ref) const FPCGMeshInstanceList& InstanceList, FPCGPackedCustomData& OutPackedCustomData) const override;
+
+public:
+	UPROPERTY(BlueprintReadWrite, EditAnywhere, Category = Settings)
+	TArray<FName> AttributeNames;
+};
diff --git a/Engine/Plugins/Experimental/PCG/Source/PCG/Public/InstancePackers/PCGInstancePackerByRegex.h b/Engine/Plugins/Experimental/PCG/Source/PCG/Public/InstancePackers/PCGInstancePackerByRegex.h
new file mode 100644
--- /dev/null
+++ b/Engine/Plugins/Experimental/PCG/Source/PCG/Public/InstancePackers/PCGInstancePackerByRegex.h
@@ -0,0 +1,20 @@
+// Copyright Epic Games, Inc. All Rights Reserved.
+
+#pragma once
+
+#include "PCGInstancePackerBase.h"
+
+#include "PCGInstancePackerByRegex.generated.h"
+
+UCLASS(BlueprintType, ClassGroup = (Procedural))
+class PCG_API UPCGInstancePackerByRegex : public UPCGInstancePackerBase 
+{
+	GENERATED_BODY()
+
+public:
+	virtual void PackInstances_Implementation(UPARAM(ref) FPCGContext& Context, const UPCGSpatialData* InSpatialData, UPARAM(ref) const FPCGMeshInstanceList& InstanceList, FPCGPackedCustomData& OutPackedCustomData) const override;
+
+public:
+	UPROPERTY(BlueprintReadWrite, EditAnywhere, Category = Settings)
+	TArray<FString> RegexPatterns; 
+};
diff --git a/Engine/Plugins/Experimental/PCG/Source/PCG/Public/MeshSelectors/PCGMeshSelectorBase.h b/Engine/Plugins/Experimental/PCG/Source/PCG/Public/MeshSelectors/PCGMeshSelectorBase.h
--- a/Engine/Plugins/Experimental/PCG/Source/PCG/Public/MeshSelectors/PCGMeshSelectorBase.h
+++ b/Engine/Plugins/Experimental/PCG/Source/PCG/Public/MeshSelectors/PCGMeshSelectorBase.h
@@ -41,7 +41,7 @@ struct FPCGMeshInstanceList
 	TArray<UMaterialInterface*> MaterialOverrides;
 
 	UPROPERTY(BlueprintReadWrite, EditAnywhere, Category = Settings)
-    TArray<FTransform> Instances;
+    TArray<FPCGPoint> Instances;
 };
 
 UCLASS(Abstract, BlueprintType, Blueprintable, ClassGroup = (Procedural))
