Allows PCG nodes to have custom determinism tests
In Editor:
 + Toggle for Basic Tests
 + Toggle for Order Independence Tests
 + Toggle for running custom blueprint
 + New blueprint base class for determinism tests

diff --git a/Engine/Plugins/Experimental/PCG/Source/PCG/Private/PCGSettings.cpp b/Engine/Plugins/Experimental/PCG/Source/PCG/Private/PCGSettings.cpp
--- a/Engine/Plugins/Experimental/PCG/Source/PCG/Private/PCGSettings.cpp
+++ b/Engine/Plugins/Experimental/PCG/Source/PCG/Private/PCGSettings.cpp
@@ -16,7 +16,7 @@ public:
 	virtual bool ShouldSkipProperty(const FProperty* InProperty) const override
 	{
 #if WITH_EDITORONLY_DATA
-		return InProperty && InProperty->GetFName() == GET_MEMBER_NAME_CHECKED(UPCGSettings, DebugSettings);
+		return InProperty && (InProperty->GetFName() == GET_MEMBER_NAME_CHECKED(UPCGSettings, DebugSettings) || InProperty->GetFName() == GET_MEMBER_NAME_CHECKED(UPCGSettings, DeterminismSettings));
 #else
 		return FArchiveObjectCrc32::ShouldSkipProperty(InProperty);
 #endif // WITH_EDITORONLY_DATA
@@ -94,7 +94,11 @@ UPCGNode* UPCGSettings::CreateNode() const
 void UPCGSettings::PostEditChangeProperty(FPropertyChangedEvent& PropertyChangedEvent)
 {
 	Super::PostEditChangeProperty(PropertyChangedEvent);
-	OnSettingsChangedDelegate.Broadcast(this, EPCGChangeType::Settings);
+
+	if (PropertyChangedEvent.GetMemberPropertyName() != GET_MEMBER_NAME_CHECKED(UPCGSettings, DeterminismSettings))
+	{
+		OnSettingsChangedDelegate.Broadcast(this, EPCGChangeType::Settings);
+	}
 }
 
 void UPCGSettings::DirtyCache()
diff --git a/Engine/Plugins/Experimental/PCG/Source/PCG/Private/Tests/Determinism/PCGDeterminismTestBlueprintBase.cpp b/Engine/Plugins/Experimental/PCG/Source/PCG/Private/Tests/Determinism/PCGDeterminismTestBlueprintBase.cpp
new file mode 100644
--- /dev/null
+++ b/Engine/Plugins/Experimental/PCG/Source/PCG/Private/Tests/Determinism/PCGDeterminismTestBlueprintBase.cpp
@@ -0,0 +1,12 @@
+// Copyright Epic Games, Inc. All Rights Reserved.
+
+#include "Tests/Determinism/PCGDeterminismTestBlueprintBase.h"
+
+void UPCGDeterminismTestBlueprintBase::ExecuteTest_Implementation(const UPCGNode* InPCGNode, FDeterminismNodeTestResult& InOutTestResult)
+{
+	// The default case if the user tries to run the blueprint test unimplemented
+	InOutTestResult.NodeTitle = InPCGNode->GetNodeTitle();
+	InOutTestResult.NodeName = InPCGNode->GetName();
+	InOutTestResult.Seed = InPCGNode->DefaultSettings->Seed;
+	PCGDeterminismTests::RunDeterminismTest(InPCGNode, InOutTestResult, PCGDeterminismTests::Defaults::DeterminismBasicTestInfo);
+}
\ No newline at end of file
diff --git a/Engine/Plugins/Experimental/PCG/Source/PCG/Private/Tests/Determinism/PCGDeterminismTestsCommon.cpp b/Engine/Plugins/Experimental/PCG/Source/PCG/Private/Tests/Determinism/PCGDeterminismTestsCommon.cpp
--- a/Engine/Plugins/Experimental/PCG/Source/PCG/Private/Tests/Determinism/PCGDeterminismTestsCommon.cpp
+++ b/Engine/Plugins/Experimental/PCG/Source/PCG/Private/Tests/Determinism/PCGDeterminismTestsCommon.cpp
@@ -19,22 +19,16 @@
 
 #define LOCTEXT_NAMESPACE "PCGDeterminism"
 
-namespace
-{
-	const FText TEXT_BasicTest = LOCTEXT("BasicTests", "Basic Tests");
-	const FText TEXT_OrderIndependenceTest = LOCTEXT("OrderIndependence", "Order Independence");
-}
-
 namespace PCGDeterminismTests
 {
-	bool LogInvalidTest(const UPCGNode* InPCGNode, FNodeTestResult& OutResult)
+	bool LogInvalidTest(const UPCGNode* InPCGNode, FDeterminismNodeTestResult& OutResult)
 	{
 		UE_LOG(LogPCG, Warning, TEXT("Attempting to run an invalid determinism test"));
 		OutResult.AdditionalDetails.Add(TEXT("Invalid test"));
 		return false;
 	}
 
-	PCG_API void RunDeterminismTest(const UPCGNode* InPCGNode, FNodeTestResult& OutResult, const FNodeTestInfo& TestToRun)
+	void RunDeterminismTest(const UPCGNode* InPCGNode, FDeterminismNodeTestResult& OutResult, const FNodeTestInfo& TestToRun)
 	{
 		if (!InPCGNode || !InPCGNode->DefaultSettings)
 		{
@@ -44,7 +38,6 @@ namespace PCGDeterminismTests
 			return;
 		}
 
-		OutResult.Seed = InPCGNode->DefaultSettings->Seed;
 		bool bTestWasSuccessful = TestToRun.TestDelegate(InPCGNode, TestToRun.TestName, OutResult);
 		OutResult.bFlagRaised |= !bTestWasSuccessful;
 
@@ -57,13 +50,7 @@ namespace PCGDeterminismTests
 		}
 	}
 
-	PCG_API void RetrieveBasicTests(TArray<FNodeTestInfo>& OutBasicTests)
-	{
-		OutBasicTests.Emplace(TEXT_BasicTest, RunBasicTestSuite, 80.f);
-		OutBasicTests.Emplace(TEXT_OrderIndependenceTest, RunOrderIndependenceSuite, 125.f);
-	}
-
-	bool RunBasicTestSuite(const UPCGNode* InPCGNode, const FName& TestName, FNodeTestResult& OutResult)
+	bool RunBasicTestSuite(const UPCGNode* InPCGNode, const FName& TestName, FDeterminismNodeTestResult& OutResult)
 	{
 		// Get Pins
 		const TArray<TObjectPtr<UPCGPin>>& InputPins = InPCGNode->GetInputPins();
@@ -75,7 +62,7 @@ namespace PCGDeterminismTests
 		return RunBasicSelfTest(NodeAndOptions) && RunBasicCopiedSelfTest(NodeAndOptions);
 	}
 
-	bool RunOrderIndependenceSuite(const UPCGNode* InPCGNode, const FName& TestName, FNodeTestResult& OutResult)
+	bool RunOrderIndependenceSuite(const UPCGNode* InPCGNode, const FName& TestName, FDeterminismNodeTestResult& OutResult)
 	{
 		// Start with not deterministic
 		UpdateTestResults(TestName, OutResult, EDeterminismLevel::NoDeterminism);
@@ -1040,7 +1027,7 @@ namespace PCGDeterminismTests
 		return BaseOptionsPerPin[PinIndex][PermutationIndex];
 	}
 
-	void UpdateTestResults(FName TestName, FNodeTestResult& OutResult, EDeterminismLevel DeterminismLevel)
+	void UpdateTestResults(FName TestName, FDeterminismNodeTestResult& OutResult, EDeterminismLevel DeterminismLevel)
 	{
 		EDeterminismLevel& Result = OutResult.TestResults.FindOrAdd(TestName);
 		Result = DeterminismLevel;
diff --git a/Engine/Plugins/Experimental/PCG/Source/PCG/Public/PCGSettings.h b/Engine/Plugins/Experimental/PCG/Source/PCG/Public/PCGSettings.h
--- a/Engine/Plugins/Experimental/PCG/Source/PCG/Public/PCGSettings.h
+++ b/Engine/Plugins/Experimental/PCG/Source/PCG/Public/PCGSettings.h
@@ -6,12 +6,13 @@
 #include "PCGData.h"
 #include "PCGElement.h"
 #include "PCGDebug.h"
+#include "Tests/Determinism//PCGDeterminismSettings.h"
 
 #include "PCGSettings.generated.h"
 
+class UPCGGraph;
 class UPCGNode;
 class UPCGSettings;
-class UPCGGraph;
 
 typedef TMap<FName, TSet<TWeakObjectPtr<const UPCGSettings>>> FPCGTagToSettingsMap;
 
@@ -63,7 +64,7 @@ public:
 
 	virtual TArray<FPCGPinProperties> InputPinProperties() const;
 	virtual TArray<FPCGPinProperties> OutputPinProperties() const;
-	
+
 	bool operator==(const UPCGSettings& Other) const;
 	uint32 GetCrc32() const;
 
@@ -104,6 +105,9 @@ public:
 #if WITH_EDITORONLY_DATA
 	UPROPERTY(EditAnywhere, BlueprintReadWrite, Category = Debug, meta = (ShowOnlyInnerProperties))
 	FPCGDebugVisualizationSettings DebugSettings;
+
+	UPROPERTY(EditAnywhere, BlueprintReadWrite, Category = Determinism, meta = (ShowOnlyInnerProperties))
+	FPCGDeterminismSettings DeterminismSettings;
 #endif
 
 #if WITH_EDITOR
@@ -118,7 +122,7 @@ protected:
 
 	/** Method that can be called to dirty the cache data from this settings objects if the operator== does not allow to detect changes */
 	void DirtyCache();
-#endif 
+#endif
 
 	/** Methods to remove boilerplate code across settings */
 	TArray<FPCGPinProperties> DefaultPointOutputPinProperties() const;
diff --git a/Engine/Plugins/Experimental/PCG/Source/PCG/Public/Tests/Determinism/PCGDeterminismSettings.h b/Engine/Plugins/Experimental/PCG/Source/PCG/Public/Tests/Determinism/PCGDeterminismSettings.h
new file mode 100644
--- /dev/null
+++ b/Engine/Plugins/Experimental/PCG/Source/PCG/Public/Tests/Determinism/PCGDeterminismSettings.h
@@ -0,0 +1,28 @@
+// Copyright Epic Games, Inc. All Rights Reserved.
+
+#pragma once
+
+#include "CoreMinimal.h"
+#include "Templates/SubclassOf.h"
+
+#include "PCGDeterminismSettings.generated.h"
+
+class UPCGDeterminismTestBlueprintBase;
+
+USTRUCT(BlueprintType)
+struct FPCGDeterminismSettings
+{
+	GENERATED_BODY()
+	
+	UPROPERTY(BlueprintReadWrite, EditAnywhere, Category = Determinism)
+	bool bBasicTests = true;
+
+	UPROPERTY(BlueprintReadWrite, EditAnywhere, Category = Determinism)
+	bool bOrderIndependenceTests = true;
+
+	UPROPERTY(BlueprintReadWrite, EditAnywhere, Category = Determinism)
+	bool bUseBlueprintDeterminismTest = false;
+
+	UPROPERTY(BlueprintReadOnly, EditAnywhere, Category = Determinism, meta = (EditCondition = "bUseBlueprintDeterminismTest"))
+	TSubclassOf<UPCGDeterminismTestBlueprintBase> DeterminismTestBlueprint;
+};
\ No newline at end of file
diff --git a/Engine/Plugins/Experimental/PCG/Source/PCG/Public/Tests/Determinism/PCGDeterminismTestBlueprintBase.h b/Engine/Plugins/Experimental/PCG/Source/PCG/Public/Tests/Determinism/PCGDeterminismTestBlueprintBase.h
new file mode 100644
--- /dev/null
+++ b/Engine/Plugins/Experimental/PCG/Source/PCG/Public/Tests/Determinism/PCGDeterminismTestBlueprintBase.h
@@ -0,0 +1,19 @@
+// Copyright Epic Games, Inc. All Rights Reserved.
+
+#pragma once
+
+#include "Tests/Determinism/PCGDeterminismTestsCommon.h"
+
+#include "PCGDeterminismTestBlueprintBase.generated.h"
+
+class UPCGNode;
+
+UCLASS(Abstract, BlueprintType, Blueprintable, hidecategories = (Object))
+class PCG_API UPCGDeterminismTestBlueprintBase : public UObject
+{
+	GENERATED_BODY()
+
+public:
+	UFUNCTION(BlueprintNativeEvent, BlueprintPure, CallInEditor, Category = Determinism)
+	void ExecuteTest(const UPCGNode* InPCGNode, UPARAM(ref)FDeterminismNodeTestResult& InOutTestResult);
+};
\ No newline at end of file
diff --git a/Engine/Plugins/Experimental/PCG/Source/PCG/Public/Tests/Determinism/PCGDeterminismTestsCommon.h b/Engine/Plugins/Experimental/PCG/Source/PCG/Public/Tests/Determinism/PCGDeterminismTestsCommon.h
--- a/Engine/Plugins/Experimental/PCG/Source/PCG/Public/Tests/Determinism/PCGDeterminismTestsCommon.h
+++ b/Engine/Plugins/Experimental/PCG/Source/PCG/Public/Tests/Determinism/PCGDeterminismTestsCommon.h
@@ -6,6 +6,8 @@
 #include "PCGSettings.h"
 #include "Tests/PCGTestsCommon.h"
 
+#include "PCGDeterminismTestsCommon.generated.h"
+
 class UPCGComponent;
 class UPCGData;
 class UPCGSpatialData;
@@ -13,6 +15,55 @@ class UPrimitiveComponent;
 class USplineComponent;
 struct FPCGDataCollection;
 
+// This will include multiple values of different meanings, but we use an enum to facilitate data passing
+UENUM()
+enum class EDeterminismLevel : uint8
+{
+	None = 0u,
+	NoDeterminism = None,
+	Basic,
+	OrderOrthogonal,
+	OrderConsistent,
+	OrderIndependent,
+	Deterministic = OrderIndependent
+};
+
+USTRUCT(BlueprintType)
+struct FDeterminismNodeTestResult
+{
+	GENERATED_BODY()
+
+	/** The node's title */
+	UPROPERTY(VisibleAnywhere, BlueprintReadOnly, Category = Determinism)
+	FName NodeTitle = TEXT("Untitled");
+
+	/** The node's name */
+	UPROPERTY(VisibleAnywhere, BlueprintReadOnly, Category = Determinism)
+	FString NodeName = TEXT("Unnamed");
+
+	// TODO: Add the seed to the UI widget
+	UPROPERTY(EditAnywhere, BlueprintReadWrite, Category = Determinism)
+	int32 Seed = -1;
+
+	/** BitFlags for which data types were tested */
+	UPROPERTY(EditAnywhere, BlueprintReadWrite, Category = Determinism)
+	EPCGDataType DataTypesTested = EPCGDataType::None;
+
+	/** A mapping of [test names : test results] */
+	UPROPERTY(EditAnywhere, BlueprintReadWrite, Category = Determinism)
+	TMap<FName, EDeterminismLevel> TestResults;
+
+	/** A mapping of [test name : additional details] */
+	UPROPERTY(EditAnywhere, BlueprintReadWrite, Category = Determinism)
+	TArray<FString> AdditionalDetails;
+
+	/** T/F whether a flag has been raised on this node's tests (for UI) */
+	UPROPERTY(EditAnywhere, BlueprintReadWrite, Category = Determinism)
+	bool bFlagRaised = false;
+
+	// TODO: Chrono Duration of how long the tests took
+};
+
 namespace PCGDeterminismTests
 {
 	namespace Defaults
@@ -48,37 +99,10 @@ namespace PCGDeterminismTests
 		EPCGDataType::Landscape
 	};
 
-	// This will include multiple values of different meanings, but we use an enum to facilitate data passing
-	UENUM()
-	enum class EDeterminismLevel : uint8
-	{
-		None = 0u,
-		NoDeterminism = None,
-		Basic,
-		OrderOrthogonal,
-		OrderConsistent,
-		OrderIndependent,
-		Deterministic = OrderIndependent
-	};
-
-	struct FNodeTestResult
-	{
-		int32 Index = -1;
-		FName NodeTitle = TEXT("Untitled");
-		FString NodeNameString = TEXT("Unnamed");
-		// TODO: Add the seed to the UI widget
-		int32 Seed = -1;
-		EPCGDataType DataTypesTested = EPCGDataType::None;
-		TMap<FName, EDeterminismLevel> TestResults;
-		TArray<FString> AdditionalDetails;
-		bool bFlagRaised = false;
-		// TODO: Chrono Duration of how long the tests took
-	};
-
 	/** A default delegate to report an unset test */
-	bool LogInvalidTest(const UPCGNode* InPCGNode, const FName& TestName, FNodeTestResult& OutResult);
+	bool LogInvalidTest(const UPCGNode* InPCGNode, const FName& TestName, FDeterminismNodeTestResult& OutResult);
 
-	typedef TFunction<bool(const UPCGNode* InPCGNode, const FName& TestName, FNodeTestResult& OutResult)> TestFunction;
+	typedef TFunction<bool(const UPCGNode* InPCGNode, const FName& TestName, FDeterminismNodeTestResult& OutResult)> TestFunction;
 
 	struct FNodeTestInfo
 	{
@@ -108,16 +132,13 @@ namespace PCGDeterminismTests
 	};
 
 	/** Validates if a PCGNode is deterministic */
-	PCG_API void RunDeterminismTest(const UPCGNode* InPCGNode, FNodeTestResult& OutResult, const FNodeTestInfo& TestToRun);
-
-	/** Adds the basic set of determinism tests to the passed in array */
-	PCG_API void RetrieveBasicTests(TArray<FNodeTestInfo>& OutBasicTests);
+	PCG_API void RunDeterminismTest(const UPCGNode* InPCGNode, FDeterminismNodeTestResult& OutResult, const FNodeTestInfo& TestToRun);
 
 	/** Validates all the generic determinism tests for any given node */
-	bool RunBasicTestSuite(const UPCGNode* InPCGNode, const FName& TestName, FNodeTestResult& OutResult);
+	PCG_API bool RunBasicTestSuite(const UPCGNode* InPCGNode, const FName& TestName, FDeterminismNodeTestResult& OutResult);
 
 	/** Validates the various levels of order independence for any given node */
-	bool RunOrderIndependenceSuite(const UPCGNode* InPCGNode, const FName& TestName, FNodeTestResult& OutResult);
+	PCG_API bool RunOrderIndependenceSuite(const UPCGNode* InPCGNode, const FName& TestName, FDeterminismNodeTestResult& OutResult);
 
 	/** Validates minimal node determinism against the same single test data */
 	bool RunBasicSelfTest(const FNodeAndOptions& NodeAndOptions);
@@ -232,7 +253,7 @@ namespace PCGDeterminismTests
 	EPCGDataType GetPermutation(int32 PermutationIteration, int32 PinIndex, const TArray<TArray<EPCGDataType>>& BaseOptionsPerPin);
 
 	/** Helper to update a tests' outgoing results */
-	void UpdateTestResults(FName TestName, FNodeTestResult& OutResult, EDeterminismLevel DeterminismLevel);
+	void UpdateTestResults(FName TestName, FDeterminismNodeTestResult& OutResult, EDeterminismLevel DeterminismLevel);
 
 	/** Gets a comparison function to compare two data objects */
 	TFunction<bool(const UPCGData*, const UPCGData*)> GetDataCompareFunction(EPCGDataType DataType, EDeterminismLevel DeterminismLevel);
@@ -325,4 +346,10 @@ namespace PCGDeterminismTests
 
 		Array = MoveTemp(TempArray);
 	}
+
+	namespace Defaults
+	{
+		static const FNodeTestInfo DeterminismBasicTestInfo = FNodeTestInfo(NSLOCTEXT("PCGDeterminism", "BasicTest", "Basic Test"), RunBasicTestSuite, 90.f);
+		static const FNodeTestInfo DeterminismOrderIndependenceInfo = FNodeTestInfo(NSLOCTEXT("PCGDeterminism", "OrderIndependenceTest", "Order Independence"), RunOrderIndependenceSuite, 140.f);
+	}
 }
\ No newline at end of file
diff --git a/Engine/Plugins/Experimental/PCG/Source/PCGEditor/Private/PCGEditor.cpp b/Engine/Plugins/Experimental/PCG/Source/PCGEditor/Private/PCGEditor.cpp
--- a/Engine/Plugins/Experimental/PCG/Source/PCGEditor/Private/PCGEditor.cpp
+++ b/Engine/Plugins/Experimental/PCG/Source/PCGEditor/Private/PCGEditor.cpp
@@ -9,6 +9,7 @@
 #include "PCGInputOutputSettings.h"
 #include "PCGPin.h"
 #include "PCGSubgraph.h"
+#include "Tests/Determinism/PCGDeterminismTestBlueprintBase.h"
 #include "Tests/Determinism/PCGDeterminismTestsCommon.h"
 
 #include "PCGEditorCommands.h"
@@ -380,7 +381,9 @@ void FPCGEditor::OnDeterminismTests()
 		return;
 	}
 
-	DeterminismWidget->Clear();
+	TMap<FName, FTestColumnInfo> TestsConducted;
+	DeterminismWidget->ClearItems();
+	DeterminismWidget->BuildBaseColumns();
 
 	int32 Index = 1;
 	for (UObject* Object : GraphEditorWidget->GetSelectedNodes())
@@ -392,20 +395,38 @@ void FPCGEditor::OnDeterminismTests()
 			if (const UPCGEditorGraphNodeBase* PCGEditorGraphNode = Cast<const UPCGEditorGraphNodeBase>(Object))
 			{
 				const UPCGNode* PCGNode = PCGEditorGraphNode->GetPCGNode();
-				check(PCGNode);
+				check(PCGNode && PCGNode->DefaultSettings);
 
-				TSharedPtr<PCGDeterminismTests::FNodeTestResult> NodeResult = MakeShared<PCGDeterminismTests::FNodeTestResult>();
-				check(NodeResult.IsValid());
-				NodeResult->Index = Index++;
+				// TODO: Check if the node has a native test
+
+				TSharedPtr<FDeterminismNodeTestResult> NodeResult = MakeShared<FDeterminismNodeTestResult>();
 				NodeResult->NodeTitle = PCGNode->GetNodeTitle();
-				NodeResult->NodeNameString = PCGNode->GetName();
+				NodeResult->NodeName = PCGNode->GetName();
+				NodeResult->Seed = PCGNode->DefaultSettings->Seed;
 
-				TArray<PCGDeterminismTests::FNodeTestInfo> BasicTests;
-				PCGDeterminismTests::RetrieveBasicTests(BasicTests);
+				if (PCGNode->DefaultSettings->DeterminismSettings.bBasicTests)
+				{
+					PCGDeterminismTests::FNodeTestInfo TestInfo = PCGDeterminismTests::Defaults::DeterminismBasicTestInfo;
+					PCGDeterminismTests::RunDeterminismTest(PCGNode, *NodeResult, TestInfo);
+					TestsConducted.FindOrAdd(TestInfo.TestName, {TestInfo.TestName, TestInfo.TestLabel, TestInfo.TestLabelWidth, HAlign_Center});
+				}
 
-				for (PCGDeterminismTests::FNodeTestInfo TestInfo : BasicTests)
+				if (PCGNode->DefaultSettings->DeterminismSettings.bOrderIndependenceTests)
 				{
+					PCGDeterminismTests::FNodeTestInfo TestInfo = PCGDeterminismTests::Defaults::DeterminismOrderIndependenceInfo;
 					PCGDeterminismTests::RunDeterminismTest(PCGNode, *NodeResult, TestInfo);
+					TestsConducted.FindOrAdd(TestInfo.TestName, {TestInfo.TestName, TestInfo.TestLabel, TestInfo.TestLabelWidth, HAlign_Center});
+				}
+
+				// Custom tests
+				if (PCGNode->DefaultSettings->DeterminismSettings.bUseBlueprintDeterminismTest)
+				{
+					TSubclassOf<UPCGDeterminismTestBlueprintBase> Blueprint = PCGNode->DefaultSettings->DeterminismSettings.DeterminismTestBlueprint;
+					Blueprint.GetDefaultObject()->ExecuteTest(PCGNode, *NodeResult);
+					FName BlueprintName(Blueprint->GetName());
+
+					// TODO: Add a dynamic tab size guess for the width
+					TestsConducted.FindOrAdd(BlueprintName, {BlueprintName, FText::FromString(Blueprint->GetName()), 140.f, HAlign_Center});
 				}
 
 				DeterminismWidget->AddItem(NodeResult);
@@ -413,6 +434,14 @@ void FPCGEditor::OnDeterminismTests()
 		}
 	}
 
+	for (const TTuple<FName, FTestColumnInfo>& Test : TestsConducted)
+	{
+		DeterminismWidget->AddColumn(Test.Value);
+	}
+
+	DeterminismWidget->AddDetailsColumn();
+	DeterminismWidget->RefreshItems();
+
 	// Give focus to the Determinism Output Tab
 	if (TabManager.IsValid())
 	{
@@ -1384,17 +1413,7 @@ TSharedRef<SPCGEditorGraphAttributeListView> FPCGEditor::CreateAttributesWidget(
 
 TSharedRef<SPCGEditorGraphDeterminismListView> FPCGEditor::CreateDeterminismWidget()
 {
-	TArray<PCGDeterminismTests::FNodeTestInfo> BasicTests;
-	PCGDeterminismTests::RetrieveBasicTests(BasicTests);
-
-	TArray<FTestColumnInfo> TestColumnInfo;
-	for (const PCGDeterminismTests::FNodeTestInfo& TestInfo : BasicTests)
-	{
-		// Build the columns based on the tests to run
-		TestColumnInfo.Emplace(FName(TestInfo.TestLabel.ToString()), TestInfo.TestLabel, TestInfo.TestLabelWidth, HAlign_Center);
-	}
-
-	return SNew(SPCGEditorGraphDeterminismListView, SharedThis(this), TestColumnInfo);
+	return SNew(SPCGEditorGraphDeterminismListView, SharedThis(this));
 }
 
 void FPCGEditor::OnSelectedNodesChanged(const TSet<UObject*>& NewSelection)
diff --git a/Engine/Plugins/Experimental/PCG/Source/PCGEditor/Private/SPCGEditorGraphDeterminism.cpp b/Engine/Plugins/Experimental/PCG/Source/PCGEditor/Private/SPCGEditorGraphDeterminism.cpp
--- a/Engine/Plugins/Experimental/PCG/Source/PCGEditor/Private/SPCGEditorGraphDeterminism.cpp
+++ b/Engine/Plugins/Experimental/PCG/Source/PCGEditor/Private/SPCGEditorGraphDeterminism.cpp
@@ -32,9 +32,10 @@ namespace
 	constexpr float ListViewRowHeight = 36.f;
 }
 
-void SPCGEditorGraphDeterminismRow::Construct(const FArguments& InArgs, const TSharedRef<STableViewBase>& InOwnerTableView, const FPCGNodeTestResultPtr& Item)
+void SPCGEditorGraphDeterminismRow::Construct(const FArguments& InArgs, const TSharedRef<STableViewBase>& InOwnerTableView, const FPCGNodeTestResultPtr& Item, int32 ItemIndex)
 {
 	CurrentItem = Item;
+	CurrentIndex = ItemIndex;
 	SMultiColumnTableRow<FPCGNodeTestResultPtr>::Construct(SMultiColumnTableRow::FArguments(), InOwnerTableView);
 }
 
@@ -42,7 +43,7 @@ TSharedRef<SWidget> SPCGEditorGraphDeterminismRow::GenerateWidgetForColumn(const
 {
 	check(CurrentItem);
 
-	FText CellText = LOCTEXT("UnknownColumn", "Unknown");
+	FText CellText;
 
 	auto ReturnColorCodedResultBlock = [this](const FText& OutCellText)
 	{
@@ -54,7 +55,7 @@ TSharedRef<SWidget> SPCGEditorGraphDeterminismRow::GenerateWidgetForColumn(const
 	// Permanent columns
 	if (ColumnId == NAME_Index)
 	{
-		return ReturnColorCodedResultBlock(FText::FromString(FString::FromInt(CurrentItem->Index)));
+		return ReturnColorCodedResultBlock(FText::FromString(FString::FromInt(CurrentIndex)));
 	}
 	else if (ColumnId == NAME_NodeTitle)
 	{
@@ -62,7 +63,7 @@ TSharedRef<SWidget> SPCGEditorGraphDeterminismRow::GenerateWidgetForColumn(const
 	}
 	else if (ColumnId == NAME_NodeName)
 	{
-		CellText = FText::FromString(CurrentItem->NodeNameString);
+		CellText = FText::FromString(CurrentItem->NodeName);
 	}
 	else if (ColumnId == NAME_DataTypesTested)
 	{
@@ -75,29 +76,28 @@ TSharedRef<SWidget> SPCGEditorGraphDeterminismRow::GenerateWidgetForColumn(const
 		FString FullDetails = FString::Join(CurrentItem->AdditionalDetails, TEXT(", "));
 		CellText = FText::FromString(FullDetails);
 	}
-
 	// Test columns
-	if (PCGDeterminismTests::EDeterminismLevel* DeterminismLevel = CurrentItem->TestResults.Find(ColumnId))
+	else if (EDeterminismLevel* DeterminismLevel = CurrentItem->TestResults.Find(ColumnId))
 	{
 		switch (*DeterminismLevel)
 		{
-		case PCGDeterminismTests::EDeterminismLevel::OrderOrthogonal:
+		case EDeterminismLevel::OrderOrthogonal:
 			return SNew(STextBlock)
 				.Text(TEXT_Orthogonal)
 				.ColorAndOpacity(FColor::Orange);
-		case PCGDeterminismTests::EDeterminismLevel::OrderConsistent:
+		case EDeterminismLevel::OrderConsistent:
 			return SNew(STextBlock)
 				.Text(TEXT_Consistent)
 				.ColorAndOpacity(FColor::Yellow);
-		case PCGDeterminismTests::EDeterminismLevel::OrderIndependent:
+		case EDeterminismLevel::OrderIndependent:
 			return SNew(STextBlock)
 				.Text(TEXT_Independent)
 				.ColorAndOpacity(FColor::Green);
-		case PCGDeterminismTests::EDeterminismLevel::Basic:
+		case EDeterminismLevel::Basic:
 			return SNew(STextBlock)
 				.Text(TEXT_Basic)
 				.ColorAndOpacity(FColor::Turquoise);
-		case PCGDeterminismTests::EDeterminismLevel::NoDeterminism:
+		case EDeterminismLevel::NoDeterminism:
 		default:
 			return SNew(STextBlock)
 				.Text(TEXT_NotDeterministic)
@@ -109,39 +109,12 @@ TSharedRef<SWidget> SPCGEditorGraphDeterminismRow::GenerateWidgetForColumn(const
 		.Text(CellText);
 }
 
-void SPCGEditorGraphDeterminismListView::Construct(const FArguments& InArgs, TWeakPtr<FPCGEditor> InPCGEditor, const TArray<FTestColumnInfo>& InTestColumns)
+void SPCGEditorGraphDeterminismListView::Construct(const FArguments& InArgs, TWeakPtr<FPCGEditor> InPCGEditor)
 {
 	check(InPCGEditor.IsValid() && !bIsConstructed);
 	PCGEditorPtr = InPCGEditor;
 
-	TSharedRef<SHeaderRow> GeneratedHeaderRow = SNew(SHeaderRow);
-
-	TArray<FTestColumnInfo> TestColumnInfo;
-	// Permanent columns
-	TestColumnInfo.Emplace(NAME_Index, TEXT_Index, SmallManualWidth, HAlign_Center);
-	TestColumnInfo.Emplace(NAME_NodeTitle, TEXT_NodeTitle, LargeManualWidth, HAlign_Left);
-	TestColumnInfo.Emplace(NAME_NodeName, TEXT_NodeName, LargeManualWidth, HAlign_Left);
-	TestColumnInfo.Emplace(NAME_DataTypesTested, TEXT_DataTypesTested, MediumManualWidth, HAlign_Center);
-
-	// Test columns
-	TestColumnInfo.Append(InTestColumns);
-
-	// Final details column
-	TestColumnInfo.Emplace(NAME_AdditionalDetails, TEXT_AdditionalDetails, 0.f, HAlign_Left);
-
-	// Build column arguments for test columns dynamically
-	for (const FTestColumnInfo& ColumnInfo : TestColumnInfo)
-	{
-		SHeaderRow::FColumn::FArguments Arguments;
-		Arguments.ColumnId(ColumnInfo.ColumnID);
-		Arguments.DefaultLabel(ColumnInfo.ColumnLabel);
-		if (ColumnInfo.Width > 0.f)
-		{
-			Arguments.ManualWidth(ColumnInfo.Width);
-		}
-		Arguments.HAlignCell(ColumnInfo.HAlign);
-		GeneratedHeaderRow->AddColumn(Arguments);
-	}
+	GeneratedHeaderRow = SNew(SHeaderRow);
 
 	SAssignNew(ListView, SListView<FPCGNodeTestResultPtr>)
 		.ListItemsSource(&ListViewItems)
@@ -161,28 +134,62 @@ void SPCGEditorGraphDeterminismListView::AddItem(const FPCGNodeTestResultPtr& It
 {
 	check(Item.IsValid());
 	ListViewItems.Emplace(Item);
-	Refresh();
 }
 
-void SPCGEditorGraphDeterminismListView::Clear()
+void SPCGEditorGraphDeterminismListView::ClearItems()
 {
+	ItemIndexCounter = -1;
 	ListViewItems.Empty();
-	Refresh();
+	RefreshItems();
 }
 
-void SPCGEditorGraphDeterminismListView::Refresh()
+void SPCGEditorGraphDeterminismListView::RefreshItems()
 {
 	ListView->RequestListRefresh();
 }
 
+void SPCGEditorGraphDeterminismListView::AddColumn(const FTestColumnInfo& ColumnInfo)
+{
+	SHeaderRow::FColumn::FArguments Arguments;
+	Arguments.ColumnId(ColumnInfo.ColumnID);
+	Arguments.DefaultLabel(ColumnInfo.ColumnLabel);
+	if (ColumnInfo.Width > 0.f)
+	{
+		Arguments.ManualWidth(ColumnInfo.Width);
+	}
+	Arguments.HAlignCell(ColumnInfo.HAlign);
+	GeneratedHeaderRow->AddColumn(Arguments);
+}
+
+void SPCGEditorGraphDeterminismListView::BuildBaseColumns()
+{
+	ClearColumns();
+	// Permanent columns
+	AddColumn({NAME_Index, TEXT_Index, SmallManualWidth, HAlign_Center});
+	AddColumn({NAME_NodeTitle, TEXT_NodeTitle, LargeManualWidth, HAlign_Left});
+	AddColumn({NAME_NodeName, TEXT_NodeName, LargeManualWidth, HAlign_Left});
+	AddColumn({NAME_DataTypesTested, TEXT_DataTypesTested, MediumManualWidth, HAlign_Center});
+}
+
+void SPCGEditorGraphDeterminismListView::AddDetailsColumn()
+{
+	// Final details column
+	AddColumn({NAME_AdditionalDetails, TEXT_AdditionalDetails, 0.f, HAlign_Left});
+}
+
 bool SPCGEditorGraphDeterminismListView::WidgetIsConstructed() const
 {
 	return bIsConstructed;
 }
 
+void SPCGEditorGraphDeterminismListView::ClearColumns()
+{
+	GeneratedHeaderRow->ClearColumns();
+}
+
 TSharedRef<ITableRow> SPCGEditorGraphDeterminismListView::OnGenerateRow(const FPCGNodeTestResultPtr Item, const TSharedRef<STableViewBase>& OwnerTable) const
 {
-	return SNew(SPCGEditorGraphDeterminismRow, OwnerTable, Item);
+	return SNew(SPCGEditorGraphDeterminismRow, OwnerTable, Item, ++ItemIndexCounter);
 }
 
 #undef LOCTEXT_NAMESPACE
\ No newline at end of file
diff --git a/Engine/Plugins/Experimental/PCG/Source/PCGEditor/Private/SPCGEditorGraphDeterminism.h b/Engine/Plugins/Experimental/PCG/Source/PCGEditor/Private/SPCGEditorGraphDeterminism.h
--- a/Engine/Plugins/Experimental/PCG/Source/PCGEditor/Private/SPCGEditorGraphDeterminism.h
+++ b/Engine/Plugins/Experimental/PCG/Source/PCGEditor/Private/SPCGEditorGraphDeterminism.h
@@ -7,7 +7,7 @@
 
 class FPCGEditor;
 
-typedef TSharedPtr<PCGDeterminismTests::FNodeTestResult> FPCGNodeTestResultPtr;
+typedef TSharedPtr<FDeterminismNodeTestResult> FPCGNodeTestResultPtr;
 
 struct FTestColumnInfo
 {
@@ -29,13 +29,14 @@ class SPCGEditorGraphDeterminismRow final : public SMultiColumnTableRow<FPCGNode
 	SLATE_END_ARGS()
 
 	/** Construct a row of the ListView */
-	void Construct(const FArguments& InArgs, const TSharedRef<STableViewBase>& InOwnerTableView, const FPCGNodeTestResultPtr& Item);
+	void Construct(const FArguments& InArgs, const TSharedRef<STableViewBase>& InOwnerTableView, const FPCGNodeTestResultPtr& Item, int32 ItemIndex);
 
 	/** Generates a column, given the column's ID */
 	virtual TSharedRef<SWidget> GenerateWidgetForColumn(const FName& ColumnId) override;
 
 protected:
 	FPCGNodeTestResultPtr CurrentItem;
+	int32 CurrentIndex = -1;
 };
 
 class SPCGEditorGraphDeterminismListView final : public SCompoundWidget
@@ -45,19 +46,29 @@ public:
 	SLATE_END_ARGS()
 
 	/** Construct the ListView */
-	void Construct(const FArguments& InArgs, TWeakPtr<FPCGEditor> InPCGEditor, const TArray<FTestColumnInfo>& InTestColumn);
+	void Construct(const FArguments& InArgs, TWeakPtr<FPCGEditor> InPCGEditor);
 	/** Add an item to the ListView */
 	void AddItem(const FPCGNodeTestResultPtr& Item);
 	/** Clear all items from the ListView */
-	void Clear();
+	void ClearItems();
 	/** Refreshes the ListView */
-	void Refresh();
+	void RefreshItems();
+	/** Adds a test column to the ListView */
+	void AddColumn(const FTestColumnInfo& ColumnInfo);
+	/** Builds the default columns for the widget */
+	void BuildBaseColumns();
+	/** Adds the additional details column */
+	void AddDetailsColumn();
+
 	/** Validates if the ListView has been constructed */
 	bool WidgetIsConstructed() const;
 
 private:
+	/** Clears the columns of the widget */
+	void ClearColumns();
 	/** Generate the row widget */
 	TSharedRef<ITableRow> OnGenerateRow(const FPCGNodeTestResultPtr Item, const TSharedRef<STableViewBase>& OwnerTable) const;
+	TSharedPtr<SHeaderRow> GeneratedHeaderRow;
 
 	TWeakPtr<FPCGEditor> PCGEditorPtr;
 
@@ -65,4 +76,5 @@ private:
 	TArray<FPCGNodeTestResultPtr> ListViewItems;
 
 	bool bIsConstructed = false;
+	mutable int32 ItemIndexCounter = -1;
 };
\ No newline at end of file
