Bypass some cleanup in PCGComponent if IsEngineExitRequested()

diff --git a/Engine/Plugins/Experimental/PCG/Source/PCG/Private/PCGComponent.cpp b/Engine/Plugins/Experimental/PCG/Source/PCG/Private/PCGComponent.cpp
--- a/Engine/Plugins/Experimental/PCG/Source/PCG/Private/PCGComponent.cpp
+++ b/Engine/Plugins/Experimental/PCG/Source/PCG/Private/PCGComponent.cpp
@@ -637,11 +637,15 @@ void UPCGComponent::BeginDestroy()
 	if (Graph)
 	{
 		Graph->OnGraphChangedDelegate.RemoveAll(this);
+		Graph = nullptr;
 	}
 
-	TeardownLandscapeTracking();
-	TeardownTrackingCallbacks();
-	TeardownActorCallbacks();
+	if (!IsEngineExitRequested())
+	{
+		TeardownLandscapeTracking();
+		TeardownTrackingCallbacks();
+		TeardownActorCallbacks();
+	}
 #endif
 
 	Super::BeginDestroy();
