Fixes and optimizations around metadata usage/locks.

diff --git a/Engine/Plugins/Experimental/PCG/Source/PCG/Private/Data/PCGDifferenceData.cpp b/Engine/Plugins/Experimental/PCG/Source/PCG/Private/Data/PCGDifferenceData.cpp
--- a/Engine/Plugins/Experimental/PCG/Source/PCG/Private/Data/PCGDifferenceData.cpp
+++ b/Engine/Plugins/Experimental/PCG/Source/PCG/Private/Data/PCGDifferenceData.cpp
@@ -105,6 +105,7 @@ FBox UPCGDifferenceData::GetStrictBounds() const
 
 bool UPCGDifferenceData::SamplePoint(const FTransform& InTransform, const FBox& InBounds, FPCGPoint& OutPoint, UPCGMetadata* OutMetadata) const
 {
+	//TRACE_CPUPROFILER_EVENT_SCOPE(UPCGDifferenceData::SamplePoint);
 	check(Source);
 
 	FPCGPoint PointFromSource;
@@ -117,14 +118,14 @@ bool UPCGDifferenceData::SamplePoint(const FTransform& InTransform, const FBox&
 
 	FPCGPoint PointFromDiff;
 	// Important note: here we will not use the point we got from the source, otherwise we are introducing severe bias
-	if (Difference && Difference->SamplePoint(InTransform, InBounds, PointFromDiff, OutMetadata))
+	if (Difference && Difference->SamplePoint(InTransform, InBounds, PointFromDiff, (bDiffMetadata ? OutMetadata : nullptr)))
 	{
 		const bool bBinaryDensity = (DensityFunction == EPCGDifferenceDensityFunction::Binary);
 		
 		// Apply difference
 		OutPoint.Density = bBinaryDensity ? 0 : FMath::Max(0, PointFromSource.Density - PointFromDiff.Density);
 		// Color?
-		if (OutMetadata && OutPoint.Density > 0)
+		if (bDiffMetadata && OutMetadata && OutPoint.Density > 0 && PointFromDiff.MetadataEntry != PCGInvalidEntryKey)
 		{
 			OutMetadata->MergePointAttributesSubset(PointFromSource, OutMetadata, Source->Metadata, PointFromDiff, OutMetadata, Difference->Metadata, OutPoint, EPCGMetadataOp::Sub);
 		}
@@ -168,22 +169,29 @@ const UPCGPointData* UPCGDifferenceData::CreatePointData(FPCGContext* Context) c
 	const TArray<FPCGPoint>& SourcePoints = SourcePointData->GetPoints();
 	TArray<FPCGPoint>& TargetPoints = Data->GetMutablePoints();
 
-	FPCGAsync::AsyncPointProcessing(Context, SourcePoints.Num(), TargetPoints, [this, Data, SourcePointData, &SourcePoints](int32 Index, FPCGPoint& OutPoint)
+	UPCGMetadata* TempDiffMetadata = nullptr;
+	if (bDiffMetadata && Data->Metadata && Difference->Metadata)
+	{
+		TempDiffMetadata = NewObject<UPCGMetadata>();
+		TempDiffMetadata->Initialize(Difference->Metadata);
+	}
+
+	FPCGAsync::AsyncPointProcessing(Context, SourcePoints.Num(), TargetPoints, [this, Data, SourcePointData, TempDiffMetadata, &SourcePoints](int32 Index, FPCGPoint& OutPoint)
 	{
 		const FPCGPoint& Point = SourcePoints[Index];
 
 		FPCGPoint PointFromDiff;
-		if (Difference && Difference->SamplePoint(Point.Transform, Point.GetLocalBounds(), PointFromDiff, Data->Metadata))
+		if (Difference && Difference->SamplePoint(Point.Transform, Point.GetLocalBounds(), PointFromDiff, TempDiffMetadata))
 		{
 			const bool bBinaryDensity = (DensityFunction == EPCGDifferenceDensityFunction::Binary);
 
 			OutPoint = Point;
-			UPCGMetadataAccessorHelpers::InitializeMetadata(OutPoint, Data->Metadata, Point);
+			//UPCGMetadataAccessorHelpers::InitializeMetadata(OutPoint, Data->Metadata, Point);
 			OutPoint.Density = bBinaryDensity ? 0 : FMath::Max(0, Point.Density - PointFromDiff.Density);
 
-			if (Data->Metadata && OutPoint.Density > 0)
+			if (TempDiffMetadata && OutPoint.Density > 0 && PointFromDiff.MetadataEntry != PCGInvalidEntryKey)
 			{
-				Data->Metadata->MergePointAttributesSubset(Point, SourcePointData->Metadata, SourcePointData->Metadata, PointFromDiff, Data->Metadata, Difference->Metadata, OutPoint, EPCGMetadataOp::Sub);
+				Data->Metadata->MergePointAttributesSubset(Point, SourcePointData->Metadata, SourcePointData->Metadata, PointFromDiff, TempDiffMetadata, TempDiffMetadata, OutPoint, EPCGMetadataOp::Sub);
 			}
 
 #if WITH_EDITORONLY_DATA
@@ -195,7 +203,7 @@ const UPCGPointData* UPCGDifferenceData::CreatePointData(FPCGContext* Context) c
 		else
 		{
 			OutPoint = Point;
-			UPCGMetadataAccessorHelpers::InitializeMetadata(OutPoint, Data->Metadata, Point);
+			//UPCGMetadataAccessorHelpers::InitializeMetadata(OutPoint, Data->Metadata, Point);
 			return true;
 		}
 	});
diff --git a/Engine/Plugins/Experimental/PCG/Source/PCG/Private/Data/PCGIntersectionData.cpp b/Engine/Plugins/Experimental/PCG/Source/PCG/Private/Data/PCGIntersectionData.cpp
--- a/Engine/Plugins/Experimental/PCG/Source/PCG/Private/Data/PCGIntersectionData.cpp
+++ b/Engine/Plugins/Experimental/PCG/Source/PCG/Private/Data/PCGIntersectionData.cpp
@@ -65,6 +65,7 @@ FBox UPCGIntersectionData::GetStrictBounds() const
 
 bool UPCGIntersectionData::SamplePoint(const FTransform& InTransform, const FBox& InBounds, FPCGPoint& OutPoint, UPCGMetadata* OutMetadata) const
 {
+	//TRACE_CPUPROFILER_EVENT_SCOPE(UPCGIntersectionData::SamplePoint);
 	check(A && B);
 	const UPCGSpatialData* X = (A->HasNonTrivialTransform() || !B->HasNonTrivialTransform()) ? A : B;
 	const UPCGSpatialData* Y = (X == A) ? B : A;
@@ -88,7 +89,18 @@ bool UPCGIntersectionData::SamplePoint(const FTransform& InTransform, const FBox
 
 	if (OutMetadata)
 	{
-		OutMetadata->MergePointAttributesSubset(PointFromX, OutMetadata, X->Metadata, PointFromY, OutMetadata, Y->Metadata, OutPoint, EPCGMetadataOp::Min);
+		if (PointFromX.MetadataEntry != PCGInvalidEntryKey && PointFromY.MetadataEntry != PCGInvalidEntryKey)
+		{
+			OutMetadata->MergePointAttributesSubset(PointFromX, OutMetadata, X->Metadata, PointFromY, OutMetadata, Y->Metadata, OutPoint, EPCGMetadataOp::Min);
+		}
+		else if (PointFromX.MetadataEntry != PCGInvalidEntryKey)
+		{
+			OutPoint.MetadataEntry = PointFromX.MetadataEntry;
+		}
+		else
+		{
+			OutPoint.MetadataEntry = PointFromY.MetadataEntry;
+		}
 	}
 
 	return true;
@@ -136,30 +148,45 @@ UPCGPointData* UPCGIntersectionData::CreateAndFilterPointData(FPCGContext* Conte
 	Data->InitializeFromData(this, SourcePointData->Metadata);
 	Data->Metadata->AddAttributes(Y->Metadata);
 
+	UPCGMetadata* TempYMetadata = nullptr;
+	if (Y->Metadata)
+	{
+		TempYMetadata = Y->Metadata ? NewObject<UPCGMetadata>() : nullptr;
+		TempYMetadata->Initialize(Y->Metadata);
+	}
+
 	TArray<FPCGPoint>& TargetPoints = Data->GetMutablePoints();
 
-	FPCGAsync::AsyncPointProcessing(Context, SourcePoints.Num(), TargetPoints, [this, Data, SourcePointData, &SourcePoints, Y](int32 Index, FPCGPoint& OutPoint)
+	FPCGAsync::AsyncPointProcessing(Context, SourcePoints.Num(), TargetPoints, [this, Data, SourcePointData, &SourcePoints, Y, TempYMetadata](int32 Index, FPCGPoint& OutPoint)
 	{
+		TRACE_CPUPROFILER_EVENT_SCOPE(UPCGIntersectionData::CreateAndFilterPointData::Iteration);
 		const FPCGPoint& Point = SourcePoints[Index];
 
 		FPCGPoint PointFromY;
 #if WITH_EDITORONLY_DATA
-		if (!Y->SamplePoint(Point.Transform, Point.GetLocalBounds(), PointFromY, Data->Metadata) && !bKeepZeroDensityPoints)
+		if (!Y->SamplePoint(Point.Transform, Point.GetLocalBounds(), PointFromY, TempYMetadata) && !bKeepZeroDensityPoints)
 #else
-		if (!Y->SamplePoint(Point.Transform, Point.GetLocalBounds(), PointFromY, Data->Metadata))
+		if (!Y->SamplePoint(Point.Transform, Point.GetLocalBounds(), PointFromY, TempYMetadata))
 #endif
 		{
 			return false;
 		}
 
 		OutPoint = Point;
-		UPCGMetadataAccessorHelpers::InitializeMetadata(OutPoint, Data->Metadata, Point);
+		//UPCGMetadataAccessorHelpers::InitializeMetadata(OutPoint, Data->Metadata, Point);
 		OutPoint.Density = PCGIntersectionDataMaths::ComputeDensity(Point.Density, PointFromY.Density, DensityFunction);
 		OutPoint.Color = Point.Color * PointFromY.Color;
 
 		if (Data->Metadata)
 		{
-			Data->Metadata->MergePointAttributesSubset(Point, SourcePointData->Metadata, SourcePointData->Metadata, PointFromY, Data->Metadata, Y->Metadata, OutPoint, EPCGMetadataOp::Min);
+			if (Point.MetadataEntry != PCGInvalidEntryKey && PointFromY.MetadataEntry != PCGInvalidEntryKey)
+			{
+				Data->Metadata->MergePointAttributesSubset(Point, SourcePointData->Metadata, SourcePointData->Metadata, PointFromY, TempYMetadata, TempYMetadata, OutPoint, EPCGMetadataOp::Min);
+			}
+			else if (PointFromY.MetadataEntry != PCGInvalidEntryKey)
+			{
+				OutPoint.MetadataEntry = PointFromY.MetadataEntry;
+			}
 		}
 
 		return true;
diff --git a/Engine/Plugins/Experimental/PCG/Source/PCG/Private/Data/PCGLandscapeData.cpp b/Engine/Plugins/Experimental/PCG/Source/PCG/Private/Data/PCGLandscapeData.cpp
--- a/Engine/Plugins/Experimental/PCG/Source/PCG/Private/Data/PCGLandscapeData.cpp
+++ b/Engine/Plugins/Experimental/PCG/Source/PCG/Private/Data/PCGLandscapeData.cpp
@@ -66,9 +66,8 @@ bool UPCGLandscapeData::SamplePoint(const FTransform& InTransform, const FBox& I
 		return false;
 	}
 
-	const FTransform LocalTransform = InTransform * Transform.Inverse();
-	const FVector LocalPoint = LocalTransform.GetLocation();
-
+	// TODO: compute full transform when we want to support bounds
+	const FVector LocalPoint = Transform.InverseTransformPosition(InTransform.GetLocation());
 	const FIntPoint ComponentMapKey(FMath::FloorToInt(LocalPoint.X / LandscapeInfo->ComponentSizeQuads), FMath::FloorToInt(LocalPoint.Y / LandscapeInfo->ComponentSizeQuads));
 
 	if (ULandscapeComponent* LandscapeComponent = LandscapeInfo->XYtoComponentMap.FindRef(ComponentMapKey))
diff --git a/Engine/Plugins/Experimental/PCG/Source/PCG/Private/Data/PCGPointData.cpp b/Engine/Plugins/Experimental/PCG/Source/PCG/Private/Data/PCGPointData.cpp
--- a/Engine/Plugins/Experimental/PCG/Source/PCG/Private/Data/PCGPointData.cpp
+++ b/Engine/Plugins/Experimental/PCG/Source/PCG/Private/Data/PCGPointData.cpp
@@ -108,9 +108,12 @@ namespace PCGPointHelpers
 		if (OutMetadata && SourceMetadata && SourceMetadata->GetAttributeCount() > 0)
 		{
 			UPCGMetadataAccessorHelpers::InitializeMetadataWithParent(OutPoint, OutMetadata, ((Ratio <= 0.5f) ? A : B), SourceMetadata);
-			OutMetadata->ResetPointWeightedAttributes(OutPoint);
-			OutMetadata->AccumulatePointWeightedAttributes(A, SourceMetadata, Ratio, (Ratio <= 0.5f), OutPoint);
-			OutMetadata->AccumulatePointWeightedAttributes(B, SourceMetadata, 1.0f-Ratio, (Ratio > 0.5f), OutPoint);
+
+			TArray<TPair<const FPCGPoint*, float>, TInlineAllocator<2>> WeightedPoints;
+			WeightedPoints.Emplace(&A, Ratio);
+			WeightedPoints.Emplace(&B, 1.0f - Ratio);
+
+			OutMetadata->ComputePointWeightedAttribute(OutPoint, MakeArrayView(WeightedPoints), SourceMetadata);
 		}
 	}
 
@@ -299,12 +302,13 @@ FPCGPoint UPCGPointData::GetPoint(int32 Index) const
 
 bool UPCGPointData::SamplePoint(const FTransform& InTransform, const FBox& InBounds, FPCGPoint& OutPoint, UPCGMetadata* OutMetadata) const
 {
+	//TRACE_CPUPROFILER_EVENT_SCOPE(UPCGPointData::SamplePoint);
 	if (bOctreeIsDirty)
 	{
 		RebuildOctree();
 	}
 
-	TArray<TPair<const FPCGPoint*, float>> Contributions;
+	TArray<TPair<const FPCGPoint*, float>, TInlineAllocator<4>> Contributions;
 	const bool bSampleInVolume = (InBounds.GetExtent() != FVector::ZeroVector);
 
 	if (!bSampleInVolume)
@@ -342,6 +346,8 @@ bool UPCGPointData::SamplePoint(const FTransform& InTransform, const FBox& InBou
 		return false;
 	}
 
+	TArray<TPair<const FPCGPoint*, float>, TInlineAllocator<4>> ContributionsForMetadata;
+
 	const FVector::FReal TransformedBoundsVolume = InBounds.TransformBy(InTransform).GetVolume();
 	auto ComputeDensityContribution = [TransformedBoundsVolume](float VolumeIntersection)
 	{
@@ -380,6 +386,8 @@ bool UPCGPointData::SamplePoint(const FTransform& InTransform, const FBox& InBou
 		WeightedBoundsMax += SourcePoint.BoundsMax * Weight;
 		WeightedColor += SourcePoint.Color * Weight;
 		WeightedSteepness += SourcePoint.Steepness * Weight;
+
+		ContributionsForMetadata.Emplace(Contribution.Key, Weight);
 	}
 
 	// Finally, apply changes to point
@@ -396,16 +404,12 @@ bool UPCGPointData::SamplePoint(const FTransform& InTransform, const FBox& InBou
 
 	if (OutMetadata)
 	{
+		//TRACE_CPUPROFILER_EVENT_SCOPE(UPCGPointData::SamplePoint::SetupMetadata);
 		UPCGMetadataAccessorHelpers::InitializeMetadataWithParent(OutPoint, OutMetadata, *MaxContributor, Metadata);
-		OutMetadata->ResetPointWeightedAttributes(OutPoint);
 
-		for (const TPair<const FPCGPoint*, float> Contribution : Contributions)
+		if (ContributionsForMetadata.Num() > 1)
 		{
-			const FPCGPoint& SourcePoint = *Contribution.Key;
-			const float Weight = Contribution.Value / SumContributions;
-			const bool bIsMaxContributor = (Contribution.Key == MaxContributor);
-
-			OutMetadata->AccumulatePointWeightedAttributes(SourcePoint, Metadata, Weight, bIsMaxContributor, OutPoint);
+			OutMetadata->ComputePointWeightedAttribute(OutPoint, MakeArrayView(ContributionsForMetadata), Metadata);
 		}
 	}
 
diff --git a/Engine/Plugins/Experimental/PCG/Source/PCG/Private/Data/PCGProjectionData.cpp b/Engine/Plugins/Experimental/PCG/Source/PCG/Private/Data/PCGProjectionData.cpp
--- a/Engine/Plugins/Experimental/PCG/Source/PCG/Private/Data/PCGProjectionData.cpp
+++ b/Engine/Plugins/Experimental/PCG/Source/PCG/Private/Data/PCGProjectionData.cpp
@@ -73,6 +73,7 @@ FBox UPCGProjectionData::ProjectBounds(const FBox& InBounds) const
 
 bool UPCGProjectionData::SamplePoint(const FTransform& InTransform, const FBox& InBounds, FPCGPoint& OutPoint, UPCGMetadata* OutMetadata) const
 {
+	//TRACE_CPUPROFILER_EVENT_SCOPE(UPCGProjectionData::SamplePoint);
 	FPCGPoint PointFromSource;
 	if (!Source->SamplePoint(InTransform, InBounds, PointFromSource, OutMetadata))
 	{
@@ -91,7 +92,7 @@ bool UPCGProjectionData::SamplePoint(const FTransform& InTransform, const FBox&
 	OutPoint.Density *= PointFromTarget.Density;
 	OutPoint.Color *= PointFromTarget.Color;
 	
-	if (OutMetadata)
+	if (OutMetadata && PointFromTarget.MetadataEntry != PCGInvalidEntryKey)
 	{
 		//METADATA TODO Review op
 		OutMetadata->MergePointAttributesSubset(PointFromSource, OutMetadata, Source->Metadata, PointFromTarget, OutMetadata, Target->Metadata, OutPoint, EPCGMetadataOp::Max);
@@ -118,17 +119,24 @@ const UPCGPointData* UPCGProjectionData::CreatePointData(FPCGContext* Context) c
 	check(PointData->Metadata);
 	PointData->Metadata->AddAttributes(Target->Metadata);
 
+	UPCGMetadata* TempTargetMetadata = nullptr;
+	if (Target->Metadata)
+	{
+		TempTargetMetadata = NewObject<UPCGMetadata>();
+		TempTargetMetadata->Initialize(Target->Metadata);
+	}
+
 	TArray<FPCGPoint>& Points = PointData->GetMutablePoints();
 
-	FPCGAsync::AsyncPointProcessing(Context, SourcePoints.Num(), Points, [this, SourcePointData, PointData, &SourcePoints](int32 Index, FPCGPoint& OutPoint)
+	FPCGAsync::AsyncPointProcessing(Context, SourcePoints.Num(), Points, [this, SourcePointData, PointData, TempTargetMetadata, &SourcePoints](int32 Index, FPCGPoint& OutPoint)
 	{
 		const FPCGPoint& SourcePoint = SourcePoints[Index];
 
 		FPCGPoint PointFromTarget;
 #if WITH_EDITORONLY_DATA
-		if (!Target->SamplePoint(SourcePoint.Transform, SourcePoint.GetLocalBounds(), PointFromTarget, PointData->Metadata) && !bKeepZeroDensityPoints)
+		if (!Target->SamplePoint(SourcePoint.Transform, SourcePoint.GetLocalBounds(), PointFromTarget, TempTargetMetadata) && !bKeepZeroDensityPoints)
 #else
-		if (!Target->SamplePoint(SourcePoint.Transform, SourcePoint.GetLocalBounds(), PointFromTarget, PointData->Metadata))
+		if (!Target->SamplePoint(SourcePoint.Transform, SourcePoint.GetLocalBounds(), PointFromTarget, TempTargetMetadata))
 #endif
 		{
 			return false;
@@ -141,10 +149,10 @@ const UPCGPointData* UPCGProjectionData::CreatePointData(FPCGContext* Context) c
 		OutPoint.Density *= PointFromTarget.Density;
 		OutPoint.Color *= PointFromTarget.Color;
 
-		if (PointData->Metadata)
+		if (PointData->Metadata && TempTargetMetadata && PointFromTarget.MetadataEntry != PCGInvalidEntryKey)
 		{
 			//METADATA TODO review op
-			PointData->Metadata->MergePointAttributesSubset(SourcePoint, SourcePointData->Metadata, SourcePointData->Metadata, PointFromTarget, PointData->Metadata, Target->Metadata, OutPoint, EPCGMetadataOp::Max);
+			PointData->Metadata->MergePointAttributesSubset(SourcePoint, SourcePointData->Metadata, SourcePointData->Metadata, PointFromTarget, TempTargetMetadata, TempTargetMetadata, OutPoint, EPCGMetadataOp::Max);
 		}
 
 		return true;
diff --git a/Engine/Plugins/Experimental/PCG/Source/PCG/Private/Data/PCGSplineData.cpp b/Engine/Plugins/Experimental/PCG/Source/PCG/Private/Data/PCGSplineData.cpp
--- a/Engine/Plugins/Experimental/PCG/Source/PCG/Private/Data/PCGSplineData.cpp
+++ b/Engine/Plugins/Experimental/PCG/Source/PCG/Private/Data/PCGSplineData.cpp
@@ -219,8 +219,15 @@ bool UPCGSplineProjectionData::SamplePoint(const FTransform& InTransform, const
 
 			if (OutMetadata)
 			{
-				//TODO review op
-				OutMetadata->MergePointAttributesSubset(SplinePoint, OutMetadata, GetSpline()->Metadata, SurfacePoint, OutMetadata, GetSurface()->Metadata, OutPoint, EPCGMetadataOp::Max);
+				if (SplinePoint.MetadataEntry != PCGInvalidEntryKey && SurfacePoint.MetadataEntry)
+				{
+					//TODO review op
+					OutMetadata->MergePointAttributesSubset(SplinePoint, OutMetadata, GetSpline()->Metadata, SurfacePoint, OutMetadata, GetSurface()->Metadata, OutPoint, EPCGMetadataOp::Max);
+				}
+				else if (SurfacePoint.MetadataEntry != PCGInvalidEntryKey)
+				{
+					OutPoint.MetadataEntry = SurfacePoint.MetadataEntry;
+				}
 			}
 
 			return true;
diff --git a/Engine/Plugins/Experimental/PCG/Source/PCG/Private/Data/PCGUnionData.cpp b/Engine/Plugins/Experimental/PCG/Source/PCG/Private/Data/PCGUnionData.cpp
--- a/Engine/Plugins/Experimental/PCG/Source/PCG/Private/Data/PCGUnionData.cpp
+++ b/Engine/Plugins/Experimental/PCG/Source/PCG/Private/Data/PCGUnionData.cpp
@@ -83,6 +83,7 @@ FBox UPCGUnionData::GetStrictBounds() const
 
 bool UPCGUnionData::SamplePoint(const FTransform& InTransform, const FBox& InBounds, FPCGPoint& OutPoint, UPCGMetadata* OutMetadata) const
 {
+	//TRACE_CPUPROFILER_EVENT_SCOPE(UPCGUnionData::SamplePoint);
 	FTransform PointTransform = InTransform;
 	bool bHasSetPoint = false;
 
@@ -131,7 +132,14 @@ bool UPCGUnionData::SamplePoint(const FTransform& InTransform, const FBox& InBou
 				// Merge properties into OutPoint
 				if (OutMetadata)
 				{
-					OutMetadata->MergePointAttributesSubset(OutPoint, OutMetadata, OutMetadata, PointInData, OutMetadata, Data[DataIndex]->Metadata, OutPoint, EPCGMetadataOp::Max);
+					if (OutPoint.MetadataEntry != PCGInvalidEntryKey && PointInData.MetadataEntry != PCGInvalidEntryKey)
+					{
+						OutMetadata->MergePointAttributesSubset(OutPoint, OutMetadata, OutMetadata, PointInData, OutMetadata, Data[DataIndex]->Metadata, OutPoint, EPCGMetadataOp::Max);
+					}
+					else if (PointInData.MetadataEntry != PCGInvalidEntryKey)
+					{
+						OutPoint.MetadataEntry = PointInData.MetadataEntry;
+					}
 				}
 			}
 			
@@ -297,7 +305,14 @@ void UPCGUnionData::CreateSequentialPointData(FPCGContext* Context, UPCGPointDat
 
 					if (PointData->Metadata)
 					{
-						PointData->Metadata->MergePointAttributesSubset(OutPoint, PointData->Metadata, PointData->Metadata, PointInData, PointData->Metadata, Data[FollowingDataIndex]->Metadata, OutPoint, EPCGMetadataOp::Max);
+						if (OutPoint.MetadataEntry != PCGInvalidEntryKey && PointInData.MetadataEntry != PCGInvalidEntryKey)
+						{
+							PointData->Metadata->MergePointAttributesSubset(OutPoint, PointData->Metadata, PointData->Metadata, PointInData, PointData->Metadata, Data[FollowingDataIndex]->Metadata, OutPoint, EPCGMetadataOp::Max);
+						}
+						else if (PointInData.MetadataEntry != PCGInvalidEntryKey)
+						{
+							OutPoint.MetadataEntry = PointInData.MetadataEntry;
+						}
 					}
 					else if (OutPoint.Density >= 1.0f)
 					{
diff --git a/Engine/Plugins/Experimental/PCG/Source/PCG/Private/Data/PCGVolumeData.cpp b/Engine/Plugins/Experimental/PCG/Source/PCG/Private/Data/PCGVolumeData.cpp
--- a/Engine/Plugins/Experimental/PCG/Source/PCG/Private/Data/PCGVolumeData.cpp
+++ b/Engine/Plugins/Experimental/PCG/Source/PCG/Private/Data/PCGVolumeData.cpp
@@ -57,6 +57,7 @@ const UPCGPointData* UPCGVolumeData::CreatePointData(FPCGContext* Context) const
 
 bool UPCGVolumeData::SamplePoint(const FTransform& InTransform, const FBox& InBounds, FPCGPoint& OutPoint, UPCGMetadata* OutMetadata) const
 {
+	//TRACE_CPUPROFILER_EVENT_SCOPE(UPCGVolumeData::SamplePoint);
 	// TODO: add metadata
 	// TODO: consider bounds
 	const FVector InPosition = InTransform.GetLocation();
diff --git a/Engine/Plugins/Experimental/PCG/Source/PCG/Private/Elements/PCGDensityFilter.cpp b/Engine/Plugins/Experimental/PCG/Source/PCG/Private/Elements/PCGDensityFilter.cpp
--- a/Engine/Plugins/Experimental/PCG/Source/PCG/Private/Elements/PCGDensityFilter.cpp
+++ b/Engine/Plugins/Experimental/PCG/Source/PCG/Private/Elements/PCGDensityFilter.cpp
@@ -48,8 +48,6 @@ bool FPCGDensityFilterElement::ExecuteInternal(FPCGContext* Context) const
 	// TODO: embarassingly parallel loop
 	for (const FPCGTaggedData& Input : Inputs)
 	{
-		TRACE_CPUPROFILER_EVENT_SCOPE(FPCGDensityFilterElement::Execute::InputLoop);
-
 		FPCGTaggedData& Output = Outputs.Add_GetRef(Input);
 
 		if (!Input.Data || Cast<UPCGSpatialData>(Input.Data) == nullptr)
diff --git a/Engine/Plugins/Experimental/PCG/Source/PCG/Private/Elements/PCGDifferenceElement.cpp b/Engine/Plugins/Experimental/PCG/Source/PCG/Private/Elements/PCGDifferenceElement.cpp
--- a/Engine/Plugins/Experimental/PCG/Source/PCG/Private/Elements/PCGDifferenceElement.cpp
+++ b/Engine/Plugins/Experimental/PCG/Source/PCG/Private/Elements/PCGDifferenceElement.cpp
@@ -101,6 +101,7 @@ bool FPCGDifferenceElement::ExecuteInternal(FPCGContext* Context) const
 				{
 					DifferenceData = (UnionData ? UnionData : FirstSpatialData)->Subtract(SpatialData);
 					DifferenceData->SetDensityFunction(Settings->DensityFunction);
+					DifferenceData->bDiffMetadata = Settings->bDiffMetadata;
 #if WITH_EDITORONLY_DATA
 					DifferenceData->bKeepZeroDensityPoints = Settings->bKeepZeroDensityPoints;
 #endif
diff --git a/Engine/Plugins/Experimental/PCG/Source/PCG/Private/Metadata/PCGMetadata.cpp b/Engine/Plugins/Experimental/PCG/Source/PCG/Private/Metadata/PCGMetadata.cpp
--- a/Engine/Plugins/Experimental/PCG/Source/PCG/Private/Metadata/PCGMetadata.cpp
+++ b/Engine/Plugins/Experimental/PCG/Source/PCG/Private/Metadata/PCGMetadata.cpp
@@ -645,6 +645,7 @@ void UPCGMetadata::MergeAttributes(PCGMetadataEntryKey InKeyA, const UPCGMetadat
 
 void UPCGMetadata::MergeAttributesSubset(PCGMetadataEntryKey InKeyA, const UPCGMetadata* InMetadataA, const UPCGMetadata* InMetadataSubsetA, PCGMetadataEntryKey InKeyB, const UPCGMetadata* InMetadataB, const UPCGMetadata* InMetadataSubsetB, PCGMetadataEntryKey& OutKey, EPCGMetadataOp Op)
 {
+	//TRACE_CPUPROFILER_EVENT_SCOPE(UPCGMetadata::MergeAttributesSubset);
 	// Early out: nothing to do if both input metadata are null / points have no assigned metadata
 	if (!InMetadataA && !InMetadataB)
 	{
@@ -749,6 +750,52 @@ void UPCGMetadata::AccumulateWeightedAttributes(PCGMetadataEntryKey InKey, const
 	AttributeLock.ReadUnlock();
 }
 
+void UPCGMetadata::ComputePointWeightedAttribute(FPCGPoint& OutPoint, const TArrayView<TPair<const FPCGPoint*, float>>& InWeightedPoints, const UPCGMetadata* InMetadata)
+{
+	TArray<TPair<PCGMetadataEntryKey, float>, TInlineAllocator<4>> InWeightedKeys;
+	InWeightedKeys.Reserve(InWeightedPoints.Num());
+
+	for (const TPair<const FPCGPoint*, float>& WeightedPoint : InWeightedPoints)
+	{
+		InWeightedKeys.Emplace(WeightedPoint.Key->MetadataEntry, WeightedPoint.Value);
+	}
+
+	ComputeWeightedAttribute(OutPoint.MetadataEntry, MakeArrayView(InWeightedKeys), InMetadata);
+}
+
+void UPCGMetadata::ComputeWeightedAttribute(PCGMetadataEntryKey& OutKey, const TArrayView<TPair<PCGMetadataEntryKey, float>>& InWeightedKeys, const UPCGMetadata* InMetadata)
+{
+	if (!InMetadata || InWeightedKeys.IsEmpty())
+	{
+		return;
+	}
+
+	// Could ensure that InitializeOnSet returns false...
+	AttributeLock.ReadLock();
+	for (const TPair<FName, FPCGMetadataAttributeBase*>& AttributePair : Attributes)
+	{
+		const FName& AttributeName = AttributePair.Key;
+		FPCGMetadataAttributeBase* Attribute = AttributePair.Value;
+
+		if (!Attribute->AllowsInterpolation())
+		{
+			continue;
+		}
+
+		if (const FPCGMetadataAttributeBase* OtherAttribute = InMetadata->GetConstAttribute(AttributeName))
+		{
+			if (OtherAttribute->GetTypeId() != Attribute->GetTypeId())
+			{
+				UE_LOG(LogPCG, Error, TEXT("Metadata type mismatch with attribute %s"), *AttributeName.ToString());
+				continue;
+			}
+
+			Attribute->SetWeightedValue(OutKey, OtherAttribute, InWeightedKeys);
+		}
+	}
+	AttributeLock.ReadUnlock();
+}
+
 void UPCGMetadata::SetAttributes(PCGMetadataEntryKey InKey, const UPCGMetadata* InMetadata, PCGMetadataEntryKey& OutKey)
 {
 	if (!InMetadata)
diff --git a/Engine/Plugins/Experimental/PCG/Source/PCG/Private/Metadata/PCGMetadataAccessor.cpp b/Engine/Plugins/Experimental/PCG/Source/PCG/Private/Metadata/PCGMetadataAccessor.cpp
--- a/Engine/Plugins/Experimental/PCG/Source/PCG/Private/Metadata/PCGMetadataAccessor.cpp
+++ b/Engine/Plugins/Experimental/PCG/Source/PCG/Private/Metadata/PCGMetadataAccessor.cpp
@@ -224,7 +224,7 @@ void UPCGMetadataAccessorHelpers::InitializeMetadata(FPCGPoint& Point, UPCGMetad
 		// If we're not given the parent metadata, we'll assume it is the current metadata's parent
 		if (!ParentMetadata || ParentMetadata->HasParent(ParentMetadata))
 		{
-			Metadata->AddEntry(ParentPoint.MetadataEntry);
+			Point.MetadataEntry = Metadata->AddEntry(ParentPoint.MetadataEntry);
 		}
 		else
 		{
diff --git a/Engine/Plugins/Experimental/PCG/Source/PCG/Public/Data/PCGDifferenceData.h b/Engine/Plugins/Experimental/PCG/Source/PCG/Public/Data/PCGDifferenceData.h
--- a/Engine/Plugins/Experimental/PCG/Source/PCG/Public/Data/PCGDifferenceData.h
+++ b/Engine/Plugins/Experimental/PCG/Source/PCG/Public/Data/PCGDifferenceData.h
@@ -39,6 +39,9 @@ public:
 	UFUNCTION(BlueprintCallable, Category = Settings)
 	void SetDensityFunction(EPCGDifferenceDensityFunction InDensityFunction);
 
+	UPROPERTY(BlueprintReadOnly, VisibleAnywhere, Category = SpatialData)
+	bool bDiffMetadata = true;
+
 #if WITH_EDITOR
 	virtual void PostEditChangeProperty(FPropertyChangedEvent& PropertyChangedEvent) override;
 #endif
diff --git a/Engine/Plugins/Experimental/PCG/Source/PCG/Public/Elements/PCGDifferenceElement.h b/Engine/Plugins/Experimental/PCG/Source/PCG/Public/Elements/PCGDifferenceElement.h
--- a/Engine/Plugins/Experimental/PCG/Source/PCG/Public/Elements/PCGDifferenceElement.h
+++ b/Engine/Plugins/Experimental/PCG/Source/PCG/Public/Elements/PCGDifferenceElement.h
@@ -40,6 +40,9 @@ public:
 	UPROPERTY(BlueprintReadWrite, EditAnywhere, Category = Settings)
 	EPCGDifferenceMode Mode = EPCGDifferenceMode::Inferred;
 
+	UPROPERTY(BlueprintReadWrite, EditAnywhere, Category = Settings)
+	bool bDiffMetadata = true;
+
 #if WITH_EDITORONLY_DATA
 	UPROPERTY(Transient, BlueprintReadWrite, EditAnywhere, Category = Debug)
 	bool bKeepZeroDensityPoints = false;
diff --git a/Engine/Plugins/Experimental/PCG/Source/PCG/Public/Metadata/PCGMetadata.h b/Engine/Plugins/Experimental/PCG/Source/PCG/Public/Metadata/PCGMetadata.h
--- a/Engine/Plugins/Experimental/PCG/Source/PCG/Public/Metadata/PCGMetadata.h
+++ b/Engine/Plugins/Experimental/PCG/Source/PCG/Public/Metadata/PCGMetadata.h
@@ -183,6 +183,9 @@ public:
 	UFUNCTION(BlueprintCallable, Category = "PCG|Metadata")
 	void AccumulatePointWeightedAttributes(const FPCGPoint& InPoint, const UPCGMetadata* InMetadata, float Weight, bool bSetNonInterpolableAttributes, UPARAM(ref) FPCGPoint& OutPoint);
 
+	void ComputePointWeightedAttribute(FPCGPoint& OutPoint, const TArrayView<TPair<const FPCGPoint*, float>>& InWeightedPoints, const UPCGMetadata* InMetadata);
+	void ComputeWeightedAttribute(PCGMetadataEntryKey& OutKey, const TArrayView<TPair<PCGMetadataEntryKey, float>>& InWeightedKeys, const UPCGMetadata* InMetadata);
+
 	int64 GetItemCountForChild() const;
 
 protected:
diff --git a/Engine/Plugins/Experimental/PCG/Source/PCG/Public/Metadata/PCGMetadataAttribute.h b/Engine/Plugins/Experimental/PCG/Source/PCG/Public/Metadata/PCGMetadataAttribute.h
--- a/Engine/Plugins/Experimental/PCG/Source/PCG/Public/Metadata/PCGMetadataAttribute.h
+++ b/Engine/Plugins/Experimental/PCG/Source/PCG/Public/Metadata/PCGMetadataAttribute.h
@@ -27,6 +27,7 @@ public:
 	virtual void SetValue(PCGMetadataEntryKey ItemKey, const FPCGMetadataAttributeBase* InAttribute, PCGMetadataEntryKey InEntryKey) = 0;
 	virtual void SetZeroValue(PCGMetadataEntryKey ItemKey) = 0;
 	virtual void AccumulateValue(PCGMetadataEntryKey ItemKey, const FPCGMetadataAttributeBase* InAttributeA, PCGMetadataEntryKey InEntryKeyA, float Weight) = 0;
+	virtual void SetWeightedValue(PCGMetadataEntryKey ItemKey, const FPCGMetadataAttributeBase* InAttribute, const TArrayView<TPair<PCGMetadataEntryKey, float>> InWeightedKeys) = 0;
 	virtual void SetValue(PCGMetadataEntryKey ItemKey, const FPCGMetadataAttributeBase* InAttributeA, PCGMetadataEntryKey InEntryKeyA, const FPCGMetadataAttributeBase* InAttributeB, PCGMetadataEntryKey InEntryKeyB, EPCGMetadataOp Op) = 0;
 	virtual bool IsEqualToDefaultValue(PCGMetadataValueKey ValueKey) const = 0;
 
diff --git a/Engine/Plugins/Experimental/PCG/Source/PCG/Public/Metadata/PCGMetadataAttributeTpl.h b/Engine/Plugins/Experimental/PCG/Source/PCG/Public/Metadata/PCGMetadataAttributeTpl.h
--- a/Engine/Plugins/Experimental/PCG/Source/PCG/Public/Metadata/PCGMetadataAttributeTpl.h
+++ b/Engine/Plugins/Experimental/PCG/Source/PCG/Public/Metadata/PCGMetadataAttributeTpl.h
@@ -93,6 +93,12 @@ public:
 		Accumulate(ItemKey, InAttribute, InEntryKey, Weight);
 	}
 
+	virtual void SetWeightedValue(PCGMetadataEntryKey ItemKey, const FPCGMetadataAttributeBase* InAttribute, const TArrayView<TPair<PCGMetadataEntryKey, float>> InWeightedKeys) override
+	{
+		check(ItemKey != PCGInvalidEntryKey);
+		Accumulate(ItemKey, InAttribute, InWeightedKeys);
+	}
+
 	virtual void SetValue(PCGMetadataEntryKey ItemKey, const FPCGMetadataAttributeBase* InAttributeA, PCGMetadataEntryKey InEntryKeyA, const FPCGMetadataAttributeBase* InAttributeB, PCGMetadataEntryKey InEntryKeyB, EPCGMetadataOp Op) override
 	{
 		check(ItemKey != PCGInvalidEntryKey);
@@ -390,6 +396,27 @@ protected:
 		// Empty on purpose
 	}
 
+	template<typename IT = T, typename TEnableIf<PCG::Private::MetadataTraits<IT>::CanInterpolate>::Type* = nullptr>
+	void Accumulate(PCGMetadataEntryKey ItemKey, const FPCGMetadataAttributeBase* InAttribute, const TArrayView<TPair<PCGMetadataEntryKey, float>>& InWeightedKeys)
+	{
+		IT Value = PCG::Private::MetadataTraits<IT>::ZeroValue();
+		for (const TPair<PCGMetadataEntryKey, float>& WeightedEntry : InWeightedKeys)
+		{
+			PCG::Private::MetadataTraits<IT>::WeightedSum(
+				Value,
+				static_cast<const FPCGMetadataAttribute<T>*>(InAttribute)->GetValueFromItemKey(WeightedEntry.Key),
+				WeightedEntry.Value);
+		}
+
+		SetValue(ItemKey, Value);
+	}
+
+	template<typename IT = T, typename TEnableIf<!PCG::Private::MetadataTraits<IT>::CanInterpolate>::Type* = nullptr>
+	void Accumulate(PCGMetadataEntryKey ItemKey, const FPCGMetadataAttributeBase* InAttribute, const TArrayView<TPair<PCGMetadataEntryKey, float>>& InWeightedKeys)
+	{
+		// Empty on purpose
+	}
+
 protected:
 	mutable FRWLock ValueLock;
 	TArray<T> Values;
