Fixed spline sampling issue

diff --git a/Engine/Plugins/Experimental/PCG/Source/PCG/Private/Data/PCGLandscapeSplineData.cpp b/Engine/Plugins/Experimental/PCG/Source/PCG/Private/Data/PCGLandscapeSplineData.cpp
--- a/Engine/Plugins/Experimental/PCG/Source/PCG/Private/Data/PCGLandscapeSplineData.cpp
+++ b/Engine/Plugins/Experimental/PCG/Source/PCG/Private/Data/PCGLandscapeSplineData.cpp
@@ -14,16 +14,24 @@
 namespace PCGLandscapeDataHelpers
 {
 	// This function assumes that the A-B segment has a "1" density, while the C-D segment has a "0" density
-	FVector::FReal GetDensityInQuad(const FVector& A, const FVector& B, const FVector& C, const FVector& D, const FVector& P)
+	FVector::FReal GetDensityInQuad(const FVector& InA, const FVector& InB, const FVector& InC, const FVector& InD, const FVector& InP)
 	{
-		FVector BaryABC = FMath::ComputeBaryCentric2D(P, A, B, C);
+		// Since the landscape has a strict Z behavior and our points might not be directly on the plane, we should only consider
+		// the 2D plane. When we support other axes, we could just remove the normal component off the position to correct.
+		const FVector2D A(InA.X, InA.Y);
+		const FVector2D B(InB.X, InB.Y);
+		const FVector2D C(InC.X, InC.Y);
+		const FVector2D D(InD.X, InD.Y);
+		const FVector2D P(InP.X, InP.Y);
+
+		FVector BaryABC = FMath::GetBaryCentric2D(P, A, B, C);
 
 		if (BaryABC.X >= 0 && BaryABC.Y >= 0 && BaryABC.Z >= 0)
 		{
 			return 1.0f - BaryABC.Z;
 		}
 
-		FVector BaryACD = FMath::ComputeBaryCentric2D(P, A, C, D);
+		FVector BaryACD = FMath::GetBaryCentric2D(P, A, C, D);
 
 		if (BaryACD.X >= 0 && BaryACD.Y >= 0 && BaryACD.Z >= 0)
 		{
