Added max depth for dependency parsing. Will prevent spurious updates and will also be faster

diff --git a/Engine/Plugins/Experimental/PCG/Source/PCG/Private/Elements/PCGExecuteBlueprint.cpp b/Engine/Plugins/Experimental/PCG/Source/PCG/Private/Elements/PCGExecuteBlueprint.cpp
--- a/Engine/Plugins/Experimental/PCG/Source/PCG/Private/Elements/PCGExecuteBlueprint.cpp
+++ b/Engine/Plugins/Experimental/PCG/Source/PCG/Private/Elements/PCGExecuteBlueprint.cpp
@@ -13,13 +13,13 @@
 #if WITH_EDITOR
 namespace PCGBlueprintHelper
 {
-	TSet<TObjectPtr<UObject>> GetDataDependencies(UPCGBlueprintElement* InElement)
+	TSet<TObjectPtr<UObject>> GetDataDependencies(UPCGBlueprintElement* InElement, int32 MaxDepth)
 	{
 		check(InElement && InElement->GetClass());
 		UClass* BPClass = InElement->GetClass();
 
 		TSet<TObjectPtr<UObject>> Dependencies;
-		PCGHelpers::GatherDependencies(InElement, Dependencies);
+		PCGHelpers::GatherDependencies(InElement, Dependencies, MaxDepth);
 		return Dependencies;
 	}
 }
@@ -78,7 +78,7 @@ void UPCGBlueprintElement::Initialize()
 {
 #if WITH_EDITOR
 	FCoreUObjectDelegates::OnObjectPropertyChanged.AddUObject(this, &UPCGBlueprintElement::OnDependencyChanged);
-	DataDependencies = PCGBlueprintHelper::GetDataDependencies(this);
+	DataDependencies = PCGBlueprintHelper::GetDataDependencies(this, DependencyParsingDepth);
 #endif
 }
 
@@ -88,7 +88,7 @@ void UPCGBlueprintElement::PostEditChangeProperty(FPropertyChangedEvent& Propert
 	Super::PostEditChangeProperty(PropertyChangedEvent);
 
 	// Since we don't really know what changed, let's just rebuild our data dependencies
-	DataDependencies = PCGBlueprintHelper::GetDataDependencies(this);
+	DataDependencies = PCGBlueprintHelper::GetDataDependencies(this, DependencyParsingDepth);
 
 	OnBlueprintChangedDelegate.Broadcast(this);
 }
diff --git a/Engine/Plugins/Experimental/PCG/Source/PCG/Private/PCGHelpers.cpp b/Engine/Plugins/Experimental/PCG/Source/PCG/Private/PCGHelpers.cpp
--- a/Engine/Plugins/Experimental/PCG/Source/PCG/Private/PCGHelpers.cpp
+++ b/Engine/Plugins/Experimental/PCG/Source/PCG/Private/PCGHelpers.cpp
@@ -167,7 +167,7 @@ namespace PCGHelpers
 		return true;
 	}
 
-	void GatherDependencies(UObject* Object, TSet<TObjectPtr<UObject>>& OutDependencies)
+	void GatherDependencies(UObject* Object, TSet<TObjectPtr<UObject>>& OutDependencies, int32 MaxDepth)
 	{
 		UClass* ObjectClass = Object ? Object->GetClass() : nullptr;
 
@@ -179,18 +179,21 @@ namespace PCGHelpers
 
 		for (FProperty* Property = ObjectClass->PropertyLink; Property != nullptr; Property = Property->PropertyLinkNext)
 		{
-			GatherDependencies(Property, Object, OutDependencies);
+			GatherDependencies(Property, Object, OutDependencies, MaxDepth);
 		}
 	}
 
 	// Inspired by IteratePropertiesRecursive in ObjectPropertyTrace.cpp
-	void GatherDependencies(FProperty* Property, const void* InContainer, TSet<TObjectPtr<UObject>>& OutDependencies)
+	void GatherDependencies(FProperty* Property, const void* InContainer, TSet<TObjectPtr<UObject>>& OutDependencies, int32 MaxDepth)
 	{
-		auto AddToDependenciesAndGatherRecursively = [&OutDependencies](UObject* Object) {
+		auto AddToDependenciesAndGatherRecursively = [&OutDependencies, MaxDepth](UObject* Object) {
 			if (Object && !OutDependencies.Contains(Object))
 			{
 				OutDependencies.Add(Object);
-				GatherDependencies(Object, OutDependencies);
+				if (MaxDepth != 0)
+				{
+					GatherDependencies(Object, OutDependencies, MaxDepth - 1);
+				}
 			}
 		};
 
@@ -219,7 +222,7 @@ namespace PCGHelpers
 			const void* StructContainer = StructProperty->ContainerPtrToValuePtr<const void>(InContainer);
 			for (TFieldIterator<FProperty> It(StructProperty->Struct); It; ++It)
 			{
-				GatherDependencies(*It, StructContainer, OutDependencies);
+				GatherDependencies(*It, StructContainer, OutDependencies, MaxDepth);
 			}
 		}
 		else if (FArrayProperty* ArrayProperty = CastField<FArrayProperty>(Property))
@@ -228,7 +231,7 @@ namespace PCGHelpers
 			for (int32 DynamicIndex = 0; DynamicIndex < Helper.Num(); ++DynamicIndex)
 			{
 				const void* ValuePtr = Helper.GetRawPtr(DynamicIndex);
-				GatherDependencies(ArrayProperty->Inner, ValuePtr, OutDependencies);
+				GatherDependencies(ArrayProperty->Inner, ValuePtr, OutDependencies, MaxDepth);
 			}
 		}
 		else if (FMapProperty* MapProperty = CastField<FMapProperty>(Property))
@@ -240,10 +243,10 @@ namespace PCGHelpers
 				if (Helper.IsValidIndex(DynamicIndex))
 				{
 					const void* KeyPtr = Helper.GetKeyPtr(DynamicIndex);
-					GatherDependencies(MapProperty->KeyProp, KeyPtr, OutDependencies);
+					GatherDependencies(MapProperty->KeyProp, KeyPtr, OutDependencies, MaxDepth);
 
 					const void* ValuePtr = Helper.GetValuePtr(DynamicIndex);
-					GatherDependencies(MapProperty->ValueProp, ValuePtr, OutDependencies);
+					GatherDependencies(MapProperty->ValueProp, ValuePtr, OutDependencies, MaxDepth);
 
 					--Num;
 				}
@@ -258,7 +261,7 @@ namespace PCGHelpers
 				if (Helper.IsValidIndex(DynamicIndex))
 				{
 					const void* ValuePtr = Helper.GetElementPtr(DynamicIndex);
-					GatherDependencies(SetProperty->ElementProp, ValuePtr, OutDependencies);
+					GatherDependencies(SetProperty->ElementProp, ValuePtr, OutDependencies, MaxDepth);
 
 					--Num;
 				}
diff --git a/Engine/Plugins/Experimental/PCG/Source/PCG/Private/PCGHelpers.h b/Engine/Plugins/Experimental/PCG/Source/PCG/Private/PCGHelpers.h
--- a/Engine/Plugins/Experimental/PCG/Source/PCG/Private/PCGHelpers.h
+++ b/Engine/Plugins/Experimental/PCG/Source/PCG/Private/PCGHelpers.h
@@ -38,7 +38,7 @@ namespace PCGHelpers
 	APCGWorldActor* GetPCGWorldActor(UWorld* InWorld);
 
 #if WITH_EDITOR
-	void GatherDependencies(UObject* Object, TSet<TObjectPtr<UObject>>& OutDependencies);
-	void GatherDependencies(FProperty* Property, const void* InContainer, TSet<TObjectPtr<UObject>>& OutDependencies);
+	void GatherDependencies(UObject* Object, TSet<TObjectPtr<UObject>>& OutDependencies, int32 MaxDepth = -1);
+	void GatherDependencies(FProperty* Property, const void* InContainer, TSet<TObjectPtr<UObject>>& OutDependencies, int32 MaxDepth);
 #endif
 };
\ No newline at end of file
diff --git a/Engine/Plugins/Experimental/PCG/Source/PCG/Public/Elements/PCGExecuteBlueprint.h b/Engine/Plugins/Experimental/PCG/Source/PCG/Public/Elements/PCGExecuteBlueprint.h
--- a/Engine/Plugins/Experimental/PCG/Source/PCG/Public/Elements/PCGExecuteBlueprint.h
+++ b/Engine/Plugins/Experimental/PCG/Source/PCG/Public/Elements/PCGExecuteBlueprint.h
@@ -137,6 +137,9 @@ public:
 
 	UPROPERTY(EditAnywhere, Category = AssetInfo, AssetRegistrySearchable)
 	FText Description;
+
+	UPROPERTY(BlueprintReadOnly, EditAnywhere, Category = "Settings|Advanced")
+	int32 DependencyParsingDepth = 1;
 #endif
 
 protected:
