Fixed isssue with landscape cache layer names being cleared on load

diff --git a/Engine/Plugins/Experimental/PCG/Source/PCG/Private/Grid/PCGLandscapeCache.cpp b/Engine/Plugins/Experimental/PCG/Source/PCG/Private/Grid/PCGLandscapeCache.cpp
--- a/Engine/Plugins/Experimental/PCG/Source/PCG/Private/Grid/PCGLandscapeCache.cpp
+++ b/Engine/Plugins/Experimental/PCG/Source/PCG/Private/Grid/PCGLandscapeCache.cpp
@@ -226,7 +226,7 @@ FPCGLandscapeCache::FPCGLandscapeCache(UObject* InOwner)
 #endif
 }
 
-void FPCGLandscapeCache::SetOwner(UObject* InOwner)
+void FPCGLandscapeCache::SetOwner(UObject* InOwner, bool bUpdateCachedNames)
 {
 	check(!Owner || Owner == InOwner || InOwner == nullptr);
 	Owner = InOwner;
@@ -240,7 +240,10 @@ void FPCGLandscapeCache::SetOwner(UObject* InOwner)
 		TeardownLandscapeCallbacks();
 	}
 
-	CacheLayerNames();
+	if (bUpdateCachedNames)
+	{
+		CacheLayerNames();
+	}
 #endif
 }
 
diff --git a/Engine/Plugins/Experimental/PCG/Source/PCG/Private/PCGWorldActor.cpp b/Engine/Plugins/Experimental/PCG/Source/PCG/Private/PCGWorldActor.cpp
--- a/Engine/Plugins/Experimental/PCG/Source/PCG/Private/PCGWorldActor.cpp
+++ b/Engine/Plugins/Experimental/PCG/Source/PCG/Private/PCGWorldActor.cpp
@@ -42,7 +42,7 @@ void APCGWorldActor::BeginCacheForCookedPlatformData(const ITargetPlatform* Targ
 void APCGWorldActor::PostLoad()
 {
 	Super::PostLoad();
-	LandscapeCache.SetOwner(this);
+	LandscapeCache.SetOwner(this, /*bUpdateCachedNames=*/false);
 	RegisterToSubsystem();
 }
 
diff --git a/Engine/Plugins/Experimental/PCG/Source/PCG/Public/Grid/PCGLandscapeCache.h b/Engine/Plugins/Experimental/PCG/Source/PCG/Public/Grid/PCGLandscapeCache.h
--- a/Engine/Plugins/Experimental/PCG/Source/PCG/Public/Grid/PCGLandscapeCache.h
+++ b/Engine/Plugins/Experimental/PCG/Source/PCG/Public/Grid/PCGLandscapeCache.h
@@ -64,7 +64,7 @@ public:
 	~FPCGLandscapeCache();
 	FPCGLandscapeCache(UObject* InOwner);
 
-	void SetOwner(UObject* InOwner);
+	void SetOwner(UObject* InOwner, bool bUpdateCachedLayerNames = true);
 	void PrimeCache();
 	void ClearCache();
 
