Fixed Copy Points target actor based on target first, then source.

diff --git a/Engine/Plugins/Experimental/PCG/Source/PCG/Private/Elements/PCGCopyPoints.cpp b/Engine/Plugins/Experimental/PCG/Source/PCG/Private/Elements/PCGCopyPoints.cpp
--- a/Engine/Plugins/Experimental/PCG/Source/PCG/Private/Elements/PCGCopyPoints.cpp
+++ b/Engine/Plugins/Experimental/PCG/Source/PCG/Private/Elements/PCGCopyPoints.cpp
@@ -100,6 +100,9 @@ bool FPCGCopyPointsElement::ExecuteInternal(FPCGContext* Context) const
 		NonRootMetadata = SourcePointData->Metadata;
 	}
 
+	// Priorize use the target actor from the target, irrespective of the source
+	OutPointData->TargetActor = TargetPointData->TargetActor ? TargetPointData->TargetActor : SourcePointData->TargetActor;
+
 	check(OutPointData->Metadata && NonRootMetadata);
 
 	TArray<FName> AttributeNames;
