Made landscape tracking update only on structural changes - makes editing graphs faster in large WP maps

diff --git a/Engine/Plugins/Experimental/PCG/Source/PCG/Private/PCGComponent.cpp b/Engine/Plugins/Experimental/PCG/Source/PCG/Private/PCGComponent.cpp
--- a/Engine/Plugins/Experimental/PCG/Source/PCG/Private/PCGComponent.cpp
+++ b/Engine/Plugins/Experimental/PCG/Source/PCG/Private/PCGComponent.cpp
@@ -671,7 +671,11 @@ void UPCGComponent::OnGraphChanged(UPCGGraph* InGraph, bool bIsStructural, bool
 		SetupTrackingCallbacks();
 		RefreshTrackingData();
 		DirtyCacheForAllTrackedTags();
-		UpdateTrackedLandscape();
+
+		if (bIsStructural)
+		{
+			UpdateTrackedLandscape();
+		}
 
 		DirtyGenerated();
 		if (InGraph && bShouldRefresh)
