Optimized graph executor scheduling for very large graphs by building a successor list so we can quickly find potential tasks to execute

diff --git a/Engine/Plugins/Experimental/PCG/Source/PCG/Private/Graph/PCGGraphExecutor.cpp b/Engine/Plugins/Experimental/PCG/Source/PCG/Private/Graph/PCGGraphExecutor.cpp
--- a/Engine/Plugins/Experimental/PCG/Source/PCG/Private/Graph/PCGGraphExecutor.cpp
+++ b/Engine/Plugins/Experimental/PCG/Source/PCG/Private/Graph/PCGGraphExecutor.cpp
@@ -146,20 +146,39 @@ void FPCGGraphExecutor::Execute()
 	// Process any newly scheduled graphs to execute
 	ScheduleLock.Lock();
 
-	for (FPCGGraphScheduleTask& ScheduledTask : ScheduledTasks)
+	for(int32 ScheduledTaskIndex = ScheduledTasks.Num() - 1; ScheduledTaskIndex >= 0; --ScheduledTaskIndex)
 	{
+		FPCGGraphScheduleTask& ScheduledTask = ScheduledTasks[ScheduledTaskIndex];
+
 		check(ScheduledTask.Tasks.Num() > 0);
+		// Push tasks to the master list & build successor list
+		for (FPCGGraphTask& Task : ScheduledTask.Tasks)
+		{
+			FPCGTaskId TaskId = Task.NodeId;
 
-		// Finally, push the tasks to the master list
-		Tasks.Append(ScheduledTask.Tasks);
-	}
+			bool bPushToReady = true;
+			for (const FPCGGraphTaskInput& Input : Task.Inputs)
+			{
+				if (!OutputData.Contains(Input.TaskId))
+				{
+					TaskSuccessors.FindOrAdd(Input.TaskId).Add(TaskId);
+					bPushToReady = false;
+				}
+			}
 
-	if (ScheduledTasks.Num() > 0)
-	{
-		ScheduledTasks.Reset();
-		// Kick off any of the newly added ready tasks
-		QueueReadyTasks();
+			// Automatically push inputless/already satisfied tasks to the ready queue
+			if (bPushToReady)
+			{
+				ReadyTasks.Emplace(MoveTemp(Task));
+			}
+			else
+			{
+				Tasks.Add(TaskId, MoveTemp(Task));
+			}
+		}
 	}
+
+	ScheduledTasks.Reset();
 	
 	ScheduleLock.Unlock();
 
@@ -188,8 +207,6 @@ void FPCGGraphExecutor::Execute()
 		const bool bMainThreadWasAvailable = bMainThreadAvailable;
 		const int32 TasksToLaunchIndex = ActiveTasks.Num();
 
-		bool bSomeTaskEndedInCurrentLoop = false;
-
 		if (bMainThreadAvailable || NumAvailableThreads > 0)
 		{
 			TRACE_CPUPROFILER_EVENT_SCOPE(FPCGGraphExecutor::Execute::PrepareTasks);
@@ -238,9 +255,12 @@ void FPCGGraphExecutor::Execute()
 				if (!bNeedsToCreateActiveTask)
 				{
 					// Fast-forward cached result to stored results
-					StoreResults(Task.NodeId, CachedOutput);
+					FPCGTaskId SkippedTaskId = Task.NodeId;
+					StoreResults(SkippedTaskId, CachedOutput);
 					ReadyTasks.RemoveAtSwap(ReadyTaskIndex);
-					bSomeTaskEndedInCurrentLoop = true;
+					QueueNextTasks(SkippedTaskId);
+					bAnyTaskEnded = true;
+
 #if WITH_EDITOR
 					UPCGComponent* SourceComponent = Task.SourceComponent;
 					if (SourceComponent && SourceComponent->IsInspecting())
@@ -333,7 +353,7 @@ void FPCGGraphExecutor::Execute()
 			}
 		}
 
-		auto PostTaskExecute = [this, &bSomeTaskEndedInCurrentLoop](int32 TaskIndex)
+		auto PostTaskExecute = [this, &bAnyTaskEnded](int32 TaskIndex)
 		{
 			FPCGGraphActiveTask& ActiveTask = ActiveTasks[TaskIndex];
 
@@ -367,7 +387,8 @@ void FPCGGraphExecutor::Execute()
 			StoreResults(ActiveTask.NodeId, ActiveTask.Context->OutputData);
 
 			// Book-keeping
-			bSomeTaskEndedInCurrentLoop = true;
+			QueueNextTasks(ActiveTask.NodeId);
+			bAnyTaskEnded = true;
 
 			// Remove current active task from list
 			ActiveTasks.RemoveAtSwap(TaskIndex);
@@ -408,12 +429,6 @@ void FPCGGraphExecutor::Execute()
 
 		check(CurrentlyUsedThreads == 0);
 
-		if (bSomeTaskEndedInCurrentLoop)
-		{
-			QueueReadyTasks();
-			bAnyTaskEnded = true;
-		}
-
 		const double EndLoopTime = FPlatformTime::Seconds();
 		if (static_cast<float>((EndLoopTime - StartTime) * 1000) > CVarTimePerFrame.GetValueOnAnyThread())
 		{
@@ -429,6 +444,7 @@ void FPCGGraphExecutor::Execute()
 		// keep a cache of intermediate results.
 		if (ReadyTasks.Num() == 0 && ActiveTasks.Num() == 0 && Tasks.Num() == 0)
 		{
+			check(TaskSuccessors.IsEmpty());
 			ClearResults();
 		}
 
@@ -440,22 +456,30 @@ void FPCGGraphExecutor::Execute()
 	}
 }
 
-void FPCGGraphExecutor::QueueReadyTasks(FPCGTaskId FinishedTaskHint)
+void FPCGGraphExecutor::QueueNextTasks(FPCGTaskId FinishedTask)
 {
-	TRACE_CPUPROFILER_EVENT_SCOPE(FPCGGraphExecutor::QueueReadyTasks);
-	for (int TaskIndex = Tasks.Num() - 1; TaskIndex >= 0; --TaskIndex)
+	TRACE_CPUPROFILER_EVENT_SCOPE(FPCGGraphExecutor::QueueNextTasks);
+
+	if (TSet<FPCGTaskId>* Successors = TaskSuccessors.Find(FinishedTask))
 	{
-		bool bAllPrerequisitesMet = true;
-		for(const FPCGGraphTaskInput& Input : Tasks[TaskIndex].Inputs)
+		for (FPCGTaskId Successor : *Successors)
 		{
-			bAllPrerequisitesMet &= OutputData.Contains(Input.TaskId);
-		}
+			bool bAllPrerequisitesMet = true;
+			FPCGGraphTask& SuccessorTask = Tasks[Successor];
 
-		if (bAllPrerequisitesMet)
-		{
-			ReadyTasks.Emplace(MoveTemp(Tasks[TaskIndex]));
-			Tasks.RemoveAtSwap(TaskIndex);
+			for (const FPCGGraphTaskInput& Input : SuccessorTask.Inputs)
+			{
+				bAllPrerequisitesMet &= OutputData.Contains(Input.TaskId);
+			}
+
+			if (bAllPrerequisitesMet)
+			{
+				ReadyTasks.Emplace(MoveTemp(SuccessorTask));
+				Tasks.Remove(Successor);
+			}
 		}
+
+		TaskSuccessors.Remove(FinishedTask);
 	}
 }
 
@@ -574,7 +598,12 @@ void FPCGGraphExecutor::ReleaseUnusedActors()
 #if WITH_EDITOR
 	if (bRunGC && !PCGHelpers::IsRuntimeOrPIE())
 	{
-		CollectGarbage(RF_NoFlags, true);
+		--CountUntilGC;
+		if (CountUntilGC <= 0)
+		{
+			CountUntilGC = 30;
+			CollectGarbage(RF_NoFlags, true);
+		}
 	}
 #endif
 }
diff --git a/Engine/Plugins/Experimental/PCG/Source/PCG/Private/Graph/PCGGraphExecutor.h b/Engine/Plugins/Experimental/PCG/Source/PCG/Private/Graph/PCGGraphExecutor.h
--- a/Engine/Plugins/Experimental/PCG/Source/PCG/Private/Graph/PCGGraphExecutor.h
+++ b/Engine/Plugins/Experimental/PCG/Source/PCG/Private/Graph/PCGGraphExecutor.h
@@ -92,7 +92,7 @@ public:
 	FPCGGraphCache& GetCache() { return GraphCache; }
 
 private:
-	void QueueReadyTasks(FPCGTaskId FinishedTaskHint = InvalidPCGTaskId);
+	void QueueNextTasks(FPCGTaskId FinishedTask);
 	void BuildTaskInput(const FPCGGraphTask& Task, FPCGDataCollection& TaskInput);
 	void StoreResults(FPCGTaskId InTaskId, const FPCGDataCollection& InTaskOutput);
 	void ClearResults();
@@ -116,9 +116,10 @@ private:
 	FCriticalSection ScheduleLock;
 	TArray<FPCGGraphScheduleTask> ScheduledTasks;
 
-	TArray<FPCGGraphTask> Tasks;
+	TMap<FPCGTaskId, FPCGGraphTask> Tasks;
 	TArray<FPCGGraphTask> ReadyTasks;
 	TArray<FPCGGraphActiveTask> ActiveTasks;
+	TMap<FPCGTaskId, TSet<FPCGTaskId>> TaskSuccessors;
 	FPCGRootSet ResultsRootSet;
 	/** Map of node instances to their output, could be cleared once execution is done */
 	/** Note: this should at some point unload based on loaded/unloaded proxies, otherwise memory cost will be unbounded */
@@ -133,6 +134,8 @@ private:
 	FCriticalSection ActorsListLock;
 	TSet<AActor*> ActorsToSave;
 	TSet<FWorldPartitionReference> ActorsToRelease;
+
+	int32 CountUntilGC = 30;
 #endif
 };
 
