Limits unwieldy Determinism Tests with too many test permutations

#rnx

diff --git a/Engine/Plugins/Experimental/PCG/Source/PCG/Private/Tests/Determinism/PCGDeterminismTestsCommon.cpp b/Engine/Plugins/Experimental/PCG/Source/PCG/Private/Tests/Determinism/PCGDeterminismTestsCommon.cpp
--- a/Engine/Plugins/Experimental/PCG/Source/PCG/Private/Tests/Determinism/PCGDeterminismTestsCommon.cpp
+++ b/Engine/Plugins/Experimental/PCG/Source/PCG/Private/Tests/Determinism/PCGDeterminismTestsCommon.cpp
@@ -19,6 +19,11 @@
 
 #define LOCTEXT_NAMESPACE "PCGDeterminism"
 
+static TAutoConsoleVariable<int32> CVarDeterminismPermutationLimit(
+	TEXT("pcg.DeterminismPermutationLimit"),
+	10000,
+	TEXT("A limit for the maximium amount of input permutations before autofailing a basic determinism test"));
+
 namespace PCGDeterminismTests
 {
 	bool LogInvalidTest(const UPCGNode* InPCGNode, FDeterminismNodeTestResult& OutResult)
@@ -59,6 +64,12 @@ namespace PCGDeterminismTests
 		FNodeAndOptions NodeAndOptions(InPCGNode, OutResult.Seed, /*bMultipleOptionsPerPin=*/true);
 		RetrieveBaseOptionsPerPin(NodeAndOptions.BaseOptionsByPin, InputPins, OutResult.DataTypesTested, Defaults::NumTestInputsPerPin);
 
+		if (GetNumPermutations(NodeAndOptions.BaseOptionsByPin) > CVarDeterminismPermutationLimit.GetValueOnGameThread())
+		{
+			UpdateTestResultForOverPermutationLimitError(OutResult);
+			return false;
+		}
+
 		return RunBasicSelfTest(NodeAndOptions) && RunBasicCopiedSelfTest(NodeAndOptions);
 	}
 
@@ -75,6 +86,12 @@ namespace PCGDeterminismTests
 		// TODO: Multiple inputs per pin might lead to redundant tests. Filter these out
 		RetrieveBaseOptionsPerPin(NodeAndOptions.BaseOptionsByPin, InputPins, OutResult.DataTypesTested, Defaults::NumTestInputsPerPin);
 
+		if (GetNumPermutations(NodeAndOptions.BaseOptionsByPin) > CVarDeterminismPermutationLimit.GetValueOnGameThread())
+		{
+			UpdateTestResultForOverPermutationLimitError(OutResult);
+			return false;
+		}
+
 		// No inputs, so just test determinism against itself
 		if (InputPins.Num() == 0 && RunBasicCopiedSelfTest(NodeAndOptions))
 		{
@@ -887,6 +904,13 @@ namespace PCGDeterminismTests
 		return false;
 	}
 
+	void UpdateTestResultForOverPermutationLimitError(FDeterminismNodeTestResult& OutResult)
+	{
+		OutResult.DataTypesTested = EPCGDataType::None;
+		OutResult.bFlagRaised = true;
+		OutResult.AdditionalDetails.Emplace(TEXT("Test did not run (permutation limit). Adjust 'pcg.DeterminismPermutationLimit' to alter this limit."));
+	}
+
 	bool ConsistencyComparisonIsUnimplemented(const UPCGData* FirstData, const UPCGData* SecondData, TArray<int32>& OutOutIndexOffsets)
 	{
 		return ComparisonIsUnimplemented(FirstData, SecondData);
diff --git a/Engine/Plugins/Experimental/PCG/Source/PCG/Public/Tests/Determinism/PCGDeterminismTestsCommon.h b/Engine/Plugins/Experimental/PCG/Source/PCG/Public/Tests/Determinism/PCGDeterminismTestsCommon.h
--- a/Engine/Plugins/Experimental/PCG/Source/PCG/Public/Tests/Determinism/PCGDeterminismTestsCommon.h
+++ b/Engine/Plugins/Experimental/PCG/Source/PCG/Public/Tests/Determinism/PCGDeterminismTestsCommon.h
@@ -223,6 +223,8 @@ namespace PCGDeterminismTests
 
 	/** A catch function for unimplemented comparisons */
 	bool ComparisonIsUnimplemented(const UPCGData* FirstData, const UPCGData* SecondData);
+	/** Updates the test result output to reflect that the permutation limit has been exceeded */
+	void UpdateTestResultForOverPermutationLimitError(FDeterminismNodeTestResult& OutResult);
 
 	/** Determines if a PCG DataType is a type that should be compared */
 	bool DataTypeIsComparable(EPCGDataType DataType);
diff --git a/Engine/Plugins/Experimental/PCG/Source/PCGEditor/Private/SPCGEditorGraphDeterminism.cpp b/Engine/Plugins/Experimental/PCG/Source/PCGEditor/Private/SPCGEditorGraphDeterminism.cpp
--- a/Engine/Plugins/Experimental/PCG/Source/PCGEditor/Private/SPCGEditorGraphDeterminism.cpp
+++ b/Engine/Plugins/Experimental/PCG/Source/PCGEditor/Private/SPCGEditorGraphDeterminism.cpp
@@ -73,7 +73,8 @@ TSharedRef<SWidget> SPCGEditorGraphDeterminismRow::GenerateWidgetForColumn(const
 	}
 	else if (ColumnId == NAME_AdditionalDetails)
 	{
-		FString FullDetails = FString::Join(CurrentItem->AdditionalDetails, TEXT(", "));
+		// TODO: Add tooltip widget to display multiple error messages, but only display the first in the row for now
+		FString FullDetails = CurrentItem->AdditionalDetails.Num() > 0 ? CurrentItem->AdditionalDetails[0] : "";
 		CellText = FText::FromString(FullDetails);
 	}
 	// Test columns
