+ Adds the Seed to the PCG Determinism Test UI

#rnx

diff --git a/Engine/Plugins/Experimental/PCG/Source/PCGEditor/Private/SPCGEditorGraphDeterminism.cpp b/Engine/Plugins/Experimental/PCG/Source/PCGEditor/Private/SPCGEditorGraphDeterminism.cpp
--- a/Engine/Plugins/Experimental/PCG/Source/PCGEditor/Private/SPCGEditorGraphDeterminism.cpp
+++ b/Engine/Plugins/Experimental/PCG/Source/PCGEditor/Private/SPCGEditorGraphDeterminism.cpp
@@ -11,12 +11,14 @@ namespace
 	const FName NAME_Index(TEXT("Index_ColumnID"));
 	const FName NAME_NodeTitle(TEXT("NodeTitle_ColumnID"));
 	const FName NAME_NodeName(TEXT("NodeName_ColumnID"));
+	const FName NAME_Seed(TEXT("Seed_ColumnID"));
 	const FName NAME_DataTypesTested(TEXT("DataTypesTested_ColumnID"));
 	const FName NAME_AdditionalDetails(TEXT("AdditionalDetails_ColumnID"));
 
 	const FText TEXT_Index(LOCTEXT("Index_Label", ""));
 	const FText TEXT_NodeTitle(LOCTEXT("NodeTitle_Label", "Title"));
 	const FText TEXT_NodeName(LOCTEXT("NodeName_Label", "Name"));
+	const FText TEXT_Seed(LOCTEXT("Seed_Label", "Seed"));
 	const FText TEXT_DataTypesTested(LOCTEXT("DataTypesTested_Label", "Input Data"));
 	const FText TEXT_AdditionalDetails(LOCTEXT("AdditionalDetails_Label", "Additional Details"));
 
@@ -65,6 +67,10 @@ TSharedRef<SWidget> SPCGEditorGraphDeterminismRow::GenerateWidgetForColumn(const
 	{
 		CellText = FText::FromString(CurrentItem->NodeName);
 	}
+	else if (ColumnId == NAME_Seed)
+	{
+		CellText = FText::AsNumber(CurrentItem->Seed);
+	}
 	else if (ColumnId == NAME_DataTypesTested)
 	{
 		FString DataTypesTestedString = *UEnum::GetValueAsString(CurrentItem->DataTypesTested);
@@ -158,7 +164,7 @@ void SPCGEditorGraphDeterminismListView::AddColumn(const FTestColumnInfo& Column
 	{
 		Arguments.ManualWidth(ColumnInfo.Width);
 	}
-	Arguments.HAlignCell(ColumnInfo.HAlign);
+	Arguments.HAlignHeader(ColumnInfo.HAlign);
 	GeneratedHeaderRow->AddColumn(Arguments);
 }
 
@@ -169,6 +175,7 @@ void SPCGEditorGraphDeterminismListView::BuildBaseColumns()
 	AddColumn({NAME_Index, TEXT_Index, SmallManualWidth, HAlign_Center});
 	AddColumn({NAME_NodeTitle, TEXT_NodeTitle, LargeManualWidth, HAlign_Left});
 	AddColumn({NAME_NodeName, TEXT_NodeName, LargeManualWidth, HAlign_Left});
+	AddColumn({NAME_Seed, TEXT_Seed, MediumManualWidth, HAlign_Center});
 	AddColumn({NAME_DataTypesTested, TEXT_DataTypesTested, MediumManualWidth, HAlign_Center});
 }
 
