Added option for metadata on landscape data + matching property on graph

diff --git a/Engine/Plugins/Experimental/PCG/Source/PCG/Private/Data/PCGLandscapeData.cpp b/Engine/Plugins/Experimental/PCG/Source/PCG/Private/Data/PCGLandscapeData.cpp
--- a/Engine/Plugins/Experimental/PCG/Source/PCG/Private/Data/PCGLandscapeData.cpp
+++ b/Engine/Plugins/Experimental/PCG/Source/PCG/Private/Data/PCGLandscapeData.cpp
@@ -13,13 +13,14 @@
 #include "Landscape.h"
 #include "LandscapeEdit.h"
 
-void UPCGLandscapeData::Initialize(ALandscapeProxy* InLandscape, const FBox& InBounds, bool bInHeightOnly)
+void UPCGLandscapeData::Initialize(ALandscapeProxy* InLandscape, const FBox& InBounds, bool bInHeightOnly, bool bInUseMetadata)
 {
 	check(InLandscape);
 	Landscape = InLandscape;
 	TargetActor = InLandscape;
 	Bounds = InBounds;
 	bHeightOnly = bInHeightOnly;
+	bUseMetadata = bInUseMetadata;
 
 	Transform = Landscape->GetActorTransform();
 
@@ -30,7 +31,7 @@ void UPCGLandscapeData::Initialize(ALandscapeProxy* InLandscape, const FBox& InB
 	// TODO: find a better way to do this - maybe there should be a prototype metadata in the landscape cache
 	if (LandscapeCache)
 	{
-		if (!bHeightOnly)
+		if (!bHeightOnly && bUseMetadata)
 		{
 			const TArray<FName> Layers = LandscapeCache->GetLayerNames(Landscape.Get());
 
@@ -87,7 +88,7 @@ bool UPCGLandscapeData::SamplePoint(const FTransform& InTransform, const FBox& I
 		}
 		else
 		{
-			LandscapeCacheEntry->GetInterpolatedPoint(ComponentLocalPoint, OutPoint, OutMetadata);
+			LandscapeCacheEntry->GetInterpolatedPoint(ComponentLocalPoint, OutPoint, bUseMetadata ? OutMetadata : nullptr);
 		}
 
 		return true;
@@ -189,7 +190,7 @@ const UPCGPointData* UPCGLandscapeData::CreatePointData(FPCGContext* Context, co
 						}
 						else
 						{
-							LandscapeCacheEntry->GetPoint(PointIndex, Point, Data->Metadata);
+							LandscapeCacheEntry->GetPoint(PointIndex, Point, bUseMetadata ? Data->Metadata : nullptr);
 						}
 					}
 				}
diff --git a/Engine/Plugins/Experimental/PCG/Source/PCG/Private/PCGComponent.cpp b/Engine/Plugins/Experimental/PCG/Source/PCG/Private/PCGComponent.cpp
--- a/Engine/Plugins/Experimental/PCG/Source/PCG/Private/PCGComponent.cpp
+++ b/Engine/Plugins/Experimental/PCG/Source/PCG/Private/PCGComponent.cpp
@@ -1362,7 +1362,8 @@ UPCGData* UPCGComponent::GetInputPCGData()
 
 UPCGData* UPCGComponent::GetActorPCGData()
 {
-	if (!CachedActorData)
+	// Actor PCG Data can be a Landscape data too
+	if (!CachedActorData || IsLandscapeCachedDataDirty(CachedActorData))
 	{
 		CachedActorData = CreateActorPCGData();
 	}
@@ -1372,7 +1373,7 @@ UPCGData* UPCGComponent::GetActorPCGData()
 
 UPCGData* UPCGComponent::GetLandscapePCGData()
 {
-	if (!CachedLandscapeData)
+	if (!CachedLandscapeData || IsLandscapeCachedDataDirty(CachedLandscapeData))
 	{
 		CachedLandscapeData = CreateLandscapePCGData(/*bHeightOnly=*/false);
 	}
@@ -1382,7 +1383,8 @@ UPCGData* UPCGComponent::GetLandscapePCGData()
 
 UPCGData* UPCGComponent::GetLandscapeHeightPCGData()
 {
-	if (!CachedLandscapeHeightData)
+
+	if (!CachedLandscapeHeightData || IsLandscapeCachedDataDirty(CachedLandscapeHeightData))
 	{
 		CachedLandscapeHeightData = CreateLandscapePCGData(/*bHeightOnly=*/true);
 	}
@@ -1510,7 +1512,7 @@ UPCGData* UPCGComponent::CreateActorPCGData(AActor* Actor)
 	else if (ALandscapeProxy* Landscape = Cast<ALandscapeProxy>(Actor))
 	{
 		UPCGLandscapeData* Data = NewObject<UPCGLandscapeData>(this);
-		Data->Initialize(Landscape, GetGridBounds(Actor), /*bHeightOnly=*/false);
+		Data->Initialize(Landscape, GetGridBounds(Actor), /*bHeightOnly=*/false, /*bUseMetadata=*/Graph && Graph->bLandscapeUsesMetadata);
 
 		return Data;
 	}
@@ -1694,7 +1696,7 @@ UPCGData* UPCGComponent::CreateLandscapePCGData(bool bHeightOnly)
 
 	// TODO: we're creating separate landscape data instances here so we can do some tweaks on it (such as storing the right target actor) but this probably should change
 	UPCGLandscapeData* LandscapeData = NewObject<UPCGLandscapeData>(this);
-	LandscapeData->Initialize(Landscape, GetGridBounds(Landscape), bHeightOnly);
+	LandscapeData->Initialize(Landscape, GetGridBounds(Landscape), bHeightOnly, /*bUseMetadata=*/Graph && Graph->bLandscapeUsesMetadata);
 	// Need to override target actor for this one, not the landscape
 	LandscapeData->TargetActor = Actor;
 
@@ -1757,6 +1759,18 @@ UPCGData* UPCGComponent::CreateInputPCGData()
 	}
 }
 
+bool UPCGComponent::IsLandscapeCachedDataDirty(const UPCGData* Data) const
+{
+	bool IsCacheDirty = false;
+
+	if (const UPCGLandscapeData* CachedData = Cast<UPCGLandscapeData>(Data))
+	{
+		IsCacheDirty = Graph && (CachedData->IsUsingMetadata() != Graph->bLandscapeUsesMetadata);
+	}
+
+	return IsCacheDirty;
+}
+
 FBox UPCGComponent::GetGridBounds() const
 {
 	return GetGridBounds(GetOwner());
diff --git a/Engine/Plugins/Experimental/PCG/Source/PCG/Public/Data/PCGLandscapeData.h b/Engine/Plugins/Experimental/PCG/Source/PCG/Public/Data/PCGLandscapeData.h
--- a/Engine/Plugins/Experimental/PCG/Source/PCG/Public/Data/PCGLandscapeData.h
+++ b/Engine/Plugins/Experimental/PCG/Source/PCG/Public/Data/PCGLandscapeData.h
@@ -17,7 +17,7 @@ class PCG_API UPCGLandscapeData : public UPCGSurfaceData
 	GENERATED_BODY()
 
 public:
-	void Initialize(ALandscapeProxy* InLandscape, const FBox& InBounds, bool bInHeightOnly);
+	void Initialize(ALandscapeProxy* InLandscape, const FBox& InBounds, bool bInHeightOnly, bool bInUseMetadata);
 
 	// ~Begin UPCGData interface
 	virtual EPCGDataType GetDataType() const override { return EPCGDataType::Landscape | Super::GetDataType(); }
@@ -40,6 +40,8 @@ public:
 	UPROPERTY(BlueprintReadWrite, EditAnywhere, Category = SourceData)
 	TSoftObjectPtr<ALandscapeProxy> Landscape;
 
+	bool IsUsingMetadata() const { return bUseMetadata; }
+
 protected:
 	UPROPERTY()
 	FBox Bounds = FBox(EForceInit::ForceInit);
@@ -47,6 +49,9 @@ protected:
 	UPROPERTY()
 	bool bHeightOnly = false;
 
+	UPROPERTY()
+	bool bUseMetadata = true;
+
 private:
 	// Transient data
 	ULandscapeInfo* LandscapeInfo = nullptr;
diff --git a/Engine/Plugins/Experimental/PCG/Source/PCG/Public/PCGComponent.h b/Engine/Plugins/Experimental/PCG/Source/PCG/Public/PCGComponent.h
--- a/Engine/Plugins/Experimental/PCG/Source/PCG/Public/PCGComponent.h
+++ b/Engine/Plugins/Experimental/PCG/Source/PCG/Public/PCGComponent.h
@@ -199,6 +199,7 @@ private:
 	UPCGData* CreateActorPCGData(AActor* Actor);
 	UPCGData* CreateLandscapePCGData(bool bHeightOnly);
 	void UpdatePCGExclusionData();
+	bool IsLandscapeCachedDataDirty(const UPCGData* Data) const;
 
 	bool ShouldGenerate(bool bForce, EPCGComponentGenerationTrigger RequestedGenerationTrigger) const;
 	void GenerateLocal(bool bForce, EPCGComponentGenerationTrigger RequestedGenerationTrigger);
diff --git a/Engine/Plugins/Experimental/PCG/Source/PCG/Public/PCGGraph.h b/Engine/Plugins/Experimental/PCG/Source/PCG/Public/PCGGraph.h
--- a/Engine/Plugins/Experimental/PCG/Source/PCG/Public/PCGGraph.h
+++ b/Engine/Plugins/Experimental/PCG/Source/PCG/Public/PCGGraph.h
@@ -40,6 +40,9 @@ public:
 	FText Description;
 #endif
 
+	UPROPERTY(EditAnywhere, Category = Settings)
+	bool bLandscapeUsesMetadata = true;
+
 	/** Creates a default node based on the settings class wanted. Returns the newly created node. */
 	UFUNCTION(BlueprintCallable, Category = Graph, meta=(DeterminesOutputType = "InSettingsClass", DynamicOutputParam = "DefaultNodeSettings"))
 	UPCGNode* AddNodeOfType(TSubclassOf<class UPCGSettings> InSettingsClass, UPCGSettings*& DefaultNodeSettings);
