[PCG] Add profiling window for graphs

diff --git a/Engine/Plugins/Experimental/PCG/Source/PCG/Private/PCGElement.cpp b/Engine/Plugins/Experimental/PCG/Source/PCG/Private/PCGElement.cpp
--- a/Engine/Plugins/Experimental/PCG/Source/PCG/Private/PCGElement.cpp
+++ b/Engine/Plugins/Experimental/PCG/Source/PCG/Private/PCGElement.cpp
@@ -75,8 +75,24 @@ bool IPCGElement::Execute(FPCGContext* Context) const
 
 #if WITH_EDITOR
 		const double EndTime = FPlatformTime::Seconds();
-		Context->ElapsedTime += (EndTime - StartTime);
+		const double ElapsedTime = EndTime - StartTime;
+		Context->ElapsedTime += ElapsedTime;
 		Context->ExecutionCount++;
+
+		// TODO: Is hardcoded, should be configurable
+		constexpr int MaxNumberOfTrackedTimers = 100;
+		{
+			FScopeLock Lock(&TimersLock);
+			if (Timers.Num() < MaxNumberOfTrackedTimers)
+			{
+				Timers.Add(ElapsedTime);
+			}
+			else
+			{
+				Timers[CurrentTimerIndex] = ElapsedTime;
+			}
+			CurrentTimerIndex = (CurrentTimerIndex + 1) % MaxNumberOfTrackedTimers;
+		}
 #endif
 
 		if (bDone)
@@ -150,6 +166,14 @@ void IPCGElement::DebugDisplay(FPCGContext* Context) const
 		}
 	}
 }
+
+void IPCGElement::ResetTimers()
+{
+	FScopeLock Lock(&TimersLock);
+	Timers.Empty();
+	CurrentTimerIndex = 0;
+
+}
 #endif // WITH_EDITOR
 
 void IPCGElement::CleanupAndValidateOutput(FPCGContext* Context) const
diff --git a/Engine/Plugins/Experimental/PCG/Source/PCG/Public/PCGElement.h b/Engine/Plugins/Experimental/PCG/Source/PCG/Public/PCGElement.h
--- a/Engine/Plugins/Experimental/PCG/Source/PCG/Public/PCGElement.h
+++ b/Engine/Plugins/Experimental/PCG/Source/PCG/Public/PCGElement.h
@@ -44,6 +44,8 @@ public:
 	/** Note: the following methods must be called from the main thread */
 #if WITH_EDITOR
 	void DebugDisplay(FPCGContext* Context) const;
+	const TArray<double>& GetTimers() const { return Timers; }
+	void ResetTimers();
 #endif
 
 protected:
@@ -56,6 +58,17 @@ protected:
 
 private:
 	void CleanupAndValidateOutput(FPCGContext* Context) const;
+
+#if WITH_EDITOR
+	// Set mutable because we need to modify them in the execute call, which is const
+	// TODO: Should be a map with PCG Components. We need a mechanism to make sure that this map is cleaned up when component doesn't exist anymore.
+	// For now, it will track all calls to execute (excluding call where the result is already in cache).
+	mutable TArray<double> Timers;
+	mutable int CurrentTimerIndex = 0;
+	// Perhaps overkill but there is a slight chance that we need to protect the timers array. If we call reset from the UI while an element is executing,
+	// it could crash while writing to the timers array.
+	mutable FCriticalSection TimersLock;
+#endif // WITH_EDITOR
 };
 
 /**
diff --git a/Engine/Plugins/Experimental/PCG/Source/PCGEditor/Private/PCGEditor.cpp b/Engine/Plugins/Experimental/PCG/Source/PCGEditor/Private/PCGEditor.cpp
--- a/Engine/Plugins/Experimental/PCG/Source/PCGEditor/Private/PCGEditor.cpp
+++ b/Engine/Plugins/Experimental/PCG/Source/PCGEditor/Private/PCGEditor.cpp
@@ -28,6 +28,7 @@
 #include "SPCGEditorGraphDeterminism.h"
 #include "SPCGEditorGraphFind.h"
 #include "SPCGEditorGraphNodePalette.h"
+#include "SPCGEditorGraphProfilingView.h"
 
 #include "AssetToolsModule.h"
 #include "EdGraphUtilities.h"
@@ -67,6 +68,7 @@ namespace FPCGEditor_private
 	const FName AttributesID = FName(TEXT("Attributes"));
 	const FName FindID = FName(TEXT("Find"));
 	const FName DeterminismID = FName(TEXT("Determinism"));
+	const FName ProfilingID = FName(TEXT("Profiling"));
 }
 
 void FPCGEditor::Initialize(const EToolkitMode::Type InMode, const TSharedPtr<class IToolkitHost>& InToolkitHost, UPCGGraph* InPCGGraph)
@@ -90,6 +92,7 @@ void FPCGEditor::Initialize(const EToolkitMode::Type InMode, const TSharedPtr<cl
 	FindWidget = CreateFindWidget();
 	AttributesWidget = CreateAttributesWidget();
 	DeterminismWidget = CreateDeterminismWidget();
+	ProfilingWidget = CreateProfilingWidget();
 
 	BindCommands();
 	RegisterToolbar();
@@ -209,6 +212,10 @@ void FPCGEditor::RegisterTabSpawners(const TSharedRef<FTabManager>& InTabManager
 	InTabManager->RegisterTabSpawner(FPCGEditor_private::DeterminismID, FOnSpawnTab::CreateSP(this, &FPCGEditor::SpawnTab_Determinism))
 		.SetDisplayName(LOCTEXT("DeterminismTab", "Determinism"))
 		.SetGroup(WorkspaceMenuCategoryRef);
+
+	InTabManager->RegisterTabSpawner(FPCGEditor_private::ProfilingID, FOnSpawnTab::CreateSP(this, &FPCGEditor::SpawnTab_Profiling))
+		.SetDisplayName(LOCTEXT("ProfilingTab", "Profiling"))
+		.SetGroup(WorkspaceMenuCategoryRef);
 }
 
 void FPCGEditor::UnregisterTabSpawners(const TSharedRef<class FTabManager>& InTabManager)
@@ -1417,6 +1424,11 @@ TSharedRef<SPCGEditorGraphDeterminismListView> FPCGEditor::CreateDeterminismWidg
 	return SNew(SPCGEditorGraphDeterminismListView, SharedThis(this));
 }
 
+TSharedRef<SPCGEditorGraphProfilingView> FPCGEditor::CreateProfilingWidget()
+{
+	return SNew(SPCGEditorGraphProfilingView, SharedThis(this));
+}
+
 void FPCGEditor::OnSelectedNodesChanged(const TSet<UObject*>& NewSelection)
 {
 	TArray<TWeakObjectPtr<UObject>> SelectedObjects;
@@ -1570,4 +1582,14 @@ TSharedRef<SDockTab> FPCGEditor::SpawnTab_Determinism(const FSpawnTabArgs& Args)
 		];
 }
 
+TSharedRef<SDockTab> FPCGEditor::SpawnTab_Profiling(const FSpawnTabArgs& Args)
+{
+	return SNew(SDockTab)
+		.Label(LOCTEXT("PCGProfilingTitle", "Profiling"))
+		.TabColorScale(GetTabColorScale())
+		[
+			ProfilingWidget.ToSharedRef()
+		];
+}
+
 #undef LOCTEXT_NAMESPACE
diff --git a/Engine/Plugins/Experimental/PCG/Source/PCGEditor/Private/PCGEditor.h b/Engine/Plugins/Experimental/PCG/Source/PCGEditor/Private/PCGEditor.h
--- a/Engine/Plugins/Experimental/PCG/Source/PCGEditor/Private/PCGEditor.h
+++ b/Engine/Plugins/Experimental/PCG/Source/PCGEditor/Private/PCGEditor.h
@@ -15,6 +15,7 @@ class SPCGEditorGraphAttributeListView;
 class SPCGEditorGraphDeterminismListView;
 class SPCGEditorGraphFind;
 class SPCGEditorGraphNodePalette;
+class SPCGEditorGraphProfilingView;
 class UPCGComponent;
 class UPCGEditorGraph;
 class UPCGEditorGraphNodeBase;
@@ -166,6 +167,9 @@ private:
 	/** Create a new determinism tab widget */
 	TSharedRef<SPCGEditorGraphDeterminismListView> CreateDeterminismWidget();
 
+	/** Create a new profiling tab widget */
+	TSharedRef<SPCGEditorGraphProfilingView> CreateProfilingWidget();
+
 	/** Called when the selection changes in the GraphEditor */
 	void OnSelectedNodesChanged(const TSet<UObject*>& NewSelection);
 
@@ -195,6 +199,7 @@ private:
 	TSharedRef<SDockTab> SpawnTab_Attributes(const FSpawnTabArgs& Args);
 	TSharedRef<SDockTab> SpawnTab_Find(const FSpawnTabArgs& Args);
 	TSharedRef<SDockTab> SpawnTab_Determinism(const FSpawnTabArgs& Args);
+	TSharedRef<SDockTab> SpawnTab_Profiling(const FSpawnTabArgs& Args);
 
 	TSharedPtr<SGraphEditor> GraphEditorWidget;
 	TSharedPtr<IDetailsView> PropertyDetailsWidget;
@@ -202,6 +207,7 @@ private:
 	TSharedPtr<SPCGEditorGraphFind> FindWidget;
 	TSharedPtr<SPCGEditorGraphAttributeListView> AttributesWidget;
 	TSharedPtr<SPCGEditorGraphDeterminismListView> DeterminismWidget;
+	TSharedPtr<SPCGEditorGraphProfilingView> ProfilingWidget;
 
 	TSharedPtr<FUICommandList> GraphEditorCommands;
 
diff --git a/Engine/Plugins/Experimental/PCG/Source/PCGEditor/Private/SPCGEditorGraphProfilingView.cpp b/Engine/Plugins/Experimental/PCG/Source/PCGEditor/Private/SPCGEditorGraphProfilingView.cpp
new file mode 100644
--- /dev/null
+++ b/Engine/Plugins/Experimental/PCG/Source/PCGEditor/Private/SPCGEditorGraphProfilingView.cpp
@@ -0,0 +1,428 @@
+// Copyright Epic Games, Inc. All Rights Reserved.
+
+#include "SPCGEditorGraphProfilingView.h"
+
+#include "PCGComponent.h"
+#include "PCGEditor.h"
+#include "PCGEditorGraph.h"
+#include "PCGEditorGraphNode.h"
+#include "PCGGraph.h"
+
+#include "Styling/AppStyle.h"
+#include "Widgets/Input/SSearchBox.h"
+#include "Widgets/Layout/SScrollBox.h"
+#include "Widgets/SBoxPanel.h"
+#include "Widgets/SWidget.h"
+
+#define LOCTEXT_NAMESPACE "SPCGEditorGraphProfilingView"
+
+namespace PCGEditorGraphProfilingView
+{
+	const FText NoDataAvailableText = LOCTEXT("NoDataAvailableText", "N/A");
+	
+	/** Names of the columns in the attribute list */
+	const FName NAME_Node = FName(TEXT("Node"));
+	const FName NAME_AvgTime = FName(TEXT("AvgTime"));
+	const FName NAME_MinTime = FName(TEXT("MinTime"));
+	const FName NAME_MaxTime = FName(TEXT("MaxTime"));
+	const FName NAME_StdTime = FName(TEXT("StdTime"));
+	const FName NAME_TotalTime = FName(TEXT("TotalTime"));
+	const FName NAME_NbCalls = FName(TEXT("NbCalls"));
+
+	/** Labels of the columns */
+	const FText TEXT_NodeLabel = LOCTEXT("NodeLabel", "Node");
+	const FText TEXT_AvgTimeLabel = LOCTEXT("AvgTimeLabel", "Avg Time(ms)");
+	const FText TEXT_MinTimeLabel = LOCTEXT("MinTimeLabel", "Min Time(ms)");
+	const FText TEXT_MaxTimeLabel = LOCTEXT("MaxTimeLabel", "Max Time(ms)");
+	const FText TEXT_StdTimeLabel = LOCTEXT("StdTimeLabel", "Std(ms)");
+	const FText TEXT_TotalTimeLabel = LOCTEXT("TotalTimeLabel", "Total time(s)");
+	const FText TEXT_NbCallsLabel = LOCTEXT("NbCallsLabel", "Number of calls");
+}
+
+void SPCGProfilingListViewItemRow::Construct(const FArguments& InArgs, const TSharedRef<STableViewBase>& InOwnerTableView, const PCGProfilingListViewItemPtr& Item)
+{
+	InternalItem = Item;
+
+	SMultiColumnTableRow<PCGProfilingListViewItemPtr>::Construct(
+		SMultiColumnTableRow::FArguments()
+		.Style(FAppStyle::Get(), "DataTableEditor.CellListViewRow"),
+		InOwnerTableView);
+}
+
+TSharedRef<SWidget> SPCGProfilingListViewItemRow::GenerateWidgetForColumn(const FName& ColumnId)
+{
+	FText ColumnData = LOCTEXT("ColumnError", "Unrecognized Column");
+	if (InternalItem.IsValid())
+	{
+		if (ColumnId == PCGEditorGraphProfilingView::NAME_Node)
+		{
+			ColumnData = FText::FromName(InternalItem->Name);
+		}
+		else if (ColumnId == PCGEditorGraphProfilingView::NAME_NbCalls)
+		{
+			ColumnData = FText::AsNumber(InternalItem->NbCalls);
+		}
+		// For all other values, if we don't have data, just write "N/A"
+		else if (!InternalItem->HasData)
+		{
+			ColumnData = PCGEditorGraphProfilingView::NoDataAvailableText;
+		}
+		else if (ColumnId == PCGEditorGraphProfilingView::NAME_AvgTime)
+		{
+			// In ms
+			ColumnData = FText::AsNumber(InternalItem->AvgTime * 1000.0);
+		}
+		else if (ColumnId == PCGEditorGraphProfilingView::NAME_MinTime)
+		{
+			// In ms
+			ColumnData = FText::AsNumber(InternalItem->MinTime * 1000.0);
+		}
+		else if (ColumnId == PCGEditorGraphProfilingView::NAME_MaxTime)
+		{
+			// In ms
+			ColumnData = FText::AsNumber(InternalItem->MaxTime * 1000.0);
+		}
+		else if (ColumnId == PCGEditorGraphProfilingView::NAME_StdTime)
+		{
+			// In ms
+			ColumnData = FText::AsNumber(InternalItem->StdTime * 1000.0);
+		}
+		else if (ColumnId == PCGEditorGraphProfilingView::NAME_TotalTime)
+		{
+			// In s
+			ColumnData = FText::AsNumber(InternalItem->TotalTime);
+		}
+	}
+
+	return SNew(STextBlock).Text(ColumnData);
+}
+
+void SPCGEditorGraphProfilingView::OnItemDoubleClicked(PCGProfilingListViewItemPtr Item)
+{
+	const TSharedPtr<FPCGEditor> PCGEditor = PCGEditorPtr.Pin();
+	if (!PCGEditor.IsValid())
+	{
+		return;
+	}
+
+	if (!Item.IsValid() || !Item->EditorNode)
+	{
+		return;
+	}
+
+	PCGEditor->JumpToNode(Item->EditorNode);
+}
+
+void SPCGEditorGraphProfilingView::Construct(const FArguments& InArgs, TSharedPtr<FPCGEditor> InPCGEditor)
+{
+	PCGEditorPtr = InPCGEditor;
+
+	const TSharedPtr<FPCGEditor> PCGEditor = PCGEditorPtr.Pin();
+	if (PCGEditor)
+	{
+		PCGEditorGraph = PCGEditor->GetPCGEditorGraph();
+	}
+
+	ListViewHeader = CreateHeaderRowWidget();
+
+	const TSharedRef<SScrollBar> HorizontalScrollBar = SNew(SScrollBar)
+		.Orientation(Orient_Horizontal)
+		.Thickness(FVector2D(12.0f, 12.0f));
+
+	const TSharedRef<SScrollBar> VerticalScrollBar = SNew(SScrollBar)
+		.Orientation(Orient_Vertical)
+		.Thickness(FVector2D(12.0f, 12.0f));
+	
+	SAssignNew(ListView, SListView<PCGProfilingListViewItemPtr>)
+		.ListItemsSource(&ListViewItems)
+		.HeaderRow(ListViewHeader)
+		.OnGenerateRow(this, &SPCGEditorGraphProfilingView::OnGenerateRow)
+		.OnMouseButtonDoubleClick(this, &SPCGEditorGraphProfilingView::OnItemDoubleClicked)
+		.AllowOverscroll(EAllowOverscroll::No)
+		.ExternalScrollbar(VerticalScrollBar)
+		.ConsumeMouseWheel(EConsumeMouseWheel::Always);
+	
+	this->ChildSlot
+	[
+		SNew(SVerticalBox)
+		+SVerticalBox::Slot()
+		.AutoHeight()
+		[
+			SNew(SHorizontalBox)
+			+SHorizontalBox::Slot()
+			.AutoWidth()
+			[
+				SNew(SButton)
+				.Text(LOCTEXT("RefreshButton", "Refresh"))
+				.OnClicked(this, &SPCGEditorGraphProfilingView::Refresh)
+			]
+			+SHorizontalBox::Slot()
+			.AutoWidth()
+			[
+				SNew(SButton)
+				.Text(LOCTEXT("ResetButton", "Reset"))
+				.OnClicked(this, &SPCGEditorGraphProfilingView::ResetTimers)
+			]
+		]
+		+SVerticalBox::Slot()
+		.FillHeight(1.0f)
+		[
+			SNew(SHorizontalBox)
+			+SHorizontalBox::Slot()
+			.FillWidth(1.0f)
+			[
+				SNew(SScrollBox)
+				.Orientation(Orient_Horizontal)
+				.ExternalScrollbar(HorizontalScrollBar)
+				+SScrollBox::Slot()
+				[
+					ListView->AsShared()
+				]
+			]
+			+SHorizontalBox::Slot()
+			.AutoWidth()
+			[
+				VerticalScrollBar
+			]
+		]
+		+SVerticalBox::Slot()
+		.AutoHeight()
+		[
+			SNew(SHorizontalBox)
+			+SHorizontalBox::Slot()
+			[
+				HorizontalScrollBar
+			]
+		]
+	];
+
+	Refresh();
+}
+
+TSharedRef<SHeaderRow> SPCGEditorGraphProfilingView::CreateHeaderRowWidget()
+{
+	return SNew(SHeaderRow)
+		.ResizeMode(ESplitterResizeMode::FixedPosition)
+		.CanSelectGeneratedColumn(true)
+		+ SHeaderRow::Column(PCGEditorGraphProfilingView::NAME_Node)
+		.ManualWidth(150)
+		.DefaultLabel(PCGEditorGraphProfilingView::TEXT_NodeLabel)
+		.HAlignHeader(HAlign_Center)
+		.HAlignCell(HAlign_Left)
+		.OnSort(this, &SPCGEditorGraphProfilingView::OnSortColumnHeader)
+		+ SHeaderRow::Column(PCGEditorGraphProfilingView::NAME_AvgTime)
+		.ManualWidth(100)
+		.DefaultLabel(PCGEditorGraphProfilingView::TEXT_AvgTimeLabel)
+		.HAlignHeader(HAlign_Center)
+		.HAlignCell(HAlign_Right)
+		.OnSort(this, &SPCGEditorGraphProfilingView::OnSortColumnHeader)
+		+ SHeaderRow::Column(PCGEditorGraphProfilingView::NAME_MinTime)
+		.ManualWidth(100)
+		.DefaultLabel(PCGEditorGraphProfilingView::TEXT_MinTimeLabel)
+		.HAlignHeader(HAlign_Center)
+		.HAlignCell(HAlign_Right)
+		.OnSort(this, &SPCGEditorGraphProfilingView::OnSortColumnHeader)
+		+ SHeaderRow::Column(PCGEditorGraphProfilingView::NAME_MaxTime)
+		.ManualWidth(100)
+		.DefaultLabel(PCGEditorGraphProfilingView::TEXT_MaxTimeLabel)
+		.HAlignHeader(HAlign_Center)
+		.HAlignCell(HAlign_Right)
+		.OnSort(this, &SPCGEditorGraphProfilingView::OnSortColumnHeader)
+		+ SHeaderRow::Column(PCGEditorGraphProfilingView::NAME_StdTime)
+		.ManualWidth(100)
+		.DefaultLabel(PCGEditorGraphProfilingView::TEXT_StdTimeLabel)
+		.HAlignHeader(HAlign_Center)
+		.HAlignCell(HAlign_Right)
+		.OnSort(this, &SPCGEditorGraphProfilingView::OnSortColumnHeader)
+		+ SHeaderRow::Column(PCGEditorGraphProfilingView::NAME_TotalTime)
+		.ManualWidth(100)
+		.DefaultLabel(PCGEditorGraphProfilingView::TEXT_TotalTimeLabel)
+		.HAlignHeader(HAlign_Center)
+		.HAlignCell(HAlign_Right)
+		.OnSort(this, &SPCGEditorGraphProfilingView::OnSortColumnHeader)
+		+ SHeaderRow::Column(PCGEditorGraphProfilingView::NAME_NbCalls)
+		.ManualWidth(125)
+		.DefaultLabel(PCGEditorGraphProfilingView::TEXT_NbCallsLabel)
+		.HAlignHeader(HAlign_Center)
+		.HAlignCell(HAlign_Right)
+		.OnSort(this, &SPCGEditorGraphProfilingView::OnSortColumnHeader);
+}
+
+void SPCGEditorGraphProfilingView::OnSortColumnHeader(const EColumnSortPriority::Type SortPriority, const FName& ColumnId, const EColumnSortMode::Type NewSortMode)
+{
+	if (SortingColumn == ColumnId)
+	{
+		// Circling
+		SortMode = EColumnSortMode::Type((SortMode + 1) % 3);
+	}
+	else
+	{
+		SortingColumn = ColumnId;
+		SortMode = NewSortMode;
+	}
+
+	Refresh();
+}
+
+FReply SPCGEditorGraphProfilingView::ResetTimers()
+{
+	const TSharedPtr<FPCGEditor> PCGEditor = PCGEditorPtr.Pin();
+	if (!PCGEditor.IsValid())
+	{
+		return FReply::Handled();
+	}
+
+	if (!PCGEditorGraph)
+	{
+		return FReply::Handled();
+	}
+
+	TArray<UPCGEditorGraphNode*> EditorNodes;
+	PCGEditorGraph->GetNodesOfClass<UPCGEditorGraphNode>(EditorNodes);
+
+	for (UPCGEditorGraphNode* PCGEditorNode : EditorNodes)
+	{
+		if (PCGEditorNode)
+		{
+			if (UPCGNode* PCGNode = PCGEditorNode->GetPCGNode())
+			{
+				if (UPCGSettings* Settings = PCGNode->DefaultSettings)
+				{
+					if (IPCGElement* Element = Settings->GetElement().Get())
+					{
+						Element->ResetTimers();
+					}
+				}
+			}
+		}
+	}
+
+	Refresh();
+
+	return FReply::Handled();
+}
+
+FReply SPCGEditorGraphProfilingView::Refresh()
+{
+	ListViewItems.Empty();
+	ListView->RequestListRefresh();
+
+	const TSharedPtr<FPCGEditor> PCGEditor = PCGEditorPtr.Pin();
+	if (!PCGEditor.IsValid())
+	{
+		return FReply::Handled();
+	}
+
+	if (!PCGEditorGraph)
+	{
+		return FReply::Handled();
+	}
+
+	TArray<UPCGEditorGraphNode*> EditorNodes;
+	PCGEditorGraph->GetNodesOfClass<UPCGEditorGraphNode>(EditorNodes);
+	ListViewItems.Reserve(EditorNodes.Num());
+
+	for (const UPCGEditorGraphNode* PCGEditorNode : EditorNodes)
+	{
+		if (PCGEditorNode)
+		{
+			PCGProfilingListViewItemPtr ListViewItem = MakeShared<FPCGProfilingListViewItem>();
+			ListViewItem->EditorNode = PCGEditorNode;
+			ListViewItem->PCGNode = PCGEditorNode->GetPCGNode();
+
+			if (ListViewItem->PCGNode)
+			{
+				ListViewItem->Name = ListViewItem->PCGNode->GetNodeTitle();
+
+				if (const UPCGSettings* Settings = ListViewItem->PCGNode->DefaultSettings)
+				{
+					if (const IPCGElement* Element = Settings->GetElement().Get())
+					{
+						const TArray<double>& Timers = Element->GetTimers();
+						if (!Timers.IsEmpty())
+						{
+							ListViewItem->MinTime = std::numeric_limits<double>::max();
+							ListViewItem->MaxTime = std::numeric_limits<double>::min();
+							for (double v : Timers)
+							{
+								ListViewItem->TotalTime += v;
+								if (v < ListViewItem->MinTime)
+								{
+									ListViewItem->MinTime = v;
+								}
+
+								if (v > ListViewItem->MaxTime)
+								{
+									ListViewItem->MaxTime = v;
+								}
+							}
+
+							ListViewItem->AvgTime = ListViewItem->TotalTime / Timers.Num();
+
+							for (double v : Timers)
+							{
+								ListViewItem->StdTime += (ListViewItem->AvgTime - v) * (ListViewItem->AvgTime - v);
+							}
+
+							ListViewItem->StdTime = FMath::Sqrt(ListViewItem->StdTime / Timers.Num());
+
+							ListViewItem->HasData = true;
+							ListViewItem->NbCalls = Timers.Num();
+						}
+					}
+				}
+			}
+
+			ListViewItems.Add(ListViewItem);
+		}
+	}
+
+	if (SortingColumn != NAME_None && SortMode != EColumnSortMode::None)
+	{
+		Algo::Sort(ListViewItems, [this](const PCGProfilingListViewItemPtr& A, const PCGProfilingListViewItemPtr& B)
+			{
+				bool isLess = false;
+				if (SortingColumn == PCGEditorGraphProfilingView::NAME_Node)
+				{
+					isLess = A->Name.FastLess(B->Name);
+				}
+				else if (SortingColumn == PCGEditorGraphProfilingView::NAME_AvgTime)
+				{
+					isLess = A->AvgTime < B->AvgTime;
+				}
+				else if (SortingColumn == PCGEditorGraphProfilingView::NAME_MinTime)
+				{
+					isLess = A->MinTime < B->MinTime;
+				}
+				else if (SortingColumn == PCGEditorGraphProfilingView::NAME_MaxTime)
+				{
+					isLess = A->MaxTime < B->MaxTime;
+				}
+				else if (SortingColumn == PCGEditorGraphProfilingView::NAME_StdTime)
+				{
+					isLess = A->StdTime < B->StdTime;
+				}
+				else if (SortingColumn == PCGEditorGraphProfilingView::NAME_TotalTime)
+				{
+					isLess = A->TotalTime < B->TotalTime;
+				}
+				else if (SortingColumn == PCGEditorGraphProfilingView::NAME_NbCalls)
+				{
+					isLess = A->NbCalls < B->NbCalls;
+				}
+
+				return SortMode == EColumnSortMode::Ascending ? isLess : !isLess;
+			});
+	}
+
+	ListView->SetListItemsSource(ListViewItems);
+
+	return FReply::Handled();
+}
+
+TSharedRef<ITableRow> SPCGEditorGraphProfilingView::OnGenerateRow(PCGProfilingListViewItemPtr Item, const TSharedRef<STableViewBase>& OwnerTable) const
+{
+	return SNew(SPCGProfilingListViewItemRow, OwnerTable, Item);
+}
+
+#undef LOCTEXT_NAMESPACE
diff --git a/Engine/Plugins/Experimental/PCG/Source/PCGEditor/Private/SPCGEditorGraphProfilingView.h b/Engine/Plugins/Experimental/PCG/Source/PCGEditor/Private/SPCGEditorGraphProfilingView.h
new file mode 100644
--- /dev/null
+++ b/Engine/Plugins/Experimental/PCG/Source/PCGEditor/Private/SPCGEditorGraphProfilingView.h
@@ -0,0 +1,80 @@
+// Copyright Epic Games, Inc. All Rights Reserved.
+
+#pragma once
+
+#include "CoreMinimal.h"
+#include "PCGElement.h"
+#include "Widgets/Input/SButton.h"
+#include "Widgets/SCompoundWidget.h"
+#include "Widgets/Views/SHeaderRow.h"
+#include "Widgets/Views/SListView.h"
+
+class FPCGEditor;
+class UPCGComponent;
+class UPCGEditorGraphNode;
+class UPCGEditorGraph;
+class UPCGNode;
+
+struct FPCGProfilingListViewItem
+{
+	const UPCGNode* PCGNode = nullptr;
+	const UPCGEditorGraphNode* EditorNode = nullptr;
+
+	FName Name = NAME_None;
+	double AvgTime = 0.0;
+	double MinTime = 0.0;
+	double MaxTime = 0.0;
+	double StdTime = 0.0;
+	double TotalTime = 0.0;
+	int32 NbCalls = 0;
+	bool HasData = false;
+};
+
+typedef TSharedPtr<FPCGProfilingListViewItem> PCGProfilingListViewItemPtr;
+
+class SPCGProfilingListViewItemRow : public SMultiColumnTableRow<PCGProfilingListViewItemPtr>
+{
+public:
+	SLATE_BEGIN_ARGS(SPCGProfilingListViewItemRow) {}
+	SLATE_END_ARGS()
+
+	void Construct(const FArguments& InArgs, const TSharedRef<STableViewBase>& InOwnerTableView, const PCGProfilingListViewItemPtr& Item);
+
+	virtual TSharedRef<SWidget> GenerateWidgetForColumn(const FName& ColumnId) override;
+
+protected:
+	PCGProfilingListViewItemPtr InternalItem;
+};
+
+class SPCGEditorGraphProfilingView : public SCompoundWidget
+{
+public:
+	SLATE_BEGIN_ARGS(SPCGEditorGraphProfilingView) { }
+	SLATE_END_ARGS()
+
+	void Construct(const FArguments& InArgs, TSharedPtr<FPCGEditor> InPCGEditor);
+
+private:
+	TSharedRef<SHeaderRow> CreateHeaderRowWidget();
+	
+	// Callbacks
+	FReply Refresh();
+	FReply ResetTimers();
+	TSharedRef<ITableRow> OnGenerateRow(PCGProfilingListViewItemPtr Item, const TSharedRef<STableViewBase>& OwnerTable) const;
+	void OnItemDoubleClicked(PCGProfilingListViewItemPtr Item);
+	void OnSortColumnHeader(const EColumnSortPriority::Type SortPriority, const FName& ColumnId, const EColumnSortMode::Type NewSortMode);
+
+	/** Pointer back to the PCG editor that owns us */
+	TWeakPtr<FPCGEditor> PCGEditorPtr;
+
+	/** Cached PCGGraph being viewed */
+	UPCGEditorGraph* PCGEditorGraph = nullptr;
+
+	TSharedPtr<SHeaderRow> ListViewHeader;
+	TSharedPtr<SListView<PCGProfilingListViewItemPtr>> ListView;
+	TArray<PCGProfilingListViewItemPtr> ListViewItems;
+
+	// To allow sorting
+	FName SortingColumn = NAME_None;
+	EColumnSortMode::Type SortMode = EColumnSortMode::Type::None;
+};
