[PCG] Allow for basic conversion double <-> float and int32 <-> int64 for attributes in param data.

diff --git a/Engine/Plugins/Experimental/PCG/Source/PCG/Public/Helpers/PCGSettingsHelpers.h b/Engine/Plugins/Experimental/PCG/Source/PCG/Public/Helpers/PCGSettingsHelpers.h
--- a/Engine/Plugins/Experimental/PCG/Source/PCG/Public/Helpers/PCGSettingsHelpers.h
+++ b/Engine/Plugins/Experimental/PCG/Source/PCG/Public/Helpers/PCGSettingsHelpers.h
@@ -6,17 +6,53 @@
 #include "Metadata/PCGMetadataAttributeTraits.h"
 #include "Metadata/PCGMetadataAttributeTpl.h"
 
+#include <type_traits>
+
 namespace PCGSettingsHelpers
 {
+	// Utility function to return a value from an attribute that doesn't match the type
+	// Use if constexpr to remove un-necessary checks.
+	// Returns true if the conversion worked
+	template <typename T>
+	inline bool GetValueWithImplicitConversion(const FPCGMetadataAttributeBase* InAttribute, PCGMetadataEntryKey InKey, T& OutValue)
+	{
+#define ImplicitConversion(InType, OutType) if constexpr (std::is_same_v<T, OutType>) { if (InAttribute->GetTypeId() == PCG::Private::MetadataTypes<InType>::Id) { OutValue = static_cast<T>(static_cast<const FPCGMetadataAttribute<InType>*>(InAttribute)->GetValueFromItemKey(InKey)); return true; } }
+
+		ImplicitConversion(float, double);
+		ImplicitConversion(double, float);
+		ImplicitConversion(int32, int64);
+		ImplicitConversion(int64, int32);
+
+#undef ImplicitConversion
+
+		return false;
+	}
+
 	template<typename T, typename TEnableIf<!TIsEnumClass<T>::Value>::Type* = nullptr>
-	T GetValue(const FName& InName, const T& InValue, UPCGParamData* InParams, PCGMetadataEntryKey InKey)
+	T GetValue(const FName& InName, const T& InValue, UPCGParamData* InParams, PCGMetadataEntryKey InKey, bool bAllowImplicitConversion = true)
 	{
 		if (InParams)
 		{
 			const FPCGMetadataAttributeBase* MatchingAttribute = InParams->Metadata ? InParams->Metadata->GetConstAttribute(InName) : nullptr;
-			if (MatchingAttribute && MatchingAttribute->GetTypeId() == PCG::Private::MetadataTypes<T>::Id)
+			if (MatchingAttribute)
 			{
-				return static_cast<const FPCGMetadataAttribute<T>*>(MatchingAttribute)->GetValueFromItemKey(InKey);
+				if (MatchingAttribute->GetTypeId() == PCG::Private::MetadataTypes<T>::Id)
+				{
+					return static_cast<const FPCGMetadataAttribute<T>*>(MatchingAttribute)->GetValueFromItemKey(InKey);
+				}
+
+				if (bAllowImplicitConversion)
+				{
+					T OutValue{};
+					if (GetValueWithImplicitConversion<T>(MatchingAttribute, InKey, OutValue))
+					{
+						UE_LOG(LogPCG, Warning, TEXT("[GetAttributeValue] Matching attribute was found, but not the right type. Implicit conversion done (%d vs %d)"), MatchingAttribute->GetTypeId(), PCG::Private::MetadataTypes<T>::Id);
+						return OutValue;
+					}
+				}
+
+				UE_LOG(LogPCG, Error, TEXT("[GetAttributeValue] Matching attribute was found, but not the right type. %d vs %d"), MatchingAttribute->GetTypeId(), PCG::Private::MetadataTypes<T>::Id);
+				return InValue;
 			}
 			else
 			{
@@ -30,17 +66,17 @@ namespace PCGSettingsHelpers
 	}
 
 	template<typename T, typename TEnableIf<!TIsEnumClass<T>::Value>::Type* = nullptr>
-	T GetValue(const FName& InName, const T& InValue, UPCGParamData* InParams)
+	T GetValue(const FName& InName, const T& InValue, UPCGParamData* InParams, bool bAllowImplicitConversion = true)
 	{
-		return GetValue(InName, InValue, InParams, 0);
+		return GetValue(InName, InValue, InParams, 0, bAllowImplicitConversion);
 	}
 
 	template<typename T, typename TEnableIf<!TIsEnumClass<T>::Value>::Type* = nullptr>
-	T GetValue(const FName& InName, const T& InValue, UPCGParamData* InParams, const FName& InParamName)
+	T GetValue(const FName& InName, const T& InValue, UPCGParamData* InParams, const FName& InParamName, bool bAllowImplicitConversion = true)
 	{
 		if (InParams && InParamName != NAME_None)
 		{
-			return GetValue(InName, InValue, InParams, InParams->FindMetadataKey(InParamName));
+			return GetValue(InName, InValue, InParams, InParams->FindMetadataKey(InParamName), bAllowImplicitConversion);
 		}
 		else
 		{
@@ -50,40 +86,40 @@ namespace PCGSettingsHelpers
 
 	/** Specialized versions for enums */
 	template<typename T, typename TEnableIf<TIsEnumClass<T>::Value>::Type* = nullptr>
-	T GetValue(const FName& InName, const T& InValue, UPCGParamData* InParams, PCGMetadataEntryKey InKey)
+	T GetValue(const FName& InName, const T& InValue, UPCGParamData* InParams, PCGMetadataEntryKey InKey, bool bAllowImplicitConversion = true)
 	{
-		return static_cast<T>(GetValue(InName, static_cast<__underlying_type(T)>(InValue), InParams, InKey));
+		return static_cast<T>(GetValue(InName, static_cast<__underlying_type(T)>(InValue), InParams, InKey, bAllowImplicitConversion));
 	}
 
 	template<typename T, typename TEnableIf<TIsEnumClass<T>::Value>::Type* = nullptr>
-	T GetValue(const FName& InName, const T& InValue, UPCGParamData* InParams)
+	T GetValue(const FName& InName, const T& InValue, UPCGParamData* InParams, bool bAllowImplicitConversion = true)
 	{
-		return static_cast<T>(GetValue(InName, static_cast<__underlying_type(T)>(InValue), InParams));
+		return static_cast<T>(GetValue(InName, static_cast<__underlying_type(T)>(InValue), InParams, bAllowImplicitConversion));
 	}
 
 	template<typename T, typename TEnableIf<TIsEnumClass<T>::Value>::Type* = nullptr>
-	T GetValue(const FName& InName, const T& InValue, UPCGParamData* InParams, const FName& InParamName)
+	T GetValue(const FName& InName, const T& InValue, UPCGParamData* InParams, const FName& InParamName, bool bAllowImplicitConversion = true)
 	{
-		return static_cast<T>(GetValue(InName, static_cast<__underlying_type(T)>(InValue), InParams, InParamName));
+		return static_cast<T>(GetValue(InName, static_cast<__underlying_type(T)>(InValue), InParams, InParamName, bAllowImplicitConversion));
 	}
 
 	/** Specialized version for names */
 	template<>
-	FORCEINLINE FName GetValue(const FName& InName, const FName& InValue, UPCGParamData* InParams, PCGMetadataEntryKey InKey)
+	FORCEINLINE FName GetValue(const FName& InName, const FName& InValue, UPCGParamData* InParams, PCGMetadataEntryKey InKey, bool bAllowImplicitConversion)
 	{
-		return FName(GetValue(InName, InValue.ToString(), InParams, InKey));
+		return FName(GetValue(InName, InValue.ToString(), InParams, InKey, bAllowImplicitConversion));
 	}
 
 	template<>
-	FORCEINLINE FName GetValue(const FName& InName, const FName& InValue, UPCGParamData* InParams)
+	FORCEINLINE FName GetValue(const FName& InName, const FName& InValue, UPCGParamData* InParams, bool bAllowImplicitConversion)
 	{
-		return FName(GetValue(InName, InValue.ToString(), InParams));
+		return FName(GetValue(InName, InValue.ToString(), InParams, bAllowImplicitConversion));
 	}
 
 	template<>
-	FORCEINLINE FName GetValue(const FName& InName, const FName& InValue, UPCGParamData* InParams, const FName& InParamName)
+	FORCEINLINE FName GetValue(const FName& InName, const FName& InValue, UPCGParamData* InParams, const FName& InParamName, bool bAllowImplicitConversion)
 	{
-		return FName(GetValue(InName, InValue.ToString(), InParams, InParamName));
+		return FName(GetValue(InName, InValue.ToString(), InParams, InParamName, bAllowImplicitConversion));
 	}
 
 	/** Sets data from the params to a given property, matched on a name basis */
