+ Added async notification for PCG Graph Generation

diff --git a/Engine/Plugins/Experimental/PCG/Source/PCG/Private/Graph/PCGGraphExecutor.cpp b/Engine/Plugins/Experimental/PCG/Source/PCG/Private/Graph/PCGGraphExecutor.cpp
--- a/Engine/Plugins/Experimental/PCG/Source/PCG/Private/Graph/PCGGraphExecutor.cpp
+++ b/Engine/Plugins/Experimental/PCG/Source/PCG/Private/Graph/PCGGraphExecutor.cpp
@@ -3,9 +3,9 @@
 #include "PCGGraphExecutor.h"
 #include "PCGComponent.h"
 #include "PCGData.h"
-#include "PCGHelpers.h"
 #include "PCGGraph.h"
 #include "PCGGraphCompiler.h"
+#include "PCGHelpers.h"
 #include "PCGInputOutputSettings.h"
 #include "PCGSubgraph.h"
 
@@ -17,8 +17,8 @@
 
 // World partition support for in-editor workflows needs these includes
 #if WITH_EDITOR
-#include "WorldPartition/WorldPartition.h"
 #include "FileHelpers.h"
+#include "WorldPartition/WorldPartition.h"
 #endif
 
 static TAutoConsoleVariable<int32> CVarMaxNumTasks(
@@ -39,6 +39,9 @@ static TAutoConsoleVariable<bool> CVarGraphMultithreading(
 FPCGGraphExecutor::FPCGGraphExecutor(UObject* InOwner)
 	: GraphCompiler(MakeUnique<FPCGGraphCompiler>())
 	, GraphCache(InOwner)
+#if WITH_EDITOR
+	, GenerationProgressNotification(GetNotificationTextFormat())
+#endif
 {
 }
 
@@ -87,6 +90,11 @@ FPCGTaskId FPCGGraphExecutor::Schedule(UPCGGraph* Graph, UPCGComponent* SourceCo
 		FPCGGraphCompiler::OffsetNodeIds(ScheduledTask.Tasks, NextTaskId);
 		NextTaskId += ScheduledTask.Tasks.Num();
 		ScheduledId = NextTaskId - 1; // This is true because the last task is from the output node or is the post-execute task
+#if WITH_EDITOR
+		bTasksAreRunning = true;
+		RemainingTaskNum += ScheduledTask.Tasks.Num();
+		GenerationProgressNotification.Update(NextTaskId);
+#endif
 
 		// Push task (not data) dependencies on the pre-execute task
 		// Note must be done after the offset ids, otherwise we'll break the dependencies
@@ -116,6 +124,10 @@ FPCGTaskId FPCGGraphExecutor::ScheduleGeneric(TFunction<bool()> InOperation, con
 
 	// Assign task id
 	Task.NodeId = NextTaskId++;
+#if WITH_EDITOR
+	++RemainingTaskNum;
+	GenerationProgressNotification.Update(NextTaskId);
+#endif
 
 	FPCGGraphScheduleTask& ScheduledTask = ScheduledTasks.Emplace_GetRef();
 	ScheduledTask.Tasks.Add(Task);
@@ -179,7 +191,7 @@ void FPCGGraphExecutor::Execute()
 	}
 
 	ScheduledTasks.Reset();
-	
+
 	ScheduleLock.Unlock();
 
 	// TODO: add optimization phase if we've added new graph(s)/tasks to be executed
@@ -268,7 +280,7 @@ void FPCGGraphExecutor::Execute()
 						SourceComponent->StoreInspectionData(Task.Node, CachedOutput);
 					}
 #endif
-					
+
 					continue;
 				}
 
@@ -322,7 +334,7 @@ void FPCGGraphExecutor::Execute()
 		// Re-launch time-sliced tasks & launch new tasks
 		// TODO: currently we don't have any time-slicing so just launch tasks
 		const int32 LastLaunchIndex = ActiveTasks.Num() - (bMainThreadWasAvailable ? 0 : 1);
-		
+
 		// Assign resources
 		const int32 NumAdditionalThreads = ((NumTasksToLaunch > 0) ? (NumAvailableThreads / NumTasksToLaunch) : 0);
 		check(NumAdditionalThreads >= 0);
@@ -336,7 +348,7 @@ void FPCGGraphExecutor::Execute()
 
 		// Dispatch async tasks
 		TMap<int32, TFuture<bool>> Futures;
-		
+
 		for(int32 ExecutionIndex = 0; ExecutionIndex < ActiveTasks.Num() - 1; ++ExecutionIndex)
 		{
 			FPCGGraphActiveTask& ActiveTask = ActiveTasks[ExecutionIndex];
@@ -408,6 +420,9 @@ void FPCGGraphExecutor::Execute()
 				{
 					PostTaskExecute(ActiveTasks.Num() - 1);
 				}
+#if WITH_EDITOR
+				--RemainingTaskNum;
+#endif
 			}
 
 			for (int32 ExecutionIndex = ActiveTasks.Num() - 1; ExecutionIndex >= 0; --ExecutionIndex)
@@ -424,6 +439,9 @@ void FPCGGraphExecutor::Execute()
 				{
 					PostTaskExecute(ExecutionIndex);
 				}
+#if WITH_EDITOR
+				--RemainingTaskNum;
+#endif
 			}
 		}
 
@@ -535,12 +553,16 @@ void FPCGGraphExecutor::StoreResults(FPCGTaskId InTaskId, const FPCGDataCollecti
 void FPCGGraphExecutor::ClearResults()
 {
 	ScheduleLock.Lock();
-	
+
 	NextTaskId = 0;
 	OutputData.Reset();
 
 	ResultsRootSet.Clear();
 
+#if WITH_EDITOR
+	EndGenerationNotification();
+#endif
+
 	ScheduleLock.Unlock();
 }
 
@@ -615,6 +637,39 @@ void FPCGGraphExecutor::NotifyGraphChanged(UPCGGraph* InGraph)
 		GraphCompiler->NotifyGraphChanged(InGraph);
 	}
 }
+
+void FPCGGraphExecutor::UpdateGenerationNotification()
+{
+	if (!bTasksAreRunning)
+	{
+		return;
+	}
+
+	GenerationProgressNotification.Update(RemainingTaskNum);
+}
+
+void FPCGGraphExecutor::EndGenerationNotification()
+{
+	if (bTasksAreRunning)
+	{
+		bTasksAreRunning = false;
+		if (RemainingTaskNum > 0)
+		{
+			UE_LOG(LogPCG, Warning, TEXT("Remaining Task Number greater than zero"));
+			RemainingTaskNum = 0;
+		}
+
+		GenerationProgressNotification.Update(0);
+		// To reset the UI notification
+		GenerationProgressNotification = FAsyncCompilationNotification(GetNotificationTextFormat());
+	}
+}
+
+FTextFormat FPCGGraphExecutor::GetNotificationTextFormat()
+{
+	return NSLOCTEXT("PCG", "PCGGenerationNotificationFormat", "{0}|plural(one=PCG Task,other=PCG Tasks)");
+}
+
 #endif // WITH_EDITOR
 
 bool FPCGFetchInputElement::ExecuteInternal(FPCGContext* Context) const
diff --git a/Engine/Plugins/Experimental/PCG/Source/PCG/Private/Graph/PCGGraphExecutor.h b/Engine/Plugins/Experimental/PCG/Source/PCG/Private/Graph/PCGGraphExecutor.h
--- a/Engine/Plugins/Experimental/PCG/Source/PCG/Private/Graph/PCGGraphExecutor.h
+++ b/Engine/Plugins/Experimental/PCG/Source/PCG/Private/Graph/PCGGraphExecutor.h
@@ -3,14 +3,15 @@
 #pragma once
 
 #include "CoreMinimal.h"
-#include "PCGData.h"
 #include "PCGContext.h"
+#include "PCGData.h"
 #include "PCGElement.h"
-#include "PCGSubsystem.h"
 #include "PCGGraphCache.h"
+#include "PCGSubsystem.h"
 
 #if WITH_EDITOR
-#include "WorldPartition/WorldPartitionHandle.h" // Needed for FWorldPartitonReference
+#include "AsyncCompilationHelpers.h"
+#include "WorldPartition/WorldPartitionHandle.h" // Needed for FWorldPartitionReference
 #endif
 
 class UPCGPin;
@@ -83,6 +84,7 @@ public:
 
 	/** Notify compiler that graph has changed so it'll be removed from the cache */
 	void NotifyGraphChanged(UPCGGraph* InGraph);
+	void UpdateGenerationNotification();
 #endif
 
 	/** "Tick" of the graph executor. This call is NOT THREADSAFE */
@@ -102,6 +104,8 @@ private:
 #if WITH_EDITOR
 	void SaveDirtyActors();
 	void ReleaseUnusedActors();
+	void EndGenerationNotification();
+	static FTextFormat GetNotificationTextFormat();
 #endif
 
 	/** Graph compiler that turns a graph into tasks */
@@ -136,6 +140,9 @@ private:
 	TSet<FWorldPartitionReference> ActorsToRelease;
 
 	int32 CountUntilGC = 30;
+	int32 RemainingTaskNum = 0;
+	bool bTasksAreRunning = false;
+	FAsyncCompilationNotification GenerationProgressNotification;
 #endif
 };
 
diff --git a/Engine/Plugins/Experimental/PCG/Source/PCG/Private/PCGSubsystem.cpp b/Engine/Plugins/Experimental/PCG/Source/PCG/Private/PCGSubsystem.cpp
--- a/Engine/Plugins/Experimental/PCG/Source/PCG/Private/PCGSubsystem.cpp
+++ b/Engine/Plugins/Experimental/PCG/Source/PCG/Private/PCGSubsystem.cpp
@@ -84,6 +84,10 @@ void UPCGSubsystem::Tick(float DeltaSeconds)
 
 	// If we have any tasks to execute, schedule some
 	GraphExecutor->Execute();
+
+#if WITH_EDITOR
+	GraphExecutor->UpdateGenerationNotification();
+#endif
 }
 
 APCGWorldActor* UPCGSubsystem::GetPCGWorldActor()
@@ -286,8 +290,8 @@ TArray<FPCGTaskId> UPCGSubsystem::ScheduleMultipleComponent(UPCGComponent* Origi
 				// Add check to avoid infinite loop
 				if (ensure(!LocalComponent->IsPartitioned()))
 				{
-				// Ensure that the PCG actor match our original
-				LocalComponent->SetPropertiesFromOriginal(OriginalComponent);
+					// Ensure that the PCG actor match our original
+					LocalComponent->SetPropertiesFromOriginal(OriginalComponent);
 
 					FPCGTaskId LocalTask = ScheduleComponent(LocalComponent, /*bSave=*/ false, Dependencies);
 
