+ Adds a default determinism test suite option for native PCG Nodes

diff --git a/Engine/Plugins/Experimental/PCG/Source/PCG/Private/Tests/Determinism/PCGDeterminismNativeTests.cpp b/Engine/Plugins/Experimental/PCG/Source/PCG/Private/Tests/Determinism/PCGDeterminismNativeTests.cpp
new file mode 100644
--- /dev/null
+++ b/Engine/Plugins/Experimental/PCG/Source/PCG/Private/Tests/Determinism/PCGDeterminismNativeTests.cpp
@@ -0,0 +1,35 @@
+// Copyright Epic Games, Inc. All Rights Reserved.
+#if WITH_EDITOR
+#include "Tests/Determinism/PCGDeterminismNativeTests.h"
+
+#include "PCGSettings.h"
+
+#include "Elements/PCGDifferenceElement.h"
+#include "Tests/Determinism/PCGDifferenceDeterminismTest.h"
+
+namespace PCGDeterminismTests
+{
+	TFunction<bool()> GetNativeTestIfExists(const UPCGSettings* PCGSettings)
+	{
+		// This way we only build it if we ever need it
+		static TMap<UClass*, TFunction<bool()>> NativeTestMapping;
+
+		if (NativeTestMapping.IsEmpty())
+		{
+			NativeTestMapping =
+			{
+				{UPCGDifferenceSettings::StaticClass(), PCGDeterminismTests::DifferenceElement::RunTestSuite}
+				// TODO: Append other native tests here
+			};
+		}
+
+		if (TFunction<bool()>* TestFunction = NativeTestMapping.Find(PCGSettings->GetClass()))
+		{
+			return *TestFunction;
+		}
+
+		return nullptr;
+	}
+}
+
+#endif
\ No newline at end of file
diff --git a/Engine/Plugins/Experimental/PCG/Source/PCG/Private/Tests/Determinism/PCGDifferenceDeterminismTest.cpp b/Engine/Plugins/Experimental/PCG/Source/PCG/Private/Tests/Determinism/PCGDifferenceDeterminismTest.cpp
--- a/Engine/Plugins/Experimental/PCG/Source/PCG/Private/Tests/Determinism/PCGDifferenceDeterminismTest.cpp
+++ b/Engine/Plugins/Experimental/PCG/Source/PCG/Private/Tests/Determinism/PCGDifferenceDeterminismTest.cpp
@@ -1,10 +1,11 @@
 // Copyright Epic Games, Inc. All Rights Reserved.
-
 #if WITH_EDITOR
+#include "Tests/Determinism/PCGDifferenceDeterminismTest.h"
 
-#include "Elements/PCGDifferenceElement.h"
 #include "Tests/Determinism/PCGDeterminismTestsCommon.h"
 
+#include "Elements/PCGDifferenceElement.h"
+
 IMPLEMENT_CUSTOM_SIMPLE_AUTOMATION_TEST(FPCGDifferenceDeterminismSingleSameDataTest, FPCGTestBaseClass, "pcg.tests.Difference.Determinism.SingleSameData", PCGTestsCommon::TestFlags)
 IMPLEMENT_CUSTOM_SIMPLE_AUTOMATION_TEST(FPCGDifferenceDeterminismSingleIdenticalDataTest, FPCGTestBaseClass, "pcg.tests.Difference.Determinism.SingleMultipleData", PCGTestsCommon::TestFlags)
 IMPLEMENT_CUSTOM_SIMPLE_AUTOMATION_TEST(FPCGDifferenceDeterminismMultipleSameDataTest, FPCGTestBaseClass, "pcg.tests.Difference.Determinism.MultipleSameData", PCGTestsCommon::TestFlags)
@@ -33,7 +34,7 @@ namespace
 
 	void DifferenceTestBase(PCGTestsCommon::FTestData& TestData, EPCGDifferenceMode DifferenceMode)
 	{
-		TFunction<void(PCGTestsCommon::FTestData& TestData)> AdditionalSettingsDelegate;
+		TFunction<void(PCGTestsCommon::FTestData&)> AdditionalSettingsDelegate;
 
 		switch (DifferenceMode)
 		{
@@ -54,7 +55,7 @@ namespace
 		// Source
 		PCGDeterminismTests::AddVolumeInputData(TestData.InputData, FVector::ZeroVector, PCGDeterminismTests::Defaults::LargeVector, PCGDeterminismTests::Defaults::SmallVector, PCGDifferenceConstants::SourceLabel);
 
-		// Difference 
+		// Difference
 		PCGDeterminismTests::AddVolumeInputData(TestData.InputData, FVector::ZeroVector, PCGDeterminismTests::Defaults::MediumVector, PCGDeterminismTests::Defaults::SmallVector, PCGDifferenceConstants::DifferencesLabel);
 	}
 
@@ -72,6 +73,30 @@ namespace
 	}
 }
 
+namespace PCGDeterminismTests
+{
+	namespace DifferenceElement
+	{
+		inline bool RunTestSuite()
+		{
+			FString EmptyParameters;
+			FPCGDifferenceDeterminismSingleSameDataTest DifferenceDeterminismSingleSameDataTest(EmptyParameters);
+			FPCGDifferenceDeterminismSingleIdenticalDataTest DifferenceDeterminismSingleIdenticalDataTest(EmptyParameters);
+			FPCGDifferenceDeterminismMultipleSameDataTest DifferenceDeterminismMultipleSameDataTest(EmptyParameters);
+			FPCGDifferenceDeterminismMultipleIdenticalDataTest DifferenceDeterminismMultipleIdenticalDataTest(EmptyParameters);
+			FPCGDifferenceDeterminismOrderIndependenceTest DifferenceDeterminismOrderIndependenceTest(EmptyParameters);
+
+			bool bSuccess = true;
+			bSuccess &= DifferenceDeterminismSingleSameDataTest.RunPCGTest(EmptyParameters);
+			bSuccess &= DifferenceDeterminismSingleIdenticalDataTest.RunPCGTest(EmptyParameters);
+			bSuccess &= DifferenceDeterminismMultipleSameDataTest.RunPCGTest(EmptyParameters);
+			bSuccess &= DifferenceDeterminismMultipleIdenticalDataTest.RunPCGTest(EmptyParameters);
+			bSuccess &= DifferenceDeterminismOrderIndependenceTest.RunPCGTest(EmptyParameters);
+			return bSuccess;
+		}
+	}
+}
+
 bool FPCGDifferenceDeterminismSingleSameDataTest::RunTest(const FString& Parameters)
 {
 	// Test single same data
diff --git a/Engine/Plugins/Experimental/PCG/Source/PCG/Public/Tests/Determinism/PCGDeterminismNativeTests.h b/Engine/Plugins/Experimental/PCG/Source/PCG/Public/Tests/Determinism/PCGDeterminismNativeTests.h
new file mode 100644
--- /dev/null
+++ b/Engine/Plugins/Experimental/PCG/Source/PCG/Public/Tests/Determinism/PCGDeterminismNativeTests.h
@@ -0,0 +1,14 @@
+// Copyright Epic Games, Inc. All Rights Reserved.
+#pragma once
+#if WITH_EDITOR
+
+#include "CoreMinimal.h"
+
+class UPCGSettings;
+
+namespace PCGDeterminismTests
+{
+	PCG_API TFunction<bool()> GetNativeTestIfExists(const UPCGSettings* PCGSettings);
+}
+
+#endif
\ No newline at end of file
diff --git a/Engine/Plugins/Experimental/PCG/Source/PCG/Public/Tests/Determinism/PCGDeterminismSettings.h b/Engine/Plugins/Experimental/PCG/Source/PCG/Public/Tests/Determinism/PCGDeterminismSettings.h
--- a/Engine/Plugins/Experimental/PCG/Source/PCG/Public/Tests/Determinism/PCGDeterminismSettings.h
+++ b/Engine/Plugins/Experimental/PCG/Source/PCG/Public/Tests/Determinism/PCGDeterminismSettings.h
@@ -13,12 +13,9 @@ USTRUCT(BlueprintType)
 struct FPCGDeterminismSettings
 {
 	GENERATED_BODY()
-	
-	UPROPERTY(BlueprintReadWrite, EditAnywhere, Category = Determinism)
-	bool bBasicTests = true;
 
 	UPROPERTY(BlueprintReadWrite, EditAnywhere, Category = Determinism)
-	bool bOrderIndependenceTests = true;
+	bool bNativeTests = true;
 
 	UPROPERTY(BlueprintReadWrite, EditAnywhere, Category = Determinism)
 	bool bUseBlueprintDeterminismTest = false;
diff --git a/Engine/Plugins/Experimental/PCG/Source/PCG/Public/Tests/Determinism/PCGDifferenceDeterminismTest.h b/Engine/Plugins/Experimental/PCG/Source/PCG/Public/Tests/Determinism/PCGDifferenceDeterminismTest.h
new file mode 100644
--- /dev/null
+++ b/Engine/Plugins/Experimental/PCG/Source/PCG/Public/Tests/Determinism/PCGDifferenceDeterminismTest.h
@@ -0,0 +1,15 @@
+// Copyright Epic Games, Inc. All Rights Reserved.
+#pragma once
+#if WITH_EDITOR
+
+#include "CoreMinimal.h"
+
+namespace PCGDeterminismTests
+{
+	namespace DifferenceElement
+	{
+		PCG_API inline bool RunTestSuite();
+	}
+}
+
+#endif
\ No newline at end of file
diff --git a/Engine/Plugins/Experimental/PCG/Source/PCG/Public/Tests/PCGTestsCommon.h b/Engine/Plugins/Experimental/PCG/Source/PCG/Public/Tests/PCGTestsCommon.h
--- a/Engine/Plugins/Experimental/PCG/Source/PCG/Public/Tests/PCGTestsCommon.h
+++ b/Engine/Plugins/Experimental/PCG/Source/PCG/Public/Tests/PCGTestsCommon.h
@@ -66,6 +66,8 @@ class FPCGTestBaseClass : public FAutomationTestBase
 {
 public:
 	using FAutomationTestBase::FAutomationTestBase;
+	/** Hook to expose private test function publicly */
+	bool RunPCGTest(const FString& Parameters) { return RunTest(Parameters); }
 
 protected:
 	/** Generates all valid input combinations */
diff --git a/Engine/Plugins/Experimental/PCG/Source/PCGEditor/Private/PCGEditor.cpp b/Engine/Plugins/Experimental/PCG/Source/PCGEditor/Private/PCGEditor.cpp
--- a/Engine/Plugins/Experimental/PCG/Source/PCGEditor/Private/PCGEditor.cpp
+++ b/Engine/Plugins/Experimental/PCG/Source/PCGEditor/Private/PCGEditor.cpp
@@ -9,6 +9,7 @@
 #include "PCGInputOutputSettings.h"
 #include "PCGPin.h"
 #include "PCGSubgraph.h"
+#include "Tests/Determinism/PCGDeterminismNativeTests.h"
 #include "Tests/Determinism/PCGDeterminismTestBlueprintBase.h"
 #include "Tests/Determinism/PCGDeterminismTestsCommon.h"
 
@@ -35,6 +36,7 @@
 #include "Editor.h"
 #include "Editor/UnrealEdEngine.h"
 #include "EditorAssetLibrary.h"
+#include "Fonts/FontMeasure.h"
 #include "Framework/Commands/GenericCommands.h"
 #include "Framework/Commands/UIAction.h"
 #include "Framework/Commands/UICommandList.h"
@@ -393,11 +395,21 @@ void FPCGEditor::OnDeterminismTests()
 	DeterminismWidget->ClearItems();
 	DeterminismWidget->BuildBaseColumns();
 
-	int32 Index = 1;
 	for (UObject* Object : GraphEditorWidget->GetSelectedNodes())
 	{
 		check(Object);
 
+		// Gets an appropriate width for each new column
+		auto GetSlateTextWidth = [](const FText& Text) -> float
+		{
+			check(FSlateApplication::Get().GetRenderer());
+			const TSharedRef<FSlateFontMeasure> FontMeasure = FSlateApplication::Get().GetRenderer()->GetFontMeasureService();
+			// TODO: Verify the below property for this part of the UI
+			FSlateFontInfo FontInfo(FAppStyle::GetFontStyle("PropertyWindow.NormalFont"));
+			constexpr float Padding = 30.f;
+			return Padding + FontMeasure->Measure(Text, FontInfo).X;
+		};
+
 		if (!Object->IsA<UPCGEditorGraphNodeInput>() && !Object->IsA<UPCGEditorGraphNodeOutput>())
 		{
 			if (const UPCGEditorGraphNodeBase* PCGEditorGraphNode = Cast<const UPCGEditorGraphNodeBase>(Object))
@@ -405,25 +417,36 @@ void FPCGEditor::OnDeterminismTests()
 				const UPCGNode* PCGNode = PCGEditorGraphNode->GetPCGNode();
 				check(PCGNode && PCGNode->DefaultSettings);
 
-				// TODO: Check if the node has a native test
-
 				TSharedPtr<FDeterminismNodeTestResult> NodeResult = MakeShared<FDeterminismNodeTestResult>();
 				NodeResult->NodeTitle = PCGNode->GetNodeTitle();
 				NodeResult->NodeName = PCGNode->GetName();
 				NodeResult->Seed = PCGNode->DefaultSettings->Seed;
 
-				if (PCGNode->DefaultSettings->DeterminismSettings.bBasicTests)
+				if (PCGNode->DefaultSettings->DeterminismSettings.bNativeTests)
 				{
-					PCGDeterminismTests::FNodeTestInfo TestInfo = PCGDeterminismTests::Defaults::DeterminismBasicTestInfo;
-					PCGDeterminismTests::RunDeterminismTest(PCGNode, *NodeResult, TestInfo);
-					TestsConducted.FindOrAdd(TestInfo.TestName, {TestInfo.TestName, TestInfo.TestLabel, TestInfo.TestLabelWidth, HAlign_Center});
-				}
+					// If the settings has a native test suite
+					if (TFunction<bool()> NativeTestSuite = PCGDeterminismTests::GetNativeTestIfExists(PCGNode->DefaultSettings))
+					{
+						FName NodeName(PCGNode->GetName());
 
-				if (PCGNode->DefaultSettings->DeterminismSettings.bOrderIndependenceTests)
-				{
-					PCGDeterminismTests::FNodeTestInfo TestInfo = PCGDeterminismTests::Defaults::DeterminismOrderIndependenceInfo;
-					PCGDeterminismTests::RunDeterminismTest(PCGNode, *NodeResult, TestInfo);
-					TestsConducted.FindOrAdd(TestInfo.TestName, {TestInfo.TestName, TestInfo.TestLabel, TestInfo.TestLabelWidth, HAlign_Center});
+						bool bSuccess = NativeTestSuite();
+						NodeResult->TestResults.Emplace(NodeName, bSuccess ? EDeterminismLevel::Basic : EDeterminismLevel::NoDeterminism);
+						NodeResult->AdditionalDetails.Emplace(FString(TEXT("Native test conducted for - ")) + NodeName.ToString());
+						NodeResult->bFlagRaised = !bSuccess;
+
+						FText ColumnText = NSLOCTEXT("PCGDeterminism", "NativeTest", "Native Test");
+						TestsConducted.FindOrAdd(NodeName, {NodeName, ColumnText, GetSlateTextWidth(ColumnText), HAlign_Center});
+					}
+					else // There is no native test suite, so run the basic tests
+					{
+						PCGDeterminismTests::FNodeTestInfo BasicTestInfo = PCGDeterminismTests::Defaults::DeterminismBasicTestInfo;
+						PCGDeterminismTests::RunDeterminismTest(PCGNode, *NodeResult, BasicTestInfo);
+						TestsConducted.FindOrAdd(BasicTestInfo.TestName, {BasicTestInfo.TestName, BasicTestInfo.TestLabel, BasicTestInfo.TestLabelWidth, HAlign_Center});
+
+						PCGDeterminismTests::FNodeTestInfo OrderIndependenceTestInfo = PCGDeterminismTests::Defaults::DeterminismOrderIndependenceInfo;
+						PCGDeterminismTests::RunDeterminismTest(PCGNode, *NodeResult, OrderIndependenceTestInfo);
+						TestsConducted.FindOrAdd(OrderIndependenceTestInfo.TestName, {OrderIndependenceTestInfo.TestName, OrderIndependenceTestInfo.TestLabel, OrderIndependenceTestInfo.TestLabelWidth, HAlign_Center});
+					}
 				}
 
 				// Custom tests
@@ -433,8 +456,8 @@ void FPCGEditor::OnDeterminismTests()
 					Blueprint.GetDefaultObject()->ExecuteTest(PCGNode, *NodeResult);
 					FName BlueprintName(Blueprint->GetName());
 
-					// TODO: Add a dynamic tab size guess for the width
-					TestsConducted.FindOrAdd(BlueprintName, {BlueprintName, FText::FromString(Blueprint->GetName()), 140.f, HAlign_Center});
+					FText ColumnText = FText::FromString(Blueprint->GetName());
+					TestsConducted.FindOrAdd(BlueprintName, {BlueprintName, ColumnText, GetSlateTextWidth(ColumnText), HAlign_Center});
 				}
 
 				DeterminismWidget->AddItem(NodeResult);
