+ Added UI button in PCG Graph to disable automatic regeneration of PCG graph by temporarily disabling notification

diff --git a/Engine/Plugins/Experimental/PCG/Source/PCG/Private/PCGGraph.cpp b/Engine/Plugins/Experimental/PCG/Source/PCG/Private/PCGGraph.cpp
--- a/Engine/Plugins/Experimental/PCG/Source/PCG/Private/PCGGraph.cpp
+++ b/Engine/Plugins/Experimental/PCG/Source/PCG/Private/PCGGraph.cpp
@@ -529,7 +529,7 @@ void UPCGGraph::NotifyGraphChanged(bool bIsStructural)
 	}
 
 	// Skip recursive cases which can happen either through direct recursivity (A -> A) or indirectly (A -> B -> A)
-	if (bIsNotifying)
+	if (bSubsystemNotificationDisabled || bIsNotifying)
 	{
 		return;
 	}
diff --git a/Engine/Plugins/Experimental/PCG/Source/PCG/Public/PCGGraph.h b/Engine/Plugins/Experimental/PCG/Source/PCG/Public/PCGGraph.h
--- a/Engine/Plugins/Experimental/PCG/Source/PCG/Public/PCGGraph.h
+++ b/Engine/Plugins/Experimental/PCG/Source/PCG/Public/PCGGraph.h
@@ -87,6 +87,8 @@ public:
 
 	const TArray<TObjectPtr<UObject>>& GetExtraEditorNodes() const { return ExtraEditorNodes; }
 	void SetExtraEditorNodes(const TArray<TObjectPtr<const UObject>>& InNodes);
+	void ToggleSubsystemNotifications() { bSubsystemNotificationDisabled = !bSubsystemNotificationDisabled; }
+	bool SubsystemNotificationsAreDisabled() const { return bSubsystemNotificationDisabled; }
 #endif
 
 #if WITH_EDITOR
@@ -127,5 +129,6 @@ private:
 	bool bDelayedChangeNotification = false;
 	bool bDelayedChangeNotificationStructural = false;
 	bool bIsNotifying = false;
+	bool bSubsystemNotificationDisabled = false;
 #endif // WITH_EDITOR
 };
diff --git a/Engine/Plugins/Experimental/PCG/Source/PCGEditor/Private/PCGEditor.cpp b/Engine/Plugins/Experimental/PCG/Source/PCGEditor/Private/PCGEditor.cpp
--- a/Engine/Plugins/Experimental/PCG/Source/PCGEditor/Private/PCGEditor.cpp
+++ b/Engine/Plugins/Experimental/PCG/Source/PCGEditor/Private/PCGEditor.cpp
@@ -289,10 +289,17 @@ void FPCGEditor::RegisterToolbar() const
 	{
 		FToolMenuSection& Section = ToolBar->AddSection("PCGToolbar", TAttribute<FText>(), InsertAfterAssetSection);
 		Section.AddEntry(FToolMenuEntry::InitToolBarButton(
- 			PCGEditorCommands.Find,
- 			TAttribute<FText>(),
- 			TAttribute<FText>(),
- 			FSlateIcon(FAppStyle::GetAppStyleSetName(), "BlueprintEditor.FindInBlueprint"))); // TODO, use own application style?
+			PCGEditorCommands.Find,
+			TAttribute<FText>(),
+			TAttribute<FText>(),
+			FSlateIcon(FAppStyle::GetAppStyleSetName(), "BlueprintEditor.FindInBlueprint"))); // TODO, use own application style?
+
+		Section.AddEntry(FToolMenuEntry::InitToolBarButton(
+			FPCGEditorCommands::Get().PauseAutoRegeneration,
+			TAttribute<FText>(),
+			TAttribute<FText>(),
+			FSlateIcon(FAppStyle::GetAppStyleSetName(), "PlayWorld.PausePlaySession")));
+
 		Section.AddSeparator(NAME_None);
 		Section.AddDynamicEntry("Debugging", FNewToolMenuSectionDelegate::CreateLambda([](FToolMenuSection& InSection)
 		{
@@ -313,6 +320,12 @@ void FPCGEditor::BindCommands()
 		PCGEditorCommands.Find,
 		FExecuteAction::CreateSP(this, &FPCGEditor::OnFind));
 
+	ToolkitCommands->MapAction(
+		PCGEditorCommands.PauseAutoRegeneration,
+		FExecuteAction::CreateSP(this, &FPCGEditor::OnPauseAutomaticRegeneration_Clicked),
+		FCanExecuteAction(),
+		FIsActionChecked::CreateSP(this, &FPCGEditor::IsAutomaticRegenerationPaused));
+
 	GraphEditorCommands->MapAction(
 		PCGEditorCommands.CollapseNodes,
 		FExecuteAction::CreateSP(this, &FPCGEditor::OnCollapseNodesInSubgraph),
@@ -336,19 +349,19 @@ void FPCGEditor::BindCommands()
 		FExecuteAction::CreateSP(this, &FPCGEditor::OnSetExecutionMode, EPCGSettingsExecutionMode::Enabled),
 		FCanExecuteAction(),
 		FIsActionChecked::CreateSP(this, &FPCGEditor::IsExecutionModeActive, EPCGSettingsExecutionMode::Enabled));
-	
+
 	GraphEditorCommands->MapAction(
 		PCGEditorCommands.ExecutionModeDebug,
 		FExecuteAction::CreateSP(this, &FPCGEditor::OnSetExecutionMode, EPCGSettingsExecutionMode::Debug),
 		FCanExecuteAction(),
 		FIsActionChecked::CreateSP(this, &FPCGEditor::IsExecutionModeActive, EPCGSettingsExecutionMode::Debug));
-	
+
 	GraphEditorCommands->MapAction(
 		PCGEditorCommands.ExecutionModeDisabled,
 		FExecuteAction::CreateSP(this, &FPCGEditor::OnSetExecutionMode, EPCGSettingsExecutionMode::Disabled),
 		FCanExecuteAction(),
 		FIsActionChecked::CreateSP(this, &FPCGEditor::IsExecutionModeActive, EPCGSettingsExecutionMode::Disabled));
-	
+
 	GraphEditorCommands->MapAction(
 		PCGEditorCommands.ExecutionModeIsolated,
 		FExecuteAction::CreateSP(this, &FPCGEditor::OnSetExecutionMode, EPCGSettingsExecutionMode::Isolated),
@@ -365,6 +378,26 @@ void FPCGEditor::OnFind()
 	}
 }
 
+void FPCGEditor::OnPauseAutomaticRegeneration_Clicked()
+{
+	if (!PCGGraphBeingEdited)
+	{
+		return;
+	}
+
+	PCGGraphBeingEdited->ToggleSubsystemNotifications();
+
+	if (!PCGGraphBeingEdited->SubsystemNotificationsAreDisabled())
+	{
+		PCGGraphBeingEdited->NotifyGraphChanged(/*bIsStructural=*/true);
+	}
+}
+
+bool FPCGEditor::IsAutomaticRegenerationPaused() const
+{
+	return PCGGraphBeingEdited && PCGGraphBeingEdited->SubsystemNotificationsAreDisabled();
+}
+
 bool FPCGEditor::CanRunDeterminismTests() const
 {
 	check(GraphEditorWidget.IsValid());
@@ -911,26 +944,26 @@ bool FPCGEditor::IsExecutionModeActive(EPCGSettingsExecutionMode InExecutionMode
 			{
 				continue;
 			}
-			
+
 			const UPCGNode* PCGNode = PCGEditorGraphNode->GetPCGNode();
 			if (!PCGNode)
 			{
 				continue;
 			}
 
-			const UPCGSettings* PCGSettings = PCGNode->DefaultSettings; 
+			const UPCGSettings* PCGSettings = PCGNode->DefaultSettings;
 			if (!PCGSettings)
 			{
-				continue;				
+				continue;
 			}
-			
+
 			if (PCGSettings->ExecutionMode == InExecutionMode)
 			{
 				return true;
 			}
 		}
 	}
-	
+
 	return false;
 }
 
diff --git a/Engine/Plugins/Experimental/PCG/Source/PCGEditor/Private/PCGEditor.h b/Engine/Plugins/Experimental/PCG/Source/PCGEditor/Private/PCGEditor.h
--- a/Engine/Plugins/Experimental/PCG/Source/PCGEditor/Private/PCGEditor.h
+++ b/Engine/Plugins/Experimental/PCG/Source/PCGEditor/Private/PCGEditor.h
@@ -36,13 +36,13 @@ public:
 
 	/** Sets the PCG component we want to debug */
 	void SetPCGComponentBeingDebugged(UPCGComponent* InPCGComponent);
-	
+
 	/** Gets the PCG component we are debugging */
 	UPCGComponent* GetPCGComponentBeingDebugged() const { return PCGComponentBeingDebugged; }
 
 	/** Sets the PCG node we want to inspect */
 	void SetPCGNodeBeingInspected(UPCGNode* InPCGNode);
-	
+
 	/** Gets the PCG node we are inspecting */
 	UPCGNode* GetPCGNodeBeingInspected() const { return PCGNodeBeingInspected; }
 
@@ -68,7 +68,7 @@ public:
 	virtual void OnClose() override;
 	virtual void InitToolMenuContext(FToolMenuContext& MenuContext) override;
 	// ~End FAssetEditorToolkit interface
-	
+
 	FOnDebugObjectChanged OnDebugObjectChangedDelegate;
 	FOnInspectedNodeChanged OnInspectedNodeChangedDelegate;
 
@@ -88,6 +88,10 @@ private:
 	/** Bring up the find tab */
 	void OnFind();
 
+	/** Enable/Disable automatic PCG node generation */
+	void OnPauseAutomaticRegeneration_Clicked();
+	bool IsAutomaticRegenerationPaused() const;
+
 	/** Start inspecting the current selected node */
 	void OnStartInspectNode();
 	/** Stop inspecting the current inspected node */
@@ -98,12 +102,12 @@ private:
 	/** Summon the Determinism tab */
 	void OnDeterminismTests();
 
-	/** Whether or not an execution mode is active for the selected ndoes */
+	/** Whether or not an execution mode is active for the selected nodes */
 	bool IsExecutionModeActive(EPCGSettingsExecutionMode InExecutionMode) const;
-		
+
 	/** Set execution mode for selected nodes */
 	void OnSetExecutionMode(EPCGSettingsExecutionMode InExecutionMode);
-	
+
 	/** Select every node in the graph */
 	void SelectAllNodes();
 	/** Whether we can select every node */
diff --git a/Engine/Plugins/Experimental/PCG/Source/PCGEditor/Private/PCGEditorCommands.cpp b/Engine/Plugins/Experimental/PCG/Source/PCGEditor/Private/PCGEditorCommands.cpp
--- a/Engine/Plugins/Experimental/PCG/Source/PCGEditor/Private/PCGEditorCommands.cpp
+++ b/Engine/Plugins/Experimental/PCG/Source/PCGEditor/Private/PCGEditorCommands.cpp
@@ -20,6 +20,7 @@ void FPCGEditorCommands::RegisterCommands()
 {
 	UI_COMMAND(CollapseNodes, "Collapse into Subgraph", "Collapse selected nodes into a separate PCGGraph asset.", EUserInterfaceActionType::Button, FInputChord(EModifierKey::Alt, EKeys::J));
 	UI_COMMAND(Find, "Find", "Finds PCG nodes and comments in the current graph", EUserInterfaceActionType::Button, FInputChord(EModifierKey::Control, EKeys::F));
+	UI_COMMAND(PauseAutoRegeneration, "Pause Regen", "Pause Automatic Regeneration of PCG nodes.", EUserInterfaceActionType::ToggleButton, FInputChord(EModifierKey::Alt, EKeys::R));
 	UI_COMMAND(RunDeterminismTest, "Run Determinism Test", "Evaluate the current node for determinism.", EUserInterfaceActionType::Button, FInputChord(EModifierKey::Alt, EKeys::D));
 	UI_COMMAND(StartInspectNode, "Start Inspecting Node", "Start Inspecting Node", EUserInterfaceActionType::Button, FInputChord(EKeys::I));
 	UI_COMMAND(StopInspectNode, "Stop Inspecting Node", "Stop Inspecting Node", EUserInterfaceActionType::Button, FInputChord(EModifierKey::Shift, EKeys::I));
diff --git a/Engine/Plugins/Experimental/PCG/Source/PCGEditor/Private/PCGEditorCommands.h b/Engine/Plugins/Experimental/PCG/Source/PCGEditor/Private/PCGEditorCommands.h
--- a/Engine/Plugins/Experimental/PCG/Source/PCGEditor/Private/PCGEditorCommands.h
+++ b/Engine/Plugins/Experimental/PCG/Source/PCGEditor/Private/PCGEditorCommands.h
@@ -17,6 +17,7 @@ public:
 	TSharedPtr<FUICommandInfo> CollapseNodes;
 	TSharedPtr<FUICommandInfo> Find;
 	TSharedPtr<FUICommandInfo> RunDeterminismTest;
+	TSharedPtr<FUICommandInfo> PauseAutoRegeneration;
 	TSharedPtr<FUICommandInfo> StartInspectNode;
 	TSharedPtr<FUICommandInfo> StopInspectNode;
 	TSharedPtr<FUICommandInfo> ExecutionModeEnabled;
