PCG UI notification: Fixes

#rnx

diff --git a/Engine/Plugins/Experimental/PCG/Source/PCG/Private/Graph/PCGGraphExecutor.cpp b/Engine/Plugins/Experimental/PCG/Source/PCG/Private/Graph/PCGGraphExecutor.cpp
--- a/Engine/Plugins/Experimental/PCG/Source/PCG/Private/Graph/PCGGraphExecutor.cpp
+++ b/Engine/Plugins/Experimental/PCG/Source/PCG/Private/Graph/PCGGraphExecutor.cpp
@@ -125,6 +125,7 @@ FPCGTaskId FPCGGraphExecutor::ScheduleGeneric(TFunction<bool()> InOperation, con
 	// Assign task id
 	Task.NodeId = NextTaskId++;
 #if WITH_EDITOR
+	bTasksAreRunning = true;
 	++RemainingTaskNum;
 	GenerationProgressNotification.Update(NextTaskId);
 #endif
@@ -250,6 +251,9 @@ void FPCGGraphExecutor::Execute()
 				if (!Task.Element)
 				{
 					ReadyTasks.RemoveAtSwap(ReadyTaskIndex);
+#if WITH_EDITOR
+					RemainingTaskNum--;
+#endif // WITH_EDITOR
 					continue;
 				}
 
@@ -279,6 +283,8 @@ void FPCGGraphExecutor::Execute()
 					{
 						SourceComponent->StoreInspectionData(Task.Node, CachedOutput);
 					}
+
+					RemainingTaskNum--;
 #endif
 
 					continue;
@@ -560,7 +566,10 @@ void FPCGGraphExecutor::ClearResults()
 	ResultsRootSet.Clear();
 
 #if WITH_EDITOR
-	EndGenerationNotification();
+	if (ScheduledTasks.IsEmpty())
+	{
+		EndGenerationNotification();
+	}
 #endif
 
 	ScheduleLock.Unlock();
