Fix minor issue with Determinism Native tests on non-unity Mac compile
#rnx

diff --git a/Engine/Plugins/Experimental/PCG/Source/PCG/Private/Tests/Determinism/PCGDifferenceDeterminismTest.cpp b/Engine/Plugins/Experimental/PCG/Source/PCG/Private/Tests/Determinism/PCGDifferenceDeterminismTest.cpp
--- a/Engine/Plugins/Experimental/PCG/Source/PCG/Private/Tests/Determinism/PCGDifferenceDeterminismTest.cpp
+++ b/Engine/Plugins/Experimental/PCG/Source/PCG/Private/Tests/Determinism/PCGDifferenceDeterminismTest.cpp
@@ -77,7 +77,7 @@ namespace PCGDeterminismTests
 {
 	namespace DifferenceElement
 	{
-		inline bool RunTestSuite()
+		bool RunTestSuite()
 		{
 			FString EmptyParameters;
 			FPCGDifferenceDeterminismSingleSameDataTest DifferenceDeterminismSingleSameDataTest(EmptyParameters);
diff --git a/Engine/Plugins/Experimental/PCG/Source/PCG/Public/Tests/Determinism/PCGDifferenceDeterminismTest.h b/Engine/Plugins/Experimental/PCG/Source/PCG/Public/Tests/Determinism/PCGDifferenceDeterminismTest.h
--- a/Engine/Plugins/Experimental/PCG/Source/PCG/Public/Tests/Determinism/PCGDifferenceDeterminismTest.h
+++ b/Engine/Plugins/Experimental/PCG/Source/PCG/Public/Tests/Determinism/PCGDifferenceDeterminismTest.h
@@ -8,7 +8,7 @@ namespace PCGDeterminismTests
 {
 	namespace DifferenceElement
 	{
-		PCG_API inline bool RunTestSuite();
+		PCG_API bool RunTestSuite();
 	}
 }
 
