Replace TSoftObjectPtr by TWeakObjectPtr in PCGLandscapeSplineData

diff --git a/Engine/Plugins/Experimental/PCG/Source/PCG/Private/Data/PCGLandscapeSplineData.cpp b/Engine/Plugins/Experimental/PCG/Source/PCG/Private/Data/PCGLandscapeSplineData.cpp
--- a/Engine/Plugins/Experimental/PCG/Source/PCG/Private/Data/PCGLandscapeSplineData.cpp
+++ b/Engine/Plugins/Experimental/PCG/Source/PCG/Private/Data/PCGLandscapeSplineData.cpp
@@ -51,13 +51,13 @@ void UPCGLandscapeSplineData::Initialize(ULandscapeSplinesComponent* InSplineCom
 
 int UPCGLandscapeSplineData::GetNumSegments() const
 {
-	check(Spline);
+	check(Spline.IsValid());
 	return Spline->GetSegments().Num();
 }
 
 FVector::FReal UPCGLandscapeSplineData::GetSegmentLength(int SegmentIndex) const
 {
-	check(Spline);
+	check(Spline.IsValid());
 	check(SegmentIndex >= 0 && SegmentIndex < Spline->GetSegments().Num());
 	
 	const ULandscapeSplineSegment* Segment = Spline->GetSegments()[SegmentIndex];
@@ -74,7 +74,7 @@ FVector::FReal UPCGLandscapeSplineData::GetSegmentLength(int SegmentIndex) const
 
 FTransform UPCGLandscapeSplineData::GetTransformAtDistance(int SegmentIndex, FVector::FReal Distance, FBox* OutBounds) const
 {
-	check(Spline);
+	check(Spline.IsValid());
 	check(SegmentIndex >= 0 && SegmentIndex < Spline->GetSegments().Num());
 
 	const ULandscapeSplineSegment* Segment = Spline->GetSegments()[SegmentIndex];
@@ -128,7 +128,7 @@ FTransform UPCGLandscapeSplineData::GetTransformAtDistance(int SegmentIndex, FVe
 
 const UPCGPointData* UPCGLandscapeSplineData::CreatePointData(FPCGContext* Context) const
 {
-	check(Spline);
+	check(Spline.IsValid());
 	TRACE_CPUPROFILER_EVENT_SCOPE(UPCGLandscapeSplineData::CreatePointData);
 
 	UPCGPointData* Data = NewObject<UPCGPointData>(const_cast<UPCGLandscapeSplineData*>(this));
@@ -148,7 +148,7 @@ const UPCGPointData* UPCGLandscapeSplineData::CreatePointData(FPCGContext* Conte
 
 FBox UPCGLandscapeSplineData::GetBounds() const
 {
-	check(Spline);
+	check(Spline.IsValid());
 
 	FBox Bounds(EForceInit::ForceInit);
 	for (const TObjectPtr<ULandscapeSplineSegment>& Segment : Spline->GetSegments())
@@ -168,16 +168,17 @@ bool UPCGLandscapeSplineData::SamplePoint(const FTransform& InTransform, const F
 {
 	// TODO : add metadata support on poly lines
 	// TODO : add support for bounds
-	check(Spline);
+	const ULandscapeSplinesComponent* SplinePtr = Spline.Get();
+	check(SplinePtr);
 
 	OutPoint.Transform = InTransform;
 	OutPoint.SetLocalBounds(InBounds); // TODO: should maybe do Min.Z = Max.Z = 0 ?
 
-	const FVector Position = Spline->GetComponentTransform().InverseTransformPosition(OutPoint.Transform.GetLocation());
+	const FVector Position = SplinePtr->GetComponentTransform().InverseTransformPosition(OutPoint.Transform.GetLocation());
 
 	float PointDensity = 0.0f;
 
-	for(const TObjectPtr<ULandscapeSplineSegment> Segment : Spline->GetSegments())
+	for(const TObjectPtr<ULandscapeSplineSegment>& Segment : SplinePtr->GetSegments())
 	{
 		// Considering the landscape spline always exists on the landscape,
 		// we'll ignore the Z component of the input here for the bounds check.
@@ -185,7 +186,7 @@ bool UPCGLandscapeSplineData::SamplePoint(const FTransform& InTransform, const F
 		{
 			continue;
 		}
-		
+
 		const TArray<FLandscapeSplineInterpPoint>& InterpPoints = Segment->GetPoints();
 
 		float SegmentDensity = 0.0f;
diff --git a/Engine/Plugins/Experimental/PCG/Source/PCG/Public/Data/PCGLandscapeSplineData.h b/Engine/Plugins/Experimental/PCG/Source/PCG/Public/Data/PCGLandscapeSplineData.h
--- a/Engine/Plugins/Experimental/PCG/Source/PCG/Public/Data/PCGLandscapeSplineData.h
+++ b/Engine/Plugins/Experimental/PCG/Source/PCG/Public/Data/PCGLandscapeSplineData.h
@@ -37,5 +37,5 @@ public:
 
 protected:
 	UPROPERTY(BlueprintReadWrite, EditAnywhere, Category = SourceData)
-	TSoftObjectPtr<ULandscapeSplinesComponent> Spline;
+	TWeakObjectPtr<ULandscapeSplinesComponent> Spline;
 };
