Improve MT contention in PCGCopyPoints and in PCGMetadata

diff --git a/Engine/Plugins/Experimental/PCG/Source/PCG/Private/Elements/PCGCopyPoints.cpp b/Engine/Plugins/Experimental/PCG/Source/PCG/Private/Elements/PCGCopyPoints.cpp
--- a/Engine/Plugins/Experimental/PCG/Source/PCG/Private/Elements/PCGCopyPoints.cpp
+++ b/Engine/Plugins/Experimental/PCG/Source/PCG/Private/Elements/PCGCopyPoints.cpp
@@ -120,6 +120,9 @@ bool FPCGCopyPointsElement::ExecuteInternal(FPCGContext* Context) const
 		}
 	}
 
+	TArray<TTuple<int64, int64>> AllMetadataEntries;
+	AllMetadataEntries.SetNum(SourcePoints.Num() * TargetPoints.Num());
+
 	FPCGAsync::AsyncPointProcessing(Context, SourcePoints.Num() * TargetPoints.Num(), OutPoints, [&](int32 Index, FPCGPoint& OutPoint)
 	{
 		const FPCGPoint& SourcePoint = SourcePoints[Index / TargetPoints.Num()];
@@ -194,7 +197,8 @@ bool FPCGCopyPointsElement::ExecuteInternal(FPCGContext* Context) const
 			NonRootPoint = &SourcePoint;
 		}
 
-		OutPoint.MetadataEntry = OutPointData->Metadata->AddEntry(RootPoint->MetadataEntry);
+		OutPoint.MetadataEntry = OutPointData->Metadata->AddEntryPlaceholder();
+		AllMetadataEntries[Index] = TTuple<int64, int64>(OutPoint.MetadataEntry, RootPoint->MetadataEntry);
 
 		// Copy EntryToValue key mappings from NonRootAttributes
 		for (FPCGMetadataAttributeBase* NonRootAttribute : NonRootAttributes)
@@ -208,6 +212,8 @@ bool FPCGCopyPointsElement::ExecuteInternal(FPCGContext* Context) const
 		return true;
 	});
 
+	OutPointData->Metadata->AddDelayedEntries(AllMetadataEntries);
+
 	// Forward any non-input data
 	Outputs.Append(Context->InputData.GetAllSettings());
 
diff --git a/Engine/Plugins/Experimental/PCG/Source/PCG/Private/Metadata/PCGMetadata.cpp b/Engine/Plugins/Experimental/PCG/Source/PCG/Private/Metadata/PCGMetadata.cpp
--- a/Engine/Plugins/Experimental/PCG/Source/PCG/Private/Metadata/PCGMetadata.cpp
+++ b/Engine/Plugins/Experimental/PCG/Source/PCG/Private/Metadata/PCGMetadata.cpp
@@ -53,7 +53,7 @@ void UPCGMetadata::Serialize(FArchive& InArchive)
 	if (InArchive.IsLoading())
 	{
 		NextAttributeId = Attributes.Num();
-		ItemKeyOffset = (Parent ? Parent->GetItemCountForChild() : 0);
+		ItemKeyOffset = (Parent.IsValid() ? Parent->GetItemCountForChild() : 0);
 	}
 }
 
@@ -72,14 +72,14 @@ void UPCGMetadata::BeginDestroy()
 
 void UPCGMetadata::Initialize(const UPCGMetadata* InParent)
 {
-	if (Parent || Attributes.Num() != 0)
+	if (Parent.IsValid() || Attributes.Num() != 0)
 	{
 		// Already initialized; note that while that might be construed as a warning, there are legit cases where this is correct
 		return;
 	}
 
 	Parent = ((InParent != this) ? InParent : nullptr);
-	ItemKeyOffset = Parent ? Parent->GetItemCountForChild() : 0;
+	ItemKeyOffset = Parent.IsValid() ? Parent->GetItemCountForChild() : 0;
 	AddAttributes(InParent);
 }
 
@@ -91,7 +91,7 @@ void UPCGMetadata::InitializeAsCopy(const UPCGMetadata* InMetadataToCopy)
 	}
 
 	check(InMetadataToCopy);
-	if (Parent || Attributes.Num() != 0)
+	if (Parent.IsValid() || Attributes.Num() != 0)
 	{
 		UE_LOG(LogPCG, Error, TEXT("Metadata has already been initialized or already contains attributes"));
 		return;
@@ -198,7 +198,7 @@ void UPCGMetadata::CopyAttribute(const UPCGMetadata* InOther, FName AttributeToC
 
 const UPCGMetadata* UPCGMetadata::GetRoot() const
 {
-	if (Parent)
+	if (Parent.IsValid())
 	{
 		return Parent->GetRoot();
 	}
@@ -363,7 +363,7 @@ FName UPCGMetadata::GetSingleAttributeNameOrNone() const
 
 bool UPCGMetadata::ParentHasAttribute(FName AttributeName) const
 {
-	return Parent && Parent->HasAttribute(AttributeName);
+	return Parent.IsValid() && Parent->HasAttribute(AttributeName);
 }
 
 void UPCGMetadata::CreateInteger64Attribute(FName AttributeName, int64 DefaultValue, bool bAllowsInterpolation, bool bOverrideParent)
@@ -472,7 +472,7 @@ FPCGMetadataAttributeBase* UPCGMetadata::CopyAttribute(FName AttributeToCopy, FN
 	}
 	AttributeLock.ReadUnlock();
 
-	if (!OriginalAttribute && Parent)
+	if (!OriginalAttribute && Parent.IsValid())
 	{
 		OriginalAttribute = Parent->GetConstAttribute(AttributeToCopy);
 	}
@@ -575,6 +575,26 @@ int64 UPCGMetadata::AddEntry(int64 ParentEntry)
 	return ParentKeys.Add(ParentEntry) + ItemKeyOffset;
 }
 
+int64 UPCGMetadata::AddEntryPlaceholder()
+{
+	FReadScopeLock ScopeLock(ItemLock);
+	return ParentKeys.Num() + DelayedEntriesIndex.IncrementExchange() + ItemKeyOffset;
+}
+
+void UPCGMetadata::AddDelayedEntries(const TArray<TTuple<int64, int64>>& AllEntries)
+{
+	FWriteScopeLock ScopeLock(ItemLock);
+	ParentKeys.AddUninitialized(AllEntries.Num());
+	for (const TTuple<int64, int64>& Entry : AllEntries)
+	{
+		int64 Index = Entry.Get<0>() - ItemKeyOffset;
+		check(Index < ParentKeys.Num());
+		ParentKeys[Index] = Entry.Get<1>();
+	}
+
+	DelayedEntriesIndex.Exchange(0);
+}
+
 bool UPCGMetadata::InitializeOnSet(PCGMetadataEntryKey& InKey, PCGMetadataEntryKey InParentKeyA, const UPCGMetadata* InParentMetadataA, PCGMetadataEntryKey InParentKeyB, const UPCGMetadata* InParentMetadataB)
 {
 	if (InKey == PCGInvalidEntryKey)
diff --git a/Engine/Plugins/Experimental/PCG/Source/PCG/Public/Metadata/PCGMetadata.h b/Engine/Plugins/Experimental/PCG/Source/PCG/Public/Metadata/PCGMetadata.h
--- a/Engine/Plugins/Experimental/PCG/Source/PCG/Public/Metadata/PCGMetadata.h
+++ b/Engine/Plugins/Experimental/PCG/Source/PCG/Public/Metadata/PCGMetadata.h
@@ -136,6 +136,20 @@ public:
 	UFUNCTION(BlueprintCallable, Category = "PCG|Metadata")
 	int64 AddEntry(int64 ParentEntryKey = -1);
 
+	/** Advanced method.
+	*   In a MT context, we might not want to add the entry directly (because of write lock). Call this to generate an unique index in the MT context
+	*   And call AddDelayedEntries at the end when you want to add all the entries.
+	*   Make sure to not call AddEntry in a different thread or even in the same thread if you use this function, or it will change all the indexes.
+	*/
+	int64 AddEntryPlaceholder();
+
+	/** Advanced method.
+	*   If you used AddEntryPlaceholder, call this function at the end of your MT processing to add all the entries in one shot.
+	*   Make sure to call this one at the end!!
+	*   @param AllEntries Array of pairs. First item is the EntryIndex generated by AddEntryPlaceholder, the second the ParentEntryKey (cf AddEntry)
+	*/
+	void AddDelayedEntries(const TArray<TTuple<int64, int64>>& AllEntries);
+
 	/** Initializes the metadata entry key. Returns true if key set from either parent */
 	bool InitializeOnSet(PCGMetadataEntryKey& InKey, PCGMetadataEntryKey InParentKeyA = PCGInvalidEntryKey, const UPCGMetadata* InParentMetadataA = nullptr, PCGMetadataEntryKey InParentKeyB = PCGInvalidEntryKey, const UPCGMetadata* InParentMetadataB = nullptr);
 
@@ -200,12 +214,12 @@ protected:
 	void RemoveAttributeInternal(FName AttributeName);
 
 	UPROPERTY()
-	TSoftObjectPtr<const UPCGMetadata> Parent;
+	TWeakObjectPtr<const UPCGMetadata> Parent;
 
 	// Set of parents kept for streams relationship and GC collection
 	// But otherwise not used directly
 	UPROPERTY()
-	TSet<TSoftObjectPtr<const UPCGMetadata>> OtherParents;
+	TSet<TWeakObjectPtr<const UPCGMetadata>> OtherParents;
 
 	TMap<FName, FPCGMetadataAttributeBase*> Attributes;
 	PCGMetadataAttributeKey NextAttributeId = 0;
@@ -215,4 +229,6 @@ protected:
 
 	mutable FRWLock AttributeLock;
 	mutable FRWLock ItemLock;
+
+	TAtomic<int64> DelayedEntriesIndex = 0;
 };
