Fix to Pause regeneration in PCG Graph Editor
- Removed redundant boolean in checks for notifying graph
* Pause will be removed when closing the graph

#rnx

diff --git a/Engine/Plugins/Experimental/PCG/Source/PCG/Private/PCGGraph.cpp b/Engine/Plugins/Experimental/PCG/Source/PCG/Private/PCGGraph.cpp
--- a/Engine/Plugins/Experimental/PCG/Source/PCG/Private/PCGGraph.cpp
+++ b/Engine/Plugins/Experimental/PCG/Source/PCG/Private/PCGGraph.cpp
@@ -482,6 +482,20 @@ void UPCGGraph::EnableNotificationsForEditor()
 	}
 }
 
+void UPCGGraph::ToggleUserPausedNotificationsForEditor()
+{
+	if (bUserPausedNotificationsInGraphEditor)
+	{
+		EnableNotificationsForEditor();
+	}
+	else
+	{
+		DisableNotificationsForEditor();
+	}
+
+	bUserPausedNotificationsInGraphEditor = !bUserPausedNotificationsInGraphEditor;
+}
+
 void UPCGGraph::SetExtraEditorNodes(const TArray<TObjectPtr<const UObject>>& InNodes)
 {
 	ExtraEditorNodes.Empty();
@@ -529,7 +543,7 @@ void UPCGGraph::NotifyGraphChanged(bool bIsStructural)
 	}
 
 	// Skip recursive cases which can happen either through direct recursivity (A -> A) or indirectly (A -> B -> A)
-	if (bSubsystemNotificationDisabled || bIsNotifying)
+	if (bIsNotifying)
 	{
 		return;
 	}
diff --git a/Engine/Plugins/Experimental/PCG/Source/PCG/Public/PCGGraph.h b/Engine/Plugins/Experimental/PCG/Source/PCG/Public/PCGGraph.h
--- a/Engine/Plugins/Experimental/PCG/Source/PCG/Public/PCGGraph.h
+++ b/Engine/Plugins/Experimental/PCG/Source/PCG/Public/PCGGraph.h
@@ -82,13 +82,13 @@ public:
 #if WITH_EDITOR
 	void DisableNotificationsForEditor();
 	void EnableNotificationsForEditor();
+	void ToggleUserPausedNotificationsForEditor();
+	bool NotificationsForEditorArePausedByUser() const { return bUserPausedNotificationsInGraphEditor; }
 	void PreNodeUndo(UPCGNode* InPCGNode);
 	void PostNodeUndo(UPCGNode* InPCGNode);
 
 	const TArray<TObjectPtr<UObject>>& GetExtraEditorNodes() const { return ExtraEditorNodes; }
 	void SetExtraEditorNodes(const TArray<TObjectPtr<const UObject>>& InNodes);
-	void ToggleSubsystemNotifications() { bSubsystemNotificationDisabled = !bSubsystemNotificationDisabled; }
-	bool SubsystemNotificationsAreDisabled() const { return bSubsystemNotificationDisabled; }
 #endif
 
 #if WITH_EDITOR
@@ -129,6 +129,6 @@ private:
 	bool bDelayedChangeNotification = false;
 	bool bDelayedChangeNotificationStructural = false;
 	bool bIsNotifying = false;
-	bool bSubsystemNotificationDisabled = false;
+	bool bUserPausedNotificationsInGraphEditor = false;
 #endif // WITH_EDITOR
 };
diff --git a/Engine/Plugins/Experimental/PCG/Source/PCGEditor/Private/PCGEditor.cpp b/Engine/Plugins/Experimental/PCG/Source/PCGEditor/Private/PCGEditor.cpp
--- a/Engine/Plugins/Experimental/PCG/Source/PCGEditor/Private/PCGEditor.cpp
+++ b/Engine/Plugins/Experimental/PCG/Source/PCGEditor/Private/PCGEditor.cpp
@@ -385,17 +385,12 @@ void FPCGEditor::OnPauseAutomaticRegeneration_Clicked()
 		return;
 	}
 
-	PCGGraphBeingEdited->ToggleSubsystemNotifications();
-
-	if (!PCGGraphBeingEdited->SubsystemNotificationsAreDisabled())
-	{
-		PCGGraphBeingEdited->NotifyGraphChanged(/*bIsStructural=*/true);
-	}
+	PCGGraphBeingEdited->ToggleUserPausedNotificationsForEditor();
 }
 
 bool FPCGEditor::IsAutomaticRegenerationPaused() const
 {
-	return PCGGraphBeingEdited && PCGGraphBeingEdited->SubsystemNotificationsAreDisabled();
+	return PCGGraphBeingEdited && PCGGraphBeingEdited->NotificationsForEditorArePausedByUser();
 }
 
 bool FPCGEditor::CanRunDeterminismTests() const
@@ -1449,6 +1444,11 @@ void FPCGEditor::OnClose()
 	ReplicateExtraNodes();
 
 	FAssetEditorToolkit::OnClose();
+
+	if (PCGGraphBeingEdited && PCGGraphBeingEdited->NotificationsForEditorArePausedByUser())
+	{
+		PCGGraphBeingEdited->ToggleUserPausedNotificationsForEditor();
+	}
 }
 
 void FPCGEditor::InitToolMenuContext(FToolMenuContext& MenuContext)
