Adding support for delayed cleanup

diff --git a/Engine/Plugins/Experimental/PCG/Source/PCG/Private/Graph/PCGGraphExecutor.cpp b/Engine/Plugins/Experimental/PCG/Source/PCG/Private/Graph/PCGGraphExecutor.cpp
--- a/Engine/Plugins/Experimental/PCG/Source/PCG/Private/Graph/PCGGraphExecutor.cpp
+++ b/Engine/Plugins/Experimental/PCG/Source/PCG/Private/Graph/PCGGraphExecutor.cpp
@@ -560,7 +560,11 @@ void FPCGGraphExecutor::ClearResults()
 {
 	ScheduleLock.Lock();
 
-	NextTaskId = 0;
+	// Only reset if we have no more scheduled tasks, to avoid breaking dependencies
+	if (ScheduledTasks.IsEmpty())
+	{
+		NextTaskId = 0;
+	}
 	OutputData.Reset();
 
 	ResultsRootSet.Clear();
diff --git a/Engine/Plugins/Experimental/PCG/Source/PCG/Private/PCGComponent.cpp b/Engine/Plugins/Experimental/PCG/Source/PCG/Private/PCGComponent.cpp
--- a/Engine/Plugins/Experimental/PCG/Source/PCG/Private/PCGComponent.cpp
+++ b/Engine/Plugins/Experimental/PCG/Source/PCG/Private/PCGComponent.cpp
@@ -54,6 +54,12 @@ bool UPCGComponent::IsPartitioned() const
 	return bIsPartitioned && CanPartition();
 }
 
+bool UPCGComponent::IsLocalComponent() const
+{
+	// TODO: Have a more robust solution for this.
+	return GetOwner() && GetOwner()->IsA<APCGPartitionActor>();
+}
+
 void UPCGComponent::SetGraph_Implementation(UPCGGraph* InGraph)
 {
 	SetGraphLocal(InGraph);
@@ -179,7 +185,7 @@ void UPCGComponent::SetPropertiesFromOriginal(const UPCGComponent* Original)
 
 void UPCGComponent::Generate()
 {
-	if (bIsGenerating)
+	if (IsGenerating())
 	{
 		return;
 	}
@@ -203,29 +209,45 @@ void UPCGComponent::GenerateLocal(bool bForce)
 
 void UPCGComponent::GenerateLocal(bool bForce, EPCGComponentGenerationTrigger RequestedGenerationTrigger)
 {
-	if (bIsGenerating)
+	if (IsGenerating() || !GetSubsystem())
 	{
 		return;
 	}
 
-	// Force component activation so it's easier to control by BP.
-	if (!bActivated)
+	if (IsCleaningUp())
 	{
-		Modify();
-		bActivated = true;
-	}
+		// If we are cleaning up, retry this function when it's done
+		TWeakObjectPtr<UPCGComponent> ThisComponentPtr(this);
 
-	FPCGTaskId TaskId = GenerateInternal(bForce, RequestedGenerationTrigger, {});
+		CurrentGenerationTask = GetSubsystem()->ScheduleGeneric([ThisComponentPtr, bForce, RequestedGenerationTrigger]()
+			{
+				if (UPCGComponent* ThisComponent = ThisComponentPtr.Get())
+				{
+					// This is necessary to have "IsGenerating" to return false.
+					// Otherwise it will early out
+					ThisComponent->CurrentGenerationTask = InvalidPCGTaskId;
+					ThisComponent->GenerateLocal(bForce, RequestedGenerationTrigger);
+				}
 
-	if (TaskId != InvalidPCGTaskId)
+				return true;
+			}, { CurrentCleanupTask });
+	}
+	else
 	{
-		bIsGenerating = true;
+		// Force component activation so it's easier to control by BP.
+		if (!bActivated)
+		{
+			Modify();
+			bActivated = true;
+		}
+
+		GenerateInternal(bForce, RequestedGenerationTrigger, {});
 	}
 }
 
 FPCGTaskId UPCGComponent::GenerateInternal(bool bForce, EPCGComponentGenerationTrigger RequestedGenerationTrigger, const TArray<FPCGTaskId>& TaskDependencies)
 {
-	if (!ShouldGenerate(bForce, RequestedGenerationTrigger))
+	if (IsGenerating() || !ShouldGenerate(bForce, RequestedGenerationTrigger))
 	{
 		return InvalidPCGTaskId;
 	}
@@ -240,10 +262,20 @@ FPCGTaskId UPCGComponent::GenerateInternal(bool bForce, EPCGComponentGenerationT
 
 	Modify();
 
-	// Immediate operation: cleanup beforehand
-	if (bGenerated)
+	// Keep track of all the dependencies
+	TArray<FPCGTaskId> AdditionalDependencies;
+	const TArray<FPCGTaskId>* AllDependencies = &TaskDependencies;
+
+	if (IsCleaningUp())
 	{
-		CleanupInternal(/*bRemoveComponents=*/false);
+		AdditionalDependencies = TaskDependencies;
+		AdditionalDependencies.Add(CurrentCleanupTask);
+		AllDependencies = &AdditionalDependencies;
+	}
+	else if (bGenerated)
+	{
+		// Immediate pass to mark all resources unused (and remove the ones that cannot be re-used)
+		CleanupInternalImmediate(/*bRemoveComponents=*/false);
 	}
 
 	const FBox NewBounds = GetGridBounds();
@@ -253,7 +285,8 @@ FPCGTaskId UPCGComponent::GenerateInternal(bool bForce, EPCGComponentGenerationT
 		return InvalidPCGTaskId;
 	}
 
-	return GetSubsystem()->ScheduleComponent(this, /*bSave=*/bForce, TaskDependencies);
+	CurrentGenerationTask = GetSubsystem()->ScheduleComponent(this, /*bSave=*/bForce, *AllDependencies);
+	return CurrentGenerationTask;
 }
 
 bool UPCGComponent::GetActorsFromTags(const TSet<FName>& InTags, TSet<TWeakObjectPtr<AActor>>& OutActors, bool bCullAgainstLocalBounds)
@@ -304,7 +337,7 @@ void UPCGComponent::PostProcessGraph(const FBox& InNewBounds, bool bInGenerated)
 	{
 		bGenerated = true;
 
-		bIsGenerating = false;
+		CurrentGenerationTask = InvalidPCGTaskId;
 
 #if WITH_EDITOR
 		bDirtyGenerated = false;
@@ -313,6 +346,16 @@ void UPCGComponent::PostProcessGraph(const FBox& InNewBounds, bool bInGenerated)
 	}
 }
 
+void UPCGComponent::PostCleanupGraph()
+{
+	bGenerated = false;
+	CurrentCleanupTask = InvalidPCGTaskId;
+
+#if WITH_EDITOR
+	OnPCGGraphCleanedDelegate.Broadcast(this);
+#endif
+}
+
 void UPCGComponent::OnProcessGraphAborted()
 {
 	UE_LOG(LogPCG, Warning, TEXT("Process Graph was called but aborted, check for errors in log if you expected a result."));
@@ -320,7 +363,7 @@ void UPCGComponent::OnProcessGraphAborted()
 	CleanupUnusedManagedResources();
 
 	bGenerated = false;
-	bIsGenerating = false;
+	CurrentGenerationTask = InvalidPCGTaskId;
 
 #if WITH_EDITOR
 	bDirtyGenerated = false;
@@ -329,7 +372,7 @@ void UPCGComponent::OnProcessGraphAborted()
 
 void UPCGComponent::Cleanup()
 {
-	if (!bGenerated || !GetSubsystem())
+	if ((!bGenerated && !IsGenerating()) || !GetSubsystem() || IsCleaningUp())
 	{
 		return;
 	}
@@ -348,23 +391,35 @@ void UPCGComponent::Cleanup_Implementation(bool bRemoveComponents, bool bSave)
 
 void UPCGComponent::CleanupLocal(bool bRemoveComponents, bool bSave)
 {
-	if (!bGenerated || !GetSubsystem())
+	if ((!bGenerated && !IsGenerating()) || !GetSubsystem() || IsCleaningUp())
 	{
 		return;
 	}
 
-	Modify();
+	if (IsGenerating())
+	{
+		// If we are generating, retry this function when it's done
+		TWeakObjectPtr<UPCGComponent> ThisComponentPtr(this);
 
-	GetSubsystem()->ScheduleCleanup(this, bRemoveComponents, bSave, {});
+		CurrentCleanupTask = GetSubsystem()->ScheduleGeneric([ThisComponentPtr, bRemoveComponents, bSave]()
+			{
+				if (UPCGComponent* ThisComponent = ThisComponentPtr.Get())
+				{
+					// This is necessary to have "IsCleaningUp" to return false.
+					// Otherwise it will early out
+					ThisComponent->CurrentCleanupTask = InvalidPCGTaskId;
+					ThisComponent->CleanupLocal(bRemoveComponents, bSave);
+				}
 
-	if (IsPartitioned())
-	{
-		bGenerated = false;
+				return true;
+			}, { CurrentGenerationTask });
 	}
+	else
+	{
+		Modify();
 
-#if WITH_EDITOR
-	OnPCGGraphCleanedDelegate.Broadcast(this);
-#endif
+		CurrentCleanupTask = GetSubsystem()->ScheduleCleanup(this, bRemoveComponents, bSave, {});
+	}
 }
 
 AActor* UPCGComponent::ClearPCGLink(UClass* TemplateActor)
@@ -375,7 +430,7 @@ AActor* UPCGComponent::ClearPCGLink(UClass* TemplateActor)
 	}
 
 	// TODO: Perhaps remove this part if we want to do it in the PCG Graph.
-	if (bIsGenerating)
+	if (IsGenerating() || IsCleaningUp())
 	{
 		return nullptr;
 	}
@@ -416,6 +471,12 @@ AActor* UPCGComponent::ClearPCGLink(UClass* TemplateActor)
 
 bool UPCGComponent::MoveResourcesToNewActor(AActor* InNewActor, bool bCreateChild)
 {
+	// Don't move resources if we are generating or cleaning up
+	if (IsGenerating() || IsCleaningUp())
+	{
+		return false;
+	}
+
 	check(InNewActor);
 	AActor* NewActor = InNewActor;
 
@@ -454,32 +515,126 @@ bool UPCGComponent::MoveResourcesToNewActor(AActor* InNewActor, bool bCreateChil
 	return bHasMovedResources;
 }
 
-void UPCGComponent::CleanupInternal(bool bRemoveComponents)
+void UPCGComponent::CleanupInternalImmediate(bool bRemoveComponents)
 {
+	if (!bGenerated || IsPartitioned() || IsCleaningUp() || IsGenerating())
+	{
+		return;
+	}
+
 	TSet<TSoftObjectPtr<AActor>> ActorsToDelete;
-	CleanupInternal(bRemoveComponents, ActorsToDelete);
+
+	{
+		FScopeLock ResourcesLock(&GeneratedResourcesLock);
+		for (int32 ResourceIndex = GeneratedResources.Num() - 1; ResourceIndex >= 0; --ResourceIndex)
+		{
+			if (GeneratedResources[ResourceIndex]->Release(bRemoveComponents, ActorsToDelete))
+			{
+				GeneratedResources.RemoveAtSwap(ResourceIndex);
+			}
+		}
+	}
+
 	UPCGActorHelpers::DeleteActors(GetWorld(), ActorsToDelete.Array());
 }
 
-void UPCGComponent::CleanupInternal(bool bHardCleanup, TSet<TSoftObjectPtr<AActor>>& OutActorsToDelete)
+FPCGTaskId UPCGComponent::CleanupInternal(bool bRemoveComponents, const TArray<FPCGTaskId>& Dependencies)
 {
-	if (!bGenerated || IsPartitioned())
+	if ((!bGenerated && !IsGenerating()) || IsPartitioned() || IsCleaningUp())
 	{
-		return;
+		return InvalidPCGTaskId;
 	}
 
-	Modify();
-	bGenerated = false;
+	// Need a shared ptr to keep track of all actors to delete.
+	// Will be protected by GeneratedResourcesLock
+	TSharedPtr<TSet<TSoftObjectPtr<AActor>>> ActorsToDeletePtr = MakeShared<TSet<TSoftObjectPtr<AActor>>>();
+
+	// Keep track of all the dependencies
+	TArray<FPCGTaskId> AdditionalDependencies;
+	const TArray<FPCGTaskId>* AllDependencies = &Dependencies;
+
+	if (IsGenerating())
+	{
+		AdditionalDependencies = Dependencies;
+		AdditionalDependencies.Add(CurrentGenerationTask);
+		AllDependencies = &AdditionalDependencies;
+	}
+
+	TArray<FPCGTaskId> AllTasks;
+
+	// Use a WeakPtr to guard against the case when the Component doesn't exist anymore.
+	TWeakObjectPtr<UPCGComponent> ThisComponentWeakPtr(this);
 
-	FScopeLock ResourcesLock(&GeneratedResourcesLock);
-	for (int32 ResourceIndex = GeneratedResources.Num() - 1; ResourceIndex >= 0; --ResourceIndex)
 	{
-		check(GeneratedResources[ResourceIndex]);
-		if (GeneratedResources[ResourceIndex]->Release(bHardCleanup, OutActorsToDelete))
+		FScopeLock ResourcesLock(&GeneratedResourcesLock);
+		for (int32 ResourceIndex = GeneratedResources.Num() - 1; ResourceIndex >= 0; --ResourceIndex)
 		{
-			GeneratedResources.RemoveAtSwap(ResourceIndex);
+			check(GeneratedResources[ResourceIndex]);
+			TWeakObjectPtr<UPCGManagedResource> GeneratedResourcePtr(GeneratedResources[ResourceIndex]);
+
+			// For now this is done on a single thread, but we could extend it to multi-thread if the release support MT.
+			FPCGTaskId ReleaseTask = GetSubsystem()->ScheduleGeneric([GeneratedResourcePtr, ResourceIndex, bRemoveComponents, ActorsToDeletePtr, ThisComponentWeakPtr]()
+			{
+				if (UPCGComponent* ThisComponent = ThisComponentWeakPtr.Get())
+				{
+					if (UPCGManagedResource* GeneratedResource = GeneratedResourcePtr.Get())
+					{
+						TSet<TSoftObjectPtr<AActor>> ActorsToDelete;
+						GeneratedResource->Release(bRemoveComponents, ActorsToDelete);
+
+						if (!ActorsToDelete.IsEmpty())
+						{
+							FScopeLock Lock(&ThisComponent->GeneratedResourcesLock);
+							ActorsToDeletePtr->Append(ActorsToDelete);
+						}
+					}
+				}
+
+				return true;
+			}, *AllDependencies);
+
+			AllTasks.Add(ReleaseTask);
 		}
 	}
+
+	// Use a WeakPtr to guard against the case when the World doesn't exist anymore.
+	TWeakObjectPtr<UWorld> WorldPtr(GetWorld());
+
+	FPCGTaskId PostReleaseTask = GetSubsystem()->ScheduleGeneric([WorldPtr, ThisComponentWeakPtr, bRemoveComponents, ActorsToDeletePtr]()
+		{
+			if (UPCGComponent* ThisComponent = ThisComponentWeakPtr.Get())
+			{
+				ThisComponent->Modify();
+				FScopeLock ResourcesLock(&ThisComponent->GeneratedResourcesLock);
+
+				// If we remove components, all generated resources are removed
+				if (bRemoveComponents)
+				{
+					ThisComponent->GeneratedResources.Empty();
+				}
+				else
+				{
+					// Otherwise, remove only those who are not marked unused (meaning that they can't be re-used)
+					for (int32 ResourceIndex = ThisComponent->GeneratedResources.Num() - 1; ResourceIndex >= 0; --ResourceIndex)
+					{
+						if (!ThisComponent->GeneratedResources[ResourceIndex]->IsMarkedUnused())
+						{
+							ThisComponent->GeneratedResources.RemoveAtSwap(ResourceIndex);
+						}
+					}
+				}
+			}
+
+			if (UWorld* World = WorldPtr.Get())
+			{
+				UPCGActorHelpers::DeleteActors(World, ActorsToDeletePtr->Array());
+			}
+
+			return true;
+		}, AllTasks);
+
+	CurrentCleanupTask = PostReleaseTask;
+	return CurrentCleanupTask;
 }
 
 void UPCGComponent::CleanupUnusedManagedResources()
@@ -1996,7 +2151,7 @@ bool UPCGComponent::GraphUsesLandscapePin() const
 
 void UPCGComponent::UpdateIsLocalComponent()
 {
-	if (GetOwner() && GetOwner()->IsA<APCGPartitionActor>())
+	if (IsLocalComponent())
 	{
 		bIsComponentLocal = true;
 	}
diff --git a/Engine/Plugins/Experimental/PCG/Source/PCG/Private/PCGManagedResource.cpp b/Engine/Plugins/Experimental/PCG/Source/PCG/Private/PCGManagedResource.cpp
--- a/Engine/Plugins/Experimental/PCG/Source/PCG/Private/PCGManagedResource.cpp
+++ b/Engine/Plugins/Experimental/PCG/Source/PCG/Private/PCGManagedResource.cpp
@@ -15,6 +15,7 @@ bool UPCGManagedResource::Release(bool bHardRelease, TSet<TSoftObjectPtr<AActor>
 		return false;
 	}
 
+	bIsMarkedUnused = false;
 	return true;
 }
 
@@ -39,6 +40,8 @@ void UPCGManagedActors::PostEditImport()
 
 bool UPCGManagedActors::Release(bool bHardRelease, TSet<TSoftObjectPtr<AActor>>& OutActorsToDelete)
 {
+	TRACE_CPUPROFILER_EVENT_SCOPE(UPCGManagedActors::Release);
+
 	if (!Super::Release(bHardRelease, OutActorsToDelete))
 	{
 		return false;
@@ -57,7 +60,9 @@ bool UPCGManagedActors::Release(bool bHardRelease, TSet<TSoftObjectPtr<AActor>>&
 
 			for (UPCGComponent* Component : ComponentsToCleanup)
 			{
-				Component->CleanupInternal(/*bRemoveComponents=*/false, OutActorsToDelete);
+				// It is more complicated to handled a non-immediate cleanup when doing it recursively in the managed actors.
+				// Do it all immediate then.
+				Component->CleanupInternalImmediate(/*bRemoveComponents=*/bHardRelease);
 			}
 
 			ComponentsToCleanup.Reset();
@@ -133,6 +138,8 @@ void UPCGManagedComponent::PostEditImport()
 
 bool UPCGManagedComponent::Release(bool bHardRelease, TSet<TSoftObjectPtr<AActor>>& /*OutActorsToDelete*/)
 {
+	TRACE_CPUPROFILER_EVENT_SCOPE(UPCGManagedComponent::Release);
+
 	const bool bSupportsComponentReset = SupportsComponentReset();
 	const bool bDeleteComponent = bHardRelease || !bSupportsComponentReset;
 
diff --git a/Engine/Plugins/Experimental/PCG/Source/PCG/Private/PCGSubsystem.cpp b/Engine/Plugins/Experimental/PCG/Source/PCG/Private/PCGSubsystem.cpp
--- a/Engine/Plugins/Experimental/PCG/Source/PCG/Private/PCGSubsystem.cpp
+++ b/Engine/Plugins/Experimental/PCG/Source/PCG/Private/PCGSubsystem.cpp
@@ -173,10 +173,14 @@ FPCGTaskId UPCGSubsystem::ScheduleComponent(UPCGComponent* PCGComponent, bool bS
 	}
 	else
 	{
-		AllTasks.Add(GraphExecutor->Schedule(PCGComponent, Dependencies));
+		FPCGTaskId TaskId = GraphExecutor->Schedule(PCGComponent, Dependencies);
+		if (TaskId != InvalidPCGTaskId)
+		{
+			AllTasks.Add(TaskId);
+		}
 	}
 
-	if (!AllTasks.IsEmpty() && AllTasks.Last() != InvalidPCGTaskId)
+	if (!AllTasks.IsEmpty())
 	{
 		TWeakObjectPtr<UPCGComponent> ComponentPtr(PCGComponent);
 
@@ -212,9 +216,7 @@ FPCGTaskId UPCGSubsystem::ScheduleCleanup(UPCGComponent* PCGComponent, bool bRem
 #if WITH_EDITOR
 	if (PCGComponent->IsPartitioned() && !PCGHelpers::IsRuntimeOrPIE())
 	{
-		// Immediate operation
-		CleanupGraph(PCGComponent, PCGComponent->LastGeneratedBounds, bRemoveComponents, bSave);
-		return InvalidPCGTaskId;
+		return CleanupGraph(PCGComponent, PCGComponent->LastGeneratedBounds, bRemoveComponents, bSave);
 	}
 #endif // WITH_EDITOR
 
@@ -236,18 +238,29 @@ FPCGTaskId UPCGSubsystem::ScheduleCleanup(UPCGComponent* PCGComponent, bool bRem
 	}
 	else
 	{
-		// In non partitioned mode, do it immediately
-		// TODO: Should it be a parameter?
-		PCGComponent->CleanupInternal(bRemoveComponents);
+		FPCGTaskId TaskId = PCGComponent->CleanupInternal(bRemoveComponents, Dependencies);
+		if (TaskId != InvalidPCGTaskId)
+		{
+			AllTasks.Add(TaskId);
+		}
 	}
 
 	if (AllTasks.IsEmpty())
 	{
 		return InvalidPCGTaskId;
 	}
+	
+	TWeakObjectPtr<UPCGComponent> ComponentPtr(PCGComponent);
 
-	// Add a dummy task to wait after all of the tasks
-	return GraphExecutor->ScheduleGeneric([]() { return true; }, AllTasks);
+	return GraphExecutor->ScheduleGeneric([ComponentPtr]() 
+		{
+			if (UPCGComponent* Component = ComponentPtr.Get())
+			{
+				Component->PostCleanupGraph();
+			}
+
+			return true;
+		}, AllTasks);
 }
 
 FPCGTaskId UPCGSubsystem::ScheduleGraph(UPCGGraph* Graph, UPCGComponent* SourceComponent, FPCGElementPtr InputElement, const TArray<FPCGTaskId>& Dependencies)
@@ -293,7 +306,7 @@ TArray<FPCGTaskId> UPCGSubsystem::ScheduleMultipleComponent(UPCGComponent* Origi
 					// Ensure that the PCG actor match our original
 					LocalComponent->SetPropertiesFromOriginal(OriginalComponent);
 
-					FPCGTaskId LocalTask = ScheduleComponent(LocalComponent, /*bSave=*/ false, Dependencies);
+					FPCGTaskId LocalTask = LocalComponent->GenerateInternal(/*bForce=*/ false, EPCGComponentGenerationTrigger::GenerateOnDemand, Dependencies);
 
 					if (LocalTask != InvalidPCGTaskId)
 					{
@@ -319,18 +332,12 @@ TArray<FPCGTaskId> UPCGSubsystem::ScheduleMultipleCleanup(UPCGComponent* Origina
 				// Add check to avoid infinite loop
 				if (ensure(!LocalComponent->IsPartitioned()))
 				{
-					TWeakObjectPtr<UPCGComponent> LocalComponentPtr(LocalComponent);
-					auto CleanupTask = [this, LocalComponentPtr, bRemoveComponents]()
-					{
-						if (UPCGComponent* Component = LocalComponentPtr.Get())
-						{
-							ScheduleCleanup(Component, bRemoveComponents, /*bSave=*/false, {});
-						}
-
-						return true;
-					};
+					FPCGTaskId LocalTask = LocalComponent->CleanupInternal(bRemoveComponents, Dependencies);
 
-					TaskIds.Add(GraphExecutor->ScheduleGeneric(CleanupTask, Dependencies));
+					if (LocalTask != InvalidPCGTaskId)
+					{
+						TaskIds.Add(LocalTask);
+					}
 				}
 			}
 		}
@@ -878,7 +885,7 @@ FPCGTaskId UPCGSubsystem::ProcessGraph(UPCGComponent* Component, const FBox& InP
 	}
 }
 
-void UPCGSubsystem::CleanupGraph(UPCGComponent* Component, const FBox& InBounds, bool bRemoveComponents, bool bSave)
+FPCGTaskId UPCGSubsystem::CleanupGraph(UPCGComponent* Component, const FBox& InBounds, bool bRemoveComponents, bool bSave)
 {
 	TWeakObjectPtr<UPCGComponent> ComponentPtr(Component);
 
@@ -903,7 +910,30 @@ void UPCGSubsystem::CleanupGraph(UPCGComponent* Component, const FBox& InBounds,
 		return InvalidPCGTaskId;
 	};
 
-	PCGSubsystem::ForEachIntersectingCell(GraphExecutor, Component->GetWorld(), InBounds, /*bCreateActor=*/false, /*bLoadCell=*/false, bSave, ScheduleTask);
+	FPCGTaskId ProcessAllCellsTaskId = PCGSubsystem::ForEachIntersectingCell(GraphExecutor, Component->GetWorld(), InBounds, /*bCreateActor=*/false, /*bLoadCell=*/false, bSave, ScheduleTask);
+
+	// Finally, call PostCleanupGraph if something happened
+	if (ProcessAllCellsTaskId != InvalidPCGTaskId)
+	{
+		auto PostCleanupGraph = [ComponentPtr]() {
+			if (UPCGComponent* Component = ComponentPtr.Get())
+			{
+				Component->PostCleanupGraph();
+			}
+			return true;
+		};
+
+		return GraphExecutor->ScheduleGeneric(PostCleanupGraph, { ProcessAllCellsTaskId });
+	}
+	else
+	{
+		if (Component)
+		{
+			Component->PostCleanupGraph();
+		}
+
+		return InvalidPCGTaskId;
+	}
 }
 
 void UPCGSubsystem::DirtyGraph(UPCGComponent* Component, const FBox& InBounds, EPCGComponentDirtyFlag DirtyFlag)
@@ -1185,21 +1215,11 @@ FPCGTaskId UPCGSubsystem::DelayGenerateGraph(UPCGComponent* Component, bool bSav
 FPCGTaskId UPCGSubsystem::DelayProcessGraph(UPCGComponent* Component, EOperation InOperation, bool bSave)
 {
 	check(Component && Component->IsPartitioned());
-	TWeakObjectPtr<UPCGComponent> ComponentPtr(Component);
 
-	auto ExecuteProcessGraph = [this, ComponentPtr, InOperation, bSave]() {
-		if (UPCGComponent* Component = ComponentPtr.Get())
-		{
-			const FBox PreviousBounds = Component->LastGeneratedBounds;
-			const FBox NewBounds = Component->GetGridBounds();
-			/*FPCGTaskId GraphTaskId = */ ProcessGraph(Component, PreviousBounds, NewBounds, InOperation, bSave);
-		}
-
-		return true;
-	};
+	const FBox PreviousBounds = Component->LastGeneratedBounds;
+	const FBox NewBounds = Component->GetGridBounds();
 
-	// Delayed graph scheduling
-	return GraphExecutor->ScheduleGeneric(ExecuteProcessGraph, {});
+	return ProcessGraph(Component, PreviousBounds, NewBounds, InOperation, bSave);
 }
 
 #endif // WITH_EDITOR
diff --git a/Engine/Plugins/Experimental/PCG/Source/PCG/Public/PCGComponent.h b/Engine/Plugins/Experimental/PCG/Source/PCG/Public/PCGComponent.h
--- a/Engine/Plugins/Experimental/PCG/Source/PCG/Public/PCGComponent.h
+++ b/Engine/Plugins/Experimental/PCG/Source/PCG/Public/PCGComponent.h
@@ -162,7 +162,8 @@ public:
 #endif
 
 	/** Return if we are currently generating the graph for this component */
-	bool IsGenerating() const { return bIsGenerating; }
+	bool IsGenerating() const { return CurrentGenerationTask != InvalidPCGTaskId; }
+	bool IsCleaningUp() const { return CurrentCleanupTask != InvalidPCGTaskId; }
 
 #if WITH_EDITOR
 	void Refresh();
@@ -180,6 +181,7 @@ public:
 #endif
 
 	bool IsPartitioned() const;
+	bool IsLocalComponent() const;
 
 	/** Updates internal properties from other component, dirties as required but does not trigger Refresh */
 	void SetPropertiesFromOriginal(const UPCGComponent* Original);
@@ -203,10 +205,18 @@ private:
 
 	bool ShouldGenerate(bool bForce, EPCGComponentGenerationTrigger RequestedGenerationTrigger) const;
 	void GenerateLocal(bool bForce, EPCGComponentGenerationTrigger RequestedGenerationTrigger);
+
+	/* Internal call to create tasks to generate the component. If there is nothing to do, an invalid task id will be returned. */
 	FPCGTaskId GenerateInternal(bool bForce, EPCGComponentGenerationTrigger RequestedGenerationTrigger, const TArray<FPCGTaskId>& Dependencies);
-	void CleanupInternal(bool bRemoveComponents);
-	void CleanupInternal(bool bRemoveComponents, TSet<TSoftObjectPtr<AActor>>& OutActorsToDelete);
+
+	/* Internal call to create tacks to cleanup the component. If there is nothing to do, an invalid task id will be returned. */
+	FPCGTaskId CleanupInternal(bool bRemoveComponents, const TArray<FPCGTaskId>& Dependencies);
+
+	/* Same as CleanupInternal, but without any delayed tasks. All is done immediately. */
+	void CleanupInternalImmediate(bool bRemoveComponents);
+
 	void PostProcessGraph(const FBox& InNewBounds, bool bInGenerated);
+	void PostCleanupGraph();
 	void OnProcessGraphAborted();
 	void CleanupUnusedManagedResources();
 	bool MoveResourcesToNewActor(AActor* InNewActor, bool bCreateChild);
@@ -285,7 +295,8 @@ private:
 	UPROPERTY()
 	FBox LastGeneratedBounds = FBox(EForceInit::ForceInit);
 
-	bool bIsGenerating = false;
+	FPCGTaskId CurrentGenerationTask = InvalidPCGTaskId;
+	FPCGTaskId CurrentCleanupTask = InvalidPCGTaskId;
 
 #if WITH_EDITOR
 	bool bIsInspecting = false;
diff --git a/Engine/Plugins/Experimental/PCG/Source/PCG/Public/PCGManagedResource.h b/Engine/Plugins/Experimental/PCG/Source/PCG/Public/PCGManagedResource.h
--- a/Engine/Plugins/Experimental/PCG/Source/PCG/Public/PCGManagedResource.h
+++ b/Engine/Plugins/Experimental/PCG/Source/PCG/Public/PCGManagedResource.h
@@ -27,6 +27,7 @@ public:
 	virtual bool ReleaseIfUnused(TSet<TSoftObjectPtr<AActor>>& OutActorsToDelete);
 
 	virtual void MarkAsUsed() { bIsMarkedUnused = false; }
+	bool IsMarkedUnused() const { return bIsMarkedUnused; }
 
 	/** Move the given resource to a new actor. Return true if it has succeeded */
 	virtual bool MoveResourceToNewActor(AActor* NewActor) { return false; };
diff --git a/Engine/Plugins/Experimental/PCG/Source/PCG/Public/PCGSubsystem.h b/Engine/Plugins/Experimental/PCG/Source/PCG/Public/PCGSubsystem.h
--- a/Engine/Plugins/Experimental/PCG/Source/PCG/Public/PCGSubsystem.h
+++ b/Engine/Plugins/Experimental/PCG/Source/PCG/Public/PCGSubsystem.h
@@ -94,7 +94,7 @@ public:
 	FPCGTaskId DelayGenerateGraph(UPCGComponent* Component, bool bSave);
 
 	/** Schedules an operation to cleanup the graph in the given bounds */
-	void CleanupGraph(UPCGComponent* Component, const FBox& InBounds, bool bRemoveComponents, bool bSave);
+	FPCGTaskId CleanupGraph(UPCGComponent* Component, const FBox& InBounds, bool bRemoveComponents, bool bSave);
 
 	/** Immediately dirties the partition actors in the given bounds */
 	void DirtyGraph(UPCGComponent* Component, const FBox& InBounds, EPCGComponentDirtyFlag DirtyFlag);
