Mark MeshSelector SelectInstances and InstancePacker PackInstances as PURE_VIRTUAL

#rnx

diff --git a/Engine/Plugins/Experimental/PCG/Source/PCG/Private/InstancePackers/PCGInstancePackerBase.cpp b/Engine/Plugins/Experimental/PCG/Source/PCG/Private/InstancePackers/PCGInstancePackerBase.cpp
--- a/Engine/Plugins/Experimental/PCG/Source/PCG/Private/InstancePackers/PCGInstancePackerBase.cpp
+++ b/Engine/Plugins/Experimental/PCG/Source/PCG/Private/InstancePackers/PCGInstancePackerBase.cpp
@@ -7,11 +7,6 @@
 #include "Metadata/PCGMetadataAttributeTraits.h"
 #include "Metadata/PCGMetadataAttributeTpl.h"
 
-void UPCGInstancePackerBase::PackInstances_Implementation(FPCGContext& Context, const UPCGSpatialData* InSpatialData, const FPCGMeshInstanceList& InstanceList, FPCGPackedCustomData& OutPackedCustomData) const
-{
-	PCGE_LOG_C(Error, &Context, "Invalid use of abstract InstancePackerBase class. Please use an already existing class or implement the PackInstances method");
-}
-
 bool UPCGInstancePackerBase::AddTypeToPacking(int TypeId, FPCGPackedCustomData& OutPackedCustomData) const
 {
 	if (TypeId == PCG::Private::MetadataTypes<float>::Id)
diff --git a/Engine/Plugins/Experimental/PCG/Source/PCG/Private/MeshSelectors/PCGMeshSelectorBase.cpp b/Engine/Plugins/Experimental/PCG/Source/PCG/Private/MeshSelectors/PCGMeshSelectorBase.cpp
--- a/Engine/Plugins/Experimental/PCG/Source/PCG/Private/MeshSelectors/PCGMeshSelectorBase.cpp
+++ b/Engine/Plugins/Experimental/PCG/Source/PCG/Private/MeshSelectors/PCGMeshSelectorBase.cpp
@@ -4,16 +4,6 @@
 
 #include "Elements/PCGStaticMeshSpawner.h"
 
-void UPCGMeshSelectorBase::SelectInstances_Implementation(
-	FPCGContext& Context, 
-	const UPCGStaticMeshSpawnerSettings* Settings, 
-	const UPCGSpatialData* InSpatialData,
-	TArray<FPCGMeshInstanceList>& OutMeshInstances,
-	UPCGPointData* OutPointData) const
-{
-	PCGE_LOG_C(Error, &Context, "Invalid use of abstract MeshSelectorBase class. Please use an already existing class or implement the SelectInstances method");
-}
-
 bool UPCGMeshSelectorBase::FindOrAddInstanceList(
 	TArray<FPCGMeshInstanceList>& OutInstanceLists,
 	const TSoftObjectPtr<UStaticMesh>& Mesh,
diff --git a/Engine/Plugins/Experimental/PCG/Source/PCG/Public/InstancePackers/PCGInstancePackerBase.h b/Engine/Plugins/Experimental/PCG/Source/PCG/Public/InstancePackers/PCGInstancePackerBase.h
--- a/Engine/Plugins/Experimental/PCG/Source/PCG/Public/InstancePackers/PCGInstancePackerBase.h
+++ b/Engine/Plugins/Experimental/PCG/Source/PCG/Public/InstancePackers/PCGInstancePackerBase.h
@@ -33,7 +33,7 @@ public:
 	UFUNCTION(BlueprintNativeEvent, Category = InstancePacking)
 	void PackInstances(FPCGContext& Context, const UPCGSpatialData* InSpatialData, const FPCGMeshInstanceList& InstanceList, FPCGPackedCustomData& OutPackedCustomData) const;
 
-	virtual void PackInstances_Implementation(FPCGContext& Context, const UPCGSpatialData* InSpatialData, const FPCGMeshInstanceList& InstanceList, FPCGPackedCustomData& OutPackedCustomData) const;
+	virtual void PackInstances_Implementation(FPCGContext& Context, const UPCGSpatialData* InSpatialData, const FPCGMeshInstanceList& InstanceList, FPCGPackedCustomData& OutPackedCustomData) const PURE_VIRTUAL(UPCGInstancePackerBase::PackInstances_Implementation);
 
 	/** Interprets Metadata TypeId and increments OutPackedCustomData.NumCustomDataFloats appropriately. Returns false if the type could not be interpreted. */
 	UFUNCTION(BlueprintCallable, Category = InstancePacking)
diff --git a/Engine/Plugins/Experimental/PCG/Source/PCG/Public/MeshSelectors/PCGMeshSelectorBase.h b/Engine/Plugins/Experimental/PCG/Source/PCG/Public/MeshSelectors/PCGMeshSelectorBase.h
--- a/Engine/Plugins/Experimental/PCG/Source/PCG/Public/MeshSelectors/PCGMeshSelectorBase.h
+++ b/Engine/Plugins/Experimental/PCG/Source/PCG/Public/MeshSelectors/PCGMeshSelectorBase.h
@@ -63,7 +63,7 @@ public:
 		const UPCGStaticMeshSpawnerSettings* Settings,
 		const UPCGSpatialData* InSpatialData,
 		TArray<FPCGMeshInstanceList>& OutMeshInstances,
-		UPCGPointData* OutPointData) const;
+		UPCGPointData* OutPointData) const PURE_VIRTUAL(UPCGMeshSelectorBase::SelectInstances_Implementation);
 
 	/** Searches OutInstanceLists for an InstanceList matching the given parameters. If nothing is found, creates a new InstanceList and adds to OutInstanceLists. Returns true if added. */
 	UFUNCTION(BlueprintCallable, Category = MeshSelection)
