Cleanup on the async notification for PCG tasks

#rnx

diff --git a/Engine/Plugins/Experimental/PCG/Source/PCG/Private/Graph/PCGGraphExecutor.cpp b/Engine/Plugins/Experimental/PCG/Source/PCG/Private/Graph/PCGGraphExecutor.cpp
--- a/Engine/Plugins/Experimental/PCG/Source/PCG/Private/Graph/PCGGraphExecutor.cpp
+++ b/Engine/Plugins/Experimental/PCG/Source/PCG/Private/Graph/PCGGraphExecutor.cpp
@@ -90,11 +90,6 @@ FPCGTaskId FPCGGraphExecutor::Schedule(UPCGGraph* Graph, UPCGComponent* SourceCo
 		FPCGGraphCompiler::OffsetNodeIds(ScheduledTask.Tasks, NextTaskId);
 		NextTaskId += ScheduledTask.Tasks.Num();
 		ScheduledId = NextTaskId - 1; // This is true because the last task is from the output node or is the post-execute task
-#if WITH_EDITOR
-		bTasksAreRunning = true;
-		RemainingTaskNum += ScheduledTask.Tasks.Num();
-		GenerationProgressNotification.Update(NextTaskId);
-#endif
 
 		// Push task (not data) dependencies on the pre-execute task
 		// Note must be done after the offset ids, otherwise we'll break the dependencies
@@ -124,11 +119,6 @@ FPCGTaskId FPCGGraphExecutor::ScheduleGeneric(TFunction<bool()> InOperation, con
 
 	// Assign task id
 	Task.NodeId = NextTaskId++;
-#if WITH_EDITOR
-	bTasksAreRunning = true;
-	++RemainingTaskNum;
-	GenerationProgressNotification.Update(NextTaskId);
-#endif
 
 	FPCGGraphScheduleTask& ScheduledTask = ScheduledTasks.Emplace_GetRef();
 	ScheduledTask.Tasks.Add(Task);
@@ -211,6 +201,10 @@ void FPCGGraphExecutor::Execute()
 	const int32 MaxNumThreads = FMath::Max(0, FMath::Min(FPlatformMisc::NumberOfCoresIncludingHyperthreads() - 2, CVarMaxNumTasks.GetValueOnAnyThread() - 1));
 	const bool bAllowMultiDispatch = CVarGraphMultithreading.GetValueOnAnyThread();
 
+#if WITH_EDITOR
+	UpdateGenerationNotification(Tasks.Num() + ReadyTasks.Num() + ActiveTasks.Num());
+#endif
+
 	while(ReadyTasks.Num() > 0 || ActiveTasks.Num() > 0)
 	{
 		// First: if we have free resources, move ready tasks to the active tasks
@@ -251,9 +245,6 @@ void FPCGGraphExecutor::Execute()
 				if (!Task.Element)
 				{
 					ReadyTasks.RemoveAtSwap(ReadyTaskIndex);
-#if WITH_EDITOR
-					RemainingTaskNum--;
-#endif // WITH_EDITOR
 					continue;
 				}
 
@@ -283,8 +274,6 @@ void FPCGGraphExecutor::Execute()
 					{
 						SourceComponent->StoreInspectionData(Task.Node, CachedOutput);
 					}
-
-					RemainingTaskNum--;
 #endif
 
 					continue;
@@ -426,9 +415,6 @@ void FPCGGraphExecutor::Execute()
 				{
 					PostTaskExecute(ActiveTasks.Num() - 1);
 				}
-#if WITH_EDITOR
-				--RemainingTaskNum;
-#endif
 			}
 
 			for (int32 ExecutionIndex = ActiveTasks.Num() - 1; ExecutionIndex >= 0; --ExecutionIndex)
@@ -445,9 +431,6 @@ void FPCGGraphExecutor::Execute()
 				{
 					PostTaskExecute(ExecutionIndex);
 				}
-#if WITH_EDITOR
-				--RemainingTaskNum;
-#endif
 			}
 		}
 
@@ -470,6 +453,10 @@ void FPCGGraphExecutor::Execute()
 		{
 			check(TaskSuccessors.IsEmpty());
 			ClearResults();
+
+#if WITH_EDITOR
+			EndGenerationNotification();
+#endif
 		}
 
 #if WITH_EDITOR
@@ -569,13 +556,6 @@ void FPCGGraphExecutor::ClearResults()
 
 	ResultsRootSet.Clear();
 
-#if WITH_EDITOR
-	if (ScheduledTasks.IsEmpty())
-	{
-		EndGenerationNotification();
-	}
-#endif
-
 	ScheduleLock.Unlock();
 }
 
@@ -651,31 +631,19 @@ void FPCGGraphExecutor::NotifyGraphChanged(UPCGGraph* InGraph)
 	}
 }
 
-void FPCGGraphExecutor::UpdateGenerationNotification()
+void FPCGGraphExecutor::UpdateGenerationNotification(int32 RemainingTaskNum)
 {
-	if (!bTasksAreRunning)
+	if(RemainingTaskNum > 0)
 	{
-		return;
+		GenerationProgressNotification.Update(RemainingTaskNum);
 	}
-
-	GenerationProgressNotification.Update(RemainingTaskNum);
 }
 
 void FPCGGraphExecutor::EndGenerationNotification()
 {
-	if (bTasksAreRunning)
-	{
-		bTasksAreRunning = false;
-		if (RemainingTaskNum > 0)
-		{
-			UE_LOG(LogPCG, Warning, TEXT("Remaining Task Number greater than zero"));
-			RemainingTaskNum = 0;
-		}
-
-		GenerationProgressNotification.Update(0);
-		// To reset the UI notification
-		GenerationProgressNotification = FAsyncCompilationNotification(GetNotificationTextFormat());
-	}
+	GenerationProgressNotification.Update(0);
+	// To reset the UI notification
+	GenerationProgressNotification = FAsyncCompilationNotification(GetNotificationTextFormat());
 }
 
 FTextFormat FPCGGraphExecutor::GetNotificationTextFormat()
diff --git a/Engine/Plugins/Experimental/PCG/Source/PCG/Private/Graph/PCGGraphExecutor.h b/Engine/Plugins/Experimental/PCG/Source/PCG/Private/Graph/PCGGraphExecutor.h
--- a/Engine/Plugins/Experimental/PCG/Source/PCG/Private/Graph/PCGGraphExecutor.h
+++ b/Engine/Plugins/Experimental/PCG/Source/PCG/Private/Graph/PCGGraphExecutor.h
@@ -84,7 +84,6 @@ public:
 
 	/** Notify compiler that graph has changed so it'll be removed from the cache */
 	void NotifyGraphChanged(UPCGGraph* InGraph);
-	void UpdateGenerationNotification();
 #endif
 
 	/** "Tick" of the graph executor. This call is NOT THREADSAFE */
@@ -104,6 +103,8 @@ private:
 #if WITH_EDITOR
 	void SaveDirtyActors();
 	void ReleaseUnusedActors();
+
+	void UpdateGenerationNotification(int32 RemainingTaskNum);
 	void EndGenerationNotification();
 	static FTextFormat GetNotificationTextFormat();
 #endif
@@ -140,8 +141,6 @@ private:
 	TSet<FWorldPartitionReference> ActorsToRelease;
 
 	int32 CountUntilGC = 30;
-	int32 RemainingTaskNum = 0;
-	bool bTasksAreRunning = false;
 	FAsyncCompilationNotification GenerationProgressNotification;
 #endif
 };
diff --git a/Engine/Plugins/Experimental/PCG/Source/PCG/Private/PCGSubsystem.cpp b/Engine/Plugins/Experimental/PCG/Source/PCG/Private/PCGSubsystem.cpp
--- a/Engine/Plugins/Experimental/PCG/Source/PCG/Private/PCGSubsystem.cpp
+++ b/Engine/Plugins/Experimental/PCG/Source/PCG/Private/PCGSubsystem.cpp
@@ -84,10 +84,6 @@ void UPCGSubsystem::Tick(float DeltaSeconds)
 
 	// If we have any tasks to execute, schedule some
 	GraphExecutor->Execute();
-
-#if WITH_EDITOR
-	GraphExecutor->UpdateGenerationNotification();
-#endif
 }
 
 APCGWorldActor* UPCGSubsystem::GetPCGWorldActor()
