Changing the use metadata flag triggers a change that appropriately dirties the landscape data in the components

#rnx

diff --git a/Engine/Plugins/Experimental/PCG/Source/PCG/Private/PCGComponent.cpp b/Engine/Plugins/Experimental/PCG/Source/PCG/Private/PCGComponent.cpp
--- a/Engine/Plugins/Experimental/PCG/Source/PCG/Private/PCGComponent.cpp
+++ b/Engine/Plugins/Experimental/PCG/Source/PCG/Private/PCGComponent.cpp
@@ -88,7 +88,7 @@ void UPCGComponent::SetGraphLocal(UPCGGraph* InGraph)
 	}
 #endif
 
-	OnGraphChanged(Graph, true, true);
+	RefreshAfterGraphChanged(Graph, /*bIsStructural=*/true, /*bDirtyInputs=*/true);
 }
 
 void UPCGComponent::AddToManagedResources(UPCGManagedResource* InResource)
@@ -806,12 +806,15 @@ void UPCGComponent::BeginDestroy()
 	Super::BeginDestroy();
 }
 
-void UPCGComponent::OnGraphChanged(UPCGGraph* InGraph, bool bIsStructural)
+void UPCGComponent::OnGraphChanged(UPCGGraph* InGraph, EPCGChangeType ChangeType)
 {
-	OnGraphChanged(InGraph, bIsStructural, true);
+	const bool bIsStructural = ((ChangeType & (EPCGChangeType::Edge | EPCGChangeType::Structural)) != EPCGChangeType::None);
+	const bool bDirtyInputs = bIsStructural || ((ChangeType & EPCGChangeType::Input) != EPCGChangeType::None);
+
+	RefreshAfterGraphChanged(InGraph, bIsStructural, bDirtyInputs);
 }
 
-void UPCGComponent::OnGraphChanged(UPCGGraph* InGraph, bool bIsStructural, bool bShouldRefresh)
+void UPCGComponent::RefreshAfterGraphChanged(UPCGGraph* InGraph, bool bIsStructural, bool bDirtyInputs)
 {
 	if (InGraph != Graph)
 	{
@@ -832,8 +835,8 @@ void UPCGComponent::OnGraphChanged(UPCGGraph* InGraph, bool bIsStructural, bool
 			UpdateTrackedLandscape();
 		}
 
-		DirtyGenerated();
-		if (InGraph && bShouldRefresh)
+		DirtyGenerated(bDirtyInputs ? (EPCGComponentDirtyFlag::Actor | EPCGComponentDirtyFlag::Landscape) : EPCGComponentDirtyFlag::None);
+		if (InGraph)
 		{
 			Refresh();
 		}
@@ -947,7 +950,7 @@ void UPCGComponent::PostEditChangeProperty(FPropertyChangedEvent& PropertyChange
 			Graph->OnGraphChangedDelegate.AddUObject(this, &UPCGComponent::OnGraphChanged);
 		}
 
-		OnGraphChanged(Graph, /*bIsStructural=*/true, /*bShouldRefresh=*/true);
+		RefreshAfterGraphChanged(Graph, /*bIsStructural=*/true, /*bDirtyInputs=*/true);
 	}
 	else if (PropName == GET_MEMBER_NAME_CHECKED(UPCGComponent, InputType))
 	{
@@ -1538,7 +1541,6 @@ UPCGData* UPCGComponent::GetLandscapePCGData()
 
 UPCGData* UPCGComponent::GetLandscapeHeightPCGData()
 {
-
 	if (!CachedLandscapeHeightData || IsLandscapeCachedDataDirty(CachedLandscapeHeightData))
 	{
 		CachedLandscapeHeightData = CreateLandscapePCGData(/*bHeightOnly=*/true);
diff --git a/Engine/Plugins/Experimental/PCG/Source/PCG/Private/PCGGraph.cpp b/Engine/Plugins/Experimental/PCG/Source/PCG/Private/PCGGraph.cpp
--- a/Engine/Plugins/Experimental/PCG/Source/PCG/Private/PCGGraph.cpp
+++ b/Engine/Plugins/Experimental/PCG/Source/PCG/Private/PCGGraph.cpp
@@ -229,7 +229,7 @@ void UPCGGraph::OnNodeAdded(UPCGNode* InNode)
 {
 #if WITH_EDITOR
 	InNode->OnNodeChangedDelegate.AddUObject(this, &UPCGGraph::OnNodeChanged);
-	NotifyGraphChanged(/*bIsStructural=*/true);
+	NotifyGraphChanged(EPCGChangeType::Structural);
 #endif
 }
 
@@ -239,7 +239,7 @@ void UPCGGraph::OnNodeRemoved(UPCGNode* InNode)
 	if (InNode)
 	{
 		InNode->OnNodeChangedDelegate.RemoveAll(this);
-		NotifyGraphChanged(/*bIsStructural=*/true);
+		NotifyGraphChanged(EPCGChangeType::Structural);
 	}
 #endif
 }
@@ -277,7 +277,7 @@ UPCGNode* UPCGGraph::AddLabeledEdge(UPCGNode* From, const FName& InboundLabel, U
 	FromPin->AddEdgeTo(ToPin);
 	
 #if WITH_EDITOR
-	NotifyGraphChanged(/*bIsStructural=*/true);
+	NotifyGraphChanged(EPCGChangeType::Structural);
 #endif
 
 	return To;
@@ -358,7 +358,7 @@ bool UPCGGraph::RemoveEdge(UPCGNode* From, const FName& FromLabel, UPCGNode* To,
 #if WITH_EDITOR
 	if (bChanged)
 	{
-		NotifyGraphChanged(/*bIsStructural=*/true);
+		NotifyGraphChanged(EPCGChangeType::Structural);
 	}
 #endif
 
@@ -378,7 +378,7 @@ bool UPCGGraph::RemoveAllInboundEdges(UPCGNode* InNode)
 #if WITH_EDITOR
 	if (bChanged)
 	{
-		NotifyGraphChanged(/*bIsStructural=*/true);
+		NotifyGraphChanged(EPCGChangeType::Structural);
 	}
 #endif
 
@@ -397,7 +397,7 @@ bool UPCGGraph::RemoveAllOutboundEdges(UPCGNode* InNode)
 #if WITH_EDITOR
 	if (bChanged)
 	{
-		NotifyGraphChanged(/*bIsStructural=*/true);
+		NotifyGraphChanged(EPCGChangeType::Structural);
 	}
 #endif
 
@@ -417,7 +417,7 @@ bool UPCGGraph::RemoveInboundEdges(UPCGNode* InNode, const FName& InboundLabel)
 #if WITH_EDITOR
 	if (bChanged)
 	{
-		NotifyGraphChanged(/*bIsStructural=*/true);
+		NotifyGraphChanged(EPCGChangeType::Structural);
 	}
 #endif
 
@@ -437,7 +437,7 @@ bool UPCGGraph::RemoveOutboundEdges(UPCGNode* InNode, const FName& OutboundLabel
 #if WITH_EDITOR
 	if (bChanged)
 	{
-		NotifyGraphChanged(/*bIsStructural=*/true);
+		NotifyGraphChanged(EPCGChangeType::Structural);
 	}
 #endif
 
@@ -476,9 +476,9 @@ void UPCGGraph::EnableNotificationsForEditor()
 
 	if (GraphChangeNotificationsDisableCounter == 0 && bDelayedChangeNotification)
 	{
-		NotifyGraphChanged(bDelayedChangeNotificationStructural);
+		NotifyGraphChanged(DelayedChangeType);
 		bDelayedChangeNotification = false;
-		bDelayedChangeNotificationStructural = false;
+		DelayedChangeType = EPCGChangeType::None;
 	}
 }
 
@@ -533,12 +533,12 @@ void UPCGGraph::GetTrackedTagsToSettings(FPCGTagToSettingsMap& OutTagsToSettings
 	}
 }
 
-void UPCGGraph::NotifyGraphChanged(bool bIsStructural)
+void UPCGGraph::NotifyGraphChanged(EPCGChangeType ChangeType)
 {
 	if(GraphChangeNotificationsDisableCounter > 0)
 	{
 		bDelayedChangeNotification = true;
-		bDelayedChangeNotificationStructural |= bIsStructural;
+		DelayedChangeType |= ChangeType;
 		return;
 	}
 
@@ -551,7 +551,8 @@ void UPCGGraph::NotifyGraphChanged(bool bIsStructural)
 	bIsNotifying = true;
 
 	// Notify the subsystem/compiler cache before so it gets recompiled properly
-	if (bIsStructural && GEditor)
+	const bool bNotifySubsystem = ((ChangeType & (EPCGChangeType::Structural | EPCGChangeType::Edge)) != EPCGChangeType::None);
+	if (bNotifySubsystem && GEditor)
 	{
 		UWorld* World = GEditor->GetEditorWorldContext().World();
 		UPCGSubsystem* PCGSubsystem = World ? World->GetSubsystem<UPCGSubsystem>() : nullptr;
@@ -561,7 +562,7 @@ void UPCGGraph::NotifyGraphChanged(bool bIsStructural)
 		}
 	}
 
-	OnGraphChangedDelegate.Broadcast(this, /*bIsStructural=*/bIsStructural);
+	OnGraphChangedDelegate.Broadcast(this, ChangeType);
 
 	bIsNotifying = false;
 }
@@ -570,7 +571,17 @@ void UPCGGraph::OnNodeChanged(UPCGNode* InNode, EPCGChangeType ChangeType)
 {
 	if((ChangeType & ~EPCGChangeType::Cosmetic) != EPCGChangeType::None)
 	{
-		NotifyGraphChanged(/*bIsStructural=*/(ChangeType & (EPCGChangeType::Structural | EPCGChangeType::Edge)) != EPCGChangeType::None);
+		NotifyGraphChanged(ChangeType);
+	}
+}
+
+void UPCGGraph::PostEditChangeProperty(FPropertyChangedEvent& PropertyChangedEvent)
+{
+	Super::PostEditChangeProperty(PropertyChangedEvent);
+
+	if (PropertyChangedEvent.GetMemberPropertyName() == GET_MEMBER_NAME_CHECKED(UPCGGraph, bLandscapeUsesMetadata))
+	{
+		NotifyGraphChanged(EPCGChangeType::Input);
 	}
 }
 #endif // WITH_EDITOR
diff --git a/Engine/Plugins/Experimental/PCG/Source/PCG/Private/PCGSubgraph.cpp b/Engine/Plugins/Experimental/PCG/Source/PCG/Private/PCGSubgraph.cpp
--- a/Engine/Plugins/Experimental/PCG/Source/PCG/Private/PCGSubgraph.cpp
+++ b/Engine/Plugins/Experimental/PCG/Source/PCG/Private/PCGSubgraph.cpp
@@ -64,11 +64,11 @@ void UPCGBaseSubgraphSettings::GetTrackedActorTags(FPCGTagToSettingsMap& OutTagT
 	}
 }
 
-void UPCGBaseSubgraphSettings::OnSubgraphChanged(UPCGGraph* InGraph, bool bIsStructural)
+void UPCGBaseSubgraphSettings::OnSubgraphChanged(UPCGGraph* InGraph, EPCGChangeType ChangeType)
 {
 	if (InGraph == GetSubgraph())
 	{
-		OnSettingsChangedDelegate.Broadcast(this, ((bIsStructural ? EPCGChangeType::Structural : EPCGChangeType::None) | EPCGChangeType::Settings));
+		OnSettingsChangedDelegate.Broadcast(this, (ChangeType | EPCGChangeType::Settings));
 	}
 }
 #endif // WITH_EDITOR
diff --git a/Engine/Plugins/Experimental/PCG/Source/PCG/Public/PCGCommon.h b/Engine/Plugins/Experimental/PCG/Source/PCG/Public/PCGCommon.h
--- a/Engine/Plugins/Experimental/PCG/Source/PCG/Public/PCGCommon.h
+++ b/Engine/Plugins/Experimental/PCG/Source/PCG/Public/PCGCommon.h
@@ -13,9 +13,10 @@ enum class EPCGChangeType : uint8
 	None = 0,
 	Cosmetic = 1 << 0,
 	Settings = 1 << 1,
-	Edge = 1 << 2,
-	Node = 1 << 3,
-	Structural = 1 << 4
+	Input = 1 << 2,
+	Edge = 1 << 3,
+	Node = 1 << 4,
+	Structural = 1 << 5
 };
 ENUM_CLASS_FLAGS(EPCGChangeType);
 
diff --git a/Engine/Plugins/Experimental/PCG/Source/PCG/Public/PCGComponent.h b/Engine/Plugins/Experimental/PCG/Source/PCG/Public/PCGComponent.h
--- a/Engine/Plugins/Experimental/PCG/Source/PCG/Public/PCGComponent.h
+++ b/Engine/Plugins/Experimental/PCG/Source/PCG/Public/PCGComponent.h
@@ -223,8 +223,8 @@ private:
 
 	bool GetActorsFromTags(const TSet<FName>& InTags, TSet<TWeakObjectPtr<AActor>>& OutActors, bool bCullAgainstLocalBounds);
 
-	void OnGraphChanged(UPCGGraph* InGraph, bool bIsStructural, bool bShouldRefresh);
-	void OnGraphChanged(UPCGGraph* InGraph, bool bIsStructural);
+	void RefreshAfterGraphChanged(UPCGGraph* InGraph, bool bIsStructural, bool bDirtyInputs);
+	void OnGraphChanged(UPCGGraph* InGraph, EPCGChangeType ChangeType);
 
 #if WITH_EDITOR
 	virtual void PreEditChange(FProperty* PropertyAboutToChange) override;
diff --git a/Engine/Plugins/Experimental/PCG/Source/PCG/Public/PCGGraph.h b/Engine/Plugins/Experimental/PCG/Source/PCG/Public/PCGGraph.h
--- a/Engine/Plugins/Experimental/PCG/Source/PCG/Public/PCGGraph.h
+++ b/Engine/Plugins/Experimental/PCG/Source/PCG/Public/PCGGraph.h
@@ -10,7 +10,7 @@
 #include "PCGGraph.generated.h"
 
 #if WITH_EDITOR
-	DECLARE_MULTICAST_DELEGATE_TwoParams(FOnPCGGraphChanged, UPCGGraph* /*Graph*/, bool /*bIsStructural*/);
+	DECLARE_MULTICAST_DELEGATE_TwoParams(FOnPCGGraphChanged, UPCGGraph* /*Graph*/, EPCGChangeType /*ChangeType*/);
 #endif // WITH_EDITOR
 
 UCLASS(BlueprintType, ClassGroup = (Procedural), hidecategories=(Object))
@@ -27,6 +27,10 @@ public:
 	/** ~Begin UObject interface */
 	virtual void PostLoad() override;
 	virtual void BeginDestroy() override;
+
+#if WITH_EDITOR
+	virtual void PostEditChangeProperty(FPropertyChangedEvent& PropertyChangedEvent) override;
+#endif
 	/** ~End UObject interface */
 
 #if WITH_EDITORONLY_DATA
@@ -122,12 +126,12 @@ protected:
 
 #if WITH_EDITOR
 private:
-	void NotifyGraphChanged(bool bIsStructural);
+	void NotifyGraphChanged(EPCGChangeType ChangeType);
 	void OnNodeChanged(UPCGNode* InNode, EPCGChangeType ChangeType);
 
 	int32 GraphChangeNotificationsDisableCounter = 0;
 	bool bDelayedChangeNotification = false;
-	bool bDelayedChangeNotificationStructural = false;
+	EPCGChangeType DelayedChangeType = EPCGChangeType::None;
 	bool bIsNotifying = false;
 	bool bUserPausedNotificationsInGraphEditor = false;
 #endif // WITH_EDITOR
diff --git a/Engine/Plugins/Experimental/PCG/Source/PCG/Public/PCGSubgraph.h b/Engine/Plugins/Experimental/PCG/Source/PCG/Public/PCGSubgraph.h
--- a/Engine/Plugins/Experimental/PCG/Source/PCG/Public/PCGSubgraph.h
+++ b/Engine/Plugins/Experimental/PCG/Source/PCG/Public/PCGSubgraph.h
@@ -34,7 +34,7 @@ protected:
 	//~End UPCGSettings interface
 
 #if WITH_EDITOR
-	void OnSubgraphChanged(UPCGGraph* InGraph, bool bIsStructural);
+	void OnSubgraphChanged(UPCGGraph* InGraph, EPCGChangeType ChangeType);
 
 	virtual bool IsStructuralProperty(const FName& InPropertyName) const { return false; }
 #endif
diff --git a/Engine/Plugins/Experimental/PCG/Source/PCGEditor/Private/PCGEditor.cpp b/Engine/Plugins/Experimental/PCG/Source/PCGEditor/Private/PCGEditor.cpp
--- a/Engine/Plugins/Experimental/PCG/Source/PCGEditor/Private/PCGEditor.cpp
+++ b/Engine/Plugins/Experimental/PCG/Source/PCGEditor/Private/PCGEditor.cpp
@@ -243,7 +243,7 @@ void FPCGEditor::PostUndo(bool bSuccess)
 	{
 		if (PCGGraphBeingEdited)
 		{
-			PCGGraphBeingEdited->NotifyGraphChanged(true);
+			PCGGraphBeingEdited->NotifyGraphChanged(EPCGChangeType::Structural);
 		}
 
 		if (GraphEditorWidget.IsValid())
@@ -880,8 +880,8 @@ void FPCGEditor::OnCollapseNodesInSubgraph()
 
 	// And notify everyone
 	GraphEditorWidget->NotifyGraphChanged();
-	PCGGraph->NotifyGraphChanged(true);
-	NewPCGGraph->NotifyGraphChanged(true);
+	PCGGraph->NotifyGraphChanged(EPCGChangeType::Structural);
+	NewPCGGraph->NotifyGraphChanged(EPCGChangeType::Structural);
 }
 
 void FPCGEditor::OnStartInspectNode()
@@ -1055,7 +1055,7 @@ void FPCGEditor::DeleteSelectedNodes()
 		{
 			GraphEditorWidget->ClearSelectionSet();
 			GraphEditorWidget->NotifyGraphChanged();
-			PCGGraphBeingEdited->NotifyGraphChanged(true);
+			PCGGraphBeingEdited->NotifyGraphChanged(EPCGChangeType::Structural);
 		}
 	}
 }
