* Shifting native determinism tests to a subscription based mapping callable via modules

#rnx

diff --git a/Engine/Plugins/Experimental/PCG/Source/PCG/Private/PCGModule.cpp b/Engine/Plugins/Experimental/PCG/Source/PCG/Private/PCGModule.cpp
--- a/Engine/Plugins/Experimental/PCG/Source/PCG/Private/PCGModule.cpp
+++ b/Engine/Plugins/Experimental/PCG/Source/PCG/Private/PCGModule.cpp
@@ -1,29 +1,72 @@
 // Copyright Epic Games, Inc. All Rights Reserved.
 
 #include "PCGModule.h"
+
+#if WITH_EDITOR
+#include "Elements/PCGDifferenceElement.h"
+
+#include "Tests/Determinism/PCGDeterminismNativeTests.h"
+#include "Tests/Determinism/PCGDifferenceDeterminismTest.h"
+
+#endif
+
 #include "Modules/ModuleInterface.h"
 
-class FPCGModule : public IModuleInterface
+class FPCGModule final : public IModuleInterface
 {
 public:
 	//~ IModuleInterface implementation
-	virtual void StartupModule() override
-	{
-
-	}
-
-	virtual void ShutdownModule() override
-	{
+	virtual void StartupModule() override;
 
-	}
+	virtual void ShutdownModule() override;
 
 	virtual bool SupportsDynamicReloading() override
 	{
 		return true;
 	}
+
 	//~ End IModuleInterface implementation
+
+private:
+#if WITH_EDITOR
+	void RegisterNativeElementDeterminismTests();
+	void DeregisterNativeElementDeterminismTests();
+#endif
 };
 
+void FPCGModule::StartupModule()
+{
+#if WITH_EDITOR
+	PCGDeterminismTests::FNativeTestRegistry::Create();
+
+	RegisterNativeElementDeterminismTests();
+#endif
+}
+
+void FPCGModule::ShutdownModule()
+{
+#if WITH_EDITOR
+	DeregisterNativeElementDeterminismTests();
+
+	PCGDeterminismTests::FNativeTestRegistry::Destroy();
+#endif
+}
+
+#if WITH_EDITOR
+
+void FPCGModule::RegisterNativeElementDeterminismTests()
+{
+	PCGDeterminismTests::FNativeTestRegistry::RegisterTestFunction(UPCGDifferenceSettings::StaticClass(), PCGDeterminismTests::DifferenceElement::RunTestSuite);
+	// TODO: Add other native test functions
+}
+
+void FPCGModule::DeregisterNativeElementDeterminismTests()
+{
+	PCGDeterminismTests::FNativeTestRegistry::DeregisterTestFunction(UPCGDifferenceSettings::StaticClass());
+}
+
+#endif
+
 IMPLEMENT_MODULE(FPCGModule, PCG);
 
-DEFINE_LOG_CATEGORY(LogPCG);
+DEFINE_LOG_CATEGORY(LogPCG);
\ No newline at end of file
diff --git a/Engine/Plugins/Experimental/PCG/Source/PCG/Private/Tests/Determinism/PCGDeterminismNativeTests.cpp b/Engine/Plugins/Experimental/PCG/Source/PCG/Private/Tests/Determinism/PCGDeterminismNativeTests.cpp
--- a/Engine/Plugins/Experimental/PCG/Source/PCG/Private/Tests/Determinism/PCGDeterminismNativeTests.cpp
+++ b/Engine/Plugins/Experimental/PCG/Source/PCG/Private/Tests/Determinism/PCGDeterminismNativeTests.cpp
@@ -4,32 +4,47 @@
 
 #include "PCGSettings.h"
 
-#include "Elements/PCGDifferenceElement.h"
-#include "Tests/Determinism/PCGDifferenceDeterminismTest.h"
-
 namespace PCGDeterminismTests
 {
-	TFunction<bool()> GetNativeTestIfExists(const UPCGSettings* PCGSettings)
+	FNativeTestRegistry* FNativeTestRegistry::RegistryPtr = nullptr;
+
+	void FNativeTestRegistry::Create()
 	{
-		// This way we only build it if we ever need it
-		static TMap<UClass*, TFunction<bool()>> NativeTestMapping;
+		check(!RegistryPtr);
+		RegistryPtr = new FNativeTestRegistry;
+	}
 
-		if (NativeTestMapping.IsEmpty())
+	void FNativeTestRegistry::Destroy()
+	{
+		check(RegistryPtr);
+		delete RegistryPtr;
+		RegistryPtr = nullptr;
+	}
+
+	void FNativeTestRegistry::RegisterTestFunction(UClass* SettingsStaticClass, TFunction<bool()> TestFunction)
+	{
+		check(RegistryPtr && SettingsStaticClass);
+		RegistryPtr->NativeTestMapping.Emplace(SettingsStaticClass, TestFunction);
+	}
+
+	void FNativeTestRegistry::DeregisterTestFunction(const UClass* SettingsStaticClass)
+	{
+		check(RegistryPtr && SettingsStaticClass);
+		if (RegistryPtr->NativeTestMapping.Contains(SettingsStaticClass))
 		{
-			NativeTestMapping =
-			{
-				{UPCGDifferenceSettings::StaticClass(), PCGDeterminismTests::DifferenceElement::RunTestSuite}
-				// TODO: Append other native tests here
-			};
+			RegistryPtr->NativeTestMapping.Remove(SettingsStaticClass);
 		}
+	}
 
-		if (TFunction<bool()>* TestFunction = NativeTestMapping.Find(PCGSettings->GetClass()))
+	TFunction<bool()> FNativeTestRegistry::GetNativeTestFunction(const UPCGSettings* PCGSettings)
+	{
+		check(RegistryPtr);
+		if (TFunction<bool()>* TestFunction = RegistryPtr->NativeTestMapping.Find(PCGSettings->GetClass()))
 		{
 			return *TestFunction;
 		}
 
-		return nullptr;
+		return {};
 	}
 }
-
 #endif
\ No newline at end of file
diff --git a/Engine/Plugins/Experimental/PCG/Source/PCG/Public/Tests/Determinism/PCGDeterminismNativeTests.h b/Engine/Plugins/Experimental/PCG/Source/PCG/Public/Tests/Determinism/PCGDeterminismNativeTests.h
--- a/Engine/Plugins/Experimental/PCG/Source/PCG/Public/Tests/Determinism/PCGDeterminismNativeTests.h
+++ b/Engine/Plugins/Experimental/PCG/Source/PCG/Public/Tests/Determinism/PCGDeterminismNativeTests.h
@@ -1,14 +1,36 @@
 // Copyright Epic Games, Inc. All Rights Reserved.
 #pragma once
 #if WITH_EDITOR
-
 #include "CoreMinimal.h"
 
 class UPCGSettings;
 
 namespace PCGDeterminismTests
 {
-	PCG_API TFunction<bool()> GetNativeTestIfExists(const UPCGSettings* PCGSettings);
-}
+	class FNativeTestRegistry
+	{
+	public:
+		/** Allocates the static class */
+		static void Create();
+		/** Deallocates the static class */
+		static void Destroy();
+
+		/** Registers a test function given the StaticClass of an element's settings */
+		static PCG_API void RegisterTestFunction(UClass* SettingsStaticClass, TFunction<bool()> TestFunction);
+		/** Deregisters a test function given the StaticClass of an element's settings */
+		static PCG_API void DeregisterTestFunction(const UClass* SettingsStaticClass);
+		/** Gets a native test function, given an element's UPCGSettings */
+		static PCG_API TFunction<bool()> GetNativeTestFunction(const UPCGSettings* PCGSettings);
 
-#endif
\ No newline at end of file
+	private:
+		FNativeTestRegistry() = default;
+		FNativeTestRegistry(const FNativeTestRegistry&) = delete;
+		FNativeTestRegistry& operator=(const FNativeTestRegistry&) = delete;
+		FNativeTestRegistry(FNativeTestRegistry&&) = delete;
+		FNativeTestRegistry& operator=(FNativeTestRegistry&&) = delete;
+
+		static FNativeTestRegistry* RegistryPtr;
+		TMap<UClass*, TFunction<bool()>> NativeTestMapping;
+	};
+}
+#endif
diff --git a/Engine/Plugins/Experimental/PCG/Source/PCGEditor/Private/PCGEditor.cpp b/Engine/Plugins/Experimental/PCG/Source/PCGEditor/Private/PCGEditor.cpp
--- a/Engine/Plugins/Experimental/PCG/Source/PCGEditor/Private/PCGEditor.cpp
+++ b/Engine/Plugins/Experimental/PCG/Source/PCGEditor/Private/PCGEditor.cpp
@@ -453,7 +453,7 @@ void FPCGEditor::OnDeterminismTests()
 				if (PCGNode->DefaultSettings->DeterminismSettings.bNativeTests)
 				{
 					// If the settings has a native test suite
-					if (TFunction<bool()> NativeTestSuite = PCGDeterminismTests::GetNativeTestIfExists(PCGNode->DefaultSettings))
+					if (TFunction<bool()> NativeTestSuite = PCGDeterminismTests::FNativeTestRegistry::GetNativeTestFunction(PCGNode->DefaultSettings))
 					{
 						FName NodeName(PCGNode->GetName());
 
