Minor bug fix for tautology and logic re-ordering in PCG test framework

#rnx

diff --git a/Engine/Plugins/Experimental/PCG/Source/PCG/Private/Tests/Determinism/PCGDeterminismTestsCommon.cpp b/Engine/Plugins/Experimental/PCG/Source/PCG/Private/Tests/Determinism/PCGDeterminismTestsCommon.cpp
--- a/Engine/Plugins/Experimental/PCG/Source/PCG/Private/Tests/Determinism/PCGDeterminismTestsCommon.cpp
+++ b/Engine/Plugins/Experimental/PCG/Source/PCG/Private/Tests/Determinism/PCGDeterminismTestsCommon.cpp
@@ -440,6 +440,11 @@ namespace PCGDeterminismTests
 			EPCGDataType FirstDataType = FirstCollection.TaggedData[I].Data->GetDataType();
 			EPCGDataType SecondDataType = SecondCollection.TaggedData[I].Data->GetDataType();
 
+			if (!DataTypeIsComparable(FirstDataType))
+			{
+				continue;
+			}
+
 			// Only compare if they are the same type and same pin label
 			if (FirstDataType != SecondDataType || FirstCollection.TaggedData[I].Pin != SecondCollection.TaggedData[I].Pin)
 			{
@@ -918,6 +923,12 @@ namespace PCGDeterminismTests
 
 	bool DataTypeIsComparable(EPCGDataType DataType)
 	{
+		// Data types that don't need to be compared
+		if (DataType == EPCGDataType::None || DataType == EPCGDataType::Other || DataType == EPCGDataType::Settings)
+		{
+			return false;
+		}
+
 		// Comparable data types
 		// TODO: Add metadata/param data as comparable in the future
 		if (!!(DataType & EPCGDataType::Spatial))
@@ -925,12 +936,6 @@ namespace PCGDeterminismTests
 			return true;
 		}
 
-		// Data types that don't need to be compared
-		if (DataType == EPCGDataType::None || DataType == EPCGDataType::Other || DataType == EPCGDataType::Settings)
-		{
-			return false;
-		}
-
 		UE_LOG(LogPCG, Warning, TEXT("Unknown data comparison type: %s"), *UEnum::GetValueAsString(DataType));
 		return false;
 	}
@@ -1175,7 +1180,7 @@ namespace PCGDeterminismTests
 
 		FPCGDataCollection ThirdOutput;
 		FPCGDataCollection FourthOutput;
-		ExecuteWithSameTestDataSameElement(TestData, PCGNode, FirstOutput, SecondOutput);
+		ExecuteWithSameTestDataSameElement(TestData, PCGNode, ThirdOutput, FourthOutput);
 
 		// Same data, so it ought to be exact
 		return DataCollectionsAreIdentical(FirstOutput, SecondOutput) && DataCollectionsAreIdentical(ThirdOutput, FourthOutput);
