Clarify the role of PCGSubsystem and PCGComponent

* PCGComponent is responsible to managing data (Graph, Resources, etc...) and the human interface to all PCG operations (via BP, API or Editor buttons)
* PCGSubsystem is responsible to the dispatch the calls from the component and to tick the graph executor. It is also the central entity to register to (PCG Components and PCG Partition Actors)
* Also fix tests

diff --git a/Engine/Plugins/Experimental/PCG/Source/PCG/Private/Graph/PCGGraphExecutor.cpp b/Engine/Plugins/Experimental/PCG/Source/PCG/Private/Graph/PCGGraphExecutor.cpp
--- a/Engine/Plugins/Experimental/PCG/Source/PCG/Private/Graph/PCGGraphExecutor.cpp
+++ b/Engine/Plugins/Experimental/PCG/Source/PCG/Private/Graph/PCGGraphExecutor.cpp
@@ -403,6 +403,7 @@ void FPCGGraphExecutor::Execute()
 
 		{
 			TRACE_CPUPROFILER_EVENT_SCOPE(FPCGGraphExecutor::Execute::ExecuteTasks);
+			bool bMainTaskDone = false;
 			// Execute main thread task
 			if (!ActiveTasks.IsEmpty())
 			{
@@ -413,13 +414,19 @@ void FPCGGraphExecutor::Execute()
 				if (!MainThreadTask.Context->bIsPaused && MainThreadTask.Element->Execute(MainThreadTask.Context.Get()))
 #endif
 				{
+					bMainTaskDone = true;
 					PostTaskExecute(ActiveTasks.Num() - 1);
 				}
 			}
 
-			for (int32 ExecutionIndex = ActiveTasks.Num() - 1; ExecutionIndex >= 0; --ExecutionIndex)
+			// If the main thread task is done, it was removed from ActiveTasks, and ActiveTasks should contains only futures (if there are some)
+			// Otherwise, ActiveTasks.Num() - 1 still point to the main thread tasks, so futures indexes are within [0, ActiveTasks.Num() - 2].
+			int32 FuturesLastIndex = bMainTaskDone ? ActiveTasks.Num() - 1 : ActiveTasks.Num() - 2;
+
+			// Then wait after all futures
+			for (int32 ExecutionIndex = FuturesLastIndex; ExecutionIndex >= 0; --ExecutionIndex)
 			{
-				bool bTaskDone = true;
+				bool bTaskDone = false;
 				// Wait on the future if any
 				if (TFuture<bool>* Future = Futures.Find(ExecutionIndex))
 				{
@@ -434,7 +441,7 @@ void FPCGGraphExecutor::Execute()
 			}
 		}
 
-		check(CurrentlyUsedThreads == 0);
+		check(ActiveTasks.Num() > 0 || CurrentlyUsedThreads == 0);
 
 		const double EndLoopTime = FPlatformTime::Seconds();
 		if (static_cast<float>((EndLoopTime - StartTime) * 1000) > CVarTimePerFrame.GetValueOnAnyThread())
diff --git a/Engine/Plugins/Experimental/PCG/Source/PCG/Private/Grid/PCGPartitionActor.cpp b/Engine/Plugins/Experimental/PCG/Source/PCG/Private/Grid/PCGPartitionActor.cpp
--- a/Engine/Plugins/Experimental/PCG/Source/PCG/Private/Grid/PCGPartitionActor.cpp
+++ b/Engine/Plugins/Experimental/PCG/Source/PCG/Private/Grid/PCGPartitionActor.cpp
@@ -47,6 +47,12 @@ void APCGPartitionActor::PostLoad()
 	{
 		BoundsComponent->SetBoxExtent(GetFixedBounds().GetExtent());
 	}
+
+	// Mark all our local components as local
+	for (UPCGComponent* LocalComponent : GetAllLocalPCGComponents())
+	{
+		LocalComponent->MarkAsLocalComponent();
+	}
 #endif // WITH_EDITOR
 }
 
@@ -205,6 +211,7 @@ void APCGPartitionActor::AddGraphInstance(UPCGComponent* OriginalComponent)
 	{
 		// Update properties as needed and early out
 		LocalComponent->SetPropertiesFromOriginal(OriginalComponent);
+		LocalComponent->MarkAsLocalComponent();
 		return;
 	}
 
@@ -215,6 +222,7 @@ void APCGPartitionActor::AddGraphInstance(UPCGComponent* OriginalComponent)
 	LocalComponent->SetPropertiesFromOriginal(OriginalComponent);
 
 	LocalComponent->RegisterComponent();
+	LocalComponent->MarkAsLocalComponent();
 	// TODO: check if we should use a non-instanced component?
 	AddInstanceComponent(LocalComponent);
 
@@ -237,7 +245,7 @@ bool APCGPartitionActor::RemoveGraphInstance(UPCGComponent* OriginalComponent)
 	LocalToOriginalMap.Remove(LocalComponent);
 
 	// TODO Add option to not cleanup?
-	LocalComponent->CleanupLocal(/*bRemoveComponents=*/true);
+	LocalComponent->CleanupLocalImmediate(/*bRemoveComponents=*/true);
 	LocalComponent->DestroyComponent();
 
 	return OriginalToLocalMap.IsEmpty();
diff --git a/Engine/Plugins/Experimental/PCG/Source/PCG/Private/PCGComponent.cpp b/Engine/Plugins/Experimental/PCG/Source/PCG/Private/PCGComponent.cpp
--- a/Engine/Plugins/Experimental/PCG/Source/PCG/Private/PCGComponent.cpp
+++ b/Engine/Plugins/Experimental/PCG/Source/PCG/Private/PCGComponent.cpp
@@ -54,12 +54,6 @@ bool UPCGComponent::IsPartitioned() const
 	return bIsPartitioned && CanPartition();
 }
 
-bool UPCGComponent::IsLocalComponent() const
-{
-	// TODO: Have a more robust solution for this.
-	return GetOwner() && GetOwner()->IsA<APCGPartitionActor>();
-}
-
 void UPCGComponent::SetGraph_Implementation(UPCGGraph* InGraph)
 {
 	SetGraphLocal(InGraph);
@@ -96,6 +90,7 @@ void UPCGComponent::AddToManagedResources(UPCGManagedResource* InResource)
 	if (InResource)
 	{
 		FScopeLock ResourcesLock(&GeneratedResourcesLock);
+		check(!GeneratedResourcesInaccessible);
 		GeneratedResources.Add(InResource);
 	}
 }
@@ -103,6 +98,7 @@ void UPCGComponent::AddToManagedResources(UPCGManagedResource* InResource)
 void UPCGComponent::ForEachManagedResource(TFunctionRef<void(UPCGManagedResource*)> Func)
 {
 	FScopeLock ResourcesLock(&GeneratedResourcesLock);
+	check(!GeneratedResourcesInaccessible);
 	for (TObjectPtr<UPCGManagedResource> ManagedResource : GeneratedResources)
 	{
 		Func(ManagedResource);
@@ -204,50 +200,24 @@ void UPCGComponent::Generate_Implementation(bool bForce)
 
 void UPCGComponent::GenerateLocal(bool bForce)
 {
-	GenerateLocal(bForce, EPCGComponentGenerationTrigger::GenerateOnDemand);
+	GenerateInternal(bForce, EPCGComponentGenerationTrigger::GenerateOnDemand, {});
 }
 
-void UPCGComponent::GenerateLocal(bool bForce, EPCGComponentGenerationTrigger RequestedGenerationTrigger)
+FPCGTaskId UPCGComponent::GenerateInternal(bool bForce, EPCGComponentGenerationTrigger RequestedGenerationTrigger, const TArray<FPCGTaskId>& Dependencies)
 {
-	if (IsGenerating() || !GetSubsystem())
+	if (IsGenerating() || !GetSubsystem() || !ShouldGenerate(bForce, RequestedGenerationTrigger))
 	{
-		return;
+		return InvalidPCGTaskId;
 	}
 
-	if (IsCleaningUp())
-	{
-		// If we are cleaning up, retry this function when it's done
-		TWeakObjectPtr<UPCGComponent> ThisComponentPtr(this);
-
-		CurrentGenerationTask = GetSubsystem()->ScheduleGeneric([ThisComponentPtr, bForce, RequestedGenerationTrigger]()
-			{
-				if (UPCGComponent* ThisComponent = ThisComponentPtr.Get())
-				{
-					// This is necessary to have "IsGenerating" to return false.
-					// Otherwise it will early out
-					ThisComponent->CurrentGenerationTask = InvalidPCGTaskId;
-					ThisComponent->GenerateLocal(bForce, RequestedGenerationTrigger);
-				}
+	CurrentGenerationTask = GetSubsystem()->ScheduleComponent(this, /*bSave=*/bForce, Dependencies);
 
-				return true;
-			}, { CurrentCleanupTask });
-	}
-	else
-	{
-		// Force component activation so it's easier to control by BP.
-		if (!bActivated)
-		{
-			Modify();
-			bActivated = true;
-		}
-
-		GenerateInternal(bForce, RequestedGenerationTrigger, {});
-	}
+	return CurrentGenerationTask;
 }
 
-FPCGTaskId UPCGComponent::GenerateInternal(bool bForce, EPCGComponentGenerationTrigger RequestedGenerationTrigger, const TArray<FPCGTaskId>& TaskDependencies)
+FPCGTaskId UPCGComponent::CreateGenerateTask(bool bForce, const TArray<FPCGTaskId>& Dependencies)
 {
-	if (IsGenerating() || !ShouldGenerate(bForce, RequestedGenerationTrigger))
+	if (IsGenerating())
 	{
 		return InvalidPCGTaskId;
 	}
@@ -264,18 +234,18 @@ FPCGTaskId UPCGComponent::GenerateInternal(bool bForce, EPCGComponentGenerationT
 
 	// Keep track of all the dependencies
 	TArray<FPCGTaskId> AdditionalDependencies;
-	const TArray<FPCGTaskId>* AllDependencies = &TaskDependencies;
+	const TArray<FPCGTaskId>* AllDependencies = &Dependencies;
 
 	if (IsCleaningUp())
 	{
-		AdditionalDependencies = TaskDependencies;
+		AdditionalDependencies = Dependencies;
 		AdditionalDependencies.Add(CurrentCleanupTask);
 		AllDependencies = &AdditionalDependencies;
 	}
 	else if (bGenerated)
 	{
 		// Immediate pass to mark all resources unused (and remove the ones that cannot be re-used)
-		CleanupInternalImmediate(/*bRemoveComponents=*/false);
+		CleanupLocalImmediate(/*bRemoveComponents=*/false);
 	}
 
 	const FBox NewBounds = GetGridBounds();
@@ -285,8 +255,7 @@ FPCGTaskId UPCGComponent::GenerateInternal(bool bForce, EPCGComponentGenerationT
 		return InvalidPCGTaskId;
 	}
 
-	CurrentGenerationTask = GetSubsystem()->ScheduleComponent(this, /*bSave=*/bForce, *AllDependencies);
-	return CurrentGenerationTask;
+	return GetSubsystem()->ScheduleGraph(this, *AllDependencies);
 }
 
 bool UPCGComponent::GetActorsFromTags(const TSet<FName>& InTags, TSet<TWeakObjectPtr<AActor>>& OutActors, bool bCullAgainstLocalBounds)
@@ -391,35 +360,18 @@ void UPCGComponent::Cleanup_Implementation(bool bRemoveComponents, bool bSave)
 
 void UPCGComponent::CleanupLocal(bool bRemoveComponents, bool bSave)
 {
-	if ((!bGenerated && !IsGenerating()) || !GetSubsystem() || IsCleaningUp())
-	{
-		return;
-	}
+	CleanupInternal(bRemoveComponents, bSave, {});
+}
 
-	if (IsGenerating())
+FPCGTaskId UPCGComponent::CleanupInternal(bool bRemoveComponents, bool bSave, const TArray<FPCGTaskId>& Dependencies)
+{
+	if ((!bGenerated && !IsGenerating()) || !GetSubsystem() || IsCleaningUp())
 	{
-		// If we are generating, retry this function when it's done
-		TWeakObjectPtr<UPCGComponent> ThisComponentPtr(this);
-
-		CurrentCleanupTask = GetSubsystem()->ScheduleGeneric([ThisComponentPtr, bRemoveComponents, bSave]()
-			{
-				if (UPCGComponent* ThisComponent = ThisComponentPtr.Get())
-				{
-					// This is necessary to have "IsCleaningUp" to return false.
-					// Otherwise it will early out
-					ThisComponent->CurrentCleanupTask = InvalidPCGTaskId;
-					ThisComponent->CleanupLocal(bRemoveComponents, bSave);
-				}
-
-				return true;
-			}, { CurrentGenerationTask });
+		return InvalidPCGTaskId;
 	}
-	else
-	{
-		Modify();
 
-		CurrentCleanupTask = GetSubsystem()->ScheduleCleanup(this, bRemoveComponents, bSave, {});
-	}
+	CurrentCleanupTask = GetSubsystem()->ScheduleCleanup(this, bRemoveComponents, bSave, Dependencies);
+	return CurrentCleanupTask;
 }
 
 AActor* UPCGComponent::ClearPCGLink(UClass* TemplateActor)
@@ -493,6 +445,7 @@ bool UPCGComponent::MoveResourcesToNewActor(AActor* InNewActor, bool bCreateChil
 	// Trying to move all resources for now. Perhaps in the future we won't want that.
 	{
 		FScopeLock ResourcesLock(&GeneratedResourcesLock);
+		check(!GeneratedResourcesInaccessible);
 		for (TObjectPtr<UPCGManagedResource>& GeneratedResource : GeneratedResources)
 		{
 			check(GeneratedResource);
@@ -515,7 +468,7 @@ bool UPCGComponent::MoveResourcesToNewActor(AActor* InNewActor, bool bCreateChil
 	return bHasMovedResources;
 }
 
-void UPCGComponent::CleanupInternalImmediate(bool bRemoveComponents)
+void UPCGComponent::CleanupLocalImmediate(bool bRemoveComponents)
 {
 	if (!bGenerated || IsPartitioned() || IsCleaningUp() || IsGenerating())
 	{
@@ -526,6 +479,7 @@ void UPCGComponent::CleanupInternalImmediate(bool bRemoveComponents)
 
 	{
 		FScopeLock ResourcesLock(&GeneratedResourcesLock);
+		check(!GeneratedResourcesInaccessible);
 		for (int32 ResourceIndex = GeneratedResources.Num() - 1; ResourceIndex >= 0; --ResourceIndex)
 		{
 			if (GeneratedResources[ResourceIndex]->Release(bRemoveComponents, ActorsToDelete))
@@ -538,17 +492,13 @@ void UPCGComponent::CleanupInternalImmediate(bool bRemoveComponents)
 	UPCGActorHelpers::DeleteActors(GetWorld(), ActorsToDelete.Array());
 }
 
-FPCGTaskId UPCGComponent::CleanupInternal(bool bRemoveComponents, const TArray<FPCGTaskId>& Dependencies)
+FPCGTaskId UPCGComponent::CreateCleanupTask(bool bRemoveComponents, const TArray<FPCGTaskId>& Dependencies)
 {
 	if ((!bGenerated && !IsGenerating()) || IsPartitioned() || IsCleaningUp())
 	{
 		return InvalidPCGTaskId;
 	}
 
-	// Need a shared ptr to keep track of all actors to delete.
-	// Will be protected by GeneratedResourcesLock
-	TSharedPtr<TSet<TSoftObjectPtr<AActor>>> ActorsToDeletePtr = MakeShared<TSet<TSoftObjectPtr<AActor>>>();
-
 	// Keep track of all the dependencies
 	TArray<FPCGTaskId> AdditionalDependencies;
 	const TArray<FPCGTaskId>* AllDependencies = &Dependencies;
@@ -560,81 +510,63 @@ FPCGTaskId UPCGComponent::CleanupInternal(bool bRemoveComponents, const TArray<F
 		AllDependencies = &AdditionalDependencies;
 	}
 
-	TArray<FPCGTaskId> AllTasks;
+	struct CleanupContext
+	{
+		bool bIsFirstIteration = true;
+		int32 ResourceIndex = -1;
+		TSet<TSoftObjectPtr<AActor>> ActorsToDelete;
+	};
 
-	// Use a WeakPtr to guard against the case when the Component doesn't exist anymore.
+	TSharedPtr<CleanupContext> Context = MakeShared<CleanupContext>();
 	TWeakObjectPtr<UPCGComponent> ThisComponentWeakPtr(this);
+	TWeakObjectPtr<UWorld> WorldPtr(GetWorld());
 
+	auto CleanupTask = [Context, ThisComponentWeakPtr, WorldPtr, bRemoveComponents]()
 	{
-		FScopeLock ResourcesLock(&GeneratedResourcesLock);
-		for (int32 ResourceIndex = GeneratedResources.Num() - 1; ResourceIndex >= 0; --ResourceIndex)
+		if (UPCGComponent* ThisComponent = ThisComponentWeakPtr.Get())
 		{
-			check(GeneratedResources[ResourceIndex]);
-			TWeakObjectPtr<UPCGManagedResource> GeneratedResourcePtr(GeneratedResources[ResourceIndex]);
+			FScopeLock ResourcesLock(&ThisComponent->GeneratedResourcesLock);
 
-			// For now this is done on a single thread, but we could extend it to multi-thread if the release support MT.
-			FPCGTaskId ReleaseTask = GetSubsystem()->ScheduleGeneric([GeneratedResourcePtr, ResourceIndex, bRemoveComponents, ActorsToDeletePtr, ThisComponentWeakPtr]()
+			// Safeguard to track illegal modifications of the generated resources array while doing cleanup
+			if (Context->bIsFirstIteration)
 			{
-				if (UPCGComponent* ThisComponent = ThisComponentWeakPtr.Get())
-				{
-					if (UPCGManagedResource* GeneratedResource = GeneratedResourcePtr.Get())
-					{
-						TSet<TSoftObjectPtr<AActor>> ActorsToDelete;
-						GeneratedResource->Release(bRemoveComponents, ActorsToDelete);
-
-						if (!ActorsToDelete.IsEmpty())
-						{
-							FScopeLock Lock(&ThisComponent->GeneratedResourcesLock);
-							ActorsToDeletePtr->Append(ActorsToDelete);
-						}
-					}
-				}
-
-				return true;
-			}, *AllDependencies);
-
-			AllTasks.Add(ReleaseTask);
-		}
-	}
-
-	// Use a WeakPtr to guard against the case when the World doesn't exist anymore.
-	TWeakObjectPtr<UWorld> WorldPtr(GetWorld());
+				check(!ThisComponent->GeneratedResourcesInaccessible);
+				ThisComponent->GeneratedResourcesInaccessible = true;
+				Context->ResourceIndex = ThisComponent->GeneratedResources.Num() - 1;
+				Context->bIsFirstIteration = false;
+			}
 
-	FPCGTaskId PostReleaseTask = GetSubsystem()->ScheduleGeneric([WorldPtr, ThisComponentWeakPtr, bRemoveComponents, ActorsToDeletePtr]()
-		{
-			if (UPCGComponent* ThisComponent = ThisComponentWeakPtr.Get())
+			// Going backward
+			if (Context->ResourceIndex >= 0)
 			{
-				ThisComponent->Modify();
-				FScopeLock ResourcesLock(&ThisComponent->GeneratedResourcesLock);
+				UPCGManagedResource* Resource = ThisComponent->GeneratedResources[Context->ResourceIndex];
+				check(Resource);
 
-				// If we remove components, all generated resources are removed
-				if (bRemoveComponents)
-				{
-					ThisComponent->GeneratedResources.Empty();
-				}
-				else
+				if (Resource->Release(bRemoveComponents, Context->ActorsToDelete))
 				{
-					// Otherwise, remove only those who are not marked unused (meaning that they can't be re-used)
-					for (int32 ResourceIndex = ThisComponent->GeneratedResources.Num() - 1; ResourceIndex >= 0; --ResourceIndex)
-					{
-						if (!ThisComponent->GeneratedResources[ResourceIndex]->IsMarkedUnused())
-						{
-							ThisComponent->GeneratedResources.RemoveAtSwap(ResourceIndex);
-						}
-					}
+					ThisComponent->GeneratedResources.RemoveAtSwap(Context->ResourceIndex);
 				}
-			}
 
-			if (UWorld* World = WorldPtr.Get())
+				Context->ResourceIndex--;
+
+				// Returning false means the task is not over
+				return false;
+			}
+			else
 			{
-				UPCGActorHelpers::DeleteActors(World, ActorsToDeletePtr->Array());
+				ThisComponent->GeneratedResourcesInaccessible = false;
 			}
+		}
+
+		if (UWorld* World = WorldPtr.Get())
+		{
+			UPCGActorHelpers::DeleteActors(World, Context->ActorsToDelete.Array());
+		}
 
-			return true;
-		}, AllTasks);
+		return true;
+	};
 
-	CurrentCleanupTask = PostReleaseTask;
-	return CurrentCleanupTask;
+	return GetSubsystem()->ScheduleGeneric(CleanupTask, *AllDependencies);
 }
 
 void UPCGComponent::CleanupUnusedManagedResources()
@@ -643,6 +575,7 @@ void UPCGComponent::CleanupUnusedManagedResources()
 
 	{
 		FScopeLock ResourcesLock(&GeneratedResourcesLock);
+		check(!GeneratedResourcesInaccessible);
 		for (int32 ResourceIndex = GeneratedResources.Num() - 1; ResourceIndex >= 0; --ResourceIndex)
 		{
 			check(GeneratedResources[ResourceIndex]);
@@ -681,7 +614,7 @@ void UPCGComponent::BeginPlay()
 		}
 		else
 		{
-			GenerateLocal(/*bForce=*/false, EPCGComponentGenerationTrigger::GenerateOnLoad);
+			GenerateInternal(/*bForce=*/false, EPCGComponentGenerationTrigger::GenerateOnLoad, {});
 			bRuntimeGenerated = true;
 		}
 	}
@@ -705,7 +638,6 @@ void UPCGComponent::OnComponentCreated()
 
 #if WITH_EDITOR
 	SetupActorCallbacks();
-	UpdateIsLocalComponent();
 #endif
 }
 
@@ -781,8 +713,6 @@ void UPCGComponent::PostLoad()
 	{
 		Graph->OnGraphChangedDelegate.AddUObject(this, &UPCGComponent::OnGraphChanged);
 	}
-	
-	UpdateIsLocalComponent();
 #endif
 }
 
@@ -925,11 +855,11 @@ void UPCGComponent::PostEditChangeProperty(FPropertyChangedEvent& PropertyChange
 
 					if (bIsNowPartitioned)
 					{
-						GetSubsystem()->DelayPartitionGraph(this);
+						GetSubsystem()->SchedulePartitionGraph(this);
 					}
 					else
 					{
-						GetSubsystem()->DelayUnpartitionGraph(this);
+						GetSubsystem()->ScheduleUnpartitionGraph(this);
 
 					}
 
@@ -1291,8 +1221,10 @@ void UPCGComponent::OnObjectPropertyChanged(UObject* InObject, FPropertyChangedE
 	bool bValueNotInteractive = (InEvent.ChangeType != EPropertyChangeType::Interactive);
 	// Special exception for actor tags, as we can't track otherwise an actor "losing" a tag
 	bool bActorTagChange = (InEvent.Property && InEvent.Property->GetFName() == GET_MEMBER_NAME_CHECKED(AActor, Tags));
+	// Discard property changed if this is this component. They will be caught by the PostEditChangeProperty
+	bool bIsThisComponent = InObject == this;
 
-	if(!bValueNotInteractive && !bActorTagChange)
+	if (bIsThisComponent || (!bValueNotInteractive && !bActorTagChange))
 	{
 		return;
 	}
@@ -1474,13 +1406,13 @@ void UPCGComponent::Refresh()
 		{
 			if (LastGeneratedBounds.IsValid)
 			{
-				GetSubsystem()->DelayUnpartitionGraph(this);
+				GetSubsystem()->ScheduleUnpartitionGraph(this);
 			}
 		}
 		else
 		{
 			bool bWasGenerated = bGenerated;
-			CleanupLocal(/*bRemoveComponents=*/false);
+			CleanupLocalImmediate(/*bRemoveComponents=*/true);
 			bGenerated = bWasGenerated;
 		}
 	}
@@ -1492,7 +1424,7 @@ void UPCGComponent::Refresh()
 		}
 		else if (IsPartitioned() && GetSubsystem())
 		{
-			GetSubsystem()->DelayPartitionGraph(this);
+			GetSubsystem()->SchedulePartitionGraph(this);
 		}
 	}
 }
@@ -2151,14 +2083,6 @@ bool UPCGComponent::GraphUsesLandscapePin() const
 	return Graph && Graph->GetInputNode()->IsOutputPinConnected(PCGInputOutputConstants::DefaultLandscapeLabel);
 }
 
-void UPCGComponent::UpdateIsLocalComponent()
-{
-	if (IsLocalComponent())
-	{
-		bIsComponentLocal = true;
-	}
-}
-
 #endif // WITH_EDITOR
 
 #undef LOCTEXT_NAMESPACE
diff --git a/Engine/Plugins/Experimental/PCG/Source/PCG/Private/PCGManagedResource.cpp b/Engine/Plugins/Experimental/PCG/Source/PCG/Private/PCGManagedResource.cpp
--- a/Engine/Plugins/Experimental/PCG/Source/PCG/Private/PCGManagedResource.cpp
+++ b/Engine/Plugins/Experimental/PCG/Source/PCG/Private/PCGManagedResource.cpp
@@ -62,7 +62,7 @@ bool UPCGManagedActors::Release(bool bHardRelease, TSet<TSoftObjectPtr<AActor>>&
 			{
 				// It is more complicated to handled a non-immediate cleanup when doing it recursively in the managed actors.
 				// Do it all immediate then.
-				Component->CleanupInternalImmediate(/*bRemoveComponents=*/bHardRelease);
+				Component->CleanupLocalImmediate(/*bRemoveComponents=*/bHardRelease);
 			}
 
 			ComponentsToCleanup.Reset();
diff --git a/Engine/Plugins/Experimental/PCG/Source/PCG/Private/PCGSubsystem.cpp b/Engine/Plugins/Experimental/PCG/Source/PCG/Private/PCGSubsystem.cpp
--- a/Engine/Plugins/Experimental/PCG/Source/PCG/Private/PCGSubsystem.cpp
+++ b/Engine/Plugins/Experimental/PCG/Source/PCG/Private/PCGSubsystem.cpp
@@ -149,7 +149,7 @@ FPCGTaskId UPCGSubsystem::ScheduleComponent(UPCGComponent* PCGComponent, bool bS
 #if WITH_EDITOR
 	if (PCGComponent->IsPartitioned() && !PCGHelpers::IsRuntimeOrPIE())
 	{
-		return DelayGenerateGraph(PCGComponent, bSave);
+		return ScheduleProcessGraph(PCGComponent, bSave);
 	}
 #endif // WITH_EDITOR
 
@@ -169,7 +169,7 @@ FPCGTaskId UPCGSubsystem::ScheduleComponent(UPCGComponent* PCGComponent, bool bS
 	}
 	else
 	{
-		FPCGTaskId TaskId = GraphExecutor->Schedule(PCGComponent, Dependencies);
+		FPCGTaskId TaskId = PCGComponent->CreateGenerateTask(/*bForce=*/bSave, Dependencies);
 		if (TaskId != InvalidPCGTaskId)
 		{
 			AllTasks.Add(TaskId);
@@ -234,7 +234,7 @@ FPCGTaskId UPCGSubsystem::ScheduleCleanup(UPCGComponent* PCGComponent, bool bRem
 	}
 	else
 	{
-		FPCGTaskId TaskId = PCGComponent->CleanupInternal(bRemoveComponents, Dependencies);
+		FPCGTaskId TaskId = PCGComponent->CreateCleanupTask(bRemoveComponents, Dependencies);
 		if (TaskId != InvalidPCGTaskId)
 		{
 			AllTasks.Add(TaskId);
@@ -259,6 +259,18 @@ FPCGTaskId UPCGSubsystem::ScheduleCleanup(UPCGComponent* PCGComponent, bool bRem
 		}, AllTasks);
 }
 
+FPCGTaskId UPCGSubsystem::ScheduleGraph(UPCGComponent* SourceComponent, const TArray<FPCGTaskId>& Dependencies)
+{
+	if (SourceComponent)
+	{
+		return GraphExecutor->Schedule(SourceComponent, Dependencies);
+	}
+	else
+	{
+		return InvalidPCGTaskId;
+	}
+}
+
 FPCGTaskId UPCGSubsystem::ScheduleGraph(UPCGGraph* Graph, UPCGComponent* SourceComponent, FPCGElementPtr InputElement, const TArray<FPCGTaskId>& Dependencies)
 {
 	if (SourceComponent)
@@ -328,7 +340,7 @@ TArray<FPCGTaskId> UPCGSubsystem::ScheduleMultipleCleanup(UPCGComponent* Origina
 				// Add check to avoid infinite loop
 				if (ensure(!LocalComponent->IsPartitioned()))
 				{
-					FPCGTaskId LocalTask = LocalComponent->CleanupInternal(bRemoveComponents, Dependencies);
+					FPCGTaskId LocalTask = LocalComponent->CleanupInternal(bRemoveComponents, /*bSave=*/ false, Dependencies);
 
 					if (LocalTask != InvalidPCGTaskId)
 					{
@@ -770,7 +782,7 @@ FPCGTaskId UPCGSubsystem::ProcessGraph(UPCGComponent* Component, const FBox& InP
 
 	// TODO: optimal implementation would find the difference between the previous bounds and the new bounds
 	// and process these only. This is esp. important because of the CreateActor parameter.
-	auto ScheduleTask = [this, ComponentPtr, InOperation](APCGPartitionActor* PCGActor, const FBox& InBounds, const TArray<FPCGTaskId>& TaskDependencies){
+	auto ScheduleTask = [this, ComponentPtr, InOperation, bSave](APCGPartitionActor* PCGActor, const FBox& InBounds, const TArray<FPCGTaskId>& TaskDependencies){
 		TWeakObjectPtr<APCGPartitionActor> PCGActorPtr(PCGActor);
 
 		auto UnpartitionTask = [ComponentPtr, PCGActorPtr]() {
@@ -804,7 +816,7 @@ FPCGTaskId UPCGSubsystem::ProcessGraph(UPCGComponent* Component, const FBox& InP
 			return true;
 		};
 
-		auto ScheduleGraph = [this, PCGActorPtr, ComponentPtr, &TaskDependencies]() {
+		auto ScheduleGraph = [this, PCGActorPtr, ComponentPtr, bSave, TaskDependencies]() {
 			if (APCGPartitionActor* PCGActor = PCGActorPtr.Get())
 			{
 				if (UPCGComponent* Component = ComponentPtr.Get())
@@ -820,7 +832,7 @@ FPCGTaskId UPCGSubsystem::ProcessGraph(UPCGComponent* Component, const FBox& InP
 						return InvalidPCGTaskId;
 					}
 
-					return LocalComponent->GenerateInternal(/*bForce=*/false, EPCGComponentGenerationTrigger::GenerateOnDemand, TaskDependencies);
+					return LocalComponent->GenerateInternal(bSave, EPCGComponentGenerationTrigger::GenerateOnDemand, TaskDependencies);
 				}
 				else
 				{
@@ -899,7 +911,7 @@ FPCGTaskId UPCGSubsystem::CleanupGraph(UPCGComponent* Component, const FBox& InB
 			// Ensure to avoid infinite loop
 			if (ensure(!LocalComponent->IsPartitioned()))
 			{
-				return ScheduleCleanup(LocalComponent, bRemoveComponents, bSave, TaskDependencies);
+				return LocalComponent->CleanupInternal(bRemoveComponents, bSave, TaskDependencies);
 			}
 		}
 
@@ -1193,29 +1205,34 @@ void UPCGSubsystem::ResetPartitionActorsMap()
 }
 
 ///////////////////////////////////////////////////////////////////////////////////////////
-void UPCGSubsystem::DelayPartitionGraph(UPCGComponent* Component)
+FPCGTaskId UPCGSubsystem::SchedulePartitionGraph(UPCGComponent* Component)
 {
-	DelayProcessGraph(Component, EOperation::Partition, /*bSave=*/false);
-}
+	check(Component && Component->IsPartitioned());
 
-void UPCGSubsystem::DelayUnpartitionGraph(UPCGComponent* Component)
-{
-	DelayProcessGraph(Component, EOperation::Unpartition, /*bSave=*/false);
+	const FBox PreviousBounds = Component->LastGeneratedBounds;
+	const FBox NewBounds = Component->GetGridBounds();
+
+	return ProcessGraph(Component, PreviousBounds, NewBounds, EOperation::Partition, /*bSave=*/false);
 }
 
-FPCGTaskId UPCGSubsystem::DelayGenerateGraph(UPCGComponent* Component, bool bSave)
+FPCGTaskId UPCGSubsystem::ScheduleUnpartitionGraph(UPCGComponent* Component)
 {
-	return DelayProcessGraph(Component, EOperation::Generate, bSave);
+	check(Component && Component->IsPartitioned());
+
+	const FBox PreviousBounds = Component->LastGeneratedBounds;
+	const FBox NewBounds = Component->GetGridBounds();
+
+	return ProcessGraph(Component, PreviousBounds, NewBounds, EOperation::Unpartition, /*bSave=*/false);
 }
 
-FPCGTaskId UPCGSubsystem::DelayProcessGraph(UPCGComponent* Component, EOperation InOperation, bool bSave)
+FPCGTaskId UPCGSubsystem::ScheduleProcessGraph(UPCGComponent* Component, bool bSave)
 {
 	check(Component && Component->IsPartitioned());
 
 	const FBox PreviousBounds = Component->LastGeneratedBounds;
 	const FBox NewBounds = Component->GetGridBounds();
 
-	return ProcessGraph(Component, PreviousBounds, NewBounds, InOperation, bSave);
+	return ProcessGraph(Component, PreviousBounds, NewBounds, EOperation::Generate, bSave);
 }
 
 #endif // WITH_EDITOR
diff --git a/Engine/Plugins/Experimental/PCG/Source/PCG/Private/Tests/Elements/PCGInstancePackerTest.cpp b/Engine/Plugins/Experimental/PCG/Source/PCG/Private/Tests/Elements/PCGInstancePackerTest.cpp
--- a/Engine/Plugins/Experimental/PCG/Source/PCG/Private/Tests/Elements/PCGInstancePackerTest.cpp
+++ b/Engine/Plugins/Experimental/PCG/Source/PCG/Private/Tests/Elements/PCGInstancePackerTest.cpp
@@ -95,7 +95,7 @@ namespace
 		}
 
 		TestData.TestPCGComponent->bGenerated = true;
-		TestData.TestPCGComponent->Cleanup();
+		TestData.TestPCGComponent->CleanupLocalImmediate(true);
 
 		return bTestPassed;
 	};
diff --git a/Engine/Plugins/Experimental/PCG/Source/PCG/Private/Tests/Elements/PCGStaticMeshSpawnerTest.cpp b/Engine/Plugins/Experimental/PCG/Source/PCG/Private/Tests/Elements/PCGStaticMeshSpawnerTest.cpp
--- a/Engine/Plugins/Experimental/PCG/Source/PCG/Private/Tests/Elements/PCGStaticMeshSpawnerTest.cpp
+++ b/Engine/Plugins/Experimental/PCG/Source/PCG/Private/Tests/Elements/PCGStaticMeshSpawnerTest.cpp
@@ -251,7 +251,7 @@ void TestMeshSelectorByAttribute(
 	}
 
 	TestData.TestPCGComponent->bGenerated = true;
-	TestData.TestPCGComponent->Cleanup();
+	TestData.TestPCGComponent->CleanupLocalImmediate(true);
 	
 	Test->TestEqual("Valid total instance count", TotalInstanceCount, ValidInstanceCount);
 }
@@ -340,7 +340,7 @@ void TestMeshSelectorWeighted(
 	}
 
 	TestData.TestPCGComponent->bGenerated = true;
-	TestData.TestPCGComponent->Cleanup();
+	TestData.TestPCGComponent->CleanupLocalImmediate(true);
 	
 	Test->TestEqual("Valid total instance count", TotalInstanceCount, PointCount);
 }
@@ -423,7 +423,7 @@ void TestMeshSelectorWeightedByCategory(
 	}
 
 	TestData.TestPCGComponent->bGenerated = true;
-	TestData.TestPCGComponent->Cleanup();
+	TestData.TestPCGComponent->CleanupLocalImmediate(true);
 
 	Test->TestEqual("Valid total instance count", TotalInstanceCount, PointCount);
 }
diff --git a/Engine/Plugins/Experimental/PCG/Source/PCG/Public/PCGComponent.h b/Engine/Plugins/Experimental/PCG/Source/PCG/Public/PCGComponent.h
--- a/Engine/Plugins/Experimental/PCG/Source/PCG/Public/PCGComponent.h
+++ b/Engine/Plugins/Experimental/PCG/Source/PCG/Public/PCGComponent.h
@@ -97,14 +97,17 @@ public:
 	void Generate();
 	void Cleanup();
 
-	/** Starts generation from a local (vs. remote) standpoint. Will not be replicated. */
+	/** Starts generation from a local (vs. remote) standpoint. Will not be replicated. Will be delayed. */
 	UFUNCTION(BlueprintCallable, Category = PCG)
 	void GenerateLocal(bool bForce);
 
-	/** Cleans up the generation from a local (vs. remote) standpoint. Will not be replicated. */
+	/** Cleans up the generation from a local (vs. remote) standpoint. Will not be replicated. Will be delayed. */
 	UFUNCTION(BlueprintCallable, Category = PCG)
 	void CleanupLocal(bool bRemoveComponents, bool bSave = false);
 
+	/* Same as CleanupLocal, but without any delayed tasks. All is done immediately. */
+	void CleanupLocalImmediate(bool bRemoveComponents);
+
 	/** Networked generation call that also activates the component as needed */
 	UFUNCTION(BlueprintCallable, NetMulticast, Reliable, Category = PCG)
 	void Generate(bool bForce);
@@ -148,9 +151,6 @@ public:
 	bool bRuntimeGenerated = false;
 
 #if WITH_EDITORONLY_DATA
-	UPROPERTY(VisibleAnywhere, Transient, Category = Properties, meta = (EditCondition = false, EditConditionHides))
-	bool bIsComponentLocal = false;
-
 	UPROPERTY(BlueprintReadWrite, EditAnywhere, AdvancedDisplay, Category = Properties, meta = (DisplayName = "Regenerate PCG volume in editor"))
 	bool bRegenerateInEditor = true;
 
@@ -181,7 +181,10 @@ public:
 #endif
 
 	bool IsPartitioned() const;
-	bool IsLocalComponent() const;
+	bool IsLocalComponent() const { return bIsComponentLocal; }
+
+	/* Responsibility of the PCG Partition Actor to mark is local */
+	void MarkAsLocalComponent() { bIsComponentLocal = true; }
 
 	/** Updates internal properties from other component, dirties as required but does not trigger Refresh */
 	void SetPropertiesFromOriginal(const UPCGComponent* Original);
@@ -204,16 +207,16 @@ private:
 	bool IsLandscapeCachedDataDirty(const UPCGData* Data) const;
 
 	bool ShouldGenerate(bool bForce, EPCGComponentGenerationTrigger RequestedGenerationTrigger) const;
-	void GenerateLocal(bool bForce, EPCGComponentGenerationTrigger RequestedGenerationTrigger);
 
-	/* Internal call to create tasks to generate the component. If there is nothing to do, an invalid task id will be returned. */
+	/* Internal call that allows to delay a Generate/Cleanup call, chain with dependencies and keep track of the task id created. This task id is also returned. */
 	FPCGTaskId GenerateInternal(bool bForce, EPCGComponentGenerationTrigger RequestedGenerationTrigger, const TArray<FPCGTaskId>& Dependencies);
+	FPCGTaskId CleanupInternal(bool bRemoveComponents, bool bSave, const TArray<FPCGTaskId>& Dependencies);
 
-	/* Internal call to create tacks to cleanup the component. If there is nothing to do, an invalid task id will be returned. */
-	FPCGTaskId CleanupInternal(bool bRemoveComponents, const TArray<FPCGTaskId>& Dependencies);
+	/* Internal call to create tasks to generate the component. If there is nothing to do, an invalid task id will be returned. Should only be used by the subsystem. */
+	FPCGTaskId CreateGenerateTask(bool bForce, const TArray<FPCGTaskId>& Dependencies);
 
-	/* Same as CleanupInternal, but without any delayed tasks. All is done immediately. */
-	void CleanupInternalImmediate(bool bRemoveComponents);
+	/* Internal call to create tasks to cleanup the component. If there is nothing to do, an invalid task id will be returned. Should only be used by the subsystem. */
+	FPCGTaskId CreateCleanupTask(bool bRemoveComponents, const TArray<FPCGTaskId>& Dependencies);
 
 	void PostProcessGraph(const FBox& InNewBounds, bool bInGenerated);
 	void PostCleanupGraph();
@@ -292,12 +295,18 @@ private:
 	UPROPERTY()
 	TArray<TObjectPtr<UPCGManagedResource>> GeneratedResources;
 
+	// When doing a cleanup, locking resource modification. Used as sentinel.
+	bool GeneratedResourcesInaccessible = false;
+
 	UPROPERTY()
 	FBox LastGeneratedBounds = FBox(EForceInit::ForceInit);
 
 	FPCGTaskId CurrentGenerationTask = InvalidPCGTaskId;
 	FPCGTaskId CurrentCleanupTask = InvalidPCGTaskId;
 
+	UPROPERTY(VisibleAnywhere, Transient, Category = Properties, meta = (EditCondition = false, EditConditionHides))
+	bool bIsComponentLocal = false;
+
 #if WITH_EDITOR
 	bool bIsInspecting = false;
 	FBox LastGeneratedBoundsPriorToUndo = FBox(EForceInit::ForceInit);
@@ -307,7 +316,6 @@ private:
 	void TeardownLandscapeTracking();
 	void UpdateTrackedLandscape(bool bBoundsCheck = true);
 	void OnLandscapeChanged(ALandscapeProxy* Landscape, const FLandscapeProxyComponentDataChangedParams& ChangeParams);
-	void UpdateIsLocalComponent();
 #endif
 
 #if WITH_EDITORONLY_DATA
diff --git a/Engine/Plugins/Experimental/PCG/Source/PCG/Public/PCGSubsystem.h b/Engine/Plugins/Experimental/PCG/Source/PCG/Public/PCGSubsystem.h
--- a/Engine/Plugins/Experimental/PCG/Source/PCG/Public/PCGSubsystem.h
+++ b/Engine/Plugins/Experimental/PCG/Source/PCG/Public/PCGSubsystem.h
@@ -65,6 +65,9 @@ public:
 	// Schedule graph (used internally for dynamic subgraph execution)
 	FPCGTaskId ScheduleGraph(UPCGGraph* Graph, UPCGComponent* SourceComponent, FPCGElementPtr InputElement, const TArray<FPCGTaskId>& Dependencies);
 
+	// Schedule graph (used internally for dynamic subgraph execution)
+	FPCGTaskId ScheduleGraph(UPCGComponent* SourceComponent, const TArray<FPCGTaskId>& Dependencies);
+
 	/** General job scheduling, used to control loading/unloading */
 	FPCGTaskId ScheduleGeneric(TFunction<bool()> InOperation, const TArray<FPCGTaskId>& TaskDependencies);
 
@@ -89,9 +92,9 @@ public:
 #if WITH_EDITOR
 public:
 	/** Operations to schedule a later operation, which enables to delay bounds querying */
-	void DelayPartitionGraph(UPCGComponent* Component);
-	void DelayUnpartitionGraph(UPCGComponent* Component);
-	FPCGTaskId DelayGenerateGraph(UPCGComponent* Component, bool bSave);
+	FPCGTaskId SchedulePartitionGraph(UPCGComponent* Component);
+	FPCGTaskId ScheduleUnpartitionGraph(UPCGComponent* Component);
+	FPCGTaskId ScheduleProcessGraph(UPCGComponent* Component, bool bSave);
 
 	/** Schedules an operation to cleanup the graph in the given bounds */
 	FPCGTaskId CleanupGraph(UPCGComponent* Component, const FBox& InBounds, bool bRemoveComponents, bool bSave);
@@ -132,7 +135,6 @@ private:
 		Generate
 	};
 
-	FPCGTaskId DelayProcessGraph(UPCGComponent* Component, EOperation InOperation, bool bSave);
 	FPCGTaskId ProcessGraph(UPCGComponent* Component, const FBox& InPreviousBounds, const FBox& InNewBounds, EOperation InOperation, bool bSave);
 #endif // WITH_EDITOR
 	
