Fixed issues related to blueprint actors and PCG (needed a custom component data instance class to manage resources properly)
Fixed non-disappearing async notification when quitting PIE.

#rnx

diff --git a/Engine/Plugins/Experimental/PCG/Source/PCG/Private/Elements/PCGDebugElement.cpp b/Engine/Plugins/Experimental/PCG/Source/PCG/Private/Elements/PCGDebugElement.cpp
--- a/Engine/Plugins/Experimental/PCG/Source/PCG/Private/Elements/PCGDebugElement.cpp
+++ b/Engine/Plugins/Experimental/PCG/Source/PCG/Private/Elements/PCGDebugElement.cpp
@@ -2,6 +2,7 @@
 
 #include "Elements/PCGDebugElement.h"
 
+#include "PCGComponent.h"
 #include "PCGHelpers.h"
 #include "Data/PCGPointData.h"
 #include "Helpers/PCGActorHelpers.h"
@@ -16,6 +17,12 @@ namespace PCGDebugElement
 	{
 		TRACE_CPUPROFILER_EVENT_SCOPE(PCGDebugElement::ExecuteDebugDisplay);
 #if WITH_EDITORONLY_DATA
+		// Early validation: if we don't have a valid PCG component, we're not going to add the debug display info
+		if (!Context->SourceComponent || !IsValid(Context->SourceComponent))
+		{
+			return;
+		}
+
 		const UPCGSettings* Settings = Context->GetInputSettings<UPCGSettings>();
 
 		if (!Settings)
@@ -140,6 +147,7 @@ namespace PCGDebugElement
 			Params.CollisionProfile = UCollisionProfile::NoCollision_ProfileName;
 
 			UInstancedStaticMeshComponent* ISMC = UPCGActorHelpers::GetOrCreateISMC(TargetActor, Context->SourceComponent, Params);
+			check(ISMC);
 			
 			ISMC->ComponentTags.AddUnique(PCGHelpers::DefaultPCGDebugTag);
 			ISMC->NumCustomDataFloats = NumCustomData;
diff --git a/Engine/Plugins/Experimental/PCG/Source/PCG/Private/Graph/PCGGraphExecutor.cpp b/Engine/Plugins/Experimental/PCG/Source/PCG/Private/Graph/PCGGraphExecutor.cpp
--- a/Engine/Plugins/Experimental/PCG/Source/PCG/Private/Graph/PCGGraphExecutor.cpp
+++ b/Engine/Plugins/Experimental/PCG/Source/PCG/Private/Graph/PCGGraphExecutor.cpp
@@ -45,7 +45,12 @@ FPCGGraphExecutor::FPCGGraphExecutor(UObject* InOwner)
 {
 }
 
-FPCGGraphExecutor::~FPCGGraphExecutor() = default;
+FPCGGraphExecutor::~FPCGGraphExecutor()
+{
+#if WITH_EDITOR
+	EndGenerationNotification();
+#endif
+}
 
 void FPCGGraphExecutor::Compile(UPCGGraph* Graph)
 {
diff --git a/Engine/Plugins/Experimental/PCG/Source/PCG/Private/PCGComponent.cpp b/Engine/Plugins/Experimental/PCG/Source/PCG/Private/PCGComponent.cpp
--- a/Engine/Plugins/Experimental/PCG/Source/PCG/Private/PCGComponent.cpp
+++ b/Engine/Plugins/Experimental/PCG/Source/PCG/Private/PCGComponent.cpp
@@ -736,6 +736,12 @@ void UPCGComponent::BeginDestroy()
 	Super::BeginDestroy();
 }
 
+TStructOnScope<FActorComponentInstanceData> UPCGComponent::GetComponentInstanceData() const
+{
+	TStructOnScope<FActorComponentInstanceData> InstanceData = MakeStructOnScope<FActorComponentInstanceData, FPCGComponentInstanceData>(this);
+	return InstanceData;
+}
+
 void UPCGComponent::OnGraphChanged(UPCGGraph* InGraph, EPCGChangeType ChangeType)
 {
 	const bool bIsStructural = ((ChangeType & (EPCGChangeType::Edge | EPCGChangeType::Structural)) != EPCGChangeType::None);
@@ -816,7 +822,7 @@ void UPCGComponent::PostEditChangeProperty(FPropertyChangedEvent& PropertyChange
 {
 	Super::PostEditChangeProperty(PropertyChangedEvent);
 
-	if (!PropertyChangedEvent.Property)
+	if (!PropertyChangedEvent.Property || !IsValid(this))
 	{
 		return;
 	}
@@ -1393,6 +1399,8 @@ void UPCGComponent::Refresh()
 	}
 #endif
 
+	//check(IsValid(this));
+
 	// Following a change in some properties or in some spatial information related to this component,
 	// We need to regenerate the graph, depending of the state in the editor.
 	// In the case of a non-partitioned graph, we need to generate the graph only if it was previously generated & tagged for regeneration
@@ -2085,4 +2093,55 @@ bool UPCGComponent::GraphUsesLandscapePin() const
 
 #endif // WITH_EDITOR
 
+void UPCGComponent::SetManagedResources(const TArray<TObjectPtr<UPCGManagedResource>>& Resources)
+{
+	FScopeLock ResourcesLock(&GeneratedResourcesLock);
+	check(GeneratedResources.IsEmpty());
+	GeneratedResources = Resources;
+}
+
+void UPCGComponent::GetManagedResources(TArray<TObjectPtr<UPCGManagedResource>>& Resources) const
+{
+	FScopeLock ResourcesLock(&GeneratedResourcesLock);
+	Resources = GeneratedResources;
+}
+
+FPCGComponentInstanceData::FPCGComponentInstanceData(const UPCGComponent* SourceComponent)
+	: FActorComponentInstanceData(SourceComponent)
+{
+	if (SourceComponent)
+	{
+		SourceComponent->GetManagedResources(GeneratedResources);
+	}
+}
+
+bool FPCGComponentInstanceData::ContainsData() const
+{
+	return GeneratedResources.Num() > 0 || Super::ContainsData();
+}
+
+void FPCGComponentInstanceData::ApplyToComponent(UActorComponent* Component, const ECacheApplyPhase CacheApplyPhase)
+{
+	Super::ApplyToComponent(Component, CacheApplyPhase);
+
+	if (CacheApplyPhase == ECacheApplyPhase::PostUserConstructionScript)
+	{
+		UPCGComponent* PCGComponent = CastChecked<UPCGComponent>(Component);
+
+		// Duplicate generated resources + retarget them
+		TArray<TObjectPtr<UPCGManagedResource>> DuplicatedResources;
+		for (TObjectPtr<UPCGManagedResource> Resource : GeneratedResources)
+		{
+			UPCGManagedResource* DuplicatedResource = CastChecked<UPCGManagedResource>(StaticDuplicateObject(Resource, PCGComponent, FName()));
+			DuplicatedResource->PostEditImport();
+			DuplicatedResources.Add(DuplicatedResource);
+		}
+
+		if (DuplicatedResources.Num() > 0)
+		{
+			PCGComponent->SetManagedResources(DuplicatedResources);
+		}
+	}
+}
+
 #undef LOCTEXT_NAMESPACE
diff --git a/Engine/Plugins/Experimental/PCG/Source/PCG/Private/PCGElement.cpp b/Engine/Plugins/Experimental/PCG/Source/PCG/Private/PCGElement.cpp
--- a/Engine/Plugins/Experimental/PCG/Source/PCG/Private/PCGElement.cpp
+++ b/Engine/Plugins/Experimental/PCG/Source/PCG/Private/PCGElement.cpp
@@ -1,6 +1,7 @@
 // Copyright Epic Games, Inc. All Rights Reserved.
 
 #include "PCGElement.h"
+#include "PCGComponent.h"
 #include "PCGContext.h"
 #include "PCGModule.h"
 #include "PCGSettings.h"
@@ -12,7 +13,7 @@ bool IPCGElement::Execute(FPCGContext* Context) const
 {
 	TRACE_CPUPROFILER_EVENT_SCOPE(IPCGElement::Execute);
 	// Early out to stop execution
-	if (Context->InputData.bCancelExecution)
+	if (Context->InputData.bCancelExecution || (Context->SourceComponent && !IsValid(Context->SourceComponent)))
 	{
 		Context->OutputData.bCancelExecution = true;
 
diff --git a/Engine/Plugins/Experimental/PCG/Source/PCG/Public/PCGComponent.h b/Engine/Plugins/Experimental/PCG/Source/PCG/Public/PCGComponent.h
--- a/Engine/Plugins/Experimental/PCG/Source/PCG/Public/PCGComponent.h
+++ b/Engine/Plugins/Experimental/PCG/Source/PCG/Public/PCGComponent.h
@@ -72,8 +72,13 @@ public:
 	virtual void EndPlay(const EEndPlayReason::Type EndPlayReason) override;
 	virtual void OnComponentCreated() override;
 	virtual void OnComponentDestroyed(bool bDestroyingHierarchy) override;
+
+protected:
+	friend struct FPCGComponentInstanceData;
+	virtual TStructOnScope<FActorComponentInstanceData> GetComponentInstanceData() const override;
 	//~End UActorComponent Interface
 
+public:
 	UPCGData* GetPCGData();
 	UPCGData* GetInputPCGData();
 	UPCGData* GetActorPCGData();
@@ -224,6 +229,10 @@ private:
 	void CleanupUnusedManagedResources();
 	bool MoveResourcesToNewActor(AActor* InNewActor, bool bCreateChild);
 
+	/* Called as part of the RerunConstructionScript, just takes resources as-is, assumes current state is empty + the resources have been retargeted if needed */
+	void GetManagedResources(TArray<TObjectPtr<UPCGManagedResource>>& Resources) const;
+	void SetManagedResources(const TArray<TObjectPtr<UPCGManagedResource>>& Resources);
+
 	bool GetActorsFromTags(const TSet<FName>& InTags, TSet<TWeakObjectPtr<AActor>>& OutActors, bool bCullAgainstLocalBounds);
 
 	void RefreshAfterGraphChanged(UPCGGraph* InGraph, bool bIsStructural, bool bDirtyInputs);
@@ -336,5 +345,22 @@ private:
 	TMap<TObjectPtr<const UPCGNode>, FPCGDataCollection> InspectionCache;
 #endif
 
-	FCriticalSection GeneratedResourcesLock;
+	mutable FCriticalSection GeneratedResourcesLock;
+};
+
+/** Used to store generated resources data during RerunConstructionScripts */
+USTRUCT()
+struct FPCGComponentInstanceData : public FActorComponentInstanceData
+{
+	GENERATED_BODY()
+public:
+	FPCGComponentInstanceData() = default;
+	explicit FPCGComponentInstanceData(const UPCGComponent* SourceComponent);
+
+protected:
+	virtual bool ContainsData() const override;
+	virtual void ApplyToComponent(UActorComponent* Component, const ECacheApplyPhase CacheApplyPhase) override;
+
+	UPROPERTY()
+	TArray<TObjectPtr<UPCGManagedResource>> GeneratedResources;
 };
