Cleanup/Commonalize point processing elements

diff --git a/Engine/Plugins/Experimental/PCG/Source/PCG/Private/Elements/PCGDensityNoise.cpp b/Engine/Plugins/Experimental/PCG/Source/PCG/Private/Elements/PCGDensityNoise.cpp
--- a/Engine/Plugins/Experimental/PCG/Source/PCG/Private/Elements/PCGDensityNoise.cpp
+++ b/Engine/Plugins/Experimental/PCG/Source/PCG/Private/Elements/PCGDensityNoise.cpp
@@ -3,8 +3,6 @@
 #include "Elements/PCGDensityNoise.h"
 
 #include "PCGHelpers.h"
-#include "Data/PCGPointData.h"
-#include "Helpers/PCGAsync.h"
 #include "Helpers/PCGSettingsHelpers.h"
 
 #include "Math/RandomStream.h"
@@ -30,73 +28,43 @@ bool FPCGDensityNoiseElement::ExecuteInternal(FPCGContext* Context) const
 	const float DensityNoiseMax = PCG_GET_OVERRIDEN_VALUE(Settings, DensityNoiseMax, Params);
 	const bool bInvertSourceDensity = PCG_GET_OVERRIDEN_VALUE(Settings, bInvertSourceDensity, Params);
 
-	// TODO: embarassingly parallel loop
-	for (const FPCGTaggedData& Input : Inputs)
+	ProcessPoints(Context, Inputs, Outputs, [&](const FPCGPoint& InPoint, FPCGPoint& OutPoint)
 	{
-		TRACE_CPUPROFILER_EVENT_SCOPE(FPCGDensityNoiseElement::Execute::InputLoop);
-		FPCGTaggedData& Output = Outputs.Add_GetRef(Input);
+		OutPoint = InPoint;
+		FRandomStream RandomSource(PCGHelpers::ComputeSeed(Settings->Seed, OutPoint.Seed));
 
-		const UPCGSpatialData* SpatialData = Cast<UPCGSpatialData>(Input.Data);
+		const float DensityNoise = RandomSource.FRandRange(DensityNoiseMin, DensityNoiseMax);
 
-		if (!SpatialData)
-		{
-			PCGE_LOG(Error, "Unable to get SpatialData from input");
-			continue;
-		}
+		// This inversion was previously calculated as OutPoint.Density *= 1.f - FMath::Abs(OutPoint.Density * 2.f - 1.f);
+		const float SourceDensity = bInvertSourceDensity ? (1.0f - OutPoint.Density) : OutPoint.Density;
 
-		const UPCGPointData* PointData = SpatialData->ToPointData(Context);
+		float UnclampedDensity = 0;
 
-		if (!PointData)
+		if (DensityMode == EPCGDensityNoiseMode::Minimum)
 		{
-			PCGE_LOG(Error, "Unable to get PointData from input");
-			continue;
+			UnclampedDensity = FMath::Min(SourceDensity, DensityNoise);
+		}
+		else if (DensityMode == EPCGDensityNoiseMode::Maximum)
+		{
+			UnclampedDensity = FMath::Max(SourceDensity, DensityNoise);
+		}
+		else if (DensityMode == EPCGDensityNoiseMode::Add)
+		{
+			UnclampedDensity = SourceDensity + DensityNoise;
+		}
+		else if (DensityMode == EPCGDensityNoiseMode::Multiply)
+		{
+			UnclampedDensity = SourceDensity * DensityNoise;
+		}
+		else //if (DensityMode == EPCGDensityNoiseMode::Set)
+		{
+			UnclampedDensity = DensityNoise;
 		}
 
-		const TArray<FPCGPoint>& Points = PointData->GetPoints();
-
-		UPCGPointData* SampledData = NewObject<UPCGPointData>();
-		SampledData->InitializeFromData(PointData);
-		TArray<FPCGPoint>& SampledPoints = SampledData->GetMutablePoints();
-		Output.Data = SampledData;
+		OutPoint.Density = FMath::Clamp(UnclampedDensity, 0.f, 1.f);
 
-		FPCGAsync::AsyncPointProcessing(Context, Points.Num(), SampledPoints, [&Points, Settings, DensityMode, DensityNoiseMin, DensityNoiseMax, bInvertSourceDensity](int32 Index, FPCGPoint& OutPoint)
-		{
-			OutPoint = Points[Index];
-			FRandomStream RandomSource(PCGHelpers::ComputeSeed(Settings->Seed, OutPoint.Seed));
-
-			const float DensityNoise = RandomSource.FRandRange(DensityNoiseMin, DensityNoiseMax);
-
-			// This inversion was previously calculated as OutPoint.Density *= 1.f - FMath::Abs(OutPoint.Density * 2.f - 1.f);
-			const float SourceDensity = bInvertSourceDensity ? (1.0f - OutPoint.Density) : OutPoint.Density;
-
-			float UnclampedDensity = 0;
-
-			if (DensityMode == EPCGDensityNoiseMode::Minimum)
-			{
-				UnclampedDensity = FMath::Min(SourceDensity, DensityNoise);
-			}
-			else if (DensityMode == EPCGDensityNoiseMode::Maximum)
-			{
-				UnclampedDensity = FMath::Max(SourceDensity, DensityNoise);
-			}
-			else if (DensityMode == EPCGDensityNoiseMode::Add)
-			{
-				UnclampedDensity = SourceDensity + DensityNoise;
-			}
-			else if (DensityMode == EPCGDensityNoiseMode::Multiply)
-			{
-				UnclampedDensity = SourceDensity * DensityNoise;
-			}
-			else //if (DensityMode == EPCGDensityNoiseMode::Set)
-			{
-				UnclampedDensity = DensityNoise;
-			}
-
-			OutPoint.Density = FMath::Clamp(UnclampedDensity, 0.f, 1.f);
-
-			return true;
-		});
-	}
+		return true;
+	});
 
 	// Forward any non-input data
 	Outputs.Append(Context->InputData.GetAllSettings());
diff --git a/Engine/Plugins/Experimental/PCG/Source/PCG/Private/Elements/PCGDensityRemapElement.cpp b/Engine/Plugins/Experimental/PCG/Source/PCG/Private/Elements/PCGDensityRemapElement.cpp
--- a/Engine/Plugins/Experimental/PCG/Source/PCG/Private/Elements/PCGDensityRemapElement.cpp
+++ b/Engine/Plugins/Experimental/PCG/Source/PCG/Private/Elements/PCGDensityRemapElement.cpp
@@ -3,8 +3,6 @@
 #include "Elements/PCGDensityRemapElement.h"
 
 #include "PCGHelpers.h"
-#include "Data/PCGPointData.h"
-#include "Helpers/PCGAsync.h"
 #include "Helpers/PCGSettingsHelpers.h"
 
 #include "Math/RandomStream.h"
@@ -53,49 +51,19 @@ bool FPCGDensityRemapElement::ExecuteInternal(FPCGContext* Context) const
 		Intercept = OutRangeMin;
 	}
 
-	// TODO: embarassingly parallel loop
-	for (const FPCGTaggedData& Input : Inputs)
+	ProcessPoints(Context, Inputs, Outputs, [&](const FPCGPoint& InPoint, FPCGPoint& OutPoint)
 	{
-		TRACE_CPUPROFILER_EVENT_SCOPE(FPCGDensityRemapElement::Execute::InputLoop);
-		FPCGTaggedData& Output = Outputs.Add_GetRef(Input);
+		OutPoint = InPoint;
+		const float SourceDensity = InPoint.Density;
 
-		const UPCGSpatialData* SpatialData = Cast<UPCGSpatialData>(Input.Data);
-
-		if (!SpatialData)
-		{
-			PCGE_LOG(Error, "Unable to get SpatialData from input");
-			continue;
-		}
-
-		const UPCGPointData* PointData = SpatialData->ToPointData(Context);
-
-		if (!PointData)
+		if (!bExcludeValuesOutsideInputRange || (SourceDensity >= InRangeTrueMin && SourceDensity <= InRangeTrueMax))
 		{
-			PCGE_LOG(Error, "Unable to get PointData from input");
-			continue;
+			const float UnclampedDensity = Slope * (SourceDensity - InRangeMin) + Intercept;
+			OutPoint.Density = FMath::Clamp(UnclampedDensity, 0.f, 1.f);
 		}
 
-		const TArray<FPCGPoint>& Points = PointData->GetPoints();
-
-		UPCGPointData* SampledData = NewObject<UPCGPointData>();
-		SampledData->InitializeFromData(PointData);
-		TArray<FPCGPoint>& SampledPoints = SampledData->GetMutablePoints();
-		Output.Data = SampledData;
-
-		FPCGAsync::AsyncPointProcessing(Context, Points.Num(), SampledPoints, [&](int32 Index, FPCGPoint& OutPoint)
-		{
-			OutPoint = Points[Index];
-			const float SourceDensity = Points[Index].Density;
-
-			if (!bExcludeValuesOutsideInputRange || (SourceDensity >= InRangeTrueMin && SourceDensity <= InRangeTrueMax))
-			{
-				const float UnclampedDensity = Slope * (SourceDensity - InRangeMin) + Intercept;
-				OutPoint.Density = FMath::Clamp(UnclampedDensity, 0.f, 1.f);
-			}
-
-			return true;
-		});
-	}
+		return true;
+	});
 
 	// Forward any non-input data
 	Outputs.Append(Context->InputData.GetAllSettings());
diff --git a/Engine/Plugins/Experimental/PCG/Source/PCG/Private/Elements/PCGPointProcessingElementBase.cpp b/Engine/Plugins/Experimental/PCG/Source/PCG/Private/Elements/PCGPointProcessingElementBase.cpp
new file mode 100644
--- /dev/null
+++ b/Engine/Plugins/Experimental/PCG/Source/PCG/Private/Elements/PCGPointProcessingElementBase.cpp
@@ -0,0 +1,40 @@
+// Copyright Epic Games, Inc. All Rights Reserved.
+
+#include "Elements/PCGPointProcessingElementBase.h"
+
+#include "Data/PCGPointData.h"
+
+void FPCGPointProcessingElementBase::ProcessPoints(FPCGContext* Context, const TArray<FPCGTaggedData>& Inputs, TArray<FPCGTaggedData>& Outputs, const TFunction<bool(const FPCGPoint&, FPCGPoint&)>& PointFunc) const
+{
+	// TODO: embarassingly parallel loop
+	for (const FPCGTaggedData& Input : Inputs)
+	{
+		TRACE_CPUPROFILER_EVENT_SCOPE(FPCGPointProcessingElement::Execute::InputLoop);
+		FPCGTaggedData& Output = Outputs.Add_GetRef(Input);
+
+		const UPCGSpatialData* SpatialData = Cast<UPCGSpatialData>(Input.Data);
+
+		if (!SpatialData)
+		{
+			PCGE_LOG(Error, "Unable to get SpatialData from input");
+			continue;
+		}
+
+		const UPCGPointData* PointData = SpatialData->ToPointData(Context);
+
+		if (!PointData)
+		{
+			PCGE_LOG(Error, "Unable to get PointData from input");
+			continue;
+		}
+
+		const TArray<FPCGPoint>& Points = PointData->GetPoints();
+
+		UPCGPointData* OutputData = NewObject<UPCGPointData>();
+		OutputData->InitializeFromData(PointData);
+		TArray<FPCGPoint>& OutputPoints = OutputData->GetMutablePoints();
+		Output.Data = OutputData;
+
+		FPCGAsync::AsyncPointProcessing(Context, Points, OutputPoints, PointFunc);
+	}
+}
diff --git a/Engine/Plugins/Experimental/PCG/Source/PCG/Private/Elements/PCGTransformPoints.cpp b/Engine/Plugins/Experimental/PCG/Source/PCG/Private/Elements/PCGTransformPoints.cpp
--- a/Engine/Plugins/Experimental/PCG/Source/PCG/Private/Elements/PCGTransformPoints.cpp
+++ b/Engine/Plugins/Experimental/PCG/Source/PCG/Private/Elements/PCGTransformPoints.cpp
@@ -3,9 +3,6 @@
 #include "Elements/PCGTransformPoints.h"
 
 #include "PCGHelpers.h"
-#include "Data/PCGPointData.h"
-#include "Data/PCGSpatialData.h"
-#include "Helpers/PCGAsync.h"
 #include "Helpers/PCGSettingsHelpers.h"
 
 TArray<FPCGPinProperties> UPCGTransformPointsSettings::InputPinProperties() const
@@ -45,108 +42,70 @@ bool FPCGTransformPointsElement::ExecuteInternal(FPCGContext* Context) const
 	const bool bUniformScale = PCG_GET_OVERRIDEN_VALUE(Settings, bUniformScale, Params);
 	const bool bRecomputeSeed = PCG_GET_OVERRIDEN_VALUE(Settings, bRecomputeSeed, Params);
 
-	for (const FPCGTaggedData& Input : Inputs)
+	ProcessPoints(Context, Inputs, Outputs, [&](const FPCGPoint& InPoint, FPCGPoint& OutPoint)
 	{
-		TRACE_CPUPROFILER_EVENT_SCOPE(FPCGTransformPointsElement::Execute::InputLoop);
+		OutPoint = InPoint;
 
-		FPCGTaggedData& Output = Outputs.Add_GetRef(Input);
+		FRandomStream RandomSource(PCGHelpers::ComputeSeed(Settings->Seed, InPoint.Seed));
 
-		if (!Input.Data || Cast<UPCGSpatialData>(Input.Data) == nullptr)
-		{
-			PCGE_LOG(Error, "Invalid input data");
-			continue;
-		}
+		const float OffsetX = RandomSource.FRandRange(OffsetMin.X, OffsetMax.X);
+		const float OffsetY = RandomSource.FRandRange(OffsetMin.Y, OffsetMax.Y);
+		const float OffsetZ = RandomSource.FRandRange(OffsetMin.Z, OffsetMax.Z);
+		const FVector RandomOffset(OffsetX, OffsetY, OffsetZ);
 
-		const UPCGSpatialData* SpatialData = Cast<UPCGSpatialData>(Input.Data);
+		const float RotationX = RandomSource.FRandRange(RotationMin.Pitch, RotationMax.Pitch);
+		const float RotationY = RandomSource.FRandRange(RotationMin.Yaw, RotationMax.Yaw);
+		const float RotationZ = RandomSource.FRandRange(RotationMin.Roll, RotationMax.Roll);
+		const FQuat RandomRotation(FRotator(RotationX, RotationY, RotationZ).Quaternion());
 
-		if (!SpatialData)
+		FVector RandomScale;
+		if (bUniformScale)
 		{
-			PCGE_LOG(Error, "Unable to get SpatialData from input");
-			continue;
+			RandomScale = FVector(RandomSource.FRandRange(ScaleMin.X, ScaleMax.X));
+		}
+		else
+		{
+			RandomScale.X = RandomSource.FRandRange(ScaleMin.X, ScaleMax.X);
+			RandomScale.Y = RandomSource.FRandRange(ScaleMin.Y, ScaleMax.Y);
+			RandomScale.Z = RandomSource.FRandRange(ScaleMin.Z, ScaleMax.Z);
 		}
 
-		const UPCGPointData* PointData = SpatialData->ToPointData(Context);
-
-		if (!PointData)
+		if (bAbsoluteOffset)
 		{
-			PCGE_LOG(Error, "Unable to get PointData from input");
-			continue;
+			OutPoint.Transform.SetLocation(InPoint.Transform.GetLocation() + RandomOffset); 
+		}
+		else
+		{
+			const FTransform RotatedTransform(InPoint.Transform.GetRotation());
+			OutPoint.Transform.SetLocation(InPoint.Transform.GetLocation() + RotatedTransform.TransformPosition(RandomOffset)); 
 		}
 
-		const TArray<FPCGPoint>& Points = PointData->GetPoints();
+		if (bAbsoluteRotation)
+		{
+			OutPoint.Transform.SetRotation(RandomRotation);
+		}
+		else
+		{
+			OutPoint.Transform.SetRotation(InPoint.Transform.GetRotation() * RandomRotation);
+		}
 
-		UPCGPointData* OutPointData = NewObject<UPCGPointData>();
-		OutPointData->InitializeFromData(PointData);
-		TArray<FPCGPoint>& OutPoints = OutPointData->GetMutablePoints();
-		Output.Data = OutPointData;
+		if (bAbsoluteScale)
+		{
+			OutPoint.Transform.SetScale3D(RandomScale);
+		}
+		else
+		{
+			OutPoint.Transform.SetScale3D(InPoint.Transform.GetScale3D() * RandomScale);
+		}
 
-		// loop on points
-		FPCGAsync::AsyncPointProcessing(Context, Points.Num(), OutPoints, [&](int32 Index, FPCGPoint& OutPoint)
+		if (bRecomputeSeed)
 		{
-			const FPCGPoint& InPoint = Points[Index];
-			OutPoint = InPoint;
-
-			FRandomStream RandomSource(PCGHelpers::ComputeSeed(Settings->Seed, InPoint.Seed));
-
-			const float OffsetX = RandomSource.FRandRange(OffsetMin.X, OffsetMax.X);
-			const float OffsetY = RandomSource.FRandRange(OffsetMin.Y, OffsetMax.Y);
-			const float OffsetZ = RandomSource.FRandRange(OffsetMin.Z, OffsetMax.Z);
-			const FVector RandomOffset(OffsetX, OffsetY, OffsetZ);
-
-			const float RotationX = RandomSource.FRandRange(RotationMin.Pitch, RotationMax.Pitch);
-			const float RotationY = RandomSource.FRandRange(RotationMin.Yaw, RotationMax.Yaw);
-			const float RotationZ = RandomSource.FRandRange(RotationMin.Roll, RotationMax.Roll);
-			const FQuat RandomRotation(FRotator(RotationX, RotationY, RotationZ).Quaternion());
-
-			FVector RandomScale;
-			if (bUniformScale)
-			{
-				RandomScale = FVector(RandomSource.FRandRange(ScaleMin.X, ScaleMax.X));
-			}
-			else
-			{
-				RandomScale.X = RandomSource.FRandRange(ScaleMin.X, ScaleMax.X);
-				RandomScale.Y = RandomSource.FRandRange(ScaleMin.Y, ScaleMax.Y);
-				RandomScale.Z = RandomSource.FRandRange(ScaleMin.Z, ScaleMax.Z);
-			}
-
-			if (bAbsoluteOffset)
-			{
-				OutPoint.Transform.SetLocation(InPoint.Transform.GetLocation() + RandomOffset); 
-			}
-			else
-			{
-				const FTransform RotatedTransform(InPoint.Transform.GetRotation());
-				OutPoint.Transform.SetLocation(InPoint.Transform.GetLocation() + RotatedTransform.TransformPosition(RandomOffset)); 
-			}
-
-			if (bAbsoluteRotation)
-			{
-				OutPoint.Transform.SetRotation(RandomRotation);
-			}
-			else
-			{
-				OutPoint.Transform.SetRotation(InPoint.Transform.GetRotation() * RandomRotation);
-			}
-
-			if (bAbsoluteScale)
-			{
-				OutPoint.Transform.SetScale3D(RandomScale);
-			}
-			else
-			{
-				OutPoint.Transform.SetScale3D(InPoint.Transform.GetScale3D() * RandomScale);
-			}
-
-			if (bRecomputeSeed)
-			{
-				const FVector& Position = OutPoint.Transform.GetLocation();
-				OutPoint.Seed = PCGHelpers::ComputeSeed((int)Position.X, (int)Position.Y, (int)Position.Z);
-			}
-
-			return true;
-		});
-	}
+			const FVector& Position = OutPoint.Transform.GetLocation();
+			OutPoint.Seed = PCGHelpers::ComputeSeed((int)Position.X, (int)Position.Y, (int)Position.Z);
+		}
+
+		return true;
+	});
 
 	// Forward any non-input data
 	Outputs.Append(Context->InputData.GetAllSettings());
diff --git a/Engine/Plugins/Experimental/PCG/Source/PCG/Private/Helpers/PCGAsync.cpp b/Engine/Plugins/Experimental/PCG/Source/PCG/Private/Helpers/PCGAsync.cpp
--- a/Engine/Plugins/Experimental/PCG/Source/PCG/Private/Helpers/PCGAsync.cpp
+++ b/Engine/Plugins/Experimental/PCG/Source/PCG/Private/Helpers/PCGAsync.cpp
@@ -8,15 +8,10 @@
 
 namespace FPCGAsync
 {
-	void AsyncPointProcessing(FPCGContext* Context, int32 NumIterations, TArray<FPCGPoint>& OutPoints, const TFunction<bool(int32, FPCGPoint&)>& PointFunc)
-	{
-		const int32 MinIterationsPerTask = 256;
-		AsyncPointProcessing(Context ? Context->NumAvailableTasks : 1, MinIterationsPerTask, NumIterations, OutPoints, PointFunc);
-	}
-
-	void AsyncPointProcessing(int32 NumAvailableTasks, int32 MinIterationsPerTask, int32 NumIterations, TArray<FPCGPoint>& OutPoints, const TFunction<bool(int32, FPCGPoint&)>& PointFunc)
+	void AsyncPointProcessing(int32 NumAvailableTasks, int32 MinIterationsPerTask, int32 NumIterations, TArray<FPCGPoint>& OutPoints, const TFunction<int32(int32, int32)>& IterationInnerLoop)
 	{
 		TRACE_CPUPROFILER_EVENT_SCOPE(IPCGElement::AsyncPointProcessing);
+
 		check(NumAvailableTasks > 0 && MinIterationsPerTask > 0 && NumIterations >= 0);
 		if (NumAvailableTasks <= 0 || MinIterationsPerTask <= 0 || NumIterations <= 0)
 		{
@@ -34,22 +29,6 @@ namespace FPCGAsync
 			OutPoints.SetNum(NumIterations);
 		}
 
-		auto IterationInnerLoop = [&PointFunc, &OutPoints](int32 StartIndex, int32 EndIndex) -> int32
-		{
-			TRACE_CPUPROFILER_EVENT_SCOPE(IPCGElement::AsyncPointProcessing::InnerLoop);
-			int32 NumPointsWritten = 0;
-
-			for (int32 Index = StartIndex; Index < EndIndex; ++Index)
-			{
-				if (PointFunc(Index, OutPoints[StartIndex + NumPointsWritten]))
-				{
-					++NumPointsWritten;
-				}
-			}
-
-			return NumPointsWritten;
-		};
-
 		// Setup [current, last, nb points] data per dispatch
 		TArray<TFuture<int32>> AsyncTasks;
 		AsyncTasks.Reserve(NumFutures);
@@ -105,6 +84,56 @@ namespace FPCGAsync
 		}
 	}
 
+	void AsyncPointProcessing(FPCGContext* Context, int32 NumIterations, TArray<FPCGPoint>& OutPoints, const TFunction<bool(int32, FPCGPoint&)>& PointFunc)
+	{
+		const int32 MinIterationsPerTask = 256;
+		AsyncPointProcessing(Context ? Context->NumAvailableTasks : 1, MinIterationsPerTask, NumIterations, OutPoints, PointFunc);
+	}
+
+	void AsyncPointProcessing(FPCGContext* Context, const TArray<FPCGPoint>& InPoints, TArray<FPCGPoint>& OutPoints, const TFunction<bool(const FPCGPoint&, FPCGPoint&)>& PointFunc)
+	{
+		const int32 MinIterationsPerTask = 256;
+
+		auto IterationInnerLoop = [&PointFunc, &InPoints, &OutPoints](int32 StartIndex, int32 EndIndex) -> int32
+		{
+			TRACE_CPUPROFILER_EVENT_SCOPE(IPCGElement::AsyncPointProcessing::InnerLoop);
+			int32 NumPointsWritten = 0;
+
+			for (int32 Index = StartIndex; Index < EndIndex; ++Index)
+			{
+				if (PointFunc(InPoints[Index], OutPoints[StartIndex + NumPointsWritten]))
+				{
+					++NumPointsWritten;
+				}
+			}
+
+			return NumPointsWritten;
+		};
+
+		AsyncPointProcessing(Context ? Context->NumAvailableTasks : 1, MinIterationsPerTask, InPoints.Num(), OutPoints, IterationInnerLoop);
+	}
+
+	void AsyncPointProcessing(int32 NumAvailableTasks, int32 MinIterationsPerTask, int32 NumIterations, TArray<FPCGPoint>& OutPoints, const TFunction<bool(int32, FPCGPoint&)>& PointFunc)
+	{
+		auto IterationInnerLoop = [&PointFunc, &OutPoints](int32 StartIndex, int32 EndIndex) -> int32
+		{
+			TRACE_CPUPROFILER_EVENT_SCOPE(IPCGElement::AsyncPointProcessing::InnerLoop);
+			int32 NumPointsWritten = 0;
+
+			for (int32 Index = StartIndex; Index < EndIndex; ++Index)
+			{
+				if (PointFunc(Index, OutPoints[StartIndex + NumPointsWritten]))
+				{
+					++NumPointsWritten;
+				}
+			}
+
+			return NumPointsWritten;
+		};
+
+		AsyncPointProcessing(NumAvailableTasks, MinIterationsPerTask, NumIterations, OutPoints, IterationInnerLoop);
+	}
+
 	void AsyncPointFilterProcessing(FPCGContext* Context, int32 NumIterations, TArray<FPCGPoint>& InFilterPoints, TArray<FPCGPoint>& OutFilterPoints, const TFunction<bool(int32, FPCGPoint&, FPCGPoint&)>& PointFunc)
 	{
 		const int32 MinIterationsPerTask = 256;
diff --git a/Engine/Plugins/Experimental/PCG/Source/PCG/Public/Elements/PCGDensityNoise.h b/Engine/Plugins/Experimental/PCG/Source/PCG/Public/Elements/PCGDensityNoise.h
--- a/Engine/Plugins/Experimental/PCG/Source/PCG/Public/Elements/PCGDensityNoise.h
+++ b/Engine/Plugins/Experimental/PCG/Source/PCG/Public/Elements/PCGDensityNoise.h
@@ -7,6 +7,8 @@
 #include "PCGNode.h"
 #include "PCGSettings.h"
 
+#include "Elements/PCGPointProcessingElementBase.h"
+
 #include "PCGDensityNoise.generated.h"
 
 UENUM()
@@ -53,7 +55,7 @@ public:
 	bool bInvertSourceDensity = false;
 };
 
-class FPCGDensityNoiseElement : public FSimplePCGElement
+class FPCGDensityNoiseElement : public FPCGPointProcessingElementBase
 {
 protected:
 	virtual bool ExecuteInternal(FPCGContext* Context) const;
diff --git a/Engine/Plugins/Experimental/PCG/Source/PCG/Public/Elements/PCGDensityRemapElement.h b/Engine/Plugins/Experimental/PCG/Source/PCG/Public/Elements/PCGDensityRemapElement.h
--- a/Engine/Plugins/Experimental/PCG/Source/PCG/Public/Elements/PCGDensityRemapElement.h
+++ b/Engine/Plugins/Experimental/PCG/Source/PCG/Public/Elements/PCGDensityRemapElement.h
@@ -7,6 +7,8 @@
 #include "PCGNode.h"
 #include "PCGSettings.h"
 
+#include "Elements/PCGPointProcessingElementBase.h"
+
 #include "PCGDensityRemapElement.generated.h"
 
 UCLASS(BlueprintType, ClassGroup = (Procedural))
@@ -47,7 +49,7 @@ public:
 	bool bExcludeValuesOutsideInputRange = true;
 };
 
-class FPCGDensityRemapElement : public FSimplePCGElement
+class FPCGDensityRemapElement : public FPCGPointProcessingElementBase
 {
 protected:
 	virtual bool ExecuteInternal(FPCGContext* Context) const;
diff --git a/Engine/Plugins/Experimental/PCG/Source/PCG/Public/Elements/PCGPointProcessingElementBase.h b/Engine/Plugins/Experimental/PCG/Source/PCG/Public/Elements/PCGPointProcessingElementBase.h
new file mode 100644
--- /dev/null
+++ b/Engine/Plugins/Experimental/PCG/Source/PCG/Public/Elements/PCGPointProcessingElementBase.h
@@ -0,0 +1,15 @@
+// Copyright Epic Games, Inc. All Rights Reserved.
+
+#pragma once
+
+#include "CoreMinimal.h"
+#include "PCGElement.h"
+
+#include "Helpers/PCGAsync.h"
+
+class FPCGPointProcessingElementBase : public FSimplePCGElement
+{
+protected:
+	virtual bool ExecuteInternal(FPCGContext* Context) const = 0;
+	void ProcessPoints(FPCGContext* Context, const TArray<FPCGTaggedData>& Inputs, TArray<FPCGTaggedData>& Outputs, const TFunction<bool(const FPCGPoint&, FPCGPoint&)>& PointFunc) const;
+};
diff --git a/Engine/Plugins/Experimental/PCG/Source/PCG/Public/Elements/PCGTransformPoints.h b/Engine/Plugins/Experimental/PCG/Source/PCG/Public/Elements/PCGTransformPoints.h
--- a/Engine/Plugins/Experimental/PCG/Source/PCG/Public/Elements/PCGTransformPoints.h
+++ b/Engine/Plugins/Experimental/PCG/Source/PCG/Public/Elements/PCGTransformPoints.h
@@ -5,6 +5,8 @@
 #include "CoreMinimal.h"
 #include "PCGSettings.h"
 
+#include "Elements/PCGPointProcessingElementBase.h"
+
 #include "PCGTransformPoints.generated.h"
 
 UCLASS(BlueprintType, ClassGroup = (Procedural))
@@ -66,7 +68,7 @@ public:
 	bool bRecomputeSeed = false;
 };
 
-class FPCGTransformPointsElement : public FSimplePCGElement
+class FPCGTransformPointsElement : public FPCGPointProcessingElementBase
 {
 protected:
 	virtual bool ExecuteInternal(FPCGContext* Context) const override;
diff --git a/Engine/Plugins/Experimental/PCG/Source/PCG/Public/Helpers/PCGAsync.h b/Engine/Plugins/Experimental/PCG/Source/PCG/Public/Helpers/PCGAsync.h
--- a/Engine/Plugins/Experimental/PCG/Source/PCG/Public/Helpers/PCGAsync.h
+++ b/Engine/Plugins/Experimental/PCG/Source/PCG/Public/Helpers/PCGAsync.h
@@ -17,13 +17,22 @@ namespace FPCGAsync
 	* @param PointFunc - A function that has the index [0; NumIterations] and has to write to the point & return true when the current call generates a point
 	*/
 	PCG_API void AsyncPointProcessing(FPCGContext* Context, int32 NumIterations, TArray<FPCGPoint>& OutPoints, const TFunction<bool(int32, FPCGPoint&)>& PointFunc);
+	
+	/** 
+	* Helper to do simple point processing loops
+	* @param Context - The context containing the information about how many tasks to launch
+	* @param InPoints - The array in which the source points will be read from
+	* @param OutPoints - The array in which the results will be written to. Note that the array will be cleared before execution
+	* @param PointFunc - A function that has the input point and has to write to the output point & return true when the current call generates a point
+	*/
+	PCG_API void AsyncPointProcessing(FPCGContext* Context, const TArray<FPCGPoint>& InPoints, TArray<FPCGPoint>& OutPoints, const TFunction<bool(const FPCGPoint&, FPCGPoint&)>& PointFunc);
 
 	/** 
 	* Helper to do more general 1:1 point processing loops
 	* @param NumAvailableTasks - The upper bound on the number of async tasks we'll start
 	* @param MinIterationsPerTask - The lower bound on the number of iterations per task we'll dispatch
 	* @param NumIterations - The number of calls that will be done to the provided function, also an upper bound on the number of points generated
- 	* @param OutPoints - The array in which the results will be written to. Note that the array will be cleared before execution
+	* @param OutPoints - The array in which the results will be written to. Note that the array will be cleared before execution
 	* @param PointFunc - A function that has the index [0; NumIterations] and has to write to the point & return true when the current call generates a point
 	*/
 	void AsyncPointProcessing(int32 NumAvailableTasks, int32 MinIterationsPerTask, int32 NumIterations, TArray<FPCGPoint>& OutPoints, const TFunction<bool(int32, FPCGPoint&)>& PointFunc);
