Fixed some terminology issues

#rnx

diff --git a/Engine/Plugins/Experimental/PCG/Source/PCG/Private/Graph/PCGGraphExecutor.cpp b/Engine/Plugins/Experimental/PCG/Source/PCG/Private/Graph/PCGGraphExecutor.cpp
--- a/Engine/Plugins/Experimental/PCG/Source/PCG/Private/Graph/PCGGraphExecutor.cpp
+++ b/Engine/Plugins/Experimental/PCG/Source/PCG/Private/Graph/PCGGraphExecutor.cpp
@@ -159,7 +159,7 @@ void FPCGGraphExecutor::Execute()
 		FPCGGraphScheduleTask& ScheduledTask = ScheduledTasks[ScheduledTaskIndex];
 
 		check(ScheduledTask.Tasks.Num() > 0);
-		// Push tasks to the master list & build successor list
+		// Push tasks to the primary task list & build successors map
 		for (FPCGGraphTask& Task : ScheduledTask.Tasks)
 		{
 			FPCGTaskId TaskId = Task.NodeId;
