+ Adds button to the PCG Graph UI to force refresh on the graph

diff --git a/Engine/Plugins/Experimental/PCG/Source/PCG/Private/PCGGraph.cpp b/Engine/Plugins/Experimental/PCG/Source/PCG/Private/PCGGraph.cpp
--- a/Engine/Plugins/Experimental/PCG/Source/PCG/Private/PCGGraph.cpp
+++ b/Engine/Plugins/Experimental/PCG/Source/PCG/Private/PCGGraph.cpp
@@ -445,6 +445,17 @@ bool UPCGGraph::RemoveOutboundEdges(UPCGNode* InNode, const FName& OutboundLabel
 }
 
 #if WITH_EDITOR
+void UPCGGraph::ForceNotificationForEditor()
+{
+	// Queue up the delayed change
+	NotifyGraphChanged(EPCGChangeType::Structural);
+	if (bUserPausedNotificationsInGraphEditor)
+	{
+		EnableNotificationsForEditor();
+		DisableNotificationsForEditor();
+	}
+}
+
 void UPCGGraph::PreNodeUndo(UPCGNode* InPCGNode)
 {
 	if (InPCGNode)
@@ -535,7 +546,7 @@ void UPCGGraph::GetTrackedTagsToSettings(FPCGTagToSettingsMap& OutTagsToSettings
 
 void UPCGGraph::NotifyGraphChanged(EPCGChangeType ChangeType)
 {
-	if(GraphChangeNotificationsDisableCounter > 0)
+	if (GraphChangeNotificationsDisableCounter > 0)
 	{
 		bDelayedChangeNotification = true;
 		DelayedChangeType |= ChangeType;
diff --git a/Engine/Plugins/Experimental/PCG/Source/PCG/Public/PCGGraph.h b/Engine/Plugins/Experimental/PCG/Source/PCG/Public/PCGGraph.h
--- a/Engine/Plugins/Experimental/PCG/Source/PCG/Public/PCGGraph.h
+++ b/Engine/Plugins/Experimental/PCG/Source/PCG/Public/PCGGraph.h
@@ -88,6 +88,7 @@ public:
 	void EnableNotificationsForEditor();
 	void ToggleUserPausedNotificationsForEditor();
 	bool NotificationsForEditorArePausedByUser() const { return bUserPausedNotificationsInGraphEditor; }
+	void ForceNotificationForEditor();
 	void PreNodeUndo(UPCGNode* InPCGNode);
 	void PostNodeUndo(UPCGNode* InPCGNode);
 
diff --git a/Engine/Plugins/Experimental/PCG/Source/PCGEditor/Private/PCGEditor.cpp b/Engine/Plugins/Experimental/PCG/Source/PCGEditor/Private/PCGEditor.cpp
--- a/Engine/Plugins/Experimental/PCG/Source/PCGEditor/Private/PCGEditor.cpp
+++ b/Engine/Plugins/Experimental/PCG/Source/PCGEditor/Private/PCGEditor.cpp
@@ -300,6 +300,12 @@ void FPCGEditor::RegisterToolbar() const
 			TAttribute<FText>(),
 			FSlateIcon(FAppStyle::GetAppStyleSetName(), "PlayWorld.PausePlaySession")));
 
+		Section.AddEntry(FToolMenuEntry::InitToolBarButton(
+			FPCGEditorCommands::Get().ForceGraphRegeneration,
+			TAttribute<FText>(),
+			TAttribute<FText>(),
+			FSlateIcon(FAppStyle::GetAppStyleSetName(), "SourceControl.Actions.Refresh")));
+
 		Section.AddSeparator(NAME_None);
 		Section.AddDynamicEntry("Debugging", FNewToolMenuSectionDelegate::CreateLambda([](FToolMenuSection& InSection)
 		{
@@ -326,6 +332,10 @@ void FPCGEditor::BindCommands()
 		FCanExecuteAction(),
 		FIsActionChecked::CreateSP(this, &FPCGEditor::IsAutomaticRegenerationPaused));
 
+	ToolkitCommands->MapAction(
+		PCGEditorCommands.ForceGraphRegeneration,
+		FExecuteAction::CreateSP(this, &FPCGEditor::OnForceGraphRegeneration_Clicked));
+
 	GraphEditorCommands->MapAction(
 		PCGEditorCommands.CollapseNodes,
 		FExecuteAction::CreateSP(this, &FPCGEditor::OnCollapseNodesInSubgraph),
@@ -393,6 +403,14 @@ bool FPCGEditor::IsAutomaticRegenerationPaused() const
 	return PCGGraphBeingEdited && PCGGraphBeingEdited->NotificationsForEditorArePausedByUser();
 }
 
+void FPCGEditor::OnForceGraphRegeneration_Clicked()
+{
+	if (PCGGraphBeingEdited)
+	{
+		PCGGraphBeingEdited->ForceNotificationForEditor();
+	}
+}
+
 bool FPCGEditor::CanRunDeterminismTests() const
 {
 	check(GraphEditorWidget.IsValid());
diff --git a/Engine/Plugins/Experimental/PCG/Source/PCGEditor/Private/PCGEditor.h b/Engine/Plugins/Experimental/PCG/Source/PCGEditor/Private/PCGEditor.h
--- a/Engine/Plugins/Experimental/PCG/Source/PCGEditor/Private/PCGEditor.h
+++ b/Engine/Plugins/Experimental/PCG/Source/PCGEditor/Private/PCGEditor.h
@@ -90,8 +90,12 @@ private:
 
 	/** Enable/Disable automatic PCG node generation */
 	void OnPauseAutomaticRegeneration_Clicked();
+	/** Has the user paused automatic regeneration in the Graph Editor */
 	bool IsAutomaticRegenerationPaused() const;
 
+	/** Force a regeneration by invoking the graph notifications  */
+	void OnForceGraphRegeneration_Clicked();
+
 	/** Start inspecting the current selected node */
 	void OnStartInspectNode();
 	/** Stop inspecting the current inspected node */
diff --git a/Engine/Plugins/Experimental/PCG/Source/PCGEditor/Private/PCGEditorCommands.cpp b/Engine/Plugins/Experimental/PCG/Source/PCGEditor/Private/PCGEditorCommands.cpp
--- a/Engine/Plugins/Experimental/PCG/Source/PCGEditor/Private/PCGEditorCommands.cpp
+++ b/Engine/Plugins/Experimental/PCG/Source/PCGEditor/Private/PCGEditorCommands.cpp
@@ -20,7 +20,8 @@ void FPCGEditorCommands::RegisterCommands()
 {
 	UI_COMMAND(CollapseNodes, "Collapse into Subgraph", "Collapse selected nodes into a separate PCGGraph asset.", EUserInterfaceActionType::Button, FInputChord(EModifierKey::Alt, EKeys::J));
 	UI_COMMAND(Find, "Find", "Finds PCG nodes and comments in the current graph", EUserInterfaceActionType::Button, FInputChord(EModifierKey::Control, EKeys::F));
-	UI_COMMAND(PauseAutoRegeneration, "Pause Regen", "Pause Automatic Regeneration of PCG nodes.", EUserInterfaceActionType::ToggleButton, FInputChord(EModifierKey::Alt, EKeys::R));
+	UI_COMMAND(PauseAutoRegeneration, "Pause Regen", "Pause automatic regeneration of the current graph.", EUserInterfaceActionType::ToggleButton, FInputChord(EModifierKey::Alt, EKeys::R));
+	UI_COMMAND(ForceGraphRegeneration, "Force Regen", "Manually force a regeneration of the current graph.", EUserInterfaceActionType::Button, FInputChord());
 	UI_COMMAND(RunDeterminismTest, "Run Determinism Test", "Evaluate the current node for determinism.", EUserInterfaceActionType::Button, FInputChord(EModifierKey::Alt, EKeys::D));
 	UI_COMMAND(StartInspectNode, "Start Inspecting Node", "Start Inspecting Node", EUserInterfaceActionType::Button, FInputChord(EKeys::I));
 	UI_COMMAND(StopInspectNode, "Stop Inspecting Node", "Stop Inspecting Node", EUserInterfaceActionType::Button, FInputChord(EModifierKey::Shift, EKeys::I));
diff --git a/Engine/Plugins/Experimental/PCG/Source/PCGEditor/Private/PCGEditorCommands.h b/Engine/Plugins/Experimental/PCG/Source/PCGEditor/Private/PCGEditorCommands.h
--- a/Engine/Plugins/Experimental/PCG/Source/PCGEditor/Private/PCGEditorCommands.h
+++ b/Engine/Plugins/Experimental/PCG/Source/PCGEditor/Private/PCGEditorCommands.h
@@ -18,6 +18,7 @@ public:
 	TSharedPtr<FUICommandInfo> Find;
 	TSharedPtr<FUICommandInfo> RunDeterminismTest;
 	TSharedPtr<FUICommandInfo> PauseAutoRegeneration;
+	TSharedPtr<FUICommandInfo> ForceGraphRegeneration;
 	TSharedPtr<FUICommandInfo> StartInspectNode;
 	TSharedPtr<FUICommandInfo> StopInspectNode;
 	TSharedPtr<FUICommandInfo> ExecutionModeEnabled;
