Fixed an issue where duplicated objects with PCG components wouldn't react to graph changes
Changed to PCG graph in editor

#rnx

diff --git a/Engine/Plugins/Experimental/PCG/Source/PCG/Private/PCGComponent.cpp b/Engine/Plugins/Experimental/PCG/Source/PCG/Private/PCGComponent.cpp
--- a/Engine/Plugins/Experimental/PCG/Source/PCG/Private/PCGComponent.cpp
+++ b/Engine/Plugins/Experimental/PCG/Source/PCG/Private/PCGComponent.cpp
@@ -697,6 +697,13 @@ void UPCGComponent::PostLoad()
 #endif
 
 #if WITH_EDITOR
+	SetupCallbacksOnCreation();
+#endif
+}
+
+#if WITH_EDITOR
+void UPCGComponent::SetupCallbacksOnCreation()
+{
 	SetupActorCallbacks();
 	SetupTrackingCallbacks();
 
@@ -707,14 +714,14 @@ void UPCGComponent::PostLoad()
 	else
 	{
 		UpdateTrackedLandscape(/*bBoundsCheck=*/false);
-	}	
+	}
 
 	if (Graph)
 	{
 		Graph->OnGraphChangedDelegate.AddUObject(this, &UPCGComponent::OnGraphChanged);
 	}
-#endif
 }
+#endif
 
 void UPCGComponent::BeginDestroy()
 {
@@ -916,6 +923,13 @@ void UPCGComponent::PostEditChangeProperty(FPropertyChangedEvent& PropertyChange
 	}
 }
 
+void UPCGComponent::PostEditImport()
+{
+	Super::PostEditImport();
+
+	SetupCallbacksOnCreation();
+}
+
 void UPCGComponent::PreEditUndo()
 {
 	// Here we will keep a copy of flags that we require to keep through the undo
diff --git a/Engine/Plugins/Experimental/PCG/Source/PCG/Public/PCGComponent.h b/Engine/Plugins/Experimental/PCG/Source/PCG/Public/PCGComponent.h
--- a/Engine/Plugins/Experimental/PCG/Source/PCG/Public/PCGComponent.h
+++ b/Engine/Plugins/Experimental/PCG/Source/PCG/Public/PCGComponent.h
@@ -65,6 +65,10 @@ public:
 	/** ~Begin UObject interface */
 	virtual void PostLoad() override;
 	virtual void BeginDestroy() override;
+
+#if WITH_EDITOR
+	virtual void PostEditImport() override;
+#endif
 	/** ~End UObject interface */
 
 	//~Begin UActorComponent Interface
@@ -229,7 +233,7 @@ private:
 	void CleanupUnusedManagedResources();
 	bool MoveResourcesToNewActor(AActor* InNewActor, bool bCreateChild);
 
-	/* Called as part of the RerunConstructionScript, just takes resources as-is, assumes current state is empty + the resources have been retargeted if needed */
+	/** Called as part of the RerunConstructionScript, just takes resources as-is, assumes current state is empty + the resources have been retargeted if needed */
 	void GetManagedResources(TArray<TObjectPtr<UPCGManagedResource>>& Resources) const;
 	void SetManagedResources(const TArray<TObjectPtr<UPCGManagedResource>>& Resources);
 
@@ -244,6 +248,9 @@ private:
 	virtual void PreEditUndo() override;
 	virtual void PostEditUndo() override;
 
+	/** Sets up actor, tracking, landscape and graph callbacks */
+	void SetupCallbacksOnCreation();
+
 	void SetupActorCallbacks();
 	void TeardownActorCallbacks();
 	void SetupTrackingCallbacks();
diff --git a/Engine/Plugins/Experimental/PCG/Source/PCGEditor/Private/PCGEditor.cpp b/Engine/Plugins/Experimental/PCG/Source/PCGEditor/Private/PCGEditor.cpp
--- a/Engine/Plugins/Experimental/PCG/Source/PCGEditor/Private/PCGEditor.cpp
+++ b/Engine/Plugins/Experimental/PCG/Source/PCGEditor/Private/PCGEditor.cpp
@@ -1410,7 +1410,7 @@ TSharedRef<SGraphEditor> FPCGEditor::CreateGraphEditorWidget()
 	);
 
 	FGraphAppearanceInfo AppearanceInfo;
-	AppearanceInfo.CornerText = LOCTEXT("PCGGraphEditorCornerText", "Procedural Graph");
+	AppearanceInfo.CornerText = LOCTEXT("PCGGraphEditorCornerText", "PCG Graph");
 
 	SGraphEditor::FGraphEditorEvents InEvents;
 	InEvents.OnSelectionChanged = SGraphEditor::FOnSelectionChanged::CreateSP(this, &FPCGEditor::OnSelectedNodesChanged);
