PCG:
*Added SortMode callback to the PCG profiler view so that the UI arrows are displayed properly.

diff --git a/Engine/Plugins/Experimental/PCG/Source/PCGEditor/Private/SPCGEditorGraphProfilingView.cpp b/Engine/Plugins/Experimental/PCG/Source/PCGEditor/Private/SPCGEditorGraphProfilingView.cpp
--- a/Engine/Plugins/Experimental/PCG/Source/PCGEditor/Private/SPCGEditorGraphProfilingView.cpp
+++ b/Engine/Plugins/Experimental/PCG/Source/PCGEditor/Private/SPCGEditorGraphProfilingView.cpp
@@ -209,42 +209,49 @@ TSharedRef<SHeaderRow> SPCGEditorGraphProfilingView::CreateHeaderRowWidget()
 		.DefaultLabel(PCGEditorGraphProfilingView::TEXT_NodeLabel)
 		.HAlignHeader(HAlign_Center)
 		.HAlignCell(HAlign_Left)
+		.SortMode(this, &SPCGEditorGraphProfilingView::GetColumnSortMode, PCGEditorGraphProfilingView::NAME_Node)
 		.OnSort(this, &SPCGEditorGraphProfilingView::OnSortColumnHeader)
 		+ SHeaderRow::Column(PCGEditorGraphProfilingView::NAME_AvgTime)
 		.ManualWidth(100)
 		.DefaultLabel(PCGEditorGraphProfilingView::TEXT_AvgTimeLabel)
 		.HAlignHeader(HAlign_Center)
 		.HAlignCell(HAlign_Right)
+		.SortMode(this, &SPCGEditorGraphProfilingView::GetColumnSortMode, PCGEditorGraphProfilingView::NAME_AvgTime)
 		.OnSort(this, &SPCGEditorGraphProfilingView::OnSortColumnHeader)
 		+ SHeaderRow::Column(PCGEditorGraphProfilingView::NAME_MinTime)
 		.ManualWidth(100)
 		.DefaultLabel(PCGEditorGraphProfilingView::TEXT_MinTimeLabel)
 		.HAlignHeader(HAlign_Center)
 		.HAlignCell(HAlign_Right)
+		.SortMode(this, &SPCGEditorGraphProfilingView::GetColumnSortMode, PCGEditorGraphProfilingView::NAME_MinTime)
 		.OnSort(this, &SPCGEditorGraphProfilingView::OnSortColumnHeader)
 		+ SHeaderRow::Column(PCGEditorGraphProfilingView::NAME_MaxTime)
 		.ManualWidth(100)
 		.DefaultLabel(PCGEditorGraphProfilingView::TEXT_MaxTimeLabel)
 		.HAlignHeader(HAlign_Center)
 		.HAlignCell(HAlign_Right)
+		.SortMode(this, &SPCGEditorGraphProfilingView::GetColumnSortMode, PCGEditorGraphProfilingView::NAME_MaxTime)
 		.OnSort(this, &SPCGEditorGraphProfilingView::OnSortColumnHeader)
 		+ SHeaderRow::Column(PCGEditorGraphProfilingView::NAME_StdTime)
 		.ManualWidth(100)
 		.DefaultLabel(PCGEditorGraphProfilingView::TEXT_StdTimeLabel)
 		.HAlignHeader(HAlign_Center)
 		.HAlignCell(HAlign_Right)
+		.SortMode(this, &SPCGEditorGraphProfilingView::GetColumnSortMode, PCGEditorGraphProfilingView::NAME_StdTime)
 		.OnSort(this, &SPCGEditorGraphProfilingView::OnSortColumnHeader)
 		+ SHeaderRow::Column(PCGEditorGraphProfilingView::NAME_TotalTime)
 		.ManualWidth(100)
 		.DefaultLabel(PCGEditorGraphProfilingView::TEXT_TotalTimeLabel)
 		.HAlignHeader(HAlign_Center)
 		.HAlignCell(HAlign_Right)
+		.SortMode(this, &SPCGEditorGraphProfilingView::GetColumnSortMode, PCGEditorGraphProfilingView::NAME_TotalTime)
 		.OnSort(this, &SPCGEditorGraphProfilingView::OnSortColumnHeader)
 		+ SHeaderRow::Column(PCGEditorGraphProfilingView::NAME_NbCalls)
 		.ManualWidth(125)
 		.DefaultLabel(PCGEditorGraphProfilingView::TEXT_NbCallsLabel)
 		.HAlignHeader(HAlign_Center)
 		.HAlignCell(HAlign_Right)
+		.SortMode(this, &SPCGEditorGraphProfilingView::GetColumnSortMode, PCGEditorGraphProfilingView::NAME_NbCalls)
 		.OnSort(this, &SPCGEditorGraphProfilingView::OnSortColumnHeader);
 }
 
@@ -264,6 +271,16 @@ void SPCGEditorGraphProfilingView::OnSortColumnHeader(const EColumnSortPriority:
 	Refresh();
 }
 
+EColumnSortMode::Type SPCGEditorGraphProfilingView::GetColumnSortMode(const FName ColumnId) const
+{
+	if (SortingColumn != ColumnId)
+	{
+		return EColumnSortMode::None;
+	}
+
+	return SortMode;
+}
+
 FReply SPCGEditorGraphProfilingView::ResetTimers()
 {
 	const TSharedPtr<FPCGEditor> PCGEditor = PCGEditorPtr.Pin();
diff --git a/Engine/Plugins/Experimental/PCG/Source/PCGEditor/Private/SPCGEditorGraphProfilingView.h b/Engine/Plugins/Experimental/PCG/Source/PCGEditor/Private/SPCGEditorGraphProfilingView.h
--- a/Engine/Plugins/Experimental/PCG/Source/PCGEditor/Private/SPCGEditorGraphProfilingView.h
+++ b/Engine/Plugins/Experimental/PCG/Source/PCGEditor/Private/SPCGEditorGraphProfilingView.h
@@ -63,6 +63,7 @@ private:
 	TSharedRef<ITableRow> OnGenerateRow(PCGProfilingListViewItemPtr Item, const TSharedRef<STableViewBase>& OwnerTable) const;
 	void OnItemDoubleClicked(PCGProfilingListViewItemPtr Item);
 	void OnSortColumnHeader(const EColumnSortPriority::Type SortPriority, const FName& ColumnId, const EColumnSortMode::Type NewSortMode);
+	EColumnSortMode::Type GetColumnSortMode(const FName ColumnId) const;
 
 	/** Pointer back to the PCG editor that owns us */
 	TWeakPtr<FPCGEditor> PCGEditorPtr;
