Added an option to ignore data on the actor and consider it as a point only, which is a requirement for spawn actor workflow.
Added equivalent flag on the spawn actor side to force this if needed.
Fixed an issue where the collapse option would not be considered a structural change (which was needed toggling from/to NoMerging)

#rnx

diff --git a/Engine/Plugins/Experimental/PCG/Source/PCG/Private/Elements/PCGSpawnActor.cpp b/Engine/Plugins/Experimental/PCG/Source/PCG/Private/Elements/PCGSpawnActor.cpp
--- a/Engine/Plugins/Experimental/PCG/Source/PCG/Private/Elements/PCGSpawnActor.cpp
+++ b/Engine/Plugins/Experimental/PCG/Source/PCG/Private/Elements/PCGSpawnActor.cpp
@@ -83,6 +83,7 @@ bool FPCGSpawnActorElement::ExecuteInternal(FPCGContext* Context) const
 	}
 
 	bool bShouldPassThroughInputs = (Settings->Option != EPCGSpawnActorOption::NoMerging && Settings->GetSubgraph() != nullptr);
+	const bool bForceDisableActorParsing = (Settings->bForceDisableActorParsing);
 
 	// Pass-through exclusions & settings
 	TArray<FPCGTaggedData>& Outputs = Context->OutputData.TaggedData;
@@ -199,6 +200,11 @@ bool FPCGSpawnActorElement::ExecuteInternal(FPCGContext* Context) const
 						{
 							if (!bShouldPassThroughInputs)
 							{
+								if (bForceDisableActorParsing)
+								{
+									PCGComponent->bParseActorComponents = false;
+								}
+
 								PCGComponent->Generate();
 							}
 							else
diff --git a/Engine/Plugins/Experimental/PCG/Source/PCG/Private/PCGComponent.cpp b/Engine/Plugins/Experimental/PCG/Source/PCG/Private/PCGComponent.cpp
--- a/Engine/Plugins/Experimental/PCG/Source/PCG/Private/PCGComponent.cpp
+++ b/Engine/Plugins/Experimental/PCG/Source/PCG/Private/PCGComponent.cpp
@@ -901,6 +901,11 @@ void UPCGComponent::PostEditChangeProperty(FPropertyChangedEvent& PropertyChange
 		DirtyGenerated(EPCGComponentDirtyFlag::Input);
 		Refresh();
 	}
+	else if (PropName == GET_MEMBER_NAME_CHECKED(UPCGComponent, bParseActorComponents))
+	{
+		DirtyGenerated(EPCGComponentDirtyFlag::Input);
+		Refresh();
+	}
 	// General properties that don't affect behavior
 	else if (PropName == GET_MEMBER_NAME_CHECKED(UPCGComponent, Seed))
 	{
@@ -1587,10 +1592,10 @@ void UPCGComponent::UpdatePCGExclusionData()
 
 UPCGData* UPCGComponent::CreateActorPCGData()
 {
-	return CreateActorPCGData(GetOwner());
+	return CreateActorPCGData(GetOwner(), bParseActorComponents);
 }
 
-UPCGData* UPCGComponent::CreateActorPCGData(AActor* Actor)
+UPCGData* UPCGComponent::CreateActorPCGData(AActor* Actor, bool bParseActor)
 {
 	TRACE_CPUPROFILER_EVENT_SCOPE(UPCGComponent::CreateActorPCGData);
 
@@ -1634,7 +1639,7 @@ UPCGData* UPCGComponent::CreateActorPCGData(AActor* Actor)
 
 		return Data;
 	}
-	else // Prepare data on a component basis
+	else if (bParseActor)// Prepare data on a component basis
 	{
 		TInlineComponentArray<ULandscapeSplinesComponent*, 1> LandscapeSplines;
 		Actor->GetComponents<ULandscapeSplinesComponent>(LandscapeSplines);
@@ -1722,14 +1727,12 @@ UPCGData* UPCGComponent::CreateActorPCGData(AActor* Actor)
 		{
 			return Union;
 		}
-		else // No parsed components: default
-		{
-			//Default behavior on unknown actors is to write a single point at the actor location
-			UPCGPointData* Data = NewObject<UPCGPointData>(this);
-			Data->InitializeFromActor(Actor);
-			return Data;
-		}
 	}
+
+	// Finally, if it's not a special actor and there are not parsed components, then return a single point at the actor position
+	UPCGPointData* Data = NewObject<UPCGPointData>(this);
+	Data->InitializeFromActor(Actor);
+	return Data;
 }
 
 UPCGData* UPCGComponent::CreatePCGData()
diff --git a/Engine/Plugins/Experimental/PCG/Source/PCG/Private/PCGSettings.cpp b/Engine/Plugins/Experimental/PCG/Source/PCG/Private/PCGSettings.cpp
--- a/Engine/Plugins/Experimental/PCG/Source/PCG/Private/PCGSettings.cpp
+++ b/Engine/Plugins/Experimental/PCG/Source/PCG/Private/PCGSettings.cpp
@@ -97,7 +97,7 @@ void UPCGSettings::PostEditChangeProperty(FPropertyChangedEvent& PropertyChanged
 
 	if (PropertyChangedEvent.GetMemberPropertyName() != GET_MEMBER_NAME_CHECKED(UPCGSettings, DeterminismSettings))
 	{
-		OnSettingsChangedDelegate.Broadcast(this, EPCGChangeType::Settings);
+		OnSettingsChangedDelegate.Broadcast(this, IsStructuralProperty(PropertyChangedEvent.GetMemberPropertyName()) ? EPCGChangeType::Structural : EPCGChangeType::Settings);
 	}
 }
 
diff --git a/Engine/Plugins/Experimental/PCG/Source/PCG/Public/Elements/PCGSpawnActor.h b/Engine/Plugins/Experimental/PCG/Source/PCG/Public/Elements/PCGSpawnActor.h
--- a/Engine/Plugins/Experimental/PCG/Source/PCG/Public/Elements/PCGSpawnActor.h
+++ b/Engine/Plugins/Experimental/PCG/Source/PCG/Public/Elements/PCGSpawnActor.h
@@ -32,6 +32,9 @@ public:
 	UPROPERTY(BlueprintReadWrite, EditAnywhere, Category = Settings)
 	EPCGSpawnActorOption Option = EPCGSpawnActorOption::CollapseActors;
 
+	UPROPERTY(BlueprintReadWrite, EditAnywhere, Category = Settings, meta = (EditCondition = "Option==EPCGSpawnActorOption::NoMerging"))
+	bool bForceDisableActorParsing = true;
+
 	//~Begin UCPGSettings interface
 	virtual UPCGNode* CreateNode() const override;
 
diff --git a/Engine/Plugins/Experimental/PCG/Source/PCG/Public/PCGComponent.h b/Engine/Plugins/Experimental/PCG/Source/PCG/Public/PCGComponent.h
--- a/Engine/Plugins/Experimental/PCG/Source/PCG/Public/PCGComponent.h
+++ b/Engine/Plugins/Experimental/PCG/Source/PCG/Public/PCGComponent.h
@@ -132,6 +132,9 @@ public:
 	UPROPERTY(BlueprintReadWrite, EditAnywhere, Category = Settings)
 	EPCGComponentInput InputType = EPCGComponentInput::Actor;
 
+	UPROPERTY(BlueprintReadWrite, EditAnywhere, Category = Settings)
+	bool bParseActorComponents = true;
+
 	UPROPERTY(BlueprintReadWrite, EditAnywhere, Category = Settings)
 	int Seed = 42;
 
@@ -210,7 +213,7 @@ private:
 	UPCGData* CreatePCGData();
 	UPCGData* CreateInputPCGData();
 	UPCGData* CreateActorPCGData();
-	UPCGData* CreateActorPCGData(AActor* Actor);
+	UPCGData* CreateActorPCGData(AActor* Actor, bool bParseActor = true);
 	UPCGData* CreateLandscapePCGData(bool bHeightOnly);
 	void UpdatePCGExclusionData();
 	bool IsLandscapeCachedDataDirty(const UPCGData* Data) const;
diff --git a/Engine/Plugins/Experimental/PCG/Source/PCG/Public/PCGSettings.h b/Engine/Plugins/Experimental/PCG/Source/PCG/Public/PCGSettings.h
--- a/Engine/Plugins/Experimental/PCG/Source/PCG/Public/PCGSettings.h
+++ b/Engine/Plugins/Experimental/PCG/Source/PCG/Public/PCGSettings.h
@@ -119,6 +119,7 @@ protected:
 
 #if WITH_EDITOR
 	virtual void PostEditChangeProperty(FPropertyChangedEvent& PropertyChangedEvent) override;
+	virtual bool IsStructuralProperty(const FName& InPropertyName) const { return false; }
 
 	/** Method that can be called to dirty the cache data from this settings objects if the operator== does not allow to detect changes */
 	void DirtyCache();
diff --git a/Engine/Plugins/Experimental/PCG/Source/PCG/Public/PCGSubgraph.h b/Engine/Plugins/Experimental/PCG/Source/PCG/Public/PCGSubgraph.h
--- a/Engine/Plugins/Experimental/PCG/Source/PCG/Public/PCGSubgraph.h
+++ b/Engine/Plugins/Experimental/PCG/Source/PCG/Public/PCGSubgraph.h
@@ -35,8 +35,6 @@ protected:
 
 #if WITH_EDITOR
 	void OnSubgraphChanged(UPCGGraph* InGraph, EPCGChangeType ChangeType);
-
-	virtual bool IsStructuralProperty(const FName& InPropertyName) const { return false; }
 #endif
 };
 
