+ Add support for param metadata in the attribute list

diff --git a/Engine/Plugins/Experimental/PCG/Source/PCG/Private/Metadata/PCGMetadata.cpp b/Engine/Plugins/Experimental/PCG/Source/PCG/Private/Metadata/PCGMetadata.cpp
--- a/Engine/Plugins/Experimental/PCG/Source/PCG/Private/Metadata/PCGMetadata.cpp
+++ b/Engine/Plugins/Experimental/PCG/Source/PCG/Private/Metadata/PCGMetadata.cpp
@@ -816,6 +816,11 @@ void UPCGMetadata::ComputeWeightedAttribute(PCGMetadataEntryKey& OutKey, const T
 	AttributeLock.ReadUnlock();
 }
 
+int64 UPCGMetadata::GetItemKeyCountForParent() const
+{
+	return ItemKeyOffset;
+}
+
 void UPCGMetadata::SetAttributes(PCGMetadataEntryKey InKey, const UPCGMetadata* InMetadata, PCGMetadataEntryKey& OutKey)
 {
 	if (!InMetadata)
diff --git a/Engine/Plugins/Experimental/PCG/Source/PCG/Public/Metadata/PCGMetadata.h b/Engine/Plugins/Experimental/PCG/Source/PCG/Public/Metadata/PCGMetadata.h
--- a/Engine/Plugins/Experimental/PCG/Source/PCG/Public/Metadata/PCGMetadata.h
+++ b/Engine/Plugins/Experimental/PCG/Source/PCG/Public/Metadata/PCGMetadata.h
@@ -159,7 +159,6 @@ public:
 	/** Attributes operations */
 	void MergeAttributes(PCGMetadataEntryKey InKeyA, const UPCGMetadata* InMetadataA, PCGMetadataEntryKey InKeyB, const UPCGMetadata* InMetadataB, PCGMetadataEntryKey& OutKey, EPCGMetadataOp Op);
 	void MergeAttributesSubset(PCGMetadataEntryKey InKeyA, const UPCGMetadata* InMetadataA, const UPCGMetadata* InMetadataSubetA, PCGMetadataEntryKey InKeyB, const UPCGMetadata* InMetadataB, const UPCGMetadata* InMetadataSubsetb, PCGMetadataEntryKey& OutKey, EPCGMetadataOp Op);
-	void MergeAttributes(PCGMetadataEntryKey InKeyA, PCGMetadataEntryKey InKeyB, PCGMetadataEntryKey& OutKey, EPCGMetadataOp Op);
 
 	void ResetWeightedAttributes(PCGMetadataEntryKey& OutKey);
 	void AccumulateWeightedAttributes(PCGMetadataEntryKey InKey, const UPCGMetadata* InMetadata, float Weight, bool bSetNonInterpolableAttributes, PCGMetadataEntryKey& OutKey);
@@ -200,6 +199,7 @@ public:
 	void ComputePointWeightedAttribute(FPCGPoint& OutPoint, const TArrayView<TPair<const FPCGPoint*, float>>& InWeightedPoints, const UPCGMetadata* InMetadata);
 	void ComputeWeightedAttribute(PCGMetadataEntryKey& OutKey, const TArrayView<TPair<PCGMetadataEntryKey, float>>& InWeightedKeys, const UPCGMetadata* InMetadata);
 
+	int64 GetItemKeyCountForParent() const;
 	int64 GetItemCountForChild() const;
 
 protected:
diff --git a/Engine/Plugins/Experimental/PCG/Source/PCGEditor/Private/SPCGEditorGraphAttributeListView.cpp b/Engine/Plugins/Experimental/PCG/Source/PCGEditor/Private/SPCGEditorGraphAttributeListView.cpp
--- a/Engine/Plugins/Experimental/PCG/Source/PCGEditor/Private/SPCGEditorGraphAttributeListView.cpp
+++ b/Engine/Plugins/Experimental/PCG/Source/PCGEditor/Private/SPCGEditorGraphAttributeListView.cpp
@@ -9,8 +9,8 @@
 #include "PCGComponent.h"
 #include "PCGData.h"
 #include "PCGEditor.h"
-#include "PCGGraph.h"
 #include "PCGNode.h"
+#include "PCGParamData.h"
 
 #include "Styling/AppStyle.h"
 #include "Widgets/Input/SSearchBox.h"
@@ -23,7 +23,7 @@
 namespace PCGEditorGraphAttributeListView
 {
 	const FText NoDataAvailableText = LOCTEXT("NoDataAvailableText", "No data available");
-	
+
 	/** Names of the columns in the attribute list */
 	const FName NAME_IndexColumn = FName(TEXT("IndexColumn"));
 	const FName NAME_PointPositionX = FName(TEXT("PointPositionX"));
@@ -88,246 +88,267 @@ void SPCGListViewItemRow::Construct(const FArguments& InArgs, const TSharedRef<S
 TSharedRef<SWidget> SPCGListViewItemRow::GenerateWidgetForColumn(const FName& ColumnId)
 {
 	FText ColumnData = LOCTEXT("ColumnError", "Unrecognized Column");
+
+	if (ColumnId == PCGEditorGraphAttributeListView::NAME_IndexColumn)
+	{
+		ColumnData = FText::FromString(FString::FromInt(InternalItem->Index));
+		return SNew(STextBlock).Text(ColumnData);
+	}
+
 	if (const FPCGPoint* PCGPoint = InternalItem->PCGPoint)
 	{
-		const FTransform& Transform = PCGPoint->Transform;
-		if (ColumnId == PCGEditorGraphAttributeListView::NAME_IndexColumn)
-		{
-			ColumnData = FText::FromString(FString::FromInt(InternalItem->Index));
-		}
-		else if (ColumnId == PCGEditorGraphAttributeListView::NAME_PointPositionX)
-		{
-			const FVector& Position = Transform.GetLocation();
-			ColumnData = FText::AsNumber(Position.X);
-		}
-		else if (ColumnId == PCGEditorGraphAttributeListView::NAME_PointPositionY)
-		{
-			const FVector& Position = Transform.GetLocation();
-			ColumnData = FText::AsNumber(Position.Y);
-		}
-		else if (ColumnId == PCGEditorGraphAttributeListView::NAME_PointPositionZ)
-		{
-			const FVector& Position = Transform.GetLocation();
-			ColumnData = FText::AsNumber(Position.Z);
-		}
-		else if (ColumnId == PCGEditorGraphAttributeListView::NAME_PointRotationX)
-		{
-			const FRotator& Rotation = Transform.Rotator();
-			ColumnData = FText::AsNumber(Rotation.Roll);
-		}
-		else if (ColumnId == PCGEditorGraphAttributeListView::NAME_PointRotationY)
-		{
-			const FRotator& Rotation = Transform.Rotator();
-			ColumnData = FText::AsNumber(Rotation.Pitch);
-		}
-		else if (ColumnId == PCGEditorGraphAttributeListView::NAME_PointRotationZ)
-		{
-			const FRotator& Rotation = Transform.Rotator();
-			ColumnData = FText::AsNumber(Rotation.Yaw);
-		}
-		else if (ColumnId == PCGEditorGraphAttributeListView::NAME_PointScaleX)
-		{
-			const FVector& Scale = Transform.GetScale3D();
-			ColumnData = FText::AsNumber(Scale.X);
-		}
-		else if (ColumnId == PCGEditorGraphAttributeListView::NAME_PointScaleY)
-		{
-			const FVector& Scale = Transform.GetScale3D();
-			ColumnData = FText::AsNumber(Scale.Y);
-		}
-		else if (ColumnId == PCGEditorGraphAttributeListView::NAME_PointScaleZ)
+		ColumnData = ConvertPointDataToText(PCGPoint, ColumnId);
+	}
+	else if (const UPCGMetadata* PCGMetadata = InternalItem->PCGMetadata)
+	{
+		if (const FPCGMetadataInfo* MetadataInfo = InternalItem->MetadataInfos->Find(ColumnId))
 		{
-			const FVector& Scale = Transform.GetScale3D();
-			ColumnData = FText::AsNumber(Scale.Z);
+			if (const FPCGMetadataAttributeBase* AttributeBase = PCGMetadata->GetConstAttribute((*MetadataInfo).MetadataId))
+			{
+				ColumnData = ConvertMetadataAttributeToText(AttributeBase, MetadataInfo, InternalItem->MetaDataItemKey);
+			}
 		}
-		else if (ColumnId == PCGEditorGraphAttributeListView::NAME_PointBoundsMinX)
+	}
+
+	return SNew(STextBlock).Text(ColumnData);
+}
+
+FText SPCGListViewItemRow::ConvertPointDataToText(const FPCGPoint* PCGPoint, const FName& ColumnId) const
+{
+	check(PCGPoint);
+
+	const FTransform& Transform = PCGPoint->Transform;
+	if (ColumnId == PCGEditorGraphAttributeListView::NAME_PointPositionX)
+	{
+		const FVector& Position = Transform.GetLocation();
+		return FText::AsNumber(Position.X);
+	}
+	else if (ColumnId == PCGEditorGraphAttributeListView::NAME_PointPositionY)
+	{
+		const FVector& Position = Transform.GetLocation();
+		return FText::AsNumber(Position.Y);
+	}
+	else if (ColumnId == PCGEditorGraphAttributeListView::NAME_PointPositionZ)
+	{
+		const FVector& Position = Transform.GetLocation();
+		return FText::AsNumber(Position.Z);
+	}
+	else if (ColumnId == PCGEditorGraphAttributeListView::NAME_PointRotationX)
+	{
+		const FRotator& Rotation = Transform.Rotator();
+		return FText::AsNumber(Rotation.Roll);
+	}
+	else if (ColumnId == PCGEditorGraphAttributeListView::NAME_PointRotationY)
+	{
+		const FRotator& Rotation = Transform.Rotator();
+		return FText::AsNumber(Rotation.Pitch);
+	}
+	else if (ColumnId == PCGEditorGraphAttributeListView::NAME_PointRotationZ)
+	{
+		const FRotator& Rotation = Transform.Rotator();
+		return FText::AsNumber(Rotation.Yaw);
+	}
+	else if (ColumnId == PCGEditorGraphAttributeListView::NAME_PointScaleX)
+	{
+		const FVector& Scale = Transform.GetScale3D();
+		return FText::AsNumber(Scale.X);
+	}
+	else if (ColumnId == PCGEditorGraphAttributeListView::NAME_PointScaleY)
+	{
+		const FVector& Scale = Transform.GetScale3D();
+		return FText::AsNumber(Scale.Y);
+	}
+	else if (ColumnId == PCGEditorGraphAttributeListView::NAME_PointScaleZ)
+	{
+		const FVector& Scale = Transform.GetScale3D();
+		return FText::AsNumber(Scale.Z);
+	}
+	else if (ColumnId == PCGEditorGraphAttributeListView::NAME_PointBoundsMinX)
+	{
+		return FText::AsNumber(PCGPoint->BoundsMin.X);
+	}
+	else if (ColumnId == PCGEditorGraphAttributeListView::NAME_PointBoundsMinY)
+	{
+		return FText::AsNumber(PCGPoint->BoundsMin.Y);
+	}
+	else if (ColumnId == PCGEditorGraphAttributeListView::NAME_PointBoundsMinZ)
+	{
+		return FText::AsNumber(PCGPoint->BoundsMin.Z);
+	}
+	else if (ColumnId == PCGEditorGraphAttributeListView::NAME_PointBoundsMaxX)
+	{
+		return FText::AsNumber(PCGPoint->BoundsMax.X);
+	}
+	else if (ColumnId == PCGEditorGraphAttributeListView::NAME_PointBoundsMaxY)
+	{
+		return FText::AsNumber(PCGPoint->BoundsMax.Y);
+	}
+	else if (ColumnId == PCGEditorGraphAttributeListView::NAME_PointBoundsMaxZ)
+	{
+		return FText::AsNumber(PCGPoint->BoundsMax.Z);
+	}
+	else if (ColumnId == PCGEditorGraphAttributeListView::NAME_PointColorR)
+	{
+		return FText::AsNumber(PCGPoint->Color.X);
+	}
+	else if (ColumnId == PCGEditorGraphAttributeListView::NAME_PointColorG)
+	{
+		return FText::AsNumber(PCGPoint->Color.Y);
+	}
+	else if (ColumnId == PCGEditorGraphAttributeListView::NAME_PointColorB)
+	{
+		return FText::AsNumber(PCGPoint->Color.Z);
+	}
+	else if (ColumnId == PCGEditorGraphAttributeListView::NAME_PointColorA)
+	{
+		return FText::AsNumber(PCGPoint->Color.W);
+	}
+	else if (ColumnId == PCGEditorGraphAttributeListView::NAME_PointDensity)
+	{
+		const float Density = PCGPoint->Density;
+		return FText::AsNumber(Density);
+	}
+	else if (ColumnId == PCGEditorGraphAttributeListView::NAME_PointSteepness)
+	{
+		const float Steepness = PCGPoint->Steepness;
+		return FText::AsNumber(Steepness);
+	}
+	else if (ColumnId == PCGEditorGraphAttributeListView::NAME_PointSeed)
+	{
+		const int32 Seed = PCGPoint->Seed;
+		return FText::AsNumber(Seed);
+	}
+
+	return PCGEditorGraphAttributeListView::NoDataAvailableText;
+}
+
+FText SPCGListViewItemRow::ConvertMetadataAttributeToText(const FPCGMetadataAttributeBase* AttributeBase, const FPCGMetadataInfo* MetadataInfo, int64 ItemKey) const
+{
+	switch (AttributeBase->GetTypeId())
+	{
+	case PCG::Private::MetadataTypes<float>::Id:
 		{
-			ColumnData = FText::AsNumber(PCGPoint->BoundsMin.X);
+			const float MetaFloat = static_cast<const FPCGMetadataAttribute<float>*>(AttributeBase)->GetValueFromItemKey(ItemKey);
+			return FText::AsNumber(MetaFloat);
 		}
-		else if (ColumnId == PCGEditorGraphAttributeListView::NAME_PointBoundsMinY)
+	case PCG::Private::MetadataTypes<double>::Id:
 		{
-			ColumnData = FText::AsNumber(PCGPoint->BoundsMin.Y);
+			const double MetaDouble = static_cast<const FPCGMetadataAttribute<double>*>(AttributeBase)->GetValueFromItemKey(ItemKey);
+			return FText::AsNumber(MetaDouble);
 		}
-		else if (ColumnId == PCGEditorGraphAttributeListView::NAME_PointBoundsMinZ)
+	case PCG::Private::MetadataTypes<bool>::Id:
 		{
-			ColumnData = FText::AsNumber(PCGPoint->BoundsMin.Z);
+			const bool bMetaBool = static_cast<const FPCGMetadataAttribute<bool>*>(AttributeBase)->GetValueFromItemKey(ItemKey);
+			return bMetaBool ? FText::FromString(TEXT("true")) : FText::FromString(TEXT("false"));
 		}
-		else if (ColumnId == PCGEditorGraphAttributeListView::NAME_PointBoundsMaxX)
+	case PCG::Private::MetadataTypes<FVector>::Id:
 		{
-			ColumnData = FText::AsNumber(PCGPoint->BoundsMax.X);
+			const FVector MetaVector = static_cast<const FPCGMetadataAttribute<FVector>*>(AttributeBase)->GetValueFromItemKey(ItemKey);
+			return FText::AsNumber(MetaVector[(*MetadataInfo).Index]);
 		}
-		else if (ColumnId == PCGEditorGraphAttributeListView::NAME_PointBoundsMaxY)
+	case PCG::Private::MetadataTypes<FVector4>::Id:
 		{
-			ColumnData = FText::AsNumber(PCGPoint->BoundsMax.Y);
+			const FVector4 MetaVector4 = static_cast<const FPCGMetadataAttribute<FVector4>*>(AttributeBase)->GetValueFromItemKey(ItemKey);
+			return FText::AsNumber(MetaVector4[(*MetadataInfo).Index]);
 		}
-		else if (ColumnId == PCGEditorGraphAttributeListView::NAME_PointBoundsMaxZ)
+	case PCG::Private::MetadataTypes<int32>::Id:
 		{
-			ColumnData = FText::AsNumber(PCGPoint->BoundsMax.Z);
+			const int32 MetaInt32 = static_cast<const FPCGMetadataAttribute<int32>*>(AttributeBase)->GetValueFromItemKey(ItemKey);
+			return FText::AsNumber(MetaInt32);
 		}
-		else if (ColumnId == PCGEditorGraphAttributeListView::NAME_PointColorR)
+	case PCG::Private::MetadataTypes<int64>::Id:
 		{
-			ColumnData = FText::AsNumber(PCGPoint->Color.X);
+			const int64 MetaInt64 = static_cast<const FPCGMetadataAttribute<int64>*>(AttributeBase)->GetValueFromItemKey(ItemKey);
+			return FText::AsNumber(MetaInt64);
 		}
-		else if (ColumnId == PCGEditorGraphAttributeListView::NAME_PointColorG)
+	case PCG::Private::MetadataTypes<FString>::Id:
 		{
-			ColumnData = FText::AsNumber(PCGPoint->Color.Y);
+			const FString MetaString = static_cast<const FPCGMetadataAttribute<FString>*>(AttributeBase)->GetValueFromItemKey(ItemKey);
+			return FText::FromString(MetaString);
 		}
-		else if (ColumnId == PCGEditorGraphAttributeListView::NAME_PointColorB)
+	case PCG::Private::MetadataTypes<FName>::Id:
 		{
-			ColumnData = FText::AsNumber(PCGPoint->Color.Z);
+			const FName MetaName = static_cast<const FPCGMetadataAttribute<FName>*>(AttributeBase)->GetValueFromItemKey(ItemKey);
+			return FText::FromName(MetaName);
 		}
-		else if (ColumnId == PCGEditorGraphAttributeListView::NAME_PointColorA)
+	case PCG::Private::MetadataTypes<FQuat>::Id:
 		{
-			ColumnData = FText::AsNumber(PCGPoint->Color.W);
-		}		
-		else if (ColumnId == PCGEditorGraphAttributeListView::NAME_PointDensity)
-		{
-			const float Density = PCGPoint->Density;
-			ColumnData = FText::AsNumber(Density);
+			const FQuat MetaQuat = static_cast<const FPCGMetadataAttribute<FQuat>*>(AttributeBase)->GetValueFromItemKey(ItemKey);
+			float QuatValue = 0.0f;
+			if ((*MetadataInfo).Index == 0)
+			{
+				QuatValue = MetaQuat.X;
+			}
+			else if ((*MetadataInfo).Index == 1)
+			{
+				QuatValue = MetaQuat.Y;
+			}
+			else if ((*MetadataInfo).Index == 2)
+			{
+				QuatValue = MetaQuat.Z;
+			}
+			else if ((*MetadataInfo).Index == 3)
+			{
+				QuatValue = MetaQuat.W;
+			}
+			return FText::AsNumber(QuatValue);
 		}
-		else if (ColumnId == PCGEditorGraphAttributeListView::NAME_PointSteepness)
+	case PCG::Private::MetadataTypes<FRotator>::Id:
 		{
-			const float Steepness = PCGPoint->Steepness;
-			ColumnData = FText::AsNumber(Steepness);
+			const FRotator MetaRotator = static_cast<const FPCGMetadataAttribute<FRotator>*>(AttributeBase)->GetValueFromItemKey(ItemKey);
+			float RotatorValue = 0.0f;
+			if ((*MetadataInfo).Index == 0)
+			{
+				RotatorValue = MetaRotator.Roll;
+			}
+			else if ((*MetadataInfo).Index == 1)
+			{
+				RotatorValue = MetaRotator.Pitch;
+			}
+			else if ((*MetadataInfo).Index == 2)
+			{
+				RotatorValue = MetaRotator.Yaw;
+			}
+			return FText::AsNumber(RotatorValue);
 		}
-		else if (ColumnId == PCGEditorGraphAttributeListView::NAME_PointSeed)
+	case PCG::Private::MetadataTypes<FTransform>::Id:
 		{
-			const int32 Seed = PCGPoint->Seed;
-			ColumnData = FText::AsNumber(Seed);
-		}
+			const FTransform& MetaTransform = static_cast<const FPCGMetadataAttribute<FTransform>*>(AttributeBase)->GetValueFromItemKey(ItemKey);
+			const int8 ComponentIndex = (*MetadataInfo).Index / 3;
+			const int8 ValueIndex = (*MetadataInfo).Index % 3;
 
-		if (const UPCGMetadata* PCGMetadata = InternalItem->PCGMetadata)
-		{
-			if (const FPCGMetadataInfo* MetadataInfo = InternalItem->MetadataInfos->Find(ColumnId))
+			if (ComponentIndex == 0)
+			{
+				const FVector& MetaVector = MetaTransform.GetLocation();
+				return FText::AsNumber(MetaVector[ValueIndex]);
+			}
+			else if (ComponentIndex == 1)
 			{
-				if (const FPCGMetadataAttributeBase* AttributeBase = PCGMetadata->GetConstAttribute((*MetadataInfo).MetadataId)) // todo investigate deadlock caused by read lock
+				const FRotator& MetaRotator = MetaTransform.Rotator();
+				float RotatorValue = 0.0f;
+				if (ValueIndex == 0)
+				{
+					RotatorValue = MetaRotator.Roll;
+				}
+				else if (ValueIndex == 1)
+				{
+					RotatorValue = MetaRotator.Pitch;
+				}
+				else if (ValueIndex == 2)
 				{
-					if (AttributeBase->GetTypeId() == PCG::Private::MetadataTypes<float>::Id)
-					{
-						const float MetaFloat = static_cast<const FPCGMetadataAttribute<float>*>(AttributeBase)->GetValueFromItemKey(PCGPoint->MetadataEntry);
-						ColumnData = FText::AsNumber(MetaFloat);
-					}
-					else if (AttributeBase->GetTypeId() == PCG::Private::MetadataTypes<double>::Id)
-					{
-						const double MetaDouble = static_cast<const FPCGMetadataAttribute<double>*>(AttributeBase)->GetValueFromItemKey(PCGPoint->MetadataEntry);
-						ColumnData = FText::AsNumber(MetaDouble);
-					}
-					else if (AttributeBase->GetTypeId() == PCG::Private::MetadataTypes<bool>::Id)
-					{
-						const bool bMetaBool = static_cast<const FPCGMetadataAttribute<bool>*>(AttributeBase)->GetValueFromItemKey(PCGPoint->MetadataEntry);
-						ColumnData = bMetaBool ? FText::FromString(TEXT("true")) : FText::FromString(TEXT("false"));
-					}
-					else if (AttributeBase->GetTypeId() == PCG::Private::MetadataTypes<FVector>::Id)
-					{
-						const FVector MetaVector = static_cast<const FPCGMetadataAttribute<FVector>*>(AttributeBase)->GetValueFromItemKey(PCGPoint->MetadataEntry);
-						ColumnData = FText::AsNumber(MetaVector[(*MetadataInfo).Index]);
-					}
-					else if (AttributeBase->GetTypeId() == PCG::Private::MetadataTypes<FVector4>::Id)
-					{
-						const FVector4 MetaVector4 = static_cast<const FPCGMetadataAttribute<FVector4>*>(AttributeBase)->GetValueFromItemKey(PCGPoint->MetadataEntry);
-						ColumnData = FText::AsNumber(MetaVector4[(*MetadataInfo).Index]);
-					}
-					else if (AttributeBase->GetTypeId() == PCG::Private::MetadataTypes<int32>::Id)
-					{
-						const int32 MetaInt32 = static_cast<const FPCGMetadataAttribute<int32>*>(AttributeBase)->GetValueFromItemKey(PCGPoint->MetadataEntry);
-						ColumnData = FText::AsNumber(MetaInt32);
-					}
-					else if (AttributeBase->GetTypeId() == PCG::Private::MetadataTypes<int64>::Id)
-					{
-						const int64 MetaInt64 = static_cast<const FPCGMetadataAttribute<int64>*>(AttributeBase)->GetValueFromItemKey(PCGPoint->MetadataEntry);
-						ColumnData = FText::AsNumber(MetaInt64);
-					}
-					else if (AttributeBase->GetTypeId() == PCG::Private::MetadataTypes<FString>::Id)
-					{
-						const FString MetaString = static_cast<const FPCGMetadataAttribute<FString>*>(AttributeBase)->GetValueFromItemKey(PCGPoint->MetadataEntry);
-						ColumnData = FText::FromString(MetaString);
-					}
-					else if (AttributeBase->GetTypeId() == PCG::Private::MetadataTypes<FName>::Id)
-					{
-						const FName MetaName = static_cast<const FPCGMetadataAttribute<FName>*>(AttributeBase)->GetValueFromItemKey(PCGPoint->MetadataEntry);
-						ColumnData = FText::FromName(MetaName);
-					}
-					else if (AttributeBase->GetTypeId() == PCG::Private::MetadataTypes<FQuat>::Id)
-					{
-						const FQuat MetaQuat = static_cast<const FPCGMetadataAttribute<FQuat>*>(AttributeBase)->GetValueFromItemKey(PCGPoint->MetadataEntry);
-						float QuatValue = 0.0f;
-						if ((*MetadataInfo).Index == 0)
-						{
-							QuatValue = MetaQuat.X;
-						}
-						else if ((*MetadataInfo).Index == 1)
-						{
-							QuatValue = MetaQuat.Y;
-						}
-						else if ((*MetadataInfo).Index == 2)
-						{
-							QuatValue = MetaQuat.Z;
-						}
-						else if ((*MetadataInfo).Index == 3)
-						{
-							QuatValue = MetaQuat.W;
-						}
-						
-						ColumnData = FText::AsNumber(QuatValue);
-					}
-					else if (AttributeBase->GetTypeId() == PCG::Private::MetadataTypes<FRotator>::Id)
-					{
-						const FRotator MetaRotator = static_cast<const FPCGMetadataAttribute<FRotator>*>(AttributeBase)->GetValueFromItemKey(PCGPoint->MetadataEntry);
-						float RotatorValue = 0.0f;
-						if ((*MetadataInfo).Index == 0)
-						{
-							RotatorValue = MetaRotator.Roll;
-						}
-						else if ((*MetadataInfo).Index == 1)
-						{
-							RotatorValue = MetaRotator.Pitch;
-						}
-						else if ((*MetadataInfo).Index == 2)
-						{
-							RotatorValue = MetaRotator.Yaw;
-						}
-						ColumnData = FText::AsNumber(RotatorValue);
-					}
-					else if (AttributeBase->GetTypeId() == PCG::Private::MetadataTypes<FTransform>::Id)
-					{
-						const FTransform& MetaTransform = static_cast<const FPCGMetadataAttribute<FTransform>*>(AttributeBase)->GetValueFromItemKey(PCGPoint->MetadataEntry);
-						const int8 ComponentIndex = (*MetadataInfo).Index / 3;
-						const int8 ValueIndex = (*MetadataInfo).Index % 3;
-						if (ComponentIndex == 0)
-						{
-							const FVector& MetaVector = MetaTransform.GetLocation();
-							ColumnData = FText::AsNumber(MetaVector[ValueIndex]);
-						}
-						else if(ComponentIndex == 1)
-						{
-							const FRotator& MetaRotator = MetaTransform.Rotator();
-							float RotatorValue = 0.0f;
-							if (ValueIndex == 0)
-							{
-								RotatorValue = MetaRotator.Roll;
-							}
-							else if (ValueIndex == 1)
-							{
-								RotatorValue = MetaRotator.Pitch;
-							}
-							else if (ValueIndex == 2)
-							{
-								RotatorValue = MetaRotator.Yaw;
-							}
-							ColumnData = FText::AsNumber(RotatorValue);
-						}
-						else if(ComponentIndex == 2)
-						{
-							const FVector& MetaVector = MetaTransform.GetScale3D();
-							ColumnData = FText::AsNumber(MetaVector[ValueIndex]);
-						}
-					}
+					RotatorValue = MetaRotator.Yaw;
 				}
+				return FText::AsNumber(RotatorValue);
+			}
+			else if (ComponentIndex == 2)
+			{
+				const FVector& MetaVector = MetaTransform.GetScale3D();
+				return FText::AsNumber(MetaVector[ValueIndex]);
 			}
 		}
+	default:
+		return PCGEditorGraphAttributeListView::NoDataAvailableText;
 	}
-
-	return SNew(STextBlock).Text(ColumnData);
 }
 
 SPCGEditorGraphAttributeListView::~SPCGEditorGraphAttributeListView()
@@ -355,7 +376,7 @@ void SPCGEditorGraphAttributeListView::Construct(const FArguments& InArgs, TShar
 	const TSharedRef<SScrollBar> VerticalScrollBar = SNew(SScrollBar)
 		.Orientation(Orient_Vertical)
 		.Thickness(FVector2D(12.0f, 12.0f));
-	
+
 	SAssignNew(ListView, SListView<PCGListviewItemPtr>)
 		.ListItemsSource(&ListViewItems)
 		.HeaderRow(ListViewHeader)
@@ -372,7 +393,7 @@ void SPCGEditorGraphAttributeListView::Construct(const FArguments& InArgs, TShar
 			SNew(STextBlock)
 			.Text(this, &SPCGEditorGraphAttributeListView::OnGenerateSelectedDataText)
 		];
-	
+
 	this->ChildSlot
 	[
 		SNew(SVerticalBox)
@@ -418,122 +439,7 @@ TSharedRef<SHeaderRow> SPCGEditorGraphAttributeListView::CreateHeaderRowWidget()
 {
 	return SNew(SHeaderRow)
 			.ResizeMode(ESplitterResizeMode::FixedPosition)
-			.CanSelectGeneratedColumn(true)
-			+ SHeaderRow::Column(PCGEditorGraphAttributeListView::NAME_IndexColumn)
-			.DefaultLabel(PCGEditorGraphAttributeListView::TEXT_IndexLabel)
-			.HAlignHeader(HAlign_Center)
-			.HAlignCell(HAlign_Right)
-			.ManualWidth(44)
-			+ SHeaderRow::Column(PCGEditorGraphAttributeListView::NAME_PointPositionX)
-			.DefaultLabel(PCGEditorGraphAttributeListView::TEXT_PointPositionLabelX)
-			.HAlignHeader(HAlign_Center)
-			.HAlignCell(HAlign_Right)
-			.ManualWidth(94)
-			+ SHeaderRow::Column(PCGEditorGraphAttributeListView::NAME_PointPositionY)
-			.DefaultLabel(PCGEditorGraphAttributeListView::TEXT_PointPositionLabelY)
-			.HAlignHeader(HAlign_Center)
-			.HAlignCell(HAlign_Right)
-			.ManualWidth(94)
-			+ SHeaderRow::Column(PCGEditorGraphAttributeListView::NAME_PointPositionZ)
-			.DefaultLabel(PCGEditorGraphAttributeListView::TEXT_PointPositionLabelZ)
-			.HAlignHeader(HAlign_Center)
-			.HAlignCell(HAlign_Right)
-			.ManualWidth(94)
-			+ SHeaderRow::Column(PCGEditorGraphAttributeListView::NAME_PointRotationX)
-			.DefaultLabel(PCGEditorGraphAttributeListView::TEXT_PointRotationLabelX)
-			.HAlignHeader(HAlign_Center)
-			.HAlignCell(HAlign_Right)
-			.ManualWidth(68)
-			+ SHeaderRow::Column(PCGEditorGraphAttributeListView::NAME_PointRotationY)
-			.DefaultLabel(PCGEditorGraphAttributeListView::TEXT_PointRotationLabelY)
-			.HAlignHeader(HAlign_Center)
-			.HAlignCell(HAlign_Right)
-			.ManualWidth(68)
-			+ SHeaderRow::Column(PCGEditorGraphAttributeListView::NAME_PointRotationZ)
-			.DefaultLabel(PCGEditorGraphAttributeListView::TEXT_PointRotationLabelZ)
-			.HAlignHeader(HAlign_Center)
-			.HAlignCell(HAlign_Right)
-			.ManualWidth(68)
-			+ SHeaderRow::Column(PCGEditorGraphAttributeListView::NAME_PointScaleX)
-			.DefaultLabel(PCGEditorGraphAttributeListView::TEXT_PointScaleLabelX)
-			.HAlignHeader(HAlign_Center)
-			.HAlignCell(HAlign_Right)
-			.ManualWidth(50)
-			+ SHeaderRow::Column(PCGEditorGraphAttributeListView::NAME_PointScaleY)
-			.DefaultLabel(PCGEditorGraphAttributeListView::TEXT_PointScaleLabelY)
-			.HAlignHeader(HAlign_Center)
-			.HAlignCell(HAlign_Right)
-			.ManualWidth(50)
-			+ SHeaderRow::Column(PCGEditorGraphAttributeListView::NAME_PointScaleZ)
-			.DefaultLabel(PCGEditorGraphAttributeListView::TEXT_PointScaleLabelZ)
-			.HAlignHeader(HAlign_Center)
-			.HAlignCell(HAlign_Right)
-			.ManualWidth(50)
-			+ SHeaderRow::Column(PCGEditorGraphAttributeListView::NAME_PointBoundsMinX)
-			.DefaultLabel(PCGEditorGraphAttributeListView::TEXT_PointBoundsLabelMinX)
-			.HAlignHeader(HAlign_Center)
-			.HAlignCell(HAlign_Right)
-			.ManualWidth(80)
-			+ SHeaderRow::Column(PCGEditorGraphAttributeListView::NAME_PointBoundsMinY)
-			.DefaultLabel(PCGEditorGraphAttributeListView::TEXT_PointBoundsLabelMinY)
-			.HAlignHeader(HAlign_Center)
-			.HAlignCell(HAlign_Right)
-			.ManualWidth(80)
-			+ SHeaderRow::Column(PCGEditorGraphAttributeListView::NAME_PointBoundsMinZ)
-			.DefaultLabel(PCGEditorGraphAttributeListView::TEXT_PointBoundsLabelMinZ)
-			.HAlignHeader(HAlign_Center)
-			.HAlignCell(HAlign_Right)
-			.ManualWidth(80)
-			+ SHeaderRow::Column(PCGEditorGraphAttributeListView::NAME_PointBoundsMaxX)
-			.DefaultLabel(PCGEditorGraphAttributeListView::TEXT_PointBoundsLabelMaxX)
-			.HAlignHeader(HAlign_Center)
-			.HAlignCell(HAlign_Right)
-			.ManualWidth(88)
-			+ SHeaderRow::Column(PCGEditorGraphAttributeListView::NAME_PointBoundsMaxY)
-			.DefaultLabel(PCGEditorGraphAttributeListView::TEXT_PointBoundsLabelMaxY)
-			.HAlignHeader(HAlign_Center)
-			.HAlignCell(HAlign_Right)
-			.ManualWidth(88)
-			+ SHeaderRow::Column(PCGEditorGraphAttributeListView::NAME_PointBoundsMaxZ)
-			.DefaultLabel(PCGEditorGraphAttributeListView::TEXT_PointBoundsLabelMaxZ)
-			.HAlignHeader(HAlign_Center)
-			.HAlignCell(HAlign_Right)
-			.ManualWidth(88)
-			+ SHeaderRow::Column(PCGEditorGraphAttributeListView::NAME_PointColorR)
-			.DefaultLabel(PCGEditorGraphAttributeListView::TEXT_PointColorLabelR)
-			.HAlignHeader(HAlign_Center)
-			.HAlignCell(HAlign_Right)
-			.ManualWidth(50)
-			+ SHeaderRow::Column(PCGEditorGraphAttributeListView::NAME_PointColorG)
-			.DefaultLabel(PCGEditorGraphAttributeListView::TEXT_PointColorLabelG)
-			.HAlignHeader(HAlign_Center)
-			.HAlignCell(HAlign_Right)
-			.ManualWidth(50)
-			+ SHeaderRow::Column(PCGEditorGraphAttributeListView::NAME_PointColorB)
-			.DefaultLabel(PCGEditorGraphAttributeListView::TEXT_PointColorLabelB)
-			.HAlignHeader(HAlign_Center)
-			.HAlignCell(HAlign_Right)
-			.ManualWidth(50)
-			+ SHeaderRow::Column(PCGEditorGraphAttributeListView::NAME_PointColorA)
-			.DefaultLabel(PCGEditorGraphAttributeListView::TEXT_PointColorLabelA)
-			.HAlignHeader(HAlign_Center)
-			.HAlignCell(HAlign_Right)
-			.ManualWidth(50)
-			+ SHeaderRow::Column(PCGEditorGraphAttributeListView::NAME_PointDensity)
-			.DefaultLabel(PCGEditorGraphAttributeListView::TEXT_PointDensityLabel)
-			.HAlignHeader(HAlign_Center)
-			.HAlignCell(HAlign_Right)
-			.ManualWidth(54)
-			+ SHeaderRow::Column(PCGEditorGraphAttributeListView::NAME_PointSteepness)
-			.DefaultLabel(PCGEditorGraphAttributeListView::TEXT_PointSteepnessLabel)
-			.HAlignHeader(HAlign_Center)
-			.HAlignCell(HAlign_Right)
-			.ManualWidth(73)
-			+ SHeaderRow::Column(PCGEditorGraphAttributeListView::NAME_PointSeed)
-			.DefaultLabel(PCGEditorGraphAttributeListView::TEXT_PointSeedLabel)
-			.HAlignHeader(HAlign_Center)
-			.HAlignCell(HAlign_Right)
-			.ManualWidth(88);
+			.CanSelectGeneratedColumn(true);
 }
 
 void SPCGEditorGraphAttributeListView::OnDebugObjectChanged(UPCGComponent* InPCGComponent)
@@ -544,14 +450,14 @@ void SPCGEditorGraphAttributeListView::OnDebugObjectChanged(UPCGComponent* InPCG
 		PCGComponent->OnPCGGraphCleanedDelegate.RemoveAll(this);
 		PCGComponent->DisableInspection();
 	}
-	
+
 	PCGComponent = InPCGComponent;
 
 	if (PCGComponent)
 	{
 		PCGComponent->EnableInspection();
 		PCGComponent->OnPCGGraphGeneratedDelegate.AddSP(this, &SPCGEditorGraphAttributeListView::OnGenerateUpdated);
-		PCGComponent->OnPCGGraphCleanedDelegate.AddSP(this, &SPCGEditorGraphAttributeListView::OnGenerateUpdated);	
+		PCGComponent->OnPCGGraphCleanedDelegate.AddSP(this, &SPCGEditorGraphAttributeListView::OnGenerateUpdated);
 	}
 	else
 	{
@@ -575,9 +481,7 @@ void SPCGEditorGraphAttributeListView::OnGenerateUpdated(UPCGComponent* /*InPCGC
 void SPCGEditorGraphAttributeListView::RefreshAttributeList()
 {
 	ListViewItems.Empty();
-	ListView->SetListItemsSource(EmptyList); // TODO: Investigate if it's possible to avoid this hack that's used to force refresh the list view. Without it the AddColumn will try to create widgets for old data.
-
-	RemoveMetadataColumns();
+	ListViewHeader->ClearColumns();
 	MetadataInfos.Empty();
 
 	const TSharedPtr<FPCGEditor> PCGEditor = PCGEditorPtr.Pin();
@@ -596,7 +500,7 @@ void SPCGEditorGraphAttributeListView::RefreshAttributeList()
 	{
 		return;
 	}
-	
+
 	const FPCGDataCollection* InspectionData = PCGComponent->GetInspectionData(PCGNode);
 	if (!InspectionData)
 	{
@@ -608,74 +512,41 @@ void SPCGEditorGraphAttributeListView::RefreshAttributeList()
 	{
 		return;
 	}
-	
+
 	const FPCGTaggedData& TaggedData = InspectionData->TaggedData[DataIndex];
 	const UPCGData* PCGData = TaggedData.Data;
 
-	if (const UPCGSpatialData* PCGSpatialData = Cast<UPCGSpatialData>(PCGData))
+	if (const UPCGParamData* PCGParamData = Cast<const UPCGParamData>(PCGData))
 	{
-		if (const UPCGPointData* PCGPointData = PCGSpatialData->ToPointData())
+		if (const UPCGMetadata* PCGMetadata = PCGParamData->ConstMetadata())
 		{
-			const UPCGMetadata* PCGMetadata = PCGPointData->ConstMetadata(); // TODO: add support for paramdata as well.
-		
-			TArray<FName> AttributeNames;
-			TArray<EPCGMetadataTypes> AttributeTypes;
-			PCGMetadata->GetAttributes(AttributeNames, AttributeTypes);
+			AddIndexColumn();
+			GenerateColumnsFromMetadata(PCGMetadata);
 
-			for (int32 i = 0; i < AttributeNames.Num(); i++)
+			PCGMetadataEntryKey ItemKeyLowerBound = PCGMetadata->GetItemKeyCountForParent();
+			PCGMetadataEntryKey ItemKeyUpperBound = PCGMetadata->GetItemCountForChild();
+			for (PCGMetadataEntryKey MetadataItemKey = ItemKeyLowerBound; MetadataItemKey < ItemKeyUpperBound; ++MetadataItemKey)
 			{
-				const FName& AttributeName = AttributeNames[i];
-				const EPCGMetadataTypes& AttributeType = AttributeTypes[i];
-
-				switch (AttributeType)
-				{
-				case EPCGMetadataTypes::Float:
-				case EPCGMetadataTypes::Double:
-				case EPCGMetadataTypes::Integer32:
-				case EPCGMetadataTypes::Integer64:
-				case EPCGMetadataTypes::Boolean:
-				case EPCGMetadataTypes::String:
-				case EPCGMetadataTypes::Name:
-					{
-						AddMetadataColumn(AttributeName);
-						break;
-					}
-				case EPCGMetadataTypes::Vector:
-					{
-						AddMetadataColumn(AttributeName, 0, TEXT("_X"));
-						AddMetadataColumn(AttributeName, 1, TEXT("_Y"));
-						AddMetadataColumn(AttributeName, 2, TEXT("_Z"));
-						break;
-					}
-				case EPCGMetadataTypes::Vector4:
-				case EPCGMetadataTypes::Quaternion:
-					{
-						AddMetadataColumn(AttributeName, 0, TEXT("_X"));
-						AddMetadataColumn(AttributeName, 1, TEXT("_Y"));
-						AddMetadataColumn(AttributeName, 2, TEXT("_Z"));
-						AddMetadataColumn(AttributeName, 3, TEXT("_W"));
-						break;
-					}
-				case EPCGMetadataTypes::Transform:
-					{
-						AddMetadataColumn(AttributeName, 0, TEXT("_tX"));
-						AddMetadataColumn(AttributeName, 1, TEXT("_tY"));
-						AddMetadataColumn(AttributeName, 2, TEXT("_tZ"));
-						AddMetadataColumn(AttributeName, 3, TEXT("_rX"));
-						AddMetadataColumn(AttributeName, 4, TEXT("_rY"));
-						AddMetadataColumn(AttributeName, 5, TEXT("_rZ"));
-						AddMetadataColumn(AttributeName, 6, TEXT("_sX"));
-						AddMetadataColumn(AttributeName, 7, TEXT("_sY"));
-						AddMetadataColumn(AttributeName, 8, TEXT("_sZ"));
-						break;
-					}
-				default:
-					break;
-				}
+				PCGListviewItemPtr ListViewItem = MakeShared<FPCGListViewItem>();
+				ListViewItem->Index = MetadataItemKey - ItemKeyLowerBound;
+				ListViewItem->PCGParamData = PCGParamData;
+				ListViewItem->PCGMetadata = PCGMetadata;
+				ListViewItem->MetaDataItemKey = MetadataItemKey;
+				ListViewItem->MetadataInfos = &MetadataInfos;
+				ListViewItems.Add(ListViewItem);
 			}
+		}
+	}
+	else if (const UPCGSpatialData* PCGSpatialData = Cast<const UPCGSpatialData>(PCGData))
+	{
+		if (const UPCGPointData* PCGPointData = PCGSpatialData->ToPointData())
+		{
+			const UPCGMetadata* PCGMetadata = PCGPointData->ConstMetadata();
 
-			const TArray<FPCGPoint>& PCGPoints = PCGPointData->GetPoints();
+			AddPointDataColumns();
+			GenerateColumnsFromMetadata(PCGMetadata);
 
+			const TArray<FPCGPoint>& PCGPoints = PCGPointData->GetPoints();
 			ListViewItems.Reserve(PCGPoints.Num());
 			for (int32 PointIndex = 0; PointIndex < PCGPoints.Num(); PointIndex++)
 			{
@@ -685,13 +556,14 @@ void SPCGEditorGraphAttributeListView::RefreshAttributeList()
 				ListViewItem->Index = PointIndex;
 				ListViewItem->PCGPoint = &PCGPoint;
 				ListViewItem->PCGMetadata = PCGMetadata;
+				ListViewItem->MetaDataItemKey = PCGPoint.MetadataEntry;
 				ListViewItem->MetadataInfos = &MetadataInfos;
 				ListViewItems.Add(ListViewItem);
 			}
 		}
 	}
 
-	ListView->SetListItemsSource(ListViewItems);
+	ListView->RequestListRefresh();
 }
 
 void SPCGEditorGraphAttributeListView::RefreshDataComboBox()
@@ -716,7 +588,7 @@ void SPCGEditorGraphAttributeListView::RefreshDataComboBox()
 	{
 		return;
 	}
-	
+
 	const FPCGDataCollection* InspectionData = PCGComponent->GetInspectionData(PCGNode);
 	if (!InspectionData)
 	{
@@ -770,12 +642,150 @@ int32 SPCGEditorGraphAttributeListView::GetSelectedDataIndex() const
 	return Index;
 }
 
+void SPCGEditorGraphAttributeListView::GenerateColumnsFromMetadata(const UPCGMetadata* PCGMetadata)
+{
+	TArray<FName> AttributeNames;
+	TArray<EPCGMetadataTypes> AttributeTypes;
+	PCGMetadata->GetAttributes(AttributeNames, AttributeTypes);
+
+	for (int32 I = 0; I < AttributeNames.Num(); I++)
+	{
+		const FName& AttributeName = AttributeNames[I];
+
+		switch (AttributeTypes[I])
+		{
+		case EPCGMetadataTypes::Float:
+		case EPCGMetadataTypes::Double:
+		case EPCGMetadataTypes::Integer32:
+		case EPCGMetadataTypes::Integer64:
+		case EPCGMetadataTypes::Boolean:
+		case EPCGMetadataTypes::String:
+		case EPCGMetadataTypes::Name:
+			{
+				AddMetadataColumn(AttributeName);
+				break;
+			}
+		case EPCGMetadataTypes::Vector:
+			{
+				AddMetadataColumn(AttributeName, 0, TEXT("_X"));
+				AddMetadataColumn(AttributeName, 1, TEXT("_Y"));
+				AddMetadataColumn(AttributeName, 2, TEXT("_Z"));
+				break;
+			}
+		case EPCGMetadataTypes::Vector4:
+		case EPCGMetadataTypes::Quaternion:
+			{
+				AddMetadataColumn(AttributeName, 0, TEXT("_X"));
+				AddMetadataColumn(AttributeName, 1, TEXT("_Y"));
+				AddMetadataColumn(AttributeName, 2, TEXT("_Z"));
+				AddMetadataColumn(AttributeName, 3, TEXT("_W"));
+				break;
+			}
+		case EPCGMetadataTypes::Transform:
+			{
+				AddMetadataColumn(AttributeName, 0, TEXT("_tX"));
+				AddMetadataColumn(AttributeName, 1, TEXT("_tY"));
+				AddMetadataColumn(AttributeName, 2, TEXT("_tZ"));
+				AddMetadataColumn(AttributeName, 3, TEXT("_rX"));
+				AddMetadataColumn(AttributeName, 4, TEXT("_rY"));
+				AddMetadataColumn(AttributeName, 5, TEXT("_rZ"));
+				AddMetadataColumn(AttributeName, 6, TEXT("_sX"));
+				AddMetadataColumn(AttributeName, 7, TEXT("_sY"));
+				AddMetadataColumn(AttributeName, 8, TEXT("_sZ"));
+				break;
+			}
+		default:
+			break;
+		}
+	}
+}
+
 TSharedRef<ITableRow> SPCGEditorGraphAttributeListView::OnGenerateRow(PCGListviewItemPtr Item, const TSharedRef<STableViewBase>& OwnerTable) const
 {
 	return SNew(SPCGListViewItemRow, OwnerTable, Item);
 }
 
-void SPCGEditorGraphAttributeListView::AddMetadataColumn(const FName& InColumnId, const int8 InValueIndex,  const TCHAR* PostFix)
+void SPCGEditorGraphAttributeListView::AddColumn(const FName& InColumnID, const FText& ColumnLabel, float ColumnWidth, EHorizontalAlignment HeaderHAlign, EHorizontalAlignment CellHAlign)
+{
+	SHeaderRow::FColumn::FArguments Arguments;
+	Arguments.ColumnId(InColumnID);
+	Arguments.DefaultLabel(ColumnLabel);
+	Arguments.ManualWidth(ColumnWidth);
+	Arguments.HAlignHeader(HeaderHAlign);
+	Arguments.HAlignCell(CellHAlign);
+	ListViewHeader->AddColumn(Arguments);
+}
+
+void SPCGEditorGraphAttributeListView::RemoveColumn(const FName& InColumnID)
+{
+	ListViewHeader->RemoveColumn(InColumnID);
+}
+
+void SPCGEditorGraphAttributeListView::AddIndexColumn()
+{
+	AddColumn(PCGEditorGraphAttributeListView::NAME_IndexColumn, PCGEditorGraphAttributeListView::TEXT_IndexLabel, 44);
+}
+
+void SPCGEditorGraphAttributeListView::RemoveIndexColumn()
+{
+	RemoveColumn(PCGEditorGraphAttributeListView::NAME_IndexColumn);
+}
+
+void SPCGEditorGraphAttributeListView::AddPointDataColumns()
+{
+	AddIndexColumn();
+	AddColumn(PCGEditorGraphAttributeListView::NAME_PointPositionX, PCGEditorGraphAttributeListView::TEXT_PointPositionLabelX, 94);
+	AddColumn(PCGEditorGraphAttributeListView::NAME_PointPositionY, PCGEditorGraphAttributeListView::TEXT_PointPositionLabelY, 94);
+	AddColumn(PCGEditorGraphAttributeListView::NAME_PointPositionZ, PCGEditorGraphAttributeListView::TEXT_PointPositionLabelZ, 94);
+	AddColumn(PCGEditorGraphAttributeListView::NAME_PointRotationX, PCGEditorGraphAttributeListView::TEXT_PointRotationLabelX, 68);
+	AddColumn(PCGEditorGraphAttributeListView::NAME_PointRotationY, PCGEditorGraphAttributeListView::TEXT_PointRotationLabelY, 68);
+	AddColumn(PCGEditorGraphAttributeListView::NAME_PointRotationZ, PCGEditorGraphAttributeListView::TEXT_PointRotationLabelZ, 68);
+	AddColumn(PCGEditorGraphAttributeListView::NAME_PointScaleX, PCGEditorGraphAttributeListView::TEXT_PointScaleLabelX, 50);
+	AddColumn(PCGEditorGraphAttributeListView::NAME_PointScaleY, PCGEditorGraphAttributeListView::TEXT_PointScaleLabelY, 50);
+	AddColumn(PCGEditorGraphAttributeListView::NAME_PointScaleZ, PCGEditorGraphAttributeListView::TEXT_PointScaleLabelZ, 50);
+	AddColumn(PCGEditorGraphAttributeListView::NAME_PointBoundsMinX, PCGEditorGraphAttributeListView::TEXT_PointBoundsLabelMinX, 80);
+	AddColumn(PCGEditorGraphAttributeListView::NAME_PointBoundsMinY, PCGEditorGraphAttributeListView::TEXT_PointBoundsLabelMinY, 80);
+	AddColumn(PCGEditorGraphAttributeListView::NAME_PointBoundsMinZ, PCGEditorGraphAttributeListView::TEXT_PointBoundsLabelMinZ, 80);
+	AddColumn(PCGEditorGraphAttributeListView::NAME_PointBoundsMaxX, PCGEditorGraphAttributeListView::TEXT_PointBoundsLabelMaxX, 88);
+	AddColumn(PCGEditorGraphAttributeListView::NAME_PointBoundsMaxY, PCGEditorGraphAttributeListView::TEXT_PointBoundsLabelMaxY, 88);
+	AddColumn(PCGEditorGraphAttributeListView::NAME_PointBoundsMaxZ, PCGEditorGraphAttributeListView::TEXT_PointBoundsLabelMaxZ, 88);
+	AddColumn(PCGEditorGraphAttributeListView::NAME_PointColorR, PCGEditorGraphAttributeListView::TEXT_PointColorLabelR, 50);
+	AddColumn(PCGEditorGraphAttributeListView::NAME_PointColorG, PCGEditorGraphAttributeListView::TEXT_PointColorLabelG, 50);
+	AddColumn(PCGEditorGraphAttributeListView::NAME_PointColorB, PCGEditorGraphAttributeListView::TEXT_PointColorLabelB, 50);
+	AddColumn(PCGEditorGraphAttributeListView::NAME_PointColorA, PCGEditorGraphAttributeListView::TEXT_PointColorLabelA, 50);
+	AddColumn(PCGEditorGraphAttributeListView::NAME_PointDensity, PCGEditorGraphAttributeListView::TEXT_PointDensityLabel, 54);
+	AddColumn(PCGEditorGraphAttributeListView::NAME_PointSteepness, PCGEditorGraphAttributeListView::TEXT_PointSteepnessLabel, 73);
+	AddColumn(PCGEditorGraphAttributeListView::NAME_PointSeed, PCGEditorGraphAttributeListView::TEXT_PointSeedLabel, 88);
+}
+
+void SPCGEditorGraphAttributeListView::RemovePointDataColumns()
+{
+	RemoveIndexColumn();
+	RemoveColumn(PCGEditorGraphAttributeListView::NAME_PointPositionX);
+	RemoveColumn(PCGEditorGraphAttributeListView::NAME_PointPositionY);
+	RemoveColumn(PCGEditorGraphAttributeListView::NAME_PointPositionZ);
+	RemoveColumn(PCGEditorGraphAttributeListView::NAME_PointRotationX);
+	RemoveColumn(PCGEditorGraphAttributeListView::NAME_PointRotationY);
+	RemoveColumn(PCGEditorGraphAttributeListView::NAME_PointRotationZ);
+	RemoveColumn(PCGEditorGraphAttributeListView::NAME_PointScaleX);
+	RemoveColumn(PCGEditorGraphAttributeListView::NAME_PointScaleY);
+	RemoveColumn(PCGEditorGraphAttributeListView::NAME_PointScaleZ);
+	RemoveColumn(PCGEditorGraphAttributeListView::NAME_PointBoundsMinX);
+	RemoveColumn(PCGEditorGraphAttributeListView::NAME_PointBoundsMinY);
+	RemoveColumn(PCGEditorGraphAttributeListView::NAME_PointBoundsMinZ);
+	RemoveColumn(PCGEditorGraphAttributeListView::NAME_PointBoundsMaxX);
+	RemoveColumn(PCGEditorGraphAttributeListView::NAME_PointBoundsMaxY);
+	RemoveColumn(PCGEditorGraphAttributeListView::NAME_PointBoundsMaxZ);
+	RemoveColumn(PCGEditorGraphAttributeListView::NAME_PointColorR);
+	RemoveColumn(PCGEditorGraphAttributeListView::NAME_PointColorG);
+	RemoveColumn(PCGEditorGraphAttributeListView::NAME_PointColorB);
+	RemoveColumn(PCGEditorGraphAttributeListView::NAME_PointColorA);
+	RemoveColumn(PCGEditorGraphAttributeListView::NAME_PointDensity);
+	RemoveColumn(PCGEditorGraphAttributeListView::NAME_PointSteepness);
+	RemoveColumn(PCGEditorGraphAttributeListView::NAME_PointSeed);
+}
+
+void SPCGEditorGraphAttributeListView::AddMetadataColumn(const FName& InColumnId, const int8 InValueIndex, const TCHAR* PostFix)
 {
 	FString ColumnIdString = InColumnId.ToString();
 
@@ -785,7 +795,7 @@ void SPCGEditorGraphAttributeListView::AddMetadataColumn(const FName& InColumnId
 	}
 
 	const FName ColumnId(ColumnIdString);
-	
+
 	FPCGMetadataInfo MetadataInfo;
 	MetadataInfo.MetadataId = InColumnId;
 	MetadataInfo.Index = InValueIndex;
@@ -802,10 +812,9 @@ void SPCGEditorGraphAttributeListView::AddMetadataColumn(const FName& InColumnId
 
 void SPCGEditorGraphAttributeListView::RemoveMetadataColumns()
 {
-	for (auto MetadataInfo : MetadataInfos)
+	for (const TTuple<FName, FPCGMetadataInfo>& MetadataInfo : MetadataInfos)
 	{
-		const FName& MetadataKey = MetadataInfo.Key;
-		ListViewHeader->RemoveColumn(MetadataKey);
+		ListViewHeader->RemoveColumn(MetadataInfo.Key);
 	}
 }
 
diff --git a/Engine/Plugins/Experimental/PCG/Source/PCGEditor/Private/SPCGEditorGraphAttributeListView.h b/Engine/Plugins/Experimental/PCG/Source/PCGEditor/Private/SPCGEditorGraphAttributeListView.h
--- a/Engine/Plugins/Experimental/PCG/Source/PCGEditor/Private/SPCGEditorGraphAttributeListView.h
+++ b/Engine/Plugins/Experimental/PCG/Source/PCGEditor/Private/SPCGEditorGraphAttributeListView.h
@@ -12,6 +12,7 @@
 class FPCGEditor;
 class UPCGComponent;
 class UPCGNode;
+class UPCGParamData;
 struct FPCGPoint;
 
 struct FPCGMetadataInfo
@@ -25,6 +26,8 @@ struct FPCGListViewItem
 	int32 Index = INDEX_NONE;
 	const FPCGPoint* PCGPoint = nullptr;
 	const UPCGMetadata* PCGMetadata = nullptr;
+	const UPCGParamData* PCGParamData = nullptr;
+	int64 MetaDataItemKey = -1;
 	const TMap<FName, FPCGMetadataInfo>* MetadataInfos = nullptr;
 };
 
@@ -39,6 +42,9 @@ public:
 	void Construct(const FArguments& InArgs, const TSharedRef<STableViewBase>& InOwnerTableView, const PCGListviewItemPtr& Item);
 
 	virtual TSharedRef<SWidget> GenerateWidgetForColumn(const FName& ColumnId) override;
+private:
+	FText ConvertPointDataToText(const FPCGPoint* PCGPoint, const FName& ColumnId) const;
+	FText ConvertMetadataAttributeToText(const FPCGMetadataAttributeBase* AttributeBase, const FPCGMetadataInfo* MetadataInfo, int64 ItemKey = 0) const;
 
 protected:
 	PCGListviewItemPtr InternalItem;
@@ -50,18 +56,18 @@ public:
 	SLATE_BEGIN_ARGS(SPCGEditorGraphAttributeListView) { }
 	SLATE_END_ARGS()
 
-	virtual ~SPCGEditorGraphAttributeListView();
+	virtual ~SPCGEditorGraphAttributeListView() override;
 
 	void Construct(const FArguments& InArgs, TSharedPtr<FPCGEditor> InPCGEditor);
 
 private:
 	TSharedRef<SHeaderRow> CreateHeaderRowWidget() const;
-	
+
 	void OnDebugObjectChanged(UPCGComponent* InPCGComponent);
 	void OnInspectedNodeChanged(UPCGNode* InPCGNode);
 
 	void OnGenerateUpdated(UPCGComponent* InPCGComponent);
-	
+
 	void RefreshAttributeList();
 	void RefreshDataComboBox();
 
@@ -69,9 +75,19 @@ private:
 	void OnSelectionChanged(TSharedPtr<FName> Item, ESelectInfo::Type SelectInfo);
 	FText OnGenerateSelectedDataText() const;
 	int32 GetSelectedDataIndex() const;
+	void GenerateColumnsFromMetadata(const UPCGMetadata* PCGMetadata);
 
 	TSharedRef<ITableRow> OnGenerateRow(PCGListviewItemPtr Item, const TSharedRef<STableViewBase>& OwnerTable) const;
 
+	void AddColumn(const FName& InColumnID, const FText& ColumnLabel, float ColumnWidth, EHorizontalAlignment HeaderHAlign = HAlign_Center, EHorizontalAlignment CellHAlign = HAlign_Right);
+	void RemoveColumn(const FName& InColumnID);
+
+	void AddIndexColumn();
+	void RemoveIndexColumn();
+
+	void AddPointDataColumns();
+	void RemovePointDataColumns();
+
 	void AddMetadataColumn(const FName& InColumnId, const int8 InValueIndex = INDEX_NONE, const TCHAR* PostFix = nullptr);
 	void RemoveMetadataColumns();
 
@@ -84,8 +100,6 @@ private:
 	TSharedPtr<SHeaderRow> ListViewHeader;
 	TSharedPtr<SListView<PCGListviewItemPtr>> ListView;
 	TArray<PCGListviewItemPtr> ListViewItems;
-	/** Empty list to force refresh the ListView when regenerating */
-	TArray<PCGListviewItemPtr> EmptyList;
 
 	TSharedPtr<SComboBox<TSharedPtr<FName>>> DataComboBox;
 	TArray<TSharedPtr<FName>> DataComboBoxItems;
