* Fixes horizontal alignment of PCG Determinism Test Widget column values

#rnx

diff --git a/Engine/Plugins/Experimental/PCG/Source/PCGEditor/Private/SPCGEditorGraphDeterminism.cpp b/Engine/Plugins/Experimental/PCG/Source/PCGEditor/Private/SPCGEditorGraphDeterminism.cpp
--- a/Engine/Plugins/Experimental/PCG/Source/PCGEditor/Private/SPCGEditorGraphDeterminism.cpp
+++ b/Engine/Plugins/Experimental/PCG/Source/PCGEditor/Private/SPCGEditorGraphDeterminism.cpp
@@ -29,7 +29,7 @@ namespace
 	const FText TEXT_Basic(LOCTEXT("BasicDeterminism", "Pass"));
 
 	constexpr float SmallManualWidth = 25.f;
-	constexpr float MediumManualWidth = 70.f;
+	constexpr float MediumManualWidth = 80.f;
 	constexpr float LargeManualWidth = 160.f;
 	constexpr float ListViewRowHeight = 36.f;
 }
@@ -165,6 +165,7 @@ void SPCGEditorGraphDeterminismListView::AddColumn(const FTestColumnInfo& Column
 		Arguments.ManualWidth(ColumnInfo.Width);
 	}
 	Arguments.HAlignHeader(ColumnInfo.HAlign);
+	Arguments.HAlignCell(ColumnInfo.HAlign);
 	GeneratedHeaderRow->AddColumn(Arguments);
 }
 
