[PCG] PCGComponent seed is now taken into account and can be overriden in nodes that uses randomness.

diff --git a/Engine/Plugins/Experimental/PCG/Source/PCG/Private/Elements/PCGDensityNoise.cpp b/Engine/Plugins/Experimental/PCG/Source/PCG/Private/Elements/PCGDensityNoise.cpp
--- a/Engine/Plugins/Experimental/PCG/Source/PCG/Private/Elements/PCGDensityNoise.cpp
+++ b/Engine/Plugins/Experimental/PCG/Source/PCG/Private/Elements/PCGDensityNoise.cpp
@@ -28,10 +28,13 @@ bool FPCGDensityNoiseElement::ExecuteInternal(FPCGContext* Context) const
 	const float DensityNoiseMax = PCG_GET_OVERRIDEN_VALUE(Settings, DensityNoiseMax, Params);
 	const bool bInvertSourceDensity = PCG_GET_OVERRIDEN_VALUE(Settings, bInvertSourceDensity, Params);
 
+	// Precompute a seed based on the settings one and the component one
+	const int Seed = PCGSettingsHelpers::ComputeSeedWithOverride(Settings, Context->SourceComponent, Params);
+
 	ProcessPoints(Context, Inputs, Outputs, [&](const FPCGPoint& InPoint, FPCGPoint& OutPoint)
 	{
 		OutPoint = InPoint;
-		FRandomStream RandomSource(PCGHelpers::ComputeSeed(Settings->Seed, OutPoint.Seed));
+		FRandomStream RandomSource(PCGHelpers::ComputeSeed(Seed, OutPoint.Seed));
 
 		const float DensityNoise = RandomSource.FRandRange(DensityNoiseMin, DensityNoiseMax);
 
diff --git a/Engine/Plugins/Experimental/PCG/Source/PCG/Private/Elements/PCGDensityRemap.cpp b/Engine/Plugins/Experimental/PCG/Source/PCG/Private/Elements/PCGDensityRemap.cpp
--- a/Engine/Plugins/Experimental/PCG/Source/PCG/Private/Elements/PCGDensityRemap.cpp
+++ b/Engine/Plugins/Experimental/PCG/Source/PCG/Private/Elements/PCGDensityRemap.cpp
@@ -30,6 +30,8 @@ bool FPCGLinearDensityRemapElement::ExecuteInternal(FPCGContext* Context) const
 	const float SettingsRemapMax = PCGSettingsHelpers::GetValue(GET_MEMBER_NAME_CHECKED(UPCGLinearDensityRemapSettings, RemapMax), Settings->RemapMax, Params);
 	const bool bMultiplyDensity = PCGSettingsHelpers::GetValue(GET_MEMBER_NAME_CHECKED(UPCGLinearDensityRemapSettings, bMultiplyDensity), Settings->bMultiplyDensity, Params);
 
+	const int Seed = PCGSettingsHelpers::ComputeSeedWithOverride(Settings, Context->SourceComponent, Params);
+
 	const float RemapMin = FMath::Min(SettingsRemapMin, SettingsRemapMax);
 	const float RemapMax = FMath::Max(SettingsRemapMin, SettingsRemapMax);
 
@@ -63,20 +65,20 @@ bool FPCGLinearDensityRemapElement::ExecuteInternal(FPCGContext* Context) const
 
 		if (bMultiplyDensity)
 		{
-			FPCGAsync::AsyncPointProcessing(Context, OriginalPointCount, SampledPoints, [&Points, Settings, RemapMin, RemapMax](int32 Index, FPCGPoint& OutPoint)
+			FPCGAsync::AsyncPointProcessing(Context, OriginalPointCount, SampledPoints, [&Points, Seed, RemapMin, RemapMax](int32 Index, FPCGPoint& OutPoint)
 			{
 				OutPoint = Points[Index];
-				FRandomStream RandomSource(PCGHelpers::ComputeSeed(Settings->Seed, OutPoint.Seed));
+				FRandomStream RandomSource(PCGHelpers::ComputeSeed(Seed, OutPoint.Seed));
 				OutPoint.Density *= RandomSource.FRandRange(RemapMin, RemapMax);
 				return true;
 			});
 		}
 		else
 		{
-			FPCGAsync::AsyncPointProcessing(Context, OriginalPointCount, SampledPoints, [&Points, Settings, RemapMin, RemapMax](int32 Index, FPCGPoint& OutPoint)
+			FPCGAsync::AsyncPointProcessing(Context, OriginalPointCount, SampledPoints, [&Points, Seed, RemapMin, RemapMax](int32 Index, FPCGPoint& OutPoint)
 			{
 				OutPoint = Points[Index];
-				FRandomStream RandomSource(PCGHelpers::ComputeSeed(Settings->Seed, OutPoint.Seed));
+				FRandomStream RandomSource(PCGHelpers::ComputeSeed(Seed, OutPoint.Seed));
 				OutPoint.Density = RandomSource.FRandRange(RemapMin, RemapMax);
 				return true;
 			});
diff --git a/Engine/Plugins/Experimental/PCG/Source/PCG/Private/Elements/PCGPointSampler.cpp b/Engine/Plugins/Experimental/PCG/Source/PCG/Private/Elements/PCGPointSampler.cpp
--- a/Engine/Plugins/Experimental/PCG/Source/PCG/Private/Elements/PCGPointSampler.cpp
+++ b/Engine/Plugins/Experimental/PCG/Source/PCG/Private/Elements/PCGPointSampler.cpp
@@ -34,6 +34,8 @@ bool FPCGPointSamplerElement::ExecuteInternal(FPCGContext* Context) const
 	const bool bKeepZeroDensityPoints = false;
 #endif
 
+	const int Seed = PCGSettingsHelpers::ComputeSeedWithOverride(Settings, Context->SourceComponent, Params);
+
 	const bool bNoSampling = (Ratio <= 0.0f);
 	const bool bTrivialSampling = (Ratio >= 1.0f);
 
@@ -96,12 +98,12 @@ bool FPCGPointSamplerElement::ExecuteInternal(FPCGContext* Context) const
 		{
 			TRACE_CPUPROFILER_EVENT_SCOPE(FPCGPointSamplerElement::Execute::SelectPoints);
 
-			FPCGAsync::AsyncPointProcessing(Context, OriginalPointCount, SampledPoints, [&Points, Settings, Ratio, bKeepZeroDensityPoints](int32 Index, FPCGPoint& OutPoint)
+			FPCGAsync::AsyncPointProcessing(Context, OriginalPointCount, SampledPoints, [&Points, Seed, Ratio, bKeepZeroDensityPoints](int32 Index, FPCGPoint& OutPoint)
 			{
 				const FPCGPoint& Point = Points[Index];
 
 				// Apply a high-pass filter based on selected ratio
-				FRandomStream RandomSource(PCGHelpers::ComputeSeed(Settings->Seed, Point.Seed));
+				FRandomStream RandomSource(PCGHelpers::ComputeSeed(Seed, Point.Seed));
 				float Chance = RandomSource.FRand();
 
 				if (Chance < Ratio)
diff --git a/Engine/Plugins/Experimental/PCG/Source/PCG/Private/Elements/PCGSurfaceSampler.cpp b/Engine/Plugins/Experimental/PCG/Source/PCG/Private/Elements/PCGSurfaceSampler.cpp
--- a/Engine/Plugins/Experimental/PCG/Source/PCG/Private/Elements/PCGSurfaceSampler.cpp
+++ b/Engine/Plugins/Experimental/PCG/Source/PCG/Private/Elements/PCGSurfaceSampler.cpp
@@ -34,6 +34,8 @@ struct FPCGSurfaceSamplerLoopData
 
 	float Ratio;
 
+	int Seed;
+
 	bool Initialize(const UPCGSurfaceSamplerSettings* InSettings, FPCGContext* Context, const FBox& InputBounds)
 	{
 		UPCGParamData* Params = Context->InputData.GetParams();
@@ -49,6 +51,8 @@ struct FPCGSurfaceSamplerLoopData
 		bKeepZeroDensityPoints = PCGSettingsHelpers::GetValue(GET_MEMBER_NAME_CHECKED(UPCGSurfaceSamplerSettings, bKeepZeroDensityPoints), Settings->bKeepZeroDensityPoints, Params);
 #endif
 
+		Seed = PCGSettingsHelpers::ComputeSeedWithOverride(InSettings, Context->SourceComponent, Params);
+
 		// Conceptually, we will break down the surface bounds in a N x M grid
 		InterstitialDistance = PointExtents * 2;
 		InnerCellSize = InterstitialDistance * Looseness;
@@ -180,7 +184,7 @@ bool FPCGSurfaceSamplerElement::ExecuteInternal(FPCGContext* Context) const
 			const FVector::FReal CurrentY = CellY * LoopData.CellSize.Y;
 			const FVector InnerCellSize = LoopData.InnerCellSize;
 
-			FRandomStream RandomSource(PCGHelpers::ComputeSeed(LoopData.Settings->Seed, CellX, CellY));
+			FRandomStream RandomSource(PCGHelpers::ComputeSeed(LoopData.Seed, CellX, CellY));
 			float Chance = RandomSource.FRand();
 
 			const float Ratio = LoopData.Ratio;
diff --git a/Engine/Plugins/Experimental/PCG/Source/PCG/Private/Elements/PCGTransformPoints.cpp b/Engine/Plugins/Experimental/PCG/Source/PCG/Private/Elements/PCGTransformPoints.cpp
--- a/Engine/Plugins/Experimental/PCG/Source/PCG/Private/Elements/PCGTransformPoints.cpp
+++ b/Engine/Plugins/Experimental/PCG/Source/PCG/Private/Elements/PCGTransformPoints.cpp
@@ -42,11 +42,13 @@ bool FPCGTransformPointsElement::ExecuteInternal(FPCGContext* Context) const
 	const bool bUniformScale = PCG_GET_OVERRIDEN_VALUE(Settings, bUniformScale, Params);
 	const bool bRecomputeSeed = PCG_GET_OVERRIDEN_VALUE(Settings, bRecomputeSeed, Params);
 
+	const int Seed = PCGSettingsHelpers::ComputeSeedWithOverride(Settings, Context->SourceComponent, Params);
+
 	ProcessPoints(Context, Inputs, Outputs, [&](const FPCGPoint& InPoint, FPCGPoint& OutPoint)
 	{
 		OutPoint = InPoint;
 
-		FRandomStream RandomSource(PCGHelpers::ComputeSeed(Settings->Seed, InPoint.Seed));
+		FRandomStream RandomSource(PCGHelpers::ComputeSeed(Seed, InPoint.Seed));
 
 		const float OffsetX = RandomSource.FRandRange(OffsetMin.X, OffsetMax.X);
 		const float OffsetY = RandomSource.FRandRange(OffsetMin.Y, OffsetMax.Y);
diff --git a/Engine/Plugins/Experimental/PCG/Source/PCG/Private/Helpers/PCGBlueprintHelpers.cpp b/Engine/Plugins/Experimental/PCG/Source/PCG/Private/Helpers/PCGBlueprintHelpers.cpp
--- a/Engine/Plugins/Experimental/PCG/Source/PCG/Private/Helpers/PCGBlueprintHelpers.cpp
+++ b/Engine/Plugins/Experimental/PCG/Source/PCG/Private/Helpers/PCGBlueprintHelpers.cpp
@@ -17,16 +17,24 @@ void UPCGBlueprintHelpers::SetSeedFromPosition(FPCGPoint& InPoint)
 	InPoint.Seed = ComputeSeedFromPosition(InPoint.Transform.GetLocation());
 }
 
-FRandomStream UPCGBlueprintHelpers::GetRandomStream(const FPCGPoint& InPoint, const UPCGSettings* OptionalSettings)
+FRandomStream UPCGBlueprintHelpers::GetRandomStream(const FPCGPoint& InPoint, const UPCGSettings* OptionalSettings, const UPCGComponent* OptionalComponent)
 {
-	if (OptionalSettings)
+	int Seed = InPoint.Seed;
+
+	if (OptionalSettings && OptionalComponent)
 	{
-		return FRandomStream(PCGHelpers::ComputeSeed(InPoint.Seed, OptionalSettings->Seed));
+		Seed = PCGHelpers::ComputeSeed(InPoint.Seed, OptionalSettings->Seed, OptionalComponent->Seed);
 	}
-	else
+	else if (OptionalSettings)
 	{
-		return FRandomStream(InPoint.Seed);
+		Seed = PCGHelpers::ComputeSeed(InPoint.Seed, OptionalSettings->Seed);
 	}
+	else if (OptionalComponent)
+	{
+		Seed = PCGHelpers::ComputeSeed(InPoint.Seed, OptionalComponent->Seed);
+	}
+
+	return FRandomStream(Seed);
 }
 
 UPCGData* UPCGBlueprintHelpers::GetActorData(FPCGContext& Context)
diff --git a/Engine/Plugins/Experimental/PCG/Source/PCG/Private/Helpers/PCGSettingsHelpers.cpp b/Engine/Plugins/Experimental/PCG/Source/PCG/Private/Helpers/PCGSettingsHelpers.cpp
--- a/Engine/Plugins/Experimental/PCG/Source/PCG/Private/Helpers/PCGSettingsHelpers.cpp
+++ b/Engine/Plugins/Experimental/PCG/Source/PCG/Private/Helpers/PCGSettingsHelpers.cpp
@@ -2,6 +2,10 @@
 
 #include "Helpers/PCGSettingsHelpers.h"
 
+#include "PCGComponent.h"
+#include "PCGHelpers.h"
+#include "PCGSettings.h"
+
 namespace PCGSettingsHelpers
 {
 	template<typename T>
@@ -200,4 +204,12 @@ namespace PCGSettingsHelpers
 			UE_LOG(LogPCG, Warning, TEXT("Mismatching type in params vs. property %s"), *Property->GetFName().ToString());
 		}
 	}
+
+	int ComputeSeedWithOverride(const UPCGSettings* InSettings, const UPCGComponent* InComponent, UPCGParamData* InParams)
+	{
+		check(InSettings);
+
+		const int SettingsSeed = InParams ? PCGSettingsHelpers::GetValue(GET_MEMBER_NAME_CHECKED(UPCGSettings, Seed), InSettings->Seed, InParams) : InSettings->Seed;
+		return InComponent ? PCGHelpers::ComputeSeed(SettingsSeed, InComponent->Seed) : SettingsSeed;
+	}
 }
\ No newline at end of file
diff --git a/Engine/Plugins/Experimental/PCG/Source/PCG/Private/MeshSelectors/PCGMeshSelectorWeighted.cpp b/Engine/Plugins/Experimental/PCG/Source/PCG/Private/MeshSelectors/PCGMeshSelectorWeighted.cpp
--- a/Engine/Plugins/Experimental/PCG/Source/PCG/Private/MeshSelectors/PCGMeshSelectorWeighted.cpp
+++ b/Engine/Plugins/Experimental/PCG/Source/PCG/Private/MeshSelectors/PCGMeshSelectorWeighted.cpp
@@ -2,9 +2,9 @@
 
 #include "MeshSelectors/PCGMeshSelectorWeighted.h"
 
-#include "PCGHelpers.h"
 #include "PCGPoint.h"
 #include "Elements/PCGStaticMeshSpawner.h"
+#include "Helpers/PCGBlueprintHelpers.h"
 #include "Metadata/PCGMetadata.h"
 #include "Metadata/PCGMetadataAttribute.h"
 #include "Metadata/PCGMetadataAttributeTraits.h"
@@ -91,7 +91,7 @@ void UPCGMeshSelectorWeighted::SelectInstances_Implementation(
 				continue;
 			}
 
-			FRandomStream RandomSource(PCGHelpers::ComputeSeed(Point.Seed, Settings->Seed));
+			FRandomStream RandomSource = UPCGBlueprintHelpers::GetRandomStream(Point, Settings, Context.SourceComponent);
 			int RandomWeightedPick = RandomSource.RandRange(0, TotalWeight - 1);
 
 			int RandomPick = 0;
diff --git a/Engine/Plugins/Experimental/PCG/Source/PCG/Private/MeshSelectors/PCGMeshSelectorWeightedByCategory.cpp b/Engine/Plugins/Experimental/PCG/Source/PCG/Private/MeshSelectors/PCGMeshSelectorWeightedByCategory.cpp
--- a/Engine/Plugins/Experimental/PCG/Source/PCG/Private/MeshSelectors/PCGMeshSelectorWeightedByCategory.cpp
+++ b/Engine/Plugins/Experimental/PCG/Source/PCG/Private/MeshSelectors/PCGMeshSelectorWeightedByCategory.cpp
@@ -2,13 +2,13 @@
 
 #include "MeshSelectors/PCGMeshSelectorWeightedByCategory.h"
 
-#include "PCGHelpers.h"
+#include "Elements/PCGPointFilter.h"
+#include "Elements/PCGStaticMeshSpawner.h"
+#include "Helpers/PCGBlueprintHelpers.h"
 #include "Metadata/PCGMetadata.h"
 #include "Metadata/PCGMetadataAttribute.h"
 #include "Metadata/PCGMetadataAttributeTraits.h"
 #include "Metadata/PCGMetadataAttributeTpl.h"
-#include "Elements/PCGPointFilter.h"
-#include "Elements/PCGStaticMeshSpawner.h"
 
 #include "Math/RandomStream.h"
 #include "Engine/StaticMesh.h"
@@ -185,7 +185,7 @@ void UPCGMeshSelectorWeightedByCategory::SelectInstances_Implementation(
 
 		const int TotalWeight = InstancesAndWeights->CumulativeWeights.Last();
 
-		FRandomStream RandomSource(PCGHelpers::ComputeSeed(Point.Seed, Settings->Seed));
+		FRandomStream RandomSource = UPCGBlueprintHelpers::GetRandomStream(Point, Settings, Context.SourceComponent);
 		int RandomWeightedPick = RandomSource.RandRange(0, TotalWeight - 1);
 
 		int RandomPick = 0;
diff --git a/Engine/Plugins/Experimental/PCG/Source/PCG/Public/Helpers/PCGBlueprintHelpers.h b/Engine/Plugins/Experimental/PCG/Source/PCG/Public/Helpers/PCGBlueprintHelpers.h
--- a/Engine/Plugins/Experimental/PCG/Source/PCG/Public/Helpers/PCGBlueprintHelpers.h
+++ b/Engine/Plugins/Experimental/PCG/Source/PCG/Public/Helpers/PCGBlueprintHelpers.h
@@ -26,9 +26,9 @@ public:
 	UFUNCTION(BlueprintCallable, Category = "PCG|Helpers")
 	static void SetSeedFromPosition(UPARAM(ref) FPCGPoint& InPoint);
 
-	/** Creates a random stream from a point's seed and settings's seed (optional) */
+	/** Creates a random stream from a point's seed and settings/component's seed (optional) */
 	UFUNCTION(BlueprintCallable, Category = "PCG|Helpers", meta = (ScriptMethod))
-	static FRandomStream GetRandomStream(const FPCGPoint& InPoint, const UPCGSettings* OptionalSettings = nullptr);
+	static FRandomStream GetRandomStream(const FPCGPoint& InPoint, const UPCGSettings* OptionalSettings = nullptr, const UPCGComponent* OptionalComponent = nullptr);
 
 	UFUNCTION(BlueprintCallable, Category = "PCG|Temporary", meta = (ScriptMethod))
 	static UPCGData* GetActorData(UPARAM(ref) FPCGContext& Context);
diff --git a/Engine/Plugins/Experimental/PCG/Source/PCG/Public/Helpers/PCGSettingsHelpers.h b/Engine/Plugins/Experimental/PCG/Source/PCG/Public/Helpers/PCGSettingsHelpers.h
--- a/Engine/Plugins/Experimental/PCG/Source/PCG/Public/Helpers/PCGSettingsHelpers.h
+++ b/Engine/Plugins/Experimental/PCG/Source/PCG/Public/Helpers/PCGSettingsHelpers.h
@@ -8,6 +8,9 @@
 
 #include <type_traits>
 
+class UPCGComponent;
+class UPCGSettings;
+
 namespace PCGSettingsHelpers
 {
 	// Utility function to return a value from an attribute that doesn't match the type
@@ -215,6 +218,8 @@ namespace PCGSettingsHelpers
 
 		return ReturnType();
 	}
+
+	int ComputeSeedWithOverride(const UPCGSettings* InSettings, const UPCGComponent* InComponent, UPCGParamData* InParams);
 }
 
 #define PCG_GET_OVERRIDEN_VALUE(Settings, Variable, Params) PCGSettingsHelpers::GetValue(GET_MEMBER_NAME_CHECKED(TRemovePointer<TRemoveConst<decltype(Settings)>::Type>::Type, Variable), (Settings)->Variable, Params);
\ No newline at end of file
