Fixed issue with generated actors resources that were not behaving properly when moving/duplicating

#rnx

diff --git a/Engine/Plugins/Experimental/PCG/Source/PCG/Private/PCGComponent.cpp b/Engine/Plugins/Experimental/PCG/Source/PCG/Private/PCGComponent.cpp
--- a/Engine/Plugins/Experimental/PCG/Source/PCG/Private/PCGComponent.cpp
+++ b/Engine/Plugins/Experimental/PCG/Source/PCG/Private/PCGComponent.cpp
@@ -2149,9 +2149,12 @@ void FPCGComponentInstanceData::ApplyToComponent(UActorComponent* Component, con
 		TArray<TObjectPtr<UPCGManagedResource>> DuplicatedResources;
 		for (TObjectPtr<UPCGManagedResource> Resource : GeneratedResources)
 		{
-			UPCGManagedResource* DuplicatedResource = CastChecked<UPCGManagedResource>(StaticDuplicateObject(Resource, PCGComponent, FName()));
-			DuplicatedResource->PostEditImport();
-			DuplicatedResources.Add(DuplicatedResource);
+			if (Resource)
+			{
+				UPCGManagedResource* DuplicatedResource = CastChecked<UPCGManagedResource>(StaticDuplicateObject(Resource, PCGComponent, FName()));
+				DuplicatedResource->PostApplyToComponent();
+				DuplicatedResources.Add(DuplicatedResource);
+			}
 		}
 
 		if (DuplicatedResources.Num() > 0)
diff --git a/Engine/Plugins/Experimental/PCG/Source/PCG/Private/PCGManagedResource.cpp b/Engine/Plugins/Experimental/PCG/Source/PCG/Private/PCGManagedResource.cpp
--- a/Engine/Plugins/Experimental/PCG/Source/PCG/Private/PCGManagedResource.cpp
+++ b/Engine/Plugins/Experimental/PCG/Source/PCG/Private/PCGManagedResource.cpp
@@ -6,6 +6,11 @@
 #include "Components/InstancedStaticMeshComponent.h"
 #include "Components/SceneComponent.h"
 
+void UPCGManagedResource::PostApplyToComponent()
+{
+	PostEditImport();
+}
+
 // By default, if it is not a hard release, we mark the resource unused.
 bool UPCGManagedResource::Release(bool bHardRelease, TSet<TSoftObjectPtr<AActor>>& /*OutActorsToDelete*/)
 {
@@ -38,6 +43,11 @@ void UPCGManagedActors::PostEditImport()
 	GeneratedActors.Reset();
 }
 
+void UPCGManagedActors::PostApplyToComponent()
+{
+	// In this case, we want to preserve the data, so we need to do nothing
+}
+
 bool UPCGManagedActors::Release(bool bHardRelease, TSet<TSoftObjectPtr<AActor>>& OutActorsToDelete)
 {
 	TRACE_CPUPROFILER_EVENT_SCOPE(UPCGManagedActors::Release);
diff --git a/Engine/Plugins/Experimental/PCG/Source/PCG/Public/PCGManagedResource.h b/Engine/Plugins/Experimental/PCG/Source/PCG/Public/PCGManagedResource.h
--- a/Engine/Plugins/Experimental/PCG/Source/PCG/Public/PCGManagedResource.h
+++ b/Engine/Plugins/Experimental/PCG/Source/PCG/Public/PCGManagedResource.h
@@ -21,6 +21,9 @@ class PCG_API UPCGManagedResource : public UObject
 {
 	GENERATED_BODY()
 public:
+	/** Called when after a PCG component is applied to (such as after a RerunConstructionScript) */
+	virtual void PostApplyToComponent();
+
 	/** Releases/Mark Unsued the resource depending on the bHardRelease flag. Returns true if resource can be removed from the PCG component */
 	virtual bool Release(bool bHardRelease, TSet<TSoftObjectPtr<AActor>>& OutActorsToDelete);
 	/** Releases resource if empty or unused. Returns true if the resource can be removed from the PCG component */
@@ -48,6 +51,7 @@ public:
 	//~End UObject interface
 
 	//~Begin UPCGManagedResource interface
+	virtual void PostApplyToComponent() override;
 	virtual bool Release(bool bHardRelease, TSet<TSoftObjectPtr<AActor>>& OutActorsToDelete) override;
 	virtual bool ReleaseIfUnused(TSet<TSoftObjectPtr<AActor>>& OutActorsToDelete) override;
 	virtual bool MoveResourceToNewActor(AActor* NewActor) override;
