Graph-wide determinism test (Iteration I)
+ Adds graph-wide test functionality

Requires Future Improvements (TODO)
-----------------------------------------------
* Comparison of Param data
* Update the sample testing volume to the bounds of the spatial data tested
* Shuffling/Permutation of Input Order
* Subgraphs should run graph test natively
* Optimizations needed for large graphs (running a single subgraph once, for example)

diff --git a/Engine/Plugins/Experimental/PCG/Source/PCG/Private/Graph/PCGGraphExecutor.cpp b/Engine/Plugins/Experimental/PCG/Source/PCG/Private/Graph/PCGGraphExecutor.cpp
--- a/Engine/Plugins/Experimental/PCG/Source/PCG/Private/Graph/PCGGraphExecutor.cpp
+++ b/Engine/Plugins/Experimental/PCG/Source/PCG/Private/Graph/PCGGraphExecutor.cpp
@@ -582,6 +582,40 @@ FPCGElementPtr FPCGGraphExecutor::GetFetchInputElement()
 }
 
 #if WITH_EDITOR
+
+FPCGTaskId FPCGGraphExecutor::ScheduleDebugWithTaskCallback(UPCGComponent* InComponent, TFunction<void(FPCGTaskId/* TaskId*/, const UPCGNode*/* Node*/, const FPCGDataCollection&/* TaskOutput*/)> TaskCompleteCallback)
+{
+	FPCGTaskId FinalTaskID = Schedule(InComponent, {});
+	TArray<FPCGGraphTask> CompiledTasks = GraphCompiler->GetCompiledTasks(InComponent->GetGraph(), /*bIsTopGraph=*/true);
+	CompiledTasks.Pop(); // Remove the final task
+
+	// Set up all final dependencies for the entire execution
+	TArray<FPCGTaskId> FinalDependencies;
+	FinalDependencies.Reserve(CompiledTasks.Num() + 1);
+	FinalDependencies.Add(FinalTaskID);
+
+	for (const FPCGGraphTask& CompiledTask : CompiledTasks)
+	{
+		// Schedule the output capture hooks
+		FPCGTaskId CaptureTaskId = ScheduleGeneric([this, TaskCompleteCallback, CompiledTask]
+		{
+			FPCGDataCollection TaskOutputData;
+			if (CompiledTask.Node && GetOutputData(CompiledTask.NodeId, TaskOutputData))
+			{
+				TaskCompleteCallback(CompiledTask.NodeId, CompiledTask.Node, TaskOutputData);
+			}
+
+			return true;
+		}, {CompiledTask.NodeId});
+
+		// Add these tasks to the final dependencies
+		FinalDependencies.Add(CaptureTaskId);
+	}
+
+	// Finally, add a task to wait on the graph itself plus the capture tasks
+	return ScheduleGeneric([] { return true; }, FinalDependencies);
+}
+
 void FPCGGraphExecutor::AddToDirtyActors(AActor* Actor)
 {
 	ActorsListLock.Lock();
diff --git a/Engine/Plugins/Experimental/PCG/Source/PCG/Private/Graph/PCGGraphExecutor.h b/Engine/Plugins/Experimental/PCG/Source/PCG/Private/Graph/PCGGraphExecutor.h
--- a/Engine/Plugins/Experimental/PCG/Source/PCG/Private/Graph/PCGGraphExecutor.h
+++ b/Engine/Plugins/Experimental/PCG/Source/PCG/Private/Graph/PCGGraphExecutor.h
@@ -79,6 +79,8 @@ public:
 	bool GetOutputData(FPCGTaskId InTaskId, FPCGDataCollection& OutData);
 
 #if WITH_EDITOR
+	FPCGTaskId ScheduleDebugWithTaskCallback(UPCGComponent* InComponent, TFunction<void(FPCGTaskId, const UPCGNode*, const FPCGDataCollection&)> TaskCompleteCallback);
+
 	void AddToDirtyActors(AActor* Actor);
 	void AddToUnusedActors(const TSet<FWorldPartitionReference>& UnusedActors);
 
@@ -166,10 +168,10 @@ protected:
 	// Important note: generic elements must always be run on the main thread
 	// as most of these will impact the editor in some way (loading, unloading, saving)
 	virtual bool ExecuteInternal(FPCGContext* Context) const override;
-	virtual bool IsCancellable() const { return false; }
+	virtual bool IsCancellable() const override { return false; }
 
 #if WITH_EDITOR
-	virtual bool ShouldLog() const { return false; }
+	virtual bool ShouldLog() const override { return false; }
 #endif
 
 private:
diff --git a/Engine/Plugins/Experimental/PCG/Source/PCG/Private/Tests/Determinism/PCGDeterminismTestsCommon.cpp b/Engine/Plugins/Experimental/PCG/Source/PCG/Private/Tests/Determinism/PCGDeterminismTestsCommon.cpp
--- a/Engine/Plugins/Experimental/PCG/Source/PCG/Private/Tests/Determinism/PCGDeterminismTestsCommon.cpp
+++ b/Engine/Plugins/Experimental/PCG/Source/PCG/Private/Tests/Determinism/PCGDeterminismTestsCommon.cpp
@@ -3,6 +3,7 @@
 
 #include "PCGComponent.h"
 #include "PCGData.h"
+#include "PCGGraph.h"
 #include "PCGHelpers.h"
 #include "PCGNode.h"
 #include "PCGPin.h"
@@ -12,6 +13,8 @@
 #include "Data/PCGSplineData.h"
 #include "Data/PCGSurfaceData.h"
 #include "Data/PCGVolumeData.h"
+#include "Graph/PCGGraphCache.h"
+#include "Graph/PCGGraphExecutor.h"
 
 #include "Components/SplineComponent.h"
 #include "Components/StaticMeshComponent.h"
@@ -55,6 +58,151 @@ namespace PCGDeterminismTests
 		}
 	}
 
+#if WITH_EDITOR
+	void RunDeterminismTest(const UPCGGraph* InPCGGraph, UPCGComponent* InPCGComponent, FDeterminismNodeTestResult& OutResult)
+	{
+		if (!InPCGGraph || !InPCGComponent)
+		{
+			OutResult.DataTypesTested = EPCGDataType::None;
+			OutResult.bFlagRaised = true;
+			OutResult.AdditionalDetails.Add(TEXT("Invalid Graph/Component"));
+			return;
+		}
+
+		// Clone the actor and component
+		AActor* PCGActor = InPCGComponent->GetOwner();
+		check(PCGActor);
+
+		AActor* PCGActorCopy = DuplicateObject<AActor>(PCGActor, GetTransientPackage());
+		PCGActorCopy->SetFlags(RF_Transient);
+		UPCGComponent* PCGComponentCopy = PCGActorCopy->FindComponentByClass<UPCGComponent>();
+		check(PCGComponentCopy);
+
+		// TODO: Update this to a public setter
+		PCGComponentCopy->bIsPartitioned = false;
+
+		FPCGGraphExecutor Executor = FPCGGraphExecutor(PCGComponentCopy);
+
+		// TODO: Record data types tested
+
+		auto ScheduleAndWaitForExecution = [&Executor](FPCGTaskId FinalTaskID)
+		{
+			// TODO: Consider randomizing/iterating through possible input orders
+			bool bTasksComplete = false;
+
+			// Clear the cache between runs
+			Executor.GetCache().ClearCache();
+
+			// Schedule final task for waiting
+			Executor.ScheduleGeneric([&bTasksComplete]
+			{
+				bTasksComplete = true;
+				return true;
+			}, {FinalTaskID});
+
+			// Run first iteration
+			while (!bTasksComplete)
+			{
+				Executor.Execute();
+			}
+		};
+
+		TMap<FPCGTaskId, TTuple<const UPCGNode*, const FPCGDataCollection>> IntermediateOutputArray;
+
+		// TODO: TaskId may not remain robust in the future
+		auto RecordMappedOutput = [&IntermediateOutputArray](FPCGTaskId TaskId, const UPCGNode* Node, const FPCGDataCollection& NodeOutput)
+		{
+			check(Node);
+			// TODO: Ensure this is threadsafe if it needs to be
+			IntermediateOutputArray.Add(TaskId, MakeTuple(Node, NodeOutput));
+		};
+
+		// Schedule the graph execution
+		FPCGTaskId FinalTaskID = Executor.ScheduleDebugWithTaskCallback(PCGComponentCopy, RecordMappedOutput);
+		ScheduleAndWaitForExecution(FinalTaskID);
+
+		// Copy the first array over, so it can be used again
+		TMap<FPCGTaskId, TTuple<const UPCGNode*, const FPCGDataCollection>> FirstIntermediateOutputArray = IntermediateOutputArray;
+
+		// Clean up generated content
+		IntermediateOutputArray.Empty();
+		PCGComponentCopy->Cleanup();
+
+		// Run again
+		FinalTaskID = Executor.ScheduleDebugWithTaskCallback(PCGComponentCopy, RecordMappedOutput);
+		ScheduleAndWaitForExecution(FinalTaskID);
+
+		// Early out for mismatched outputs
+		if (IntermediateOutputArray.Num() != FirstIntermediateOutputArray.Num())
+		{
+			OutResult.bFlagRaised = true;
+			UpdateTestResults(Defaults::GraphResultName, OutResult, EDeterminismLevel::NoDeterminism);
+			OutResult.AdditionalDetails.Add(TEXT("Number of outputs differ"));
+			return;
+		}
+
+		// Begin output evaluation...
+		EDeterminismLevel HighestDeterminismLevel = EDeterminismLevel::OrderIndependent;
+
+		for (const TTuple<FPCGTaskId, TTuple<const UPCGNode*, const FPCGDataCollection>>& TaskOutputMapping : FirstIntermediateOutputArray)
+		{
+			const FPCGTaskId TaskIdKey = TaskOutputMapping.Key;
+			const TTuple<const UPCGNode*, const FPCGDataCollection> NodeDataTuple = TaskOutputMapping.Value;
+			const UPCGNode* NodePtr = NodeDataTuple.Get<0>();
+			const FString NodeNameString = NodePtr->GetName();
+
+			UE_LOG(LogPCG, Log, TEXT("[%s] Evaluating Node [%d]..."), *InPCGGraph->GetName(), *NodeNameString);
+
+			const FPCGDataCollection& FirstOutput = NodeDataTuple.Get<1>();
+			const FPCGDataCollection& SecondOutput = IntermediateOutputArray[TaskIdKey].Get<1>();
+
+			// Check for identical outputs
+			if (HighestDeterminismLevel == EDeterminismLevel::OrderIndependent)
+			{
+				if (!DataCollectionsAreIdentical(FirstOutput, SecondOutput))
+				{
+					// Not identical, downgrade to consistent
+					HighestDeterminismLevel = EDeterminismLevel::OrderConsistent;
+					UE_LOG(LogPCG, Log, TEXT("[%s] Node failed \"independent\" test. Downgrading to \"consistent\""), *NodeNameString);
+					OutResult.AdditionalDetails.Emplace(TEXT("Downgraded to order consistent at node: ") + NodeNameString);
+				}
+			}
+
+			// Check for consistent outputs
+			if (HighestDeterminismLevel == EDeterminismLevel::OrderConsistent)
+			{
+				TArray<int32> IndexOffsets;
+				if (!DataCollectionsMatch(FirstOutput, SecondOutput, IndexOffsets) && !DataCollectionsAreConsistent(FirstOutput, SecondOutput, IndexOffsets.Num()))
+				{
+					// Not consistent, downgrade to orthogonal
+					HighestDeterminismLevel = EDeterminismLevel::OrderOrthogonal;
+					UE_LOG(LogPCG, Log, TEXT("[%s] Node failed \"consistent\" test. Downgrading to \"orthogonal\""), *NodeNameString);
+					OutResult.AdditionalDetails.Emplace(TEXT("Downgraded to order orthogonal at node: ") + NodeNameString);
+				}
+			}
+
+			// Highest possible is orthogonal, if not, then fail test
+			if (HighestDeterminismLevel == EDeterminismLevel::OrderOrthogonal)
+			{
+				if (!DataCollectionsContainSameData(FirstOutput, SecondOutput))
+				{
+					// Not orthogonal, fail test
+					HighestDeterminismLevel = EDeterminismLevel::NoDeterminism;
+					UE_LOG(LogPCG, Log, TEXT("[%s] Node failed \"orthogonal\" test. Test failed."), *NodeNameString);
+					OutResult.AdditionalDetails.Emplace(TEXT("Determinism test failed at node: ") + NodeNameString);
+					break;
+				}
+			}
+		}
+
+		// Clean up PCG generated
+		PCGComponentCopy->Cleanup();
+
+		// Finalize the results
+		UpdateTestResults(Defaults::GraphResultName, OutResult, HighestDeterminismLevel);
+	}
+#endif
+
 	bool RunBasicTestSuite(const UPCGNode* InPCGNode, const FName& TestName, FDeterminismNodeTestResult& OutResult)
 	{
 		// Get Pins
@@ -684,7 +832,16 @@ namespace PCGDeterminismTests
 
 	bool SurfaceDataIsIdentical(const UPCGData* FirstData, const UPCGData* SecondData)
 	{
+		const UPCGSurfaceData* FirstSurfaceData = CastChecked<const UPCGSurfaceData>(FirstData);
+		const UPCGSurfaceData* SecondSurfaceData = CastChecked<const UPCGSurfaceData>(SecondData);
+
 		// TODO: Implement Surface Data comparison as needed in the future
+		// In the meantime, at least check if they are the same object
+		if (FirstSurfaceData == SecondSurfaceData)
+		{
+			return true;
+		}
+
 		UE_LOG(LogPCG, Warning, TEXT("Surface comparison not fully implemented."));
 
 		return false;
@@ -726,6 +883,7 @@ namespace PCGDeterminismTests
 	bool SampledSpatialDataIsIdentical(const UPCGSpatialData* FirstSpatialData, const UPCGSpatialData* SecondSpatialData)
 	{
 		// At this point, bounds has already been checked for equality
+		// TODO: Change this testing volume relative to the bounds of the two datas
 		const FBox SampleBounds = Defaults::TestingVolume;
 		const FVector SampleExtent = SampleBounds.GetExtent();
 
diff --git a/Engine/Plugins/Experimental/PCG/Source/PCG/Public/Tests/Determinism/PCGDeterminismTestsCommon.h b/Engine/Plugins/Experimental/PCG/Source/PCG/Public/Tests/Determinism/PCGDeterminismTestsCommon.h
--- a/Engine/Plugins/Experimental/PCG/Source/PCG/Public/Tests/Determinism/PCGDeterminismTestsCommon.h
+++ b/Engine/Plugins/Experimental/PCG/Source/PCG/Public/Tests/Determinism/PCGDeterminismTestsCommon.h
@@ -29,7 +29,7 @@ enum class EDeterminismLevel : uint8
 };
 
 USTRUCT(BlueprintType)
-struct FDeterminismNodeTestResult
+struct FDeterminismNodeTestResult // TODO: Refactor the names since this isn't always a node any longer
 {
 	GENERATED_BODY()
 
@@ -41,7 +41,6 @@ struct FDeterminismNodeTestResult
 	UPROPERTY(VisibleAnywhere, BlueprintReadOnly, Category = Determinism)
 	FString NodeName = TEXT("Unnamed");
 
-	// TODO: Add the seed to the UI widget
 	UPROPERTY(EditAnywhere, BlueprintReadWrite, Category = Determinism)
 	int32 Seed = -1;
 
@@ -85,6 +84,8 @@ namespace PCGDeterminismTests
 		const FVector LargeVector = FVector::OneVector * LargeDistance;
 
 		const FName TestPinName = TEXT("Test");
+		const FName GraphResultName = TEXT("Graph Result");
+
 		/** A little bigger than the typical largest volumes */
 		const FBox TestingVolume = FBox(-1.2 * LargeVector, 1.2 * LargeVector);
 	}
@@ -134,6 +135,11 @@ namespace PCGDeterminismTests
 	/** Validates if a PCGNode is deterministic */
 	PCG_API void RunDeterminismTest(const UPCGNode* InPCGNode, FDeterminismNodeTestResult& OutResult, const FNodeTestInfo& TestToRun);
 
+#if WITH_EDITOR
+	/** Validates if a PCGGraph is deterministic */
+	PCG_API void RunDeterminismTest(const UPCGGraph* InPCGGraph, UPCGComponent* InPCGComponent, FDeterminismNodeTestResult& OutResult);
+#endif
+
 	/** Validates all the generic determinism tests for any given node */
 	PCG_API bool RunBasicTestSuite(const UPCGNode* InPCGNode, const FName& TestName, FDeterminismNodeTestResult& OutResult);
 
diff --git a/Engine/Plugins/Experimental/PCG/Source/PCGEditor/Private/PCGEditor.cpp b/Engine/Plugins/Experimental/PCG/Source/PCGEditor/Private/PCGEditor.cpp
--- a/Engine/Plugins/Experimental/PCG/Source/PCGEditor/Private/PCGEditor.cpp
+++ b/Engine/Plugins/Experimental/PCG/Source/PCGEditor/Private/PCGEditor.cpp
@@ -315,6 +315,13 @@ void FPCGEditor::RegisterToolbar() const
 				InSection.AddEntry(FToolMenuEntry::InitWidget("SelectedDebugObjectWidget", SNew(SPCGEditorGraphDebugObjectWidget, Context->PCGEditor.Pin()), FText::GetEmpty()));	
 			}
 		}));
+
+		Section.AddSeparator(NAME_None);
+		Section.AddEntry(FToolMenuEntry::InitToolBarButton(
+			 PCGEditorCommands.RunDeterminismGraphTest,
+			 TAttribute<FText>(),
+			 TAttribute<FText>(),
+			 FSlateIcon(FAppStyle::GetAppStyleSetName(), "BlueprintDebugger.TabIcon")));
 	}
 }
 
@@ -336,6 +343,11 @@ void FPCGEditor::BindCommands()
 		PCGEditorCommands.ForceGraphRegeneration,
 		FExecuteAction::CreateSP(this, &FPCGEditor::OnForceGraphRegeneration_Clicked));
 
+	ToolkitCommands->MapAction(
+		PCGEditorCommands.RunDeterminismGraphTest,
+		FExecuteAction::CreateSP(this, &FPCGEditor::OnDeterminismGraphTest),
+		FCanExecuteAction::CreateSP(this, &FPCGEditor::CanRunDeterminismGraphTest));
+
 	GraphEditorCommands->MapAction(
 		PCGEditorCommands.CollapseNodes,
 		FExecuteAction::CreateSP(this, &FPCGEditor::OnCollapseNodesInSubgraph),
@@ -350,9 +362,9 @@ void FPCGEditor::BindCommands()
 		FExecuteAction::CreateSP(this, &FPCGEditor::OnStopInspectNode));
 
 	GraphEditorCommands->MapAction(
-		PCGEditorCommands.RunDeterminismTest,
-		FExecuteAction::CreateSP(this, &FPCGEditor::OnDeterminismTests),
-		FCanExecuteAction::CreateSP(this, &FPCGEditor::CanRunDeterminismTests));
+		PCGEditorCommands.RunDeterminismNodeTest,
+		FExecuteAction::CreateSP(this, &FPCGEditor::OnDeterminismNodeTest),
+		FCanExecuteAction::CreateSP(this, &FPCGEditor::CanRunDeterminismNodeTest));
 
 	GraphEditorCommands->MapAction(
 		PCGEditorCommands.ExecutionModeEnabled,
@@ -411,24 +423,22 @@ void FPCGEditor::OnForceGraphRegeneration_Clicked()
 	}
 }
 
-bool FPCGEditor::CanRunDeterminismTests() const
+bool FPCGEditor::CanRunDeterminismNodeTest() const
 {
 	check(GraphEditorWidget.IsValid());
 
-	bool bValidObjectFound = false;
 	for (const UObject* Object : GraphEditorWidget->GetSelectedNodes())
 	{
 		if (Cast<const UPCGEditorGraphNodeBase>(Object) && !Cast<const UPCGEditorGraphNodeInput>(Object) && !Cast<const UPCGEditorGraphNodeOutput>(Object))
 		{
-			bValidObjectFound = true;
-			break;
+			return true;
 		}
 	}
 
-	return bValidObjectFound;
+	return false;
 }
 
-void FPCGEditor::OnDeterminismTests()
+void FPCGEditor::OnDeterminismNodeTest()
 {
 	check(GraphEditorWidget.IsValid());
 
@@ -526,6 +536,50 @@ void FPCGEditor::OnDeterminismTests()
 	}
 }
 
+bool FPCGEditor::CanRunDeterminismGraphTest() const
+{
+	return PCGEditorGraph && PCGComponentBeingDebugged;
+}
+
+void FPCGEditor::OnDeterminismGraphTest()
+{
+	check(GraphEditorWidget.IsValid());
+
+	if (!DeterminismWidget.IsValid() || !DeterminismWidget->WidgetIsConstructed() || !PCGGraphBeingEdited || !PCGComponentBeingDebugged)
+	{
+		return;
+	}
+
+	if (PCGComponentBeingDebugged->GetGraph() != PCGGraphBeingEdited)
+	{
+		// TODO: Should we alert the user more directly or disable this altogether?
+		UE_LOG(LogPCGEditor, Warning, TEXT("Running Determinism on a PCG Component with different/no attached PCG Graph"));
+	}
+
+	DeterminismWidget->ClearItems();
+	DeterminismWidget->BuildBaseColumns();
+
+	FTestColumnInfo ColumnInfo({PCGDeterminismTests::Defaults::GraphResultName, NSLOCTEXT("PCGDeterminism", "Result", "Result"), 120.f, HAlign_Center});
+	DeterminismWidget->AddColumn(ColumnInfo);
+
+	TSharedPtr<FDeterminismNodeTestResult> TestResult = MakeShared<FDeterminismNodeTestResult>();
+	TestResult->NodeTitle = TEXT("Full Graph Test");
+	TestResult->NodeName = PCGGraphBeingEdited->GetName();
+	TestResult->Seed = PCGComponentBeingDebugged->Seed;
+
+	PCGDeterminismTests::RunDeterminismTest(PCGGraphBeingEdited, PCGComponentBeingDebugged, *TestResult);
+
+	DeterminismWidget->AddItem(TestResult);
+	DeterminismWidget->AddDetailsColumn();
+	DeterminismWidget->RefreshItems();
+
+	// Give focus to the Determinism Output Tab
+	if (TabManager.IsValid())
+	{
+		TabManager->TryInvokeTab(FPCGEditor_private::DeterminismID);
+	}
+}
+
 bool FPCGEditor::CanCollapseNodesInSubgraph() const
 {
 	bool HasPCGNode = false;
diff --git a/Engine/Plugins/Experimental/PCG/Source/PCGEditor/Private/PCGEditor.h b/Engine/Plugins/Experimental/PCG/Source/PCGEditor/Private/PCGEditor.h
--- a/Engine/Plugins/Experimental/PCG/Source/PCGEditor/Private/PCGEditor.h
+++ b/Engine/Plugins/Experimental/PCG/Source/PCGEditor/Private/PCGEditor.h
@@ -55,7 +55,7 @@ public:
 	// ~End IToolkit interface
 
 	// ~Begin FEditorUndoClient interface
-	virtual bool MatchesContext(const FTransactionContext& InContext, const TArray<TPair<UObject *, FTransactionObjectEvent>>& TransactionObjectContexts) const override;
+	virtual bool MatchesContext(const FTransactionContext& InContext, const TArray<TPair<UObject*, FTransactionObjectEvent>>& TransactionObjectContexts) const override;
 	virtual void PostUndo(bool bSuccess) override;
 	virtual void PostRedo(bool bSuccess) override { PostUndo(bSuccess); }
 	// ~End FEditorUndoClient interface
@@ -101,10 +101,15 @@ private:
 	/** Stop inspecting the current inspected node */
 	void OnStopInspectNode();
 
-	/** Can determinism be tested on this node */
-	bool CanRunDeterminismTests() const;
-	/** Summon the Determinism tab */
-	void OnDeterminismTests();
+	/** Can determinism be tested on the selected node(s) */
+	bool CanRunDeterminismNodeTest() const;
+	/** Run the determinism test on the selected node(s) */
+	void OnDeterminismNodeTest();
+
+	/** Can determinism be tested on the current graph */
+	bool CanRunDeterminismGraphTest() const;
+	/** Run the determinism test on the current graph */
+	void OnDeterminismGraphTest();
 
 	/** Whether or not an execution mode is active for the selected nodes */
 	bool IsExecutionModeActive(EPCGSettingsExecutionMode InExecutionMode) const;
diff --git a/Engine/Plugins/Experimental/PCG/Source/PCGEditor/Private/PCGEditorCommands.cpp b/Engine/Plugins/Experimental/PCG/Source/PCGEditor/Private/PCGEditorCommands.cpp
--- a/Engine/Plugins/Experimental/PCG/Source/PCGEditor/Private/PCGEditorCommands.cpp
+++ b/Engine/Plugins/Experimental/PCG/Source/PCGEditor/Private/PCGEditorCommands.cpp
@@ -22,7 +22,8 @@ void FPCGEditorCommands::RegisterCommands()
 	UI_COMMAND(Find, "Find", "Finds PCG nodes and comments in the current graph", EUserInterfaceActionType::Button, FInputChord(EModifierKey::Control, EKeys::F));
 	UI_COMMAND(PauseAutoRegeneration, "Pause Regen", "Pause automatic regeneration of the current graph.", EUserInterfaceActionType::ToggleButton, FInputChord(EModifierKey::Alt, EKeys::R));
 	UI_COMMAND(ForceGraphRegeneration, "Force Regen", "Manually force a regeneration of the current graph.", EUserInterfaceActionType::Button, FInputChord());
-	UI_COMMAND(RunDeterminismTest, "Run Determinism Test", "Evaluate the current node for determinism.", EUserInterfaceActionType::Button, FInputChord(EModifierKey::Alt, EKeys::D));
+	UI_COMMAND(RunDeterminismNodeTest, "Run Determinism Test on Node", "Evaluate the current node for determinism.", EUserInterfaceActionType::Button, FInputChord(EModifierKey::Alt, EKeys::D));
+	UI_COMMAND(RunDeterminismGraphTest, "Graph Determinism Test", "Evaluate the current graph for determinism.", EUserInterfaceActionType::Button, FInputChord());
 	UI_COMMAND(StartInspectNode, "Start Inspecting Node", "Start Inspecting Node", EUserInterfaceActionType::Button, FInputChord(EKeys::I));
 	UI_COMMAND(StopInspectNode, "Stop Inspecting Node", "Stop Inspecting Node", EUserInterfaceActionType::Button, FInputChord(EModifierKey::Shift, EKeys::I));
 	UI_COMMAND(ExecutionModeEnabled, "Execution Mode Enabled", "Set the execution mode for this node to Enabled.", EUserInterfaceActionType::RadioButton, FInputChord());
diff --git a/Engine/Plugins/Experimental/PCG/Source/PCGEditor/Private/PCGEditorCommands.h b/Engine/Plugins/Experimental/PCG/Source/PCGEditor/Private/PCGEditorCommands.h
--- a/Engine/Plugins/Experimental/PCG/Source/PCGEditor/Private/PCGEditorCommands.h
+++ b/Engine/Plugins/Experimental/PCG/Source/PCGEditor/Private/PCGEditorCommands.h
@@ -16,9 +16,10 @@ public:
 
 	TSharedPtr<FUICommandInfo> CollapseNodes;
 	TSharedPtr<FUICommandInfo> Find;
-	TSharedPtr<FUICommandInfo> RunDeterminismTest;
 	TSharedPtr<FUICommandInfo> PauseAutoRegeneration;
 	TSharedPtr<FUICommandInfo> ForceGraphRegeneration;
+	TSharedPtr<FUICommandInfo> RunDeterminismNodeTest;
+	TSharedPtr<FUICommandInfo> RunDeterminismGraphTest;
 	TSharedPtr<FUICommandInfo> StartInspectNode;
 	TSharedPtr<FUICommandInfo> StopInspectNode;
 	TSharedPtr<FUICommandInfo> ExecutionModeEnabled;
diff --git a/Engine/Plugins/Experimental/PCG/Source/PCGEditor/Private/PCGEditorGraphNodeBase.cpp b/Engine/Plugins/Experimental/PCG/Source/PCGEditor/Private/PCGEditorGraphNodeBase.cpp
--- a/Engine/Plugins/Experimental/PCG/Source/PCGEditor/Private/PCGEditorGraphNodeBase.cpp
+++ b/Engine/Plugins/Experimental/PCG/Source/PCGEditor/Private/PCGEditorGraphNodeBase.cpp
@@ -132,7 +132,7 @@ void UPCGEditorGraphNodeBase::GetNodeContextMenuActions(UToolMenu* Menu, class U
 
 	{
 		FToolMenuSection& Section = Menu->AddSection("EdGraphSchemaDeterminism", LOCTEXT("DeterminismHeader", "Determinism"));
-		Section.AddMenuEntry(FPCGEditorCommands::Get().RunDeterminismTest,
+		Section.AddMenuEntry(FPCGEditorCommands::Get().RunDeterminismNodeTest,
 			LOCTEXT("Determinism_RunTest", "Validate Determinism on Selection"),
 			LOCTEXT("Determinism_RunTestToolTip", "Run a test to validate the selected nodes for determinism."));
 	}
