Tentative fix for static analysis warning

#rnx

diff --git a/Engine/Plugins/Experimental/PCG/Source/PCG/Private/Tests/Determinism/PCGDeterminismTestsCommon.cpp b/Engine/Plugins/Experimental/PCG/Source/PCG/Private/Tests/Determinism/PCGDeterminismTestsCommon.cpp
--- a/Engine/Plugins/Experimental/PCG/Source/PCG/Private/Tests/Determinism/PCGDeterminismTestsCommon.cpp
+++ b/Engine/Plugins/Experimental/PCG/Source/PCG/Private/Tests/Determinism/PCGDeterminismTestsCommon.cpp
@@ -88,7 +88,7 @@ namespace PCGDeterminismTests
 		auto ScheduleAndWaitForExecution = [&Executor](FPCGTaskId FinalTaskID)
 		{
 			// TODO: Consider randomizing/iterating through possible input orders
-			bool bTasksComplete = false;
+			volatile bool bTasksComplete = false;
 
 			// Clear the cache between runs
 			Executor.GetCache().ClearCache();
