PCG WP: Merge editor flow with runtime flow

* PCG Components and PCG Partition Actors are now registered to the subsystem upon creation/load in editor. But contrary to the runtime flow, the mapping Component <-> Partition Actor won't be created.
* Creating/Moving a PCGComponent will create the Partition Actors
* When generating, we create the mapping Component <-> Partition Actor.
* When CleaningUp, we delete the mapping.
* Also add the option in editor to delete all unused PA (aka the ones that intersect with no PCG Component bounds).

diff --git a/Engine/Plugins/Experimental/PCG/Source/PCG/Private/Grid/PCGPartitionActor.cpp b/Engine/Plugins/Experimental/PCG/Source/PCG/Private/Grid/PCGPartitionActor.cpp
--- a/Engine/Plugins/Experimental/PCG/Source/PCG/Private/Grid/PCGPartitionActor.cpp
+++ b/Engine/Plugins/Experimental/PCG/Source/PCG/Private/Grid/PCGPartitionActor.cpp
@@ -29,6 +29,11 @@ APCGPartitionActor::APCGPartitionActor(const FObjectInitializer& ObjectInitializ
 #endif // WITH_EDITOR
 }
 
+UPCGSubsystem* APCGPartitionActor::GetSubsystem() const 
+{
+	return GetWorld() ? GetWorld()->GetSubsystem<UPCGSubsystem>() : nullptr;
+}
+
 void APCGPartitionActor::PostLoad()
 {
 	Super::PostLoad();
@@ -42,18 +47,69 @@ void APCGPartitionActor::PostLoad()
 	// Make sure that we don't track objects that do not exist anymore
 	CleanupDeadGraphInstances();
 
+#if WITH_EDITOR
+	// Mark all our local components as local
+	for (UPCGComponent* LocalComponent : GetAllLocalPCGComponents())
+	{
+		LocalComponent->MarkAsLocalComponent();
+	}
+
+	bWasPostCreatedLoaded = true;
+#endif // WITH_EDITOR
+}
+
+void APCGPartitionActor::BeginDestroy()
+{
+#if WITH_EDITOR
+	if (!PCGHelpers::IsRuntimeOrPIE() && GetSubsystem() && PCGGridSize > 0)
+	{
+		GetSubsystem()->UnregisterPartitionActor(this);
+	}
+#endif // WITH_EDITOR
+
+	Super::BeginDestroy();
+}
+
+void APCGPartitionActor::Destroyed()
+{
+#if WITH_EDITOR
+	if (!PCGHelpers::IsRuntimeOrPIE() && GetSubsystem() && PCGGridSize > 0)
+	{
+		GetSubsystem()->UnregisterPartitionActor(this);
+	}
+#endif // WITH_EDITOR
+
+	Super::Destroyed();
+}
+
+void APCGPartitionActor::PostRegisterAllComponents()
+{
+	Super::PostRegisterAllComponents();
+
 #if WITH_EDITOR
 	if (BoundsComponent)
 	{
 		BoundsComponent->SetBoxExtent(GetFixedBounds().GetExtent());
 	}
+#endif // WITH_EDITOR
 
-	// Mark all our local components as local
-	for (UPCGComponent* LocalComponent : GetAllLocalPCGComponents())
+	// Make the Partition actor register itself to the PCG Subsystem
+	// Always do it at runtime, wait for the post load/creation in editor
+	// Only do the mapping if we are at runtime.
+
+	// Make sure the PCGGridSize is not 0, otherwise it will break everything
+
+	const bool bIsRuntimeOrPIE = PCGHelpers::IsRuntimeOrPIE();
+	if ((bIsRuntimeOrPIE 
+#if WITH_EDITOR
+		|| bWasPostCreatedLoaded
+#endif // WITH_EDITOR
+		) && ensure(PCGGridSize > 0))
 	{
-		LocalComponent->MarkAsLocalComponent();
+		UPCGSubsystem* Subsystem = GetSubsystem();
+		check(Subsystem);
+		Subsystem->RegisterPartitionActor(this, /*bDoComponentMapping*/ bIsRuntimeOrPIE);
 	}
-#endif // WITH_EDITOR
 }
 
 void APCGPartitionActor::BeginPlay()
@@ -76,24 +132,12 @@ void APCGPartitionActor::BeginPlay()
 		}
 	}
 
-	// Make the Partition actor register itself to the PCG Subsystem
-	UWorld* World = GetWorld();
-	check(World);
-
-	UPCGSubsystem* Subsystem = World->GetSubsystem<UPCGSubsystem>();
-	check(Subsystem);
-
-	Subsystem->RegisterPartitionActor(this);
-
 	Super::BeginPlay();
 }
 
 void APCGPartitionActor::EndPlay(const EEndPlayReason::Type EndPlayReason)
 {
-	UWorld* World = GetWorld();
-	check(World);
-
-	UPCGSubsystem* Subsystem = World->GetSubsystem<UPCGSubsystem>();
+	UPCGSubsystem* Subsystem = GetSubsystem();
 	check(Subsystem);
 
 	Subsystem->UnregisterPartitionActor(this);
@@ -158,42 +202,59 @@ UPCGComponent* APCGPartitionActor::GetOriginalComponent(const UPCGComponent* Loc
 	return OriginalComponent ? (*OriginalComponent).Get() : nullptr;
 }
 
-bool APCGPartitionActor::CleanupDeadGraphInstances()
+void APCGPartitionActor::CleanupDeadGraphInstances()
 {
-	TSet<TObjectPtr<UPCGComponent>> DeadLocalInstances;
+	// First find if we have any local dead instance (= nullptr) hooked to an original component.
+	TSet<TObjectPtr<UPCGComponent>> DeadOriginalInstances;
+	for (const auto& OriginalToLocalItem : OriginalToLocalMap)
+	{
+		if (!OriginalToLocalItem.Value)
+		{
+			DeadOriginalInstances.Add(OriginalToLocalItem.Key);
+		}
+	}
 
-	// Note: since we might end up with multiple nulls in the original to local map
-	// it might not be very stable to use it; we'll use the 
-	for (const auto& LocalToOriginalItem : LocalToOriginalMap)
+	if (!DeadOriginalInstances.IsEmpty())
 	{
-		if (!LocalToOriginalItem.Value)
+		Modify();
+
+		for (const TObjectPtr<UPCGComponent>& DeadInstance : DeadOriginalInstances)
 		{
-			DeadLocalInstances.Add(LocalToOriginalItem.Key);
+			OriginalToLocalMap.Remove(DeadInstance);
 		}
+
+		LocalToOriginalMap.Remove(nullptr);
 	}
 
-	if (DeadLocalInstances.Num() == 0)
+	// And do the same with dead original ones.
+	TSet<TObjectPtr<UPCGComponent>> DeadLocalInstances;
+	for (const auto& LocalToOriginalItem : LocalToOriginalMap)
 	{
-		return OriginalToLocalMap.IsEmpty();
+		if (!LocalToOriginalItem.Value)
+		{
+			DeadLocalInstances.Add(LocalToOriginalItem.Key);
+		}
 	}
 
-	Modify();
 
-	for (const TObjectPtr<UPCGComponent>& DeadInstance : DeadLocalInstances)
+	if (!DeadLocalInstances.IsEmpty())
 	{
-		LocalToOriginalMap.Remove(DeadInstance);
+		Modify();
 
-		if (DeadInstance)
+		for (const TObjectPtr<UPCGComponent>& DeadInstance : DeadLocalInstances)
 		{
-			DeadInstance->CleanupLocal(/*bRemoveComponents=*/true);
-			DeadInstance->DestroyComponent();
-		}
-	}
+			LocalToOriginalMap.Remove(DeadInstance);
 
-	// Remove all dead entries
-	OriginalToLocalMap.Remove(nullptr);
+			if (DeadInstance)
+			{
+				DeadInstance->CleanupLocalImmediate(/*bRemoveComponents=*/true);
+				DeadInstance->DestroyComponent();
+			}
+		}
 
-	return OriginalToLocalMap.IsEmpty();
+		// Remove all dead entries
+		OriginalToLocalMap.Remove(nullptr);
+	}
 }
 
 void APCGPartitionActor::AddGraphInstance(UPCGComponent* OriginalComponent)
@@ -289,6 +350,14 @@ void APCGPartitionActor::PostCreation()
 	{
 		BoundsComponent->SetBoxExtent(GetFixedBounds().GetExtent());
 	}
+
+	// Make sure PCGGrid size if greater than 0, otherwise it will break everything
+	if (!PCGHelpers::IsRuntimeOrPIE() && GetSubsystem() && ensure(PCGGridSize > 0))
+	{
+		GetSubsystem()->RegisterPartitionActor(this, /*bDoComponentMapping=*/ false);
+	}
+
+	bWasPostCreatedLoaded = true;
 #endif // WITH_EDITOR
 }
 
@@ -297,7 +366,7 @@ bool APCGPartitionActor::IsSafeForDeletion() const
 	ensure(IsInGameThread());
 	for (TObjectPtr<UPCGComponent> PCGComponent : GetAllOriginalPCGComponents())
 	{
-		if (PCGComponent && PCGComponent->IsGenerating())
+		if (PCGComponent && (PCGComponent->IsGenerating() || PCGComponent->IsCleaningUp()))
 		{
 			return false;
 		}
diff --git a/Engine/Plugins/Experimental/PCG/Source/PCG/Private/Grid/PCGPartitionActor.h b/Engine/Plugins/Experimental/PCG/Source/PCG/Private/Grid/PCGPartitionActor.h
--- a/Engine/Plugins/Experimental/PCG/Source/PCG/Private/Grid/PCGPartitionActor.h
+++ b/Engine/Plugins/Experimental/PCG/Source/PCG/Private/Grid/PCGPartitionActor.h
@@ -8,6 +8,7 @@
 #include "PCGPartitionActor.generated.h"
 
 class UPCGComponent;
+class UPCGSubsystem;
 class UBoxComponent;
 
 /** 
@@ -24,12 +25,15 @@ public:
 
 	//~Begin UObject Interface
 	virtual void PostLoad() override;
+	virtual void BeginDestroy() override;
+	virtual void Destroyed() override;
 	//~End UObject Interface
 
 	//~Begin AActor Interface
 	virtual void BeginPlay();
 	virtual void EndPlay(const EEndPlayReason::Type EndPlayReason) override;
 	virtual void GetActorBounds(bool bOnlyCollidingComponents, FVector& Origin, FVector& BoxExtent, bool bIncludeFromChildActors) const override;
+	virtual void PostRegisterAllComponents() override;
 #if WITH_EDITOR
 	virtual FBox GetStreamingBounds() const override;
 	virtual AActor* GetSceneOutlinerParent() const override;
@@ -51,7 +55,7 @@ public:
 
 	void AddGraphInstance(UPCGComponent* OriginalComponent);
 	bool RemoveGraphInstance(UPCGComponent* OriginalComponent);
-	bool CleanupDeadGraphInstances();
+	void CleanupDeadGraphInstances();
 
 #if WITH_EDITOR
 	/** To be called after the creation of a new actor to copy the GridSize property (Editor only) into the PCGGridSize property */
@@ -79,6 +83,8 @@ public:
 	FGuid PCGGuid;
 
 private:
+	UPCGSubsystem* GetSubsystem() const;
+
 	// TODO: Make these properties editor only (see comment before).
 	UPROPERTY()
 	TMap<TObjectPtr<UPCGComponent>, TObjectPtr<UPCGComponent>> OriginalToLocalMap;
@@ -96,5 +102,8 @@ private:
 	/** Box component to draw the Partition actor bounds in the Editor viewport */
 	UPROPERTY(Transient)
 	TObjectPtr<UBoxComponent> BoundsComponent;
+
+	/** Utility bool to check if PostCreation/PostLoad was called. */
+	bool bWasPostCreatedLoaded = false;
 #endif // WITH_EDITORONLY_DATA
 };
diff --git a/Engine/Plugins/Experimental/PCG/Source/PCG/Private/PCGComponent.cpp b/Engine/Plugins/Experimental/PCG/Source/PCG/Private/PCGComponent.cpp
--- a/Engine/Plugins/Experimental/PCG/Source/PCG/Private/PCGComponent.cpp
+++ b/Engine/Plugins/Experimental/PCG/Source/PCG/Private/PCGComponent.cpp
@@ -54,6 +54,32 @@ bool UPCGComponent::IsPartitioned() const
 	return bIsPartitioned && CanPartition();
 }
 
+void UPCGComponent::SetIsPartitioned(bool bIsNowPartitioned)
+{
+	if (bIsNowPartitioned == bIsPartitioned)
+	{
+		return;
+	}
+
+	bool bDoActorMapping = bGenerated || PCGHelpers::IsRuntimeOrPIE();
+
+	if (bGenerated)
+	{
+		CleanupLocalImmediate(/*bRemoveComponents=*/true);
+	}
+
+	if (bIsNowPartitioned)
+	{
+		bIsPartitioned = bIsNowPartitioned;
+		GetSubsystem()->RegisterOrUpdatePCGComponent(this, bDoActorMapping);
+	}
+	else
+	{
+		GetSubsystem()->UnregisterPCGComponent(this);
+		bIsPartitioned = bIsNowPartitioned;
+	}
+}
+
 void UPCGComponent::SetGraph_Implementation(UPCGGraph* InGraph)
 {
 	SetGraphLocal(InGraph);
@@ -210,6 +236,8 @@ FPCGTaskId UPCGComponent::GenerateInternal(bool bForce, EPCGComponentGenerationT
 		return InvalidPCGTaskId;
 	}
 
+	Modify();
+
 	CurrentGenerationTask = GetSubsystem()->ScheduleComponent(this, /*bSave=*/bForce, Dependencies);
 
 	return CurrentGenerationTask;
@@ -230,8 +258,6 @@ FPCGTaskId UPCGComponent::CreateGenerateTask(bool bForce, const TArray<FPCGTaskI
 	}
 #endif
 
-	Modify();
-
 	// Keep track of all the dependencies
 	TArray<FPCGTaskId> AdditionalDependencies;
 	const TArray<FPCGTaskId>* AllDependencies = &Dependencies;
@@ -322,6 +348,7 @@ void UPCGComponent::PostCleanupGraph()
 
 #if WITH_EDITOR
 	OnPCGGraphCleanedDelegate.Broadcast(this);
+	bDirtyGenerated = false;
 #endif
 }
 
@@ -331,7 +358,6 @@ void UPCGComponent::OnProcessGraphAborted()
 
 	CleanupUnusedManagedResources();
 
-	bGenerated = false;
 	CurrentGenerationTask = InvalidPCGTaskId;
 
 #if WITH_EDITOR
@@ -370,6 +396,8 @@ FPCGTaskId UPCGComponent::CleanupInternal(bool bRemoveComponents, bool bSave, co
 		return InvalidPCGTaskId;
 	}
 
+	Modify();
+
 	CurrentCleanupTask = GetSubsystem()->ScheduleCleanup(this, bRemoveComponents, bSave, Dependencies);
 	return CurrentCleanupTask;
 }
@@ -470,7 +498,7 @@ bool UPCGComponent::MoveResourcesToNewActor(AActor* InNewActor, bool bCreateChil
 
 void UPCGComponent::CleanupLocalImmediate(bool bRemoveComponents)
 {
-	if (!bGenerated || IsPartitioned() || IsCleaningUp() || IsGenerating())
+	if (!bGenerated || IsCleaningUp() || IsGenerating())
 	{
 		return;
 	}
@@ -490,6 +518,13 @@ void UPCGComponent::CleanupLocalImmediate(bool bRemoveComponents)
 	}
 
 	UPCGActorHelpers::DeleteActors(GetWorld(), ActorsToDelete.Array());
+
+	// If bRemoveComponents is true, it means we are in a "real" cleanup, not a pre-cleanup before a generate.
+	// So call PostCleanup in this case.
+	if (bRemoveComponents)
+	{
+		PostCleanupGraph();
+	}
 }
 
 FPCGTaskId UPCGComponent::CreateCleanupTask(bool bRemoveComponents, const TArray<FPCGTaskId>& Dependencies)
@@ -596,7 +631,7 @@ void UPCGComponent::BeginPlay()
 	// First if it is partitioned, register itself to the PCGSubsystem, to map the component to all its corresponding PartitionActors
 	if (IsPartitioned() && GetSubsystem())
 	{
-		GetSubsystem()->RegisterPCGComponent(this);
+		GetSubsystem()->RegisterOrUpdatePCGComponent(this);
 	}
 
 	if(bActivated && !bGenerated && GenerationTrigger == EPCGComponentGenerationTrigger::GenerateOnLoad)
@@ -638,6 +673,12 @@ void UPCGComponent::OnComponentCreated()
 
 #if WITH_EDITOR
 	SetupActorCallbacks();
+
+	if (!PCGHelpers::IsRuntimeOrPIE() && IsPartitioned())
+	{
+		check(GetSubsystem());
+		GetSubsystem()->RegisterOrUpdatePCGComponent(this, /*bDoActorMapping=*/bGenerated);
+	}
 #endif
 }
 
@@ -653,6 +694,11 @@ void UPCGComponent::OnComponentDestroyed(bool bDestroyingHierarchy)
 		{
 			Subsystem->CleanupPartitionActors(LastGeneratedBounds);
 		}
+
+		if (!PCGHelpers::IsRuntimeOrPIE() && IsPartitioned())
+		{
+			Subsystem->UnregisterPCGComponent(this);
+		}
 	}
 #endif
 
@@ -720,12 +766,24 @@ void UPCGComponent::SetupCallbacksOnCreation()
 	{
 		Graph->OnGraphChangedDelegate.AddUObject(this, &UPCGComponent::OnGraphChanged);
 	}
+
+	if (!PCGHelpers::IsRuntimeOrPIE() && IsPartitioned())
+	{
+		check(GetSubsystem());
+		GetSubsystem()->RegisterOrUpdatePCGComponent(this, /*bDoActorMapping=*/bGenerated);
+	}
 }
 #endif
 
 void UPCGComponent::BeginDestroy()
 {
 #if WITH_EDITOR
+	if (!PCGHelpers::IsRuntimeOrPIE() && IsPartitioned())
+	{
+		check(GetSubsystem());
+		GetSubsystem()->UnregisterPCGComponent(this);
+	}
+
 	if (Graph)
 	{
 		Graph->OnGraphChangedDelegate.RemoveAll(this);
@@ -842,42 +900,20 @@ void UPCGComponent::PostEditChangeProperty(FPropertyChangedEvent& PropertyChange
 	{
 		if (CanPartition())
 		{
-			if (bActivated)
-			{
-				bool bIsNowPartitioned = bIsPartitioned;
-				if (bGenerated)
-				{
-					bIsPartitioned = !bIsPartitioned;
-
-					// First, we'll cleanup
-					bActivated = false;
-					Refresh();
-
-					// Then invalidate the previous bounds to force actor creation (as if we moved the volume)
-					// and do a normal refresh
-					bActivated = true;
-					bIsPartitioned = bIsNowPartitioned;
-					ResetLastGeneratedBounds();
-					DirtyGenerated();
-					Refresh();
-				}
-				else
-				{
-					// We need the component to be partitioned if we use the subsystem.
-					bIsPartitioned = true;
+			// At this point, bIsPartitioned is already set with the new value.
+			// But we need to do some cleanup before
+			// So keep this new value, and take its negation for the cleanup.
+			bool bIsNowPartitioned = bIsPartitioned;
+			bIsPartitioned = !bIsPartitioned;
 
-					if (bIsNowPartitioned)
-					{
-						GetSubsystem()->SchedulePartitionGraph(this);
-					}
-					else
-					{
-						GetSubsystem()->ScheduleUnpartitionGraph(this);
-
-					}
+			// SetIsPartioned cleans up before, so keep track if we were generated or not.
+			bool bWasGenerated = bGenerated;
+			SetIsPartitioned(bIsNowPartitioned);
 
-					bIsPartitioned = bIsNowPartitioned;
-				}
+			// And finally, re-generate if we were generated and activated
+			if (bWasGenerated && bActivated)
+			{
+				GenerateLocal(/*bForce=*/false);
 			}
 		}
 		else
@@ -950,7 +986,7 @@ void UPCGComponent::PreEditUndo()
 	if (bGenerated)
 	{
 		// Cleanup so managed resources are cleaned in all cases
-		CleanupLocal(/*bRemoveComponents=*/true, /*bSave=*/PCGComponent::bSaveOnCleanupAndGenerate);
+		CleanupLocalImmediate(/*bRemoveComponents=*/true);
 		// Put back generated flag to its original value so it is captured properly
 		bGenerated = true;
 	}	
@@ -1418,30 +1454,20 @@ void UPCGComponent::Refresh()
 	}
 #endif
 
-	//check(IsValid(this));
+	// Before doing a refresh, update the component to the subsystem if we are partitioned
+	// Only redo the mapping if we are generated
+	if (IsPartitioned())
+	{
+		GetSubsystem()->RegisterOrUpdatePCGComponent(this, /*bDoActorMapping=*/ bGenerated);
+	}
 
 	// Following a change in some properties or in some spatial information related to this component,
-	// We need to regenerate the graph, depending of the state in the editor.
-	// In the case of a non-partitioned graph, we need to generate the graph only if it was previously generated & tagged for regeneration
-	// In the partitioned graph case, however, we need to do a bit more:
-	// 1. Regenerate the graph if it was previously generated & tagged for regeneration;
-	//  notice that the associated partition actors will not (and should not) have the regenerate flag on.
-	// 2. Otherwise, we need to update the partitioning if the spatial data has changed.
+	// We need to regenerate/cleanup the graph, depending of the state in the editor.
 	if (!bActivated)
 	{
-		if (IsPartitioned() && GetSubsystem())
-		{
-			if (LastGeneratedBounds.IsValid)
-			{
-				GetSubsystem()->ScheduleUnpartitionGraph(this);
-			}
-		}
-		else
-		{
-			bool bWasGenerated = bGenerated;
-			CleanupLocalImmediate(/*bRemoveComponents=*/true);
-			bGenerated = bWasGenerated;
-		}
+		bool bWasGenerated = bGenerated;
+		CleanupLocalImmediate(/*bRemoveComponents=*/true);
+		bGenerated = bWasGenerated;
 	}
 	else
 	{
@@ -1449,10 +1475,6 @@ void UPCGComponent::Refresh()
 		{
 			GenerateLocal(/*bForce=*/false);
 		}
-		else if (IsPartitioned() && GetSubsystem())
-		{
-			GetSubsystem()->SchedulePartitionGraph(this);
-		}
 	}
 }
 #endif
@@ -2161,6 +2183,19 @@ void FPCGComponentInstanceData::ApplyToComponent(UActorComponent* Component, con
 		{
 			PCGComponent->SetManagedResources(DuplicatedResources);
 		}
+
+		// We need to make sure that the component is registered to the subsystem if it is partitioned
+		// or unregister it if it is not.
+		if (PCGComponent->IsPartitioned())
+		{
+			// Mapping needs to be done if we are at runtime or previously generated
+			bool bDoActorMapping = PCGComponent->bGenerated || PCGHelpers::IsRuntimeOrPIE();
+			PCGComponent->GetSubsystem()->RegisterOrUpdatePCGComponent(PCGComponent, bDoActorMapping);
+		}
+		else
+		{
+			PCGComponent->GetSubsystem()->UnregisterPCGComponent(PCGComponent);
+		}
 	}
 }
 
diff --git a/Engine/Plugins/Experimental/PCG/Source/PCG/Private/PCGSubsystem.cpp b/Engine/Plugins/Experimental/PCG/Source/PCG/Private/PCGSubsystem.cpp
--- a/Engine/Plugins/Experimental/PCG/Source/PCG/Private/PCGSubsystem.cpp
+++ b/Engine/Plugins/Experimental/PCG/Source/PCG/Private/PCGSubsystem.cpp
@@ -39,8 +39,22 @@ namespace PCGSubsystemConsole
 				}
 			}));
 }
+
+namespace PCGSubsystem
+{
+	FPCGTaskId ForEachIntersectingCell(FPCGGraphExecutor* GraphExecutor, UWorld* World, const FBox& InBounds, bool bCreateActor, bool bLoadCell, bool bSaveActors, TFunctionRef<FPCGTaskId(APCGPartitionActor*, const FBox&, const TArray<FPCGTaskId>&)> InOperation);
+}
 #endif
 
+UPCGSubsystem::UPCGSubsystem()
+	: Super()
+{
+	// TODO: For now we set our octree to be 2km wide, but it would be perhaps better to
+	// scale it to the size of our world.
+	constexpr FVector::FReal OctreeExtent = 200000; // 2km
+	PCGComponentOctree = FPCGComponentOctree(FVector::ZeroVector, OctreeExtent);
+}
+
 void UPCGSubsystem::Deinitialize()
 {
 	// Cancel all tasks
@@ -71,11 +85,6 @@ void UPCGSubsystem::PostInitialize()
 			});
 		}
 	}
-
-	// TODO: For now we set our octree to be 2km wide, but it would be perhaps better to
-	// scale it to the size of our world.
-	constexpr FVector::FReal OctreeExtent = 200000; // 2km
-	PCGComponentOctree = FPCGComponentOctree(FVector::ZeroVector, OctreeExtent);
 }
 
 void UPCGSubsystem::Tick(float DeltaSeconds)
@@ -145,11 +154,16 @@ FPCGTaskId UPCGSubsystem::ScheduleComponent(UPCGComponent* PCGComponent, bool bS
 		return InvalidPCGTaskId;
 	}
 
-	// TODO: Merge both path
 #if WITH_EDITOR
 	if (PCGComponent->IsPartitioned() && !PCGHelpers::IsRuntimeOrPIE())
 	{
-		return ScheduleProcessGraph(PCGComponent, bSave);
+		// In this case create the PA and update the mapping
+		// Note: This is an immediate operation, as we need the PA for the generation.
+		auto ScheduleTask = [](APCGPartitionActor* PCGActor, const FBox& InIntersectedBounds, const TArray<FPCGTaskId>& TaskDependencies) { return InvalidPCGTaskId; };
+
+		PCGSubsystem::ForEachIntersectingCell(GraphExecutor, PCGComponent->GetWorld(), PCGComponent->GetGridBounds(), /*bCreateActor=*/true, /*bLoadCell=*/false, /*bSave=*/false, ScheduleTask);
+
+		UpdateMappingPCGComponentPartitionActor(PCGComponent);
 	}
 #endif // WITH_EDITOR
 
@@ -158,14 +172,29 @@ FPCGTaskId UPCGSubsystem::ScheduleComponent(UPCGComponent* PCGComponent, bool bS
 	// If the component is partitioned, we will forward the calls to its registered PCG Partition actors
 	if (PCGComponent->IsPartitioned())
 	{
+		auto LocalGenerateTask = [OriginalComponent = PCGComponent, &Dependencies](UPCGComponent* LocalComponent)
 		{
-			// TODO: Might be more interesting to copy the set and release the lock.
-			FReadScopeLock ReadLock(ComponentToPartitionActorsMapLock);
-			if (TSet<TObjectPtr<APCGPartitionActor>>* PartitionActorsPtr = ComponentToPartitionActorsMap.Find(PCGComponent))
+			// If the local component is currently generating, it's probably because it was requested by a refresh.
+			// Wait after this one instead
+			if (LocalComponent->IsGenerating())
 			{
-				AllTasks = ScheduleMultipleComponent(PCGComponent, *PartitionActorsPtr, Dependencies);
+				return LocalComponent->CurrentGenerationTask;
 			}
-		}
+
+			// Ensure that the PCG actor match our original
+			LocalComponent->SetPropertiesFromOriginal(OriginalComponent);
+
+			if (LocalComponent->bGenerated && !OriginalComponent->bGenerated)
+			{
+				// Detected a mismatch between the original component and the local component.
+				// Request a cleanup first
+				LocalComponent->CleanupLocalImmediate(true);
+			}
+
+			return LocalComponent->GenerateInternal(/*bForce=*/ false, EPCGComponentGenerationTrigger::GenerateOnDemand, Dependencies);
+		};
+
+		AllTasks = DispatchToRegisteredLocalComponents(PCGComponent, LocalGenerateTask);
 	}
 	else
 	{
@@ -208,29 +237,24 @@ FPCGTaskId UPCGSubsystem::ScheduleCleanup(UPCGComponent* PCGComponent, bool bRem
 		return InvalidPCGTaskId;
 	}
 
-	// TODO: Merge both path
-#if WITH_EDITOR
-	if (PCGComponent->IsPartitioned() && !PCGHelpers::IsRuntimeOrPIE())
-	{
-		return CleanupGraph(PCGComponent, PCGComponent->LastGeneratedBounds, bRemoveComponents, bSave);
-	}
-#endif // WITH_EDITOR
-
 	TArray<FPCGTaskId> AllTasks;
 
 	// If the component is partitioned, we will forward the calls to its registered PCG Partition actors
 	if (PCGComponent->IsPartitioned())
 	{
+		auto LocalCleanupTask = [bRemoveComponents, &Dependencies](UPCGComponent* LocalComponent)
 		{
-			// TODO: Might be more interesting to copy the set and release the lock.
-			FReadScopeLock ReadLock(ComponentToPartitionActorsMapLock);
-			TSet<TObjectPtr<APCGPartitionActor>>* PartitionActorsPtr = ComponentToPartitionActorsMap.Find(PCGComponent);
-
-			if (PartitionActorsPtr)
+			// If the local component is currently cleaning up, it's probably because it was requested by a refresh.
+			// Wait after this one instead
+			if (LocalComponent->IsCleaningUp())
 			{
-				AllTasks = ScheduleMultipleCleanup(PCGComponent, *PartitionActorsPtr, bRemoveComponents, Dependencies);
+				return LocalComponent->CurrentCleanupTask;
 			}
-		}
+
+			return LocalComponent->CleanupInternal(bRemoveComponents, /*bSave=*/ false, Dependencies);
+		};
+
+		AllTasks = DispatchToRegisteredLocalComponents(PCGComponent, LocalCleanupTask);
 	}
 	else
 	{
@@ -241,22 +265,39 @@ FPCGTaskId UPCGSubsystem::ScheduleCleanup(UPCGComponent* PCGComponent, bool bRem
 		}
 	}
 
-	if (AllTasks.IsEmpty())
-	{
-		return InvalidPCGTaskId;
-	}
-	
 	TWeakObjectPtr<UPCGComponent> ComponentPtr(PCGComponent);
-
-	return GraphExecutor->ScheduleGeneric([ComponentPtr]() 
+	auto PostCleanupTask = [this, ComponentPtr]()
+	{
+		if (UPCGComponent* Component = ComponentPtr.Get())
 		{
-			if (UPCGComponent* Component = ComponentPtr.Get())
+			Component->PostCleanupGraph();
+
+#if WITH_EDITOR
+			// If we are in Editor and partitioned, delete the mapping.
+			if (Component->IsPartitioned() && !PCGHelpers::IsRuntimeOrPIE())
 			{
-				Component->PostCleanupGraph();
+				DeleteMappingPCGComponentPartitionActor(Component);
 			}
+#endif // WITH_EDITOR
+		}
 
-			return true;
-		}, AllTasks);
+		return true;
+	};
+
+	FPCGTaskId PostCleanupTaskId = InvalidPCGTaskId;
+
+	// If we have no tasks to do, just call PostCleanup immediatly
+	// otherwise wait for all the tasks to be done to call PostCleanup.
+	if (AllTasks.IsEmpty())
+	{
+		PostCleanupTask();
+	}
+	else
+	{
+		PostCleanupTaskId = GraphExecutor->ScheduleGeneric(PostCleanupTask, AllTasks);
+	}
+
+	return PostCleanupTaskId;
 }
 
 FPCGTaskId UPCGSubsystem::ScheduleGraph(UPCGComponent* SourceComponent, const TArray<FPCGTaskId>& Dependencies)
@@ -299,36 +340,23 @@ FPCGTaskId UPCGSubsystem::ScheduleGeneric(TFunction<bool()> InOperation, const T
 	return GraphExecutor->ScheduleGeneric(InOperation, TaskDependencies);
 }
 
-TArray<FPCGTaskId> UPCGSubsystem::ScheduleMultipleComponent(UPCGComponent* OriginalComponent, TSet<TObjectPtr<APCGPartitionActor>>& PartitionActors, const TArray<FPCGTaskId>& Dependencies)
+TArray<FPCGTaskId> UPCGSubsystem::DispatchToRegisteredLocalComponents(UPCGComponent* OriginalComponent, const TFunction<FPCGTaskId(UPCGComponent*)>& InFunc) const
 {
-	TArray<FPCGTaskId> TaskIds;
-	for (APCGPartitionActor* PartitionActor : PartitionActors)
-	{
-		if (PartitionActor)
-		{
-			if (UPCGComponent* LocalComponent = PartitionActor->GetLocalComponent(OriginalComponent))
-			{
-				// Add check to avoid infinite loop
-				if (ensure(!LocalComponent->IsPartitioned()))
-				{
-					// Ensure that the PCG actor match our original
-					LocalComponent->SetPropertiesFromOriginal(OriginalComponent);
+	TRACE_CPUPROFILER_EVENT_SCOPE(UPCGSubsystem::DispatchToRegisteredLocalComponents);
 
-					FPCGTaskId LocalTask = LocalComponent->GenerateInternal(/*bForce=*/ false, EPCGComponentGenerationTrigger::GenerateOnDemand, Dependencies);
+	// TODO: Might be more interesting to copy the set and release the lock.
+	FReadScopeLock ReadLock(ComponentToPartitionActorsMapLock);
+	const TSet<TObjectPtr<APCGPartitionActor>>* PartitionActorsPtr = ComponentToPartitionActorsMap.Find(OriginalComponent);
 
-					if (LocalTask != InvalidPCGTaskId)
-					{
-						TaskIds.Add(LocalTask);
-					}
-				}
-			}
-		}
+	if (!PartitionActorsPtr)
+	{
+		return TArray<FPCGTaskId>();
 	}
 
-	return TaskIds;
+	return DispatchToLocalComponents(OriginalComponent, *PartitionActorsPtr, InFunc);
 }
 
-TArray<FPCGTaskId> UPCGSubsystem::ScheduleMultipleCleanup(UPCGComponent* OriginalComponent, TSet<TObjectPtr<APCGPartitionActor>>& PartitionActors, bool bRemoveComponents, const TArray<FPCGTaskId>& Dependencies)
+TArray<FPCGTaskId> UPCGSubsystem::DispatchToLocalComponents(UPCGComponent* OriginalComponent, const TSet<TObjectPtr<APCGPartitionActor>>& PartitionActors, const TFunction<FPCGTaskId(UPCGComponent*)>& InFunc) const
 {
 	TArray<FPCGTaskId> TaskIds;
 	for (APCGPartitionActor* PartitionActor : PartitionActors)
@@ -340,7 +368,7 @@ TArray<FPCGTaskId> UPCGSubsystem::ScheduleMultipleCleanup(UPCGComponent* Origina
 				// Add check to avoid infinite loop
 				if (ensure(!LocalComponent->IsPartitioned()))
 				{
-					FPCGTaskId LocalTask = LocalComponent->CleanupInternal(bRemoveComponents, /*bSave=*/ false, Dependencies);
+					FPCGTaskId LocalTask = InFunc(LocalComponent);
 
 					if (LocalTask != InvalidPCGTaskId)
 					{
@@ -360,90 +388,100 @@ bool UPCGSubsystem::GetOutputData(FPCGTaskId TaskId, FPCGDataCollection& OutData
 	return GraphExecutor->GetOutputData(TaskId, OutData);
 }
 
-void UPCGSubsystem::RegisterPCGComponent(UPCGComponent* InComponent)
+bool UPCGSubsystem::RegisterOrUpdatePCGComponent(UPCGComponent* InComponent, bool bDoActorMapping)
 {
 	// Just make sure that we don't register components that are from the PCG Partition Actor
 	check(InComponent && InComponent->GetOwner() && !InComponent->GetOwner()->IsA<APCGPartitionActor>());
 
+	// If the component is not partitioned, don't register/update it.
+	if (!InComponent->IsPartitioned())
+	{
+		return false;
+	}
+
 	// Check also that the bounds are valid. If not early out.
 	if (!InComponent->GetGridBounds().IsValid)
 	{
-		UE_LOG(LogPCG, Error, TEXT("[RegisterPCGComponent] Component has invalid bounds, not registered."));
-		return;
+		UE_LOG(LogPCG, Error, TEXT("[RegisterOrUpdatePCGComponent] Component has invalid bounds, not registered nor updated."));
+		return false;
 	}
 
 	FBox Bounds(EForceInit::ForceInit);
+	bool bComponentHasChanged = false;
+	bool bComponentWasAdded = false;
 
 	{
 		FWriteScopeLock WriteLock(PCGVolumeOctreeLock);
 
-		// We should not register a component twice.
-		check(!ComponentToIdMap.Contains(InComponent));
-
-		FPCGComponentOctreeIDSharedRef IdShared = MakeShared<FPCGComponentOctreeID>();
-		FPCGComponentRef ComponentRef(InComponent, IdShared);
-		Bounds = ComponentRef.Bounds.GetBox();
-		check(Bounds.IsValid);
-		PCGComponentOctree.AddElement(ComponentRef);
+		FPCGComponentOctreeIDSharedRef* ElementIdPtr = ComponentToIdMap.Find(InComponent);
 
-		// Store the shared ptr, because if we add/remove components in the octree, the id might change.
-		// We need to make sure that we always have the latest id for the given component.
-		ComponentToIdMap.Add(InComponent, MoveTemp(IdShared));
-	}
+		if (!ElementIdPtr)
+		{
+			// Does not exist yet, add it.
+			FPCGComponentOctreeIDSharedRef IdShared = MakeShared<FPCGComponentOctreeID>();
+			FPCGComponentRef ComponentRef(InComponent, IdShared);
+			Bounds = ComponentRef.Bounds.GetBox();
+			check(Bounds.IsValid);
+			PCGComponentOctree.AddElement(ComponentRef);
+
+			bComponentHasChanged = true;
+			bComponentWasAdded = true;
+
+			// Store the shared ptr, because if we add/remove components in the octree, the id might change.
+			// We need to make sure that we always have the latest id for the given component.
+			ComponentToIdMap.Add(InComponent, MoveTemp(IdShared));
+		}
+		else
+		{
+			// It already exists, update it if the bounds changed.
 
-	// After adding, do the mapping
-	{
-		FWriteScopeLock WriteLock(ComponentToPartitionActorsMapLock);
+			// Do a copy here.
+			FPCGComponentRef ComponentRef = PCGComponentOctree.GetElementById((*ElementIdPtr)->Id);
+			FBox PreviousBounds = ComponentRef.Bounds.GetBox();
+			ComponentRef.UpdateBounds();
+			Bounds = ComponentRef.Bounds.GetBox();
+			check(Bounds.IsValid);
 
-		TSet<TObjectPtr<APCGPartitionActor>>& PartitionActors = ComponentToPartitionActorsMap.Emplace(InComponent);
-		FindAllIntersectingPartitionActors(Bounds, [&PartitionActors, InComponent](APCGPartitionActor* Actor)
+			// If bounds changed, remove and re-add to the octree
+			if (!PreviousBounds.Equals(Bounds))
 			{
-				PartitionActors.Add(Actor);
-				Actor->AddGraphInstance(InComponent);
-			});
+				PCGComponentOctree.RemoveElement((*ElementIdPtr)->Id);
+				PCGComponentOctree.AddElement(ComponentRef);
+				bComponentHasChanged = true;
+			}
+		}
 	}
-}
 
-void UPCGSubsystem::UpdatePCGComponentBounds(UPCGComponent* InComponent)
-{
-	check(InComponent);
-	FBox Bounds;
-
-	// Remove it and add it again to the octree
+#if WITH_EDITOR
+	// In Editor only, we will create new partition actors depending on the new bounds
+	// TODO: For now it will always create the PA. But if we want to create them only when we generate, we need to make
+	// sure to update the runtime flow, for them to also create PA if they need to.
+	if (!PCGHelpers::IsRuntimeOrPIE())
 	{
-		FWriteScopeLock WriteLock(PCGVolumeOctreeLock);
+		auto ScheduleTask = [](APCGPartitionActor* PCGActor, const FBox& InIntersectedBounds, const TArray<FPCGTaskId>& TaskDependencies) { return InvalidPCGTaskId; };
 
-		FPCGComponentOctreeIDSharedRef* ElementIdPtr = ComponentToIdMap.Find(InComponent);
+		PCGSubsystem::ForEachIntersectingCell(GraphExecutor, InComponent->GetWorld(), Bounds, /*bCreateActor=*/true, /*bLoadCell=*/false, /*bSave=*/false, ScheduleTask);
+	}
+#endif // WITH_EDITOR
 
-		if (!ElementIdPtr)
+	// After adding/updating, try to do the mapping (if we asked for it and the component changed)
+	if (bDoActorMapping)
+	{
+		if (bComponentHasChanged)
 		{
-			return;
+			UpdateMappingPCGComponentPartitionActor(InComponent);
 		}
-
-		FPCGComponentRef ComponentRef = PCGComponentOctree.GetElementById((*ElementIdPtr)->Id);
-		PCGComponentOctree.RemoveElement((*ElementIdPtr)->Id);
-
-		ComponentRef.UpdateBounds();
-		Bounds = ComponentRef.Bounds.GetBox();
-
-		PCGComponentOctree.AddElement(ComponentRef);
 	}
-
-	// Redo the mapping between components and PCG PartitionActors
+	else
 	{
-		FWriteScopeLock WriteLock(ComponentToPartitionActorsMapLock);
-		TSet<TObjectPtr<APCGPartitionActor>>* PartitionActorsPtr = ComponentToPartitionActorsMap.Find(InComponent);
-		if (ensure(PartitionActorsPtr))
+		if (!bComponentWasAdded)
 		{
-			PartitionActorsPtr->Empty();
-			FindAllIntersectingPartitionActors(Bounds, [PartitionActorsPtr, InComponent](APCGPartitionActor* Actor)
-				{
-					PartitionActorsPtr->Add(Actor);
-					Actor->AddGraphInstance(InComponent);
-				});
+			// If we do not want a mapping, delete the existing one
+			DeleteMappingPCGComponentPartitionActor(InComponent);
 		}
 	}
 
+	return bComponentHasChanged;
 }
 
 void UPCGSubsystem::UnregisterPCGComponent(UPCGComponent* InComponent)
@@ -478,24 +516,21 @@ void UPCGSubsystem::UnregisterPCGComponent(UPCGComponent* InComponent)
 	}
 
 	{
-		FReadScopeLock ReadLock(ComponentToPartitionActorsMapLock);
+		FWriteScopeLock WriteLock(ComponentToPartitionActorsMapLock);
 		TSet<TObjectPtr<APCGPartitionActor>>* PartitionActorsPtr = ComponentToPartitionActorsMap.Find(InComponent);
-		if (ensure(PartitionActorsPtr))
+		if (PartitionActorsPtr)
 		{
 			for (APCGPartitionActor* Actor : *PartitionActorsPtr)
 			{
 				Actor->RemoveGraphInstance(InComponent);
 			}
-		}
-	}
 
-	{
-		FWriteScopeLock WriteLock(ComponentToPartitionActorsMapLock);
-		ComponentToPartitionActorsMap.Remove(InComponent);
+			ComponentToPartitionActorsMap.Remove(InComponent);
+		}
 	}
 }
 
-void UPCGSubsystem::FindAllIntersectingComponents(const FBoxCenterAndExtent& InBounds, TFunction<void(UPCGComponent*)> InFunc) const
+void UPCGSubsystem::ForAllIntersectingComponents(const FBoxCenterAndExtent& InBounds, TFunction<void(UPCGComponent*)> InFunc) const
 {
 	FReadScopeLock ReadLock(PCGVolumeOctreeLock);
 
@@ -505,7 +540,7 @@ void UPCGSubsystem::FindAllIntersectingComponents(const FBoxCenterAndExtent& InB
 		});
 }
 
-void UPCGSubsystem::RegisterPartitionActor(APCGPartitionActor* Actor)
+void UPCGSubsystem::RegisterPartitionActor(APCGPartitionActor* Actor, bool bDoComponentMapping)
 {
 	check(Actor);
 
@@ -524,13 +559,17 @@ void UPCGSubsystem::RegisterPartitionActor(APCGPartitionActor* Actor)
 	// And then register itself to all the components that intersect with it
 	{
 		FWriteScopeLock WriteLock(ComponentToPartitionActorsMapLock);
-		FindAllIntersectingComponents(FBoxCenterAndExtent(Actor->GetFixedBounds()), [this, Actor](UPCGComponent* Component)
+		ForAllIntersectingComponents(FBoxCenterAndExtent(Actor->GetFixedBounds()), [this, Actor, bDoComponentMapping](UPCGComponent* Component)
 			{
-				TSet<TObjectPtr<APCGPartitionActor>>* PartitionActorsPtr = ComponentToPartitionActorsMap.Find(Component);
-				if (ensure(PartitionActorsPtr))
+				// For each component, do the mapping if we ask it explicitly, or if the component is generated
+				if (bDoComponentMapping || Component->bGenerated)
 				{
-					Actor->AddGraphInstance(Component);
-					PartitionActorsPtr->Add(Actor);
+					TSet<TObjectPtr<APCGPartitionActor>>* PartitionActorsPtr = ComponentToPartitionActorsMap.Find(Component);
+					if (ensure(PartitionActorsPtr))
+					{
+						Actor->AddGraphInstance(Component);
+						PartitionActorsPtr->Add(Actor);
+					}
 				}
 			});
 	}
@@ -550,10 +589,10 @@ void UPCGSubsystem::UnregisterPartitionActor(APCGPartitionActor* Actor)
 	// And then unregister itself to all the components that intersect with it
 	{
 		FWriteScopeLock WriteLock(ComponentToPartitionActorsMapLock);
-		FindAllIntersectingComponents(FBoxCenterAndExtent(Actor->GetFixedBounds()), [this, Actor](UPCGComponent* Component)
+		ForAllIntersectingComponents(FBoxCenterAndExtent(Actor->GetFixedBounds()), [this, Actor](UPCGComponent* Component)
 			{
 				TSet<TObjectPtr<APCGPartitionActor>>* PartitionActorsPtr = ComponentToPartitionActorsMap.Find(Component);
-				if (ensure(PartitionActorsPtr))
+				if (PartitionActorsPtr)
 				{
 					Actor->RemoveGraphInstance(Component);
 					PartitionActorsPtr->Remove(Actor);
@@ -562,7 +601,7 @@ void UPCGSubsystem::UnregisterPartitionActor(APCGPartitionActor* Actor)
 	}
 }
 
-void UPCGSubsystem::FindAllIntersectingPartitionActors(const FBox& InBounds, TFunction<void(APCGPartitionActor*)> InFunc) const
+void UPCGSubsystem::ForAllIntersectingPartitionActors(const FBox& InBounds, TFunction<void(APCGPartitionActor*)> InFunc) const
 {
 
 	// No PCGWorldActor just early out. Same for invalid bounds.
@@ -604,6 +643,95 @@ void UPCGSubsystem::FindAllIntersectingPartitionActors(const FBox& InBounds, TFu
 	}
 }
 
+void UPCGSubsystem::UpdateMappingPCGComponentPartitionActor(UPCGComponent* InComponent)
+{
+	check(InComponent);
+
+	// Get the bounds
+	FBox Bounds(EForceInit::ForceInit);
+	{
+		FReadScopeLock OctreeReadLock(PCGVolumeOctreeLock);
+
+		FPCGComponentOctreeIDSharedRef* ElementIdPtr = ComponentToIdMap.Find(InComponent);
+
+		if (ensure(ElementIdPtr))
+		{
+			const FPCGComponentRef& ComponentRef = PCGComponentOctree.GetElementById((*ElementIdPtr)->Id);
+			Bounds = ComponentRef.Bounds.GetBox();
+		}
+		else
+		{
+			return;
+		}
+	}
+
+	if (!Bounds.IsValid)
+	{
+		return;
+	}
+
+	{
+		FWriteScopeLock WriteLock(ComponentToPartitionActorsMapLock);
+
+		TSet<TObjectPtr<APCGPartitionActor>>* PartitionActorsPtr = ComponentToPartitionActorsMap.Find(InComponent);
+
+		if (!PartitionActorsPtr)
+		{
+			// Does not yet exists, add it
+			PartitionActorsPtr = &ComponentToPartitionActorsMap.Emplace(InComponent);
+			check(PartitionActorsPtr);
+		}
+
+		TSet<TObjectPtr<APCGPartitionActor>> NewMapping;
+		ForAllIntersectingPartitionActors(Bounds, [&NewMapping, InComponent](APCGPartitionActor* Actor)
+		{
+			Actor->AddGraphInstance(InComponent);
+			NewMapping.Add(Actor);
+		});
+
+		// Find the ones that were removed
+		TSet<TObjectPtr<APCGPartitionActor>> RemovedActors = PartitionActorsPtr->Difference(NewMapping);
+
+		for (APCGPartitionActor* RemovedActor : RemovedActors)
+		{
+			RemovedActor->RemoveGraphInstance(InComponent);
+		}
+
+		*PartitionActorsPtr = MoveTemp(NewMapping);
+	}
+}
+
+void UPCGSubsystem::DeleteMappingPCGComponentPartitionActor(UPCGComponent* InComponent)
+{
+	check(InComponent);
+
+	if (!InComponent->IsPartitioned())
+	{
+		return;
+	}
+
+	FWriteScopeLock WriteLock(ComponentToPartitionActorsMapLock);
+
+	TSet<TObjectPtr<APCGPartitionActor>>* PartitionActorsPtr = ComponentToPartitionActorsMap.Find(InComponent);
+	if (PartitionActorsPtr)
+	{
+		for (APCGPartitionActor* Actor : *PartitionActorsPtr)
+		{
+			Actor->RemoveGraphInstance(InComponent);
+		}
+
+		PartitionActorsPtr->Empty();
+	}
+}
+
+TSet<TObjectPtr<UPCGComponent>> UPCGSubsystem::GetAllRegisteredComponents() const
+{
+	TSet<TObjectPtr<UPCGComponent>> Components;
+	ComponentToIdMap.GetKeys(Components);
+
+	return Components;
+}
+
 #if WITH_EDITOR
 
 namespace PCGSubsystem
@@ -946,29 +1074,16 @@ FPCGTaskId UPCGSubsystem::CleanupGraph(UPCGComponent* Component, const FBox& InB
 
 void UPCGSubsystem::DirtyGraph(UPCGComponent* Component, const FBox& InBounds, EPCGComponentDirtyFlag DirtyFlag)
 {
-	TWeakObjectPtr<UPCGComponent> ComponentPtr(Component);
-
-	auto ScheduleTask = [this, ComponentPtr, DirtyFlag](APCGPartitionActor* PCGActor, const FBox& InIntersectedBounds, const TArray<FPCGTaskId>& TaskDependencies) {
-		// In the specific case of the dirty, we want to bypass the execution queue, esp. since there's nothing happening here
-		// so we will run the command now, and not delay it.
-		check(ComponentPtr.Get() != nullptr && PCGActor != nullptr);
-		UPCGComponent* Component = ComponentPtr.Get();
-
-		if (!PCGActor || !Component)
-		{
-			return InvalidPCGTaskId;
-		}
-
-		if (UPCGComponent* LocalComponent = PCGActor->GetLocalComponent(Component))
-		{
-			LocalComponent->DirtyGenerated(DirtyFlag);
-		}
+	check(Component);
 
+	// Immediate operation
+	auto DirtyTask = [DirtyFlag](UPCGComponent* LocalComponent)
+	{
+		LocalComponent->DirtyGenerated(DirtyFlag);
 		return InvalidPCGTaskId;
 	};
 
-	// TODO: Could implement a "ForEachLoadedIntersectingCellImmediate" method to clarify this use case
-	PCGSubsystem::ForEachIntersectingCell(GraphExecutor, Component->GetWorld(), InBounds, /*bCreateActor=*/false, /*bLoadCell=*/false, /*bSave=*/false, ScheduleTask);
+	DispatchToRegisteredLocalComponents(Component, DirtyTask);
 }
 
 void UPCGSubsystem::CleanupPartitionActors(const FBox& InBounds)
@@ -1053,7 +1168,7 @@ void UPCGSubsystem::ClearPCGLink(UPCGComponent* InComponent, const FBox& InBound
 	}
 }
 
-void UPCGSubsystem::DeletePartitionActors()
+void UPCGSubsystem::DeletePartitionActors(bool bOnlyDeleteUnused)
 {
 	TRACE_CPUPROFILER_EVENT_SCOPE(UPCGSubsystem::DeletePartitionActors);
 
@@ -1122,13 +1237,42 @@ void UPCGSubsystem::DeletePartitionActors()
 		GEditor->SelectNone(true, true, false);
 	}
 
-	FWorldPartitionHelpers::ForEachActorDesc<APCGPartitionActor>(World->GetWorldPartition(), GatherAndDestroyActors);
+	if (bOnlyDeleteUnused)
+	{		
+		// Need to copy to avoid deadlock
+		TMap<FIntVector, TObjectPtr<APCGPartitionActor>> PartitionActorsMapCopy;
+		{
+			FReadScopeLock ReadLock(PartitionActorsMapLock);
+			PartitionActorsMapCopy = PartitionActorsMap;
+		}
 
-	// Also cleanup the remaining actors that don't have descriptors, if we have a loaded level
-	if (ULevel* Level = World->GetCurrentLevel())
+		for (const TPair<FIntVector, TObjectPtr<APCGPartitionActor>>& Item : PartitionActorsMapCopy)
+		{
+			if (APCGPartitionActor* PartitionActor = Item.Value)
+			{
+				bool bIntersectWithOneComponent = false;
+
+				auto FoundComponents = [&bIntersectWithOneComponent](UPCGComponent*) { bIntersectWithOneComponent = true; };
+
+				ForAllIntersectingComponents(PartitionActor->GetFixedBounds(), FoundComponents);
+
+				if (!bIntersectWithOneComponent)
+				{
+					GatherAndDestroyLoadedActors(PartitionActor);
+				}
+			}
+		}
+	}
+	else
 	{
-		TRACE_CPUPROFILER_EVENT_SCOPE(UPCGSubsystem::DeletePartitionActors::ForEachActorInLevel);
-		UPCGActorHelpers::ForEachActorInLevel<APCGPartitionActor>(Level, GatherAndDestroyLoadedActors);
+		FWorldPartitionHelpers::ForEachActorDesc<APCGPartitionActor>(World->GetWorldPartition(), GatherAndDestroyActors);
+
+		// Also cleanup the remaining actors that don't have descriptors, if we have a loaded level
+		if (ULevel* Level = World->GetCurrentLevel())
+		{
+			TRACE_CPUPROFILER_EVENT_SCOPE(UPCGSubsystem::DeletePartitionActors::ForEachActorInLevel);
+			UPCGActorHelpers::ForEachActorInLevel<APCGPartitionActor>(Level, GatherAndDestroyLoadedActors);
+		}
 	}
 
 	if (!ActorsNotSafeToBeDeleted.IsEmpty())
@@ -1204,35 +1348,4 @@ void UPCGSubsystem::ResetPartitionActorsMap()
 	PartitionActorsMapLock.WriteUnlock();
 }
 
-///////////////////////////////////////////////////////////////////////////////////////////
-FPCGTaskId UPCGSubsystem::SchedulePartitionGraph(UPCGComponent* Component)
-{
-	check(Component && Component->IsPartitioned());
-
-	const FBox PreviousBounds = Component->LastGeneratedBounds;
-	const FBox NewBounds = Component->GetGridBounds();
-
-	return ProcessGraph(Component, PreviousBounds, NewBounds, EOperation::Partition, /*bSave=*/false);
-}
-
-FPCGTaskId UPCGSubsystem::ScheduleUnpartitionGraph(UPCGComponent* Component)
-{
-	check(Component && Component->IsPartitioned());
-
-	const FBox PreviousBounds = Component->LastGeneratedBounds;
-	const FBox NewBounds = Component->GetGridBounds();
-
-	return ProcessGraph(Component, PreviousBounds, NewBounds, EOperation::Unpartition, /*bSave=*/false);
-}
-
-FPCGTaskId UPCGSubsystem::ScheduleProcessGraph(UPCGComponent* Component, bool bSave)
-{
-	check(Component && Component->IsPartitioned());
-
-	const FBox PreviousBounds = Component->LastGeneratedBounds;
-	const FBox NewBounds = Component->GetGridBounds();
-
-	return ProcessGraph(Component, PreviousBounds, NewBounds, EOperation::Generate, bSave);
-}
-
 #endif // WITH_EDITOR
diff --git a/Engine/Plugins/Experimental/PCG/Source/PCG/Private/PCGWorldActor.cpp b/Engine/Plugins/Experimental/PCG/Source/PCG/Private/PCGWorldActor.cpp
--- a/Engine/Plugins/Experimental/PCG/Source/PCG/Private/PCGWorldActor.cpp
+++ b/Engine/Plugins/Experimental/PCG/Source/PCG/Private/PCGWorldActor.cpp
@@ -98,13 +98,10 @@ void APCGWorldActor::OnPartitionGridSizeChanged()
 	{
 		return;
 	}
-
-	// First gather all the PCG components that linked to PCGPartitionActor
-	TSet<TObjectPtr<UPCGComponent>> AllPartitionedComponents;
 	
 	bool bAllSafeToDelete = true;
 
-	auto AddPartitionComponentAndCheckIfSafeToDelete = [&AllPartitionedComponents, &bAllSafeToDelete](AActor* Actor) -> bool
+	auto AddPartitionComponentAndCheckIfSafeToDelete = [&bAllSafeToDelete](AActor* Actor) -> bool
 	{
 		TObjectPtr<APCGPartitionActor> PartitionActor = CastChecked<APCGPartitionActor>(Actor);
 
@@ -114,14 +111,6 @@ void APCGWorldActor::OnPartitionGridSizeChanged()
 			return true;
 		}
 
-		for (UPCGComponent* PCGComponent : PartitionActor->GetAllOriginalPCGComponents())
-		{
-			if (PCGComponent)
-			{
-				AllPartitionedComponents.Add(PCGComponent);
-			}
-		}
-
 		return true;
 	};
 
@@ -136,11 +125,13 @@ void APCGWorldActor::OnPartitionGridSizeChanged()
 	}
 
 	// Then delete all PCGPartitionActors
-	PCGSubsystem->DeletePartitionActors();
+	PCGSubsystem->DeletePartitionActors(/*bDeleteOnlyUnused=*/false);
 
-	// And finally, refresh all components
-	for (TObjectPtr<UPCGComponent> PCGComponent : AllPartitionedComponents)
+	// And finally, refresh all components that are partitioned (registered to the PCGSubsystem)
+	// to let them recreate the needed PCG Partition Actors.
+	for (UPCGComponent* PCGComponent : PCGSubsystem->GetAllRegisteredComponents())
 	{
+		check(PCGComponent);
 		PCGComponent->DirtyGenerated();
 		PCGComponent->Refresh();
 	}
diff --git a/Engine/Plugins/Experimental/PCG/Source/PCG/Public/PCGComponent.h b/Engine/Plugins/Experimental/PCG/Source/PCG/Public/PCGComponent.h
--- a/Engine/Plugins/Experimental/PCG/Source/PCG/Public/PCGComponent.h
+++ b/Engine/Plugins/Experimental/PCG/Source/PCG/Public/PCGComponent.h
@@ -192,6 +192,12 @@ public:
 	const FPCGDataCollection* GetInspectionData(const UPCGNode* InNode) const;
 #endif
 
+	/** Utility function (mostly for tests) to properly set the value of bIsPartitioned.
+	*   Will do an immediate cleanup first and then register/unregister the component to the subsystem.
+	*   It's your responsibility after to regenerate the graph if you want to.
+	*/
+	void SetIsPartitioned(bool bIsNowPartitioned);
+
 	bool IsPartitioned() const;
 	bool IsLocalComponent() const { return bIsComponentLocal; }
 
diff --git a/Engine/Plugins/Experimental/PCG/Source/PCG/Public/PCGSubsystem.h b/Engine/Plugins/Experimental/PCG/Source/PCG/Public/PCGSubsystem.h
--- a/Engine/Plugins/Experimental/PCG/Source/PCG/Public/PCGSubsystem.h
+++ b/Engine/Plugins/Experimental/PCG/Source/PCG/Public/PCGSubsystem.h
@@ -31,6 +31,8 @@ class PCG_API UPCGSubsystem : public UTickableWorldSubsystem
 	GENERATED_BODY()
 
 public:
+	UPCGSubsystem();
+
 	//~ Begin USubsystem Interface.
 	virtual void Deinitialize() override;
 	//~ End USubsystem Interface
@@ -74,28 +76,22 @@ public:
 	/** Gets the output data for a given task */
 	bool GetOutputData(FPCGTaskId InTaskId, FPCGDataCollection& OutData);
 
-	/** Register a new PCG Component, will be added to the octree. Returns an id that needs to be kept for update and removal. Thread safe */
-	void RegisterPCGComponent(UPCGComponent* InComponent);
-
-	/** Update a PCG Component, if it has changed its transform. Thread safe */
-	void UpdatePCGComponentBounds(UPCGComponent* InComponent);
+	/** Register a new PCG Component or update it, will be added to the octree if it doesn't exists yet. Returns true if it was added/updated. Thread safe */
+	bool RegisterOrUpdatePCGComponent(UPCGComponent* InComponent, bool bDoActorMapping = true);
 
 	/** Unregister a PCG Component, will be removed from the octree. Thread safe */
 	void UnregisterPCGComponent(UPCGComponent* InComponent);
 
-	/** Register a new Partition actor, will be added to a map and will query all intersecting volume to bind to them. Thread safe */
-	void RegisterPartitionActor(APCGPartitionActor* InActor);
+	/** Register a new Partition actor, will be added to a map and will query all intersecting volume to bind to them if asked. Thread safe */
+	void RegisterPartitionActor(APCGPartitionActor* InActor, bool bDoComponentMapping = true);
 
 	/** Unregister a Partition actor, will be removed from the map and remove itself to all intersecting volumes. Thread safe */
 	void UnregisterPartitionActor(APCGPartitionActor* InActor);
 
+	TSet<TObjectPtr<UPCGComponent>> GetAllRegisteredComponents() const;
+
 #if WITH_EDITOR
 public:
-	/** Operations to schedule a later operation, which enables to delay bounds querying */
-	FPCGTaskId SchedulePartitionGraph(UPCGComponent* Component);
-	FPCGTaskId ScheduleUnpartitionGraph(UPCGComponent* Component);
-	FPCGTaskId ScheduleProcessGraph(UPCGComponent* Component, bool bSave);
-
 	/** Schedules an operation to cleanup the graph in the given bounds */
 	FPCGTaskId CleanupGraph(UPCGComponent* Component, const FBox& InBounds, bool bRemoveComponents, bool bSave);
 
@@ -104,7 +100,7 @@ public:
 
 	/** Partition actors methods */
 	void CleanupPartitionActors(const FBox& InBounds);
-	void DeletePartitionActors();
+	void DeletePartitionActors(bool bOnlyDeleteUnused);
 
 	/** Propagate to the graph compiler graph changes */
 	void NotifyGraphChanged(UPCGGraph* InGraph);
@@ -139,17 +135,23 @@ private:
 #endif // WITH_EDITOR
 	
 private:
-	// Schedule multiple graphs
-	TArray<FPCGTaskId> ScheduleMultipleComponent(UPCGComponent* OriginalComponent, TSet<TObjectPtr<APCGPartitionActor>>& PartitionActors, const TArray<FPCGTaskId>& Dependencies);
+	/* Call the InFunc function to all local component registered to the original component. Return the list of all the tasks scheduled. Thread safe*/
+	TArray<FPCGTaskId> DispatchToRegisteredLocalComponents(UPCGComponent* OriginalComponent, const TFunction<FPCGTaskId(UPCGComponent*)>& InFunc) const;
 
-	// Schedule multiple cleanups
-	TArray<FPCGTaskId> ScheduleMultipleCleanup(UPCGComponent* OriginalComponent, TSet<TObjectPtr<APCGPartitionActor>>& PartitionActors, bool bRemoveComponents, const TArray<FPCGTaskId>& Dependencies);
+	/* Call the InFunc function to all local component from the set of partition actors. Return the list of all the tasks scheduled. */
+	TArray<FPCGTaskId> DispatchToLocalComponents(UPCGComponent* OriginalComponent, const TSet<TObjectPtr<APCGPartitionActor>>& PartitionActors, const TFunction<FPCGTaskId(UPCGComponent*)>& InFunc) const;
 
 	/** Iterate other all the components which bounds intersect the box in param and call a callback. Thread safe */
-	void FindAllIntersectingComponents(const FBoxCenterAndExtent& InBounds, TFunction<void(UPCGComponent*)> InFunc) const;
+	void ForAllIntersectingComponents(const FBoxCenterAndExtent& InBounds, TFunction<void(UPCGComponent*)> InFunc) const;
 
 	/** Iterate other all the int coordinates given a box and call a callback. Thread safe */
-	void FindAllIntersectingPartitionActors(const FBox& InBounds, TFunction<void(APCGPartitionActor*)> InFunc) const;
+	void ForAllIntersectingPartitionActors(const FBox& InBounds, TFunction<void(APCGPartitionActor*)> InFunc) const;
+
+	/** Update the current mapping between a PCG component and its PCG Partition actors */
+	void UpdateMappingPCGComponentPartitionActor(UPCGComponent* InComponent);
+
+	/** Delete the current mapping between a PCG component and its PCG Partition actors */
+	void DeleteMappingPCGComponentPartitionActor(UPCGComponent* InComponent);
 	
 private:
 	APCGWorldActor* PCGWorldActor = nullptr;
@@ -160,7 +162,7 @@ private:
 #endif
 
 	FPCGComponentOctree PCGComponentOctree;
-	TMap<TObjectPtr<const UPCGComponent>, FPCGComponentOctreeIDSharedRef> ComponentToIdMap;
+	TMap<TObjectPtr<UPCGComponent>, FPCGComponentOctreeIDSharedRef> ComponentToIdMap;
 	mutable FRWLock PCGVolumeOctreeLock;
 
 	TMap<FIntVector, TObjectPtr<APCGPartitionActor>> PartitionActorsMap;
diff --git a/Engine/Plugins/Experimental/PCG/Source/PCGEditor/Private/PCGEditorModule.cpp b/Engine/Plugins/Experimental/PCG/Source/PCGEditor/Private/PCGEditorModule.cpp
--- a/Engine/Plugins/Experimental/PCG/Source/PCGEditor/Private/PCGEditorModule.cpp
+++ b/Engine/Plugins/Experimental/PCG/Source/PCGEditor/Private/PCGEditorModule.cpp
@@ -160,8 +160,8 @@ void FPCGEditorModule::AddMenuEntry(FMenuBuilder& MenuBuilder)
 void FPCGEditorModule::PopulateMenuActions(FMenuBuilder& MenuBuilder)
 {
 	MenuBuilder.AddMenuEntry(
-		LOCTEXT("DeletePCGActors", "Delete all PCG partition actors"),
-		LOCTEXT("DeletePCGActors_Tooltip", "Deletes all PCG partition actors in the current world"),
+		LOCTEXT("DeletePCGPartitionActors", "Delete all PCG partition actors"),
+		LOCTEXT("DeletePCGPartitionActors_Tooltip", "Deletes all PCG partition actors in the current world"),
 		FSlateIcon(),
 		FUIAction(
 			FExecuteAction::CreateLambda([]() {
@@ -169,7 +169,23 @@ void FPCGEditorModule::PopulateMenuActions(FMenuBuilder& MenuBuilder)
 				{
 					if (UWorld* World = GEditor->GetEditorWorldContext().World())
 					{
-						World->GetSubsystem<UPCGSubsystem>()->DeletePartitionActors();
+						World->GetSubsystem<UPCGSubsystem>()->DeletePartitionActors(/*bOnlyDeleteUnused=*/false);
+					}
+				}
+			})),
+		NAME_None);
+
+	MenuBuilder.AddMenuEntry(
+		LOCTEXT("DeleteUnusedPCGPartitionActors", "Delete all unused PCG partition actors"),
+		LOCTEXT("DeleteUnusedPCGPartitionActors_Tooltip", "Deletes all PCG partition actors in the current world that doesn't intersect with any PCG Component."),
+		FSlateIcon(),
+		FUIAction(
+			FExecuteAction::CreateLambda([]() {
+				if (GEditor)
+				{
+					if (UWorld* World = GEditor->GetEditorWorldContext().World())
+					{
+						World->GetSubsystem<UPCGSubsystem>()->DeletePartitionActors(/*bOnlyDeleteUnused=*/true);
 					}
 				}
 			})),
