PCG: Add Component Seed in the Cache

diff --git a/Engine/Plugins/Experimental/PCG/Source/PCG/Private/Graph/PCGGraphCache.cpp b/Engine/Plugins/Experimental/PCG/Source/PCG/Private/Graph/PCGGraphCache.cpp
--- a/Engine/Plugins/Experimental/PCG/Source/PCG/Private/Graph/PCGGraphCache.cpp
+++ b/Engine/Plugins/Experimental/PCG/Source/PCG/Private/Graph/PCGGraphCache.cpp
@@ -1,27 +1,35 @@
 // Copyright Epic Games, Inc. All Rights Reserved.
 
 #include "PCGGraphCache.h"
+#include "PCGComponent.h"
 #include "PCGSettings.h"
 
 #include "Misc/ScopeRWLock.h"
 #include "Algo/AnyOf.h"
 #include "GameFramework/Actor.h"
 
-FPCGGraphCacheEntry::FPCGGraphCacheEntry(const FPCGDataCollection& InInput, const UPCGSettings* InSettings, const FPCGDataCollection& InOutput, TWeakObjectPtr<UObject> InOwner, FPCGRootSet& OutRootSet)
+namespace PCGGraphCacheConstants
+{
+	constexpr int32 NullComponentSeed = 0;
+	constexpr int32 NullSettingsCrc32 = 0;
+}
+
+FPCGGraphCacheEntry::FPCGGraphCacheEntry(const FPCGDataCollection& InInput, const UPCGSettings* InSettings, const UPCGComponent* InComponent, const FPCGDataCollection& InOutput, TWeakObjectPtr<UObject> InOwner, FPCGRootSet& OutRootSet)
 	: Input(InInput)
 	, Output(InOutput)
 {
 	// Note: we don't need to root the settings since they'll be owned by the subsystem
 	Settings = InSettings ? Cast<UPCGSettings>(StaticDuplicateObject(InSettings, InOwner.Get())) : nullptr;
-	SettingsCrc32 = InSettings ? InSettings->GetCrc32() : 0;
+	SettingsCrc32 = InSettings ? InSettings->GetCrc32() : PCGGraphCacheConstants::NullSettingsCrc32;
+	ComponentSeed = InComponent ? InComponent->Seed : PCGGraphCacheConstants::NullComponentSeed;
 
 	Input.AddToRootSet(OutRootSet);
 	Output.AddToRootSet(OutRootSet);
 }
 
-bool FPCGGraphCacheEntry::Matches(const FPCGDataCollection& InInput, int32 InSettingsCrc32) const
+bool FPCGGraphCacheEntry::Matches(const FPCGDataCollection& InInput, int32 InSettingsCrc32, int32 InComponentSeed) const
 {
-	return (SettingsCrc32 == InSettingsCrc32) && (Input == InInput);
+	return (SettingsCrc32 == InSettingsCrc32) && (Input == InInput) && (ComponentSeed == InComponentSeed);
 }
 
 FPCGGraphCache::FPCGGraphCache(TWeakObjectPtr<UObject> InOwner)
@@ -34,7 +42,7 @@ FPCGGraphCache::~FPCGGraphCache()
 	ClearCache();
 }
 
-bool FPCGGraphCache::GetFromCache(const IPCGElement* InElement, const FPCGDataCollection& InInput, const UPCGSettings* InSettings, FPCGDataCollection& OutOutput) const
+bool FPCGGraphCache::GetFromCache(const IPCGElement* InElement, const FPCGDataCollection& InInput, const UPCGSettings* InSettings, const UPCGComponent* InComponent, FPCGDataCollection& OutOutput) const
 {
 	if (!Owner.IsValid())
 	{
@@ -46,11 +54,12 @@ bool FPCGGraphCache::GetFromCache(const IPCGElement* InElement, const FPCGDataCo
 
 	if (const FPCGGraphCacheEntries* Entries = CacheData.Find(InElement))
 	{
-		int32 InSettingsCrc32 = (InSettings ? InSettings->GetCrc32() : 0);
+		int32 InSettingsCrc32 = (InSettings ? InSettings->GetCrc32() : PCGGraphCacheConstants::NullSettingsCrc32);
+		int32 InComponentSeed = (InComponent ? InComponent->Seed : PCGGraphCacheConstants::NullComponentSeed);
 
 		for (const FPCGGraphCacheEntry& Entry : *Entries)
 		{
-			if (Entry.Matches(InInput, InSettingsCrc32))
+			if (Entry.Matches(InInput, InSettingsCrc32, InComponentSeed))
 			{
 				OutOutput = Entry.Output;
 				return true;
@@ -63,7 +72,7 @@ bool FPCGGraphCache::GetFromCache(const IPCGElement* InElement, const FPCGDataCo
 	return false;
 }
 
-void FPCGGraphCache::StoreInCache(const IPCGElement* InElement, const FPCGDataCollection& InInput, const UPCGSettings* InSettings, const FPCGDataCollection& InOutput)
+void FPCGGraphCache::StoreInCache(const IPCGElement* InElement, const FPCGDataCollection& InInput, const UPCGSettings* InSettings, const UPCGComponent* InComponent, const FPCGDataCollection& InOutput)
 {
 	if (!Owner.IsValid())
 	{
@@ -79,7 +88,7 @@ void FPCGGraphCache::StoreInCache(const IPCGElement* InElement, const FPCGDataCo
 		Entries = &(CacheData.Add(InElement));
 	}
 
-	Entries->Emplace(InInput, InSettings, InOutput, Owner, RootSet);
+	Entries->Emplace(InInput, InSettings, InComponent, InOutput, Owner, RootSet);
 }
 
 void FPCGGraphCache::ClearCache()
diff --git a/Engine/Plugins/Experimental/PCG/Source/PCG/Private/Graph/PCGGraphCache.h b/Engine/Plugins/Experimental/PCG/Source/PCG/Private/Graph/PCGGraphCache.h
--- a/Engine/Plugins/Experimental/PCG/Source/PCG/Private/Graph/PCGGraphCache.h
+++ b/Engine/Plugins/Experimental/PCG/Source/PCG/Private/Graph/PCGGraphCache.h
@@ -7,19 +7,21 @@
 #include "PCGElement.h"
 
 class UPCGSettings;
+class UPCGComponent;
 class AActor;
 
 struct FPCGGraphCacheEntry
 {
 	FPCGGraphCacheEntry() = default;
-	FPCGGraphCacheEntry(const FPCGDataCollection& InInput, const UPCGSettings* InSettings, const FPCGDataCollection& InOutput, TWeakObjectPtr<UObject> InOwner, FPCGRootSet& OutRootSet);
+	FPCGGraphCacheEntry(const FPCGDataCollection& InInput, const UPCGSettings* InSettings, const UPCGComponent* InComponent, const FPCGDataCollection& InOutput, TWeakObjectPtr<UObject> InOwner, FPCGRootSet& OutRootSet);
 
-	bool Matches(const FPCGDataCollection& InInput, int32 InSettingsCrc32) const;
+	bool Matches(const FPCGDataCollection& InInput, int32 InSettingsCrc32, int32 InComponentSeed) const;
 
 	FPCGDataCollection Input;
 	TObjectPtr<UPCGSettings> Settings;
 	FPCGDataCollection Output;
 	int32 SettingsCrc32;
+	int32 ComponentSeed;
 };
 
 // TODO: investigate if we need a more evolved data structure here
@@ -36,10 +38,10 @@ struct FPCGGraphCache
 	~FPCGGraphCache();
 
 	/** Returns true if data was found from the cache, in which case the outputs are written in OutOutput */
-	bool GetFromCache(const IPCGElement* InElement, const FPCGDataCollection& InInput, const UPCGSettings* InSettings, FPCGDataCollection& OutOutput) const;
+	bool GetFromCache(const IPCGElement* InElement, const FPCGDataCollection& InInput, const UPCGSettings* InSettings, const UPCGComponent* InComponent, FPCGDataCollection& OutOutput) const;
 
 	/** Stores data in the cache for later use */
-	void StoreInCache(const IPCGElement* InElement, const FPCGDataCollection& InInput, const UPCGSettings* InSettings, const FPCGDataCollection& InOutput);
+	void StoreInCache(const IPCGElement* InElement, const FPCGDataCollection& InInput, const UPCGSettings* InSettings, const UPCGComponent* InComponent, const FPCGDataCollection& InOutput);
 
 	/** Removes all entries from the cache, unroots data, etc. */
 	void ClearCache();
diff --git a/Engine/Plugins/Experimental/PCG/Source/PCG/Private/Graph/PCGGraphExecutor.cpp b/Engine/Plugins/Experimental/PCG/Source/PCG/Private/Graph/PCGGraphExecutor.cpp
--- a/Engine/Plugins/Experimental/PCG/Source/PCG/Private/Graph/PCGGraphExecutor.cpp
+++ b/Engine/Plugins/Experimental/PCG/Source/PCG/Private/Graph/PCGGraphExecutor.cpp
@@ -257,7 +257,7 @@ void FPCGGraphExecutor::Execute()
 				// there is an execution mode that would prevent us from doing so.
 				const UPCGSettings* TaskSettings = PCGContextHelpers::GetInputSettings<UPCGSettings>(Task.Node, TaskInput);
 				FPCGDataCollection CachedOutput;
-				const bool bResultAlreadyInCache = Task.Element->IsCacheable(TaskSettings) && GraphCache.GetFromCache(Task.Element.Get(), TaskInput, TaskSettings, CachedOutput);
+				const bool bResultAlreadyInCache = Task.Element->IsCacheable(TaskSettings) && GraphCache.GetFromCache(Task.Element.Get(), TaskInput, TaskSettings, Task.SourceComponent, CachedOutput);
 #if WITH_EDITOR
 				const bool bNeedsToCreateActiveTask = !bResultAlreadyInCache || TaskSettings->ExecutionMode == EPCGSettingsExecutionMode::Debug || TaskSettings->ExecutionMode == EPCGSettingsExecutionMode::Isolated;
 #else
@@ -377,7 +377,7 @@ void FPCGGraphExecutor::Execute()
 				const UPCGSettings* ActiveTaskSettings = ActiveTask.Context->GetInputSettings<UPCGSettings>();
 				if (ActiveTaskSettings && ActiveTask.Element->IsCacheable(ActiveTaskSettings))
 				{
-					GraphCache.StoreInCache(ActiveTask.Element.Get(), ActiveTask.Context->InputData, ActiveTaskSettings, ActiveTask.Context->OutputData);
+					GraphCache.StoreInCache(ActiveTask.Element.Get(), ActiveTask.Context->InputData, ActiveTaskSettings, ActiveTask.Context->SourceComponent, ActiveTask.Context->OutputData);
 				}
 			}
 
