PCG: Removed the option to change the seed if the settings doesn't use randomness.

diff --git a/Engine/Plugins/Experimental/PCG/Source/PCG/Private/Elements/PCGDensityNoise.cpp b/Engine/Plugins/Experimental/PCG/Source/PCG/Private/Elements/PCGDensityNoise.cpp
--- a/Engine/Plugins/Experimental/PCG/Source/PCG/Private/Elements/PCGDensityNoise.cpp
+++ b/Engine/Plugins/Experimental/PCG/Source/PCG/Private/Elements/PCGDensityNoise.cpp
@@ -7,6 +7,11 @@
 
 #include "Math/RandomStream.h"
 
+UPCGDensityNoiseSettings::UPCGDensityNoiseSettings()
+{
+	bUseSeed = true;
+}
+
 FPCGElementPtr UPCGDensityNoiseSettings::CreateElement() const
 {
 	return MakeShared<FPCGDensityNoiseElement>();
diff --git a/Engine/Plugins/Experimental/PCG/Source/PCG/Private/Elements/PCGDensityRemap.cpp b/Engine/Plugins/Experimental/PCG/Source/PCG/Private/Elements/PCGDensityRemap.cpp
--- a/Engine/Plugins/Experimental/PCG/Source/PCG/Private/Elements/PCGDensityRemap.cpp
+++ b/Engine/Plugins/Experimental/PCG/Source/PCG/Private/Elements/PCGDensityRemap.cpp
@@ -7,6 +7,11 @@
 #include "Helpers/PCGSettingsHelpers.h"
 #include "Math/RandomStream.h"
 
+UPCGLinearDensityRemapSettings::UPCGLinearDensityRemapSettings()
+{
+	bUseSeed = true;
+}
+
 FPCGElementPtr UPCGLinearDensityRemapSettings::CreateElement() const
 {
 	return MakeShared<FPCGLinearDensityRemapElement>();
diff --git a/Engine/Plugins/Experimental/PCG/Source/PCG/Private/Elements/PCGExecuteBlueprint.cpp b/Engine/Plugins/Experimental/PCG/Source/PCG/Private/Elements/PCGExecuteBlueprint.cpp
--- a/Engine/Plugins/Experimental/PCG/Source/PCG/Private/Elements/PCGExecuteBlueprint.cpp
+++ b/Engine/Plugins/Experimental/PCG/Source/PCG/Private/Elements/PCGExecuteBlueprint.cpp
@@ -151,6 +151,11 @@ TSet<FName> UPCGBlueprintElement::OutputLabels() const
 	return Labels;
 }
 
+UPCGBlueprintSettings::UPCGBlueprintSettings()
+{
+	bUseSeed = true;
+}
+
 void UPCGBlueprintSettings::SetupBlueprintEvent()
 {
 #if WITH_EDITOR
diff --git a/Engine/Plugins/Experimental/PCG/Source/PCG/Private/Elements/PCGPointSampler.cpp b/Engine/Plugins/Experimental/PCG/Source/PCG/Private/Elements/PCGPointSampler.cpp
--- a/Engine/Plugins/Experimental/PCG/Source/PCG/Private/Elements/PCGPointSampler.cpp
+++ b/Engine/Plugins/Experimental/PCG/Source/PCG/Private/Elements/PCGPointSampler.cpp
@@ -7,6 +7,11 @@
 #include "Helpers/PCGSettingsHelpers.h"
 #include "Math/RandomStream.h"
 
+UPCGPointSamplerSettings::UPCGPointSamplerSettings()
+{
+	bUseSeed = true;
+}
+
 FPCGElementPtr UPCGPointSamplerSettings::CreateElement() const
 {
 	return MakeShared<FPCGPointSamplerElement>();
diff --git a/Engine/Plugins/Experimental/PCG/Source/PCG/Private/Elements/PCGStaticMeshSpawner.cpp b/Engine/Plugins/Experimental/PCG/Source/PCG/Private/Elements/PCGStaticMeshSpawner.cpp
--- a/Engine/Plugins/Experimental/PCG/Source/PCG/Private/Elements/PCGStaticMeshSpawner.cpp
+++ b/Engine/Plugins/Experimental/PCG/Source/PCG/Private/Elements/PCGStaticMeshSpawner.cpp
@@ -17,6 +17,8 @@
 
 UPCGStaticMeshSpawnerSettings::UPCGStaticMeshSpawnerSettings(const FObjectInitializer &ObjectInitializer)
 {
+	bUseSeed = true;
+
 	MeshSelectorType = UPCGMeshSelectorWeighted::StaticClass();
 	MeshSelectorInstance = ObjectInitializer.CreateDefaultSubobject<UPCGMeshSelectorWeighted>(this, TEXT("DefaultSelectorInstance"));
 }
diff --git a/Engine/Plugins/Experimental/PCG/Source/PCG/Private/Elements/PCGSurfaceSampler.cpp b/Engine/Plugins/Experimental/PCG/Source/PCG/Private/Elements/PCGSurfaceSampler.cpp
--- a/Engine/Plugins/Experimental/PCG/Source/PCG/Private/Elements/PCGSurfaceSampler.cpp
+++ b/Engine/Plugins/Experimental/PCG/Source/PCG/Private/Elements/PCGSurfaceSampler.cpp
@@ -102,6 +102,11 @@ struct FPCGSurfaceSamplerLoopData
 	}
 };
 
+UPCGSurfaceSamplerSettings::UPCGSurfaceSamplerSettings()
+{
+	bUseSeed = true;
+}
+
 void UPCGSurfaceSamplerSettings::PostLoad()
 {
 	Super::PostLoad();
diff --git a/Engine/Plugins/Experimental/PCG/Source/PCG/Private/Elements/PCGTransformPoints.cpp b/Engine/Plugins/Experimental/PCG/Source/PCG/Private/Elements/PCGTransformPoints.cpp
--- a/Engine/Plugins/Experimental/PCG/Source/PCG/Private/Elements/PCGTransformPoints.cpp
+++ b/Engine/Plugins/Experimental/PCG/Source/PCG/Private/Elements/PCGTransformPoints.cpp
@@ -5,6 +5,11 @@
 #include "PCGHelpers.h"
 #include "Helpers/PCGSettingsHelpers.h"
 
+UPCGTransformPointsSettings::UPCGTransformPointsSettings()
+{
+	bUseSeed = true;
+}
+
 TArray<FPCGPinProperties> UPCGTransformPointsSettings::InputPinProperties() const
 {
 	TArray<FPCGPinProperties> PinProperties;
diff --git a/Engine/Plugins/Experimental/PCG/Source/PCG/Private/Graph/PCGGraphCache.cpp b/Engine/Plugins/Experimental/PCG/Source/PCG/Private/Graph/PCGGraphCache.cpp
--- a/Engine/Plugins/Experimental/PCG/Source/PCG/Private/Graph/PCGGraphCache.cpp
+++ b/Engine/Plugins/Experimental/PCG/Source/PCG/Private/Graph/PCGGraphCache.cpp
@@ -8,10 +8,22 @@
 #include "Algo/AnyOf.h"
 #include "GameFramework/Actor.h"
 
-namespace PCGGraphCacheConstants
+namespace PCGGraphCache
 {
 	constexpr int32 NullComponentSeed = 0;
 	constexpr int32 NullSettingsCrc32 = 0;
+
+	// Component seed is only get if we have a component and either no settings or the settings use seed.
+	// TODO: Perhaps we should not get the component seed if we have no settings...
+	int32 GetComponentSeed(const UPCGSettings* InSettings, const UPCGComponent* InComponent)
+	{
+		if ((!InSettings || InSettings->UseSeed()) && InComponent)
+		{
+			return InComponent->Seed;
+		}
+
+		return PCGGraphCache::NullComponentSeed;
+	}
 }
 
 FPCGGraphCacheEntry::FPCGGraphCacheEntry(const FPCGDataCollection& InInput, const UPCGSettings* InSettings, const UPCGComponent* InComponent, const FPCGDataCollection& InOutput, TWeakObjectPtr<UObject> InOwner, FPCGRootSet& OutRootSet)
@@ -20,8 +32,8 @@ FPCGGraphCacheEntry::FPCGGraphCacheEntry(const FPCGDataCollection& InInput, cons
 {
 	// Note: we don't need to root the settings since they'll be owned by the subsystem
 	Settings = InSettings ? Cast<UPCGSettings>(StaticDuplicateObject(InSettings, InOwner.Get())) : nullptr;
-	SettingsCrc32 = InSettings ? InSettings->GetCrc32() : PCGGraphCacheConstants::NullSettingsCrc32;
-	ComponentSeed = InComponent ? InComponent->Seed : PCGGraphCacheConstants::NullComponentSeed;
+	SettingsCrc32 = InSettings ? InSettings->GetCrc32() : PCGGraphCache::NullSettingsCrc32;
+	ComponentSeed = PCGGraphCache::GetComponentSeed(InSettings, InComponent);
 
 	Input.AddToRootSet(OutRootSet);
 	Output.AddToRootSet(OutRootSet);
@@ -54,8 +66,8 @@ bool FPCGGraphCache::GetFromCache(const IPCGElement* InElement, const FPCGDataCo
 
 	if (const FPCGGraphCacheEntries* Entries = CacheData.Find(InElement))
 	{
-		int32 InSettingsCrc32 = (InSettings ? InSettings->GetCrc32() : PCGGraphCacheConstants::NullSettingsCrc32);
-		int32 InComponentSeed = (InComponent ? InComponent->Seed : PCGGraphCacheConstants::NullComponentSeed);
+		int32 InSettingsCrc32 = (InSettings ? InSettings->GetCrc32() : PCGGraphCache::NullSettingsCrc32);
+		int32 InComponentSeed = PCGGraphCache::GetComponentSeed(InSettings, InComponent);
 
 		for (const FPCGGraphCacheEntry& Entry : *Entries)
 		{
diff --git a/Engine/Plugins/Experimental/PCG/Source/PCG/Public/Elements/PCGDensityNoise.h b/Engine/Plugins/Experimental/PCG/Source/PCG/Public/Elements/PCGDensityNoise.h
--- a/Engine/Plugins/Experimental/PCG/Source/PCG/Public/Elements/PCGDensityNoise.h
+++ b/Engine/Plugins/Experimental/PCG/Source/PCG/Public/Elements/PCGDensityNoise.h
@@ -27,6 +27,8 @@ class PCG_API UPCGDensityNoiseSettings : public UPCGSettings
 	GENERATED_BODY()
 
 public:
+	UPCGDensityNoiseSettings();
+
 	// ~Begin UPCGSettings interface
 #if WITH_EDITOR
 	virtual FName GetDefaultNodeName() const override { return FName(TEXT("DensityNoiseNode")); }
diff --git a/Engine/Plugins/Experimental/PCG/Source/PCG/Public/Elements/PCGDensityRemap.h b/Engine/Plugins/Experimental/PCG/Source/PCG/Public/Elements/PCGDensityRemap.h
--- a/Engine/Plugins/Experimental/PCG/Source/PCG/Public/Elements/PCGDensityRemap.h
+++ b/Engine/Plugins/Experimental/PCG/Source/PCG/Public/Elements/PCGDensityRemap.h
@@ -15,6 +15,8 @@ class PCG_API UPCGLinearDensityRemapSettings : public UPCGSettings
 	GENERATED_BODY()
 
 public:
+	UPCGLinearDensityRemapSettings();
+
 	// ~Begin UPCGSettings interface
 #if WITH_EDITOR
 	virtual FName GetDefaultNodeName() const override { return FName(TEXT("LinearDensityRemap")); }
diff --git a/Engine/Plugins/Experimental/PCG/Source/PCG/Public/Elements/PCGExecuteBlueprint.h b/Engine/Plugins/Experimental/PCG/Source/PCG/Public/Elements/PCGExecuteBlueprint.h
--- a/Engine/Plugins/Experimental/PCG/Source/PCG/Public/Elements/PCGExecuteBlueprint.h
+++ b/Engine/Plugins/Experimental/PCG/Source/PCG/Public/Elements/PCGExecuteBlueprint.h
@@ -159,6 +159,8 @@ class PCG_API UPCGBlueprintSettings : public UPCGSettings
 	GENERATED_BODY()
 
 public:
+	UPCGBlueprintSettings();
+
 	friend class FPCGExecuteBlueprintElement;
 
 	// ~Begin UPCGSettings interface
diff --git a/Engine/Plugins/Experimental/PCG/Source/PCG/Public/Elements/PCGPointSampler.h b/Engine/Plugins/Experimental/PCG/Source/PCG/Public/Elements/PCGPointSampler.h
--- a/Engine/Plugins/Experimental/PCG/Source/PCG/Public/Elements/PCGPointSampler.h
+++ b/Engine/Plugins/Experimental/PCG/Source/PCG/Public/Elements/PCGPointSampler.h
@@ -15,6 +15,8 @@ class PCG_API UPCGPointSamplerSettings : public UPCGSettings
 	GENERATED_BODY()
 
 public:
+	UPCGPointSamplerSettings();
+
 	//~Begin UPCGSettings interface
 #if WITH_EDITOR
 	virtual FName GetDefaultNodeName() const override { return FName(TEXT("PointSamplerNode")); }
diff --git a/Engine/Plugins/Experimental/PCG/Source/PCG/Public/Elements/PCGSurfaceSampler.h b/Engine/Plugins/Experimental/PCG/Source/PCG/Public/Elements/PCGSurfaceSampler.h
--- a/Engine/Plugins/Experimental/PCG/Source/PCG/Public/Elements/PCGSurfaceSampler.h
+++ b/Engine/Plugins/Experimental/PCG/Source/PCG/Public/Elements/PCGSurfaceSampler.h
@@ -17,6 +17,8 @@ class PCG_API UPCGSurfaceSamplerSettings : public UPCGSettings
 	GENERATED_BODY()
 
 public:
+	UPCGSurfaceSamplerSettings();
+
 	UPROPERTY(BlueprintReadWrite, EditAnywhere, Category = Settings, meta=(ClampMin="0"))
 	float PointsPerSquaredMeter = 0.1f;
 
diff --git a/Engine/Plugins/Experimental/PCG/Source/PCG/Public/Elements/PCGTransformPoints.h b/Engine/Plugins/Experimental/PCG/Source/PCG/Public/Elements/PCGTransformPoints.h
--- a/Engine/Plugins/Experimental/PCG/Source/PCG/Public/Elements/PCGTransformPoints.h
+++ b/Engine/Plugins/Experimental/PCG/Source/PCG/Public/Elements/PCGTransformPoints.h
@@ -15,6 +15,8 @@ class PCG_API UPCGTransformPointsSettings : public UPCGSettings
 	GENERATED_BODY()
 
 public:
+	UPCGTransformPointsSettings();
+
 	//~Begin UPCGSettings interface
 #if WITH_EDITOR
 	virtual FName GetDefaultNodeName() const override { return FName(TEXT("TransformPointsNode")); }
diff --git a/Engine/Plugins/Experimental/PCG/Source/PCG/Public/PCGSettings.h b/Engine/Plugins/Experimental/PCG/Source/PCG/Public/PCGSettings.h
--- a/Engine/Plugins/Experimental/PCG/Source/PCG/Public/PCGSettings.h
+++ b/Engine/Plugins/Experimental/PCG/Source/PCG/Public/PCGSettings.h
@@ -67,6 +67,7 @@ public:
 
 	bool operator==(const UPCGSettings& Other) const;
 	uint32 GetCrc32() const;
+	bool UseSeed() const { return bUseSeed; }
 
 #if WITH_EDITOR
 	virtual FName GetDefaultNodeName() const { return NAME_None; }
@@ -84,7 +85,7 @@ public:
 	/** Derived classes can implement this to expose additional name information in the logs */
 	virtual FName AdditionalTaskName() const { return NAME_None; }
 
-	UPROPERTY(EditAnywhere, BlueprintReadWrite, Category = Settings)
+	UPROPERTY(EditAnywhere, BlueprintReadWrite, Category = Settings, meta=(EditCondition=bUseSeed, EditConditionHides))
 	int Seed = 0xC35A9631; // random prime number
 
 	/** TODO: Remove this - Placeholder feature until we have a nodegraph */
@@ -125,6 +126,10 @@ protected:
 	void DirtyCache();
 #endif
 
+	// By default, settings won't use a seed. Set this bool at true in the child ctor to allow edition and use it.
+	UPROPERTY(VisibleAnywhere, Transient, Category = Settings, meta = (EditCondition = false, EditConditionHides))
+	bool bUseSeed = false;
+
 	/** Methods to remove boilerplate code across settings */
 	TArray<FPCGPinProperties> DefaultPointOutputPinProperties() const;
 
