Quick fix for the PCG Graph Attribute List View
* Should now properly display attributes associated with PCG Points

#rnx

diff --git a/Engine/Plugins/Experimental/PCG/Source/PCGEditor/Private/SPCGEditorGraphAttributeListView.cpp b/Engine/Plugins/Experimental/PCG/Source/PCGEditor/Private/SPCGEditorGraphAttributeListView.cpp
--- a/Engine/Plugins/Experimental/PCG/Source/PCGEditor/Private/SPCGEditorGraphAttributeListView.cpp
+++ b/Engine/Plugins/Experimental/PCG/Source/PCGEditor/Private/SPCGEditorGraphAttributeListView.cpp
@@ -218,6 +218,17 @@ FText SPCGListViewItemRow::ConvertPointDataToText(const FPCGPoint* PCGPoint, con
 		const int32 Seed = PCGPoint->Seed;
 		return FText::AsNumber(Seed);
 	}
+	// None of the default point columns were caught, see if its metadata
+	else if (const UPCGMetadata* PCGMetadata = InternalItem->PCGMetadata)
+	{
+		if (const FPCGMetadataInfo* MetadataInfo = InternalItem->MetadataInfos->Find(ColumnId))
+		{
+			if (const FPCGMetadataAttributeBase* AttributeBase = PCGMetadata->GetConstAttribute((*MetadataInfo).MetadataId))
+			{
+				return ConvertMetadataAttributeToText(AttributeBase, MetadataInfo, InternalItem->MetaDataItemKey);
+			}
+		}
+	}
 
 	return PCGEditorGraphAttributeListView::NoDataAvailableText;
 }
