[PCG] Factorize metadata elements in their folder + refactor common methods.

diff --git a/Engine/Plugins/Experimental/PCG/Source/PCG/Private/Elements/Metadata/PCGMetadataBreakVector.cpp b/Engine/Plugins/Experimental/PCG/Source/PCG/Private/Elements/Metadata/PCGMetadataBreakVector.cpp
new file mode 100644
--- /dev/null
+++ b/Engine/Plugins/Experimental/PCG/Source/PCG/Private/Elements/Metadata/PCGMetadataBreakVector.cpp
@@ -0,0 +1,274 @@
+// Copyright Epic Games, Inc. All Rights Reserved.
+
+#include "Elements/Metadata/PCGMetadataBreakVector.h"
+
+#include "PCGParamData.h"
+#include "Data/PCGSpatialData.h"
+#include "Elements/Metadata/PCGMetadataElementCommon.h"
+#include "Helpers/PCGSettingsHelpers.h"
+#include "Metadata/PCGMetadata.h"
+#include "Metadata/PCGMetadataAttribute.h"
+#include "Metadata/PCGMetadataAttributeTpl.h"
+
+TArray<FPCGPinProperties> UPCGMetadataBreakVectorSettings::InputPinProperties() const
+{
+	TArray<FPCGPinProperties> PinProperties;
+	PinProperties.Emplace(PCGMetadataBreakVectorConstants::ParamsLabel, EPCGDataType::Param);
+	PinProperties.Emplace(PCGMetadataBreakVectorConstants::SourceLabel, EPCGDataType::Any);
+
+	return PinProperties;
+}
+
+TArray<FPCGPinProperties> UPCGMetadataBreakVectorSettings::OutputPinProperties() const
+{
+	TArray<FPCGPinProperties> PinProperties;
+	PinProperties.Emplace(PCGMetadataBreakVectorConstants::XLabel, EPCGDataType::Any);
+	PinProperties.Emplace(PCGMetadataBreakVectorConstants::YLabel, EPCGDataType::Any);
+	PinProperties.Emplace(PCGMetadataBreakVectorConstants::ZLabel, EPCGDataType::Any);
+	PinProperties.Emplace(PCGMetadataBreakVectorConstants::WLabel, EPCGDataType::Any);
+
+	return PinProperties;
+}
+
+FPCGElementPtr UPCGMetadataBreakVectorSettings::CreateElement() const
+{
+	return MakeShared<FPCGMetadataBreakVectorElement>();
+}
+
+bool FPCGMetadataBreakVectorElement::ExecuteInternal(FPCGContext* Context) const
+{
+	TRACE_CPUPROFILER_EVENT_SCOPE(FPCGMetadataBreakVectorElement::Execute);
+
+	const UPCGMetadataBreakVectorSettings* Settings = Context->GetInputSettings<UPCGMetadataBreakVectorSettings>();
+	check(Settings);
+
+	const TArray<FPCGTaggedData>& Inputs = Context->InputData.TaggedData;
+	TArray<FPCGTaggedData>& Outputs = Context->OutputData.TaggedData;
+
+	const TArray<FPCGTaggedData>& ParamData = Context->InputData.GetParamsByPin(PCGMetadataBreakVectorConstants::ParamsLabel);
+	UPCGParamData* Params = nullptr;
+
+	// TODO: this only captures the first Param
+	for (const FPCGTaggedData& TaggedDatum : ParamData)
+	{
+		if (UPCGParamData* ParamDatum = Cast<UPCGParamData>(TaggedDatum.Data))
+		{
+			Params = ParamDatum;
+			break;
+		}
+	}
+
+	const FName SourceAttributeName = PCG_GET_OVERRIDEN_VALUE(Settings, SourceAttributeName, Params);
+
+	// These exist to facilitate unit testing of the MetadataBreakVector element in isolation, when it otherwise would not have any pin connections to generate output
+	bool bXPinConnected = Settings->bForceConnectX;
+	bool bYPinConnected = Settings->bForceConnectY;
+	bool bZPinConnected = Settings->bForceConnectZ;
+	bool bWPinConnected = Settings->bForceConnectW;
+
+	if (Context->Node)
+	{
+		bXPinConnected |= Context->Node->IsOutputPinConnected(PCGMetadataBreakVectorConstants::XLabel);
+		bYPinConnected |= Context->Node->IsOutputPinConnected(PCGMetadataBreakVectorConstants::YLabel);
+		bZPinConnected |= Context->Node->IsOutputPinConnected(PCGMetadataBreakVectorConstants::ZLabel);
+		bWPinConnected |= Context->Node->IsOutputPinConnected(PCGMetadataBreakVectorConstants::WLabel);
+	}
+
+	for (const FPCGTaggedData& Input : Inputs)
+	{
+		if (!Input.Data)
+		{
+			PCGE_LOG(Error, "Unable to get data from input");
+			continue;
+		}
+
+		UPCGMetadata* SourceMetadata = nullptr;
+		if (const UPCGSpatialData* SpatialInput = Cast<const UPCGSpatialData>(Input.Data))
+		{
+			SourceMetadata = SpatialInput->Metadata;
+		}
+		else if (const UPCGParamData* ParamsInput = Cast<const UPCGParamData>(Input.Data))
+		{
+			SourceMetadata = ParamsInput->Metadata;
+		}
+		else // If the Data type does not handle Metadata, forward it into the valid pins
+		{
+			if (bXPinConnected)
+			{
+				FPCGTaggedData& Data = Outputs.Add_GetRef(Input);
+				Data.Pin = PCGMetadataBreakVectorConstants::XLabel;
+			}
+
+			if (bYPinConnected)
+			{
+				FPCGTaggedData& Data = Outputs.Add_GetRef(Input);
+				Data.Pin = PCGMetadataBreakVectorConstants::YLabel;
+			}
+
+			if (bZPinConnected)
+			{
+				FPCGTaggedData& Data = Outputs.Add_GetRef(Input);
+				Data.Pin = PCGMetadataBreakVectorConstants::ZLabel;
+			}
+
+			if (bWPinConnected)
+			{
+				FPCGTaggedData& Data = Outputs.Add_GetRef(Input);
+				Data.Pin = PCGMetadataBreakVectorConstants::WLabel;
+			}
+
+			continue;
+		}
+
+		if (!SourceMetadata)
+		{
+			PCGE_LOG(Warning, "Invalid metadata");
+			continue;
+		}
+
+		const FName LocalSourceAttributeName = (SourceAttributeName == NAME_None) ? SourceMetadata->GetSingleAttributeNameOrNone() : SourceAttributeName;
+		FName DestinationAttributeForX = NAME_None;
+		FName DestinationAttributeForY = NAME_None;
+		FName DestinationAttributeForZ = NAME_None;
+		FName DestinationAttributeForW = NAME_None;
+
+		if (LocalSourceAttributeName != NAME_None)
+		{
+			DestinationAttributeForX = FName(LocalSourceAttributeName.ToString() + ".X");
+			DestinationAttributeForY = FName(LocalSourceAttributeName.ToString() + ".Y");
+			DestinationAttributeForZ = FName(LocalSourceAttributeName.ToString() + ".Z");
+			DestinationAttributeForW = FName(LocalSourceAttributeName.ToString() + ".W");
+		}
+
+		const FPCGMetadataAttributeBase* SourceAttribute = SourceMetadata->GetConstAttribute(LocalSourceAttributeName);
+		if (!SourceAttribute)
+		{
+			PCGE_LOG(Error, "Attribute %s does not exist", *LocalSourceAttributeName.ToString());
+			continue;
+		}
+
+		if (SourceAttribute->GetTypeId() != PCG::Private::MetadataTypes<FVector>::Id 
+			&& SourceAttribute->GetTypeId() != PCG::Private::MetadataTypes<FVector4>::Id
+			&& SourceAttribute->GetTypeId() != PCG::Private::MetadataTypes<FRotator>::Id)
+		{
+			PCGE_LOG(Error, "Attribute %s is not a breakable type", *LocalSourceAttributeName.ToString());
+			continue;
+		}
+
+		FVector4 DefaultValue;
+
+		if (SourceAttribute->GetTypeId() == PCG::Private::MetadataTypes<FVector>::Id)
+		{
+			DefaultValue = static_cast<const FPCGMetadataAttribute<FVector>*>(SourceAttribute)->GetValue(PCGDefaultValueKey);
+		}
+		else if (SourceAttribute->GetTypeId() == PCG::Private::MetadataTypes<FVector4>::Id)
+		{
+			DefaultValue = static_cast<const FPCGMetadataAttribute<FVector4>*>(SourceAttribute)->GetValue(PCGDefaultValueKey);
+		}
+		else //if (SourceAttribute->GetTypeId() == PCG::Private::MetadataTypes<FRotator>::Id)
+		{
+			const FRotator Rotator = static_cast<const FPCGMetadataAttribute<FRotator>*>(SourceAttribute)->GetValue(PCGDefaultValueKey);
+			DefaultValue.X = Rotator.Roll;
+			DefaultValue.Y = Rotator.Pitch;
+			DefaultValue.Z = Rotator.Yaw;
+		}
+
+		FPCGMetadataAttribute<double>* AttributeX = nullptr;
+		FPCGMetadataAttribute<double>* AttributeY = nullptr;
+		FPCGMetadataAttribute<double>* AttributeZ = nullptr;
+		FPCGMetadataAttribute<double>* AttributeW = nullptr;
+
+		if (bXPinConnected)
+		{
+			FPCGTaggedData& OutputX = Outputs.Add_GetRef(Input);
+			OutputX.Pin = PCGMetadataBreakVectorConstants::XLabel;
+
+			UPCGMetadata* OutMetadata = nullptr;
+
+			PCGMetadataElementCommon::DuplicateTaggedData(Input, OutputX, OutMetadata);
+			AttributeX = PCGMetadataElementCommon::ClearOrCreateAttribute(OutMetadata, DestinationAttributeForX, DefaultValue.X);
+			PCGMetadataElementCommon::CopyEntryToValueKeyMap(SourceMetadata, SourceAttribute, AttributeX);
+		}
+
+		if (bYPinConnected)
+		{
+			FPCGTaggedData& OutputY = Outputs.Add_GetRef(Input);
+			OutputY.Pin = PCGMetadataBreakVectorConstants::YLabel;
+
+			UPCGMetadata* OutMetadata = nullptr;
+
+			PCGMetadataElementCommon::DuplicateTaggedData(Input, OutputY, OutMetadata);
+			AttributeY = PCGMetadataElementCommon::ClearOrCreateAttribute(OutMetadata, DestinationAttributeForY, DefaultValue.Y);
+			PCGMetadataElementCommon::CopyEntryToValueKeyMap(SourceMetadata, SourceAttribute, AttributeY);
+		}
+
+		if (bZPinConnected)
+		{
+			FPCGTaggedData& OutputZ = Outputs.Add_GetRef(Input);
+			OutputZ.Pin = PCGMetadataBreakVectorConstants::ZLabel;
+
+			UPCGMetadata* OutMetadata = nullptr;
+
+			PCGMetadataElementCommon::DuplicateTaggedData(Input, OutputZ, OutMetadata);
+			AttributeZ = PCGMetadataElementCommon::ClearOrCreateAttribute(OutMetadata, DestinationAttributeForZ, DefaultValue.Z);
+			PCGMetadataElementCommon::CopyEntryToValueKeyMap(SourceMetadata, SourceAttribute, AttributeZ);
+		}
+
+		if (bWPinConnected && SourceAttribute->GetTypeId() == PCG::Private::MetadataTypes<FVector4>::Id)
+		{
+			FPCGTaggedData& OutputW = Outputs.Add_GetRef(Input);
+			OutputW.Pin = PCGMetadataBreakVectorConstants::WLabel;
+
+			UPCGMetadata* OutMetadata = nullptr;
+
+			PCGMetadataElementCommon::DuplicateTaggedData(Input, OutputW, OutMetadata);
+			AttributeW = PCGMetadataElementCommon::ClearOrCreateAttribute(OutMetadata, DestinationAttributeForW, DefaultValue.W);
+			PCGMetadataElementCommon::CopyEntryToValueKeyMap(SourceMetadata, SourceAttribute, AttributeW);
+		}
+
+		// Copy all Value pairs from the parent hierarchy into our new attributes. Assumes that adding values in order of the parent hierarchy will result in identical ValueKey->Value mappings.
+		const PCGMetadataValueKey NumValueKeys = SourceAttribute->GetValueKeyOffsetForChild();
+		for (PCGMetadataValueKey ValueKey = 0; ValueKey < NumValueKeys; ++ValueKey)
+		{
+			FVector4 Value;
+
+			if (SourceAttribute->GetTypeId() == PCG::Private::MetadataTypes<FVector>::Id)
+			{
+				Value = static_cast<const FPCGMetadataAttribute<FVector>*>(SourceAttribute)->GetValue(ValueKey);
+			}
+			else if (SourceAttribute->GetTypeId() == PCG::Private::MetadataTypes<FVector4>::Id)
+			{
+				Value = static_cast<const FPCGMetadataAttribute<FVector4>*>(SourceAttribute)->GetValue(ValueKey);
+			}
+			else //if (SourceAttribute->GetTypeId() == PCG::Private::MetadataTypes<FRotator>::Id)
+			{
+				const FRotator Rotator = static_cast<const FPCGMetadataAttribute<FRotator>*>(SourceAttribute)->GetValue(ValueKey);
+				Value.X = Rotator.Roll;
+				Value.Y = Rotator.Pitch;
+				Value.Z = Rotator.Yaw;
+			}
+
+			if (AttributeX)
+			{
+				AttributeX->AddValue(Value.X);
+			}
+
+			if (AttributeY)
+			{
+				AttributeY->AddValue(Value.Y);
+			}
+
+			if (AttributeZ)
+			{
+				AttributeZ->AddValue(Value.Z);
+			}
+
+			if (AttributeW)
+			{
+				AttributeW->AddValue(Value.W);
+			}
+		}
+	}
+
+	return true;
+}
diff --git a/Engine/Plugins/Experimental/PCG/Source/PCG/Private/Elements/Metadata/PCGMetadataElement.cpp b/Engine/Plugins/Experimental/PCG/Source/PCG/Private/Elements/Metadata/PCGMetadataElement.cpp
new file mode 100644
--- /dev/null
+++ b/Engine/Plugins/Experimental/PCG/Source/PCG/Private/Elements/Metadata/PCGMetadataElement.cpp
@@ -0,0 +1,400 @@
+// Copyright Epic Games, Inc. All Rights Reserved.
+
+#include "Elements/Metadata/PCGMetadataElement.h"
+#include "Data/PCGPointData.h"
+#include "Helpers/PCGSettingsHelpers.h"
+#include "Metadata/PCGMetadata.h"
+#include "Metadata/PCGMetadataAttribute.h"
+#include "Metadata/PCGMetadataAttributeTpl.h"
+
+namespace PCGMetadataOperations
+{
+	template<typename U, typename T>
+	bool SetValue(TArray<FPCGPoint>& InPoints, FPCGMetadataAttributeBase* AttributeBase, UPCGMetadata* Metadata, TFunctionRef<U(const FPCGPoint& InPoint)> PropGetter)
+	{
+		if (AttributeBase && AttributeBase->GetTypeId() == PCG::Private::MetadataTypes<T>::Id)
+		{
+			FPCGMetadataAttribute<T>* Attribute = static_cast<FPCGMetadataAttribute<T>*>(AttributeBase);
+			for (FPCGPoint& Point : InPoints)
+			{
+				Metadata->InitializeOnSet(Point.MetadataEntry);
+				Attribute->SetValue(Point.MetadataEntry, PropGetter(Point));
+			}
+
+			return true;
+		}
+		else
+		{
+			return false;
+		}
+	}
+
+	template<typename T, typename U>
+	bool SetValue(const FPCGMetadataAttributeBase* AttributeBase, TArray<FPCGPoint>& InPoints, TFunctionRef<void(FPCGPoint& OutPoint, const U& InValue)> PropSetter)
+	{
+		if (AttributeBase && AttributeBase->GetTypeId() == PCG::Private::MetadataTypes<T>::Id)
+		{
+			const FPCGMetadataAttribute<T>* Attribute = static_cast<const FPCGMetadataAttribute<T>*>(AttributeBase);
+			for (FPCGPoint& Point : InPoints)
+			{
+				PropSetter(Point, U(Attribute->GetValueFromItemKey(Point.MetadataEntry)));
+			}
+
+			return true;
+		}
+		else
+		{
+			return false;
+		}
+	}
+}
+
+FPCGElementPtr UPCGMetadataOperationSettings::CreateElement() const
+{
+	return MakeShared<FPCGMetadataOperationElement>();
+}
+
+bool FPCGMetadataOperationElement::ExecuteInternal(FPCGContext* Context) const
+{
+	TRACE_CPUPROFILER_EVENT_SCOPE(FPCGMetadataOperationElement::Execute);
+
+	const UPCGMetadataOperationSettings* Settings = Context->GetInputSettings<UPCGMetadataOperationSettings>();
+	check(Settings);
+
+	TArray<FPCGTaggedData> Inputs = Context->InputData.GetInputs();
+	UPCGParamData* Params = Context->InputData.GetParams();
+
+	TArray<FPCGTaggedData>& Outputs = Context->OutputData.TaggedData;
+
+	const FName SourceAttribute = PCGSettingsHelpers::GetValue(GET_MEMBER_NAME_CHECKED(UPCGMetadataOperationSettings, SourceAttribute), Settings->SourceAttribute, Params);
+	const EPCGPointProperties PointProperty = PCGSettingsHelpers::GetValue(GET_MEMBER_NAME_CHECKED(UPCGMetadataOperationSettings, PointProperty), Settings->PointProperty, Params);
+	const FName DestinationAttribute = PCGSettingsHelpers::GetValue(GET_MEMBER_NAME_CHECKED(UPCGMetadataOperationSettings, DestinationAttribute), Settings->DestinationAttribute, Params);
+	const EPCGMetadataOperationTarget Target = PCGSettingsHelpers::GetValue(GET_MEMBER_NAME_CHECKED(UPCGMetadataOperationSettings, Target), Settings->Target, Params);
+
+	// Forward any non-input data
+	Outputs.Append(Context->InputData.GetAllSettings());
+
+	for (const FPCGTaggedData& Input : Inputs)
+	{
+		FPCGTaggedData& Output = Outputs.Add_GetRef(Input);
+
+		const UPCGSpatialData* SpatialInput = Cast<const UPCGSpatialData>(Input.Data);
+
+		if (!SpatialInput)
+		{
+			PCGE_LOG(Error, "Invalid input data");
+			continue;
+		}
+
+		const UPCGPointData* OriginalData = SpatialInput->ToPointData(Context);
+
+		if (!OriginalData)
+		{
+			PCGE_LOG(Error, "Unable to get point data from input");
+			continue;
+		}
+
+		if (!OriginalData->Metadata)
+		{
+			PCGE_LOG(Warning, "Input has no metadata");
+			continue;
+		}
+
+		const FName LocalSourceAttribute = ((SourceAttribute != NAME_None) ? SourceAttribute : OriginalData->Metadata->GetSingleAttributeNameOrNone());
+
+		// Check if the attribute exists
+		if ((Target == EPCGMetadataOperationTarget::AttributeToProperty || Target == EPCGMetadataOperationTarget::AttributeToAttribute) && !OriginalData->Metadata->HasAttribute(LocalSourceAttribute))
+		{
+			PCGE_LOG(Warning, "Input does not have the %s attribute", *LocalSourceAttribute.ToString());
+			continue;
+		}
+
+		const TArray<FPCGPoint>& Points = OriginalData->GetPoints();
+		const int OriginalPointCount = Points.Num();
+
+		UPCGPointData* SampledData = NewObject<UPCGPointData>();
+		SampledData->InitializeFromData(OriginalData);
+		TArray<FPCGPoint>& SampledPoints = SampledData->GetMutablePoints();
+
+		Output.Data = SampledData;
+
+		// Copy points and then apply the operation
+		SampledPoints = Points;
+
+		if (Target == EPCGMetadataOperationTarget::PropertyToAttribute)
+		{
+			if (PointProperty == EPCGPointProperties::Density)
+			{
+				auto DensityGetter = [](const FPCGPoint& InPoint) { return InPoint.Density; };
+
+				if (!SampledData->Metadata->HasAttribute(DestinationAttribute))
+				{
+					SampledData->Metadata->CreateFloatAttribute(DestinationAttribute, 0.0f, /*bAllowsInterpolation=*/true, /*bOverrideParent=*/true);
+				}
+
+				FPCGMetadataAttributeBase* AttributeBase = SampledData->Metadata->GetMutableAttribute(DestinationAttribute);
+				if(!PCGMetadataOperations::SetValue<float, float>(SampledPoints, AttributeBase, SampledData->Metadata, DensityGetter) &&
+					!PCGMetadataOperations::SetValue<float, FVector>(SampledPoints, AttributeBase, SampledData->Metadata, DensityGetter) &&
+					!PCGMetadataOperations::SetValue<float, FVector4>(SampledPoints, AttributeBase, SampledData->Metadata, DensityGetter))
+				{
+					PCGE_LOG(Error, "Attribute %s already exists but its type is not compatible", *DestinationAttribute.ToString());
+				}
+			}
+			else if (PointProperty == EPCGPointProperties::BoundsMin)
+			{
+				auto ExtentsGetter = [](const FPCGPoint& InPoint) { return InPoint.BoundsMin; };
+
+				if (!SampledData->Metadata->HasAttribute(DestinationAttribute))
+				{
+					SampledData->Metadata->CreateVectorAttribute(DestinationAttribute, FVector::Zero(), /*bAllowsInterpolation=*/true, /*bOverrideParent=*/true);
+				}
+
+				FPCGMetadataAttributeBase* AttributeBase = SampledData->Metadata->GetMutableAttribute(DestinationAttribute);
+				if (!PCGMetadataOperations::SetValue<FVector, FVector>(SampledPoints, AttributeBase, SampledData->Metadata, ExtentsGetter) &&
+					!PCGMetadataOperations::SetValue<FVector, FTransform>(SampledPoints, AttributeBase, SampledData->Metadata, ExtentsGetter))
+				{
+					PCGE_LOG(Error, "Attribute %s already exists but its type is not compatible", *DestinationAttribute.ToString());
+				}
+			}
+			else if (PointProperty == EPCGPointProperties::BoundsMax)
+			{
+				auto ExtentsGetter = [](const FPCGPoint& InPoint) { return InPoint.BoundsMax; };
+
+				if (!SampledData->Metadata->HasAttribute(DestinationAttribute))
+				{
+					SampledData->Metadata->CreateVectorAttribute(DestinationAttribute, FVector::Zero(), /*bAllowsInterpolation=*/true, /*bOverrideParent=*/true);
+				}
+
+				FPCGMetadataAttributeBase* AttributeBase = SampledData->Metadata->GetMutableAttribute(DestinationAttribute);
+				if (!PCGMetadataOperations::SetValue<FVector, FVector>(SampledPoints, AttributeBase, SampledData->Metadata, ExtentsGetter) &&
+					!PCGMetadataOperations::SetValue<FVector, FTransform>(SampledPoints, AttributeBase, SampledData->Metadata, ExtentsGetter))
+				{
+					PCGE_LOG(Error, "Attribute %s already exists but its type is not compatible", *DestinationAttribute.ToString());
+				}
+			}
+			else if (PointProperty == EPCGPointProperties::Extents)
+			{
+				auto ExtentsGetter = [](const FPCGPoint& InPoint) { return InPoint.GetExtents(); };
+
+				if (!SampledData->Metadata->HasAttribute(DestinationAttribute))
+				{
+					SampledData->Metadata->CreateVectorAttribute(DestinationAttribute, FVector::Zero() , /*bAllowsInterpolation=*/true, /*bOverrideParent=*/true);
+				}
+
+				FPCGMetadataAttributeBase* AttributeBase = SampledData->Metadata->GetMutableAttribute(DestinationAttribute);
+				if(!PCGMetadataOperations::SetValue<FVector, FVector>(SampledPoints, AttributeBase, SampledData->Metadata, ExtentsGetter) &&
+					!PCGMetadataOperations::SetValue<FVector, FTransform>(SampledPoints, AttributeBase, SampledData->Metadata, ExtentsGetter))
+				{
+					PCGE_LOG(Error, "Attribute %s already exists but its type is not compatible", *DestinationAttribute.ToString());
+				}
+			}
+			else if(PointProperty == EPCGPointProperties::Color)
+			{
+				auto ColorGetter = [](const FPCGPoint& InPoint) { return InPoint.Color; };
+
+				if (!SampledData->Metadata->HasAttribute(DestinationAttribute))
+				{
+					SampledData->Metadata->CreateVector4Attribute(DestinationAttribute, FVector4::Zero(), /*bAllowsInterpolation=*/true, /*bOverrideParent=*/true);
+				}
+
+				FPCGMetadataAttributeBase* AttributeBase = SampledData->Metadata->GetMutableAttribute(DestinationAttribute);
+				if(!PCGMetadataOperations::SetValue<FVector4, FVector4>(SampledPoints, AttributeBase, SampledData->Metadata, ColorGetter))
+				{
+					PCGE_LOG(Error, "Attribute %s already exists but its type is not compatible", *DestinationAttribute.ToString());
+				}
+			}
+			else if (PointProperty == EPCGPointProperties::Position)
+			{
+				auto PositionGetter = [](const FPCGPoint& InPoint) { return InPoint.Transform.GetLocation(); };
+
+				if (!SampledData->Metadata->HasAttribute(DestinationAttribute))
+				{
+					SampledData->Metadata->CreateVectorAttribute(DestinationAttribute, FVector::Zero(), /*bAllowsInterpolation=*/true, /*bOverrideParent=*/true);
+				}
+
+				FPCGMetadataAttributeBase* AttributeBase = SampledData->Metadata->GetMutableAttribute(DestinationAttribute);
+				if (!PCGMetadataOperations::SetValue<FVector, FVector>(SampledPoints, AttributeBase, SampledData->Metadata, PositionGetter) &&
+					!PCGMetadataOperations::SetValue<FVector, FTransform>(SampledPoints, AttributeBase, SampledData->Metadata, PositionGetter))
+				{
+					PCGE_LOG(Error, "Attribute %s already exists but its type is not compatible", *DestinationAttribute.ToString());
+				}
+			}
+			else if (PointProperty == EPCGPointProperties::Rotation)
+			{
+				auto RotationGetter = [](const FPCGPoint& InPoint) { return InPoint.Transform.GetRotation(); };
+
+				if (!SampledData->Metadata->HasAttribute(DestinationAttribute))
+				{
+					SampledData->Metadata->CreateQuatAttribute(DestinationAttribute, FQuat::Identity, /*bAllowsInterpolation=*/true, /*bOverrideParent=*/true);
+				}
+
+				FPCGMetadataAttributeBase* AttributeBase = SampledData->Metadata->GetMutableAttribute(DestinationAttribute);
+				if (!PCGMetadataOperations::SetValue<FQuat, FQuat>(SampledPoints, AttributeBase, SampledData->Metadata, RotationGetter))
+				{
+					PCGE_LOG(Error, "Attribute %s already exists but its type is not compatible", *DestinationAttribute.ToString());
+				}
+			}
+			else if (PointProperty == EPCGPointProperties::Scale)
+			{
+				auto ScaleGetter = [](const FPCGPoint& InPoint) { return InPoint.Transform.GetScale3D(); };
+
+				if (!SampledData->Metadata->HasAttribute(DestinationAttribute))
+				{
+					SampledData->Metadata->CreateVectorAttribute(DestinationAttribute, FVector::One(), /*bAllowsInterpolation=*/true, /*bOverrideParent=*/true);
+				}
+
+				FPCGMetadataAttributeBase* AttributeBase = SampledData->Metadata->GetMutableAttribute(DestinationAttribute);
+				if (!PCGMetadataOperations::SetValue<FVector, FVector>(SampledPoints, AttributeBase, SampledData->Metadata, ScaleGetter))
+				{
+					PCGE_LOG(Error, "Attribute %s already exists but its type is not compatible", *DestinationAttribute.ToString());
+				}
+			}
+			else if (PointProperty == EPCGPointProperties::Transform)
+			{
+				auto TransformGetter = [](const FPCGPoint& InPoint) { return InPoint.Transform; };
+
+				if (!SampledData->Metadata->HasAttribute(DestinationAttribute))
+				{
+					SampledData->Metadata->CreateTransformAttribute(DestinationAttribute, FTransform::Identity, /*bAllowsInterpolation=*/true, /*bOverrideParent=*/true);
+				}
+
+				FPCGMetadataAttributeBase* AttributeBase = SampledData->Metadata->GetMutableAttribute(DestinationAttribute);
+				if (!PCGMetadataOperations::SetValue<FTransform, FTransform>(SampledPoints, AttributeBase, SampledData->Metadata, TransformGetter))
+				{
+					PCGE_LOG(Error, "Attribute %s already exists but its type is not compatible", *DestinationAttribute.ToString());
+				}
+			}
+			else if (PointProperty == EPCGPointProperties::Steepness)
+			{
+				auto SteepnessGetter = [](const FPCGPoint& InPoint) { return InPoint.Steepness; };
+
+				if (!SampledData->Metadata->HasAttribute(DestinationAttribute))
+				{
+					SampledData->Metadata->CreateFloatAttribute(DestinationAttribute, 0.5f, /*bAllowsInterpolation=*/true, /*bOverrideParent=*/true);
+				}
+
+				FPCGMetadataAttributeBase* AttributeBase = SampledData->Metadata->GetMutableAttribute(DestinationAttribute);
+				if (!PCGMetadataOperations::SetValue<float, float>(SampledPoints, AttributeBase, SampledData->Metadata, SteepnessGetter))
+				{
+					PCGE_LOG(Error, "Attribute %s already exists but its type is not compatible", *DestinationAttribute.ToString());
+				}
+			}
+		}
+		else if(Target == EPCGMetadataOperationTarget::AttributeToProperty) // Attribute to property
+		{
+			if (PointProperty == EPCGPointProperties::Density)
+			{
+				auto DensitySetter = [](FPCGPoint& InPoint, const float& InValue) { InPoint.Density = InValue; };
+
+				const FPCGMetadataAttributeBase* AttributeBase = SampledData->Metadata->GetConstAttribute(LocalSourceAttribute);
+				if(!PCGMetadataOperations::SetValue<float, float>(AttributeBase, SampledPoints, DensitySetter))
+				{
+					PCGE_LOG(Error, "Attribute %s already exists but its type is not compatible", *LocalSourceAttribute.ToString());
+				}
+			}
+			else if (PointProperty == EPCGPointProperties::BoundsMin)
+			{
+				auto ExtentsSetter = [](FPCGPoint& InPoint, const FVector& InValue) { InPoint.BoundsMin = InValue; };
+
+				const FPCGMetadataAttributeBase* AttributeBase = SampledData->Metadata->GetConstAttribute(LocalSourceAttribute);
+				if (!PCGMetadataOperations::SetValue<float, FVector>(AttributeBase, SampledPoints, ExtentsSetter) &&
+					!PCGMetadataOperations::SetValue<FVector, FVector>(AttributeBase, SampledPoints, ExtentsSetter))
+				{
+					PCGE_LOG(Error, "Attribute %s already exists but its type is not compatible", *LocalSourceAttribute.ToString());
+				}
+			}
+			else if (PointProperty == EPCGPointProperties::BoundsMax)
+			{
+				auto ExtentsSetter = [](FPCGPoint& InPoint, const FVector& InValue) { InPoint.BoundsMax = InValue; };
+
+				const FPCGMetadataAttributeBase* AttributeBase = SampledData->Metadata->GetConstAttribute(LocalSourceAttribute);
+				if (!PCGMetadataOperations::SetValue<float, FVector>(AttributeBase, SampledPoints, ExtentsSetter) &&
+					!PCGMetadataOperations::SetValue<FVector, FVector>(AttributeBase, SampledPoints, ExtentsSetter))
+				{
+					PCGE_LOG(Error, "Attribute %s already exists but its type is not compatible", *LocalSourceAttribute.ToString());
+				}
+			}
+			else if (PointProperty == EPCGPointProperties::Extents)
+			{
+				auto ExtentsSetter = [](FPCGPoint& InPoint, const FVector& InValue) { InPoint.SetExtents(InValue); };
+
+				const FPCGMetadataAttributeBase* AttributeBase = SampledData->Metadata->GetConstAttribute(LocalSourceAttribute);
+				if(!PCGMetadataOperations::SetValue<float, FVector>(AttributeBase, SampledPoints, ExtentsSetter) &&
+					!PCGMetadataOperations::SetValue<FVector, FVector>(AttributeBase, SampledPoints, ExtentsSetter))
+				{
+					PCGE_LOG(Error, "Attribute %s already exists but its type is not compatible", *LocalSourceAttribute.ToString());
+				}
+			}
+			else if (PointProperty == EPCGPointProperties::Color)
+			{
+				auto ColorSetter = [](FPCGPoint& InPoint, const FVector4& InValue) { InPoint.Color = InValue; };
+
+				const FPCGMetadataAttributeBase* AttributeBase = SampledData->Metadata->GetConstAttribute(LocalSourceAttribute);
+				if(!PCGMetadataOperations::SetValue<float, FVector4>(AttributeBase, SampledPoints, ColorSetter) &&
+					!PCGMetadataOperations::SetValue<FVector4, FVector4>(AttributeBase, SampledPoints, ColorSetter))
+				{
+					PCGE_LOG(Error, "Attribute %s already exists but its type is not compatible", *LocalSourceAttribute.ToString());
+				}
+			}
+			else if (PointProperty == EPCGPointProperties::Position)
+			{
+				auto PositionSetter = [](FPCGPoint& InPoint, const FVector& InValue) { InPoint.Transform.SetLocation(InValue); };
+
+				const FPCGMetadataAttributeBase* AttributeBase = SampledData->Metadata->GetConstAttribute(LocalSourceAttribute);
+				if (!PCGMetadataOperations::SetValue<float, FVector>(AttributeBase, SampledPoints, PositionSetter) &&
+					!PCGMetadataOperations::SetValue<FVector, FVector>(AttributeBase, SampledPoints, PositionSetter))
+				{
+					PCGE_LOG(Error, "Attribute %s already exists but its type is not compatible", *LocalSourceAttribute.ToString());
+				}
+			}
+			else if (PointProperty == EPCGPointProperties::Rotation)
+			{
+				auto RotationSetter = [](FPCGPoint& InPoint, const FQuat& InValue) { InPoint.Transform.SetRotation(InValue.GetNormalized()); };
+
+				const FPCGMetadataAttributeBase* AttributeBase = SampledData->Metadata->GetConstAttribute(LocalSourceAttribute);
+				if (!PCGMetadataOperations::SetValue<FQuat, FQuat>(AttributeBase, SampledPoints, RotationSetter))
+				{
+					PCGE_LOG(Error, "Attribute %s already exists but its type is not compatible", *LocalSourceAttribute.ToString());
+				}
+			}
+			else if (PointProperty == EPCGPointProperties::Scale)
+			{
+				auto ScaleSetter = [](FPCGPoint& InPoint, const FVector& InValue) { InPoint.Transform.SetScale3D(InValue); };
+
+				const FPCGMetadataAttributeBase* AttributeBase = SampledData->Metadata->GetConstAttribute(LocalSourceAttribute);
+				if (!PCGMetadataOperations::SetValue<float, FVector>(AttributeBase, SampledPoints, ScaleSetter) &&
+					!PCGMetadataOperations::SetValue<FVector, FVector>(AttributeBase, SampledPoints, ScaleSetter))
+				{
+					PCGE_LOG(Error, "Attribute %s already exists but its type is not compatible", *LocalSourceAttribute.ToString());
+				}
+			}
+			else if (PointProperty == EPCGPointProperties::Transform)
+			{
+				auto TransformSetter = [](FPCGPoint& InPoint, const FTransform& InValue) { InPoint.Transform = InValue; };
+
+				const FPCGMetadataAttributeBase* AttributeBase = SampledData->Metadata->GetConstAttribute(LocalSourceAttribute);
+				if (!PCGMetadataOperations::SetValue<FVector, FTransform>(AttributeBase, SampledPoints, TransformSetter) &&
+					!PCGMetadataOperations::SetValue<FTransform, FTransform>(AttributeBase, SampledPoints, TransformSetter))
+				{
+					PCGE_LOG(Error, "Attribute %s already exists but its type is not compatible", *LocalSourceAttribute.ToString());
+				}
+			}
+			else if (PointProperty == EPCGPointProperties::Steepness)
+			{
+				auto SteepnessSetter = [](FPCGPoint& InPoint, const float& InValue) { InPoint.Steepness = InValue; };
+
+				const FPCGMetadataAttributeBase* AttributeBase = SampledData->Metadata->GetConstAttribute(LocalSourceAttribute);
+				if (!PCGMetadataOperations::SetValue<float, float>(AttributeBase, SampledPoints, SteepnessSetter))
+				{
+					PCGE_LOG(Error, "Attribute %s already exists but its type is not compatible", *LocalSourceAttribute.ToString());
+				}
+			}
+		}
+		else // Attribute to attribute
+		{
+			SampledData->Metadata->CopyExistingAttribute(LocalSourceAttribute, DestinationAttribute);
+		}
+	}
+
+	return true;
+}
\ No newline at end of file
diff --git a/Engine/Plugins/Experimental/PCG/Source/PCG/Private/Elements/Metadata/PCGMetadataElementCommon.cpp b/Engine/Plugins/Experimental/PCG/Source/PCG/Private/Elements/Metadata/PCGMetadataElementCommon.cpp
new file mode 100644
--- /dev/null
+++ b/Engine/Plugins/Experimental/PCG/Source/PCG/Private/Elements/Metadata/PCGMetadataElementCommon.cpp
@@ -0,0 +1,46 @@
+// Copyright Epic Games, Inc. All Rights Reserved.
+
+#include "Elements/Metadata/PCGMetadataElementCommon.h"
+
+#include "PCGParamData.h"
+#include "Data/PCGSpatialData.h"
+#include "Helpers/PCGSettingsHelpers.h"
+
+namespace PCGMetadataElementCommon
+{
+	void DuplicateTaggedData(const FPCGTaggedData& InTaggedData, FPCGTaggedData& OutTaggedData, UPCGMetadata*& OutMetadata)
+	{
+		if (const UPCGSpatialData* SpatialInput = Cast<const UPCGSpatialData>(InTaggedData.Data))
+		{
+			UPCGSpatialData* NewSpatialData = Cast<UPCGSpatialData>(StaticDuplicateObject(SpatialInput, const_cast<UPCGSpatialData*>(SpatialInput), FName()));
+			NewSpatialData->InitializeFromData(SpatialInput);
+			OutTaggedData.Data = NewSpatialData;
+
+			OutMetadata = NewSpatialData->Metadata;
+		}
+		else if (const UPCGParamData* ParamsInput = Cast<const UPCGParamData>(InTaggedData.Data))
+		{
+			UPCGParamData* NewParamData = Cast<UPCGParamData>(StaticDuplicateObject(ParamsInput, const_cast<UPCGParamData*>(ParamsInput), FName()));
+			OutTaggedData.Data = NewParamData;
+
+			NewParamData->Metadata->Initialize(ParamsInput->Metadata);
+			OutMetadata = NewParamData->Metadata;
+		}
+	}
+
+	void CopyEntryToValueKeyMap(const UPCGMetadata* MetadataToCopy, const FPCGMetadataAttributeBase* AttributeToCopy, FPCGMetadataAttributeBase* OutAttribute)
+	{
+		if (!OutAttribute)
+		{
+			UE_LOG(LogPCG, Error, TEXT("Failed to create output attribute"));
+			return;
+		}
+
+		const PCGMetadataEntryKey EntryKeyCount = MetadataToCopy->GetItemCountForChild();
+		for (PCGMetadataEntryKey EntryKey = 0; EntryKey < EntryKeyCount; ++EntryKey)
+		{
+			const PCGMetadataValueKey ValueKey = AttributeToCopy->GetValueKey(EntryKey);
+			OutAttribute->SetValueFromValueKey(EntryKey, ValueKey);
+		}
+	}
+}
\ No newline at end of file
diff --git a/Engine/Plugins/Experimental/PCG/Source/PCG/Private/Elements/Metadata/PCGMetadataPartition.cpp b/Engine/Plugins/Experimental/PCG/Source/PCG/Private/Elements/Metadata/PCGMetadataPartition.cpp
new file mode 100644
--- /dev/null
+++ b/Engine/Plugins/Experimental/PCG/Source/PCG/Private/Elements/Metadata/PCGMetadataPartition.cpp
@@ -0,0 +1,126 @@
+// Copyright Epic Games, Inc. All Rights Reserved.
+
+#include "Elements/Metadata/PCGMetadataPartition.h"
+#include "Data/PCGSpatialData.h"
+#include "Data/PCGPointData.h"
+#include "Helpers/PCGSettingsHelpers.h"
+
+FPCGElementPtr UPCGMetadataPartitionSettings::CreateElement() const
+{
+	return MakeShared<FPCGMetadataPartitionElement>();
+}
+
+bool FPCGMetadataPartitionElement::ExecuteInternal(FPCGContext* Context) const
+{
+	TRACE_CPUPROFILER_EVENT_SCOPE(FPCGMetadataPartitionElement::Execute);
+	check(Context);
+
+	const UPCGMetadataPartitionSettings* Settings = Context->GetInputSettings<UPCGMetadataPartitionSettings>();
+	check(Settings);
+
+	TArray<FPCGTaggedData> Inputs = Context->InputData.GetInputs();
+	UPCGParamData* Params = Context->InputData.GetParams();
+	TArray<FPCGTaggedData>& Outputs = Context->OutputData.TaggedData;
+
+	// Forward any non-input data
+	Outputs.Append(Context->InputData.GetAllSettings());
+
+	const FName PartitionAttribute = PCGSettingsHelpers::GetValue(GET_MEMBER_NAME_CHECKED(UPCGMetadataPartitionSettings, PartitionAttribute), Settings->PartitionAttribute, Params);
+
+	for (const FPCGTaggedData& Input : Inputs)
+	{
+		const UPCGSpatialData* SpatialInput = Cast<const UPCGSpatialData>(Input.Data);
+		check(SpatialInput);
+
+		if (!SpatialInput->ConstMetadata())
+		{
+			PCGE_LOG(Error, "Input does not have metadata");
+			continue;
+		}
+
+		const FName LocalPartitionAttribute = ((PartitionAttribute != NAME_None) ? PartitionAttribute : SpatialInput->ConstMetadata()->GetSingleAttributeNameOrNone());
+
+		if (!SpatialInput->ConstMetadata()->HasAttribute(LocalPartitionAttribute))
+		{
+			PCGE_LOG(Error, "Input does not have the %s attribute", *LocalPartitionAttribute.ToString());
+			continue;
+		}
+
+		// Currently, we support only to parse points, so we'll do so here
+		const UPCGPointData* InputData = SpatialInput->ToPointData(Context);
+
+		if (!InputData)
+		{
+			PCGE_LOG(Error, "Unable to get points out of spatial data");
+			continue;
+		}
+
+		const TArray<FPCGPoint>& SourcePoints = InputData->GetPoints();
+		const UPCGMetadata* SourceMetadata = InputData->ConstMetadata();
+		const FPCGMetadataAttributeBase* AttributeBase = SourceMetadata->GetConstAttribute(LocalPartitionAttribute);
+		check(AttributeBase);
+
+		TMap<PCGMetadataValueKey, UPCGPointData*> ValueToData;
+
+		// Get all value keys (-1 + 0 - N)
+		int64 MetadataValueKeyCount = AttributeBase->GetValueKeyOffsetForChild();
+
+		// For every value key, check if it should be merged with the default value
+		// Keep track of that mapping (should be only one, but that might not be true in uncompressed data)
+		TArray<PCGMetadataValueKey> ValueKeyMapping;
+		ValueKeyMapping.Reserve(MetadataValueKeyCount);
+
+		for (PCGMetadataValueKey ValueKey = 0; ValueKey < MetadataValueKeyCount; ++ValueKey)
+		{
+			if (AttributeBase->IsEqualToDefaultValue(ValueKey))
+			{
+				ValueKeyMapping.Add(-1);
+			}
+			else
+			{
+				ValueKeyMapping.Add(ValueKey);
+			}
+		}
+
+		TArray<UPCGPointData*> PartitionedData;
+		PartitionedData.SetNum(1 + ValueKeyMapping.Num());
+		for (int32 Index = 0; Index < PartitionedData.Num(); ++Index)
+		{
+			PartitionedData[Index] = nullptr;
+		}
+
+		// Loop on points, find matching data pointer, allocate if null, etc.and add point to it
+		for (const FPCGPoint& Point : SourcePoints)
+		{
+			PCGMetadataValueKey ValueKey = AttributeBase->GetValueKey(Point.MetadataEntry);
+			// Remap
+			if (ValueKey != PCGDefaultValueKey)
+			{
+				ValueKey = ValueKeyMapping[ValueKey];
+			}
+
+			const int32 PartitionDataIndex = 1 + ValueKey;
+
+			if (!PartitionedData[PartitionDataIndex])
+			{
+				PartitionedData[PartitionDataIndex] = NewObject<UPCGPointData>();
+				PartitionedData[PartitionDataIndex]->InitializeFromData(InputData);
+			}
+
+			PartitionedData[PartitionDataIndex]->GetMutablePoints().Add(Point);
+		}
+
+		// Finally, push back the partitioned data to the outputs
+		for (UPCGPointData* PartitionData : PartitionedData)
+		{
+			if (PartitionData)
+			{
+				FPCGTaggedData& Output = Outputs.Emplace_GetRef();
+				Output.Data = PartitionData;
+				Output.Tags = Input.Tags;
+			}
+		}
+	}
+
+	return true;
+}
\ No newline at end of file
diff --git a/Engine/Plugins/Experimental/PCG/Source/PCG/Private/Elements/Metadata/PCGMetadataRenameElement.cpp b/Engine/Plugins/Experimental/PCG/Source/PCG/Private/Elements/Metadata/PCGMetadataRenameElement.cpp
new file mode 100644
--- /dev/null
+++ b/Engine/Plugins/Experimental/PCG/Source/PCG/Private/Elements/Metadata/PCGMetadataRenameElement.cpp
@@ -0,0 +1,118 @@
+// Copyright Epic Games, Inc. All Rights Reserved.
+
+#include "Elements/Metadata/PCGMetadataRenameElement.h"
+#include "Data/PCGSpatialData.h"
+#include "Helpers/PCGSettingsHelpers.h"
+#include "Metadata/PCGMetadata.h"
+#include "Metadata/PCGMetadataAttribute.h"
+
+TArray<FPCGPinProperties> UPCGMetadataRenameSettings::OutputPinProperties() const
+{
+	TArray<FPCGPinProperties> PinProperties;
+	PinProperties.Emplace(PCGPinConstants::DefaultOutputLabel, EPCGDataType::Any);
+
+	return PinProperties;
+}
+
+FPCGElementPtr UPCGMetadataRenameSettings::CreateElement() const
+{
+	return MakeShared<FPCGMetadataRenameElement>();
+}
+
+bool FPCGMetadataRenameElement::ExecuteInternal(FPCGContext* Context) const
+{
+	TRACE_CPUPROFILER_EVENT_SCOPE(FPCGMetadataOperationElement::Execute);
+
+	const UPCGMetadataRenameSettings* Settings = Context->GetInputSettings<UPCGMetadataRenameSettings>();
+	check(Settings);
+
+	TArray<FPCGTaggedData> Inputs = Context->InputData.GetInputs();
+	TArray<FPCGTaggedData> AllParams = Context->InputData.GetAllParams();
+	UPCGParamData* Params = Context->InputData.GetParams();
+
+	TArray<FPCGTaggedData>& Outputs = Context->OutputData.TaggedData;
+	
+	const FName AttributeToRename = PCGSettingsHelpers::GetValue(GET_MEMBER_NAME_CHECKED(UPCGMetadataRenameSettings, AttributeToRename), Settings->AttributeToRename, Params);
+	const FName NewAttributeName = PCGSettingsHelpers::GetValue(GET_MEMBER_NAME_CHECKED(UPCGMetadataRenameSettings, NewAttributeName), Settings->NewAttributeName, Params);
+
+	if (NewAttributeName == NAME_None)
+	{
+		UE_LOG(LogPCG, Warning, TEXT("Metadata rename operation cannot rename from %s to %s"), *AttributeToRename.ToString(), *NewAttributeName.ToString());
+		// Bypass
+		Context->OutputData = Context->InputData;
+		return true;
+	}
+
+	for (const FPCGTaggedData& Input : Inputs)
+	{
+		FPCGTaggedData& Output = Outputs.Add_GetRef(Input);
+
+		const UPCGSpatialData* SpatialInput = Cast<const UPCGSpatialData>(Input.Data);
+		if (!SpatialInput)
+		{
+			continue;
+		}
+
+		// If the data has a metadata & the attribute to rename, then duplicate the data
+		// otherwise, keep it as is.
+		const UPCGMetadata* Metadata = SpatialInput->Metadata;
+
+		if (!Metadata)
+		{
+			continue;
+		}
+
+		const FName LocalAttributeToRename = ((AttributeToRename != NAME_None) ? AttributeToRename : Metadata->GetSingleAttributeNameOrNone());
+
+		if (!Metadata->HasAttribute(LocalAttributeToRename))
+		{
+			continue;
+		}
+
+		//TODO: this might require to execute on the main thread
+		UPCGSpatialData* NewSpatialData = Cast<UPCGSpatialData>(StaticDuplicateObject(SpatialInput, const_cast<UPCGSpatialData*>(SpatialInput), FName()));
+		NewSpatialData->InitializeFromData(SpatialInput);
+		NewSpatialData->Metadata->RenameAttribute(LocalAttributeToRename, NewAttributeName);
+
+		Output.Data = NewSpatialData;
+	}
+
+	for (const FPCGTaggedData& ParamTaggedData : AllParams)
+	{
+		FPCGTaggedData& Output = Outputs.Add_GetRef(ParamTaggedData);
+
+		const UPCGParamData* InputParams = Cast<const UPCGParamData>(ParamTaggedData.Data);
+		if (!InputParams)
+		{
+			continue;
+		}
+		
+		const UPCGMetadata* Metadata = InputParams->Metadata;
+
+		if (!Metadata)
+		{
+			continue;
+		}
+
+		const FName LocalAttributeToRename = ((AttributeToRename != NAME_None) ? AttributeToRename : Metadata->GetSingleAttributeNameOrNone());
+
+		if (!Metadata->HasAttribute(LocalAttributeToRename))
+		{
+			continue;
+		}
+
+		//TODO: this might require to execute on the main thread
+		UPCGParamData* NewParams = Cast<UPCGParamData>(StaticDuplicateObject(InputParams, const_cast<UPCGParamData*>(InputParams), FName()));
+		// Note: we will not parent the metadata here to ensure that the 0th entry is present on this metadata
+		// We will instead do a copy (and keep its parent, etc.)
+		NewParams->Metadata->InitializeAsCopy(Metadata);
+		NewParams->Metadata->RenameAttribute(LocalAttributeToRename, NewAttributeName);
+
+		Output.Data = NewParams;
+	}
+
+	// Pass-through settings
+	Outputs.Append(Context->InputData.GetAllSettings());
+
+	return true;
+}
\ No newline at end of file
diff --git a/Engine/Plugins/Experimental/PCG/Source/PCG/Private/Elements/PCGMetadataBreakVector.cpp b/Engine/Plugins/Experimental/PCG/Source/PCG/Private/Elements/PCGMetadataBreakVector.cpp
deleted file mode 100644
diff --git a/Engine/Plugins/Experimental/PCG/Source/PCG/Private/Elements/PCGMetadataElement.cpp b/Engine/Plugins/Experimental/PCG/Source/PCG/Private/Elements/PCGMetadataElement.cpp
deleted file mode 100644
diff --git a/Engine/Plugins/Experimental/PCG/Source/PCG/Private/Elements/PCGMetadataPartition.cpp b/Engine/Plugins/Experimental/PCG/Source/PCG/Private/Elements/PCGMetadataPartition.cpp
deleted file mode 100644
diff --git a/Engine/Plugins/Experimental/PCG/Source/PCG/Private/Elements/PCGMetadataRenameElement.cpp b/Engine/Plugins/Experimental/PCG/Source/PCG/Private/Elements/PCGMetadataRenameElement.cpp
deleted file mode 100644
diff --git a/Engine/Plugins/Experimental/PCG/Source/PCG/Private/Helpers/PCGMetadataHelpers.cpp b/Engine/Plugins/Experimental/PCG/Source/PCG/Private/Helpers/PCGMetadataHelpers.cpp
new file mode 100644
--- /dev/null
+++ b/Engine/Plugins/Experimental/PCG/Source/PCG/Private/Helpers/PCGMetadataHelpers.cpp
@@ -0,0 +1,12 @@
+// Copyright Epic Games, Inc. All Rights Reserved.
+
+#include "Helpers/PCGMetadataHelpers.h"
+#include "Metadata/PCGMetadata.h"
+
+namespace PCGMetadataHelpers
+{
+	bool HasSameRoot(const UPCGMetadata* Metadata1, const UPCGMetadata* Metadata2)
+	{
+		return Metadata1 && Metadata2 && Metadata1->GetRoot() == Metadata2->GetRoot();
+	}
+}
\ No newline at end of file
diff --git a/Engine/Plugins/Experimental/PCG/Source/PCG/Private/Metadata/PCGMetadata.cpp b/Engine/Plugins/Experimental/PCG/Source/PCG/Private/Metadata/PCGMetadata.cpp
--- a/Engine/Plugins/Experimental/PCG/Source/PCG/Private/Metadata/PCGMetadata.cpp
+++ b/Engine/Plugins/Experimental/PCG/Source/PCG/Private/Metadata/PCGMetadata.cpp
@@ -1,8 +1,6 @@
 // Copyright Epic Games, Inc. All Rights Reserved.
 
 #include "Metadata/PCGMetadata.h"
-#include "Metadata/PCGMetadataAttributeTpl.h"
-#include "PCGModule.h"
 
 #include "Misc/ScopeRWLock.h"
 #include "Helpers/PCGSettingsHelpers.h"
@@ -235,35 +233,6 @@ void UPCGMetadata::RemoveAttributeInternal(FName AttributeName)
 	Attributes.Remove(AttributeName);
 }
 
-template<typename T>
-FPCGMetadataAttributeBase* UPCGMetadata::CreateAttribute(FName AttributeName, const T& DefaultValue, bool bAllowsInterpolation, bool bOverrideParent)
-{
-	const FPCGMetadataAttributeBase* ParentAttribute = nullptr;
-
-	if (bOverrideParent && Parent.IsValid())
-	{
-		ParentAttribute = Parent->GetConstAttribute(AttributeName);
-	}
-
-	FPCGMetadataAttributeBase* NewAttribute = new FPCGMetadataAttribute<T>(this, AttributeName, ParentAttribute, DefaultValue, bAllowsInterpolation);
-
-	AttributeLock.WriteLock();
-	if (FPCGMetadataAttributeBase** ExistingAttribute = Attributes.Find(AttributeName))
-	{
-		UE_LOG(LogPCG, Warning, TEXT("Attribute %s already exists"), *AttributeName.ToString());
-		delete NewAttribute;
-		NewAttribute = *ExistingAttribute;
-	}
-	else
-	{
-		NewAttribute->AttributeId = NextAttributeId++;
-		AddAttributeInternal(AttributeName, NewAttribute);
-	}
-	AttributeLock.WriteUnlock();
-
-	return NewAttribute;
-}
-
 FPCGMetadataAttributeBase* UPCGMetadata::GetMutableAttribute(FName AttributeName)
 {
 	FPCGMetadataAttributeBase* Attribute = nullptr;
diff --git a/Engine/Plugins/Experimental/PCG/Source/PCG/Private/Metadata/PCGMetadataAttribute.cpp b/Engine/Plugins/Experimental/PCG/Source/PCG/Private/Metadata/PCGMetadataAttribute.cpp
--- a/Engine/Plugins/Experimental/PCG/Source/PCG/Private/Metadata/PCGMetadataAttribute.cpp
+++ b/Engine/Plugins/Experimental/PCG/Source/PCG/Private/Metadata/PCGMetadataAttribute.cpp
@@ -1,6 +1,7 @@
 // Copyright Epic Games, Inc. All Rights Reserved.
 
 #include "Metadata/PCGMetadataAttribute.h"
+#include "Metadata/PCGMetadata.h"
 
 FPCGMetadataAttributeBase::FPCGMetadataAttributeBase(UPCGMetadata* InMetadata, FName InName, const FPCGMetadataAttributeBase* InParent, bool bInAllowsInterpolation)
 	: Metadata(InMetadata)
diff --git a/Engine/Plugins/Experimental/PCG/Source/PCG/Private/Tests/Elements/PCGMetadataBreakVectorTest.cpp b/Engine/Plugins/Experimental/PCG/Source/PCG/Private/Tests/Elements/PCGMetadataBreakVectorTest.cpp
--- a/Engine/Plugins/Experimental/PCG/Source/PCG/Private/Tests/Elements/PCGMetadataBreakVectorTest.cpp
+++ b/Engine/Plugins/Experimental/PCG/Source/PCG/Private/Tests/Elements/PCGMetadataBreakVectorTest.cpp
@@ -8,7 +8,7 @@
 #include "PCGParamData.h"
 
 #include "Data/PCGPointData.h"
-#include "Elements/PCGMetadataBreakVector.h"
+#include "Elements/Metadata/PCGMetadataBreakVector.h"
 #include "Metadata/PCGMetadata.h"
 #include "Metadata/PCGMetadataAttribute.h"
 #include "Metadata/PCGMetadataAttributeTpl.h"
diff --git a/Engine/Plugins/Experimental/PCG/Source/PCG/Public/Elements/Metadata/PCGMetadataBreakVector.h b/Engine/Plugins/Experimental/PCG/Source/PCG/Public/Elements/Metadata/PCGMetadataBreakVector.h
new file mode 100644
--- /dev/null
+++ b/Engine/Plugins/Experimental/PCG/Source/PCG/Public/Elements/Metadata/PCGMetadataBreakVector.h
@@ -0,0 +1,59 @@
+// Copyright Epic Games, Inc. All Rights Reserved.
+
+#pragma once
+
+#include "PCGElement.h"
+#include "PCGSettings.h"
+
+#include "PCGMetadataBreakVector.generated.h"
+
+namespace PCGMetadataBreakVectorConstants
+{
+	const FName ParamsLabel = TEXT("Params");
+	const FName SourceLabel = TEXT("Source");
+	const FName XLabel = TEXT("X");
+	const FName YLabel = TEXT("Y");
+	const FName ZLabel = TEXT("Z");
+	const FName WLabel = TEXT("W");
+}
+
+UCLASS(BlueprintType, ClassGroup = (Procedural))
+class PCG_API UPCGMetadataBreakVectorSettings : public UPCGSettings
+{
+	GENERATED_BODY()
+
+public:
+	//~Begin UPCGSettings interface
+#if WITH_EDITOR
+	virtual FName GetDefaultNodeName() const override { return FName(TEXT("MetadataBreakVectorNode")); }
+	virtual EPCGSettingsType GetType() const override { return EPCGSettingsType::Metadata; }
+#endif
+
+	virtual TArray<FPCGPinProperties> InputPinProperties() const override;
+	virtual TArray<FPCGPinProperties> OutputPinProperties() const override;
+
+protected:
+	virtual FPCGElementPtr CreateElement() const override;
+	//~End UPCGSettings interface
+
+public:
+	UPROPERTY(BlueprintReadWrite, EditAnywhere, Category = Settings)
+	FName SourceAttributeName = NAME_None;
+
+	// These exist to facilitate unit testing of the MetadataBreakVector element in isolation, when it otherwise would not have any pin connections to generate output
+	bool bForceConnectX = false;
+	bool bForceConnectY = false;
+	bool bForceConnectZ = false;
+	bool bForceConnectW = false;
+};
+
+class FPCGMetadataBreakVectorElement : public FSimplePCGElement
+{
+public:
+	// MetadataBreakVector relies on StaticDuplicateObject, so we cannot run outside of the main thread
+	virtual bool CanExecuteOnlyOnMainThread(const UPCGSettings* InSettings) const override { return true; }
+	virtual bool IsCacheable(const UPCGSettings* InSettings) const override { return false; }
+
+protected:
+	virtual bool ExecuteInternal(FPCGContext* Context) const override;
+};
diff --git a/Engine/Plugins/Experimental/PCG/Source/PCG/Public/Elements/Metadata/PCGMetadataElement.h b/Engine/Plugins/Experimental/PCG/Source/PCG/Public/Elements/Metadata/PCGMetadataElement.h
new file mode 100644
--- /dev/null
+++ b/Engine/Plugins/Experimental/PCG/Source/PCG/Public/Elements/Metadata/PCGMetadataElement.h
@@ -0,0 +1,55 @@
+// Copyright Epic Games, Inc. All Rights Reserved.
+
+#pragma once
+
+#include "PCGSettings.h"
+#include "PCGElement.h"
+#include "PCGPoint.h"
+
+#include "PCGMetadataElement.generated.h"
+
+UENUM()
+enum class EPCGMetadataOperationTarget : uint8
+{
+	PropertyToAttribute,
+	AttributeToProperty,
+	AttributeToAttribute
+};
+
+UCLASS(BlueprintType, ClassGroup = (Procedural))
+class PCG_API UPCGMetadataOperationSettings : public UPCGSettings
+{
+	GENERATED_BODY()
+
+public:
+	//~Begin UPCGSettings interface
+#if WITH_EDITOR
+	virtual FName GetDefaultNodeName() const override { return FName(TEXT("MetadataOperationNode")); }
+	virtual EPCGSettingsType GetType() const override { return EPCGSettingsType::Metadata; }
+#endif
+
+	virtual TArray<FPCGPinProperties> OutputPinProperties() const override { return Super::DefaultPointOutputPinProperties(); }
+
+protected:
+	virtual FPCGElementPtr CreateElement() const override;
+	//~End UPCGSettings interface
+
+public:
+	UPROPERTY(BlueprintReadWrite, EditAnywhere, Category = Settings, meta = (EditCondition = "Target==EPCGMetadataOperationTarget::AttributeToProperty||Target==EPCGMetadataOperationTarget::AttributeToAttribute"))
+	FName SourceAttribute = NAME_None;
+
+	UPROPERTY(BlueprintReadWrite, EditAnywhere, Category = Settings)
+	EPCGPointProperties PointProperty = EPCGPointProperties::Density;
+
+	UPROPERTY(BlueprintReadWrite, EditAnywhere, Category = Settings, meta = (EditCondition="Target==EPCGMetadataOperationTarget::PropertyToAttribute||Target==EPCGMetadataOperationTarget::AttributeToAttribute"))
+	FName DestinationAttribute = NAME_None;
+
+	UPROPERTY(BlueprintReadWrite, EditAnywhere, Category = Settings)
+	EPCGMetadataOperationTarget Target = EPCGMetadataOperationTarget::PropertyToAttribute;
+};
+
+class FPCGMetadataOperationElement : public FSimplePCGElement
+{
+protected:
+	virtual bool ExecuteInternal(FPCGContext* Context) const override;
+};
\ No newline at end of file
diff --git a/Engine/Plugins/Experimental/PCG/Source/PCG/Public/Elements/Metadata/PCGMetadataElementCommon.h b/Engine/Plugins/Experimental/PCG/Source/PCG/Public/Elements/Metadata/PCGMetadataElementCommon.h
new file mode 100644
--- /dev/null
+++ b/Engine/Plugins/Experimental/PCG/Source/PCG/Public/Elements/Metadata/PCGMetadataElementCommon.h
@@ -0,0 +1,39 @@
+// Copyright Epic Games, Inc. All Rights Reserved.
+
+#pragma once
+
+#include "PCGModule.h"
+#include "Metadata/PCGMetadata.h"
+#include "Metadata/PCGMetadataAttribute.h"
+#include "Metadata/PCGMetadataAttributeTpl.h"
+
+struct FPCGTaggedData;
+
+namespace PCGMetadataElementCommon
+{
+	void DuplicateTaggedData(const FPCGTaggedData& InTaggedData, FPCGTaggedData& OutTaggedData, UPCGMetadata*& OutMetadata);
+
+	/** Copies the entry to value key relationship stored in the given Metadata, including its parents */
+	void CopyEntryToValueKeyMap(const UPCGMetadata* MetadataToCopy, const FPCGMetadataAttributeBase* AttributeToCopy, FPCGMetadataAttributeBase* OutAttribute);
+
+	/** Creates a new attribute, or clears the attribute if it already exists and is a 'T' type */
+	template<typename T>
+	FPCGMetadataAttribute<T>* ClearOrCreateAttribute(UPCGMetadata* Metadata, const FName& DestinationAttribute, T DefaultValue)
+	{
+		if (!Metadata)
+		{
+			UE_LOG(LogPCG, Error, TEXT("Failed to create metadata"));
+			return nullptr;
+		}
+
+		if (Metadata->HasAttribute(DestinationAttribute))
+		{
+			UE_LOG(LogPCG, Warning, TEXT("Attribute %s already exists and has been overwritten"), *DestinationAttribute.ToString());
+			Metadata->DeleteAttribute(DestinationAttribute);
+		}
+
+		Metadata->CreateAttribute<T>(DestinationAttribute, DefaultValue, /*bAllowsInterpolation=*/true, /*bOverrideParent=*/false);
+
+		return static_cast<FPCGMetadataAttribute<T>*>(Metadata->GetMutableAttribute(DestinationAttribute));
+	}
+}
\ No newline at end of file
diff --git a/Engine/Plugins/Experimental/PCG/Source/PCG/Public/Elements/Metadata/PCGMetadataPartition.h b/Engine/Plugins/Experimental/PCG/Source/PCG/Public/Elements/Metadata/PCGMetadataPartition.h
new file mode 100644
--- /dev/null
+++ b/Engine/Plugins/Experimental/PCG/Source/PCG/Public/Elements/Metadata/PCGMetadataPartition.h
@@ -0,0 +1,37 @@
+// Copyright Epic Games, Inc. All Rights Reserved.
+
+#pragma once
+
+#include "PCGSettings.h"
+#include "PCGElement.h"
+
+#include "PCGMetadataPartition.generated.h"
+
+UCLASS(BlueprintType, ClassGroup = (Procedural))
+class PCG_API UPCGMetadataPartitionSettings : public UPCGSettings
+{
+	GENERATED_BODY()
+
+public:
+	//~Begin UPCGSettings interface
+#if WITH_EDITOR
+	virtual FName GetDefaultNodeName() const override { return FName(TEXT("MetadataPartitionNode")); }
+	virtual EPCGSettingsType GetType() const override { return EPCGSettingsType::Metadata; }
+#endif
+
+	virtual TArray<FPCGPinProperties> OutputPinProperties() const override { return Super::DefaultPointOutputPinProperties(); }
+
+protected:
+	virtual FPCGElementPtr CreateElement() const override;
+	//~End UPCGSettings interface
+
+public:
+	UPROPERTY(BlueprintReadWrite, EditAnywhere, Category = Settings)
+	FName PartitionAttribute = NAME_None;
+};
+
+class FPCGMetadataPartitionElement : public FSimplePCGElement
+{
+protected:
+	virtual bool ExecuteInternal(FPCGContext* Context) const override;
+};
\ No newline at end of file
diff --git a/Engine/Plugins/Experimental/PCG/Source/PCG/Public/Elements/Metadata/PCGMetadataRenameElement.h b/Engine/Plugins/Experimental/PCG/Source/PCG/Public/Elements/Metadata/PCGMetadataRenameElement.h
new file mode 100644
--- /dev/null
+++ b/Engine/Plugins/Experimental/PCG/Source/PCG/Public/Elements/Metadata/PCGMetadataRenameElement.h
@@ -0,0 +1,43 @@
+// Copyright Epic Games, Inc. All Rights Reserved.
+
+#pragma once
+
+#include "PCGSettings.h"
+#include "PCGElement.h"
+
+#include "PCGMetadataRenameElement.generated.h"
+
+UCLASS(BlueprintType, ClassGroup = (Procedural))
+class PCG_API UPCGMetadataRenameSettings : public UPCGSettings
+{
+	GENERATED_BODY()
+
+public:
+	//~Begin UPCGSettings interface
+#if WITH_EDITOR
+	virtual FName GetDefaultNodeName() const override { return FName(TEXT("MetadataRenameNode")); }
+	virtual EPCGSettingsType GetType() const override { return EPCGSettingsType::Metadata; }
+#endif
+
+	virtual TArray<FPCGPinProperties> OutputPinProperties() const override;
+
+protected:
+	virtual FPCGElementPtr CreateElement() const override;
+	//~End UPCGSettings interface
+
+public:
+	UPROPERTY(BlueprintReadWrite, EditAnywhere, Category = Settings)
+	FName AttributeToRename = NAME_None;
+
+	UPROPERTY(BlueprintReadWrite, EditAnywhere, Category = Settings)
+	FName NewAttributeName = NAME_None;
+};
+
+class FPCGMetadataRenameElement : public FSimplePCGElement
+{
+public:
+	virtual bool CanExecuteOnlyOnMainThread(const UPCGSettings* InSettings) const override { return true; }
+
+protected:
+	virtual bool ExecuteInternal(FPCGContext* Context) const override;
+};
\ No newline at end of file
diff --git a/Engine/Plugins/Experimental/PCG/Source/PCG/Public/Elements/PCGMetadataBreakVector.h b/Engine/Plugins/Experimental/PCG/Source/PCG/Public/Elements/PCGMetadataBreakVector.h
deleted file mode 100644
diff --git a/Engine/Plugins/Experimental/PCG/Source/PCG/Public/Elements/PCGMetadataElement.h b/Engine/Plugins/Experimental/PCG/Source/PCG/Public/Elements/PCGMetadataElement.h
deleted file mode 100644
diff --git a/Engine/Plugins/Experimental/PCG/Source/PCG/Public/Elements/PCGMetadataPartition.h b/Engine/Plugins/Experimental/PCG/Source/PCG/Public/Elements/PCGMetadataPartition.h
deleted file mode 100644
diff --git a/Engine/Plugins/Experimental/PCG/Source/PCG/Public/Elements/PCGMetadataRenameElement.h b/Engine/Plugins/Experimental/PCG/Source/PCG/Public/Elements/PCGMetadataRenameElement.h
deleted file mode 100644
diff --git a/Engine/Plugins/Experimental/PCG/Source/PCG/Public/Helpers/PCGMetadataHelpers.h b/Engine/Plugins/Experimental/PCG/Source/PCG/Public/Helpers/PCGMetadataHelpers.h
new file mode 100644
--- /dev/null
+++ b/Engine/Plugins/Experimental/PCG/Source/PCG/Public/Helpers/PCGMetadataHelpers.h
@@ -0,0 +1,10 @@
+// Copyright Epic Games, Inc. All Rights Reserved.
+
+#pragma once
+
+class UPCGMetadata;
+
+namespace PCGMetadataHelpers
+{
+	bool HasSameRoot(const UPCGMetadata* Metadata1, const UPCGMetadata* Metadata2);
+}
\ No newline at end of file
diff --git a/Engine/Plugins/Experimental/PCG/Source/PCG/Public/Metadata/PCGMetadata.h b/Engine/Plugins/Experimental/PCG/Source/PCG/Public/Metadata/PCGMetadata.h
--- a/Engine/Plugins/Experimental/PCG/Source/PCG/Public/Metadata/PCGMetadata.h
+++ b/Engine/Plugins/Experimental/PCG/Source/PCG/Public/Metadata/PCGMetadata.h
@@ -3,24 +3,14 @@
 #pragma once
 
 #include "CoreMinimal.h"
+#include "PCGModule.h"
 #include "PCGMetadataAccessor.h"
 #include "PCGMetadataAttributeTraits.h"
 #include "PCGMetadataCommon.h"
+#include "Metadata/PCGMetadataAttributeTpl.h"
 
 #include "PCGMetadata.generated.h"
 
-class FPCGMetadataAttributeBase;
-
-UENUM()
-enum class EPCGMetadataOp : uint8
-{
-	Min,
-	Max,
-	Sub,
-	Add,
-	Mul,
-	Div
-};
 
 UCLASS(BlueprintType)
 class PCG_API UPCGMetadata : public UObject
@@ -202,10 +192,10 @@ public:
 	int64 GetItemKeyCountForParent() const;
 	int64 GetItemCountForChild() const;
 
-protected:
 	template<typename T>
 	FPCGMetadataAttributeBase* CreateAttribute(FName AttributeName, const T& DefaultValue, bool bAllowsInterpolation, bool bOverrideParent);
 
+protected:
 	FPCGMetadataAttributeBase* CopyAttribute(FName AttributeToCopy, FName NewAttributeName, bool bKeepParent, bool bCopyEntries, bool bCopyValues);
 
 	bool ParentHasAttribute(FName AttributeName) const;
@@ -232,3 +222,32 @@ protected:
 
 	TAtomic<int64> DelayedEntriesIndex = 0;
 };
+
+template<typename T>
+FPCGMetadataAttributeBase* UPCGMetadata::CreateAttribute(FName AttributeName, const T& DefaultValue, bool bAllowsInterpolation, bool bOverrideParent)
+{
+	const FPCGMetadataAttributeBase* ParentAttribute = nullptr;
+
+	if (bOverrideParent && Parent.IsValid())
+	{
+		ParentAttribute = Parent->GetConstAttribute(AttributeName);
+	}
+
+	FPCGMetadataAttributeBase* NewAttribute = new FPCGMetadataAttribute<T>(this, AttributeName, ParentAttribute, DefaultValue, bAllowsInterpolation);
+
+	AttributeLock.WriteLock();
+	if (FPCGMetadataAttributeBase** ExistingAttribute = Attributes.Find(AttributeName))
+	{
+		UE_LOG(LogPCG, Warning, TEXT("Attribute %s already exists"), *AttributeName.ToString());
+		delete NewAttribute;
+		NewAttribute = *ExistingAttribute;
+	}
+	else
+	{
+		NewAttribute->AttributeId = NextAttributeId++;
+		AddAttributeInternal(AttributeName, NewAttribute);
+	}
+	AttributeLock.WriteUnlock();
+
+	return NewAttribute;
+}
\ No newline at end of file
diff --git a/Engine/Plugins/Experimental/PCG/Source/PCG/Public/Metadata/PCGMetadataAttribute.h b/Engine/Plugins/Experimental/PCG/Source/PCG/Public/Metadata/PCGMetadataAttribute.h
--- a/Engine/Plugins/Experimental/PCG/Source/PCG/Public/Metadata/PCGMetadataAttribute.h
+++ b/Engine/Plugins/Experimental/PCG/Source/PCG/Public/Metadata/PCGMetadataAttribute.h
@@ -4,11 +4,12 @@
 
 #include "CoreMinimal.h"
 #include "PCGMetadataCommon.h"
-#include "PCGMetadata.h"
 
 #include "Misc/ScopeRWLock.h"
 #include "Templates/EnableIf.h"
 
+class UPCGMetadata;
+
 class PCG_API FPCGMetadataAttributeBase
 {
 public:
diff --git a/Engine/Plugins/Experimental/PCG/Source/PCG/Public/Metadata/PCGMetadataAttributeTpl.h b/Engine/Plugins/Experimental/PCG/Source/PCG/Public/Metadata/PCGMetadataAttributeTpl.h
--- a/Engine/Plugins/Experimental/PCG/Source/PCG/Public/Metadata/PCGMetadataAttributeTpl.h
+++ b/Engine/Plugins/Experimental/PCG/Source/PCG/Public/Metadata/PCGMetadataAttributeTpl.h
@@ -6,6 +6,10 @@
 #include "Metadata/PCGMetadataAttribute.h"
 #include "Metadata/PCGMetadataAttributeTraits.h"
 
+#include "Helpers/PCGMetadataHelpers.h"
+
+class UPCGMetadata;
+
 template<typename T>
 class FPCGMetadataAttribute : public FPCGMetadataAttributeBase
 {
@@ -47,7 +51,7 @@ public:
 	virtual FPCGMetadataAttributeBase* Copy(FName NewName, UPCGMetadata* InMetadata, bool bKeepParent, bool bCopyEntries = true, bool bCopyValues = true) const override
 	{
 		// this copies to a new attribute
-		check(!bKeepParent || Metadata->GetRoot() == InMetadata->GetRoot());
+		check(!bKeepParent || PCGMetadataHelpers::HasSameRoot(Metadata, InMetadata));
 		FPCGMetadataAttribute<T>* AttributeCopy = new FPCGMetadataAttribute<T>(InMetadata, NewName, bKeepParent ? this : nullptr, DefaultValue, bAllowsInterpolation);
 
 		if (bCopyEntries)
diff --git a/Engine/Plugins/Experimental/PCG/Source/PCG/Public/Metadata/PCGMetadataCommon.h b/Engine/Plugins/Experimental/PCG/Source/PCG/Public/Metadata/PCGMetadataCommon.h
--- a/Engine/Plugins/Experimental/PCG/Source/PCG/Public/Metadata/PCGMetadataCommon.h
+++ b/Engine/Plugins/Experimental/PCG/Source/PCG/Public/Metadata/PCGMetadataCommon.h
@@ -4,9 +4,22 @@
 
 #include "CoreMinimal.h"
 
+#include "PCGMetadataCommon.generated.h"
+
 typedef int64 PCGMetadataEntryKey;
 typedef int32 PCGMetadataAttributeKey;
 typedef int32 PCGMetadataValueKey;
 
 const PCGMetadataEntryKey PCGInvalidEntryKey = -1;
-const PCGMetadataValueKey PCGDefaultValueKey = -1;
\ No newline at end of file
+const PCGMetadataValueKey PCGDefaultValueKey = -1;
+
+UENUM()
+enum class EPCGMetadataOp : uint8
+{
+	Min,
+	Max,
+	Sub,
+	Add,
+	Mul,
+	Div
+};
\ No newline at end of file
