Updates graph-wide determinism test (Iteration II)
+ Comparison of Param data
* Update the sample testing volume to the bounds of the spatial data tested

Future Improvements (TODO)
-----------------------------------------------
* Shuffling/Permutation of Input Order

diff --git a/Engine/Plugins/Experimental/PCG/Source/PCG/Private/PCGComponent.cpp b/Engine/Plugins/Experimental/PCG/Source/PCG/Private/PCGComponent.cpp
--- a/Engine/Plugins/Experimental/PCG/Source/PCG/Private/PCGComponent.cpp
+++ b/Engine/Plugins/Experimental/PCG/Source/PCG/Private/PCGComponent.cpp
@@ -63,20 +63,27 @@ void UPCGComponent::SetIsPartitioned(bool bIsNowPartitioned)
 
 	bool bDoActorMapping = bGenerated || PCGHelpers::IsRuntimeOrPIE();
 
-	if (bGenerated)
+	if (UPCGSubsystem* Subsystem = GetSubsystem())
 	{
-		CleanupLocalImmediate(/*bRemoveComponents=*/true);
-	}
+		if (bGenerated)
+		{
+			CleanupLocalImmediate(/*bRemoveComponents=*/true);
+		}
 
-	if (bIsNowPartitioned)
-	{
-		bIsPartitioned = bIsNowPartitioned;
-		GetSubsystem()->RegisterOrUpdatePCGComponent(this, bDoActorMapping);
+		if (bIsNowPartitioned)
+		{
+			bIsPartitioned = bIsNowPartitioned;
+			Subsystem->RegisterOrUpdatePCGComponent(this, bDoActorMapping);
+		}
+		else
+		{
+			Subsystem->UnregisterPCGComponent(this);
+			bIsPartitioned = bIsNowPartitioned;
+		}
 	}
 	else
 	{
-		GetSubsystem()->UnregisterPCGComponent(this);
-		bIsPartitioned = bIsNowPartitioned;
+		bIsPartitioned = false;
 	}
 }
 
diff --git a/Engine/Plugins/Experimental/PCG/Source/PCG/Private/Tests/Determinism/PCGDeterminismTestsCommon.cpp b/Engine/Plugins/Experimental/PCG/Source/PCG/Private/Tests/Determinism/PCGDeterminismTestsCommon.cpp
--- a/Engine/Plugins/Experimental/PCG/Source/PCG/Private/Tests/Determinism/PCGDeterminismTestsCommon.cpp
+++ b/Engine/Plugins/Experimental/PCG/Source/PCG/Private/Tests/Determinism/PCGDeterminismTestsCommon.cpp
@@ -6,6 +6,7 @@
 #include "PCGGraph.h"
 #include "PCGHelpers.h"
 #include "PCGNode.h"
+#include "PCGParamData.h"
 #include "PCGPin.h"
 #include "PCGPoint.h"
 #include "Data/PCGPointData.h"
@@ -15,6 +16,9 @@
 #include "Data/PCGVolumeData.h"
 #include "Graph/PCGGraphCache.h"
 #include "Graph/PCGGraphExecutor.h"
+#include "Metadata/PCGMetadata.h"
+#include "Metadata/PCGMetadataAttribute.h"
+#include "Metadata/PCGMetadataAttributeTpl.h"
 
 #include "Components/SplineComponent.h"
 #include "Components/StaticMeshComponent.h"
@@ -38,6 +42,8 @@ namespace PCGDeterminismTests
 
 	void RunDeterminismTest(const UPCGNode* InPCGNode, FDeterminismNodeTestResult& OutResult, const FNodeTestInfo& TestToRun)
 	{
+		TRACE_CPUPROFILER_EVENT_SCOPE(PCGDeterminismTests::RunDeterminismTest::Node);
+
 		if (!InPCGNode || !InPCGNode->DefaultSettings)
 		{
 			OutResult.DataTypesTested = EPCGDataType::None;
@@ -61,6 +67,8 @@ namespace PCGDeterminismTests
 #if WITH_EDITOR
 	void RunDeterminismTest(const UPCGGraph* InPCGGraph, UPCGComponent* InPCGComponent, FDeterminismNodeTestResult& OutResult)
 	{
+		TRACE_CPUPROFILER_EVENT_SCOPE(PCGDeterminismTests::RunDeterminismTest::Graph);
+
 		if (!InPCGGraph || !InPCGComponent)
 		{
 			OutResult.DataTypesTested = EPCGDataType::None;
@@ -69,6 +77,11 @@ namespace PCGDeterminismTests
 			return;
 		}
 
+		// Update the data types tested
+		const UPCGData* InputPCGData = InPCGComponent->GetInputPCGData();
+		check(InputPCGData);
+		OutResult.DataTypesTested |= InputPCGData->GetDataType();
+
 		// Clone the actor and component
 		AActor* PCGActor = InPCGComponent->GetOwner();
 		check(PCGActor);
@@ -77,14 +90,10 @@ namespace PCGDeterminismTests
 		PCGActorCopy->SetFlags(RF_Transient);
 		UPCGComponent* PCGComponentCopy = PCGActorCopy->FindComponentByClass<UPCGComponent>();
 		check(PCGComponentCopy);
-
-		// TODO: Update this to a public setter
-		PCGComponentCopy->bIsPartitioned = false;
+		PCGComponentCopy->SetIsPartitioned(false);
 
 		FPCGGraphExecutor Executor = FPCGGraphExecutor(PCGComponentCopy);
 
-		// TODO: Record data types tested
-
 		auto ScheduleAndWaitForExecution = [&Executor](FPCGTaskId FinalTaskID)
 		{
 			// TODO: Consider randomizing/iterating through possible input orders
@@ -113,7 +122,6 @@ namespace PCGDeterminismTests
 		auto RecordMappedOutput = [&IntermediateOutputArray](FPCGTaskId TaskId, const UPCGNode* Node, const FPCGDataCollection& NodeOutput)
 		{
 			check(Node);
-			// TODO: Ensure this is threadsafe if it needs to be
 			IntermediateOutputArray.Add(TaskId, MakeTuple(Node, NodeOutput));
 		};
 
@@ -195,7 +203,7 @@ namespace PCGDeterminismTests
 			}
 		}
 
-		// Clean up PCG generated
+		// Clean up anything generated
 		PCGComponentCopy->Cleanup();
 
 		// Finalize the results
@@ -320,6 +328,8 @@ namespace PCGDeterminismTests
 		int32 NumInputsPerPin,
 		EDeterminismLevel MaxLevel)
 	{
+		TRACE_CPUPROFILER_EVENT_SCOPE(PCGDeterminismTests::GetHighestDeterminismLevel);
+
 		const UPCGNode* PCGNode = NodeAndOptions.PCGNode;
 		const TArray<TArray<EPCGDataType>>& BaseOptionsByPin = NodeAndOptions.BaseOptionsByPin;
 
@@ -577,6 +587,8 @@ namespace PCGDeterminismTests
 
 	bool DataCollectionsAreIdentical(const FPCGDataCollection& FirstCollection, const FPCGDataCollection& SecondCollection)
 	{
+		TRACE_CPUPROFILER_EVENT_SCOPE(PCGDeterminismTests::DataCollectionsAreIdentical);
+
 		if (FirstCollection.TaggedData.Num() != SecondCollection.TaggedData.Num())
 		{
 			return false;
@@ -611,6 +623,8 @@ namespace PCGDeterminismTests
 
 	bool DataCollectionsAreConsistent(const FPCGDataCollection& FirstCollection, const FPCGDataCollection& SecondCollection, int32 NumInputs)
 	{
+		TRACE_CPUPROFILER_EVENT_SCOPE(PCGDeterminismTests::DataCollectionsAreConsistent);
+
 		TArray<int32> IndexOffsets;
 		if (!DataCollectionsMatch(FirstCollection, SecondCollection, IndexOffsets))
 		{
@@ -658,6 +672,8 @@ namespace PCGDeterminismTests
 		const FPCGDataCollection& SecondCollection,
 		TArray<int32>& OutIndexOffsets)
 	{
+		TRACE_CPUPROFILER_EVENT_SCOPE(PCGDeterminismTests::DataCollectionsMatch);
+
 		OutIndexOffsets.Empty();
 
 		if (FirstCollection.TaggedData.Num() != SecondCollection.TaggedData.Num())
@@ -748,6 +764,8 @@ namespace PCGDeterminismTests
 
 	bool InternalDataMatches(const UPCGData* FirstData, const UPCGData* SecondData, TArray<int32>& OutIndexOffsets)
 	{
+		TRACE_CPUPROFILER_EVENT_SCOPE(PCGDeterminismTests::InternalDataMatches);
+
 		check(FirstData->GetDataType() == SecondData->GetDataType());
 
 		// Filter output differences
@@ -757,6 +775,8 @@ namespace PCGDeterminismTests
 
 	bool SpatialDataIsIdentical(const UPCGData* FirstData, const UPCGData* SecondData)
 	{
+		TRACE_CPUPROFILER_EVENT_SCOPE(PCGDeterminismTests::SpatialDataIsIdentical);
+
 		const UPCGSpatialData* FirstSpatialData = CastChecked<const UPCGSpatialData>(FirstData);
 		const UPCGSpatialData* SecondSpatialData = CastChecked<const UPCGSpatialData>(SecondData);
 
@@ -795,6 +815,8 @@ namespace PCGDeterminismTests
 
 	bool PointDataIsIdentical(const UPCGData* FirstData, const UPCGData* SecondData)
 	{
+		TRACE_CPUPROFILER_EVENT_SCOPE(PCGDeterminismTests::PointDataIsIdentical);
+
 		const UPCGPointData* FirstPointData = CastChecked<const UPCGPointData>(FirstData);
 		const UPCGPointData* SecondPointData = CastChecked<const UPCGPointData>(SecondData);
 
@@ -824,6 +846,8 @@ namespace PCGDeterminismTests
 
 	bool VolumeDataIsIdentical(const UPCGData* FirstData, const UPCGData* SecondData)
 	{
+		TRACE_CPUPROFILER_EVENT_SCOPE(PCGDeterminismTests::VolumeDataIsIdentical);
+
 		const UPCGVolumeData* FirstVolumeData = CastChecked<const UPCGVolumeData>(FirstData);
 		const UPCGVolumeData* SecondVolumeData = CastChecked<const UPCGVolumeData>(SecondData);
 
@@ -882,9 +906,16 @@ namespace PCGDeterminismTests
 
 	bool SampledSpatialDataIsIdentical(const UPCGSpatialData* FirstSpatialData, const UPCGSpatialData* SecondSpatialData)
 	{
-		// At this point, bounds has already been checked for equality
-		// TODO: Change this testing volume relative to the bounds of the two datas
-		const FBox SampleBounds = Defaults::TestingVolume;
+		TRACE_CPUPROFILER_EVENT_SCOPE(PCGDeterminismTests::SampledSpatialDataIsIdentical);
+
+		check(FirstSpatialData && SecondSpatialData);
+
+		// The bounds should have already been checked by now
+		check(FirstSpatialData->GetBounds() == SecondSpatialData->GetBounds());
+
+		// Combine the bounds of both spatial data and expand by a little to get the overall bounds
+		FBox SampleBounds = FirstSpatialData->GetBounds() + SecondSpatialData->GetBounds();
+		SampleBounds = SampleBounds.ExpandBy(SampleBounds.GetExtent() * Defaults::TestingVolumeExpandByFactor);
 		const FVector SampleExtent = SampleBounds.GetExtent();
 
 		FPCGPoint FirstPoint;
@@ -900,6 +931,7 @@ namespace PCGDeterminismTests
 				for (FVector::FReal Z = StartingOffset.Z; Z < SampleBounds.Max.Z; Z += StepInterval.Z)
 				{
 					FTransform PointTransform(FVector(X, Y, Z));
+
 					bool bFirstPointWasSampled = FirstSpatialData->SamplePoint(PointTransform, SampleBounds, FirstPoint, nullptr);
 					bool bSecondPointWasSampled = SecondSpatialData->SamplePoint(PointTransform, SampleBounds, SecondPoint, nullptr);
 
@@ -1062,6 +1094,122 @@ namespace PCGDeterminismTests
 		return true;
 	}
 
+	bool ParamDataIsIdentical(const UPCGData* FirstData, const UPCGData* SecondData)
+	{
+		const UPCGParamData* FirstParamData = CastChecked<const UPCGParamData>(FirstData);
+		const UPCGParamData* SecondParamData = CastChecked<const UPCGParamData>(SecondData);
+
+		return MetadataIsIdentical(FirstParamData->Metadata, SecondParamData->Metadata);
+	}
+
+	bool MetadataIsIdentical(const UPCGMetadata* FirstMetadata, const UPCGMetadata* SecondMetadata)
+	{
+		TRACE_CPUPROFILER_EVENT_SCOPE(PCGDeterminismTests::MetadataIsIdentical);
+
+		check(FirstMetadata && SecondMetadata);
+
+		// Early out for different counts
+		if (FirstMetadata->GetAttributeCount() != SecondMetadata->GetAttributeCount())
+		{
+			return false;
+		}
+
+		TArray<FName> FirstAttributeNames;
+		TArray<EPCGMetadataTypes> FirstAttributeTypes;
+		FirstMetadata->GetAttributes(FirstAttributeNames, FirstAttributeTypes);
+
+		TArray<FName> SecondAttributeNames;
+		TArray<EPCGMetadataTypes> SecondAttributeTypes;
+		SecondMetadata->GetAttributes(SecondAttributeNames, SecondAttributeTypes);
+
+		for (const FName& AttributeName : FirstAttributeNames)
+		{
+			const FPCGMetadataAttributeBase* FirstAttributeBase = FirstMetadata->GetConstAttribute(AttributeName);
+			const FPCGMetadataAttributeBase* SecondAttributeBase = SecondMetadata->GetConstAttribute(AttributeName);
+
+			// Early out if it isn't found
+			if (!FirstAttributeBase || !SecondAttributeBase)
+			{
+				return false;
+			}
+
+			// Check the Entry:Value key mappings
+			const PCGMetadataEntryKey EntryKeyNum = FirstMetadata->GetItemCountForChild();
+			for (PCGMetadataEntryKey EntryKey = 0; EntryKey < EntryKeyNum; ++EntryKey)
+			{
+				if (FirstAttributeBase->GetValueKey(EntryKey) != SecondAttributeBase->GetValueKey(EntryKey))
+				{
+					return false;
+				}
+			}
+
+			// Compare the values with the value keys
+			const PCGMetadataValueKey FirstNumValueKeys = FirstAttributeBase->GetValueKeyOffsetForChild();
+			const PCGMetadataValueKey SecondNumValueKeys = SecondAttributeBase->GetValueKeyOffsetForChild();
+
+			if (FirstNumValueKeys != SecondNumValueKeys)
+			{
+				return false;
+			}
+
+			bool bAttributesAreEqual = false;
+
+			for (PCGMetadataValueKey ValueKey = 0; ValueKey < FirstNumValueKeys; ++ValueKey)
+			{
+				// Note: these must all have operator==() defined
+				switch (FirstAttributeBase->GetTypeId())
+				{
+				case PCG::Private::MetadataTypes<float>::Id:
+					bAttributesAreEqual = MetadataAttributesAreEqual<float>(FirstAttributeBase, SecondAttributeBase, ValueKey);
+					break;
+				case PCG::Private::MetadataTypes<double>::Id:
+					bAttributesAreEqual = MetadataAttributesAreEqual<double>(FirstAttributeBase, SecondAttributeBase, ValueKey);
+					break;
+				case PCG::Private::MetadataTypes<bool>::Id:
+					bAttributesAreEqual = MetadataAttributesAreEqual<bool>(FirstAttributeBase, SecondAttributeBase, ValueKey);
+					break;
+				case PCG::Private::MetadataTypes<FVector>::Id:
+					bAttributesAreEqual = MetadataAttributesAreEqual<FVector>(FirstAttributeBase, SecondAttributeBase, ValueKey);
+					break;
+				case PCG::Private::MetadataTypes<FVector4>::Id:
+					bAttributesAreEqual = MetadataAttributesAreEqual<FVector4>(FirstAttributeBase, SecondAttributeBase, ValueKey);
+					break;
+				case PCG::Private::MetadataTypes<int32>::Id:
+					bAttributesAreEqual = MetadataAttributesAreEqual<int32>(FirstAttributeBase, SecondAttributeBase, ValueKey);
+					break;
+				case PCG::Private::MetadataTypes<int64>::Id:
+					bAttributesAreEqual = MetadataAttributesAreEqual<int64>(FirstAttributeBase, SecondAttributeBase, ValueKey);
+					break;
+				case PCG::Private::MetadataTypes<FString>::Id:
+					bAttributesAreEqual = MetadataAttributesAreEqual<FString>(FirstAttributeBase, SecondAttributeBase, ValueKey);
+					break;
+				case PCG::Private::MetadataTypes<FName>::Id:
+					bAttributesAreEqual = MetadataAttributesAreEqual<FName>(FirstAttributeBase, SecondAttributeBase, ValueKey);
+					break;
+				case PCG::Private::MetadataTypes<FQuat>::Id:
+					bAttributesAreEqual = MetadataAttributesAreEqual<FQuat>(FirstAttributeBase, SecondAttributeBase, ValueKey);
+					break;
+				case PCG::Private::MetadataTypes<FRotator>::Id:
+					bAttributesAreEqual = MetadataAttributesAreEqual<FRotator>(FirstAttributeBase, SecondAttributeBase, ValueKey);
+					break;
+				case PCG::Private::MetadataTypes<FTransform>::Id:
+					bAttributesAreEqual = MetadataAttributesAreEqual<FTransform>(FirstAttributeBase, SecondAttributeBase, ValueKey);
+					break;
+				default:
+					UE_LOG(LogPCG, Warning, TEXT("Invalid type to compare for metadata being identical: %d"), FirstAttributeBase->GetTypeId());
+					break;
+				}
+
+				if (!bAttributesAreEqual)
+				{
+					return false;
+				}
+			}
+		}
+
+		return true;
+	}
+
 	bool ComparisonIsUnimplemented(const UPCGData* FirstData, const UPCGData* SecondData)
 	{
 		return false;
@@ -1088,8 +1236,7 @@ namespace PCGDeterminismTests
 		}
 
 		// Comparable data types
-		// TODO: Add metadata/param data as comparable in the future
-		if (!!(DataType & EPCGDataType::Spatial))
+		if (!!(DataType & (EPCGDataType::Param | EPCGDataType::Spatial)))
 		{
 			return true;
 		}
@@ -1222,14 +1369,18 @@ namespace PCGDeterminismTests
 
 	TFunction<bool(const UPCGData*, const UPCGData*)> GetDataCompareFunction(EPCGDataType DataType, EDeterminismLevel DeterminismLevel)
 	{
-		// TODO: Currently only spatial data is comparable. Other types to be added.
-		if (!DataTypeIsComparable(DataType) || DataType != EPCGDataType::Spatial)
+		if (!DataTypeIsComparable(DataType))
 		{
 			// Should never reach here
 			UE_LOG(LogPCG, Warning, TEXT("Attempting to compare incomparable data."));
 			return ComparisonIsUnimplemented;
 		}
 
+		if (DataType == EPCGDataType::Param)
+		{
+			return ParamDataIsIdentical;
+		}
+
 		switch (DeterminismLevel)
 		{
 		case EDeterminismLevel::OrderOrthogonal:
diff --git a/Engine/Plugins/Experimental/PCG/Source/PCG/Public/Tests/Determinism/PCGDeterminismTestsCommon.h b/Engine/Plugins/Experimental/PCG/Source/PCG/Public/Tests/Determinism/PCGDeterminismTestsCommon.h
--- a/Engine/Plugins/Experimental/PCG/Source/PCG/Public/Tests/Determinism/PCGDeterminismTestsCommon.h
+++ b/Engine/Plugins/Experimental/PCG/Source/PCG/Public/Tests/Determinism/PCGDeterminismTestsCommon.h
@@ -5,11 +5,13 @@
 #include "PCGCommon.h"
 #include "PCGSettings.h"
 #include "Tests/PCGTestsCommon.h"
+#include "Metadata/PCGMetadataAttributeTpl.h"
 
 #include "PCGDeterminismTestsCommon.generated.h"
 
 class UPCGComponent;
 class UPCGData;
+class UPCGMetadata;
 class UPCGSpatialData;
 class UPrimitiveComponent;
 class USplineComponent;
@@ -72,12 +74,14 @@ namespace PCGDeterminismTests
 		constexpr static int32 NumTestPointsToGenerate = 100;
 		constexpr static int32 NumPolyLinePointsToGenerate = 6;
 		constexpr static int32 NumTestPolyLinePointsToGenerate = 6;
-		constexpr static int32 NumSamplingStepsPerDimension = 100;
+		constexpr static int32 NumSamplingStepsPerDimension = 15;
 		constexpr static int32 NumTestInputsPerPin = 2;
 
 		constexpr static FVector::FReal SmallDistance = 50.0;
 		constexpr static FVector::FReal MediumDistance = 200.0;
 		constexpr static FVector::FReal LargeDistance = 500.0;
+		/** This will expand the testing volume by an increase of this ratio - ex. 0.25 == 125% */
+		constexpr static FVector::FReal TestingVolumeExpandByFactor = 0.25;
 
 		const FVector SmallVector = FVector::OneVector * SmallDistance;
 		const FVector MediumVector = FVector::OneVector * MediumDistance;
@@ -85,9 +89,6 @@ namespace PCGDeterminismTests
 
 		const FName TestPinName = TEXT("Test");
 		const FName GraphResultName = TEXT("Graph Result");
-
-		/** A little bigger than the typical largest volumes */
-		const FBox TestingVolume = FBox(-1.2 * LargeVector, 1.2 * LargeVector);
 	}
 
 	constexpr static EPCGDataType TestableDataTypes[6] =
@@ -227,6 +228,12 @@ namespace PCGDeterminismTests
 	/** Validates whether two PointData objects contain the same points */
 	bool PointDataIsOrthogonal(const UPCGData* FirstData, const UPCGData* SecondData);
 
+	/** Validates whether two ParamData objects are identical */
+	bool ParamDataIsIdentical(const UPCGData* FirstData, const UPCGData* SecondData);
+
+	/** Validates whether two Metadata objects are identical */
+	bool MetadataIsIdentical(const UPCGMetadata* FirstMetadata, const UPCGMetadata* SecondMetadata);
+
 	/** A catch function for unimplemented comparisons */
 	bool ComparisonIsUnimplemented(const UPCGData* FirstData, const UPCGData* SecondData);
 	/** Updates the test result output to reflect that the permutation limit has been exceeded */
@@ -355,6 +362,27 @@ namespace PCGDeterminismTests
 		Array = MoveTemp(TempArray);
 	}
 
+	template<typename MetadataAttributeType>
+	bool MetadataAttributesAreEqual(const FPCGMetadataAttributeBase* FirstAttributeBase, const FPCGMetadataAttributeBase* SecondAttributeBase, PCGMetadataValueKey ValueKey)
+	{
+		check(FirstAttributeBase && SecondAttributeBase);
+
+		if constexpr (TIsSame<FTransform, MetadataAttributeType>::Value)
+		{
+			const FTransform FirstTypedAttribute = static_cast<const FPCGMetadataAttribute<FTransform>*>(FirstAttributeBase)->GetValue(ValueKey);
+			const FTransform SecondTypedAttribute = static_cast<const FPCGMetadataAttribute<FTransform>*>(SecondAttributeBase)->GetValue(ValueKey);
+
+			return FirstTypedAttribute.Equals(SecondTypedAttribute);
+		}
+		else
+		{
+			const MetadataAttributeType FirstTypedAttribute = static_cast<const FPCGMetadataAttribute<MetadataAttributeType>*>(FirstAttributeBase)->GetValue(ValueKey);
+			const MetadataAttributeType SecondTypedAttribute = static_cast<const FPCGMetadataAttribute<MetadataAttributeType>*>(SecondAttributeBase)->GetValue(ValueKey);
+
+			return (FirstTypedAttribute == SecondTypedAttribute);
+		}
+	}
+
 	namespace Defaults
 	{
 		static const FNodeTestInfo DeterminismBasicTestInfo = FNodeTestInfo(NSLOCTEXT("PCGDeterminism", "BasicTest", "Basic Test"), RunBasicTestSuite, 90.f);
