Add support for areas bound by closed splines/poly line

diff --git a/Engine/Plugins/Experimental/PCG/Source/PCG/PCG.Build.cs b/Engine/Plugins/Experimental/PCG/Source/PCG/PCG.Build.cs
--- a/Engine/Plugins/Experimental/PCG/Source/PCG/PCG.Build.cs
+++ b/Engine/Plugins/Experimental/PCG/Source/PCG/PCG.Build.cs
@@ -23,6 +23,7 @@ namespace UnrealBuildTool.Rules
 			PrivateDependencyModuleNames.AddRange(
 				new string[]
 				{					
+					"Voronoi",
 				}
 			);
 
diff --git a/Engine/Plugins/Experimental/PCG/Source/PCG/Private/Elements/PCGSplineSampler.cpp b/Engine/Plugins/Experimental/PCG/Source/PCG/Private/Elements/PCGSplineSampler.cpp
--- a/Engine/Plugins/Experimental/PCG/Source/PCG/Private/Elements/PCGSplineSampler.cpp
+++ b/Engine/Plugins/Experimental/PCG/Source/PCG/Private/Elements/PCGSplineSampler.cpp
@@ -10,6 +10,291 @@
 #include "Helpers/PCGBlueprintHelpers.h"
 #include "Helpers/PCGSettingsHelpers.h"
 
+#include "Components/SplineComponent.h"
+#include "Voronoi/Voronoi.h"
+
+namespace PCGSplineSamplerHelpers
+{
+	/**
+	* Robust 2D line segment intersection. Returns true if the segments intersect and stores the intersection in OutIntersectionPoint.
+	* If the segments are collinear, one of the contained endpoints will be chosen as the intersection.
+	*/
+	bool SegmentIntersection2D(const FVector2D& SegmentStartA, const FVector2D& SegmentEndA, const FVector2D& SegmentStartB, const FVector2D& SegmentEndB, FVector2D& OutIntersectionPoint)
+	{
+		const FVector2D VectorA = SegmentEndA - SegmentStartA;
+		const FVector2D VectorB = SegmentEndB - SegmentStartB;
+
+		const FVector::FReal Determinant = VectorA.X * VectorB.Y - VectorB.X * VectorA.Y;
+
+		// Determinant is zero means the segments are parallel. Check if they are also collinear by projecting (0, 0) onto both lines.
+		if (FMath::IsNearlyZero(Determinant))
+		{
+			const FVector::FReal SquareMagnitudeA = VectorA.SquaredLength();
+			if (FMath::IsNearlyZero(SquareMagnitudeA))
+			{
+				const FVector2D& MinB = SegmentStartB.X < SegmentEndB.X ? SegmentStartB : SegmentEndB;
+				const FVector2D& MaxB = SegmentStartB.X > SegmentEndB.X ? SegmentStartB : SegmentEndB;
+
+				if (SegmentStartA.X > MinB.X && SegmentStartA.X < MaxB.X)
+				{
+					OutIntersectionPoint = SegmentStartA;
+					return true;
+				}
+
+				return false;
+			}
+
+			const FVector::FReal SquareMagnitudeB = VectorB.SquaredLength();
+			if (FMath::IsNearlyZero(SquareMagnitudeB))
+			{
+				const FVector2D& MinA = SegmentStartA.X < SegmentEndA.X ? SegmentStartA : SegmentEndA;
+				const FVector2D& MaxA = SegmentStartA.X > SegmentEndA.X ? SegmentStartA : SegmentEndA;
+
+				if (SegmentStartB.X > MinA.X && SegmentStartB.X < MaxA.X)
+				{
+					OutIntersectionPoint = SegmentStartB;
+					return true;
+				}
+
+				return false;
+			}
+
+			// Taking the ray from SegmentStart to (0, 0) for both segments, and projecting the ray onto its
+			// respective line segment should yield two coincident points if the segments are collinear
+			const FVector2D AToZero = -SegmentStartA; // (0, 0) - SegmentStartA
+			const FVector2D ProjectionOnA = VectorA * VectorA.Dot(AToZero) / SquareMagnitudeA;
+			const FVector2D ProjectedPointOnA = SegmentStartA + ProjectionOnA;
+
+			const FVector2D BToZero = -SegmentStartB; // (0, 0) - SegmentStartB
+			const FVector2D ProjectionOnB = VectorB * VectorB.Dot(BToZero) / SquareMagnitudeB;
+			const FVector2D ProjectedPointOnB = SegmentStartB + ProjectionOnB;
+			
+			// If the projected points are not equal, then we are not collinear
+			if (!ProjectedPointOnA.Equals(ProjectedPointOnB))
+			{
+				return false;
+			}
+			
+			const FVector2D& MinA = SegmentStartA.X < SegmentEndA.X ? SegmentStartA : SegmentEndA;
+			const FVector2D& MaxA = SegmentStartA.X > SegmentEndA.X ? SegmentStartA : SegmentEndA;
+			const FVector2D& MinB = SegmentStartB.X < SegmentEndB.X ? SegmentStartB : SegmentEndB;
+			const FVector2D& MaxB = SegmentStartB.X > SegmentEndB.X ? SegmentStartB : SegmentEndB;
+			
+			if (MinA.X > MinB.X && MinA.X < MaxB.X)
+			{
+				OutIntersectionPoint = MinA;
+				return true;
+			}
+			else if (MaxA.X > MinB.X && MaxA.X < MaxB.X)
+			{
+				OutIntersectionPoint = MaxA;
+				return true;
+			}
+			if (MinB.X > MinA.X && MinB.X < MaxA.X)
+			{
+				OutIntersectionPoint = MinB;
+				return true;
+			}
+			else if (MaxB.X > MinA.X && MaxB.X < MaxA.X)
+			{
+				OutIntersectionPoint = MaxB;
+				return true;
+			}
+			
+			return false;
+		}
+
+		const FVector::FReal DeltaX = SegmentStartA.X - SegmentStartB.X;
+		const FVector::FReal DeltaY = SegmentStartA.Y - SegmentStartB.Y;
+
+		const FVector::FReal S = (VectorA.X * DeltaY - VectorA.Y * DeltaX) / Determinant;
+		const FVector::FReal T = (VectorB.X * DeltaY - VectorB.Y * DeltaX) / Determinant;
+
+		if (S >= 0 && S <= 1 && T >= 0 && T <= 1)
+		{
+			OutIntersectionPoint = SegmentStartA + T * VectorA;
+			return true;
+		}
+
+		return false;
+	}
+
+	/** Intersects a line segment with all line segments of a polygon. May return duplicates if the segment intersects a vertex of the polygon. */
+	void SegmentPolygonIntersection2D(const FVector2D& SegmentStart, const FVector2D& SegmentEnd, const TArray<FVector2D>& PolygonPoints, TArray<FVector2D>& OutIntersectionPoints)
+	{
+		const int32 PointCount = PolygonPoints.Num();
+
+		if (PointCount < 3)
+		{
+			return;
+		}
+
+		// Checks if a duplicate point should be counted as a valid intersection and adds it to the intersection list if so. Collinear points are also considered duplicates.
+		auto TestDuplicatePoint = [&](const FVector2D& IntersectionPoint, int32 Index) -> bool
+		{
+			const int32 BackwardIndex = Index - 1;
+			const int32 ForwardIndex = Index + 1;
+
+			// Find DeltaY for the nearest backward point
+			FVector::FReal DeltaY1 = 0.f;
+			for (int32 Offset = 0; Offset < PointCount; ++Offset)
+			{
+				DeltaY1 = PolygonPoints[(BackwardIndex - Offset + PointCount) % PointCount].Y - IntersectionPoint.Y;
+
+				// Choose the first non-zero delta
+				if (!FMath::IsNearlyZero(DeltaY1))
+				{
+					break;
+				}
+			}
+
+			// Find DeltaY for the nearest forward point
+			FVector::FReal DeltaY2 = 0.f;
+			for (int32 Offset = 0; Offset < PointCount; ++Offset)
+			{
+				DeltaY2 = PolygonPoints[(ForwardIndex + Offset) % PointCount].Y - IntersectionPoint.Y;
+
+				// Choose the first non-zero delta
+				if (!FMath::IsNearlyZero(DeltaY2))
+				{
+					break;
+				}
+			}
+
+			// Allow double counting when adjacent points lie on opposite sides of our ray
+			return (DeltaY1 > 0 && DeltaY2 > 0) || (DeltaY1 < 0 && DeltaY2 < 0);
+		};
+
+		auto AreSegmentsCollinear = [](const FVector2D& SegmentStartA, const FVector2D& SegmentEndA, const FVector2D& SegmentStartB, const FVector2D& SegmentEndB) -> bool
+		{
+			const FVector2D VectorA = SegmentEndA - SegmentStartA;
+			const FVector2D VectorB = SegmentEndB - SegmentStartB;
+
+			const FVector::FReal SquareMagnitudeA = VectorA.SquaredLength();
+			const FVector::FReal SquareMagnitudeB = VectorB.SquaredLength();
+
+			if (!FMath::IsNearlyZero(FVector2D::CrossProduct(VectorA, VectorB)) || FMath::IsNearlyZero(SquareMagnitudeA) || FMath::IsNearlyZero(SquareMagnitudeB))
+			{
+				return false;
+			}
+
+			// Taking the ray from SegmentStart to (0, 0) for both segments, and projecting the ray onto its
+			// respective line segment should yield two coincident points if the segments are collinear
+			const FVector2D AToZero = -SegmentStartA; // (0, 0) - SegmentStartA
+			const FVector2D ProjectionOnA = VectorA * VectorA.Dot(AToZero) / SquareMagnitudeA;
+			const FVector2D ProjectedPointOnA = SegmentStartA + ProjectionOnA;
+
+			const FVector2D BToZero = -SegmentStartB; // (0, 0) - SegmentStartB
+			const FVector2D ProjectionOnB = VectorB * VectorB.Dot(BToZero) / SquareMagnitudeB;
+			const FVector2D ProjectedPointOnB = SegmentStartB + ProjectionOnB;
+
+			// If the projected points are not equal, then we are not collinear
+			return ProjectedPointOnA.Equals(ProjectedPointOnB);
+		};
+
+		FVector2D IntersectionPoint;
+		FVector2D PreviousIntersectionPoint;
+		bool bLastSegmentWasCollinear = false;
+
+		for (int32 PointIndex = 0; PointIndex < PointCount - 1; ++PointIndex)
+		{
+			const FVector2D& PolySegmentStart = PolygonPoints[PointIndex];
+			const FVector2D& PolySegmentEnd = PolygonPoints[PointIndex + 1];
+
+			// Discard zero length segments
+			if (PolySegmentStart.Equals(PolySegmentEnd))
+			{
+				continue;
+			}
+
+			const bool bSegmentIsCollinear = AreSegmentsCollinear(PolySegmentEnd, PolySegmentStart, SegmentStart, SegmentEnd);
+
+			// Discard collinear segments and test if duplicate points should be included
+			if (!bSegmentIsCollinear && SegmentIntersection2D(SegmentStart, SegmentEnd, PolySegmentStart, PolySegmentEnd, IntersectionPoint))
+			{
+				if ((OutIntersectionPoints.Num() > 0 && IntersectionPoint.Equals(PreviousIntersectionPoint)) || bLastSegmentWasCollinear)
+				{
+					if (TestDuplicatePoint(IntersectionPoint, PointIndex))
+					{
+						OutIntersectionPoints.Add(IntersectionPoint);
+					}
+				}
+				else
+				{
+					OutIntersectionPoints.Add(IntersectionPoint);
+				}
+
+				PreviousIntersectionPoint = IntersectionPoint;
+			}
+
+			bLastSegmentWasCollinear = bSegmentIsCollinear;
+		}
+
+		// Special case to wrap the last segment
+		{
+			const FVector2D& PolySegmentStart = PolygonPoints.Last();
+			const FVector2D& PolySegmentEnd = PolygonPoints[0];
+
+			// Discard zero length segments
+			if (PolySegmentStart.Equals(PolySegmentEnd))
+			{
+				return;
+			}
+
+			const bool bSegmentIsCollinear = AreSegmentsCollinear(PolySegmentEnd, PolySegmentStart, SegmentStart, SegmentEnd);
+
+			if (bSegmentIsCollinear && OutIntersectionPoints.Num() > 0 && OutIntersectionPoints[0].Equals(PolygonPoints[0]))
+			{
+				// If the last segment is collinear, then we will double count the segment unless we cull the first polygon point
+				OutIntersectionPoints.RemoveAt(0);
+			}
+			else if (!bSegmentIsCollinear && SegmentIntersection2D(SegmentStart, SegmentEnd, PolySegmentStart, PolySegmentEnd, IntersectionPoint))
+			{
+				if (OutIntersectionPoints.Num() > 0 || bLastSegmentWasCollinear)
+				{
+					if (IntersectionPoint.Equals(PolygonPoints[0]))
+					{
+						if (TestDuplicatePoint(IntersectionPoint, PointCount))
+						{
+							OutIntersectionPoints.Add(IntersectionPoint);
+						}
+					}
+					else if (IntersectionPoint.Equals(PolygonPoints.Last()))
+					{
+						if (TestDuplicatePoint(IntersectionPoint, PointCount - 1))
+						{
+							OutIntersectionPoints.Add(IntersectionPoint);
+						}
+					}
+					else
+					{
+						OutIntersectionPoints.Add(IntersectionPoint);
+					}
+				}
+				else
+				{
+					OutIntersectionPoints.Add(IntersectionPoint);
+				}
+			}
+		}
+	}
+
+	/** Tests if a point lies inside the given polygon by casting a ray to MaxDistance and counting the intersections */
+	bool PointInsidePolygon2D(const TArray<FVector2D>& PolygonPoints, const FVector2D& Point, FVector::FReal MaxDistance) 
+	{
+		if (PolygonPoints.Num() < 3)
+		{
+			return false;
+		}
+
+		TArray<FVector2D> IntersectionPoints;
+		SegmentPolygonIntersection2D(Point, Point + FVector2D(MaxDistance, 0), PolygonPoints, IntersectionPoints);
+
+		// Odd number of intersections means we are inside the polygon
+		return (IntersectionPoints.Num() % 2) == 1;
+	}
+}
+
 namespace PCGSplineSampler
 {
 	struct FStepSampler
@@ -265,6 +550,301 @@ namespace PCGSplineSampler
 			Point.Seed = UPCGBlueprintHelpers::ComputeSeedFromPosition(Point.Transform.GetLocation());
 		}
 	}
+ 
+	void SampleInteriorData(const UPCGPolyLineData* LineData, const UPCGSpatialData* SpatialData, const FPCGSplineSamplerParams& Params, UPCGPointData* OutPointData)
+	{
+		check(LineData && OutPointData);
+
+		TSoftObjectPtr<USplineComponent> Spline;
+
+		// TODO: handle projected splines
+		if (const UPCGSplineData* SplineData = Cast<UPCGSplineData>(LineData))
+		{
+			Spline = SplineData->Spline;
+		}
+		else if (const UPCGLandscapeSplineData* LandscapeSplineData = Cast<UPCGLandscapeSplineData>(LineData))
+		{
+			UE_LOG(LogPCG, Error, TEXT("LandscapeSplines are not supported for interior sampling"));
+			return;
+		}
+		else
+		{
+			UE_LOG(LogPCG, Error, TEXT("Could not create UPCGSplineData from LineData"));
+			return;
+		}
+
+		check(Spline);
+
+		if (!Spline->IsClosedLoop())
+		{
+			UE_LOG(LogPCG, Error, TEXT("Interior sampling only generates for closed shapes"));
+			return;
+		}
+
+		const FVector MinPoint = Spline->Bounds.Origin - Spline->Bounds.BoxExtent;
+		const FVector MaxPoint = Spline->Bounds.Origin + Spline->Bounds.BoxExtent;
+
+		const FVector::FReal MaxDimension = FMath::Max(Spline->Bounds.BoxExtent.X, Spline->Bounds.BoxExtent.Y) * 2.f;
+		const FVector::FReal MaxDimensionSquared = MaxDimension * MaxDimension;
+
+		const FRichCurve* DensityFalloffCurve = Params.InteriorDensityFalloffCurve.GetRichCurveConst();
+		const bool bGenerateMedialAxis = DensityFalloffCurve != nullptr && DensityFalloffCurve->GetNumKeys() > 0;
+		const bool bFindNearestSplineKey = Params.InteriorOrientation == EPCGSplineSamplingInteriorOrientation::FollowCurvature || (bGenerateMedialAxis && !Params.bTreatSplineAsPolyline);
+		const bool bProjectOntoSurface = Params.bProjectOntoSurface || bFindNearestSplineKey;
+
+		const FVector::FReal BoundExtents = Params.InteriorBorderSampleSpacing * 0.5f;
+		const FVector BoundsMin = FVector::One() * -BoundExtents;
+		const FVector BoundsMax = FVector::One() * BoundExtents;
+
+		const FVector::FReal SplineLength = Spline->GetSplineLength();
+
+		TArray<FVector> SplineSamplePoints;
+
+		if (Params.bTreatSplineAsPolyline)
+		{
+			// Treat spline interface points as vertices of a polyline
+			const int NumSegments = LineData->GetNumSegments();
+			for (int32 SegmentIndex = 0; SegmentIndex < NumSegments; ++SegmentIndex)
+			{
+				SplineSamplePoints.Add(LineData->GetLocationAtDistance(SegmentIndex, 0));
+			}
+		}
+		else
+		{
+			// Get sample points along the spline that are higher resolution than our PolyLine
+			for (FVector::FReal Length = 0.f; Length < SplineLength; Length += Params.InteriorBorderSampleSpacing)
+			{
+				SplineSamplePoints.Add(Spline->GetLocationAtDistanceAlongSpline(Length, ESplineCoordinateSpace::World));
+			}
+		}
+
+		// Flat polygon representation of our spline points
+		TArray<FVector2D> SplineSamplePoints2D;
+		for (const FVector& Point : SplineSamplePoints)
+		{
+			SplineSamplePoints2D.Add(FVector2D(Point));
+		}
+
+		TArray<TTuple<FVector2D, FVector2D>> MedialAxisEdges;
+
+		// Compute the Medial Axis as a subset of the Voronoi Diagram of the PolyLine points
+		if (bGenerateMedialAxis)
+		{
+			TArray<FVector> PolygonPoints; // Top-down 2D projection polygon of the spline points
+			if (Params.bTreatSplineAsPolyline)
+			{
+				// If we already computed the polygon, use a copy instead of generating it again
+				for (const FVector& SplinePoint : SplineSamplePoints)
+				{
+					FVector& PolygonPoint = PolygonPoints.Add_GetRef(SplinePoint);
+					PolygonPoint.Z = MinPoint.Z;
+				}
+			}
+			else
+			{
+				const int NumSegments = LineData->GetNumSegments();
+				for (int32 SegmentIndex = 0; SegmentIndex < NumSegments; ++SegmentIndex)
+				{
+					FVector& PolygonPoint = PolygonPoints.Add_GetRef(LineData->GetLocationAtDistance(SegmentIndex, 0));
+					PolygonPoint.Z = MinPoint.Z;
+				}
+			}
+
+			TArray<FVector2D> PolygonPoints2D;
+			for (const FVector& Point : PolygonPoints)
+			{
+				PolygonPoints2D.Add(FVector2D(Point));
+			}
+
+			TArray<TTuple<FVector, FVector>> VoronoiEdges;
+			TArray<int32> CellMember;
+
+			GetVoronoiEdges(PolygonPoints, FBox(MinPoint, FVector(MaxPoint.X, MaxPoint.Y, MinPoint.Z)), VoronoiEdges, CellMember);
+
+			// Find the subset of the Voronoi Diagram which composes the Medial Axis
+			for (const TTuple<FVector, FVector>& Edge : VoronoiEdges)
+			{
+				// Discard any edges which intersect the polygon
+				bool bDiscard = false;
+				for (int32 PointIndex = 0; PointIndex < PolygonPoints.Num(); ++PointIndex)
+				{
+					FVector2D IntersectionPoint;
+					if (PCGSplineSamplerHelpers::SegmentIntersection2D(PolygonPoints2D[PointIndex], PolygonPoints2D[(PointIndex + 1) % PolygonPoints2D.Num()], FVector2D(Edge.Get<0>()), FVector2D(Edge.Get<1>()), IntersectionPoint))
+					{
+						bDiscard = true;
+						break;
+					}
+				}
+
+				if (bDiscard)
+				{
+					continue;
+				}
+
+				// If either of the points lies within the polygon, the segment must lie within the polygon
+				if (PCGSplineSamplerHelpers::PointInsidePolygon2D(PolygonPoints2D, FVector2D(Edge.Get<0>()), MaxDimension))
+				{
+					MedialAxisEdges.Add(Edge);
+				}
+			}
+		}
+
+		TArray<FPCGPoint>& OutPoints = OutPointData->GetMutablePoints();
+		TArray<FVector::FReal> DistancesSquaredOnXY;
+
+		// RayPadding should be large enough to avoid potential misclassifications in SegmentPolygonIntersection2D
+		const FVector::FReal RayPadding = 100.f;
+		const FVector::FReal MinY = FMath::CeilToDouble(MinPoint.Y / Params.InteriorSampleSpacing) * Params.InteriorSampleSpacing;
+		const FVector::FReal MaxY = FMath::FloorToDouble(MaxPoint.Y / Params.InteriorSampleSpacing) * Params.InteriorSampleSpacing;
+
+		// Point sampling
+		for (FVector::FReal Y = MinY; Y < MaxY + UE_KINDA_SMALL_NUMBER; Y += Params.InteriorSampleSpacing)
+		{
+			const FVector2D RayMin(MinPoint.X - RayPadding, Y);
+			const FVector2D RayMax(MaxPoint.X + RayPadding, Y);
+
+			// Get the intersections along this ray, sorted by Y value
+			TArray<FVector2D> Intersections;
+			PCGSplineSamplerHelpers::SegmentPolygonIntersection2D(RayMin, RayMax, SplineSamplePoints2D, Intersections);
+
+			if (Intersections.Num() % 2 != 0)
+			{
+				UE_LOG(LogPCG, Error, TEXT("Intersection test failed, skipping samples for this row"));
+				continue;
+			}
+
+			Intersections.Sort([](const FVector2D& LHS, const FVector2D& RHS) { return LHS.X < RHS.X; });
+
+			// TODO: async processing
+			// Each pair of intersections defines a range in which point samples may lie
+			for (int32 RangeIndex = 0; RangeIndex < Intersections.Num(); RangeIndex += 2)
+			{
+				const FVector::FReal MinX = FMath::CeilToDouble(Intersections[RangeIndex].X / Params.InteriorSampleSpacing) * Params.InteriorSampleSpacing;
+				const FVector::FReal MaxX = FMath::FloorToDouble(Intersections[RangeIndex + 1].X / Params.InteriorSampleSpacing) * Params.InteriorSampleSpacing;
+
+				if (MaxX - MinX < Params.InteriorSampleSpacing)
+				{
+					continue;
+				}
+
+				for (FVector::FReal X = MinX; X < MaxX + UE_KINDA_SMALL_NUMBER; X += Params.InteriorSampleSpacing)
+				{
+					FPCGPoint& Point = OutPoints.Emplace_GetRef();
+
+					const FVector2D SampleLocation = FVector2D(X, Y);
+					FVector SurfaceLocation = FVector(SampleLocation, MinPoint.Z);
+
+					if (bProjectOntoSurface)
+					{
+						// Precompute 2D distance to every spline point
+						DistancesSquaredOnXY.Reset(SplineSamplePoints.Num());
+						for (const FVector& SplinePoint : SplineSamplePoints)
+						{
+							FVector::FReal DistanceSquared = FVector::DistSquaredXY(SurfaceLocation, SplinePoint);
+							DistancesSquaredOnXY.Add(DistanceSquared);
+						}
+
+						// Compute average Z value weighted by 1 / Distance^2
+						FVector::FReal SumZ = 0.f;
+						FVector::FReal SumWeights = 0.f;
+						for (int32 PointIndex = 0; PointIndex < SplineSamplePoints.Num(); ++PointIndex)
+						{
+							const FVector::FReal DistanceSquared = DistancesSquaredOnXY[PointIndex];
+
+							// If sample point overlaps exactly with a border point, then that must be the height
+							if (FMath::IsNearlyZero(DistanceSquared))
+							{
+								SumZ = SplineSamplePoints[PointIndex].Z;
+								SumWeights = 1.f;
+								break;
+							}
+
+							// TODO: it would be more accurate to use distance to the polyline instead of distance to the polyline points,
+							// however it would also be much more expensive. Perhaps worth investigating when Params.bTreatSplineAsPolyline is true
+							const FVector::FReal Weight = 1.f / DistanceSquared;
+
+							SumWeights += Weight;
+							SumZ += SplineSamplePoints[PointIndex].Z * Weight;
+						}
+
+						SurfaceLocation.Z = SumZ / SumWeights;
+					}
+
+					// if bTreatAsPolyline, then we shouldnt use this, we should use nearest point on the polygon line segments
+					const float NearestSplineKey = bFindNearestSplineKey ? Spline->FindInputKeyClosestToWorldLocation(SurfaceLocation) : 0.f;
+
+					if (Params.bProjectOntoSurface)
+					{
+						Point.Transform.SetLocation(SurfaceLocation);
+					}
+					else
+					{
+						Point.Transform.SetLocation(FVector(SampleLocation, MinPoint.Z));
+					}
+
+					if (Params.InteriorOrientation == EPCGSplineSamplingInteriorOrientation::FollowCurvature)
+					{
+						Point.Transform.SetRotation(Spline->GetQuaternionAtSplineInputKey(NearestSplineKey, ESplineCoordinateSpace::World));
+					}
+
+					// Calculate density fall off
+					if (bGenerateMedialAxis && MedialAxisEdges.Num())
+					{
+						FVector::FReal SmallestDistSquared = MaxDimensionSquared;
+
+						// Find distance from SampleLocation to MedialAxis
+						for (const TTuple<FVector2D, FVector2D>& Edge : MedialAxisEdges)
+						{
+							const FVector::FReal DistSquared = FMath::PointDistToSegmentSquared(FVector(SampleLocation, 0), FVector(Edge.Get<0>(), 0), FVector(Edge.Get<1>(), 0));
+
+							if (DistSquared < SmallestDistSquared)
+							{
+								SmallestDistSquared = DistSquared;
+							}
+						}
+
+						const FVector::FReal SmallestDist = FMath::Sqrt(SmallestDistSquared);
+
+						FVector::FReal PointToSplineDist = 0.f;
+
+						if (Params.bTreatSplineAsPolyline)
+						{
+							SmallestDistSquared = MaxDimensionSquared;
+							for (int32 PointIndex = 0; PointIndex < SplineSamplePoints.Num(); ++PointIndex)
+							{
+								const FVector::FReal DistSquared = FMath::PointDistToSegmentSquared(SurfaceLocation, SplineSamplePoints[PointIndex], SplineSamplePoints[(PointIndex + 1) % SplineSamplePoints.Num()]);
+
+								if (DistSquared < SmallestDistSquared)
+								{
+									SmallestDistSquared = DistSquared;
+								}
+							}
+							
+							PointToSplineDist = FMath::Sqrt(SmallestDistSquared);
+						}
+						else
+						{
+							const FVector NearestSplineLocation = Spline->GetLocationAtSplineInputKey(NearestSplineKey, ESplineCoordinateSpace::World);
+							PointToSplineDist = FVector2D::Distance(SampleLocation, FVector2D(NearestSplineLocation.X, NearestSplineLocation.Y));
+						}
+
+						// Linear fall off in the range [0, 1]
+						const FVector::FReal T = SmallestDist / (SmallestDist + PointToSplineDist + UE_KINDA_SMALL_NUMBER);
+
+						Point.Density = DensityFalloffCurve->Eval(T);
+					}
+					else
+					{
+						Point.Density = 1.f;
+					}
+
+					Point.BoundsMin = BoundsMin;
+					Point.BoundsMax = BoundsMax;
+					Point.Seed = UPCGBlueprintHelpers::ComputeSeedFromPosition(Point.Transform.GetLocation());
+				}
+			}
+		}
+	}
 
 	const UPCGPolyLineData* GetPolyLineData(const UPCGSpatialData* InSpatialData)
 	{
@@ -328,6 +908,10 @@ bool FPCGSplineSamplerElement::ExecuteInternal(FPCGContext* Context) const
 	SamplerParams.DistanceIncrement = PCG_GET_OVERRIDEN_VALUE(&SamplerParams, DistanceIncrement, Params);
 	SamplerParams.NumPlanarSubdivisions = PCG_GET_OVERRIDEN_VALUE(&SamplerParams, NumPlanarSubdivisions, Params);
 	SamplerParams.NumHeightSubdivisions = PCG_GET_OVERRIDEN_VALUE(&SamplerParams, NumHeightSubdivisions, Params);
+	SamplerParams.InteriorSampleSpacing = PCG_GET_OVERRIDEN_VALUE(&SamplerParams, InteriorSampleSpacing, Params);
+	SamplerParams.InteriorBorderSampleSpacing = PCG_GET_OVERRIDEN_VALUE(&SamplerParams, InteriorBorderSampleSpacing, Params);
+	SamplerParams.InteriorOrientation = PCG_GET_OVERRIDEN_VALUE(&SamplerParams, InteriorOrientation, Params);
+	SamplerParams.bProjectOntoSurface = PCG_GET_OVERRIDEN_VALUE(&SamplerParams, bProjectOntoSurface, Params);
 
 	TArray<FPCGTaggedData>& Outputs = Context->OutputData.TaggedData;
 
@@ -353,11 +937,17 @@ bool FPCGSplineSamplerElement::ExecuteInternal(FPCGContext* Context) const
 
 		UPCGPointData* SampledPointData = NewObject<UPCGPointData>(const_cast<UPCGPolyLineData*>(LineData));
 		SampledPointData->InitializeFromData(SpatialData);
-
 		Output.Data = SampledPointData;
 
-		PCGSplineSampler::SampleLineData(LineData, SpatialData, SamplerParams, SampledPointData);
+		if (SamplerParams.Dimension == EPCGSplineSamplingDimension::OnInterior)
+		{
+			PCGSplineSampler::SampleInteriorData(LineData, SpatialData, SamplerParams, SampledPointData);
+		}
+		else
+		{
+			PCGSplineSampler::SampleLineData(LineData, SpatialData, SamplerParams, SampledPointData);
+		}
 	}
 
 	return true;
-}
\ No newline at end of file
+}
diff --git a/Engine/Plugins/Experimental/PCG/Source/PCG/Private/Tests/Elements/PCGSplineSamplerTest.cpp b/Engine/Plugins/Experimental/PCG/Source/PCG/Private/Tests/Elements/PCGSplineSamplerTest.cpp
new file mode 100644
--- /dev/null
+++ b/Engine/Plugins/Experimental/PCG/Source/PCG/Private/Tests/Elements/PCGSplineSamplerTest.cpp
@@ -0,0 +1,97 @@
+// Copyright Epic Games, Inc. All Rights Reserved.
+
+#include "Tests/PCGTestsCommon.h"
+#include "Tests/Determinism/PCGDeterminismTestsCommon.h"
+#include "Elements/PCGSplineSampler.h"
+
+#if WITH_EDITOR
+
+IMPLEMENT_CUSTOM_SIMPLE_AUTOMATION_TEST(FPCGSplineSamplerTest, FPCGTestBaseClass, "pcg.tests.SplineSampler.PointInPolygon", PCGTestsCommon::TestFlags)
+
+bool FPCGSplineSamplerTest::RunTest(const FString& Parameters)
+{	
+	/* Clock-wise torch shape, to test edge cases
+	 * |\/ \/|
+	 * |_   _|
+	 *   |_|
+	 */
+
+	const TArray<FVector2D> PolygonPoints =
+	{
+		FVector2D(0, 0),
+		FVector2D(0, 1),
+		FVector2D(0, 2),
+		FVector2D(1, 1),
+		FVector2D(1.5, 2),
+		FVector2D(1.5, 2), // Some duplicate points
+		FVector2D(2, 1),
+		FVector2D(3, 2),
+		FVector2D(3, 1),
+		FVector2D(3, 0),
+		FVector2D(3, 0),
+		FVector2D(2.5, 0),
+		FVector2D(2, 0),
+		FVector2D(2, -1),
+		FVector2D(2, -1),
+		FVector2D(1, -1),
+		FVector2D(1, 0),
+		FVector2D(0.5, 0),
+	};
+
+	constexpr FVector::FReal MaxDistance = 1000.f;
+	bool bTestPassed = true;
+
+	bTestPassed &= TestTrue("Interior point lies inside polygon", PCGSplineSamplerHelpers::PointInsidePolygon2D(PolygonPoints, FVector2D(0.5, 0.5), MaxDistance));
+	bTestPassed &= TestTrue("Exterior point lies outside polygon", !PCGSplineSamplerHelpers::PointInsidePolygon2D(PolygonPoints, FVector2D(10.5, 0.5), MaxDistance));
+
+	// Y = -1
+	{
+		bTestPassed &= TestTrue("(1-e, -1) lies outside polygon", !PCGSplineSamplerHelpers::PointInsidePolygon2D(PolygonPoints, FVector2D(1 - UE_KINDA_SMALL_NUMBER, -1), MaxDistance));
+		bTestPassed &= TestTrue("(2+e, -1) lies outside polygon", !PCGSplineSamplerHelpers::PointInsidePolygon2D(PolygonPoints, FVector2D(2 + UE_KINDA_SMALL_NUMBER, -1), MaxDistance));
+	}
+
+	// Y = 0
+	{
+		bTestPassed &= TestTrue("(0-e, 0) lies outside polygon", !PCGSplineSamplerHelpers::PointInsidePolygon2D(PolygonPoints, FVector2D(0 - UE_KINDA_SMALL_NUMBER, 0), MaxDistance));
+		bTestPassed &= TestTrue("(1.5, 0) lies inside polygon", PCGSplineSamplerHelpers::PointInsidePolygon2D(PolygonPoints, FVector2D(1.5, 0), MaxDistance));
+		bTestPassed &= TestTrue("(3+e, 0) lies outside polygon", !PCGSplineSamplerHelpers::PointInsidePolygon2D(PolygonPoints, FVector2D(3 + UE_KINDA_SMALL_NUMBER, 0), MaxDistance));
+	}
+
+	// Y = 1
+	{
+		// Test around (0, 1)
+		bTestPassed &= TestTrue("(0-e, 1) lies outside polygon", !PCGSplineSamplerHelpers::PointInsidePolygon2D(PolygonPoints, FVector2D(0 - UE_KINDA_SMALL_NUMBER, 1), MaxDistance));
+		bTestPassed &= TestTrue("(0+e, 1) lies inside polygon", PCGSplineSamplerHelpers::PointInsidePolygon2D(PolygonPoints, FVector2D(0 + UE_KINDA_SMALL_NUMBER, 1), MaxDistance));
+
+		// Test around (1, 1)
+		bTestPassed &= TestTrue("(1-e, 1) lies inside polygon", PCGSplineSamplerHelpers::PointInsidePolygon2D(PolygonPoints, FVector2D(1 - UE_KINDA_SMALL_NUMBER, 1), MaxDistance));
+		bTestPassed &= TestTrue("(1+e, 1) lies inside polygon", PCGSplineSamplerHelpers::PointInsidePolygon2D(PolygonPoints, FVector2D(1 + UE_KINDA_SMALL_NUMBER, 1), MaxDistance));
+
+		// Test around (2, 1)
+		bTestPassed &= TestTrue("(2-e, 1) lies inside polygon", PCGSplineSamplerHelpers::PointInsidePolygon2D(PolygonPoints, FVector2D(2 - UE_KINDA_SMALL_NUMBER, 1), MaxDistance));
+		bTestPassed &= TestTrue("(2+e, 1) lies inside polygon", PCGSplineSamplerHelpers::PointInsidePolygon2D(PolygonPoints, FVector2D(2 + UE_KINDA_SMALL_NUMBER, 1), MaxDistance));
+
+		// Test around (3, 1)
+		bTestPassed &= TestTrue("(3-e, 1) lies inside polygon", PCGSplineSamplerHelpers::PointInsidePolygon2D(PolygonPoints, FVector2D(3 - UE_KINDA_SMALL_NUMBER, 1), MaxDistance));
+		bTestPassed &= TestTrue("(3+e, 1) lies outside polygon", !PCGSplineSamplerHelpers::PointInsidePolygon2D(PolygonPoints, FVector2D(3 + UE_KINDA_SMALL_NUMBER, 1), MaxDistance));
+	}
+
+	// Y = 2
+	{
+		// Test around (0, 2)
+		bTestPassed &= TestTrue("(0-e, 2) lies outside polygon", !PCGSplineSamplerHelpers::PointInsidePolygon2D(PolygonPoints, FVector2D(0 - UE_KINDA_SMALL_NUMBER, 2), MaxDistance));
+		bTestPassed &= TestTrue("(0+e, 2) lies outside polygon", !PCGSplineSamplerHelpers::PointInsidePolygon2D(PolygonPoints, FVector2D(0 + UE_KINDA_SMALL_NUMBER, 2), MaxDistance));
+
+		// Test around (1.5, 2)
+		bTestPassed &= TestTrue("(1.5-e, 2) lies outside polygon", !PCGSplineSamplerHelpers::PointInsidePolygon2D(PolygonPoints, FVector2D(1.5 - UE_KINDA_SMALL_NUMBER, 2), MaxDistance));
+		bTestPassed &= TestTrue("(1.5+e, 2) lies outside polygon", !PCGSplineSamplerHelpers::PointInsidePolygon2D(PolygonPoints, FVector2D(1.5 + UE_KINDA_SMALL_NUMBER, 2), MaxDistance));
+
+		// Test around (3, 2)
+		bTestPassed &= TestTrue("(3-e, 2) lies outside polygon", !PCGSplineSamplerHelpers::PointInsidePolygon2D(PolygonPoints, FVector2D(3 - UE_KINDA_SMALL_NUMBER, 2), MaxDistance));
+		bTestPassed &= TestTrue("(3+e, 2) lies outside polygon", !PCGSplineSamplerHelpers::PointInsidePolygon2D(PolygonPoints, FVector2D(3 + UE_KINDA_SMALL_NUMBER, 2), MaxDistance));
+	}
+
+	return bTestPassed;
+}
+
+#endif // WITH_EDITOR
diff --git a/Engine/Plugins/Experimental/PCG/Source/PCG/Public/Elements/PCGSplineSampler.h b/Engine/Plugins/Experimental/PCG/Source/PCG/Public/Elements/PCGSplineSampler.h
--- a/Engine/Plugins/Experimental/PCG/Source/PCG/Public/Elements/PCGSplineSampler.h
+++ b/Engine/Plugins/Experimental/PCG/Source/PCG/Public/Elements/PCGSplineSampler.h
@@ -4,6 +4,7 @@
 
 #include "PCGSettings.h"
 #include "PCGElement.h"
+#include "Curves/CurveFloat.h"
 
 #include "PCGSplineSampler.generated.h"
 
@@ -24,7 +25,8 @@ enum class EPCGSplineSamplingDimension : uint8
 	OnSpline = 0,
 	OnHorizontal,
 	OnVertical,
-	OnVolume
+	OnVolume,
+	OnInterior
 };
 
 UENUM()
@@ -34,24 +36,31 @@ enum class EPCGSplineSamplingFill : uint8
 	EdgesOnly = 1
 };
 
+UENUM()
+enum class EPCGSplineSamplingInteriorOrientation : uint8
+{
+	Uniform = 0,
+	FollowCurvature = 1
+};
+
 USTRUCT(BlueprintType)
 struct PCG_API FPCGSplineSamplerParams
 {
 	GENERATED_BODY()
 
-	UPROPERTY(BlueprintReadWrite, EditAnywhere, Category = Settings)
+	UPROPERTY(BlueprintReadWrite, EditAnywhere, Category = Settings, meta = (EditCondition = "Dimension!=EPCGSplineSamplingDimension::OnInterior"))
 	EPCGSplineSamplingMode Mode = EPCGSplineSamplingMode::Subdivision;
 
 	UPROPERTY(BlueprintReadWrite, EditAnywhere, Category = Settings)
 	EPCGSplineSamplingDimension Dimension = EPCGSplineSamplingDimension::OnSpline;
 
-	UPROPERTY(BlueprintReadWrite, EditAnywhere, Category = Settings, meta = (EditCondition = "Dimension!=EPCGSplineSamplingDimension::OnSpline"))
+	UPROPERTY(BlueprintReadWrite, EditAnywhere, Category = Settings, meta = (EditCondition = "Dimension!=EPCGSplineSamplingDimension::OnSpline&&Dimension!=EPCGSplineSamplingDimension::OnInterior"))
 	EPCGSplineSamplingFill Fill = EPCGSplineSamplingFill::Fill;
 
-	UPROPERTY(BlueprintReadWrite, EditAnywhere, Category = Settings, meta = (ClampMin = "0", EditCondition = "Mode==EPCGSplineSamplingMode::Subdivision"))
+	UPROPERTY(BlueprintReadWrite, EditAnywhere, Category = Settings, meta = (ClampMin = "0", EditCondition = "Mode==EPCGSplineSamplingMode::Subdivision&&Dimension!=EPCGSplineSamplingDimension::OnInterior"))
 	int32 SubdivisionsPerSegment = 1;
 
-	UPROPERTY(BlueprintReadWrite, EditAnywhere, Category = Settings, meta = (ClampMin = "0.1", EditCondition = "Mode==EPCGSplineSamplingMode::Distance"))
+	UPROPERTY(BlueprintReadWrite, EditAnywhere, Category = Settings, meta = (ClampMin = "0.1", EditCondition = "Mode==EPCGSplineSamplingMode::Distance&&Dimension!=EPCGSplineSamplingDimension::OnInterior"))
 	float DistanceIncrement = 100.0f;
 
 	UPROPERTY(BlueprintReadWrite, EditAnywhere, Category = Settings, meta = (ClampMin = "0", EditCondition = "Dimension==EPCGSplineSamplingDimension::OnHorizontal||Dimension==EPCGSplineSamplingDimension::OnVolume"))
@@ -59,11 +68,43 @@ struct PCG_API FPCGSplineSamplerParams
 
 	UPROPERTY(BlueprintReadWrite, EditAnywhere, Category = Settings, meta = (ClampMin = "0", EditCondition = "Dimension==EPCGSplineSamplingDimension::OnVertical||Dimension==EPCGSplineSamplingDimension::OnVolume"))
 	int32 NumHeightSubdivisions = 8;
+
+	/** The space between each sample point */
+	UPROPERTY(BlueprintReadWrite, EditAnywhere, Category = Settings, meta = (ClampMin = "0.1", EditCondition = "Dimension==EPCGSplineSamplingDimension::OnInterior"))
+	float InteriorSampleSpacing = 100.0f;
+
+	/** The space between each sample point on the spline boundary. Used for computation; lower spacing is more expensive but more accurate. */
+	UPROPERTY(BlueprintReadWrite, EditAnywhere, Category = Settings, meta = (ClampMin = "0.1", EditCondition = "Dimension==EPCGSplineSamplingDimension::OnInterior&&!bTreatSplineAsPolyline"))
+	float InteriorBorderSampleSpacing = 100.0f;
+
+	/** Use the spline points to form a polyline, instead of computing many sample points along the spline. This is more accurate if your spline is linear. */
+	UPROPERTY(BlueprintReadWrite, EditAnywhere, Category = Settings, meta = (EditCondition = "Dimension==EPCGSplineSamplingDimension::OnInterior"))
+	bool bTreatSplineAsPolyline = false;
+
+	/** Determines the orientation of interior points */
+	UPROPERTY(BlueprintReadWrite, EditAnywhere, Category = Settings, meta = (EditCondition = "Dimension==EPCGSplineSamplingDimension::OnInterior"))
+	EPCGSplineSamplingInteriorOrientation InteriorOrientation = EPCGSplineSamplingInteriorOrientation::Uniform;
+
+	/** Project sample points onto one possible surface given by the spline boundary */
+	UPROPERTY(EditAnywhere, Category = Settings, meta = (EditCondition = "Dimension==EPCGSplineSamplingDimension::OnInterior"))
+	bool bProjectOntoSurface = false;
+
+	// TODO: DirtyCache for OnDependencyChanged when this float curve is an external asset
+	/** Defines the density for each sample based on its distance from the spline */
+	UPROPERTY(EditAnywhere, Category = Settings, meta = (EditCondition = "Dimension==EPCGSplineSamplingDimension::OnInterior"))
+	FRuntimeFloatCurve InteriorDensityFalloffCurve;
 };
 
+namespace PCGSplineSamplerHelpers
+{
+	/** Tests if a point lies inside the given polygon by casting a ray to MaxDistance and counting the intersections */
+	bool PointInsidePolygon2D(const TArray<FVector2D>& PolygonPoints, const FVector2D& Point, FVector::FReal MaxDistance);
+}
+
 namespace PCGSplineSampler
 {
 	void SampleLineData(const UPCGPolyLineData* LineData, const UPCGSpatialData* SpatialData, const FPCGSplineSamplerParams& Params, UPCGPointData* OutPointData);
+	void SampleInteriorData(const UPCGPolyLineData* LineData, const UPCGSpatialData* SpatialData, const FPCGSplineSamplerParams& Params, UPCGPointData* OutPointData);
 	const UPCGPolyLineData* GetPolyLineData(const UPCGSpatialData* InSpatialData);
 }
 
