* Renaming PCG Determinism Test for clarification

#rnx

diff --git a/Engine/Plugins/Experimental/PCG/Source/PCG/Private/Tests/Determinism/PCGDeterminismTestBlueprintBase.cpp b/Engine/Plugins/Experimental/PCG/Source/PCG/Private/Tests/Determinism/PCGDeterminismTestBlueprintBase.cpp
--- a/Engine/Plugins/Experimental/PCG/Source/PCG/Private/Tests/Determinism/PCGDeterminismTestBlueprintBase.cpp
+++ b/Engine/Plugins/Experimental/PCG/Source/PCG/Private/Tests/Determinism/PCGDeterminismTestBlueprintBase.cpp
@@ -2,11 +2,11 @@
 
 #include "Tests/Determinism/PCGDeterminismTestBlueprintBase.h"
 
-void UPCGDeterminismTestBlueprintBase::ExecuteTest_Implementation(const UPCGNode* InPCGNode, FDeterminismNodeTestResult& InOutTestResult)
+void UPCGDeterminismTestBlueprintBase::ExecuteTest_Implementation(const UPCGNode* InPCGNode, FDeterminismTestResult& InOutTestResult)
 {
 	// The default case if the user tries to run the blueprint test unimplemented
-	InOutTestResult.NodeTitle = InPCGNode->GetNodeTitle();
-	InOutTestResult.NodeName = InPCGNode->GetName();
+	InOutTestResult.TestResultTitle = InPCGNode->GetNodeTitle();
+	InOutTestResult.TestResultName = InPCGNode->GetName();
 	InOutTestResult.Seed = InPCGNode->DefaultSettings->Seed;
 	PCGDeterminismTests::RunDeterminismTest(InPCGNode, InOutTestResult, PCGDeterminismTests::Defaults::DeterminismBasicTestInfo);
 }
\ No newline at end of file
diff --git a/Engine/Plugins/Experimental/PCG/Source/PCG/Private/Tests/Determinism/PCGDeterminismTestsCommon.cpp b/Engine/Plugins/Experimental/PCG/Source/PCG/Private/Tests/Determinism/PCGDeterminismTestsCommon.cpp
--- a/Engine/Plugins/Experimental/PCG/Source/PCG/Private/Tests/Determinism/PCGDeterminismTestsCommon.cpp
+++ b/Engine/Plugins/Experimental/PCG/Source/PCG/Private/Tests/Determinism/PCGDeterminismTestsCommon.cpp
@@ -33,14 +33,14 @@ static TAutoConsoleVariable<int32> CVarDeterminismPermutationLimit(
 
 namespace PCGDeterminismTests
 {
-	bool LogInvalidTest(const UPCGNode* InPCGNode, FDeterminismNodeTestResult& OutResult)
+	bool LogInvalidTest(const UPCGNode* InPCGNode, FDeterminismTestResult& OutResult)
 	{
 		UE_LOG(LogPCG, Warning, TEXT("Attempting to run an invalid determinism test"));
 		OutResult.AdditionalDetails.Add(TEXT("Invalid test"));
 		return false;
 	}
 
-	void RunDeterminismTest(const UPCGNode* InPCGNode, FDeterminismNodeTestResult& OutResult, const FNodeTestInfo& TestToRun)
+	void RunDeterminismTest(const UPCGNode* InPCGNode, FDeterminismTestResult& OutResult, const FNodeTestInfo& TestToRun)
 	{
 		TRACE_CPUPROFILER_EVENT_SCOPE(PCGDeterminismTests::RunDeterminismTest::Node);
 
@@ -65,7 +65,7 @@ namespace PCGDeterminismTests
 	}
 
 #if WITH_EDITOR
-	void RunDeterminismTest(const UPCGGraph* InPCGGraph, UPCGComponent* InPCGComponent, FDeterminismNodeTestResult& OutResult)
+	void RunDeterminismTest(const UPCGGraph* InPCGGraph, UPCGComponent* InPCGComponent, FDeterminismTestResult& OutResult)
 	{
 		TRACE_CPUPROFILER_EVENT_SCOPE(PCGDeterminismTests::RunDeterminismTest::Graph);
 
@@ -211,7 +211,7 @@ namespace PCGDeterminismTests
 	}
 #endif
 
-	bool RunBasicTestSuite(const UPCGNode* InPCGNode, const FName& TestName, FDeterminismNodeTestResult& OutResult)
+	bool RunBasicTestSuite(const UPCGNode* InPCGNode, const FName& TestName, FDeterminismTestResult& OutResult)
 	{
 		// Get Pins
 		const TArray<TObjectPtr<UPCGPin>>& InputPins = InPCGNode->GetInputPins();
@@ -229,7 +229,7 @@ namespace PCGDeterminismTests
 		return RunBasicSelfTest(NodeAndOptions) && RunBasicCopiedSelfTest(NodeAndOptions);
 	}
 
-	bool RunOrderIndependenceSuite(const UPCGNode* InPCGNode, const FName& TestName, FDeterminismNodeTestResult& OutResult)
+	bool RunOrderIndependenceSuite(const UPCGNode* InPCGNode, const FName& TestName, FDeterminismTestResult& OutResult)
 	{
 		// Start with not deterministic
 		UpdateTestResults(TestName, OutResult, EDeterminismLevel::NoDeterminism);
@@ -1215,7 +1215,7 @@ namespace PCGDeterminismTests
 		return false;
 	}
 
-	void UpdateTestResultForOverPermutationLimitError(FDeterminismNodeTestResult& OutResult)
+	void UpdateTestResultForOverPermutationLimitError(FDeterminismTestResult& OutResult)
 	{
 		OutResult.DataTypesTested = EPCGDataType::None;
 		OutResult.bFlagRaised = true;
@@ -1361,7 +1361,7 @@ namespace PCGDeterminismTests
 		return BaseOptionsPerPin[PinIndex][PermutationIndex];
 	}
 
-	void UpdateTestResults(FName TestName, FDeterminismNodeTestResult& OutResult, EDeterminismLevel DeterminismLevel)
+	void UpdateTestResults(FName TestName, FDeterminismTestResult& OutResult, EDeterminismLevel DeterminismLevel)
 	{
 		EDeterminismLevel& Result = OutResult.TestResults.FindOrAdd(TestName);
 		Result = DeterminismLevel;
diff --git a/Engine/Plugins/Experimental/PCG/Source/PCG/Public/Tests/Determinism/PCGDeterminismTestBlueprintBase.h b/Engine/Plugins/Experimental/PCG/Source/PCG/Public/Tests/Determinism/PCGDeterminismTestBlueprintBase.h
--- a/Engine/Plugins/Experimental/PCG/Source/PCG/Public/Tests/Determinism/PCGDeterminismTestBlueprintBase.h
+++ b/Engine/Plugins/Experimental/PCG/Source/PCG/Public/Tests/Determinism/PCGDeterminismTestBlueprintBase.h
@@ -15,5 +15,5 @@ class PCG_API UPCGDeterminismTestBlueprintBase : public UObject
 
 public:
 	UFUNCTION(BlueprintNativeEvent, BlueprintPure, CallInEditor, Category = Determinism)
-	void ExecuteTest(const UPCGNode* InPCGNode, UPARAM(ref)FDeterminismNodeTestResult& InOutTestResult);
+	void ExecuteTest(const UPCGNode* InPCGNode, UPARAM(ref)FDeterminismTestResult& InOutTestResult);
 };
\ No newline at end of file
diff --git a/Engine/Plugins/Experimental/PCG/Source/PCG/Public/Tests/Determinism/PCGDeterminismTestsCommon.h b/Engine/Plugins/Experimental/PCG/Source/PCG/Public/Tests/Determinism/PCGDeterminismTestsCommon.h
--- a/Engine/Plugins/Experimental/PCG/Source/PCG/Public/Tests/Determinism/PCGDeterminismTestsCommon.h
+++ b/Engine/Plugins/Experimental/PCG/Source/PCG/Public/Tests/Determinism/PCGDeterminismTestsCommon.h
@@ -31,17 +31,17 @@ enum class EDeterminismLevel : uint8
 };
 
 USTRUCT(BlueprintType)
-struct FDeterminismNodeTestResult // TODO: Refactor the names since this isn't always a node any longer
+struct FDeterminismTestResult
 {
 	GENERATED_BODY()
 
 	/** The node's title */
 	UPROPERTY(VisibleAnywhere, BlueprintReadOnly, Category = Determinism)
-	FName NodeTitle = TEXT("Untitled");
+	FName TestResultTitle = TEXT("Untitled");
 
 	/** The node's name */
 	UPROPERTY(VisibleAnywhere, BlueprintReadOnly, Category = Determinism)
-	FString NodeName = TEXT("Unnamed");
+	FString TestResultName = TEXT("Unnamed");
 
 	UPROPERTY(EditAnywhere, BlueprintReadWrite, Category = Determinism)
 	int32 Seed = -1;
@@ -102,9 +102,9 @@ namespace PCGDeterminismTests
 	};
 
 	/** A default delegate to report an unset test */
-	bool LogInvalidTest(const UPCGNode* InPCGNode, const FName& TestName, FDeterminismNodeTestResult& OutResult);
+	bool LogInvalidTest(const UPCGNode* InPCGNode, const FName& TestName, FDeterminismTestResult& OutResult);
 
-	typedef TFunction<bool(const UPCGNode* InPCGNode, const FName& TestName, FDeterminismNodeTestResult& OutResult)> TestFunction;
+	typedef TFunction<bool(const UPCGNode* InPCGNode, const FName& TestName, FDeterminismTestResult& OutResult)> TestFunction;
 
 	struct FNodeTestInfo
 	{
@@ -134,18 +134,18 @@ namespace PCGDeterminismTests
 	};
 
 	/** Validates if a PCGNode is deterministic */
-	PCG_API void RunDeterminismTest(const UPCGNode* InPCGNode, FDeterminismNodeTestResult& OutResult, const FNodeTestInfo& TestToRun);
+	PCG_API void RunDeterminismTest(const UPCGNode* InPCGNode, FDeterminismTestResult& OutResult, const FNodeTestInfo& TestToRun);
 
 #if WITH_EDITOR
 	/** Validates if a PCGGraph is deterministic */
-	PCG_API void RunDeterminismTest(const UPCGGraph* InPCGGraph, UPCGComponent* InPCGComponent, FDeterminismNodeTestResult& OutResult);
+	PCG_API void RunDeterminismTest(const UPCGGraph* InPCGGraph, UPCGComponent* InPCGComponent, FDeterminismTestResult& OutResult);
 #endif
 
 	/** Validates all the generic determinism tests for any given node */
-	PCG_API bool RunBasicTestSuite(const UPCGNode* InPCGNode, const FName& TestName, FDeterminismNodeTestResult& OutResult);
+	PCG_API bool RunBasicTestSuite(const UPCGNode* InPCGNode, const FName& TestName, FDeterminismTestResult& OutResult);
 
 	/** Validates the various levels of order independence for any given node */
-	PCG_API bool RunOrderIndependenceSuite(const UPCGNode* InPCGNode, const FName& TestName, FDeterminismNodeTestResult& OutResult);
+	PCG_API bool RunOrderIndependenceSuite(const UPCGNode* InPCGNode, const FName& TestName, FDeterminismTestResult& OutResult);
 
 	/** Validates minimal node determinism against the same single test data */
 	bool RunBasicSelfTest(const FNodeAndOptions& NodeAndOptions);
@@ -237,7 +237,7 @@ namespace PCGDeterminismTests
 	/** A catch function for unimplemented comparisons */
 	bool ComparisonIsUnimplemented(const UPCGData* FirstData, const UPCGData* SecondData);
 	/** Updates the test result output to reflect that the permutation limit has been exceeded */
-	void UpdateTestResultForOverPermutationLimitError(FDeterminismNodeTestResult& OutResult);
+	void UpdateTestResultForOverPermutationLimitError(FDeterminismTestResult& OutResult);
 
 	/** Determines if a PCG DataType is a type that should be compared */
 	bool DataTypeIsComparable(EPCGDataType DataType);
@@ -268,7 +268,7 @@ namespace PCGDeterminismTests
 	EPCGDataType GetPermutation(int32 PermutationIteration, int32 PinIndex, const TArray<TArray<EPCGDataType>>& BaseOptionsPerPin);
 
 	/** Helper to update a tests' outgoing results */
-	void UpdateTestResults(FName TestName, FDeterminismNodeTestResult& OutResult, EDeterminismLevel DeterminismLevel);
+	void UpdateTestResults(FName TestName, FDeterminismTestResult& OutResult, EDeterminismLevel DeterminismLevel);
 
 	/** Gets a comparison function to compare two data objects */
 	TFunction<bool(const UPCGData*, const UPCGData*)> GetDataCompareFunction(EPCGDataType DataType, EDeterminismLevel DeterminismLevel);
diff --git a/Engine/Plugins/Experimental/PCG/Source/PCGEditor/Private/PCGEditor.cpp b/Engine/Plugins/Experimental/PCG/Source/PCGEditor/Private/PCGEditor.cpp
--- a/Engine/Plugins/Experimental/PCG/Source/PCGEditor/Private/PCGEditor.cpp
+++ b/Engine/Plugins/Experimental/PCG/Source/PCGEditor/Private/PCGEditor.cpp
@@ -473,9 +473,9 @@ void FPCGEditor::OnDeterminismNodeTest()
 				const UPCGNode* PCGNode = PCGEditorGraphNode->GetPCGNode();
 				check(PCGNode && PCGNode->DefaultSettings);
 
-				TSharedPtr<FDeterminismNodeTestResult> NodeResult = MakeShared<FDeterminismNodeTestResult>();
-				NodeResult->NodeTitle = PCGNode->GetNodeTitle();
-				NodeResult->NodeName = PCGNode->GetName();
+				TSharedPtr<FDeterminismTestResult> NodeResult = MakeShared<FDeterminismTestResult>();
+				NodeResult->TestResultTitle = PCGNode->GetNodeTitle();
+				NodeResult->TestResultName = PCGNode->GetName();
 				NodeResult->Seed = PCGNode->DefaultSettings->Seed;
 
 				if (PCGNode->DefaultSettings->DeterminismSettings.bNativeTests)
@@ -562,9 +562,9 @@ void FPCGEditor::OnDeterminismGraphTest()
 	FTestColumnInfo ColumnInfo({PCGDeterminismTests::Defaults::GraphResultName, NSLOCTEXT("PCGDeterminism", "Result", "Result"), 120.f, HAlign_Center});
 	DeterminismWidget->AddColumn(ColumnInfo);
 
-	TSharedPtr<FDeterminismNodeTestResult> TestResult = MakeShared<FDeterminismNodeTestResult>();
-	TestResult->NodeTitle = TEXT("Full Graph Test");
-	TestResult->NodeName = PCGGraphBeingEdited->GetName();
+	TSharedPtr<FDeterminismTestResult> TestResult = MakeShared<FDeterminismTestResult>();
+	TestResult->TestResultTitle = TEXT("Full Graph Test");
+	TestResult->TestResultName = PCGGraphBeingEdited->GetName();
 	TestResult->Seed = PCGComponentBeingDebugged->Seed;
 
 	PCGDeterminismTests::RunDeterminismTest(PCGGraphBeingEdited, PCGComponentBeingDebugged, *TestResult);
diff --git a/Engine/Plugins/Experimental/PCG/Source/PCGEditor/Private/SPCGEditorGraphDeterminism.cpp b/Engine/Plugins/Experimental/PCG/Source/PCGEditor/Private/SPCGEditorGraphDeterminism.cpp
--- a/Engine/Plugins/Experimental/PCG/Source/PCGEditor/Private/SPCGEditorGraphDeterminism.cpp
+++ b/Engine/Plugins/Experimental/PCG/Source/PCGEditor/Private/SPCGEditorGraphDeterminism.cpp
@@ -61,11 +61,11 @@ TSharedRef<SWidget> SPCGEditorGraphDeterminismRow::GenerateWidgetForColumn(const
 	}
 	else if (ColumnId == NAME_NodeTitle)
 	{
-		return ReturnColorCodedResultBlock(FText::FromName(CurrentItem->NodeTitle));
+		return ReturnColorCodedResultBlock(FText::FromName(CurrentItem->TestResultTitle));
 	}
 	else if (ColumnId == NAME_NodeName)
 	{
-		CellText = FText::FromString(CurrentItem->NodeName);
+		CellText = FText::FromString(CurrentItem->TestResultName);
 	}
 	else if (ColumnId == NAME_Seed)
 	{
diff --git a/Engine/Plugins/Experimental/PCG/Source/PCGEditor/Private/SPCGEditorGraphDeterminism.h b/Engine/Plugins/Experimental/PCG/Source/PCGEditor/Private/SPCGEditorGraphDeterminism.h
--- a/Engine/Plugins/Experimental/PCG/Source/PCGEditor/Private/SPCGEditorGraphDeterminism.h
+++ b/Engine/Plugins/Experimental/PCG/Source/PCGEditor/Private/SPCGEditorGraphDeterminism.h
@@ -7,7 +7,7 @@
 
 class FPCGEditor;
 
-typedef TSharedPtr<FDeterminismNodeTestResult> FPCGNodeTestResultPtr;
+typedef TSharedPtr<FDeterminismTestResult> FPCGNodeTestResultPtr;
 
 struct FTestColumnInfo
 {
