+ Add simple fallthrough determinism test for permutation fail

#rnx

diff --git a/Engine/Plugins/Experimental/PCG/Source/PCG/Private/Tests/Determinism/PCGDeterminismTestsCommon.cpp b/Engine/Plugins/Experimental/PCG/Source/PCG/Private/Tests/Determinism/PCGDeterminismTestsCommon.cpp
--- a/Engine/Plugins/Experimental/PCG/Source/PCG/Private/Tests/Determinism/PCGDeterminismTestsCommon.cpp
+++ b/Engine/Plugins/Experimental/PCG/Source/PCG/Private/Tests/Determinism/PCGDeterminismTestsCommon.cpp
@@ -222,6 +222,14 @@ namespace PCGDeterminismTests
 
 		if (GetNumPermutations(NodeAndOptions.BaseOptionsByPin) > CVarDeterminismPermutationLimit.GetValueOnGameThread())
 		{
+			EDeterminismLevel SimpleTestDeterminismLevel = GetHighestDeterminismLevelSimpleShuffleInput(NodeAndOptions);
+			// Downgrading the level here, because the Basic Test is intended to be Pass/Fail
+			if (SimpleTestDeterminismLevel != EDeterminismLevel::NoDeterminism)
+			{
+				SimpleTestDeterminismLevel = EDeterminismLevel::Basic;
+			}
+
+			UpdateTestResults(TestName, OutResult, SimpleTestDeterminismLevel);
 			UpdateTestResultForOverPermutationLimitError(OutResult);
 			return false;
 		}
@@ -244,6 +252,10 @@ namespace PCGDeterminismTests
 
 		if (GetNumPermutations(NodeAndOptions.BaseOptionsByPin) > CVarDeterminismPermutationLimit.GetValueOnGameThread())
 		{
+			EDeterminismLevel SimpleTestDeterminismLevel = GetHighestDeterminismLevelSimpleShuffleInput(NodeAndOptions);
+			// We can still test the levels of determinism with the simple test
+			UpdateTestResults(TestName, OutResult, SimpleTestDeterminismLevel);
+			// However, the result will have additional details highlighting that a full test wasn't conducted
 			UpdateTestResultForOverPermutationLimitError(OutResult);
 			return false;
 		}
@@ -406,6 +418,70 @@ namespace PCGDeterminismTests
 		return MaxLevel;
 	}
 
+	EDeterminismLevel GetHighestDeterminismLevelSimpleShuffleInput(const FNodeAndOptions& NodeAndOptions, int32 NumInputsPerPin, EDeterminismLevel MaxLevel)
+	{
+		const UPCGNode* PCGNode = NodeAndOptions.PCGNode;
+		const TArray<TArray<EPCGDataType>>& BaseOptionsByPin = NodeAndOptions.BaseOptionsByPin;
+
+		PCGTestsCommon::FTestData FirstTestData(NodeAndOptions.Seed, PCGNode->DefaultSettings);
+		PCGTestsCommon::FTestData SecondTestData(NodeAndOptions.Seed, PCGNode->DefaultSettings);
+
+		int32 NumInputs = NodeAndOptions.bMultipleOptionsPerPin ? NumInputsPerPin : 1;
+
+		// Add inputs for each pin. Unlike the full test, this will just shuffle the input order
+		for (int32 PinIndex = 0; PinIndex < BaseOptionsByPin.Num(); ++PinIndex)
+		{
+			for (int32 I = 0; I < NumInputs; ++I)
+			{
+				EPCGDataType DataType = BaseOptionsByPin[PinIndex][I];
+				AddRandomizedInputData(FirstTestData, DataType);
+				AddRandomizedInputData(SecondTestData, DataType);
+			}
+		}
+
+		ShuffleInputOrder(SecondTestData);
+
+		ExecuteWithTestData(FirstTestData, PCGNode);
+		ExecuteWithTestData(SecondTestData, PCGNode);
+
+		// Gauge the highest level
+		switch (MaxLevel)
+		{
+		case EDeterminismLevel::OrderIndependent:
+			if (DataCollectionsAreIdentical(FirstTestData.OutputData, SecondTestData.OutputData))
+			{
+				break;
+			}
+			MaxLevel = EDeterminismLevel::OrderConsistent;
+			// Falls through
+		case EDeterminismLevel::OrderConsistent:
+			if (DataCollectionsAreConsistent(FirstTestData.OutputData, SecondTestData.OutputData, NumInputs))
+			{
+				// With only 1 input, orthogonal == consistent
+				if (BaseOptionsByPin.Num() < 2)
+				{
+					if (DataCollectionsContainSameData(FirstTestData.OutputData, SecondTestData.OutputData))
+					{
+						break;
+					}
+				}
+				break;
+			}
+			MaxLevel = EDeterminismLevel::OrderOrthogonal;
+			// Falls through
+		case EDeterminismLevel::OrderOrthogonal:
+			if (DataCollectionsContainSameData(FirstTestData.OutputData, SecondTestData.OutputData))
+			{
+				break;
+			}
+			// Falls through
+		default:
+			MaxLevel = EDeterminismLevel::NoDeterminism;
+		}
+
+		return MaxLevel;
+	}
+
 	void AddRandomizedInputData(PCGTestsCommon::FTestData& TestData, EPCGDataType DataType, const FName& PinName)
 	{
 		switch (DataType)
@@ -1217,9 +1293,7 @@ namespace PCGDeterminismTests
 
 	void UpdateTestResultForOverPermutationLimitError(FDeterminismTestResult& OutResult)
 	{
-		OutResult.DataTypesTested = EPCGDataType::None;
-		OutResult.bFlagRaised = true;
-		OutResult.AdditionalDetails.Emplace(TEXT("Test did not run (permutation limit). Adjust 'pcg.DeterminismPermutationLimit' to alter this limit."));
+		OutResult.AdditionalDetails.Emplace(TEXT("A full test was not conducted, due to a high number of permutations. Adjust 'pcg.DeterminismPermutationLimit' to alter this limit."));
 	}
 
 	bool ConsistencyComparisonIsUnimplemented(const UPCGData* FirstData, const UPCGData* SecondData, TArray<int32>& OutOutIndexOffsets)
diff --git a/Engine/Plugins/Experimental/PCG/Source/PCG/Public/Tests/Determinism/PCGDeterminismTestsCommon.h b/Engine/Plugins/Experimental/PCG/Source/PCG/Public/Tests/Determinism/PCGDeterminismTestsCommon.h
--- a/Engine/Plugins/Experimental/PCG/Source/PCG/Public/Tests/Determinism/PCGDeterminismTestsCommon.h
+++ b/Engine/Plugins/Experimental/PCG/Source/PCG/Public/Tests/Determinism/PCGDeterminismTestsCommon.h
@@ -156,6 +156,10 @@ namespace PCGDeterminismTests
 	EDeterminismLevel GetHighestDeterminismLevel(const FNodeAndOptions& NodeAndOptions,
 		int32 NumInputsPerPin = Defaults::NumTestInputsPerPin,
 		EDeterminismLevel MaxLevel = EDeterminismLevel::OrderIndependent);
+	/** Conducts simplified shuffled input tests and determines the highest level of determinism */
+	EDeterminismLevel GetHighestDeterminismLevelSimpleShuffleInput(const FNodeAndOptions& NodeAndOptions,
+		int32 NumInputsPerPin = Defaults::NumTestInputsPerPin,
+		EDeterminismLevel MaxLevel = EDeterminismLevel::OrderIndependent);
 
 	/** Adds input data to test data, based on input pins' allowed types */
 	void AddRandomizedInputData(PCGTestsCommon::FTestData& TestData, EPCGDataType DataType, const FName& PinName = Defaults::TestPinName);
