[PCG] Safeguard usage of PCGSubsystem.

diff --git a/Engine/Plugins/Experimental/PCG/Source/PCG/Private/Grid/PCGPartitionActor.cpp b/Engine/Plugins/Experimental/PCG/Source/PCG/Private/Grid/PCGPartitionActor.cpp
--- a/Engine/Plugins/Experimental/PCG/Source/PCG/Private/Grid/PCGPartitionActor.cpp
+++ b/Engine/Plugins/Experimental/PCG/Source/PCG/Private/Grid/PCGPartitionActor.cpp
@@ -106,9 +106,10 @@ void APCGPartitionActor::PostRegisterAllComponents()
 #endif // WITH_EDITOR
 		) && ensure(PCGGridSize > 0))
 	{
-		UPCGSubsystem* Subsystem = GetSubsystem();
-		check(Subsystem);
-		Subsystem->RegisterPartitionActor(this, /*bDoComponentMapping*/ bIsRuntimeOrPIE);
+		if (UPCGSubsystem* Subsystem = GetSubsystem())
+		{
+			Subsystem->RegisterPartitionActor(this, /*bDoComponentMapping*/ bIsRuntimeOrPIE);
+		}
 	}
 }
 
@@ -137,10 +138,10 @@ void APCGPartitionActor::BeginPlay()
 
 void APCGPartitionActor::EndPlay(const EEndPlayReason::Type EndPlayReason)
 {
-	UPCGSubsystem* Subsystem = GetSubsystem();
-	check(Subsystem);
-
-	Subsystem->UnregisterPartitionActor(this);
+	if (UPCGSubsystem* Subsystem = GetSubsystem())
+	{
+		Subsystem->UnregisterPartitionActor(this);
+	}
 
 	Super::EndPlay(EndPlayReason);
 }
diff --git a/Engine/Plugins/Experimental/PCG/Source/PCG/Private/PCGComponent.cpp b/Engine/Plugins/Experimental/PCG/Source/PCG/Private/PCGComponent.cpp
--- a/Engine/Plugins/Experimental/PCG/Source/PCG/Private/PCGComponent.cpp
+++ b/Engine/Plugins/Experimental/PCG/Source/PCG/Private/PCGComponent.cpp
@@ -258,11 +258,11 @@ FPCGTaskId UPCGComponent::CreateGenerateTask(bool bForce, const TArray<FPCGTaskI
 	}
 
 #if WITH_EDITOR
-	if (bForce && bGenerated && !bDirtyGenerated)
-	{
-		// TODO: generate new seed
-		++Seed;
-	}
+	// TODO: Have a better way to know when we need to generate a new seed.
+	//if (bForce && bGenerated && !bDirtyGenerated)
+	//{
+	//	++Seed;
+	//}
 #endif
 
 	// Keep track of all the dependencies
@@ -681,9 +681,8 @@ void UPCGComponent::OnComponentCreated()
 #if WITH_EDITOR
 	SetupActorCallbacks();
 
-	if (!PCGHelpers::IsRuntimeOrPIE() && IsPartitioned())
+	if (!PCGHelpers::IsRuntimeOrPIE() && IsPartitioned() && GetSubsystem())
 	{
-		check(GetSubsystem());
 		GetSubsystem()->RegisterOrUpdatePCGComponent(this, /*bDoActorMapping=*/bGenerated);
 	}
 #endif
@@ -774,9 +773,8 @@ void UPCGComponent::SetupCallbacksOnCreation()
 		Graph->OnGraphChangedDelegate.AddUObject(this, &UPCGComponent::OnGraphChanged);
 	}
 
-	if (!PCGHelpers::IsRuntimeOrPIE() && IsPartitioned())
+	if (!PCGHelpers::IsRuntimeOrPIE() && IsPartitioned() && GetSubsystem())
 	{
-		check(GetSubsystem());
 		GetSubsystem()->RegisterOrUpdatePCGComponent(this, /*bDoActorMapping=*/bGenerated);
 	}
 }
@@ -785,9 +783,8 @@ void UPCGComponent::SetupCallbacksOnCreation()
 void UPCGComponent::BeginDestroy()
 {
 #if WITH_EDITOR
-	if (!PCGHelpers::IsRuntimeOrPIE() && IsPartitioned())
+	if (!PCGHelpers::IsRuntimeOrPIE() && IsPartitioned() && GetSubsystem())
 	{
-		check(GetSubsystem());
 		GetSubsystem()->UnregisterPCGComponent(this);
 	}
 
@@ -923,11 +920,6 @@ void UPCGComponent::PostEditChangeProperty(FPropertyChangedEvent& PropertyChange
 				GenerateLocal(/*bForce=*/false);
 			}
 		}
-		else
-		{
-			// Just ignore the change
-			bIsPartitioned = false;
-		}
 	}
 	else if (PropName == GET_MEMBER_NAME_CHECKED(UPCGComponent, Graph))
 	{
@@ -1463,7 +1455,7 @@ void UPCGComponent::Refresh()
 
 	// Before doing a refresh, update the component to the subsystem if we are partitioned
 	// Only redo the mapping if we are generated
-	if (IsPartitioned())
+	if (IsPartitioned() && GetSubsystem())
 	{
 		GetSubsystem()->RegisterOrUpdatePCGComponent(this, /*bDoActorMapping=*/ bGenerated);
 	}
@@ -2193,15 +2185,18 @@ void FPCGComponentInstanceData::ApplyToComponent(UActorComponent* Component, con
 
 		// We need to make sure that the component is registered to the subsystem if it is partitioned
 		// or unregister it if it is not.
-		if (PCGComponent->IsPartitioned())
+		if (UPCGSubsystem* Subsystem = PCGComponent->GetSubsystem())
 		{
-			// Mapping needs to be done if we are at runtime or previously generated
-			bool bDoActorMapping = PCGComponent->bGenerated || PCGHelpers::IsRuntimeOrPIE();
-			PCGComponent->GetSubsystem()->RegisterOrUpdatePCGComponent(PCGComponent, bDoActorMapping);
-		}
-		else
-		{
-			PCGComponent->GetSubsystem()->UnregisterPCGComponent(PCGComponent);
+			if (PCGComponent->IsPartitioned())
+			{
+				// Mapping needs to be done if we are at runtime or previously generated
+				bool bDoActorMapping = PCGComponent->bGenerated || PCGHelpers::IsRuntimeOrPIE();
+				Subsystem->RegisterOrUpdatePCGComponent(PCGComponent, bDoActorMapping);
+			}
+			else
+			{
+				Subsystem->UnregisterPCGComponent(PCGComponent);
+			}
 		}
 	}
 }
