Add BP interface to manage resources

#rnx

diff --git a/Engine/Plugins/Experimental/PCG/Source/PCG/Public/PCGComponent.h b/Engine/Plugins/Experimental/PCG/Source/PCG/Public/PCGComponent.h
--- a/Engine/Plugins/Experimental/PCG/Source/PCG/Public/PCGComponent.h
+++ b/Engine/Plugins/Experimental/PCG/Source/PCG/Public/PCGComponent.h
@@ -99,7 +99,10 @@ public:
 	UFUNCTION(BlueprintCallable, NetMulticast, Reliable, Category = PCG)
 	void SetGraph(UPCGGraph* InGraph);
 
+	/** Registers some managed resource to the current component */
+	UFUNCTION(BlueprintCallable, Category = PCG)
 	void AddToManagedResources(UPCGManagedResource* InResource);
+
 	void ForEachManagedResource(TFunctionRef<void(UPCGManagedResource*)> InFunction);
 
 	/** Transactionable methods to be called from details UI */
