Added normal/tangent normalization by default.
We can revisit based on different use cases here (might need unnormalized tangentX/Y)

#rnx

diff --git a/Engine/Plugins/Experimental/PCG/Source/PCG/Private/Grid/PCGLandscapeCache.cpp b/Engine/Plugins/Experimental/PCG/Source/PCG/Private/Grid/PCGLandscapeCache.cpp
--- a/Engine/Plugins/Experimental/PCG/Source/PCG/Private/Grid/PCGLandscapeCache.cpp
+++ b/Engine/Plugins/Experimental/PCG/Source/PCG/Private/Grid/PCGLandscapeCache.cpp
@@ -95,7 +95,7 @@ void FPCGLandscapeCacheEntry::GetPoint(int32 PointIndex, FPCGPoint& OutPoint, UP
 	TangentX = FVector(Normal.Z, 0.f, -Normal.X);
 	TangentY = Normal ^ TangentX;
 
-	OutPoint.Transform = FTransform(TangentX, TangentY, Normal, Position);
+	OutPoint.Transform = FTransform(TangentX.GetSafeNormal(), TangentY.GetSafeNormal(), Normal.GetSafeNormal(), Position);
 	OutPoint.BoundsMin = -PointHalfSize;
 	OutPoint.BoundsMax = PointHalfSize;
 	OutPoint.Seed = 1 + PointIndex;
@@ -153,19 +153,16 @@ void FPCGLandscapeCacheEntry::GetInterpolatedPoint(const FVector2D& LocalPoint,
 	const FVector LerpPositionY1 = FMath::Lerp(PositionX0Y1, PositionX1Y1, XFactor);
 	const FVector Position = FMath::Lerp(LerpPositionY0, LerpPositionY1, YFactor);
 
-	const FVector LerpNormalY0 = FMath::Lerp(NormalX0Y0, NormalX1Y0, XFactor);
-	const FVector LerpNormalY1 = FMath::Lerp(NormalX0Y1, NormalX1Y1, XFactor);
-	FVector Normal = FMath::Lerp(LerpNormalY0, LerpNormalY1, YFactor);
-
-	// TODO: we could preserve normal length by lerping this too?
-	Normal.Normalize();
+	const FVector LerpNormalY0 = FMath::Lerp(NormalX0Y0.GetSafeNormal(), NormalX1Y0.GetSafeNormal(), XFactor).GetSafeNormal();
+	const FVector LerpNormalY1 = FMath::Lerp(NormalX0Y1.GetSafeNormal(), NormalX1Y1.GetSafeNormal(), XFactor).GetSafeNormal();
+	const FVector Normal = FMath::Lerp(LerpNormalY0, LerpNormalY1, YFactor);
 
 	FVector TangentX;
 	FVector TangentY;
 	TangentX = FVector(Normal.Z, 0.f, -Normal.X);
 	TangentY = Normal ^ TangentX;
 
-	OutPoint.Transform = FTransform(TangentX, TangentY, Normal, Position);
+	OutPoint.Transform = FTransform(TangentX.GetSafeNormal(), TangentY.GetSafeNormal(), Normal.GetSafeNormal(), Position);
 	OutPoint.BoundsMin = -PointHalfSize;
 	OutPoint.BoundsMax = PointHalfSize;
 	OutPoint.Seed = UPCGBlueprintHelpers::ComputeSeedFromPosition(Position);
