Improve robustness of texture data, though we need to expand on it further.
Fixed issue with FColor parsing, was wrongfully done.

diff --git a/Engine/Plugins/Experimental/PCG/Source/PCG/Private/Data/PCGTextureData.cpp b/Engine/Plugins/Experimental/PCG/Source/PCG/Private/Data/PCGTextureData.cpp
--- a/Engine/Plugins/Experimental/PCG/Source/PCG/Private/Data/PCGTextureData.cpp
+++ b/Engine/Plugins/Experimental/PCG/Source/PCG/Private/Data/PCGTextureData.cpp
@@ -215,47 +215,61 @@ void UPCGTextureData::Initialize(UTexture2D* InTexture, const FTransform& InTran
 
 	if (Texture)
 	{
-#if WITH_EDITORONLY_DATA
-		if (Texture->GetPlatformData()->Mips.Num() > 0)
+		FTexturePlatformData* PlatformData = Texture->GetPlatformData();
+		if (PlatformData && PlatformData->Mips.Num() > 0)
 		{
 			TRACE_CPUPROFILER_EVENT_SCOPE(UPCGTextureData::Initialize::ReadData);
 
-			if (Texture->GetPlatformData()->PixelFormat == PF_R8G8B8A8 || Texture->GetPlatformData()->PixelFormat == PF_G8)
+			if (PlatformData->PixelFormat == PF_B8G8R8A8 ||
+				PlatformData->PixelFormat == PF_R8G8B8A8 ||
+				PlatformData->PixelFormat == PF_G8)
 			{
-				Width = Texture->GetSizeX();
-				Height = Texture->GetSizeY();
-				const int32 PixelCount = Width * Height;
-				ColorData.SetNum(PixelCount);
-
-				// TODO: previously this code created a duplicate of InTexture and read the ColorData from that source 
-				// however, there was a problem with the duplicate not having texture data in Mips[0], so we've removed the duplication for now
-				// original code can be found in ProceduralTextureReader.h
-				const uint8_t* BulkData = reinterpret_cast<const uint8_t*>(InTexture->GetPlatformData()->Mips[0].BulkData.LockReadOnly());
-
-				if (Texture->GetPlatformData()->PixelFormat == PF_R8G8B8A8)
+				if (const uint8_t* BulkData = reinterpret_cast<const uint8_t*>(PlatformData->Mips[0].BulkData.LockReadOnly()))
 				{
-					const FColor* FormattedImageData = reinterpret_cast<const FColor*>(BulkData);
-					for (int32 D = 0; D < PixelCount; ++D)
+					Width = Texture->GetSizeX();
+					Height = Texture->GetSizeY();
+					const int32 PixelCount = Width * Height;
+					ColorData.SetNum(PixelCount);
+
+					if (PlatformData->PixelFormat == PF_B8G8R8A8)
 					{
-						ColorData[D] = FormattedImageData[D].ReinterpretAsLinear();
+						// Memory representation of FColor is BGRA
+						const FColor* FormattedImageData = reinterpret_cast<const FColor*>(BulkData);
+						for (int32 D = 0; D < PixelCount; ++D)
+						{
+							ColorData[D] = FormattedImageData[D].ReinterpretAsLinear();
+						}
 					}
-				}
-				else if (Texture->GetPlatformData()->PixelFormat == PF_G8)
-				{
-					for (int32 D = 0; D < PixelCount; ++D)
+					else if (PlatformData->PixelFormat == PF_R8G8B8A8)
 					{
-						ColorData[D] = FColor(BulkData[D], BulkData[D], BulkData[D]).ReinterpretAsLinear();
+						// Since the memory representation is BGRA, we just need to swap B & R to obtain RGBA 
+						const FColor* FormattedImageData = reinterpret_cast<const FColor*>(BulkData);
+						for (int32 D = 0; D < PixelCount; ++D)
+						{
+							ColorData[D] = FormattedImageData[D].ReinterpretAsLinear();
+							Swap(ColorData[D].R, ColorData[D].B);
+						}
+					}
+					else if (PlatformData->PixelFormat == PF_G8)
+					{
+						for (int32 D = 0; D < PixelCount; ++D)
+						{
+							ColorData[D] = FColor(BulkData[D], BulkData[D], BulkData[D]).ReinterpretAsLinear();
+						}
 					}
 				}
-
-				InTexture->GetPlatformData()->Mips[0].BulkData.Unlock();
+				else
+				{
+					UE_LOG(LogPCG, Error, TEXT("PCGTextureData unable to get bulk data from %s"), *Texture->GetFName().ToString());
+				}
+				
+				PlatformData->Mips[0].BulkData.Unlock();
 			}
 			else
 			{
 				UE_LOG(LogPCG, Error, TEXT("PCGTextureData does not support the format of %s"), *Texture->GetFName().ToString());
 			}
 		}
-#endif
 	}
 
 	Bounds = FBox(EForceInit::ForceInit);
diff --git a/Engine/Plugins/Experimental/PCG/Source/PCG/Private/Tests/Data/PCGTextureDataTest.cpp b/Engine/Plugins/Experimental/PCG/Source/PCG/Private/Tests/Data/PCGTextureDataTest.cpp
--- a/Engine/Plugins/Experimental/PCG/Source/PCG/Private/Tests/Data/PCGTextureDataTest.cpp
+++ b/Engine/Plugins/Experimental/PCG/Source/PCG/Private/Tests/Data/PCGTextureDataTest.cpp
@@ -22,7 +22,7 @@ bool FPCGTextureDataOffsetTilingRotation::RunTest(const FString& Parameters)
 	Pixels.Init(FColor::Black, TextureSize * TextureSize);
 	Pixels[WhitePixelY * TextureSize + WhitePixelX] = FColor::White;
 
-	UTexture2D* Texture2D = UTexture2D::CreateTransient(TextureSize, TextureSize, EPixelFormat::PF_R8G8B8A8);
+	UTexture2D* Texture2D = UTexture2D::CreateTransient(TextureSize, TextureSize, EPixelFormat::PF_B8G8R8A8);
 	Texture2D->CompressionSettings = TextureCompressionSettings::TC_VectorDisplacementmap;
 	Texture2D->SRGB = 0;
 	Texture2D->MipGenSettings = TMGS_NoMipmaps;
