Fixed a crash in FPCGGraphExecutor::Execute resulting from using a reference to an element of ReadyTasks after the array's reallocation.

diff --git a/Engine/Plugins/Experimental/PCG/Source/PCG/Private/Graph/PCGGraphExecutor.cpp b/Engine/Plugins/Experimental/PCG/Source/PCG/Private/Graph/PCGGraphExecutor.cpp
--- a/Engine/Plugins/Experimental/PCG/Source/PCG/Private/Graph/PCGGraphExecutor.cpp
+++ b/Engine/Plugins/Experimental/PCG/Source/PCG/Private/Graph/PCGGraphExecutor.cpp
@@ -266,20 +266,20 @@ void FPCGGraphExecutor::Execute()
 
 				if (!bNeedsToCreateActiveTask)
 				{
-					// Fast-forward cached result to stored results
-					FPCGTaskId SkippedTaskId = Task.NodeId;
-					StoreResults(SkippedTaskId, CachedOutput);
-					ReadyTasks.RemoveAtSwap(ReadyTaskIndex);
-					QueueNextTasks(SkippedTaskId);
-					bAnyTaskEnded = true;
-
 #if WITH_EDITOR
+					// doing this now since we're about to modify ReadyTasks potentially reallocating while Task is a reference. 
 					UPCGComponent* SourceComponent = Task.SourceComponent;
 					if (SourceComponent && SourceComponent->IsInspecting())
 					{
 						SourceComponent->StoreInspectionData(Task.Node, CachedOutput);
 					}
 #endif
+					// Fast-forward cached result to stored results
+					FPCGTaskId SkippedTaskId = Task.NodeId;
+					StoreResults(SkippedTaskId, CachedOutput);
+					ReadyTasks.RemoveAtSwap(ReadyTaskIndex);
+					QueueNextTasks(SkippedTaskId);
+					bAnyTaskEnded = true;
 
 					continue;
 				}
