Made the landscape spline a bit more consistent in computing the density - added tolerance to barycentric check.

#rnx

diff --git a/Engine/Plugins/Experimental/PCG/Source/PCG/Private/Data/PCGLandscapeSplineData.cpp b/Engine/Plugins/Experimental/PCG/Source/PCG/Private/Data/PCGLandscapeSplineData.cpp
--- a/Engine/Plugins/Experimental/PCG/Source/PCG/Private/Data/PCGLandscapeSplineData.cpp
+++ b/Engine/Plugins/Experimental/PCG/Source/PCG/Private/Data/PCGLandscapeSplineData.cpp
@@ -14,28 +14,23 @@
 namespace PCGLandscapeDataHelpers
 {
 	// This function assumes that the A-B segment has a "1" density, while the C-D segment has a "0" density
-	FVector::FReal GetDensityInQuad(const FVector& InA, const FVector& InB, const FVector& InC, const FVector& InD, const FVector& InP)
+	FVector::FReal GetDensityInQuad(const FVector& A, const FVector& B, const FVector& C, const FVector& D, const FVector& P)
 	{
 		// Since the landscape has a strict Z behavior and our points might not be directly on the plane, we should only consider
 		// the 2D plane. When we support other axes, we could just remove the normal component off the position to correct.
-		const FVector2D A(InA.X, InA.Y);
-		const FVector2D B(InB.X, InB.Y);
-		const FVector2D C(InC.X, InC.Y);
-		const FVector2D D(InD.X, InD.Y);
-		const FVector2D P(InP.X, InP.Y);
-
+		const FVector::FReal Tolerance = UE_SMALL_NUMBER;
 		FVector BaryABC = FMath::GetBaryCentric2D(P, A, B, C);
 
-		if (BaryABC.X >= 0 && BaryABC.Y >= 0 && BaryABC.Z >= 0)
+		if (BaryABC.X >= -Tolerance && BaryABC.Y >= -Tolerance && BaryABC.Z >= -Tolerance)
 		{
-			return 1.0f - BaryABC.Z;
+			return 1.0f - FMath::Max(BaryABC.Z, 0);
 		}
 
 		FVector BaryACD = FMath::GetBaryCentric2D(P, A, C, D);
 
-		if (BaryACD.X >= 0 && BaryACD.Y >= 0 && BaryACD.Z >= 0)
+		if (BaryACD.X >= -Tolerance && BaryACD.Y >= -Tolerance && BaryACD.Z >= -Tolerance)
 		{
-			return BaryACD.X;
+			return FMath::Max(BaryACD.X, 0);
 		}
 
 		return FVector::FReal(-1);
