PCG:
*Fixed an issue were the PCGEditor debug component dropdown would pick up invalid garbage objects.

diff --git a/Engine/Plugins/Experimental/PCG/Source/PCGEditor/Private/SPCGEditorGraphDebugObjectWidget.cpp b/Engine/Plugins/Experimental/PCG/Source/PCGEditor/Private/SPCGEditorGraphDebugObjectWidget.cpp
--- a/Engine/Plugins/Experimental/PCG/Source/PCGEditor/Private/SPCGEditorGraphDebugObjectWidget.cpp
+++ b/Engine/Plugins/Experimental/PCG/Source/PCGEditor/Private/SPCGEditorGraphDebugObjectWidget.cpp
@@ -93,6 +93,11 @@ void SPCGEditorGraphDebugObjectWidget::OnComboBoxOpening()
 	GetObjectsOfClass(UPCGComponent::StaticClass(), PCGComponents, true);
 	for (UObject* PCGComponentObject : PCGComponents)
 	{
+		if (!IsValid(PCGComponentObject))
+		{
+			continue;
+		}
+		
 		UPCGComponent* PCGComponent = Cast<UPCGComponent>(PCGComponentObject);
 		if (!PCGComponent)
 		{
