Removed engine version from PCG plugin definition file

diff --git a/Engine/Plugins/Experimental/PCG/PCG.uplugin b/Engine/Plugins/Experimental/PCG/PCG.uplugin
--- a/Engine/Plugins/Experimental/PCG/PCG.uplugin
+++ b/Engine/Plugins/Experimental/PCG/PCG.uplugin
@@ -2,7 +2,6 @@
 	"FileVersion": 3,
 	"Version": 8,
 	"VersionName": "0.1",
-	"EngineVersion": "5.1.0",
 	"FriendlyName": "PCG",
 	"Description": "Procedural generation framework for UE",
 	"Category": "Editor",
