Moved landscape cache to bulk data, loaded on demand.
Added some CVars to control frequency and threshold of cleanup.

diff --git a/Engine/Plugins/Experimental/PCG/Source/PCG/Private/Grid/PCGLandscapeCache.cpp b/Engine/Plugins/Experimental/PCG/Source/PCG/Private/Grid/PCGLandscapeCache.cpp
--- a/Engine/Plugins/Experimental/PCG/Source/PCG/Private/Grid/PCGLandscapeCache.cpp
+++ b/Engine/Plugins/Experimental/PCG/Source/PCG/Private/Grid/PCGLandscapeCache.cpp
@@ -20,12 +20,24 @@
 #include "LandscapeDataAccess.h"
 #include "Landscape.h"
 #include "Kismet/GameplayStatics.h"
+#include "Serialization/BufferReader.h"
+#include "Serialization/BufferWriter.h"
+
+static TAutoConsoleVariable<float> CVarLandscapeCacheGCFrequency(
+	TEXT("pcg.LandscapeCacheGCFrequency"),
+	10,
+	TEXT("Rate at which to release landscape cache data, in seconds"));
+
+static TAutoConsoleVariable<int32> CVarLandscapeCacheSizeThreshold(
+	TEXT("pcg.LandscapeCacheSizeThreshold"),
+	64,
+	TEXT("Memory Threhold at which we start cleaning up the landscape cache"));
 
 #if WITH_EDITOR
-void FPCGLandscapeCacheEntry::BuildCacheData(ULandscapeInfo* LandscapeInfo, ULandscapeComponent* InComponent, UObject* Owner)
+void FPCGLandscapeCacheEntry::BuildCacheData(ULandscapeInfo* LandscapeInfo, ULandscapeComponent* InComponent)
 {
 	TRACE_CPUPROFILER_EVENT_SCOPE(FPCGLandscapeCacheEntry::BuildCacheData);
-	check(Owner && !Component.Get() && InComponent && PositionsAndNormals.Num() == 0);
+	check(!Component.Get() && InComponent && PositionsAndNormals.Num() == 0 && !bDataLoaded);
 	Component = InComponent;
 
 	ALandscape* Landscape = Component->GetLandscapeActor();
@@ -72,20 +84,21 @@ void FPCGLandscapeCacheEntry::BuildCacheData(ULandscapeInfo* LandscapeInfo, ULan
 		FLandscapeComponentDataInterface CDI(InComponent, 0, /*bUseEditingLayer=*/true);
 		if (CDI.GetWeightmapTextureData(LayerInfo, LayerCache, /*bUseEditingLayer=*/true))
 		{
-			FPCGLandscapeCacheLayer& PCGLayer = LayerData.Emplace_GetRef();
-			PCGLayer.Name = Layer.LayerName;
-			PCGLayer.Data = LayerCache;
+			LayerDataNames.Add(Layer.LayerName);
+			LayerData.Add(LayerCache);
 		}
 
 		LayerCache.Reset();
 	}
+	
+	bDataLoaded = true;
 }
 #endif // WITH_EDITOR
 
 void FPCGLandscapeCacheEntry::GetPoint(int32 PointIndex, FPCGPoint& OutPoint, UPCGMetadata* OutMetadata) const
 {
 	//TRACE_CPUPROFILER_EVENT_SCOPE(FPCGLandscapeCacheEntry::GetPoint);
-	check(PointIndex >= 0 && 2 * PointIndex < PositionsAndNormals.Num());
+	check(bDataLoaded && PointIndex >= 0 && 2 * PointIndex < PositionsAndNormals.Num());	
 
 	const FVector& Position = PositionsAndNormals[2 * PointIndex];
 	const FVector& Normal = PositionsAndNormals[2 * PointIndex + 1];
@@ -102,14 +115,18 @@ void FPCGLandscapeCacheEntry::GetPoint(int32 PointIndex, FPCGPoint& OutPoint, UP
 	
 	if (OutMetadata && !LayerData.IsEmpty())
 	{
+		check(LayerData.Num() == LayerDataNames.Num());
 		OutPoint.MetadataEntry = OutMetadata->AddEntry();
 
-		for (const FPCGLandscapeCacheLayer& Layer : LayerData)
+		for(int32 LayerIndex = 0; LayerIndex < LayerData.Num(); ++LayerIndex)
 		{
-			if (FPCGMetadataAttributeBase* Attribute = OutMetadata->GetMutableAttribute(Layer.Name))
+			const FName& CurrentLayerName = LayerDataNames[LayerIndex];
+			const TArray<uint8>& CurrentLayerData = LayerData[LayerIndex];
+
+			if (FPCGMetadataAttributeBase* Attribute = OutMetadata->GetMutableAttribute(CurrentLayerName))
 			{
 				check(Attribute->GetTypeId() == PCG::Private::MetadataTypes<float>::Id);
-				static_cast<FPCGMetadataAttribute<float>*>(Attribute)->SetValue(OutPoint.MetadataEntry, (float)Layer.Data[PointIndex] / 255.0f);
+				static_cast<FPCGMetadataAttribute<float>*>(Attribute)->SetValue(OutPoint.MetadataEntry, (float)CurrentLayerData[PointIndex] / 255.0f);
 			}
 		}
 	}
@@ -118,7 +135,7 @@ void FPCGLandscapeCacheEntry::GetPoint(int32 PointIndex, FPCGPoint& OutPoint, UP
 void FPCGLandscapeCacheEntry::GetPointHeightOnly(int32 PointIndex, FPCGPoint& OutPoint) const
 {
 	//TRACE_CPUPROFILER_EVENT_SCOPE(FPCGLandscapeCacheEntry::GetPointHeightOnly);
-	check(PointIndex >= 0 && 2 * PointIndex < PositionsAndNormals.Num());
+	check(bDataLoaded && PointIndex >= 0 && 2 * PointIndex < PositionsAndNormals.Num());
 
 	const FVector& Position = PositionsAndNormals[2 * PointIndex];
 	
@@ -130,6 +147,8 @@ void FPCGLandscapeCacheEntry::GetPointHeightOnly(int32 PointIndex, FPCGPoint& Ou
 
 void FPCGLandscapeCacheEntry::GetInterpolatedPoint(const FVector2D& LocalPoint, FPCGPoint& OutPoint, UPCGMetadata* OutMetadata) const
 {
+	check(bDataLoaded);
+
 	const int32 X0Y0 = FMath::FloorToInt(LocalPoint.X) + FMath::FloorToInt(LocalPoint.Y) * Stride;
 	const int32 X1Y0 = X0Y0 + 1;
 	const int32 X0Y1 = X0Y0 + Stride;
@@ -171,13 +190,18 @@ void FPCGLandscapeCacheEntry::GetInterpolatedPoint(const FVector2D& LocalPoint,
 	{
 		OutPoint.MetadataEntry = OutMetadata->AddEntry();
 
-		for (const FPCGLandscapeCacheLayer& Layer : LayerData)
+		check(LayerData.Num() == LayerDataNames.Num());
+
+		for(int32 LayerIndex = 0; LayerIndex < LayerData.Num(); ++LayerIndex)
 		{
-			if (FPCGMetadataAttributeBase* Attribute = OutMetadata->GetMutableAttribute(Layer.Name))
+			const FName& CurrentLayerName = LayerDataNames[LayerIndex];
+			const TArray<uint8>& CurrentLayerData = LayerData[LayerIndex];
+
+			if (FPCGMetadataAttributeBase* Attribute = OutMetadata->GetMutableAttribute(CurrentLayerName))
 			{
 				check(Attribute->GetTypeId() == PCG::Private::MetadataTypes<float>::Id);
-				float Y0Data = FMath::Lerp((float)Layer.Data[X0Y0] / 255.0f, (float)Layer.Data[X1Y0] / 255.0f, XFactor);
-				float Y1Data = FMath::Lerp((float)Layer.Data[X0Y1] / 255.0f, (float)Layer.Data[X1Y1] / 255.0f, XFactor);
+				float Y0Data = FMath::Lerp((float)CurrentLayerData[X0Y0] / 255.0f, (float)CurrentLayerData[X1Y0] / 255.0f, XFactor);
+				float Y1Data = FMath::Lerp((float)CurrentLayerData[X0Y1] / 255.0f, (float)CurrentLayerData[X1Y1] / 255.0f, XFactor);
 				float Data = FMath::Lerp(Y0Data, Y1Data, YFactor);
 
 				static_cast<FPCGMetadataAttribute<float>*>(Attribute)->SetValue(OutPoint.MetadataEntry, Data);
@@ -188,6 +212,8 @@ void FPCGLandscapeCacheEntry::GetInterpolatedPoint(const FVector2D& LocalPoint,
 
 void FPCGLandscapeCacheEntry::GetInterpolatedPointHeightOnly(const FVector2D& LocalPoint, FPCGPoint& OutPoint) const
 {
+	check(bDataLoaded);
+
 	const int32 X0Y0 = FMath::FloorToInt(LocalPoint.X) + FMath::FloorToInt(LocalPoint.Y) * Stride;
 	const int32 X1Y0 = X0Y0 + 1;
 	const int32 X0Y1 = X0Y0 + Stride;
@@ -213,53 +239,248 @@ void FPCGLandscapeCacheEntry::GetInterpolatedPointHeightOnly(const FVector2D& Lo
 	OutPoint.Seed = UPCGBlueprintHelpers::ComputeSeedFromPosition(Position);
 }
 
-FPCGLandscapeCache::FPCGLandscapeCache(UObject* InOwner)
-	: Owner(InOwner)
+bool FPCGLandscapeCacheEntry::TouchAndLoad(int32 InTouch) const
 {
-	check(Owner);
+	Touch = InTouch; // technically, this could be an atomic, but we don't need this to be very precise
+
 #if WITH_EDITOR
-	SetupLandscapeCallbacks();
-	CacheLayerNames();
+	check(bDataLoaded);
+	return false; // In editor, we're "always" loaded since it's created on the fly
+#else
+	if (!bDataLoaded)
+	{
+		FScopeLock ScopeDataLock(&DataLock);
+		if (!bDataLoaded)
+		{
+			SerializeFromBulkData();
+			return true;
+		}
+	}
+
+	return false;
 #endif
 }
 
-void FPCGLandscapeCache::SetOwner(UObject* InOwner, bool bUpdateCachedNames)
+void FPCGLandscapeCacheEntry::Unload()
 {
-	check(!Owner || Owner == InOwner || InOwner == nullptr);
-	Owner = InOwner;
-#if WITH_EDITOR
-	if (InOwner)
+	check(bDataLoaded);
+	PositionsAndNormals.Reset();
+	PositionsAndNormals.Shrink();
+	LayerData.Reset();
+	LayerData.Shrink();
+	bDataLoaded = false;
+	Touch = 0;
+}
+
+int32 FPCGLandscapeCacheEntry::GetMemorySize() const
+{
+	check(bDataLoaded);
+	int32 MemSize = PositionsAndNormals.GetAllocatedSize();
+	for (const TArray<uint8>& Layer : LayerData)
 	{
-		SetupLandscapeCallbacks();
+		MemSize += Layer.GetAllocatedSize();
 	}
-	else
+
+	return MemSize;
+}
+
+void FPCGLandscapeCacheEntry::SerializeToBulkData()
+{
+	// Move data from local arrays to the bulk data
+	BulkData.Lock(LOCK_READ_WRITE);
+
+	int32 NumBytes = 0;
+	// Number of entries in the array + size of the array
+	NumBytes += sizeof(int32) + PositionsAndNormals.Num() * PositionsAndNormals.GetTypeSize();
+
+	// Number of layers
+	NumBytes += sizeof(int32); 
+	for (TArray<uint8>& CurrentLayerData : LayerData)
 	{
-		TeardownLandscapeCallbacks();
+		// Number of entries in the layer data array + size of the array
+		NumBytes += sizeof(int32) + CurrentLayerData.Num() * CurrentLayerData.GetTypeSize();
 	}
 
-	if (bUpdateCachedNames)
+	uint8* Dest = (uint8*)BulkData.Realloc(NumBytes);
+	FBufferWriter Ar(Dest, NumBytes);
+	Ar.SetIsPersistent(true);
+
+	Ar << PositionsAndNormals;
+	
+	int32 LayerDataCount = LayerData.Num();
+	Ar << LayerDataCount;
+
+	for (TArray<uint8>& CurrentLayerData : LayerData)
 	{
-		CacheLayerNames();
+		Ar << CurrentLayerData;
 	}
+
+	BulkData.Unlock();
+}
+
+void FPCGLandscapeCacheEntry::SerializeFromBulkData() const
+{
+	check(!bDataLoaded);
+
+	// Note: this call is not threadsafe by itself, it is meant to be called from a locked region
+	uint8* Data = nullptr;
+	BulkData.GetCopy((void**)&Data);
+	int32 DataSize = BulkData.GetBulkDataSize();
+
+	FBufferReader Ar(Data, DataSize, /*bInFreeOnClose=*/true, /*bIsPersistent=*/true);
+
+	Ar << PositionsAndNormals;
+
+	int32 LayerCount = 0;
+	Ar << LayerCount;
+
+	LayerData.SetNum(LayerCount);
+
+	for (int32 LayerIndex = 0; LayerIndex < LayerCount; ++LayerIndex)
+	{
+		Ar << LayerData[LayerIndex];
+	}
+
+	bDataLoaded = true;
+}
+
+void FPCGLandscapeCacheEntry::Serialize(FArchive& Archive, UObject* Owner, int32 Index)
+{
+	check(bDataLoaded || Archive.IsLoading());
+
+	if (bDataLoaded && Archive.IsSaving() && Archive.IsCooking())
+	{
+		SerializeToBulkData();
+	}
+
+	Archive << Component;
+	Archive << PointHalfSize;
+	Archive << Stride;
+	Archive << LayerDataNames;
+
+	// We force it not inline that means bulk data won't automatically be loaded when we deserialize
+	// later but only when we explicitly take action to load it
+	BulkData.SetBulkDataFlags(BULKDATA_Force_NOT_InlinePayload);
+	BulkData.Serialize(Archive, Owner, Index);
+}
+
+UPCGLandscapeCache::UPCGLandscapeCache()
+{
+#if WITH_EDITOR
+	SetupLandscapeCallbacks();
+	CacheLayerNames();
 #endif
 }
 
-FPCGLandscapeCache::~FPCGLandscapeCache()
+UPCGLandscapeCache::~UPCGLandscapeCache()
 {
+	ClearCache();
 #if WITH_EDITOR
 	TeardownLandscapeCallbacks();
 #endif
 }
 
-void FPCGLandscapeCache::PrimeCache()
+void UPCGLandscapeCache::Serialize(FArchive& Archive)
 {
-#if WITH_EDITOR
-	if (!Owner)
+	Super::Serialize(Archive);
+
+	if (Archive.IsSaving() && !Archive.IsCooking())
 	{
-		return;
+		ClearCache();
 	}
 
-	UWorld* World = Owner->GetWorld();
+	// Serialize cache entries
+	int32 NumEntries = (Archive.IsLoading() ? 0 : CachedData.Num());
+	Archive << NumEntries;
+
+	if (Archive.IsLoading())
+	{
+		CachedData.Reserve(NumEntries);
+
+		for (int32 EntryIndex = 0; EntryIndex < NumEntries; ++EntryIndex)
+		{
+			TPair<FGuid, FIntPoint> Key;
+			Archive << Key;
+
+			FPCGLandscapeCacheEntry* Entry = new FPCGLandscapeCacheEntry();
+			Entry->Serialize(Archive, this, EntryIndex);
+
+			CachedData.Add(Key, Entry);
+		}
+	}
+	else
+	{
+		int32 EntryIndex = 0;
+		for (auto& CacheEntry : CachedData)
+		{
+			Archive << CacheEntry.Key;
+			CacheEntry.Value->Serialize(Archive, this, EntryIndex++);
+		}
+	}
+}
+
+void UPCGLandscapeCache::Tick(float DeltaSeconds)
+{
+#if !WITH_EDITOR
+	// Important implementation note:
+	// If the threshold is too low, it could lead to some issues - namely the check(bDataLoaded) in the cache entries
+	// as we currently do not have a state in the landscape cache to know whether something is still under use.
+	// It is possible to do so however, by adding a scoped variable + a counter, so we could safely remove these from
+	// the cleanup process done below, but at this point in time it doesn't seem to be needed or likely,
+	// since we do not keep (wrt time-slicing) the landscape cache entry pointers at any point in the PCG landscape data implementation
+	TimeSinceLastCleanupInSeconds += DeltaSeconds;
+
+	if (TimeSinceLastCleanupInSeconds >= CVarLandscapeCacheGCFrequency.GetValueOnAnyThread())
+	{
+		TimeSinceLastCleanupInSeconds = 0;
+
+		const int32 MemoryThresholdInBytes = CVarLandscapeCacheSizeThreshold.GetValueOnAnyThread() * 1024 * 1024;
+
+		if (CacheMemorySize > MemoryThresholdInBytes)
+		{
+			TArray<FPCGLandscapeCacheEntry*> LoadedEntries;
+			for (TPair<TPair<FGuid, FIntPoint>, FPCGLandscapeCacheEntry*>& CacheEntry : CachedData)
+			{
+				if (CacheEntry.Value->bDataLoaded)
+				{
+					LoadedEntries.Add(CacheEntry.Value);
+				}
+			}
+
+			Algo::Sort(LoadedEntries, [](const FPCGLandscapeCacheEntry* A, const FPCGLandscapeCacheEntry* B)
+			{
+				return A->Touch < B->Touch;
+			});
+
+			// Go through the oldest data and free them until we're under the threshold
+			int32 EntryIndex = 0;
+			int32 HighestTouchUnloaded = 0;
+			while(EntryIndex < LoadedEntries.Num() && CacheMemorySize > MemoryThresholdInBytes)
+			{
+				CacheMemorySize -= LoadedEntries[EntryIndex]->GetMemorySize();
+				HighestTouchUnloaded = LoadedEntries[EntryIndex]->Touch;
+				LoadedEntries[EntryIndex]->Unload();
+				check(LoadedEntries[EntryIndex]->GetMemorySize()==0);
+				++EntryIndex;
+			}
+
+			// Adjust touch on the other loaded entries so we can keep Touch bounded
+			while (EntryIndex < LoadedEntries.Num())
+			{
+				LoadedEntries[EntryIndex++]->Touch -= HighestTouchUnloaded;
+			}
+
+			CacheTouch -= HighestTouchUnloaded;
+		}
+	}
+#endif // !WITH_EDITOR
+}
+
+void UPCGLandscapeCache::PrimeCache()
+{
+#if WITH_EDITOR
+	UWorld* World = GetWorld();
+
 	if (!World)
 	{
 		return;
@@ -273,6 +494,9 @@ void FPCGLandscapeCache::PrimeCache()
 			// Build per-component information
 			LandscapeInfo->ForAllLandscapeProxies([this, LandscapeInfo](const ALandscapeProxy* LandscapeProxy)
 			{
+				check(LandscapeProxy);
+				const FGuid LandscapeGuid = LandscapeProxy->GetLandscapeGuid();
+
 				for (ULandscapeComponent* LandscapeComponent : LandscapeProxy->LandscapeComponents)
 				{
 					if (!LandscapeComponent)
@@ -280,11 +504,14 @@ void FPCGLandscapeCache::PrimeCache()
 						continue;
 					}
 
-					FIntPoint ComponentKey(LandscapeComponent->SectionBaseX / LandscapeComponent->ComponentSizeQuads, LandscapeComponent->SectionBaseY / LandscapeComponent->ComponentSizeQuads);
+					FIntPoint Coordinate(LandscapeComponent->SectionBaseX / LandscapeComponent->ComponentSizeQuads, LandscapeComponent->SectionBaseY / LandscapeComponent->ComponentSizeQuads);
+					TPair<FGuid, FIntPoint> ComponentKey(LandscapeGuid, Coordinate);
 
 					if (!CachedData.Contains(ComponentKey))
 					{
-						CachedData.Add(ComponentKey).BuildCacheData(LandscapeInfo, LandscapeComponent, Owner);
+						FPCGLandscapeCacheEntry* NewEntry = new FPCGLandscapeCacheEntry();
+						NewEntry->BuildCacheData(LandscapeInfo, LandscapeComponent);
+						CachedData.Add(ComponentKey, NewEntry);
 					}
 				}
 			});
@@ -295,43 +522,64 @@ void FPCGLandscapeCache::PrimeCache()
 #endif
 }
 
-void FPCGLandscapeCache::ClearCache()
+void UPCGLandscapeCache::ClearCache()
 {
+	for (TPair<TPair<FGuid, FIntPoint>, FPCGLandscapeCacheEntry*>& CacheEntry : CachedData)
+	{
+		delete CacheEntry.Value;
+		CacheEntry.Value = nullptr;
+	}
+
 	CachedData.Reset();
 }
 
-const FPCGLandscapeCacheEntry* FPCGLandscapeCache::GetCacheEntry(ULandscapeComponent* LandscapeComponent, const FIntPoint& ComponentKey)
+const FPCGLandscapeCacheEntry* UPCGLandscapeCache::GetCacheEntry(ULandscapeComponent* LandscapeComponent, const FIntPoint& ComponentCoordinate)
 {
-	const FPCGLandscapeCacheEntry* FoundEntry = nullptr;
+	const FPCGLandscapeCacheEntry* CacheEntry = nullptr;
+	TPair<FGuid, FIntPoint> ComponentKey(LandscapeComponent && LandscapeComponent->GetLandscapeProxy() ? LandscapeComponent->GetLandscapeProxy()->GetLandscapeGuid() : FGuid(), ComponentCoordinate);
 
 	{
 #if WITH_EDITOR
 		FReadScopeLock ScopeLock(CacheLock);
 #endif
-		FoundEntry = CachedData.Find(ComponentKey);
+		if (FPCGLandscapeCacheEntry** FoundEntry = CachedData.Find(ComponentKey))
+		{
+			CacheEntry = *FoundEntry;
+		}
 	}
 
 #if WITH_EDITOR
-	if (!FoundEntry && LandscapeComponent && LandscapeComponent->GetLandscapeInfo())
+	if (!CacheEntry && LandscapeComponent && LandscapeComponent->GetLandscapeInfo())
 	{
-		check(LandscapeComponent->SectionBaseX / LandscapeComponent->ComponentSizeQuads == ComponentKey.X && LandscapeComponent->SectionBaseY / LandscapeComponent->ComponentSizeQuads == ComponentKey.Y);
-		// Create entry
-		FPCGLandscapeCacheEntry NewEntry;
-		NewEntry.BuildCacheData(LandscapeComponent->GetLandscapeInfo(), LandscapeComponent, Owner);
-
-		// Try to store it
+		FWriteScopeLock ScopeLock(CacheLock);
+		if (FPCGLandscapeCacheEntry** FoundEntry = CachedData.Find(ComponentKey))
 		{
-			FWriteScopeLock ScopeLock(CacheLock);
-			CachedData.FindOrAdd(ComponentKey, NewEntry);
-			FoundEntry = CachedData.Find(ComponentKey);
+			CacheEntry = *FoundEntry;
+		}
+		else
+		{
+			check(LandscapeComponent->SectionBaseX / LandscapeComponent->ComponentSizeQuads == ComponentKey.Value.X && LandscapeComponent->SectionBaseY / LandscapeComponent->ComponentSizeQuads == ComponentKey.Value.Y);
+			FPCGLandscapeCacheEntry* NewEntry = new FPCGLandscapeCacheEntry();
+			NewEntry->BuildCacheData(LandscapeComponent->GetLandscapeInfo(), LandscapeComponent);
+
+			CacheEntry = NewEntry;
+			CachedData.Add(ComponentKey, NewEntry);
 		}
 	}
 #endif
-	
-	return FoundEntry;
+
+	if (CacheEntry)
+	{
+		if (CacheEntry->TouchAndLoad(CacheTouch++))
+		{
+			CacheMemorySize += CacheEntry->GetMemorySize();
+		}
+	}
+
+	return CacheEntry;
 }
 
-TArray<FName> FPCGLandscapeCache::GetLayerNames(ALandscapeProxy* Landscape)
+TArray<FName> UPCGLandscapeCache::GetLayerNames(ALandscapeProxy* Landscape)
 {
 #if WITH_EDITOR
 	FReadScopeLock ScopeLock(CacheLock);
@@ -340,17 +588,12 @@ TArray<FName> FPCGLandscapeCache::GetLayerNames(ALandscapeProxy* Landscape)
 }
 
 #if WITH_EDITOR
-void FPCGLandscapeCache::SetupLandscapeCallbacks()
+void UPCGLandscapeCache::SetupLandscapeCallbacks()
 {
 	// Remove previous callbacks, if any
 	TeardownLandscapeCallbacks();
 
-	if (!Owner)
-	{
-		return;
-	}
-
-	UWorld* World = Owner->GetWorld();
+	UWorld* World = GetWorld();
 	if (!World)
 	{
 		return;
@@ -365,11 +608,11 @@ void FPCGLandscapeCache::SetupLandscapeCallbacks()
 		ALandscapeProxy* Landscape = CastChecked<ALandscapeProxy>(FoundLandscape);
 
 		Landscapes.Add(Landscape);
-		Landscape->OnComponentDataChanged.AddRaw(this, &FPCGLandscapeCache::OnLandscapeChanged);
+		Landscape->OnComponentDataChanged.AddUObject(this, &UPCGLandscapeCache::OnLandscapeChanged);
 	}
 }
 
-void FPCGLandscapeCache::TeardownLandscapeCallbacks()
+void UPCGLandscapeCache::TeardownLandscapeCallbacks()
 {
 	for (TWeakObjectPtr<ALandscapeProxy> LandscapeWeakPtr : Landscapes)
 	{
@@ -380,7 +623,7 @@ void FPCGLandscapeCache::TeardownLandscapeCallbacks()
 	}
 }
 
-void FPCGLandscapeCache::OnLandscapeChanged(ALandscapeProxy* Landscape, const FLandscapeProxyComponentDataChangedParams& ChangeParams)
+void UPCGLandscapeCache::OnLandscapeChanged(ALandscapeProxy* Landscape, const FLandscapeProxyComponentDataChangedParams& ChangeParams)
 {
 	if (!Landscapes.Contains(Landscape))
 	{
@@ -390,12 +633,18 @@ void FPCGLandscapeCache::OnLandscapeChanged(ALandscapeProxy* Landscape, const FL
 	CacheLock.WriteLock();
 
 	// Just remove these from the cache, they'll be added back on demand
-	ChangeParams.ForEachComponent([this](const ULandscapeComponent* LandscapeComponent)
+	ChangeParams.ForEachComponent([this, Landscape](const ULandscapeComponent* LandscapeComponent)
 	{
 		if (LandscapeComponent)
 		{
-			FIntPoint ComponentKey(LandscapeComponent->SectionBaseX / LandscapeComponent->ComponentSizeQuads, LandscapeComponent->SectionBaseY / LandscapeComponent->ComponentSizeQuads);
-			CachedData.Remove(ComponentKey);
+			FIntPoint Coordinate(LandscapeComponent->SectionBaseX / LandscapeComponent->ComponentSizeQuads, LandscapeComponent->SectionBaseY / LandscapeComponent->ComponentSizeQuads);
+			TPair<FGuid, FIntPoint> ComponentKey(Landscape->GetLandscapeGuid(), Coordinate);
+
+			if (FPCGLandscapeCacheEntry** FoundEntry = CachedData.Find(ComponentKey))
+			{
+				delete *FoundEntry;
+				CachedData.Remove(ComponentKey);
+			}
 		}
 	});
 
@@ -404,7 +653,7 @@ void FPCGLandscapeCache::OnLandscapeChanged(ALandscapeProxy* Landscape, const FL
 	CacheLock.WriteUnlock();
 }
 
-void FPCGLandscapeCache::CacheLayerNames()
+void UPCGLandscapeCache::CacheLayerNames()
 {
 	CachedLayerNames.Reset();
 
@@ -417,7 +666,7 @@ void FPCGLandscapeCache::CacheLayerNames()
 	}
 }
 
-void FPCGLandscapeCache::CacheLayerNames(ALandscapeProxy* Landscape)
+void UPCGLandscapeCache::CacheLayerNames(ALandscapeProxy* Landscape)
 {
 	check(Landscape);
 
diff --git a/Engine/Plugins/Experimental/PCG/Source/PCG/Private/PCGSubsystem.cpp b/Engine/Plugins/Experimental/PCG/Source/PCG/Private/PCGSubsystem.cpp
--- a/Engine/Plugins/Experimental/PCG/Source/PCG/Private/PCGSubsystem.cpp
+++ b/Engine/Plugins/Experimental/PCG/Source/PCG/Private/PCGSubsystem.cpp
@@ -93,6 +93,12 @@ void UPCGSubsystem::Tick(float DeltaSeconds)
 
 	// If we have any tasks to execute, schedule some
 	GraphExecutor->Execute();
+
+	// Lose references to landscape cache as needed
+	if (UPCGLandscapeCache* LandscapeCache = GetLandscapeCache())
+	{
+		LandscapeCache->Tick(DeltaSeconds);
+	}
 }
 
 APCGWorldActor* UPCGSubsystem::GetPCGWorldActor()
@@ -139,10 +145,10 @@ void UPCGSubsystem::UnregisterPCGWorldActor(APCGWorldActor* InActor)
 	}
 }
 
-FPCGLandscapeCache* UPCGSubsystem::GetLandscapeCache()
+UPCGLandscapeCache* UPCGSubsystem::GetLandscapeCache()
 {
 	APCGWorldActor* LandscapeCacheOwner = GetPCGWorldActor();
-	return LandscapeCacheOwner ? &LandscapeCacheOwner->LandscapeCache : nullptr;
+	return LandscapeCacheOwner ? LandscapeCacheOwner->LandscapeCache.Get() : nullptr;
 }
 
 FPCGTaskId UPCGSubsystem::ScheduleComponent(UPCGComponent* PCGComponent, bool bSave, const TArray<FPCGTaskId>& Dependencies)
@@ -1322,7 +1328,7 @@ void UPCGSubsystem::FlushCache()
 
 void UPCGSubsystem::BuildLandscapeCache()
 {
-	if (FPCGLandscapeCache* LandscapeCache = GetLandscapeCache())
+	if (UPCGLandscapeCache* LandscapeCache = GetLandscapeCache())
 	{
 		PCGWorldActor->Modify();
 		LandscapeCache->PrimeCache();
@@ -1335,7 +1341,7 @@ void UPCGSubsystem::BuildLandscapeCache()
 
 void UPCGSubsystem::ClearLandscapeCache()
 {
-	if (FPCGLandscapeCache* LandscapeCache = GetLandscapeCache())
+	if (UPCGLandscapeCache* LandscapeCache = GetLandscapeCache())
 	{
 		LandscapeCache->ClearCache();
 	}
diff --git a/Engine/Plugins/Experimental/PCG/Source/PCG/Private/PCGWorldActor.cpp b/Engine/Plugins/Experimental/PCG/Source/PCG/Private/PCGWorldActor.cpp
--- a/Engine/Plugins/Experimental/PCG/Source/PCG/Private/PCGWorldActor.cpp
+++ b/Engine/Plugins/Experimental/PCG/Source/PCG/Private/PCGWorldActor.cpp
@@ -18,31 +18,20 @@ APCGWorldActor::APCGWorldActor(const FObjectInitializer& ObjectInitializer)
 #endif
 
 	PartitionGridSize = DefaultPartitionGridSize;
-	LandscapeCache.SetOwner(this);
-}
-
-void APCGWorldActor::Serialize(FArchive& InArchive)
-{
-	if (InArchive.IsSaving() && !InArchive.IsCooking() && !(GetWorld() && GetWorld()->GetWorldPartition()))
-	{
-		LandscapeCache.ClearCache();
-	}
-
-	Super::Serialize(InArchive);
+	LandscapeCache = ObjectInitializer.CreateDefaultSubobject<UPCGLandscapeCache>(this, TEXT("LandscapeCache"));
 }
 
 #if WITH_EDITOR
 void APCGWorldActor::BeginCacheForCookedPlatformData(const ITargetPlatform* TargetPlatform)
 {
 	Super::BeginCacheForCookedPlatformData(TargetPlatform);
-	LandscapeCache.PrimeCache();
+	LandscapeCache->PrimeCache();
 }
 #endif
 
 void APCGWorldActor::PostLoad()
 {
 	Super::PostLoad();
-	LandscapeCache.SetOwner(this, /*bUpdateCachedNames=*/false);
 	RegisterToSubsystem();
 }
 
diff --git a/Engine/Plugins/Experimental/PCG/Source/PCG/Public/Data/PCGLandscapeData.h b/Engine/Plugins/Experimental/PCG/Source/PCG/Public/Data/PCGLandscapeData.h
--- a/Engine/Plugins/Experimental/PCG/Source/PCG/Public/Data/PCGLandscapeData.h
+++ b/Engine/Plugins/Experimental/PCG/Source/PCG/Public/Data/PCGLandscapeData.h
@@ -9,7 +9,7 @@
 
 class ALandscapeProxy;
 class ULandscapeInfo;
-struct FPCGLandscapeCache;
+class UPCGLandscapeCache;
 
 UCLASS(BlueprintType, ClassGroup = (Procedural))
 class PCG_API UPCGLandscapeData : public UPCGSurfaceData
@@ -55,5 +55,5 @@ protected:
 private:
 	// Transient data
 	ULandscapeInfo* LandscapeInfo = nullptr;
-	FPCGLandscapeCache* LandscapeCache = nullptr;
+	UPCGLandscapeCache* LandscapeCache = nullptr;
 };
diff --git a/Engine/Plugins/Experimental/PCG/Source/PCG/Public/Grid/PCGLandscapeCache.h b/Engine/Plugins/Experimental/PCG/Source/PCG/Public/Grid/PCGLandscapeCache.h
--- a/Engine/Plugins/Experimental/PCG/Source/PCG/Public/Grid/PCGLandscapeCache.h
+++ b/Engine/Plugins/Experimental/PCG/Source/PCG/Public/Grid/PCGLandscapeCache.h
@@ -2,6 +2,8 @@
 
 #pragma once
 
+#include "Serialization/BulkData.h"
+
 #include "PCGLandscapeCache.generated.h"
 
 class ALandscapeProxy;
@@ -12,61 +14,67 @@ class UPCGPointData;
 class UPCGMetadata;
 struct FPCGPoint;
 
-USTRUCT()
-struct FPCGLandscapeCacheLayer
-{
-	GENERATED_BODY()
-
-	UPROPERTY(VisibleAnywhere, Category = "Cache")
-	FName Name;
-
-	UPROPERTY()
-	TArray<uint8> Data;
-};
-
-USTRUCT()
 struct FPCGLandscapeCacheEntry
 {
-	GENERATED_BODY()
-public:
-#if WITH_EDITOR
-	void BuildCacheData(ULandscapeInfo* LandscapeInfo, ULandscapeComponent* InComponent, UObject* Owner);
-#endif
+	friend class UPCGLandscapeCache;
 
+public:
 	void GetPoint(int32 PointIndex, FPCGPoint& OutPoint, UPCGMetadata* OutMetadata) const;
 	void GetPointHeightOnly(int32 PointIndex, FPCGPoint& OutPoint) const;
 	void GetInterpolatedPoint(const FVector2D& LocalPoint, FPCGPoint& OutPoint, UPCGMetadata* OutMetadata) const;
 	void GetInterpolatedPointHeightOnly(const FVector2D& LocalPoint, FPCGPoint& OutPoint) const;
 
-	// TODO: this isn't really needed
-	UPROPERTY()
-	TWeakObjectPtr<const ULandscapeComponent> Component = nullptr;
+private:
+	// Private API for UPCGLandscapeCache usage
+	bool TouchAndLoad(int32 Touch) const;
+	void Unload();
+	int32 GetMemorySize() const;
 
-	UPROPERTY()
-	TArray<FVector> PositionsAndNormals;
+#if WITH_EDITOR
+	void BuildCacheData(ULandscapeInfo* LandscapeInfo, ULandscapeComponent* InComponent);
+#endif
 
-	UPROPERTY()
-	TArray<FPCGLandscapeCacheLayer> LayerData;
+	// Serialize called from the landscape cache
+	void Serialize(FArchive& Ar, UObject* Owner, int32 BulkIndex);
 
-	UPROPERTY()
-	FVector PointHalfSize = FVector::One();
+	// Internal usage methods
+	void SerializeToBulkData();
+	void SerializeFromBulkData() const;
 
-	UPROPERTY()
+	// Serialized data
+	TWeakObjectPtr<const ULandscapeComponent> Component = nullptr;
+	TArray<FName> LayerDataNames;
+	FVector PointHalfSize = FVector::One();
 	int32 Stride = 0;
+
+	// Data built in editor or loaded from the bulk data
+	mutable FByteBulkData BulkData;
+
+	// Data stored in the BulkData
+	mutable TArray<FVector> PositionsAndNormals;
+	mutable TArray<TArray<uint8>> LayerData;
+
+	// Transient data
+	mutable FCriticalSection DataLock;
+	mutable int32 Touch = 0;
+	mutable bool bDataLoaded = false;
 };
 
-USTRUCT()
-struct FPCGLandscapeCache
+UCLASS()
+class UPCGLandscapeCache : public UObject
 {
 	GENERATED_BODY()
 public:
-	FPCGLandscapeCache() = default;
-	~FPCGLandscapeCache();
-	FPCGLandscapeCache(UObject* InOwner);
+	UPCGLandscapeCache();
+	~UPCGLandscapeCache();
+
+	//~Begin UObject interface
+	virtual void Serialize(FArchive& Archive) override;
+	//~End UObject interface
 
-	void SetOwner(UObject* InOwner, bool bUpdateCachedLayerNames = true);
 	void PrimeCache();
 	void ClearCache();
+	void Tick(float DeltaSeconds);
 
 	const FPCGLandscapeCacheEntry* GetCacheEntry(ULandscapeComponent* LandscapeComponent, const FIntPoint& ComponentKey);
 	TArray<FName> GetLayerNames(ALandscapeProxy* Landscape);
@@ -80,16 +88,15 @@ private:
 	void CacheLayerNames();
 #endif
 
-	// TODO: need to have an indirection for multiple landscape support
-	UPROPERTY(VisibleAnywhere, Category = "Cache")
-	TMap<FIntPoint, FPCGLandscapeCacheEntry> CachedData;
+	TMap<TPair<FGuid, FIntPoint>, FPCGLandscapeCacheEntry*> CachedData;
 
 	//TODO: separate by landscape
 	UPROPERTY(VisibleAnywhere, Category = "Cache")
 	TSet<FName> CachedLayerNames;
 
-	// Transient by design
-	UObject* Owner = nullptr;
+	std::atomic<int32> CacheMemorySize = 0;
+	std::atomic<int32> CacheTouch = 0;
+	float TimeSinceLastCleanupInSeconds = 0;
 
 #if WITH_EDITOR
 	TSet<TWeakObjectPtr<ALandscapeProxy>> Landscapes;
@@ -98,13 +105,4 @@ private:
 #if WITH_EDITOR
 	FRWLock CacheLock;
 #endif
-};
-
-template<>
-struct TStructOpsTypeTraits<FPCGLandscapeCache> : public TStructOpsTypeTraitsBase2< FPCGLandscapeCache>
-{
-	enum
-	{
-		WithCopy = false
-	};
 };
\ No newline at end of file
diff --git a/Engine/Plugins/Experimental/PCG/Source/PCG/Public/PCGSubsystem.h b/Engine/Plugins/Experimental/PCG/Source/PCG/Public/PCGSubsystem.h
--- a/Engine/Plugins/Experimental/PCG/Source/PCG/Public/PCGSubsystem.h
+++ b/Engine/Plugins/Experimental/PCG/Source/PCG/Public/PCGSubsystem.h
@@ -16,7 +16,7 @@ class APCGPartitionActor;
 class APCGWorldActor;
 class UPCGGraph;
 struct FPCGDataCollection;
-struct FPCGLandscapeCache;
+class UPCGLandscapeCache;
 
 class IPCGElement;
 typedef TSharedPtr<IPCGElement, ESPMode::ThreadSafe> FPCGElementPtr;
@@ -56,7 +56,7 @@ public:
 	void RegisterPCGWorldActor(APCGWorldActor* InActor);
 	void UnregisterPCGWorldActor(APCGWorldActor* InActor);
 
-	FPCGLandscapeCache* GetLandscapeCache();
+	UPCGLandscapeCache* GetLandscapeCache();
 
 	// Schedule graph (owner -> graph)
 	FPCGTaskId ScheduleComponent(UPCGComponent* PCGComponent, bool bSave, const TArray<FPCGTaskId>& Dependencies);
diff --git a/Engine/Plugins/Experimental/PCG/Source/PCG/Public/PCGWorldActor.h b/Engine/Plugins/Experimental/PCG/Source/PCG/Public/PCGWorldActor.h
--- a/Engine/Plugins/Experimental/PCG/Source/PCG/Public/PCGWorldActor.h
+++ b/Engine/Plugins/Experimental/PCG/Source/PCG/Public/PCGWorldActor.h
@@ -18,7 +18,6 @@ public:
 	APCGWorldActor(const FObjectInitializer& ObjectInitializer);
 
 	//~Begin AActor Interface
-	virtual void Serialize(FArchive& InArchive) override;
 	virtual void PostLoad() override;
 	virtual void BeginDestroy() override;
 
@@ -43,8 +42,9 @@ public:
 	UPROPERTY(config, EditAnywhere, Category = WorldPartition)
 	uint32 PartitionGridSize;
 
+	/** Contains all the PCG data required to query the landscape complete. Serialized in cooked builds only */
 	UPROPERTY(VisibleAnywhere, Category = CachedData)
-	FPCGLandscapeCache LandscapeCache;
+	TObjectPtr<UPCGLandscapeCache> LandscapeCache = nullptr;
 
 	/** Disable creation of Partition Actors on the Z axis. Can improve performances if 3D partitioning is not needed. */
 	UPROPERTY(config, EditAnywhere, Category = WorldPartition)
