Added spawn actor tagging and tag inheritance. Note that this will not work seamlessly in collapsed hierarchies.

diff --git a/Engine/Plugins/Experimental/PCG/Source/PCG/Private/Elements/PCGSpawnActor.cpp b/Engine/Plugins/Experimental/PCG/Source/PCG/Private/Elements/PCGSpawnActor.cpp
--- a/Engine/Plugins/Experimental/PCG/Source/PCG/Private/Elements/PCGSpawnActor.cpp
+++ b/Engine/Plugins/Experimental/PCG/Source/PCG/Private/Elements/PCGSpawnActor.cpp
@@ -7,6 +7,7 @@
 #include "PCGManagedResource.h"
 
 #include "Data/PCGPointData.h"
+#include "Grid/PCGPartitionActor.h"
 #include "Helpers/PCGActorHelpers.h"
 #include "Helpers/PCGSettingsHelpers.h"
 
@@ -182,12 +183,15 @@ bool FPCGSpawnActorElement::ExecuteInternal(FPCGContext* Context) const
 					SpawnParams.Owner = TargetActor;
 					SpawnParams.SpawnCollisionHandlingOverride = ESpawnActorCollisionHandlingMethod::AlwaysSpawn;
 
+					TArray<FName> NewActorTags = GetNewActorTags(Context, TargetActor, Settings->bInheritActorTags, Settings->TagsToAddOnActors);
+
+					// Create managed resource for actor tracking
 					UPCGManagedActors* ManagedActors = NewObject<UPCGManagedActors>(Context->SourceComponent);
 
 					for (const FPCGPoint& Point : Points)
 					{
 						AActor* GeneratedActor = TargetActor->GetWorld()->SpawnActor(Settings->TemplateActorClass, &Point.Transform, SpawnParams);
-						GeneratedActor->Tags.Add(PCGHelpers::DefaultPCGActorTag);
+						GeneratedActor->Tags.Append(NewActorTags);
 						GeneratedActor->AttachToActor(TargetActor, FAttachmentTransformRules::KeepWorldTransform);
 
 						ManagedActors->GeneratedActors.Add(GeneratedActor);
@@ -241,6 +245,38 @@ bool FPCGSpawnActorElement::ExecuteInternal(FPCGContext* Context) const
 	return true;
 }
 
+TArray<FName> FPCGSpawnActorElement::GetNewActorTags(FPCGContext* Context, AActor* TargetActor, bool bInheritActorTags, const TArray<FName>& AdditionalTags) const
+{
+	TArray<FName> NewActorTags;
+	// Prepare actor tags
+	if (bInheritActorTags)
+	{
+		// Special case: if the current target actor is a partition, we'll reach out
+		// and find the original actor tags
+		if (APCGPartitionActor* PartitionActor = Cast<APCGPartitionActor>(TargetActor))
+		{
+			if (UPCGComponent* OriginalComponent = PartitionActor->GetOriginalComponent(Context->SourceComponent))
+			{
+				check(OriginalComponent->GetOwner());
+				NewActorTags = OriginalComponent->GetOwner()->Tags;
+			}
+		}
+		else
+		{
+			NewActorTags = TargetActor->Tags;
+		}
+	}
+
+	NewActorTags.AddUnique(PCGHelpers::DefaultPCGActorTag);
+
+	for (const FName& AdditionalTag : AdditionalTags)
+	{
+		NewActorTags.AddUnique(AdditionalTag);
+	}
+
+	return NewActorTags;
+}
+
 #if WITH_EDITOR
 EPCGSettingsType UPCGSpawnActorSettings::GetType() const
 {
diff --git a/Engine/Plugins/Experimental/PCG/Source/PCG/Public/Elements/PCGSpawnActor.h b/Engine/Plugins/Experimental/PCG/Source/PCG/Public/Elements/PCGSpawnActor.h
--- a/Engine/Plugins/Experimental/PCG/Source/PCG/Public/Elements/PCGSpawnActor.h
+++ b/Engine/Plugins/Experimental/PCG/Source/PCG/Public/Elements/PCGSpawnActor.h
@@ -35,6 +35,13 @@ public:
 	UPROPERTY(BlueprintReadWrite, EditAnywhere, Category = Settings, meta = (EditCondition = "Option==EPCGSpawnActorOption::NoMerging"))
 	bool bForceDisableActorParsing = true;
 
+	/** Warning: inheriting parent actor tags work only in non-collapsed actor hierarchies */
+	UPROPERTY(BlueprintReadWrite, EditAnywhere, Category = Settings, meta = (EditCondition = "Option!=EPCGSpawnActorOption::CollapseActors"))
+	bool bInheritActorTags = false;
+
+	UPROPERTY(BlueprintReadWrite, EditAnywhere, Category = Settings, meta = (EditCondition = "Option!=EPCGSpawnActorOption::CollapseActors"))
+	TArray<FName> TagsToAddOnActors;
+
 	//~Begin UCPGSettings interface
 	virtual UPCGNode* CreateNode() const override;
 
@@ -76,4 +83,7 @@ public:
 
 protected:
 	virtual bool ExecuteInternal(FPCGContext* Context) const override;
+
+private:
+	TArray<FName> GetNewActorTags(FPCGContext* Context, AActor* TargetActor, bool bInheritActorTags, const TArray<FName>& AdditionalTags) const;
 };
