Renamed landscape cache object to fix name collision

#rnx

diff --git a/Engine/Plugins/Experimental/PCG/Source/PCG/Private/PCGSubsystem.cpp b/Engine/Plugins/Experimental/PCG/Source/PCG/Private/PCGSubsystem.cpp
--- a/Engine/Plugins/Experimental/PCG/Source/PCG/Private/PCGSubsystem.cpp
+++ b/Engine/Plugins/Experimental/PCG/Source/PCG/Private/PCGSubsystem.cpp
@@ -148,7 +148,7 @@ void UPCGSubsystem::UnregisterPCGWorldActor(APCGWorldActor* InActor)
 UPCGLandscapeCache* UPCGSubsystem::GetLandscapeCache()
 {
 	APCGWorldActor* LandscapeCacheOwner = GetPCGWorldActor();
-	return LandscapeCacheOwner ? LandscapeCacheOwner->LandscapeCache.Get() : nullptr;
+	return LandscapeCacheOwner ? LandscapeCacheOwner->LandscapeCacheObject.Get() : nullptr;
 }
 
 FPCGTaskId UPCGSubsystem::ScheduleComponent(UPCGComponent* PCGComponent, bool bSave, const TArray<FPCGTaskId>& Dependencies)
diff --git a/Engine/Plugins/Experimental/PCG/Source/PCG/Private/PCGWorldActor.cpp b/Engine/Plugins/Experimental/PCG/Source/PCG/Private/PCGWorldActor.cpp
--- a/Engine/Plugins/Experimental/PCG/Source/PCG/Private/PCGWorldActor.cpp
+++ b/Engine/Plugins/Experimental/PCG/Source/PCG/Private/PCGWorldActor.cpp
@@ -18,14 +18,14 @@ APCGWorldActor::APCGWorldActor(const FObjectInitializer& ObjectInitializer)
 #endif
 
 	PartitionGridSize = DefaultPartitionGridSize;
-	LandscapeCache = ObjectInitializer.CreateDefaultSubobject<UPCGLandscapeCache>(this, TEXT("LandscapeCache"));
+	LandscapeCacheObject = ObjectInitializer.CreateDefaultSubobject<UPCGLandscapeCache>(this, TEXT("LandscapeCache"));
 }
 
 #if WITH_EDITOR
 void APCGWorldActor::BeginCacheForCookedPlatformData(const ITargetPlatform* TargetPlatform)
 {
 	Super::BeginCacheForCookedPlatformData(TargetPlatform);
-	LandscapeCache->PrimeCache();
+	LandscapeCacheObject->PrimeCache();
 }
 #endif
 
diff --git a/Engine/Plugins/Experimental/PCG/Source/PCG/Public/PCGWorldActor.h b/Engine/Plugins/Experimental/PCG/Source/PCG/Public/PCGWorldActor.h
--- a/Engine/Plugins/Experimental/PCG/Source/PCG/Public/PCGWorldActor.h
+++ b/Engine/Plugins/Experimental/PCG/Source/PCG/Public/PCGWorldActor.h
@@ -43,8 +43,8 @@ public:
 	uint32 PartitionGridSize;
 
 	/** Contains all the PCG data required to query the landscape complete. Serialized in cooked builds only */
-	UPROPERTY(VisibleAnywhere, Category = CachedData)
-	TObjectPtr<UPCGLandscapeCache> LandscapeCache = nullptr;
+	UPROPERTY(VisibleAnywhere, Category = CachedData, meta = (DisplayName="Landscape Cache"))
+	TObjectPtr<UPCGLandscapeCache> LandscapeCacheObject = nullptr;
 
 	/** Disable creation of Partition Actors on the Z axis. Can improve performances if 3D partitioning is not needed. */
 	UPROPERTY(config, EditAnywhere, Category = WorldPartition)
