Fixed PCG landscape changes callbacks in WP maps. Fixes issues with missing metadata.

#rnx

diff --git a/Engine/Plugins/Experimental/PCG/Source/PCG/Private/Grid/PCGLandscapeCache.cpp b/Engine/Plugins/Experimental/PCG/Source/PCG/Private/Grid/PCGLandscapeCache.cpp
--- a/Engine/Plugins/Experimental/PCG/Source/PCG/Private/Grid/PCGLandscapeCache.cpp
+++ b/Engine/Plugins/Experimental/PCG/Source/PCG/Private/Grid/PCGLandscapeCache.cpp
@@ -48,7 +48,7 @@ void FPCGLandscapeCacheEntry::BuildCacheData(ULandscapeInfo* LandscapeInfo, ULan
 
 	// Get landscape default layer (heightmap/tangents/normal)
 	{
-		FLandscapeComponentDataInterface CDI(InComponent, 0);
+		FLandscapeComponentDataInterface CDI(InComponent, 0, /*bWorkInEditingLayer=*/false);
 		FVector WorldPos;
 		FVector WorldTangentX;
 		FVector WorldTangentY;
@@ -81,8 +81,8 @@ void FPCGLandscapeCacheEntry::BuildCacheData(ULandscapeInfo* LandscapeInfo, ULan
 			continue;
 		}
 
-		FLandscapeComponentDataInterface CDI(InComponent, 0, /*bUseEditingLayer=*/true);
-		if (CDI.GetWeightmapTextureData(LayerInfo, LayerCache, /*bUseEditingLayer=*/true))
+		FLandscapeComponentDataInterface CDI(InComponent, 0, /*bWorkInEditingLayer=*/false);
+		if (CDI.GetWeightmapTextureData(LayerInfo, LayerCache, /*bUseEditingLayer=*/false))
 		{
 			LayerDataNames.Add(Layer.LayerName);
 			LayerData.Add(LayerCache);
@@ -364,20 +364,14 @@ void FPCGLandscapeCacheEntry::Serialize(FArchive& Archive, UObject* Owner, int32
 	BulkData.Serialize(Archive, Owner, Index);
 }
 
-UPCGLandscapeCache::UPCGLandscapeCache()
-{
-#if WITH_EDITOR
-	SetupLandscapeCallbacks();
-	CacheLayerNames();
-#endif
-}
-
-UPCGLandscapeCache::~UPCGLandscapeCache()
+void UPCGLandscapeCache::BeginDestroy()
 {
 	ClearCache();
 #if WITH_EDITOR
 	TeardownLandscapeCallbacks();
 #endif
+
+	Super::BeginDestroy();
 }
 
 void UPCGLandscapeCache::Serialize(FArchive& Archive)
@@ -421,7 +415,14 @@ void UPCGLandscapeCache::Serialize(FArchive& Archive)
 
 void UPCGLandscapeCache::Tick(float DeltaSeconds)
 {
-#if !WITH_EDITOR
+#if WITH_EDITOR
+	if (!bInitialized && GetWorld())
+	{
+		SetupLandscapeCallbacks();
+		CacheLayerNames();
+		bInitialized = true;
+	}
+#else
 	// Important implementation note:
 	// If the threshold is too low, it could lead to some issues - namely the check(bDataLoaded) in the cache entries
 	// as we currently do not have a state in the landscape cache to know whether something is still under use.
diff --git a/Engine/Plugins/Experimental/PCG/Source/PCG/Private/PCGComponent.cpp b/Engine/Plugins/Experimental/PCG/Source/PCG/Private/PCGComponent.cpp
--- a/Engine/Plugins/Experimental/PCG/Source/PCG/Private/PCGComponent.cpp
+++ b/Engine/Plugins/Experimental/PCG/Source/PCG/Private/PCGComponent.cpp
@@ -759,7 +759,7 @@ void UPCGComponent::SetupCallbacksOnCreation()
 	SetupActorCallbacks();
 	SetupTrackingCallbacks();
 
-	if (TrackedLandscape.IsValid())
+	if(!TrackedLandscapes.IsEmpty())
 	{
 		SetupLandscapeTracking();
 	}
@@ -1150,7 +1150,7 @@ void UPCGComponent::OnActorDeleted(AActor* InActor)
 void UPCGComponent::OnActorMoved(AActor* InActor)
 {
 	const bool bOwnerMoved = (InActor == GetOwner());
-	const bool bLandscapeMoved = (InActor && InActor == TrackedLandscape);
+	const bool bLandscapeMoved = (InActor && TrackedLandscapes.Contains(InActor));
 
 	if (bOwnerMoved || bLandscapeMoved)
 	{
@@ -1189,11 +1189,11 @@ void UPCGComponent::OnActorMoved(AActor* InActor)
 void UPCGComponent::UpdateTrackedLandscape(bool bBoundsCheck)
 {
 	TeardownLandscapeTracking();
-	TrackedLandscape = nullptr;
+	TrackedLandscapes.Reset();
 
 	if (ALandscapeProxy* Landscape = Cast<ALandscapeProxy>(GetOwner()))
 	{
-		TrackedLandscape = Landscape;
+		TrackedLandscapes.Add(Landscape);
 	}
 	else if (InputType == EPCGComponentInput::Landscape || GraphUsesLandscapePin())
 	{
@@ -1204,12 +1204,12 @@ void UPCGComponent::UpdateTrackedLandscape(bool bBoundsCheck)
 				UPCGData* ActorData = GetActorPCGData();
 				if (const UPCGSpatialData* ActorSpatialData = Cast<const UPCGSpatialData>(ActorData))
 				{
-					TrackedLandscape = PCGHelpers::GetLandscape(World, ActorSpatialData->GetBounds());
+					TrackedLandscapes = PCGHelpers::GetLandscapeProxies(World, ActorSpatialData->GetBounds());
 				}
 			}
 			else
 			{
-				TrackedLandscape = PCGHelpers::GetAnyLandscape(World);
+				TrackedLandscapes = PCGHelpers::GetAllLandscapeProxies(World);
 			}
 		}
 	}
@@ -1219,28 +1219,34 @@ void UPCGComponent::UpdateTrackedLandscape(bool bBoundsCheck)
 
 void UPCGComponent::SetupLandscapeTracking()
 {
-	if (TrackedLandscape.IsValid())
+	for (TWeakObjectPtr<ALandscapeProxy> LandscapeProxy : TrackedLandscapes)
 	{
-		TrackedLandscape->OnComponentDataChanged.AddUObject(this, &UPCGComponent::OnLandscapeChanged);
+		if (LandscapeProxy.IsValid())
+		{
+			LandscapeProxy->OnComponentDataChanged.AddUObject(this, &UPCGComponent::OnLandscapeChanged);
+		}
 	}
 }
 
 void UPCGComponent::TeardownLandscapeTracking()
 {
-	if (TrackedLandscape.IsValid())
+	for (TWeakObjectPtr<ALandscapeProxy> LandscapeProxy : TrackedLandscapes)
 	{
-		TrackedLandscape->OnComponentDataChanged.RemoveAll(this);
+		if (LandscapeProxy.IsValid())
+		{
+			LandscapeProxy->OnComponentDataChanged.RemoveAll(this);
+		}
 	}
 }
 
 void UPCGComponent::OnLandscapeChanged(ALandscapeProxy* Landscape, const FLandscapeProxyComponentDataChangedParams& ChangeParams)
 {
-	if (Landscape == TrackedLandscape)
+	if(TrackedLandscapes.Contains(Landscape))
 	{
 		// Check if there is an overlap in the changed components vs. the current actor data
 		EPCGComponentDirtyFlag DirtyFlag = EPCGComponentDirtyFlag::None;
 
-		if (GetOwner() == TrackedLandscape)
+		if (GetOwner() == Landscape)
 		{
 			DirtyFlag = EPCGComponentDirtyFlag::Actor;
 		}
diff --git a/Engine/Plugins/Experimental/PCG/Source/PCG/Private/PCGHelpers.cpp b/Engine/Plugins/Experimental/PCG/Source/PCG/Private/PCGHelpers.cpp
--- a/Engine/Plugins/Experimental/PCG/Source/PCG/Private/PCGHelpers.cpp
+++ b/Engine/Plugins/Experimental/PCG/Source/PCG/Private/PCGHelpers.cpp
@@ -111,16 +111,19 @@ namespace PCGHelpers
 		}
 	}
 
-	ALandscape* GetAnyLandscape(UWorld* InWorld)
+	TArray<TWeakObjectPtr<ALandscapeProxy>> GetAllLandscapeProxies(UWorld* InWorld)
 	{
-		for (TObjectIterator<ALandscape> It; It; ++It)
+		TArray<TWeakObjectPtr<ALandscapeProxy>> LandscapeProxies;
+
+		for (TObjectIterator<ALandscapeProxy> It; It; ++It)
 		{
 			if (It->GetWorld() == InWorld)
 			{
-				return *It;
+				LandscapeProxies.Add(*It);
 			}
 		}
-		return nullptr;
+
+		return LandscapeProxies;
 	}
 
 	ALandscape* GetLandscape(UWorld* InWorld, const FBox& InBounds)
@@ -148,6 +151,30 @@ namespace PCGHelpers
 		return Landscape;
 	}
 
+	TArray<TWeakObjectPtr<ALandscapeProxy>> GetLandscapeProxies(UWorld* InWorld, const FBox& InBounds)
+	{
+		TArray<TWeakObjectPtr<ALandscapeProxy>> LandscapeProxies;
+
+		if (!InBounds.IsValid)
+		{
+			return LandscapeProxies;
+		}
+
+		for (TObjectIterator<ALandscapeProxy> It; It; ++It)
+		{
+			if (It->GetWorld() == InWorld)
+			{
+				const FBox LandscapeBounds = GetLandscapeBounds(*It);
+				if (LandscapeBounds.IsValid && LandscapeBounds.Intersect(InBounds))
+				{
+					LandscapeProxies.Add(*It);
+				}
+			}
+		}
+
+		return LandscapeProxies;
+	}
+
 	APCGWorldActor* GetPCGWorldActor(UWorld* InWorld)
 	{
 		return (InWorld && InWorld->GetSubsystem<UPCGSubsystem>()) ? InWorld->GetSubsystem<UPCGSubsystem>()->GetPCGWorldActor() : nullptr;
diff --git a/Engine/Plugins/Experimental/PCG/Source/PCG/Private/PCGHelpers.h b/Engine/Plugins/Experimental/PCG/Source/PCG/Private/PCGHelpers.h
--- a/Engine/Plugins/Experimental/PCG/Source/PCG/Private/PCGHelpers.h
+++ b/Engine/Plugins/Experimental/PCG/Source/PCG/Private/PCGHelpers.h
@@ -31,7 +31,8 @@ namespace PCGHelpers
 	FBox GetLandscapeBounds(ALandscapeProxy* InLandscape);
 
 	ALandscape* GetLandscape(UWorld* InWorld, const FBox& InActorBounds);
-	ALandscape* GetAnyLandscape(UWorld* InWorld);
+	TArray<TWeakObjectPtr<ALandscapeProxy>> GetLandscapeProxies(UWorld* InWorld, const FBox& InActorBounds);
+	TArray<TWeakObjectPtr<ALandscapeProxy>> GetAllLandscapeProxies(UWorld* InWorld);
 
 	bool IsRuntimeOrPIE();
 
diff --git a/Engine/Plugins/Experimental/PCG/Source/PCG/Public/Grid/PCGLandscapeCache.h b/Engine/Plugins/Experimental/PCG/Source/PCG/Public/Grid/PCGLandscapeCache.h
--- a/Engine/Plugins/Experimental/PCG/Source/PCG/Public/Grid/PCGLandscapeCache.h
+++ b/Engine/Plugins/Experimental/PCG/Source/PCG/Public/Grid/PCGLandscapeCache.h
@@ -65,10 +65,8 @@ class UPCGLandscapeCache : public UObject
 {
 	GENERATED_BODY()
 public:
-	UPCGLandscapeCache();
-	~UPCGLandscapeCache();
-
 	//~Begin UObject interface
+	virtual void BeginDestroy() override;
 	virtual void Serialize(FArchive& Archive) override;
 	//~End UObject interface
 
@@ -97,6 +95,7 @@ private:
 	std::atomic<int32> CacheMemorySize = 0;
 	std::atomic<int32> CacheTouch = 0;
 	float TimeSinceLastCleanupInSeconds = 0;
+	bool bInitialized = false;
 
 #if WITH_EDITOR
 	TSet<TWeakObjectPtr<ALandscapeProxy>> Landscapes;
diff --git a/Engine/Plugins/Experimental/PCG/Source/PCG/Public/PCGComponent.h b/Engine/Plugins/Experimental/PCG/Source/PCG/Public/PCGComponent.h
--- a/Engine/Plugins/Experimental/PCG/Source/PCG/Public/PCGComponent.h
+++ b/Engine/Plugins/Experimental/PCG/Source/PCG/Public/PCGComponent.h
@@ -348,7 +348,7 @@ private:
 
 #if WITH_EDITORONLY_DATA
 	UPROPERTY()
-	TWeakObjectPtr<ALandscapeProxy> TrackedLandscape = nullptr;
+	TArray<TWeakObjectPtr<ALandscapeProxy>> TrackedLandscapes;
 #endif
 
 #if WITH_EDITORONLY_DATA
