Removed pcg world actor assert in subsystem, not an actual issue at the moment.

#rnx

diff --git a/Engine/Plugins/Experimental/PCG/Source/PCG/Private/PCGSubsystem.cpp b/Engine/Plugins/Experimental/PCG/Source/PCG/Private/PCGSubsystem.cpp
--- a/Engine/Plugins/Experimental/PCG/Source/PCG/Private/PCGSubsystem.cpp
+++ b/Engine/Plugins/Experimental/PCG/Source/PCG/Private/PCGSubsystem.cpp
@@ -109,8 +109,8 @@ APCGWorldActor* UPCGSubsystem::GetPCGWorldActor()
 		PCGWorldActorLock.Lock();
 		if (!PCGWorldActor)
 		{
-		PCGWorldActor = APCGWorldActor::CreatePCGWorldActor(GetWorld());
-	}
+			PCGWorldActor = APCGWorldActor::CreatePCGWorldActor(GetWorld());
+		}
 		PCGWorldActorLock.Unlock();
 	}
 #endif
@@ -133,8 +133,11 @@ void UPCGSubsystem::DestroyPCGWorldActor()
 
 void UPCGSubsystem::RegisterPCGWorldActor(APCGWorldActor* InActor)
 {
-	check(!PCGWorldActor || PCGWorldActor == InActor);
-	PCGWorldActor = InActor;
+	// TODO: we should support merging or multi world actor support when relevant
+	if (!PCGWorldActor)
+	{
+		PCGWorldActor = InActor;
+	}
 }
 
 void UPCGSubsystem::UnregisterPCGWorldActor(APCGWorldActor* InActor)
@@ -609,7 +612,6 @@ void UPCGSubsystem::UnregisterPartitionActor(APCGPartitionActor* Actor)
 
 void UPCGSubsystem::ForAllIntersectingPartitionActors(const FBox& InBounds, TFunction<void(APCGPartitionActor*)> InFunc) const
 {
-
 	// No PCGWorldActor just early out. Same for invalid bounds.
 	if (!PCGWorldActor || !InBounds.IsValid)
 	{
