Tentative fix for spurious pcg world actor creation

#rnx

diff --git a/Engine/Plugins/Experimental/PCG/Source/PCG/Private/PCGSubsystem.cpp b/Engine/Plugins/Experimental/PCG/Source/PCG/Private/PCGSubsystem.cpp
--- a/Engine/Plugins/Experimental/PCG/Source/PCG/Private/PCGSubsystem.cpp
+++ b/Engine/Plugins/Experimental/PCG/Source/PCG/Private/PCGSubsystem.cpp
@@ -95,9 +95,9 @@ void UPCGSubsystem::Tick(float DeltaSeconds)
 	GraphExecutor->Execute();
 
 	// Lose references to landscape cache as needed
-	if (UPCGLandscapeCache* LandscapeCache = GetLandscapeCache())
+	if(PCGWorldActor && GetLandscapeCache())
 	{
-		LandscapeCache->Tick(DeltaSeconds);
+		GetLandscapeCache()->Tick(DeltaSeconds);
 	}
 }
 
