PCG: Tentative fix for BP Actor PCG Component not getting unregistered.

diff --git a/Engine/Plugins/Experimental/PCG/Source/PCG/Private/Grid/PCGPartitionActor.cpp b/Engine/Plugins/Experimental/PCG/Source/PCG/Private/Grid/PCGPartitionActor.cpp
--- a/Engine/Plugins/Experimental/PCG/Source/PCG/Private/Grid/PCGPartitionActor.cpp
+++ b/Engine/Plugins/Experimental/PCG/Source/PCG/Private/Grid/PCGPartitionActor.cpp
@@ -292,12 +292,34 @@ void APCGPartitionActor::AddGraphInstance(UPCGComponent* OriginalComponent)
 	LocalToOriginalMap.Add(LocalComponent, OriginalComponent);
 }
 
+void APCGPartitionActor::RemapGraphInstance(const UPCGComponent* OldOriginalComponent, UPCGComponent* NewOriginalComponent)
+{
+	UPCGComponent* LocalComponent = GetLocalComponent(OldOriginalComponent);
+
+	if (!LocalComponent)
+	{
+		return;
+	}
+
+	Modify();
+
+	OriginalToLocalMap.Remove(OldOriginalComponent);
+	LocalToOriginalMap.Remove(LocalComponent);
+
+	LocalComponent->SetPropertiesFromOriginal(NewOriginalComponent);
+	OriginalToLocalMap.Add(NewOriginalComponent, LocalComponent);
+	LocalToOriginalMap.Add(LocalComponent, NewOriginalComponent);
+}
+
 bool APCGPartitionActor::RemoveGraphInstance(UPCGComponent* OriginalComponent)
 {
 	UPCGComponent* LocalComponent = GetLocalComponent(OriginalComponent);
 
 	if (!LocalComponent)
 	{
+		// If we don't have a local component, perhaps the original component is already dead,
+		// so do some clean up
+		CleanupDeadGraphInstances();
 		return false;
 	}
 
diff --git a/Engine/Plugins/Experimental/PCG/Source/PCG/Private/Grid/PCGPartitionActor.h b/Engine/Plugins/Experimental/PCG/Source/PCG/Private/Grid/PCGPartitionActor.h
--- a/Engine/Plugins/Experimental/PCG/Source/PCG/Private/Grid/PCGPartitionActor.h
+++ b/Engine/Plugins/Experimental/PCG/Source/PCG/Private/Grid/PCGPartitionActor.h
@@ -26,12 +26,12 @@ public:
 	//~Begin UObject Interface
 	virtual void PostLoad() override;
 	virtual void BeginDestroy() override;
-	virtual void Destroyed() override;
 	//~End UObject Interface
 
 	//~Begin AActor Interface
 	virtual void BeginPlay();
 	virtual void EndPlay(const EEndPlayReason::Type EndPlayReason) override;
+	virtual void Destroyed() override;
 	virtual void GetActorBounds(bool bOnlyCollidingComponents, FVector& Origin, FVector& BoxExtent, bool bIncludeFromChildActors) const override;
 	virtual void PostRegisterAllComponents() override;
 #if WITH_EDITOR
@@ -54,6 +54,7 @@ public:
 	bool IsUsing2DGrid() const { return bUse2DGrid; }
 
 	void AddGraphInstance(UPCGComponent* OriginalComponent);
+	void RemapGraphInstance(const UPCGComponent* OldOriginalComponent, UPCGComponent* NewOriginalComponent);
 	bool RemoveGraphInstance(UPCGComponent* OriginalComponent);
 	void CleanupDeadGraphInstances();
 
diff --git a/Engine/Plugins/Experimental/PCG/Source/PCG/Private/PCGComponent.cpp b/Engine/Plugins/Experimental/PCG/Source/PCG/Private/PCGComponent.cpp
--- a/Engine/Plugins/Experimental/PCG/Source/PCG/Private/PCGComponent.cpp
+++ b/Engine/Plugins/Experimental/PCG/Source/PCG/Private/PCGComponent.cpp
@@ -505,11 +505,6 @@ bool UPCGComponent::MoveResourcesToNewActor(AActor* InNewActor, bool bCreateChil
 
 void UPCGComponent::CleanupLocalImmediate(bool bRemoveComponents)
 {
-	if (!bGenerated || IsCleaningUp() || IsGenerating())
-	{
-		return;
-	}
-
 	TSet<TSoftObjectPtr<AActor>> ActorsToDelete;
 
 	{
@@ -567,6 +562,12 @@ FPCGTaskId UPCGComponent::CreateCleanupTask(bool bRemoveComponents, const TArray
 	{
 		if (UPCGComponent* ThisComponent = ThisComponentWeakPtr.Get())
 		{
+			// If the component is not valid anymore, just early out
+			if (!IsValid(ThisComponent))
+			{
+				return true;
+			}
+
 			FScopeLock ResourcesLock(&ThisComponent->GeneratedResourcesLock);
 
 			// Safeguard to track illegal modifications of the generated resources array while doing cleanup
@@ -680,11 +681,6 @@ void UPCGComponent::OnComponentCreated()
 
 #if WITH_EDITOR
 	SetupActorCallbacks();
-
-	if (!PCGHelpers::IsRuntimeOrPIE() && IsPartitioned() && GetSubsystem())
-	{
-		GetSubsystem()->RegisterOrUpdatePCGComponent(this, /*bDoActorMapping=*/bGenerated);
-	}
 #endif
 }
 
@@ -695,13 +691,7 @@ void UPCGComponent::OnComponentDestroyed(bool bDestroyingHierarchy)
 	// In the case of level change or exit, the subsystem will be null
 	if (UPCGSubsystem* Subsystem = GetSubsystem())
 	{
-		// The RF_BeginDestroyed flag is set when the object is being unloaded, but not in the editor-destroy context we're interested in.
-		if (!HasAnyFlags(RF_BeginDestroyed) && !IsUnreachable() && IsPartitioned() && !GetOwner()->GetWorld()->IsGameWorld())
-		{
-			Subsystem->CleanupPartitionActors(LastGeneratedBounds);
-		}
-
-		if (!PCGHelpers::IsRuntimeOrPIE() && IsPartitioned())
+		if (!PCGHelpers::IsRuntimeOrPIE())
 		{
 			Subsystem->UnregisterPCGComponent(this);
 		}
@@ -772,22 +762,12 @@ void UPCGComponent::SetupCallbacksOnCreation()
 	{
 		Graph->OnGraphChangedDelegate.AddUObject(this, &UPCGComponent::OnGraphChanged);
 	}
-
-	if (!PCGHelpers::IsRuntimeOrPIE() && IsPartitioned() && GetSubsystem())
-	{
-		GetSubsystem()->RegisterOrUpdatePCGComponent(this, /*bDoActorMapping=*/bGenerated);
-	}
 }
 #endif
 
 void UPCGComponent::BeginDestroy()
 {
 #if WITH_EDITOR
-	if (!PCGHelpers::IsRuntimeOrPIE() && IsPartitioned() && GetSubsystem())
-	{
-		GetSubsystem()->UnregisterPCGComponent(this);
-	}
-
 	if (Graph)
 	{
 		Graph->OnGraphChangedDelegate.RemoveAll(this);
@@ -805,6 +785,28 @@ void UPCGComponent::BeginDestroy()
 	Super::BeginDestroy();
 }
 
+void UPCGComponent::OnRegister()
+{
+	Super::OnRegister();
+
+#if WITH_EDITOR
+	// We can't register to the subsystem in OnRegister if we are at runtime because
+	// the landscape can be not loaded yet.
+	// It will be done in BeginPlay at runtime.
+	if (!PCGHelpers::IsRuntimeOrPIE() && IsPartitioned() && GetSubsystem())
+	{
+		if (UWorld* World = GetWorld())
+		{
+			// We won't be able to spawn any actors if we are currently running a construction script.
+			if (!World->bIsRunningConstructionScript)
+			{
+				GetSubsystem()->RegisterOrUpdatePCGComponent(this, bGenerated);
+			}
+		}
+	}
+#endif //WITH_EDITOR
+}
+
 TStructOnScope<FActorComponentInstanceData> UPCGComponent::GetComponentInstanceData() const
 {
 	TStructOnScope<FActorComponentInstanceData> InstanceData = MakeStructOnScope<FActorComponentInstanceData, FPCGComponentInstanceData>(this);
@@ -898,8 +900,8 @@ void UPCGComponent::PostEditChangeProperty(FPropertyChangedEvent& PropertyChange
 
 	const FName PropName = PropertyChangedEvent.Property->GetFName();
 
-	// Important note: all property changes already go through the OnObjectPropertyChanged, 
-	// So there is no need to add cases that do simple Refresh() calls
+	// Important note: all property changes already go through the OnObjectPropertyChanged, and will be dirtied here.
+	// So where only a Refresh is needed, it goes through the "capture all" else case.
 	if (PropName == GET_MEMBER_NAME_CHECKED(UPCGComponent, bIsPartitioned))
 	{
 		if (CanPartition())
@@ -942,11 +944,6 @@ void UPCGComponent::PostEditChangeProperty(FPropertyChangedEvent& PropertyChange
 		Refresh();
 	}
 	// General properties that don't affect behavior
-	else if (PropName == GET_MEMBER_NAME_CHECKED(UPCGComponent, Seed))
-	{
-		DirtyGenerated();
-		Refresh();
-	}
 	else if (PropName == GET_MEMBER_NAME_CHECKED(UPCGComponent, ExcludedTags))
 	{
 		SetupTrackingCallbacks();
@@ -961,6 +958,10 @@ void UPCGComponent::PostEditChangeProperty(FPropertyChangedEvent& PropertyChange
 			Refresh();
 		}
 	}
+	else
+	{
+		Refresh();
+	}
 }
 
 void UPCGComponent::PostEditImport()
@@ -1287,11 +1288,15 @@ void UPCGComponent::OnObjectPropertyChanged(UObject* InObject, FPropertyChangedE
 	bool bValueNotInteractive = (InEvent.ChangeType != EPropertyChangeType::Interactive);
 	// Special exception for actor tags, as we can't track otherwise an actor "losing" a tag
 	bool bActorTagChange = (InEvent.Property && InEvent.Property->GetFName() == GET_MEMBER_NAME_CHECKED(AActor, Tags));
-	// Discard property changed if this is this component. They will be caught by the PostEditChangeProperty
-	bool bIsThisComponent = InObject == this;
+	if (!bValueNotInteractive && !bActorTagChange)
+	{
+		return;
+	}
 
-	if (bIsThisComponent || (!bValueNotInteractive && !bActorTagChange))
+	// If the object changed is this PCGComponent, dirty ourselves and exit. It will be picked up by PostEditChangeProperty
+	if (InObject == this)
 	{
+		DirtyGenerated();
 		return;
 	}
 
@@ -1420,8 +1425,6 @@ void UPCGComponent::ResetLastGeneratedBounds()
 {
 	LastGeneratedBounds = FBox(EForceInit::ForceInit);
 }
-
-#if WITH_EDITOR
 void UPCGComponent::DisableInspection()
 {
 	bIsInspecting = false;
@@ -1447,17 +1450,20 @@ const FPCGDataCollection* UPCGComponent::GetInspectionData(const UPCGNode* InNod
 {
 	return InspectionCache.Find(InNode);
 }
-#endif
 
 void UPCGComponent::Refresh()
 {
-#if WITH_EDITOR
 	// Disable auto-refreshing on preview actors until we have something more robust on the execution side.
 	if (GetOwner() && GetOwner()->bIsEditorPreviewActor)
 	{
 		return;
 	}
-#endif
+
+	// Discard any refresh if have already one scheduled.
+	if (CurrentRefreshTask != InvalidPCGTaskId)
+	{
+		return;
+	}
 
 	// Before doing a refresh, update the component to the subsystem if we are partitioned
 	// Only redo the mapping if we are generated
@@ -1482,7 +1488,7 @@ void UPCGComponent::Refresh()
 		}
 	}
 }
-#endif
+#endif // WITH_EDITOR
 
 UPCGData* UPCGComponent::GetPCGData()
 {
@@ -2150,8 +2156,9 @@ void UPCGComponent::GetManagedResources(TArray<TObjectPtr<UPCGManagedResource>>&
 	Resources = GeneratedResources;
 }
 
-FPCGComponentInstanceData::FPCGComponentInstanceData(const UPCGComponent* SourceComponent)
-	: FActorComponentInstanceData(SourceComponent)
+FPCGComponentInstanceData::FPCGComponentInstanceData(const UPCGComponent* InSourceComponent)
+	: FActorComponentInstanceData(InSourceComponent)
+	, SourceComponent(InSourceComponent)
 {
 	if (SourceComponent)
 	{
@@ -2189,22 +2196,75 @@ void FPCGComponentInstanceData::ApplyToComponent(UActorComponent* Component, con
 			PCGComponent->SetManagedResources(DuplicatedResources);
 		}
 
-		// We need to make sure that the component is registered to the subsystem if it is partitioned
-		// or unregister it if it is not.
-		if (UPCGSubsystem* Subsystem = PCGComponent->GetSubsystem())
+		// Also remap if we are partitioned
+		UPCGSubsystem* Subsystem = PCGComponent->GetSubsystem();
+		if (Subsystem && PCGComponent->IsPartitioned() && SourceComponent)
 		{
-			if (PCGComponent->IsPartitioned())
-			{
-				// Mapping needs to be done if we are at runtime or previously generated
-				bool bDoActorMapping = PCGComponent->bGenerated || PCGHelpers::IsRuntimeOrPIE();
-				Subsystem->RegisterOrUpdatePCGComponent(PCGComponent, bDoActorMapping);
-			}
-			else
-			{
-				Subsystem->UnregisterPCGComponent(PCGComponent);
-			}
+			Subsystem->RemapPCGComponent(SourceComponent, PCGComponent);
 		}
+
+#if WITH_EDITOR
+		// Finally, start a delayed refresh task (if there is not one already), in editor only
+		// It is important to be delayed, because we cannot spawn Partition Actors within this scope,
+		// because we are in a construction script.
+		DelayedRefresh(PCGComponent);
+#endif // WITH_EDITOR
 	}
 }
 
+#if WITH_EDITOR
+void FPCGComponentInstanceData::DelayedRefresh(UPCGComponent* PCGComponent)
+{
+	if (UPCGSubsystem* Subsystem = PCGComponent->GetSubsystem())
+	{
+		if (PCGComponent->CurrentRefreshTask == InvalidPCGTaskId)
+		{
+			TWeakObjectPtr<UPCGComponent> ComponentPtr(PCGComponent);
+			PCGComponent->CurrentRefreshTask = Subsystem->ScheduleGeneric([ComponentPtr]()
+				{
+					UPCGComponent* Component = ComponentPtr.Get();
+					if (Component && IsValid(Component))
+					{
+						bool bWasGenerated = Component->bGenerated;
+						if (UPCGSubsystem* Subsystem = Component->GetSubsystem())
+						{
+							// Register/Unregister depending on the partition flag.
+							if (Component->IsPartitioned())
+							{
+								// If we are partitioned but we have resources, we need to force a cleanup
+								if (!Component->GeneratedResources.IsEmpty())
+								{
+									Component->CleanupLocalImmediate(true);
+								}
+
+								Subsystem->RegisterOrUpdatePCGComponent(Component, /*bDoActorMapping=*/bWasGenerated);
+							}
+							else
+							{
+								Subsystem->UnregisterPCGComponent(Component);
+							}
+						}
+
+						// Mark the refresh task invalid (since we check in Refresh if there is already a task scheduled)
+						Component->CurrentRefreshTask = InvalidPCGTaskId;
+
+						// If we had a cleanup, just generate.
+						if (bWasGenerated && !Component->bGenerated)
+						{
+							Component->GenerateLocal(/*bForce=*/false);
+						}
+						// Otherwise, only refresh if we are dirty.
+						else if (Component->bDirtyGenerated)
+						{
+							Component->Refresh();
+						}
+					}
+
+					return true;
+				}, {});
+		}
+	}
+}
+#endif // WITH_EDITOR
+
 #undef LOCTEXT_NAMESPACE
diff --git a/Engine/Plugins/Experimental/PCG/Source/PCG/Private/PCGSubsystem.cpp b/Engine/Plugins/Experimental/PCG/Source/PCG/Private/PCGSubsystem.cpp
--- a/Engine/Plugins/Experimental/PCG/Source/PCG/Private/PCGSubsystem.cpp
+++ b/Engine/Plugins/Experimental/PCG/Source/PCG/Private/PCGSubsystem.cpp
@@ -99,6 +99,17 @@ void UPCGSubsystem::Tick(float DeltaSeconds)
 	{
 		GetLandscapeCache()->Tick(DeltaSeconds);
 	}
+
+	TSet<TObjectPtr<UPCGComponent>> ComponentToUnregister;
+	{
+		FScopeLock Lock(&DelayedComponentToUnregisterLock);
+		ComponentToUnregister = MoveTemp(DelayedComponentToUnregister);
+	}
+
+	for (UPCGComponent* Component : ComponentToUnregister)
+	{
+		UnregisterPCGComponent(Component, /*bForce=*/true);
+	}
 }
 
 APCGWorldActor* UPCGSubsystem::GetPCGWorldActor()
@@ -221,6 +232,12 @@ FPCGTaskId UPCGSubsystem::ScheduleComponent(UPCGComponent* PCGComponent, bool bS
 		return GraphExecutor->ScheduleGeneric([ComponentPtr]() {
 			if (UPCGComponent* Component = ComponentPtr.Get())
 			{
+				// If the component is not valid anymore, just early out.
+				if (!IsValid(Component))
+				{
+					return true;
+				}
+
 				const FBox NewBounds = Component->GetGridBounds();
 				Component->PostProcessGraph(NewBounds, /*bGenerate=*/true);
 			}
@@ -279,6 +296,12 @@ FPCGTaskId UPCGSubsystem::ScheduleCleanup(UPCGComponent* PCGComponent, bool bRem
 	{
 		if (UPCGComponent* Component = ComponentPtr.Get())
 		{
+			// If the component is not valid anymore, just early out
+			if (!IsValid(Component))
+			{
+				return true;
+			}
+
 			Component->PostCleanupGraph();
 
 #if WITH_EDITOR
@@ -402,8 +425,8 @@ bool UPCGSubsystem::RegisterOrUpdatePCGComponent(UPCGComponent* InComponent, boo
 	// Just make sure that we don't register components that are from the PCG Partition Actor
 	check(InComponent && InComponent->GetOwner() && !InComponent->GetOwner()->IsA<APCGPartitionActor>());
 
-	// If the component is not partitioned, don't register/update it.
-	if (!InComponent->IsPartitioned())
+	// If the component is not partitioned nor valid, don't register/update it.
+	if (!InComponent->IsPartitioned() || !IsValid(InComponent))
 	{
 		return false;
 	}
@@ -431,6 +454,17 @@ bool UPCGSubsystem::RegisterOrUpdatePCGComponent(UPCGComponent* InComponent, boo
 			FPCGComponentRef ComponentRef(InComponent, IdShared);
 			Bounds = ComponentRef.Bounds.GetBox();
 			check(Bounds.IsValid);
+
+			// If the Component is already generated, it probably mean we are in loading. The component bounds and last
+			// generated bounds should be the same.
+			// If the bounds depends on other components on the owner however, it might not be the same, because of the registration order.
+			// In this case, override the bounds by the last generated ones.
+			if (InComponent->bGenerated && !InComponent->LastGeneratedBounds.Equals(Bounds))
+			{
+				Bounds = InComponent->LastGeneratedBounds;
+				ComponentRef.Bounds = Bounds;
+			}
+
 			PCGComponentOctree.AddElement(ComponentRef);
 
 			bComponentHasChanged = true;
@@ -469,7 +503,13 @@ bool UPCGSubsystem::RegisterOrUpdatePCGComponent(UPCGComponent* InComponent, boo
 	{
 		auto ScheduleTask = [](APCGPartitionActor* PCGActor, const FBox& InIntersectedBounds, const TArray<FPCGTaskId>& TaskDependencies) { return InvalidPCGTaskId; };
 
-		PCGSubsystem::ForEachIntersectingCell(GraphExecutor, InComponent->GetWorld(), Bounds, /*bCreateActor=*/true, /*bLoadCell=*/false, /*bSave=*/false, ScheduleTask);
+		// We can't spawn actors if we are running constructions scripts, asserting when we try to get the actor with the WP API.
+		// We should never enter this if we are in a construction script. If the ensure is hit, we need to fix it.
+		UWorld* World = InComponent->GetWorld();
+		if (ensure(World && !World->bIsRunningConstructionScript))
+		{
+			PCGSubsystem::ForEachIntersectingCell(GraphExecutor, World, Bounds, /*bCreateActor=*/true, /*bLoadCell=*/false, /*bSave=*/false, ScheduleTask);
+		}
 	}
 #endif // WITH_EDITOR
 
@@ -493,7 +533,88 @@ bool UPCGSubsystem::RegisterOrUpdatePCGComponent(UPCGComponent* InComponent, boo
 	return bComponentHasChanged;
 }
 
-void UPCGSubsystem::UnregisterPCGComponent(UPCGComponent* InComponent)
+bool UPCGSubsystem::RemapPCGComponent(const UPCGComponent* OldComponent, UPCGComponent* NewComponent)
+{
+	check(OldComponent && NewComponent);
+
+	// First verification we have the old component registered
+	{
+		FReadScopeLock ReadLock(PCGVolumeOctreeLock);
+
+		FPCGComponentOctreeIDSharedRef* ElementIdPtr = ComponentToIdMap.Find(OldComponent);
+
+		if (!ElementIdPtr)
+		{
+			return false;
+		}
+	}
+
+	// If so, lock again in write and recheck if it has not been remapped already.
+	bool bBoundsChanged = false;
+	{
+		FWriteScopeLock WriteLock(PCGVolumeOctreeLock);
+
+		FPCGComponentOctreeIDSharedRef* ElementIdPtr = ComponentToIdMap.Find(OldComponent);
+
+		if (!ElementIdPtr)
+		{
+			// It was remapped, so return true there.
+			return true;
+		}
+
+		FPCGComponentOctreeIDSharedRef ElementId = *ElementIdPtr;
+
+		FPCGComponentRef ComponentRef = PCGComponentOctree.GetElementById(ElementId->Id);
+		FBox PreviousBounds = ComponentRef.Bounds.GetBox();
+		ComponentRef.Component = NewComponent;
+		ComponentRef.UpdateBounds();
+		FBox Bounds = ComponentRef.Bounds.GetBox();
+		check(Bounds.IsValid);
+
+		// If bounds changed, we need to update the mapping
+		bBoundsChanged = !Bounds.Equals(PreviousBounds);
+
+		PCGComponentOctree.RemoveElement((*ElementIdPtr)->Id);
+		PCGComponentOctree.AddElement(ComponentRef);
+		ComponentToIdMap.Remove(OldComponent);
+		ComponentToIdMap.Add(NewComponent, MoveTemp(ElementId));
+	}
+
+	// Remove it from the delayed
+	{
+		FScopeLock Lock(&DelayedComponentToUnregisterLock);
+		DelayedComponentToUnregister.Remove(OldComponent);
+	}
+
+	// Remap all previous instances
+	{
+		FWriteScopeLock WriteLock(ComponentToPartitionActorsMapLock);
+		TSet<TObjectPtr<APCGPartitionActor>>* PartitionActorsPtr = ComponentToPartitionActorsMap.Find(OldComponent);
+
+		if (PartitionActorsPtr)
+		{
+			TSet<TObjectPtr<APCGPartitionActor>> PartitionActorsToRemap = MoveTemp(*PartitionActorsPtr);
+			ComponentToPartitionActorsMap.Remove(OldComponent);
+
+			for (APCGPartitionActor* Actor : PartitionActorsToRemap)
+			{
+				Actor->RemapGraphInstance(OldComponent, NewComponent);
+			}
+
+			ComponentToPartitionActorsMap.Add(NewComponent, MoveTemp(PartitionActorsToRemap));
+		}
+	}
+
+	// And update the mapping if bounds changed
+	if (bBoundsChanged)
+	{
+		UpdateMappingPCGComponentPartitionActor(NewComponent);
+	}
+
+	return true;
+}
+
+void UPCGSubsystem::UnregisterPCGComponent(UPCGComponent* InComponent, bool bForce)
 {
 	check(InComponent);
 
@@ -509,6 +630,18 @@ void UPCGSubsystem::UnregisterPCGComponent(UPCGComponent* InComponent)
 		}
 	}
 
+	// We also need to check that our current PCG Component is not deleted while being reconstructed by a construction script.
+	// If so, it will be "re-created" at some point with the same properties.
+	// In this particular case, we don't remove the PCG component from the octree and we won't delete the mapping, but mark it to be removed
+	// at next Subsystem tick. If we call "RemapPCGComponent" before, we will re-connect everything correctly.
+	// Ignore this if we force (aka when we actually unregister the delayed one)
+	if (InComponent->IsCreatedByConstructionScript() && !bForce)
+	{
+		FScopeLock Lock(&DelayedComponentToUnregisterLock);
+		DelayedComponentToUnregister.Add(InComponent);
+		return;
+	}
+
 	// If so, lock in write and retry to find it (to avoid removing it twice)
 	{
 		FWriteScopeLock WriteLock(PCGVolumeOctreeLock);
@@ -524,19 +657,24 @@ void UPCGSubsystem::UnregisterPCGComponent(UPCGComponent* InComponent)
 		ComponentToIdMap.Remove(InComponent);
 	}
 
+	// Because of recursive component deletes actors that has components, we cannot do RemoveGraphInstance
+	// inside a lock. So copy the actors to clean up and release the lock before doing RemoveGraphInstance.
+	TSet<TObjectPtr<APCGPartitionActor>> PartitionActorsToCleanUp;
 	{
 		FWriteScopeLock WriteLock(ComponentToPartitionActorsMapLock);
 		TSet<TObjectPtr<APCGPartitionActor>>* PartitionActorsPtr = ComponentToPartitionActorsMap.Find(InComponent);
+		
 		if (PartitionActorsPtr)
 		{
-			for (APCGPartitionActor* Actor : *PartitionActorsPtr)
-			{
-				Actor->RemoveGraphInstance(InComponent);
-			}
-
+			PartitionActorsToCleanUp = MoveTemp(*PartitionActorsPtr);
 			ComponentToPartitionActorsMap.Remove(InComponent);
 		}
 	}
+
+	for (APCGPartitionActor* Actor : PartitionActorsToCleanUp)
+	{
+		Actor->RemoveGraphInstance(InComponent);
+	}
 }
 
 void UPCGSubsystem::ForAllIntersectingComponents(const FBoxCenterAndExtent& InBounds, TFunction<void(UPCGComponent*)> InFunc) const
@@ -574,7 +712,9 @@ void UPCGSubsystem::RegisterPartitionActor(APCGPartitionActor* Actor, bool bDoCo
 				if (bDoComponentMapping || Component->bGenerated)
 				{
 					TSet<TObjectPtr<APCGPartitionActor>>* PartitionActorsPtr = ComponentToPartitionActorsMap.Find(Component);
-					if (ensure(PartitionActorsPtr))
+					// In editor we might load/create partition actors while the component is registering. Because of that,
+					// the mapping might not already exists, even if the component is marked generated.
+					if (PartitionActorsPtr)
 					{
 						Actor->AddGraphInstance(Component);
 						PartitionActorsPtr->Add(Actor);
@@ -678,6 +818,8 @@ void UPCGSubsystem::UpdateMappingPCGComponentPartitionActor(UPCGComponent* InCom
 		return;
 	}
 
+	TSet<TObjectPtr<APCGPartitionActor>> RemovedActors;
+
 	{
 		FWriteScopeLock WriteLock(ComponentToPartitionActorsMapLock);
 
@@ -698,15 +840,16 @@ void UPCGSubsystem::UpdateMappingPCGComponentPartitionActor(UPCGComponent* InCom
 		});
 
 		// Find the ones that were removed
-		TSet<TObjectPtr<APCGPartitionActor>> RemovedActors = PartitionActorsPtr->Difference(NewMapping);
-
-		for (APCGPartitionActor* RemovedActor : RemovedActors)
-		{
-			RemovedActor->RemoveGraphInstance(InComponent);
-		}
+		RemovedActors = PartitionActorsPtr->Difference(NewMapping);
 
 		*PartitionActorsPtr = MoveTemp(NewMapping);
 	}
+
+	// No need to be locked to do this.
+	for (APCGPartitionActor* RemovedActor : RemovedActors)
+	{
+		RemovedActor->RemoveGraphInstance(InComponent);
+	}
 }
 
 void UPCGSubsystem::DeleteMappingPCGComponentPartitionActor(UPCGComponent* InComponent)
@@ -1037,7 +1180,7 @@ FPCGTaskId UPCGSubsystem::CleanupGraph(UPCGComponent* Component, const FBox& InB
 		UPCGComponent* Component = ComponentPtr.Get();
 		check(Component != nullptr && PCGActor != nullptr);
 
-		if (!PCGActor || !Component)
+		if (!PCGActor || !Component || !IsValid(Component))
 		{
 			return InvalidPCGTaskId;
 		}
diff --git a/Engine/Plugins/Experimental/PCG/Source/PCG/Public/PCGComponent.h b/Engine/Plugins/Experimental/PCG/Source/PCG/Public/PCGComponent.h
--- a/Engine/Plugins/Experimental/PCG/Source/PCG/Public/PCGComponent.h
+++ b/Engine/Plugins/Experimental/PCG/Source/PCG/Public/PCGComponent.h
@@ -76,6 +76,7 @@ public:
 	virtual void EndPlay(const EEndPlayReason::Type EndPlayReason) override;
 	virtual void OnComponentCreated() override;
 	virtual void OnComponentDestroyed(bool bDestroyingHierarchy) override;
+	virtual void OnRegister() override;
 
 protected:
 	friend struct FPCGComponentInstanceData;
@@ -332,6 +333,10 @@ private:
 	FPCGTaskId CurrentGenerationTask = InvalidPCGTaskId;
 	FPCGTaskId CurrentCleanupTask = InvalidPCGTaskId;
 
+#if WITH_EDITOR
+	FPCGTaskId CurrentRefreshTask = InvalidPCGTaskId;
+#endif // WITH_EDITOR
+
 	UPROPERTY(VisibleAnywhere, Transient, Category = Properties, meta = (EditCondition = false, EditConditionHides))
 	bool bIsComponentLocal = false;
 
@@ -374,12 +379,19 @@ struct FPCGComponentInstanceData : public FActorComponentInstanceData
 	GENERATED_BODY()
 public:
 	FPCGComponentInstanceData() = default;
-	explicit FPCGComponentInstanceData(const UPCGComponent* SourceComponent);
+	explicit FPCGComponentInstanceData(const UPCGComponent* InSourceComponent);
 
 protected:
 	virtual bool ContainsData() const override;
 	virtual void ApplyToComponent(UActorComponent* Component, const ECacheApplyPhase CacheApplyPhase) override;
 
+#if WITH_EDITOR
+	void DelayedRefresh(UPCGComponent* PCGComponent);
+#endif // WITH_EDITOR
+
 	UPROPERTY()
 	TArray<TObjectPtr<UPCGManagedResource>> GeneratedResources;
+
+	UPROPERTY()
+	TObjectPtr<const UPCGComponent> SourceComponent;
 };
diff --git a/Engine/Plugins/Experimental/PCG/Source/PCG/Public/PCGSubsystem.h b/Engine/Plugins/Experimental/PCG/Source/PCG/Public/PCGSubsystem.h
--- a/Engine/Plugins/Experimental/PCG/Source/PCG/Public/PCGSubsystem.h
+++ b/Engine/Plugins/Experimental/PCG/Source/PCG/Public/PCGSubsystem.h
@@ -79,8 +79,11 @@ public:
 	/** Register a new PCG Component or update it, will be added to the octree if it doesn't exists yet. Returns true if it was added/updated. Thread safe */
 	bool RegisterOrUpdatePCGComponent(UPCGComponent* InComponent, bool bDoActorMapping = true);
 
-	/** Unregister a PCG Component, will be removed from the octree. Thread safe */
-	void UnregisterPCGComponent(UPCGComponent* InComponent);
+	/** In case of BP Actors, we need to remap the old component destroyed by the construction script to the new one. Returns true if re-mapping succeeded. */
+	bool RemapPCGComponent(const UPCGComponent* OldComponent, UPCGComponent* NewComponent);
+
+	/** Unregister a PCG Component, will be removed from the octree. Can force it, if we have a delayed unregister. Thread safe */
+	void UnregisterPCGComponent(UPCGComponent* InComponent, bool bForce = false);
 
 	/** Register a new Partition actor, will be added to a map and will query all intersecting volume to bind to them if asked. Thread safe */
 	void RegisterPartitionActor(APCGPartitionActor* InActor, bool bDoComponentMapping = true);
@@ -170,4 +173,7 @@ private:
 
 	TMap<TObjectPtr<const UPCGComponent>, TSet<TObjectPtr<APCGPartitionActor>>> ComponentToPartitionActorsMap;
 	mutable FRWLock ComponentToPartitionActorsMapLock;
+
+	TSet<TObjectPtr<UPCGComponent>> DelayedComponentToUnregister;
+	mutable FCriticalSection DelayedComponentToUnregisterLock;
 };
\ No newline at end of file
