FName -> FSoftObjectPath refactoring.

diff --git a/Engine/Plugins/Animation/MLDeformer/MLDeformerFramework/Source/MLDeformerFrameworkEditor/Private/MLDeformerVizSettingsDetails.cpp b/Engine/Plugins/Animation/MLDeformer/MLDeformerFramework/Source/MLDeformerFrameworkEditor/Private/MLDeformerVizSettingsDetails.cpp
--- a/Engine/Plugins/Animation/MLDeformer/MLDeformerFramework/Source/MLDeformerFrameworkEditor/Private/MLDeformerVizSettingsDetails.cpp
+++ b/Engine/Plugins/Animation/MLDeformer/MLDeformerFramework/Source/MLDeformerFrameworkEditor/Private/MLDeformerVizSettingsDetails.cpp
@@ -206,7 +206,7 @@ namespace UE::MLDeformer
 		{
 			// Check if we already assigned the default asset.
 			const FAssetData CurrentGraphAssetData(CurrentGraph);
-			const FString CurrentPath = CurrentGraphAssetData.ObjectPath.ToString();
+			const FString CurrentPath = CurrentGraphAssetData.GetObjectPathString();
 			const FString DefaultPath = Model->GetDefaultDeformerGraphAssetPath();
 			return (DefaultPath != CurrentPath);
 		}
diff --git a/Engine/Plugins/Editor/AssetRegistryExport/Source/Private/AssetRegistryExportCommandlet.cpp b/Engine/Plugins/Editor/AssetRegistryExport/Source/Private/AssetRegistryExportCommandlet.cpp
--- a/Engine/Plugins/Editor/AssetRegistryExport/Source/Private/AssetRegistryExportCommandlet.cpp
+++ b/Engine/Plugins/Editor/AssetRegistryExport/Source/Private/AssetRegistryExportCommandlet.cpp
@@ -242,7 +242,7 @@ int32 UAssetRegistryExportCommandlet::Main(const FString& CmdLineParams)
 	bool bGotCompressedSize = false;
 	AssetRegistry.EnumerateAllAssets(TSet<FName>(), [&](const FAssetData& AssetData)
 	{
-		FString AssetObjectPath = AssetData.ObjectPath.ToString();
+		FString AssetObjectPath = AssetData.GetObjectPathString();
 
 		// Add to the 'Assets' table.
 		{
diff --git a/Engine/Plugins/Editor/GuidedTutorials/Source/IntroTutorials/Private/ClassTypeActions_EditorTutorial.cpp b/Engine/Plugins/Editor/GuidedTutorials/Source/IntroTutorials/Private/ClassTypeActions_EditorTutorial.cpp
--- a/Engine/Plugins/Editor/GuidedTutorials/Source/IntroTutorials/Private/ClassTypeActions_EditorTutorial.cpp
+++ b/Engine/Plugins/Editor/GuidedTutorials/Source/IntroTutorials/Private/ClassTypeActions_EditorTutorial.cpp
@@ -21,7 +21,7 @@ UClass* FClassTypeActions_EditorTutorial::GetSupportedClass() const
 
 TSharedPtr<SWidget> FClassTypeActions_EditorTutorial::GetThumbnailOverlay(const FAssetData& AssetData) const
 {
-	const FString FullTutorialAssetPath = AssetData.ObjectPath.ToString();
+	const FString FullTutorialAssetPath = AssetData.GetObjectPathString();
 	auto OnLaunchTutorialClicked = [FullTutorialAssetPath]() -> FReply
 	{
 		if (IIntroTutorials::IsAvailable())
diff --git a/Engine/Plugins/Editor/GuidedTutorials/Source/IntroTutorials/Private/STutorialsBrowser.cpp b/Engine/Plugins/Editor/GuidedTutorials/Source/IntroTutorials/Private/STutorialsBrowser.cpp
--- a/Engine/Plugins/Editor/GuidedTutorials/Source/IntroTutorials/Private/STutorialsBrowser.cpp
+++ b/Engine/Plugins/Editor/GuidedTutorials/Source/IntroTutorials/Private/STutorialsBrowser.cpp
@@ -768,7 +768,7 @@ void STutorialsBrowser::RebuildTutorials(TSharedPtr<FTutorialListEntry_Category>
 
 	for (const auto& TutorialAsset : AssetData)
 	{
-		UBlueprint* Blueprint = LoadObject<UBlueprint>(nullptr, *TutorialAsset.ObjectPath.ToString());
+		UBlueprint* Blueprint = LoadObject<UBlueprint>(nullptr, *TutorialAsset.GetObjectPathString());
 		if (Blueprint && Blueprint->GeneratedClass && Blueprint->BlueprintType == BPTYPE_Normal)
 		{
 			UEditorTutorial* Tutorial = NewObject<UEditorTutorial>(GetTransientPackage(), Blueprint->GeneratedClass);
diff --git a/Engine/Plugins/Enterprise/LidarPointCloud/Source/LidarPointCloudRuntime/Private/LidarPointCloud.cpp b/Engine/Plugins/Enterprise/LidarPointCloud/Source/LidarPointCloudRuntime/Private/LidarPointCloud.cpp
--- a/Engine/Plugins/Enterprise/LidarPointCloud/Source/LidarPointCloudRuntime/Private/LidarPointCloud.cpp
+++ b/Engine/Plugins/Enterprise/LidarPointCloud/Source/LidarPointCloudRuntime/Private/LidarPointCloud.cpp
@@ -173,7 +173,7 @@ public:
 				Info.Hyperlink = FSimpleDelegate::CreateLambda([SoftObjectPath] {
 					// Select the cloud in Content Browser when the hyperlink is clicked
 					TArray<FAssetData> AssetData;
-					AssetData.Add(FModuleManager::LoadModuleChecked<FAssetRegistryModule>("AssetRegistry").Get().GetAssetByObjectPath(SoftObjectPath.GetAssetPathName()));
+					AssetData.Add(FModuleManager::LoadModuleChecked<FAssetRegistryModule>("AssetRegistry").Get().GetAssetByObjectPath(SoftObjectPath.GetWithoutSubPath()));
 					FModuleManager::LoadModuleChecked<FContentBrowserModule>("ContentBrowser").Get().SyncBrowserToAssets(AssetData);
 					});
 				Info.HyperlinkText = FText::FromString(FPaths::GetBaseFilename(SoftObjectPath.ToString()));
diff --git a/Engine/Plugins/Experimental/GameFeatures/Source/GameFeatures/Private/GameFeaturesSubsystem.cpp b/Engine/Plugins/Experimental/GameFeatures/Source/GameFeatures/Private/GameFeaturesSubsystem.cpp
--- a/Engine/Plugins/Experimental/GameFeatures/Source/GameFeatures/Private/GameFeaturesSubsystem.cpp
+++ b/Engine/Plugins/Experimental/GameFeatures/Source/GameFeatures/Private/GameFeaturesSubsystem.cpp
@@ -388,7 +388,7 @@ void UGameFeaturesSubsystem::AddGameFeatureToAssetManager(const UGameFeatureData
 	LocalAssetManager.RegisterSpecificPrimaryAsset(GameFeatureToAdd->GetPrimaryAssetId(), FAssetData(GameFeatureToAdd));
 #else
 	// In non-editor, the asset bundle data is compiled out, so it must be gathered from the asset registry instead
-	LocalAssetManager.RegisterSpecificPrimaryAsset(GameFeatureToAdd->GetPrimaryAssetId(), LocalAssetRegistry.GetAssetByObjectPath(FName(*GameFeatureToAdd->GetPathName()), true));
+	LocalAssetManager.RegisterSpecificPrimaryAsset(GameFeatureToAdd->GetPrimaryAssetId(), LocalAssetRegistry.GetAssetByObjectPath(FSoftObjectPath(GameFeatureToAdd), true));
 #endif // WITH_EDITOR
 
 	// @TODO: HACK - There is no guarantee that the plugin mount point was added before inte initial asset scan.
diff --git a/Engine/Plugins/Experimental/PCG/Source/PCGEditor/Private/PCGEditorGraphSchema.cpp b/Engine/Plugins/Experimental/PCG/Source/PCGEditor/Private/PCGEditorGraphSchema.cpp
--- a/Engine/Plugins/Experimental/PCG/Source/PCGEditor/Private/PCGEditorGraphSchema.cpp
+++ b/Engine/Plugins/Experimental/PCG/Source/PCGEditor/Private/PCGEditorGraphSchema.cpp
@@ -274,7 +274,9 @@ void UPCGEditorGraphSchema::GetSubgraphElementActions(FGraphActionMenuBuilder& A
 			const FText Description = AssetData.GetTagValueRef<FText>(TEXT("Description"));
 
 			TSharedPtr<FPCGEditorGraphSchemaAction_NewSubgraphElement> NewSubgraphAction(new FPCGEditorGraphSchemaAction_NewSubgraphElement(Category, MenuDesc, Description, 0));
+PRAGMA_DISABLE_DEPRECATION_WARNINGS
 			NewSubgraphAction->SubgraphObjectPath = AssetData.ObjectPath;
+PRAGMA_ENABLE_DEPRECATION_WARNINGS
 			ActionMenuBuilder.AddAction(NewSubgraphAction);
 		}
 	}
@@ -304,7 +306,9 @@ void UPCGEditorGraphSchema::DroppedAssetsOnGraph(const TArray<FAssetData>& Asset
 			if (Asset->IsA<UPCGGraph>())
 			{
 				FPCGEditorGraphSchemaAction_NewSubgraphElement NewSubgraphAction;
+PRAGMA_DISABLE_DEPRECATION_WARNINGS
 				NewSubgraphAction.SubgraphObjectPath = AssetData.ObjectPath;
+PRAGMA_ENABLE_DEPRECATION_WARNINGS
 				NewSubgraphAction.PerformAction(Graph, NullFromPin, GraphPositionOffset);
 				GraphPositionOffset.Y += PositionOffsetIncrementY;
 			}
diff --git a/Engine/Plugins/Experimental/Shotgrid/Source/Shotgrid/Private/ShotgridEngine.cpp b/Engine/Plugins/Experimental/Shotgrid/Source/Shotgrid/Private/ShotgridEngine.cpp
--- a/Engine/Plugins/Experimental/Shotgrid/Source/Shotgrid/Private/ShotgridEngine.cpp
+++ b/Engine/Plugins/Experimental/Shotgrid/Source/Shotgrid/Private/ShotgridEngine.cpp
@@ -62,7 +62,7 @@ void UShotgridEngine::SetSelection(const TArray<FAssetData>* InSelectedAssets, c
 			{
 				if (IsValid(Asset) && Asset->IsAsset())
 				{
-					FAssetData AssetData = AssetRegistry.GetAssetByObjectPath(FName(*Asset->GetPathName()));
+					FAssetData AssetData = AssetRegistry.GetAssetByObjectPath(FSoftObjectPath(Asset));
 					AllReferencedAssets.AddUnique(MoveTemp(AssetData));
 				}
 			}
diff --git a/Engine/Plugins/Experimental/VirtualCamera/Source/VirtualCamera/Private/LevelSequencePlaybackController.cpp b/Engine/Plugins/Experimental/VirtualCamera/Source/VirtualCamera/Private/LevelSequencePlaybackController.cpp
--- a/Engine/Plugins/Experimental/VirtualCamera/Source/VirtualCamera/Private/LevelSequencePlaybackController.cpp
+++ b/Engine/Plugins/Experimental/VirtualCamera/Source/VirtualCamera/Private/LevelSequencePlaybackController.cpp
@@ -56,7 +56,7 @@ void ULevelSequencePlaybackController::GetLevelSequences(TArray<FLevelSequenceDa
 		);
 
 		FLevelSequenceData NewSequenceData = FLevelSequenceData(
-			LevelSequence.ObjectPath.ToString(),
+			LevelSequence.GetObjectPathString(),
 			LevelSequence.AssetName.ToString(),
 			FileManager.GetTimeStamp(*LevelSequenceFile)
 		);
diff --git a/Engine/Plugins/Runtime/NetworkPredictionExtras/Source/NetworkPredictionExtras/Private/MockRootMotionSourceObject.cpp b/Engine/Plugins/Runtime/NetworkPredictionExtras/Source/NetworkPredictionExtras/Private/MockRootMotionSourceObject.cpp
--- a/Engine/Plugins/Runtime/NetworkPredictionExtras/Source/NetworkPredictionExtras/Private/MockRootMotionSourceObject.cpp
+++ b/Engine/Plugins/Runtime/NetworkPredictionExtras/Source/NetworkPredictionExtras/Private/MockRootMotionSourceObject.cpp
@@ -99,7 +99,7 @@ void UMockRootMotionSourceClassMap::BuildClassMap_Internal()
 
 	for (FAssetData& AssetData : AssetDataList)
 	{
-		FString ClassPathString = AssetData.ObjectPath.ToString() + TEXT("_c");
+		FString ClassPathString = AssetData.GetObjectPathString() + TEXT("_c");
 		FSoftObjectPath SoftObjPath(ClassPathString);
 		SourceList.Add(TSoftClassPtr<UMockRootMotionSource>(SoftObjPath));
 	}
