PCG:
*Added GC ref to PCG editor graph object being edited.

diff --git a/Engine/Plugins/Experimental/PCG/Source/PCGEditor/Private/PCGEditor.cpp b/Engine/Plugins/Experimental/PCG/Source/PCGEditor/Private/PCGEditor.cpp
--- a/Engine/Plugins/Experimental/PCG/Source/PCGEditor/Private/PCGEditor.cpp
+++ b/Engine/Plugins/Experimental/PCG/Source/PCGEditor/Private/PCGEditor.cpp
@@ -232,6 +232,11 @@ void FPCGEditor::UnregisterTabSpawners(const TSharedRef<class FTabManager>& InTa
 	FAssetEditorToolkit::UnregisterTabSpawners(InTabManager);
 }
 
+void FPCGEditor::AddReferencedObjects(FReferenceCollector& Collector)
+{
+	Collector.AddReferencedObject(PCGGraphBeingEdited);
+}
+
 bool FPCGEditor::MatchesContext(const FTransactionContext& InContext, const TArray<TPair<UObject*, FTransactionObjectEvent>>& TransactionObjectContexts) const
 {
 	return InContext.Context == FPCGEditorCommon::ContextIdentifier;
diff --git a/Engine/Plugins/Experimental/PCG/Source/PCGEditor/Private/PCGEditor.h b/Engine/Plugins/Experimental/PCG/Source/PCGEditor/Private/PCGEditor.h
--- a/Engine/Plugins/Experimental/PCG/Source/PCGEditor/Private/PCGEditor.h
+++ b/Engine/Plugins/Experimental/PCG/Source/PCGEditor/Private/PCGEditor.h
@@ -25,7 +25,7 @@ class UPCGNode;
 DECLARE_MULTICAST_DELEGATE_OneParam(FOnDebugObjectChanged, UPCGComponent*);
 DECLARE_MULTICAST_DELEGATE_OneParam(FOnInspectedNodeChanged, UPCGNode*);
 
-class FPCGEditor : public FAssetEditorToolkit, public FSelfRegisteringEditorUndoClient
+class FPCGEditor : public FAssetEditorToolkit, public FGCObject, public FSelfRegisteringEditorUndoClient
 {
 public:
 	/** Edits the specified PCGGraph */
@@ -54,6 +54,14 @@ public:
 	virtual void UnregisterTabSpawners(const TSharedRef<class FTabManager>& InTabManager) override;
 	// ~End IToolkit interface
 
+	// ~Begin FGCObject interface
+	virtual void AddReferencedObjects(FReferenceCollector& Collector) override;
+	virtual FString GetReferencerName() const override
+	{
+		return TEXT("FPCGEditor");
+	}
+	// ~End FGCObject interface
+	
 	// ~Begin FEditorUndoClient interface
 	virtual bool MatchesContext(const FTransactionContext& InContext, const TArray<TPair<UObject*, FTransactionObjectEvent>>& TransactionObjectContexts) const override;
 	virtual void PostUndo(bool bSuccess) override;
