PCG:
*Fixed an assert when creating attribute list view columns for metadata without a proper name set. We will Now stringify the metadata type enum and use that as a column name instead.

diff --git a/Engine/Plugins/Experimental/PCG/Source/PCGEditor/Private/SPCGEditorGraphAttributeListView.cpp b/Engine/Plugins/Experimental/PCG/Source/PCGEditor/Private/SPCGEditorGraphAttributeListView.cpp
--- a/Engine/Plugins/Experimental/PCG/Source/PCGEditor/Private/SPCGEditorGraphAttributeListView.cpp
+++ b/Engine/Plugins/Experimental/PCG/Source/PCGEditor/Private/SPCGEditorGraphAttributeListView.cpp
@@ -662,6 +662,13 @@ void SPCGEditorGraphAttributeListView::GenerateColumnsFromMetadata(const UPCGMet
 	for (int32 I = 0; I < AttributeNames.Num(); I++)
 	{
 		const FName& AttributeName = AttributeNames[I];
+		FName ColumnName = AttributeName;
+
+		if (ColumnName == NAME_None)
+		{
+			const FString TypeString = UEnum::GetDisplayValueAsText(AttributeTypes[I]).ToString();
+			ColumnName = *TypeString;
+		}
 
 		switch (AttributeTypes[I])
 		{
@@ -673,36 +680,36 @@ void SPCGEditorGraphAttributeListView::GenerateColumnsFromMetadata(const UPCGMet
 		case EPCGMetadataTypes::String:
 		case EPCGMetadataTypes::Name:
 			{
-				AddMetadataColumn(AttributeName);
+				AddMetadataColumn(ColumnName, AttributeName);
 				break;
 			}
 		case EPCGMetadataTypes::Vector:
 			{
-				AddMetadataColumn(AttributeName, 0, TEXT("_X"));
-				AddMetadataColumn(AttributeName, 1, TEXT("_Y"));
-				AddMetadataColumn(AttributeName, 2, TEXT("_Z"));
+				AddMetadataColumn(ColumnName, AttributeName, 0, TEXT("_X"));
+				AddMetadataColumn(ColumnName, AttributeName, 1, TEXT("_Y"));
+				AddMetadataColumn(ColumnName, AttributeName, 2, TEXT("_Z"));
 				break;
 			}
 		case EPCGMetadataTypes::Vector4:
 		case EPCGMetadataTypes::Quaternion:
 			{
-				AddMetadataColumn(AttributeName, 0, TEXT("_X"));
-				AddMetadataColumn(AttributeName, 1, TEXT("_Y"));
-				AddMetadataColumn(AttributeName, 2, TEXT("_Z"));
-				AddMetadataColumn(AttributeName, 3, TEXT("_W"));
+				AddMetadataColumn(ColumnName, AttributeName, 0, TEXT("_X"));
+				AddMetadataColumn(ColumnName, AttributeName, 1, TEXT("_Y"));
+				AddMetadataColumn(ColumnName, AttributeName, 2, TEXT("_Z"));
+				AddMetadataColumn(ColumnName, AttributeName, 3, TEXT("_W"));
 				break;
 			}
 		case EPCGMetadataTypes::Transform:
 			{
-				AddMetadataColumn(AttributeName, 0, TEXT("_tX"));
-				AddMetadataColumn(AttributeName, 1, TEXT("_tY"));
-				AddMetadataColumn(AttributeName, 2, TEXT("_tZ"));
-				AddMetadataColumn(AttributeName, 3, TEXT("_rX"));
-				AddMetadataColumn(AttributeName, 4, TEXT("_rY"));
-				AddMetadataColumn(AttributeName, 5, TEXT("_rZ"));
-				AddMetadataColumn(AttributeName, 6, TEXT("_sX"));
-				AddMetadataColumn(AttributeName, 7, TEXT("_sY"));
-				AddMetadataColumn(AttributeName, 8, TEXT("_sZ"));
+				AddMetadataColumn(ColumnName, AttributeName, 0, TEXT("_tX"));
+				AddMetadataColumn(ColumnName, AttributeName, 1, TEXT("_tY"));
+				AddMetadataColumn(ColumnName, AttributeName, 2, TEXT("_tZ"));
+				AddMetadataColumn(ColumnName, AttributeName, 3, TEXT("_rX"));
+				AddMetadataColumn(ColumnName, AttributeName, 4, TEXT("_rY"));
+				AddMetadataColumn(ColumnName, AttributeName, 5, TEXT("_rZ"));
+				AddMetadataColumn(ColumnName, AttributeName, 6, TEXT("_sX"));
+				AddMetadataColumn(ColumnName, AttributeName, 7, TEXT("_sY"));
+				AddMetadataColumn(ColumnName, AttributeName, 8, TEXT("_sZ"));
 				break;
 			}
 		default:
@@ -796,7 +803,7 @@ void SPCGEditorGraphAttributeListView::RemovePointDataColumns()
 	RemoveColumn(PCGEditorGraphAttributeListView::NAME_PointSeed);
 }
 
-void SPCGEditorGraphAttributeListView::AddMetadataColumn(const FName& InColumnId, const int8 InValueIndex, const TCHAR* PostFix)
+void SPCGEditorGraphAttributeListView::AddMetadataColumn(const FName& InColumnId, const FName& InMetadataId, const int8 InValueIndex, const TCHAR* PostFix)
 {
 	FString ColumnIdString = InColumnId.ToString();
 
@@ -808,7 +815,7 @@ void SPCGEditorGraphAttributeListView::AddMetadataColumn(const FName& InColumnId
 	const FName ColumnId(ColumnIdString);
 
 	FPCGMetadataInfo MetadataInfo;
-	MetadataInfo.MetadataId = InColumnId;
+	MetadataInfo.MetadataId = InMetadataId;
 	MetadataInfo.Index = InValueIndex;
 	MetadataInfos.Add(ColumnId, MetadataInfo);
 
diff --git a/Engine/Plugins/Experimental/PCG/Source/PCGEditor/Private/SPCGEditorGraphAttributeListView.h b/Engine/Plugins/Experimental/PCG/Source/PCGEditor/Private/SPCGEditorGraphAttributeListView.h
--- a/Engine/Plugins/Experimental/PCG/Source/PCGEditor/Private/SPCGEditorGraphAttributeListView.h
+++ b/Engine/Plugins/Experimental/PCG/Source/PCGEditor/Private/SPCGEditorGraphAttributeListView.h
@@ -88,7 +88,7 @@ private:
 	void AddPointDataColumns();
 	void RemovePointDataColumns();
 
-	void AddMetadataColumn(const FName& InColumnId, const int8 InValueIndex = INDEX_NONE, const TCHAR* PostFix = nullptr);
+	void AddMetadataColumn(const FName& InColumnId, const FName& InMetadataId, const int8 InValueIndex = INDEX_NONE, const TCHAR* PostFix = nullptr);
 	void RemoveMetadataColumns();
 
 	/** Pointer back to the PCG editor that owns us */
