Nativized texture sampler node
Improved texture sampler behavior

diff --git a/Engine/Plugins/Experimental/PCG/Source/PCG/Private/Data/PCGTextureData.cpp b/Engine/Plugins/Experimental/PCG/Source/PCG/Private/Data/PCGTextureData.cpp
--- a/Engine/Plugins/Experimental/PCG/Source/PCG/Private/Data/PCGTextureData.cpp
+++ b/Engine/Plugins/Experimental/PCG/Source/PCG/Private/Data/PCGTextureData.cpp
@@ -9,43 +9,54 @@
 namespace PCGTextureSampling
 {
 	template<typename ValueType>
-	ValueType Sample(const FVector2D& InPosition, const FBox2D& InSurface, const UPCGBaseTextureData* InTextureData, int32 Width, int32 Height, TFunctionRef<ValueType(int32 Index)> SamplingFunction)
+	bool Sample(const FVector2D& InPosition, 
+		const FBox2D& InSurface, 
+		const UPCGBaseTextureData* InTextureData, 
+		int32 Width, 
+		int32 Height, 
+		ValueType& SampledValue,
+		TFunctionRef<ValueType(int32 Index)> SamplingFunction)
 	{
 		check(Width > 0 && Height > 0);
-		if (Width <= 0 || Height <= 0)
+		if (Width <= 0 || Height <= 0 || InSurface.GetSize().SquaredLength() <= 0)
 		{
-			return ValueType{};
+			return false;
 		}
 
 		check(InTextureData);
 
+		const FVector2D LocalSpacePos = (InPosition - InSurface.Min) / InSurface.GetSize();
 		FVector2D Pos = FVector2D::ZeroVector;
-		if (InTextureData->bStretchToFit)
+		if (!InTextureData->bUseAdvancedTiling)
 		{
 			// TODO: There seems to be a bias issue here, as the bounds size are not in the same space as the texels.
-			Pos = (InPosition - InSurface.Min) * FVector2D(Width, Height) / InSurface.GetSize();
+			Pos = LocalSpacePos * FVector2D(Width, Height);
 		}
 		else
 		{
-			const float TextureSizeX = InTextureData->TexelSize * Width;
-			const float TextureSizeY = InTextureData->TexelSize * Height;
-
-			const FRotator Rotation = FRotator(0.f, -InTextureData->Rotation, 0.f);
-			const FVector Translation = FVector(-InTextureData->XOffset * TextureSizeX, -InTextureData->YOffset * TextureSizeY, 0.f);
-			const FTransform Transform(Rotation, Translation);
-			const FVector2D SamplePosition = FVector2D(Transform.TransformPosition(FVector(InPosition, 0.f)));
-
-			float X = FMath::Clamp(static_cast<float>(FMath::Fmod(SamplePosition.X, TextureSizeX)) / TextureSizeX, -1.f, 1.f);
-			if (X < 0.f)
+			// Conceptually, we are building "tiles" in texture space with the origin being in the middle of the [0, 0] tile.
+			// The offset is given in a ratio of [0, 1], applied "before" scaling & rotation.
+			// Rotation is done around the center given, where the center is (0.5, 0.5) + offset
+			// Scaling controls the horizon of tiles, and the tile selection is done through min-max bounds, in tile space,
+			// with the origin tile being from -0.5 to 0.5.
+			const FRotator Rotation = FRotator(0.0, -InTextureData->Rotation, 0.0);
+			FVector Scale = FVector(InTextureData->Tiling, 1.0);
+			Scale.X = ((FMath::Abs(Scale.X) > SMALL_NUMBER) ? (1.0 / Scale.X) : 0.0);
+			Scale.Y = ((FMath::Abs(Scale.Y) > SMALL_NUMBER) ? (1.0 / Scale.Y) : 0.0);
+			const FVector Translation = FVector(0.5 + InTextureData->CenterOffset.X, 0.5 + InTextureData->CenterOffset.Y, 0);
+
+			FTransform Transform = FTransform(Rotation, Translation, Scale);
+			
+			// Transform to tile-space
+			const FVector2D SamplePosition = FVector2D(Transform.InverseTransformPosition(FVector(LocalSpacePos, 0.f)));
+
+			if (InTextureData->bUseTileBounds && !InTextureData->TileBounds.IsInsideOrOn(SamplePosition))
 			{
-				X = 1.f + X;
+				return false;
 			}
 
-			float Y = FMath::Clamp(static_cast<float>(FMath::Fmod(SamplePosition.Y, TextureSizeY)) / TextureSizeY, -1.f, 1.f);
-			if (Y < 0.f)
-			{
-				Y = 1.f + Y;
-			}
+			FVector::FReal X = FMath::Frac(SamplePosition.X + 0.5);
+			FVector::FReal Y = FMath::Frac(SamplePosition.Y + 0.5);
 
 			X *= Width;
 			Y *= Height;
@@ -83,8 +94,8 @@ namespace PCGTextureSampling
 		ValueType SampleX0Y1 = SamplingFunction(X0 + Y1 * Width);
 		ValueType SampleX1Y1 = SamplingFunction(X1 + Y1 * Width);
 
-		ValueType BilinearInterpolation = FMath::BiLerp(SampleX0Y0, SampleX1Y0, SampleX0Y1, SampleX1Y1, Pos.X - X0, Pos.Y - Y0);
-		return BilinearInterpolation;
+		SampledValue = FMath::BiLerp(SampleX0Y0, SampleX1Y0, SampleX0Y1, SampleX1Y1, Pos.X - X0, Pos.Y - Y0);
+		return true;
 	}
 
 	float SampleFloatChannel(const FLinearColor& InColor, EPCGTextureColorChannel ColorChannel)
@@ -136,11 +147,17 @@ bool UPCGBaseTextureData::SamplePoint(const FTransform& InTransform, const FBox&
 	FVector2D Position2D(PointPositionInLocalSpace.X, PointPositionInLocalSpace.Y);
 	FBox2D Surface(FVector2D(-1.0f, -1.0f), FVector2D(1.0f, 1.0f));
 
-	FLinearColor Color = PCGTextureSampling::Sample<FLinearColor>(Position2D, Surface, this, Width, Height, [this](int32 Index) { return ColorData[Index]; });
-	OutPoint.Color = Color;
-	OutPoint.Density = ((DensityFunction == EPCGTextureDensityFunction::Ignore) ? 1.0f : PCGTextureSampling::SampleFloatChannel(Color, ColorChannel));
-
-	return OutPoint.Density > 0;
+	FLinearColor Color = FLinearColor(EForceInit::ForceInit);
+	if (PCGTextureSampling::Sample<FLinearColor>(Position2D, Surface, this, Width, Height, Color, [this](int32 Index) { return ColorData[Index]; }))
+	{
+		OutPoint.Color = Color;
+		OutPoint.Density = ((DensityFunction == EPCGTextureDensityFunction::Ignore) ? 1.0f : PCGTextureSampling::SampleFloatChannel(Color, ColorChannel));
+		return OutPoint.Density > 0;
+	}
+	else
+	{
+		return false;
+	}
 }
 
 const UPCGPointData* UPCGBaseTextureData::CreatePointData(FPCGContext* Context) const
@@ -163,39 +180,53 @@ const UPCGPointData* UPCGBaseTextureData::CreatePointData(FPCGContext* Context)
 
 	TArray<FPCGPoint>& Points = Data->GetMutablePoints();
 
-	// TODO: There's a bias issue here where we should correct by a 0.5 unit...
-	const FVector::FReal XScale = LocalSurfaceBounds.GetSize().X / Width;
-	const FVector::FReal YScale = LocalSurfaceBounds.GetSize().Y / Height;
-	const FVector2D Bias = LocalSurfaceBounds.Min;
+	// Map target texel size to the current physical size of the texture data.
+	const FVector::FReal XSize = 2.0 * Transform.GetScale3D().X;
+	const FVector::FReal YSize = 2.0 * Transform.GetScale3D().Y;
 
-	FPCGAsync::AsyncPointProcessing(Context, Width * Height, Points, [this, XScale, YScale, Bias](int32 Index, FPCGPoint& OutPoint)
+	const int32 XCount = FMath::Floor(XSize / TexelSize);
+	const int32 YCount = FMath::Floor(YSize / TexelSize);
+	const int32 PointCount = XCount * YCount;
+
+	if (PointCount <= 0)
 	{
-		const int X = Index % Width;
-		const int Y = Index / Width;
+		UE_LOG(LogPCG, Warning, TEXT("Texture data has a texel size larger than its data - will return empty data"));
+		return Data;
+	}
 
-		const float Density = PCGTextureSampling::SampleFloatChannel(ColorData[X + Y * Width], ColorChannel);
+	const FBox2D Surface(FVector2D(-1.0f, -1.0f), FVector2D(1.0f, 1.0f));
 
+	FPCGAsync::AsyncPointProcessing(Context, PointCount, Points, [this, XCount, YCount, &Surface](int32 Index, FPCGPoint& OutPoint)
+	{
+		const int X = (Index % XCount);
+		const int Y = (Index / YCount);
+
+		// TODO: we should have a 0.5 bias here
+		FVector2D LocalCoordinate((2.0 * X + 0.5) / XCount - 1.0, (2.0 * Y + 0.5) / YCount - 1.0);
+		FLinearColor Color = FLinearColor(EForceInit::ForceInit);
+
+		if (PCGTextureSampling::Sample<FLinearColor>(LocalCoordinate, Surface, this, Width, Height, Color, [this](int32 Index) { return ColorData[Index]; }))
+		{
+			const float Density = ((DensityFunction == EPCGTextureDensityFunction::Ignore) ? 1.0f : PCGTextureSampling::SampleFloatChannel(Color, ColorChannel));
 #if WITH_EDITORONLY_DATA
-		if (Density > 0 || bKeepZeroDensityPoints)
+			if (Density > 0 || bKeepZeroDensityPoints)
 #else
-		if (Density > 0)
+			if (Density > 0)
 #endif
-		{
-			FVector LocalPosition(X * XScale + Bias.X, Y * YScale + Bias.Y, 0);
-			OutPoint = FPCGPoint(FTransform(Transform.TransformPosition(LocalPosition)),
-				Density,
-				PCGHelpers::ComputeSeed(X, Y));
+			{
+				FVector LocalPosition(LocalCoordinate, 0);
+				OutPoint = FPCGPoint(FTransform(Transform.TransformPosition(LocalPosition)),
+					Density,
+					PCGHelpers::ComputeSeed(X, Y));
 
-			const FVector TransformScale = Transform.GetScale3D();
-			OutPoint.SetExtents(FVector(TransformScale.X * XScale / 2.0, TransformScale.Y * YScale / 2.0, 1.0));
-			OutPoint.Color = ColorData[X + Y * Width];
+				OutPoint.SetExtents(FVector(TexelSize / 2.0));
+				OutPoint.Color = Color;
 
-			return true;
-		}
-		else
-		{
-			return false;
+				return true;
+			}
 		}
+
+		return false;
 	});
 
 	return Data;
diff --git a/Engine/Plugins/Experimental/PCG/Source/PCG/Private/Elements/PCGTextureSampler.cpp b/Engine/Plugins/Experimental/PCG/Source/PCG/Private/Elements/PCGTextureSampler.cpp
new file mode 100644
--- /dev/null
+++ b/Engine/Plugins/Experimental/PCG/Source/PCG/Private/Elements/PCGTextureSampler.cpp
@@ -0,0 +1,90 @@
+// Copyright Epic Games, Inc. All Rights Reserved.
+
+#include "Elements/PCGTextureSampler.h"
+
+#include "PCGComponent.h"
+#include "PCGHelpers.h"
+#include "Data/PCGSpatialData.h"
+#include "Helpers/PCGBlueprintHelpers.h"
+#include "Helpers/PCGSettingsHelpers.h"
+
+FPCGElementPtr UPCGTextureSamplerSettings::CreateElement() const
+{
+	return MakeShared<FPCGTextureSamplerElement>();
+}
+
+TArray<FPCGPinProperties> UPCGTextureSamplerSettings::InputPinProperties() const
+{
+	TArray<FPCGPinProperties> PinProperties;
+	PinProperties.Emplace(PCGPinConstants::DefaultParamsLabel, EPCGDataType::Param);
+
+	return PinProperties;
+}
+
+bool FPCGTextureSamplerElement::IsCacheable(const UPCGSettings* InSettings) const
+{
+	// It is currently possible to cache the texture sampler only if using an absolute transform
+	// since otherwise it depends on the source component.
+	// However, more tricky here is that the bUseAbsoluteTransform can be overriden, in which case we can't take a decision
+	return false;
+}
+
+bool FPCGTextureSamplerElement::ExecuteInternal(FPCGContext* Context) const
+{
+	TRACE_CPUPROFILER_EVENT_SCOPE(FPCGTextureSamplerElement::Execute);
+
+	const UPCGTextureSamplerSettings* Settings = Context->GetInputSettings<UPCGTextureSamplerSettings>();
+	check(Settings);
+
+	UPCGParamData* Params = Context->InputData.GetParams();
+
+	FTransform Transform = PCG_GET_OVERRIDEN_VALUE(Settings, Transform, Params);
+	bool bUseAbsoluteTransform = PCG_GET_OVERRIDEN_VALUE(Settings, bUseAbsoluteTransform, Params);
+	//TObjectPtr<UTexture2D> PCG_OVERRIDEABLE_VALUE(Texture);
+	TObjectPtr<UTexture2D> Texture = Settings->Texture;
+	EPCGTextureDensityFunction DensityFunction = PCG_GET_OVERRIDEN_VALUE(Settings, DensityFunction, Params);
+	EPCGTextureColorChannel ColorChannel = PCG_GET_OVERRIDEN_VALUE(Settings, ColorChannel, Params);
+	float TexelSize = PCG_GET_OVERRIDEN_VALUE(Settings, TexelSize, Params);
+	bool bUseAdvancedTiling = PCG_GET_OVERRIDEN_VALUE(Settings, bUseAdvancedTiling, Params);
+	FVector2D Tiling = PCG_GET_OVERRIDEN_VALUE(Settings, Tiling, Params);
+	FVector2D CenterOffset = PCG_GET_OVERRIDEN_VALUE(Settings, CenterOffset, Params);
+	float Rotation = PCG_GET_OVERRIDEN_VALUE(Settings, Rotation, Params);
+	bool bUseTileBounds = PCG_GET_OVERRIDEN_VALUE(Settings, bUseTileBounds, Params);
+	FVector2D TileBoundsMin = PCG_GET_OVERRIDEN_VALUE(Settings, TileBoundsMin, Params);
+	FVector2D TileBoundsMax = PCG_GET_OVERRIDEN_VALUE(Settings, TileBoundsMax, Params);
+
+	TArray<FPCGTaggedData>& Outputs = Context->OutputData.TaggedData;
+	FPCGTaggedData& Output = Outputs.Emplace_GetRef();
+
+	UPCGTextureData* TextureData = NewObject<UPCGTextureData>();
+	Output.Data = TextureData;
+
+	TextureData->TargetActor = Context->SourceComponent->GetOwner();
+
+	AActor* OriginalActor = UPCGBlueprintHelpers::GetOriginalComponent(*Context)->GetOwner();
+	
+	FTransform FinalTransform = Transform;
+
+	if (!bUseAbsoluteTransform)
+	{
+		FTransform OriginalActorTransform = OriginalActor->GetTransform();
+		FinalTransform = Transform * OriginalActorTransform;
+		
+		FBox OriginalActorLocalBounds = PCGHelpers::GetActorLocalBounds(OriginalActor);
+		FinalTransform.SetScale3D(FinalTransform.GetScale3D() * 0.5 * (OriginalActorLocalBounds.Max - OriginalActorLocalBounds.Min));
+	}
+
+	// Initialize & set properties
+	TextureData->Initialize(Texture, FinalTransform);
+	TextureData->DensityFunction = DensityFunction;
+	TextureData->ColorChannel = ColorChannel;
+	TextureData->TexelSize = TexelSize;
+	TextureData->bUseAdvancedTiling = bUseAdvancedTiling;
+	TextureData->Tiling = Tiling;
+	TextureData->CenterOffset = CenterOffset;
+	TextureData->Rotation = Rotation;
+	TextureData->bUseTileBounds = bUseTileBounds;
+	TextureData->TileBounds = FBox2D(TileBoundsMin, TileBoundsMax);
+
+	return true;
+}
\ No newline at end of file
diff --git a/Engine/Plugins/Experimental/PCG/Source/PCG/Private/PCGHelpers.cpp b/Engine/Plugins/Experimental/PCG/Source/PCG/Private/PCGHelpers.cpp
--- a/Engine/Plugins/Experimental/PCG/Source/PCG/Private/PCGHelpers.cpp
+++ b/Engine/Plugins/Experimental/PCG/Source/PCG/Private/PCGHelpers.cpp
@@ -79,6 +79,31 @@ namespace PCGHelpers
 		return Box;
 	}
 
+	FBox GetActorLocalBounds(AActor* InActor)
+	{
+		// Specialized version of CalculateComponentsBoundingBoxInLocalScape that skips over PCG generated components
+		// This is to ensure stable bounds and no timing issues (cleared ISMs, etc.)
+		FBox Box(EForceInit::ForceInit);
+
+		const bool bNonColliding = true;
+		const bool bIncludeFromChildActors = true;
+
+		const FTransform& ActorToWorld = InActor->GetTransform();
+		const FTransform WorldToActor = ActorToWorld.Inverse();
+
+		InActor->ForEachComponent<UPrimitiveComponent>(bIncludeFromChildActors, [&](const UPrimitiveComponent* InPrimComp)
+		{
+				if ((bNonColliding || InPrimComp->IsCollisionEnabled()) &&
+					!InPrimComp->ComponentTags.Contains(DefaultPCGTag))
+				{
+					const FTransform ComponentToActor = InPrimComp->GetComponentTransform() * WorldToActor;
+					Box += InPrimComp->CalcBounds(ComponentToActor).GetBox();
+				}
+		});
+
+		return Box;
+	}
+
 	bool IsRuntimeOrPIE()
 	{
 #if WITH_EDITOR
diff --git a/Engine/Plugins/Experimental/PCG/Source/PCG/Private/PCGHelpers.h b/Engine/Plugins/Experimental/PCG/Source/PCG/Private/PCGHelpers.h
--- a/Engine/Plugins/Experimental/PCG/Source/PCG/Private/PCGHelpers.h
+++ b/Engine/Plugins/Experimental/PCG/Source/PCG/Private/PCGHelpers.h
@@ -28,6 +28,7 @@ namespace PCGHelpers
 	FBox OverlapBounds(const FBox& InBoxA, const FBox& InBoxB);
 
 	FBox GetActorBounds(AActor* InActor);
+	FBox GetActorLocalBounds(AActor* InActor);
 	FBox GetLandscapeBounds(ALandscapeProxy* InLandscape);
 
 	ALandscape* GetLandscape(UWorld* InWorld, const FBox& InActorBounds);
diff --git a/Engine/Plugins/Experimental/PCG/Source/PCG/Private/Tests/Data/PCGTextureDataTest.cpp b/Engine/Plugins/Experimental/PCG/Source/PCG/Private/Tests/Data/PCGTextureDataTest.cpp
--- a/Engine/Plugins/Experimental/PCG/Source/PCG/Private/Tests/Data/PCGTextureDataTest.cpp
+++ b/Engine/Plugins/Experimental/PCG/Source/PCG/Private/Tests/Data/PCGTextureDataTest.cpp
@@ -38,79 +38,65 @@ bool FPCGTextureDataOffsetTilingRotation::RunTest(const FString& Parameters)
 	UPCGTextureData* TextureData = NewObject<UPCGTextureData>();
 	TextureData->Initialize(Texture2D, FTransform());
 
-	TextureData->bStretchToFit = false;
+	TextureData->bUseAdvancedTiling = true;
 
 	FRandomStream RandomStream;
 	FPCGPoint OutPoint;
 
-	const float ScaledPixelX = static_cast<float>(WhitePixelX) / TextureSize;
-	const float ScaledPixelY = static_cast<float>(WhitePixelY) / TextureSize;
-	const float TexelScalar = 1.f / 3.f;
+	const float TextureSpacePixelX = static_cast<float>(WhitePixelX) / TextureSize;
+	const float ScaledPixelX = (2.0 * TextureSpacePixelX) - 1.0;
+	const float TextureSpacePixelY = static_cast<float>(WhitePixelY) / TextureSize;
+	const float ScaledPixelY = (2.0 * TextureSpacePixelY) - 1;
 
-	for (int TexelSizeFactor = 0; TexelSizeFactor <= 3; ++TexelSizeFactor)
 	{
-		const float TexelSize = FMath::Pow(10.f, TexelSizeFactor);
-		const float ScaledTexelSize = TexelSize * TexelScalar;
-
-		TextureData->TexelSize = TexelSize;
 		TextureData->Rotation = 0.f;
 	
 		// sampling with no offset
-		TextureData->XOffset = 0.f;
-		TextureData->YOffset = 0.f;
+		TextureData->CenterOffset = FVector2D::ZeroVector;
 		TextureData->SamplePoint(FTransform(), FBox(), OutPoint, nullptr);
 		TestEqual("Valid color sampled for no offset at (0, 0)", OutPoint.Color, static_cast<FVector>(FColor::Black));
 
 		// sampling at position, with no offset
-		TextureData->XOffset = 0.f;
-		TextureData->YOffset = 0.f;
-		TextureData->SamplePoint(FTransform(FVector(ScaledTexelSize, ScaledTexelSize, 0.f)), FBox(), OutPoint, nullptr);
+		TextureData->CenterOffset = FVector2D::ZeroVector;
+		TextureData->SamplePoint(FTransform(FVector(0.5, 0.5, 0.0)), FBox(), OutPoint, nullptr);
 		TestEqual("Valid color sampled for no offset at position", OutPoint.Color, static_cast<FVector>(FColor::Black));
+
+		// sampling at expected position
+		TextureData->CenterOffset = FVector2D::ZeroVector;
+		TextureData->SamplePoint(FTransform(FVector(ScaledPixelX, ScaledPixelY, 0)), FBox(), OutPoint, nullptr);
+		TestEqual("Valid color sampled at expected position", OutPoint.Color, static_cast<FVector>(FColor::White));
 	
 		// sampling with offset
-		TextureData->XOffset = 1 - ScaledPixelX;
-		TextureData->YOffset = 1 - ScaledPixelY;
+		TextureData->CenterOffset = FVector2D(0.5 - TextureSpacePixelX, 0.5 -TextureSpacePixelY);
 		TextureData->SamplePoint(FTransform(), FBox(), OutPoint, nullptr);
 		TestEqual("Valid color sampled for offset from (0, 0)", OutPoint.Color, static_cast<FVector>(FColor::White));
 
 		// sampling at position, with offset
-		TextureData->XOffset = (TexelScalar - WhitePixelX) / TextureSize;
-		TextureData->YOffset = (TexelScalar - WhitePixelY) / TextureSize;
-		TextureData->SamplePoint(FTransform(FVector(ScaledTexelSize, ScaledTexelSize, 0.f)), FBox(), OutPoint, nullptr);
+		TextureData->CenterOffset = FVector2D(1.0 - TextureSpacePixelX, 1.0 - TextureSpacePixelY);
+		TextureData->SamplePoint(FTransform(FVector(1.0, 1.0, 0.0)), FBox(), OutPoint, nullptr);
 		TestEqual("Valid color sampled for offset from position", OutPoint.Color, static_cast<FVector>(FColor::White));
 	
-		for (float Rotation = 0.f; Rotation < 360.f; Rotation += 10.f)
+		for (float Rotation = -360.f; Rotation < 360.f; Rotation += 10.f)
 		{
 			const float Theta = FMath::DegreesToRadians(Rotation);
 			const float CosTheta = FMath::Cos(Theta);
 			const float SinTheta = FMath::Sin(Theta);
-			const float X = ((WhitePixelX * CosTheta) - (WhitePixelY * SinTheta)) * TexelSize;
-			const float Y = ((WhitePixelY * CosTheta) + (WhitePixelX * SinTheta)) * TexelSize;
+			const FVector::FReal X = (ScaledPixelX * CosTheta) - (ScaledPixelY * SinTheta);
+			const FVector::FReal Y = (ScaledPixelY * CosTheta) + (ScaledPixelX * SinTheta);
 
 			FTransform RotatedTransform(FVector(X, Y, 0.f));
 
 			// sampling with rotation at black position
-			TextureData->Rotation = Rotation;
-			TextureData->XOffset = 0.f; 
-			TextureData->YOffset = 0.f; 
-			TextureData->SamplePoint(FTransform(FVector(ScaledTexelSize, ScaledTexelSize, 0.f)), FBox(), OutPoint, nullptr);
-			TestEqual("Valid color sampled for non-rotated position", OutPoint.Color, static_cast<FVector>(FColor::Black));
+			TextureData->Rotation = -Rotation;
+			TextureData->CenterOffset = FVector2D::ZeroVector;
+			TextureData->SamplePoint(FTransform(FVector(0.5, 0.5, 0.0)), FBox(), OutPoint, nullptr);
+			TestEqual("Valid color sampled for rotated off-position", OutPoint.Color, static_cast<FVector>(FColor::Black));
 	
 			// sampling with rotation at white position
-			TextureData->Rotation = Rotation;
-			TextureData->XOffset = 0.f; 
-			TextureData->YOffset = 0.f; 
+			TextureData->Rotation = -Rotation;
+			TextureData->CenterOffset = FVector2D::ZeroVector;
 			TextureData->SamplePoint(RotatedTransform, FBox(), OutPoint, nullptr);
 			TestEqual("Valid color sampled for rotated position", OutPoint.Color, static_cast<FVector>(FColor::White));
-			
-			RotatedTransform.SetLocation(FVector(X / 2, Y / 2, 0.f));
-	
-			// sampling with rotation and offset
-			TextureData->Rotation = Rotation;
-			TextureData->XOffset = -ScaledPixelX / 2;
-			TextureData->YOffset = -ScaledPixelY / 2;
-			TextureData->SamplePoint(RotatedTransform, FBox(), OutPoint, nullptr);
-			TestEqual("Valid color sampled for offset and rotation", OutPoint.Color, static_cast<FVector>(FColor::White));
 		}
 	}
 
diff --git a/Engine/Plugins/Experimental/PCG/Source/PCG/Public/Data/PCGTextureData.h b/Engine/Plugins/Experimental/PCG/Source/PCG/Public/Data/PCGTextureData.h
--- a/Engine/Plugins/Experimental/PCG/Source/PCG/Public/Data/PCGTextureData.h
+++ b/Engine/Plugins/Experimental/PCG/Source/PCG/Public/Data/PCGTextureData.h
@@ -16,7 +16,6 @@ enum class EPCGTextureColorChannel : uint8
 	Alpha
 };
 
-
 UENUM(BlueprintType)
 enum class EPCGTextureDensityFunction : uint8
 {
@@ -53,26 +52,30 @@ public:
 	UPROPERTY(BlueprintReadWrite, EditAnywhere, Category = Settings)
 	EPCGTextureColorChannel ColorChannel = EPCGTextureColorChannel::Alpha;
 
+	/** The size of one texel in cm, used when calling ToPointData. */
+	UPROPERTY(BlueprintReadWrite, EditAnywhere, Category = Settings, meta = (UIMin = "1.0", ClampMin = "1.0"))
+	float TexelSize = 50.0f;
+
 	/** Whether to tile the source or to stretch it to fit target area. */
 	UPROPERTY(BlueprintReadWrite, EditAnywhere, Category = Settings)
-	bool bStretchToFit = true;
-
-	/** The size of one texel in cm. */
-	UPROPERTY(BlueprintReadWrite, EditAnywhere, Category = Settings, meta = (UIMin = "1.0", ClampMin = "1.0", EditCondition = "!bStretchToFit"))
-	float TexelSize = 50;
+	bool bUseAdvancedTiling = false;
 
-	/** X Offset factor to apply when sampling texture (Factor of texture dimension). */
-	UPROPERTY(BlueprintReadWrite, EditAnywhere, Category = Settings, meta = (UIMin = "0.0", ClampMin = "0.0", UIMax = "1.0", ClampMax = "1.0", DisplayName = "X Offset", EditCondition = "!bStretchToFit"))
-	float XOffset = 0;
+	UPROPERTY(BlueprintReadWrite, EditAnywhere, Category = "Settings|Tiling", meta = (EditCondition = "bUseAdvancedTiling"))
+	FVector2D Tiling = FVector2D(1.0, 1.0);
 
-	/** Y Offset factor to apply when sampling texture (Factor of texture dimension). */
-	UPROPERTY(BlueprintReadWrite, EditAnywhere, Category = Settings, meta = (UIMin = "0.0", ClampMin = "0.0", UIMax = "1.0", ClampMax = "1.0", DisplayName = "Y Offset", EditCondition = "!bStretchToFit"))
-	float YOffset = 0;
+	UPROPERTY(BlueprintReadWrite, EditAnywhere, Category = "Settings|Tiling", meta = (EditCondition = "bUseAdvancedTiling"))
+	FVector2D CenterOffset = FVector2D::ZeroVector;
 
 	/** Rotation to apply when sampling texture. */
-	UPROPERTY(BlueprintReadWrite, EditAnywhere, Category = Settings, meta = (UIMin = 0, ClampMin = 0, UIMax = 359, ClampMax = 359, Units = deg, EditCondition = "!bStretchToFit"))
+	UPROPERTY(BlueprintReadWrite, EditAnywhere, Category = Settings, meta = (UIMin = -360, ClampMin = -360, UIMax = 360, ClampMax = 360, Units = deg, EditCondition = "bUseAdvancedTiling"))
 	float Rotation = 0;
 
+	UPROPERTY(BlueprintReadWrite, EditAnywhere, Category = "Settings|Tiling", meta = (EditionCondition = "bUseAdvancedTiling"))
+	bool bUseTileBounds = false;
+
+	UPROPERTY(BlueprintReadWrite, EditAnywhere, Category = "Settings|Tiling", meta = (EditCondition = "bUseAdvancedTiling && bUseTileBounds"))
+	FBox2D TileBounds = FBox2D(FVector2D(-0.5, -0.5), FVector2D(0.5, 0.5));
+
 protected:
 	UPROPERTY()
 	TArray<FLinearColor> ColorData;
diff --git a/Engine/Plugins/Experimental/PCG/Source/PCG/Public/Elements/PCGTextureSampler.h b/Engine/Plugins/Experimental/PCG/Source/PCG/Public/Elements/PCGTextureSampler.h
new file mode 100644
--- /dev/null
+++ b/Engine/Plugins/Experimental/PCG/Source/PCG/Public/Elements/PCGTextureSampler.h
@@ -0,0 +1,82 @@
+// Copyright Epic Games, Inc. All Rights Reserved.
+
+#pragma once
+
+#include "PCGSettings.h"
+#include "PCGElement.h"
+#include "Data/PCGTextureData.h"
+
+#include "PCGTextureSampler.generated.h"
+
+UCLASS(BlueprintType, ClassGroup = (Procedural))
+class PCG_API UPCGTextureSamplerSettings : public UPCGSettings
+{
+	GENERATED_BODY()
+
+public:
+	//~Begin UPCGSettings interface
+#if WITH_EDITOR
+	virtual FName GetDefaultNodeName() const override { return FName(TEXT("TextureSamplerNode")); }
+	virtual EPCGSettingsType GetType() const override { return EPCGSettingsType::Sampler; }
+#endif
+
+	virtual TArray<FPCGPinProperties> InputPinProperties() const override;
+	virtual TArray<FPCGPinProperties> OutputPinProperties() const override { return Super::DefaultPointOutputPinProperties(); }
+
+protected:
+	virtual FPCGElementPtr CreateElement() const override;
+	//~End UPCGSettings interface
+
+public:
+	// Surface transform
+	UPROPERTY(BlueprintReadWrite, EditAnywhere, Category = Settings)
+	FTransform Transform = FTransform::Identity;
+
+	UPROPERTY(BlueprintReadWrite, EditAnywhere, Category = Settings)
+	bool bUseAbsoluteTransform = false;
+
+	// Texture specific parameters
+	UPROPERTY(BlueprintReadWrite, EditAnywhere, Category = Settings)
+	TObjectPtr<UTexture2D> Texture = nullptr;
+
+	// Common members in BaseTextureData
+	UPROPERTY(BlueprintReadWrite, VisibleAnywhere, Category = SpatialData)
+	EPCGTextureDensityFunction DensityFunction = EPCGTextureDensityFunction::Multiply;
+
+	UPROPERTY(BlueprintReadWrite, EditAnywhere, Category = Settings)
+	EPCGTextureColorChannel ColorChannel = EPCGTextureColorChannel::Alpha;
+
+	/** The size of one texel in cm, used when calling ToPointData. */
+	UPROPERTY(BlueprintReadWrite, EditAnywhere, Category = Settings, meta = (UIMin = "1.0", ClampMin = "1.0"))
+	float TexelSize = 50.0f;
+
+	/** Whether to tile the source or to stretch it to fit target area. */
+	UPROPERTY(BlueprintReadWrite, EditAnywhere, Category = Settings)
+	bool bUseAdvancedTiling = false;
+
+	UPROPERTY(BlueprintReadWrite, EditAnywhere, Category = "Settings|Tiling", meta = (EditCondition = "bUseAdvancedTiling"))
+	FVector2D Tiling = FVector2D(1.0, 1.0);
+
+	UPROPERTY(BlueprintReadWrite, EditAnywhere, Category = "Settings|Tiling", meta = (EditCondition = "bUseAdvancedTiling"))
+	FVector2D CenterOffset = FVector2D::ZeroVector;
+
+	/** Rotation to apply when sampling texture. */
+	UPROPERTY(BlueprintReadWrite, EditAnywhere, Category = "Settings|Tiling", meta = (UIMin = -360, ClampMin = -360, UIMax = 360, ClampMax = 360, Units = deg, EditCondition = "bUseAdvancedTiling"))
+	float Rotation = 0;
+
+	UPROPERTY(BlueprintReadWrite, EditAnywhere, Category = "Settings|Tiling", meta = (EditionCondition = "bUseAdvancedTiling"))
+	bool bUseTileBounds = false;
+
+	UPROPERTY(BlueprintReadWrite, EditAnywhere, Category = "Settings|Tiling", meta = (EditCondition = "bUseAdvancedTiling && bUseTileBounds"))
+	FVector2D TileBoundsMin = FVector2D(-0.5, -0.5);
+
+	UPROPERTY(BlueprintReadWrite, EditAnywhere, Category = "Settings|Tiling", meta = (EditCondition = "bUseAdvancedTiling && bUseTileBounds"))
+	FVector2D TileBoundsMax = FVector2D(0.5, 0.5);
+};
+
+class FPCGTextureSamplerElement : public FSimplePCGElement
+{
+protected:
+	virtual bool IsCacheable(const UPCGSettings* InSettings) const;
+	virtual bool ExecuteInternal(FPCGContext* Context) const override;
+};
\ No newline at end of file
diff --git a/Engine/Plugins/Experimental/PCG/Source/PCG/Public/Metadata/PCGMetadataAttributeTraits.h b/Engine/Plugins/Experimental/PCG/Source/PCG/Public/Metadata/PCGMetadataAttributeTraits.h
--- a/Engine/Plugins/Experimental/PCG/Source/PCG/Public/Metadata/PCGMetadataAttributeTraits.h
+++ b/Engine/Plugins/Experimental/PCG/Source/PCG/Public/Metadata/PCGMetadataAttributeTraits.h
@@ -13,6 +13,7 @@ enum class EPCGMetadataTypes : uint8
 	Double,
 	Integer32,
 	Integer64,
+	Vector2,
 	Vector,
 	Vector4,
 	Quaternion,
@@ -40,6 +41,7 @@ namespace PCG
 		PCGMetadataGenerateDataTypes(double, Double);
 		PCGMetadataGenerateDataTypes(int32, Integer32);
 		PCGMetadataGenerateDataTypes(int64, Integer64);
+		PCGMetadataGenerateDataTypes(FVector2D, Vector2);
 		PCGMetadataGenerateDataTypes(FVector, Vector);
 		PCGMetadataGenerateDataTypes(FVector4, Vector4);
 		PCGMetadataGenerateDataTypes(FQuat, Quaternion);
@@ -153,41 +155,50 @@ namespace PCG
 		};
 
 		// Vector types
-		template<>
-		struct MetadataTraits<FVector> : DefaultOperationTraits<FVector>
+		template<typename T>
+		struct VectorTraits : DefaultOperationTraits<T>, DefaultWeightedSumTraits<T>
 		{
 			enum { CompressData = false };
 			enum { CanMinMax = true };
 			enum { CanInterpolate = true };
 
-			static FVector Min(const FVector& A, const FVector& B) 
+			static T ZeroValue()
 			{
-				return FVector(FMath::Min(A.X, B.X), FMath::Min(A.Y, B.Y), FMath::Min(A.Z, B.Z));
+				return T::Zero();
 			}
+		};
 
-			static FVector Max(const FVector& A, const FVector& B)
+		template<>
+		struct MetadataTraits<FVector2D> : VectorTraits<FVector2D>
+		{
+			static FVector2D Min(const FVector2D& A, const FVector2D& B)
 			{
-				return FVector(FMath::Max(A.X, B.X), FMath::Max(A.Y, B.Y), FMath::Max(A.Z, B.Z));
+				return FVector2D::Min(A, B);
 			}
 
-			static FVector WeightedSum(const FVector& A, const FVector& B, float Weight)
+			static FVector2D Max(const FVector2D& A, const FVector2D& B)
 			{
-				return A + B * Weight;
+				return FVector2D::Max(A, B);
 			}
+		};
 
-			static FVector ZeroValue()
+		template<>
+		struct MetadataTraits<FVector> : VectorTraits<FVector>
+		{
+			static FVector Min(const FVector& A, const FVector& B)
 			{
-				return FVector::ZeroVector;
+				return FVector::Min(A, B);
 			}
-		};
 
+			static FVector Max(const FVector& A, const FVector& B)
+			{
+				return FVector::Max(A, B);
+			}
+		};
+		
 		template<>
-		struct MetadataTraits<FVector4> : DefaultOperationTraits<FVector4>
+		struct MetadataTraits<FVector4> : VectorTraits<FVector4>
 		{
-			enum { CompressData = false };
-			enum { CanMinMax = true };
-			enum { CanInterpolate = true };
-
 			static FVector4 Min(const FVector4& A, const FVector4& B)
 			{
 				return FVector4(FMath::Min(A.X, B.X), FMath::Min(A.Y, B.Y), FMath::Min(A.Z, B.Z), FMath::Min(A.W, B.W));
@@ -197,16 +208,6 @@ namespace PCG
 			{
 				return FVector4(FMath::Max(A.X, B.X), FMath::Max(A.Y, B.Y), FMath::Max(A.Z, B.Z), FMath::Max(A.W, B.W));
 			}
-
-			static FVector4 WeightedSum(const FVector4& A, const FVector4& B, float Weight)
-			{
-				return A + B * Weight;
-			}
-
-			static FVector4 ZeroValue()
-			{
-				return FVector4::Zero();
-			}
 		};
 
 		// Quaternion
diff --git a/Engine/Plugins/Experimental/PCG/Source/PCG/Public/PCGCommon.h b/Engine/Plugins/Experimental/PCG/Source/PCG/Public/PCGCommon.h
--- a/Engine/Plugins/Experimental/PCG/Source/PCG/Public/PCGCommon.h
+++ b/Engine/Plugins/Experimental/PCG/Source/PCG/Public/PCGCommon.h
@@ -52,4 +52,5 @@ namespace PCGPinConstants
 {
 	const FName DefaultInputLabel = TEXT("In");
 	const FName DefaultOutputLabel = TEXT("Out");
+	const FName DefaultParamsLabel = TEXT("Params");
 }
\ No newline at end of file
