[PCG] Metadata native operation: Math operations

diff --git a/Engine/Plugins/Experimental/PCG/Source/PCG/Private/Elements/Metadata/PCGMetadataMaths.inl b/Engine/Plugins/Experimental/PCG/Source/PCG/Private/Elements/Metadata/PCGMetadataMaths.inl
new file mode 100644
--- /dev/null
+++ b/Engine/Plugins/Experimental/PCG/Source/PCG/Private/Elements/Metadata/PCGMetadataMaths.inl
@@ -0,0 +1,342 @@
+// Copyright Epic Games, Inc. All Rights Reserved.
+
+namespace PCGMetadataMaths
+{
+	////////////////////
+	// Sign
+	////////////////////
+	template <typename T>
+	inline T Sign(const T& Value)
+	{
+		double Temp(Value);
+		return Temp > 0.0 ? static_cast<T>(1) : (Temp < 0.0 ? static_cast<T>(-1) : static_cast<T>(0));
+	}
+
+	template<>
+	inline FVector2D Sign(const FVector2D& Value)
+	{
+		return FVector2D(Sign(Value.X), Sign(Value.Y));
+	}
+
+	template<>
+	inline FVector Sign(const FVector& Value)
+	{
+		return FVector(Sign(Value.X), Sign(Value.Y), Sign(Value.Z));
+	}
+
+	template<>
+	inline FVector4 Sign(const FVector4& Value)
+	{
+		return FVector4(Sign(Value.X), Sign(Value.Y), Sign(Value.Z), Sign(Value.W));
+	}
+
+	////////////////////
+	// Frac
+	////////////////////
+	template <typename T>
+	inline T Frac(const T& Value)
+	{
+		return static_cast<T>(FMath::Frac((double)Value));
+	}
+
+	template<>
+	inline FVector2D Frac(const FVector2D& Value)
+	{
+		return FVector2D(Frac(Value.X), Frac(Value.Y));
+	}
+
+	template<>
+	inline FVector Frac(const FVector& Value)
+	{
+		return FVector(Frac(Value.X), Frac(Value.Y), Frac(Value.Z));
+	}
+
+	template<>
+	inline FVector4 Frac(const FVector4& Value)
+	{
+		return FVector4(Frac(Value.X), Frac(Value.Y), Frac(Value.Z), Frac(Value.W));
+	}
+
+	////////////////////
+	// Truncate
+	////////////////////
+	template <typename T>
+	inline T Truncate(const T& Value)
+	{
+		return static_cast<T>(FMath::TruncToDouble((double)Value));
+	}
+
+	template<>
+	inline FVector2D Truncate(const FVector2D& Value)
+	{
+		return FVector2D(Truncate(Value.X), Truncate(Value.Y));
+	}
+
+	template<>
+	inline FVector Truncate(const FVector& Value)
+	{
+		return FVector(Truncate(Value.X), Truncate(Value.Y), Truncate(Value.Z));
+	}
+
+	template<>
+	inline FVector4 Truncate(const FVector4& Value)
+	{
+		return FVector4(Truncate(Value.X), Truncate(Value.Y), Truncate(Value.Z), Truncate(Value.W));
+	}
+
+	////////////////////
+	// Round
+	////////////////////
+	template <typename T>
+	inline T Round(const T& Value)
+	{
+		return static_cast<T>(FMath::RoundToDouble((double)Value));
+	}
+
+	template<>
+	inline FVector2D Round(const FVector2D& Value)
+	{
+		return FVector2D(Round(Value.X), Round(Value.Y));
+	}
+
+	template<>
+	inline FVector Round(const FVector& Value)
+	{
+		return FVector(Round(Value.X), Round(Value.Y), Round(Value.Z));
+	}
+
+	template<>
+	inline FVector4 Round(const FVector4& Value)
+	{
+		return FVector4(Round(Value.X), Round(Value.Y), Round(Value.Z), Round(Value.W));
+	}
+
+	////////////////////
+	// Sqrt
+	////////////////////
+	template <typename T>
+	inline T Sqrt(const T& Value)
+	{
+		return static_cast<T>(FMath::Sqrt((double)Value));
+	}
+
+	template<>
+	inline FVector2D Sqrt(const FVector2D& Value)
+	{
+		return FVector2D(Sqrt(Value.X), Sqrt(Value.Y));
+	}
+
+	template<>
+	inline FVector Sqrt(const FVector& Value)
+	{
+		return FVector(Sqrt(Value.X), Sqrt(Value.Y), Sqrt(Value.Z));
+	}
+
+	template<>
+	inline FVector4 Sqrt(const FVector4& Value)
+	{
+		return FVector4(Sqrt(Value.X), Sqrt(Value.Y), Sqrt(Value.Z), Sqrt(Value.W));
+	}
+
+	////////////////////
+	// Abs
+	////////////////////
+	template <typename T>
+	inline T Abs(const T& Value)
+	{
+		return FMath::Abs(Value);
+	}
+
+	template<>
+	inline FVector2D Abs(const FVector2D& Value)
+	{
+		return Value.GetAbs();
+	}
+
+	template<>
+	inline FVector Abs(const FVector& Value)
+	{
+		return Value.GetAbs();
+	}
+
+	template<>
+	inline FVector4 Abs(const FVector4& Value)
+	{
+		return FVector4(FMath::Abs(Value.X), FMath::Abs(Value.Y), FMath::Abs(Value.Z), FMath::Abs(Value.W));
+	}
+
+	////////////////////
+	// Max
+	////////////////////
+	template <typename T>
+	inline T Max(const T& Value1, const T& Value2)
+	{
+		return FMath::Max(Value1, Value2);
+	}
+
+	template<>
+	inline FVector2D Max(const FVector2D& Value1, const FVector2D& Value2)
+	{
+		return FVector2D::Max(Value1, Value2);
+	}
+
+	template<>
+	inline FVector Max(const FVector& Value1, const FVector& Value2)
+	{
+		return FVector::Max(Value1, Value2);
+	}
+
+	template<>
+	inline FVector4 Max(const FVector4& Value1, const FVector4& Value2)
+	{
+		return FVector4(FMath::Max(Value1.X, Value2.X), FMath::Max(Value1.Y, Value2.Y), FMath::Max(Value1.Z, Value2.Z), FMath::Max(Value1.W, Value2.W));
+	}
+
+	////////////////////
+	// Min
+	////////////////////
+	template <typename T>
+	inline T Min(const T& Value1, const T& Value2)
+	{
+		return FMath::Min(Value1, Value2);
+	}
+
+	template<>
+	inline FVector2D Min(const FVector2D& Value1, const FVector2D& Value2)
+	{
+		return FVector2D::Min(Value1, Value2);
+	}
+
+	template<>
+	inline FVector Min(const FVector& Value1, const FVector& Value2)
+	{
+		return FVector::Min(Value1, Value2);
+	}
+
+	template<>
+	inline FVector4 Min(const FVector4& Value1, const FVector4& Value2)
+	{
+		return FVector4(FMath::Min(Value1.X, Value2.X), FMath::Min(Value1.Y, Value2.Y), FMath::Min(Value1.Z, Value2.Z), FMath::Min(Value1.W, Value2.W));
+	}
+
+	////////////////////
+	// Pow
+	////////////////////
+	template <typename T>
+	inline T Pow(const T& Value, const T& Power)
+	{
+		return static_cast<T>(FMath::Pow((double)Value, (double)Power));
+	}
+
+	template<>
+	inline FVector2D Pow(const FVector2D& Value, const FVector2D& Power)
+	{
+		return FVector2D(FMath::Pow(Value.X, Power.X), FMath::Pow(Value.Y, Power.Y));
+	}
+
+	template<>
+	inline FVector Pow(const FVector& Value, const FVector& Power)
+	{
+		return FVector(FMath::Pow(Value.X, Power.X), FMath::Pow(Value.Y, Power.Y), FMath::Pow(Value.Z, Power.Z));
+	}
+
+	template<>
+	inline FVector4 Pow(const FVector4& Value, const FVector4& Power)
+	{
+		return FVector4(FMath::Pow(Value.X, Power.X), FMath::Pow(Value.Y, Power.Y), FMath::Pow(Value.Z, Power.Z), FMath::Pow(Value.W, Power.W));
+	}
+
+	////////////////////
+	// Clamp
+	////////////////////
+	template <typename T>
+	inline T Clamp(const T& Value, const T& MinValue, const T& MaxValue)
+	{
+		return FMath::Clamp(Value, MinValue, MaxValue);
+	}
+
+	template<>
+	inline FVector2D Clamp(const FVector2D& Value, const FVector2D& MinValue, const FVector2D& MaxValue)
+	{
+		return FVector2D(Clamp(Value.X, MinValue.X, MaxValue.X), Clamp(Value.Y, MinValue.Y, MaxValue.Y));
+	}
+
+	template<>
+	inline FVector Clamp(const FVector& Value, const FVector& MinValue, const FVector& MaxValue)
+	{
+		return Value.BoundToBox(MinValue, MaxValue);
+	}
+
+	template<>
+	inline FVector4 Clamp(const FVector4& Value, const FVector4& MinValue, const FVector4& MaxValue)
+	{
+		return FVector4(
+			Clamp(Value.X, MinValue.X, MaxValue.X),
+			Clamp(Value.Y, MinValue.Y, MaxValue.Y),
+			Clamp(Value.Z, MinValue.Z, MaxValue.Z),
+			Clamp(Value.W, MinValue.W, MaxValue.W));
+	}
+
+	////////////////////
+	// Lerp
+	////////////////////
+	template <typename T, typename U>
+	inline T Lerp(const T& Value1, const T& Value2, const U& Ratio)
+	{
+		return static_cast<T>(FMath::Lerp(Value1, Value2, Ratio));
+	}
+
+	template<typename U>
+	inline FVector2D Lerp(const FVector2D& Value1, const FVector2D& Value2, const U& Ratio)
+	{
+		return FVector2D(
+			Lerp(Value1.X, Value2.X, Ratio),
+			Lerp(Value1.Y, Value2.Y, Ratio));
+	}
+
+	template<>
+	inline FVector2D Lerp(const FVector2D& Value1, const FVector2D& Value2, const FVector2D& Ratio)
+	{
+		return FVector2D(
+			Lerp(Value1.X, Value2.X, Ratio.X),
+			Lerp(Value1.Y, Value2.Y, Ratio.Y));
+	}
+
+	template<typename U>
+	inline FVector Lerp(const FVector& Value1, const FVector& Value2, const U& Ratio)
+	{
+		return FVector(
+			Lerp(Value1.X, Value2.X, Ratio),
+			Lerp(Value1.Y, Value2.Y, Ratio),
+			Lerp(Value1.Z, Value2.Z, Ratio));
+	}
+
+	template<>
+	inline FVector Lerp(const FVector& Value1, const FVector& Value2, const FVector& Ratio)
+	{
+		return FVector(
+			Lerp(Value1.X, Value2.X, Ratio.X), 
+			Lerp(Value1.Y, Value2.Y, Ratio.Y), 
+			Lerp(Value1.Z, Value2.Z, Ratio.Z));
+	}
+
+	template<typename U>
+	inline FVector4 Lerp(const FVector4& Value1, const FVector4& Value2, const U& Ratio)
+	{
+		return FVector4(
+			Lerp(Value1.X, Value2.X, Ratio),
+			Lerp(Value1.Y, Value2.Y, Ratio),
+			Lerp(Value1.Z, Value2.Z, Ratio),
+			Lerp(Value1.W, Value2.W, Ratio));
+	}
+
+	template<>
+	inline FVector4 Lerp(const FVector4& Value1, const FVector4& Value2, const FVector4& Ratio)
+	{
+		return FVector4(
+			Lerp(Value1.X, Value2.X, Ratio.X),
+			Lerp(Value1.Y, Value2.Y, Ratio.Y),
+			Lerp(Value1.Z, Value2.Z, Ratio.Z),
+			Lerp(Value1.W, Value2.W, Ratio.W));
+	}
+}
\ No newline at end of file
diff --git a/Engine/Plugins/Experimental/PCG/Source/PCG/Private/Elements/Metadata/PCGMetadataMathsOpElement.cpp b/Engine/Plugins/Experimental/PCG/Source/PCG/Private/Elements/Metadata/PCGMetadataMathsOpElement.cpp
new file mode 100644
--- /dev/null
+++ b/Engine/Plugins/Experimental/PCG/Source/PCG/Private/Elements/Metadata/PCGMetadataMathsOpElement.cpp
@@ -0,0 +1,401 @@
+// Copyright Epic Games, Inc. All Rights Reserved.
+
+#include "Elements/Metadata/PCGMetadataMathsOpElement.h"
+
+#include "PCGParamData.h"
+#include "Data/PCGSpatialData.h"
+#include "Helpers/PCGSettingsHelpers.h"
+#include "Metadata/PCGMetadata.h"
+#include "Metadata/PCGMetadataAttribute.h"
+#include "Metadata/PCGMetadataAttributeTpl.h"
+#include "Metadata/PCGMetadataEntryKeyIterator.h"
+#include "Elements/Metadata/PCGMetadataElementCommon.h"
+
+#include "Elements/Metadata/PCGMetadataMaths.inl"
+
+namespace PCGMetadataMathsSettings
+{
+	inline constexpr bool IsUnaryOp(EPCGMedadataMathsOperation Operation)
+	{
+		return !!(Operation & EPCGMedadataMathsOperation::UnaryOp);
+	}
+
+	inline constexpr bool IsBinaryOp(EPCGMedadataMathsOperation Operation)
+	{
+		return !!(Operation & EPCGMedadataMathsOperation::BinaryOp);
+	}
+
+	inline constexpr bool IsTernaryOp(EPCGMedadataMathsOperation Operation)
+	{
+		return !!(Operation & EPCGMedadataMathsOperation::TernaryOp);
+	}
+
+	inline FName GetFirstPinLabel(EPCGMedadataMathsOperation Operation)
+	{
+		if (PCGMetadataMathsSettings::IsUnaryOp(Operation)
+			|| Operation == EPCGMedadataMathsOperation::Clamp
+			|| Operation == EPCGMedadataMathsOperation::ClampMin
+			|| Operation == EPCGMedadataMathsOperation::ClampMax)
+		{
+			return PCGPinConstants::DefaultInputLabel;
+		}
+
+		if (PCGMetadataMathsSettings::IsBinaryOp(Operation)
+			|| Operation == EPCGMedadataMathsOperation::Lerp)
+		{
+			return PCGMetadataSettingsBaseConstants::DoubleInputFirstLabel;
+		}
+
+		return NAME_None;
+	}
+
+	inline FName GetSecondPinLabel(EPCGMedadataMathsOperation Operation)
+	{
+		if (Operation == EPCGMedadataMathsOperation::ClampMin || Operation == EPCGMedadataMathsOperation::Clamp)
+		{
+			return PCGMetadataSettingsBaseConstants::ClampMinLabel;
+		}
+
+		if (Operation == EPCGMedadataMathsOperation::ClampMax)
+		{
+			return PCGMetadataSettingsBaseConstants::ClampMaxLabel;
+		}
+
+		if (PCGMetadataMathsSettings::IsBinaryOp(Operation) || PCGMetadataMathsSettings::IsTernaryOp(Operation))
+		{
+			return PCGMetadataSettingsBaseConstants::DoubleInputSecondLabel;
+		}
+
+		return NAME_None;
+	}
+
+	inline FName GetThirdPinLabel(EPCGMedadataMathsOperation Operation)
+	{
+		if (Operation == EPCGMedadataMathsOperation::Clamp)
+		{
+			return PCGMetadataSettingsBaseConstants::ClampMaxLabel;
+		}
+
+		if (Operation == EPCGMedadataMathsOperation::Lerp)
+		{
+			return PCGMetadataSettingsBaseConstants::LerpRatioLabel;
+		}
+
+		return NAME_None;
+	}
+
+	template <typename T>
+	void UnaryOp(T& Value, EPCGMedadataMathsOperation Op)
+	{
+		switch (Op)
+		{
+		case EPCGMedadataMathsOperation::Sign:
+			Value = PCGMetadataMaths::Sign(Value);
+			break;
+		case EPCGMedadataMathsOperation::Frac:
+			Value = PCGMetadataMaths::Frac(Value);
+			break;
+		case EPCGMedadataMathsOperation::Truncate:
+			Value = PCGMetadataMaths::Truncate(Value);
+			break;
+		case EPCGMedadataMathsOperation::Round:
+			Value = PCGMetadataMaths::Round(Value);
+			break;
+		case EPCGMedadataMathsOperation::Sqrt:
+			Value = PCGMetadataMaths::Sqrt(Value);
+			break;
+		case EPCGMedadataMathsOperation::Abs:
+			Value = PCGMetadataMaths::Abs(Value);
+			break;
+		}
+	}
+
+	template <typename T>
+	void BinaryOp(T& Value1, const T& Value2, EPCGMedadataMathsOperation Op)
+	{
+		switch (Op)
+		{
+		case EPCGMedadataMathsOperation::Add:
+			Value1 = Value1 + Value2;
+			break;
+		case EPCGMedadataMathsOperation::Subtract:
+			Value1 = Value1 - Value2;
+			break;
+		case EPCGMedadataMathsOperation::Multiply:
+			Value1 = Value1 * Value2;
+			break;
+		case EPCGMedadataMathsOperation::Divide:
+			Value1 = Value1 / Value2;
+			break;
+		case EPCGMedadataMathsOperation::Max:
+			Value1 = PCGMetadataMaths::Max(Value1, Value2);
+			break;
+		case EPCGMedadataMathsOperation::Min:
+			Value1 = PCGMetadataMaths::Min(Value1, Value2);
+			break;
+		case EPCGMedadataMathsOperation::ClampMin:
+			Value1 = PCGMetadataMaths::Clamp(Value1, Value2, Value1);
+			break;
+		case EPCGMedadataMathsOperation::ClampMax:
+			Value1 = PCGMetadataMaths::Clamp(Value1, Value1, Value2);
+			break;
+		case EPCGMedadataMathsOperation::Pow:
+			Value1 = PCGMetadataMaths::Pow(Value1, Value2);
+			break;
+		}
+	}
+
+	template <typename T>
+	void TernaryOp(T& Value1, const T& Value2, const T& Value3, EPCGMedadataMathsOperation Op)
+	{
+		switch (Op)
+		{
+		case EPCGMedadataMathsOperation::Clamp:
+			Value1 = PCGMetadataMaths::Clamp(Value1, Value2, Value3);
+			break;
+		case EPCGMedadataMathsOperation::Lerp:
+			Value1 = PCGMetadataMaths::Lerp(Value1, Value2, Value3);
+			break;
+		}
+	}
+}
+
+FName UPCGMetadataMathsSettings::GetInputPinLabel(uint32 Index) const
+{
+	switch (Index)
+	{
+	case 0:
+		return PCGMetadataMathsSettings::GetFirstPinLabel(Operation);
+	case 1:
+		return PCGMetadataMathsSettings::GetSecondPinLabel(Operation);
+	case 2:
+		return PCGMetadataMathsSettings::GetThirdPinLabel(Operation);
+	default:
+		return NAME_None;
+	}
+}
+
+uint32 UPCGMetadataMathsSettings::GetInputPinNum() const
+{
+	if (PCGMetadataMathsSettings::IsUnaryOp(Operation))
+	{
+		return 1;
+	}
+
+	if (PCGMetadataMathsSettings::IsBinaryOp(Operation))
+	{
+		return 2;
+	}
+
+	if (PCGMetadataMathsSettings::IsTernaryOp(Operation))
+	{
+		return 3;
+	}
+
+	return 0;
+}
+
+// By default: Float/Double, Int32/Int64, Vector2, Vector, Vector4
+bool UPCGMetadataMathsSettings::IsSupportedInputType(uint16 TypeId, uint32 InputIndex, bool& bHasSpecialRequirement) const
+{
+	bHasSpecialRequirement = false;
+	return TypeId <= (uint16)EPCGMetadataTypes::Vector4;
+}
+
+FName UPCGMetadataMathsSettings::GetInputAttributeNameWithOverride(uint32 Index, UPCGParamData* Params) const
+{
+	switch (Index)
+	{
+	case 0:
+		return PCG_GET_OVERRIDEN_VALUE(this, Input1AttributeName, Params);
+	case 1:
+		return PCG_GET_OVERRIDEN_VALUE(this, Input2AttributeName, Params);
+	case 2:
+		return PCG_GET_OVERRIDEN_VALUE(this, Input3AttributeName, Params);
+	default:
+		return NAME_None;
+	}
+}
+
+#if WITH_EDITOR
+FName UPCGMetadataMathsSettings::GetDefaultNodeName() const
+{
+	if (const UEnum* EnumPtr = FindObject<UEnum>(nullptr, TEXT("EPCGMedadataMathsOperation"), true))
+	{ 
+		return EnumPtr->GetNameByValue(static_cast<int>(Operation)); 
+	}
+
+	return TEXT("Metadata Maths Node");
+}
+#endif // WITH_EDITOR
+
+FPCGElementPtr UPCGMetadataMathsSettings::CreateElement() const
+{
+	return MakeShared<FPCGMetadataMathsElement>();
+}
+
+bool FPCGMetadataMathsElement::DoOperation(FOperationData& OperationData) const
+{
+	TRACE_CPUPROFILER_EVENT_SCOPE(FPCGMetadataMathsElement::Execute);
+
+	const UPCGMetadataMathsSettings* Settings = CastChecked<UPCGMetadataMathsSettings>(OperationData.Settings);
+
+	// All values in OperationData.Iterators are UniquePtr. Dereference them here to make the syntax less cumbersome.
+	// Also some iterators can be null, it means they need to use the first iterator (that should never be null).
+	// We just need to make sure to not increment the first iterator multiple times in one loop.
+
+	check(OperationData.Iterators[0]);
+
+	IPCGMetadataEntryIterator& Iterator1 = *OperationData.Iterators[0];
+	IPCGMetadataEntryIterator& Iterator2 = (OperationData.Iterators.Num() >= 2 && OperationData.Iterators[1]) ? *OperationData.Iterators[1] : Iterator1;
+	IPCGMetadataEntryIterator& Iterator3 = (OperationData.Iterators.Num() >= 3 && OperationData.Iterators[2]) ? *OperationData.Iterators[2] : Iterator1;
+
+	bool bShouldIncrementIterator2 = &Iterator1 != &Iterator2;
+	bool bShouldIncrementIterator3 = &Iterator1 != &Iterator3;
+
+	// And then do the operation for all elements
+	if (PCGMetadataMathsSettings::IsUnaryOp(Settings->Operation))
+	{
+		TRACE_CPUPROFILER_EVENT_SCOPE(FPCGMetadataMathsElement::ExecuteInternal::UnaryOp);
+		auto UnaryFunc = [Operation = Settings->Operation, &OperationData, &Iterator1](auto DummyOutValue) -> void
+		{
+			using AttributeType = decltype(DummyOutValue);
+
+			// Need to remove types that would not compile
+			if constexpr (PCG::Private::MetadataTypes<AttributeType>::Id > (uint16)EPCGMetadataTypes::Vector4)
+			{
+				return;
+			}
+			else
+			{
+				FPCGMetadataAttribute<AttributeType>* OutputAttribute = static_cast<FPCGMetadataAttribute<AttributeType>*>(OperationData.OutputAttribute);
+				AttributeType DefaultValue = PCGMetadataAttribute::GetValueWithBroadcast<AttributeType>(OperationData.SourceAttributes[0], PCGInvalidEntryKey);
+				PCGMetadataMathsSettings::UnaryOp(DefaultValue, Operation);
+				OutputAttribute->SetDefaultValue(DefaultValue);
+
+				for (int32 i = 0; i < OperationData.NumberOfElementsToProcess; ++i)
+				{
+					PCGMetadataEntryKey EntryKey = *Iterator1;
+
+					// If the entry key is invalid, nothing to do
+					if (EntryKey == PCGInvalidEntryKey)
+					{
+						AttributeType Value = PCGMetadataAttribute::GetValueWithBroadcast<AttributeType>(OperationData.SourceAttributes[0], EntryKey);
+
+						PCGMetadataMathsSettings::UnaryOp(Value, Operation);
+
+						OutputAttribute->SetValue(EntryKey, Value);
+					}
+
+					++Iterator1;
+				}
+			}
+		};
+
+		PCGMetadataAttribute::CallbackWithRightType(OperationData.OutputType, UnaryFunc);
+	}
+	else if (PCGMetadataMathsSettings::IsBinaryOp(Settings->Operation))
+	{
+		auto BinaryFunc = [Operation = Settings->Operation, &OperationData, &Iterator1, &Iterator2, bShouldIncrementIterator2](auto DummyOutValue) -> void
+		{
+			TRACE_CPUPROFILER_EVENT_SCOPE(FPCGMetadataMathsElement::ExecuteInternal::BinaryOp);
+			using AttributeType = decltype(DummyOutValue);
+
+			// Need to remove types that would not compile
+			if constexpr (PCG::Private::MetadataTypes<AttributeType>::Id > (uint16)EPCGMetadataTypes::Vector4)
+			{
+				return;
+			}
+			else
+			{
+				FPCGMetadataAttribute<AttributeType>* OutputAttribute = static_cast<FPCGMetadataAttribute<AttributeType>*>(OperationData.OutputAttribute);
+				AttributeType DefaultValue1 = PCGMetadataAttribute::GetValueWithBroadcast<AttributeType>(OperationData.SourceAttributes[0], PCGInvalidEntryKey);
+				AttributeType DefaultValue2 = PCGMetadataAttribute::GetValueWithBroadcast<AttributeType>(OperationData.SourceAttributes[1], PCGInvalidEntryKey);
+				PCGMetadataMathsSettings::BinaryOp(DefaultValue1, DefaultValue2, Operation);
+				OutputAttribute->SetDefaultValue(DefaultValue1);
+
+				for (int32 i = 0; i < OperationData.NumberOfElementsToProcess; ++i)
+				{
+					PCGMetadataEntryKey EntryKey1 = *Iterator1;
+					PCGMetadataEntryKey EntryKey2 = *Iterator2;
+
+					// If the entry key is invalid, nothing to do
+					if (EntryKey1 == PCGInvalidEntryKey)
+					{
+						AttributeType Value1 = PCGMetadataAttribute::GetValueWithBroadcast<AttributeType>(OperationData.SourceAttributes[0], EntryKey1);
+						AttributeType Value2 = PCGMetadataAttribute::GetValueWithBroadcast<AttributeType>(OperationData.SourceAttributes[1], EntryKey2);
+
+						PCGMetadataMathsSettings::BinaryOp(Value1, Value2, Operation);
+
+						OutputAttribute->SetValue(EntryKey1, Value1);
+					}
+
+					++Iterator1;
+					if (bShouldIncrementIterator2)
+					{
+						++Iterator2;
+					}
+				}
+			}
+		};
+
+		PCGMetadataAttribute::CallbackWithRightType(OperationData.OutputType, BinaryFunc);
+	}
+	else if (PCGMetadataMathsSettings::IsTernaryOp(Settings->Operation))
+	{
+		auto TernaryFunc = [Operation = Settings->Operation, &OperationData, &Iterator1, &Iterator2, &Iterator3, bShouldIncrementIterator2, bShouldIncrementIterator3](auto DummyOutValue) -> void
+		{
+			TRACE_CPUPROFILER_EVENT_SCOPE(FPCGMetadataMathsElement::ExecuteInternal::TernaryOp);
+			using AttributeType = decltype(DummyOutValue);
+
+			// Need to remove types that would not compile
+			if constexpr (PCG::Private::MetadataTypes<AttributeType>::Id > (uint16)EPCGMetadataTypes::Vector4)
+			{
+				return;
+			}
+			else
+			{
+				FPCGMetadataAttribute<AttributeType>* OutputAttribute = static_cast<FPCGMetadataAttribute<AttributeType>*>(OperationData.OutputAttribute);
+				AttributeType DefaultValue1 = PCGMetadataAttribute::GetValueWithBroadcast<AttributeType>(OperationData.SourceAttributes[0], PCGInvalidEntryKey);
+				AttributeType DefaultValue2 = PCGMetadataAttribute::GetValueWithBroadcast<AttributeType>(OperationData.SourceAttributes[1], PCGInvalidEntryKey);
+				AttributeType DefaultValue3 = PCGMetadataAttribute::GetValueWithBroadcast<AttributeType>(OperationData.SourceAttributes[2], PCGInvalidEntryKey);
+				PCGMetadataMathsSettings::TernaryOp(DefaultValue1, DefaultValue2, DefaultValue3, Operation);
+				OutputAttribute->SetDefaultValue(DefaultValue1);
+
+				for (PCGMetadataValueKey ValueKey = 0; ValueKey < OperationData.NumberOfElementsToProcess; ++ValueKey)
+				{
+					PCGMetadataEntryKey EntryKey1 = *Iterator1;
+					PCGMetadataEntryKey EntryKey2 = *Iterator2;
+					PCGMetadataEntryKey EntryKey3 = *Iterator3;
+
+					// If the entry key is invalid, nothing to do
+					if (EntryKey1 == PCGInvalidEntryKey)
+					{
+						AttributeType Value1 = PCGMetadataAttribute::GetValueWithBroadcast<AttributeType>(OperationData.SourceAttributes[0], EntryKey1);
+						AttributeType Value2 = PCGMetadataAttribute::GetValueWithBroadcast<AttributeType>(OperationData.SourceAttributes[1], EntryKey2);
+						AttributeType Value3 = PCGMetadataAttribute::GetValueWithBroadcast<AttributeType>(OperationData.SourceAttributes[2], EntryKey3);
+
+						PCGMetadataMathsSettings::TernaryOp(Value1, Value2, Value3, Operation);
+
+						OutputAttribute->SetValue(EntryKey1, Value1);
+					}
+
+					++Iterator1;
+					if (bShouldIncrementIterator2)
+					{
+						++Iterator2;
+					}
+
+					if (bShouldIncrementIterator3)
+					{
+						++Iterator3;
+					}
+				}
+			}
+		};
+
+		PCGMetadataAttribute::CallbackWithRightType(OperationData.OutputType, TernaryFunc);
+	}
+
+	return true;
+}
\ No newline at end of file
diff --git a/Engine/Plugins/Experimental/PCG/Source/PCG/Private/Elements/Metadata/PCGMetadataOpElementBase.cpp b/Engine/Plugins/Experimental/PCG/Source/PCG/Private/Elements/Metadata/PCGMetadataOpElementBase.cpp
new file mode 100644
--- /dev/null
+++ b/Engine/Plugins/Experimental/PCG/Source/PCG/Private/Elements/Metadata/PCGMetadataOpElementBase.cpp
@@ -0,0 +1,226 @@
+// Copyright Epic Games, Inc. All Rights Reserved.
+
+#include "Elements/Metadata/PCGMetadataOpElementBase.h"
+
+#include "PCGParamData.h"
+#include "Data/PCGSpatialData.h"
+#include "Data/PCGPointData.h"
+#include "Elements/Metadata/PCGMetadataElementCommon.h"
+#include "Helpers/PCGSettingsHelpers.h"
+#include "Metadata/PCGMetadata.h"
+#include "Metadata/PCGMetadataAttribute.h"
+#include "Metadata/PCGMetadataAttributeTpl.h"
+#include "Metadata/PCGMetadataEntryKeyIterator.h"
+
+TArray<FPCGPinProperties> UPCGMetadataSettingsBase::InputPinProperties() const
+{
+	TArray<FPCGPinProperties> PinProperties;
+
+	for (uint32 i = 0; i < GetInputPinNum(); ++i)
+	{
+		const FName PinLabel = GetInputPinLabel(i);
+		if (PinLabel != NAME_None)
+		{
+			PinProperties.Emplace(PinLabel, EPCGDataType::Any, /*bAllowMultipleConnections=*/false);
+		}
+	}
+
+	PinProperties.Emplace(PCGPinConstants::DefaultParamsLabel, EPCGDataType::Param, /*bAllowMultipleConnections=*/false);
+
+	return PinProperties;
+}
+
+TArray<FPCGPinProperties> UPCGMetadataSettingsBase::OutputPinProperties() const
+{
+	TArray<FPCGPinProperties> PinProperties;
+	PinProperties.Emplace(PCGPinConstants::DefaultOutputLabel, EPCGDataType::Any);
+
+	return PinProperties;
+}
+
+bool UPCGMetadataSettingsBase::IsMoreComplexType(uint16 FirstType, uint16 SecondType) const
+{
+	return FirstType != SecondType && FirstType <= (uint16)(EPCGMetadataTypes::Count) && SecondType <= (uint16)(EPCGMetadataTypes::Count) && PCG::Private::BroadcastableTypes[SecondType][FirstType];
+}
+
+bool FPCGMetadataElementBase::ExecuteInternal(FPCGContext* Context) const
+{
+	TRACE_CPUPROFILER_EVENT_SCOPE(FPCGMetadataElementBase::Execute);
+
+	const UPCGMetadataSettingsBase* Settings = Context->GetInputSettings<UPCGMetadataSettingsBase>();
+	check(Settings);
+
+	const TArray<FPCGTaggedData>& Inputs = Context->InputData.TaggedData;
+	TArray<FPCGTaggedData>& Outputs = Context->OutputData.TaggedData;
+
+	// Nothing to do if we have no output
+	if (!Context->Node || !Context->Node->IsOutputPinConnected(PCGPinConstants::DefaultOutputLabel))
+	{
+		return true;
+	}
+
+	const TArray<FPCGTaggedData>& ParamData = Context->InputData.GetParamsByPin(PCGPinConstants::DefaultParamsLabel);
+	UPCGParamData* Params = ParamData.IsEmpty() ? nullptr : Cast<UPCGParamData>(ParamData[0].Data);
+
+	FName OutputAttributeName = PCG_GET_OVERRIDEN_VALUE(Settings, OutputAttributeName, Params);
+
+	const uint32 NumberOfInputs = Settings->GetInputPinNum();
+
+	// Gathering all the inputs metadata
+	TArray<const UPCGMetadata*> SourceMetadata;
+	TArray<const UPCGData*> Data;
+	SourceMetadata.SetNum(NumberOfInputs);
+	Data.SetNum(NumberOfInputs);
+
+	for (uint32 i = 0; i < NumberOfInputs; ++i)
+	{
+		TArray<FPCGTaggedData> InputData = Context->InputData.GetInputsByPin(Settings->GetInputPinLabel(i));
+		if (InputData.IsEmpty())
+		{
+			PCGE_LOG(Error, "Invalid inputs for pin %d", i);
+			return true;
+		}
+
+		// By construction, there can only be one of then(hence the 0 index)
+		check(InputData.Num() == 1);
+		Data[i] = InputData[0].Data;
+
+		// Only gather Spacial and Params input. 
+		if (const UPCGSpatialData* SpatialInput = Cast<const UPCGSpatialData>(Data[i]))
+		{
+			SourceMetadata[i] = SpatialInput->Metadata;
+		}
+		else if (const UPCGParamData* ParamsInput = Cast<const UPCGParamData>(Data[i]))
+		{
+			SourceMetadata[i] = ParamsInput->Metadata;
+		}
+		else
+		{
+			PCGE_LOG(Error, "Invalid inputs for pin %d", i);
+			return true;
+		}
+	}
+
+	FOperationData OperationData;
+	OperationData.Iterators.SetNum(NumberOfInputs);
+	TArray<int32> NumberOfElements;
+	NumberOfElements.SetNum(NumberOfInputs);
+
+	OperationData.SourceAttributes.SetNum(NumberOfInputs);
+
+	uint16 MostComplexType = (uint16)EPCGMetadataTypes::Unknown;
+	OperationData.NumberOfElementsToProcess = -1;
+
+	for (uint32 i = 0; i < NumberOfInputs; ++i)
+	{
+		FName SourceAttributeName = Settings->GetInputAttributeNameWithOverride(i, Params);
+		if (SourceAttributeName == NAME_None)
+		{
+			SourceAttributeName = SourceMetadata[i]->GetSingleAttributeNameOrNone();
+		}
+
+		OperationData.SourceAttributes[i] = SourceMetadata[i]->GetConstAttribute(SourceAttributeName);
+
+		if (OperationData.SourceAttributes[i] == nullptr)
+		{
+			PCGE_LOG(Error, "Attribute %s does not exist for input %d", *SourceAttributeName.ToString(), i);
+			return true;
+		}
+
+		// Then verify that the type is OK
+		bool bHasSpecialRequirement = false;
+		if (!Settings->IsSupportedInputType(OperationData.SourceAttributes[i]->GetTypeId(), i, bHasSpecialRequirement))
+		{
+			PCGE_LOG(Error, "Attribute %s is not a supported type for input %d", *SourceAttributeName.ToString(), i);
+			return true;
+		}
+
+		if (!bHasSpecialRequirement)
+		{
+			// In this case, check if we have a more complex type, or if we can broadcast to the most complex type.
+			if (MostComplexType == (uint16)EPCGMetadataTypes::Unknown || Settings->IsMoreComplexType(OperationData.SourceAttributes[i]->GetTypeId(), MostComplexType))
+			{
+				MostComplexType = OperationData.SourceAttributes[i]->GetTypeId();
+			}
+			else if (MostComplexType != OperationData.SourceAttributes[i]->GetTypeId() && !PCG::Private::IsBroadcastable(OperationData.SourceAttributes[i]->GetTypeId(), MostComplexType))
+			{
+				PCGE_LOG(Error, "Attribute %s cannot be broadcasted to match types for input %d", *SourceAttributeName.ToString(), i);
+				return true;
+			}
+		}
+
+		// Finally check that we have the right number of inputs, depending on the source
+		if (const UPCGPointData* PointData = Cast<UPCGPointData>(Data[i]))
+		{
+			NumberOfElements[i] = PointData->GetPoints().Num();
+
+			// If we are not the first input, we only have an iterator if we have a single point and we broadcast.
+			bool bShouldBroadcast = i != 0 && NumberOfElements[i] == 1 && Settings->Mode == EPCGMetadataSettingsBaseMode::Broadcast;
+			if (i == 0 || NumberOfElements[i] == NumberOfElements[0] || bShouldBroadcast)
+			{
+				OperationData.Iterators[i] = MakeUnique<FPCGMetadataEntryPointIterator>(PointData, bShouldBroadcast);
+			}
+		}
+		else
+		{
+			NumberOfElements[i] = OperationData.SourceAttributes[i]->GetNumberOfEntries();
+
+			// No element means only dealing with the default value.
+			if (NumberOfElements[i] == 0)
+			{
+				OperationData.Iterators[i] = MakeUnique<FPCGMetadataEntryConstantIterator>(PCGInvalidEntryKey, /*bRepeat=*/ true);
+			}
+			else
+			{
+				// Broadcast only if we have a single element and we are in broadcast mode, or inferred mode and we are param data.
+				bool bShouldBroadcast = i != 0 && 
+					NumberOfElements[i] == 1 && 
+					(Settings->Mode == EPCGMetadataSettingsBaseMode::Broadcast || (Settings->Mode == EPCGMetadataSettingsBaseMode::Inferred && Data[i]->IsA<UPCGParamData>()));
+
+				if (i == 0 || NumberOfElements[i] == NumberOfElements[0] || bShouldBroadcast)
+				{
+					OperationData.Iterators[i] = MakeUnique<FPCGMetadataEntryAttributeIterator>(*OperationData.SourceAttributes[i], bShouldBroadcast);
+				}
+			}
+		}
+
+		if (OperationData.NumberOfElementsToProcess == -1)
+		{
+			OperationData.NumberOfElementsToProcess = NumberOfElements[i];
+		}
+	}
+
+	// At this point, we verified everything, so we can go forward with the computation, depending on the most complex type
+	// So first forward outputs and create the attribute
+	OperationData.OutputAttribute = nullptr;
+	auto CreateAttribute = [&SourceAttributes = OperationData.SourceAttributes, &Inputs, &Outputs, &SourceMetadata, OutputAttributeName, &OutputAttribute = OperationData.OutputAttribute](auto DummyOutValue) -> void
+	{
+		using AttributeType = decltype(DummyOutValue);
+
+		AttributeType DefaultValue = PCGMetadataAttribute::GetValueWithBroadcast<AttributeType>(SourceAttributes[0], PCGDefaultValueKey);
+
+		FPCGTaggedData& OutputData = Outputs.Add_GetRef(Inputs[0]);
+		OutputData.Pin = PCGPinConstants::DefaultOutputLabel;
+
+		UPCGMetadata* OutMetadata = nullptr;
+
+		PCGMetadataElementCommon::DuplicateTaggedData(Inputs[0], OutputData, OutMetadata);
+		OutputAttribute = PCGMetadataElementCommon::ClearOrCreateAttribute(OutMetadata, OutputAttributeName, DefaultValue);
+		PCGMetadataElementCommon::CopyEntryToValueKeyMap(SourceMetadata[0], SourceAttributes[0], OutputAttribute);
+	};
+
+	OperationData.OutputType = Settings->GetOutputType(MostComplexType);
+
+	PCGMetadataAttribute::CallbackWithRightType(OperationData.OutputType, CreateAttribute);
+
+	if (OperationData.OutputAttribute == nullptr)
+	{
+		PCGE_LOG(Error, "Error while creating output attribute");
+		Outputs.Empty();
+		return true;
+	}
+
+	OperationData.Settings = Settings;
+
+	return DoOperation(OperationData);
+}
\ No newline at end of file
diff --git a/Engine/Plugins/Experimental/PCG/Source/PCG/Private/Metadata/PCGMetadataEntryKeyIterator.cpp b/Engine/Plugins/Experimental/PCG/Source/PCG/Private/Metadata/PCGMetadataEntryKeyIterator.cpp
new file mode 100644
--- /dev/null
+++ b/Engine/Plugins/Experimental/PCG/Source/PCG/Private/Metadata/PCGMetadataEntryKeyIterator.cpp
@@ -0,0 +1,149 @@
+// Copyright Epic Games, Inc. All Rights Reserved.
+
+#pragma once
+
+#include "Metadata/PCGMetadataEntryKeyIterator.h"
+
+#include "PCGPoint.h"
+#include "Data/PCGPointData.h"
+#include "Metadata/PCGMetadataAttribute.h"
+
+// FPCGMetadataEntryAttributeIterator
+
+FPCGMetadataEntryAttributeIterator::FPCGMetadataEntryAttributeIterator(const FPCGMetadataAttributeBase& InAttribute, bool bRepeat)
+	: IPCGMetadataEntryIterator(bRepeat)
+	, OriginalAttribute(InAttribute)
+	, CurrentAttribute(&InAttribute)
+{
+	ItPtr = new Iterator(InAttribute.GetEntryToValueKeyMap_NotThreadSafe());
+
+	FindValidPtr();
+
+	if (!(*ItPtr) && CurrentAttribute->GetParent() == nullptr)
+	{
+		bIsInvalid = true;
+	}
+}
+
+void FPCGMetadataEntryAttributeIterator::FindValidPtr()
+{
+	while (!(*ItPtr) && CurrentAttribute->GetParent() != nullptr)
+	{
+		CurrentAttribute = CurrentAttribute->GetParent();
+		ResetItPtr();
+	}
+}
+
+void FPCGMetadataEntryAttributeIterator::ResetItPtr()
+{
+	// Since Iterators doesn't support copy nor move constructors we need to
+	// do this ugly dance with an explicit call to the destructor and a placement new to
+	// call the constructor again.
+	if (!bIsInvalid && CurrentAttribute)
+	{
+		ItPtr->~Iterator();
+		new(ItPtr) Iterator(CurrentAttribute->GetEntryToValueKeyMap_NotThreadSafe());
+	}
+}
+
+FPCGMetadataEntryAttributeIterator::~FPCGMetadataEntryAttributeIterator()
+{
+	delete ItPtr;
+}
+
+PCGMetadataEntryKey FPCGMetadataEntryAttributeIterator::operator*()
+{
+	return (*ItPtr)->Key;
+}
+
+IPCGMetadataEntryIterator& FPCGMetadataEntryAttributeIterator::operator++()
+{
+	Iterator& It = *ItPtr;
+
+	check(!bIsInvalid && !!It);
+
+	// Special case where there is no parent, has a single element and is repeat, we do nothing.
+	if (OriginalAttribute.GetNumberOfEntries() == 1 && OriginalAttribute.GetParent() == nullptr && IsRepeat())
+	{
+		return *this;
+	}
+
+	// Otherwise, increment the iterator
+	if (!(++It))
+	{
+		// If it is not valid, try to find a valid one in its parents
+		FindValidPtr();
+
+		// If we reached the root of the parents, the iterator is still not valid and we repeat, 
+		// we reset to the original attribute.
+		if (!It && CurrentAttribute->GetParent() == nullptr && IsRepeat())
+		{
+			CurrentAttribute = &OriginalAttribute;
+			ResetItPtr();
+			FindValidPtr();
+		}
+	}
+
+	return *this;
+}
+
+bool FPCGMetadataEntryAttributeIterator::IsEnd() const
+{
+	return bIsInvalid || !!(*ItPtr);
+}
+
+// FPCGMetadataEntryPointIterator
+
+FPCGMetadataEntryPointIterator::FPCGMetadataEntryPointIterator(const UPCGPointData* InPointData, bool bRepeat)
+	: IPCGMetadataEntryIterator(bRepeat)
+	, PointData(InPointData)
+	, CurrentIndex(0)
+{
+
+}
+
+PCGMetadataEntryKey FPCGMetadataEntryPointIterator::operator*()
+{
+	return PointData->GetPoints()[CurrentIndex].MetadataEntry;
+}
+
+IPCGMetadataEntryIterator& FPCGMetadataEntryPointIterator::operator++()
+{
+	if (++CurrentIndex >= PointData->GetPoints().Num() && IsRepeat())
+	{
+		CurrentIndex = 0;
+	}
+
+	return *this;
+}
+
+bool FPCGMetadataEntryPointIterator::IsEnd() const
+{
+	return CurrentIndex >= PointData->GetPoints().Num();
+}
+
+// FPCGMetadataEntryConstantIterator
+
+FPCGMetadataEntryConstantIterator::FPCGMetadataEntryConstantIterator(PCGMetadataEntryKey InKey, bool bRepeat)
+	: IPCGMetadataEntryIterator(bRepeat)
+	, Key(InKey)
+	, bHasEnded(false)
+{
+
+}
+
+PCGMetadataEntryKey FPCGMetadataEntryConstantIterator::operator*()
+{
+	return Key;
+}
+
+IPCGMetadataEntryIterator& FPCGMetadataEntryConstantIterator::operator++()
+{
+	bHasEnded = !IsRepeat();
+	return *this;
+}
+
+bool FPCGMetadataEntryConstantIterator::IsEnd() const
+{
+	return bHasEnded;
+}
\ No newline at end of file
diff --git a/Engine/Plugins/Experimental/PCG/Source/PCG/Public/Elements/Metadata/PCGMetadataMathsOpElement.h b/Engine/Plugins/Experimental/PCG/Source/PCG/Public/Elements/Metadata/PCGMetadataMathsOpElement.h
new file mode 100644
--- /dev/null
+++ b/Engine/Plugins/Experimental/PCG/Source/PCG/Public/Elements/Metadata/PCGMetadataMathsOpElement.h
@@ -0,0 +1,86 @@
+// Copyright Epic Games, Inc. All Rights Reserved.
+
+#pragma once
+
+#include "PCGElement.h"
+#include "PCGSettings.h"
+#include "Elements/Metadata/PCGMetadataOpElementBase.h"
+
+#include "PCGMetadataMathsOpElement.generated.h"
+
+UENUM(Meta=(Bitflags))
+enum class EPCGMedadataMathsOperation : uint16
+{
+	// Unary op
+	UnaryOp = 1 << 10 UMETA(Hidden),
+	Sign,
+	Frac,
+	Truncate,
+	Round,
+	Sqrt,
+	Abs,
+
+	// Binary op
+	BinaryOp = 1 << 11 UMETA(Hidden),
+	Add,
+	Subtract,
+	Multiply,
+	Divide,
+	Max,
+	Min,
+	Pow,
+	ClampMin,
+	ClampMax,
+
+	// Ternary op
+	TernaryOp = 1 << 12 UMETA(Hidden),
+	Clamp,
+	Lerp,
+};
+ENUM_CLASS_FLAGS(EPCGMedadataMathsOperation);
+
+UCLASS(BlueprintType, ClassGroup = (Procedural))
+class PCG_API UPCGMetadataMathsSettings : public UPCGMetadataSettingsBase
+{
+	GENERATED_BODY()
+
+public:
+	//~Begin UPCGSettings interface
+#if WITH_EDITOR
+	virtual FName GetDefaultNodeName() const override;
+#endif
+	//~End UPCGSettings interface
+
+	//~Begin UPCGMetadataSettingsBase interface
+	virtual FName GetInputAttributeNameWithOverride(uint32 Index, UPCGParamData* Params) const override;
+
+	virtual FName GetInputPinLabel(uint32 Index) const override;
+	virtual uint32 GetInputPinNum() const override;
+
+	virtual bool IsSupportedInputType(uint16 TypeId, uint32 InputIndex, bool& bHasSpecialRequirement) const override;
+	//~End UPCGMetadataSettingsBase interface
+
+protected:
+	//~Begin UPCGSettings interface
+	virtual FPCGElementPtr CreateElement() const override;
+	//~End UPCGSettings interface
+
+public:
+	UPROPERTY(BlueprintReadWrite, EditAnywhere, Category = Settings)
+	EPCGMedadataMathsOperation Operation = EPCGMedadataMathsOperation::Add;
+
+	UPROPERTY(BlueprintReadWrite, EditAnywhere, Category = Input)
+	FName Input1AttributeName = NAME_None;
+
+	UPROPERTY(BlueprintReadWrite, EditAnywhere, Category = Input, meta = (EditCondition="(Operation & '/Script/PCG.EPCGMedadataMathsOperation::Binary') || (Operation & '/Script/PCG.EPCGMedadataMathsOperation::TernaryOp')", EditConditionHides))
+	FName Input2AttributeName = NAME_None;
+
+	UPROPERTY(BlueprintReadWrite, EditAnywhere, Category = Input, meta = (EditCondition = "Operation & '/Script/PCG.EPCGMedadataMathsOperation::TernaryOp'", EditConditionHides))
+	FName Input3AttributeName = NAME_None;
+};
+
+class FPCGMetadataMathsElement : public FPCGMetadataElementBase
+{
+protected:
+	virtual bool DoOperation(FOperationData& OperationData) const override;
+};
diff --git a/Engine/Plugins/Experimental/PCG/Source/PCG/Public/Elements/Metadata/PCGMetadataOpElementBase.h b/Engine/Plugins/Experimental/PCG/Source/PCG/Public/Elements/Metadata/PCGMetadataOpElementBase.h
new file mode 100644
--- /dev/null
+++ b/Engine/Plugins/Experimental/PCG/Source/PCG/Public/Elements/Metadata/PCGMetadataOpElementBase.h
@@ -0,0 +1,111 @@
+// Copyright Epic Games, Inc. All Rights Reserved.
+
+#pragma once
+
+#include "PCGElement.h"
+#include "PCGSettings.h"
+
+#include "PCGMetadataOpElementBase.generated.h"
+
+class FPCGMetadataAttributeBase;
+class IPCGMetadataEntryIterator;
+
+namespace PCGMetadataSettingsBaseConstants
+{
+	const FName DoubleInputFirstLabel = TEXT("InA");
+	const FName DoubleInputSecondLabel = TEXT("InB");
+
+	const FName ClampMinLabel = TEXT("Min");
+	const FName ClampMaxLabel = TEXT("Max");
+	const FName LerpRatioLabel = TEXT("Ratio");
+}
+
+// Defines behavior when number of entries doesn't match in inputs
+UENUM()
+enum class EPCGMetadataSettingsBaseMode
+{
+	Inferred     UMETA(Tooltip = "Broadcast for ParamData and no broadcast for SpatialData."),
+	NoBroadcast  UMETA(ToolTip = "If number of entries doesn't match, will use the default value."),
+	Broadcast    UMETA(ToolTip = "If there is no entry or a single entry, will repeat this value.")
+};
+
+UENUM()
+enum class EPCGMetadataSettingsBaseTypes
+{
+	AutoUpcastTypes,
+	StrictTypes
+};
+
+/**
+ * Base class for all Metadata operations
+ * A metadata operation can work on 2 different type of inputs: ParamData and SpatialData
+ * Each of those inputs can have some metadata.
+ * The output will contain the metadata of the first input (all its attributes) + the result of the operation (in a separate attribute)
+ * The new attribute can collide with one of the attributes in the incoming metadata. In this case, the attribute value will be overridden by the result
+ * of the operation. It will also override the type of the attribute if it doesn't match the original.
+ * 
+ * You can specify the name of the attribute for each input and for the output. If they are None, they will take the default attribute.
+ * Attribute names can also be overridden by ParamData, just connect the Param pin with some param data that matches exactly the name of the property you want to override.
+ * 
+ * Each operation has some requirements for the input types, and can broadcast some values into others (example Vector + Float -> Vector).
+ * For example, if the op only accept booleans, all other value types will throw an error.
+ * 
+ * If there are multiple values for an attribute, the operation will be done on all values. If one input has N elements and the second has 1 element,
+ * the second will be repeated for each element of the first for the operation. We only support N-N operations and N-1 operation (ie. The number of values
+ * needs to be all the same or 1)
+ * 
+ * If the node doesn't provide an output, check the logs to know why it failed.
+ */
+UCLASS(BlueprintType, Abstract, ClassGroup = (Procedural))
+class PCG_API UPCGMetadataSettingsBase : public UPCGSettings
+{
+	GENERATED_BODY()
+
+public:
+	//~Begin UPCGSettings interface
+#if WITH_EDITOR
+	virtual EPCGSettingsType GetType() const override { return EPCGSettingsType::Metadata; }
+#endif
+
+	virtual TArray<FPCGPinProperties> InputPinProperties() const override;
+	virtual TArray<FPCGPinProperties> OutputPinProperties() const override;
+	//~End UPCGSettings interface
+
+	virtual FName GetInputAttributeNameWithOverride(uint32 Index, UPCGParamData* Params) const { return NAME_None; };
+
+	virtual FName GetInputPinLabel(uint32 Index) const { return NAME_None; }
+	virtual uint32 GetInputPinNum() const { return 0; };
+
+	virtual bool IsSupportedInputType(uint16 TypeId, uint32 InputIndex, bool& bHasSpecialRequirement) const { return false; };
+	virtual uint16 GetOutputType(uint16 InputTypeId) const { return InputTypeId; };
+
+	bool IsMoreComplexType(uint16 FirstType, uint16 SecondType) const;
+
+	UPROPERTY(BlueprintReadWrite, EditAnywhere, Category = Output)
+	FName OutputAttributeName = NAME_None;
+
+	UPROPERTY(BlueprintReadWrite, EditAnywhere, Category = Settings)
+	EPCGMetadataSettingsBaseMode Mode = EPCGMetadataSettingsBaseMode::Inferred;
+};
+
+
+class FPCGMetadataElementBase : public FSimplePCGElement
+{
+public:
+	virtual bool IsCacheable(const UPCGSettings* InSettings) const override { return true; }
+
+protected:
+	virtual bool ExecuteInternal(FPCGContext* Context) const override;
+
+	struct FOperationData
+	{
+		TArray<const FPCGMetadataAttributeBase*> SourceAttributes;
+		int32 NumberOfElementsToProcess = -1;
+		uint16 OutputType;
+		FPCGMetadataAttributeBase* OutputAttribute = nullptr;
+		const UPCGMetadataSettingsBase* Settings = nullptr;
+		TArray<TUniquePtr<IPCGMetadataEntryIterator>> Iterators;
+	};
+
+	virtual bool DoOperation(FOperationData& OperationData) const { return true; };
+};
\ No newline at end of file
diff --git a/Engine/Plugins/Experimental/PCG/Source/PCG/Public/Metadata/PCGMetadataAttribute.h b/Engine/Plugins/Experimental/PCG/Source/PCG/Public/Metadata/PCGMetadataAttribute.h
--- a/Engine/Plugins/Experimental/PCG/Source/PCG/Public/Metadata/PCGMetadataAttribute.h
+++ b/Engine/Plugins/Experimental/PCG/Source/PCG/Public/Metadata/PCGMetadataAttribute.h
@@ -39,6 +39,13 @@ public:
 
 	bool AllowsInterpolation() const { return bAllowsInterpolation; }
 
+	int32 GetNumberOfEntries() const { return EntryToValueKeyMap.Num(); }
+
+	// This call is not thread safe
+	const TMap<PCGMetadataEntryKey, PCGMetadataValueKey>& GetEntryToValueKeyMap_NotThreadSafe() const { return EntryToValueKeyMap; }
+
+	const FPCGMetadataAttributeBase* GetParent() const { return Parent; }
+
 protected:
 	TMap<PCGMetadataEntryKey, PCGMetadataValueKey> EntryToValueKeyMap;
 	mutable FRWLock EntryMapLock;
diff --git a/Engine/Plugins/Experimental/PCG/Source/PCG/Public/Metadata/PCGMetadataAttributeTpl.h b/Engine/Plugins/Experimental/PCG/Source/PCG/Public/Metadata/PCGMetadataAttributeTpl.h
--- a/Engine/Plugins/Experimental/PCG/Source/PCG/Public/Metadata/PCGMetadataAttributeTpl.h
+++ b/Engine/Plugins/Experimental/PCG/Source/PCG/Public/Metadata/PCGMetadataAttributeTpl.h
@@ -267,6 +267,11 @@ public:
 		return PCGDefaultValueKey;
 	}
 
+	void SetDefaultValue(const T& Value)
+	{
+		DefaultValue = Value;
+	}
+
 protected:
 	/** Code related to computing compared values (min, max, sub, add) */
 	template<typename IT = T, typename TEnableIf<PCG::Private::MetadataTraits<IT>::CanMinMax>::Type* = nullptr>
@@ -441,6 +446,7 @@ namespace PCGMetadataAttribute
 			AllocatePCGMetadataAttributeOnType(int32);
 			AllocatePCGMetadataAttributeOnType(int64);
 			AllocatePCGMetadataAttributeOnType(FVector);
+			AllocatePCGMetadataAttributeOnType(FVector2D);
 			AllocatePCGMetadataAttributeOnType(FVector4);
 			AllocatePCGMetadataAttributeOnType(FQuat);
 			AllocatePCGMetadataAttributeOnType(FTransform);
@@ -455,4 +461,118 @@ namespace PCGMetadataAttribute
 
 #undef AllocatePCGMetadataAttributeOnType
 	}
+
+	template <typename Func>
+	inline decltype(auto) CallbackWithRightType(uint16 TypeId, Func Callback)
+	{
+		using ReturnType = decltype(Callback(double{}));
+
+		switch (TypeId)
+		{
+		case (uint16)EPCGMetadataTypes::Integer32:
+			return Callback(int32{});
+		case (uint16)EPCGMetadataTypes::Integer64:
+			return Callback(int64{});
+		case (uint16)EPCGMetadataTypes::Float:
+			return Callback(float{});
+		case (uint16)EPCGMetadataTypes::Double:
+			return Callback(double{});
+		case (uint16)EPCGMetadataTypes::Vector2:
+			return Callback(FVector2D{});
+		case (uint16)EPCGMetadataTypes::Vector:
+			return Callback(FVector{});
+		case (uint16)EPCGMetadataTypes::Vector4:
+			return Callback(FVector4{});
+		case (uint16)EPCGMetadataTypes::Quaternion:
+			return Callback(FQuat{});
+		case (uint16)EPCGMetadataTypes::Transform:
+			return Callback(FTransform{});
+		case (uint16)EPCGMetadataTypes::String:
+			return Callback(FString{});
+		case (uint16)EPCGMetadataTypes::Boolean:
+			return Callback(bool{});
+		case (uint16)EPCGMetadataTypes::Rotator:
+			return Callback(FRotator{});
+		case (uint16)EPCGMetadataTypes::Name:
+			return Callback(FName{});
+		default:
+		{
+			// ReturnType{} is invalid if ReturnType is void
+			if constexpr (std::is_same_v<ReturnType, void>)
+			{
+				return;
+			}
+			else
+			{
+				return ReturnType{};
+			}
+		}
+		}
+	}
+
+	template <typename OutType>
+	inline OutType GetValueWithBroadcast(const FPCGMetadataAttributeBase* InAttribute, PCGMetadataEntryKey InKey)
+	{
+		auto Func = [InAttribute, InKey](auto DummyInType) -> OutType
+		{
+			using InType = decltype(DummyInType);
+			InType Value = static_cast<const FPCGMetadataAttribute<InType>*>(InAttribute)->GetValueFromItemKey(InKey);
+
+			if constexpr (std::is_same_v<OutType, InType>)
+			{
+				return Value;
+			}
+			else
+			{
+				if constexpr (!PCG::Private::IsBroadcastable(PCG::Private::MetadataTypes<InType>::Id, PCG::Private::MetadataTypes<OutType>::Id))
+				{
+					return OutType{};
+				}
+				else
+				{
+					if constexpr (std::is_same_v<OutType, FVector4>)
+					{
+						if constexpr (std::is_same_v<InType, FVector> || std::is_same_v<InType, FVector2D>)
+						{
+							// TODO: Should it be 0? 1? Something else? Depending on operation?
+							// For now it is too ambiguous, so don't support it
+							return FVector4();
+						}
+						else
+						{
+							return FVector4(Value, Value, Value, Value);
+						}
+					}
+					else
+					{
+						// Seems like the && condition is not evaluated correctly on Linux, so we cut the condition in two `if constexpr`.
+						if constexpr (std::is_same_v<OutType, FVector>)
+						{
+							if constexpr (std::is_same_v<InType, FVector2D>)
+							{
+								return FVector(Value, 0.0);
+							}
+							else
+							{
+								return OutType(Value);
+							}
+						}
+						else
+						{
+							return OutType(Value);
+						}
+					}
+				}
+			}
+		};
+
+		if (PCG::Private::MetadataTypes<OutType>::Id == InAttribute->GetTypeId())
+		{
+			return Func(OutType{});
+		}
+		else
+		{
+			return CallbackWithRightType(InAttribute->GetTypeId(), Func);
+		}
+	}
 }
\ No newline at end of file
diff --git a/Engine/Plugins/Experimental/PCG/Source/PCG/Public/Metadata/PCGMetadataAttributeTraits.h b/Engine/Plugins/Experimental/PCG/Source/PCG/Public/Metadata/PCGMetadataAttributeTraits.h
--- a/Engine/Plugins/Experimental/PCG/Source/PCG/Public/Metadata/PCGMetadataAttributeTraits.h
+++ b/Engine/Plugins/Experimental/PCG/Source/PCG/Public/Metadata/PCGMetadataAttributeTraits.h
@@ -22,6 +22,9 @@ enum class EPCGMetadataTypes : uint8
 	Boolean,
 	Rotator,
 	Name,
+
+	Count UMETA(Hidden),
+
 	Unknown = 255
 };
 
@@ -53,6 +56,68 @@ namespace PCG
 
 #undef PCGMetadataGenerateDataTypes
 
+		// Wrapper around a standard 2-dimensional CArray that is constexpr, to know if a type is broadcastable to another.
+		// First index is the original type, second index is the wanted type. Returns true if we can broadcast first type into second type.
+		struct UBroadcastableTypes
+		{
+			constexpr UBroadcastableTypes() : Values{{false}}
+			{
+				for (uint8 i = 0; i < (uint8)EPCGMetadataTypes::Count; ++i)
+				{
+					Values[i][i] = true;
+				}
+
+#define PCGMetadataBroadcastable(FirstType, SecondType) Values[(uint8)EPCGMetadataTypes::FirstType][(uint8)EPCGMetadataTypes::SecondType] = true
+
+				// Finally set all cases where it is broadcastable
+				PCGMetadataBroadcastable(Float, Double);
+				PCGMetadataBroadcastable(Float, Vector2);
+				PCGMetadataBroadcastable(Float, Vector);
+				PCGMetadataBroadcastable(Float, Vector4);
+
+				PCGMetadataBroadcastable(Double, Vector2);
+				PCGMetadataBroadcastable(Double, Vector);
+				PCGMetadataBroadcastable(Double, Vector4);
+
+				PCGMetadataBroadcastable(Integer32, Float);
+				PCGMetadataBroadcastable(Integer32, Double);
+				PCGMetadataBroadcastable(Integer32, Integer64);
+				PCGMetadataBroadcastable(Integer32, Vector2);
+				PCGMetadataBroadcastable(Integer32, Vector);
+				PCGMetadataBroadcastable(Integer32, Vector4);
+
+				PCGMetadataBroadcastable(Integer64, Float);
+				PCGMetadataBroadcastable(Integer64, Double);
+				PCGMetadataBroadcastable(Integer64, Vector2);
+				PCGMetadataBroadcastable(Integer64, Vector);
+				PCGMetadataBroadcastable(Integer64, Vector4);
+
+				PCGMetadataBroadcastable(Vector2, Vector);
+
+				PCGMetadataBroadcastable(Quaternion, Rotator);
+				PCGMetadataBroadcastable(Rotator, Quaternion);
+
+#undef PCGMetadataBroadcastable
+			}
+
+			constexpr const bool* operator[](int i) const { return Values[i]; }
+
+			bool Values[(uint8)EPCGMetadataTypes::Count][(uint8)EPCGMetadataTypes::Count];
+		};
+
+		inline static constexpr UBroadcastableTypes BroadcastableTypes{};
+
+		constexpr inline bool IsBroadcastable(uint16 FirstType, uint16 SecondType)
+		{
+			// Unknown types aren't broadcastable
+			if (FirstType >= static_cast<uint16>(EPCGMetadataTypes::Count) || SecondType >= static_cast<uint16>(EPCGMetadataTypes::Count))
+			{
+				return false;
+			}
+
+			return BroadcastableTypes[FirstType][SecondType];
+		}
+
 		template<typename T>
 		struct DefaultOperationTraits
 		{
diff --git a/Engine/Plugins/Experimental/PCG/Source/PCG/Public/Metadata/PCGMetadataEntryKeyIterator.h b/Engine/Plugins/Experimental/PCG/Source/PCG/Public/Metadata/PCGMetadataEntryKeyIterator.h
new file mode 100644
--- /dev/null
+++ b/Engine/Plugins/Experimental/PCG/Source/PCG/Public/Metadata/PCGMetadataEntryKeyIterator.h
@@ -0,0 +1,101 @@
+// Copyright Epic Games, Inc. All Rights Reserved.
+
+#pragma once
+
+#include "Metadata/PCGMetadataCommon.h"
+
+class UPCGPointData;
+class FPCGMetadataAttributeBase;
+
+/**
+ * Base class for an iterator on EntryKeys for Metadata
+ * This is used to abstract multiple ways to iterate over entry keys.
+ * Its main use is for metadata operations
+ * 
+ * It is behaving as a normal iterator. It can be dereferenced, pre-incremented
+ * and we can check if we reached the end of the iterator.
+ * Also, the iterator can be marked to automatically "reset" if it reached the end,
+ * allowing to loop over the same data.
+ */
+class IPCGMetadataEntryIterator
+{
+public:
+	IPCGMetadataEntryIterator(bool bInRepeat) : bRepeat(bInRepeat) {}
+	virtual ~IPCGMetadataEntryIterator() = default;
+
+	// Can't copy
+	IPCGMetadataEntryIterator(const IPCGMetadataEntryIterator& Other) = delete;
+	IPCGMetadataEntryIterator& operator=(const IPCGMetadataEntryIterator& Other) = delete;
+
+	IPCGMetadataEntryIterator(IPCGMetadataEntryIterator&& Other) = default;
+	IPCGMetadataEntryIterator& operator=(IPCGMetadataEntryIterator&& Other) = default;
+
+	virtual PCGMetadataEntryKey operator*() = 0;
+	virtual IPCGMetadataEntryIterator& operator++() = 0;
+	virtual bool IsEnd() const = 0;
+
+	bool IsRepeat() { return bRepeat; }
+
+private:
+	bool bRepeat = false;
+};
+
+/**
+ * Will iterate over all the EntryKeys in the EntryKeyToValueKey mapping
+ * for this attribute, including its parents.
+ * 
+ * Note that this iterator is not thread safe on the Attribute entry key map.
+ */
+class FPCGMetadataEntryAttributeIterator : public IPCGMetadataEntryIterator
+{
+public:
+	FPCGMetadataEntryAttributeIterator(const FPCGMetadataAttributeBase& InAttribute, bool bRepeat);
+	~FPCGMetadataEntryAttributeIterator();
+	virtual PCGMetadataEntryKey operator*() override;
+	virtual IPCGMetadataEntryIterator& operator++() override;
+	virtual bool IsEnd() const override;
+
+private:
+	using Iterator = TMap<PCGMetadataEntryKey, PCGMetadataValueKey>::TConstIterator;
+
+	void FindValidPtr();
+	void ResetItPtr();
+
+	const FPCGMetadataAttributeBase& OriginalAttribute;
+	const FPCGMetadataAttributeBase* CurrentAttribute = nullptr;
+	Iterator* ItPtr = nullptr;
+	bool bIsInvalid = false;
+};
+
+/**
+ * Will iterate over all the EntryKeys for each point in the point data.
+ */
+class FPCGMetadataEntryPointIterator : public IPCGMetadataEntryIterator
+{
+public:
+	FPCGMetadataEntryPointIterator(const UPCGPointData* InPointData, bool bRepeat);
+	virtual PCGMetadataEntryKey operator*() override;
+	virtual IPCGMetadataEntryIterator& operator++() override;
+	virtual bool IsEnd() const override;
+
+private:
+	const UPCGPointData* PointData;
+	int32 CurrentIndex = 0;
+};
+
+/**
+ * Will return a constant value. If repeat is not set, it will return a constant value
+ * only once before marking it ended.
+ */
+class FPCGMetadataEntryConstantIterator : public IPCGMetadataEntryIterator
+{
+public:
+	FPCGMetadataEntryConstantIterator(PCGMetadataEntryKey InKey, bool bRepeat);
+	virtual PCGMetadataEntryKey operator*() override;
+	virtual IPCGMetadataEntryIterator& operator++() override;
+	virtual bool IsEnd() const override;
+
+private:
+	PCGMetadataEntryKey Key;
+	bool bHasEnded = false;
+};
\ No newline at end of file
diff --git a/Engine/Plugins/Experimental/PCG/Source/PCGEditor/Private/SPCGEditorGraphAttributeListView.cpp b/Engine/Plugins/Experimental/PCG/Source/PCGEditor/Private/SPCGEditorGraphAttributeListView.cpp
--- a/Engine/Plugins/Experimental/PCG/Source/PCGEditor/Private/SPCGEditorGraphAttributeListView.cpp
+++ b/Engine/Plugins/Experimental/PCG/Source/PCGEditor/Private/SPCGEditorGraphAttributeListView.cpp
@@ -252,6 +252,11 @@ FText SPCGListViewItemRow::ConvertMetadataAttributeToText(const FPCGMetadataAttr
 			const bool bMetaBool = static_cast<const FPCGMetadataAttribute<bool>*>(AttributeBase)->GetValueFromItemKey(ItemKey);
 			return bMetaBool ? FText::FromString(TEXT("true")) : FText::FromString(TEXT("false"));
 		}
+	case PCG::Private::MetadataTypes<FVector2D>::Id:
+	{
+		const FVector2D MetaVector = static_cast<const FPCGMetadataAttribute<FVector2D>*>(AttributeBase)->GetValueFromItemKey(ItemKey);
+		return FText::AsNumber(MetaVector[(*MetadataInfo).Index]);
+	}
 	case PCG::Private::MetadataTypes<FVector>::Id:
 		{
 			const FVector MetaVector = static_cast<const FPCGMetadataAttribute<FVector>*>(AttributeBase)->GetValueFromItemKey(ItemKey);
@@ -683,6 +688,12 @@ void SPCGEditorGraphAttributeListView::GenerateColumnsFromMetadata(const UPCGMet
 				AddMetadataColumn(ColumnName, AttributeName);
 				break;
 			}
+		case EPCGMetadataTypes::Vector2:
+		{
+			AddMetadataColumn(ColumnName, AttributeName, 0, TEXT("_X"));
+			AddMetadataColumn(ColumnName, AttributeName, 1, TEXT("_Y"));
+			break;
+		}
 		case EPCGMetadataTypes::Vector:
 			{
 				AddMetadataColumn(ColumnName, AttributeName, 0, TEXT("_X"));
