Fixed rooting issue between the graph cache and the temporary results in the graph executor

#rnx

diff --git a/Engine/Plugins/Experimental/PCG/Source/PCG/Private/Graph/PCGGraphCache.cpp b/Engine/Plugins/Experimental/PCG/Source/PCG/Private/Graph/PCGGraphCache.cpp
--- a/Engine/Plugins/Experimental/PCG/Source/PCG/Private/Graph/PCGGraphCache.cpp
+++ b/Engine/Plugins/Experimental/PCG/Source/PCG/Private/Graph/PCGGraphCache.cpp
@@ -30,8 +30,6 @@ FPCGGraphCacheEntry::FPCGGraphCacheEntry(const FPCGDataCollection& InInput, cons
 	: Input(InInput)
 	, Output(InOutput)
 {
-	// Note: we don't need to root the settings since they'll be owned by the subsystem
-	Settings = InSettings ? Cast<UPCGSettings>(StaticDuplicateObject(InSettings, InOwner.Get())) : nullptr;
 	SettingsCrc32 = InSettings ? InSettings->GetCrc32() : PCGGraphCache::NullSettingsCrc32;
 	ComponentSeed = PCGGraphCache::GetComponentSeed(InSettings, InComponent);
 
@@ -44,9 +42,10 @@ bool FPCGGraphCacheEntry::Matches(const FPCGDataCollection& InInput, int32 InSet
 	return (SettingsCrc32 == InSettingsCrc32) && (Input == InInput) && (ComponentSeed == InComponentSeed);
 }
 
-FPCGGraphCache::FPCGGraphCache(TWeakObjectPtr<UObject> InOwner)
-	: Owner(InOwner)
+FPCGGraphCache::FPCGGraphCache(TWeakObjectPtr<UObject> InOwner, FPCGRootSet* InRootSet)
+	: Owner(InOwner), RootSet(InRootSet)
 {
+	check(InOwner.Get() && InRootSet);
 }
 
 FPCGGraphCache::~FPCGGraphCache()
@@ -100,18 +99,25 @@ void FPCGGraphCache::StoreInCache(const IPCGElement* InElement, const FPCGDataCo
 		Entries = &(CacheData.Add(InElement));
 	}
 
-	Entries->Emplace(InInput, InSettings, InComponent, InOutput, Owner, RootSet);
+	Entries->Emplace(InInput, InSettings, InComponent, InOutput, Owner, *RootSet);
 }
 
 void FPCGGraphCache::ClearCache()
 {
 	FWriteScopeLock ScopedWriteLock(CacheLock);
 
+	// Unroot all previously rooted data
+	for (TPair<const IPCGElement*, FPCGGraphCacheEntries>& CacheEntry : CacheData)
+	{
+		for (FPCGGraphCacheEntry& Entry : CacheEntry.Value)
+		{
+			Entry.Input.RemoveFromRootSet(*RootSet);
+			Entry.Output.RemoveFromRootSet(*RootSet);
+		}
+	}
+
 	// Remove all entries
 	CacheData.Reset();
-
-	// Unroot all previously rooted data
-	RootSet.Clear();
 }
 
 #if WITH_EDITOR
@@ -128,13 +134,8 @@ void FPCGGraphCache::CleanFromCache(const IPCGElement* InElement)
 	{
 		for (FPCGGraphCacheEntry& Entry : *Entries)
 		{
-			Entry.Input.RemoveFromRootSet(RootSet);
-			Entry.Output.RemoveFromRootSet(RootSet);
-
-			if (Entry.Settings)
-			{
-				RootSet.Remove(Entry.Settings);
-			}
+			Entry.Input.RemoveFromRootSet(*RootSet);
+			Entry.Output.RemoveFromRootSet(*RootSet);
 		}
 	}
 
diff --git a/Engine/Plugins/Experimental/PCG/Source/PCG/Private/Graph/PCGGraphCache.h b/Engine/Plugins/Experimental/PCG/Source/PCG/Private/Graph/PCGGraphCache.h
--- a/Engine/Plugins/Experimental/PCG/Source/PCG/Private/Graph/PCGGraphCache.h
+++ b/Engine/Plugins/Experimental/PCG/Source/PCG/Private/Graph/PCGGraphCache.h
@@ -18,7 +18,6 @@ struct FPCGGraphCacheEntry
 	bool Matches(const FPCGDataCollection& InInput, int32 InSettingsCrc32, int32 InComponentSeed) const;
 
 	FPCGDataCollection Input;
-	TObjectPtr<UPCGSettings> Settings;
 	FPCGDataCollection Output;
 	int32 SettingsCrc32;
 	int32 ComponentSeed;
@@ -34,7 +33,7 @@ typedef TArray<FPCGGraphCacheEntry> FPCGGraphCacheEntries;
 */
 struct FPCGGraphCache
 {
-	FPCGGraphCache(TWeakObjectPtr<UObject> InOwner);
+	FPCGGraphCache(TWeakObjectPtr<UObject> InOwner, FPCGRootSet* InRootSet);
 	~FPCGGraphCache();
 
 	/** Returns true if data was found from the cache, in which case the outputs are written in OutOutput */
@@ -58,7 +57,7 @@ private:
 	TWeakObjectPtr<UObject> Owner = nullptr;
 
 	/** To prevent garbage collection on data in the cache, we'll need to root some data */
-	FPCGRootSet RootSet;
+	FPCGRootSet* RootSet = nullptr;
 
 	mutable FRWLock CacheLock;
 };
\ No newline at end of file
diff --git a/Engine/Plugins/Experimental/PCG/Source/PCG/Private/Graph/PCGGraphExecutor.cpp b/Engine/Plugins/Experimental/PCG/Source/PCG/Private/Graph/PCGGraphExecutor.cpp
--- a/Engine/Plugins/Experimental/PCG/Source/PCG/Private/Graph/PCGGraphExecutor.cpp
+++ b/Engine/Plugins/Experimental/PCG/Source/PCG/Private/Graph/PCGGraphExecutor.cpp
@@ -38,7 +38,7 @@ static TAutoConsoleVariable<bool> CVarGraphMultithreading(
 
 FPCGGraphExecutor::FPCGGraphExecutor(UObject* InOwner)
 	: GraphCompiler(MakeUnique<FPCGGraphCompiler>())
-	, GraphCache(InOwner)
+	, GraphCache(InOwner, &DataRootSet)
 #if WITH_EDITOR
 	, GenerationProgressNotification(GetNotificationTextFormat())
 #endif
@@ -47,6 +47,11 @@ FPCGGraphExecutor::FPCGGraphExecutor(UObject* InOwner)
 
 FPCGGraphExecutor::~FPCGGraphExecutor()
 {
+	// We don't really need to do this here (it would be done in the destructor of these both)
+	// but this is to clarify/ensure the order in which this happens
+	GraphCache.ClearCache();
+	DataRootSet.Clear();
+
 #if WITH_EDITOR
 	EndGenerationNotification();
 #endif
@@ -297,7 +302,6 @@ void FPCGGraphExecutor::Execute()
 
 					FPCGContext* Context = Task.Element->Initialize(TaskInput, Task.SourceComponent, Task.Node);
 					Context->TaskId = Task.NodeId;
-					Context->Cache = &GraphCache;
 					ActiveTask.Context = TUniquePtr<FPCGContext>(Context);
 
 #if WITH_EDITOR
@@ -552,7 +556,7 @@ void FPCGGraphExecutor::StoreResults(FPCGTaskId InTaskId, const FPCGDataCollecti
 	OutputData.Add(InTaskId, InTaskOutput);
 
 	// Root any non-rooted results, otherwise they'll get garbage-collected
-	InTaskOutput.AddToRootSet(ResultsRootSet);
+	InTaskOutput.AddToRootSet(DataRootSet);
 }
 
 void FPCGGraphExecutor::ClearResults()
@@ -564,9 +568,12 @@ void FPCGGraphExecutor::ClearResults()
 	{
 		NextTaskId = 0;
 	}
-	OutputData.Reset();
 
-	ResultsRootSet.Clear();
+	for (const TPair<FPCGTaskId, FPCGDataCollection>& TaskOutput : OutputData)
+	{
+		TaskOutput.Value.RemoveFromRootSet(DataRootSet);
+	}
+	OutputData.Reset();
 
 	ScheduleLock.Unlock();
 }
diff --git a/Engine/Plugins/Experimental/PCG/Source/PCG/Private/Graph/PCGGraphExecutor.h b/Engine/Plugins/Experimental/PCG/Source/PCG/Private/Graph/PCGGraphExecutor.h
--- a/Engine/Plugins/Experimental/PCG/Source/PCG/Private/Graph/PCGGraphExecutor.h
+++ b/Engine/Plugins/Experimental/PCG/Source/PCG/Private/Graph/PCGGraphExecutor.h
@@ -114,6 +114,9 @@ private:
 	/** Graph compiler that turns a graph into tasks */
 	TUniquePtr<FPCGGraphCompiler> GraphCompiler;
 
+	/** Rootset to hold temporary results (for current computation) + graph cache */
+	FPCGRootSet DataRootSet;
+
 	/** Graph results cache */
 	FPCGGraphCache GraphCache;
 
@@ -127,7 +130,6 @@ private:
 	TArray<FPCGGraphTask> ReadyTasks;
 	TArray<FPCGGraphActiveTask> ActiveTasks;
 	TMap<FPCGTaskId, TSet<FPCGTaskId>> TaskSuccessors;
-	FPCGRootSet ResultsRootSet;
 	/** Map of node instances to their output, could be cleared once execution is done */
 	/** Note: this should at some point unload based on loaded/unloaded proxies, otherwise memory cost will be unbounded */
 	TMap<FPCGTaskId, FPCGDataCollection> OutputData;
diff --git a/Engine/Plugins/Experimental/PCG/Source/PCG/Public/PCGContext.h b/Engine/Plugins/Experimental/PCG/Source/PCG/Public/PCGContext.h
--- a/Engine/Plugins/Experimental/PCG/Source/PCG/Public/PCGContext.h
+++ b/Engine/Plugins/Experimental/PCG/Source/PCG/Public/PCGContext.h
@@ -37,7 +37,6 @@ struct PCG_API FPCGContext
 	FPCGDataCollection InputData;
 	FPCGDataCollection OutputData;
 	UPCGComponent* SourceComponent = nullptr;
-	FPCGGraphCache* Cache = nullptr;
 	int32 NumAvailableTasks = 0;
 
 	// TODO: add RNG source
