PCG:
*Moved class default details to toolbar instead of clicking the background of the node editor.

diff --git a/Engine/Plugins/Experimental/PCG/Source/PCGEditor/Private/PCGEditor.cpp b/Engine/Plugins/Experimental/PCG/Source/PCGEditor/Private/PCGEditor.cpp
--- a/Engine/Plugins/Experimental/PCG/Source/PCGEditor/Private/PCGEditor.cpp
+++ b/Engine/Plugins/Experimental/PCG/Source/PCGEditor/Private/PCGEditor.cpp
@@ -327,6 +327,13 @@ void FPCGEditor::RegisterToolbar() const
 			 TAttribute<FText>(),
 			 TAttribute<FText>(),
 			 FSlateIcon(FAppStyle::GetAppStyleSetName(), "BlueprintDebugger.TabIcon")));
+
+		Section.AddSeparator(NAME_None);
+		Section.AddEntry(FToolMenuEntry::InitToolBarButton(
+			 PCGEditorCommands.EditClassDefaults,
+			 TAttribute<FText>(),
+			 TAttribute<FText>(),
+			 FSlateIcon(FAppStyle::GetAppStyleSetName(), "FullBlueprintEditor.EditClassDefaults")));
 	}
 }
 
@@ -353,6 +360,12 @@ void FPCGEditor::BindCommands()
 		FExecuteAction::CreateSP(this, &FPCGEditor::OnDeterminismGraphTest),
 		FCanExecuteAction::CreateSP(this, &FPCGEditor::CanRunDeterminismGraphTest));
 
+	ToolkitCommands->MapAction(
+		PCGEditorCommands.EditClassDefaults,
+		FExecuteAction::CreateSP(this, &FPCGEditor::OnEditClassDefaults),
+		FCanExecuteAction(),
+		FIsActionChecked::CreateSP(this, &FPCGEditor::IsEditClassDefaultsToggled));
+	
 	GraphEditorCommands->MapAction(
 		PCGEditorCommands.CollapseNodes,
 		FExecuteAction::CreateSP(this, &FPCGEditor::OnCollapseNodesInSubgraph),
@@ -585,6 +598,17 @@ void FPCGEditor::OnDeterminismGraphTest()
 	}
 }
 
+void FPCGEditor::OnEditClassDefaults() const
+{
+	PropertyDetailsWidget->SetObject(PCGGraphBeingEdited);
+}
+
+bool FPCGEditor::IsEditClassDefaultsToggled() const
+{
+	const TArray<TWeakObjectPtr<UObject>>& SelectedObjects = PropertyDetailsWidget->GetSelectedObjects();
+	return SelectedObjects.Num() == 1 && SelectedObjects[0] == PCGGraphBeingEdited;
+}
+
 bool FPCGEditor::CanCollapseNodesInSubgraph() const
 {
 	bool HasPCGNode = false;
@@ -1565,27 +1589,20 @@ TSharedRef<SPCGEditorGraphProfilingView> FPCGEditor::CreateProfilingWidget()
 void FPCGEditor::OnSelectedNodesChanged(const TSet<UObject*>& NewSelection)
 {
 	TArray<TWeakObjectPtr<UObject>> SelectedObjects;
-
-	if (NewSelection.Num() == 0)
-	{
-		SelectedObjects.Add(PCGGraphBeingEdited);
-	}
-	else
+	
+	for (UObject* Object : NewSelection)
 	{
-		for (UObject* Object : NewSelection)
+		if (UPCGEditorGraphNodeBase* PCGGraphNode = Cast<UPCGEditorGraphNodeBase>(Object))
 		{
-			if (UPCGEditorGraphNodeBase* PCGGraphNode = Cast<UPCGEditorGraphNodeBase>(Object))
+			if (UPCGNode* PCGNode = PCGGraphNode->GetPCGNode())
 			{
-				if (UPCGNode* PCGNode = PCGGraphNode->GetPCGNode())
-				{
-					SelectedObjects.Add(PCGNode->DefaultSettings);
-				}
-			}
-			else if (UEdGraphNode* GraphNode = Cast<UEdGraphNode>(Object))
-			{
-				SelectedObjects.Add(GraphNode);
+				SelectedObjects.Add(PCGNode->DefaultSettings);
 			}
 		}
+		else if (UEdGraphNode* GraphNode = Cast<UEdGraphNode>(Object))
+		{
+			SelectedObjects.Add(GraphNode);
+		}
 	}
 
 	PropertyDetailsWidget->SetObjects(SelectedObjects, /*bForceRefresh=*/true);
diff --git a/Engine/Plugins/Experimental/PCG/Source/PCGEditor/Private/PCGEditor.h b/Engine/Plugins/Experimental/PCG/Source/PCGEditor/Private/PCGEditor.h
--- a/Engine/Plugins/Experimental/PCG/Source/PCGEditor/Private/PCGEditor.h
+++ b/Engine/Plugins/Experimental/PCG/Source/PCGEditor/Private/PCGEditor.h
@@ -119,6 +119,11 @@ private:
 	/** Run the determinism test on the current graph */
 	void OnDeterminismGraphTest();
 
+	/** Open details view for the PCG object being edited */
+	void OnEditClassDefaults() const;
+	/** Whether the PCG object being edited is opened in details view or not */
+	bool IsEditClassDefaultsToggled() const;
+
 	/** Whether or not an execution mode is active for the selected nodes */
 	bool IsExecutionModeActive(EPCGSettingsExecutionMode InExecutionMode) const;
 
diff --git a/Engine/Plugins/Experimental/PCG/Source/PCGEditor/Private/PCGEditorCommands.cpp b/Engine/Plugins/Experimental/PCG/Source/PCGEditor/Private/PCGEditorCommands.cpp
--- a/Engine/Plugins/Experimental/PCG/Source/PCGEditor/Private/PCGEditorCommands.cpp
+++ b/Engine/Plugins/Experimental/PCG/Source/PCGEditor/Private/PCGEditorCommands.cpp
@@ -24,6 +24,7 @@ void FPCGEditorCommands::RegisterCommands()
 	UI_COMMAND(ForceGraphRegeneration, "Force Regen", "Manually force a regeneration of the current graph.", EUserInterfaceActionType::Button, FInputChord());
 	UI_COMMAND(RunDeterminismNodeTest, "Run Determinism Test on Node", "Evaluate the current node for determinism.", EUserInterfaceActionType::Button, FInputChord(EModifierKey::Alt, EKeys::D));
 	UI_COMMAND(RunDeterminismGraphTest, "Graph Determinism Test", "Evaluate the current graph for determinism.", EUserInterfaceActionType::Button, FInputChord());
+	UI_COMMAND(EditClassDefaults, "Class Defaults", "Edit the initial values of your class.", EUserInterfaceActionType::ToggleButton, FInputChord());
 	UI_COMMAND(StartInspectNode, "Start Inspecting Node", "Start Inspecting Node", EUserInterfaceActionType::Button, FInputChord(EKeys::I));
 	UI_COMMAND(StopInspectNode, "Stop Inspecting Node", "Stop Inspecting Node", EUserInterfaceActionType::Button, FInputChord(EModifierKey::Shift, EKeys::I));
 	UI_COMMAND(ExecutionModeEnabled, "Execution Mode Enabled", "Set the execution mode for this node to Enabled.", EUserInterfaceActionType::RadioButton, FInputChord());
diff --git a/Engine/Plugins/Experimental/PCG/Source/PCGEditor/Private/PCGEditorCommands.h b/Engine/Plugins/Experimental/PCG/Source/PCGEditor/Private/PCGEditorCommands.h
--- a/Engine/Plugins/Experimental/PCG/Source/PCGEditor/Private/PCGEditorCommands.h
+++ b/Engine/Plugins/Experimental/PCG/Source/PCGEditor/Private/PCGEditorCommands.h
@@ -20,6 +20,7 @@ public:
 	TSharedPtr<FUICommandInfo> ForceGraphRegeneration;
 	TSharedPtr<FUICommandInfo> RunDeterminismNodeTest;
 	TSharedPtr<FUICommandInfo> RunDeterminismGraphTest;
+	TSharedPtr<FUICommandInfo> EditClassDefaults;
 	TSharedPtr<FUICommandInfo> StartInspectNode;
 	TSharedPtr<FUICommandInfo> StopInspectNode;
 	TSharedPtr<FUICommandInfo> ExecutionModeEnabled;
