[PCG] Metadata native operation: Boolean operations + CreateParamData native node

diff --git a/Engine/Plugins/Experimental/PCG/Source/PCG/Private/Elements/Metadata/PCGMetadataBooleanOpElement.cpp b/Engine/Plugins/Experimental/PCG/Source/PCG/Private/Elements/Metadata/PCGMetadataBooleanOpElement.cpp
new file mode 100644
--- /dev/null
+++ b/Engine/Plugins/Experimental/PCG/Source/PCG/Private/Elements/Metadata/PCGMetadataBooleanOpElement.cpp
@@ -0,0 +1,149 @@
+// Copyright Epic Games, Inc. All Rights Reserved.
+
+#include "Elements/Metadata/PCGMetadataBooleanOpElement.h"
+
+#include "PCGParamData.h"
+#include "Data/PCGSpatialData.h"
+#include "Helpers/PCGSettingsHelpers.h"
+#include "Metadata/PCGMetadata.h"
+#include "Metadata/PCGMetadataAttribute.h"
+#include "Metadata/PCGMetadataAttributeTpl.h"
+#include "Metadata/PCGMetadataEntryKeyIterator.h"
+#include "Elements/Metadata/PCGMetadataElementCommon.h"
+
+FName UPCGMetadataBooleanSettings::GetInputPinLabel(uint32 Index) const
+{
+	switch (Index)
+	{
+	case 0:
+		return (Operation == EPCGMedadataBooleanOperation::Not) ? PCGPinConstants::DefaultInputLabel : PCGMetadataSettingsBaseConstants::DoubleInputFirstLabel;
+	case 1:
+		return PCGMetadataSettingsBaseConstants::DoubleInputSecondLabel;
+	default:
+		return NAME_None;
+	}
+}
+
+uint32 UPCGMetadataBooleanSettings::GetInputPinNum() const
+{
+	return (Operation == EPCGMedadataBooleanOperation::Not) ? 1 : 2;
+}
+
+bool UPCGMetadataBooleanSettings::IsSupportedInputType(uint16 TypeId, uint32 InputIndex, bool& bHasSpecialRequirement) const
+{
+	bHasSpecialRequirement = false;
+	return TypeId == (uint16)EPCGMetadataTypes::Boolean;
+}
+
+FName UPCGMetadataBooleanSettings::GetInputAttributeNameWithOverride(uint32 Index, UPCGParamData* Params) const
+{
+	switch (Index)
+	{
+	case 0:
+		return PCG_GET_OVERRIDEN_VALUE(this, Input1AttributeName, Params);
+	case 1:
+		return PCG_GET_OVERRIDEN_VALUE(this, Input2AttributeName, Params);
+	default:
+		return NAME_None;
+	}
+}
+
+uint16 UPCGMetadataBooleanSettings::GetOutputType(uint16 InputTypeId) const
+{
+	return (uint16)EPCGMetadataTypes::Boolean;
+}
+
+#if WITH_EDITOR
+FName UPCGMetadataBooleanSettings::GetDefaultNodeName() const
+{
+	if (const UEnum* EnumPtr = FindObject<UEnum>(nullptr, TEXT("EPCGMedadataBooleanOperation"), true))
+	{ 
+		return EnumPtr->GetNameByValue(static_cast<int>(Operation)); 
+	}
+
+	return TEXT("Metadata Boolean Node");
+}
+#endif // WITH_EDITOR
+
+FPCGElementPtr UPCGMetadataBooleanSettings::CreateElement() const
+{
+	return MakeShared<FPCGMetadataBooleanElement>();
+}
+
+bool FPCGMetadataBooleanElement::DoOperation(FOperationData& OperationData) const
+{
+	TRACE_CPUPROFILER_EVENT_SCOPE(FPCGMetadataBooleanElement::Execute);
+
+	const UPCGMetadataBooleanSettings* Settings = CastChecked<UPCGMetadataBooleanSettings>(OperationData.Settings);
+
+	auto BoolFunc = [Settings](bool& Value1, const bool& Value2)
+	{
+		switch (Settings->Operation)
+		{
+		case EPCGMedadataBooleanOperation::And:
+			Value1 = (Value1 && Value2);
+			break;
+		case EPCGMedadataBooleanOperation::Not:
+			Value1 = !Value1;
+			break;
+		case EPCGMedadataBooleanOperation::Or:
+			Value1 = (Value1 || Value2);
+			break;
+		case EPCGMedadataBooleanOperation::Xor:
+			Value1 = (Value1 != Value2);
+			break;
+		}
+	};
+
+	// All values in OperationData.Iterators are UniquePtr. Dereference them here to make the syntax less cumbersome.
+	// Also some iterators can be null, it means they need to use the first iterator (that should never be null).
+	// We just need to make sure to not increment the first iterator multiple times in one loop.
+
+	check(OperationData.Iterators[0]);
+
+	IPCGMetadataEntryIterator& Iterator1 = *OperationData.Iterators[0];
+	IPCGMetadataEntryIterator& Iterator2 = (OperationData.Iterators.Num() >= 2 && OperationData.Iterators[1]) ? *OperationData.Iterators[1] : Iterator1;
+
+	bool bShouldIncrementIterator2 = &Iterator1 != &Iterator2;
+
+	PCGMetadataEntryKey EntryKey1 = 0;
+	PCGMetadataEntryKey EntryKey2 = 0;
+
+	bool Value1 = false;
+	bool Value2 = false;
+
+	FPCGMetadataAttribute<bool>* OutputAttribute = static_cast<FPCGMetadataAttribute<bool>*>(OperationData.OutputAttribute);
+	bool DefaultValue1 = PCGMetadataAttribute::GetValueWithBroadcast<bool>(OperationData.SourceAttributes[0], PCGInvalidEntryKey);
+	bool DefaultValue2 = (Settings->Operation != EPCGMedadataBooleanOperation::Not) ? PCGMetadataAttribute::GetValueWithBroadcast<bool>(OperationData.SourceAttributes[1], PCGInvalidEntryKey) : false;
+	BoolFunc(DefaultValue1, DefaultValue2);
+	OutputAttribute->SetDefaultValue(DefaultValue1);
+
+	for (int32 i = 0; i < OperationData.NumberOfElementsToProcess; ++i)
+	{
+		EntryKey1 = *Iterator1;
+
+		// If the entry key is invalid, nothing to do
+		if (EntryKey1 != PCGInvalidEntryKey)
+		{
+			Value1 = PCGMetadataAttribute::GetValueWithBroadcast<bool>(OperationData.SourceAttributes[0], EntryKey1);
+
+			if (Settings->Operation != EPCGMedadataBooleanOperation::Not)
+			{
+				EntryKey2 = *Iterator2;
+				Value2 = PCGMetadataAttribute::GetValueWithBroadcast<bool>(OperationData.SourceAttributes[1], EntryKey2);
+			}
+
+			BoolFunc(Value1, Value2);
+
+			OutputAttribute->SetValue(EntryKey1, Value1);
+		}
+
+		++Iterator1;
+		if (bShouldIncrementIterator2)
+		{
+			++Iterator2;
+		}
+	}
+
+	return true;
+}
\ No newline at end of file
diff --git a/Engine/Plugins/Experimental/PCG/Source/PCG/Private/Elements/Metadata/PCGMetadataMathsOpElement.cpp b/Engine/Plugins/Experimental/PCG/Source/PCG/Private/Elements/Metadata/PCGMetadataMathsOpElement.cpp
--- a/Engine/Plugins/Experimental/PCG/Source/PCG/Private/Elements/Metadata/PCGMetadataMathsOpElement.cpp
+++ b/Engine/Plugins/Experimental/PCG/Source/PCG/Private/Elements/Metadata/PCGMetadataMathsOpElement.cpp
@@ -278,7 +278,7 @@ bool FPCGMetadataMathsElement::DoOperation(FOperationData& OperationData) const
 					PCGMetadataEntryKey EntryKey = *Iterator1;
 
 					// If the entry key is invalid, nothing to do
-					if (EntryKey == PCGInvalidEntryKey)
+					if (EntryKey != PCGInvalidEntryKey)
 					{
 						AttributeType Value = PCGMetadataAttribute::GetValueWithBroadcast<AttributeType>(OperationData.SourceAttributes[0], EntryKey);
 
@@ -320,7 +320,7 @@ bool FPCGMetadataMathsElement::DoOperation(FOperationData& OperationData) const
 					PCGMetadataEntryKey EntryKey2 = *Iterator2;
 
 					// If the entry key is invalid, nothing to do
-					if (EntryKey1 == PCGInvalidEntryKey)
+					if (EntryKey1 != PCGInvalidEntryKey)
 					{
 						AttributeType Value1 = PCGMetadataAttribute::GetValueWithBroadcast<AttributeType>(OperationData.SourceAttributes[0], EntryKey1);
 						AttributeType Value2 = PCGMetadataAttribute::GetValueWithBroadcast<AttributeType>(OperationData.SourceAttributes[1], EntryKey2);
@@ -369,7 +369,7 @@ bool FPCGMetadataMathsElement::DoOperation(FOperationData& OperationData) const
 					PCGMetadataEntryKey EntryKey3 = *Iterator3;
 
 					// If the entry key is invalid, nothing to do
-					if (EntryKey1 == PCGInvalidEntryKey)
+					if (EntryKey1 != PCGInvalidEntryKey)
 					{
 						AttributeType Value1 = PCGMetadataAttribute::GetValueWithBroadcast<AttributeType>(OperationData.SourceAttributes[0], EntryKey1);
 						AttributeType Value2 = PCGMetadataAttribute::GetValueWithBroadcast<AttributeType>(OperationData.SourceAttributes[1], EntryKey2);
diff --git a/Engine/Plugins/Experimental/PCG/Source/PCG/Private/Elements/PCGCreateParamData.cpp b/Engine/Plugins/Experimental/PCG/Source/PCG/Private/Elements/PCGCreateParamData.cpp
new file mode 100644
--- /dev/null
+++ b/Engine/Plugins/Experimental/PCG/Source/PCG/Private/Elements/PCGCreateParamData.cpp
@@ -0,0 +1,113 @@
+// Copyright Epic Games, Inc. All Rights Reserved.
+
+#pragma once
+
+#include "Elements/PCGCreateParamData.h"
+
+#include "PCGParamData.h"
+#include "Metadata/PCGMetadata.h"
+#include "Elements/Metadata/PCGMetadataOpElementBase.h"
+
+#if WITH_EDITOR
+FName UPCGCreateParamDataSettings::GetDefaultNodeName() const
+{
+	return TEXT("CreateParamDataNode");
+}
+#endif
+
+TArray<FPCGPinProperties> UPCGCreateParamDataSettings::InputPinProperties() const
+{
+	return TArray<FPCGPinProperties>();
+}
+
+TArray<FPCGPinProperties> UPCGCreateParamDataSettings::OutputPinProperties() const
+{
+	TArray<FPCGPinProperties> PinProperties;
+	PinProperties.Emplace(PCGPinConstants::DefaultOutputLabel, EPCGDataType::Param);
+
+	return PinProperties;
+}
+
+FPCGElementPtr UPCGCreateParamDataSettings::CreateElement() const
+{
+	return MakeShared<FPCGCreateParamDataElement>();
+}
+
+bool FPCGCreateParamDataElement::ExecuteInternal(FPCGContext* Context) const
+{
+	TRACE_CPUPROFILER_EVENT_SCOPE(FPCGCreateParamDataElement::Execute);
+
+	check(Context);
+
+	const UPCGCreateParamDataSettings* Settings = Context->GetInputSettings<UPCGCreateParamDataSettings>();
+	check(Settings);
+
+	// If we have no output connected, nothing to do
+	if (!Context->Node || !Context->Node->IsOutputPinConnected(PCGPinConstants::DefaultOutputLabel))
+	{
+		PCGE_LOG(Verbose, "Node is not connected, nothing to do");
+		return true;
+	}
+
+	// From there, we should be able to create the data.
+	UPCGParamData* ParamData = NewObject<UPCGParamData>();
+	UPCGMetadata* Metadata = ParamData->MutableMetadata();
+	check(Metadata);
+	PCGMetadataEntryKey EntryKey = Metadata->AddEntry();
+
+	auto CreateAndSetAttribute = [&Metadata, EntryKey, AttributeName = Settings->OutputAttributeName](const auto& Value)
+	{
+		using AttributeType = std::remove_const_t<std::remove_reference_t<decltype(Value)>>;
+
+		FPCGMetadataAttributeBase* BaseAttribute = Metadata->CreateAttribute<AttributeType>(AttributeName, Value, false, false);
+
+		FPCGMetadataAttribute<AttributeType>* Attribute = static_cast<FPCGMetadataAttribute<AttributeType>*>(BaseAttribute);
+		Attribute->SetValue(EntryKey, Value);
+	};
+
+	switch (Settings->Type)
+	{
+
+	case EPCGMetadataTypes::Integer64:
+		CreateAndSetAttribute(Settings->IntValue);
+		break;
+	case EPCGMetadataTypes::Double:
+		CreateAndSetAttribute(Settings->DoubleValue);
+		break;
+	case EPCGMetadataTypes::Vector2:
+		CreateAndSetAttribute(Settings->Vector2Value);
+		break;
+	case EPCGMetadataTypes::Vector:
+		CreateAndSetAttribute(Settings->VectorValue);
+		break;
+	case EPCGMetadataTypes::Vector4:
+		CreateAndSetAttribute(Settings->Vector4Value);
+		break;
+	case EPCGMetadataTypes::Quaternion:
+		CreateAndSetAttribute(Settings->QuatValue);
+		break;
+	case EPCGMetadataTypes::Transform:
+		CreateAndSetAttribute(Settings->TransformValue);
+		break;
+	case EPCGMetadataTypes::String:
+		CreateAndSetAttribute(Settings->StringValue);
+		break;
+	case EPCGMetadataTypes::Boolean:
+		CreateAndSetAttribute(Settings->BoolValue);
+		break;
+	case EPCGMetadataTypes::Rotator:
+		CreateAndSetAttribute(Settings->RotatorValue);
+		break;
+	case EPCGMetadataTypes::Name:
+		CreateAndSetAttribute(Settings->NameValue);
+		break;
+	default:
+		break;
+	}
+
+	TArray<FPCGTaggedData>& Outputs = Context->OutputData.TaggedData;
+	FPCGTaggedData& Output = Outputs.Emplace_GetRef();
+	Output.Data = ParamData;
+
+	return true;
+}
\ No newline at end of file
diff --git a/Engine/Plugins/Experimental/PCG/Source/PCG/Public/Elements/Metadata/PCGMetadataBooleanOpElement.h b/Engine/Plugins/Experimental/PCG/Source/PCG/Public/Elements/Metadata/PCGMetadataBooleanOpElement.h
new file mode 100644
--- /dev/null
+++ b/Engine/Plugins/Experimental/PCG/Source/PCG/Public/Elements/Metadata/PCGMetadataBooleanOpElement.h
@@ -0,0 +1,62 @@
+// Copyright Epic Games, Inc. All Rights Reserved.
+
+#pragma once
+
+#include "PCGElement.h"
+#include "PCGSettings.h"
+#include "Elements/Metadata/PCGMetadataOpElementBase.h"
+
+#include "PCGMetadataBooleanOpElement.generated.h"
+
+UENUM()
+enum class EPCGMedadataBooleanOperation : uint16
+{
+	And,
+	Not,
+	Or,
+	Xor
+};
+
+UCLASS(BlueprintType, ClassGroup = (Procedural))
+class PCG_API UPCGMetadataBooleanSettings : public UPCGMetadataSettingsBase
+{
+	GENERATED_BODY()
+
+public:
+	//~Begin UPCGSettings interface
+#if WITH_EDITOR
+	virtual FName GetDefaultNodeName() const override;
+#endif
+	//~End UPCGSettings interface
+
+	//~Begin UPCGMetadataSettingsBase interface
+	virtual FName GetInputAttributeNameWithOverride(uint32 Index, UPCGParamData* Params) const override;
+
+	virtual FName GetInputPinLabel(uint32 Index) const override;
+	virtual uint32 GetInputPinNum() const override;
+
+	virtual bool IsSupportedInputType(uint16 TypeId, uint32 InputIndex, bool& bHasSpecialRequirement) const override;
+	virtual uint16 GetOutputType(uint16 InputTypeId) const override;
+	//~End UPCGMetadataSettingsBase interface
+
+protected:
+	//~Begin UPCGSettings interface
+	virtual FPCGElementPtr CreateElement() const override;
+	//~End UPCGSettings interface
+
+public:
+	UPROPERTY(BlueprintReadWrite, EditAnywhere, Category = Settings)
+	EPCGMedadataBooleanOperation Operation = EPCGMedadataBooleanOperation::And;
+
+	UPROPERTY(BlueprintReadWrite, EditAnywhere, Category = Input)
+	FName Input1AttributeName = NAME_None;
+
+	UPROPERTY(BlueprintReadWrite, EditAnywhere, Category = Input, meta = (EditCondition="Operation != EPCGMedadataBooleanOperation::Not", EditConditionHides))
+	FName Input2AttributeName = NAME_None;
+};
+
+class FPCGMetadataBooleanElement : public FPCGMetadataElementBase
+{
+protected:
+	virtual bool DoOperation(FOperationData& OperationData) const override;
+};
diff --git a/Engine/Plugins/Experimental/PCG/Source/PCG/Public/Elements/PCGCreateParamData.h b/Engine/Plugins/Experimental/PCG/Source/PCG/Public/Elements/PCGCreateParamData.h
new file mode 100644
--- /dev/null
+++ b/Engine/Plugins/Experimental/PCG/Source/PCG/Public/Elements/PCGCreateParamData.h
@@ -0,0 +1,79 @@
+// Copyright Epic Games, Inc. All Rights Reserved.
+
+#pragma once
+
+#include "PCGElement.h"
+#include "PCGSettings.h"
+#include "Metadata/PCGMetadataAttributeTraits.h"
+
+#include "PCGCreateParamData.generated.h"
+
+UCLASS(BlueprintType, ClassGroup = (Procedural))
+class PCG_API UPCGCreateParamDataSettings : public UPCGSettings
+{
+	GENERATED_BODY()
+
+public:
+	//~Begin UPCGSettings interface
+#if WITH_EDITOR
+	virtual FName GetDefaultNodeName() const override;
+	virtual EPCGSettingsType GetType() const override { return EPCGSettingsType::Param; }
+#endif
+
+	virtual TArray<FPCGPinProperties> InputPinProperties() const override;
+	virtual TArray<FPCGPinProperties> OutputPinProperties() const override;
+	//~End UPCGSettings interface
+
+	UPROPERTY(BlueprintReadWrite, EditAnywhere, Category = Settings, meta = (ValidEnumValues = "Double, Integer64, Vector2, Vector, Vector4, Quaternion, Transform, String, Boolean, Rotator, Name"))
+	EPCGMetadataTypes Type = EPCGMetadataTypes::Double;
+
+	UPROPERTY(BlueprintReadWrite, EditAnywhere, Category = Settings)
+	FName OutputAttributeName = NAME_None;
+
+	// All different types
+	UPROPERTY(BlueprintReadWrite, EditAnywhere, Category = Settings, meta = (EditCondition = "Type == EPCGMetadataTypes::Double", EditConditionHides))
+	double DoubleValue = 0.0;
+
+	UPROPERTY(BlueprintReadWrite, EditAnywhere, Category = Settings, meta = (EditCondition = "Type == EPCGMetadataTypes::Integer64", EditConditionHides))
+	int64 IntValue = 0;
+
+	UPROPERTY(BlueprintReadWrite, EditAnywhere, Category = Settings, meta = (EditCondition = "Type == EPCGMetadataTypes::Vector2", EditConditionHides))
+	FVector2D Vector2Value = FVector2D::ZeroVector;
+
+	UPROPERTY(BlueprintReadWrite, EditAnywhere, Category = Settings, meta = (EditCondition = "Type == EPCGMetadataTypes::Vector", EditConditionHides))
+	FVector VectorValue = FVector::ZeroVector;
+
+	UPROPERTY(BlueprintReadWrite, EditAnywhere, Category = Settings, meta = (EditCondition = "Type == EPCGMetadataTypes::Vector4", EditConditionHides))
+	FVector4 Vector4Value = FVector4::Zero();
+
+	UPROPERTY(BlueprintReadWrite, EditAnywhere, Category = Settings, meta = (EditCondition = "Type == EPCGMetadataTypes::Quaternion", EditConditionHides))
+	FQuat QuatValue = FQuat::Identity;
+
+	UPROPERTY(BlueprintReadWrite, EditAnywhere, Category = Settings, meta = (EditCondition = "Type == EPCGMetadataTypes::Transform", EditConditionHides))
+	FTransform TransformValue = FTransform::Identity;
+
+	UPROPERTY(BlueprintReadWrite, EditAnywhere, Category = Settings, meta = (EditCondition = "Type == EPCGMetadataTypes::String", EditConditionHides))
+	FString StringValue = "";
+
+	UPROPERTY(BlueprintReadWrite, EditAnywhere, Category = Settings, meta = (EditCondition = "Type == EPCGMetadataTypes::Boolean", EditConditionHides))
+	bool BoolValue = false;
+
+	UPROPERTY(BlueprintReadWrite, EditAnywhere, Category = Settings, meta = (EditCondition = "Type == EPCGMetadataTypes::Rotator", EditConditionHides))
+	FRotator RotatorValue = FRotator::ZeroRotator;
+
+	UPROPERTY(BlueprintReadWrite, EditAnywhere, Category = Settings, meta = (EditCondition = "Type == EPCGMetadataTypes::Name", EditConditionHides))
+	FName NameValue = NAME_None;
+
+protected:
+	virtual FPCGElementPtr CreateElement() const override;
+};
+
+
+class FPCGCreateParamDataElement : public FSimplePCGElement
+{
+public:
+	virtual bool IsCacheable(const UPCGSettings* InSettings) const override { return false; }
+
+protected:
+	virtual bool ExecuteInternal(FPCGContext* Context) const override;
+};
\ No newline at end of file
diff --git a/Engine/Plugins/Experimental/PCG/Source/PCG/Public/PCGSettings.h b/Engine/Plugins/Experimental/PCG/Source/PCG/Public/PCGSettings.h
--- a/Engine/Plugins/Experimental/PCG/Source/PCG/Public/PCGSettings.h
+++ b/Engine/Plugins/Experimental/PCG/Source/PCG/Public/PCGSettings.h
@@ -38,7 +38,8 @@ enum class EPCGSettingsType : uint8
 	Spawner,
 	Subgraph,
 	Debug,
-	Generic
+	Generic,
+	Param
 };
 
 #if WITH_EDITOR
diff --git a/Engine/Plugins/Experimental/PCG/Source/PCGEditor/Private/PCGEditorSettings.cpp b/Engine/Plugins/Experimental/PCG/Source/PCGEditor/Private/PCGEditorSettings.cpp
--- a/Engine/Plugins/Experimental/PCG/Source/PCGEditor/Private/PCGEditorSettings.cpp
+++ b/Engine/Plugins/Experimental/PCG/Source/PCGEditor/Private/PCGEditorSettings.cpp
@@ -19,6 +19,7 @@ UPCGEditorSettings::UPCGEditorSettings(const FObjectInitializer& ObjectInitializ
 	SamplerNodeColor = FLinearColor(0.8f, 1.0f, 0.4f);
 	SpawnerNodeColor = FLinearColor(1.0f, 0.6f, 0.4f);
 	SubgraphNodeColor = FLinearColor(1.0f, 0.1f, 0.1f);
+	ParamDataNodeColor = FLinearColor(1.0f, 0.6f, 0.0f);
 	DebugNodeColor = FLinearColor(1.0f, 0.0f, 1.0f);
 
 	DefaultPinColor = FLinearColor(1.0f, 1.0f, 1.0f);
@@ -82,6 +83,10 @@ FLinearColor UPCGEditorSettings::GetColor(UPCGSettings* Settings) const
 	{
 		return DebugNodeColor;
 	}
+	else if (Settings->GetType() == EPCGSettingsType::Param)
+	{
+		return ParamDataNodeColor;
+	}
 	else
 	{
 		// Finally, we couldn't find any match, so return the default value
diff --git a/Engine/Plugins/Experimental/PCG/Source/PCGEditor/Public/PCGEditorSettings.h b/Engine/Plugins/Experimental/PCG/Source/PCGEditor/Public/PCGEditorSettings.h
--- a/Engine/Plugins/Experimental/PCG/Source/PCGEditor/Public/PCGEditorSettings.h
+++ b/Engine/Plugins/Experimental/PCG/Source/PCGEditor/Public/PCGEditorSettings.h
@@ -65,6 +65,10 @@ public:
 	UPROPERTY(EditAnywhere, config, Category = Node, meta = (HideAlphaChannel))
 	FLinearColor SubgraphNodeColor;
 
+	/** Color used for param data operations */
+	UPROPERTY(EditAnywhere, config, Category = Node, meta = (HideAlphaChannel))
+	FLinearColor ParamDataNodeColor;
+
 	/** Color used for debug operations */
 	UPROPERTY(EditAnywhere, config, Category = Node, meta = (HideAlphaChannel))
 	FLinearColor DebugNodeColor;
