[PCG] Metadata native operation: Bitwise operations

diff --git a/Engine/Plugins/Experimental/PCG/Source/PCG/Private/Elements/Metadata/PCGMetadataBitwiseOpElement.cpp b/Engine/Plugins/Experimental/PCG/Source/PCG/Private/Elements/Metadata/PCGMetadataBitwiseOpElement.cpp
new file mode 100644
--- /dev/null
+++ b/Engine/Plugins/Experimental/PCG/Source/PCG/Private/Elements/Metadata/PCGMetadataBitwiseOpElement.cpp
@@ -0,0 +1,149 @@
+// Copyright Epic Games, Inc. All Rights Reserved.
+
+#include "Elements/Metadata/PCGMetadataBitwiseOpElement.h"
+
+#include "PCGParamData.h"
+#include "Data/PCGSpatialData.h"
+#include "Helpers/PCGSettingsHelpers.h"
+#include "Metadata/PCGMetadata.h"
+#include "Metadata/PCGMetadataAttribute.h"
+#include "Metadata/PCGMetadataAttributeTpl.h"
+#include "Metadata/PCGMetadataEntryKeyIterator.h"
+#include "Elements/Metadata/PCGMetadataElementCommon.h"
+
+FName UPCGMetadataBitwiseSettings::GetInputPinLabel(uint32 Index) const
+{
+	switch (Index)
+	{
+	case 0:
+		return (Operation == EPCGMedadataBitwiseOperation::BitwiseNot) ? PCGPinConstants::DefaultInputLabel : PCGMetadataSettingsBaseConstants::DoubleInputFirstLabel;
+	case 1:
+		return PCGMetadataSettingsBaseConstants::DoubleInputSecondLabel;
+	default:
+		return NAME_None;
+	}
+}
+
+uint32 UPCGMetadataBitwiseSettings::GetInputPinNum() const
+{
+	return (Operation == EPCGMedadataBitwiseOperation::BitwiseNot) ? 1 : 2;
+}
+
+bool UPCGMetadataBitwiseSettings::IsSupportedInputType(uint16 TypeId, uint32 InputIndex, bool& bHasSpecialRequirement) const
+{
+	bHasSpecialRequirement = false;
+	return (TypeId == (uint16)EPCGMetadataTypes::Integer32) || (TypeId == (uint16)EPCGMetadataTypes::Integer64);
+}
+
+FName UPCGMetadataBitwiseSettings::GetInputAttributeNameWithOverride(uint32 Index, UPCGParamData* Params) const
+{
+	switch (Index)
+	{
+	case 0:
+		return PCG_GET_OVERRIDEN_VALUE(this, Input1AttributeName, Params);
+	case 1:
+		return PCG_GET_OVERRIDEN_VALUE(this, Input2AttributeName, Params);
+	default:
+		return NAME_None;
+	}
+}
+
+#if WITH_EDITOR
+FName UPCGMetadataBitwiseSettings::GetDefaultNodeName() const
+{
+	if (const UEnum* EnumPtr = FindObject<UEnum>(nullptr, TEXT("EPCGMedadataBitwiseOperation"), true))
+	{
+		return EnumPtr->GetNameByValue(static_cast<int>(Operation));
+	}
+
+	return TEXT("Metadata Bitwise Node");
+}
+#endif // WITH_EDITOR
+
+FPCGElementPtr UPCGMetadataBitwiseSettings::CreateElement() const
+{
+	return MakeShared<FPCGMetadataBitwiseElement>();
+}
+
+uint16 UPCGMetadataBitwiseSettings::GetOutputType(uint16 InputTypeId) const
+{
+	return (uint16)EPCGMetadataTypes::Integer64;
+}
+
+bool FPCGMetadataBitwiseElement::DoOperation(FOperationData& OperationData) const
+{
+	TRACE_CPUPROFILER_EVENT_SCOPE(FPCGMetadataBitwiseElement::Execute);
+
+	const UPCGMetadataBitwiseSettings* Settings = CastChecked<UPCGMetadataBitwiseSettings>(OperationData.Settings);
+
+	auto BitwiseFunc = [Settings](int64& Value1, const int64& Value2)
+	{
+		switch (Settings->Operation)
+		{
+		case EPCGMedadataBitwiseOperation::BitwiseAnd:
+			Value1 = (Value1 & Value2);
+			break;
+		case EPCGMedadataBitwiseOperation::BitwiseNot:
+			Value1 = ~Value1;
+			break;
+		case EPCGMedadataBitwiseOperation::BitwiseOr:
+			Value1 = (Value1 | Value2);
+			break;
+		case EPCGMedadataBitwiseOperation::BitwiseXor:
+			Value1 = (Value1 ^ Value2);
+			break;
+		}
+	};
+
+	// All values in OperationData.Iterators are UniquePtr. Dereference them here to make the syntax less cumbersome.
+	// Also some iterators can be null, it means they need to use the first iterator (that should never be null).
+	// We just need to make sure to not increment the first iterator multiple times in one loop.
+
+	check(OperationData.Iterators[0]);
+
+	IPCGMetadataEntryIterator& Iterator1 = *OperationData.Iterators[0];
+	IPCGMetadataEntryIterator& Iterator2 = (OperationData.Iterators.Num() >= 2 && OperationData.Iterators[1]) ? *OperationData.Iterators[1] : Iterator1;
+
+	bool bShouldIncrementIterator2 = &Iterator1 != &Iterator2;
+
+	PCGMetadataEntryKey EntryKey1 = 0;
+	PCGMetadataEntryKey EntryKey2 = 0;
+
+	int64 Value1 = 0;
+	int64 Value2 = 0;
+
+	FPCGMetadataAttribute<int64>* OutputAttribute = static_cast<FPCGMetadataAttribute<int64>*>(OperationData.OutputAttribute);
+	int64 DefaultValue1 = PCGMetadataAttribute::GetValueWithBroadcast<int64>(OperationData.SourceAttributes[0], PCGInvalidEntryKey);
+	int64 DefaultValue2 = (Settings->Operation != EPCGMedadataBitwiseOperation::BitwiseNot) ? PCGMetadataAttribute::GetValueWithBroadcast<int64>(OperationData.SourceAttributes[1], PCGInvalidEntryKey) : 0;
+	BitwiseFunc(DefaultValue1, DefaultValue2);
+	OutputAttribute->SetDefaultValue(DefaultValue1);
+
+	for (int32 i = 0; i < OperationData.NumberOfElementsToProcess; ++i)
+	{
+		EntryKey1 = *Iterator1;
+
+		if (EntryKey1 != PCGInvalidEntryKey)
+		{
+			Value1 = PCGMetadataAttribute::GetValueWithBroadcast<int64>(OperationData.SourceAttributes[0], EntryKey1);
+
+			if (Settings->Operation != EPCGMedadataBitwiseOperation::BitwiseNot)
+			{
+				EntryKey2 = *Iterator2;
+				Value2 = PCGMetadataAttribute::GetValueWithBroadcast<int64>(OperationData.SourceAttributes[1], EntryKey2);
+			}
+
+			BitwiseFunc(Value1, Value2);
+
+			OutputAttribute->SetValue(EntryKey1, Value1);
+		}
+
+		++Iterator1;
+
+		if (bShouldIncrementIterator2)
+		{
+			++Iterator2;
+		}
+	}
+
+	return true;
+}
\ No newline at end of file
diff --git a/Engine/Plugins/Experimental/PCG/Source/PCG/Public/Elements/Metadata/PCGMetadataBitwiseOpElement.h b/Engine/Plugins/Experimental/PCG/Source/PCG/Public/Elements/Metadata/PCGMetadataBitwiseOpElement.h
new file mode 100644
--- /dev/null
+++ b/Engine/Plugins/Experimental/PCG/Source/PCG/Public/Elements/Metadata/PCGMetadataBitwiseOpElement.h
@@ -0,0 +1,62 @@
+// Copyright Epic Games, Inc. All Rights Reserved.
+
+#pragma once
+
+#include "PCGElement.h"
+#include "PCGSettings.h"
+#include "Elements/Metadata/PCGMetadataOpElementBase.h"
+
+#include "PCGMetadataBitwiseOpElement.generated.h"
+
+UENUM()
+enum class EPCGMedadataBitwiseOperation : uint16
+{
+	BitwiseAnd,
+	BitwiseNot,
+	BitwiseOr,
+	BitwiseXor
+};
+
+UCLASS(BlueprintType, ClassGroup = (Procedural))
+class PCG_API UPCGMetadataBitwiseSettings : public UPCGMetadataSettingsBase
+{
+	GENERATED_BODY()
+
+public:
+	//~Begin UPCGSettings interface
+#if WITH_EDITOR
+	virtual FName GetDefaultNodeName() const override;
+#endif
+	//~End UPCGSettings interface
+
+	//~Begin UPCGMetadataSettingsBase interface
+	virtual FName GetInputAttributeNameWithOverride(uint32 Index, UPCGParamData* Params) const override;
+
+	virtual FName GetInputPinLabel(uint32 Index) const override;
+	virtual uint32 GetInputPinNum() const override;
+
+	virtual bool IsSupportedInputType(uint16 TypeId, uint32 InputIndex, bool& bHasSpecialRequirement) const override;
+	virtual uint16 GetOutputType(uint16 InputTypeId) const override;
+	//~End UPCGMetadataSettingsBase interface
+
+protected:
+	//~Begin UPCGSettings interface
+	virtual FPCGElementPtr CreateElement() const override;
+	//~End UPCGSettings interface
+
+public:
+	UPROPERTY(BlueprintReadWrite, EditAnywhere, Category = Settings)
+	EPCGMedadataBitwiseOperation Operation = EPCGMedadataBitwiseOperation::BitwiseAnd;
+
+	UPROPERTY(BlueprintReadWrite, EditAnywhere, Category = Input)
+	FName Input1AttributeName = NAME_None;
+
+	UPROPERTY(BlueprintReadWrite, EditAnywhere, Category = Input, meta = (EditCondition = "Operation != EPCGMedadataBitwiseOperation::BitwiseNot", EditConditionHides))
+	FName Input2AttributeName = NAME_None;
+};
+
+class FPCGMetadataBitwiseElement : public FPCGMetadataElementBase
+{
+protected:
+	virtual bool DoOperation(FOperationData& OperationData) const override;
+};
