[PCG] Metadata native operation: Trig operations

diff --git a/Engine/Plugins/Experimental/PCG/Source/PCG/Private/Elements/Metadata/PCGMetadataTrigOpElement.cpp b/Engine/Plugins/Experimental/PCG/Source/PCG/Private/Elements/Metadata/PCGMetadataTrigOpElement.cpp
new file mode 100644
--- /dev/null
+++ b/Engine/Plugins/Experimental/PCG/Source/PCG/Private/Elements/Metadata/PCGMetadataTrigOpElement.cpp
@@ -0,0 +1,190 @@
+// Copyright Epic Games, Inc. All Rights Reserved.
+
+#include "Elements/Metadata/PCGMetadataTrigOpElement.h"
+
+#include "Helpers/PCGSettingsHelpers.h"
+#include "Metadata/PCGMetadata.h"
+#include "Metadata/PCGMetadataAttribute.h"
+#include "Metadata/PCGMetadataAttributeTpl.h"
+#include "Metadata/PCGMetadataEntryKeyIterator.h"
+#include "Elements/Metadata/PCGMetadataElementCommon.h"
+
+namespace PCGMetadataTrigSettings
+{
+	template <typename OutType>
+	void ApplyTrigOperation(OutType& Input1, const OutType& Input2, EPCGMedadataTrigOperation Operation)
+	{
+		switch (Operation)
+		{
+		case EPCGMedadataTrigOperation::Acos:
+			Input1 = FMath::Acos(Input1);
+			break;
+		case EPCGMedadataTrigOperation::Asin:
+			Input1 = FMath::Asin(Input1);
+			break;
+		case EPCGMedadataTrigOperation::Atan:
+			Input1 = FMath::Atan(Input1);
+			break;
+		case EPCGMedadataTrigOperation::Atan2:
+			Input1 = FMath::Atan2(Input1, Input2);
+			break;
+		case EPCGMedadataTrigOperation::Cos:
+			Input1 = FMath::Cos(Input1);
+			break;
+		case EPCGMedadataTrigOperation::Sin:
+			Input1 = FMath::Sin(Input1);
+			break;
+		case EPCGMedadataTrigOperation::Tan:
+			Input1 = FMath::Tan(Input1);
+			break;
+		case EPCGMedadataTrigOperation::DegToRad:
+			Input1 = FMath::DegreesToRadians(Input1);
+			break;
+		case EPCGMedadataTrigOperation::RadToDeg:
+			Input1 = FMath::RadiansToDegrees(Input1);
+			break;
+		default:
+			break;
+		}
+	}
+}
+
+FName UPCGMetadataTrigSettings::GetInputPinLabel(uint32 Index) const
+{
+	switch (Index)
+	{
+	case 0:
+		return (Operation != EPCGMedadataTrigOperation::Atan2) ? PCGPinConstants::DefaultInputLabel : PCGMetadataSettingsBaseConstants::DoubleInputFirstLabel;
+	case 1:
+		return PCGMetadataSettingsBaseConstants::DoubleInputSecondLabel;
+	default:
+		return NAME_None;
+	}
+}
+
+uint32 UPCGMetadataTrigSettings::GetInputPinNum() const
+{
+	return (Operation != EPCGMedadataTrigOperation::Atan2) ? 1 : 2;
+}
+
+bool UPCGMetadataTrigSettings::IsSupportedInputType(uint16 TypeId, uint32 InputIndex, bool& bHasSpecialRequirement) const
+{
+	bHasSpecialRequirement = false;
+	return TypeId <= (uint16)EPCGMetadataTypes::Integer64;
+}
+
+FName UPCGMetadataTrigSettings::GetInputAttributeNameWithOverride(uint32 Index, UPCGParamData* Params) const
+{
+	switch (Index)
+	{
+	case 0:
+		return PCG_GET_OVERRIDEN_VALUE(this, Input1AttributeName, Params);
+	case 1:
+		return PCG_GET_OVERRIDEN_VALUE(this, Input2AttributeName, Params);
+	default:
+		return NAME_None;
+	}
+}
+
+uint16 UPCGMetadataTrigSettings::GetOutputType(uint16 InputTypeId) const
+{
+	if (InputTypeId == (uint16)EPCGMetadataTypes::Integer32 || InputTypeId == (uint16)EPCGMetadataTypes::Integer64)
+	{
+		return (uint16)EPCGMetadataTypes::Double;
+	}
+	else
+	{
+		return InputTypeId;
+	}
+}
+
+#if WITH_EDITOR
+FName UPCGMetadataTrigSettings::GetDefaultNodeName() const
+{
+	if (const UEnum* EnumPtr = FindObject<UEnum>(nullptr, TEXT("EPCGMedadataTrigOperation"), true))
+	{ 
+		return EnumPtr->GetNameByValue(static_cast<int>(Operation)); 
+	}
+
+	return TEXT("Metadata Trig Node");
+}
+#endif // WITH_EDITOR
+
+FPCGElementPtr UPCGMetadataTrigSettings::CreateElement() const
+{
+	return MakeShared<FPCGMetadataTrigElement>();
+}
+
+bool FPCGMetadataTrigElement::DoOperation(FOperationData& OperationData) const
+{
+	TRACE_CPUPROFILER_EVENT_SCOPE(FPCGMetadataTrigElement::Execute);
+
+	const UPCGMetadataTrigSettings* Settings = CastChecked<UPCGMetadataTrigSettings>(OperationData.Settings);
+
+	// All values in OperationData.Iterators are UniquePtr. Dereference them here to make the syntax less cumbersome.
+	// Also some iterators can be null, it means they need to use the first iterator (that should never be null).
+	// We just need to make sure to not increment the first iterator multiple times in one loop.
+
+	check(OperationData.Iterators[0]);
+
+	IPCGMetadataEntryIterator& Iterator1 = *OperationData.Iterators[0];
+	IPCGMetadataEntryIterator& Iterator2 = (OperationData.Iterators.Num() >= 2 && OperationData.Iterators[1]) ? *OperationData.Iterators[1] : Iterator1;
+
+	bool bShouldIncrementIterator2 = &Iterator1 != &Iterator2;
+
+	auto TrigFunc = [Operation = Settings->Operation, &Iterator1, &Iterator2, bShouldIncrementIterator2, &OperationData](auto DummyValue)
+	{
+		using AttributeType = decltype(DummyValue);
+
+		// Need to remove types that would not compile
+		if constexpr (PCG::Private::MetadataTypes<AttributeType>::Id > (uint16)EPCGMetadataTypes::Double)
+		{
+			return;
+		}
+		else
+		{
+			PCGMetadataEntryKey EntryKey1 = 0;
+			PCGMetadataEntryKey EntryKey2 = 0;
+
+			FPCGMetadataAttribute<AttributeType>* OutputAttribute = static_cast<FPCGMetadataAttribute<AttributeType>*>(OperationData.OutputAttribute);
+			AttributeType DefaultValue1 = PCGMetadataAttribute::GetValueWithBroadcast<AttributeType>(OperationData.SourceAttributes[0], PCGInvalidEntryKey);
+			AttributeType DefaultValue2 = (Operation == EPCGMedadataTrigOperation::Atan2) ? PCGMetadataAttribute::GetValueWithBroadcast<AttributeType>(OperationData.SourceAttributes[1], PCGInvalidEntryKey) : AttributeType{};
+			PCGMetadataTrigSettings::ApplyTrigOperation(DefaultValue1, DefaultValue2, Operation);
+			OutputAttribute->SetDefaultValue(DefaultValue1);
+
+			AttributeType Value1{};
+			AttributeType Value2{};
+
+			for (int32 i = 0; i < OperationData.NumberOfElementsToProcess; ++i)
+			{
+				EntryKey1 = *Iterator1;
+
+				// If the entry key is invalid, nothing to do
+				if (EntryKey1 != PCGInvalidEntryKey)
+				{
+					Value1 = PCGMetadataAttribute::GetValueWithBroadcast<AttributeType>(OperationData.SourceAttributes[0], EntryKey1);
+
+					if (Operation == EPCGMedadataTrigOperation::Atan2)
+					{
+						EntryKey2 = *Iterator2;
+						Value2 = PCGMetadataAttribute::GetValueWithBroadcast<AttributeType>(OperationData.SourceAttributes[1], EntryKey2);
+					}
+
+					PCGMetadataTrigSettings::ApplyTrigOperation(Value1, Value2, Operation);
+
+					OutputAttribute->SetValue(EntryKey1, Value1);
+				}
+
+				++Iterator1;
+				if (bShouldIncrementIterator2)
+				{
+					++Iterator2;
+				}
+			}
+		}
+	};
+
+	PCGMetadataAttribute::CallbackWithRightType(OperationData.OutputType, TrigFunc);
+
+	return true;
+}
\ No newline at end of file
diff --git a/Engine/Plugins/Experimental/PCG/Source/PCG/Public/Elements/Metadata/PCGMetadataTrigOpElement.h b/Engine/Plugins/Experimental/PCG/Source/PCG/Public/Elements/Metadata/PCGMetadataTrigOpElement.h
new file mode 100644
--- /dev/null
+++ b/Engine/Plugins/Experimental/PCG/Source/PCG/Public/Elements/Metadata/PCGMetadataTrigOpElement.h
@@ -0,0 +1,67 @@
+// Copyright Epic Games, Inc. All Rights Reserved.
+
+#pragma once
+
+#include "PCGElement.h"
+#include "PCGSettings.h"
+#include "Elements/Metadata/PCGMetadataOpElementBase.h"
+
+#include "PCGMetadataTrigOpElement.generated.h"
+
+UENUM()
+enum class EPCGMedadataTrigOperation : uint16
+{
+	Acos,
+	Asin,
+	Atan,
+	Atan2,
+	Cos,
+	Sin,
+	Tan,
+	DegToRad,
+	RadToDeg
+};
+
+UCLASS(BlueprintType, ClassGroup = (Procedural))
+class PCG_API UPCGMetadataTrigSettings : public UPCGMetadataSettingsBase
+{
+	GENERATED_BODY()
+
+public:
+	//~Begin UPCGSettings interface
+#if WITH_EDITOR
+	virtual FName GetDefaultNodeName() const override;
+#endif
+	//~End UPCGSettings interface
+
+	//~Begin UPCGMetadataSettingsBase interface
+	virtual FName GetInputAttributeNameWithOverride(uint32 Index, UPCGParamData* Params) const override;
+
+	virtual FName GetInputPinLabel(uint32 Index) const override;
+	virtual uint32 GetInputPinNum() const override;
+
+	virtual bool IsSupportedInputType(uint16 TypeId, uint32 InputIndex, bool& bHasSpecialRequirement) const override;
+	virtual uint16 GetOutputType(uint16 InputTypeId) const override;
+	//~End UPCGMetadataSettingsBase interface
+
+protected:
+	//~Begin UPCGSettings interface
+	virtual FPCGElementPtr CreateElement() const override;
+	//~End UPCGSettings interface
+
+public:
+	UPROPERTY(BlueprintReadWrite, EditAnywhere, Category = Settings)
+	EPCGMedadataTrigOperation Operation = EPCGMedadataTrigOperation::Acos;
+
+	UPROPERTY(BlueprintReadWrite, EditAnywhere, Category = Input)
+	FName Input1AttributeName = NAME_None;
+
+	UPROPERTY(BlueprintReadWrite, EditAnywhere, Category = Input, meta = (EditCondition="Operation = EPCGMedadataTrigOperation::Atan2", EditConditionHides))
+	FName Input2AttributeName = NAME_None;
+};
+
+class FPCGMetadataTrigElement : public FPCGMetadataElementBase
+{
+protected:
+	virtual bool DoOperation(FOperationData& OperationData) const override;
+};
