[PCG] Metadata native operation: Refactor unary/binary/ternary op

diff --git a/Engine/Plugins/Experimental/PCG/Source/PCG/Private/Elements/Metadata/PCGMetadataBitwiseOpElement.cpp b/Engine/Plugins/Experimental/PCG/Source/PCG/Private/Elements/Metadata/PCGMetadataBitwiseOpElement.cpp
--- a/Engine/Plugins/Experimental/PCG/Source/PCG/Private/Elements/Metadata/PCGMetadataBitwiseOpElement.cpp
+++ b/Engine/Plugins/Experimental/PCG/Source/PCG/Private/Elements/Metadata/PCGMetadataBitwiseOpElement.cpp
@@ -11,6 +11,30 @@
 #include "Metadata/PCGMetadataEntryKeyIterator.h"
 #include "Elements/Metadata/PCGMetadataElementCommon.h"
 
+namespace PCGMetadataBitwiseSettings
+{
+	inline int64 UnaryOp(const int64& Value)
+	{
+		// EPCGMedadataBitwiseOperation::BitwiseNot
+		return ~Value;
+	}
+
+	inline int64 BinaryOp(const int64& Value1, const int64& Value2, EPCGMedadataBitwiseOperation Operation)
+	{
+		switch (Operation)
+		{
+		case EPCGMedadataBitwiseOperation::BitwiseAnd:
+			return (Value1 & Value2);
+		case EPCGMedadataBitwiseOperation::BitwiseOr:
+			return (Value1 | Value2);
+		case EPCGMedadataBitwiseOperation::BitwiseXor:
+			return (Value1 ^ Value2);
+		default:
+			return 0;
+		}
+	}
+}
+
 FName UPCGMetadataBitwiseSettings::GetInputPinLabel(uint32 Index) const
 {
 	switch (Index)
@@ -76,73 +100,13 @@ bool FPCGMetadataBitwiseElement::DoOperation(FOperationData& OperationData) cons
 
 	const UPCGMetadataBitwiseSettings* Settings = CastChecked<UPCGMetadataBitwiseSettings>(OperationData.Settings);
 
-	auto BitwiseFunc = [Settings](int64& Value1, const int64& Value2)
+	if (Settings->Operation == EPCGMedadataBitwiseOperation::BitwiseNot)
 	{
-		switch (Settings->Operation)
-		{
-		case EPCGMedadataBitwiseOperation::BitwiseAnd:
-			Value1 = (Value1 & Value2);
-			break;
-		case EPCGMedadataBitwiseOperation::BitwiseNot:
-			Value1 = ~Value1;
-			break;
-		case EPCGMedadataBitwiseOperation::BitwiseOr:
-			Value1 = (Value1 | Value2);
-			break;
-		case EPCGMedadataBitwiseOperation::BitwiseXor:
-			Value1 = (Value1 ^ Value2);
-			break;
-		}
-	};
-
-	// All values in OperationData.Iterators are UniquePtr. Dereference them here to make the syntax less cumbersome.
-	// Also some iterators can be null, it means they need to use the first iterator (that should never be null).
-	// We just need to make sure to not increment the first iterator multiple times in one loop.
-
-	check(OperationData.Iterators[0]);
-
-	IPCGMetadataEntryIterator& Iterator1 = *OperationData.Iterators[0];
-	IPCGMetadataEntryIterator& Iterator2 = (OperationData.Iterators.Num() >= 2 && OperationData.Iterators[1]) ? *OperationData.Iterators[1] : Iterator1;
-
-	bool bShouldIncrementIterator2 = &Iterator1 != &Iterator2;
-
-	PCGMetadataEntryKey EntryKey1 = 0;
-	PCGMetadataEntryKey EntryKey2 = 0;
-
-	int64 Value1 = 0;
-	int64 Value2 = 0;
-
-	FPCGMetadataAttribute<int64>* OutputAttribute = static_cast<FPCGMetadataAttribute<int64>*>(OperationData.OutputAttribute);
-	int64 DefaultValue1 = PCGMetadataAttribute::GetValueWithBroadcast<int64>(OperationData.SourceAttributes[0], PCGInvalidEntryKey);
-	int64 DefaultValue2 = (Settings->Operation != EPCGMedadataBitwiseOperation::BitwiseNot) ? PCGMetadataAttribute::GetValueWithBroadcast<int64>(OperationData.SourceAttributes[1], PCGInvalidEntryKey) : 0;
-	BitwiseFunc(DefaultValue1, DefaultValue2);
-	OutputAttribute->SetDefaultValue(DefaultValue1);
-
-	for (int32 i = 0; i < OperationData.NumberOfElementsToProcess; ++i)
+		DoUnaryOp<int64, int64>(OperationData, [](const int64& Value) -> int64 { return PCGMetadataBitwiseSettings::UnaryOp(Value); });
+	}
+	else
 	{
-		EntryKey1 = *Iterator1;
-
-		if (EntryKey1 != PCGInvalidEntryKey)
-		{
-			Value1 = PCGMetadataAttribute::GetValueWithBroadcast<int64>(OperationData.SourceAttributes[0], EntryKey1);
-
-			if (Settings->Operation != EPCGMedadataBitwiseOperation::BitwiseNot)
-			{
-				EntryKey2 = *Iterator2;
-				Value2 = PCGMetadataAttribute::GetValueWithBroadcast<int64>(OperationData.SourceAttributes[1], EntryKey2);
-			}
-
-			BitwiseFunc(Value1, Value2);
-
-			OutputAttribute->SetValue(EntryKey1, Value1);
-		}
-
-		++Iterator1;
-
-		if (bShouldIncrementIterator2)
-		{
-			++Iterator2;
-		}
+		DoBinaryOp<int64, int64, int64>(OperationData, [Operation = Settings->Operation](const int64& Value1, const int64& Value2) -> int64 { return PCGMetadataBitwiseSettings::BinaryOp(Value1, Value2, Operation); });
 	}
 
 	return true;
diff --git a/Engine/Plugins/Experimental/PCG/Source/PCG/Private/Elements/Metadata/PCGMetadataBooleanOpElement.cpp b/Engine/Plugins/Experimental/PCG/Source/PCG/Private/Elements/Metadata/PCGMetadataBooleanOpElement.cpp
--- a/Engine/Plugins/Experimental/PCG/Source/PCG/Private/Elements/Metadata/PCGMetadataBooleanOpElement.cpp
+++ b/Engine/Plugins/Experimental/PCG/Source/PCG/Private/Elements/Metadata/PCGMetadataBooleanOpElement.cpp
@@ -11,6 +11,30 @@
 #include "Metadata/PCGMetadataEntryKeyIterator.h"
 #include "Elements/Metadata/PCGMetadataElementCommon.h"
 
+namespace PCGMetadataBooleanSettings
+{
+	inline bool UnaryOp(const bool& Value)
+	{
+		// EPCGMedadataBooleanOperation::Not
+		return !Value;
+	}
+
+	inline bool BinaryOp(const bool& Value1, const bool& Value2, EPCGMedadataBooleanOperation Operation)
+	{
+		switch (Operation)
+		{
+		case EPCGMedadataBooleanOperation::And:
+			return (Value1 && Value2);
+		case EPCGMedadataBooleanOperation::Or:
+			return (Value1 || Value2);
+		case EPCGMedadataBooleanOperation::Xor:
+			return (Value1 != Value2);
+		default:
+			return false;
+		}
+	}
+}
+
 FName UPCGMetadataBooleanSettings::GetInputPinLabel(uint32 Index) const
 {
 	switch (Index)
@@ -76,73 +100,13 @@ bool FPCGMetadataBooleanElement::DoOperation(FOperationData& OperationData) cons
 
 	const UPCGMetadataBooleanSettings* Settings = CastChecked<UPCGMetadataBooleanSettings>(OperationData.Settings);
 
-	auto BoolFunc = [Settings](bool& Value1, const bool& Value2)
+	if (Settings->Operation == EPCGMedadataBooleanOperation::Not)
 	{
-		switch (Settings->Operation)
-		{
-		case EPCGMedadataBooleanOperation::And:
-			Value1 = (Value1 && Value2);
-			break;
-		case EPCGMedadataBooleanOperation::Not:
-			Value1 = !Value1;
-			break;
-		case EPCGMedadataBooleanOperation::Or:
-			Value1 = (Value1 || Value2);
-			break;
-		case EPCGMedadataBooleanOperation::Xor:
-			Value1 = (Value1 != Value2);
-			break;
-		}
-	};
-
-	// All values in OperationData.Iterators are UniquePtr. Dereference them here to make the syntax less cumbersome.
-	// Also some iterators can be null, it means they need to use the first iterator (that should never be null).
-	// We just need to make sure to not increment the first iterator multiple times in one loop.
-
-	check(OperationData.Iterators[0]);
-
-	IPCGMetadataEntryIterator& Iterator1 = *OperationData.Iterators[0];
-	IPCGMetadataEntryIterator& Iterator2 = (OperationData.Iterators.Num() >= 2 && OperationData.Iterators[1]) ? *OperationData.Iterators[1] : Iterator1;
-
-	bool bShouldIncrementIterator2 = &Iterator1 != &Iterator2;
-
-	PCGMetadataEntryKey EntryKey1 = 0;
-	PCGMetadataEntryKey EntryKey2 = 0;
-
-	bool Value1 = false;
-	bool Value2 = false;
-
-	FPCGMetadataAttribute<bool>* OutputAttribute = static_cast<FPCGMetadataAttribute<bool>*>(OperationData.OutputAttribute);
-	bool DefaultValue1 = PCGMetadataAttribute::GetValueWithBroadcast<bool>(OperationData.SourceAttributes[0], PCGInvalidEntryKey);
-	bool DefaultValue2 = (Settings->Operation != EPCGMedadataBooleanOperation::Not) ? PCGMetadataAttribute::GetValueWithBroadcast<bool>(OperationData.SourceAttributes[1], PCGInvalidEntryKey) : false;
-	BoolFunc(DefaultValue1, DefaultValue2);
-	OutputAttribute->SetDefaultValue(DefaultValue1);
-
-	for (int32 i = 0; i < OperationData.NumberOfElementsToProcess; ++i)
+		DoUnaryOp<bool, bool>(OperationData, [](const bool& Value) -> bool { return PCGMetadataBooleanSettings::UnaryOp(Value); });
+	}
+	else
 	{
-		EntryKey1 = *Iterator1;
-
-		// If the entry key is invalid, nothing to do
-		if (EntryKey1 != PCGInvalidEntryKey)
-		{
-			Value1 = PCGMetadataAttribute::GetValueWithBroadcast<bool>(OperationData.SourceAttributes[0], EntryKey1);
-
-			if (Settings->Operation != EPCGMedadataBooleanOperation::Not)
-			{
-				EntryKey2 = *Iterator2;
-				Value2 = PCGMetadataAttribute::GetValueWithBroadcast<bool>(OperationData.SourceAttributes[1], EntryKey2);
-			}
-
-			BoolFunc(Value1, Value2);
-
-			OutputAttribute->SetValue(EntryKey1, Value1);
-		}
-
-		++Iterator1;
-		if (bShouldIncrementIterator2)
-		{
-			++Iterator2;
-		}
+		DoBinaryOp<bool, bool, bool>(OperationData, [Operation = Settings->Operation](const bool& Value1, const bool& Value2) -> bool { return PCGMetadataBooleanSettings::BinaryOp(Value1, Value2, Operation); });
 	}
 
 	return true;
diff --git a/Engine/Plugins/Experimental/PCG/Source/PCG/Private/Elements/Metadata/PCGMetadataMathsOpElement.cpp b/Engine/Plugins/Experimental/PCG/Source/PCG/Private/Elements/Metadata/PCGMetadataMathsOpElement.cpp
--- a/Engine/Plugins/Experimental/PCG/Source/PCG/Private/Elements/Metadata/PCGMetadataMathsOpElement.cpp
+++ b/Engine/Plugins/Experimental/PCG/Source/PCG/Private/Elements/Metadata/PCGMetadataMathsOpElement.cpp
@@ -85,77 +85,66 @@ namespace PCGMetadataMathsSettings
 	}
 
 	template <typename T>
-	void UnaryOp(T& Value, EPCGMedadataMathsOperation Op)
+	T UnaryOp(const T& Value, EPCGMedadataMathsOperation Op)
 	{
 		switch (Op)
 		{
 		case EPCGMedadataMathsOperation::Sign:
-			Value = PCGMetadataMaths::Sign(Value);
-			break;
+			return PCGMetadataMaths::Sign(Value);
 		case EPCGMedadataMathsOperation::Frac:
-			Value = PCGMetadataMaths::Frac(Value);
-			break;
+			return PCGMetadataMaths::Frac(Value);
 		case EPCGMedadataMathsOperation::Truncate:
-			Value = PCGMetadataMaths::Truncate(Value);
-			break;
+			return PCGMetadataMaths::Truncate(Value);
 		case EPCGMedadataMathsOperation::Round:
-			Value = PCGMetadataMaths::Round(Value);
-			break;
+			return PCGMetadataMaths::Round(Value);
 		case EPCGMedadataMathsOperation::Sqrt:
-			Value = PCGMetadataMaths::Sqrt(Value);
-			break;
+			return PCGMetadataMaths::Sqrt(Value);
 		case EPCGMedadataMathsOperation::Abs:
-			Value = PCGMetadataMaths::Abs(Value);
-			break;
+			return PCGMetadataMaths::Abs(Value);
+		default:
+			return T{};
 		}
 	}
 
 	template <typename T>
-	void BinaryOp(T& Value1, const T& Value2, EPCGMedadataMathsOperation Op)
+	T BinaryOp(const T& Value1, const T& Value2, EPCGMedadataMathsOperation Op)
 	{
 		switch (Op)
 		{
 		case EPCGMedadataMathsOperation::Add:
-			Value1 = Value1 + Value2;
-			break;
+			return Value1 + Value2;
 		case EPCGMedadataMathsOperation::Subtract:
-			Value1 = Value1 - Value2;
-			break;
+			return Value1 - Value2;
 		case EPCGMedadataMathsOperation::Multiply:
-			Value1 = Value1 * Value2;
-			break;
+			return Value1 * Value2;
 		case EPCGMedadataMathsOperation::Divide:
-			Value1 = Value1 / Value2;
-			break;
+			return Value1 / Value2;
 		case EPCGMedadataMathsOperation::Max:
-			Value1 = PCGMetadataMaths::Max(Value1, Value2);
-			break;
+			return PCGMetadataMaths::Max(Value1, Value2);
 		case EPCGMedadataMathsOperation::Min:
-			Value1 = PCGMetadataMaths::Min(Value1, Value2);
-			break;
+			return PCGMetadataMaths::Min(Value1, Value2);
 		case EPCGMedadataMathsOperation::ClampMin:
-			Value1 = PCGMetadataMaths::Clamp(Value1, Value2, Value1);
-			break;
+			return PCGMetadataMaths::Clamp(Value1, Value2, Value1);
 		case EPCGMedadataMathsOperation::ClampMax:
-			Value1 = PCGMetadataMaths::Clamp(Value1, Value1, Value2);
-			break;
+			return PCGMetadataMaths::Clamp(Value1, Value1, Value2);
 		case EPCGMedadataMathsOperation::Pow:
-			Value1 = PCGMetadataMaths::Pow(Value1, Value2);
-			break;
+			return PCGMetadataMaths::Pow(Value1, Value2);
+		default:
+			return T{};
 		}
 	}
 
 	template <typename T>
-	void TernaryOp(T& Value1, const T& Value2, const T& Value3, EPCGMedadataMathsOperation Op)
+	T TernaryOp(const T& Value1, const T& Value2, const T& Value3, EPCGMedadataMathsOperation Op)
 	{
 		switch (Op)
 		{
 		case EPCGMedadataMathsOperation::Clamp:
-			Value1 = PCGMetadataMaths::Clamp(Value1, Value2, Value3);
-			break;
+			return PCGMetadataMaths::Clamp(Value1, Value2, Value3);
 		case EPCGMedadataMathsOperation::Lerp:
-			Value1 = PCGMetadataMaths::Lerp(Value1, Value2, Value3);
-			break;
+			return PCGMetadataMaths::Lerp(Value1, Value2, Value3);
+		default:
+			return T{};
 		}
 	}
 }
@@ -240,162 +229,36 @@ bool FPCGMetadataMathsElement::DoOperation(FOperationData& OperationData) const
 
 	const UPCGMetadataMathsSettings* Settings = CastChecked<UPCGMetadataMathsSettings>(OperationData.Settings);
 
-	// All values in OperationData.Iterators are UniquePtr. Dereference them here to make the syntax less cumbersome.
-	// Also some iterators can be null, it means they need to use the first iterator (that should never be null).
-	// We just need to make sure to not increment the first iterator multiple times in one loop.
-
-	check(OperationData.Iterators[0]);
-
-	IPCGMetadataEntryIterator& Iterator1 = *OperationData.Iterators[0];
-	IPCGMetadataEntryIterator& Iterator2 = (OperationData.Iterators.Num() >= 2 && OperationData.Iterators[1]) ? *OperationData.Iterators[1] : Iterator1;
-	IPCGMetadataEntryIterator& Iterator3 = (OperationData.Iterators.Num() >= 3 && OperationData.Iterators[2]) ? *OperationData.Iterators[2] : Iterator1;
-
-	bool bShouldIncrementIterator2 = &Iterator1 != &Iterator2;
-	bool bShouldIncrementIterator3 = &Iterator1 != &Iterator3;
-
-	// And then do the operation for all elements
-	if (PCGMetadataMathsSettings::IsUnaryOp(Settings->Operation))
+	auto MathFunc = [this, Operation = Settings->Operation, &OperationData](auto DummyOutValue) -> void
 	{
-		TRACE_CPUPROFILER_EVENT_SCOPE(FPCGMetadataMathsElement::ExecuteInternal::UnaryOp);
-		auto UnaryFunc = [Operation = Settings->Operation, &OperationData, &Iterator1](auto DummyOutValue) -> void
-		{
-			using AttributeType = decltype(DummyOutValue);
-
-			// Need to remove types that would not compile
-			if constexpr (PCG::Private::MetadataTypes<AttributeType>::Id > (uint16)EPCGMetadataTypes::Vector4)
-			{
-				return;
-			}
-			else
-			{
-				FPCGMetadataAttribute<AttributeType>* OutputAttribute = static_cast<FPCGMetadataAttribute<AttributeType>*>(OperationData.OutputAttribute);
-				AttributeType DefaultValue = PCGMetadataAttribute::GetValueWithBroadcast<AttributeType>(OperationData.SourceAttributes[0], PCGInvalidEntryKey);
-				PCGMetadataMathsSettings::UnaryOp(DefaultValue, Operation);
-				OutputAttribute->SetDefaultValue(DefaultValue);
-
-				for (int32 i = 0; i < OperationData.NumberOfElementsToProcess; ++i)
-				{
-					PCGMetadataEntryKey EntryKey = *Iterator1;
-
-					// If the entry key is invalid, nothing to do
-					if (EntryKey != PCGInvalidEntryKey)
-					{
-						AttributeType Value = PCGMetadataAttribute::GetValueWithBroadcast<AttributeType>(OperationData.SourceAttributes[0], EntryKey);
-
-						PCGMetadataMathsSettings::UnaryOp(Value, Operation);
-
-						OutputAttribute->SetValue(EntryKey, Value);
-					}
-
-					++Iterator1;
-				}
-			}
-		};
+		using AttributeType = decltype(DummyOutValue);
 
-		PCGMetadataAttribute::CallbackWithRightType(OperationData.OutputType, UnaryFunc);
-	}
-	else if (PCGMetadataMathsSettings::IsBinaryOp(Settings->Operation))
-	{
-		auto BinaryFunc = [Operation = Settings->Operation, &OperationData, &Iterator1, &Iterator2, bShouldIncrementIterator2](auto DummyOutValue) -> void
+		// Need to remove types that would not compile
+		if constexpr (PCG::Private::MetadataTypes<AttributeType>::Id > (uint16)EPCGMetadataTypes::Vector4)
 		{
-			TRACE_CPUPROFILER_EVENT_SCOPE(FPCGMetadataMathsElement::ExecuteInternal::BinaryOp);
-			using AttributeType = decltype(DummyOutValue);
-
-			// Need to remove types that would not compile
-			if constexpr (PCG::Private::MetadataTypes<AttributeType>::Id > (uint16)EPCGMetadataTypes::Vector4)
-			{
-				return;
-			}
-			else
+			return;
+		}
+		else
+		{
+			if (PCGMetadataMathsSettings::IsUnaryOp(Operation))
 			{
-				FPCGMetadataAttribute<AttributeType>* OutputAttribute = static_cast<FPCGMetadataAttribute<AttributeType>*>(OperationData.OutputAttribute);
-				AttributeType DefaultValue1 = PCGMetadataAttribute::GetValueWithBroadcast<AttributeType>(OperationData.SourceAttributes[0], PCGInvalidEntryKey);
-				AttributeType DefaultValue2 = PCGMetadataAttribute::GetValueWithBroadcast<AttributeType>(OperationData.SourceAttributes[1], PCGInvalidEntryKey);
-				PCGMetadataMathsSettings::BinaryOp(DefaultValue1, DefaultValue2, Operation);
-				OutputAttribute->SetDefaultValue(DefaultValue1);
-
-				for (int32 i = 0; i < OperationData.NumberOfElementsToProcess; ++i)
-				{
-					PCGMetadataEntryKey EntryKey1 = *Iterator1;
-					PCGMetadataEntryKey EntryKey2 = *Iterator2;
-
-					// If the entry key is invalid, nothing to do
-					if (EntryKey1 != PCGInvalidEntryKey)
-					{
-						AttributeType Value1 = PCGMetadataAttribute::GetValueWithBroadcast<AttributeType>(OperationData.SourceAttributes[0], EntryKey1);
-						AttributeType Value2 = PCGMetadataAttribute::GetValueWithBroadcast<AttributeType>(OperationData.SourceAttributes[1], EntryKey2);
-
-						PCGMetadataMathsSettings::BinaryOp(Value1, Value2, Operation);
-
-						OutputAttribute->SetValue(EntryKey1, Value1);
-					}
-
-					++Iterator1;
-					if (bShouldIncrementIterator2)
-					{
-						++Iterator2;
-					}
-				}
+				TRACE_CPUPROFILER_EVENT_SCOPE(FPCGMetadataMathsElement::ExecuteInternal::UnaryOp);
+				DoUnaryOp<AttributeType, AttributeType>(OperationData, [Operation](const AttributeType& Value) -> AttributeType { return PCGMetadataMathsSettings::UnaryOp(Value, Operation); });
 			}
-		};
-
-		PCGMetadataAttribute::CallbackWithRightType(OperationData.OutputType, BinaryFunc);
-	}
-	else if (PCGMetadataMathsSettings::IsTernaryOp(Settings->Operation))
-	{
-		auto TernaryFunc = [Operation = Settings->Operation, &OperationData, &Iterator1, &Iterator2, &Iterator3, bShouldIncrementIterator2, bShouldIncrementIterator3](auto DummyOutValue) -> void
-		{
-			TRACE_CPUPROFILER_EVENT_SCOPE(FPCGMetadataMathsElement::ExecuteInternal::TernaryOp);
-			using AttributeType = decltype(DummyOutValue);
-
-			// Need to remove types that would not compile
-			if constexpr (PCG::Private::MetadataTypes<AttributeType>::Id > (uint16)EPCGMetadataTypes::Vector4)
+			else if (PCGMetadataMathsSettings::IsBinaryOp(Operation))
 			{
-				return;
+				TRACE_CPUPROFILER_EVENT_SCOPE(FPCGMetadataMathsElement::ExecuteInternal::BinaryOp);
+				DoBinaryOp<AttributeType, AttributeType, AttributeType>(OperationData, [Operation](const AttributeType& Value1, const AttributeType& Value2) -> AttributeType { return PCGMetadataMathsSettings::BinaryOp(Value1, Value2, Operation); });
 			}
-			else
+			else if (PCGMetadataMathsSettings::IsTernaryOp(Operation))
 			{
-				FPCGMetadataAttribute<AttributeType>* OutputAttribute = static_cast<FPCGMetadataAttribute<AttributeType>*>(OperationData.OutputAttribute);
-				AttributeType DefaultValue1 = PCGMetadataAttribute::GetValueWithBroadcast<AttributeType>(OperationData.SourceAttributes[0], PCGInvalidEntryKey);
-				AttributeType DefaultValue2 = PCGMetadataAttribute::GetValueWithBroadcast<AttributeType>(OperationData.SourceAttributes[1], PCGInvalidEntryKey);
-				AttributeType DefaultValue3 = PCGMetadataAttribute::GetValueWithBroadcast<AttributeType>(OperationData.SourceAttributes[2], PCGInvalidEntryKey);
-				PCGMetadataMathsSettings::TernaryOp(DefaultValue1, DefaultValue2, DefaultValue3, Operation);
-				OutputAttribute->SetDefaultValue(DefaultValue1);
-
-				for (PCGMetadataValueKey ValueKey = 0; ValueKey < OperationData.NumberOfElementsToProcess; ++ValueKey)
-				{
-					PCGMetadataEntryKey EntryKey1 = *Iterator1;
-					PCGMetadataEntryKey EntryKey2 = *Iterator2;
-					PCGMetadataEntryKey EntryKey3 = *Iterator3;
-
-					// If the entry key is invalid, nothing to do
-					if (EntryKey1 != PCGInvalidEntryKey)
-					{
-						AttributeType Value1 = PCGMetadataAttribute::GetValueWithBroadcast<AttributeType>(OperationData.SourceAttributes[0], EntryKey1);
-						AttributeType Value2 = PCGMetadataAttribute::GetValueWithBroadcast<AttributeType>(OperationData.SourceAttributes[1], EntryKey2);
-						AttributeType Value3 = PCGMetadataAttribute::GetValueWithBroadcast<AttributeType>(OperationData.SourceAttributes[2], EntryKey3);
-
-						PCGMetadataMathsSettings::TernaryOp(Value1, Value2, Value3, Operation);
-
-						OutputAttribute->SetValue(EntryKey1, Value1);
-					}
-
-					++Iterator1;
-					if (bShouldIncrementIterator2)
-					{
-						++Iterator2;
-					}
-
-					if (bShouldIncrementIterator3)
-					{
-						++Iterator3;
-					}
-				}
+				TRACE_CPUPROFILER_EVENT_SCOPE(FPCGMetadataMathsElement::ExecuteInternal::TernaryOp);
+				DoTernaryOp<AttributeType, AttributeType, AttributeType, AttributeType>(OperationData, [Operation](const AttributeType& Value1, const AttributeType& Value2, const AttributeType& Value3) -> AttributeType { return PCGMetadataMathsSettings::TernaryOp(Value1, Value2, Value3, Operation); });
 			}
-		};
+		}
+	};
 
-		PCGMetadataAttribute::CallbackWithRightType(OperationData.OutputType, TernaryFunc);
-	}
+	PCGMetadataAttribute::CallbackWithRightType(OperationData.OutputType, MathFunc);
 
 	return true;
 }
\ No newline at end of file
diff --git a/Engine/Plugins/Experimental/PCG/Source/PCG/Private/Elements/Metadata/PCGMetadataOpElementBase.cpp b/Engine/Plugins/Experimental/PCG/Source/PCG/Private/Elements/Metadata/PCGMetadataOpElementBase.cpp
--- a/Engine/Plugins/Experimental/PCG/Source/PCG/Private/Elements/Metadata/PCGMetadataOpElementBase.cpp
+++ b/Engine/Plugins/Experimental/PCG/Source/PCG/Private/Elements/Metadata/PCGMetadataOpElementBase.cpp
@@ -108,7 +108,7 @@ bool FPCGMetadataElementBase::ExecuteInternal(FPCGContext* Context) const
 
 	OperationData.SourceAttributes.SetNum(NumberOfInputs);
 
-	uint16 MostComplexType = (uint16)EPCGMetadataTypes::Unknown;
+	OperationData.MostComplexInputType = (uint16)EPCGMetadataTypes::Unknown;
 	OperationData.NumberOfElementsToProcess = -1;
 
 	for (uint32 i = 0; i < NumberOfInputs; ++i)
@@ -138,11 +138,11 @@ bool FPCGMetadataElementBase::ExecuteInternal(FPCGContext* Context) const
 		if (!bHasSpecialRequirement)
 		{
 			// In this case, check if we have a more complex type, or if we can broadcast to the most complex type.
-			if (MostComplexType == (uint16)EPCGMetadataTypes::Unknown || Settings->IsMoreComplexType(OperationData.SourceAttributes[i]->GetTypeId(), MostComplexType))
+			if (OperationData.MostComplexInputType == (uint16)EPCGMetadataTypes::Unknown || Settings->IsMoreComplexType(OperationData.SourceAttributes[i]->GetTypeId(), OperationData.MostComplexInputType))
 			{
-				MostComplexType = OperationData.SourceAttributes[i]->GetTypeId();
+				OperationData.MostComplexInputType = OperationData.SourceAttributes[i]->GetTypeId();
 			}
-			else if (MostComplexType != OperationData.SourceAttributes[i]->GetTypeId() && !PCG::Private::IsBroadcastable(OperationData.SourceAttributes[i]->GetTypeId(), MostComplexType))
+			else if (OperationData.MostComplexInputType != OperationData.SourceAttributes[i]->GetTypeId() && !PCG::Private::IsBroadcastable(OperationData.SourceAttributes[i]->GetTypeId(), OperationData.MostComplexInputType))
 			{
 				PCGE_LOG(Error, "Attribute %s cannot be broadcasted to match types for input %d", *SourceAttributeName.ToString(), i);
 				return true;
@@ -163,7 +163,7 @@ bool FPCGMetadataElementBase::ExecuteInternal(FPCGContext* Context) const
 		}
 		else
 		{
-			NumberOfElements[i] = OperationData.SourceAttributes[i]->GetNumberOfEntries();
+			NumberOfElements[i] = OperationData.SourceAttributes[i]->GetNumberOfEntriesWithParents();
 
 			// No element means only dealing with the default value.
 			if (NumberOfElements[i] == 0)
@@ -209,7 +209,7 @@ bool FPCGMetadataElementBase::ExecuteInternal(FPCGContext* Context) const
 		PCGMetadataElementCommon::CopyEntryToValueKeyMap(SourceMetadata[0], SourceAttributes[0], OutputAttribute);
 	};
 
-	OperationData.OutputType = Settings->GetOutputType(MostComplexType);
+	OperationData.OutputType = Settings->GetOutputType(OperationData.MostComplexInputType);
 
 	PCGMetadataAttribute::CallbackWithRightType(OperationData.OutputType, CreateAttribute);
 
diff --git a/Engine/Plugins/Experimental/PCG/Source/PCG/Private/Elements/Metadata/PCGMetadataTrigOpElement.cpp b/Engine/Plugins/Experimental/PCG/Source/PCG/Private/Elements/Metadata/PCGMetadataTrigOpElement.cpp
--- a/Engine/Plugins/Experimental/PCG/Source/PCG/Private/Elements/Metadata/PCGMetadataTrigOpElement.cpp
+++ b/Engine/Plugins/Experimental/PCG/Source/PCG/Private/Elements/Metadata/PCGMetadataTrigOpElement.cpp
@@ -12,41 +12,37 @@
 namespace PCGMetadataTrigSettings
 {
 	template <typename OutType>
-	void ApplyTrigOperation(OutType& Input1, const OutType& Input2, EPCGMedadataTrigOperation Operation)
+	OutType UnaryOp(const OutType& Input1, EPCGMedadataTrigOperation Operation)
 	{
 		switch (Operation)
 		{
 		case EPCGMedadataTrigOperation::Acos:
-			Input1 = FMath::Acos(Input1);
-			break;
+			return FMath::Acos(Input1);
 		case EPCGMedadataTrigOperation::Asin:
-			Input1 = FMath::Asin(Input1);
-			break;
+			return FMath::Asin(Input1);
 		case EPCGMedadataTrigOperation::Atan:
-			Input1 = FMath::Atan(Input1);
-			break;
-		case EPCGMedadataTrigOperation::Atan2:
-			Input1 = FMath::Atan2(Input1, Input2);
-			break;
+			return FMath::Atan(Input1);
 		case EPCGMedadataTrigOperation::Cos:
-			Input1 = FMath::Cos(Input1);
-			break;
+			return FMath::Cos(Input1);
 		case EPCGMedadataTrigOperation::Sin:
-			Input1 = FMath::Sin(Input1);
-			break;
+			return FMath::Sin(Input1);
 		case EPCGMedadataTrigOperation::Tan:
-			Input1 = FMath::Tan(Input1);
-			break;
+			return FMath::Tan(Input1);
 		case EPCGMedadataTrigOperation::DegToRad:
-			Input1 = FMath::DegreesToRadians(Input1);
-			break;
+			return FMath::DegreesToRadians(Input1);
 		case EPCGMedadataTrigOperation::RadToDeg:
-			Input1 = FMath::RadiansToDegrees(Input1);
-			break;
+			return FMath::RadiansToDegrees(Input1);
 		default:
-			break;
+			return OutType{};
 		}
 	}
+
+	template <typename OutType>
+	OutType BinaryOp(const OutType& Input1, const OutType& Input2)
+	{
+		// EPCGMedadataTrigOperation::Atan2:
+		return FMath::Atan2(Input1, Input2);
+	}
 }
 
 FName UPCGMetadataTrigSettings::GetInputPinLabel(uint32 Index) const
@@ -121,70 +117,14 @@ bool FPCGMetadataTrigElement::DoOperation(FOperationData& OperationData) const
 
 	const UPCGMetadataTrigSettings* Settings = CastChecked<UPCGMetadataTrigSettings>(OperationData.Settings);
 
-	// All values in OperationData.Iterators are UniquePtr. Dereference them here to make the syntax less cumbersome.
-	// Also some iterators can be null, it means they need to use the first iterator (that should never be null).
-	// We just need to make sure to not increment the first iterator multiple times in one loop.
-
-	check(OperationData.Iterators[0]);
-
-	IPCGMetadataEntryIterator& Iterator1 = *OperationData.Iterators[0];
-	IPCGMetadataEntryIterator& Iterator2 = (OperationData.Iterators.Num() >= 2 && OperationData.Iterators[1]) ? *OperationData.Iterators[1] : Iterator1;
-
-	bool bShouldIncrementIterator2 = &Iterator1 != &Iterator2;
-
-	auto TrigFunc = [Operation = Settings->Operation, &Iterator1, &Iterator2, bShouldIncrementIterator2, &OperationData](auto DummyValue)
+	if (Settings->Operation == EPCGMedadataTrigOperation::Atan2)
 	{
-		using AttributeType = decltype(DummyValue);
-
-		// Need to remove types that would not compile
-		if constexpr (PCG::Private::MetadataTypes<AttributeType>::Id > (uint16)EPCGMetadataTypes::Double)
-		{
-			return;
-		}
-		else
-		{
-			PCGMetadataEntryKey EntryKey1 = 0;
-			PCGMetadataEntryKey EntryKey2 = 0;
-
-			FPCGMetadataAttribute<AttributeType>* OutputAttribute = static_cast<FPCGMetadataAttribute<AttributeType>*>(OperationData.OutputAttribute);
-			AttributeType DefaultValue1 = PCGMetadataAttribute::GetValueWithBroadcast<AttributeType>(OperationData.SourceAttributes[0], PCGInvalidEntryKey);
-			AttributeType DefaultValue2 = (Operation == EPCGMedadataTrigOperation::Atan2) ? PCGMetadataAttribute::GetValueWithBroadcast<AttributeType>(OperationData.SourceAttributes[1], PCGInvalidEntryKey) : AttributeType{};
-			PCGMetadataTrigSettings::ApplyTrigOperation(DefaultValue1, DefaultValue2, Operation);
-			OutputAttribute->SetDefaultValue(DefaultValue1);
-
-			AttributeType Value1{};
-			AttributeType Value2{};
-
-			for (int32 i = 0; i < OperationData.NumberOfElementsToProcess; ++i)
-			{
-				EntryKey1 = *Iterator1;
-
-				// If the entry key is invalid, nothing to do
-				if (EntryKey1 != PCGInvalidEntryKey)
-				{
-					Value1 = PCGMetadataAttribute::GetValueWithBroadcast<AttributeType>(OperationData.SourceAttributes[0], EntryKey1);
-
-					if (Operation == EPCGMedadataTrigOperation::Atan2)
-					{
-						EntryKey2 = *Iterator2;
-						Value2 = PCGMetadataAttribute::GetValueWithBroadcast<AttributeType>(OperationData.SourceAttributes[1], EntryKey2);
-					}
-
-					PCGMetadataTrigSettings::ApplyTrigOperation(Value1, Value2, Operation);
-
-					OutputAttribute->SetValue(EntryKey1, Value1);
-				}
-
-				++Iterator1;
-				if (bShouldIncrementIterator2)
-				{
-					++Iterator2;
-				}
-			}
-		}
-	};
-
-	PCGMetadataAttribute::CallbackWithRightType(OperationData.OutputType, TrigFunc);
+		DoBinaryOp<double, double, double>(OperationData, [](const double& Value1, const double& Value2) -> double { return PCGMetadataTrigSettings::BinaryOp(Value1, Value2); });
+	}
+	else
+	{
+		DoUnaryOp<double, double>(OperationData, [Operation = Settings->Operation](const double& Value) -> double { return PCGMetadataTrigSettings::UnaryOp(Value, Operation); });
+	}
 
 	return true;
 }
\ No newline at end of file
diff --git a/Engine/Plugins/Experimental/PCG/Source/PCG/Public/Elements/Metadata/PCGMetadataOpElementBase.h b/Engine/Plugins/Experimental/PCG/Source/PCG/Public/Elements/Metadata/PCGMetadataOpElementBase.h
--- a/Engine/Plugins/Experimental/PCG/Source/PCG/Public/Elements/Metadata/PCGMetadataOpElementBase.h
+++ b/Engine/Plugins/Experimental/PCG/Source/PCG/Public/Elements/Metadata/PCGMetadataOpElementBase.h
@@ -5,6 +5,10 @@
 #include "PCGElement.h"
 #include "PCGSettings.h"
 
+#include "Metadata/PCGMetadataAttribute.h"
+#include "Metadata/PCGMetadataAttributeTpl.h"
+#include "Metadata/PCGMetadataEntryKeyIterator.h"
+
 #include "PCGMetadataOpElementBase.generated.h"
 
 class FPCGMetadataAttributeBase;
@@ -94,18 +98,165 @@ class FPCGMetadataElementBase : public FSimplePCGElement
 public:
 	virtual bool IsCacheable(const UPCGSettings* InSettings) const override { return true; }
 
-protected:
-	virtual bool ExecuteInternal(FPCGContext* Context) const override;
-
 	struct FOperationData
 	{
 		TArray<const FPCGMetadataAttributeBase*> SourceAttributes;
 		int32 NumberOfElementsToProcess = -1;
+		uint16 MostComplexInputType;
 		uint16 OutputType;
 		FPCGMetadataAttributeBase* OutputAttribute = nullptr;
 		const UPCGMetadataSettingsBase* Settings = nullptr;
 		TArray<TUniquePtr<IPCGMetadataEntryIterator>> Iterators;
 	};
 
-	virtual bool DoOperation(FOperationData& OperationData) const { return true; };
-};
\ No newline at end of file
+protected:
+	virtual bool ExecuteInternal(FPCGContext* Context) const override;
+
+	virtual bool DoOperation(FOperationData& OperationData) const = 0;
+
+	template <typename InType, typename OutType>
+	void DoUnaryOp(FOperationData& OperationData, TFunction<OutType(const InType&)> UnaryOp) const;
+
+	template <typename InType1, typename InType2, typename OutType>
+	void DoBinaryOp(FOperationData& OperationData, TFunction<OutType(const InType1&, const InType2&)> BinaryOp) const;
+
+	template <typename InType1, typename InType2, typename InType3, typename OutType>
+	void DoTernaryOp(FOperationData& OperationData, TFunction<OutType(const InType1&, const InType2&, const InType3&)> TernaryOp) const;
+};
+
+template <typename InType, typename OutType>
+inline void FPCGMetadataElementBase::DoUnaryOp(FOperationData& OperationData, TFunction<OutType(const InType&)> UnaryOp) const
+{
+	check(OperationData.SourceAttributes[0]);
+
+	// All values in OperationData.Iterators are UniquePtr. Dereference them here to make the syntax less cumbersome.
+	check(OperationData.Iterators[0]);
+	IPCGMetadataEntryIterator& Iterator1 = *OperationData.Iterators[0];
+
+	FPCGMetadataAttribute<OutType>* OutputAttribute = static_cast<FPCGMetadataAttribute<OutType>*>(OperationData.OutputAttribute);
+	InType DefaultValue = PCGMetadataAttribute::GetValueWithBroadcast<InType>(OperationData.SourceAttributes[0], PCGInvalidEntryKey);
+	OutType DefaultOutputValue = UnaryOp(DefaultValue);
+	OutputAttribute->SetDefaultValue(DefaultOutputValue);
+
+	for (int32 i = 0; i < OperationData.NumberOfElementsToProcess; ++i)
+	{
+		PCGMetadataEntryKey EntryKey = *Iterator1;
+
+		// If the entry key is invalid, nothing to do
+		if (EntryKey != PCGInvalidEntryKey)
+		{
+			InType Value = PCGMetadataAttribute::GetValueWithBroadcast<InType>(OperationData.SourceAttributes[0], EntryKey);
+
+			OutType OutputValue = UnaryOp(Value);
+
+			OutputAttribute->SetValue(EntryKey, Value);
+		}
+
+		++Iterator1;
+	}
+}
+
+template <typename InType1, typename InType2, typename OutType>
+inline void FPCGMetadataElementBase::DoBinaryOp(FOperationData& OperationData, TFunction<OutType(const InType1&, const InType2&)> BinaryOp) const
+{
+	check(OperationData.SourceAttributes[0]);
+	check(OperationData.SourceAttributes[1]);
+
+	// All values in OperationData.Iterators are UniquePtr. Dereference them here to make the syntax less cumbersome.
+	// Also some iterators can be null, it means they need to use the first iterator (that should never be null).
+	// We just need to make sure to not increment the first iterator multiple times in one loop.
+	check(OperationData.Iterators[0]);
+
+	IPCGMetadataEntryIterator& Iterator1 = *OperationData.Iterators[0];
+	IPCGMetadataEntryIterator& Iterator2 = (OperationData.Iterators.Num() >= 2 && OperationData.Iterators[1]) ? *OperationData.Iterators[1] : Iterator1;
+
+	bool bShouldIncrementIterator2 = &Iterator1 != &Iterator2;
+
+	FPCGMetadataAttribute<OutType>* OutputAttribute = static_cast<FPCGMetadataAttribute<OutType>*>(OperationData.OutputAttribute);
+	InType1 DefaultValue1 = PCGMetadataAttribute::GetValueWithBroadcast<InType1>(OperationData.SourceAttributes[0], PCGInvalidEntryKey);
+	InType2 DefaultValue2 = PCGMetadataAttribute::GetValueWithBroadcast<InType2>(OperationData.SourceAttributes[1], PCGInvalidEntryKey);
+	OutType DefaultOutputValue = BinaryOp(DefaultValue1, DefaultValue2);
+	OutputAttribute->SetDefaultValue(DefaultOutputValue);
+
+	for (int32 i = 0; i < OperationData.NumberOfElementsToProcess; ++i)
+	{
+		PCGMetadataEntryKey EntryKey1 = *Iterator1;
+
+		// If the entry key is invalid, nothing to do
+		if (EntryKey1 != PCGInvalidEntryKey)
+		{
+			PCGMetadataEntryKey EntryKey2 = *Iterator2;
+
+			InType1 Value1 = PCGMetadataAttribute::GetValueWithBroadcast<InType1>(OperationData.SourceAttributes[0], EntryKey1);
+			InType2 Value2 = PCGMetadataAttribute::GetValueWithBroadcast<InType2>(OperationData.SourceAttributes[1], EntryKey2);
+
+			OutType OutputValue = BinaryOp(Value1, Value2);
+
+			OutputAttribute->SetValue(EntryKey1, OutputValue);
+		}
+
+		++Iterator1;
+		if (bShouldIncrementIterator2)
+		{
+			++Iterator2;
+		}
+	}
+}
+
+template <typename InType1, typename InType2, typename InType3, typename OutType>
+inline void FPCGMetadataElementBase::DoTernaryOp(FOperationData& OperationData, TFunction<OutType(const InType1&, const InType2&, const InType3&)> TernaryOp) const
+{
+	check(OperationData.SourceAttributes[0]);
+	check(OperationData.SourceAttributes[1]);
+	check(OperationData.SourceAttributes[2]);
+
+	// All values in OperationData.Iterators are UniquePtr. Dereference them here to make the syntax less cumbersome.
+	// Also some iterators can be null, it means they need to use the first iterator (that should never be null).
+	// We just need to make sure to not increment the first iterator multiple times in one loop.
+	check(OperationData.Iterators[0]);
+
+	IPCGMetadataEntryIterator& Iterator1 = *OperationData.Iterators[0];
+	IPCGMetadataEntryIterator& Iterator2 = (OperationData.Iterators.Num() >= 2 && OperationData.Iterators[1]) ? *OperationData.Iterators[1] : Iterator1;
+	IPCGMetadataEntryIterator& Iterator3 = (OperationData.Iterators.Num() >= 3 && OperationData.Iterators[2]) ? *OperationData.Iterators[2] : Iterator1;
+
+	bool bShouldIncrementIterator2 = &Iterator1 != &Iterator2;
+	bool bShouldIncrementIterator3 = &Iterator1 != &Iterator3;
+
+	FPCGMetadataAttribute<OutType>* OutputAttribute = static_cast<FPCGMetadataAttribute<OutType>*>(OperationData.OutputAttribute);
+	InType1 DefaultValue1 = PCGMetadataAttribute::GetValueWithBroadcast<InType1>(OperationData.SourceAttributes[0], PCGInvalidEntryKey);
+	InType2 DefaultValue2 = PCGMetadataAttribute::GetValueWithBroadcast<InType2>(OperationData.SourceAttributes[1], PCGInvalidEntryKey);
+	InType3 DefaultValue3 = PCGMetadataAttribute::GetValueWithBroadcast<InType3>(OperationData.SourceAttributes[2], PCGInvalidEntryKey);
+	OutType DefaultOutputValue = TernaryOp(DefaultValue1, DefaultValue2, DefaultValue3);
+	OutputAttribute->SetDefaultValue(DefaultOutputValue);
+
+	for (int32 i = 0; i < OperationData.NumberOfElementsToProcess; ++i)
+	{
+		PCGMetadataEntryKey EntryKey1 = *Iterator1;
+
+		// If the entry key is invalid, nothing to do
+		if (EntryKey1 != PCGInvalidEntryKey)
+		{
+			PCGMetadataEntryKey EntryKey2 = *Iterator2;
+			PCGMetadataEntryKey EntryKey3 = *Iterator3;
+
+			InType1 Value1 = PCGMetadataAttribute::GetValueWithBroadcast<InType1>(OperationData.SourceAttributes[0], EntryKey1);
+			InType2 Value2 = PCGMetadataAttribute::GetValueWithBroadcast<InType2>(OperationData.SourceAttributes[1], EntryKey2);
+			InType3 Value3 = PCGMetadataAttribute::GetValueWithBroadcast<InType3>(OperationData.SourceAttributes[2], EntryKey3);
+
+			OutType OutputValue = TernaryOp(Value1, Value2, Value3);
+
+			OutputAttribute->SetValue(EntryKey1, OutputValue);
+		}
+
+		++Iterator1;
+		if (bShouldIncrementIterator2)
+		{
+			++Iterator2;
+		}
+
+		if (bShouldIncrementIterator3)
+		{
+			++Iterator3;
+		}
+	}
+}
\ No newline at end of file
diff --git a/Engine/Plugins/Experimental/PCG/Source/PCG/Public/Elements/Metadata/PCGMetadataTrigOpElement.h b/Engine/Plugins/Experimental/PCG/Source/PCG/Public/Elements/Metadata/PCGMetadataTrigOpElement.h
--- a/Engine/Plugins/Experimental/PCG/Source/PCG/Public/Elements/Metadata/PCGMetadataTrigOpElement.h
+++ b/Engine/Plugins/Experimental/PCG/Source/PCG/Public/Elements/Metadata/PCGMetadataTrigOpElement.h
@@ -56,7 +56,7 @@ public:
 	UPROPERTY(BlueprintReadWrite, EditAnywhere, Category = Input)
 	FName Input1AttributeName = NAME_None;
 
-	UPROPERTY(BlueprintReadWrite, EditAnywhere, Category = Input, meta = (EditCondition="Operation = EPCGMedadataTrigOperation::Atan2", EditConditionHides))
+	UPROPERTY(BlueprintReadWrite, EditAnywhere, Category = Input, meta = (EditCondition="Operation == EPCGMedadataTrigOperation::Atan2", EditConditionHides))
 	FName Input2AttributeName = NAME_None;
 };
 
diff --git a/Engine/Plugins/Experimental/PCG/Source/PCG/Public/Metadata/PCGMetadataAttribute.h b/Engine/Plugins/Experimental/PCG/Source/PCG/Public/Metadata/PCGMetadataAttribute.h
--- a/Engine/Plugins/Experimental/PCG/Source/PCG/Public/Metadata/PCGMetadataAttribute.h
+++ b/Engine/Plugins/Experimental/PCG/Source/PCG/Public/Metadata/PCGMetadataAttribute.h
@@ -40,6 +40,7 @@ public:
 	bool AllowsInterpolation() const { return bAllowsInterpolation; }
 
 	int32 GetNumberOfEntries() const { return EntryToValueKeyMap.Num(); }
+	int32 GetNumberOfEntriesWithParents() const { return EntryToValueKeyMap.Num() + (Parent ? Parent->GetNumberOfEntries() : 0); }
 
 	// This call is not thread safe
 	const TMap<PCGMetadataEntryKey, PCGMetadataValueKey>& GetEntryToValueKeyMap_NotThreadSafe() const { return EntryToValueKeyMap; }
