[PCG] Metadata native operation: Vector operations

diff --git a/Engine/Plugins/Experimental/PCG/Source/PCG/Private/Elements/Metadata/PCGMetadataVectorOpElement.cpp b/Engine/Plugins/Experimental/PCG/Source/PCG/Private/Elements/Metadata/PCGMetadataVectorOpElement.cpp
new file mode 100644
--- /dev/null
+++ b/Engine/Plugins/Experimental/PCG/Source/PCG/Private/Elements/Metadata/PCGMetadataVectorOpElement.cpp
@@ -0,0 +1,410 @@
+// Copyright Epic Games, Inc. All Rights Reserved.
+
+#include "Elements/Metadata/PCGMetadataVectorOpElement.h"
+
+#include "PCGParamData.h"
+#include "Data/PCGSpatialData.h"
+#include "Helpers/PCGSettingsHelpers.h"
+#include "Metadata/PCGMetadata.h"
+#include "Metadata/PCGMetadataAttribute.h"
+#include "Metadata/PCGMetadataAttributeTpl.h"
+#include "Metadata/PCGMetadataEntryKeyIterator.h"
+#include "Elements/Metadata/PCGMetadataElementCommon.h"
+
+namespace PCGMetadataVectorConstants
+{
+	const FName AxisLabel = TEXT("Axis");
+	const FName AngleLabel = TEXT("Angle(Deg)");
+}
+
+namespace PCGMetadataVectorSettings
+{
+	inline constexpr bool IsUnaryOp(EPCGMedadataVectorOperation Operation)
+	{
+		return Operation == EPCGMedadataVectorOperation::Normalize || 
+			Operation == EPCGMedadataVectorOperation::Length;
+	}
+
+	inline constexpr bool IsTernaryOp(EPCGMedadataVectorOperation Operation)
+	{
+		return Operation == EPCGMedadataVectorOperation::RotateAroundAxis;
+	}
+
+	inline constexpr bool IsTransformOp(EPCGMedadataVectorOperation Operation)
+	{
+		return (uint16)Operation >= (uint16)EPCGMedadataVectorOperation::TransformOp;
+	}
+
+	template <typename InType>
+	inline constexpr bool IsInvalidInType()
+	{
+		return PCG::Private::MetadataTypes<InType>::Id < (uint16)EPCGMetadataTypes::Vector2 || PCG::Private::MetadataTypes<InType>::Id > (uint16)EPCGMetadataTypes::Vector4;
+	}
+
+	// Mimic KismetMathLibrary
+	template <typename InType>
+	inline InType ApplyTransformOperation(const InType& Value, const FTransform& Transform, EPCGMedadataVectorOperation Operation)
+	{
+		switch (Operation)
+		{
+		case EPCGMedadataVectorOperation::TransformDirection:
+		{
+			if constexpr (std::is_same_v<InType, FVector2D>)
+			{
+				// Operation doesn't exist in KismetMathLibrary, do the operation as Vec3, Z as 0 and discard Z in result
+				FVector TempValue{ Value, 0.0 };
+				return FVector2D{ Transform.TransformVectorNoScale(TempValue) };
+			}
+			else if constexpr (std::is_same_v<InType, FVector4>)
+			{
+				return Transform.TransformFVector4NoScale(Value);
+			}
+			else
+			{
+				return Transform.TransformVectorNoScale(Value);
+			}
+		}
+		case EPCGMedadataVectorOperation::TransformLocation:
+		{
+			if constexpr (std::is_same_v<InType, FVector2D>)
+			{
+				// Operation doesn't exist in KismetMathLibrary, do the operation as Vec3, Z as 0 and discard Z in result
+				FVector TempValue{ Value, 0.0 };
+				return FVector2D{ Transform.TransformPosition(TempValue) };
+			}
+			else if constexpr (std::is_same_v<InType, FVector4>)
+			{
+				// Operation doesn't exist in KismetMathLibrary, do the operation as Vec3 and set W to 0
+				FVector TempValue{ Value };
+				return FVector4{ Transform.TransformPosition(TempValue), 0.0 };
+			}
+			else
+			{
+				return Transform.TransformPosition(Value);
+			}
+			break;
+		}
+		case EPCGMedadataVectorOperation::InverseTransformDirection:
+		{
+			if constexpr (std::is_same_v<InType, FVector2D>)
+			{
+				// Operation doesn't exist in KismetMathLibrary, do the operation as Vec3, Z as 0 and discard Z in result
+				FVector TempValue{ Value, 0.0 };
+				return FVector2D{ Transform.InverseTransformVectorNoScale(TempValue) };
+			}
+			else if constexpr (std::is_same_v<InType, FVector4>)
+			{
+				// Operation doesn't exist in KismetMathLibrary, do the operation as Vec3 and set W to 0
+				FVector TempValue{ Value };
+				return FVector4{ Transform.InverseTransformVectorNoScale(TempValue), 0.0 };
+			}
+			else
+			{
+				return Transform.InverseTransformVectorNoScale(Value);
+			}
+			break;
+		}
+		case EPCGMedadataVectorOperation::InverseTransformLocation:
+		{
+			if constexpr (std::is_same_v<InType, FVector2D>)
+			{
+				// Operation doesn't exist in KismetMathLibrary, do the operation as Vec3, Z as 0 and discard Z in result
+				FVector TempValue{ Value, 0.0 };
+				return FVector2D{ Transform.InverseTransformPosition(TempValue) };
+			}
+			else if constexpr (std::is_same_v<InType, FVector4>)
+			{
+				// Operation doesn't exist in KismetMathLibrary, do the operation as Vec3 and set W to 0
+				FVector TempValue{ Value };
+				return FVector4{ Transform.InverseTransformPosition(TempValue), 0.0 };
+			}
+			else
+			{
+				return Transform.InverseTransformPosition(Value);
+			}
+			break;
+		}
+		default:
+			break;
+		}
+
+		return InType{};
+	}
+
+	template <typename InType>
+	inline InType ApplyVectorOperation(const InType& Value1, const InType& Value2, double& DoubleValue, EPCGMedadataVectorOperation Operation)
+	{
+		switch (Operation)
+		{
+		case EPCGMedadataVectorOperation::Cross:
+		{
+			if constexpr (std::is_same_v<InType, FVector2D>)
+			{
+				DoubleValue = Value1 ^ Value2;
+			}
+			else
+			{
+				return Value1 ^ Value2;
+			}
+			break;
+		}
+		case EPCGMedadataVectorOperation::Dot:
+		{
+			if constexpr (std::is_same_v<FVector4, InType>)
+			{
+				DoubleValue = Dot4(Value1, Value2);
+			}
+			else
+			{
+				DoubleValue = Value1.Dot(Value2);
+			}
+			break;
+		}
+		case EPCGMedadataVectorOperation::Distance:
+			DoubleValue = (Value1 - Value2).Size();
+			break;
+		case EPCGMedadataVectorOperation::RotateAroundAxis:
+		{
+			if constexpr (std::is_same_v<FVector4, InType>)
+			{
+				FVector TempValue1{ Value1 };
+				FVector TempValue2{ Value2 };
+
+				return FVector4{ TempValue1.RotateAngleAxis(DoubleValue, TempValue2), 0.0 };
+			}
+			else if constexpr (std::is_same_v<FVector2D, InType>)
+			{
+				// Ignore the axis if it is Vec2
+				return Value1.GetRotated(DoubleValue);
+			}
+			else
+			{
+				return Value1.RotateAngleAxis(DoubleValue, Value2);
+			}
+			break;
+		}
+		case EPCGMedadataVectorOperation::Normalize:
+			if constexpr (std::is_same_v<FVector4, InType>)
+			{
+				double Length = Value1.Size();
+				if (Length >= UE_SMALL_NUMBER)
+				{
+					return Value1 / Length;
+				}
+				else
+				{
+					return FVector4::Zero();
+				}
+			}
+			else
+			{
+				InType Res = Value1;
+				Res.Normalize();
+				return Res;
+			}
+			break;
+		case EPCGMedadataVectorOperation::Length:
+			DoubleValue = Value1.Size();
+			break;
+		default:
+			break;
+		}
+
+		return InType{};
+	};
+}
+
+FName UPCGMetadataVectorSettings::GetInputPinLabel(uint32 Index) const
+{
+	switch (Index)
+	{
+	case 0:
+	{
+		if (PCGMetadataVectorSettings::IsUnaryOp(Operation) || 
+			PCGMetadataVectorSettings::IsTransformOp(Operation) ||
+			Operation == EPCGMedadataVectorOperation::RotateAroundAxis)
+		{
+			return PCGPinConstants::DefaultInputLabel;
+		}
+		else
+		{
+			return PCGMetadataSettingsBaseConstants::DoubleInputFirstLabel;
+		}
+	}
+	case 1:
+	{
+		if (Operation == EPCGMedadataVectorOperation::RotateAroundAxis)
+		{
+			return PCGMetadataVectorConstants::AxisLabel;
+		}
+		else if (PCGMetadataVectorSettings::IsTransformOp(Operation))
+		{
+			return PCGMetadataSettingsBaseConstants::TransformLabel;
+		}
+		else
+		{
+			return PCGMetadataSettingsBaseConstants::DoubleInputSecondLabel;
+		}
+	}
+	case 2:
+		return PCGMetadataVectorConstants::AngleLabel;
+	default:
+		return NAME_None;
+	}
+}
+
+uint32 UPCGMetadataVectorSettings::GetInputPinNum() const
+{
+	if (PCGMetadataVectorSettings::IsUnaryOp(Operation))
+	{
+		return 1;
+	}
+	else if (PCGMetadataVectorSettings::IsTernaryOp(Operation))
+	{
+		return 3;
+	}
+	else
+	{
+		return 2;
+	}
+}
+
+bool UPCGMetadataVectorSettings::IsSupportedInputType(uint16 TypeId, uint32 InputIndex, bool& bHasSpecialRequirement) const
+{
+	bHasSpecialRequirement = false;
+
+	if (Operation == EPCGMedadataVectorOperation::RotateAroundAxis && InputIndex == 2)
+	{
+		bHasSpecialRequirement = true;
+		return (TypeId == (uint16)EPCGMetadataTypes::Double) || (TypeId == (uint16)EPCGMetadataTypes::Float);
+	}
+	else if (PCGMetadataVectorSettings::IsTransformOp(Operation) && InputIndex == 1)
+	{
+		bHasSpecialRequirement = true;
+		return TypeId == (uint16)EPCGMetadataTypes::Transform;
+	}
+	else
+	{
+		return (TypeId == (uint16)EPCGMetadataTypes::Vector2) || (TypeId == (uint16)EPCGMetadataTypes::Vector) || (TypeId == (uint16)EPCGMetadataTypes::Vector4);
+	}
+}
+
+FName UPCGMetadataVectorSettings::GetInputAttributeNameWithOverride(uint32 Index, UPCGParamData* Params) const
+{
+	switch (Index)
+	{
+	case 0:
+		return PCG_GET_OVERRIDEN_VALUE(this, Input1AttributeName, Params);
+	case 1:
+		return PCG_GET_OVERRIDEN_VALUE(this, Input2AttributeName, Params);
+	case 2:
+		return PCG_GET_OVERRIDEN_VALUE(this, Input3AttributeName, Params);
+	default:
+		return NAME_None;
+	}
+}
+
+#if WITH_EDITOR
+FName UPCGMetadataVectorSettings::GetDefaultNodeName() const
+{
+	if (const UEnum* EnumPtr = FindObject<UEnum>(nullptr, TEXT("EPCGMedadataVectorOperation"), true))
+	{
+		return EnumPtr->GetNameByValue(static_cast<int>(Operation));
+	}
+
+	return TEXT("Metadata Vector Node");
+}
+#endif // WITH_EDITOR
+
+FPCGElementPtr UPCGMetadataVectorSettings::CreateElement() const
+{
+	return MakeShared<FPCGMetadataVectorElement>();
+}
+
+uint16 UPCGMetadataVectorSettings::GetOutputType(uint16 InputTypeId) const
+{
+	// Dot, Length and Cross product with Vec2 output Double values
+	if (Operation == EPCGMedadataVectorOperation::Dot ||
+		Operation == EPCGMedadataVectorOperation::Length ||
+		Operation == EPCGMedadataVectorOperation::Distance ||
+		(Operation == EPCGMedadataVectorOperation::Cross && InputTypeId == (uint16)EPCGMetadataTypes::Vector2))
+	{
+		return (uint16)EPCGMetadataTypes::Double;
+	}
+	else
+	{
+		// Otherwise, output the input type
+		return InputTypeId;
+	}
+}
+
+bool FPCGMetadataVectorElement::DoOperation(FOperationData& OperationData) const
+{
+	TRACE_CPUPROFILER_EVENT_SCOPE(FPCGMetadataVectorElement::Execute);
+
+	const UPCGMetadataVectorSettings* Settings = CastChecked<UPCGMetadataVectorSettings>(OperationData.Settings);
+
+	auto VectorFunc = [this, Operation = Settings->Operation, &OperationData](auto DummyValue)
+	{
+		using AttributeType = decltype(DummyValue);
+
+		if constexpr (PCGMetadataVectorSettings::IsInvalidInType<AttributeType>())
+		{
+			return;
+		}
+		else
+		{
+			if (PCGMetadataVectorSettings::IsTransformOp(Operation))
+			{
+				DoBinaryOp<AttributeType, FTransform, AttributeType>(OperationData, [Operation](const AttributeType& Value, const FTransform& Transform) -> AttributeType { return PCGMetadataVectorSettings::ApplyTransformOperation(Value, Transform, Operation);});
+			}
+			else if (OperationData.OutputType == (uint16)EPCGMetadataTypes::Double)
+			{
+				if (PCGMetadataVectorSettings::IsUnaryOp(Operation))
+				{
+					DoUnaryOp<AttributeType, double>(OperationData, [Operation](const AttributeType& Value) -> double {
+						double Res = 0.0;
+						PCGMetadataVectorSettings::ApplyVectorOperation(Value, AttributeType{}, Res, Operation);
+						return Res;
+						});
+				}
+				else
+				{
+					DoBinaryOp<AttributeType, AttributeType, double>(OperationData, [Operation](const AttributeType& Value1, const AttributeType& Value2) -> double {
+						double Res = 0.0;
+						PCGMetadataVectorSettings::ApplyVectorOperation(Value1, Value2, Res, Operation);
+						return Res;
+						});
+				}
+			}
+			else
+			{
+				if (PCGMetadataVectorSettings::IsUnaryOp(Operation))
+				{
+					DoUnaryOp<AttributeType, AttributeType>(OperationData, [Operation](const AttributeType& Value) -> AttributeType {
+						double DummyDouble = 0.0;
+						return PCGMetadataVectorSettings::ApplyVectorOperation(Value, AttributeType{}, DummyDouble, Operation);
+						});
+				}
+				else if (PCGMetadataVectorSettings::IsTernaryOp(Operation))
+				{
+					DoTernaryOp<AttributeType, AttributeType, double, AttributeType>(OperationData, [Operation](const AttributeType& Value1, const AttributeType& Value2, const double& DoubleValue) -> AttributeType {
+						double DoubleCopy = DoubleValue;
+						return PCGMetadataVectorSettings::ApplyVectorOperation(Value1, Value2, DoubleCopy, Operation);
+						});
+				}
+				else
+				{
+					DoBinaryOp<AttributeType, AttributeType, AttributeType>(OperationData, [Operation](const AttributeType& Value1, const AttributeType& Value2) -> AttributeType {
+						double DummyDouble = 0.0;
+						return PCGMetadataVectorSettings::ApplyVectorOperation(Value1, Value2, DummyDouble, Operation);
+						});
+				}
+			}
+		}
+	};
+
+	PCGMetadataAttribute::CallbackWithRightType(OperationData.MostComplexInputType, VectorFunc);
+
+	return true;
+}
\ No newline at end of file
diff --git a/Engine/Plugins/Experimental/PCG/Source/PCG/Public/Elements/Metadata/PCGMetadataOpElementBase.h b/Engine/Plugins/Experimental/PCG/Source/PCG/Public/Elements/Metadata/PCGMetadataOpElementBase.h
--- a/Engine/Plugins/Experimental/PCG/Source/PCG/Public/Elements/Metadata/PCGMetadataOpElementBase.h
+++ b/Engine/Plugins/Experimental/PCG/Source/PCG/Public/Elements/Metadata/PCGMetadataOpElementBase.h
@@ -22,6 +22,7 @@ namespace PCGMetadataSettingsBaseConstants
 	const FName ClampMinLabel = TEXT("Min");
 	const FName ClampMaxLabel = TEXT("Max");
 	const FName LerpRatioLabel = TEXT("Ratio");
+	const FName TransformLabel = TEXT("Transform");
 }
 
 // Defines behavior when number of entries doesn't match in inputs
@@ -149,7 +150,7 @@ inline void FPCGMetadataElementBase::DoUnaryOp(FOperationData& OperationData, TF
 
 			OutType OutputValue = UnaryOp(Value);
 
-			OutputAttribute->SetValue(EntryKey, Value);
+			OutputAttribute->SetValue(EntryKey, OutputValue);
 		}
 
 		++Iterator1;
diff --git a/Engine/Plugins/Experimental/PCG/Source/PCG/Public/Elements/Metadata/PCGMetadataVectorOpElement.h b/Engine/Plugins/Experimental/PCG/Source/PCG/Public/Elements/Metadata/PCGMetadataVectorOpElement.h
new file mode 100644
--- /dev/null
+++ b/Engine/Plugins/Experimental/PCG/Source/PCG/Public/Elements/Metadata/PCGMetadataVectorOpElement.h
@@ -0,0 +1,74 @@
+// Copyright Epic Games, Inc. All Rights Reserved.
+
+#pragma once
+
+#include "PCGElement.h"
+#include "PCGSettings.h"
+#include "Elements/Metadata/PCGMetadataOpElementBase.h"
+
+#include "PCGMetadataVectorOpElement.generated.h"
+
+UENUM()
+enum class EPCGMedadataVectorOperation : uint16
+{
+	VectorOp = 0 UMETA(Hidden),
+	Cross,
+	Dot,
+	Distance,
+	Normalize,
+	Length,
+	RotateAroundAxis,
+
+	TransformOp = 100 UMETA(Hidden),
+	TransformDirection,
+	TransformLocation,
+	InverseTransformDirection,
+	InverseTransformLocation
+};
+
+UCLASS(BlueprintType, ClassGroup = (Procedural))
+class PCG_API UPCGMetadataVectorSettings : public UPCGMetadataSettingsBase
+{
+	GENERATED_BODY()
+
+public:
+	//~Begin UPCGSettings interface
+#if WITH_EDITOR
+	virtual FName GetDefaultNodeName() const override;
+#endif
+	//~End UPCGSettings interface
+
+	//~Begin UPCGMetadataSettingsBase interface
+	virtual FName GetInputAttributeNameWithOverride(uint32 Index, UPCGParamData* Params) const override;
+
+	virtual FName GetInputPinLabel(uint32 Index) const override;
+	virtual uint32 GetInputPinNum() const override;
+
+	virtual bool IsSupportedInputType(uint16 TypeId, uint32 InputIndex, bool& bHasSpecialRequirement) const override;
+	virtual uint16 GetOutputType(uint16 InputTypeId) const override;
+	//~End UPCGMetadataSettingsBase interface
+
+protected:
+	//~Begin UPCGSettings interface
+	virtual FPCGElementPtr CreateElement() const override;
+	//~End UPCGSettings interface
+
+public:
+	UPROPERTY(BlueprintReadWrite, EditAnywhere, Category = Settings)
+	EPCGMedadataVectorOperation Operation = EPCGMedadataVectorOperation::Cross;
+
+	UPROPERTY(BlueprintReadWrite, EditAnywhere, Category = Input)
+	FName Input1AttributeName = NAME_None;
+
+	UPROPERTY(BlueprintReadWrite, EditAnywhere, Category = Input, meta = (EditCondition = "Operation != EPCGMedadataVectorOperation::Normalize && Operation != EPCGMedadataVectorOperation::Length", EditConditionHides))
+	FName Input2AttributeName = NAME_None;
+
+	UPROPERTY(BlueprintReadWrite, EditAnywhere, Category = Input, meta = (EditCondition = "Operation == EPCGMedadataVectorOperation::RotateAroundAxis", EditConditionHides))
+	FName Input3AttributeName = NAME_None;
+};
+
+class FPCGMetadataVectorElement : public FPCGMetadataElementBase
+{
+protected:
+	virtual bool DoOperation(FOperationData& OperationData) const override;
+};
