[PCG] Metadata native operation: Compare operations

diff --git a/Engine/Plugins/Experimental/PCG/Source/PCG/Private/Elements/Metadata/PCGMetadataCompareOpElement.cpp b/Engine/Plugins/Experimental/PCG/Source/PCG/Private/Elements/Metadata/PCGMetadataCompareOpElement.cpp
new file mode 100644
--- /dev/null
+++ b/Engine/Plugins/Experimental/PCG/Source/PCG/Private/Elements/Metadata/PCGMetadataCompareOpElement.cpp
@@ -0,0 +1,141 @@
+// Copyright Epic Games, Inc. All Rights Reserved.
+
+#include "Elements/Metadata/PCGMetadataCompareOpElement.h"
+
+#include "Helpers/PCGSettingsHelpers.h"
+#include "Metadata/PCGMetadata.h"
+#include "Metadata/PCGMetadataAttribute.h"
+#include "Metadata/PCGMetadataAttributeTpl.h"
+#include "Metadata/PCGMetadataEntryKeyIterator.h"
+#include "Elements/Metadata/PCGMetadataElementCommon.h"
+
+namespace PCGMetadataCompareSettings
+{
+	template <typename T>
+	bool ApplyCompare(const T& Input1, const T& Input2, EPCGMedadataCompareOperation Operation, double Tolerance)
+	{
+		switch (Operation)
+		{
+		case EPCGMedadataCompareOperation::Equal:
+		{
+			if constexpr (std::is_same_v<T, int32> || std::is_same_v<T, int64>)
+			{
+				return Input1 == Input2;
+			}
+			else
+			{
+				return FMath::IsNearlyEqual(Input1, Input2, static_cast<T>(Tolerance));
+			}
+		}
+		case EPCGMedadataCompareOperation::NotEqual:
+		{
+			if constexpr (std::is_same_v<T, int32> || std::is_same_v<T, int64>)
+			{
+				return Input1 != Input2;
+			}
+			else
+			{
+				return !FMath::IsNearlyEqual(Input1, Input2, static_cast<T>(Tolerance));
+			}
+		}
+		case EPCGMedadataCompareOperation::Greater:
+			return Input1 > Input2;
+		case EPCGMedadataCompareOperation::GreaterOrEqual:
+			return Input1 >= Input2;
+		case EPCGMedadataCompareOperation::Less:
+			return Input1 < Input2;
+		case EPCGMedadataCompareOperation::LessOrEqual:
+			return Input1 <= Input2;
+		default:
+			return false;
+		}
+	}
+}
+
+FName UPCGMetadataCompareSettings::GetInputPinLabel(uint32 Index) const
+{
+	switch (Index)
+	{
+	case 0:
+		return PCGMetadataSettingsBaseConstants::DoubleInputFirstLabel;
+	case 1:
+		return PCGMetadataSettingsBaseConstants::DoubleInputSecondLabel;
+	default:
+		return NAME_None;
+	}
+}
+
+uint32 UPCGMetadataCompareSettings::GetInputPinNum() const
+{
+	return 2;
+}
+
+bool UPCGMetadataCompareSettings::IsSupportedInputType(uint16 TypeId, uint32 InputIndex, bool& bHasSpecialRequirement) const
+{
+	bHasSpecialRequirement = false;
+	return TypeId <= (uint16)EPCGMetadataTypes::Integer64;
+}
+
+FName UPCGMetadataCompareSettings::GetInputAttributeNameWithOverride(uint32 Index, UPCGParamData* Params) const
+{
+	switch (Index)
+	{
+	case 0:
+		return PCG_GET_OVERRIDEN_VALUE(this, Input1AttributeName, Params);
+	case 1:
+		return PCG_GET_OVERRIDEN_VALUE(this, Input2AttributeName, Params);
+	default:
+		return NAME_None;
+	}
+}
+
+uint16 UPCGMetadataCompareSettings::GetOutputType(uint16 InputTypeId) const
+{
+	return (uint16)EPCGMetadataTypes::Boolean;
+}
+
+#if WITH_EDITOR
+FName UPCGMetadataCompareSettings::GetDefaultNodeName() const
+{
+	if (const UEnum* EnumPtr = FindObject<UEnum>(nullptr, TEXT("EPCGMedadataCompareOperation"), true))
+	{ 
+		return EnumPtr->GetNameByValue(static_cast<int>(Operation)); 
+	}
+
+	return TEXT("Metadata Compare Node");
+}
+#endif // WITH_EDITOR
+
+FPCGElementPtr UPCGMetadataCompareSettings::CreateElement() const
+{
+	return MakeShared<FPCGMetadataCompareElement>();
+}
+
+bool FPCGMetadataCompareElement::DoOperation(FOperationData& OperationData) const
+{
+	TRACE_CPUPROFILER_EVENT_SCOPE(FPCGMetadataTrigElement::Execute);
+
+	const UPCGMetadataCompareSettings* Settings = CastChecked<UPCGMetadataCompareSettings>(OperationData.Settings);
+
+	auto CompareFunc = [this, Operation = Settings->Operation, Tolerance = Settings->Tolerance, &OperationData](auto DummyValue)
+	{
+		using AttributeType = decltype(DummyValue);
+
+		// Need to remove types that would not compile
+		if constexpr (PCG::Private::MetadataTypes<AttributeType>::Id > (uint16)EPCGMetadataTypes::Integer64)
+		{
+			return;
+		}
+		else
+		{
+			DoBinaryOp<AttributeType, AttributeType, bool>(OperationData, 
+				[Operation, Tolerance](const AttributeType& Value1, const AttributeType& Value2) -> bool { 
+					return PCGMetadataCompareSettings::ApplyCompare(Value1, Value2, Operation, Tolerance); 
+				});
+		}
+	};
+
+	PCGMetadataAttribute::CallbackWithRightType(OperationData.MostComplexInputType, CompareFunc);
+
+	return true;
+}
\ No newline at end of file
diff --git a/Engine/Plugins/Experimental/PCG/Source/PCG/Public/Elements/Metadata/PCGMetadataCompareOpElement.h b/Engine/Plugins/Experimental/PCG/Source/PCG/Public/Elements/Metadata/PCGMetadataCompareOpElement.h
new file mode 100644
--- /dev/null
+++ b/Engine/Plugins/Experimental/PCG/Source/PCG/Public/Elements/Metadata/PCGMetadataCompareOpElement.h
@@ -0,0 +1,67 @@
+// Copyright Epic Games, Inc. All Rights Reserved.
+
+#pragma once
+
+#include "PCGElement.h"
+#include "PCGSettings.h"
+#include "Elements/Metadata/PCGMetadataOpElementBase.h"
+
+#include "PCGMetadataCompareOpElement.generated.h"
+
+UENUM()
+enum class EPCGMedadataCompareOperation : uint16
+{
+	Equal,
+	NotEqual,
+	Greater,
+	GreaterOrEqual,
+	Less,
+	LessOrEqual
+};
+
+UCLASS(BlueprintType, ClassGroup = (Procedural))
+class PCG_API UPCGMetadataCompareSettings : public UPCGMetadataSettingsBase
+{
+	GENERATED_BODY()
+
+public:
+	//~Begin UPCGSettings interface
+#if WITH_EDITOR
+	virtual FName GetDefaultNodeName() const override;
+#endif
+	//~End UPCGSettings interface
+
+	//~Begin UPCGMetadataSettingsBase interface
+	virtual FName GetInputAttributeNameWithOverride(uint32 Index, UPCGParamData* Params) const override;
+
+	virtual FName GetInputPinLabel(uint32 Index) const override;
+	virtual uint32 GetInputPinNum() const override;
+
+	virtual bool IsSupportedInputType(uint16 TypeId, uint32 InputIndex, bool& bHasSpecialRequirement) const override;
+	virtual uint16 GetOutputType(uint16 InputTypeId) const override;
+	//~End UPCGMetadataSettingsBase interface
+
+protected:
+	//~Begin UPCGSettings interface
+	virtual FPCGElementPtr CreateElement() const override;
+	//~End UPCGSettings interface
+
+public:
+	UPROPERTY(BlueprintReadWrite, EditAnywhere, Category = Settings)
+	EPCGMedadataCompareOperation Operation = EPCGMedadataCompareOperation::Equal;
+
+	UPROPERTY(BlueprintReadWrite, EditAnywhere, Category = Input)
+	FName Input1AttributeName = NAME_None;
+
+	UPROPERTY(BlueprintReadWrite, EditAnywhere, Category = Input)
+	FName Input2AttributeName = NAME_None;
+
+	UPROPERTY(BlueprintReadWrite, EditAnywhere, Category = Settings, meta = (EditCondition = "Operation == EPCGMedadataCompareOperation::Equal || Operation == EPCGMedadataCompareOperation::NotEqual", EditConditionHides))
+	double Tolerance = UE_DOUBLE_SMALL_NUMBER;
+};
+
+class FPCGMetadataCompareElement : public FPCGMetadataElementBase
+{
+protected:
+	virtual bool DoOperation(FOperationData& OperationData) const override;
+};
