Various Fixes in the PCG graph executor:
-  Improved sleeping tasks management
-  Fixed thread count allocation issues

Made the target actor a weak pointer to prevent GC issues
Removed object outering when we can run from another thread

diff --git a/Engine/Plugins/Experimental/PCG/Source/PCG/Private/Data/PCGDifferenceData.cpp b/Engine/Plugins/Experimental/PCG/Source/PCG/Private/Data/PCGDifferenceData.cpp
--- a/Engine/Plugins/Experimental/PCG/Source/PCG/Private/Data/PCGDifferenceData.cpp
+++ b/Engine/Plugins/Experimental/PCG/Source/PCG/Private/Data/PCGDifferenceData.cpp
@@ -163,7 +163,7 @@ const UPCGPointData* UPCGDifferenceData::CreatePointData(FPCGContext* Context) c
 		return SourcePointData;
 	}
 
-	UPCGPointData* Data = NewObject<UPCGPointData>(const_cast<UPCGDifferenceData*>(this));
+	UPCGPointData* Data = NewObject<UPCGPointData>();
 	Data->InitializeFromData(this, SourcePointData->Metadata);
 
 	const TArray<FPCGPoint>& SourcePoints = SourcePointData->GetPoints();
diff --git a/Engine/Plugins/Experimental/PCG/Source/PCG/Private/Data/PCGIntersectionData.cpp b/Engine/Plugins/Experimental/PCG/Source/PCG/Private/Data/PCGIntersectionData.cpp
--- a/Engine/Plugins/Experimental/PCG/Source/PCG/Private/Data/PCGIntersectionData.cpp
+++ b/Engine/Plugins/Experimental/PCG/Source/PCG/Private/Data/PCGIntersectionData.cpp
@@ -144,7 +144,7 @@ UPCGPointData* UPCGIntersectionData::CreateAndFilterPointData(FPCGContext* Conte
 
 	const TArray<FPCGPoint>& SourcePoints = SourcePointData->GetPoints();
 
-	UPCGPointData* Data = NewObject<UPCGPointData>(const_cast<UPCGIntersectionData*>(this));
+	UPCGPointData* Data = NewObject<UPCGPointData>();
 	Data->InitializeFromData(this, SourcePointData->Metadata);
 	Data->Metadata->AddAttributes(Y->Metadata);
 
diff --git a/Engine/Plugins/Experimental/PCG/Source/PCG/Private/Data/PCGLandscapeData.cpp b/Engine/Plugins/Experimental/PCG/Source/PCG/Private/Data/PCGLandscapeData.cpp
--- a/Engine/Plugins/Experimental/PCG/Source/PCG/Private/Data/PCGLandscapeData.cpp
+++ b/Engine/Plugins/Experimental/PCG/Source/PCG/Private/Data/PCGLandscapeData.cpp
@@ -111,7 +111,7 @@ const UPCGPointData* UPCGLandscapeData::CreatePointData(FPCGContext* Context, co
 
 	const int32 ComponentSizeQuads = LandscapeInfo->ComponentSizeQuads;
 
-	UPCGPointData* Data = NewObject<UPCGPointData>(const_cast<UPCGLandscapeData*>(this));
+	UPCGPointData* Data = NewObject<UPCGPointData>();
 	Data->InitializeFromData(this);
 	TArray<FPCGPoint>& Points = Data->GetMutablePoints();
 
diff --git a/Engine/Plugins/Experimental/PCG/Source/PCG/Private/Data/PCGLandscapeSplineData.cpp b/Engine/Plugins/Experimental/PCG/Source/PCG/Private/Data/PCGLandscapeSplineData.cpp
--- a/Engine/Plugins/Experimental/PCG/Source/PCG/Private/Data/PCGLandscapeSplineData.cpp
+++ b/Engine/Plugins/Experimental/PCG/Source/PCG/Private/Data/PCGLandscapeSplineData.cpp
@@ -126,7 +126,7 @@ const UPCGPointData* UPCGLandscapeSplineData::CreatePointData(FPCGContext* Conte
 	check(Spline.IsValid());
 	TRACE_CPUPROFILER_EVENT_SCOPE(UPCGLandscapeSplineData::CreatePointData);
 
-	UPCGPointData* Data = NewObject<UPCGPointData>(const_cast<UPCGLandscapeSplineData*>(this));
+	UPCGPointData* Data = NewObject<UPCGPointData>();
 	Data->InitializeFromData(this);
 	TArray<FPCGPoint>& Points = Data->GetMutablePoints();
 
diff --git a/Engine/Plugins/Experimental/PCG/Source/PCG/Private/Data/PCGProjectionData.cpp b/Engine/Plugins/Experimental/PCG/Source/PCG/Private/Data/PCGProjectionData.cpp
--- a/Engine/Plugins/Experimental/PCG/Source/PCG/Private/Data/PCGProjectionData.cpp
+++ b/Engine/Plugins/Experimental/PCG/Source/PCG/Private/Data/PCGProjectionData.cpp
@@ -114,7 +114,7 @@ const UPCGPointData* UPCGProjectionData::CreatePointData(FPCGContext* Context) c
 	const UPCGPointData* SourcePointData = Source->ToPointData(Context);
 	const TArray<FPCGPoint>& SourcePoints = SourcePointData->GetPoints();
 
-	UPCGPointData* PointData = NewObject<UPCGPointData>(const_cast<UPCGProjectionData*>(this));
+	UPCGPointData* PointData = NewObject<UPCGPointData>();
 	PointData->InitializeFromData(this, SourcePointData->Metadata);
 	check(PointData->Metadata);
 	PointData->Metadata->AddAttributes(Target->Metadata);
diff --git a/Engine/Plugins/Experimental/PCG/Source/PCG/Private/Data/PCGSpatialData.cpp b/Engine/Plugins/Experimental/PCG/Source/PCG/Private/Data/PCGSpatialData.cpp
--- a/Engine/Plugins/Experimental/PCG/Source/PCG/Private/Data/PCGSpatialData.cpp
+++ b/Engine/Plugins/Experimental/PCG/Source/PCG/Private/Data/PCGSpatialData.cpp
@@ -92,7 +92,7 @@ UPCGIntersectionData* UPCGSpatialData::IntersectWith(const UPCGSpatialData* InOt
 
 UPCGProjectionData* UPCGSpatialData::ProjectOn(const UPCGSpatialData* InOther) const
 {
-	UPCGProjectionData* ProjectionData = NewObject<UPCGProjectionData>(const_cast<UPCGSpatialData*>(this));
+	UPCGProjectionData* ProjectionData = NewObject<UPCGProjectionData>();
 	ProjectionData->Initialize(this, InOther);
 
 	return ProjectionData;
@@ -128,7 +128,7 @@ UPCGMetadata* UPCGSpatialData::CreateEmptyMetadata()
 
 void UPCGSpatialData::InitializeFromData(const UPCGSpatialData* InSource, const UPCGMetadata* InMetadataParentOverride, bool bInheritMetadata)
 {
-	if (InSource && !TargetActor)
+	if (InSource && TargetActor.IsExplicitlyNull())
 	{
 		TargetActor = InSource->TargetActor;
 	}
diff --git a/Engine/Plugins/Experimental/PCG/Source/PCG/Private/Data/PCGSplineData.cpp b/Engine/Plugins/Experimental/PCG/Source/PCG/Private/Data/PCGSplineData.cpp
--- a/Engine/Plugins/Experimental/PCG/Source/PCG/Private/Data/PCGSplineData.cpp
+++ b/Engine/Plugins/Experimental/PCG/Source/PCG/Private/Data/PCGSplineData.cpp
@@ -57,7 +57,7 @@ FTransform UPCGSplineData::GetTransformAtDistance(int SegmentIndex, FVector::FRe
 const UPCGPointData* UPCGSplineData::CreatePointData(FPCGContext* Context) const
 {
 	TRACE_CPUPROFILER_EVENT_SCOPE(UPCGSplineData::CreatePointData);
-	UPCGPointData* Data = NewObject<UPCGPointData>(const_cast<UPCGSplineData*>(this));
+	UPCGPointData* Data = NewObject<UPCGPointData>();
 	Data->InitializeFromData(this);
 
 	FPCGSplineSamplerParams SamplerParams;
diff --git a/Engine/Plugins/Experimental/PCG/Source/PCG/Private/Data/PCGTextureData.cpp b/Engine/Plugins/Experimental/PCG/Source/PCG/Private/Data/PCGTextureData.cpp
--- a/Engine/Plugins/Experimental/PCG/Source/PCG/Private/Data/PCGTextureData.cpp
+++ b/Engine/Plugins/Experimental/PCG/Source/PCG/Private/Data/PCGTextureData.cpp
@@ -168,7 +168,7 @@ const UPCGPointData* UPCGBaseTextureData::CreatePointData(FPCGContext* Context)
 	// or based on a given texel size
 	FBox2D LocalSurfaceBounds(FVector2D(-1.0f, -1.0f), FVector2D(1.0f, 1.0f));
 
-	UPCGPointData* Data = NewObject<UPCGPointData>(const_cast<UPCGBaseTextureData*>(this));
+	UPCGPointData* Data = NewObject<UPCGPointData>();
 	Data->InitializeFromData(this);
 
 	// Early out for invalid data
diff --git a/Engine/Plugins/Experimental/PCG/Source/PCG/Private/Data/PCGUnionData.cpp b/Engine/Plugins/Experimental/PCG/Source/PCG/Private/Data/PCGUnionData.cpp
--- a/Engine/Plugins/Experimental/PCG/Source/PCG/Private/Data/PCGUnionData.cpp
+++ b/Engine/Plugins/Experimental/PCG/Source/PCG/Private/Data/PCGUnionData.cpp
@@ -176,7 +176,7 @@ const UPCGPointData* UPCGUnionData::CreatePointData(FPCGContext* Context) const
 		return Data[0]->ToPointData(Context);
 	}
 
-	UPCGPointData* PointData = NewObject<UPCGPointData>(const_cast<UPCGUnionData*>(this));
+	UPCGPointData* PointData = NewObject<UPCGPointData>();
 	PointData->InitializeFromData(this, Data[0]->Metadata);
 
 	// Initialize metadata
diff --git a/Engine/Plugins/Experimental/PCG/Source/PCG/Private/Elements/PCGCopyPoints.cpp b/Engine/Plugins/Experimental/PCG/Source/PCG/Private/Elements/PCGCopyPoints.cpp
--- a/Engine/Plugins/Experimental/PCG/Source/PCG/Private/Elements/PCGCopyPoints.cpp
+++ b/Engine/Plugins/Experimental/PCG/Source/PCG/Private/Elements/PCGCopyPoints.cpp
@@ -101,7 +101,7 @@ bool FPCGCopyPointsElement::ExecuteInternal(FPCGContext* Context) const
 	}
 
 	// Priorize use the target actor from the target, irrespective of the source
-	OutPointData->TargetActor = TargetPointData->TargetActor ? TargetPointData->TargetActor : SourcePointData->TargetActor;
+	OutPointData->TargetActor = TargetPointData->TargetActor.IsValid() ? TargetPointData->TargetActor : SourcePointData->TargetActor;
 
 	check(OutPointData->Metadata && NonRootMetadata);
 
diff --git a/Engine/Plugins/Experimental/PCG/Source/PCG/Private/Elements/PCGDebugElement.cpp b/Engine/Plugins/Experimental/PCG/Source/PCG/Private/Elements/PCGDebugElement.cpp
--- a/Engine/Plugins/Experimental/PCG/Source/PCG/Private/Elements/PCGDebugElement.cpp
+++ b/Engine/Plugins/Experimental/PCG/Source/PCG/Private/Elements/PCGDebugElement.cpp
@@ -83,7 +83,7 @@ namespace PCGDebugElement
 				continue;
 			}
 
-			AActor* TargetActor = SpatialData->TargetActor;
+			AActor* TargetActor = SpatialData->TargetActor.Get();
 
 			if (!TargetActor)
 			{
diff --git a/Engine/Plugins/Experimental/PCG/Source/PCG/Private/Elements/PCGExecuteBlueprint.cpp b/Engine/Plugins/Experimental/PCG/Source/PCG/Private/Elements/PCGExecuteBlueprint.cpp
--- a/Engine/Plugins/Experimental/PCG/Source/PCG/Private/Elements/PCGExecuteBlueprint.cpp
+++ b/Engine/Plugins/Experimental/PCG/Source/PCG/Private/Elements/PCGExecuteBlueprint.cpp
@@ -515,7 +515,7 @@ void UPCGBlueprintElement::LoopOnPoints(FPCGContext& InContext, const UPCGPointD
 	}
 	else
 	{
-		OutData = NewObject<UPCGPointData>(const_cast<UPCGPointData*>(InData));
+		OutData = NewObject<UPCGPointData>();
 		OutData->InitializeFromData(InData);
 	}
 
@@ -542,7 +542,7 @@ void UPCGBlueprintElement::MultiLoopOnPoints(FPCGContext& InContext, const UPCGP
 	}
 	else
 	{
-		OutData = NewObject<UPCGPointData>(const_cast<UPCGPointData*>(InData));
+		OutData = NewObject<UPCGPointData>();
 		OutData->InitializeFromData(InData);
 	}
 
@@ -569,7 +569,7 @@ void UPCGBlueprintElement::LoopOnPointPairs(FPCGContext& InContext, const UPCGPo
 	}
 	else
 	{
-		OutData = NewObject<UPCGPointData>(const_cast<UPCGPointData*>(InA));
+		OutData = NewObject<UPCGPointData>();
 		OutData->InitializeFromData(InA);
 		OutData->Metadata->AddAttributes(InB->Metadata);
 	}
@@ -599,7 +599,7 @@ void UPCGBlueprintElement::LoopNTimes(FPCGContext& InContext, int64 NumIteration
 	else
 	{
 		const UPCGSpatialData* Owner = (InA ? InA : InB);
-		OutData = Owner ? NewObject<UPCGPointData>(const_cast<UPCGSpatialData*>(Owner)) : NewObject<UPCGPointData>();
+		OutData = NewObject<UPCGPointData>();
 
 		if (Owner)
 		{
diff --git a/Engine/Plugins/Experimental/PCG/Source/PCG/Private/Elements/PCGSpawnActor.cpp b/Engine/Plugins/Experimental/PCG/Source/PCG/Private/Elements/PCGSpawnActor.cpp
--- a/Engine/Plugins/Experimental/PCG/Source/PCG/Private/Elements/PCGSpawnActor.cpp
+++ b/Engine/Plugins/Experimental/PCG/Source/PCG/Private/Elements/PCGSpawnActor.cpp
@@ -100,7 +100,7 @@ bool FPCGSpawnActorElement::ExecuteInternal(FPCGContext* Context) const
 			continue;
 		}
 
-		AActor* TargetActor = SpatialData->TargetActor;
+		AActor* TargetActor = SpatialData->TargetActor.Get();
 
 		if (!TargetActor)
 		{
diff --git a/Engine/Plugins/Experimental/PCG/Source/PCG/Private/Elements/PCGSplineSampler.cpp b/Engine/Plugins/Experimental/PCG/Source/PCG/Private/Elements/PCGSplineSampler.cpp
--- a/Engine/Plugins/Experimental/PCG/Source/PCG/Private/Elements/PCGSplineSampler.cpp
+++ b/Engine/Plugins/Experimental/PCG/Source/PCG/Private/Elements/PCGSplineSampler.cpp
@@ -935,7 +935,7 @@ bool FPCGSplineSamplerElement::ExecuteInternal(FPCGContext* Context) const
 		FPCGTaggedData& Output = Outputs.Emplace_GetRef();
 		Output = Input;
 
-		UPCGPointData* SampledPointData = NewObject<UPCGPointData>(const_cast<UPCGPolyLineData*>(LineData));
+		UPCGPointData* SampledPointData = NewObject<UPCGPointData>();
 		SampledPointData->InitializeFromData(SpatialData);
 		Output.Data = SampledPointData;
 
diff --git a/Engine/Plugins/Experimental/PCG/Source/PCG/Private/Elements/PCGStaticMeshSpawner.cpp b/Engine/Plugins/Experimental/PCG/Source/PCG/Private/Elements/PCGStaticMeshSpawner.cpp
--- a/Engine/Plugins/Experimental/PCG/Source/PCG/Private/Elements/PCGStaticMeshSpawner.cpp
+++ b/Engine/Plugins/Experimental/PCG/Source/PCG/Private/Elements/PCGStaticMeshSpawner.cpp
@@ -63,7 +63,7 @@ bool FPCGStaticMeshSpawnerElement::ExecuteInternal(FPCGContext* Context) const
 			continue;
 		}
 
-		AActor* TargetActor = SpatialData->TargetActor;
+		AActor* TargetActor = SpatialData->TargetActor.Get();
 
 		if (!TargetActor)
 		{
diff --git a/Engine/Plugins/Experimental/PCG/Source/PCG/Private/Elements/PCGSurfaceSampler.cpp b/Engine/Plugins/Experimental/PCG/Source/PCG/Private/Elements/PCGSurfaceSampler.cpp
--- a/Engine/Plugins/Experimental/PCG/Source/PCG/Private/Elements/PCGSurfaceSampler.cpp
+++ b/Engine/Plugins/Experimental/PCG/Source/PCG/Private/Elements/PCGSurfaceSampler.cpp
@@ -149,7 +149,7 @@ bool FPCGSurfaceSamplerElement::ExecuteInternal(FPCGContext* Context) const
 	{
 		const UPCGSpatialData* SpatialInput = Cast<UPCGSpatialData>(Input.Data);
 
-		if (!SpatialInput || !SpatialInput->TargetActor)
+		if (!SpatialInput || !SpatialInput->TargetActor.IsValid())
 		{
 			PCGE_LOG(Error, "Invalid input data");
 			continue;
diff --git a/Engine/Plugins/Experimental/PCG/Source/PCG/Private/Elements/PCGVolumeSampler.cpp b/Engine/Plugins/Experimental/PCG/Source/PCG/Private/Elements/PCGVolumeSampler.cpp
--- a/Engine/Plugins/Experimental/PCG/Source/PCG/Private/Elements/PCGVolumeSampler.cpp
+++ b/Engine/Plugins/Experimental/PCG/Source/PCG/Private/Elements/PCGVolumeSampler.cpp
@@ -11,7 +11,7 @@ namespace PCGVolumeSampler
 {
 	UPCGPointData* SampleVolume(FPCGContext* Context, const UPCGSpatialData* SpatialData, const FVolumeSamplerSettings& SamplerSettings)
 	{
-		UPCGPointData* Data = NewObject<UPCGPointData>(const_cast<UPCGSpatialData*>(SpatialData));
+		UPCGPointData* Data = NewObject<UPCGPointData>();
 		Data->InitializeFromData(SpatialData);
 		TArray<FPCGPoint>& Points = Data->GetMutablePoints();
 
diff --git a/Engine/Plugins/Experimental/PCG/Source/PCG/Private/Graph/PCGGraphExecutor.cpp b/Engine/Plugins/Experimental/PCG/Source/PCG/Private/Graph/PCGGraphExecutor.cpp
--- a/Engine/Plugins/Experimental/PCG/Source/PCG/Private/Graph/PCGGraphExecutor.cpp
+++ b/Engine/Plugins/Experimental/PCG/Source/PCG/Private/Graph/PCGGraphExecutor.cpp
@@ -198,7 +198,7 @@ void FPCGGraphExecutor::Execute()
 	// TODO: add optimization phase if we've added new graph(s)/tasks to be executed
 
 	// This is a safeguard to check if we're in a stuck state
-	if (ReadyTasks.Num() == 0 && ActiveTasks.Num() == 0 && Tasks.Num() > 0)
+	if (ReadyTasks.Num() == 0 && ActiveTasks.Num() == 0 && SleepingTasks.Num() == 0 && Tasks.Num() > 0)
 	{
 		UE_LOG(LogPCG, Error, TEXT("PCG Graph executor error: tasks are in a deadlocked state. Will drop all tasks."));
 		Tasks.Reset();
@@ -215,23 +215,84 @@ void FPCGGraphExecutor::Execute()
 	UpdateGenerationNotification(Tasks.Num() + ReadyTasks.Num() + ActiveTasks.Num());
 #endif
 
-	while(ReadyTasks.Num() > 0 || ActiveTasks.Num() > 0)
+	while(ReadyTasks.Num() > 0 || ActiveTasks.Num() > 0 || SleepingTasks.Num() > 0)
 	{
 		// First: if we have free resources, move ready tasks to the active tasks
 		bool bMainThreadAvailable = (ActiveTasks.Num() == 0 || !ActiveTasks.Last().Element->CanExecuteOnlyOnMainThread(ActiveTasks.Last().Context->GetInputSettings<UPCGSettings>()));
 		int32 NumAvailableThreads = FMath::Max(0, MaxNumThreads - CurrentlyUsedThreads);
 
 		const bool bMainThreadWasAvailable = bMainThreadAvailable;
-		const int32 TasksToLaunchIndex = ActiveTasks.Num();
+		const int32 TasksAlreadyLaunchedCount = ActiveTasks.Num();
 
 		if (bMainThreadAvailable || NumAvailableThreads > 0)
 		{
 			TRACE_CPUPROFILER_EVENT_SCOPE(FPCGGraphExecutor::Execute::PrepareTasks);
 			// Sort tasks by priority (highest priority should be at the end)
 			// TODO
+			bool bHasDispatchedTasks = false;
+
+			auto CannotDispatchMoreTasks = [&bMainThreadAvailable, &NumAvailableThreads, &bHasDispatchedTasks, bAllowMultiDispatch]
+			{
+				return ((!bAllowMultiDispatch && bHasDispatchedTasks) || (!bMainThreadAvailable && NumAvailableThreads == 0));
+			};
+
+			auto TaskDispatchBookKeeping = [&bMainThreadAvailable, &NumAvailableThreads, &bHasDispatchedTasks](bool bIsMainThreadTask)
+			{
+				bHasDispatchedTasks = true;
+
+				if (bIsMainThreadTask || NumAvailableThreads == 0)
+				{
+					bMainThreadAvailable = false;
+				}
+				else
+				{
+					--NumAvailableThreads;
+				}
+			};
+
+			// First, wake up any sleeping tasks that can be reactivated
+			for (int32 SleepingTaskIndex = SleepingTasks.Num() - 1; SleepingTaskIndex >= 0; --SleepingTaskIndex)
+			{
+				if (CannotDispatchMoreTasks())
+				{
+					break;
+				}
+
+				FPCGGraphActiveTask& SleepingTask = SleepingTasks[SleepingTaskIndex];
+
+				if (SleepingTask.Context->bIsPaused)
+				{
+					continue; // still sleeping
+				}
+
+				// Validate that we can start this task now
+				const bool bIsMainThreadTask = SleepingTask.Element->CanExecuteOnlyOnMainThread(SleepingTask.Context->GetInputSettings<UPCGSettings>());
+
+				if (!bIsMainThreadTask || bMainThreadAvailable)
+				{
+					const bool bInsertLast = (bIsMainThreadTask || bMainThreadAvailable || ActiveTasks.Num() == 0);
+					if (bInsertLast)
+					{
+						ActiveTasks.Emplace(MoveTemp(SleepingTask));
+					}
+					else
+					{
+						ActiveTasks.Insert(MoveTemp(SleepingTask), ActiveTasks.Num() - 1);
+					}
+
+					TaskDispatchBookKeeping(bIsMainThreadTask);
+
+					SleepingTasks.RemoveAtSwap(SleepingTaskIndex);
+				}
+			}
 
 			for(int32 ReadyTaskIndex = ReadyTasks.Num() - 1; ReadyTaskIndex >= 0; --ReadyTaskIndex)
 			{
+				if (CannotDispatchMoreTasks())
+				{
+					break;
+				}
+
 				FPCGGraphTask& Task = ReadyTasks[ReadyTaskIndex];
 
 				// Build input
@@ -312,60 +373,56 @@ void FPCGGraphExecutor::Execute()
 					}
 #endif
 
-					if (bIsMainThreadTask || NumAvailableThreads == 0)
-					{
-						bMainThreadAvailable = false;
-					}
-					else
-					{
-						--NumAvailableThreads;
-					}
+					TaskDispatchBookKeeping(bIsMainThreadTask);
 
 					ReadyTasks.RemoveAtSwap(ReadyTaskIndex);
-
-					if (!bAllowMultiDispatch || (!bMainThreadAvailable && NumAvailableThreads == 0))
-					{
-						break;
-					}
 				}
 			}
 		}
 
 		check(NumAvailableThreads >= 0);
 
-		const int32 NumTasksToLaunch = ActiveTasks.Num() - TasksToLaunchIndex;
-
-		// Re-launch time-sliced tasks & launch new tasks
-		// TODO: currently we don't have any time-slicing so just launch tasks
-		const int32 LastLaunchIndex = ActiveTasks.Num() - (bMainThreadWasAvailable ? 0 : 1);
+		const int32 NumTasksToLaunch = ActiveTasks.Num() - TasksAlreadyLaunchedCount;
 
 		// Assign resources
-		const int32 NumAdditionalThreads = ((NumTasksToLaunch > 0) ? (NumAvailableThreads / NumTasksToLaunch) : 0);
-		check(NumAdditionalThreads >= 0);
-
-		for (int32 ExecutionIndex = TasksToLaunchIndex; ExecutionIndex < LastLaunchIndex; ++ExecutionIndex)
+		if (NumTasksToLaunch > 0)
 		{
-			FPCGGraphActiveTask& ActiveTask = ActiveTasks[ExecutionIndex];
-			ActiveTask.Context->NumAvailableTasks = 1 + NumAdditionalThreads;
-			CurrentlyUsedThreads += ActiveTask.Context->NumAvailableTasks;
+			TRACE_CPUPROFILER_EVENT_SCOPE(FPCGGraphExecutor::Execute::AssignResources);
+			const int32 NumAdditionalThreads = ((NumTasksToLaunch > 0) ? (NumAvailableThreads / NumTasksToLaunch) : 0);
+			check(NumAdditionalThreads >= 0);
+
+			for (int32 ExecutionIndex = 0; ExecutionIndex < ActiveTasks.Num(); ++ExecutionIndex)
+			{
+				FPCGGraphActiveTask& ActiveTask = ActiveTasks[ExecutionIndex];
+
+				// Tasks that were already launched already have assigned tasks, so don't touch them
+				if (ActiveTask.Context->NumAvailableTasks == 0)
+				{
+					ActiveTask.Context->NumAvailableTasks = 1 + NumAdditionalThreads;
+					CurrentlyUsedThreads += ActiveTask.Context->NumAvailableTasks;
+				}
+			}
 		}
 
 		// Dispatch async tasks
 		TMap<int32, TFuture<bool>> Futures;
 
-		for(int32 ExecutionIndex = 0; ExecutionIndex < ActiveTasks.Num() - 1; ++ExecutionIndex)
 		{
-			FPCGGraphActiveTask& ActiveTask = ActiveTasks[ExecutionIndex];
-#if WITH_EDITOR
-			if(!ActiveTask.bIsBypassed && !ActiveTask.Context->bIsPaused)
-#else
-			if(!ActiveTask.Context->bIsPaused)
-#endif
+			TRACE_CPUPROFILER_EVENT_SCOPE(FPCGGraphExecutor::Execute::StartFutures);
+			for(int32 ExecutionIndex = 0; ExecutionIndex < ActiveTasks.Num() - 1; ++ExecutionIndex)
 			{
-				Futures.Emplace(ExecutionIndex, Async(EAsyncExecution::ThreadPool, [&ActiveTask]()
+				FPCGGraphActiveTask& ActiveTask = ActiveTasks[ExecutionIndex];
+				check(!ActiveTask.Context || !ActiveTask.Context->bIsPaused);
+
+	#if WITH_EDITOR
+				if(!ActiveTask.bIsBypassed)
+	#endif
 				{
-					return ActiveTask.Element->Execute(ActiveTask.Context.Get());
-				}));
+					Futures.Emplace(ExecutionIndex, Async(EAsyncExecution::ThreadPool, [&ActiveTask]()
+					{
+						return ActiveTask.Element->Execute(ActiveTask.Context.Get());
+					}));
+				}
 			}
 		}
 
@@ -385,6 +442,7 @@ void FPCGGraphExecutor::Execute()
 				}
 			}
 
+			check(ActiveTask.Context->NumAvailableTasks >= 0);
 			CurrentlyUsedThreads -= ActiveTask.Context->NumAvailableTasks;
 
 #if WITH_EDITOR
@@ -416,11 +474,13 @@ void FPCGGraphExecutor::Execute()
 			// Execute main thread task
 			if (!ActiveTasks.IsEmpty())
 			{
+				TRACE_CPUPROFILER_EVENT_SCOPE(FPCGGraphExecutor::Execute::ExecuteTasks::MainThreadTask);
 				FPCGGraphActiveTask& MainThreadTask = ActiveTasks.Last();
+				check(!MainThreadTask.Context || !MainThreadTask.Context->bIsPaused);
 #if WITH_EDITOR
-				if (MainThreadTask.bIsBypassed || (!MainThreadTask.Context->bIsPaused && MainThreadTask.Element->Execute(MainThreadTask.Context.Get())))
+				if(MainThreadTask.bIsBypassed || MainThreadTask.Element->Execute(MainThreadTask.Context.Get()))
 #else
-				if (!MainThreadTask.Context->bIsPaused && MainThreadTask.Element->Execute(MainThreadTask.Context.Get()))
+				if(MainThreadTask.Element->Execute(MainThreadTask.Context.Get()))
 #endif
 				{
 					bMainTaskDone = true;
@@ -431,21 +491,44 @@ void FPCGGraphExecutor::Execute()
 			// If the main thread task is done, it was removed from ActiveTasks, and ActiveTasks should contains only futures (if there are some)
 			// Otherwise, ActiveTasks.Num() - 1 still point to the main thread tasks, so futures indexes are within [0, ActiveTasks.Num() - 2].
 			int32 FuturesLastIndex = bMainTaskDone ? ActiveTasks.Num() - 1 : ActiveTasks.Num() - 2;
-
-			// Then wait after all futures
-			for (int32 ExecutionIndex = FuturesLastIndex; ExecutionIndex >= 0; --ExecutionIndex)
+			if(FuturesLastIndex >= 0)
 			{
-				bool bTaskDone = false;
-				// Wait on the future if any
-				if (TFuture<bool>* Future = Futures.Find(ExecutionIndex))
+				TRACE_CPUPROFILER_EVENT_SCOPE(FPCGGraphExecutor::Execute::ExecuteTasks::WaitForFutures);
+
+				// Then wait after all futures
+				for (int32 ExecutionIndex = FuturesLastIndex; ExecutionIndex >= 0; --ExecutionIndex)
 				{
-					Future->Wait();
-					bTaskDone = Future->Get();
+					bool bTaskDone = false;
+					// Wait on the future if any
+					if (TFuture<bool>* Future = Futures.Find(ExecutionIndex))
+					{
+						Future->Wait();
+						bTaskDone = Future->Get();
+					}
+
+					if (bTaskDone)
+					{
+						PostTaskExecute(ExecutionIndex);
+					}
 				}
+			}
+		}
 
-				if (bTaskDone)
+		// Any paused tasks at that point should relinquish their resources
+		if (ActiveTasks.Num() > 0)
+		{
+			TRACE_CPUPROFILER_EVENT_SCOPE(FPCGGraphExecutor::Execute::ExecuteTasks::CheckSleepingTasks);
+			for (int32 ActiveTaskIndex = ActiveTasks.Num() - 1; ActiveTaskIndex >= 0; --ActiveTaskIndex)
+			{
+				if (ActiveTasks[ActiveTaskIndex].Context->bIsPaused)
 				{
-					PostTaskExecute(ExecutionIndex);
+					FPCGGraphActiveTask& ActiveTask = ActiveTasks[ActiveTaskIndex];
+					check(ActiveTask.Context->NumAvailableTasks > 0);
+					CurrentlyUsedThreads -= ActiveTask.Context->NumAvailableTasks;
+					ActiveTask.Context->NumAvailableTasks = 0;
+
+					SleepingTasks.Emplace(MoveTemp(ActiveTask));
+					ActiveTasks.RemoveAtSwap(ActiveTaskIndex);
 				}
 			}
 		}
diff --git a/Engine/Plugins/Experimental/PCG/Source/PCG/Private/Graph/PCGGraphExecutor.h b/Engine/Plugins/Experimental/PCG/Source/PCG/Private/Graph/PCGGraphExecutor.h
--- a/Engine/Plugins/Experimental/PCG/Source/PCG/Private/Graph/PCGGraphExecutor.h
+++ b/Engine/Plugins/Experimental/PCG/Source/PCG/Private/Graph/PCGGraphExecutor.h
@@ -129,6 +129,7 @@ private:
 	TMap<FPCGTaskId, FPCGGraphTask> Tasks;
 	TArray<FPCGGraphTask> ReadyTasks;
 	TArray<FPCGGraphActiveTask> ActiveTasks;
+	TArray<FPCGGraphActiveTask> SleepingTasks;
 	TMap<FPCGTaskId, TSet<FPCGTaskId>> TaskSuccessors;
 	/** Map of node instances to their output, could be cleared once execution is done */
 	/** Note: this should at some point unload based on loaded/unloaded proxies, otherwise memory cost will be unbounded */
@@ -153,6 +154,7 @@ class FPCGFetchInputElement : public FSimplePCGElement
 {
 public:
 	virtual bool IsCacheable(const UPCGSettings* InSettings) const override { return false; }
+	virtual bool CanExecuteOnlyOnMainThread(const UPCGSettings* InSettings) const override { return true; }
 
 protected:
 	virtual bool ExecuteInternal(FPCGContext* Context) const override;
diff --git a/Engine/Plugins/Experimental/PCG/Source/PCG/Private/Grid/PCGLandscapeCache.cpp b/Engine/Plugins/Experimental/PCG/Source/PCG/Private/Grid/PCGLandscapeCache.cpp
--- a/Engine/Plugins/Experimental/PCG/Source/PCG/Private/Grid/PCGLandscapeCache.cpp
+++ b/Engine/Plugins/Experimental/PCG/Source/PCG/Private/Grid/PCGLandscapeCache.cpp
@@ -274,7 +274,6 @@ void FPCGLandscapeCacheEntry::Unload()
 
 int32 FPCGLandscapeCacheEntry::GetMemorySize() const
 {
-	check(bDataLoaded);
 	int32 MemSize = PositionsAndNormals.GetAllocatedSize();
 	for (const TArray<uint8>& Layer : LayerData)
 	{
diff --git a/Engine/Plugins/Experimental/PCG/Source/PCG/Private/PCGElement.cpp b/Engine/Plugins/Experimental/PCG/Source/PCG/Private/PCGElement.cpp
--- a/Engine/Plugins/Experimental/PCG/Source/PCG/Private/PCGElement.cpp
+++ b/Engine/Plugins/Experimental/PCG/Source/PCG/Private/PCGElement.cpp
@@ -12,6 +12,8 @@
 bool IPCGElement::Execute(FPCGContext* Context) const
 {
 	TRACE_CPUPROFILER_EVENT_SCOPE(IPCGElement::Execute);
+	check(Context && Context->NumAvailableTasks > 0);
+
 	// Early out to stop execution
 	if (Context->InputData.bCancelExecution || (Context->SourceComponent && !IsValid(Context->SourceComponent)))
 	{
@@ -102,8 +104,6 @@ bool IPCGElement::Execute(FPCGContext* Context) const
 
 #if WITH_EDITOR
 			PCGE_LOG(Log, "Executed in (%f)s and (%d) call(s)", Context->ElapsedTime, Context->ExecutionCount);
-#else
-			PCGE_LOG(Log, "Executed");
 #endif
 		}
 
diff --git a/Engine/Plugins/Experimental/PCG/Source/PCG/Private/PCGParamData.cpp b/Engine/Plugins/Experimental/PCG/Source/PCG/Private/PCGParamData.cpp
--- a/Engine/Plugins/Experimental/PCG/Source/PCG/Private/PCGParamData.cpp
+++ b/Engine/Plugins/Experimental/PCG/Source/PCG/Private/PCGParamData.cpp
@@ -51,7 +51,7 @@ UPCGParamData* UPCGParamData::FilterParamsByName(const FName& InName) const
 
 UPCGParamData* UPCGParamData::FilterParamsByKey(int64 InKey) const
 {
-	UPCGParamData* NewParams = NewObject<UPCGParamData>(const_cast<UPCGParamData*>(this));
+	UPCGParamData* NewParams = NewObject<UPCGParamData>();
 
 	// Here instead of parenting the metadata, we will create a copy
 	// so that the only entry in the metadata (if any) will have the 0 key.
diff --git a/Engine/Plugins/Experimental/PCG/Source/PCG/Public/Data/PCGSpatialData.h b/Engine/Plugins/Experimental/PCG/Source/PCG/Public/Data/PCGSpatialData.h
--- a/Engine/Plugins/Experimental/PCG/Source/PCG/Public/Data/PCGSpatialData.h
+++ b/Engine/Plugins/Experimental/PCG/Source/PCG/Public/Data/PCGSpatialData.h
@@ -106,7 +106,7 @@ public:
 	void InitializeFromData(const UPCGSpatialData* InSource, const UPCGMetadata* InMetadataParentOverride = nullptr, bool bInheritMetadata = true);
 
 	UPROPERTY(Transient, BlueprintReadWrite, EditAnywhere, Category = Data)
-	TObjectPtr<AActor> TargetActor = nullptr;
+	TWeakObjectPtr<AActor> TargetActor = nullptr;
 
 #if WITH_EDITORONLY_DATA
 	UPROPERTY(Transient, BlueprintReadWrite, EditAnywhere, Category = Debug)
