[PCG] Metadata native operation: Rotator operations

diff --git a/Engine/Plugins/Experimental/PCG/Source/PCG/Private/Elements/Metadata/PCGMetadataRotatorOpElement.cpp b/Engine/Plugins/Experimental/PCG/Source/PCG/Private/Elements/Metadata/PCGMetadataRotatorOpElement.cpp
new file mode 100644
--- /dev/null
+++ b/Engine/Plugins/Experimental/PCG/Source/PCG/Private/Elements/Metadata/PCGMetadataRotatorOpElement.cpp
@@ -0,0 +1,235 @@
+// Copyright Epic Games, Inc. All Rights Reserved.
+
+#include "Elements/Metadata/PCGMetadataRotatorOpElement.h"
+
+#include "Helpers/PCGSettingsHelpers.h"
+#include "Metadata/PCGMetadata.h"
+#include "Metadata/PCGMetadataAttribute.h"
+#include "Metadata/PCGMetadataAttributeTpl.h"
+#include "Metadata/PCGMetadataEntryKeyIterator.h"
+#include "Elements/Metadata/PCGMetadataElementCommon.h"
+
+// Taken from Kismet Math library
+FRotator PCGMetadataRotatorHelpers::RLerp(const FRotator& A, const FRotator& B, double Alpha, bool bShortestPath)
+{
+	// if shortest path, we use Quaternion to interpolate instead of using FRotator
+	if (bShortestPath)
+	{
+		FQuat AQuat(A);
+		FQuat BQuat(B);
+
+		FQuat Result = FQuat::Slerp(AQuat, BQuat, Alpha);
+
+		return Result.Rotator();
+	}
+
+	const FRotator DeltaAngle = B - A;
+	return A + Alpha * DeltaAngle;
+}
+
+namespace PCGMetadataRotatorSettings
+{
+	inline constexpr bool IsUnaryOp(EPCGMedadataRotatorOperation Operation)
+	{
+		return Operation == EPCGMedadataRotatorOperation::Invert;
+	}
+
+	inline constexpr bool IsTernaryOp(EPCGMedadataRotatorOperation Operation)
+	{
+		return Operation == EPCGMedadataRotatorOperation::Lerp;
+	}
+
+	inline constexpr bool IsTransfromOp(EPCGMedadataRotatorOperation Operation)
+	{
+		return (uint16)Operation >= (uint16)EPCGMedadataRotatorOperation::TransformOp;
+	}
+
+	inline FRotator ApplyRotatorOperation(const FRotator& Input1, const FRotator& Input2, double Ratio, EPCGMedadataRotatorOperation Operation)
+	{
+		switch (Operation)
+		{
+		case EPCGMedadataRotatorOperation::Combine:
+			return FRotator(Input2.Quaternion() * Input1.Quaternion());
+		case EPCGMedadataRotatorOperation::Lerp:
+			return PCGMetadataRotatorHelpers::RLerp(Input1, Input2, Ratio, false);
+		case EPCGMedadataRotatorOperation::Invert:
+			return Input1.GetInverse();
+		default:
+			return FRotator{};
+		}
+	}
+
+	inline FQuat ApplyRotatorOperation(const FQuat& Input1, const FQuat& Input2, double Ratio, EPCGMedadataRotatorOperation Operation)
+	{
+		switch (Operation)
+		{
+		case EPCGMedadataRotatorOperation::Combine:
+			return Input2 * Input1;
+		case EPCGMedadataRotatorOperation::Lerp:
+			return FQuat::Slerp(Input1, Input2, Ratio);
+		case EPCGMedadataRotatorOperation::Invert:
+			return Input1.Inverse();
+		default:
+			return FQuat{};
+		}
+	}
+
+	inline FQuat ApplyTransformOperation(const FQuat& Input, const FTransform& Transform, EPCGMedadataRotatorOperation Operation)
+	{
+		switch (Operation)
+		{
+		case EPCGMedadataRotatorOperation::TransformRotation:
+			return Transform.TransformRotation(Input);
+		case EPCGMedadataRotatorOperation::InverseTransformRotation:
+			return Transform.InverseTransformRotation(Input);
+		default:
+			return FQuat{};
+		}
+	}
+
+	// In Kismet Math Library, they transform Rotators in Quaternions.
+	inline FRotator ApplyTransformOperation(const FRotator& Input, const FTransform& Transform, EPCGMedadataRotatorOperation Operation)
+	{
+		return FRotator(ApplyTransformOperation(Input.Quaternion(), Transform, Operation));
+	}
+}
+
+FName UPCGMetadataRotatorSettings::GetInputPinLabel(uint32 Index) const
+{
+	switch (Index)
+	{
+	case 0:
+		if (!PCGMetadataRotatorSettings::IsUnaryOp(Operation) && !PCGMetadataRotatorSettings::IsTransfromOp(Operation))
+		{
+			return PCGMetadataSettingsBaseConstants::DoubleInputFirstLabel;
+		}
+		else
+		{
+			return PCGPinConstants::DefaultInputLabel;
+		}
+	case 1:
+		return PCGMetadataRotatorSettings::IsTransfromOp(Operation) ? PCGMetadataSettingsBaseConstants::TransformLabel : PCGMetadataSettingsBaseConstants::DoubleInputSecondLabel;
+	case 2:
+		return PCGMetadataSettingsBaseConstants::LerpRatioLabel;
+	default:
+		return NAME_None;
+	}
+}
+
+uint32 UPCGMetadataRotatorSettings::GetInputPinNum() const
+{
+	if (PCGMetadataRotatorSettings::IsUnaryOp(Operation))
+	{
+		return 1;
+	}
+	else if (PCGMetadataRotatorSettings::IsTernaryOp(Operation))
+	{
+		return 3;
+	}
+	else
+	{
+		return 2;
+	}
+}
+
+bool UPCGMetadataRotatorSettings::IsSupportedInputType(uint16 TypeId, uint32 InputIndex, bool& bHasSpecialRequirement) const
+{
+	if (InputIndex == 2)
+	{
+		bHasSpecialRequirement = true;
+		return TypeId == (uint16)EPCGMetadataTypes::Float || TypeId == (uint16)EPCGMetadataTypes::Double;
+	}
+	else if (InputIndex == 1 && PCGMetadataRotatorSettings::IsTransfromOp(Operation))
+	{
+		bHasSpecialRequirement = true;
+		return TypeId == (uint16)EPCGMetadataTypes::Transform;
+	}
+	else
+	{
+		bHasSpecialRequirement = false;
+		return TypeId == (uint16)EPCGMetadataTypes::Rotator || TypeId == (uint16)EPCGMetadataTypes::Quaternion;
+	}
+}
+
+FName UPCGMetadataRotatorSettings::GetInputAttributeNameWithOverride(uint32 Index, UPCGParamData* Params) const
+{
+	switch (Index)
+	{
+	case 0:
+		return PCG_GET_OVERRIDEN_VALUE(this, Input1AttributeName, Params);
+	case 1:
+		return PCG_GET_OVERRIDEN_VALUE(this, Input2AttributeName, Params);
+	case 2:
+		return PCG_GET_OVERRIDEN_VALUE(this, Input3AttributeName, Params);
+	default:
+		return NAME_None;
+	}
+}
+
+#if WITH_EDITOR
+FName UPCGMetadataRotatorSettings::GetDefaultNodeName() const
+{
+	if (const UEnum* EnumPtr = FindObject<UEnum>(nullptr, TEXT("EPCGMedadataRotatorOperation"), true))
+	{ 
+		return EnumPtr->GetNameByValue(static_cast<int>(Operation)); 
+	}
+
+	return TEXT("Metadata Rotator Node");
+}
+#endif // WITH_EDITOR
+
+FPCGElementPtr UPCGMetadataRotatorSettings::CreateElement() const
+{
+	return MakeShared<FPCGMetadataRotatorElement>();
+}
+
+bool FPCGMetadataRotatorElement::DoOperation(FOperationData& OperationData) const
+{
+	TRACE_CPUPROFILER_EVENT_SCOPE(FPCGMetadataRotatorElement::Execute);
+
+	const UPCGMetadataRotatorSettings* Settings = CastChecked<UPCGMetadataRotatorSettings>(OperationData.Settings);
+
+	auto RotatorFunc = [this, &OperationData, Operation = Settings->Operation](auto DummyValue)
+	{
+		using AttributeType = decltype(DummyValue);
+
+		if constexpr (PCG::Private::MetadataTypes<AttributeType>::Id != (uint16)EPCGMetadataTypes::Quaternion &&
+					PCG::Private::MetadataTypes<AttributeType>::Id != (uint16)EPCGMetadataTypes::Rotator)
+		{
+			return;
+		}
+		else
+		{
+			if (PCGMetadataRotatorSettings::IsTransfromOp(Operation))
+			{
+				DoBinaryOp<AttributeType, FTransform, AttributeType>(OperationData, [Operation](const AttributeType& Value, const FTransform& Transform)->AttributeType {
+					return PCGMetadataRotatorSettings::ApplyTransformOperation(Value, Transform, Operation);
+					});
+			}
+			else if (PCGMetadataRotatorSettings::IsUnaryOp(Operation))
+			{
+				DoUnaryOp<AttributeType, AttributeType>(OperationData, [Operation](const AttributeType& Value)->AttributeType {
+					double DummyDouble = 0.0;
+					return PCGMetadataRotatorSettings::ApplyRotatorOperation(Value, AttributeType{}, DummyDouble, Operation);
+				});
+			}
+			else if (PCGMetadataRotatorSettings::IsTernaryOp(Operation))
+			{
+				DoTernaryOp<AttributeType, AttributeType, double, AttributeType>(OperationData, [Operation](const AttributeType& Value1, const AttributeType& Value2, const double& Ratio)->AttributeType {
+					return PCGMetadataRotatorSettings::ApplyRotatorOperation(Value1, Value2, Ratio, Operation);
+				});
+			}
+			else
+			{
+				DoBinaryOp<AttributeType, AttributeType, AttributeType>(OperationData, [Operation](const AttributeType& Value1, const AttributeType& Value2)->AttributeType {
+					double DummyDouble = 0.0;
+					return PCGMetadataRotatorSettings::ApplyRotatorOperation(Value1, Value2, DummyDouble, Operation);
+				});
+			}
+		}
+	};
+
+	PCGMetadataAttribute::CallbackWithRightType(OperationData.OutputType, RotatorFunc);
+
+	return true;
+}
\ No newline at end of file
diff --git a/Engine/Plugins/Experimental/PCG/Source/PCG/Public/Elements/Metadata/PCGMetadataRotatorOpElement.h b/Engine/Plugins/Experimental/PCG/Source/PCG/Public/Elements/Metadata/PCGMetadataRotatorOpElement.h
new file mode 100644
--- /dev/null
+++ b/Engine/Plugins/Experimental/PCG/Source/PCG/Public/Elements/Metadata/PCGMetadataRotatorOpElement.h
@@ -0,0 +1,74 @@
+// Copyright Epic Games, Inc. All Rights Reserved.
+
+#pragma once
+
+#include "PCGElement.h"
+#include "PCGSettings.h"
+#include "Elements/Metadata/PCGMetadataOpElementBase.h"
+
+#include "PCGMetadataRotatorOpElement.generated.h"
+
+UENUM()
+enum class EPCGMedadataRotatorOperation : uint16
+{
+	RotatorOp = 0 UMETA(Hidden),
+	Combine,
+	Invert,
+	Lerp,
+
+	TransformOp = 100 UMETA(Hidden),
+	TransformRotation,
+	InverseTransformRotation
+};
+
+namespace PCGMetadataRotatorHelpers
+{
+	// Taken from Kismet Math Library
+	FRotator RLerp(const FRotator& A, const FRotator& B, double Alpha, bool bShortestPath);
+}
+
+UCLASS(BlueprintType, ClassGroup = (Procedural))
+class PCG_API UPCGMetadataRotatorSettings : public UPCGMetadataSettingsBase
+{
+	GENERATED_BODY()
+
+public:
+	//~Begin UPCGSettings interface
+#if WITH_EDITOR
+	virtual FName GetDefaultNodeName() const override;
+#endif
+	//~End UPCGSettings interface
+
+	//~Begin UPCGMetadataSettingsBase interface
+	virtual FName GetInputAttributeNameWithOverride(uint32 Index, UPCGParamData* Params) const override;
+
+	virtual FName GetInputPinLabel(uint32 Index) const override;
+	virtual uint32 GetInputPinNum() const override;
+
+	virtual bool IsSupportedInputType(uint16 TypeId, uint32 InputIndex, bool& bHasSpecialRequirement) const override;
+	//~End UPCGMetadataSettingsBase interface
+
+protected:
+	//~Begin UPCGSettings interface
+	virtual FPCGElementPtr CreateElement() const override;
+	//~End UPCGSettings interface
+
+public:
+	UPROPERTY(BlueprintReadWrite, EditAnywhere, Category = Settings)
+	EPCGMedadataRotatorOperation Operation = EPCGMedadataRotatorOperation::Combine;
+
+	UPROPERTY(BlueprintReadWrite, EditAnywhere, Category = Input)
+	FName Input1AttributeName = NAME_None;
+
+	UPROPERTY(BlueprintReadWrite, EditAnywhere, Category = Input, meta = (EditCondition = "Operation != EPCGMedadataTrigOperation::Invert", EditConditionHides))
+	FName Input2AttributeName = NAME_None;
+
+	UPROPERTY(BlueprintReadWrite, EditAnywhere, Category = Input, meta = (EditCondition = "Operation == EPCGMedadataTrigOperation::Lerp", EditConditionHides))
+	FName Input3AttributeName = NAME_None;
+};
+
+class FPCGMetadataRotatorElement : public FPCGMetadataElementBase
+{
+protected:
+	virtual bool DoOperation(FOperationData& OperationData) const override;
+};
