[PCG] Metadata native operation: Transform operations

diff --git a/Engine/Plugins/Experimental/PCG/Source/PCG/Private/Elements/Metadata/PCGMetadataTransformOpElement.cpp b/Engine/Plugins/Experimental/PCG/Source/PCG/Private/Elements/Metadata/PCGMetadataTransformOpElement.cpp
new file mode 100644
--- /dev/null
+++ b/Engine/Plugins/Experimental/PCG/Source/PCG/Private/Elements/Metadata/PCGMetadataTransformOpElement.cpp
@@ -0,0 +1,183 @@
+// Copyright Epic Games, Inc. All Rights Reserved.
+
+#include "Elements/Metadata/PCGMetadataTransformOpElement.h"
+
+#include "Elements/Metadata/PCGMetadataElementCommon.h"
+#include "Elements/Metadata/PCGMetadataRotatorOpElement.h"
+#include "Helpers/PCGSettingsHelpers.h"
+#include "Metadata/PCGMetadata.h"
+#include "Metadata/PCGMetadataAttribute.h"
+#include "Metadata/PCGMetadataAttributeTpl.h"
+#include "Metadata/PCGMetadataEntryKeyIterator.h"
+
+#include "Math/DualQuat.h"
+
+namespace PCGMetadataTransfromSettings
+{
+	inline bool IsUnaryOp(EPCGMedadataTransformOperation Operation)
+	{
+		return Operation == EPCGMedadataTransformOperation::Invert;
+	}
+
+	inline bool IsTernaryOp(EPCGMedadataTransformOperation Operation)
+	{
+		return Operation == EPCGMedadataTransformOperation::Lerp;
+	}
+
+	// Taken from Kismet Math Library
+	inline FTransform LerpTransform(const FTransform& A, const FTransform& B, double Alpha, EPCGTransformLerpMode Mode)
+	{
+		FTransform Result;
+
+		FTransform NA = A;
+		FTransform NB = B;
+		NA.NormalizeRotation();
+		NB.NormalizeRotation();
+
+		// Quaternion interpolation
+		if (Mode == EPCGTransformLerpMode::QuatInterp)
+		{
+			Result.Blend(NA, NB, Alpha);
+			return Result;
+		}
+		// Euler Angle interpolation
+		else if (Mode == EPCGTransformLerpMode::EulerInterp)
+		{
+			Result.SetTranslation(FMath::Lerp(NA.GetTranslation(), NB.GetTranslation(), Alpha));
+			Result.SetScale3D(FMath::Lerp(NA.GetScale3D(), NB.GetScale3D(), Alpha));
+			Result.SetRotation(FQuat(PCGMetadataRotatorHelpers::RLerp(NA.Rotator(), NB.Rotator(), Alpha, false)));
+			return Result;
+		}
+		// Dual quaternion interpolation
+		else
+		{
+			if ((NB.GetRotation() | NA.GetRotation()) < 0.0f)
+			{
+				NB.SetRotation(NB.GetRotation() * -1.0f);
+			}
+			return (FDualQuat(NA) * (1 - Alpha) + FDualQuat(NB) * Alpha).Normalized().AsFTransform(FMath::Lerp(NA.GetScale3D(), NB.GetScale3D(), Alpha));
+		}
+	}
+
+	inline FTransform ApplyTransformOperation(const FTransform& Value1, const FTransform& Value2, double Ratio, EPCGMedadataTransformOperation Operation, EPCGTransformLerpMode Mode)
+	{
+		switch (Operation)
+		{
+		case EPCGMedadataTransformOperation::Invert:
+			return Value1.Inverse();
+		case EPCGMedadataTransformOperation::Compose:
+			return Value1 * Value2;
+		case EPCGMedadataTransformOperation::Lerp:
+			return LerpTransform(Value1, Value2, Ratio, Mode);
+		default:
+			return FTransform{};
+		}
+	}
+}
+
+FName UPCGMetadataTransformSettings::GetInputPinLabel(uint32 Index) const
+{
+	switch (Index)
+	{
+	case 0:
+		return (Operation == EPCGMedadataTransformOperation::Invert) ? PCGPinConstants::DefaultInputLabel : PCGMetadataSettingsBaseConstants::DoubleInputFirstLabel;
+	case 1:
+		return PCGMetadataSettingsBaseConstants::DoubleInputSecondLabel;
+	case 2:
+		return PCGMetadataSettingsBaseConstants::LerpRatioLabel;
+	default:
+		return NAME_None;
+	}
+}
+
+uint32 UPCGMetadataTransformSettings::GetInputPinNum() const
+{
+	if (Operation == EPCGMedadataTransformOperation::Invert)
+	{
+		return 1;
+	}
+	else if (Operation == EPCGMedadataTransformOperation::Lerp)
+	{
+		return 3;
+	}
+	else
+	{
+		return 2;
+	}
+}
+
+bool UPCGMetadataTransformSettings::IsSupportedInputType(uint16 TypeId, uint32 InputIndex, bool& bHasSpecialRequirement) const
+{
+	if (InputIndex == 2)
+	{
+		bHasSpecialRequirement = true;
+		return TypeId == (uint16)EPCGMetadataTypes::Float || TypeId == (uint16)EPCGMetadataTypes::Double;
+	}
+	else
+	{
+		bHasSpecialRequirement = false;
+		return TypeId == (uint16)EPCGMetadataTypes::Transform;
+	}
+}
+
+FName UPCGMetadataTransformSettings::GetInputAttributeNameWithOverride(uint32 Index, UPCGParamData* Params) const
+{
+	switch (Index)
+	{
+	case 0:
+		return PCG_GET_OVERRIDEN_VALUE(this, Input1AttributeName, Params);
+	case 1:
+		return PCG_GET_OVERRIDEN_VALUE(this, Input2AttributeName, Params);
+	case 2:
+		return PCG_GET_OVERRIDEN_VALUE(this, Input3AttributeName, Params);
+	default:
+		return NAME_None;
+	}
+}
+
+#if WITH_EDITOR
+FName UPCGMetadataTransformSettings::GetDefaultNodeName() const
+{
+	if (const UEnum* EnumPtr = FindObject<UEnum>(nullptr, TEXT("EPCGMedadataTransformOperation"), true))
+	{ 
+		return EnumPtr->GetNameByValue(static_cast<int>(Operation)); 
+	}
+
+	return TEXT("Metadata Transform Node");
+}
+#endif // WITH_EDITOR
+
+FPCGElementPtr UPCGMetadataTransformSettings::CreateElement() const
+{
+	return MakeShared<FPCGMetadataTransformElement>();
+}
+
+bool FPCGMetadataTransformElement::DoOperation(FOperationData& OperationData) const
+{
+	TRACE_CPUPROFILER_EVENT_SCOPE(FPCGMetadataTransformElement::Execute);
+
+	const UPCGMetadataTransformSettings* Settings = CastChecked<UPCGMetadataTransformSettings>(OperationData.Settings);
+
+	if (PCGMetadataTransfromSettings::IsUnaryOp(Settings->Operation))
+	{
+		DoUnaryOp<FTransform, FTransform>(OperationData, [Operation = Settings->Operation, Mode = Settings->TransformLerpMode](const FTransform& Value)->FTransform {
+			double DummyDouble = 0.0;
+			return PCGMetadataTransfromSettings::ApplyTransformOperation(Value, FTransform{}, DummyDouble, Operation, Mode);
+			});
+	}
+	else if (PCGMetadataTransfromSettings::IsTernaryOp(Settings->Operation))
+	{
+		DoTernaryOp<FTransform, FTransform, double, FTransform>(OperationData, [Operation = Settings->Operation, Mode = Settings->TransformLerpMode](const FTransform& Value1, const FTransform& Value2, const double& Ratio)->FTransform {
+			return PCGMetadataTransfromSettings::ApplyTransformOperation(Value1, Value2, Ratio, Operation, Mode);
+			});
+	}
+	else
+	{
+		DoBinaryOp<FTransform, FTransform, FTransform>(OperationData, [Operation = Settings->Operation, Mode = Settings->TransformLerpMode](const FTransform& Value1, const FTransform& Value2)->FTransform {
+			double DummyDouble = 0.0;
+			return PCGMetadataTransfromSettings::ApplyTransformOperation(Value1, Value2, DummyDouble, Operation, Mode);
+			});
+	}
+
+	return true;
+}
\ No newline at end of file
diff --git a/Engine/Plugins/Experimental/PCG/Source/PCG/Public/Elements/Metadata/PCGMetadataTransformOpElement.h b/Engine/Plugins/Experimental/PCG/Source/PCG/Public/Elements/Metadata/PCGMetadataTransformOpElement.h
new file mode 100644
--- /dev/null
+++ b/Engine/Plugins/Experimental/PCG/Source/PCG/Public/Elements/Metadata/PCGMetadataTransformOpElement.h
@@ -0,0 +1,80 @@
+// Copyright Epic Games, Inc. All Rights Reserved.
+
+#pragma once
+
+#include "PCGElement.h"
+#include "PCGSettings.h"
+#include "Elements/Metadata/PCGMetadataOpElementBase.h"
+
+#include "PCGMetadataTransformOpElement.generated.h"
+
+UENUM()
+enum class EPCGMedadataTransformOperation : uint16
+{
+	Compose,
+	Invert,
+	Lerp,
+};
+
+// Taken from Kismet Math Library
+UENUM()
+enum class EPCGTransformLerpMode : uint16
+{
+	/** Shortest Path or Quaternion interpolation for the rotation. */
+	QuatInterp,
+
+	/** Rotor or Euler Angle interpolation. */
+	EulerInterp,
+
+	/** Dual quaternion interpolation, follows helix or screw-motion path between keyframes.   */
+	DualQuatInterp
+};
+
+UCLASS(BlueprintType, ClassGroup = (Procedural))
+class PCG_API UPCGMetadataTransformSettings : public UPCGMetadataSettingsBase
+{
+	GENERATED_BODY()
+
+public:
+	//~Begin UPCGSettings interface
+#if WITH_EDITOR
+	virtual FName GetDefaultNodeName() const override;
+#endif
+	//~End UPCGSettings interface
+
+	//~Begin UPCGMetadataSettingsBase interface
+	virtual FName GetInputAttributeNameWithOverride(uint32 Index, UPCGParamData* Params) const override;
+
+	virtual FName GetInputPinLabel(uint32 Index) const override;
+	virtual uint32 GetInputPinNum() const override;
+
+	virtual bool IsSupportedInputType(uint16 TypeId, uint32 InputIndex, bool& bHasSpecialRequirement) const override;
+	//~End UPCGMetadataSettingsBase interface
+
+protected:
+	//~Begin UPCGSettings interface
+	virtual FPCGElementPtr CreateElement() const override;
+	//~End UPCGSettings interface
+
+public:
+	UPROPERTY(BlueprintReadWrite, EditAnywhere, Category = Settings)
+	EPCGMedadataTransformOperation Operation = EPCGMedadataTransformOperation::Compose;
+
+	UPROPERTY(BlueprintReadWrite, EditAnywhere, Category = Settings, meta = (EditCondition = "Operation == EPCGMedadataTransformOperation::Lerp", EditConditionHides))
+	EPCGTransformLerpMode TransformLerpMode = EPCGTransformLerpMode::QuatInterp;
+
+	UPROPERTY(BlueprintReadWrite, EditAnywhere, Category = Input)
+	FName Input1AttributeName = NAME_None;
+
+	UPROPERTY(BlueprintReadWrite, EditAnywhere, Category = Input, meta = (EditCondition = "Operation != EPCGMedadataTransformOperation::Invert", EditConditionHides))
+	FName Input2AttributeName = NAME_None;
+
+	UPROPERTY(BlueprintReadWrite, EditAnywhere, Category = Input, meta = (EditCondition = "Operation == EPCGMedadataTransformOperation::Lerp", EditConditionHides))
+	FName Input3AttributeName = NAME_None;
+};
+
+class FPCGMetadataTransformElement : public FPCGMetadataElementBase
+{
+protected:
+	virtual bool DoOperation(FOperationData& OperationData) const override;
+};
