PCG:
*Fixed a crash when a user tried to rename a PCG node with an input larger than 1023 characters.
*Removed an old unused function override and include.

diff --git a/Engine/Plugins/Experimental/PCG/Source/PCGEditor/Private/PCGEditorGraphNode.cpp b/Engine/Plugins/Experimental/PCG/Source/PCGEditor/Private/PCGEditorGraphNode.cpp
--- a/Engine/Plugins/Experimental/PCG/Source/PCGEditor/Private/PCGEditorGraphNode.cpp
+++ b/Engine/Plugins/Experimental/PCG/Source/PCGEditor/Private/PCGEditorGraphNode.cpp
@@ -3,8 +3,8 @@
 #include "PCGEditorGraphNode.h"
 
 #include "PCGEditorGraphSchema.h"
+#include "PCGEditorModule.h"
 #include "PCGNode.h"
-#include "PCGSettings.h"
 
 #include "Framework/Commands/GenericCommands.h"
 #include "GraphEditorActions.h"
@@ -63,20 +63,30 @@ void UPCGEditorGraphNode::AllocateDefaultPins()
 	}
 }
 
-void UPCGEditorGraphNode::ReconstructNode()
-{
-	Super::ReconstructNode();
-	//TODO: Implement special version for settings to avoid the enum type
-}
-
 void UPCGEditorGraphNode::OnRenameNode(const FString& NewName)
 {
-	FName TentativeName(NewName);
+	if (!GetCanRenameNode())
+	{
+		return;
+	}
+
+	if(NewName.Len() >= NAME_SIZE)
+	{
+		UE_LOG(LogPCGEditor, Error, TEXT("New name for PCG node is too long."));
+		return;
+	}
+
+	if (!PCGNode)
+	{
+		return;
+	}
+
+	const FName TentativeName(*NewName);
 
-	if (GetCanRenameNode() && PCGNode && PCGNode->GetNodeTitle() != TentativeName)
+	if (PCGNode->GetNodeTitle() != TentativeName)
 	{
 		PCGNode->Modify();
-		PCGNode->NodeTitle = FName(NewName);
+		PCGNode->NodeTitle = TentativeName;
 	}
 }
 
diff --git a/Engine/Plugins/Experimental/PCG/Source/PCGEditor/Private/PCGEditorGraphNode.h b/Engine/Plugins/Experimental/PCG/Source/PCGEditor/Private/PCGEditorGraphNode.h
--- a/Engine/Plugins/Experimental/PCG/Source/PCGEditor/Private/PCGEditorGraphNode.h
+++ b/Engine/Plugins/Experimental/PCG/Source/PCGEditor/Private/PCGEditorGraphNode.h
@@ -22,7 +22,6 @@ public:
 	virtual FText GetNodeTitle(ENodeTitleType::Type TitleType) const override;
 	virtual void GetNodeContextMenuActions(UToolMenu* Menu, class UGraphNodeContextMenuContext* Context) const override;
 	virtual void AllocateDefaultPins() override;
-	virtual void ReconstructNode() override;
 	virtual void OnRenameNode(const FString& NewName);
 	// ~End UEdGraphNode interface
 };
