Made the source component pointer in the PCG context weak so we can more easily prevent computation/issues.
Added some robustness in the PCG blueprint.
Prior work in the graph executor for task cancellation.

diff --git a/Engine/Plugins/Experimental/PCG/Source/PCG/Private/Elements/PCGDebugElement.cpp b/Engine/Plugins/Experimental/PCG/Source/PCG/Private/Elements/PCGDebugElement.cpp
--- a/Engine/Plugins/Experimental/PCG/Source/PCG/Private/Elements/PCGDebugElement.cpp
+++ b/Engine/Plugins/Experimental/PCG/Source/PCG/Private/Elements/PCGDebugElement.cpp
@@ -18,7 +18,7 @@ namespace PCGDebugElement
 		TRACE_CPUPROFILER_EVENT_SCOPE(PCGDebugElement::ExecuteDebugDisplay);
 #if WITH_EDITORONLY_DATA
 		// Early validation: if we don't have a valid PCG component, we're not going to add the debug display info
-		if (!Context->SourceComponent || !IsValid(Context->SourceComponent))
+		if (!Context->SourceComponent.IsValid())
 		{
 			return;
 		}
@@ -146,7 +146,7 @@ namespace PCGDebugElement
 			Params.MaterialOverrides = Materials;
 			Params.CollisionProfile = UCollisionProfile::NoCollision_ProfileName;
 
-			UInstancedStaticMeshComponent* ISMC = UPCGActorHelpers::GetOrCreateISMC(TargetActor, Context->SourceComponent, Params);
+			UInstancedStaticMeshComponent* ISMC = UPCGActorHelpers::GetOrCreateISMC(TargetActor, Context->SourceComponent.Get(), Params);
 			check(ISMC);
 			
 			ISMC->ComponentTags.AddUnique(PCGHelpers::DefaultPCGDebugTag);
diff --git a/Engine/Plugins/Experimental/PCG/Source/PCG/Private/Elements/PCGExecuteBlueprint.cpp b/Engine/Plugins/Experimental/PCG/Source/PCG/Private/Elements/PCGExecuteBlueprint.cpp
--- a/Engine/Plugins/Experimental/PCG/Source/PCG/Private/Elements/PCGExecuteBlueprint.cpp
+++ b/Engine/Plugins/Experimental/PCG/Source/PCG/Private/Elements/PCGExecuteBlueprint.cpp
@@ -615,7 +615,7 @@ void UPCGBlueprintElement::LoopNTimes(FPCGContext& InContext, int64 NumIteration
 	});
 }
 
-FPCGContext* FPCGExecuteBlueprintElement::Initialize(const FPCGDataCollection& InputData, UPCGComponent* SourceComponent, const UPCGNode* Node)
+FPCGContext* FPCGExecuteBlueprintElement::Initialize(const FPCGDataCollection& InputData, TWeakObjectPtr<UPCGComponent> SourceComponent, const UPCGNode* Node)
 {
 	TRACE_CPUPROFILER_EVENT_SCOPE(FPCGExecuteBlueprintElement::Initialize);
 	FPCGBlueprintExecutionContext* Context = new FPCGBlueprintExecutionContext();
@@ -628,7 +628,10 @@ FPCGContext* FPCGExecuteBlueprintElement::Initialize(const FPCGDataCollection& I
 	{
 		Context->BlueprintElementInstance = CastChecked<UPCGBlueprintElement>(StaticDuplicateObject(Settings->BlueprintElementInstance, GetTransientPackage(), FName()));
 #if !WITH_EDITOR
-		Context->BlueprintElementInstance->SetInstanceWorld(SourceComponent ? SourceComponent->GetOwner()->GetWorld() : nullptr);
+		if (SourceComponent.IsValid() && SourceComponent->GetOwner())
+		{
+			Context->BlueprintElementInstance->SetInstanceWorld(SourceComponent->GetOwner()->GetWorld());
+		}
 #endif
 	}
 	else
diff --git a/Engine/Plugins/Experimental/PCG/Source/PCG/Private/Elements/PCGPropertyToParamData.cpp b/Engine/Plugins/Experimental/PCG/Source/PCG/Private/Elements/PCGPropertyToParamData.cpp
--- a/Engine/Plugins/Experimental/PCG/Source/PCG/Private/Elements/PCGPropertyToParamData.cpp
+++ b/Engine/Plugins/Experimental/PCG/Source/PCG/Private/Elements/PCGPropertyToParamData.cpp
@@ -72,7 +72,7 @@ namespace PCGPropertyToParamDataHelpers
 
 		AActor* FoundActor = nullptr;
 
-		if (!InContext.SourceComponent)
+		if (!InContext.SourceComponent.IsValid())
 		{
 			return FoundActor;
 		}
diff --git a/Engine/Plugins/Experimental/PCG/Source/PCG/Private/Elements/PCGSpawnActor.cpp b/Engine/Plugins/Experimental/PCG/Source/PCG/Private/Elements/PCGSpawnActor.cpp
--- a/Engine/Plugins/Experimental/PCG/Source/PCG/Private/Elements/PCGSpawnActor.cpp
+++ b/Engine/Plugins/Experimental/PCG/Source/PCG/Private/Elements/PCGSpawnActor.cpp
@@ -108,7 +108,7 @@ bool FPCGSpawnActorElement::ExecuteInternal(FPCGContext* Context) const
 			continue;
 		}
 
-		const bool bHasAuthority = !Context->SourceComponent || (Context->SourceComponent->GetOwner() && Context->SourceComponent->GetOwner()->HasAuthority());
+		const bool bHasAuthority = !Context->SourceComponent.IsValid() || (Context->SourceComponent->GetOwner() && Context->SourceComponent->GetOwner()->HasAuthority());
 		const bool bSpawnedActorsRequireAuthority = Settings->TemplateActorClass->GetDefaultObject<AActor>()->GetIsReplicated();
 
 		// First, create target instance transforms
@@ -159,7 +159,7 @@ bool FPCGSpawnActorElement::ExecuteInternal(FPCGContext* Context) const
 					FPCGISMCBuilderParameters Params;
 					Params.Mesh = Mesh;
 
-					ISMC = UPCGActorHelpers::GetOrCreateISMC(TargetActor, Context->SourceComponent, Params);
+					ISMC = UPCGActorHelpers::GetOrCreateISMC(TargetActor, Context->SourceComponent.Get(), Params);
 					UEngine::CopyPropertiesForUnrelatedObjects(FirstSMC, ISMC);
 				}
 				else
@@ -186,7 +186,7 @@ bool FPCGSpawnActorElement::ExecuteInternal(FPCGContext* Context) const
 					TArray<FName> NewActorTags = GetNewActorTags(Context, TargetActor, Settings->bInheritActorTags, Settings->TagsToAddOnActors);
 
 					// Create managed resource for actor tracking
-					UPCGManagedActors* ManagedActors = NewObject<UPCGManagedActors>(Context->SourceComponent);
+					UPCGManagedActors* ManagedActors = NewObject<UPCGManagedActors>(Context->SourceComponent.Get());
 
 					for (const FPCGPoint& Point : Points)
 					{
@@ -255,7 +255,7 @@ TArray<FName> FPCGSpawnActorElement::GetNewActorTags(FPCGContext* Context, AActo
 		// and find the original actor tags
 		if (APCGPartitionActor* PartitionActor = Cast<APCGPartitionActor>(TargetActor))
 		{
-			if (UPCGComponent* OriginalComponent = PartitionActor->GetOriginalComponent(Context->SourceComponent))
+			if (UPCGComponent* OriginalComponent = PartitionActor->GetOriginalComponent(Context->SourceComponent.Get()))
 			{
 				check(OriginalComponent->GetOwner());
 				NewActorTags = OriginalComponent->GetOwner()->Tags;
diff --git a/Engine/Plugins/Experimental/PCG/Source/PCG/Private/Elements/PCGStaticMeshSpawner.cpp b/Engine/Plugins/Experimental/PCG/Source/PCG/Private/Elements/PCGStaticMeshSpawner.cpp
--- a/Engine/Plugins/Experimental/PCG/Source/PCG/Private/Elements/PCGStaticMeshSpawner.cpp
+++ b/Engine/Plugins/Experimental/PCG/Source/PCG/Private/Elements/PCGStaticMeshSpawner.cpp
@@ -145,7 +145,7 @@ void FPCGStaticMeshSpawnerElement::SpawnStaticMeshInstances(FPCGContext* Context
 
 	Params.NumCustomDataFloats = PackedCustomData.NumCustomDataFloats;
 
-	UInstancedStaticMeshComponent* ISMC = UPCGActorHelpers::GetOrCreateISMC(TargetActor, Context->SourceComponent, Params);
+	UInstancedStaticMeshComponent* ISMC = UPCGActorHelpers::GetOrCreateISMC(TargetActor, Context->SourceComponent.Get(), Params);
 
 	const int32 PreExistingInstanceCount = ISMC->GetInstanceCount();
 	const int32 NewInstanceCount = InstanceList.Instances.Num();
diff --git a/Engine/Plugins/Experimental/PCG/Source/PCG/Private/Graph/PCGGraphExecutor.cpp b/Engine/Plugins/Experimental/PCG/Source/PCG/Private/Graph/PCGGraphExecutor.cpp
--- a/Engine/Plugins/Experimental/PCG/Source/PCG/Private/Graph/PCGGraphExecutor.cpp
+++ b/Engine/Plugins/Experimental/PCG/Source/PCG/Private/Graph/PCGGraphExecutor.cpp
@@ -115,7 +115,7 @@ FPCGTaskId FPCGGraphExecutor::Schedule(UPCGGraph* Graph, UPCGComponent* SourceCo
 	return ScheduledId;
 }
 
-FPCGTaskId FPCGGraphExecutor::ScheduleGeneric(TFunction<bool()> InOperation, const TArray<FPCGTaskId>& TaskDependencies)
+FPCGTaskId FPCGGraphExecutor::ScheduleGeneric(TFunction<bool()> InOperation, UPCGComponent* InSourceComponent, const TArray<FPCGTaskId>& TaskDependencies)
 {
 	// Build task & element to hold the operation to perform
 	FPCGGraphTask Task;
@@ -123,6 +123,8 @@ FPCGTaskId FPCGGraphExecutor::ScheduleGeneric(TFunction<bool()> InOperation, con
 	{
 		Task.Inputs.Emplace(TaskDependency, nullptr, nullptr);
 	}
+
+	Task.SourceComponent = InSourceComponent;
 	Task.Element = MakeShared<FPCGGenericElement>(InOperation);
 
 	ScheduleLock.Lock();
@@ -323,7 +325,7 @@ void FPCGGraphExecutor::Execute()
 				// there is an execution mode that would prevent us from doing so.
 				const UPCGSettings* TaskSettings = PCGContextHelpers::GetInputSettings<UPCGSettings>(Task.Node, TaskInput);
 				FPCGDataCollection CachedOutput;
-				const bool bResultAlreadyInCache = Task.Element->IsCacheable(TaskSettings) && GraphCache.GetFromCache(Task.Element.Get(), TaskInput, TaskSettings, Task.SourceComponent, CachedOutput);
+				const bool bResultAlreadyInCache = Task.Element->IsCacheable(TaskSettings) && GraphCache.GetFromCache(Task.Element.Get(), TaskInput, TaskSettings, Task.SourceComponent.Get(), CachedOutput);
 #if WITH_EDITOR
 				const bool bNeedsToCreateActiveTask = !bResultAlreadyInCache || TaskSettings->ExecutionMode == EPCGSettingsExecutionMode::Debug || TaskSettings->ExecutionMode == EPCGSettingsExecutionMode::Isolated;
 #else
@@ -334,7 +336,7 @@ void FPCGGraphExecutor::Execute()
 				{
 #if WITH_EDITOR
 					// doing this now since we're about to modify ReadyTasks potentially reallocating while Task is a reference. 
-					UPCGComponent* SourceComponent = Task.SourceComponent;
+					UPCGComponent* SourceComponent = Task.SourceComponent.Get();
 					if (SourceComponent && SourceComponent->IsInspecting())
 					{
 						SourceComponent->StoreInspectionData(Task.Node, CachedOutput);
@@ -438,7 +440,7 @@ void FPCGGraphExecutor::Execute()
 				const UPCGSettings* ActiveTaskSettings = ActiveTask.Context->GetInputSettings<UPCGSettings>();
 				if (ActiveTaskSettings && ActiveTask.Element->IsCacheable(ActiveTaskSettings))
 				{
-					GraphCache.StoreInCache(ActiveTask.Element.Get(), ActiveTask.Context->InputData, ActiveTaskSettings, ActiveTask.Context->SourceComponent, ActiveTask.Context->OutputData);
+					GraphCache.StoreInCache(ActiveTask.Element.Get(), ActiveTask.Context->InputData, ActiveTaskSettings, ActiveTask.Context->SourceComponent.Get(), ActiveTask.Context->OutputData);
 				}
 			}
 
@@ -450,7 +452,7 @@ void FPCGGraphExecutor::Execute()
 			// Additional note: this needs to be executed before the StoreResults since debugging might cancel further tasks
 			ActiveTask.Element->DebugDisplay(ActiveTask.Context.Get());
 
-			UPCGComponent* SourceComponent = ActiveTask.Context->SourceComponent;
+			UPCGComponent* SourceComponent = ActiveTask.Context->SourceComponent.Get();
 			if (SourceComponent && SourceComponent->IsInspecting())
 			{
 				SourceComponent->StoreInspectionData(ActiveTask.Context->Node, ActiveTask.Context->OutputData);
@@ -696,14 +698,14 @@ FPCGTaskId FPCGGraphExecutor::ScheduleDebugWithTaskCallback(UPCGComponent* InCom
 			}
 
 			return true;
-		}, {CompiledTask.NodeId});
+		}, InComponent, {CompiledTask.NodeId});
 
 		// Add these tasks to the final dependencies
 		FinalDependencies.Add(CaptureTaskId);
 	}
 
 	// Finally, add a task to wait on the graph itself plus the capture tasks
-	return ScheduleGeneric([] { return true; }, FinalDependencies);
+	return ScheduleGeneric([] { return true; }, InComponent, FinalDependencies);
 }
 
 void FPCGGraphExecutor::AddToDirtyActors(AActor* Actor)
@@ -791,12 +793,20 @@ FTextFormat FPCGGraphExecutor::GetNotificationTextFormat()
 
 bool FPCGFetchInputElement::ExecuteInternal(FPCGContext* Context) const
 {
+	check(Context);
 	// First: any input can be passed through to the output trivially
 	Context->OutputData = Context->InputData;
 
 	// Second: fetch the inputs provided by the component
-	UPCGComponent* Component = Context->SourceComponent;
-	check(Component);
+	UPCGComponent* Component = Context->SourceComponent.Get();
+	
+	// Early out if the component has been deleted/is invalid
+	if (!Component)
+	{
+		// If the component should exist but it doesn't (which is all the time here, previously we checked for it), then this should be cancelled
+		Context->OutputData.bCancelExecution = true;
+		return true;
+	}
 
 	check(Context->Node);
 
diff --git a/Engine/Plugins/Experimental/PCG/Source/PCG/Private/Graph/PCGGraphExecutor.h b/Engine/Plugins/Experimental/PCG/Source/PCG/Private/Graph/PCGGraphExecutor.h
--- a/Engine/Plugins/Experimental/PCG/Source/PCG/Private/Graph/PCGGraphExecutor.h
+++ b/Engine/Plugins/Experimental/PCG/Source/PCG/Private/Graph/PCGGraphExecutor.h
@@ -39,7 +39,7 @@ struct FPCGGraphTask
 	TArray<FPCGGraphTaskInput> Inputs;
 	//TArray<DataId> Outputs;
 	const UPCGNode* Node = nullptr;
-	UPCGComponent* SourceComponent = nullptr;
+	TWeakObjectPtr<UPCGComponent> SourceComponent = nullptr;
 	FPCGElementPtr Element; // Added to have tasks that aren't node-bound
 	FPCGTaskId NodeId = InvalidPCGTaskId;
 };
@@ -73,7 +73,7 @@ public:
 	FPCGTaskId Schedule(UPCGGraph* Graph, UPCGComponent* InSourceComponent, FPCGElementPtr InputElement, const TArray<FPCGTaskId>& TaskDependency);
 
 	/** General job scheduling, used to control loading/unloading */
-	FPCGTaskId ScheduleGeneric(TFunction<bool()> InOperation, const TArray<FPCGTaskId>& TaskDependencies);
+	FPCGTaskId ScheduleGeneric(TFunction<bool()> InOperation, UPCGComponent* InSourceComponent, const TArray<FPCGTaskId>& TaskDependencies);
 
 	/** Gets data in the output results. Returns false if data is not ready. */
 	bool GetOutputData(FPCGTaskId InTaskId, FPCGDataCollection& OutData);
diff --git a/Engine/Plugins/Experimental/PCG/Source/PCG/Private/Helpers/PCGBlueprintHelpers.cpp b/Engine/Plugins/Experimental/PCG/Source/PCG/Private/Helpers/PCGBlueprintHelpers.cpp
--- a/Engine/Plugins/Experimental/PCG/Source/PCG/Private/Helpers/PCGBlueprintHelpers.cpp
+++ b/Engine/Plugins/Experimental/PCG/Source/PCG/Private/Helpers/PCGBlueprintHelpers.cpp
@@ -39,35 +39,35 @@ FRandomStream UPCGBlueprintHelpers::GetRandomStream(const FPCGPoint& InPoint, co
 
 UPCGData* UPCGBlueprintHelpers::GetActorData(FPCGContext& Context)
 {
-	return Context.SourceComponent ? Context.SourceComponent->GetActorPCGData() : nullptr;
+	return Context.SourceComponent.IsValid() ? Context.SourceComponent->GetActorPCGData() : nullptr;
 }
 
 UPCGData* UPCGBlueprintHelpers::GetInputData(FPCGContext& Context)
 {
-	return Context.SourceComponent ? Context.SourceComponent->GetInputPCGData() : nullptr;
+	return Context.SourceComponent.IsValid() ? Context.SourceComponent->GetInputPCGData() : nullptr;
 }
 
 TArray<UPCGData*> UPCGBlueprintHelpers::GetExclusionData(FPCGContext& Context)
 {
-	return Context.SourceComponent ? Context.SourceComponent->GetPCGExclusionData() : TArray<UPCGData*>();
+	return Context.SourceComponent.IsValid() ? Context.SourceComponent->GetPCGExclusionData() : TArray<UPCGData*>();
 }
 
 UPCGComponent* UPCGBlueprintHelpers::GetComponent(FPCGContext& Context)
 {
-	return Context.SourceComponent;
+	return Context.SourceComponent.Get();
 }
 
 UPCGComponent* UPCGBlueprintHelpers::GetOriginalComponent(FPCGContext& Context)
 {
-	if (Context.SourceComponent &&
+	if (Context.SourceComponent.IsValid() &&
 		Cast<APCGPartitionActor>(Context.SourceComponent->GetOwner()) &&
-		Cast<APCGPartitionActor>(Context.SourceComponent->GetOwner())->GetOriginalComponent(Context.SourceComponent))
+		Cast<APCGPartitionActor>(Context.SourceComponent->GetOwner())->GetOriginalComponent(Context.SourceComponent.Get()))
 	{
-		return Cast<APCGPartitionActor>(Context.SourceComponent->GetOwner())->GetOriginalComponent(Context.SourceComponent);
+		return Cast<APCGPartitionActor>(Context.SourceComponent->GetOwner())->GetOriginalComponent(Context.SourceComponent.Get());
 	}
 	else
 	{
-		return Context.SourceComponent;
+		return Context.SourceComponent.Get();
 	}
 }
 
diff --git a/Engine/Plugins/Experimental/PCG/Source/PCG/Private/MeshSelectors/PCGMeshSelectorWeighted.cpp b/Engine/Plugins/Experimental/PCG/Source/PCG/Private/MeshSelectors/PCGMeshSelectorWeighted.cpp
--- a/Engine/Plugins/Experimental/PCG/Source/PCG/Private/MeshSelectors/PCGMeshSelectorWeighted.cpp
+++ b/Engine/Plugins/Experimental/PCG/Source/PCG/Private/MeshSelectors/PCGMeshSelectorWeighted.cpp
@@ -91,7 +91,7 @@ void UPCGMeshSelectorWeighted::SelectInstances_Implementation(
 				continue;
 			}
 
-			FRandomStream RandomSource = UPCGBlueprintHelpers::GetRandomStream(Point, Settings, Context.SourceComponent);
+			FRandomStream RandomSource = UPCGBlueprintHelpers::GetRandomStream(Point, Settings, Context.SourceComponent.Get());
 			int RandomWeightedPick = RandomSource.RandRange(0, TotalWeight - 1);
 
 			int RandomPick = 0;
diff --git a/Engine/Plugins/Experimental/PCG/Source/PCG/Private/MeshSelectors/PCGMeshSelectorWeightedByCategory.cpp b/Engine/Plugins/Experimental/PCG/Source/PCG/Private/MeshSelectors/PCGMeshSelectorWeightedByCategory.cpp
--- a/Engine/Plugins/Experimental/PCG/Source/PCG/Private/MeshSelectors/PCGMeshSelectorWeightedByCategory.cpp
+++ b/Engine/Plugins/Experimental/PCG/Source/PCG/Private/MeshSelectors/PCGMeshSelectorWeightedByCategory.cpp
@@ -185,7 +185,7 @@ void UPCGMeshSelectorWeightedByCategory::SelectInstances_Implementation(
 
 		const int TotalWeight = InstancesAndWeights->CumulativeWeights.Last();
 
-		FRandomStream RandomSource = UPCGBlueprintHelpers::GetRandomStream(Point, Settings, Context.SourceComponent);
+		FRandomStream RandomSource = UPCGBlueprintHelpers::GetRandomStream(Point, Settings, Context.SourceComponent.Get());
 		int RandomWeightedPick = RandomSource.RandRange(0, TotalWeight - 1);
 
 		int RandomPick = 0;
diff --git a/Engine/Plugins/Experimental/PCG/Source/PCG/Private/PCGComponent.cpp b/Engine/Plugins/Experimental/PCG/Source/PCG/Private/PCGComponent.cpp
--- a/Engine/Plugins/Experimental/PCG/Source/PCG/Private/PCGComponent.cpp
+++ b/Engine/Plugins/Experimental/PCG/Source/PCG/Private/PCGComponent.cpp
@@ -609,7 +609,7 @@ FPCGTaskId UPCGComponent::CreateCleanupTask(bool bRemoveComponents, const TArray
 		return true;
 	};
 
-	return GetSubsystem()->ScheduleGeneric(CleanupTask, *AllDependencies);
+	return GetSubsystem()->ScheduleGeneric(CleanupTask, this, *AllDependencies);
 }
 
 void UPCGComponent::CleanupUnusedManagedResources()
@@ -1648,7 +1648,7 @@ UPCGData* UPCGComponent::CreateActorPCGData(AActor* Actor, bool bParseActor)
 			// TODO: cache/share the original component's actor pcg data
 			if (const UPCGSpatialData* OriginalComponentSpatialData = Cast<const UPCGSpatialData>(OriginalComponent->GetActorPCGData()))
 			{
-				UPCGVolumeData* Data = NewObject<UPCGVolumeData>(this);
+				UPCGVolumeData* Data = NewObject<UPCGVolumeData>();
 				Data->Initialize(PartitionActor->GetFixedBounds(), PartitionActor);
 
 				return Data->IntersectWith(OriginalComponentSpatialData);
@@ -1660,14 +1660,14 @@ UPCGData* UPCGComponent::CreateActorPCGData(AActor* Actor, bool bParseActor)
 	}
 	else if (ALandscapeProxy* Landscape = Cast<ALandscapeProxy>(Actor))
 	{
-		UPCGLandscapeData* Data = NewObject<UPCGLandscapeData>(this);
+		UPCGLandscapeData* Data = NewObject<UPCGLandscapeData>();
 		Data->Initialize(Landscape, GetGridBounds(Actor), /*bHeightOnly=*/false, /*bUseMetadata=*/Graph && Graph->bLandscapeUsesMetadata);
 
 		return Data;
 	}
 	else if (AVolume* Volume = Cast<AVolume>(Actor))
 	{
-		UPCGVolumeData* Data = NewObject<UPCGVolumeData>(this);
+		UPCGVolumeData* Data = NewObject<UPCGVolumeData>();
 		Data->Initialize(Volume);
 
 		return Data;
@@ -1693,12 +1693,12 @@ UPCGData* UPCGComponent::CreateActorPCGData(AActor* Actor, bool bParseActor)
 		UPCGUnionData* Union = nullptr;
 		if (LandscapeSplines.Num() + Splines.Num() + Shapes.Num() + OtherPrimitives.Num() > 1)
 		{
-			Union = NewObject<UPCGUnionData>(this);
+			Union = NewObject<UPCGUnionData>();
 		}
 
 		for (ULandscapeSplinesComponent* SplineComponent : LandscapeSplines)
 		{
-			UPCGLandscapeSplineData* SplineData = NewObject<UPCGLandscapeSplineData>(this);
+			UPCGLandscapeSplineData* SplineData = NewObject<UPCGLandscapeSplineData>();
 			SplineData->Initialize(SplineComponent);
 
 			if (Union)
@@ -1713,7 +1713,7 @@ UPCGData* UPCGComponent::CreateActorPCGData(AActor* Actor, bool bParseActor)
 
 		for (USplineComponent* SplineComponent : Splines)
 		{
-			UPCGSplineData* SplineData = NewObject<UPCGSplineData>(this);
+			UPCGSplineData* SplineData = NewObject<UPCGSplineData>();
 			SplineData->Initialize(SplineComponent);
 
 			if (Union)
@@ -1728,7 +1728,7 @@ UPCGData* UPCGComponent::CreateActorPCGData(AActor* Actor, bool bParseActor)
 
 		for (UShapeComponent* ShapeComponent : Shapes)
 		{
-			UPCGPrimitiveData* ShapeData = NewObject<UPCGPrimitiveData>(this);
+			UPCGPrimitiveData* ShapeData = NewObject<UPCGPrimitiveData>();
 			ShapeData->Initialize(ShapeComponent);
 			
 			if (Union)
@@ -1743,7 +1743,7 @@ UPCGData* UPCGComponent::CreateActorPCGData(AActor* Actor, bool bParseActor)
 
 		for (UPrimitiveComponent* PrimitiveComponent : OtherPrimitives)
 		{
-			UPCGPrimitiveData* PrimitiveData = NewObject<UPCGPrimitiveData>(this);
+			UPCGPrimitiveData* PrimitiveData = NewObject<UPCGPrimitiveData>();
 			PrimitiveData->Initialize(PrimitiveComponent);
 
 			if (Union)
@@ -1763,7 +1763,7 @@ UPCGData* UPCGComponent::CreateActorPCGData(AActor* Actor, bool bParseActor)
 	}
 
 	// Finally, if it's not a special actor and there are not parsed components, then return a single point at the actor position
-	UPCGPointData* Data = NewObject<UPCGPointData>(this);
+	UPCGPointData* Data = NewObject<UPCGPointData>();
 	Data->InitializeFromActor(Actor);
 	return Data;
 }
@@ -1842,7 +1842,7 @@ UPCGData* UPCGComponent::CreateLandscapePCGData(bool bHeightOnly)
 	}
 
 	// TODO: we're creating separate landscape data instances here so we can do some tweaks on it (such as storing the right target actor) but this probably should change
-	UPCGLandscapeData* LandscapeData = NewObject<UPCGLandscapeData>(this);
+	UPCGLandscapeData* LandscapeData = NewObject<UPCGLandscapeData>();
 	LandscapeData->Initialize(Landscape, GetGridBounds(Landscape), bHeightOnly, /*bUseMetadata=*/Graph && Graph->bLandscapeUsesMetadata);
 	// Need to override target actor for this one, not the landscape
 	LandscapeData->TargetActor = Actor;
@@ -2261,7 +2261,7 @@ void FPCGComponentInstanceData::DelayedRefresh(UPCGComponent* PCGComponent)
 					}
 
 					return true;
-				}, {});
+				}, PCGComponent, {});
 		}
 	}
 }
diff --git a/Engine/Plugins/Experimental/PCG/Source/PCG/Private/PCGContext.cpp b/Engine/Plugins/Experimental/PCG/Source/PCG/Private/PCGContext.cpp
--- a/Engine/Plugins/Experimental/PCG/Source/PCG/Private/PCGContext.cpp
+++ b/Engine/Plugins/Experimental/PCG/Source/PCG/Private/PCGContext.cpp
@@ -30,5 +30,5 @@ FString FPCGContext::GetTaskName() const
 
 FString FPCGContext::GetComponentName() const
 {
-	return SourceComponent && SourceComponent->GetOwner() ? SourceComponent->GetOwner()->GetFName().ToString() : TEXT("Non-PCG Component");
+	return SourceComponent.IsValid() && SourceComponent->GetOwner() ? SourceComponent->GetOwner()->GetFName().ToString() : TEXT("Non-PCG Component");
 }
\ No newline at end of file
diff --git a/Engine/Plugins/Experimental/PCG/Source/PCG/Private/PCGElement.cpp b/Engine/Plugins/Experimental/PCG/Source/PCG/Private/PCGElement.cpp
--- a/Engine/Plugins/Experimental/PCG/Source/PCG/Private/PCGElement.cpp
+++ b/Engine/Plugins/Experimental/PCG/Source/PCG/Private/PCGElement.cpp
@@ -15,7 +15,7 @@ bool IPCGElement::Execute(FPCGContext* Context) const
 	check(Context && Context->NumAvailableTasks > 0);
 
 	// Early out to stop execution
-	if (Context->InputData.bCancelExecution || (Context->SourceComponent && !IsValid(Context->SourceComponent)))
+	if (Context->InputData.bCancelExecution || (!Context->SourceComponent.IsExplicitlyNull() && !Context->SourceComponent.IsValid()))
 	{
 		Context->OutputData.bCancelExecution = true;
 
@@ -217,7 +217,7 @@ void IPCGElement::CleanupAndValidateOutput(FPCGContext* Context) const
 	}
 }
 
-FPCGContext* FSimplePCGElement::Initialize(const FPCGDataCollection& InputData, UPCGComponent* SourceComponent, const UPCGNode* Node)
+FPCGContext* FSimplePCGElement::Initialize(const FPCGDataCollection& InputData, TWeakObjectPtr<UPCGComponent> SourceComponent, const UPCGNode* Node)
 {
 	FPCGContext* Context = new FPCGContext();
 	Context->InputData = InputData;
diff --git a/Engine/Plugins/Experimental/PCG/Source/PCG/Private/PCGSubgraph.cpp b/Engine/Plugins/Experimental/PCG/Source/PCG/Private/PCGSubgraph.cpp
--- a/Engine/Plugins/Experimental/PCG/Source/PCG/Private/PCGSubgraph.cpp
+++ b/Engine/Plugins/Experimental/PCG/Source/PCG/Private/PCGSubgraph.cpp
@@ -139,7 +139,7 @@ TObjectPtr<UPCGGraph> UPCGSubgraphNode::GetSubgraph() const
 	return Settings ? Settings->Subgraph : nullptr;
 }
 
-FPCGContext* FPCGSubgraphElement::Initialize(const FPCGDataCollection& InputData, UPCGComponent* SourceComponent, const UPCGNode* Node)
+FPCGContext* FPCGSubgraphElement::Initialize(const FPCGDataCollection& InputData, TWeakObjectPtr<UPCGComponent> SourceComponent, const UPCGNode* Node)
 {
 	FPCGSubgraphContext* Context = new FPCGSubgraphContext();
 	Context->InputData = InputData;
@@ -164,13 +164,13 @@ bool FPCGSubgraphElement::ExecuteInternal(FPCGContext* InContext) const
 			check(Settings);
 			UPCGGraph* Subgraph = Settings->Subgraph;
 
-			UPCGSubsystem* Subsystem = Context->SourceComponent->GetSubsystem();
+			UPCGSubsystem* Subsystem = Context->SourceComponent.IsValid() ? Context->SourceComponent->GetSubsystem() : nullptr;
 
 			if (Subsystem && Subgraph)
 			{
 				// Dispatch graph to execute with the given information we have
 				// using this node's task id as additional inputs
-				FPCGTaskId SubgraphTaskId = Subsystem->ScheduleGraph(Subgraph, Context->SourceComponent, MakeShared<FPCGInputForwardingElement>(Context->InputData), {});
+				FPCGTaskId SubgraphTaskId = Subsystem->ScheduleGraph(Subgraph, Context->SourceComponent.Get(), MakeShared<FPCGInputForwardingElement>(Context->InputData), {});
 
 				Context->SubgraphTaskId = SubgraphTaskId;
 				Context->bScheduledSubgraph = true;
@@ -181,7 +181,7 @@ bool FPCGSubgraphElement::ExecuteInternal(FPCGContext* InContext) const
 					// Wake up the current task
 					Context->bIsPaused = false;
 					return true;
-					}, { Context->SubgraphTaskId });
+					}, Context->SourceComponent.Get(), { Context->SubgraphTaskId });
 
 				return false;
 			}
diff --git a/Engine/Plugins/Experimental/PCG/Source/PCG/Private/PCGSubsystem.cpp b/Engine/Plugins/Experimental/PCG/Source/PCG/Private/PCGSubsystem.cpp
--- a/Engine/Plugins/Experimental/PCG/Source/PCG/Private/PCGSubsystem.cpp
+++ b/Engine/Plugins/Experimental/PCG/Source/PCG/Private/PCGSubsystem.cpp
@@ -243,7 +243,7 @@ FPCGTaskId UPCGSubsystem::ScheduleComponent(UPCGComponent* PCGComponent, bool bS
 			}
 
 			return true;
-			}, AllTasks);
+			}, PCGComponent, AllTasks);
 	}
 	else
 	{
@@ -326,7 +326,7 @@ FPCGTaskId UPCGSubsystem::ScheduleCleanup(UPCGComponent* PCGComponent, bool bRem
 	}
 	else
 	{
-		PostCleanupTaskId = GraphExecutor->ScheduleGeneric(PostCleanupTask, AllTasks);
+		PostCleanupTaskId = GraphExecutor->ScheduleGeneric(PostCleanupTask, PCGComponent, AllTasks);
 	}
 
 	return PostCleanupTaskId;
@@ -366,10 +366,10 @@ TStatId UPCGSubsystem::GetStatId() const
 	RETURN_QUICK_DECLARE_CYCLE_STAT(UPCGSubsystem, STATGROUP_Tickables);
 }
 
-FPCGTaskId UPCGSubsystem::ScheduleGeneric(TFunction<bool()> InOperation, const TArray<FPCGTaskId>& TaskDependencies)
+FPCGTaskId UPCGSubsystem::ScheduleGeneric(TFunction<bool()> InOperation, UPCGComponent* Component, const TArray<FPCGTaskId>& TaskDependencies)
 {
 	check(GraphExecutor);
-	return GraphExecutor->ScheduleGeneric(InOperation, TaskDependencies);
+	return GraphExecutor->ScheduleGeneric(InOperation, Component, TaskDependencies);
 }
 
 TArray<FPCGTaskId> UPCGSubsystem::DispatchToRegisteredLocalComponents(UPCGComponent* OriginalComponent, const TFunction<FPCGTaskId(UPCGComponent*)>& InFunc) const
@@ -1000,7 +1000,7 @@ namespace PCGSubsystem
 						return true;
 					};
 
-					FPCGTaskId LoadTaskId = GraphExecutor->ScheduleGeneric(LoadActorsTask, {});
+					FPCGTaskId LoadTaskId = GraphExecutor->ScheduleGeneric(LoadActorsTask, nullptr, {});
 					SetPreviousTaskIfValid(LoadTaskId);
 				}
 
@@ -1016,7 +1016,7 @@ namespace PCGSubsystem
 						return true;
 					};
 
-					FPCGTaskId SaveTaskId = GraphExecutor->ScheduleGeneric(SaveActorTask, PreviousTasks);
+					FPCGTaskId SaveTaskId = GraphExecutor->ScheduleGeneric(SaveActorTask, nullptr, PreviousTasks);
 					SetPreviousTaskIfValid(SaveTaskId);
 				}
 
@@ -1027,7 +1027,7 @@ namespace PCGSubsystem
 				};
 
 				// Schedule after the save (if valid), then the execute so we can queue this after the load.
-				FPCGTaskId UnloadTaskId = GraphExecutor->ScheduleGeneric(UnloadActorsTask, PreviousTasks);
+				FPCGTaskId UnloadTaskId = GraphExecutor->ScheduleGeneric(UnloadActorsTask, nullptr, PreviousTasks);
 				SetPreviousTaskIfValid(UnloadTaskId);
 
 				// Finally, mark "last" valid task in the cell tasks.
@@ -1044,7 +1044,7 @@ namespace PCGSubsystem
 		// Finally, create a dummy generic task to wait on all cells
 		if (!CellTasks.IsEmpty())
 		{
-			return GraphExecutor->ScheduleGeneric([]() { return true; }, CellTasks);
+			return GraphExecutor->ScheduleGeneric([]() { return true; }, nullptr, CellTasks);
 		}
 		else
 		{
@@ -1126,9 +1126,9 @@ FPCGTaskId UPCGSubsystem::ProcessGraph(UPCGComponent* Component, const FBox& InP
 		switch (InOperation)
 		{
 		case EOperation::Unpartition:
-			return GraphExecutor->ScheduleGeneric(UnpartitionTask, TaskDependencies);
+			return GraphExecutor->ScheduleGeneric(UnpartitionTask, ComponentPtr.Get(), TaskDependencies);
 		case EOperation::Partition:
-			return GraphExecutor->ScheduleGeneric(PartitionTask, TaskDependencies);
+			return GraphExecutor->ScheduleGeneric(PartitionTask, ComponentPtr.Get(), TaskDependencies);
 		case EOperation::Generate:
 		default:
 			return ScheduleGraph();
@@ -1159,7 +1159,7 @@ FPCGTaskId UPCGSubsystem::ProcessGraph(UPCGComponent* Component, const FBox& InP
 			return true;
 		};
 
-		return GraphExecutor->ScheduleGeneric(PostProcessGraph, { ProcessAllCellsTaskId });
+		return GraphExecutor->ScheduleGeneric(PostProcessGraph, Component, { ProcessAllCellsTaskId });
 	}
 	else
 	{
@@ -1210,7 +1210,7 @@ FPCGTaskId UPCGSubsystem::CleanupGraph(UPCGComponent* Component, const FBox& InB
 			return true;
 		};
 
-		return GraphExecutor->ScheduleGeneric(PostCleanupGraph, { ProcessAllCellsTaskId });
+		return GraphExecutor->ScheduleGeneric(PostCleanupGraph, Component, { ProcessAllCellsTaskId });
 	}
 	else
 	{
@@ -1247,7 +1247,7 @@ void UPCGSubsystem::CleanupPartitionActors(const FBox& InBounds)
 			return true;
 		};
 
-		return GraphExecutor->ScheduleGeneric(CleanupTask, TaskDependencies);
+		return GraphExecutor->ScheduleGeneric(CleanupTask, nullptr, TaskDependencies);
 	};
 
 	PCGSubsystem::ForEachIntersectingCell(GraphExecutor, GetWorld(), InBounds, /*bCreateActor=*/false, /*bLoadCell=*/false, /*bSave=*/false, ScheduleTask);
@@ -1270,7 +1270,7 @@ void UPCGSubsystem::ClearPCGLink(UPCGComponent* InComponent, const FBox& InBound
 			return true;
 		};
 
-		return GraphExecutor->ScheduleGeneric(MoveTask, TaskDependencies);
+		return GraphExecutor->ScheduleGeneric(MoveTask, ComponentPtr.Get(), TaskDependencies);
 	};
 
 	FPCGTaskId TaskId = PCGSubsystem::ForEachIntersectingCell(GraphExecutor, GetWorld(), InBounds, /*bCreateActor=*/false, /*bLoadCell=*/false, /*bSave=*/false, ScheduleTask);
@@ -1311,7 +1311,7 @@ void UPCGSubsystem::ClearPCGLink(UPCGComponent* InComponent, const FBox& InBound
 			return true;
 		};
 
-		GraphExecutor->ScheduleGeneric(CleanupTask, { TaskId });
+		GraphExecutor->ScheduleGeneric(CleanupTask, InComponent, { TaskId });
 	}
 	else
 	{
diff --git a/Engine/Plugins/Experimental/PCG/Source/PCG/Private/Tests/Determinism/PCGDeterminismTestsCommon.cpp b/Engine/Plugins/Experimental/PCG/Source/PCG/Private/Tests/Determinism/PCGDeterminismTestsCommon.cpp
--- a/Engine/Plugins/Experimental/PCG/Source/PCG/Private/Tests/Determinism/PCGDeterminismTestsCommon.cpp
+++ b/Engine/Plugins/Experimental/PCG/Source/PCG/Private/Tests/Determinism/PCGDeterminismTestsCommon.cpp
@@ -94,7 +94,7 @@ namespace PCGDeterminismTests
 
 		FPCGGraphExecutor Executor = FPCGGraphExecutor(PCGComponentCopy);
 
-		auto ScheduleAndWaitForExecution = [&Executor](FPCGTaskId FinalTaskID)
+		auto ScheduleAndWaitForExecution = [&Executor, PCGComponentCopy](FPCGTaskId FinalTaskID)
 		{
 			// TODO: Consider randomizing/iterating through possible input orders
 			volatile bool bTasksComplete = false;
@@ -107,7 +107,7 @@ namespace PCGDeterminismTests
 			{
 				bTasksComplete = true;
 				return true;
-			}, {FinalTaskID});
+			}, PCGComponentCopy, {FinalTaskID});
 
 			// Run first iteration
 			while (!bTasksComplete)
diff --git a/Engine/Plugins/Experimental/PCG/Source/PCG/Public/Elements/PCGExecuteBlueprint.h b/Engine/Plugins/Experimental/PCG/Source/PCG/Public/Elements/PCGExecuteBlueprint.h
--- a/Engine/Plugins/Experimental/PCG/Source/PCG/Public/Elements/PCGExecuteBlueprint.h
+++ b/Engine/Plugins/Experimental/PCG/Source/PCG/Public/Elements/PCGExecuteBlueprint.h
@@ -240,5 +240,5 @@ public:
 
 protected:
 	virtual bool ExecuteInternal(FPCGContext* Context) const override;
-	virtual FPCGContext* Initialize(const FPCGDataCollection& InputData, UPCGComponent* SourceComponent, const UPCGNode* Node) override;	
+	virtual FPCGContext* Initialize(const FPCGDataCollection& InputData, TWeakObjectPtr<UPCGComponent> SourceComponent, const UPCGNode* Node) override;	
 };
diff --git a/Engine/Plugins/Experimental/PCG/Source/PCG/Public/Helpers/PCGSettingsHelpers.h b/Engine/Plugins/Experimental/PCG/Source/PCG/Public/Helpers/PCGSettingsHelpers.h
--- a/Engine/Plugins/Experimental/PCG/Source/PCG/Public/Helpers/PCGSettingsHelpers.h
+++ b/Engine/Plugins/Experimental/PCG/Source/PCG/Public/Helpers/PCGSettingsHelpers.h
@@ -220,6 +220,11 @@ namespace PCGSettingsHelpers
 	}
 
 	int ComputeSeedWithOverride(const UPCGSettings* InSettings, const UPCGComponent* InComponent, UPCGParamData* InParams);
+
+	FORCEINLINE int ComputeSeedWithOverride(const UPCGSettings* InSettings, TWeakObjectPtr<UPCGComponent> InComponent, UPCGParamData* InParams)
+	{
+		return ComputeSeedWithOverride(InSettings, InComponent.Get(), InParams);
+	}
 }
 
 #define PCG_GET_OVERRIDEN_VALUE(Settings, Variable, Params) PCGSettingsHelpers::GetValue(GET_MEMBER_NAME_CHECKED(TRemovePointer<TRemoveConst<decltype(Settings)>::Type>::Type, Variable), (Settings)->Variable, Params);
\ No newline at end of file
diff --git a/Engine/Plugins/Experimental/PCG/Source/PCG/Public/PCGContext.h b/Engine/Plugins/Experimental/PCG/Source/PCG/Public/PCGContext.h
--- a/Engine/Plugins/Experimental/PCG/Source/PCG/Public/PCGContext.h
+++ b/Engine/Plugins/Experimental/PCG/Source/PCG/Public/PCGContext.h
@@ -36,7 +36,7 @@ struct PCG_API FPCGContext
 
 	FPCGDataCollection InputData;
 	FPCGDataCollection OutputData;
-	UPCGComponent* SourceComponent = nullptr;
+	TWeakObjectPtr<UPCGComponent> SourceComponent = nullptr;
 	int32 NumAvailableTasks = 0;
 
 	// TODO: add RNG source
diff --git a/Engine/Plugins/Experimental/PCG/Source/PCG/Public/PCGElement.h b/Engine/Plugins/Experimental/PCG/Source/PCG/Public/PCGElement.h
--- a/Engine/Plugins/Experimental/PCG/Source/PCG/Public/PCGElement.h
+++ b/Engine/Plugins/Experimental/PCG/Source/PCG/Public/PCGElement.h
@@ -34,7 +34,7 @@ class PCG_API IPCGElement
 {
 public:
 	virtual ~IPCGElement() = default;
-	virtual FPCGContext* Initialize(const FPCGDataCollection& InputData, UPCGComponent* SourceComponent, const UPCGNode* Node) = 0;
+	virtual FPCGContext* Initialize(const FPCGDataCollection& InputData, TWeakObjectPtr<UPCGComponent> SourceComponent, const UPCGNode* Node) = 0;
 
 	virtual bool CanExecuteOnlyOnMainThread(const UPCGSettings* InSettings) const { return false; }
 	virtual bool IsCacheable(const UPCGSettings* InSettings) const { return true; }
@@ -77,5 +77,5 @@ private:
 class PCG_API FSimplePCGElement : public IPCGElement
 {
 public:
-	virtual FPCGContext* Initialize(const FPCGDataCollection& InputData, UPCGComponent* SourceComponent, const UPCGNode* Node) override;
+	virtual FPCGContext* Initialize(const FPCGDataCollection& InputData, TWeakObjectPtr<UPCGComponent> SourceComponent, const UPCGNode* Node) override;
 };
\ No newline at end of file
diff --git a/Engine/Plugins/Experimental/PCG/Source/PCG/Public/PCGSubgraph.h b/Engine/Plugins/Experimental/PCG/Source/PCG/Public/PCGSubgraph.h
--- a/Engine/Plugins/Experimental/PCG/Source/PCG/Public/PCGSubgraph.h
+++ b/Engine/Plugins/Experimental/PCG/Source/PCG/Public/PCGSubgraph.h
@@ -105,7 +105,7 @@ struct PCG_API FPCGSubgraphContext : public FPCGContext
 class PCG_API FPCGSubgraphElement : public IPCGElement
 {
 public:
-	virtual FPCGContext* Initialize(const FPCGDataCollection& InputData, UPCGComponent* SourceComponent, const UPCGNode* Node) override;
+	virtual FPCGContext* Initialize(const FPCGDataCollection& InputData, TWeakObjectPtr<UPCGComponent> SourceComponent, const UPCGNode* Node) override;
 
 protected:
 	virtual bool ExecuteInternal(FPCGContext* Context) const override;
diff --git a/Engine/Plugins/Experimental/PCG/Source/PCG/Public/PCGSubsystem.h b/Engine/Plugins/Experimental/PCG/Source/PCG/Public/PCGSubsystem.h
--- a/Engine/Plugins/Experimental/PCG/Source/PCG/Public/PCGSubsystem.h
+++ b/Engine/Plugins/Experimental/PCG/Source/PCG/Public/PCGSubsystem.h
@@ -71,7 +71,7 @@ public:
 	FPCGTaskId ScheduleGraph(UPCGComponent* SourceComponent, const TArray<FPCGTaskId>& Dependencies);
 
 	/** General job scheduling, used to control loading/unloading */
-	FPCGTaskId ScheduleGeneric(TFunction<bool()> InOperation, const TArray<FPCGTaskId>& TaskDependencies);
+	FPCGTaskId ScheduleGeneric(TFunction<bool()> InOperation, UPCGComponent* SourceComponent, const TArray<FPCGTaskId>& TaskDependencies);
 
 	/** Gets the output data for a given task */
 	bool GetOutputData(FPCGTaskId InTaskId, FPCGDataCollection& OutData);
