[PCG] Metadata native operation: Update BreakVector to use new common class

diff --git a/Engine/Plugins/Experimental/PCG/Source/PCG/Private/Elements/Metadata/PCGMetadataBreakVector.cpp b/Engine/Plugins/Experimental/PCG/Source/PCG/Private/Elements/Metadata/PCGMetadataBreakVector.cpp
--- a/Engine/Plugins/Experimental/PCG/Source/PCG/Private/Elements/Metadata/PCGMetadataBreakVector.cpp
+++ b/Engine/Plugins/Experimental/PCG/Source/PCG/Private/Elements/Metadata/PCGMetadataBreakVector.cpp
@@ -10,265 +10,148 @@
 #include "Metadata/PCGMetadataAttribute.h"
 #include "Metadata/PCGMetadataAttributeTpl.h"
 
-TArray<FPCGPinProperties> UPCGMetadataBreakVectorSettings::InputPinProperties() const
+namespace PCGMetadataBreakVectorSettings
 {
-	TArray<FPCGPinProperties> PinProperties;
-	PinProperties.Emplace(PCGMetadataBreakVectorConstants::ParamsLabel, EPCGDataType::Param);
-	PinProperties.Emplace(PCGMetadataBreakVectorConstants::SourceLabel, EPCGDataType::Any);
-
-	return PinProperties;
-}
-
-TArray<FPCGPinProperties> UPCGMetadataBreakVectorSettings::OutputPinProperties() const
-{
-	TArray<FPCGPinProperties> PinProperties;
-	PinProperties.Emplace(PCGMetadataBreakVectorConstants::XLabel, EPCGDataType::Any);
-	PinProperties.Emplace(PCGMetadataBreakVectorConstants::YLabel, EPCGDataType::Any);
-	PinProperties.Emplace(PCGMetadataBreakVectorConstants::ZLabel, EPCGDataType::Any);
-	PinProperties.Emplace(PCGMetadataBreakVectorConstants::WLabel, EPCGDataType::Any);
-
-	return PinProperties;
-}
-
-FPCGElementPtr UPCGMetadataBreakVectorSettings::CreateElement() const
-{
-	return MakeShared<FPCGMetadataBreakVectorElement>();
-}
-
-bool FPCGMetadataBreakVectorElement::ExecuteInternal(FPCGContext* Context) const
-{
-	TRACE_CPUPROFILER_EVENT_SCOPE(FPCGMetadataBreakVectorElement::Execute);
+	template <typename InType = FVector>
+	inline void DoBreak(const InType& Value, TArray<double>& OutValues, uint32 NumOutputs)
+	{
+		check(NumOutputs == 4);
 
-	const UPCGMetadataBreakVectorSettings* Settings = Context->GetInputSettings<UPCGMetadataBreakVectorSettings>();
-	check(Settings);
+		OutValues.Add(Value.X);
+		OutValues.Add(Value.Y);
+		OutValues.Add(Value.Z);
+		OutValues.Add(0.0);
+	}
 
-	const TArray<FPCGTaggedData>& Inputs = Context->InputData.TaggedData;
-	TArray<FPCGTaggedData>& Outputs = Context->OutputData.TaggedData;
+	template <>
+	inline void DoBreak<FRotator>(const FRotator& Value, TArray<double>& OutValues, uint32 NumOutputs)
+	{
+		check(NumOutputs == 4);
 
-	const TArray<FPCGTaggedData>& ParamData = Context->InputData.GetParamsByPin(PCGMetadataBreakVectorConstants::ParamsLabel);
-	UPCGParamData* Params = nullptr;
+		OutValues.Add(Value.Roll);
+		OutValues.Add(Value.Pitch);
+		OutValues.Add(Value.Yaw);
+		OutValues.Add(0.0);
+	}
 
-	// TODO: this only captures the first Param
-	for (const FPCGTaggedData& TaggedDatum : ParamData)
+	template <>
+	inline void DoBreak<FVector2D>(const FVector2D& Value, TArray<double>& OutValues, uint32 NumOutputs)
 	{
-		if (UPCGParamData* ParamDatum = Cast<UPCGParamData>(TaggedDatum.Data))
-		{
-			Params = ParamDatum;
-			break;
-		}
+		check(NumOutputs == 4);
+
+		OutValues.Add(Value.X);
+		OutValues.Add(Value.Y);
+		OutValues.Add(0.0);
+		OutValues.Add(0.0);
 	}
 
-	const FName SourceAttributeName = PCG_GET_OVERRIDEN_VALUE(Settings, SourceAttributeName, Params);
+	template <>
+	inline void DoBreak<FVector4>(const FVector4& Value, TArray<double>& OutValues, uint32 NumOutputs)
+	{
+		check(NumOutputs == 4);
 
-	// These exist to facilitate unit testing of the MetadataBreakVector element in isolation, when it otherwise would not have any pin connections to generate output
-	bool bXPinConnected = Settings->bForceConnectX;
-	bool bYPinConnected = Settings->bForceConnectY;
-	bool bZPinConnected = Settings->bForceConnectZ;
-	bool bWPinConnected = Settings->bForceConnectW;
+		OutValues.Add(Value.X);
+		OutValues.Add(Value.Y);
+		OutValues.Add(Value.Z);
+		OutValues.Add(Value.W);
+	}
 
-	if (Context->Node)
+	inline constexpr bool IsValidType(uint16 TypeId)
 	{
-		bXPinConnected |= Context->Node->IsOutputPinConnected(PCGMetadataBreakVectorConstants::XLabel);
-		bYPinConnected |= Context->Node->IsOutputPinConnected(PCGMetadataBreakVectorConstants::YLabel);
-		bZPinConnected |= Context->Node->IsOutputPinConnected(PCGMetadataBreakVectorConstants::ZLabel);
-		bWPinConnected |= Context->Node->IsOutputPinConnected(PCGMetadataBreakVectorConstants::WLabel);
+		return TypeId == (uint16)EPCGMetadataTypes::Vector2 ||
+			TypeId == (uint16)EPCGMetadataTypes::Vector ||
+			TypeId == (uint16)EPCGMetadataTypes::Vector4 ||
+			TypeId == (uint16)EPCGMetadataTypes::Rotator;
 	}
 
-	for (const FPCGTaggedData& Input : Inputs)
+	template <typename T>
+	inline constexpr bool IsValidType()
 	{
-		if (!Input.Data)
-		{
-			PCGE_LOG(Error, "Unable to get data from input");
-			continue;
-		}
-
-		UPCGMetadata* SourceMetadata = nullptr;
-		if (const UPCGSpatialData* SpatialInput = Cast<const UPCGSpatialData>(Input.Data))
-		{
-			SourceMetadata = SpatialInput->Metadata;
-		}
-		else if (const UPCGParamData* ParamsInput = Cast<const UPCGParamData>(Input.Data))
-		{
-			SourceMetadata = ParamsInput->Metadata;
-		}
-		else // If the Data type does not handle Metadata, forward it into the valid pins
-		{
-			if (bXPinConnected)
-			{
-				FPCGTaggedData& Data = Outputs.Add_GetRef(Input);
-				Data.Pin = PCGMetadataBreakVectorConstants::XLabel;
-			}
-
-			if (bYPinConnected)
-			{
-				FPCGTaggedData& Data = Outputs.Add_GetRef(Input);
-				Data.Pin = PCGMetadataBreakVectorConstants::YLabel;
-			}
-
-			if (bZPinConnected)
-			{
-				FPCGTaggedData& Data = Outputs.Add_GetRef(Input);
-				Data.Pin = PCGMetadataBreakVectorConstants::ZLabel;
-			}
-
-			if (bWPinConnected)
-			{
-				FPCGTaggedData& Data = Outputs.Add_GetRef(Input);
-				Data.Pin = PCGMetadataBreakVectorConstants::WLabel;
-			}
-
-			continue;
-		}
-
-		if (!SourceMetadata)
-		{
-			PCGE_LOG(Warning, "Invalid metadata");
-			continue;
-		}
-
-		const FName LocalSourceAttributeName = (SourceAttributeName == NAME_None) ? SourceMetadata->GetSingleAttributeNameOrNone() : SourceAttributeName;
-		FName DestinationAttributeForX = NAME_None;
-		FName DestinationAttributeForY = NAME_None;
-		FName DestinationAttributeForZ = NAME_None;
-		FName DestinationAttributeForW = NAME_None;
-
-		if (LocalSourceAttributeName != NAME_None)
-		{
-			DestinationAttributeForX = FName(LocalSourceAttributeName.ToString() + ".X");
-			DestinationAttributeForY = FName(LocalSourceAttributeName.ToString() + ".Y");
-			DestinationAttributeForZ = FName(LocalSourceAttributeName.ToString() + ".Z");
-			DestinationAttributeForW = FName(LocalSourceAttributeName.ToString() + ".W");
-		}
-
-		const FPCGMetadataAttributeBase* SourceAttribute = SourceMetadata->GetConstAttribute(LocalSourceAttributeName);
-		if (!SourceAttribute)
-		{
-			PCGE_LOG(Error, "Attribute %s does not exist", *LocalSourceAttributeName.ToString());
-			continue;
-		}
-
-		if (SourceAttribute->GetTypeId() != PCG::Private::MetadataTypes<FVector>::Id 
-			&& SourceAttribute->GetTypeId() != PCG::Private::MetadataTypes<FVector4>::Id
-			&& SourceAttribute->GetTypeId() != PCG::Private::MetadataTypes<FRotator>::Id)
-		{
-			PCGE_LOG(Error, "Attribute %s is not a breakable type", *LocalSourceAttributeName.ToString());
-			continue;
-		}
-
-		FVector4 DefaultValue;
+		return IsValidType(PCG::Private::MetadataTypes<T>::Id);
+	}
+}
 
-		if (SourceAttribute->GetTypeId() == PCG::Private::MetadataTypes<FVector>::Id)
-		{
-			DefaultValue = static_cast<const FPCGMetadataAttribute<FVector>*>(SourceAttribute)->GetValue(PCGDefaultValueKey);
-		}
-		else if (SourceAttribute->GetTypeId() == PCG::Private::MetadataTypes<FVector4>::Id)
-		{
-			DefaultValue = static_cast<const FPCGMetadataAttribute<FVector4>*>(SourceAttribute)->GetValue(PCGDefaultValueKey);
-		}
-		else //if (SourceAttribute->GetTypeId() == PCG::Private::MetadataTypes<FRotator>::Id)
-		{
-			const FRotator Rotator = static_cast<const FPCGMetadataAttribute<FRotator>*>(SourceAttribute)->GetValue(PCGDefaultValueKey);
-			DefaultValue.X = Rotator.Roll;
-			DefaultValue.Y = Rotator.Pitch;
-			DefaultValue.Z = Rotator.Yaw;
-		}
+FName UPCGMetadataBreakVectorSettings::GetInputAttributeNameWithOverride(uint32 Index, UPCGParamData* Params) const
+{
+	return PCG_GET_OVERRIDEN_VALUE(this, InputAttributeName, Params);
+}
 
-		FPCGMetadataAttribute<double>* AttributeX = nullptr;
-		FPCGMetadataAttribute<double>* AttributeY = nullptr;
-		FPCGMetadataAttribute<double>* AttributeZ = nullptr;
-		FPCGMetadataAttribute<double>* AttributeW = nullptr;
+FName UPCGMetadataBreakVectorSettings::GetOutputPinLabel(uint32 Index) const 
+{
+	switch (Index)
+	{
+	case 0:
+		return PCGMetadataBreakVectorConstants::XLabel;
+	case 1:
+		return PCGMetadataBreakVectorConstants::YLabel;
+	case 2:
+		return PCGMetadataBreakVectorConstants::ZLabel;
+	default:
+		return PCGMetadataBreakVectorConstants::WLabel;
+	}
+}
 
-		if (bXPinConnected)
-		{
-			FPCGTaggedData& OutputX = Outputs.Add_GetRef(Input);
-			OutputX.Pin = PCGMetadataBreakVectorConstants::XLabel;
+uint32 UPCGMetadataBreakVectorSettings::GetOutputPinNum() const
+{
+	return 4;
+}
 
-			UPCGMetadata* OutMetadata = nullptr;
+bool UPCGMetadataBreakVectorSettings::IsSupportedInputType(uint16 TypeId, uint32 InputIndex, bool& bHasSpecialRequirement) const
+{
+	bHasSpecialRequirement = false;
+	return PCGMetadataBreakVectorSettings::IsValidType(TypeId);
+}
 
-			PCGMetadataElementCommon::DuplicateTaggedData(Input, OutputX, OutMetadata);
-			AttributeX = PCGMetadataElementCommon::ClearOrCreateAttribute(OutMetadata, DestinationAttributeForX, DefaultValue.X);
-			PCGMetadataElementCommon::CopyEntryToValueKeyMap(SourceMetadata, SourceAttribute, AttributeX);
-		}
+uint16 UPCGMetadataBreakVectorSettings::GetOutputType(uint16 InputTypeId) const
+{
+	return (uint16)EPCGMetadataTypes::Double;
+}
 
-		if (bYPinConnected)
-		{
-			FPCGTaggedData& OutputY = Outputs.Add_GetRef(Input);
-			OutputY.Pin = PCGMetadataBreakVectorConstants::YLabel;
+FName UPCGMetadataBreakVectorSettings::GetOutputAttributeName(FName BaseName, uint32 Index) const
+{
+	if (BaseName == NAME_None)
+	{
+		return NAME_None;
+	}
 
-			UPCGMetadata* OutMetadata = nullptr;
+	switch (Index)
+	{
+	case 0:
+		return FName(BaseName.ToString() + "." + PCGMetadataBreakVectorConstants::XLabel.ToString());
+	case 1:
+		return FName(BaseName.ToString() + "." + PCGMetadataBreakVectorConstants::YLabel.ToString());
+	case 2:
+		return FName(BaseName.ToString() + "." + PCGMetadataBreakVectorConstants::ZLabel.ToString());
+	default:
+		return FName(BaseName.ToString() + "." + PCGMetadataBreakVectorConstants::WLabel.ToString());
+	}
+}
 
-			PCGMetadataElementCommon::DuplicateTaggedData(Input, OutputY, OutMetadata);
-			AttributeY = PCGMetadataElementCommon::ClearOrCreateAttribute(OutMetadata, DestinationAttributeForY, DefaultValue.Y);
-			PCGMetadataElementCommon::CopyEntryToValueKeyMap(SourceMetadata, SourceAttribute, AttributeY);
-		}
+FPCGElementPtr UPCGMetadataBreakVectorSettings::CreateElement() const
+{
+	return MakeShared<FPCGMetadataBreakVectorElement>();
+}
 
-		if (bZPinConnected)
-		{
-			FPCGTaggedData& OutputZ = Outputs.Add_GetRef(Input);
-			OutputZ.Pin = PCGMetadataBreakVectorConstants::ZLabel;
+bool FPCGMetadataBreakVectorElement::DoOperation(FOperationData& OperationData) const
+{
+	TRACE_CPUPROFILER_EVENT_SCOPE(FPCGMetadataBreakVectorElement::Execute);
 
-			UPCGMetadata* OutMetadata = nullptr;
+	const UPCGMetadataBreakVectorSettings* Settings = static_cast<const UPCGMetadataBreakVectorSettings*>(OperationData.Settings);
+	check(Settings);
 
-			PCGMetadataElementCommon::DuplicateTaggedData(Input, OutputZ, OutMetadata);
-			AttributeZ = PCGMetadataElementCommon::ClearOrCreateAttribute(OutMetadata, DestinationAttributeForZ, DefaultValue.Z);
-			PCGMetadataElementCommon::CopyEntryToValueKeyMap(SourceMetadata, SourceAttribute, AttributeZ);
-		}
+	auto BreakFunc = [this, &OperationData](auto DummyValue) -> bool
+	{
+		using AttributeType = decltype(DummyValue);
 
-		if (bWPinConnected && SourceAttribute->GetTypeId() == PCG::Private::MetadataTypes<FVector4>::Id)
+		if constexpr (!PCGMetadataBreakVectorSettings::IsValidType<AttributeType>())
 		{
-			FPCGTaggedData& OutputW = Outputs.Add_GetRef(Input);
-			OutputW.Pin = PCGMetadataBreakVectorConstants::WLabel;
-
-			UPCGMetadata* OutMetadata = nullptr;
-
-			PCGMetadataElementCommon::DuplicateTaggedData(Input, OutputW, OutMetadata);
-			AttributeW = PCGMetadataElementCommon::ClearOrCreateAttribute(OutMetadata, DestinationAttributeForW, DefaultValue.W);
-			PCGMetadataElementCommon::CopyEntryToValueKeyMap(SourceMetadata, SourceAttribute, AttributeW);
+			return false;
 		}
-
-		// Copy all Value pairs from the parent hierarchy into our new attributes. Assumes that adding values in order of the parent hierarchy will result in identical ValueKey->Value mappings.
-		const PCGMetadataValueKey NumValueKeys = SourceAttribute->GetValueKeyOffsetForChild();
-		for (PCGMetadataValueKey ValueKey = 0; ValueKey < NumValueKeys; ++ValueKey)
+		else
 		{
-			FVector4 Value;
-
-			if (SourceAttribute->GetTypeId() == PCG::Private::MetadataTypes<FVector>::Id)
-			{
-				Value = static_cast<const FPCGMetadataAttribute<FVector>*>(SourceAttribute)->GetValue(ValueKey);
-			}
-			else if (SourceAttribute->GetTypeId() == PCG::Private::MetadataTypes<FVector4>::Id)
-			{
-				Value = static_cast<const FPCGMetadataAttribute<FVector4>*>(SourceAttribute)->GetValue(ValueKey);
-			}
-			else //if (SourceAttribute->GetTypeId() == PCG::Private::MetadataTypes<FRotator>::Id)
-			{
-				const FRotator Rotator = static_cast<const FPCGMetadataAttribute<FRotator>*>(SourceAttribute)->GetValue(ValueKey);
-				Value.X = Rotator.Roll;
-				Value.Y = Rotator.Pitch;
-				Value.Z = Rotator.Yaw;
-			}
-
-			if (AttributeX)
-			{
-				AttributeX->AddValue(Value.X);
-			}
-
-			if (AttributeY)
-			{
-				AttributeY->AddValue(Value.Y);
-			}
-
-			if (AttributeZ)
-			{
-				AttributeZ->AddValue(Value.Z);
-			}
-
-			if (AttributeW)
-			{
-				AttributeW->AddValue(Value.W);
-			}
+			return DoUnaryOpMultipleOutputs<AttributeType, double>(OperationData, PCGMetadataBreakVectorSettings::DoBreak<AttributeType>);
 		}
-	}
+	};
 
-	return true;
+	return PCGMetadataAttribute::CallbackWithRightType(OperationData.MostComplexInputType, BreakFunc);
 }
diff --git a/Engine/Plugins/Experimental/PCG/Source/PCG/Private/Elements/Metadata/PCGMetadataOpElementBase.cpp b/Engine/Plugins/Experimental/PCG/Source/PCG/Private/Elements/Metadata/PCGMetadataOpElementBase.cpp
--- a/Engine/Plugins/Experimental/PCG/Source/PCG/Private/Elements/Metadata/PCGMetadataOpElementBase.cpp
+++ b/Engine/Plugins/Experimental/PCG/Source/PCG/Private/Elements/Metadata/PCGMetadataOpElementBase.cpp
@@ -33,7 +33,15 @@ TArray<FPCGPinProperties> UPCGMetadataSettingsBase::InputPinProperties() const
 TArray<FPCGPinProperties> UPCGMetadataSettingsBase::OutputPinProperties() const
 {
 	TArray<FPCGPinProperties> PinProperties;
-	PinProperties.Emplace(PCGPinConstants::DefaultOutputLabel, EPCGDataType::Any);
+	
+	for (uint32 i = 0; i < GetOutputPinNum(); ++i)
+	{
+		const FName PinLabel = GetOutputPinLabel(i);
+		if (PinLabel != NAME_None)
+		{
+			PinProperties.Emplace(PinLabel, EPCGDataType::Any);
+		}
+	}
 
 	return PinProperties;
 }
@@ -50,11 +58,30 @@ bool FPCGMetadataElementBase::ExecuteInternal(FPCGContext* Context) const
 	const UPCGMetadataSettingsBase* Settings = Context->GetInputSettings<UPCGMetadataSettingsBase>();
 	check(Settings);
 
+	const uint32 NumberOfInputs = Settings->GetInputPinNum();
+	const uint32 NumberOfOutputs = Settings->GetOutputPinNum();
+
+	check(NumberOfInputs > 0);
+	check(NumberOfOutputs <= UPCGMetadataSettingsBase::MaxNumberOfOutputs);
+
 	const TArray<FPCGTaggedData>& Inputs = Context->InputData.TaggedData;
 	TArray<FPCGTaggedData>& Outputs = Context->OutputData.TaggedData;
 
-	// Nothing to do if we have no output
-	if (!Context->Node || !Context->Node->IsOutputPinConnected(PCGPinConstants::DefaultOutputLabel))
+	TStaticArray<bool, UPCGMetadataSettingsBase::MaxNumberOfOutputs> ConnectedOutputs;
+	bool AnyOutputConnected = false;
+
+	for (uint32 i = 0; i < NumberOfOutputs; ++i)
+	{
+		ConnectedOutputs[i] = Context->Node && Context->Node->IsOutputPinConnected(Settings->GetOutputPinLabel(i));
+
+#if WITH_EDITOR
+		ConnectedOutputs[i] |= Settings->ForceOutputConnections[i];
+#endif // WITH_EDITOR
+
+		AnyOutputConnected |= ConnectedOutputs[i];
+	}
+
+	if (!AnyOutputConnected)
 	{
 		return true;
 	}
@@ -64,13 +91,11 @@ bool FPCGMetadataElementBase::ExecuteInternal(FPCGContext* Context) const
 
 	FName OutputAttributeName = PCG_GET_OVERRIDEN_VALUE(Settings, OutputAttributeName, Params);
 
-	const uint32 NumberOfInputs = Settings->GetInputPinNum();
-
 	// Gathering all the inputs metadata
 	TArray<const UPCGMetadata*> SourceMetadata;
-	TArray<const UPCGData*> Data;
+	TArray<FPCGTaggedData> InputTaggedData;
 	SourceMetadata.SetNum(NumberOfInputs);
-	Data.SetNum(NumberOfInputs);
+	InputTaggedData.SetNum(NumberOfInputs);
 
 	for (uint32 i = 0; i < NumberOfInputs; ++i)
 	{
@@ -83,14 +108,14 @@ bool FPCGMetadataElementBase::ExecuteInternal(FPCGContext* Context) const
 
 		// By construction, there can only be one of then(hence the 0 index)
 		check(InputData.Num() == 1);
-		Data[i] = InputData[0].Data;
+		InputTaggedData[i] = MoveTemp(InputData[0]);
 
 		// Only gather Spacial and Params input. 
-		if (const UPCGSpatialData* SpatialInput = Cast<const UPCGSpatialData>(Data[i]))
+		if (const UPCGSpatialData* SpatialInput = Cast<const UPCGSpatialData>(InputTaggedData[i].Data))
 		{
 			SourceMetadata[i] = SpatialInput->Metadata;
 		}
-		else if (const UPCGParamData* ParamsInput = Cast<const UPCGParamData>(Data[i]))
+		else if (const UPCGParamData* ParamsInput = Cast<const UPCGParamData>(InputTaggedData[i].Data))
 		{
 			SourceMetadata[i] = ParamsInput->Metadata;
 		}
@@ -150,7 +175,7 @@ bool FPCGMetadataElementBase::ExecuteInternal(FPCGContext* Context) const
 		}
 
 		// Finally check that we have the right number of inputs, depending on the source
-		if (const UPCGPointData* PointData = Cast<UPCGPointData>(Data[i]))
+		if (const UPCGPointData* PointData = Cast<UPCGPointData>(InputTaggedData[i].Data))
 		{
 			NumberOfElements[i] = PointData->GetPoints().Num();
 
@@ -175,7 +200,7 @@ bool FPCGMetadataElementBase::ExecuteInternal(FPCGContext* Context) const
 				// Broadcast only if we have a single element and we are in broadcast mode, or inferred mode and we are param data.
 				bool bShouldBroadcast = i != 0 && 
 					NumberOfElements[i] == 1 && 
-					(Settings->Mode == EPCGMetadataSettingsBaseMode::Broadcast || (Settings->Mode == EPCGMetadataSettingsBaseMode::Inferred && Data[i]->IsA<UPCGParamData>()));
+					(Settings->Mode == EPCGMetadataSettingsBaseMode::Broadcast || (Settings->Mode == EPCGMetadataSettingsBaseMode::Inferred && InputTaggedData[i].Data->IsA<UPCGParamData>()));
 
 				if (i == 0 || NumberOfElements[i] == NumberOfElements[0] || bShouldBroadcast)
 				{
@@ -192,35 +217,59 @@ bool FPCGMetadataElementBase::ExecuteInternal(FPCGContext* Context) const
 
 	// At this point, we verified everything, so we can go forward with the computation, depending on the most complex type
 	// So first forward outputs and create the attribute
-	OperationData.OutputAttribute = nullptr;
-	auto CreateAttribute = [&SourceAttributes = OperationData.SourceAttributes, &Inputs, &Outputs, &SourceMetadata, OutputAttributeName, &OutputAttribute = OperationData.OutputAttribute](auto DummyOutValue) -> void
+	OperationData.OutputAttributes.SetNum(Settings->GetOutputPinNum());
+
+	auto CreateAttributes = [&](auto DummyOutValue) -> bool
 	{
 		using AttributeType = decltype(DummyOutValue);
 
-		AttributeType DefaultValue = PCGMetadataAttribute::GetValueWithBroadcast<AttributeType>(SourceAttributes[0], PCGDefaultValueKey);
+		AttributeType DefaultValue = PCGMetadataAttribute::GetValueWithBroadcast<AttributeType>(OperationData.SourceAttributes[0], PCGDefaultValueKey);
 
-		FPCGTaggedData& OutputData = Outputs.Add_GetRef(Inputs[0]);
-		OutputData.Pin = PCGPinConstants::DefaultOutputLabel;
+		for (uint32 i = 0; i < NumberOfOutputs; ++i)
+		{
+			if (!ConnectedOutputs[i])
+			{
+				OperationData.OutputAttributes[i] = nullptr;
+				continue;
+			}
+
+			FPCGTaggedData& OutputData = Outputs.Add_GetRef(InputTaggedData[0]);
+			OutputData.Pin = Settings->GetOutputPinLabel(i);
+
+			FName OutputAttributeFinalName = Settings->GetOutputAttributeName(OutputAttributeName, i);
 
-		UPCGMetadata* OutMetadata = nullptr;
+			UPCGMetadata* OutMetadata = nullptr;
 
-		PCGMetadataElementCommon::DuplicateTaggedData(Inputs[0], OutputData, OutMetadata);
-		OutputAttribute = PCGMetadataElementCommon::ClearOrCreateAttribute(OutMetadata, OutputAttributeName, DefaultValue);
-		PCGMetadataElementCommon::CopyEntryToValueKeyMap(SourceMetadata[0], SourceAttributes[0], OutputAttribute);
+			PCGMetadataElementCommon::DuplicateTaggedData(InputTaggedData[0], OutputData, OutMetadata);
+			OperationData.OutputAttributes[i] = PCGMetadataElementCommon::ClearOrCreateAttribute(OutMetadata, OutputAttributeFinalName, DefaultValue);
+			PCGMetadataElementCommon::CopyEntryToValueKeyMap(SourceMetadata[0], OperationData.SourceAttributes[0], OperationData.OutputAttributes[i]);
+
+			// If something went wrong stop here.
+			if (OperationData.OutputAttributes[i] == nullptr)
+			{
+				return false;
+			}
+		}
+
+		return true;
 	};
 
 	OperationData.OutputType = Settings->GetOutputType(OperationData.MostComplexInputType);
 
-	PCGMetadataAttribute::CallbackWithRightType(OperationData.OutputType, CreateAttribute);
-
-	if (OperationData.OutputAttribute == nullptr)
+	if (!PCGMetadataAttribute::CallbackWithRightType(OperationData.OutputType, CreateAttributes))
 	{
-		PCGE_LOG(Error, "Error while creating output attribute");
+		PCGE_LOG(Error, "Error while creating output attributes");
 		Outputs.Empty();
 		return true;
 	}
 
 	OperationData.Settings = Settings;
 
-	return DoOperation(OperationData);
-}
\ No newline at end of file
+	if (!DoOperation(OperationData))
+	{
+		PCGE_LOG(Error, "Error while performing the metadata operation, check logs for more information");
+		Outputs.Empty();
+	}
+
+	return true;
+}
diff --git a/Engine/Plugins/Experimental/PCG/Source/PCG/Private/Metadata/PCGMetadata.cpp b/Engine/Plugins/Experimental/PCG/Source/PCG/Private/Metadata/PCGMetadata.cpp
--- a/Engine/Plugins/Experimental/PCG/Source/PCG/Private/Metadata/PCGMetadata.cpp
+++ b/Engine/Plugins/Experimental/PCG/Source/PCG/Private/Metadata/PCGMetadata.cpp
@@ -360,6 +360,11 @@ void UPCGMetadata::CreateVector4Attribute(FName AttributeName, FVector4 DefaultV
 	CreateAttribute<FVector4>(AttributeName, DefaultValue, bAllowsInterpolation, bOverrideParent);
 }
 
+void UPCGMetadata::CreateVector2Attribute(FName AttributeName, FVector2D DefaultValue, bool bAllowsInterpolation, bool bOverrideParent)
+{
+	CreateAttribute<FVector2D>(AttributeName, DefaultValue, bAllowsInterpolation, bOverrideParent);
+}
+
 void UPCGMetadata::CreateRotatorAttribute(FName AttributeName, FRotator DefaultValue, bool bAllowsInterpolation, bool bOverrideParent)
 {
 	CreateAttribute<FRotator>(AttributeName, DefaultValue, bAllowsInterpolation, bOverrideParent);
diff --git a/Engine/Plugins/Experimental/PCG/Source/PCG/Private/Metadata/PCGMetadataAccessor.cpp b/Engine/Plugins/Experimental/PCG/Source/PCG/Private/Metadata/PCGMetadataAccessor.cpp
--- a/Engine/Plugins/Experimental/PCG/Source/PCG/Private/Metadata/PCGMetadataAccessor.cpp
+++ b/Engine/Plugins/Experimental/PCG/Source/PCG/Private/Metadata/PCGMetadataAccessor.cpp
@@ -139,6 +139,16 @@ void UPCGMetadataAccessorHelpers::SetVector4AttributeByMetadataKey(int64& Key, U
 	SetAttribute(Key, Metadata, AttributeName, Value);
 }
 
+FVector2D UPCGMetadataAccessorHelpers::GetVector2AttributeByMetadataKey(int64 Key, const UPCGMetadata* Metadata, FName AttributeName)
+{
+	return GetAttribute<FVector2D>(Key, Metadata, AttributeName);
+}
+
+void UPCGMetadataAccessorHelpers::SetVector2AttributeByMetadataKey(int64& Key, UPCGMetadata* Metadata, FName AttributeName, const FVector2D& Value)
+{
+	SetAttribute(Key, Metadata, AttributeName, Value);
+}
+
 FRotator UPCGMetadataAccessorHelpers::GetRotatorAttributeByMetadataKey(int64 Key, const UPCGMetadata* Metadata, FName AttributeName)
 {
 	return GetAttribute<FRotator>(Key, Metadata, AttributeName);
@@ -294,6 +304,16 @@ void UPCGMetadataAccessorHelpers::SetVector4Attribute(FPCGPoint& Point, UPCGMeta
 	SetAttribute(Point.MetadataEntry, Metadata, AttributeName, Value);
 }
 
+FVector2D UPCGMetadataAccessorHelpers::GetVector2Attribute(const FPCGPoint& Point, const UPCGMetadata* Metadata, FName AttributeName)
+{
+	return GetAttribute<FVector2D>(Point.MetadataEntry, Metadata, AttributeName);
+}
+
+void UPCGMetadataAccessorHelpers::SetVector2Attribute(FPCGPoint& Point, UPCGMetadata* Metadata, FName AttributeName, const FVector2D& Value)
+{
+	SetAttribute(Point.MetadataEntry, Metadata, AttributeName, Value);
+}
+
 FRotator UPCGMetadataAccessorHelpers::GetRotatorAttribute(const FPCGPoint& Point, const UPCGMetadata* Metadata, FName AttributeName)
 {
 	return GetAttribute<FRotator>(Point.MetadataEntry, Metadata, AttributeName);
diff --git a/Engine/Plugins/Experimental/PCG/Source/PCG/Private/Tests/Elements/Metadata/PCGMetadataBreakVectorTest.cpp b/Engine/Plugins/Experimental/PCG/Source/PCG/Private/Tests/Elements/Metadata/PCGMetadataBreakVectorTest.cpp
new file mode 100644
--- /dev/null
+++ b/Engine/Plugins/Experimental/PCG/Source/PCG/Private/Tests/Elements/Metadata/PCGMetadataBreakVectorTest.cpp
@@ -0,0 +1,409 @@
+// Copyright Epic Games, Inc. All Rights Reserved.
+
+#include "Tests/Determinism/PCGDeterminismTestsCommon.h"
+#include "Tests/PCGTestsCommon.h"
+
+#include "PCGComponent.h"
+#include "PCGData.h"
+#include "PCGParamData.h"
+
+#include "Data/PCGPointData.h"
+#include "Elements/Metadata/PCGMetadataBreakVector.h"
+#include "Metadata/PCGMetadata.h"
+#include "Metadata/PCGMetadataAttribute.h"
+#include "Metadata/PCGMetadataAttributeTpl.h"
+
+#if WITH_EDITOR
+
+IMPLEMENT_CUSTOM_SIMPLE_AUTOMATION_TEST(FPCGMetadataBreakVectorTest, FPCGTestBaseClass, "pcg.tests.Metadata.BreakVector", PCGTestsCommon::TestFlags)
+
+namespace
+{
+	enum class EPCGComponentToCheck : uint8
+	{
+		X,
+		Z,
+		Y,
+		W
+	};
+
+	const FName Vec2Attribute = TEXT("Vec2");
+	const FName Vec3Attribute = TEXT("Vec3");
+	const FName Vec4Attribute = TEXT("Vec4");
+	const FName RotatorAttribute = TEXT("Rotator");
+	const FName FloatAttribute = TEXT("Float");
+
+	void GenerateSpatialData(PCGTestsCommon::FTestData& TestData)
+	{
+		FPCGTaggedData& SpatialTaggedData = TestData.InputData.TaggedData.Emplace_GetRef();
+		TObjectPtr<UPCGPointData> PointData = PCGTestsCommon::CreateRandomPointData(10, TestData.Seed);
+		SpatialTaggedData.Data = PointData;
+		SpatialTaggedData.Pin = PCGPinConstants::DefaultInputLabel;
+
+		const bool bAllowsInterpolation = false;
+		const bool bOverrideParent = false;
+
+		PointData->Metadata->CreateVector2Attribute(Vec2Attribute, FVector2D::ZeroVector, bAllowsInterpolation, bOverrideParent);
+		PointData->Metadata->CreateVectorAttribute(Vec3Attribute, FVector::Zero(), bAllowsInterpolation, bOverrideParent);
+		PointData->Metadata->CreateVector4Attribute(Vec4Attribute, FVector4::Zero(), bAllowsInterpolation, bOverrideParent);
+		PointData->Metadata->CreateRotatorAttribute(RotatorAttribute, FRotator::ZeroRotator, bAllowsInterpolation, bOverrideParent);
+		PointData->Metadata->CreateFloatAttribute(FloatAttribute, 0, bAllowsInterpolation, bOverrideParent);
+
+		TArray<FPCGPoint>& SourcePoints = PointData->GetMutablePoints();
+
+		// Randomize attribute values, and leave the second half to default value
+		FRandomStream RandomSource(TestData.Seed);
+		for (int I = 0; I < SourcePoints.Num() / 2; ++I)
+		{
+			UPCGMetadataAccessorHelpers::SetVector2Attribute(SourcePoints[I], PointData->Metadata, Vec2Attribute, FVector2D(RandomSource.VRand()));
+			UPCGMetadataAccessorHelpers::SetVectorAttribute(SourcePoints[I], PointData->Metadata, Vec3Attribute, RandomSource.VRand());
+			UPCGMetadataAccessorHelpers::SetVector4Attribute(SourcePoints[I], PointData->Metadata, Vec4Attribute, FVector4(RandomSource.VRand(), RandomSource.FRand()));
+			UPCGMetadataAccessorHelpers::SetRotatorAttribute(SourcePoints[I], PointData->Metadata, RotatorAttribute, FRotator(RandomSource.FRand(), RandomSource.FRand(), RandomSource.FRand()));
+			UPCGMetadataAccessorHelpers::SetFloatAttribute(SourcePoints[I], PointData->Metadata, FloatAttribute, RandomSource.FRand());
+		}
+	}
+
+	void GenerateParamData(PCGTestsCommon::FTestData& TestData)
+	{
+		FPCGTaggedData& ParamTaggedData = TestData.InputData.TaggedData.Emplace_GetRef();
+		TObjectPtr<UPCGParamData> ParamData = PCGTestsCommon::CreateEmptyParamData();
+		ParamTaggedData.Data = ParamData;
+		ParamTaggedData.Pin = PCGPinConstants::DefaultInputLabel;
+
+		PCGMetadataEntryKey EntryKey = ParamData->Metadata->AddEntry();
+		const bool bAllowsInterpolation = false;
+		const bool bOverrideParent = false;
+
+		ParamData->Metadata->CreateVector2Attribute(Vec2Attribute, FVector2D::ZeroVector, bAllowsInterpolation, bOverrideParent);
+		ParamData->Metadata->CreateVectorAttribute(Vec3Attribute, FVector::Zero(), bAllowsInterpolation, bOverrideParent);
+		ParamData->Metadata->CreateVector4Attribute(Vec4Attribute, FVector4::Zero(), bAllowsInterpolation, bOverrideParent);
+		ParamData->Metadata->CreateRotatorAttribute(RotatorAttribute, FRotator::ZeroRotator, bAllowsInterpolation, bOverrideParent);
+		ParamData->Metadata->CreateFloatAttribute(FloatAttribute, 0, bAllowsInterpolation, bOverrideParent);
+
+		// Randomize attribute values
+		FRandomStream RandomSource(TestData.Seed);
+		if (FPCGMetadataAttribute<FVector2D>* Attribute = static_cast<FPCGMetadataAttribute<FVector2D>*>(ParamData->Metadata->GetMutableAttribute(Vec2Attribute)))
+		{
+			Attribute->SetValue(EntryKey, FVector2D(RandomSource.VRand()));
+		}
+
+		if (FPCGMetadataAttribute<FVector>* Attribute = static_cast<FPCGMetadataAttribute<FVector>*>(ParamData->Metadata->GetMutableAttribute(Vec3Attribute)))
+		{
+			Attribute->SetValue(EntryKey, RandomSource.VRand());
+		}
+
+		if (FPCGMetadataAttribute<FVector4>* Attribute = static_cast<FPCGMetadataAttribute<FVector4>*>(ParamData->Metadata->GetMutableAttribute(Vec4Attribute)))
+		{
+			Attribute->SetValue(EntryKey, FVector4(RandomSource.VRand(), RandomSource.FRand()));
+		}
+
+		if (FPCGMetadataAttribute<FRotator>* Attribute = static_cast<FPCGMetadataAttribute<FRotator>*>(ParamData->Metadata->GetMutableAttribute(RotatorAttribute)))
+		{
+			Attribute->SetValue(EntryKey, FRotator(RandomSource.FRand(), RandomSource.FRand(), RandomSource.FRand()));
+		}
+
+		if (FPCGMetadataAttribute<float>* Attribute = static_cast<FPCGMetadataAttribute<float>*>(ParamData->Metadata->GetMutableAttribute(FloatAttribute)))
+		{
+			Attribute->SetValue(EntryKey, RandomSource.FRand());
+		}
+	}
+}
+
+
+bool FPCGMetadataBreakVectorTest::RunTest(const FString& Parameters)
+{	
+	auto ValidateMetadataBreakVector = [this](PCGTestsCommon::FTestData& TestData, bool bIsValid = true) -> bool
+	{
+		const FName OutputAttributeName = TEXT("Output");
+
+		UPCGMetadataBreakVectorSettings* Settings = CastChecked<UPCGMetadataBreakVectorSettings>(TestData.Settings);
+		Settings->OutputAttributeName = OutputAttributeName;
+		FPCGElementPtr MetadataBreakVectorElement = TestData.Settings->GetElement();
+
+		TUniquePtr<FPCGContext> Context = TUniquePtr<FPCGContext>(MetadataBreakVectorElement->Initialize(TestData.InputData, TestData.TestPCGComponent, nullptr));
+		Context->NumAvailableTasks = 1;
+
+		while (!MetadataBreakVectorElement->Execute(Context.Get()))
+		{}
+
+		// If the test is not valid, just check that we have no output and early out
+		if (!bIsValid)
+		{
+			return TestTrue(TEXT("Invalid operation, no outputs"), Context->OutputData.GetInputs().IsEmpty());
+		}
+
+		const TArray<FPCGTaggedData>& Inputs = Context->InputData.GetInputsByPin(PCGPinConstants::DefaultInputLabel);
+
+		if (!TestEqual("Right number of inputs", Inputs.Num(), 1))
+		{
+			return false;
+		}
+
+		bool bTestPassed = true;
+
+		const FPCGTaggedData& Input = Inputs[0];
+
+		UPCGMetadata* SourceMetadata = nullptr; 
+		if (const UPCGSpatialData* SpatialInput = Cast<const UPCGSpatialData>(Input.Data))
+		{
+			SourceMetadata = SpatialInput->Metadata;
+		}
+		else if (const UPCGParamData* ParamsInput = Cast<const UPCGParamData>(Input.Data))
+		{
+			SourceMetadata = ParamsInput->Metadata;
+		}
+		else
+		{
+			return false;
+		}
+
+		check(SourceMetadata);
+
+		const FPCGMetadataAttributeBase* SourceAttributeBase = SourceMetadata->GetConstAttribute(Settings->InputAttributeName);
+		check(SourceAttributeBase);
+
+		auto ValidateComponentOutput = [&](const FPCGTaggedData& Output, const FName& OutAttributeName, EPCGComponentToCheck ComponentToCheck)
+		{
+			if (!TestTrue("Valid output data", Output.Data != nullptr))
+			{
+				bTestPassed = false;
+				return;
+			}
+
+			UPCGMetadata* OutMetadata = nullptr;
+			if (const UPCGSpatialData* SpatialInput = Cast<const UPCGSpatialData>(Output.Data))
+			{
+				OutMetadata = SpatialInput->Metadata;
+			}
+			else if (const UPCGParamData* ParamsInput = Cast<const UPCGParamData>(Output.Data))
+			{
+				OutMetadata = ParamsInput->Metadata;
+			}
+			else
+			{
+				return;
+			}
+
+			if (!TestNotNull("Valid output metadata", OutMetadata))
+			{
+				bTestPassed = false;
+				return;
+			}
+
+			const FPCGMetadataAttributeBase* OutAttributeBase = OutMetadata->GetConstAttribute(OutAttributeName);
+
+			if (!TestNotNull("Valid output attribute", OutAttributeBase))
+			{
+				bTestPassed = false;
+				return;
+			}
+
+			if (!TestEqual("Output attribute is a valid type (double)", OutAttributeBase->GetTypeId(), PCG::Private::MetadataTypes<double>::Id))
+			{
+				bTestPassed = false;
+				return;
+			}
+
+			const FPCGMetadataAttribute<double>* OutAttribute = static_cast<const FPCGMetadataAttribute<double>*>(OutAttributeBase);
+			check(OutAttribute);
+
+			const PCGMetadataEntryKey InEntryKeyCount = SourceMetadata->GetItemCountForChild();
+			const PCGMetadataEntryKey OutEntryKeyCount = OutMetadata->GetItemCountForChild();
+
+			if (!TestEqual("Identical EntryKey counts", InEntryKeyCount, OutEntryKeyCount))
+			{
+				bTestPassed = false;
+				return;
+			}
+
+			const PCGMetadataValueKey InValueKeyCount = SourceAttributeBase->GetValueKeyOffsetForChild();
+			const PCGMetadataValueKey OutValueKeyCount = OutAttribute->GetValueKeyOffsetForChild();
+
+			if (!TestEqual("Identical ValueKey counts", InValueKeyCount, OutValueKeyCount))
+			{
+				bTestPassed = false;
+				return;
+			}
+
+			for (PCGMetadataEntryKey EntryKey = 0; EntryKey < InEntryKeyCount; ++EntryKey)
+			{
+				const PCGMetadataValueKey InValueKey = SourceAttributeBase->GetValueKey(EntryKey);
+				const PCGMetadataValueKey OutValueKey = OutAttribute->GetValueKey(EntryKey);
+
+				bTestPassed &= TestEqual("Identical value keys", InValueKey, OutValueKey);
+
+				FVector4 SourceValue;
+				if (SourceAttributeBase->GetTypeId() == PCG::Private::MetadataTypes<FVector>::Id)
+				{
+					SourceValue = static_cast<const FPCGMetadataAttribute<FVector>*>(SourceAttributeBase)->GetValue(InValueKey);
+				}
+				else if (SourceAttributeBase->GetTypeId() == PCG::Private::MetadataTypes<FVector2D>::Id)
+				{
+					const FVector2D Vec2 = static_cast<const FPCGMetadataAttribute<FVector2D>*>(SourceAttributeBase)->GetValue(InValueKey);
+					SourceValue = FVector4(Vec2.X, Vec2.Y);
+				}
+				else if (SourceAttributeBase->GetTypeId() == PCG::Private::MetadataTypes<FVector4>::Id)
+				{
+					SourceValue = static_cast<const FPCGMetadataAttribute<FVector4>*>(SourceAttributeBase)->GetValue(InValueKey);
+				}
+				else //if (SourceAttributeBase->GetTypeId() == PCG::Private::MetadataTypes<FRotator>::Id)
+				{
+					const FRotator Rotator = static_cast<const FPCGMetadataAttribute<FRotator>*>(SourceAttributeBase)->GetValue(InValueKey);
+					SourceValue.X = Rotator.Roll;
+					SourceValue.Y = Rotator.Pitch;
+					SourceValue.Z = Rotator.Yaw;
+				}
+
+				const double OutValue = OutAttribute->GetValue(OutValueKey);
+
+				if (ComponentToCheck == EPCGComponentToCheck::X)
+				{
+					bTestPassed &= TestEqual("X", SourceValue.X, OutValue);
+				}
+				else if (ComponentToCheck == EPCGComponentToCheck::Y)
+				{
+					bTestPassed &= TestEqual("Y", SourceValue.Y, OutValue);
+				}
+				else if (ComponentToCheck == EPCGComponentToCheck::Z)
+				{
+					bTestPassed &= TestEqual("Z", SourceValue.Z, OutValue);
+				}
+				else if (ComponentToCheck == EPCGComponentToCheck::W)
+				{
+					bTestPassed &= TestEqual("W", SourceValue.W, OutValue);
+				}
+			}
+		};
+
+		const TArray<FPCGTaggedData> OutputsX = Context->OutputData.GetInputsByPin(PCGMetadataBreakVectorConstants::XLabel);
+		const TArray<FPCGTaggedData> OutputsY = Context->OutputData.GetInputsByPin(PCGMetadataBreakVectorConstants::YLabel);
+		const TArray<FPCGTaggedData> OutputsZ = Context->OutputData.GetInputsByPin(PCGMetadataBreakVectorConstants::ZLabel);
+		const TArray<FPCGTaggedData> OutputsW = Context->OutputData.GetInputsByPin(PCGMetadataBreakVectorConstants::WLabel);
+
+		const FName DestinationAttributeForX = Settings->GetOutputAttributeName(Settings->OutputAttributeName, 0);
+		const FName DestinationAttributeForY = Settings->GetOutputAttributeName(Settings->OutputAttributeName, 1);
+		const FName DestinationAttributeForZ = Settings->GetOutputAttributeName(Settings->OutputAttributeName, 2);
+		const FName DestinationAttributeForW = Settings->GetOutputAttributeName(Settings->OutputAttributeName, 3);
+
+		if (Settings->ForceOutputConnections[0])
+		{
+			if (TestEqual("Appropriate number of outputs generated fox X", OutputsX.Num(), 1))
+			{
+				ValidateComponentOutput(OutputsX[0], DestinationAttributeForX, EPCGComponentToCheck::X);
+			}
+			else
+			{
+				bTestPassed = false;
+			}
+		}
+
+		if (Settings->ForceOutputConnections[1])
+		{
+			if (TestEqual("Appropriate number of outputs generated for Y", OutputsY.Num(), 1))
+			{
+				ValidateComponentOutput(OutputsY[0], DestinationAttributeForY, EPCGComponentToCheck::Y);
+			}
+			else
+			{
+				bTestPassed = false;
+			}
+		}
+
+		if (Settings->ForceOutputConnections[2])
+		{
+			if (TestEqual("Appropriate number of outputs generated for Z", OutputsZ.Num(), 1))
+			{
+				ValidateComponentOutput(OutputsZ[0], DestinationAttributeForZ, EPCGComponentToCheck::Z);
+			}
+			else
+			{
+				bTestPassed = false;
+			}
+		}
+
+		if (Settings->ForceOutputConnections[3] && SourceAttributeBase->GetTypeId() == PCG::Private::MetadataTypes<FVector4>::Id)
+		{
+			if (TestEqual("Appropriate number of outputs generated for W", OutputsW.Num(), 1))
+			{
+				ValidateComponentOutput(OutputsW[0], DestinationAttributeForW, EPCGComponentToCheck::W);
+			}
+			else
+			{
+				bTestPassed = false;
+			}
+		}
+
+		return bTestPassed;
+	};
+
+	bool bTestPassed = true;
+
+	PCGTestsCommon::FTestData TestDataSpatial(PCGDeterminismTests::Defaults::Seed);
+	PCGDeterminismTests::GenerateSettings<UPCGMetadataBreakVectorSettings>(TestDataSpatial);
+	GenerateSpatialData(TestDataSpatial);
+
+	PCGTestsCommon::FTestData TestDataParams(PCGDeterminismTests::Defaults::Seed);
+	PCGDeterminismTests::GenerateSettings<UPCGMetadataBreakVectorSettings>(TestDataParams);
+	GenerateParamData(TestDataParams);
+
+	using PairWhatData = TPair<FString, PCGTestsCommon::FTestData*>;
+
+	TStaticArray<PairWhatData, 2> AllTestData;
+	AllTestData[0] = PairWhatData("Testing with point data as input", &TestDataSpatial);
+	AllTestData[1] = PairWhatData("Testing with param data as input", &TestDataParams);
+
+	// Setup error catching. This error should happen only twice (invalid type)
+	AddExpectedError(TEXT("Attribute Float is not a supported type for input 0"), EAutomationExpectedErrorFlags::Contains, 2);
+
+	for (PairWhatData& PairTestData : AllTestData)
+	{
+		AddInfo(PairTestData.Key);
+		PCGTestsCommon::FTestData* TestData = PairTestData.Value;
+
+		UPCGMetadataBreakVectorSettings* Settings = CastChecked<UPCGMetadataBreakVectorSettings>(TestData->Settings);
+
+		Settings->ForceOutputConnections[0] = true;
+		Settings->ForceOutputConnections[1] = true;
+		Settings->ForceOutputConnections[2] = false;
+		Settings->ForceOutputConnections[3] = false;
+
+		{
+			AddInfo(TEXT("Test with FVector2D as input attribute"));
+			Settings->InputAttributeName = Vec2Attribute;
+			bTestPassed &= ValidateMetadataBreakVector(*TestData);
+		}
+
+		Settings->ForceOutputConnections[2] = true;
+
+		{
+			AddInfo(TEXT("Test with FVector as input attribute"));
+			Settings->InputAttributeName = Vec3Attribute;
+			bTestPassed &= ValidateMetadataBreakVector(*TestData);
+		}
+
+		{
+			AddInfo(TEXT("Test with FVector4 as input attribute"));
+			Settings->InputAttributeName = Vec4Attribute;
+			Settings->ForceOutputConnections[3] = true;
+			bTestPassed &= ValidateMetadataBreakVector(*TestData);
+			Settings->ForceOutputConnections[3] = false;
+		}
+
+		{
+			AddInfo(TEXT("Test with FRotator as input attribute"));
+			Settings->InputAttributeName = RotatorAttribute;
+			bTestPassed &= ValidateMetadataBreakVector(*TestData);
+		}
+
+		{
+			AddInfo(TEXT("Test with Float as input attribute (invalid type)"));
+			Settings->InputAttributeName = FloatAttribute;
+
+			bTestPassed &= ValidateMetadataBreakVector(*TestData, /*bIsValid=*/false);
+		}
+	}
+
+	return bTestPassed;
+}
+
+#endif // WITH_EDITOR
diff --git a/Engine/Plugins/Experimental/PCG/Source/PCG/Private/Tests/Elements/PCGMetadataBreakVectorTest.cpp b/Engine/Plugins/Experimental/PCG/Source/PCG/Private/Tests/Elements/PCGMetadataBreakVectorTest.cpp
deleted file mode 100644
diff --git a/Engine/Plugins/Experimental/PCG/Source/PCG/Public/Elements/Metadata/PCGMetadataBreakVector.h b/Engine/Plugins/Experimental/PCG/Source/PCG/Public/Elements/Metadata/PCGMetadataBreakVector.h
--- a/Engine/Plugins/Experimental/PCG/Source/PCG/Public/Elements/Metadata/PCGMetadataBreakVector.h
+++ b/Engine/Plugins/Experimental/PCG/Source/PCG/Public/Elements/Metadata/PCGMetadataBreakVector.h
@@ -5,12 +5,12 @@
 #include "PCGElement.h"
 #include "PCGSettings.h"
 
+#include "Elements/Metadata/PCGMetadataOpElementBase.h"
+
 #include "PCGMetadataBreakVector.generated.h"
 
 namespace PCGMetadataBreakVectorConstants
 {
-	const FName ParamsLabel = TEXT("Params");
-	const FName SourceLabel = TEXT("Source");
 	const FName XLabel = TEXT("X");
 	const FName YLabel = TEXT("Y");
 	const FName ZLabel = TEXT("Z");
@@ -18,7 +18,7 @@ namespace PCGMetadataBreakVectorConstants
 }
 
 UCLASS(BlueprintType, ClassGroup = (Procedural))
-class PCG_API UPCGMetadataBreakVectorSettings : public UPCGSettings
+class PCG_API UPCGMetadataBreakVectorSettings : public UPCGMetadataSettingsBase
 {
 	GENERATED_BODY()
 
@@ -26,11 +26,16 @@ public:
 	//~Begin UPCGSettings interface
 #if WITH_EDITOR
 	virtual FName GetDefaultNodeName() const override { return FName(TEXT("MetadataBreakVectorNode")); }
-	virtual EPCGSettingsType GetType() const override { return EPCGSettingsType::Metadata; }
 #endif
 
-	virtual TArray<FPCGPinProperties> InputPinProperties() const override;
-	virtual TArray<FPCGPinProperties> OutputPinProperties() const override;
+	virtual FName GetInputAttributeNameWithOverride(uint32 Index, UPCGParamData* Params) const override;
+
+	virtual FName GetOutputPinLabel(uint32 Index) const override;
+	virtual uint32 GetOutputPinNum() const override;
+
+	virtual bool IsSupportedInputType(uint16 TypeId, uint32 InputIndex, bool& bHasSpecialRequirement) const override;
+	virtual uint16 GetOutputType(uint16 InputTypeId) const override;
+	virtual FName GetOutputAttributeName(FName BaseName, uint32 Index) const override;
 
 protected:
 	virtual FPCGElementPtr CreateElement() const override;
@@ -38,22 +43,11 @@ protected:
 
 public:
 	UPROPERTY(BlueprintReadWrite, EditAnywhere, Category = Settings)
-	FName SourceAttributeName = NAME_None;
-
-	// These exist to facilitate unit testing of the MetadataBreakVector element in isolation, when it otherwise would not have any pin connections to generate output
-	bool bForceConnectX = false;
-	bool bForceConnectY = false;
-	bool bForceConnectZ = false;
-	bool bForceConnectW = false;
+	FName InputAttributeName = NAME_None;
 };
 
-class FPCGMetadataBreakVectorElement : public FSimplePCGElement
+class FPCGMetadataBreakVectorElement : public FPCGMetadataElementBase
 {
-public:
-	// MetadataBreakVector relies on StaticDuplicateObject, so we cannot run outside of the main thread
-	virtual bool CanExecuteOnlyOnMainThread(const UPCGSettings* InSettings) const override { return true; }
-	virtual bool IsCacheable(const UPCGSettings* InSettings) const override { return false; }
-
 protected:
-	virtual bool ExecuteInternal(FPCGContext* Context) const override;
+	virtual bool DoOperation(FOperationData& OperationData) const override;
 };
diff --git a/Engine/Plugins/Experimental/PCG/Source/PCG/Public/Elements/Metadata/PCGMetadataOpElementBase.h b/Engine/Plugins/Experimental/PCG/Source/PCG/Public/Elements/Metadata/PCGMetadataOpElementBase.h
--- a/Engine/Plugins/Experimental/PCG/Source/PCG/Public/Elements/Metadata/PCGMetadataOpElementBase.h
+++ b/Engine/Plugins/Experimental/PCG/Source/PCG/Public/Elements/Metadata/PCGMetadataOpElementBase.h
@@ -9,6 +9,8 @@
 #include "Metadata/PCGMetadataAttributeTpl.h"
 #include "Metadata/PCGMetadataEntryKeyIterator.h"
 
+#include "Containers/StaticArray.h"
+
 #include "PCGMetadataOpElementBase.generated.h"
 
 class FPCGMetadataAttributeBase;
@@ -78,11 +80,15 @@ public:
 
 	virtual FName GetInputAttributeNameWithOverride(uint32 Index, UPCGParamData* Params) const { return NAME_None; };
 
-	virtual FName GetInputPinLabel(uint32 Index) const { return NAME_None; }
-	virtual uint32 GetInputPinNum() const { return 0; };
+	virtual FName GetInputPinLabel(uint32 Index) const { return PCGPinConstants::DefaultInputLabel; }
+	virtual uint32 GetInputPinNum() const { return 1; };
+
+	virtual FName GetOutputPinLabel(uint32 Index) const { return PCGPinConstants::DefaultOutputLabel; }
+	virtual uint32 GetOutputPinNum() const { return 1; }
 
 	virtual bool IsSupportedInputType(uint16 TypeId, uint32 InputIndex, bool& bHasSpecialRequirement) const { return false; };
 	virtual uint16 GetOutputType(uint16 InputTypeId) const { return InputTypeId; };
+	virtual FName GetOutputAttributeName(FName BaseName, uint32 Index) const { return BaseName; }
 
 	bool IsMoreComplexType(uint16 FirstType, uint16 SecondType) const;
 
@@ -91,12 +97,22 @@ public:
 
 	UPROPERTY(BlueprintReadWrite, EditAnywhere, Category = Settings)
 	EPCGMetadataSettingsBaseMode Mode = EPCGMetadataSettingsBaseMode::Inferred;
+
+	static constexpr uint32 MaxNumberOfOutputs = 4;
+
+#if WITH_EDITORONLY_DATA
+	// Useful for unit tests. Allow to force a connection to allow the node to do its operation, even if nothing is connected to it.
+	TStaticArray<bool, MaxNumberOfOutputs> ForceOutputConnections{};
+#endif // WITH_EDITORONLY_DATA
 };
 
 
 class FPCGMetadataElementBase : public FSimplePCGElement
 {
 public:
+	// FPCGMetadataElementBase relies on StaticDuplicateObject (when we duplicate tagged data), so we cannot run outside of the main thread
+	virtual bool CanExecuteOnlyOnMainThread(const UPCGSettings* InSettings) const override { return true; }
+
 	virtual bool IsCacheable(const UPCGSettings* InSettings) const override { return true; }
 
 	struct FOperationData
@@ -105,7 +121,7 @@ public:
 		int32 NumberOfElementsToProcess = -1;
 		uint16 MostComplexInputType;
 		uint16 OutputType;
-		FPCGMetadataAttributeBase* OutputAttribute = nullptr;
+		TArray<FPCGMetadataAttributeBase*> OutputAttributes;
 		const UPCGMetadataSettingsBase* Settings = nullptr;
 		TArray<TUniquePtr<IPCGMetadataEntryIterator>> Iterators;
 	};
@@ -115,26 +131,42 @@ protected:
 
 	virtual bool DoOperation(FOperationData& OperationData) const = 0;
 
+	/* All operations return true if it succeeded false otherwise.
+	* For each operation, the callback takes as many inputs as needed, as const ref of InType.
+	* and must return OutType by value.
+	*/
 	template <typename InType, typename OutType>
-	void DoUnaryOp(FOperationData& OperationData, TFunction<OutType(const InType&)> UnaryOp) const;
+	bool DoUnaryOp(FOperationData& OperationData, TFunction<OutType(const InType&)> UnaryOp) const;
 
 	template <typename InType1, typename InType2, typename OutType>
-	void DoBinaryOp(FOperationData& OperationData, TFunction<OutType(const InType1&, const InType2&)> BinaryOp) const;
+	bool DoBinaryOp(FOperationData& OperationData, TFunction<OutType(const InType1&, const InType2&)> BinaryOp) const;
 
 	template <typename InType1, typename InType2, typename InType3, typename OutType>
-	void DoTernaryOp(FOperationData& OperationData, TFunction<OutType(const InType1&, const InType2&, const InType3&)> TernaryOp) const;
+	bool DoTernaryOp(FOperationData& OperationData, TFunction<OutType(const InType1&, const InType2&, const InType3&)> TernaryOp) const;
+
+	template <typename InType1, typename InType2, typename InType3, typename InType4, typename OutType>
+	bool DoQuaternaryOp(FOperationData& OperationData, TFunction<OutType(const InType1&, const InType2&, const InType3&, const InType4&)> QuaternaryOp) const;
+
+	/* For multiple outputs, the callback won't return anything. But the callback should write results in
+	* the array passed as ref. Number of requested outputs is also passed as a parameter.
+	* Note that the array passed is "reserved" but not "resized", meaning that you need to push your results
+	* and not access the items directly with operator[].
+	*/
+	template <typename InType, typename OutType>
+	bool DoUnaryOpMultipleOutputs(FOperationData& OperationData, TFunction<void(const InType&, TArray<OutType>&, uint32)> UnaryOp) const;
 };
 
 template <typename InType, typename OutType>
-inline void FPCGMetadataElementBase::DoUnaryOp(FOperationData& OperationData, TFunction<OutType(const InType&)> UnaryOp) const
+inline bool FPCGMetadataElementBase::DoUnaryOp(FOperationData& OperationData, TFunction<OutType(const InType&)> UnaryOp) const
 {
 	check(OperationData.SourceAttributes[0]);
+	check(OperationData.OutputAttributes[0]);
 
 	// All values in OperationData.Iterators are UniquePtr. Dereference them here to make the syntax less cumbersome.
 	check(OperationData.Iterators[0]);
 	IPCGMetadataEntryIterator& Iterator1 = *OperationData.Iterators[0];
 
-	FPCGMetadataAttribute<OutType>* OutputAttribute = static_cast<FPCGMetadataAttribute<OutType>*>(OperationData.OutputAttribute);
+	FPCGMetadataAttribute<OutType>* OutputAttribute = static_cast<FPCGMetadataAttribute<OutType>*>(OperationData.OutputAttributes[0]);
 	InType DefaultValue = PCGMetadataAttribute::GetValueWithBroadcast<InType>(OperationData.SourceAttributes[0], PCGInvalidEntryKey);
 	OutType DefaultOutputValue = UnaryOp(DefaultValue);
 	OutputAttribute->SetDefaultValue(DefaultOutputValue);
@@ -155,25 +187,28 @@ inline void FPCGMetadataElementBase::DoUnaryOp(FOperationData& OperationData, TF
 
 		++Iterator1;
 	}
+
+	return true;
 }
 
 template <typename InType1, typename InType2, typename OutType>
-inline void FPCGMetadataElementBase::DoBinaryOp(FOperationData& OperationData, TFunction<OutType(const InType1&, const InType2&)> BinaryOp) const
+inline bool FPCGMetadataElementBase::DoBinaryOp(FOperationData& OperationData, TFunction<OutType(const InType1&, const InType2&)> BinaryOp) const
 {
 	check(OperationData.SourceAttributes[0]);
 	check(OperationData.SourceAttributes[1]);
+	check(OperationData.OutputAttributes[0]);
 
 	// All values in OperationData.Iterators are UniquePtr. Dereference them here to make the syntax less cumbersome.
 	// Also some iterators can be null, it means they need to use the first iterator (that should never be null).
 	// We just need to make sure to not increment the first iterator multiple times in one loop.
 	check(OperationData.Iterators[0]);
 
-	IPCGMetadataEntryIterator& Iterator1 = *OperationData.Iterators[0];
-	IPCGMetadataEntryIterator& Iterator2 = (OperationData.Iterators.Num() >= 2 && OperationData.Iterators[1]) ? *OperationData.Iterators[1] : Iterator1;
+	bool bShouldIncrementIterator2 = OperationData.Iterators.Num() >= 2 && OperationData.Iterators[1];
 
-	bool bShouldIncrementIterator2 = &Iterator1 != &Iterator2;
+	IPCGMetadataEntryIterator& Iterator1 = *OperationData.Iterators[0];
+	IPCGMetadataEntryIterator& Iterator2 = bShouldIncrementIterator2 ? *OperationData.Iterators[1] : Iterator1;
 
-	FPCGMetadataAttribute<OutType>* OutputAttribute = static_cast<FPCGMetadataAttribute<OutType>*>(OperationData.OutputAttribute);
+	FPCGMetadataAttribute<OutType>* OutputAttribute = static_cast<FPCGMetadataAttribute<OutType>*>(OperationData.OutputAttributes[0]);
 	InType1 DefaultValue1 = PCGMetadataAttribute::GetValueWithBroadcast<InType1>(OperationData.SourceAttributes[0], PCGInvalidEntryKey);
 	InType2 DefaultValue2 = PCGMetadataAttribute::GetValueWithBroadcast<InType2>(OperationData.SourceAttributes[1], PCGInvalidEntryKey);
 	OutType DefaultOutputValue = BinaryOp(DefaultValue1, DefaultValue2);
@@ -202,28 +237,31 @@ inline void FPCGMetadataElementBase::DoBinaryOp(FOperationData& OperationData, T
 			++Iterator2;
 		}
 	}
+
+	return true;
 }
 
 template <typename InType1, typename InType2, typename InType3, typename OutType>
-inline void FPCGMetadataElementBase::DoTernaryOp(FOperationData& OperationData, TFunction<OutType(const InType1&, const InType2&, const InType3&)> TernaryOp) const
+inline bool FPCGMetadataElementBase::DoTernaryOp(FOperationData& OperationData, TFunction<OutType(const InType1&, const InType2&, const InType3&)> TernaryOp) const
 {
 	check(OperationData.SourceAttributes[0]);
 	check(OperationData.SourceAttributes[1]);
 	check(OperationData.SourceAttributes[2]);
+	check(OperationData.OutputAttributes[0]);
 
 	// All values in OperationData.Iterators are UniquePtr. Dereference them here to make the syntax less cumbersome.
 	// Also some iterators can be null, it means they need to use the first iterator (that should never be null).
 	// We just need to make sure to not increment the first iterator multiple times in one loop.
 	check(OperationData.Iterators[0]);
 
-	IPCGMetadataEntryIterator& Iterator1 = *OperationData.Iterators[0];
-	IPCGMetadataEntryIterator& Iterator2 = (OperationData.Iterators.Num() >= 2 && OperationData.Iterators[1]) ? *OperationData.Iterators[1] : Iterator1;
-	IPCGMetadataEntryIterator& Iterator3 = (OperationData.Iterators.Num() >= 3 && OperationData.Iterators[2]) ? *OperationData.Iterators[2] : Iterator1;
+	bool bShouldIncrementIterator2 = OperationData.Iterators.Num() >= 2 && OperationData.Iterators[1];
+	bool bShouldIncrementIterator3 = OperationData.Iterators.Num() >= 3 && OperationData.Iterators[2];
 
-	bool bShouldIncrementIterator2 = &Iterator1 != &Iterator2;
-	bool bShouldIncrementIterator3 = &Iterator1 != &Iterator3;
+	IPCGMetadataEntryIterator& Iterator1 = *OperationData.Iterators[0];
+	IPCGMetadataEntryIterator& Iterator2 = bShouldIncrementIterator2 ? *OperationData.Iterators[1] : Iterator1;
+	IPCGMetadataEntryIterator& Iterator3 = bShouldIncrementIterator3 ? *OperationData.Iterators[2] : Iterator1;
 
-	FPCGMetadataAttribute<OutType>* OutputAttribute = static_cast<FPCGMetadataAttribute<OutType>*>(OperationData.OutputAttribute);
+	FPCGMetadataAttribute<OutType>* OutputAttribute = static_cast<FPCGMetadataAttribute<OutType>*>(OperationData.OutputAttributes[0]);
 	InType1 DefaultValue1 = PCGMetadataAttribute::GetValueWithBroadcast<InType1>(OperationData.SourceAttributes[0], PCGInvalidEntryKey);
 	InType2 DefaultValue2 = PCGMetadataAttribute::GetValueWithBroadcast<InType2>(OperationData.SourceAttributes[1], PCGInvalidEntryKey);
 	InType3 DefaultValue3 = PCGMetadataAttribute::GetValueWithBroadcast<InType3>(OperationData.SourceAttributes[2], PCGInvalidEntryKey);
@@ -260,4 +298,150 @@ inline void FPCGMetadataElementBase::DoTernaryOp(FOperationData& OperationData,
 			++Iterator3;
 		}
 	}
+
+	return true;
+}
+
+template <typename InType1, typename InType2, typename InType3, typename InType4, typename OutType>
+inline bool FPCGMetadataElementBase::DoQuaternaryOp(FOperationData& OperationData, TFunction<OutType(const InType1&, const InType2&, const InType3&, const InType4&)> QuaternaryOp) const
+{
+	check(OperationData.SourceAttributes[0]);
+	check(OperationData.SourceAttributes[1]);
+	check(OperationData.SourceAttributes[2]);
+	check(OperationData.SourceAttributes[3]);
+	check(OperationData.OutputAttributes[0]);
+
+	// All values in OperationData.Iterators are UniquePtr. Dereference them here to make the syntax less cumbersome.
+	// Also some iterators can be null, it means they need to use the first iterator (that should never be null).
+	// We just need to make sure to not increment the first iterator multiple times in one loop.
+	check(OperationData.Iterators[0]);
+
+	bool bShouldIncrementIterator2 = OperationData.Iterators.Num() >= 2 && OperationData.Iterators[1];
+	bool bShouldIncrementIterator3 = OperationData.Iterators.Num() >= 3 && OperationData.Iterators[2];
+	bool bShouldIncrementIterator4 = OperationData.Iterators.Num() >= 4 && OperationData.Iterators[3];
+
+	IPCGMetadataEntryIterator& Iterator1 = *OperationData.Iterators[0];
+	IPCGMetadataEntryIterator& Iterator2 = bShouldIncrementIterator2 ? *OperationData.Iterators[1] : Iterator1;
+	IPCGMetadataEntryIterator& Iterator3 = bShouldIncrementIterator3 ? *OperationData.Iterators[2] : Iterator1;
+	IPCGMetadataEntryIterator& Iterator4 = bShouldIncrementIterator4 ? *OperationData.Iterators[3] : Iterator1;
+
+	FPCGMetadataAttribute<OutType>* OutputAttribute = static_cast<FPCGMetadataAttribute<OutType>*>(OperationData.OutputAttributes[0]);
+	InType1 DefaultValue1 = PCGMetadataAttribute::GetValueWithBroadcast<InType1>(OperationData.SourceAttributes[0], PCGInvalidEntryKey);
+	InType2 DefaultValue2 = PCGMetadataAttribute::GetValueWithBroadcast<InType2>(OperationData.SourceAttributes[1], PCGInvalidEntryKey);
+	InType3 DefaultValue3 = PCGMetadataAttribute::GetValueWithBroadcast<InType3>(OperationData.SourceAttributes[2], PCGInvalidEntryKey);
+	InType3 DefaultValue4 = PCGMetadataAttribute::GetValueWithBroadcast<InType3>(OperationData.SourceAttributes[3], PCGInvalidEntryKey);
+	OutType DefaultOutputValue = QuaternaryOp(DefaultValue1, DefaultValue2, DefaultValue3, DefaultValue4);
+	OutputAttribute->SetDefaultValue(DefaultOutputValue);
+
+	for (int32 i = 0; i < OperationData.NumberOfElementsToProcess; ++i)
+	{
+		PCGMetadataEntryKey EntryKey1 = *Iterator1;
+
+		// If the entry key is invalid, nothing to do
+		if (EntryKey1 != PCGInvalidEntryKey)
+		{
+			PCGMetadataEntryKey EntryKey2 = *Iterator2;
+			PCGMetadataEntryKey EntryKey3 = *Iterator3;
+			PCGMetadataEntryKey EntryKey4 = *Iterator4;
+
+			InType1 Value1 = PCGMetadataAttribute::GetValueWithBroadcast<InType1>(OperationData.SourceAttributes[0], EntryKey1);
+			InType2 Value2 = PCGMetadataAttribute::GetValueWithBroadcast<InType2>(OperationData.SourceAttributes[1], EntryKey2);
+			InType3 Value3 = PCGMetadataAttribute::GetValueWithBroadcast<InType3>(OperationData.SourceAttributes[2], EntryKey3);
+			InType4 Value4 = PCGMetadataAttribute::GetValueWithBroadcast<InType4>(OperationData.SourceAttributes[3], EntryKey4);
+
+			OutType OutputValue = QuaternaryOp(Value1, Value2, Value3, Value4);
+
+			OutputAttribute->SetValue(EntryKey1, OutputValue);
+		}
+
+		++Iterator1;
+		if (bShouldIncrementIterator2)
+		{
+			++Iterator2;
+		}
+
+		if (bShouldIncrementIterator3)
+		{
+			++Iterator3;
+		}
+
+		if (bShouldIncrementIterator4)
+		{
+			++Iterator4;
+		}
+	}
+
+	return true;
+}
+
+template <typename InType, typename OutType>
+inline bool FPCGMetadataElementBase::DoUnaryOpMultipleOutputs(FOperationData & OperationData, TFunction<void(const InType&, TArray<OutType>&, uint32)> UnaryOp) const
+{
+	check(OperationData.SourceAttributes[0]);
+
+	// All values in OperationData.Iterators are UniquePtr. Dereference them here to make the syntax less cumbersome.
+	check(OperationData.Iterators[0]);
+	IPCGMetadataEntryIterator& Iterator1 = *OperationData.Iterators[0];
+
+	const uint32 NumOutputs = OperationData.OutputAttributes.Num();
+
+	TArray<OutType> OutputValues;
+	TArray<FPCGMetadataAttribute<OutType>*> OutputAttributes;
+
+	// Output value array is reserved not resized. Allow us to validate that the callback properly pushed
+	// the expected amount of outputs.
+	OutputValues.Reserve(NumOutputs);
+	OutputAttributes.SetNum(NumOutputs);
+
+	InType DefaultValue = PCGMetadataAttribute::GetValueWithBroadcast<InType>(OperationData.SourceAttributes[0], PCGInvalidEntryKey);
+	UnaryOp(DefaultValue, OutputValues, NumOutputs);
+
+	if (OutputValues.Num() != NumOutputs)
+	{
+		UE_LOG(LogPCG, Error, TEXT("[DoUnaryOpMultipleOutputs] Callback didn't provided the number of expected outputs."));
+		return false;
+	}
+
+	for (uint32 j = 0; j < NumOutputs; ++j)
+	{
+		// Output Attributes can be null if there is nothing connected to the output pin.
+		OutputAttributes[j] = static_cast<FPCGMetadataAttribute<OutType>*>(OperationData.OutputAttributes[j]);
+		if (OutputAttributes[j])
+		{
+			OutputAttributes[j]->SetDefaultValue(OutputValues[j]);
+		}
+	}
+
+	for (int32 i = 0; i < OperationData.NumberOfElementsToProcess; ++i)
+	{
+		PCGMetadataEntryKey EntryKey = *Iterator1;
+
+		// If the entry key is invalid, nothing to do
+		if (EntryKey != PCGInvalidEntryKey)
+		{
+			InType Value = PCGMetadataAttribute::GetValueWithBroadcast<InType>(OperationData.SourceAttributes[0], EntryKey);
+
+			// Make sure the array is emptied (but reserved)
+			OutputValues.Empty(/*Slack=*/NumOutputs);
+			UnaryOp(Value, OutputValues, NumOutputs);
+
+			if (OutputValues.Num() != NumOutputs)
+			{
+				UE_LOG(LogPCG, Error, TEXT("[DoUnaryOpMultipleOutputs] Callback didn't provided the number of expected outputs."));
+				return false;
+			}
+
+			for (uint32 j = 0; j < NumOutputs; ++j)
+			{
+				if (OutputAttributes[j])
+				{
+					OutputAttributes[j]->SetValue(EntryKey, OutputValues[j]);
+				}
+			}
+		}
+
+		++Iterator1;
+	}
+
+	return true;
 }
\ No newline at end of file
diff --git a/Engine/Plugins/Experimental/PCG/Source/PCG/Public/Metadata/PCGMetadata.h b/Engine/Plugins/Experimental/PCG/Source/PCG/Public/Metadata/PCGMetadata.h
--- a/Engine/Plugins/Experimental/PCG/Source/PCG/Public/Metadata/PCGMetadata.h
+++ b/Engine/Plugins/Experimental/PCG/Source/PCG/Public/Metadata/PCGMetadata.h
@@ -71,6 +71,9 @@ public:
 	UFUNCTION(BlueprintCallable, Category = "PCG|Metadata", meta=(DisplayName = "Create Vector4 Attribute"))
 	void CreateVector4Attribute(FName AttributeName, FVector4 DefaultValue, bool bAllowsInterpolation, bool bOverrideParent = true);
 
+	UFUNCTION(BlueprintCallable, Category = "PCG|Metadata")
+	void CreateVector2Attribute(FName AttributeName, FVector2D DefaultValue, bool bAllowsInterpolation, bool bOverrideParent = true);
+
 	UFUNCTION(BlueprintCallable, Category = "PCG|Metadata")
 	void CreateRotatorAttribute(FName AttributeName, FRotator DefaultValue, bool bAllowsInterpolation, bool bOverrideParent = true);
 
diff --git a/Engine/Plugins/Experimental/PCG/Source/PCG/Public/Metadata/PCGMetadataAccessor.h b/Engine/Plugins/Experimental/PCG/Source/PCG/Public/Metadata/PCGMetadataAccessor.h
--- a/Engine/Plugins/Experimental/PCG/Source/PCG/Public/Metadata/PCGMetadataAccessor.h
+++ b/Engine/Plugins/Experimental/PCG/Source/PCG/Public/Metadata/PCGMetadataAccessor.h
@@ -48,6 +48,12 @@ public:
 	UFUNCTION(BlueprintCallable, Category = "PCG|Metadata")
 	static void SetVector4AttributeByMetadataKey(UPARAM(ref) int64& Key, UPCGMetadata* Metadata, FName AttributeName, const FVector4& Value);
 
+	UFUNCTION(BlueprintCallable, Category = "PCG|Metadata")
+	static FVector2D GetVector2AttributeByMetadataKey(int64 Key, const UPCGMetadata* Metadata, FName AttributeName);
+
+	UFUNCTION(BlueprintCallable, Category = "PCG|Metadata")
+	static void SetVector2AttributeByMetadataKey(UPARAM(ref) int64& Key, UPCGMetadata* Metadata, FName AttributeName, const FVector2D& Value);
+
 	UFUNCTION(BlueprintCallable, Category = "PCG|Metadata")
 	static FRotator GetRotatorAttributeByMetadataKey(int64 Key, const UPCGMetadata* Metadata, FName AttributeName);
 
@@ -120,6 +126,12 @@ public:
 	UFUNCTION(BlueprintCallable, Category = "PCG|Metadata", meta = (ScriptMethod))
 	static void SetVector4Attribute(UPARAM(ref) FPCGPoint& Point, UPCGMetadata* Metadata, FName AttributeName, const FVector4& Value);
 
+	UFUNCTION(BlueprintCallable, Category = "PCG|Metadata", meta = (ScriptMethod))
+	static FVector2D GetVector2Attribute(const FPCGPoint& Point, const UPCGMetadata* Metadata, FName AttributeName);
+
+	UFUNCTION(BlueprintCallable, Category = "PCG|Metadata", meta = (ScriptMethod))
+	static void SetVector2Attribute(UPARAM(ref) FPCGPoint& Point, UPCGMetadata* Metadata, FName AttributeName, const FVector2D& Value);
+
 	UFUNCTION(BlueprintCallable, Category = "PCG|Metadata", meta = (ScriptMethod))
 	static FRotator GetRotatorAttribute(const FPCGPoint& Point, const UPCGMetadata* Metadata, FName AttributeName);
 
