Added non-overlapping multi-landscape support in PCG landscape data

diff --git a/Engine/Plugins/Experimental/PCG/Source/PCG/Private/Data/PCGLandscapeData.cpp b/Engine/Plugins/Experimental/PCG/Source/PCG/Private/Data/PCGLandscapeData.cpp
--- a/Engine/Plugins/Experimental/PCG/Source/PCG/Private/Data/PCGLandscapeData.cpp
+++ b/Engine/Plugins/Experimental/PCG/Source/PCG/Private/Data/PCGLandscapeData.cpp
@@ -13,31 +13,46 @@
 #include "Landscape.h"
 #include "LandscapeEdit.h"
 
-void UPCGLandscapeData::Initialize(ALandscapeProxy* InLandscape, const FBox& InBounds, bool bInHeightOnly, bool bInUseMetadata)
+void UPCGLandscapeData::Initialize(const TArray<TWeakObjectPtr<ALandscapeProxy>>& InLandscapes, const FBox& InBounds, bool bInHeightOnly, bool bInUseMetadata)
 {
-	check(InLandscape);
-	Landscape = InLandscape;
-	TargetActor = InLandscape;
+	for (TWeakObjectPtr<ALandscapeProxy> InLandscape : InLandscapes)
+	{
+		if (InLandscape.IsValid())
+		{
+			Landscapes.Emplace(InLandscape.Get());
+
+			// Build landscape info list
+			LandscapeInfos.Emplace(PCGHelpers::GetLandscapeBounds(InLandscape.Get()), InLandscape->GetLandscapeInfo());
+		}
+	}
+
+	check(!Landscapes.IsEmpty());
+
+	ALandscapeProxy* FirstLandscape = Landscapes[0].Get();
+
+	TargetActor = FirstLandscape;
 	Bounds = InBounds;
 	bHeightOnly = bInHeightOnly;
 	bUseMetadata = bInUseMetadata;
 
-	Transform = Landscape->GetActorTransform();
+	Transform = FirstLandscape->GetActorTransform();
 
-	// Cache additional data
-	LandscapeInfo = Landscape ? Landscape->GetLandscapeInfo() : nullptr;
-	LandscapeCache = Landscape && Landscape->GetWorld() && Landscape->GetWorld()->GetSubsystem<UPCGSubsystem>() ? Landscape->GetWorld()->GetSubsystem<UPCGSubsystem>()->GetLandscapeCache() : nullptr;
+	// Store cache pointer for easier access
+	LandscapeCache = (FirstLandscape->GetWorld() && FirstLandscape->GetWorld()->GetSubsystem<UPCGSubsystem>()) ? FirstLandscape->GetWorld()->GetSubsystem<UPCGSubsystem>()->GetLandscapeCache() : nullptr;
 
 	// TODO: find a better way to do this - maybe there should be a prototype metadata in the landscape cache
 	if (LandscapeCache)
 	{
 		if (!bHeightOnly && bUseMetadata)
 		{
-			const TArray<FName> Layers = LandscapeCache->GetLayerNames(Landscape.Get());
-
-			for (const FName& Layer : Layers)
+			for (TSoftObjectPtr<ALandscapeProxy> Landscape : Landscapes)
 			{
-				Metadata->CreateFloatAttribute(Layer, 0.0f, /*bAllowInterpolation=*/true);
+				const TArray<FName> Layers = LandscapeCache->GetLayerNames(Landscape.Get());
+
+				for (const FName& Layer : Layers)
+				{
+					Metadata->CreateFloatAttribute(Layer, 0.0f, /*bAllowInterpolation=*/true);
+				}
 			}
 		}
 	}
@@ -47,6 +62,32 @@ void UPCGLandscapeData::Initialize(ALandscapeProxy* InLandscape, const FBox& InB
 	}
 }
 
+void UPCGLandscapeData::PostLoad()
+{
+	ALandscapeProxy* FirstLandscape = nullptr;
+
+	for (TSoftObjectPtr<ALandscapeProxy> Landscape : Landscapes)
+	{
+		Landscape.LoadSynchronous();
+		if (Landscape.Get())
+		{
+			LandscapeInfos.Emplace(PCGHelpers::GetLandscapeBounds(Landscape.Get()), Landscape->GetLandscapeInfo());
+
+			if (!FirstLandscape)
+			{
+				FirstLandscape = Landscape.Get();
+			}
+		}
+		else
+		{
+			UE_LOG(LogPCG, Warning, TEXT("Was unable to load landscape in landscape data"));
+		}
+	}
+
+	TargetActor = FirstLandscape;
+	LandscapeCache = (FirstLandscape && FirstLandscape->GetWorld() && FirstLandscape->GetWorld()->GetSubsystem<UPCGSubsystem>()) ? FirstLandscape->GetWorld()->GetSubsystem<UPCGSubsystem>()->GetLandscapeCache() : nullptr;
+}
+
 FBox UPCGLandscapeData::GetBounds() const
 {
 	return Bounds;
@@ -62,13 +103,21 @@ FBox UPCGLandscapeData::GetStrictBounds() const
 bool UPCGLandscapeData::SamplePoint(const FTransform& InTransform, const FBox& InBounds, FPCGPoint& OutPoint, UPCGMetadata* OutMetadata) const
 {
 	//TRACE_CPUPROFILER_EVENT_SCOPE(UPCGLandscapeData::SamplePoint);
-	if (!LandscapeInfo || !LandscapeCache)
+	if (!LandscapeCache)
 	{
 		return false;
 	}
 
+	const ULandscapeInfo* LandscapeInfo = GetLandscapeInfo(InTransform.GetLocation());
+	if (!LandscapeInfo)
+	{
+		return false;
+	}
+
+	const FTransform& LandscapeTransform = LandscapeInfo->GetLandscapeProxy()->GetTransform();
+
 	// TODO: compute full transform when we want to support bounds
-	const FVector LocalPoint = Transform.InverseTransformPosition(InTransform.GetLocation());
+	const FVector LocalPoint = LandscapeTransform.InverseTransformPosition(InTransform.GetLocation());
 	const FIntPoint ComponentMapKey(FMath::FloorToInt(LocalPoint.X / LandscapeInfo->ComponentSizeQuads), FMath::FloorToInt(LocalPoint.Y / LandscapeInfo->ComponentSizeQuads));
 
 	if (ULandscapeComponent* LandscapeComponent = LandscapeInfo->XYtoComponentMap.FindRef(ComponentMapKey))
@@ -103,14 +152,12 @@ const UPCGPointData* UPCGLandscapeData::CreatePointData(FPCGContext* Context, co
 {
 	TRACE_CPUPROFILER_EVENT_SCOPE(UPCGLandscapeData::CreatePointData);
 
-	if (!LandscapeInfo || !LandscapeCache)
+	if (LandscapeInfos.IsEmpty() || !LandscapeCache)
 	{
 		UE_LOG(LogPCG, Error, TEXT("PCG Landscape cache or Landscape info are not initialized"));
 		return nullptr;
 	}
 
-	const int32 ComponentSizeQuads = LandscapeInfo->ComponentSizeQuads;
-
 	UPCGPointData* Data = NewObject<UPCGPointData>();
 	Data->InitializeFromData(this);
 	TArray<FPCGPoint>& Points = Data->GetMutablePoints();
@@ -127,79 +174,111 @@ const UPCGPointData* UPCGLandscapeData::CreatePointData(FPCGContext* Context, co
 		return Data;
 	}
 
-	// TODO: add offset to nearest edge, will have an impact if the grid size doesn't match the landscape size
-	const FVector MinPt = Transform.InverseTransformPosition(EffectiveBounds.Min);
-	const FVector MaxPt = Transform.InverseTransformPosition(EffectiveBounds.Max);
+	for (const TPair<FBox, ULandscapeInfo*>& LandscapeInfoPair : LandscapeInfos)
+	{
+		ULandscapeInfo* LandscapeInfo = LandscapeInfoPair.Value;
 
-	// Note: the MaxX/Y here are inclusive, hence the floor & the +1 in the sizes
-	const int32 MinX = FMath::CeilToInt(MinPt.X);
-	const int32 MaxX = FMath::FloorToInt(MaxPt.X);
-	const int32 MinY = FMath::CeilToInt(MinPt.Y);
-	const int32 MaxY = FMath::FloorToInt(MaxPt.Y);
+		const FTransform& LandscapeTransform = LandscapeInfo->GetLandscapeProxy()->GetTransform();
+		const int32 ComponentSizeQuads = LandscapeInfo->ComponentSizeQuads;
 
-	//Early out if the bounds do not overlap any landscape vertices
-	if (MaxX < MinX || MaxY < MinY)
-	{
-		return Data;
-	}
+		// TODO: add offset to nearest edge, will have an impact if the grid size doesn't match the landscape size
+		const FVector MinPt = LandscapeTransform.InverseTransformPosition(EffectiveBounds.Min);
+		const FVector MaxPt = LandscapeTransform.InverseTransformPosition(EffectiveBounds.Max);
 
-	const int64 PointCountUpperBound = (1 + MaxX - MinX) * (1 + MaxY - MinY);
-	if (PointCountUpperBound > 0)
-	{
-		Points.Reserve(PointCountUpperBound);
-	}
+		// Note: the MaxX/Y here are inclusive, hence the floor & the +1 in the sizes
+		const int32 MinX = FMath::CeilToInt(MinPt.X);
+		const int32 MaxX = FMath::FloorToInt(MaxPt.X);
+		const int32 MinY = FMath::CeilToInt(MinPt.Y);
+		const int32 MaxY = FMath::FloorToInt(MaxPt.Y);
 
-	const int32 MinComponentX = MinX / ComponentSizeQuads;
-	const int32 MaxComponentX = MaxX / ComponentSizeQuads;
-	const int32 MinComponentY = MinY / ComponentSizeQuads;
-	const int32 MaxComponentY = MaxY / ComponentSizeQuads;
+		//Early out if the bounds do not overlap any landscape vertices
+		if (MaxX < MinX || MaxY < MinY)
+		{
+			continue;
+		}
 
-	for (int32 ComponentX = MinComponentX; ComponentX <= MaxComponentX; ++ComponentX)
-	{
-		for (int32 ComponentY = MinComponentY; ComponentY <= MaxComponentY; ++ComponentY)
+		const int64 PointCountUpperBound = (1 + MaxX - MinX) * (1 + MaxY - MinY);
+		const int32 PointsBeforeNum = Points.Num();
+		if (PointCountUpperBound > 0)
 		{
-			FIntPoint ComponentMapKey(ComponentX, ComponentY);
-			if (ULandscapeComponent* LandscapeComponent = LandscapeInfo->XYtoComponentMap.FindRef(ComponentMapKey))
-			{
-				const FPCGLandscapeCacheEntry* LandscapeCacheEntry = LandscapeCache->GetCacheEntry(LandscapeComponent, ComponentMapKey);
+			Points.Reserve(PointsBeforeNum + PointCountUpperBound);
+		}
 
-				if (!LandscapeCacheEntry)
+		const int32 MinComponentX = MinX / ComponentSizeQuads;
+		const int32 MaxComponentX = MaxX / ComponentSizeQuads;
+		const int32 MinComponentY = MinY / ComponentSizeQuads;
+		const int32 MaxComponentY = MaxY / ComponentSizeQuads;
+
+		for (int32 ComponentX = MinComponentX; ComponentX <= MaxComponentX; ++ComponentX)
+		{
+			for (int32 ComponentY = MinComponentY; ComponentY <= MaxComponentY; ++ComponentY)
+			{
+				FIntPoint ComponentMapKey(ComponentX, ComponentY);
+				if (ULandscapeComponent* LandscapeComponent = LandscapeInfo->XYtoComponentMap.FindRef(ComponentMapKey))
 				{
-					continue;
-				}
+					const FPCGLandscapeCacheEntry* LandscapeCacheEntry = LandscapeCache->GetCacheEntry(LandscapeComponent, ComponentMapKey);
 
-				// Rebase our bounds in the component referential
-				const int32 LocalMinX = FMath::Clamp(MinX - ComponentMapKey.X * ComponentSizeQuads, 0, ComponentSizeQuads - 1);
-				const int32 LocalMaxX = FMath::Clamp(MaxX - ComponentMapKey.X * ComponentSizeQuads, 0, ComponentSizeQuads - 1);
+					if (!LandscapeCacheEntry)
+					{
+						continue;
+					}
 
-				const int32 LocalMinY = FMath::Clamp(MinY - ComponentMapKey.Y * ComponentSizeQuads, 0, ComponentSizeQuads - 1);
-				const int32 LocalMaxY = FMath::Clamp(MaxY - ComponentMapKey.Y * ComponentSizeQuads, 0, ComponentSizeQuads - 1);
+					// Rebase our bounds in the component referential
+					const int32 LocalMinX = FMath::Clamp(MinX - ComponentMapKey.X * ComponentSizeQuads, 0, ComponentSizeQuads - 1);
+					const int32 LocalMaxX = FMath::Clamp(MaxX - ComponentMapKey.X * ComponentSizeQuads, 0, ComponentSizeQuads - 1);
 
-				// We can't really copy data from the component points wholesale because the component points have an additional boundary point.
-				// TODO: consider optimizing this, though it will impact the Sample then
-				for (int32 LocalX = LocalMinX; LocalX <= LocalMaxX; ++LocalX)
-				{
-					for (int32 LocalY = LocalMinY; LocalY <= LocalMaxY; ++LocalY)
-					{
-						const int32 PointIndex = LocalX + LocalY * (ComponentSizeQuads + 1);
+					const int32 LocalMinY = FMath::Clamp(MinY - ComponentMapKey.Y * ComponentSizeQuads, 0, ComponentSizeQuads - 1);
+					const int32 LocalMaxY = FMath::Clamp(MaxY - ComponentMapKey.Y * ComponentSizeQuads, 0, ComponentSizeQuads - 1);
 
-						FPCGPoint& Point = Points.Emplace_GetRef();
-						if (bHeightOnly)
-						{
-							LandscapeCacheEntry->GetPointHeightOnly(PointIndex, Point);
-						}
-						else
+					// We can't really copy data from the component points wholesale because the component points have an additional boundary point.
+					// TODO: consider optimizing this, though it will impact the Sample then
+					for (int32 LocalX = LocalMinX; LocalX <= LocalMaxX; ++LocalX)
+					{
+						for (int32 LocalY = LocalMinY; LocalY <= LocalMaxY; ++LocalY)
 						{
-							LandscapeCacheEntry->GetPoint(PointIndex, Point, bUseMetadata ? Data->Metadata : nullptr);
+							const int32 PointIndex = LocalX + LocalY * (ComponentSizeQuads + 1);
+
+							FPCGPoint& Point = Points.Emplace_GetRef();
+							if (bHeightOnly)
+							{
+								LandscapeCacheEntry->GetPointHeightOnly(PointIndex, Point);
+							}
+							else
+							{
+								LandscapeCacheEntry->GetPoint(PointIndex, Point, bUseMetadata ? Data->Metadata : nullptr);
+							}
 						}
 					}
 				}
 			}
 		}
-	}
 
-	check(Points.Num() <= PointCountUpperBound);
-	UE_LOG(LogPCG, Verbose, TEXT("Landscape %s extracted %d of %d potential points"), *Landscape->GetFName().ToString(), Points.Num(), PointCountUpperBound);
+		check(Points.Num() - PointsBeforeNum <= PointCountUpperBound);
+		UE_LOG(LogPCG, Verbose, TEXT("Landscape %s extracted %d of %d potential points"), *LandscapeInfo->GetLandscapeProxy()->GetFName().ToString(), Points.Num() - PointsBeforeNum, PointCountUpperBound);
+	}
 
 	return Data;
+}
+
+const ULandscapeInfo* UPCGLandscapeData::GetLandscapeInfo(const FVector& InPosition) const
+{
+	check(!LandscapeInfos.IsEmpty());
+
+	// Early out
+	if (LandscapeInfos.Num() == 1)
+	{
+		return LandscapeInfos[0].Value;
+	}
+
+	// As discussed in the header, this loop here is the reason why we do not really support overlapping landscapes.
+	// TODO: we could maybe improve on this if we find the "nearest" landscape on a Z perspective, but this might still lead to issues
+	for (const TPair<FBox, ULandscapeInfo*>& LandscapeInfoPair : LandscapeInfos)
+	{
+		if (LandscapeInfoPair.Key.IsInside(InPosition))
+		{
+			return LandscapeInfoPair.Value;
+		}
+	}
+
+	return nullptr;
 }
\ No newline at end of file
diff --git a/Engine/Plugins/Experimental/PCG/Source/PCG/Private/PCGComponent.cpp b/Engine/Plugins/Experimental/PCG/Source/PCG/Private/PCGComponent.cpp
--- a/Engine/Plugins/Experimental/PCG/Source/PCG/Private/PCGComponent.cpp
+++ b/Engine/Plugins/Experimental/PCG/Source/PCG/Private/PCGComponent.cpp
@@ -1661,7 +1661,7 @@ UPCGData* UPCGComponent::CreateActorPCGData(AActor* Actor, bool bParseActor)
 	else if (ALandscapeProxy* Landscape = Cast<ALandscapeProxy>(Actor))
 	{
 		UPCGLandscapeData* Data = NewObject<UPCGLandscapeData>();
-		Data->Initialize(Landscape, GetGridBounds(Actor), /*bHeightOnly=*/false, /*bUseMetadata=*/Graph && Graph->bLandscapeUsesMetadata);
+		Data->Initialize({ Landscape }, GetGridBounds(Actor), /*bHeightOnly=*/false, /*bUseMetadata=*/Graph && Graph->bLandscapeUsesMetadata);
 
 		return Data;
 	}
@@ -1819,31 +1819,42 @@ UPCGData* UPCGComponent::CreateLandscapePCGData(bool bHeightOnly)
 	}
 
 	const UPCGSpatialData* ActorSpatialData = Cast<const UPCGSpatialData>(ActorData);
-	ALandscapeProxy* Landscape = nullptr;
+
+	FBox ActorBounds;
 
 	if (ActorSpatialData)
 	{
-		const FBox ActorDataBounds = ActorSpatialData->GetBounds();
-		Landscape = PCGHelpers::GetLandscape(Actor->GetWorld(), ActorDataBounds);
+		ActorBounds = ActorSpatialData->GetBounds();
 	}
 	else
 	{
 		FVector Origin;
 		FVector Extent;
 		Actor->GetActorBounds(/*bOnlyCollidingComponents=*/false, Origin, Extent);
-
-		Landscape = PCGHelpers::GetLandscape(Actor->GetWorld(), FBox::BuildAABB(Origin, Extent));
+		ActorBounds = FBox::BuildAABB(Origin, Extent);
 	}
 
-	if (!Landscape)
+	TArray<TWeakObjectPtr<ALandscapeProxy>> Landscapes = PCGHelpers::GetLandscapeProxies(Actor->GetWorld(), ActorBounds);
+
+	if (Landscapes.IsEmpty())
 	{
 		// No landscape found
 		return nullptr;
 	}
 
+	FBox LandscapeBounds(EForceInit::ForceInit);
+
+	for (TWeakObjectPtr<ALandscapeProxy> Landscape : Landscapes)
+	{
+		if (Landscape.IsValid())
+		{
+			LandscapeBounds += GetGridBounds(Landscape.Get());
+		}
+	}
+
 	// TODO: we're creating separate landscape data instances here so we can do some tweaks on it (such as storing the right target actor) but this probably should change
 	UPCGLandscapeData* LandscapeData = NewObject<UPCGLandscapeData>();
-	LandscapeData->Initialize(Landscape, GetGridBounds(Landscape), bHeightOnly, /*bUseMetadata=*/Graph && Graph->bLandscapeUsesMetadata);
+	LandscapeData->Initialize(Landscapes, LandscapeBounds, bHeightOnly, /*bUseMetadata=*/Graph && Graph->bLandscapeUsesMetadata);
 	// Need to override target actor for this one, not the landscape
 	LandscapeData->TargetActor = Actor;
 
diff --git a/Engine/Plugins/Experimental/PCG/Source/PCG/Public/Data/PCGLandscapeData.h b/Engine/Plugins/Experimental/PCG/Source/PCG/Public/Data/PCGLandscapeData.h
--- a/Engine/Plugins/Experimental/PCG/Source/PCG/Public/Data/PCGLandscapeData.h
+++ b/Engine/Plugins/Experimental/PCG/Source/PCG/Public/Data/PCGLandscapeData.h
@@ -11,13 +11,20 @@ class ALandscapeProxy;
 class ULandscapeInfo;
 class UPCGLandscapeCache;
 
+/**
+* Landscape data access abstraction for PCG. Supports multi-landscape access, but it assumes that they are not overlapping.
+*/
 UCLASS(BlueprintType, ClassGroup = (Procedural))
 class PCG_API UPCGLandscapeData : public UPCGSurfaceData
 {
 	GENERATED_BODY()
 
 public:
-	void Initialize(ALandscapeProxy* InLandscape, const FBox& InBounds, bool bInHeightOnly, bool bInUseMetadata);
+	void Initialize(const TArray<TWeakObjectPtr<ALandscapeProxy>>& InLandscapes, const FBox& InBounds, bool bInHeightOnly, bool bInUseMetadata);
+
+	// ~Begin UObject interface
+	virtual void PostLoad();
+	// ~End UObject interface
 
 	// ~Begin UPCGData interface
 	virtual EPCGDataType GetDataType() const override { return EPCGDataType::Landscape | Super::GetDataType(); }
@@ -36,13 +43,18 @@ public:
 	virtual const UPCGPointData* CreatePointData(FPCGContext* Context, const FBox& InBounds) const override;
 	// ~End UPCGConcreteDataWithPointCache interface
 
-	// TODO: add on property changed to clear cached data
+	// TODO: add on property changed to clear cached data. This is used to populate the LandscapeInfos array.
 	UPROPERTY(BlueprintReadWrite, EditAnywhere, Category = SourceData)
-	TSoftObjectPtr<ALandscapeProxy> Landscape;
+	TArray<TSoftObjectPtr<ALandscapeProxy>> Landscapes;
 
 	bool IsUsingMetadata() const { return bUseMetadata; }
 
 protected:
+	/** Returns the landscape info associated to the first landscape that contains the given position
+	* Note that this implicitly removes support for overlapping landscapes, which might be a future TODO
+	*/
+	const ULandscapeInfo* GetLandscapeInfo(const FVector& InPosition) const;
+
 	UPROPERTY()
 	FBox Bounds = FBox(EForceInit::ForceInit);
 
@@ -54,6 +66,6 @@ protected:
 
 private:
 	// Transient data
-	ULandscapeInfo* LandscapeInfo = nullptr;
+	TArray<TPair<FBox, ULandscapeInfo*>> LandscapeInfos;
 	UPCGLandscapeCache* LandscapeCache = nullptr;
 };
