Minor changes
- Spawn actor now outputs the points  properly in all cases
- Flushing cache should now work in engine/pie too
- Default debug scaling is now extents
- Changed the caching flag on "from actor" pcg blueprints
- Added some BP methods to get local/world bounds on actors and also ignore pcg generated components if needed

diff --git a/Engine/Plugins/Experimental/PCG/Source/PCG/Private/Elements/PCGSpawnActor.cpp b/Engine/Plugins/Experimental/PCG/Source/PCG/Private/Elements/PCGSpawnActor.cpp
--- a/Engine/Plugins/Experimental/PCG/Source/PCG/Private/Elements/PCGSpawnActor.cpp
+++ b/Engine/Plugins/Experimental/PCG/Source/PCG/Private/Elements/PCGSpawnActor.cpp
@@ -83,7 +83,6 @@ bool FPCGSpawnActorElement::ExecuteInternal(FPCGContext* Context) const
 		return true;
 	}
 
-	bool bShouldPassThroughInputs = (Settings->Option != EPCGSpawnActorOption::NoMerging && Settings->GetSubgraph() != nullptr);
 	const bool bForceDisableActorParsing = (Settings->bForceDisableActorParsing);
 
 	// Pass-through exclusions & settings
@@ -202,7 +201,7 @@ bool FPCGSpawnActorElement::ExecuteInternal(FPCGContext* Context) const
 
 						for (UPCGComponent* PCGComponent : PCGComponents)
 						{
-							if (!bShouldPassThroughInputs)
+							if (Settings->Option == EPCGSpawnActorOption::NoMerging)
 							{
 								if (bForceDisableActorParsing)
 								{
@@ -235,8 +234,10 @@ bool FPCGSpawnActorElement::ExecuteInternal(FPCGContext* Context) const
 			}
 		}
 
-		// Finally, pass through the input if we're merging the PCG here
-		if (bShouldPassThroughInputs)
+		// Finally, pass through the input, in all cases: 
+		// - if it's not merged, will be the input points directly
+		// - if it's merged but there is no subgraph, will be the input points directly
+		// - if it's merged and there is a subgraph, we'd need to pass the data for it to be given to the subgraph
 		{
 			Outputs.Add(Input);
 		}
diff --git a/Engine/Plugins/Experimental/PCG/Source/PCG/Private/Helpers/PCGBlueprintHelpers.cpp b/Engine/Plugins/Experimental/PCG/Source/PCG/Private/Helpers/PCGBlueprintHelpers.cpp
--- a/Engine/Plugins/Experimental/PCG/Source/PCG/Private/Helpers/PCGBlueprintHelpers.cpp
+++ b/Engine/Plugins/Experimental/PCG/Source/PCG/Private/Helpers/PCGBlueprintHelpers.cpp
@@ -91,3 +91,12 @@ FVector UPCGBlueprintHelpers::GetLocalCenter(const FPCGPoint& InPoint)
 	return InPoint.GetLocalCenter();
 }
 
+FBox UPCGBlueprintHelpers::GetActorBoundsPCG(AActor* InActor, bool bIgnorePCGCreatedComponents)
+{
+	return PCGHelpers::GetActorBounds(InActor, bIgnorePCGCreatedComponents);
+}
+
+FBox UPCGBlueprintHelpers::GetActorLocalBoundsPCG(AActor* InActor, bool bIgnorePCGCreatedComponents)
+{
+	return PCGHelpers::GetActorLocalBounds(InActor, bIgnorePCGCreatedComponents);
+}
\ No newline at end of file
diff --git a/Engine/Plugins/Experimental/PCG/Source/PCG/Private/PCGHelpers.cpp b/Engine/Plugins/Experimental/PCG/Source/PCG/Private/PCGHelpers.cpp
--- a/Engine/Plugins/Experimental/PCG/Source/PCG/Private/PCGHelpers.cpp
+++ b/Engine/Plugins/Experimental/PCG/Source/PCG/Private/PCGHelpers.cpp
@@ -56,7 +56,7 @@ namespace PCGHelpers
 		}
 	}
 
-	FBox GetActorBounds(AActor* InActor)
+	FBox GetActorBounds(AActor* InActor, bool bIgnorePCGCreatedComponents)
 	{
 		// Specialized version of GetComponentsBoundingBox that skips over PCG generated components
 		// This is to ensure stable bounds and no timing issues (cleared ISMs, etc.)
@@ -65,21 +65,28 @@ namespace PCGHelpers
 		const bool bNonColliding = true;
 		const bool bIncludeFromChildActors = true;
 
-		InActor->ForEachComponent<UPrimitiveComponent>(bIncludeFromChildActors, [&](const UPrimitiveComponent* InPrimComp)
+		if (InActor)
+		{
+			InActor->ForEachComponent<UPrimitiveComponent>(bIncludeFromChildActors, [&](const UPrimitiveComponent* InPrimComp)
 			{
 				// Note: we omit the IsRegistered check here (e.g. InPrimComp->IsRegistered() )
 				// since this can be called in a scope where the components are temporarily unregistered
 				if ((bNonColliding || InPrimComp->IsCollisionEnabled()) &&
-					!InPrimComp->ComponentTags.Contains(DefaultPCGTag))
+					(!bIgnorePCGCreatedComponents || !InPrimComp->ComponentTags.Contains(DefaultPCGTag)))
 				{
 					Box += InPrimComp->Bounds.GetBox();
 				}
 			});
+		}
+		else
+		{
+			UE_LOG(LogPCG, Error, TEXT("Actor is invalid in GetActorBounds"));
+		}
 
 		return Box;
 	}
 
-	FBox GetActorLocalBounds(AActor* InActor)
+	FBox GetActorLocalBounds(AActor* InActor, bool bIgnorePCGCreatedComponents)
 	{
 		// Specialized version of CalculateComponentsBoundingBoxInLocalScape that skips over PCG generated components
 		// This is to ensure stable bounds and no timing issues (cleared ISMs, etc.)
@@ -88,18 +95,25 @@ namespace PCGHelpers
 		const bool bNonColliding = true;
 		const bool bIncludeFromChildActors = true;
 
-		const FTransform& ActorToWorld = InActor->GetTransform();
-		const FTransform WorldToActor = ActorToWorld.Inverse();
-
-		InActor->ForEachComponent<UPrimitiveComponent>(bIncludeFromChildActors, [&](const UPrimitiveComponent* InPrimComp)
+		if (InActor)
 		{
+			const FTransform& ActorToWorld = InActor->GetTransform();
+			const FTransform WorldToActor = ActorToWorld.Inverse();
+
+			InActor->ForEachComponent<UPrimitiveComponent>(bIncludeFromChildActors, [&](const UPrimitiveComponent* InPrimComp)
+			{
 				if ((bNonColliding || InPrimComp->IsCollisionEnabled()) &&
-					!InPrimComp->ComponentTags.Contains(DefaultPCGTag))
+					(!bIgnorePCGCreatedComponents || !InPrimComp->ComponentTags.Contains(DefaultPCGTag)))
 				{
 					const FTransform ComponentToActor = InPrimComp->GetComponentTransform() * WorldToActor;
 					Box += InPrimComp->CalcBounds(ComponentToActor).GetBox();
 				}
-		});
+			});
+		}
+		else
+		{
+			UE_LOG(LogPCG, Error, TEXT("Actor is invalid in GetActorLocalBounds"));
+		}
 
 		return Box;
 	}
diff --git a/Engine/Plugins/Experimental/PCG/Source/PCG/Private/PCGHelpers.h b/Engine/Plugins/Experimental/PCG/Source/PCG/Private/PCGHelpers.h
--- a/Engine/Plugins/Experimental/PCG/Source/PCG/Private/PCGHelpers.h
+++ b/Engine/Plugins/Experimental/PCG/Source/PCG/Private/PCGHelpers.h
@@ -27,8 +27,8 @@ namespace PCGHelpers
 
 	FBox OverlapBounds(const FBox& InBoxA, const FBox& InBoxB);
 
-	FBox GetActorBounds(AActor* InActor);
-	FBox GetActorLocalBounds(AActor* InActor);
+	FBox GetActorBounds(AActor* InActor, bool bIgnorePCGCreatedComponents = true);
+	FBox GetActorLocalBounds(AActor* InActor, bool bIgnorePCGCreatedComponents = true);
 	FBox GetLandscapeBounds(ALandscapeProxy* InLandscape);
 
 	ALandscape* GetLandscape(UWorld* InWorld, const FBox& InActorBounds);
diff --git a/Engine/Plugins/Experimental/PCG/Source/PCG/Private/PCGSubsystem.cpp b/Engine/Plugins/Experimental/PCG/Source/PCG/Private/PCGSubsystem.cpp
--- a/Engine/Plugins/Experimental/PCG/Source/PCG/Private/PCGSubsystem.cpp
+++ b/Engine/Plugins/Experimental/PCG/Source/PCG/Private/PCGSubsystem.cpp
@@ -8,13 +8,13 @@
 #include "Grid/PCGPartitionActor.h"
 #include "Helpers/PCGActorHelpers.h"
 
+#include "Engine/Engine.h"
 #include "Engine/World.h"
 #include "HAL/IConsoleManager.h"
 #include "Math/GenericOctree.h"
 
 #if WITH_EDITOR
 #include "Editor.h"
-#include "Engine/Engine.h"
 #include "ActorPartition/ActorPartitionSubsystem.h"
 #include "WorldPartition/WorldPartition.h"
 #include "WorldPartition/WorldPartitionHelpers.h"
@@ -22,7 +22,6 @@
 #include "ObjectTools.h"
 #endif
 
-#if WITH_EDITOR
 namespace PCGSubsystemConsole
 {
 	static FAutoConsoleCommand CommandFlushCache(
@@ -30,16 +29,35 @@ namespace PCGSubsystemConsole
 		TEXT("Clears the PCG results cache."),
 		FConsoleCommandDelegate::CreateLambda([]()
 			{
+				UWorld* World = nullptr;
+
+#if WITH_EDITOR
 				if (GEditor)
 				{
-					if (UWorld* World = GEditor->GetEditorWorldContext().World())
+					if (GEditor->PlayWorld)
+					{
+						World = GEditor->PlayWorld;
+					}
+					else
 					{
-						World->GetSubsystem<UPCGSubsystem>()->FlushCache();
+						World = GEditor->GetEditorWorldContext().World();
 					}
 				}
+				else
+#endif
+				if (GEngine)
+				{
+					World = GEngine->GetCurrentPlayWorld();
+				}
+
+				if (World && World->GetSubsystem<UPCGSubsystem>())
+				{
+					World->GetSubsystem<UPCGSubsystem>()->FlushCache();
+				}
 			}));
 }
 
+#if WITH_EDITOR
 namespace PCGSubsystem
 {
 	FPCGTaskId ForEachIntersectingCell(FPCGGraphExecutor* GraphExecutor, UWorld* World, const FBox& InBounds, bool bCreateActor, bool bLoadCell, bool bSaveActors, TFunctionRef<FPCGTaskId(APCGPartitionActor*, const FBox&, const TArray<FPCGTaskId>&)> InOperation);
@@ -1463,14 +1481,6 @@ void UPCGSubsystem::CleanFromCache(const IPCGElement* InElement)
 	}
 }
 
-void UPCGSubsystem::FlushCache()
-{
-	if (GraphExecutor)
-	{
-		GraphExecutor->GetCache().ClearCache();
-	}
-}
-
 void UPCGSubsystem::BuildLandscapeCache()
 {
 	if (UPCGLandscapeCache* LandscapeCache = GetLandscapeCache())
@@ -1500,3 +1510,11 @@ void UPCGSubsystem::ResetPartitionActorsMap()
 }
 
 #endif // WITH_EDITOR
+
+void UPCGSubsystem::FlushCache()
+{
+	if (GraphExecutor)
+	{
+		GraphExecutor->GetCache().ClearCache();
+	}
+}
\ No newline at end of file
diff --git a/Engine/Plugins/Experimental/PCG/Source/PCG/Public/Helpers/PCGBlueprintHelpers.h b/Engine/Plugins/Experimental/PCG/Source/PCG/Public/Helpers/PCGBlueprintHelpers.h
--- a/Engine/Plugins/Experimental/PCG/Source/PCG/Public/Helpers/PCGBlueprintHelpers.h
+++ b/Engine/Plugins/Experimental/PCG/Source/PCG/Public/Helpers/PCGBlueprintHelpers.h
@@ -56,4 +56,10 @@ public:
 
 	UFUNCTION(BlueprintCallable, Category = "PCG|Points", meta = (ScriptMethod))
 	static FVector GetLocalCenter(const FPCGPoint& InPoint);
+
+	UFUNCTION(BlueprintCallable, Category = "PCG|Helpers", meta = (ScriptMethod))
+	static FBox GetActorBoundsPCG(AActor* InActor, bool bIgnorePCGCreatedComponents = true);
+
+	UFUNCTION(BlueprintCallable, Category = "PCG|Helpers", meta = (ScriptMethod))
+	static FBox GetActorLocalBoundsPCG(AActor* InActor, bool bIgnorePCGCreatedComponents = true);
 };
\ No newline at end of file
diff --git a/Engine/Plugins/Experimental/PCG/Source/PCG/Public/PCGDebug.h b/Engine/Plugins/Experimental/PCG/Source/PCG/Public/PCGDebug.h
--- a/Engine/Plugins/Experimental/PCG/Source/PCG/Public/PCGDebug.h
+++ b/Engine/Plugins/Experimental/PCG/Source/PCG/Public/PCGDebug.h
@@ -29,7 +29,7 @@ public:
 	float PointScale = 1.0f;
 
 	UPROPERTY(BlueprintReadWrite, EditAnywhere, Category = Debug)
-	EPCGDebugVisScaleMethod ScaleMethod = EPCGDebugVisScaleMethod::Relative;
+	EPCGDebugVisScaleMethod ScaleMethod = EPCGDebugVisScaleMethod::Extents;
 
 	UPROPERTY(BlueprintReadWrite, EditAnywhere, Category = Debug)
 	TSoftObjectPtr<UStaticMesh> PointMesh;
