[PCG] Metadata native operation: Refactoring operation functions to support multiple outputs for each one.

diff --git a/Engine/Plugins/Experimental/PCG/Source/PCG/Private/Elements/Metadata/PCGMetadataBitwiseOpElement.cpp b/Engine/Plugins/Experimental/PCG/Source/PCG/Private/Elements/Metadata/PCGMetadataBitwiseOpElement.cpp
--- a/Engine/Plugins/Experimental/PCG/Source/PCG/Private/Elements/Metadata/PCGMetadataBitwiseOpElement.cpp
+++ b/Engine/Plugins/Experimental/PCG/Source/PCG/Private/Elements/Metadata/PCGMetadataBitwiseOpElement.cpp
@@ -102,11 +102,11 @@ bool FPCGMetadataBitwiseElement::DoOperation(FOperationData& OperationData) cons
 
 	if (Settings->Operation == EPCGMedadataBitwiseOperation::BitwiseNot)
 	{
-		DoUnaryOp<int64, int64>(OperationData, [](const int64& Value) -> int64 { return PCGMetadataBitwiseSettings::UnaryOp(Value); });
+		DoUnaryOp<int64>(OperationData, [](const int64& Value) -> int64 { return PCGMetadataBitwiseSettings::UnaryOp(Value); });
 	}
 	else
 	{
-		DoBinaryOp<int64, int64, int64>(OperationData, [Operation = Settings->Operation](const int64& Value1, const int64& Value2) -> int64 { return PCGMetadataBitwiseSettings::BinaryOp(Value1, Value2, Operation); });
+		DoBinaryOp<int64, int64>(OperationData, [Operation = Settings->Operation](const int64& Value1, const int64& Value2) -> int64 { return PCGMetadataBitwiseSettings::BinaryOp(Value1, Value2, Operation); });
 	}
 
 	return true;
diff --git a/Engine/Plugins/Experimental/PCG/Source/PCG/Private/Elements/Metadata/PCGMetadataBooleanOpElement.cpp b/Engine/Plugins/Experimental/PCG/Source/PCG/Private/Elements/Metadata/PCGMetadataBooleanOpElement.cpp
--- a/Engine/Plugins/Experimental/PCG/Source/PCG/Private/Elements/Metadata/PCGMetadataBooleanOpElement.cpp
+++ b/Engine/Plugins/Experimental/PCG/Source/PCG/Private/Elements/Metadata/PCGMetadataBooleanOpElement.cpp
@@ -102,11 +102,11 @@ bool FPCGMetadataBooleanElement::DoOperation(FOperationData& OperationData) cons
 
 	if (Settings->Operation == EPCGMedadataBooleanOperation::Not)
 	{
-		DoUnaryOp<bool, bool>(OperationData, [](const bool& Value) -> bool { return PCGMetadataBooleanSettings::UnaryOp(Value); });
+		DoUnaryOp<bool>(OperationData, [](const bool& Value) -> bool { return PCGMetadataBooleanSettings::UnaryOp(Value); });
 	}
 	else
 	{
-		DoBinaryOp<bool, bool, bool>(OperationData, [Operation = Settings->Operation](const bool& Value1, const bool& Value2) -> bool { return PCGMetadataBooleanSettings::BinaryOp(Value1, Value2, Operation); });
+		DoBinaryOp<bool, bool>(OperationData, [Operation = Settings->Operation](const bool& Value1, const bool& Value2) -> bool { return PCGMetadataBooleanSettings::BinaryOp(Value1, Value2, Operation); });
 	}
 
 	return true;
diff --git a/Engine/Plugins/Experimental/PCG/Source/PCG/Private/Elements/Metadata/PCGMetadataBreakVector.cpp b/Engine/Plugins/Experimental/PCG/Source/PCG/Private/Elements/Metadata/PCGMetadataBreakVector.cpp
--- a/Engine/Plugins/Experimental/PCG/Source/PCG/Private/Elements/Metadata/PCGMetadataBreakVector.cpp
+++ b/Engine/Plugins/Experimental/PCG/Source/PCG/Private/Elements/Metadata/PCGMetadataBreakVector.cpp
@@ -12,48 +12,58 @@
 
 namespace PCGMetadataBreakVectorSettings
 {
-	template <typename InType = FVector>
-	inline void DoBreak(const InType& Value, TArray<double>& OutValues, uint32 NumOutputs)
+	template <typename InType>
+	inline double GetFirstIndex(const InType& Value)
 	{
-		check(NumOutputs == 4);
+		return Value.X;
+	}
 
-		OutValues.Add(Value.X);
-		OutValues.Add(Value.Y);
-		OutValues.Add(Value.Z);
-		OutValues.Add(0.0);
+	template<>
+	inline double GetFirstIndex<FRotator>(const FRotator& Value)
+	{
+		return Value.Roll;
 	}
 
-	template <>
-	inline void DoBreak<FRotator>(const FRotator& Value, TArray<double>& OutValues, uint32 NumOutputs)
+	template <typename InType>
+	inline double GetSecondIndex(const InType& Value)
 	{
-		check(NumOutputs == 4);
+		return Value.Y;
+	}
 
-		OutValues.Add(Value.Roll);
-		OutValues.Add(Value.Pitch);
-		OutValues.Add(Value.Yaw);
-		OutValues.Add(0.0);
+	template<>
+	inline double GetSecondIndex<FRotator>(const FRotator& Value)
+	{
+		return Value.Pitch;
 	}
 
-	template <>
-	inline void DoBreak<FVector2D>(const FVector2D& Value, TArray<double>& OutValues, uint32 NumOutputs)
+	template <typename InType>
+	inline double GetThirdIndex(const InType& Value)
 	{
-		check(NumOutputs == 4);
+		return Value.Z;
+	}
 
-		OutValues.Add(Value.X);
-		OutValues.Add(Value.Y);
-		OutValues.Add(0.0);
-		OutValues.Add(0.0);
+	template<>
+	inline double GetThirdIndex<FRotator>(const FRotator& Value)
+	{
+		return Value.Yaw;
 	}
 
-	template <>
-	inline void DoBreak<FVector4>(const FVector4& Value, TArray<double>& OutValues, uint32 NumOutputs)
+	template<>
+	inline double GetThirdIndex<FVector2D>(const FVector2D& Value)
 	{
-		check(NumOutputs == 4);
+		return 0.0;
+	}
 
-		OutValues.Add(Value.X);
-		OutValues.Add(Value.Y);
-		OutValues.Add(Value.Z);
-		OutValues.Add(Value.W);
+	template <typename InType>
+	inline double GetFourthIndex(const InType& Value)
+	{
+		return 0.0;
+	}
+
+	template<>
+	inline double GetFourthIndex<FVector4>(const FVector4& Value)
+	{
+		return Value.W;
 	}
 
 	inline constexpr bool IsValidType(uint16 TypeId)
@@ -149,7 +159,11 @@ bool FPCGMetadataBreakVectorElement::DoOperation(FOperationData& OperationData)
 		}
 		else
 		{
-			return DoUnaryOpMultipleOutputs<AttributeType, double>(OperationData, PCGMetadataBreakVectorSettings::DoBreak<AttributeType>);
+			return DoUnaryOp<AttributeType>(OperationData, 
+				PCGMetadataBreakVectorSettings::GetFirstIndex<AttributeType>,
+				PCGMetadataBreakVectorSettings::GetSecondIndex<AttributeType>,
+				PCGMetadataBreakVectorSettings::GetThirdIndex<AttributeType>,
+				PCGMetadataBreakVectorSettings::GetFourthIndex<AttributeType>);
 		}
 	};
 
diff --git a/Engine/Plugins/Experimental/PCG/Source/PCG/Private/Elements/Metadata/PCGMetadataCompareOpElement.cpp b/Engine/Plugins/Experimental/PCG/Source/PCG/Private/Elements/Metadata/PCGMetadataCompareOpElement.cpp
--- a/Engine/Plugins/Experimental/PCG/Source/PCG/Private/Elements/Metadata/PCGMetadataCompareOpElement.cpp
+++ b/Engine/Plugins/Experimental/PCG/Source/PCG/Private/Elements/Metadata/PCGMetadataCompareOpElement.cpp
@@ -128,7 +128,7 @@ bool FPCGMetadataCompareElement::DoOperation(FOperationData& OperationData) cons
 		}
 		else
 		{
-			DoBinaryOp<AttributeType, AttributeType, bool>(OperationData, 
+			DoBinaryOp<AttributeType, AttributeType>(OperationData, 
 				[Operation, Tolerance](const AttributeType& Value1, const AttributeType& Value2) -> bool { 
 					return PCGMetadataCompareSettings::ApplyCompare(Value1, Value2, Operation, Tolerance); 
 				});
diff --git a/Engine/Plugins/Experimental/PCG/Source/PCG/Private/Elements/Metadata/PCGMetadataMathsOpElement.cpp b/Engine/Plugins/Experimental/PCG/Source/PCG/Private/Elements/Metadata/PCGMetadataMathsOpElement.cpp
--- a/Engine/Plugins/Experimental/PCG/Source/PCG/Private/Elements/Metadata/PCGMetadataMathsOpElement.cpp
+++ b/Engine/Plugins/Experimental/PCG/Source/PCG/Private/Elements/Metadata/PCGMetadataMathsOpElement.cpp
@@ -243,17 +243,17 @@ bool FPCGMetadataMathsElement::DoOperation(FOperationData& OperationData) const
 			if (PCGMetadataMathsSettings::IsUnaryOp(Operation))
 			{
 				TRACE_CPUPROFILER_EVENT_SCOPE(FPCGMetadataMathsElement::ExecuteInternal::UnaryOp);
-				DoUnaryOp<AttributeType, AttributeType>(OperationData, [Operation](const AttributeType& Value) -> AttributeType { return PCGMetadataMathsSettings::UnaryOp(Value, Operation); });
+				DoUnaryOp<AttributeType>(OperationData, [Operation](const AttributeType& Value) -> AttributeType { return PCGMetadataMathsSettings::UnaryOp(Value, Operation); });
 			}
 			else if (PCGMetadataMathsSettings::IsBinaryOp(Operation))
 			{
 				TRACE_CPUPROFILER_EVENT_SCOPE(FPCGMetadataMathsElement::ExecuteInternal::BinaryOp);
-				DoBinaryOp<AttributeType, AttributeType, AttributeType>(OperationData, [Operation](const AttributeType& Value1, const AttributeType& Value2) -> AttributeType { return PCGMetadataMathsSettings::BinaryOp(Value1, Value2, Operation); });
+				DoBinaryOp<AttributeType, AttributeType>(OperationData, [Operation](const AttributeType& Value1, const AttributeType& Value2) -> AttributeType { return PCGMetadataMathsSettings::BinaryOp(Value1, Value2, Operation); });
 			}
 			else if (PCGMetadataMathsSettings::IsTernaryOp(Operation))
 			{
 				TRACE_CPUPROFILER_EVENT_SCOPE(FPCGMetadataMathsElement::ExecuteInternal::TernaryOp);
-				DoTernaryOp<AttributeType, AttributeType, AttributeType, AttributeType>(OperationData, [Operation](const AttributeType& Value1, const AttributeType& Value2, const AttributeType& Value3) -> AttributeType { return PCGMetadataMathsSettings::TernaryOp(Value1, Value2, Value3, Operation); });
+				DoTernaryOp<AttributeType, AttributeType, AttributeType>(OperationData, [Operation](const AttributeType& Value1, const AttributeType& Value2, const AttributeType& Value3) -> AttributeType { return PCGMetadataMathsSettings::TernaryOp(Value1, Value2, Value3, Operation); });
 			}
 		}
 	};
diff --git a/Engine/Plugins/Experimental/PCG/Source/PCG/Private/Elements/Metadata/PCGMetadataOpElementBase.cpp b/Engine/Plugins/Experimental/PCG/Source/PCG/Private/Elements/Metadata/PCGMetadataOpElementBase.cpp
--- a/Engine/Plugins/Experimental/PCG/Source/PCG/Private/Elements/Metadata/PCGMetadataOpElementBase.cpp
+++ b/Engine/Plugins/Experimental/PCG/Source/PCG/Private/Elements/Metadata/PCGMetadataOpElementBase.cpp
@@ -219,33 +219,37 @@ bool FPCGMetadataElementBase::ExecuteInternal(FPCGContext* Context) const
 	// So first forward outputs and create the attribute
 	OperationData.OutputAttributes.SetNum(Settings->GetOutputPinNum());
 
-	auto CreateAttributes = [&](auto DummyOutValue) -> bool
+	auto CreateAttribute = [&](uint32 OutputIndex, auto DummyOutValue) -> bool
 	{
 		using AttributeType = decltype(DummyOutValue);
 
 		AttributeType DefaultValue = PCGMetadataAttribute::GetValueWithBroadcast<AttributeType>(OperationData.SourceAttributes[0], PCGDefaultValueKey);
 
-		for (uint32 i = 0; i < NumberOfOutputs; ++i)
+		if (!ConnectedOutputs[OutputIndex])
 		{
-			if (!ConnectedOutputs[i])
-			{
-				OperationData.OutputAttributes[i] = nullptr;
-				continue;
-			}
+			OperationData.OutputAttributes[OutputIndex] = nullptr;
+			return true;
+		}
 
-			FPCGTaggedData& OutputData = Outputs.Add_GetRef(InputTaggedData[0]);
-			OutputData.Pin = Settings->GetOutputPinLabel(i);
+		FPCGTaggedData& OutputData = Outputs.Add_GetRef(InputTaggedData[0]);
+		OutputData.Pin = Settings->GetOutputPinLabel(OutputIndex);
 
-			FName OutputAttributeFinalName = Settings->GetOutputAttributeName(OutputAttributeName, i);
+		FName OutputAttributeFinalName = Settings->GetOutputAttributeName(OutputAttributeName, OutputIndex);
 
-			UPCGMetadata* OutMetadata = nullptr;
+		UPCGMetadata* OutMetadata = nullptr;
 
-			PCGMetadataElementCommon::DuplicateTaggedData(InputTaggedData[0], OutputData, OutMetadata);
-			OperationData.OutputAttributes[i] = PCGMetadataElementCommon::ClearOrCreateAttribute(OutMetadata, OutputAttributeFinalName, DefaultValue);
-			PCGMetadataElementCommon::CopyEntryToValueKeyMap(SourceMetadata[0], OperationData.SourceAttributes[0], OperationData.OutputAttributes[i]);
+		PCGMetadataElementCommon::DuplicateTaggedData(InputTaggedData[0], OutputData, OutMetadata);
+		OperationData.OutputAttributes[OutputIndex] = PCGMetadataElementCommon::ClearOrCreateAttribute(OutMetadata, OutputAttributeFinalName, DefaultValue);
+		PCGMetadataElementCommon::CopyEntryToValueKeyMap(SourceMetadata[0], OperationData.SourceAttributes[0], OperationData.OutputAttributes[OutputIndex]);
 
-			// If something went wrong stop here.
-			if (OperationData.OutputAttributes[i] == nullptr)
+		return OperationData.OutputAttributes[OutputIndex] != nullptr;
+	};
+
+	auto CreateAllSameAttributes = [&](auto DummyOutValue) -> bool
+	{
+		for (uint32 i = 0; i < NumberOfOutputs; ++i)
+		{
+			if (!CreateAttribute(i, DummyOutValue))
 			{
 				return false;
 			}
@@ -256,7 +260,27 @@ bool FPCGMetadataElementBase::ExecuteInternal(FPCGContext* Context) const
 
 	OperationData.OutputType = Settings->GetOutputType(OperationData.MostComplexInputType);
 
-	if (!PCGMetadataAttribute::CallbackWithRightType(OperationData.OutputType, CreateAttributes))
+	bool bCreateAttributeSucceeded = true;
+
+	if (!Settings->HasDifferentOutputTypes())
+	{
+		bCreateAttributeSucceeded = PCGMetadataAttribute::CallbackWithRightType(OperationData.OutputType, CreateAllSameAttributes);
+	}
+	else
+	{
+		TArray<uint16> OutputTypes = Settings->GetAllOutputTypes();
+		check(OutputTypes.Num() == NumberOfOutputs);
+
+		for (uint32 i = 0; i < NumberOfOutputs && bCreateAttributeSucceeded; ++i)
+		{
+			bCreateAttributeSucceeded &= PCGMetadataAttribute::CallbackWithRightType(OutputTypes[i],
+				[&](auto DummyOutValue) -> bool {
+					return CreateAttribute(i, DummyOutValue);
+				});
+		}
+	}
+
+	if (!bCreateAttributeSucceeded)
 	{
 		PCGE_LOG(Error, "Error while creating output attributes");
 		Outputs.Empty();
diff --git a/Engine/Plugins/Experimental/PCG/Source/PCG/Private/Elements/Metadata/PCGMetadataRotatorOpElement.cpp b/Engine/Plugins/Experimental/PCG/Source/PCG/Private/Elements/Metadata/PCGMetadataRotatorOpElement.cpp
--- a/Engine/Plugins/Experimental/PCG/Source/PCG/Private/Elements/Metadata/PCGMetadataRotatorOpElement.cpp
+++ b/Engine/Plugins/Experimental/PCG/Source/PCG/Private/Elements/Metadata/PCGMetadataRotatorOpElement.cpp
@@ -202,26 +202,26 @@ bool FPCGMetadataRotatorElement::DoOperation(FOperationData& OperationData) cons
 		{
 			if (PCGMetadataRotatorSettings::IsTransfromOp(Operation))
 			{
-				DoBinaryOp<AttributeType, FTransform, AttributeType>(OperationData, [Operation](const AttributeType& Value, const FTransform& Transform)->AttributeType {
+				DoBinaryOp<AttributeType, FTransform>(OperationData, [Operation](const AttributeType& Value, const FTransform& Transform)->AttributeType {
 					return PCGMetadataRotatorSettings::ApplyTransformOperation(Value, Transform, Operation);
 					});
 			}
 			else if (PCGMetadataRotatorSettings::IsUnaryOp(Operation))
 			{
-				DoUnaryOp<AttributeType, AttributeType>(OperationData, [Operation](const AttributeType& Value)->AttributeType {
+				DoUnaryOp<AttributeType>(OperationData, [Operation](const AttributeType& Value)->AttributeType {
 					double DummyDouble = 0.0;
 					return PCGMetadataRotatorSettings::ApplyRotatorOperation(Value, AttributeType{}, DummyDouble, Operation);
 				});
 			}
 			else if (PCGMetadataRotatorSettings::IsTernaryOp(Operation))
 			{
-				DoTernaryOp<AttributeType, AttributeType, double, AttributeType>(OperationData, [Operation](const AttributeType& Value1, const AttributeType& Value2, const double& Ratio)->AttributeType {
+				DoTernaryOp<AttributeType, AttributeType, double>(OperationData, [Operation](const AttributeType& Value1, const AttributeType& Value2, const double& Ratio)->AttributeType {
 					return PCGMetadataRotatorSettings::ApplyRotatorOperation(Value1, Value2, Ratio, Operation);
 				});
 			}
 			else
 			{
-				DoBinaryOp<AttributeType, AttributeType, AttributeType>(OperationData, [Operation](const AttributeType& Value1, const AttributeType& Value2)->AttributeType {
+				DoBinaryOp<AttributeType, AttributeType>(OperationData, [Operation](const AttributeType& Value1, const AttributeType& Value2)->AttributeType {
 					double DummyDouble = 0.0;
 					return PCGMetadataRotatorSettings::ApplyRotatorOperation(Value1, Value2, DummyDouble, Operation);
 				});
diff --git a/Engine/Plugins/Experimental/PCG/Source/PCG/Private/Elements/Metadata/PCGMetadataTransformOpElement.cpp b/Engine/Plugins/Experimental/PCG/Source/PCG/Private/Elements/Metadata/PCGMetadataTransformOpElement.cpp
--- a/Engine/Plugins/Experimental/PCG/Source/PCG/Private/Elements/Metadata/PCGMetadataTransformOpElement.cpp
+++ b/Engine/Plugins/Experimental/PCG/Source/PCG/Private/Elements/Metadata/PCGMetadataTransformOpElement.cpp
@@ -160,20 +160,20 @@ bool FPCGMetadataTransformElement::DoOperation(FOperationData& OperationData) co
 
 	if (PCGMetadataTransfromSettings::IsUnaryOp(Settings->Operation))
 	{
-		DoUnaryOp<FTransform, FTransform>(OperationData, [Operation = Settings->Operation, Mode = Settings->TransformLerpMode](const FTransform& Value)->FTransform {
+		DoUnaryOp<FTransform>(OperationData, [Operation = Settings->Operation, Mode = Settings->TransformLerpMode](const FTransform& Value)->FTransform {
 			double DummyDouble = 0.0;
 			return PCGMetadataTransfromSettings::ApplyTransformOperation(Value, FTransform{}, DummyDouble, Operation, Mode);
 			});
 	}
 	else if (PCGMetadataTransfromSettings::IsTernaryOp(Settings->Operation))
 	{
-		DoTernaryOp<FTransform, FTransform, double, FTransform>(OperationData, [Operation = Settings->Operation, Mode = Settings->TransformLerpMode](const FTransform& Value1, const FTransform& Value2, const double& Ratio)->FTransform {
+		DoTernaryOp<FTransform, FTransform, double>(OperationData, [Operation = Settings->Operation, Mode = Settings->TransformLerpMode](const FTransform& Value1, const FTransform& Value2, const double& Ratio)->FTransform {
 			return PCGMetadataTransfromSettings::ApplyTransformOperation(Value1, Value2, Ratio, Operation, Mode);
 			});
 	}
 	else
 	{
-		DoBinaryOp<FTransform, FTransform, FTransform>(OperationData, [Operation = Settings->Operation, Mode = Settings->TransformLerpMode](const FTransform& Value1, const FTransform& Value2)->FTransform {
+		DoBinaryOp<FTransform, FTransform>(OperationData, [Operation = Settings->Operation, Mode = Settings->TransformLerpMode](const FTransform& Value1, const FTransform& Value2)->FTransform {
 			double DummyDouble = 0.0;
 			return PCGMetadataTransfromSettings::ApplyTransformOperation(Value1, Value2, DummyDouble, Operation, Mode);
 			});
diff --git a/Engine/Plugins/Experimental/PCG/Source/PCG/Private/Elements/Metadata/PCGMetadataTrigOpElement.cpp b/Engine/Plugins/Experimental/PCG/Source/PCG/Private/Elements/Metadata/PCGMetadataTrigOpElement.cpp
--- a/Engine/Plugins/Experimental/PCG/Source/PCG/Private/Elements/Metadata/PCGMetadataTrigOpElement.cpp
+++ b/Engine/Plugins/Experimental/PCG/Source/PCG/Private/Elements/Metadata/PCGMetadataTrigOpElement.cpp
@@ -119,11 +119,11 @@ bool FPCGMetadataTrigElement::DoOperation(FOperationData& OperationData) const
 
 	if (Settings->Operation == EPCGMedadataTrigOperation::Atan2)
 	{
-		DoBinaryOp<double, double, double>(OperationData, [](const double& Value1, const double& Value2) -> double { return PCGMetadataTrigSettings::BinaryOp(Value1, Value2); });
+		DoBinaryOp<double, double>(OperationData, [](const double& Value1, const double& Value2) -> double { return PCGMetadataTrigSettings::BinaryOp(Value1, Value2); });
 	}
 	else
 	{
-		DoUnaryOp<double, double>(OperationData, [Operation = Settings->Operation](const double& Value) -> double { return PCGMetadataTrigSettings::UnaryOp(Value, Operation); });
+		DoUnaryOp<double>(OperationData, [Operation = Settings->Operation](const double& Value) -> double { return PCGMetadataTrigSettings::UnaryOp(Value, Operation); });
 	}
 
 	return true;
diff --git a/Engine/Plugins/Experimental/PCG/Source/PCG/Private/Elements/Metadata/PCGMetadataVectorOpElement.cpp b/Engine/Plugins/Experimental/PCG/Source/PCG/Private/Elements/Metadata/PCGMetadataVectorOpElement.cpp
--- a/Engine/Plugins/Experimental/PCG/Source/PCG/Private/Elements/Metadata/PCGMetadataVectorOpElement.cpp
+++ b/Engine/Plugins/Experimental/PCG/Source/PCG/Private/Elements/Metadata/PCGMetadataVectorOpElement.cpp
@@ -356,13 +356,13 @@ bool FPCGMetadataVectorElement::DoOperation(FOperationData& OperationData) const
 		{
 			if (PCGMetadataVectorSettings::IsTransformOp(Operation))
 			{
-				DoBinaryOp<AttributeType, FTransform, AttributeType>(OperationData, [Operation](const AttributeType& Value, const FTransform& Transform) -> AttributeType { return PCGMetadataVectorSettings::ApplyTransformOperation(Value, Transform, Operation);});
+				DoBinaryOp<AttributeType, FTransform>(OperationData, [Operation](const AttributeType& Value, const FTransform& Transform) -> AttributeType { return PCGMetadataVectorSettings::ApplyTransformOperation(Value, Transform, Operation);});
 			}
 			else if (OperationData.OutputType == (uint16)EPCGMetadataTypes::Double)
 			{
 				if (PCGMetadataVectorSettings::IsUnaryOp(Operation))
 				{
-					DoUnaryOp<AttributeType, double>(OperationData, [Operation](const AttributeType& Value) -> double {
+					DoUnaryOp<AttributeType>(OperationData, [Operation](const AttributeType& Value) -> double {
 						double Res = 0.0;
 						PCGMetadataVectorSettings::ApplyVectorOperation(Value, AttributeType{}, Res, Operation);
 						return Res;
@@ -370,7 +370,7 @@ bool FPCGMetadataVectorElement::DoOperation(FOperationData& OperationData) const
 				}
 				else
 				{
-					DoBinaryOp<AttributeType, AttributeType, double>(OperationData, [Operation](const AttributeType& Value1, const AttributeType& Value2) -> double {
+					DoBinaryOp<AttributeType, AttributeType>(OperationData, [Operation](const AttributeType& Value1, const AttributeType& Value2) -> double {
 						double Res = 0.0;
 						PCGMetadataVectorSettings::ApplyVectorOperation(Value1, Value2, Res, Operation);
 						return Res;
@@ -381,21 +381,21 @@ bool FPCGMetadataVectorElement::DoOperation(FOperationData& OperationData) const
 			{
 				if (PCGMetadataVectorSettings::IsUnaryOp(Operation))
 				{
-					DoUnaryOp<AttributeType, AttributeType>(OperationData, [Operation](const AttributeType& Value) -> AttributeType {
+					DoUnaryOp<AttributeType>(OperationData, [Operation](const AttributeType& Value) -> AttributeType {
 						double DummyDouble = 0.0;
 						return PCGMetadataVectorSettings::ApplyVectorOperation(Value, AttributeType{}, DummyDouble, Operation);
 						});
 				}
 				else if (PCGMetadataVectorSettings::IsTernaryOp(Operation))
 				{
-					DoTernaryOp<AttributeType, AttributeType, double, AttributeType>(OperationData, [Operation](const AttributeType& Value1, const AttributeType& Value2, const double& DoubleValue) -> AttributeType {
+					DoTernaryOp<AttributeType, AttributeType, double>(OperationData, [Operation](const AttributeType& Value1, const AttributeType& Value2, const double& DoubleValue) -> AttributeType {
 						double DoubleCopy = DoubleValue;
 						return PCGMetadataVectorSettings::ApplyVectorOperation(Value1, Value2, DoubleCopy, Operation);
 						});
 				}
 				else
 				{
-					DoBinaryOp<AttributeType, AttributeType, AttributeType>(OperationData, [Operation](const AttributeType& Value1, const AttributeType& Value2) -> AttributeType {
+					DoBinaryOp<AttributeType, AttributeType>(OperationData, [Operation](const AttributeType& Value1, const AttributeType& Value2) -> AttributeType {
 						double DummyDouble = 0.0;
 						return PCGMetadataVectorSettings::ApplyVectorOperation(Value1, Value2, DummyDouble, Operation);
 						});
diff --git a/Engine/Plugins/Experimental/PCG/Source/PCG/Public/Elements/Metadata/PCGMetadataOpElementBase.h b/Engine/Plugins/Experimental/PCG/Source/PCG/Public/Elements/Metadata/PCGMetadataOpElementBase.h
--- a/Engine/Plugins/Experimental/PCG/Source/PCG/Public/Elements/Metadata/PCGMetadataOpElementBase.h
+++ b/Engine/Plugins/Experimental/PCG/Source/PCG/Public/Elements/Metadata/PCGMetadataOpElementBase.h
@@ -90,6 +90,9 @@ public:
 	virtual uint16 GetOutputType(uint16 InputTypeId) const { return InputTypeId; };
 	virtual FName GetOutputAttributeName(FName BaseName, uint32 Index) const { return BaseName; }
 
+	virtual bool HasDifferentOutputTypes() const { return false; }
+	virtual TArray<uint16> GetAllOutputTypes() const { return TArray<uint16>(); };
+
 	bool IsMoreComplexType(uint16 FirstType, uint16 SecondType) const;
 
 	UPROPERTY(BlueprintReadWrite, EditAnywhere, Category = Output)
@@ -131,45 +134,53 @@ protected:
 
 	virtual bool DoOperation(FOperationData& OperationData) const = 0;
 
-	/* All operations return true if it succeeded false otherwise.
-	* For each operation, the callback takes as many inputs as needed, as const ref of InType.
-	* and must return OutType by value.
+	/* All operations can have a fixed number of inputs and a variable number of outputs.
+	* Each output need to have its own callback, all taking the exact number of "const InType&" as input
+	* and each can return a different output type.
 	*/
-	template <typename InType, typename OutType>
-	bool DoUnaryOp(FOperationData& OperationData, TFunction<OutType(const InType&)> UnaryOp) const;
-
-	template <typename InType1, typename InType2, typename OutType>
-	bool DoBinaryOp(FOperationData& OperationData, TFunction<OutType(const InType1&, const InType2&)> BinaryOp) const;
+	template <typename InType, typename... Callbacks>
+	bool DoUnaryOp(FOperationData& OperationData, Callbacks&& ...InCallbacks) const;
 
-	template <typename InType1, typename InType2, typename InType3, typename OutType>
-	bool DoTernaryOp(FOperationData& OperationData, TFunction<OutType(const InType1&, const InType2&, const InType3&)> TernaryOp) const;
+	template <typename InType1, typename InType2, typename... Callbacks>
+	bool DoBinaryOp(FOperationData& OperationData, Callbacks&& ...InCallbacks) const;
 
-	template <typename InType1, typename InType2, typename InType3, typename InType4, typename OutType>
-	bool DoQuaternaryOp(FOperationData& OperationData, TFunction<OutType(const InType1&, const InType2&, const InType3&, const InType4&)> QuaternaryOp) const;
+	template <typename InType1, typename InType2, typename InType3, typename... Callbacks>
+	bool DoTernaryOp(FOperationData& OperationData, Callbacks&& ...InCallbacks) const;
 
-	/* For multiple outputs, the callback won't return anything. But the callback should write results in
-	* the array passed as ref. Number of requested outputs is also passed as a parameter.
-	* Note that the array passed is "reserved" but not "resized", meaning that you need to push your results
-	* and not access the items directly with operator[].
-	*/
-	template <typename InType, typename OutType>
-	bool DoUnaryOpMultipleOutputs(FOperationData& OperationData, TFunction<void(const InType&, TArray<OutType>&, uint32)> UnaryOp) const;
+	template <typename InType1, typename InType2, typename InType3, typename InType4, typename... Callbacks>
+	bool DoQuaternaryOp(FOperationData& OperationData, Callbacks&& ...InCallbacks) const;
 };
 
-template <typename InType, typename OutType>
-inline bool FPCGMetadataElementBase::DoUnaryOp(FOperationData& OperationData, TFunction<OutType(const InType&)> UnaryOp) const
+template <typename InType, typename... Callbacks>
+inline bool FPCGMetadataElementBase::DoUnaryOp(FOperationData& OperationData, Callbacks&& ...InCallbacks) const
 {
 	check(OperationData.SourceAttributes[0]);
-	check(OperationData.OutputAttributes[0]);
+
+	const uint32 NumOutputs = OperationData.OutputAttributes.Num();
+	constexpr uint32 NumCallbacks = (uint32)sizeof...(InCallbacks);
+	check(NumOutputs == NumCallbacks);
 
 	// All values in OperationData.Iterators are UniquePtr. Dereference them here to make the syntax less cumbersome.
 	check(OperationData.Iterators[0]);
 	IPCGMetadataEntryIterator& Iterator1 = *OperationData.Iterators[0];
 
-	FPCGMetadataAttribute<OutType>* OutputAttribute = static_cast<FPCGMetadataAttribute<OutType>*>(OperationData.OutputAttributes[0]);
 	InType DefaultValue = PCGMetadataAttribute::GetValueWithBroadcast<InType>(OperationData.SourceAttributes[0], PCGInvalidEntryKey);
-	OutType DefaultOutputValue = UnaryOp(DefaultValue);
-	OutputAttribute->SetDefaultValue(DefaultOutputValue);
+
+	// Fold expressions. It will un-roll the packed argument "InCallbacks", and we can use "InCallbacks" as the current unrolled argument.
+	// In this case, we have multiple callbacks that all take "const InType&" as input, but can return different outputs.
+	// We use this return value to know in which type we need to cast our attribute.
+	uint32 j = 0;
+	([&]
+	{
+		check(j < NumOutputs);
+		auto OutValue = InCallbacks(DefaultValue);
+		using OutType = decltype(OutValue);
+
+		if (FPCGMetadataAttribute<OutType>* OutputAttribute = static_cast<FPCGMetadataAttribute<OutType>*>(OperationData.OutputAttributes[j++]))
+		{
+			OutputAttribute->SetDefaultValue(OutValue);
+		}
+	} (), ...);
 
 	for (int32 i = 0; i < OperationData.NumberOfElementsToProcess; ++i)
 	{
@@ -180,9 +191,18 @@ inline bool FPCGMetadataElementBase::DoUnaryOp(FOperationData& OperationData, TF
 		{
 			InType Value = PCGMetadataAttribute::GetValueWithBroadcast<InType>(OperationData.SourceAttributes[0], EntryKey);
 
-			OutType OutputValue = UnaryOp(Value);
+			j = 0;
+			([&]
+			{
+				check(j < NumOutputs);
+				auto OutValue = InCallbacks(Value);
+				using OutType = decltype(OutValue);
 
-			OutputAttribute->SetValue(EntryKey, OutputValue);
+				if (FPCGMetadataAttribute<OutType>* OutputAttribute = static_cast<FPCGMetadataAttribute<OutType>*>(OperationData.OutputAttributes[j++]))
+				{
+					OutputAttribute->SetValue(EntryKey, OutValue);
+				}
+			} (), ...);
 		}
 
 		++Iterator1;
@@ -191,12 +211,15 @@ inline bool FPCGMetadataElementBase::DoUnaryOp(FOperationData& OperationData, TF
 	return true;
 }
 
-template <typename InType1, typename InType2, typename OutType>
-inline bool FPCGMetadataElementBase::DoBinaryOp(FOperationData& OperationData, TFunction<OutType(const InType1&, const InType2&)> BinaryOp) const
+template <typename InType1, typename InType2, typename... Callbacks>
+inline bool FPCGMetadataElementBase::DoBinaryOp(FOperationData& OperationData, Callbacks&& ...InCallbacks) const
 {
 	check(OperationData.SourceAttributes[0]);
 	check(OperationData.SourceAttributes[1]);
-	check(OperationData.OutputAttributes[0]);
+
+	const uint32 NumOutputs = OperationData.OutputAttributes.Num();
+	constexpr uint32 NumCallbacks = (uint32)sizeof...(InCallbacks);
+	check(NumOutputs == NumCallbacks);
 
 	// All values in OperationData.Iterators are UniquePtr. Dereference them here to make the syntax less cumbersome.
 	// Also some iterators can be null, it means they need to use the first iterator (that should never be null).
@@ -208,11 +231,22 @@ inline bool FPCGMetadataElementBase::DoBinaryOp(FOperationData& OperationData, T
 	IPCGMetadataEntryIterator& Iterator1 = *OperationData.Iterators[0];
 	IPCGMetadataEntryIterator& Iterator2 = bShouldIncrementIterator2 ? *OperationData.Iterators[1] : Iterator1;
 
-	FPCGMetadataAttribute<OutType>* OutputAttribute = static_cast<FPCGMetadataAttribute<OutType>*>(OperationData.OutputAttributes[0]);
 	InType1 DefaultValue1 = PCGMetadataAttribute::GetValueWithBroadcast<InType1>(OperationData.SourceAttributes[0], PCGInvalidEntryKey);
 	InType2 DefaultValue2 = PCGMetadataAttribute::GetValueWithBroadcast<InType2>(OperationData.SourceAttributes[1], PCGInvalidEntryKey);
-	OutType DefaultOutputValue = BinaryOp(DefaultValue1, DefaultValue2);
-	OutputAttribute->SetDefaultValue(DefaultOutputValue);
+	
+	// Fold expression, cf. Unary op
+	uint32 j = 0;
+	([&]
+	{
+		check(j < NumOutputs);
+		auto OutValue = InCallbacks(DefaultValue1, DefaultValue2);
+		using OutType = decltype(OutValue);
+
+		if (FPCGMetadataAttribute<OutType>* OutputAttribute = static_cast<FPCGMetadataAttribute<OutType>*>(OperationData.OutputAttributes[j++]))
+		{
+			OutputAttribute->SetDefaultValue(OutValue);
+		}
+	} (), ...);
 
 	for (int32 i = 0; i < OperationData.NumberOfElementsToProcess; ++i)
 	{
@@ -226,9 +260,18 @@ inline bool FPCGMetadataElementBase::DoBinaryOp(FOperationData& OperationData, T
 			InType1 Value1 = PCGMetadataAttribute::GetValueWithBroadcast<InType1>(OperationData.SourceAttributes[0], EntryKey1);
 			InType2 Value2 = PCGMetadataAttribute::GetValueWithBroadcast<InType2>(OperationData.SourceAttributes[1], EntryKey2);
 
-			OutType OutputValue = BinaryOp(Value1, Value2);
+			j = 0;
+			([&]
+			{
+				check(j < NumOutputs);
+				auto OutValue = InCallbacks(Value1, Value2);
+				using OutType = decltype(OutValue);
 
-			OutputAttribute->SetValue(EntryKey1, OutputValue);
+				if (FPCGMetadataAttribute<OutType>* OutputAttribute = static_cast<FPCGMetadataAttribute<OutType>*>(OperationData.OutputAttributes[j++]))
+				{
+					OutputAttribute->SetValue(EntryKey1, OutValue);
+				}
+			} (), ...);
 		}
 
 		++Iterator1;
@@ -241,13 +284,16 @@ inline bool FPCGMetadataElementBase::DoBinaryOp(FOperationData& OperationData, T
 	return true;
 }
 
-template <typename InType1, typename InType2, typename InType3, typename OutType>
-inline bool FPCGMetadataElementBase::DoTernaryOp(FOperationData& OperationData, TFunction<OutType(const InType1&, const InType2&, const InType3&)> TernaryOp) const
+template <typename InType1, typename InType2, typename InType3, typename... Callbacks>
+inline bool FPCGMetadataElementBase::DoTernaryOp(FOperationData& OperationData, Callbacks&& ...InCallbacks) const
 {
 	check(OperationData.SourceAttributes[0]);
 	check(OperationData.SourceAttributes[1]);
 	check(OperationData.SourceAttributes[2]);
-	check(OperationData.OutputAttributes[0]);
+
+	const uint32 NumOutputs = OperationData.OutputAttributes.Num();
+	constexpr uint32 NumCallbacks = (uint32)sizeof...(InCallbacks);
+	check(NumOutputs == NumCallbacks);
 
 	// All values in OperationData.Iterators are UniquePtr. Dereference them here to make the syntax less cumbersome.
 	// Also some iterators can be null, it means they need to use the first iterator (that should never be null).
@@ -261,12 +307,23 @@ inline bool FPCGMetadataElementBase::DoTernaryOp(FOperationData& OperationData,
 	IPCGMetadataEntryIterator& Iterator2 = bShouldIncrementIterator2 ? *OperationData.Iterators[1] : Iterator1;
 	IPCGMetadataEntryIterator& Iterator3 = bShouldIncrementIterator3 ? *OperationData.Iterators[2] : Iterator1;
 
-	FPCGMetadataAttribute<OutType>* OutputAttribute = static_cast<FPCGMetadataAttribute<OutType>*>(OperationData.OutputAttributes[0]);
 	InType1 DefaultValue1 = PCGMetadataAttribute::GetValueWithBroadcast<InType1>(OperationData.SourceAttributes[0], PCGInvalidEntryKey);
 	InType2 DefaultValue2 = PCGMetadataAttribute::GetValueWithBroadcast<InType2>(OperationData.SourceAttributes[1], PCGInvalidEntryKey);
 	InType3 DefaultValue3 = PCGMetadataAttribute::GetValueWithBroadcast<InType3>(OperationData.SourceAttributes[2], PCGInvalidEntryKey);
-	OutType DefaultOutputValue = TernaryOp(DefaultValue1, DefaultValue2, DefaultValue3);
-	OutputAttribute->SetDefaultValue(DefaultOutputValue);
+	
+	// Fold expression, cf. Unary op
+	uint32 j = 0;
+	([&]
+	{
+		check(j < NumOutputs);
+		auto OutValue = InCallbacks(DefaultValue1, DefaultValue2, DefaultValue3);
+		using OutType = decltype(OutValue);
+
+		if (FPCGMetadataAttribute<OutType>* OutputAttribute = static_cast<FPCGMetadataAttribute<OutType>*>(OperationData.OutputAttributes[j++]))
+		{
+			OutputAttribute->SetDefaultValue(OutValue);
+		}
+	} (), ...);
 
 	for (int32 i = 0; i < OperationData.NumberOfElementsToProcess; ++i)
 	{
@@ -282,9 +339,18 @@ inline bool FPCGMetadataElementBase::DoTernaryOp(FOperationData& OperationData,
 			InType2 Value2 = PCGMetadataAttribute::GetValueWithBroadcast<InType2>(OperationData.SourceAttributes[1], EntryKey2);
 			InType3 Value3 = PCGMetadataAttribute::GetValueWithBroadcast<InType3>(OperationData.SourceAttributes[2], EntryKey3);
 
-			OutType OutputValue = TernaryOp(Value1, Value2, Value3);
+			j = 0;
+			([&]
+			{
+				check(j < NumOutputs);
+				auto OutValue = InCallbacks(Value1, Value2, Value3);
+				using OutType = decltype(OutValue);
 
-			OutputAttribute->SetValue(EntryKey1, OutputValue);
+				if (FPCGMetadataAttribute<OutType>* OutputAttribute = static_cast<FPCGMetadataAttribute<OutType>*>(OperationData.OutputAttributes[j++]))
+				{
+					OutputAttribute->SetValue(EntryKey1, OutValue);
+				}
+			} (), ...);
 		}
 
 		++Iterator1;
@@ -302,14 +368,17 @@ inline bool FPCGMetadataElementBase::DoTernaryOp(FOperationData& OperationData,
 	return true;
 }
 
-template <typename InType1, typename InType2, typename InType3, typename InType4, typename OutType>
-inline bool FPCGMetadataElementBase::DoQuaternaryOp(FOperationData& OperationData, TFunction<OutType(const InType1&, const InType2&, const InType3&, const InType4&)> QuaternaryOp) const
+template <typename InType1, typename InType2, typename InType3, typename InType4, typename... Callbacks>
+inline bool FPCGMetadataElementBase::DoQuaternaryOp(FOperationData& OperationData, Callbacks&& ...InCallbacks) const
 {
 	check(OperationData.SourceAttributes[0]);
 	check(OperationData.SourceAttributes[1]);
 	check(OperationData.SourceAttributes[2]);
 	check(OperationData.SourceAttributes[3]);
-	check(OperationData.OutputAttributes[0]);
+
+	const uint32 NumOutputs = OperationData.OutputAttributes.Num();
+	constexpr uint32 NumCallbacks = (uint32)sizeof...(InCallbacks);
+	check(NumOutputs == NumCallbacks);
 
 	// All values in OperationData.Iterators are UniquePtr. Dereference them here to make the syntax less cumbersome.
 	// Also some iterators can be null, it means they need to use the first iterator (that should never be null).
@@ -325,13 +394,24 @@ inline bool FPCGMetadataElementBase::DoQuaternaryOp(FOperationData& OperationDat
 	IPCGMetadataEntryIterator& Iterator3 = bShouldIncrementIterator3 ? *OperationData.Iterators[2] : Iterator1;
 	IPCGMetadataEntryIterator& Iterator4 = bShouldIncrementIterator4 ? *OperationData.Iterators[3] : Iterator1;
 
-	FPCGMetadataAttribute<OutType>* OutputAttribute = static_cast<FPCGMetadataAttribute<OutType>*>(OperationData.OutputAttributes[0]);
 	InType1 DefaultValue1 = PCGMetadataAttribute::GetValueWithBroadcast<InType1>(OperationData.SourceAttributes[0], PCGInvalidEntryKey);
 	InType2 DefaultValue2 = PCGMetadataAttribute::GetValueWithBroadcast<InType2>(OperationData.SourceAttributes[1], PCGInvalidEntryKey);
 	InType3 DefaultValue3 = PCGMetadataAttribute::GetValueWithBroadcast<InType3>(OperationData.SourceAttributes[2], PCGInvalidEntryKey);
 	InType3 DefaultValue4 = PCGMetadataAttribute::GetValueWithBroadcast<InType3>(OperationData.SourceAttributes[3], PCGInvalidEntryKey);
-	OutType DefaultOutputValue = QuaternaryOp(DefaultValue1, DefaultValue2, DefaultValue3, DefaultValue4);
-	OutputAttribute->SetDefaultValue(DefaultOutputValue);
+	
+	// Fold expression, cf. Unary op
+	uint32 j = 0;
+	([&]
+	{
+		check(j < NumOutputs);
+		auto OutValue = InCallbacks(DefaultValue1, DefaultValue2, DefaultValue3, DefaultValue4);
+		using OutType = decltype(OutValue);
+
+		if (FPCGMetadataAttribute<OutType>* OutputAttribute = static_cast<FPCGMetadataAttribute<OutType>*>(OperationData.OutputAttributes[j++]))
+		{
+			OutputAttribute->SetDefaultValue(OutValue);
+		}
+	} (), ...);
 
 	for (int32 i = 0; i < OperationData.NumberOfElementsToProcess; ++i)
 	{
@@ -349,9 +429,18 @@ inline bool FPCGMetadataElementBase::DoQuaternaryOp(FOperationData& OperationDat
 			InType3 Value3 = PCGMetadataAttribute::GetValueWithBroadcast<InType3>(OperationData.SourceAttributes[2], EntryKey3);
 			InType4 Value4 = PCGMetadataAttribute::GetValueWithBroadcast<InType4>(OperationData.SourceAttributes[3], EntryKey4);
 
-			OutType OutputValue = QuaternaryOp(Value1, Value2, Value3, Value4);
+			j = 0;
+			([&]
+			{
+				check(j < NumOutputs);
+				auto OutValue = InCallbacks(Value1, Value2, Value3, Value4);
+				using OutType = decltype(OutValue);
 
-			OutputAttribute->SetValue(EntryKey1, OutputValue);
+				if (FPCGMetadataAttribute<OutType>* OutputAttribute = static_cast<FPCGMetadataAttribute<OutType>*>(OperationData.OutputAttributes[j++]))
+				{
+					OutputAttribute->SetValue(EntryKey1, OutValue);
+				}
+			} (), ...);
 		}
 
 		++Iterator1;
@@ -371,77 +460,5 @@ inline bool FPCGMetadataElementBase::DoQuaternaryOp(FOperationData& OperationDat
 		}
 	}
 
-	return true;
-}
-
-template <typename InType, typename OutType>
-inline bool FPCGMetadataElementBase::DoUnaryOpMultipleOutputs(FOperationData & OperationData, TFunction<void(const InType&, TArray<OutType>&, uint32)> UnaryOp) const
-{
-	check(OperationData.SourceAttributes[0]);
-
-	// All values in OperationData.Iterators are UniquePtr. Dereference them here to make the syntax less cumbersome.
-	check(OperationData.Iterators[0]);
-	IPCGMetadataEntryIterator& Iterator1 = *OperationData.Iterators[0];
-
-	const uint32 NumOutputs = OperationData.OutputAttributes.Num();
-
-	TArray<OutType> OutputValues;
-	TArray<FPCGMetadataAttribute<OutType>*> OutputAttributes;
-
-	// Output value array is reserved not resized. Allow us to validate that the callback properly pushed
-	// the expected amount of outputs.
-	OutputValues.Reserve(NumOutputs);
-	OutputAttributes.SetNum(NumOutputs);
-
-	InType DefaultValue = PCGMetadataAttribute::GetValueWithBroadcast<InType>(OperationData.SourceAttributes[0], PCGInvalidEntryKey);
-	UnaryOp(DefaultValue, OutputValues, NumOutputs);
-
-	if (OutputValues.Num() != NumOutputs)
-	{
-		UE_LOG(LogPCG, Error, TEXT("[DoUnaryOpMultipleOutputs] Callback didn't provided the number of expected outputs."));
-		return false;
-	}
-
-	for (uint32 j = 0; j < NumOutputs; ++j)
-	{
-		// Output Attributes can be null if there is nothing connected to the output pin.
-		OutputAttributes[j] = static_cast<FPCGMetadataAttribute<OutType>*>(OperationData.OutputAttributes[j]);
-		if (OutputAttributes[j])
-		{
-			OutputAttributes[j]->SetDefaultValue(OutputValues[j]);
-		}
-	}
-
-	for (int32 i = 0; i < OperationData.NumberOfElementsToProcess; ++i)
-	{
-		PCGMetadataEntryKey EntryKey = *Iterator1;
-
-		// If the entry key is invalid, nothing to do
-		if (EntryKey != PCGInvalidEntryKey)
-		{
-			InType Value = PCGMetadataAttribute::GetValueWithBroadcast<InType>(OperationData.SourceAttributes[0], EntryKey);
-
-			// Make sure the array is emptied (but reserved)
-			OutputValues.Empty(/*Slack=*/NumOutputs);
-			UnaryOp(Value, OutputValues, NumOutputs);
-
-			if (OutputValues.Num() != NumOutputs)
-			{
-				UE_LOG(LogPCG, Error, TEXT("[DoUnaryOpMultipleOutputs] Callback didn't provided the number of expected outputs."));
-				return false;
-			}
-
-			for (uint32 j = 0; j < NumOutputs; ++j)
-			{
-				if (OutputAttributes[j])
-				{
-					OutputAttributes[j]->SetValue(EntryKey, OutputValues[j]);
-				}
-			}
-		}
-
-		++Iterator1;
-	}
-
 	return true;
 }
\ No newline at end of file
