Tentative workaround for compiler bug in version 17.3 for VS2022

#rnx

diff --git a/Engine/Plugins/Experimental/PCG/Source/PCG/Public/Elements/Metadata/PCGMetadataOpElementBase.h b/Engine/Plugins/Experimental/PCG/Source/PCG/Public/Elements/Metadata/PCGMetadataOpElementBase.h
--- a/Engine/Plugins/Experimental/PCG/Source/PCG/Public/Elements/Metadata/PCGMetadataOpElementBase.h
+++ b/Engine/Plugins/Experimental/PCG/Source/PCG/Public/Elements/Metadata/PCGMetadataOpElementBase.h
@@ -174,7 +174,7 @@ inline bool FPCGMetadataElementBase::DoUnaryOp(FOperationData& OperationData, Ca
 	{
 		check(j < NumOutputs);
 		auto OutValue = InCallbacks(DefaultValue);
-		using OutType = decltype(OutValue);
+		typedef decltype(OutValue) OutType;
 
 		if (FPCGMetadataAttribute<OutType>* OutputAttribute = static_cast<FPCGMetadataAttribute<OutType>*>(OperationData.OutputAttributes[j++]))
 		{
@@ -196,7 +196,7 @@ inline bool FPCGMetadataElementBase::DoUnaryOp(FOperationData& OperationData, Ca
 			{
 				check(j < NumOutputs);
 				auto OutValue = InCallbacks(Value);
-				using OutType = decltype(OutValue);
+				typedef decltype(OutValue) OutType;
 
 				if (FPCGMetadataAttribute<OutType>* OutputAttribute = static_cast<FPCGMetadataAttribute<OutType>*>(OperationData.OutputAttributes[j++]))
 				{
@@ -240,7 +240,7 @@ inline bool FPCGMetadataElementBase::DoBinaryOp(FOperationData& OperationData, C
 	{
 		check(j < NumOutputs);
 		auto OutValue = InCallbacks(DefaultValue1, DefaultValue2);
-		using OutType = decltype(OutValue);
+		typedef decltype(OutValue) OutType;
 
 		if (FPCGMetadataAttribute<OutType>* OutputAttribute = static_cast<FPCGMetadataAttribute<OutType>*>(OperationData.OutputAttributes[j++]))
 		{
@@ -265,7 +265,7 @@ inline bool FPCGMetadataElementBase::DoBinaryOp(FOperationData& OperationData, C
 			{
 				check(j < NumOutputs);
 				auto OutValue = InCallbacks(Value1, Value2);
-				using OutType = decltype(OutValue);
+				typedef decltype(OutValue) OutType;
 
 				if (FPCGMetadataAttribute<OutType>* OutputAttribute = static_cast<FPCGMetadataAttribute<OutType>*>(OperationData.OutputAttributes[j++]))
 				{
@@ -317,7 +317,7 @@ inline bool FPCGMetadataElementBase::DoTernaryOp(FOperationData& OperationData,
 	{
 		check(j < NumOutputs);
 		auto OutValue = InCallbacks(DefaultValue1, DefaultValue2, DefaultValue3);
-		using OutType = decltype(OutValue);
+		typedef decltype(OutValue) OutType;
 
 		if (FPCGMetadataAttribute<OutType>* OutputAttribute = static_cast<FPCGMetadataAttribute<OutType>*>(OperationData.OutputAttributes[j++]))
 		{
@@ -344,7 +344,7 @@ inline bool FPCGMetadataElementBase::DoTernaryOp(FOperationData& OperationData,
 			{
 				check(j < NumOutputs);
 				auto OutValue = InCallbacks(Value1, Value2, Value3);
-				using OutType = decltype(OutValue);
+				typedef decltype(OutValue) OutType;
 
 				if (FPCGMetadataAttribute<OutType>* OutputAttribute = static_cast<FPCGMetadataAttribute<OutType>*>(OperationData.OutputAttributes[j++]))
 				{
@@ -405,7 +405,7 @@ inline bool FPCGMetadataElementBase::DoQuaternaryOp(FOperationData& OperationDat
 	{
 		check(j < NumOutputs);
 		auto OutValue = InCallbacks(DefaultValue1, DefaultValue2, DefaultValue3, DefaultValue4);
-		using OutType = decltype(OutValue);
+		typedef decltype(OutValue) OutType;
 
 		if (FPCGMetadataAttribute<OutType>* OutputAttribute = static_cast<FPCGMetadataAttribute<OutType>*>(OperationData.OutputAttributes[j++]))
 		{
@@ -434,7 +434,7 @@ inline bool FPCGMetadataElementBase::DoQuaternaryOp(FOperationData& OperationDat
 			{
 				check(j < NumOutputs);
 				auto OutValue = InCallbacks(Value1, Value2, Value3, Value4);
-				using OutType = decltype(OutValue);
+				typedef decltype(OutValue) OutType;
 
 				if (FPCGMetadataAttribute<OutType>* OutputAttribute = static_cast<FPCGMetadataAttribute<OutType>*>(OperationData.OutputAttributes[j++]))
 				{
