[PCG] Metadata native operation: Make/Break Vectors and Transforms

diff --git a/Engine/Plugins/Experimental/PCG/Source/PCG/Private/Elements/Metadata/PCGMetadataBreakTransform.cpp b/Engine/Plugins/Experimental/PCG/Source/PCG/Private/Elements/Metadata/PCGMetadataBreakTransform.cpp
new file mode 100644
--- /dev/null
+++ b/Engine/Plugins/Experimental/PCG/Source/PCG/Private/Elements/Metadata/PCGMetadataBreakTransform.cpp
@@ -0,0 +1,93 @@
+// Copyright Epic Games, Inc. All Rights Reserved.
+
+#include "Elements/Metadata/PCGMetadataBreakTransform.h"
+
+#include "PCGParamData.h"
+#include "Data/PCGSpatialData.h"
+#include "Elements/Metadata/PCGMetadataElementCommon.h"
+#include "Helpers/PCGSettingsHelpers.h"
+#include "Metadata/PCGMetadata.h"
+#include "Metadata/PCGMetadataAttribute.h"
+#include "Metadata/PCGMetadataAttributeTpl.h"
+
+namespace PCGMetadataBreakTransformSettings
+{
+	inline FVector GetTranslation(const FTransform& Transform)
+	{
+		return Transform.GetTranslation();
+	}
+
+	inline FVector GetScale(const FTransform& Transform)
+	{
+		return Transform.GetScale3D();
+	}
+
+	inline FQuat GetRotation(const FTransform& Transform)
+	{
+		return Transform.GetRotation();
+	}
+}
+
+FName UPCGMetadataBreakTransformSettings::GetInputAttributeNameWithOverride(uint32 Index, UPCGParamData* Params) const
+{
+	return PCG_GET_OVERRIDEN_VALUE(this, InputAttributeName, Params);
+}
+
+FName UPCGMetadataBreakTransformSettings::GetOutputPinLabel(uint32 Index) const
+{
+	switch (Index)
+	{
+	case 0:
+		return PCGMetadataTransformConstants::Translation;
+	case 1:
+		return PCGMetadataTransformConstants::Rotation;
+	case 2:
+		return PCGMetadataTransformConstants::Scale;
+	default:
+		return NAME_None;
+	}
+}
+
+uint32 UPCGMetadataBreakTransformSettings::GetOutputPinNum() const
+{
+	return 3;
+}
+
+bool UPCGMetadataBreakTransformSettings::IsSupportedInputType(uint16 TypeId, uint32 InputIndex, bool& bHasSpecialRequirement) const
+{
+	bHasSpecialRequirement = false;
+	return PCG::Private::IsOfTypes<FTransform>(TypeId);
+}
+
+bool UPCGMetadataBreakTransformSettings::HasDifferentOutputTypes() const
+{
+	return true;
+}
+
+TArray<uint16> UPCGMetadataBreakTransformSettings::GetAllOutputTypes() const
+{
+	return { (uint16)EPCGMetadataTypes::Vector, (uint16)EPCGMetadataTypes::Quaternion, (uint16)EPCGMetadataTypes::Vector };
+}
+
+FName UPCGMetadataBreakTransformSettings::GetOutputAttributeName(FName BaseName, uint32 Index) const
+{
+	return BaseName != NAME_None ? FName(BaseName.ToString() + "." + GetOutputPinLabel(Index).ToString()) : NAME_None;
+}
+
+FPCGElementPtr UPCGMetadataBreakTransformSettings::CreateElement() const
+{
+	return MakeShared<FPCGMetadataBreakTransformElement>();
+}
+
+bool FPCGMetadataBreakTransformElement::DoOperation(FOperationData& OperationData) const
+{
+	TRACE_CPUPROFILER_EVENT_SCOPE(FPCGMetadataBreakVectorElement::Execute);
+
+	const UPCGMetadataBreakTransformSettings* Settings = static_cast<const UPCGMetadataBreakTransformSettings*>(OperationData.Settings);
+	check(Settings);
+
+	return DoUnaryOp<FTransform>(OperationData,
+		PCGMetadataBreakTransformSettings::GetTranslation,
+		PCGMetadataBreakTransformSettings::GetRotation,
+		PCGMetadataBreakTransformSettings::GetScale);
+}
diff --git a/Engine/Plugins/Experimental/PCG/Source/PCG/Private/Elements/Metadata/PCGMetadataBreakVector.cpp b/Engine/Plugins/Experimental/PCG/Source/PCG/Private/Elements/Metadata/PCGMetadataBreakVector.cpp
--- a/Engine/Plugins/Experimental/PCG/Source/PCG/Private/Elements/Metadata/PCGMetadataBreakVector.cpp
+++ b/Engine/Plugins/Experimental/PCG/Source/PCG/Private/Elements/Metadata/PCGMetadataBreakVector.cpp
@@ -68,10 +68,7 @@ namespace PCGMetadataBreakVectorSettings
 
 	inline constexpr bool IsValidType(uint16 TypeId)
 	{
-		return TypeId == (uint16)EPCGMetadataTypes::Vector2 ||
-			TypeId == (uint16)EPCGMetadataTypes::Vector ||
-			TypeId == (uint16)EPCGMetadataTypes::Vector4 ||
-			TypeId == (uint16)EPCGMetadataTypes::Rotator;
+		return PCG::Private::IsOfTypes<FVector2D, FVector, FVector4, FRotator>(TypeId);
 	}
 
 	template <typename T>
diff --git a/Engine/Plugins/Experimental/PCG/Source/PCG/Private/Elements/Metadata/PCGMetadataMakeTransform.cpp b/Engine/Plugins/Experimental/PCG/Source/PCG/Private/Elements/Metadata/PCGMetadataMakeTransform.cpp
new file mode 100644
--- /dev/null
+++ b/Engine/Plugins/Experimental/PCG/Source/PCG/Private/Elements/Metadata/PCGMetadataMakeTransform.cpp
@@ -0,0 +1,110 @@
+// Copyright Epic Games, Inc. All Rights Reserved.
+
+#include "Elements/Metadata/PCGMetadataMakeTransform.h"
+
+#include "Helpers/PCGSettingsHelpers.h"
+#include "Metadata/PCGMetadataAttributeTpl.h"
+#include "Elements/Metadata/PCGMetadataElementCommon.h"
+
+namespace PCGMetadataMakeTransformSettings
+{
+	template <typename VectorType = FVector4>
+	FTransform MakeTransform(const VectorType& Translation, const FQuat& Rotation, const VectorType& Scale)
+	{
+		return FTransform(Rotation, FVector(Translation), FVector(Scale));
+	}
+
+	template <>
+	FTransform MakeTransform<FVector>(const FVector& Translation, const FQuat& Rotation, const FVector& Scale)
+	{
+		return FTransform(Rotation, Translation, Scale);
+	}
+
+	template <>
+	FTransform MakeTransform<FVector2D>(const FVector2D& Translation, const FQuat& Rotation, const FVector2D& Scale)
+	{
+		return FTransform(Rotation, FVector(Translation, 0.0), FVector(Scale, 1.0));
+	}
+}
+
+FName UPCGMetadataMakeTransformSettings::GetInputPinLabel(uint32 Index) const
+{
+	switch (Index)
+	{
+	case 0:
+		return PCGMetadataTransformConstants::Translation;
+	case 1:
+		return PCGMetadataTransformConstants::Rotation;
+	case 2:
+		return PCGMetadataTransformConstants::Scale;
+	default:
+		return NAME_None;
+	}
+}
+
+uint32 UPCGMetadataMakeTransformSettings::GetInputPinNum() const
+{
+	return 3;
+}
+
+bool UPCGMetadataMakeTransformSettings::IsSupportedInputType(uint16 TypeId, uint32 InputIndex, bool& bHasSpecialRequirement) const
+{
+	if (InputIndex == 1)
+	{
+		bHasSpecialRequirement = true;
+		return PCG::Private::IsOfTypes<FRotator, FQuat>(TypeId);
+	}
+	else
+	{
+		bHasSpecialRequirement = false;
+		return PCG::Private::IsOfTypes<FVector2D, FVector, FVector4>(TypeId);
+	}
+}
+
+FName UPCGMetadataMakeTransformSettings::GetInputAttributeNameWithOverride(uint32 Index, UPCGParamData* Params) const
+{
+	switch (Index)
+	{
+	case 0:
+		return PCG_GET_OVERRIDEN_VALUE(this, Input1AttributeName, Params);
+	case 1:
+		return PCG_GET_OVERRIDEN_VALUE(this, Input2AttributeName, Params);
+	case 2:
+		return PCG_GET_OVERRIDEN_VALUE(this, Input3AttributeName, Params);
+	default:
+		return NAME_None;
+	}
+}
+
+uint16 UPCGMetadataMakeTransformSettings::GetOutputType(uint16 InputTypeId) const
+{
+	return (uint16)EPCGMetadataTypes::Transform;
+}
+
+FPCGElementPtr UPCGMetadataMakeTransformSettings::CreateElement() const
+{
+	return MakeShared<FPCGMetadataMakeTransformElement>();
+}
+
+bool FPCGMetadataMakeTransformElement::DoOperation(FOperationData& OperationData) const
+{
+	TRACE_CPUPROFILER_EVENT_SCOPE(FPCGMetadataMakeTransformElement::Execute);
+
+	const UPCGMetadataMakeTransformSettings* Settings = CastChecked<UPCGMetadataMakeTransformSettings>(OperationData.Settings);
+
+	auto TransformFunc = [this, &OperationData](auto DummyValue) -> bool
+	{
+		using AttributeType = decltype(DummyValue);
+
+		if constexpr (!PCG::Private::IsOfTypes<AttributeType, FVector2D, FVector, FVector4>())
+		{
+			return false;
+		}
+		else
+		{
+			return DoTernaryOp<AttributeType, FQuat, AttributeType>(OperationData, PCGMetadataMakeTransformSettings::MakeTransform<AttributeType>);
+		}
+	};
+
+	return PCGMetadataAttribute::CallbackWithRightType(OperationData.MostComplexInputType, TransformFunc);
+}
\ No newline at end of file
diff --git a/Engine/Plugins/Experimental/PCG/Source/PCG/Private/Elements/Metadata/PCGMetadataMakeVector.cpp b/Engine/Plugins/Experimental/PCG/Source/PCG/Private/Elements/Metadata/PCGMetadataMakeVector.cpp
new file mode 100644
--- /dev/null
+++ b/Engine/Plugins/Experimental/PCG/Source/PCG/Private/Elements/Metadata/PCGMetadataMakeVector.cpp
@@ -0,0 +1,202 @@
+// Copyright Epic Games, Inc. All Rights Reserved.
+
+#include "Elements/Metadata/PCGMetadataMakeVector.h"
+
+#include "Helpers/PCGSettingsHelpers.h"
+#include "Metadata/PCGMetadataAttributeTpl.h"
+#include "Elements/Metadata/PCGMetadataElementCommon.h"
+
+namespace PCGMetadataMakeVectorSettings
+{
+	FVector2D MakeVector2(const double& X, const double& Y)
+	{
+		return FVector2D(X, Y);
+	}
+
+	FVector MakeVector3(const double& X, const double& Y, const double& Z)
+	{
+		return FVector(X, Y, Z);
+	}
+
+	FVector MakeVector3Vec2(const FVector2D& XY, const double& Z)
+	{
+		return FVector(XY, Z);
+	}
+
+	FVector4 MakeVector4(const double& X, const double& Y, const double& Z, const double& W)
+	{
+		return FVector4(X, Y, Z, W);
+	}
+
+	FVector4 MakeVector4Vec2(const FVector2D& XY, const double& Z, const double W)
+	{
+		return FVector4(XY.X, XY.Y, Z, W);
+	}
+
+	FVector4 MakeVector4TwoVec2(const FVector2D& XY, const FVector2D& ZW)
+	{
+		return FVector4(XY, ZW);
+	}
+
+	FVector4 MakeVector4Vec3(const FVector& XYZ, const double& W)
+	{
+		return FVector4(XYZ, W);
+	}
+}
+
+FName UPCGMetadataMakeVectorSettings::GetInputPinLabel(uint32 Index) const
+{
+	static const FName DefaultNames[4] = {
+		PCGMetadataMakeVectorConstants::XLabel,
+		PCGMetadataMakeVectorConstants::YLabel,
+		PCGMetadataMakeVectorConstants::ZLabel,
+		PCGMetadataMakeVectorConstants::WLabel
+	};
+
+	if (OutputType == EPCGMetadataTypes::Vector2)
+	{
+		return DefaultNames[Index];
+	}
+	else if (OutputType == EPCGMetadataTypes::Vector)
+	{
+		if (MakeVector3Op == EPCGMetadataMakeVector3::ThreeValues)
+		{
+			return DefaultNames[Index];
+		}
+		else
+		{
+			return Index == 0 ? PCGMetadataMakeVectorConstants::XYLabel : PCGMetadataMakeVectorConstants::ZLabel;
+		}
+	}
+	else
+	{
+		if (MakeVector4Op == EPCGMetadataMakeVector4::FourValues)
+		{
+			return DefaultNames[Index];
+		}
+		else if (MakeVector4Op == EPCGMetadataMakeVector4::Vector2AndTwoValues)
+		{
+			return Index == 0 ? PCGMetadataMakeVectorConstants::XYLabel : (Index == 1 ? PCGMetadataMakeVectorConstants::ZLabel : PCGMetadataMakeVectorConstants::WLabel);
+		}
+		else if (MakeVector4Op == EPCGMetadataMakeVector4::TwoVector2)
+		{
+			return Index == 0 ? PCGMetadataMakeVectorConstants::XYLabel : PCGMetadataMakeVectorConstants::ZWLabel;
+		}
+		else
+		{
+			return Index == 0 ? PCGMetadataMakeVectorConstants::XYZLabel : PCGMetadataMakeVectorConstants::WLabel;
+		}
+	}
+}
+
+uint32 UPCGMetadataMakeVectorSettings::GetInputPinNum() const
+{
+	if (OutputType == EPCGMetadataTypes::Vector2 ||
+		(OutputType == EPCGMetadataTypes::Vector && MakeVector3Op == EPCGMetadataMakeVector3::Vector2AndValue) ||
+		(OutputType == EPCGMetadataTypes::Vector4 &&
+			(MakeVector4Op == EPCGMetadataMakeVector4::TwoVector2 || MakeVector4Op == EPCGMetadataMakeVector4::Vector3AndValue)))
+	{
+		return 2;
+	}
+	else if (OutputType == EPCGMetadataTypes::Vector ||
+		(OutputType == EPCGMetadataTypes::Vector4 && MakeVector4Op == EPCGMetadataMakeVector4::Vector2AndTwoValues))
+	{
+		return 3;
+	}
+	else
+	{
+		return 4;
+	}
+}
+
+bool UPCGMetadataMakeVectorSettings::IsSupportedInputType(uint16 TypeId, uint32 InputIndex, bool& bHasSpecialRequirement) const
+{
+	// Use labels since the logic is already done there.
+	bHasSpecialRequirement = false;
+	FName Label = GetInputPinLabel(InputIndex);
+
+	if (Label == PCGMetadataMakeVectorConstants::XYZLabel)
+	{
+		return PCG::Private::IsOfTypes<FVector>(TypeId);
+	}
+	else if (Label == PCGMetadataMakeVectorConstants::XYLabel || Label == PCGMetadataMakeVectorConstants::ZWLabel)
+	{
+		return PCG::Private::IsOfTypes<FVector2D>(TypeId);
+	}
+	else
+	{
+		return PCG::Private::IsOfTypes<float, double>(TypeId);
+	}
+}
+
+FName UPCGMetadataMakeVectorSettings::GetInputAttributeNameWithOverride(uint32 Index, UPCGParamData* Params) const
+{
+	switch (Index)
+	{
+	case 0:
+		return PCG_GET_OVERRIDEN_VALUE(this, Input1AttributeName, Params);
+	case 1:
+		return PCG_GET_OVERRIDEN_VALUE(this, Input2AttributeName, Params);
+	case 2:
+		return PCG_GET_OVERRIDEN_VALUE(this, Input3AttributeName, Params);
+	case 3:
+		return PCG_GET_OVERRIDEN_VALUE(this, Input4AttributeName, Params);
+	default:
+		return NAME_None;
+	}
+}
+
+uint16 UPCGMetadataMakeVectorSettings::GetOutputType(uint16 InputTypeId) const
+{
+	return (uint16)OutputType;
+}
+
+FPCGElementPtr UPCGMetadataMakeVectorSettings::CreateElement() const
+{
+	return MakeShared<FPCGMetadataMakeVectorElement>();
+}
+
+bool FPCGMetadataMakeVectorElement::DoOperation(FOperationData& OperationData) const
+{
+	TRACE_CPUPROFILER_EVENT_SCOPE(FPCGMetadataMakeVectorElement::Execute);
+
+	const UPCGMetadataMakeVectorSettings* Settings = CastChecked<UPCGMetadataMakeVectorSettings>(OperationData.Settings);
+
+	if (Settings->OutputType == EPCGMetadataTypes::Vector2)
+	{
+		DoBinaryOp<double, double>(OperationData, PCGMetadataMakeVectorSettings::MakeVector2);
+	}
+	else if (Settings->OutputType == EPCGMetadataTypes::Vector)
+	{
+		if (Settings->MakeVector3Op == EPCGMetadataMakeVector3::ThreeValues)
+		{
+			DoTernaryOp<double, double, double>(OperationData, PCGMetadataMakeVectorSettings::MakeVector3);
+		}
+		else
+		{
+			DoBinaryOp<FVector2D, double>(OperationData, PCGMetadataMakeVectorSettings::MakeVector3Vec2);
+		}
+	}
+	else
+	{
+		switch (Settings->MakeVector4Op)
+		{
+		case EPCGMetadataMakeVector4::FourValues:
+			DoQuaternaryOp<double, double, double, double>(OperationData, PCGMetadataMakeVectorSettings::MakeVector4);
+			break;
+		case EPCGMetadataMakeVector4::TwoVector2:
+			DoBinaryOp<FVector2D, FVector2D>(OperationData, PCGMetadataMakeVectorSettings::MakeVector4TwoVec2);
+			break;
+		case EPCGMetadataMakeVector4::Vector2AndTwoValues:
+			DoTernaryOp<FVector2D, double, double>(OperationData, PCGMetadataMakeVectorSettings::MakeVector4Vec2);
+			break;
+		case EPCGMetadataMakeVector4::Vector3AndValue:
+			DoBinaryOp<FVector, double>(OperationData, PCGMetadataMakeVectorSettings::MakeVector4Vec3);
+			break;
+		default:
+			break;
+		}
+	}
+
+	return true;
+}
\ No newline at end of file
diff --git a/Engine/Plugins/Experimental/PCG/Source/PCG/Private/Tests/Elements/Metadata/PCGMetadataBreakTransformTest.cpp b/Engine/Plugins/Experimental/PCG/Source/PCG/Private/Tests/Elements/Metadata/PCGMetadataBreakTransformTest.cpp
new file mode 100644
--- /dev/null
+++ b/Engine/Plugins/Experimental/PCG/Source/PCG/Private/Tests/Elements/Metadata/PCGMetadataBreakTransformTest.cpp
@@ -0,0 +1,336 @@
+// Copyright Epic Games, Inc. All Rights Reserved.
+
+#include "Tests/Determinism/PCGDeterminismTestsCommon.h"
+#include "Tests/PCGTestsCommon.h"
+
+#include "PCGComponent.h"
+#include "PCGData.h"
+#include "PCGParamData.h"
+
+#include "Data/PCGPointData.h"
+#include "Elements/Metadata/PCGMetadataBreakTransform.h"
+#include "Metadata/PCGMetadata.h"
+#include "Metadata/PCGMetadataAttribute.h"
+#include "Metadata/PCGMetadataAttributeTpl.h"
+
+#if WITH_EDITOR
+
+IMPLEMENT_CUSTOM_SIMPLE_AUTOMATION_TEST(FPCGMetadataBreakTransformTest, FPCGTestBaseClass, "pcg.tests.Metadata.BreakTransform", PCGTestsCommon::TestFlags)
+
+namespace PCGBreakTransformTest
+{
+	enum class EPCGComponentToCheck
+	{
+		Translation,
+		Rotation,
+		Scale
+	};
+
+	const FName TransformAttribute = TEXT("Transform");
+	const FName InvalidAttribute = TEXT("Float");
+	
+	FTransform RandomTransform(FRandomStream& RandomSource)
+	{
+		FVector RandomTranslation = RandomSource.VRand();
+		FQuat RandomRotation = FRotator(RandomSource.FRand(), RandomSource.FRand(), RandomSource.FRand()).Quaternion();
+		FVector RandomScale = RandomSource.VRand();
+
+		return FTransform(RandomRotation, RandomTranslation, RandomScale);
+	}
+
+	void GenerateSpatialData(PCGTestsCommon::FTestData& TestData)
+	{
+		FPCGTaggedData& SpatialTaggedData = TestData.InputData.TaggedData.Emplace_GetRef();
+		TObjectPtr<UPCGPointData> PointData = PCGTestsCommon::CreateRandomPointData(10, TestData.Seed);
+		SpatialTaggedData.Data = PointData;
+		SpatialTaggedData.Pin = PCGPinConstants::DefaultInputLabel;
+
+		const bool bAllowsInterpolation = false;
+		const bool bOverrideParent = false;
+
+		PointData->Metadata->CreateTransformAttribute(TransformAttribute, FTransform::Identity, bAllowsInterpolation, bOverrideParent);
+		PointData->Metadata->CreateFloatAttribute(InvalidAttribute, 0.0f, bAllowsInterpolation, bOverrideParent);
+
+		TArray<FPCGPoint>& SourcePoints = PointData->GetMutablePoints();
+
+		// Randomize attribute values, and leave the second half to default value
+		FRandomStream RandomSource(TestData.Seed);
+		for (int I = 0; I < SourcePoints.Num() / 2; ++I)
+		{
+			UPCGMetadataAccessorHelpers::SetTransformAttribute(SourcePoints[I], PointData->Metadata, TransformAttribute, RandomTransform(RandomSource));
+			UPCGMetadataAccessorHelpers::SetFloatAttribute(SourcePoints[I], PointData->Metadata, InvalidAttribute, RandomSource.FRand());
+		}
+	}
+
+	void GenerateParamData(PCGTestsCommon::FTestData& TestData)
+	{
+		FPCGTaggedData& ParamTaggedData = TestData.InputData.TaggedData.Emplace_GetRef();
+		TObjectPtr<UPCGParamData> ParamData = PCGTestsCommon::CreateEmptyParamData();
+		ParamTaggedData.Data = ParamData;
+		ParamTaggedData.Pin = PCGPinConstants::DefaultInputLabel;
+
+		PCGMetadataEntryKey EntryKey = ParamData->Metadata->AddEntry();
+		const bool bAllowsInterpolation = false;
+		const bool bOverrideParent = false;
+
+		ParamData->Metadata->CreateTransformAttribute(TransformAttribute, FTransform::Identity, bAllowsInterpolation, bOverrideParent);
+		ParamData->Metadata->CreateFloatAttribute(InvalidAttribute, 0.0f, bAllowsInterpolation, bOverrideParent);
+
+		// Randomize attribute values
+		FRandomStream RandomSource(TestData.Seed);
+		if (FPCGMetadataAttribute<FTransform>* Attribute = static_cast<FPCGMetadataAttribute<FTransform>*>(ParamData->Metadata->GetMutableAttribute(TransformAttribute)))
+		{
+			Attribute->SetValue(EntryKey, RandomTransform(RandomSource));
+		}
+
+		if (FPCGMetadataAttribute<float>* Attribute = static_cast<FPCGMetadataAttribute<float>*>(ParamData->Metadata->GetMutableAttribute(InvalidAttribute)))
+		{
+			Attribute->SetValue(EntryKey, RandomSource.FRand());
+		}
+	}
+}
+
+
+bool FPCGMetadataBreakTransformTest::RunTest(const FString& Parameters)
+{	
+	auto ValidateMetadataBreakVector = [this](PCGTestsCommon::FTestData& TestData, bool bIsValid = true) -> bool
+	{
+		const FName OutputAttributeName = TEXT("Output");
+
+		UPCGMetadataBreakTransformSettings* Settings = CastChecked<UPCGMetadataBreakTransformSettings>(TestData.Settings);
+		Settings->OutputAttributeName = OutputAttributeName;
+		FPCGElementPtr MetadataBreakTransformElement = TestData.Settings->GetElement();
+
+		TUniquePtr<FPCGContext> Context = TUniquePtr<FPCGContext>(MetadataBreakTransformElement->Initialize(TestData.InputData, TestData.TestPCGComponent, nullptr));
+		Context->NumAvailableTasks = 1;
+
+		while (!MetadataBreakTransformElement->Execute(Context.Get()))
+		{}
+
+		// If the test is not valid, just check that we have no output and early out
+		if (!bIsValid)
+		{
+			return TestTrue(TEXT("Invalid operation, no outputs"), Context->OutputData.GetInputs().IsEmpty());
+		}
+
+		const TArray<FPCGTaggedData>& Inputs = Context->InputData.GetInputsByPin(PCGPinConstants::DefaultInputLabel);
+
+		if (!TestEqual("Right number of inputs", Inputs.Num(), 1))
+		{
+			return false;
+		}
+
+		bool bTestPassed = true;
+
+		const FPCGTaggedData& Input = Inputs[0];
+
+		UPCGMetadata* SourceMetadata = nullptr; 
+		if (const UPCGSpatialData* SpatialInput = Cast<const UPCGSpatialData>(Input.Data))
+		{
+			SourceMetadata = SpatialInput->Metadata;
+		}
+		else if (const UPCGParamData* ParamsInput = Cast<const UPCGParamData>(Input.Data))
+		{
+			SourceMetadata = ParamsInput->Metadata;
+		}
+		else
+		{
+			return false;
+		}
+
+		check(SourceMetadata);
+
+		const FPCGMetadataAttributeBase* SourceAttributeBase = SourceMetadata->GetConstAttribute(Settings->InputAttributeName);
+		check(SourceAttributeBase);
+
+		auto ValidateComponentOutput = [&](const FPCGTaggedData& Output, const FName& OutAttributeName, PCGBreakTransformTest::EPCGComponentToCheck ComponentToCheck, auto DummyValue)
+		{
+			using OutType = decltype(DummyValue);
+
+			if (!TestTrue("Valid output data", Output.Data != nullptr))
+			{
+				bTestPassed = false;
+				return;
+			}
+
+			UPCGMetadata* OutMetadata = nullptr;
+			if (const UPCGSpatialData* SpatialInput = Cast<const UPCGSpatialData>(Output.Data))
+			{
+				OutMetadata = SpatialInput->Metadata;
+			}
+			else if (const UPCGParamData* ParamsInput = Cast<const UPCGParamData>(Output.Data))
+			{
+				OutMetadata = ParamsInput->Metadata;
+			}
+			else
+			{
+				return;
+			}
+
+			if (!TestNotNull("Valid output metadata", OutMetadata))
+			{
+				bTestPassed = false;
+				return;
+			}
+
+			const FPCGMetadataAttributeBase* OutAttributeBase = OutMetadata->GetConstAttribute(OutAttributeName);
+
+			if (!TestNotNull("Valid output attribute", OutAttributeBase))
+			{
+				bTestPassed = false;
+				return;
+			}
+
+			if (!TestEqual("Output attribute is a valid type", OutAttributeBase->GetTypeId(), PCG::Private::MetadataTypes<OutType>::Id))
+			{
+				bTestPassed = false;
+				return;
+			}
+
+			const FPCGMetadataAttribute<OutType>* OutAttribute = static_cast<const FPCGMetadataAttribute<OutType>*>(OutAttributeBase);
+			check(OutAttribute);
+
+			const PCGMetadataEntryKey InEntryKeyCount = SourceMetadata->GetItemCountForChild();
+			const PCGMetadataEntryKey OutEntryKeyCount = OutMetadata->GetItemCountForChild();
+
+			if (!TestEqual("Identical EntryKey counts", InEntryKeyCount, OutEntryKeyCount))
+			{
+				bTestPassed = false;
+				return;
+			}
+
+			const PCGMetadataValueKey InValueKeyCount = SourceAttributeBase->GetValueKeyOffsetForChild();
+			const PCGMetadataValueKey OutValueKeyCount = OutAttribute->GetValueKeyOffsetForChild();
+
+			if (!TestEqual("Identical ValueKey counts", InValueKeyCount, OutValueKeyCount))
+			{
+				bTestPassed = false;
+				return;
+			}
+
+			for (PCGMetadataEntryKey EntryKey = 0; EntryKey < InEntryKeyCount; ++EntryKey)
+			{
+				const PCGMetadataValueKey InValueKey = SourceAttributeBase->GetValueKey(EntryKey);
+				const PCGMetadataValueKey OutValueKey = OutAttribute->GetValueKey(EntryKey);
+
+				bTestPassed &= TestEqual("Identical value keys", InValueKey, OutValueKey);
+
+				FTransform SourceValue = static_cast<const FPCGMetadataAttribute<FTransform>*>(SourceAttributeBase)->GetValue(InValueKey);
+
+				const OutType OutValue = OutAttribute->GetValue(OutValueKey);
+
+				if constexpr (std::is_same_v<FVector, OutType>)
+				{
+					if (ComponentToCheck == PCGBreakTransformTest::EPCGComponentToCheck::Translation)
+					{
+						bTestPassed &= TestEqual("Translation", SourceValue.GetTranslation(), OutValue);
+					}
+					else if (ComponentToCheck == PCGBreakTransformTest::EPCGComponentToCheck::Scale)
+					{
+						bTestPassed &= TestEqual("Scale", SourceValue.GetScale3D(), OutValue);
+					}
+					else
+					{
+						bTestPassed = false;
+					}
+				}
+				else
+				{
+					if (ComponentToCheck == PCGBreakTransformTest::EPCGComponentToCheck::Rotation)
+					{
+						bTestPassed &= TestEqual("Rotation", SourceValue.GetRotation(), OutValue);
+					}
+					else
+					{
+						bTestPassed = false;
+					}
+				}
+			}
+		};
+
+		const TArray<FPCGTaggedData> OutputsTranslation = Context->OutputData.GetInputsByPin(PCGMetadataTransformConstants::Translation);
+		const TArray<FPCGTaggedData> OutputsRotation = Context->OutputData.GetInputsByPin(PCGMetadataTransformConstants::Rotation);
+		const TArray<FPCGTaggedData> OutputsScale = Context->OutputData.GetInputsByPin(PCGMetadataTransformConstants::Scale);
+
+		const FName DestinationAttributeForTranslation = Settings->GetOutputAttributeName(Settings->OutputAttributeName, 0);
+		const FName DestinationAttributeForRotation = Settings->GetOutputAttributeName(Settings->OutputAttributeName, 1);
+		const FName DestinationAttributeForScale = Settings->GetOutputAttributeName(Settings->OutputAttributeName, 2);
+
+		if (TestEqual("Appropriate number of outputs generated fox Translation", OutputsTranslation.Num(), 1))
+		{
+			ValidateComponentOutput(OutputsTranslation[0], DestinationAttributeForTranslation, PCGBreakTransformTest::EPCGComponentToCheck::Translation, FVector{});
+		}
+		else
+		{
+			bTestPassed = false;
+		}
+
+		if (TestEqual("Appropriate number of outputs generated for Rotation", OutputsRotation.Num(), 1))
+		{
+			ValidateComponentOutput(OutputsRotation[0], DestinationAttributeForRotation, PCGBreakTransformTest::EPCGComponentToCheck::Rotation, FQuat{});
+		}
+		else
+		{
+			bTestPassed = false;
+		}
+
+		if (TestEqual("Appropriate number of outputs generated for Scale", OutputsScale.Num(), 1))
+		{
+			ValidateComponentOutput(OutputsScale[0], DestinationAttributeForScale, PCGBreakTransformTest::EPCGComponentToCheck::Scale, FVector{});
+		}
+		else
+		{
+			bTestPassed = false;
+		}
+
+		return bTestPassed;
+	};
+
+	bool bTestPassed = true;
+
+	PCGTestsCommon::FTestData TestDataSpatial(PCGDeterminismTests::Defaults::Seed);
+	PCGDeterminismTests::GenerateSettings<UPCGMetadataBreakTransformSettings>(TestDataSpatial);
+	PCGBreakTransformTest::GenerateSpatialData(TestDataSpatial);
+
+	PCGTestsCommon::FTestData TestDataParams(PCGDeterminismTests::Defaults::Seed);
+	PCGDeterminismTests::GenerateSettings<UPCGMetadataBreakTransformSettings>(TestDataParams);
+	PCGBreakTransformTest::GenerateParamData(TestDataParams);
+
+	using PairWhatData = TPair<FString, PCGTestsCommon::FTestData*>;
+
+	TStaticArray<PairWhatData, 2> AllTestData;
+	AllTestData[0] = PairWhatData("Testing with point data as input", &TestDataSpatial);
+	AllTestData[1] = PairWhatData("Testing with param data as input", &TestDataParams);
+
+	// Setup error catching. This error should happen only twice (invalid type)
+	AddExpectedError(TEXT("Attribute Float is not a supported type for input 0"), EAutomationExpectedErrorFlags::Contains, 2);
+
+	for (PairWhatData& PairTestData : AllTestData)
+	{
+		AddInfo(PairTestData.Key);
+		PCGTestsCommon::FTestData* TestData = PairTestData.Value;
+
+		UPCGMetadataBreakTransformSettings* Settings = CastChecked<UPCGMetadataBreakTransformSettings>(TestData->Settings);
+
+		Settings->ForceOutputConnections[0] = true;
+		Settings->ForceOutputConnections[1] = true;
+		Settings->ForceOutputConnections[2] = true;
+
+		{
+			AddInfo(TEXT("Test with Transform as input attribute"));
+			Settings->InputAttributeName = PCGBreakTransformTest::TransformAttribute;
+			bTestPassed &= ValidateMetadataBreakVector(*TestData);
+		}
+
+		{
+			AddInfo(TEXT("Test with Float as input attribute (invalid type)"));
+			Settings->InputAttributeName = PCGBreakTransformTest::InvalidAttribute;
+
+			bTestPassed &= ValidateMetadataBreakVector(*TestData, /*bIsValid=*/false);
+		}
+	}
+
+	return bTestPassed;
+}
+
+#endif // WITH_EDITOR
diff --git a/Engine/Plugins/Experimental/PCG/Source/PCG/Private/Tests/Elements/Metadata/PCGMetadataBreakVectorTest.cpp b/Engine/Plugins/Experimental/PCG/Source/PCG/Private/Tests/Elements/Metadata/PCGMetadataBreakVectorTest.cpp
--- a/Engine/Plugins/Experimental/PCG/Source/PCG/Private/Tests/Elements/Metadata/PCGMetadataBreakVectorTest.cpp
+++ b/Engine/Plugins/Experimental/PCG/Source/PCG/Private/Tests/Elements/Metadata/PCGMetadataBreakVectorTest.cpp
@@ -17,7 +17,7 @@
 
 IMPLEMENT_CUSTOM_SIMPLE_AUTOMATION_TEST(FPCGMetadataBreakVectorTest, FPCGTestBaseClass, "pcg.tests.Metadata.BreakVector", PCGTestsCommon::TestFlags)
 
-namespace
+namespace PCGBreakVectorTest
 {
 	enum class EPCGComponentToCheck : uint8
 	{
@@ -162,7 +162,7 @@ bool FPCGMetadataBreakVectorTest::RunTest(const FString& Parameters)
 		const FPCGMetadataAttributeBase* SourceAttributeBase = SourceMetadata->GetConstAttribute(Settings->InputAttributeName);
 		check(SourceAttributeBase);
 
-		auto ValidateComponentOutput = [&](const FPCGTaggedData& Output, const FName& OutAttributeName, EPCGComponentToCheck ComponentToCheck)
+		auto ValidateComponentOutput = [&](const FPCGTaggedData& Output, const FName& OutAttributeName, PCGBreakVectorTest::EPCGComponentToCheck ComponentToCheck)
 		{
 			if (!TestTrue("Valid output data", Output.Data != nullptr))
 			{
@@ -256,19 +256,19 @@ bool FPCGMetadataBreakVectorTest::RunTest(const FString& Parameters)
 
 				const double OutValue = OutAttribute->GetValue(OutValueKey);
 
-				if (ComponentToCheck == EPCGComponentToCheck::X)
+				if (ComponentToCheck == PCGBreakVectorTest::EPCGComponentToCheck::X)
 				{
 					bTestPassed &= TestEqual("X", SourceValue.X, OutValue);
 				}
-				else if (ComponentToCheck == EPCGComponentToCheck::Y)
+				else if (ComponentToCheck == PCGBreakVectorTest::EPCGComponentToCheck::Y)
 				{
 					bTestPassed &= TestEqual("Y", SourceValue.Y, OutValue);
 				}
-				else if (ComponentToCheck == EPCGComponentToCheck::Z)
+				else if (ComponentToCheck == PCGBreakVectorTest::EPCGComponentToCheck::Z)
 				{
 					bTestPassed &= TestEqual("Z", SourceValue.Z, OutValue);
 				}
-				else if (ComponentToCheck == EPCGComponentToCheck::W)
+				else if (ComponentToCheck == PCGBreakVectorTest::EPCGComponentToCheck::W)
 				{
 					bTestPassed &= TestEqual("W", SourceValue.W, OutValue);
 				}
@@ -289,7 +289,7 @@ bool FPCGMetadataBreakVectorTest::RunTest(const FString& Parameters)
 		{
 			if (TestEqual("Appropriate number of outputs generated fox X", OutputsX.Num(), 1))
 			{
-				ValidateComponentOutput(OutputsX[0], DestinationAttributeForX, EPCGComponentToCheck::X);
+				ValidateComponentOutput(OutputsX[0], DestinationAttributeForX, PCGBreakVectorTest::EPCGComponentToCheck::X);
 			}
 			else
 			{
@@ -301,7 +301,7 @@ bool FPCGMetadataBreakVectorTest::RunTest(const FString& Parameters)
 		{
 			if (TestEqual("Appropriate number of outputs generated for Y", OutputsY.Num(), 1))
 			{
-				ValidateComponentOutput(OutputsY[0], DestinationAttributeForY, EPCGComponentToCheck::Y);
+				ValidateComponentOutput(OutputsY[0], DestinationAttributeForY, PCGBreakVectorTest::EPCGComponentToCheck::Y);
 			}
 			else
 			{
@@ -313,7 +313,7 @@ bool FPCGMetadataBreakVectorTest::RunTest(const FString& Parameters)
 		{
 			if (TestEqual("Appropriate number of outputs generated for Z", OutputsZ.Num(), 1))
 			{
-				ValidateComponentOutput(OutputsZ[0], DestinationAttributeForZ, EPCGComponentToCheck::Z);
+				ValidateComponentOutput(OutputsZ[0], DestinationAttributeForZ, PCGBreakVectorTest::EPCGComponentToCheck::Z);
 			}
 			else
 			{
@@ -325,7 +325,7 @@ bool FPCGMetadataBreakVectorTest::RunTest(const FString& Parameters)
 		{
 			if (TestEqual("Appropriate number of outputs generated for W", OutputsW.Num(), 1))
 			{
-				ValidateComponentOutput(OutputsW[0], DestinationAttributeForW, EPCGComponentToCheck::W);
+				ValidateComponentOutput(OutputsW[0], DestinationAttributeForW, PCGBreakVectorTest::EPCGComponentToCheck::W);
 			}
 			else
 			{
@@ -340,11 +340,11 @@ bool FPCGMetadataBreakVectorTest::RunTest(const FString& Parameters)
 
 	PCGTestsCommon::FTestData TestDataSpatial(PCGDeterminismTests::Defaults::Seed);
 	PCGDeterminismTests::GenerateSettings<UPCGMetadataBreakVectorSettings>(TestDataSpatial);
-	GenerateSpatialData(TestDataSpatial);
+	PCGBreakVectorTest::GenerateSpatialData(TestDataSpatial);
 
 	PCGTestsCommon::FTestData TestDataParams(PCGDeterminismTests::Defaults::Seed);
 	PCGDeterminismTests::GenerateSettings<UPCGMetadataBreakVectorSettings>(TestDataParams);
-	GenerateParamData(TestDataParams);
+	PCGBreakVectorTest::GenerateParamData(TestDataParams);
 
 	using PairWhatData = TPair<FString, PCGTestsCommon::FTestData*>;
 
@@ -369,7 +369,7 @@ bool FPCGMetadataBreakVectorTest::RunTest(const FString& Parameters)
 
 		{
 			AddInfo(TEXT("Test with FVector2D as input attribute"));
-			Settings->InputAttributeName = Vec2Attribute;
+			Settings->InputAttributeName = PCGBreakVectorTest::Vec2Attribute;
 			bTestPassed &= ValidateMetadataBreakVector(*TestData);
 		}
 
@@ -377,13 +377,13 @@ bool FPCGMetadataBreakVectorTest::RunTest(const FString& Parameters)
 
 		{
 			AddInfo(TEXT("Test with FVector as input attribute"));
-			Settings->InputAttributeName = Vec3Attribute;
+			Settings->InputAttributeName = PCGBreakVectorTest::Vec3Attribute;
 			bTestPassed &= ValidateMetadataBreakVector(*TestData);
 		}
 
 		{
 			AddInfo(TEXT("Test with FVector4 as input attribute"));
-			Settings->InputAttributeName = Vec4Attribute;
+			Settings->InputAttributeName = PCGBreakVectorTest::Vec4Attribute;
 			Settings->ForceOutputConnections[3] = true;
 			bTestPassed &= ValidateMetadataBreakVector(*TestData);
 			Settings->ForceOutputConnections[3] = false;
@@ -391,13 +391,13 @@ bool FPCGMetadataBreakVectorTest::RunTest(const FString& Parameters)
 
 		{
 			AddInfo(TEXT("Test with FRotator as input attribute"));
-			Settings->InputAttributeName = RotatorAttribute;
+			Settings->InputAttributeName = PCGBreakVectorTest::RotatorAttribute;
 			bTestPassed &= ValidateMetadataBreakVector(*TestData);
 		}
 
 		{
 			AddInfo(TEXT("Test with Float as input attribute (invalid type)"));
-			Settings->InputAttributeName = FloatAttribute;
+			Settings->InputAttributeName = PCGBreakVectorTest::FloatAttribute;
 
 			bTestPassed &= ValidateMetadataBreakVector(*TestData, /*bIsValid=*/false);
 		}
diff --git a/Engine/Plugins/Experimental/PCG/Source/PCG/Private/Tests/Elements/Metadata/PCGMetadataMakeTransformTest.cpp b/Engine/Plugins/Experimental/PCG/Source/PCG/Private/Tests/Elements/Metadata/PCGMetadataMakeTransformTest.cpp
new file mode 100644
--- /dev/null
+++ b/Engine/Plugins/Experimental/PCG/Source/PCG/Private/Tests/Elements/Metadata/PCGMetadataMakeTransformTest.cpp
@@ -0,0 +1,269 @@
+// Copyright Epic Games, Inc. All Rights Reserved.
+
+#include "Tests/Determinism/PCGDeterminismTestsCommon.h"
+#include "Tests/PCGTestsCommon.h"
+
+#include "PCGComponent.h"
+#include "PCGData.h"
+#include "PCGParamData.h"
+
+#include "Data/PCGPointData.h"
+#include "Elements/Metadata/PCGMetadataMakeTransform.h"
+#include "Metadata/PCGMetadata.h"
+#include "Metadata/PCGMetadataAttribute.h"
+#include "Metadata/PCGMetadataAttributeTpl.h"
+
+#if WITH_EDITOR
+
+IMPLEMENT_CUSTOM_SIMPLE_AUTOMATION_TEST(FPCGMetadataMakeTransformTest, FPCGTestBaseClass, "pcg.tests.Metadata.MakeTransform", PCGTestsCommon::TestFlags)
+
+namespace
+{
+	const FName Vec3Attribute = TEXT("Vec3");
+	const FName QuatAttribute = TEXT("Quat");
+
+	void ConnectData(FPCGDataCollection& InputData, TObjectPtr<UPCGData> Data)
+	{
+		FPCGTaggedData& TaggedData1 = InputData.TaggedData.Emplace_GetRef();
+		TaggedData1.Data = Data;
+		TaggedData1.Pin = PCGMetadataTransformConstants::Translation;
+
+		FPCGTaggedData& TaggedData2 = InputData.TaggedData.Emplace_GetRef();
+		TaggedData2.Data = Data;
+		TaggedData2.Pin = PCGMetadataTransformConstants::Rotation;
+
+		FPCGTaggedData& TaggedData3 = InputData.TaggedData.Emplace_GetRef();
+		TaggedData3.Data = Data;
+		TaggedData3.Pin = PCGMetadataTransformConstants::Scale;
+	}
+
+	void GenerateSpatialData(PCGTestsCommon::FTestData& TestData)
+	{
+		TObjectPtr<UPCGPointData> PointData = PCGTestsCommon::CreateRandomPointData(10, TestData.Seed);
+		
+		ConnectData(TestData.InputData, PointData);
+
+		const bool bAllowsInterpolation = false;
+		const bool bOverrideParent = false;
+
+		PointData->Metadata->CreateVectorAttribute(Vec3Attribute, FVector::ZeroVector, bAllowsInterpolation, bOverrideParent);
+		PointData->Metadata->CreateQuatAttribute(QuatAttribute, FQuat::Identity, bAllowsInterpolation, bOverrideParent);
+
+		TArray<FPCGPoint>& SourcePoints = PointData->GetMutablePoints();
+
+		// Randomize attribute values, and leave the second half to default value
+		FRandomStream RandomSource(TestData.Seed);
+		for (int I = 0; I < SourcePoints.Num() / 2; ++I)
+		{
+			UPCGMetadataAccessorHelpers::SetVectorAttribute(SourcePoints[I], PointData->Metadata, Vec3Attribute, RandomSource.VRand());
+			UPCGMetadataAccessorHelpers::SetQuatAttribute(SourcePoints[I], PointData->Metadata, QuatAttribute, FRotator(RandomSource.FRand(), RandomSource.FRand(), RandomSource.FRand()).Quaternion());
+		}
+	}
+
+	void GenerateParamData(PCGTestsCommon::FTestData& TestData)
+	{
+		TObjectPtr<UPCGParamData> ParamData = PCGTestsCommon::CreateEmptyParamData();
+
+		ConnectData(TestData.InputData, ParamData);
+
+		PCGMetadataEntryKey EntryKey = ParamData->Metadata->AddEntry();
+		const bool bAllowsInterpolation = false;
+		const bool bOverrideParent = false;
+
+		ParamData->Metadata->CreateVectorAttribute(Vec3Attribute, FVector::ZeroVector, bAllowsInterpolation, bOverrideParent);
+		ParamData->Metadata->CreateQuatAttribute(QuatAttribute, FQuat::Identity, bAllowsInterpolation, bOverrideParent);
+
+		// Randomize attribute values
+		FRandomStream RandomSource(TestData.Seed);
+		if (FPCGMetadataAttribute<FVector>* Attribute = static_cast<FPCGMetadataAttribute<FVector>*>(ParamData->Metadata->GetMutableAttribute(Vec3Attribute)))
+		{
+			Attribute->SetValue(EntryKey, RandomSource.VRand());
+		}
+
+		if (FPCGMetadataAttribute<FQuat>* Attribute = static_cast<FPCGMetadataAttribute<FQuat>*>(ParamData->Metadata->GetMutableAttribute(QuatAttribute)))
+		{
+			Attribute->SetValue(EntryKey, FRotator(RandomSource.FRand(), RandomSource.FRand(), RandomSource.FRand()).Quaternion());
+		}
+	}
+}
+
+
+bool FPCGMetadataMakeTransformTest::RunTest(const FString& Parameters)
+{	
+	auto ValidateMetadataMakeVector = [this](PCGTestsCommon::FTestData& TestData) -> bool
+	{
+		const FName OutputAttributeName = TEXT("Output");
+
+		UPCGMetadataMakeTransformSettings* Settings = CastChecked<UPCGMetadataMakeTransformSettings>(TestData.Settings);
+		Settings->OutputAttributeName = OutputAttributeName;
+		FPCGElementPtr MetadataBreakTransformElement = TestData.Settings->GetElement();
+
+		TUniquePtr<FPCGContext> Context = TUniquePtr<FPCGContext>(MetadataBreakTransformElement->Initialize(TestData.InputData, TestData.TestPCGComponent, nullptr));
+		Context->NumAvailableTasks = 1;
+
+		while (!MetadataBreakTransformElement->Execute(Context.Get()))
+		{}
+
+		// We have the same metadata to all pins, so just check the first one
+		const TArray<FPCGTaggedData>& TranslationInputs = Context->InputData.GetInputsByPin(PCGMetadataTransformConstants::Translation);
+
+		if (!TestEqual("Right number of translation inputs", TranslationInputs.Num(), 1))
+		{
+			return false;
+		}
+
+		const FPCGTaggedData& Input = TranslationInputs[0];
+
+		UPCGMetadata* SourceMetadata = nullptr; 
+		if (const UPCGSpatialData* SpatialInput = Cast<const UPCGSpatialData>(Input.Data))
+		{
+			SourceMetadata = SpatialInput->Metadata;
+		}
+		else if (const UPCGParamData* ParamsInput = Cast<const UPCGParamData>(Input.Data))
+		{
+			SourceMetadata = ParamsInput->Metadata;
+		}
+		else
+		{
+			return false;
+		}
+
+		check(SourceMetadata);
+
+		const FPCGMetadataAttributeBase* TranslationAttributeBase = SourceMetadata->GetConstAttribute(Settings->Input1AttributeName);
+		check(TranslationAttributeBase);
+
+		const FPCGMetadataAttributeBase* RotationAttributeBase = SourceMetadata->GetConstAttribute(Settings->Input2AttributeName);
+		check(RotationAttributeBase);
+
+		const FPCGMetadataAttributeBase* ScaleAttributeBase = SourceMetadata->GetConstAttribute(Settings->Input3AttributeName);
+		check(ScaleAttributeBase);
+
+		const TArray<FPCGTaggedData> Outputs = Context->OutputData.GetInputsByPin(PCGPinConstants::DefaultOutputLabel);
+		if (!TestEqual("Right number of outputs", Outputs.Num(), 1))
+		{
+			return false;
+		}
+
+		const FPCGTaggedData& Output = Outputs[0];
+
+		if (!TestTrue("Valid output data", Output.Data != nullptr))
+		{
+			return false;
+		}
+
+		UPCGMetadata* OutMetadata = nullptr;
+		if (const UPCGSpatialData* SpatialInput = Cast<const UPCGSpatialData>(Output.Data))
+		{
+			OutMetadata = SpatialInput->Metadata;
+		}
+		else if (const UPCGParamData* ParamsInput = Cast<const UPCGParamData>(Output.Data))
+		{
+			OutMetadata = ParamsInput->Metadata;
+		}
+		else
+		{
+			return false;
+		}
+
+		if (!TestNotNull("Valid output metadata", OutMetadata))
+		{
+			return false;
+		}
+
+		const FPCGMetadataAttributeBase* OutAttributeBase = OutMetadata->GetConstAttribute(Settings->OutputAttributeName);
+
+		if (!TestNotNull("Valid output attribute", OutAttributeBase))
+		{
+			return false;
+		}
+
+		const FPCGMetadataAttribute<FTransform>* OutAttribute = static_cast<const FPCGMetadataAttribute<FTransform>*>(OutAttributeBase);
+		check(OutAttribute);
+
+		const PCGMetadataEntryKey InEntryKeyCount = SourceMetadata->GetItemCountForChild();
+		const PCGMetadataEntryKey OutEntryKeyCount = OutMetadata->GetItemCountForChild();
+
+		if (!TestEqual("Identical EntryKey counts", InEntryKeyCount, OutEntryKeyCount))
+		{
+			return false;
+		}
+
+		const PCGMetadataValueKey InValueKeyCount = TranslationAttributeBase->GetValueKeyOffsetForChild();
+		const PCGMetadataValueKey OutValueKeyCount = OutAttribute->GetValueKeyOffsetForChild();
+
+		if (!TestEqual("Identical ValueKey counts", InValueKeyCount, OutValueKeyCount))
+		{
+			return false;
+		}
+
+		for (PCGMetadataEntryKey EntryKey = 0; EntryKey < InEntryKeyCount; ++EntryKey)
+		{
+			const PCGMetadataValueKey InValueKey = TranslationAttributeBase->GetValueKey(EntryKey);
+			const PCGMetadataValueKey OutValueKey = OutAttribute->GetValueKey(EntryKey);
+
+			if (!TestEqual("Identical value keys", InValueKey, OutValueKey))
+			{
+				return false;
+			}
+
+			FVector TranslationValue = static_cast<const FPCGMetadataAttribute<FVector>*>(TranslationAttributeBase)->GetValue(InValueKey);
+			FQuat QuatValue = static_cast<const FPCGMetadataAttribute<FQuat>*>(RotationAttributeBase)->GetValue(InValueKey);
+			FVector ScaleValue = static_cast<const FPCGMetadataAttribute<FVector>*>(ScaleAttributeBase)->GetValue(InValueKey);
+			
+			FTransform OutValue = OutAttribute->GetValue(OutValueKey);
+
+			if (!TestEqual("Mismatch between input and output for translation", OutValue.GetTranslation(), TranslationValue))
+			{
+				return false;
+			}
+
+			if (!TestEqual("Mismatch between input and output for rotation", OutValue.GetRotation(), QuatValue))
+			{
+				return false;
+			}
+
+			if (!TestEqual("Mismatch between input and output for scale", OutValue.GetScale3D(), ScaleValue))
+			{
+				return false;
+			}
+		}
+
+		return true;
+	};
+
+	bool bTestPassed = true;
+
+	PCGTestsCommon::FTestData TestDataSpatial(PCGDeterminismTests::Defaults::Seed);
+	PCGDeterminismTests::GenerateSettings<UPCGMetadataMakeTransformSettings>(TestDataSpatial);
+	GenerateSpatialData(TestDataSpatial);
+
+	PCGTestsCommon::FTestData TestDataParams(PCGDeterminismTests::Defaults::Seed);
+	PCGDeterminismTests::GenerateSettings<UPCGMetadataMakeTransformSettings>(TestDataParams);
+	GenerateParamData(TestDataParams);
+
+	using PairWhatData = TPair<FString, PCGTestsCommon::FTestData*>;
+
+	TStaticArray<PairWhatData, 2> AllTestData;
+	AllTestData[0] = PairWhatData("Testing with point data as input", &TestDataSpatial);
+	AllTestData[1] = PairWhatData("Testing with param data as input", &TestDataParams);
+
+	for (PairWhatData& PairTestData : AllTestData)
+	{
+		AddInfo(PairTestData.Key);
+		PCGTestsCommon::FTestData* TestData = PairTestData.Value;
+
+		UPCGMetadataMakeTransformSettings* Settings = CastChecked<UPCGMetadataMakeTransformSettings>(TestData->Settings);
+
+		Settings->ForceOutputConnections[0] = true;
+		Settings->Input1AttributeName = Vec3Attribute;
+		Settings->Input2AttributeName = QuatAttribute;
+		Settings->Input3AttributeName = Vec3Attribute;
+
+		bTestPassed &= ValidateMetadataMakeVector(*TestData);
+	}
+
+	return bTestPassed;
+}
+
+#endif // WITH_EDITOR
diff --git a/Engine/Plugins/Experimental/PCG/Source/PCG/Public/Elements/Metadata/PCGMetadataBreakTransform.h b/Engine/Plugins/Experimental/PCG/Source/PCG/Public/Elements/Metadata/PCGMetadataBreakTransform.h
new file mode 100644
--- /dev/null
+++ b/Engine/Plugins/Experimental/PCG/Source/PCG/Public/Elements/Metadata/PCGMetadataBreakTransform.h
@@ -0,0 +1,48 @@
+// Copyright Epic Games, Inc. All Rights Reserved.
+
+#pragma once
+
+#include "PCGElement.h"
+#include "PCGSettings.h"
+
+#include "Elements/Metadata/PCGMetadataOpElementBase.h"
+#include "Elements/Metadata/PCGMetadataMakeTransform.h"
+
+#include "PCGMetadataBreakTransform.generated.h"
+
+UCLASS(BlueprintType, ClassGroup = (Procedural))
+class PCG_API UPCGMetadataBreakTransformSettings : public UPCGMetadataSettingsBase
+{
+	GENERATED_BODY()
+
+public:
+	//~Begin UPCGSettings interface
+#if WITH_EDITOR
+	virtual FName GetDefaultNodeName() const override { return FName(TEXT("MetadataBreakTransformNode")); }
+#endif
+
+	virtual FName GetInputAttributeNameWithOverride(uint32 Index, UPCGParamData* Params) const override;
+
+	virtual FName GetOutputPinLabel(uint32 Index) const override;
+	virtual uint32 GetOutputPinNum() const override;
+
+	virtual bool IsSupportedInputType(uint16 TypeId, uint32 InputIndex, bool& bHasSpecialRequirement) const override;
+	virtual FName GetOutputAttributeName(FName BaseName, uint32 Index) const override;
+
+	virtual bool HasDifferentOutputTypes() const override;
+	virtual TArray<uint16> GetAllOutputTypes() const override;
+
+protected:
+	virtual FPCGElementPtr CreateElement() const override;
+	//~End UPCGSettings interface
+
+public:
+	UPROPERTY(BlueprintReadWrite, EditAnywhere, Category = Settings)
+	FName InputAttributeName = NAME_None;
+};
+
+class FPCGMetadataBreakTransformElement : public FPCGMetadataElementBase
+{
+protected:
+	virtual bool DoOperation(FOperationData& OperationData) const override;
+};
diff --git a/Engine/Plugins/Experimental/PCG/Source/PCG/Public/Elements/Metadata/PCGMetadataMakeTransform.h b/Engine/Plugins/Experimental/PCG/Source/PCG/Public/Elements/Metadata/PCGMetadataMakeTransform.h
new file mode 100644
--- /dev/null
+++ b/Engine/Plugins/Experimental/PCG/Source/PCG/Public/Elements/Metadata/PCGMetadataMakeTransform.h
@@ -0,0 +1,59 @@
+// Copyright Epic Games, Inc. All Rights Reserved.
+
+#pragma once
+
+#include "PCGElement.h"
+#include "PCGSettings.h"
+
+#include "Elements/Metadata/PCGMetadataOpElementBase.h"
+
+#include "PCGMetadataMakeTransform.generated.h"
+
+namespace PCGMetadataTransformConstants
+{
+	const FName Translation = TEXT("Translation");
+	const FName Rotation = TEXT("Rotation");
+	const FName Scale = TEXT("Scale");
+}
+
+UCLASS(BlueprintType, ClassGroup = (Procedural))
+class PCG_API UPCGMetadataMakeTransformSettings : public UPCGMetadataSettingsBase
+{
+	GENERATED_BODY()
+
+public:
+#if WITH_EDITOR
+	//~Begin UPCGSettings interface
+	virtual FName GetDefaultNodeName() const override { return FName(TEXT("MetadataMakeTransformNode")); }
+	//~End UPCGSettings interface
+#endif
+
+	virtual FName GetInputAttributeNameWithOverride(uint32 Index, UPCGParamData* Params) const override;
+
+	virtual FName GetInputPinLabel(uint32 Index) const override;
+	virtual uint32 GetInputPinNum() const override;
+
+	virtual bool IsSupportedInputType(uint16 TypeId, uint32 InputIndex, bool& bHasSpecialRequirement) const override;
+	virtual uint16 GetOutputType(uint16 InputTypeId) const override;
+
+protected:
+	//~Begin UPCGSettings interface
+	virtual FPCGElementPtr CreateElement() const override;
+	//~End UPCGSettings interface
+
+public:
+	UPROPERTY(BlueprintReadWrite, EditAnywhere, Category = Input)
+	FName Input1AttributeName = NAME_None;
+
+	UPROPERTY(BlueprintReadWrite, EditAnywhere, Category = Input)
+	FName Input2AttributeName = NAME_None;
+
+	UPROPERTY(BlueprintReadWrite, EditAnywhere, Category = Input)
+	FName Input3AttributeName = NAME_None;
+};
+
+class FPCGMetadataMakeTransformElement : public FPCGMetadataElementBase
+{
+protected:
+	virtual bool DoOperation(FOperationData& OperationData) const override;
+};
diff --git a/Engine/Plugins/Experimental/PCG/Source/PCG/Public/Elements/Metadata/PCGMetadataMakeVector.h b/Engine/Plugins/Experimental/PCG/Source/PCG/Public/Elements/Metadata/PCGMetadataMakeVector.h
new file mode 100644
--- /dev/null
+++ b/Engine/Plugins/Experimental/PCG/Source/PCG/Public/Elements/Metadata/PCGMetadataMakeVector.h
@@ -0,0 +1,93 @@
+// Copyright Epic Games, Inc. All Rights Reserved.
+
+#pragma once
+
+#include "PCGElement.h"
+#include "PCGSettings.h"
+
+#include "Elements/Metadata/PCGMetadataOpElementBase.h"
+
+#include "PCGMetadataMakeVector.generated.h"
+
+namespace PCGMetadataMakeVectorConstants
+{
+	const FName XLabel = TEXT("X");
+	const FName YLabel = TEXT("Y");
+	const FName ZLabel = TEXT("Z");
+	const FName WLabel = TEXT("W");
+
+	const FName XYLabel = TEXT("XY");
+	const FName ZWLabel = TEXT("ZW");
+
+	const FName XYZLabel = TEXT("XYZ");
+}
+
+UENUM()
+enum class EPCGMetadataMakeVector3 : uint8
+{
+	ThreeValues,
+	Vector2AndValue
+};
+
+UENUM()
+enum class EPCGMetadataMakeVector4 : uint8
+{
+	FourValues,
+	Vector2AndTwoValues,
+	TwoVector2,
+	Vector3AndValue
+};
+
+UCLASS(BlueprintType, ClassGroup = (Procedural))
+class PCG_API UPCGMetadataMakeVectorSettings : public UPCGMetadataSettingsBase
+{
+	GENERATED_BODY()
+
+public:
+#if WITH_EDITOR
+	//~Begin UPCGSettings interface
+	virtual FName GetDefaultNodeName() const override { return FName(TEXT("MetadataMakeVectorNode")); }
+	//~End UPCGSettings interface
+#endif
+
+	virtual FName GetInputAttributeNameWithOverride(uint32 Index, UPCGParamData* Params) const override;
+
+	virtual FName GetInputPinLabel(uint32 Index) const override;
+	virtual uint32 GetInputPinNum() const override;
+
+	virtual bool IsSupportedInputType(uint16 TypeId, uint32 InputIndex, bool& bHasSpecialRequirement) const override;
+	virtual uint16 GetOutputType(uint16 InputTypeId) const override;
+
+protected:
+	//~Begin UPCGSettings interface
+	virtual FPCGElementPtr CreateElement() const override;
+	//~End UPCGSettings interface
+
+public:
+	UPROPERTY(BlueprintReadWrite, EditAnywhere, Category = Input)
+	FName Input1AttributeName = NAME_None;
+
+	UPROPERTY(BlueprintReadWrite, EditAnywhere, Category = Input)
+	FName Input2AttributeName = NAME_None;
+
+	UPROPERTY(BlueprintReadWrite, EditAnywhere, Category = Input)
+	FName Input3AttributeName = NAME_None;
+
+	UPROPERTY(BlueprintReadWrite, EditAnywhere, Category = Input)
+	FName Input4AttributeName = NAME_None;
+
+	UPROPERTY(BlueprintReadWrite, EditAnywhere, Category = Settings, meta = (ValidEnumValues = "Vector2, Vector, Vector4"))
+	EPCGMetadataTypes OutputType = EPCGMetadataTypes::Vector2;
+
+	UPROPERTY(BlueprintReadWrite, EditAnywhere, Category = Settings, meta = (EditCondition = "OutputType == EPCGMetadataTypes::Vector", EditConditionHides))
+	EPCGMetadataMakeVector3 MakeVector3Op = EPCGMetadataMakeVector3::ThreeValues;
+
+	UPROPERTY(BlueprintReadWrite, EditAnywhere, Category = Settings, meta = (EditCondition = "OutputType == EPCGMetadataTypes::Vector4", EditConditionHides))
+	EPCGMetadataMakeVector4 MakeVector4Op = EPCGMetadataMakeVector4::FourValues;
+};
+
+class FPCGMetadataMakeVectorElement : public FPCGMetadataElementBase
+{
+protected:
+	virtual bool DoOperation(FOperationData& OperationData) const override;
+};
diff --git a/Engine/Plugins/Experimental/PCG/Source/PCG/Public/Metadata/PCGMetadataAttributeTraits.h b/Engine/Plugins/Experimental/PCG/Source/PCG/Public/Metadata/PCGMetadataAttributeTraits.h
--- a/Engine/Plugins/Experimental/PCG/Source/PCG/Public/Metadata/PCGMetadataAttributeTraits.h
+++ b/Engine/Plugins/Experimental/PCG/Source/PCG/Public/Metadata/PCGMetadataAttributeTraits.h
@@ -56,6 +56,23 @@ namespace PCG
 
 #undef PCGMetadataGenerateDataTypes
 
+		// Verify if the TypeId is a type known by PCG and matches any types provided in the template.
+		// Example: IsOfTypes<FVector2D, FVector, FVector4>(TypeId);
+		template <typename ...Types>
+		constexpr inline bool IsOfTypes(uint16 TypeId)
+		{
+			return (TypeId < (uint16)EPCGMetadataTypes::Unknown) && ((TypeId == MetadataTypes<Types>::Id) || ...);
+		}
+
+		// Verify if T is a type known by PCG and if it matches another type
+		// in all provided types (in Types...)
+		// Example: IsOfTypes<T, FVector2D, FVector, FVector4>();
+		template <typename T, typename ...Types>
+		constexpr inline bool IsOfTypes()
+		{
+			return IsOfTypes<Types...>(MetadataTypes<T>::Id);
+		}
+
 		// Wrapper around a standard 2-dimensional CArray that is constexpr, to know if a type is broadcastable to another.
 		// First index is the original type, second index is the wanted type. Returns true if we can broadcast first type into second type.
 		struct UBroadcastableTypes
