[PCG] Adaptive name for metadata operations

diff --git a/Engine/Plugins/Experimental/PCG/Source/PCG/Private/Elements/PCGCreateParamData.cpp b/Engine/Plugins/Experimental/PCG/Source/PCG/Private/Elements/PCGCreateParamData.cpp
--- a/Engine/Plugins/Experimental/PCG/Source/PCG/Private/Elements/PCGCreateParamData.cpp
+++ b/Engine/Plugins/Experimental/PCG/Source/PCG/Private/Elements/PCGCreateParamData.cpp
@@ -26,6 +26,39 @@ TArray<FPCGPinProperties> UPCGCreateParamDataSettings::OutputPinProperties() con
 	return PinProperties;
 }
 
+FName UPCGCreateParamDataSettings::AdditionalTaskName() const
+{
+	const FString Name = OutputAttributeName.ToString();
+
+	switch (Type)
+	{
+	case EPCGMetadataTypes::Integer64:
+		return FName(FString::Printf(TEXT("%s: %ll"), *Name, IntValue));
+	case EPCGMetadataTypes::Double:
+		return FName(FString::Printf(TEXT("%s: %.2f"), *Name, DoubleValue));
+	case EPCGMetadataTypes::String:
+		return FName(FString::Printf(TEXT("%s: \"%s\""), *Name, *StringValue));
+	case EPCGMetadataTypes::Name:
+		return FName(FString::Printf(TEXT("%s: N(\"%s\")"), *Name, *NameValue.ToString()));
+	case EPCGMetadataTypes::Vector2:
+		return FName(FString::Printf(TEXT("%s: V(%.2f, %.2f)"), *Name, Vector2Value.X, Vector2Value.Y));
+	case EPCGMetadataTypes::Vector:
+		return FName(FString::Printf(TEXT("%s: V(%.2f, %.2f, %.2f)"), *Name, VectorValue.X, VectorValue.Y , VectorValue.Z));
+	case EPCGMetadataTypes::Vector4:
+		return FName(FString::Printf(TEXT("%s: V(%.2f, %.2f, %.2f, %.2f)"), *Name, Vector4Value.X, Vector4Value.Y, Vector4Value.Z, Vector4Value.W));
+	case EPCGMetadataTypes::Rotator:
+		return FName(FString::Printf(TEXT("%s: R(%.2f, %.2f, %.2f)"), *Name, RotatorValue.Roll, RotatorValue.Pitch, RotatorValue.Yaw));
+	case EPCGMetadataTypes::Quaternion:
+		return FName(FString::Printf(TEXT("%s: Q(%.2f, %.2f, %.2f, %.2f)"), *Name, QuatValue.X, QuatValue.Y, QuatValue.Z, QuatValue.W));
+	case EPCGMetadataTypes::Transform:
+		return FName(FString::Printf(TEXT("%s: Transform"), *Name));
+	case EPCGMetadataTypes::Boolean:
+		return FName(FString::Printf(TEXT("%s: %s"), *Name, (BoolValue ? "True" : "False") ));
+	default:
+		return NAME_None;
+	}
+}
+
 FPCGElementPtr UPCGCreateParamDataSettings::CreateElement() const
 {
 	return MakeShared<FPCGCreateParamDataElement>();
diff --git a/Engine/Plugins/Experimental/PCG/Source/PCG/Public/Elements/PCGCreateParamData.h b/Engine/Plugins/Experimental/PCG/Source/PCG/Public/Elements/PCGCreateParamData.h
--- a/Engine/Plugins/Experimental/PCG/Source/PCG/Public/Elements/PCGCreateParamData.h
+++ b/Engine/Plugins/Experimental/PCG/Source/PCG/Public/Elements/PCGCreateParamData.h
@@ -22,6 +22,7 @@ public:
 
 	virtual TArray<FPCGPinProperties> InputPinProperties() const override;
 	virtual TArray<FPCGPinProperties> OutputPinProperties() const override;
+	virtual FName AdditionalTaskName() const override;
 	//~End UPCGSettings interface
 
 	UPROPERTY(BlueprintReadWrite, EditAnywhere, Category = Settings, meta = (ValidEnumValues = "Double, Integer64, Vector2, Vector, Vector4, Quaternion, Transform, String, Boolean, Rotator, Name"))
