[PCG] Added drag/drop support for PCG graphs and settings for volume size and whether to generate on spawn or not

diff --git a/Engine/Plugins/Experimental/PCG/Source/PCGEditor/Private/PCGVolumeFactory.cpp b/Engine/Plugins/Experimental/PCG/Source/PCGEditor/Private/PCGVolumeFactory.cpp
--- a/Engine/Plugins/Experimental/PCG/Source/PCGEditor/Private/PCGVolumeFactory.cpp
+++ b/Engine/Plugins/Experimental/PCG/Source/PCGEditor/Private/PCGVolumeFactory.cpp
@@ -1,6 +1,10 @@
 // Copyright Epic Games, Inc. All Rights Reserved.
 
 #include "PCGVolumeFactory.h"
+
+#include "PCGComponent.h"
+#include "PCGEditorSettings.h"
+#include "PCGGraph.h"
 #include "PCGVolume.h"
 
 #define LOCTEXT_NAMESPACE "PCGVolumeFactory"
@@ -13,4 +17,36 @@ UPCGVolumeFactory::UPCGVolumeFactory(const FObjectInitializer& ObjectInitializer
 	bUseSurfaceOrientation = true;
 }
 
-#undef LOCTEXT_NAMESPACE
\ No newline at end of file
+bool UPCGVolumeFactory::CanCreateActorFrom(const FAssetData& AssetData, FText& OutErrorMsg)
+{
+	if ( !AssetData.IsValid() || !AssetData.IsInstanceOf(UPCGGraph::StaticClass()))
+	{
+		OutErrorMsg = LOCTEXT("NoPCGGraph", "A valid PCG graph asset must be specified.");
+		return false;
+	}
+
+	return true;
+}
+
+void UPCGVolumeFactory::PostSpawnActor(UObject* Asset, AActor* NewActor)
+{
+	Super::PostSpawnActor(Asset, NewActor);
+
+	if (UPCGGraph* PCGGraph = Cast<UPCGGraph>(Asset))
+	{
+		const UPCGEditorSettings* PCGEditorSettings = GetDefault<UPCGEditorSettings>();
+		
+		APCGVolume* PCGVolume = CastChecked<APCGVolume>(NewActor);
+		PCGVolume->SetActorScale3D(PCGEditorSettings->VolumeScale);
+		
+		UPCGComponent* PCGComponent = CastChecked<UPCGComponent>(PCGVolume->GetComponentByClass(UPCGComponent::StaticClass()));
+		PCGComponent->SetGraph(PCGGraph);
+
+		if (PCGEditorSettings->bGenerateOnDrop)
+		{
+			PCGComponent->Generate();
+		}
+	}
+}
+
+#undef LOCTEXT_NAMESPACE
diff --git a/Engine/Plugins/Experimental/PCG/Source/PCGEditor/Private/PCGVolumeFactory.h b/Engine/Plugins/Experimental/PCG/Source/PCGEditor/Private/PCGVolumeFactory.h
--- a/Engine/Plugins/Experimental/PCG/Source/PCGEditor/Private/PCGVolumeFactory.h
+++ b/Engine/Plugins/Experimental/PCG/Source/PCGEditor/Private/PCGVolumeFactory.h
@@ -15,4 +15,9 @@ class UPCGVolumeFactory : public UActorFactoryBoxVolume
 
 public:
 	UPCGVolumeFactory(const FObjectInitializer& ObjectInitializer);
+
+	//~ Begin UActorFactory Interface
+	virtual bool CanCreateActorFrom(const FAssetData& AssetData, FText& OutErrorMsg) override;
+	virtual void PostSpawnActor(UObject* Asset, AActor* NewActor) override;
+	//~ End UActorFactory Interface
 };
\ No newline at end of file
diff --git a/Engine/Plugins/Experimental/PCG/Source/PCGEditor/Public/PCGEditorSettings.h b/Engine/Plugins/Experimental/PCG/Source/PCGEditor/Public/PCGEditorSettings.h
--- a/Engine/Plugins/Experimental/PCG/Source/PCGEditor/Public/PCGEditorSettings.h
+++ b/Engine/Plugins/Experimental/PCG/Source/PCGEditor/Public/PCGEditorSettings.h
@@ -100,4 +100,12 @@ public:
 	/** Specify if we want to jump to definition in case of double click on native PCG Nodes */
 	UPROPERTY(EditAnywhere, config, Category = Workflow)
 	bool bEnableNavigateToNativeNodes;
+
+	/** Specifies the scale of the volume created on PCG graph drag/drop */
+	UPROPERTY(EditAnywhere, Config, Category = Workflow)
+	FVector VolumeScale = FVector(25, 25, 10);
+
+	/** Whether we want to generate the PCG graph after drag/drop or not */
+	UPROPERTY(EditAnywhere, Config, Category = Workflow)
+	bool bGenerateOnDrop = true;
 };
