[PCG] Tentative fix for static analysis error

diff --git a/Engine/Plugins/Experimental/PCG/Source/PCG/Private/Elements/Metadata/PCGMetadataBreakTransform.cpp b/Engine/Plugins/Experimental/PCG/Source/PCG/Private/Elements/Metadata/PCGMetadataBreakTransform.cpp
--- a/Engine/Plugins/Experimental/PCG/Source/PCG/Private/Elements/Metadata/PCGMetadataBreakTransform.cpp
+++ b/Engine/Plugins/Experimental/PCG/Source/PCG/Private/Elements/Metadata/PCGMetadataBreakTransform.cpp
@@ -71,7 +71,14 @@ TArray<uint16> UPCGMetadataBreakTransformSettings::GetAllOutputTypes() const
 
 FName UPCGMetadataBreakTransformSettings::GetOutputAttributeName(FName BaseName, uint32 Index) const
 {
-	return BaseName != NAME_None ? FName(BaseName.ToString() + "." + GetOutputPinLabel(Index).ToString()) : NAME_None;
+	if (BaseName != NAME_None)
+	{
+		return FName(BaseName.ToString() + "." + GetOutputPinLabel(Index).ToString());
+	}
+	else
+	{
+		return NAME_None;
+	}
 }
 
 FPCGElementPtr UPCGMetadataBreakTransformSettings::CreateElement() const
