[PCG] Metadata native operation : Add Boolean op tests

diff --git a/Engine/Plugins/Experimental/PCG/Source/PCG/Private/Tests/Elements/Metadata/PCGMetadataBooleanOpTest.cpp b/Engine/Plugins/Experimental/PCG/Source/PCG/Private/Tests/Elements/Metadata/PCGMetadataBooleanOpTest.cpp
new file mode 100644
--- /dev/null
+++ b/Engine/Plugins/Experimental/PCG/Source/PCG/Private/Tests/Elements/Metadata/PCGMetadataBooleanOpTest.cpp
@@ -0,0 +1,245 @@
+// Copyright Epic Games, Inc. All Rights Reserved.
+
+#include "Tests/Determinism/PCGDeterminismTestsCommon.h"
+#include "Tests/PCGTestsCommon.h"
+
+#include "PCGComponent.h"
+#include "PCGData.h"
+#include "PCGParamData.h"
+
+#include "Data/PCGPointData.h"
+#include "Elements/Metadata/PCGMetadataBooleanOpElement.h"
+#include "Metadata/PCGMetadata.h"
+#include "Metadata/PCGMetadataAttribute.h"
+#include "Metadata/PCGMetadataAttributeTpl.h"
+
+#if WITH_EDITOR
+
+IMPLEMENT_CUSTOM_SIMPLE_AUTOMATION_TEST(FPCGMetadataBooleanOpTest, FPCGTestBaseClass, "pcg.tests.Metadata.BooleanOp", PCGTestsCommon::TestFlags)
+
+bool FPCGMetadataBooleanOpTest::RunTest(const FString& Parameters)
+{
+	PCGTestsCommon::FTestData TestData(PCGDeterminismTests::Defaults::Seed);
+	PCGDeterminismTests::GenerateSettings<UPCGMetadataBooleanSettings>(TestData);
+	UPCGMetadataBooleanSettings* Settings = CastChecked<UPCGMetadataBooleanSettings>(TestData.Settings);
+	FPCGElementPtr MetadataBooleanElement = TestData.Settings->GetElement();
+
+	FPCGTaggedData& ParamTaggedData1 = TestData.InputData.TaggedData.Emplace_GetRef(FPCGTaggedData());
+	TObjectPtr<UPCGParamData> ParamData1 = PCGTestsCommon::CreateEmptyParamData();
+	ParamTaggedData1.Data = ParamData1;
+
+	FPCGTaggedData& ParamTaggedData2 = TestData.InputData.TaggedData.Emplace_GetRef(FPCGTaggedData());
+	TObjectPtr<UPCGParamData> ParamData2 = PCGTestsCommon::CreateEmptyParamData();
+	ParamTaggedData2.Data = ParamData2;
+
+	const FName TrueAttribute = TEXT("True");
+	const FName FalseAttribute = TEXT("False");
+	const FName InvalidAttribute = TEXT("Invalid");
+	const FName OutputAttributeName = TEXT("Output");
+
+	Settings->OutputAttributeName = OutputAttributeName;
+	Settings->ForceOutputConnections[0] = true;
+
+	const bool bAllowInterpolation = false;
+	const bool bOverrideParent = false;
+
+	ParamData1->Metadata->CreateAttribute<bool>(TrueAttribute, /*DefaultValue=*/ true, bAllowInterpolation, bOverrideParent);
+	ParamData1->Metadata->CreateAttribute<bool>(FalseAttribute, /*DefaultValue=*/ false, bAllowInterpolation, bOverrideParent);
+	ParamData1->Metadata->CreateAttribute<int64>(InvalidAttribute, /*DefaultValue=*/ 0, bAllowInterpolation, bOverrideParent);
+
+	ParamData2->Metadata->CreateAttribute<bool>(TrueAttribute, /*DefaultValue=*/ true, bAllowInterpolation, bOverrideParent);
+	ParamData2->Metadata->CreateAttribute<bool>(FalseAttribute, /*DefaultValue=*/ false, bAllowInterpolation, bOverrideParent);
+	ParamData2->Metadata->CreateAttribute<int64>(InvalidAttribute, /*DefaultValue=*/ 0, bAllowInterpolation, bOverrideParent);
+
+	auto ValidateOp = [this, &MetadataBooleanElement, &TestData, &Settings](bool bExpectedResult, bool bIsValid = true) -> bool
+	{
+		TUniquePtr<FPCGContext> Context = TUniquePtr<FPCGContext>(MetadataBooleanElement->Initialize(TestData.InputData, TestData.TestPCGComponent, nullptr));
+		Context->NumAvailableTasks = 1;
+
+		while (!MetadataBooleanElement->Execute(Context.Get()))
+		{
+		}
+
+		const TArray<FPCGTaggedData> Outputs = Context->OutputData.GetInputsByPin(PCGPinConstants::DefaultOutputLabel);
+
+		// If we expect an invalid op, check that we have no output and exit immediately.
+		if (!bIsValid)
+		{
+			return TestTrue(TEXT("Invalid op == no output"), Outputs.IsEmpty());
+		}
+
+		if (!TestEqual(TEXT("Number of outputs"), Outputs.Num(), 1))
+		{
+			return false;
+		}
+
+		const UPCGData* OutputData = Outputs[0].Data;
+
+		if (!TestTrue(TEXT("Output data is param data"), OutputData->IsA<UPCGParamData>()))
+		{
+			return false;
+		}
+
+		const UPCGMetadata* OutMetadata = CastChecked<UPCGParamData>(OutputData)->Metadata;
+
+		if (!TestNotNull(TEXT("Output metadata exists"), OutMetadata))
+		{
+			return false;
+		}
+
+		const FPCGMetadataAttributeBase* OutAttributeBase = OutMetadata->GetConstAttribute(Settings->OutputAttributeName);
+		if (!TestNotNull(TEXT("Output attribute exists"), OutAttributeBase))
+		{
+			return false;
+		}
+
+		if (!TestEqual(TEXT("Output attribute is a bool"), OutAttributeBase->GetTypeId(), (uint16)PCG::Private::MetadataTypes<bool>::Id))
+		{
+			return false;
+		}
+
+		const FPCGMetadataAttribute<bool>* OutAttribute = static_cast<const FPCGMetadataAttribute<bool>*>(OutAttributeBase);
+
+		bool Out = OutAttribute->GetValueFromItemKey(PCGInvalidEntryKey);
+
+		if (!TestEqual(TEXT("Output default value"), Out, bExpectedResult))
+		{
+			return false;
+		}
+
+		return true;
+	};
+
+	bool bTestPassed = true;
+
+	/////////////////////////////////////////
+	// Not
+	/////////////////////////////////////////
+	ParamTaggedData1.Pin = PCGPinConstants::DefaultInputLabel;
+
+	Settings->Operation = EPCGMedadataBooleanOperation::Not;
+	{
+		Settings->Input1AttributeName = TrueAttribute;
+		bTestPassed &= ValidateOp(/*bExpectedResult=*/ false);
+	}
+
+	{
+		Settings->Input1AttributeName = FalseAttribute;
+		bTestPassed &= ValidateOp(true);
+	}
+
+	/////////////////////////////////////////
+	// And
+	/////////////////////////////////////////
+
+	// Validate invalid op when there is type mismatch.
+	ParamTaggedData1.Pin = PCGMetadataSettingsBaseConstants::DoubleInputFirstLabel;
+	ParamTaggedData2.Pin = PCGMetadataSettingsBaseConstants::DoubleInputSecondLabel;
+	Settings->Operation = EPCGMedadataBooleanOperation::And;
+
+	AddExpectedError(TEXT("Attribute Invalid is not a supported type for input"), EAutomationExpectedErrorFlags::Contains, 3);
+
+	{
+		Settings->Input1AttributeName = InvalidAttribute;
+		Settings->Input2AttributeName = FalseAttribute;
+		bTestPassed &= ValidateOp(/*bExpectedResult=*/false, /*bIsValid=*/ false);
+	}
+
+	{
+		Settings->Input1AttributeName = FalseAttribute;
+		Settings->Input2AttributeName = InvalidAttribute;
+		bTestPassed &= ValidateOp(/*bExpectedResult=*/false, /*bIsValid=*/ false);
+	}
+
+	{
+		Settings->Input1AttributeName = InvalidAttribute;
+		Settings->Input2AttributeName = InvalidAttribute;
+		bTestPassed &= ValidateOp(/*bExpectedResult=*/false, /*bIsValid=*/ false);
+	}
+
+	// Valid tests
+	{
+		Settings->Input1AttributeName = TrueAttribute;
+		Settings->Input2AttributeName = TrueAttribute;
+		bTestPassed &= ValidateOp(/*bExpectedResult=*/true);
+	}
+
+	{
+		Settings->Input1AttributeName = FalseAttribute;
+		Settings->Input2AttributeName = TrueAttribute;
+		bTestPassed &= ValidateOp(/*bExpectedResult=*/false);
+	}
+
+	{
+		Settings->Input1AttributeName = TrueAttribute;
+		Settings->Input2AttributeName = FalseAttribute;
+		bTestPassed &= ValidateOp(/*bExpectedResult=*/false);
+	}
+
+	{
+		Settings->Input1AttributeName = FalseAttribute;
+		Settings->Input2AttributeName = FalseAttribute;
+		bTestPassed &= ValidateOp(/*bExpectedResult=*/false);
+	}
+
+
+	/////////////////////////////////////////
+	// Or
+	/////////////////////////////////////////
+	Settings->Operation = EPCGMedadataBooleanOperation::Or;
+	{
+		Settings->Input1AttributeName = TrueAttribute;
+		Settings->Input2AttributeName = TrueAttribute;
+		bTestPassed &= ValidateOp(/*bExpectedResult=*/true);
+	}
+
+	{
+		Settings->Input1AttributeName = FalseAttribute;
+		Settings->Input2AttributeName = TrueAttribute;
+		bTestPassed &= ValidateOp(/*bExpectedResult=*/true);
+	}
+
+	{
+		Settings->Input1AttributeName = TrueAttribute;
+		Settings->Input2AttributeName = FalseAttribute;
+		bTestPassed &= ValidateOp(/*bExpectedResult=*/true);
+	}
+
+	{
+		Settings->Input1AttributeName = FalseAttribute;
+		Settings->Input2AttributeName = FalseAttribute;
+		bTestPassed &= ValidateOp(/*bExpectedResult=*/false);
+	}
+
+	/////////////////////////////////////////
+	// Xor
+	/////////////////////////////////////////
+	Settings->Operation = EPCGMedadataBooleanOperation::Xor;
+	{
+		Settings->Input1AttributeName = TrueAttribute;
+		Settings->Input2AttributeName = TrueAttribute;
+		bTestPassed &= ValidateOp(/*bExpectedResult=*/false);
+	}
+
+	{
+		Settings->Input1AttributeName = FalseAttribute;
+		Settings->Input2AttributeName = TrueAttribute;
+		bTestPassed &= ValidateOp(/*bExpectedResult=*/true);
+	}
+
+	{
+		Settings->Input1AttributeName = TrueAttribute;
+		Settings->Input2AttributeName = FalseAttribute;
+		bTestPassed &= ValidateOp(/*bExpectedResult=*/true);
+	}
+
+	{
+		Settings->Input1AttributeName = FalseAttribute;
+		Settings->Input2AttributeName = FalseAttribute;
+		bTestPassed &= ValidateOp(/*bExpectedResult=*/false);
+	}
+
+	return bTestPassed;
+}
+
+#endif // WITH_EDITOR
\ No newline at end of file
