[PCG] All metadata native operation now takes the latest attribute if it was None + remove connected output condition

diff --git a/Engine/Plugins/Experimental/PCG/Source/PCG/Private/Elements/Metadata/PCGMetadataElement.cpp b/Engine/Plugins/Experimental/PCG/Source/PCG/Private/Elements/Metadata/PCGMetadataElement.cpp
--- a/Engine/Plugins/Experimental/PCG/Source/PCG/Private/Elements/Metadata/PCGMetadataElement.cpp
+++ b/Engine/Plugins/Experimental/PCG/Source/PCG/Private/Elements/Metadata/PCGMetadataElement.cpp
@@ -100,7 +100,7 @@ bool FPCGMetadataOperationElement::ExecuteInternal(FPCGContext* Context) const
 			continue;
 		}
 
-		const FName LocalSourceAttribute = ((SourceAttribute != NAME_None) ? SourceAttribute : OriginalData->Metadata->GetSingleAttributeNameOrNone());
+		const FName LocalSourceAttribute = ((SourceAttribute != NAME_None) ? SourceAttribute : OriginalData->Metadata->GetLatestAttributeNameOrNone());
 
 		// Check if the attribute exists
 		if ((Target == EPCGMetadataOperationTarget::AttributeToProperty || Target == EPCGMetadataOperationTarget::AttributeToAttribute) && !OriginalData->Metadata->HasAttribute(LocalSourceAttribute))
diff --git a/Engine/Plugins/Experimental/PCG/Source/PCG/Private/Elements/Metadata/PCGMetadataOpElementBase.cpp b/Engine/Plugins/Experimental/PCG/Source/PCG/Private/Elements/Metadata/PCGMetadataOpElementBase.cpp
--- a/Engine/Plugins/Experimental/PCG/Source/PCG/Private/Elements/Metadata/PCGMetadataOpElementBase.cpp
+++ b/Engine/Plugins/Experimental/PCG/Source/PCG/Private/Elements/Metadata/PCGMetadataOpElementBase.cpp
@@ -67,25 +67,6 @@ bool FPCGMetadataElementBase::ExecuteInternal(FPCGContext* Context) const
 	const TArray<FPCGTaggedData>& Inputs = Context->InputData.TaggedData;
 	TArray<FPCGTaggedData>& Outputs = Context->OutputData.TaggedData;
 
-	TStaticArray<bool, UPCGMetadataSettingsBase::MaxNumberOfOutputs> ConnectedOutputs;
-	bool AnyOutputConnected = false;
-
-	for (uint32 i = 0; i < NumberOfOutputs; ++i)
-	{
-		ConnectedOutputs[i] = Context->Node && Context->Node->IsOutputPinConnected(Settings->GetOutputPinLabel(i));
-
-#if WITH_EDITOR
-		ConnectedOutputs[i] |= Settings->ForceOutputConnections[i];
-#endif // WITH_EDITOR
-
-		AnyOutputConnected |= ConnectedOutputs[i];
-	}
-
-	if (!AnyOutputConnected)
-	{
-		return true;
-	}
-
 	const TArray<FPCGTaggedData>& ParamData = Context->InputData.GetParamsByPin(PCGPinConstants::DefaultParamsLabel);
 	UPCGParamData* Params = ParamData.IsEmpty() ? nullptr : Cast<UPCGParamData>(ParamData[0].Data);
 
@@ -100,14 +81,13 @@ bool FPCGMetadataElementBase::ExecuteInternal(FPCGContext* Context) const
 	for (uint32 i = 0; i < NumberOfInputs; ++i)
 	{
 		TArray<FPCGTaggedData> InputData = Context->InputData.GetInputsByPin(Settings->GetInputPinLabel(i));
-		if (InputData.IsEmpty())
+		if (InputData.Num() != 1)
 		{
 			PCGE_LOG(Error, "Invalid inputs for pin %d", i);
 			return true;
 		}
 
 		// By construction, there can only be one of then(hence the 0 index)
-		check(InputData.Num() == 1);
 		InputTaggedData[i] = MoveTemp(InputData[0]);
 
 		// Only gather Spacial and Params input. 
@@ -141,7 +121,7 @@ bool FPCGMetadataElementBase::ExecuteInternal(FPCGContext* Context) const
 		FName SourceAttributeName = Settings->GetInputAttributeNameWithOverride(i, Params);
 		if (SourceAttributeName == NAME_None)
 		{
-			SourceAttributeName = SourceMetadata[i]->GetSingleAttributeNameOrNone();
+			SourceAttributeName = SourceMetadata[i]->GetLatestAttributeNameOrNone();
 		}
 
 		OperationData.SourceAttributes[i] = SourceMetadata[i]->GetConstAttribute(SourceAttributeName);
@@ -225,12 +205,6 @@ bool FPCGMetadataElementBase::ExecuteInternal(FPCGContext* Context) const
 
 		AttributeType DefaultValue = PCGMetadataAttribute::GetValueWithBroadcast<AttributeType>(OperationData.SourceAttributes[0], PCGDefaultValueKey);
 
-		if (!ConnectedOutputs[OutputIndex])
-		{
-			OperationData.OutputAttributes[OutputIndex] = nullptr;
-			return true;
-		}
-
 		FPCGTaggedData& OutputData = Outputs.Add_GetRef(InputTaggedData[0]);
 		OutputData.Pin = Settings->GetOutputPinLabel(OutputIndex);
 
diff --git a/Engine/Plugins/Experimental/PCG/Source/PCG/Private/Elements/Metadata/PCGMetadataPartition.cpp b/Engine/Plugins/Experimental/PCG/Source/PCG/Private/Elements/Metadata/PCGMetadataPartition.cpp
--- a/Engine/Plugins/Experimental/PCG/Source/PCG/Private/Elements/Metadata/PCGMetadataPartition.cpp
+++ b/Engine/Plugins/Experimental/PCG/Source/PCG/Private/Elements/Metadata/PCGMetadataPartition.cpp
@@ -38,7 +38,7 @@ bool FPCGMetadataPartitionElement::ExecuteInternal(FPCGContext* Context) const
 			continue;
 		}
 
-		const FName LocalPartitionAttribute = ((PartitionAttribute != NAME_None) ? PartitionAttribute : SpatialInput->ConstMetadata()->GetSingleAttributeNameOrNone());
+		const FName LocalPartitionAttribute = ((PartitionAttribute != NAME_None) ? PartitionAttribute : SpatialInput->ConstMetadata()->GetLatestAttributeNameOrNone());
 
 		if (!SpatialInput->ConstMetadata()->HasAttribute(LocalPartitionAttribute))
 		{
diff --git a/Engine/Plugins/Experimental/PCG/Source/PCG/Private/Elements/Metadata/PCGMetadataRenameElement.cpp b/Engine/Plugins/Experimental/PCG/Source/PCG/Private/Elements/Metadata/PCGMetadataRenameElement.cpp
--- a/Engine/Plugins/Experimental/PCG/Source/PCG/Private/Elements/Metadata/PCGMetadataRenameElement.cpp
+++ b/Engine/Plugins/Experimental/PCG/Source/PCG/Private/Elements/Metadata/PCGMetadataRenameElement.cpp
@@ -62,7 +62,7 @@ bool FPCGMetadataRenameElement::ExecuteInternal(FPCGContext* Context) const
 			continue;
 		}
 
-		const FName LocalAttributeToRename = ((AttributeToRename != NAME_None) ? AttributeToRename : Metadata->GetSingleAttributeNameOrNone());
+		const FName LocalAttributeToRename = ((AttributeToRename != NAME_None) ? AttributeToRename : Metadata->GetLatestAttributeNameOrNone());
 
 		if (!Metadata->HasAttribute(LocalAttributeToRename))
 		{
@@ -94,7 +94,7 @@ bool FPCGMetadataRenameElement::ExecuteInternal(FPCGContext* Context) const
 			continue;
 		}
 
-		const FName LocalAttributeToRename = ((AttributeToRename != NAME_None) ? AttributeToRename : Metadata->GetSingleAttributeNameOrNone());
+		const FName LocalAttributeToRename = ((AttributeToRename != NAME_None) ? AttributeToRename : Metadata->GetLatestAttributeNameOrNone());
 
 		if (!Metadata->HasAttribute(LocalAttributeToRename))
 		{
diff --git a/Engine/Plugins/Experimental/PCG/Source/PCG/Private/Elements/PCGPointFilter.cpp b/Engine/Plugins/Experimental/PCG/Source/PCG/Private/Elements/PCGPointFilter.cpp
--- a/Engine/Plugins/Experimental/PCG/Source/PCG/Private/Elements/PCGPointFilter.cpp
+++ b/Engine/Plugins/Experimental/PCG/Source/PCG/Private/Elements/PCGPointFilter.cpp
@@ -705,7 +705,7 @@ bool FPCGPointFilterElement::ExecuteInternal(FPCGContext* Context) const
 			ThresholdPoints = &TargetPoints;
 		}
 
-		const FName LocalTargetAttributeName = ((TargetAttributeName != NAME_None || !TargetMetadata) ? TargetAttributeName : TargetMetadata->GetSingleAttributeNameOrNone());
+		const FName LocalTargetAttributeName = ((TargetAttributeName != NAME_None || !TargetMetadata) ? TargetAttributeName : TargetMetadata->GetLatestAttributeNameOrNone());
 
 		// Additional validation
 		if (TargetFilterType == EPCGPointTargetFilterType::Metadata)
@@ -722,7 +722,7 @@ bool FPCGPointFilterElement::ExecuteInternal(FPCGContext* Context) const
 			}
 		}
 
-		const FName LocalThresholdAttributeName = ((ThresholdAttributeName != NAME_None || !ThresholdMetadata) ? ThresholdAttributeName : ThresholdMetadata->GetSingleAttributeNameOrNone());
+		const FName LocalThresholdAttributeName = ((ThresholdAttributeName != NAME_None || !ThresholdMetadata) ? ThresholdAttributeName : ThresholdMetadata->GetLatestAttributeNameOrNone());
 		
 		if (ThresholdFilterType == EPCGPointThresholdType::Metadata)
 		{
diff --git a/Engine/Plugins/Experimental/PCG/Source/PCG/Private/Metadata/PCGMetadata.cpp b/Engine/Plugins/Experimental/PCG/Source/PCG/Private/Metadata/PCGMetadata.cpp
--- a/Engine/Plugins/Experimental/PCG/Source/PCG/Private/Metadata/PCGMetadata.cpp
+++ b/Engine/Plugins/Experimental/PCG/Source/PCG/Private/Metadata/PCGMetadata.cpp
@@ -314,20 +314,19 @@ void UPCGMetadata::GetAttributes(TArray<FName>& AttributeNames, TArray<EPCGMetad
 	AttributeLock.ReadUnlock();
 }
 
-FName UPCGMetadata::GetSingleAttributeNameOrNone() const
+FName UPCGMetadata::GetLatestAttributeNameOrNone() const
 {
-	FName SingleAttributeName = NAME_None;
+	FName LatestAttributeName = NAME_None;
 	AttributeLock.ReadLock();
-	if (Attributes.Num() == 1)
+	if (!Attributes.IsEmpty())
 	{
-		TArray<FName, TInlineAllocator<1>> Keys;
+		TArray<FName> Keys;
 		Attributes.GenerateKeyArray(Keys);
-		check(Keys.Num() == 1);
-		SingleAttributeName = Keys[0];
+		LatestAttributeName = Keys.Last();
 	}
 	AttributeLock.ReadUnlock();
 
-	return SingleAttributeName;
+	return LatestAttributeName;
 }
 
 bool UPCGMetadata::ParentHasAttribute(FName AttributeName) const
@@ -415,16 +414,27 @@ bool UPCGMetadata::SetAttributeFromProperty(FName AttributeName, PCGMetadataEntr
 			return false;
 		}
 
-		// Check that the property match the attribute type!
-		if (PCG::Private::MetadataTypes<PropertyType>::Id != BaseAttribute->GetTypeId())
-		{
-			return false;
-		}
-
-		FPCGMetadataAttribute<PropertyType>* Attribute = static_cast<FPCGMetadataAttribute<PropertyType>*>(BaseAttribute);
-		Attribute->SetValue(EntryKey, PropertyValue);
-
-		return true;
+		// Allow to set the value if both type matches or if we can construct AttributeType from PropertyType.
+		return PCGMetadataAttribute::CallbackWithRightType(BaseAttribute->GetTypeId(), [&](auto AttributeValue) -> bool
+			{
+				using AttributeType = decltype(AttributeValue);
+				FPCGMetadataAttribute<AttributeType>* Attribute = static_cast<FPCGMetadataAttribute<AttributeType>*>(BaseAttribute);
+
+				if constexpr (std::is_same_v<AttributeType, PropertyType>)
+				{
+					Attribute->SetValue(EntryKey, PropertyValue);
+					return true;
+				}
+				else if constexpr (std::is_constructible_v<AttributeType, PropertyType>)
+				{
+					Attribute->SetValue(EntryKey, AttributeType(PropertyValue));
+					return true;
+				}
+				else
+				{
+					return false;
+				}
+			});
 	};
 
 	return PCGSettingsHelpers::GetPropertyValueWithCallback(Object, InProperty, CreateAttributeAndSet);
diff --git a/Engine/Plugins/Experimental/PCG/Source/PCG/Public/Metadata/PCGMetadata.h b/Engine/Plugins/Experimental/PCG/Source/PCG/Public/Metadata/PCGMetadata.h
--- a/Engine/Plugins/Experimental/PCG/Source/PCG/Public/Metadata/PCGMetadata.h
+++ b/Engine/Plugins/Experimental/PCG/Source/PCG/Public/Metadata/PCGMetadata.h
@@ -106,7 +106,7 @@ public:
 	UFUNCTION(BlueprintCallable, Category = "PCG|Metadata")
 	void GetAttributes(TArray<FName>& AttributeNames, TArray<EPCGMetadataTypes>& AttributeTypes) const;
 
-	FName GetSingleAttributeNameOrNone() const;
+	FName GetLatestAttributeNameOrNone() const;
 
 	/** Delete/Hide attribute */
 	// Due to stream inheriting, we might want to consider "hiding" parent stream and deleting local streams only
