[PCG] Allow to specify PostSpawn functions to be called on all spawned actors

diff --git a/Engine/Plugins/Experimental/PCG/Source/PCG/Private/Elements/PCGSpawnActor.cpp b/Engine/Plugins/Experimental/PCG/Source/PCG/Private/Elements/PCGSpawnActor.cpp
--- a/Engine/Plugins/Experimental/PCG/Source/PCG/Private/Elements/PCGSpawnActor.cpp
+++ b/Engine/Plugins/Experimental/PCG/Source/PCG/Private/Elements/PCGSpawnActor.cpp
@@ -83,6 +83,31 @@ bool FPCGSpawnActorElement::ExecuteInternal(FPCGContext* Context) const
 		return true;
 	}
 
+	TArray<UFunction*> PostSpawnFunctions;
+	for (FName PostSpawnFunctionName : Settings->PostSpawnFunctionNames)
+	{
+		if (PostSpawnFunctionName == NAME_None)
+		{
+			continue;
+		}
+
+		if (UFunction* PostSpawnFunction = Settings->TemplateActorClass->FindFunctionByName(PostSpawnFunctionName))
+		{
+			if (PostSpawnFunction->NumParms != 0)
+			{
+				PCGE_LOG(Warning, "PostSpawnFunction \"%s\" requires parameters. We only support parameter-less functions. Will skip the call.", *PostSpawnFunctionName.ToString());
+			}
+			else
+			{
+				PostSpawnFunctions.Add(PostSpawnFunction);
+			}
+		}
+		else
+		{
+			PCGE_LOG(Warning, "PostSpawnFunction \"%s\" was not found in class \"%s\".", *PostSpawnFunctionName.ToString(), *Settings->TemplateActorClass->GetFName().ToString());
+		}
+	}
+
 	const bool bForceDisableActorParsing = (Settings->bForceDisableActorParsing);
 
 	// Pass-through exclusions & settings
@@ -193,6 +218,11 @@ bool FPCGSpawnActorElement::ExecuteInternal(FPCGContext* Context) const
 						GeneratedActor->Tags.Append(NewActorTags);
 						GeneratedActor->AttachToActor(TargetActor, FAttachmentTransformRules::KeepWorldTransform);
 
+						for (UFunction* PostSpawnFunction : PostSpawnFunctions)
+						{
+							GeneratedActor->ProcessEvent(PostSpawnFunction, nullptr);
+						}
+
 						ManagedActors->GeneratedActors.Add(GeneratedActor);
 
 						// If the actor spawned has a PCG component, either generate it or mark it as generated if we pass through its inputs
diff --git a/Engine/Plugins/Experimental/PCG/Source/PCG/Public/Elements/PCGSpawnActor.h b/Engine/Plugins/Experimental/PCG/Source/PCG/Public/Elements/PCGSpawnActor.h
--- a/Engine/Plugins/Experimental/PCG/Source/PCG/Public/Elements/PCGSpawnActor.h
+++ b/Engine/Plugins/Experimental/PCG/Source/PCG/Public/Elements/PCGSpawnActor.h
@@ -29,6 +29,10 @@ public:
 	UPROPERTY(BlueprintReadWrite, EditAnywhere, Category = Settings)
 	TSubclassOf<AActor> TemplateActorClass = nullptr;
 
+	/* Can specify a list of functions from the template class to be called on each actor spawned, in order. Need to be parameter-less and with "CallInEditor" flag enabled.*/
+	UPROPERTY(BlueprintReadWrite, EditAnywhere, Category = Settings, meta = (EditCondition = "Option!=EPCGSpawnActorOption::CollapseActors"))
+	TArray<FName> PostSpawnFunctionNames;
+
 	UPROPERTY(BlueprintReadWrite, EditAnywhere, Category = Settings)
 	EPCGSpawnActorOption Option = EPCGSpawnActorOption::CollapseActors;
 
