[PCG] Metadata native operation: Create attribute node

* Refactor CreateParamData to make it use the same code than the Create Attribute node
* Remove the connected output condition on create param data.

diff --git a/Engine/Plugins/Experimental/PCG/Source/PCG/Private/Elements/PCGCreateAttribute.cpp b/Engine/Plugins/Experimental/PCG/Source/PCG/Private/Elements/PCGCreateAttribute.cpp
new file mode 100644
--- /dev/null
+++ b/Engine/Plugins/Experimental/PCG/Source/PCG/Private/Elements/PCGCreateAttribute.cpp
@@ -0,0 +1,129 @@
+// Copyright Epic Games, Inc. All Rights Reserved.
+
+#include "Elements/PCGCreateAttribute.h"
+
+#include "PCGData.h"
+#include "PCGParamData.h"
+#include "Data/PCGPointData.h"
+#include "Metadata/PCGMetadata.h"
+#include "Metadata/PCGMetadataAttribute.h"
+
+#if WITH_EDITOR
+FName UPCGCreateAttributeSettings::GetDefaultNodeName() const
+{
+	return TEXT("CreateAttribute");
+}
+#endif
+
+TArray<FPCGPinProperties> UPCGCreateAttributeSettings::InputPinProperties() const
+{
+	TArray<FPCGPinProperties> PinProperties;
+	PinProperties.Emplace(PCGPinConstants::DefaultInputLabel, EPCGDataType::Any, /*bInAllowMultipleConnections=*/ true);
+	PinProperties.Emplace(PCGPinConstants::DefaultParamsLabel, EPCGDataType::Param, /*bInAllowMultipleConnections=*/ false);
+
+	return PinProperties;
+}
+
+TArray<FPCGPinProperties> UPCGCreateAttributeSettings::OutputPinProperties() const
+{
+	TArray<FPCGPinProperties> PinProperties;
+	PinProperties.Emplace(PCGPinConstants::DefaultOutputLabel, EPCGDataType::Any);
+
+	return PinProperties;
+}
+
+FPCGElementPtr UPCGCreateAttributeSettings::CreateElement() const
+{
+	return MakeShared<FPCGCreateAttributeElement>();
+}
+
+bool FPCGCreateAttributeElement::ExecuteInternal(FPCGContext* Context) const
+{
+	TRACE_CPUPROFILER_EVENT_SCOPE(FPCGCreateAttributeElement::Execute);
+
+	check(Context);
+
+	const UPCGCreateAttributeSettings* Settings = Context->GetInputSettings<UPCGCreateAttributeSettings>();
+	check(Settings);
+
+	TArray<FPCGTaggedData> Params = Context->InputData.GetInputsByPin(PCGPinConstants::DefaultParamsLabel);
+	UPCGParamData* ParamData = nullptr;
+
+	if (!Params.IsEmpty())
+	{
+		ParamData = CastChecked<UPCGParamData>(Params[0].Data);
+	}
+
+	TArray<FPCGTaggedData> Inputs = Context->InputData.GetInputsByPin(PCGPinConstants::DefaultInputLabel);
+
+	if (Inputs.IsEmpty())
+	{
+		PCGE_LOG(Error, "No input connected");
+		return true;
+	}
+
+	for (const FPCGTaggedData& InputTaggedData : Inputs)
+	{
+		const UPCGData* InputData = InputTaggedData.Data;
+
+		const UPCGMetadata* ParentMetadata = nullptr;
+
+		if (const UPCGPointData* InputPointData = Cast<UPCGPointData>(InputData))
+		{
+			ParentMetadata = InputPointData->Metadata;
+		}
+		else if (const UPCGParamData* InputParamData = Cast<UPCGParamData>(InputData))
+		{
+			ParentMetadata = InputParamData->Metadata;
+		}
+		else
+		{
+			PCGE_LOG(Error, "Invalid data as input. Only support points and params");
+			return true;
+		}
+
+		UPCGMetadata* Metadata = NewObject<UPCGMetadata>();
+		if (Settings->bKeepExistingAttributes)
+		{
+			Metadata->Initialize(ParentMetadata);
+		}
+
+		FPCGMetadataAttributeBase* Attribute = Settings->ClearOrCreateAttribute(Metadata, ParamData);
+
+		if (!Attribute)
+		{
+			PCGE_LOG(Error, "Error while creating attribute %s", *Settings->OutputAttributeName.ToString());
+			return true;
+		}
+
+		TArray<FPCGTaggedData>& Outputs = Context->OutputData.TaggedData;
+		FPCGTaggedData& Output = Outputs.Emplace_GetRef();
+
+		if (const UPCGPointData* InputPointData = Cast<UPCGPointData>(InputData))
+		{
+			UPCGPointData* PointOutputData = NewObject<UPCGPointData>();
+			PointOutputData->Metadata = Metadata;
+			PointOutputData->InitializeFromData(InputPointData);
+
+			// Copy the points
+			TArray<FPCGPoint>& Points = PointOutputData->GetMutablePoints();
+			Points = InputPointData->GetPoints();
+
+			Settings->SetAttribute(Attribute, Metadata, Points, ParamData);
+
+			Output.Data = PointOutputData;
+		}
+		else if (const UPCGParamData* InputParamData = Cast<UPCGParamData>(InputData))
+		{
+			UPCGParamData* OutputParamData = NewObject<UPCGParamData>();
+			OutputParamData->Metadata = Metadata;
+
+			PCGMetadataEntryKey EntryKey = Metadata->AddEntry();
+			Settings->SetAttribute(Attribute, Metadata, EntryKey, ParamData);
+
+			Output.Data = OutputParamData;
+		}
+	}
+
+	return true;
+}
\ No newline at end of file
diff --git a/Engine/Plugins/Experimental/PCG/Source/PCG/Private/Elements/PCGCreateAttributeBase.cpp b/Engine/Plugins/Experimental/PCG/Source/PCG/Private/Elements/PCGCreateAttributeBase.cpp
new file mode 100644
--- /dev/null
+++ b/Engine/Plugins/Experimental/PCG/Source/PCG/Private/Elements/PCGCreateAttributeBase.cpp
@@ -0,0 +1,148 @@
+// Copyright Epic Games, Inc. All Rights Reserved.
+
+#include "Elements/PCGCreateAttributeBase.h"
+
+#include "PCGParamData.h"
+#include "Elements/Metadata/PCGMetadataElementCommon.h"
+#include "Metadata/PCGMetadata.h"
+#include "Helpers/PCGSettingsHelpers.h"
+
+namespace PCGAttributeCreationBaseSettings
+{
+	template <typename Func>
+	decltype(auto) Dispatcher(const UPCGAttributeCreationBaseSettings* Settings, UPCGParamData* Params, Func Callback)
+	{
+		using ReturnType = decltype(Callback(double{}));
+
+		switch (Settings->Type)
+		{
+		case EPCGMetadataTypes::Integer64:
+			return Callback(PCG_GET_OVERRIDEN_VALUE(Settings, IntValue, Params));
+		case EPCGMetadataTypes::Double:
+			return Callback(PCG_GET_OVERRIDEN_VALUE(Settings, DoubleValue, Params));
+		case EPCGMetadataTypes::Vector2:
+			return Callback(PCG_GET_OVERRIDEN_VALUE(Settings, Vector2Value, Params));
+		case EPCGMetadataTypes::Vector:
+			return Callback(PCG_GET_OVERRIDEN_VALUE(Settings, VectorValue, Params));
+		case EPCGMetadataTypes::Vector4:
+			return Callback(PCG_GET_OVERRIDEN_VALUE(Settings, Vector4Value, Params));
+		case EPCGMetadataTypes::Quaternion:
+			return Callback(PCG_GET_OVERRIDEN_VALUE(Settings, QuatValue, Params));
+		case EPCGMetadataTypes::Transform:
+			return Callback(PCG_GET_OVERRIDEN_VALUE(Settings, TransformValue, Params));
+		case EPCGMetadataTypes::String:
+			return Callback(PCG_GET_OVERRIDEN_VALUE(Settings, StringValue, Params));
+		case EPCGMetadataTypes::Boolean:
+			return Callback(PCG_GET_OVERRIDEN_VALUE(Settings, BoolValue, Params));
+		case EPCGMetadataTypes::Rotator:
+			return Callback(PCG_GET_OVERRIDEN_VALUE(Settings, RotatorValue, Params));
+		case EPCGMetadataTypes::Name:
+			return Callback(PCG_GET_OVERRIDEN_VALUE(Settings, NameValue, Params));
+		default:
+			// ReturnType{} is invalid if ReturnType is void
+			if constexpr (std::is_same_v<ReturnType, void>)
+			{
+				return;
+			}
+			else
+			{
+				return ReturnType{};
+			}
+		}
+	}
+}
+
+FPCGMetadataAttributeBase* UPCGAttributeCreationBaseSettings::ClearOrCreateAttribute(UPCGMetadata* Metadata, UPCGParamData* Params) const
+{
+	check(Metadata);
+
+	auto CreateAttribute = [this, Metadata](auto&& Value) -> FPCGMetadataAttributeBase*
+	{
+		return PCGMetadataElementCommon::ClearOrCreateAttribute(Metadata, OutputAttributeName, Value);
+	};
+
+	return PCGAttributeCreationBaseSettings::Dispatcher(this, Params, CreateAttribute);
+}
+
+PCGMetadataEntryKey UPCGAttributeCreationBaseSettings::SetAttribute(FPCGMetadataAttributeBase* Attribute, UPCGMetadata* Metadata, PCGMetadataEntryKey EntryKey, UPCGParamData* Params) const
+{
+	check(Attribute && Metadata);
+
+	auto SetAttribute = [this, Attribute, EntryKey, Metadata](auto&& Value) -> PCGMetadataEntryKey
+	{
+		using AttributeType = std::remove_reference_t<decltype(Value)>;
+
+		check(Attribute->GetTypeId() == PCG::Private::MetadataTypes<AttributeType>::Id);
+
+		const PCGMetadataEntryKey FinalKey = (EntryKey == PCGInvalidEntryKey) ? Metadata->AddEntry() : EntryKey;
+
+		static_cast<FPCGMetadataAttribute<AttributeType>*>(Attribute)->SetValue(FinalKey, Value);
+
+		return FinalKey;
+	};
+
+	return PCGAttributeCreationBaseSettings::Dispatcher(this, Params, SetAttribute);
+}
+
+void UPCGAttributeCreationBaseSettings::SetAttribute(FPCGMetadataAttributeBase* Attribute, UPCGMetadata* Metadata, TArray<FPCGPoint>& Points, UPCGParamData* Params) const
+{
+	if (Points.IsEmpty())
+	{
+		return;
+	}
+
+	check(Attribute && Metadata);
+
+	auto SetAttribute = [this, Attribute, Metadata, &Points](auto&& Value) -> void
+	{
+		using AttributeType = std::remove_reference_t<decltype(Value)>;
+
+		check(Attribute->GetTypeId() == PCG::Private::MetadataTypes<AttributeType>::Id);
+
+		FPCGMetadataAttribute<AttributeType>* TypedAttribute = static_cast<FPCGMetadataAttribute<AttributeType>*>(Attribute);
+
+		for (FPCGPoint& Point : Points)
+		{
+			PCGMetadataEntryKey EntryKey = (Point.MetadataEntry == PCGInvalidEntryKey) ? Metadata->AddEntry() : Point.MetadataEntry;
+
+			TypedAttribute->SetValue(EntryKey, Value);
+
+			Point.MetadataEntry = EntryKey;
+		}
+	};
+
+	PCGAttributeCreationBaseSettings::Dispatcher(this, Params, SetAttribute);
+}
+
+FName UPCGAttributeCreationBaseSettings::AdditionalTaskName() const
+{
+	const FString Name = OutputAttributeName.ToString();
+
+	switch (Type)
+	{
+	case EPCGMetadataTypes::Integer64:
+		return FName(FString::Printf(TEXT("%s: %ll"), *Name, IntValue));
+	case EPCGMetadataTypes::Double:
+		return FName(FString::Printf(TEXT("%s: %.2f"), *Name, DoubleValue));
+	case EPCGMetadataTypes::String:
+		return FName(FString::Printf(TEXT("%s: \"%s\""), *Name, *StringValue));
+	case EPCGMetadataTypes::Name:
+		return FName(FString::Printf(TEXT("%s: N(\"%s\")"), *Name, *NameValue.ToString()));
+	case EPCGMetadataTypes::Vector2:
+		return FName(FString::Printf(TEXT("%s: V(%.2f, %.2f)"), *Name, Vector2Value.X, Vector2Value.Y));
+	case EPCGMetadataTypes::Vector:
+		return FName(FString::Printf(TEXT("%s: V(%.2f, %.2f, %.2f)"), *Name, VectorValue.X, VectorValue.Y, VectorValue.Z));
+	case EPCGMetadataTypes::Vector4:
+		return FName(FString::Printf(TEXT("%s: V(%.2f, %.2f, %.2f, %.2f)"), *Name, Vector4Value.X, Vector4Value.Y, Vector4Value.Z, Vector4Value.W));
+	case EPCGMetadataTypes::Rotator:
+		return FName(FString::Printf(TEXT("%s: R(%.2f, %.2f, %.2f)"), *Name, RotatorValue.Roll, RotatorValue.Pitch, RotatorValue.Yaw));
+	case EPCGMetadataTypes::Quaternion:
+		return FName(FString::Printf(TEXT("%s: Q(%.2f, %.2f, %.2f, %.2f)"), *Name, QuatValue.X, QuatValue.Y, QuatValue.Z, QuatValue.W));
+	case EPCGMetadataTypes::Transform:
+		return FName(FString::Printf(TEXT("%s: Transform"), *Name));
+	case EPCGMetadataTypes::Boolean:
+		return FName(FString::Printf(TEXT("%s: %s"), *Name, (BoolValue ? "True" : "False")));
+	default:
+		return NAME_None;
+	}
+}
\ No newline at end of file
diff --git a/Engine/Plugins/Experimental/PCG/Source/PCG/Private/Elements/PCGCreateParamData.cpp b/Engine/Plugins/Experimental/PCG/Source/PCG/Private/Elements/PCGCreateParamData.cpp
--- a/Engine/Plugins/Experimental/PCG/Source/PCG/Private/Elements/PCGCreateParamData.cpp
+++ b/Engine/Plugins/Experimental/PCG/Source/PCG/Private/Elements/PCGCreateParamData.cpp
@@ -2,14 +2,16 @@
 
 #include "Elements/PCGCreateParamData.h"
 
+#include "PCGData.h"
 #include "PCGParamData.h"
 #include "Metadata/PCGMetadata.h"
-#include "Elements/Metadata/PCGMetadataOpElementBase.h"
+#include "Metadata/PCGMetadataAttribute.h"
+#include "Metadata/PCGMetadataCommon.h"
 
 #if WITH_EDITOR
 FName UPCGCreateParamDataSettings::GetDefaultNodeName() const
 {
-	return TEXT("CreateParamDataNode");
+	return TEXT("CreateParamData");
 }
 #endif
 
@@ -26,39 +28,6 @@ TArray<FPCGPinProperties> UPCGCreateParamDataSettings::OutputPinProperties() con
 	return PinProperties;
 }
 
-FName UPCGCreateParamDataSettings::AdditionalTaskName() const
-{
-	const FString Name = OutputAttributeName.ToString();
-
-	switch (Type)
-	{
-	case EPCGMetadataTypes::Integer64:
-		return FName(FString::Printf(TEXT("%s: %ll"), *Name, IntValue));
-	case EPCGMetadataTypes::Double:
-		return FName(FString::Printf(TEXT("%s: %.2f"), *Name, DoubleValue));
-	case EPCGMetadataTypes::String:
-		return FName(FString::Printf(TEXT("%s: \"%s\""), *Name, *StringValue));
-	case EPCGMetadataTypes::Name:
-		return FName(FString::Printf(TEXT("%s: N(\"%s\")"), *Name, *NameValue.ToString()));
-	case EPCGMetadataTypes::Vector2:
-		return FName(FString::Printf(TEXT("%s: V(%.2f, %.2f)"), *Name, Vector2Value.X, Vector2Value.Y));
-	case EPCGMetadataTypes::Vector:
-		return FName(FString::Printf(TEXT("%s: V(%.2f, %.2f, %.2f)"), *Name, VectorValue.X, VectorValue.Y , VectorValue.Z));
-	case EPCGMetadataTypes::Vector4:
-		return FName(FString::Printf(TEXT("%s: V(%.2f, %.2f, %.2f, %.2f)"), *Name, Vector4Value.X, Vector4Value.Y, Vector4Value.Z, Vector4Value.W));
-	case EPCGMetadataTypes::Rotator:
-		return FName(FString::Printf(TEXT("%s: R(%.2f, %.2f, %.2f)"), *Name, RotatorValue.Roll, RotatorValue.Pitch, RotatorValue.Yaw));
-	case EPCGMetadataTypes::Quaternion:
-		return FName(FString::Printf(TEXT("%s: Q(%.2f, %.2f, %.2f, %.2f)"), *Name, QuatValue.X, QuatValue.Y, QuatValue.Z, QuatValue.W));
-	case EPCGMetadataTypes::Transform:
-		return FName(FString::Printf(TEXT("%s: Transform"), *Name));
-	case EPCGMetadataTypes::Boolean:
-		return FName(FString::Printf(TEXT("%s: %s"), *Name, (BoolValue ? "True" : "False") ));
-	default:
-		return NAME_None;
-	}
-}
-
 FPCGElementPtr UPCGCreateParamDataSettings::CreateElement() const
 {
 	return MakeShared<FPCGCreateParamDataElement>();
@@ -73,69 +42,22 @@ bool FPCGCreateParamDataElement::ExecuteInternal(FPCGContext* Context) const
 	const UPCGCreateParamDataSettings* Settings = Context->GetInputSettings<UPCGCreateParamDataSettings>();
 	check(Settings);
 
-	// If we have no output connected, nothing to do
-	if (!Context->Node || !Context->Node->IsOutputPinConnected(PCGPinConstants::DefaultOutputLabel))
-	{
-		PCGE_LOG(Verbose, "Node is not connected, nothing to do");
-		return true;
-	}
-
 	// From there, we should be able to create the data.
 	UPCGParamData* ParamData = NewObject<UPCGParamData>();
 	UPCGMetadata* Metadata = ParamData->MutableMetadata();
 	check(Metadata);
 	PCGMetadataEntryKey EntryKey = Metadata->AddEntry();
 
-	auto CreateAndSetAttribute = [&Metadata, EntryKey, AttributeName = Settings->OutputAttributeName](const auto& Value)
-	{
-		using AttributeType = std::remove_const_t<std::remove_reference_t<decltype(Value)>>;
-
-		FPCGMetadataAttributeBase* BaseAttribute = Metadata->CreateAttribute<AttributeType>(AttributeName, Value, false, false);
+	FPCGMetadataAttributeBase* Attribute = Settings->ClearOrCreateAttribute(Metadata);
 
-		FPCGMetadataAttribute<AttributeType>* Attribute = static_cast<FPCGMetadataAttribute<AttributeType>*>(BaseAttribute);
-		Attribute->SetValue(EntryKey, Value);
-	};
-
-	switch (Settings->Type)
+	if (!Attribute)
 	{
-
-	case EPCGMetadataTypes::Integer64:
-		CreateAndSetAttribute(Settings->IntValue);
-		break;
-	case EPCGMetadataTypes::Double:
-		CreateAndSetAttribute(Settings->DoubleValue);
-		break;
-	case EPCGMetadataTypes::Vector2:
-		CreateAndSetAttribute(Settings->Vector2Value);
-		break;
-	case EPCGMetadataTypes::Vector:
-		CreateAndSetAttribute(Settings->VectorValue);
-		break;
-	case EPCGMetadataTypes::Vector4:
-		CreateAndSetAttribute(Settings->Vector4Value);
-		break;
-	case EPCGMetadataTypes::Quaternion:
-		CreateAndSetAttribute(Settings->QuatValue);
-		break;
-	case EPCGMetadataTypes::Transform:
-		CreateAndSetAttribute(Settings->TransformValue);
-		break;
-	case EPCGMetadataTypes::String:
-		CreateAndSetAttribute(Settings->StringValue);
-		break;
-	case EPCGMetadataTypes::Boolean:
-		CreateAndSetAttribute(Settings->BoolValue);
-		break;
-	case EPCGMetadataTypes::Rotator:
-		CreateAndSetAttribute(Settings->RotatorValue);
-		break;
-	case EPCGMetadataTypes::Name:
-		CreateAndSetAttribute(Settings->NameValue);
-		break;
-	default:
-		break;
+		PCGE_LOG(Error, "Error while creating attribute %s", *Settings->OutputAttributeName.ToString());
+		return true;
 	}
 
+	Settings->SetAttribute(Attribute, Metadata, EntryKey);
+
 	TArray<FPCGTaggedData>& Outputs = Context->OutputData.TaggedData;
 	FPCGTaggedData& Output = Outputs.Emplace_GetRef();
 	Output.Data = ParamData;
diff --git a/Engine/Plugins/Experimental/PCG/Source/PCG/Public/Elements/PCGCreateAttribute.h b/Engine/Plugins/Experimental/PCG/Source/PCG/Public/Elements/PCGCreateAttribute.h
new file mode 100644
--- /dev/null
+++ b/Engine/Plugins/Experimental/PCG/Source/PCG/Public/Elements/PCGCreateAttribute.h
@@ -0,0 +1,41 @@
+// Copyright Epic Games, Inc. All Rights Reserved.
+
+#pragma once
+
+#include "PCGElement.h"
+#include "Elements/PCGCreateAttributeBase.h"
+
+#include "PCGCreateAttribute.generated.h"
+
+UCLASS(BlueprintType, ClassGroup = (Procedural))
+class PCG_API UPCGCreateAttributeSettings : public UPCGAttributeCreationBaseSettings
+{
+	GENERATED_BODY()
+
+public:
+	//~Begin UPCGSettings interface
+#if WITH_EDITOR
+	virtual FName GetDefaultNodeName() const override;
+	virtual EPCGSettingsType GetType() const override { return EPCGSettingsType::Metadata; }
+#endif
+
+	virtual TArray<FPCGPinProperties> InputPinProperties() const override;
+	virtual TArray<FPCGPinProperties> OutputPinProperties() const override;
+	//~End UPCGSettings interface
+
+	UPROPERTY(BlueprintReadWrite, EditAnywhere, Category = Settings)
+	bool bKeepExistingAttributes = false;
+
+protected:
+	virtual FPCGElementPtr CreateElement() const override;
+};
+
+
+class FPCGCreateAttributeElement : public FSimplePCGElement
+{
+public:
+	virtual bool IsCacheable(const UPCGSettings* InSettings) const override { return false; }
+
+protected:
+	virtual bool ExecuteInternal(FPCGContext* Context) const override;
+};
\ No newline at end of file
diff --git a/Engine/Plugins/Experimental/PCG/Source/PCG/Public/Elements/PCGCreateAttributeBase.h b/Engine/Plugins/Experimental/PCG/Source/PCG/Public/Elements/PCGCreateAttributeBase.h
new file mode 100644
--- /dev/null
+++ b/Engine/Plugins/Experimental/PCG/Source/PCG/Public/Elements/PCGCreateAttributeBase.h
@@ -0,0 +1,79 @@
+// Copyright Epic Games, Inc. All Rights Reserved.
+
+#pragma once
+
+#include "PCGPoint.h"
+#include "PCGSettings.h"
+#include "Metadata/PCGMetadataCommon.h"
+#include "Metadata/PCGMetadataAttributeTpl.h"
+
+#include "PCGCreateAttributeBase.generated.h"
+
+class FPCGMetadataAttributeBase;
+class UPCGMetadata;
+class UPCGParamData;
+
+/*
+* Base settings class for any node that needs to create new attributes.
+* Contains all supported PCG types, a utility function to change the node name depending of the created attribute
+* and a utility function to create/set the attribute value depending on the type in the settings;
+*/
+UCLASS(BlueprintType, ClassGroup = (Procedural), Abstract)
+class PCG_API UPCGAttributeCreationBaseSettings : public UPCGSettings
+{
+	GENERATED_BODY()
+
+public:
+	//~Begin UPCGSettings interface
+	virtual FName AdditionalTaskName() const override;
+	//~End UPCGSettings interface
+
+	/* Create (or clear) an attribute named by OutputAttributeName and depending on the selected type. Value can be overridden by params. Default value will be set to the specified value. */
+	FPCGMetadataAttributeBase* ClearOrCreateAttribute(UPCGMetadata* Metadata, UPCGParamData* Params = nullptr) const;
+
+	/* Set an entry defined by EntryKey (or create one if it is invalid) in the Attribute depending on the selected type. Value can be overridden by params. */
+	PCGMetadataEntryKey SetAttribute(FPCGMetadataAttributeBase* Attribute, UPCGMetadata* Metadata, PCGMetadataEntryKey EntryKey, UPCGParamData* Params = nullptr) const;
+
+	/* Set attributes for each point in the Attribute depending on the selected type. Value can be overridden by params. */
+	void SetAttribute(FPCGMetadataAttributeBase* Attribute, UPCGMetadata* Metadata, TArray<FPCGPoint>& Points, UPCGParamData* Params = nullptr) const;
+
+	UPROPERTY(BlueprintReadWrite, EditAnywhere, Category = Settings, meta = (ValidEnumValues = "Double, Integer64, Vector2, Vector, Vector4, Quaternion, Transform, String, Boolean, Rotator, Name"))
+	EPCGMetadataTypes Type = EPCGMetadataTypes::Double;
+
+	UPROPERTY(BlueprintReadWrite, EditAnywhere, Category = Settings)
+	FName OutputAttributeName = NAME_None;
+
+	// All different types
+	UPROPERTY(BlueprintReadWrite, EditAnywhere, Category = Settings, meta = (EditCondition = "Type == EPCGMetadataTypes::Double", EditConditionHides))
+	double DoubleValue = 0.0;
+
+	UPROPERTY(BlueprintReadWrite, EditAnywhere, Category = Settings, meta = (EditCondition = "Type == EPCGMetadataTypes::Integer64", EditConditionHides))
+	int64 IntValue = 0;
+
+	UPROPERTY(BlueprintReadWrite, EditAnywhere, Category = Settings, meta = (EditCondition = "Type == EPCGMetadataTypes::Vector2", EditConditionHides))
+	FVector2D Vector2Value = FVector2D::ZeroVector;
+
+	UPROPERTY(BlueprintReadWrite, EditAnywhere, Category = Settings, meta = (EditCondition = "Type == EPCGMetadataTypes::Vector", EditConditionHides))
+	FVector VectorValue = FVector::ZeroVector;
+
+	UPROPERTY(BlueprintReadWrite, EditAnywhere, Category = Settings, meta = (EditCondition = "Type == EPCGMetadataTypes::Vector4", EditConditionHides))
+	FVector4 Vector4Value = FVector4::Zero();
+
+	UPROPERTY(BlueprintReadWrite, EditAnywhere, Category = Settings, meta = (EditCondition = "Type == EPCGMetadataTypes::Quaternion", EditConditionHides))
+	FQuat QuatValue = FQuat::Identity;
+
+	UPROPERTY(BlueprintReadWrite, EditAnywhere, Category = Settings, meta = (EditCondition = "Type == EPCGMetadataTypes::Transform", EditConditionHides))
+	FTransform TransformValue = FTransform::Identity;
+
+	UPROPERTY(BlueprintReadWrite, EditAnywhere, Category = Settings, meta = (EditCondition = "Type == EPCGMetadataTypes::String", EditConditionHides))
+	FString StringValue = "";
+
+	UPROPERTY(BlueprintReadWrite, EditAnywhere, Category = Settings, meta = (EditCondition = "Type == EPCGMetadataTypes::Boolean", EditConditionHides))
+	bool BoolValue = false;
+
+	UPROPERTY(BlueprintReadWrite, EditAnywhere, Category = Settings, meta = (EditCondition = "Type == EPCGMetadataTypes::Rotator", EditConditionHides))
+	FRotator RotatorValue = FRotator::ZeroRotator;
+
+	UPROPERTY(BlueprintReadWrite, EditAnywhere, Category = Settings, meta = (EditCondition = "Type == EPCGMetadataTypes::Name", EditConditionHides))
+	FName NameValue = NAME_None;
+};
\ No newline at end of file
diff --git a/Engine/Plugins/Experimental/PCG/Source/PCG/Public/Elements/PCGCreateParamData.h b/Engine/Plugins/Experimental/PCG/Source/PCG/Public/Elements/PCGCreateParamData.h
--- a/Engine/Plugins/Experimental/PCG/Source/PCG/Public/Elements/PCGCreateParamData.h
+++ b/Engine/Plugins/Experimental/PCG/Source/PCG/Public/Elements/PCGCreateParamData.h
@@ -3,13 +3,12 @@
 #pragma once
 
 #include "PCGElement.h"
-#include "PCGSettings.h"
-#include "Metadata/PCGMetadataAttributeTraits.h"
+#include "Elements/PCGCreateAttributeBase.h"
 
 #include "PCGCreateParamData.generated.h"
 
 UCLASS(BlueprintType, ClassGroup = (Procedural))
-class PCG_API UPCGCreateParamDataSettings : public UPCGSettings
+class PCG_API UPCGCreateParamDataSettings : public UPCGAttributeCreationBaseSettings
 {
 	GENERATED_BODY()
 
@@ -22,49 +21,8 @@ public:
 
 	virtual TArray<FPCGPinProperties> InputPinProperties() const override;
 	virtual TArray<FPCGPinProperties> OutputPinProperties() const override;
-	virtual FName AdditionalTaskName() const override;
 	//~End UPCGSettings interface
 
-	UPROPERTY(BlueprintReadWrite, EditAnywhere, Category = Settings, meta = (ValidEnumValues = "Double, Integer64, Vector2, Vector, Vector4, Quaternion, Transform, String, Boolean, Rotator, Name"))
-	EPCGMetadataTypes Type = EPCGMetadataTypes::Double;
-
-	UPROPERTY(BlueprintReadWrite, EditAnywhere, Category = Settings)
-	FName OutputAttributeName = NAME_None;
-
-	// All different types
-	UPROPERTY(BlueprintReadWrite, EditAnywhere, Category = Settings, meta = (EditCondition = "Type == EPCGMetadataTypes::Double", EditConditionHides))
-	double DoubleValue = 0.0;
-
-	UPROPERTY(BlueprintReadWrite, EditAnywhere, Category = Settings, meta = (EditCondition = "Type == EPCGMetadataTypes::Integer64", EditConditionHides))
-	int64 IntValue = 0;
-
-	UPROPERTY(BlueprintReadWrite, EditAnywhere, Category = Settings, meta = (EditCondition = "Type == EPCGMetadataTypes::Vector2", EditConditionHides))
-	FVector2D Vector2Value = FVector2D::ZeroVector;
-
-	UPROPERTY(BlueprintReadWrite, EditAnywhere, Category = Settings, meta = (EditCondition = "Type == EPCGMetadataTypes::Vector", EditConditionHides))
-	FVector VectorValue = FVector::ZeroVector;
-
-	UPROPERTY(BlueprintReadWrite, EditAnywhere, Category = Settings, meta = (EditCondition = "Type == EPCGMetadataTypes::Vector4", EditConditionHides))
-	FVector4 Vector4Value = FVector4::Zero();
-
-	UPROPERTY(BlueprintReadWrite, EditAnywhere, Category = Settings, meta = (EditCondition = "Type == EPCGMetadataTypes::Quaternion", EditConditionHides))
-	FQuat QuatValue = FQuat::Identity;
-
-	UPROPERTY(BlueprintReadWrite, EditAnywhere, Category = Settings, meta = (EditCondition = "Type == EPCGMetadataTypes::Transform", EditConditionHides))
-	FTransform TransformValue = FTransform::Identity;
-
-	UPROPERTY(BlueprintReadWrite, EditAnywhere, Category = Settings, meta = (EditCondition = "Type == EPCGMetadataTypes::String", EditConditionHides))
-	FString StringValue = "";
-
-	UPROPERTY(BlueprintReadWrite, EditAnywhere, Category = Settings, meta = (EditCondition = "Type == EPCGMetadataTypes::Boolean", EditConditionHides))
-	bool BoolValue = false;
-
-	UPROPERTY(BlueprintReadWrite, EditAnywhere, Category = Settings, meta = (EditCondition = "Type == EPCGMetadataTypes::Rotator", EditConditionHides))
-	FRotator RotatorValue = FRotator::ZeroRotator;
-
-	UPROPERTY(BlueprintReadWrite, EditAnywhere, Category = Settings, meta = (EditCondition = "Type == EPCGMetadataTypes::Name", EditConditionHides))
-	FName NameValue = NAME_None;
-
 protected:
 	virtual FPCGElementPtr CreateElement() const override;
 };
diff --git a/Engine/Plugins/Experimental/PCG/Source/PCG/Public/Helpers/PCGSettingsHelpers.h b/Engine/Plugins/Experimental/PCG/Source/PCG/Public/Helpers/PCGSettingsHelpers.h
--- a/Engine/Plugins/Experimental/PCG/Source/PCG/Public/Helpers/PCGSettingsHelpers.h
+++ b/Engine/Plugins/Experimental/PCG/Source/PCG/Public/Helpers/PCGSettingsHelpers.h
@@ -227,4 +227,4 @@ namespace PCGSettingsHelpers
 	}
 }
 
-#define PCG_GET_OVERRIDEN_VALUE(Settings, Variable, Params) PCGSettingsHelpers::GetValue(GET_MEMBER_NAME_CHECKED(TRemovePointer<TRemoveConst<decltype(Settings)>::Type>::Type, Variable), (Settings)->Variable, Params);
\ No newline at end of file
+#define PCG_GET_OVERRIDEN_VALUE(Settings, Variable, Params) PCGSettingsHelpers::GetValue(GET_MEMBER_NAME_CHECKED(TRemovePointer<TRemoveConst<decltype(Settings)>::Type>::Type, Variable), (Settings)->Variable, Params)
\ No newline at end of file
