[PCG] Naming pass on all metadata operation nodes

* Remove Params pin, as it was decided it is not useful
* Rename all nodes following this pattern: "Attribute {Function} ops"
* All nodes will now update their node name to match their operation, following this pattern : "{Function}: {Operation}"
* Use of templated function to check types, instead of casting manually.

diff --git a/Engine/Plugins/Experimental/PCG/Source/PCG/Private/Elements/Metadata/PCGMetadataBitwiseOpElement.cpp b/Engine/Plugins/Experimental/PCG/Source/PCG/Private/Elements/Metadata/PCGMetadataBitwiseOpElement.cpp
--- a/Engine/Plugins/Experimental/PCG/Source/PCG/Private/Elements/Metadata/PCGMetadataBitwiseOpElement.cpp
+++ b/Engine/Plugins/Experimental/PCG/Source/PCG/Private/Elements/Metadata/PCGMetadataBitwiseOpElement.cpp
@@ -15,7 +15,7 @@ namespace PCGMetadataBitwiseSettings
 {
 	inline int64 UnaryOp(const int64& Value)
 	{
-		// EPCGMedadataBitwiseOperation::BitwiseNot
+		// EPCGMedadataBitwiseOperation::Not
 		return ~Value;
 	}
 
@@ -23,11 +23,11 @@ namespace PCGMetadataBitwiseSettings
 	{
 		switch (Operation)
 		{
-		case EPCGMedadataBitwiseOperation::BitwiseAnd:
+		case EPCGMedadataBitwiseOperation::And:
 			return (Value1 & Value2);
-		case EPCGMedadataBitwiseOperation::BitwiseOr:
+		case EPCGMedadataBitwiseOperation::Or:
 			return (Value1 | Value2);
-		case EPCGMedadataBitwiseOperation::BitwiseXor:
+		case EPCGMedadataBitwiseOperation::Xor:
 			return (Value1 ^ Value2);
 		default:
 			return 0;
@@ -40,7 +40,7 @@ FName UPCGMetadataBitwiseSettings::GetInputPinLabel(uint32 Index) const
 	switch (Index)
 	{
 	case 0:
-		return (Operation == EPCGMedadataBitwiseOperation::BitwiseNot) ? PCGPinConstants::DefaultInputLabel : PCGMetadataSettingsBaseConstants::DoubleInputFirstLabel;
+		return (Operation == EPCGMedadataBitwiseOperation::Not) ? PCGPinConstants::DefaultInputLabel : PCGMetadataSettingsBaseConstants::DoubleInputFirstLabel;
 	case 1:
 		return PCGMetadataSettingsBaseConstants::DoubleInputSecondLabel;
 	default:
@@ -50,13 +50,13 @@ FName UPCGMetadataBitwiseSettings::GetInputPinLabel(uint32 Index) const
 
 uint32 UPCGMetadataBitwiseSettings::GetInputPinNum() const
 {
-	return (Operation == EPCGMedadataBitwiseOperation::BitwiseNot) ? 1 : 2;
+	return (Operation == EPCGMedadataBitwiseOperation::Not) ? 1 : 2;
 }
 
 bool UPCGMetadataBitwiseSettings::IsSupportedInputType(uint16 TypeId, uint32 InputIndex, bool& bHasSpecialRequirement) const
 {
 	bHasSpecialRequirement = false;
-	return (TypeId == (uint16)EPCGMetadataTypes::Integer32) || (TypeId == (uint16)EPCGMetadataTypes::Integer64);
+	return PCG::Private::IsOfTypes<int32, int64>(TypeId);
 }
 
 FName UPCGMetadataBitwiseSettings::GetInputAttributeNameWithOverride(uint32 Index, UPCGParamData* Params) const
@@ -72,15 +72,22 @@ FName UPCGMetadataBitwiseSettings::GetInputAttributeNameWithOverride(uint32 Inde
 	}
 }
 
-#if WITH_EDITOR
-FName UPCGMetadataBitwiseSettings::GetDefaultNodeName() const
+FName UPCGMetadataBitwiseSettings::AdditionalTaskName() const
 {
-	if (const UEnum* EnumPtr = FindObject<UEnum>(nullptr, TEXT("EPCGMedadataBitwiseOperation"), true))
+	if (const UEnum* EnumPtr = FindObject<UEnum>(nullptr, TEXT("/Script/PCG.EPCGMedadataBitwiseOperation"), true))
+	{
+		return FName(FString("Bitwise: ") + EnumPtr->GetNameStringByValue(static_cast<int>(Operation)));
+	}
+	else
 	{
-		return EnumPtr->GetNameByValue(static_cast<int>(Operation));
+		return NAME_None;
 	}
+}
 
-	return TEXT("Metadata Bitwise Node");
+#if WITH_EDITOR
+FName UPCGMetadataBitwiseSettings::GetDefaultNodeName() const
+{
+	return TEXT("Attribute Bitwise Op");
 }
 #endif // WITH_EDITOR
 
@@ -100,7 +107,7 @@ bool FPCGMetadataBitwiseElement::DoOperation(FOperationData& OperationData) cons
 
 	const UPCGMetadataBitwiseSettings* Settings = CastChecked<UPCGMetadataBitwiseSettings>(OperationData.Settings);
 
-	if (Settings->Operation == EPCGMedadataBitwiseOperation::BitwiseNot)
+	if (Settings->Operation == EPCGMedadataBitwiseOperation::Not)
 	{
 		DoUnaryOp<int64>(OperationData, [](const int64& Value) -> int64 { return PCGMetadataBitwiseSettings::UnaryOp(Value); });
 	}
diff --git a/Engine/Plugins/Experimental/PCG/Source/PCG/Private/Elements/Metadata/PCGMetadataBooleanOpElement.cpp b/Engine/Plugins/Experimental/PCG/Source/PCG/Private/Elements/Metadata/PCGMetadataBooleanOpElement.cpp
--- a/Engine/Plugins/Experimental/PCG/Source/PCG/Private/Elements/Metadata/PCGMetadataBooleanOpElement.cpp
+++ b/Engine/Plugins/Experimental/PCG/Source/PCG/Private/Elements/Metadata/PCGMetadataBooleanOpElement.cpp
@@ -56,7 +56,7 @@ uint32 UPCGMetadataBooleanSettings::GetInputPinNum() const
 bool UPCGMetadataBooleanSettings::IsSupportedInputType(uint16 TypeId, uint32 InputIndex, bool& bHasSpecialRequirement) const
 {
 	bHasSpecialRequirement = false;
-	return TypeId == (uint16)EPCGMetadataTypes::Boolean;
+	return PCG::Private::IsOfTypes<bool>(TypeId);
 }
 
 FName UPCGMetadataBooleanSettings::GetInputAttributeNameWithOverride(uint32 Index, UPCGParamData* Params) const
@@ -77,15 +77,22 @@ uint16 UPCGMetadataBooleanSettings::GetOutputType(uint16 InputTypeId) const
 	return (uint16)EPCGMetadataTypes::Boolean;
 }
 
-#if WITH_EDITOR
-FName UPCGMetadataBooleanSettings::GetDefaultNodeName() const
+FName UPCGMetadataBooleanSettings::AdditionalTaskName() const
 {
-	if (const UEnum* EnumPtr = FindObject<UEnum>(nullptr, TEXT("EPCGMedadataBooleanOperation"), true))
-	{ 
-		return EnumPtr->GetNameByValue(static_cast<int>(Operation)); 
+	if (const UEnum* EnumPtr = FindObject<UEnum>(nullptr, TEXT("/Script/PCG.EPCGMedadataBooleanOperation"), true))
+	{
+		return FName(FString("Boolean: ") + EnumPtr->GetNameStringByValue(static_cast<int>(Operation)));
 	}
+	else
+	{
+		return NAME_None;
+	}
+}
 
-	return TEXT("Metadata Boolean Node");
+#if WITH_EDITOR
+FName UPCGMetadataBooleanSettings::GetDefaultNodeName() const
+{
+	return TEXT("Attribute Boolean Op");
 }
 #endif // WITH_EDITOR
 
diff --git a/Engine/Plugins/Experimental/PCG/Source/PCG/Private/Elements/Metadata/PCGMetadataBreakTransform.cpp b/Engine/Plugins/Experimental/PCG/Source/PCG/Private/Elements/Metadata/PCGMetadataBreakTransform.cpp
--- a/Engine/Plugins/Experimental/PCG/Source/PCG/Private/Elements/Metadata/PCGMetadataBreakTransform.cpp
+++ b/Engine/Plugins/Experimental/PCG/Source/PCG/Private/Elements/Metadata/PCGMetadataBreakTransform.cpp
@@ -81,6 +81,13 @@ FName UPCGMetadataBreakTransformSettings::GetOutputAttributeName(FName BaseName,
 	}
 }
 
+#if WITH_EDITOR
+FName UPCGMetadataBreakTransformSettings::GetDefaultNodeName() const
+{
+	return TEXT("Break Transform Attribute");
+}
+#endif // WITH_EDITOR
+
 FPCGElementPtr UPCGMetadataBreakTransformSettings::CreateElement() const
 {
 	return MakeShared<FPCGMetadataBreakTransformElement>();
diff --git a/Engine/Plugins/Experimental/PCG/Source/PCG/Private/Elements/Metadata/PCGMetadataBreakVector.cpp b/Engine/Plugins/Experimental/PCG/Source/PCG/Private/Elements/Metadata/PCGMetadataBreakVector.cpp
--- a/Engine/Plugins/Experimental/PCG/Source/PCG/Private/Elements/Metadata/PCGMetadataBreakVector.cpp
+++ b/Engine/Plugins/Experimental/PCG/Source/PCG/Private/Elements/Metadata/PCGMetadataBreakVector.cpp
@@ -134,6 +134,13 @@ FName UPCGMetadataBreakVectorSettings::GetOutputAttributeName(FName BaseName, ui
 	}
 }
 
+#if WITH_EDITOR
+FName UPCGMetadataBreakVectorSettings::GetDefaultNodeName() const
+{
+	return TEXT("Break Vector Attribute");
+}
+#endif // WITH_EDITOR
+
 FPCGElementPtr UPCGMetadataBreakVectorSettings::CreateElement() const
 {
 	return MakeShared<FPCGMetadataBreakVectorElement>();
diff --git a/Engine/Plugins/Experimental/PCG/Source/PCG/Private/Elements/Metadata/PCGMetadataCompareOpElement.cpp b/Engine/Plugins/Experimental/PCG/Source/PCG/Private/Elements/Metadata/PCGMetadataCompareOpElement.cpp
--- a/Engine/Plugins/Experimental/PCG/Source/PCG/Private/Elements/Metadata/PCGMetadataCompareOpElement.cpp
+++ b/Engine/Plugins/Experimental/PCG/Source/PCG/Private/Elements/Metadata/PCGMetadataCompareOpElement.cpp
@@ -73,7 +73,7 @@ uint32 UPCGMetadataCompareSettings::GetInputPinNum() const
 bool UPCGMetadataCompareSettings::IsSupportedInputType(uint16 TypeId, uint32 InputIndex, bool& bHasSpecialRequirement) const
 {
 	bHasSpecialRequirement = false;
-	return TypeId <= (uint16)EPCGMetadataTypes::Integer64;
+	return PCG::Private::IsOfTypes<int32, int64, float, double>(TypeId);
 }
 
 FName UPCGMetadataCompareSettings::GetInputAttributeNameWithOverride(uint32 Index, UPCGParamData* Params) const
@@ -94,15 +94,22 @@ uint16 UPCGMetadataCompareSettings::GetOutputType(uint16 InputTypeId) const
 	return (uint16)EPCGMetadataTypes::Boolean;
 }
 
-#if WITH_EDITOR
-FName UPCGMetadataCompareSettings::GetDefaultNodeName() const
+FName UPCGMetadataCompareSettings::AdditionalTaskName() const
 {
-	if (const UEnum* EnumPtr = FindObject<UEnum>(nullptr, TEXT("EPCGMedadataCompareOperation"), true))
-	{ 
-		return EnumPtr->GetNameByValue(static_cast<int>(Operation)); 
+	if (const UEnum* EnumPtr = FindObject<UEnum>(nullptr, TEXT("/Script/PCG.EPCGMedadataCompareOperation"), true))
+	{
+		return FName(FString("Compare: ") + EnumPtr->GetNameStringByValue(static_cast<int>(Operation)));
+	}
+	else
+	{
+		return NAME_None;
 	}
+}
 
-	return TEXT("Metadata Compare Node");
+#if WITH_EDITOR
+FName UPCGMetadataCompareSettings::GetDefaultNodeName() const
+{
+	return TEXT("Attribute Compare Op");
 }
 #endif // WITH_EDITOR
 
@@ -122,7 +129,7 @@ bool FPCGMetadataCompareElement::DoOperation(FOperationData& OperationData) cons
 		using AttributeType = decltype(DummyValue);
 
 		// Need to remove types that would not compile
-		if constexpr (PCG::Private::MetadataTypes<AttributeType>::Id > (uint16)EPCGMetadataTypes::Integer64)
+		if constexpr (!PCG::Private::IsOfTypes<AttributeType, int32, int64, float, double>())
 		{
 			return;
 		}
diff --git a/Engine/Plugins/Experimental/PCG/Source/PCG/Private/Elements/Metadata/PCGMetadataMakeTransform.cpp b/Engine/Plugins/Experimental/PCG/Source/PCG/Private/Elements/Metadata/PCGMetadataMakeTransform.cpp
--- a/Engine/Plugins/Experimental/PCG/Source/PCG/Private/Elements/Metadata/PCGMetadataMakeTransform.cpp
+++ b/Engine/Plugins/Experimental/PCG/Source/PCG/Private/Elements/Metadata/PCGMetadataMakeTransform.cpp
@@ -81,6 +81,13 @@ uint16 UPCGMetadataMakeTransformSettings::GetOutputType(uint16 InputTypeId) cons
 	return (uint16)EPCGMetadataTypes::Transform;
 }
 
+#if WITH_EDITOR
+FName UPCGMetadataMakeTransformSettings::GetDefaultNodeName() const
+{
+	return TEXT("Make Transform Attribute");
+}
+#endif // WITH_EDITOR
+
 FPCGElementPtr UPCGMetadataMakeTransformSettings::CreateElement() const
 {
 	return MakeShared<FPCGMetadataMakeTransformElement>();
diff --git a/Engine/Plugins/Experimental/PCG/Source/PCG/Private/Elements/Metadata/PCGMetadataMakeVector.cpp b/Engine/Plugins/Experimental/PCG/Source/PCG/Private/Elements/Metadata/PCGMetadataMakeVector.cpp
--- a/Engine/Plugins/Experimental/PCG/Source/PCG/Private/Elements/Metadata/PCGMetadataMakeVector.cpp
+++ b/Engine/Plugins/Experimental/PCG/Source/PCG/Private/Elements/Metadata/PCGMetadataMakeVector.cpp
@@ -151,6 +151,13 @@ uint16 UPCGMetadataMakeVectorSettings::GetOutputType(uint16 InputTypeId) const
 	return (uint16)OutputType;
 }
 
+#if WITH_EDITOR
+FName UPCGMetadataMakeVectorSettings::GetDefaultNodeName() const
+{
+	return TEXT("Make Vector Attribute");
+}
+#endif // WITH_EDITOR
+
 FPCGElementPtr UPCGMetadataMakeVectorSettings::CreateElement() const
 {
 	return MakeShared<FPCGMetadataMakeVectorElement>();
diff --git a/Engine/Plugins/Experimental/PCG/Source/PCG/Private/Elements/Metadata/PCGMetadataMathsOpElement.cpp b/Engine/Plugins/Experimental/PCG/Source/PCG/Private/Elements/Metadata/PCGMetadataMathsOpElement.cpp
--- a/Engine/Plugins/Experimental/PCG/Source/PCG/Private/Elements/Metadata/PCGMetadataMathsOpElement.cpp
+++ b/Engine/Plugins/Experimental/PCG/Source/PCG/Private/Elements/Metadata/PCGMetadataMathsOpElement.cpp
@@ -188,7 +188,7 @@ uint32 UPCGMetadataMathsSettings::GetInputPinNum() const
 bool UPCGMetadataMathsSettings::IsSupportedInputType(uint16 TypeId, uint32 InputIndex, bool& bHasSpecialRequirement) const
 {
 	bHasSpecialRequirement = false;
-	return TypeId <= (uint16)EPCGMetadataTypes::Vector4;
+	return PCG::Private::IsOfTypes<float, double, int32, int64, FVector2D, FVector, FVector4>(TypeId);
 }
 
 FName UPCGMetadataMathsSettings::GetInputAttributeNameWithOverride(uint32 Index, UPCGParamData* Params) const
@@ -206,15 +206,22 @@ FName UPCGMetadataMathsSettings::GetInputAttributeNameWithOverride(uint32 Index,
 	}
 }
 
-#if WITH_EDITOR
-FName UPCGMetadataMathsSettings::GetDefaultNodeName() const
+FName UPCGMetadataMathsSettings::AdditionalTaskName() const
 {
-	if (const UEnum* EnumPtr = FindObject<UEnum>(nullptr, TEXT("EPCGMedadataMathsOperation"), true))
-	{ 
-		return EnumPtr->GetNameByValue(static_cast<int>(Operation)); 
+	if (const UEnum* EnumPtr = FindObject<UEnum>(nullptr, TEXT("/Script/PCG.EPCGMedadataMathsOperation"), true))
+	{
+		return FName(FString("Maths: ") + EnumPtr->GetNameStringByValue(static_cast<int>(Operation)));
+	}
+	else
+	{
+		return NAME_None;
 	}
+}
 
-	return TEXT("Metadata Maths Node");
+#if WITH_EDITOR
+FName UPCGMetadataMathsSettings::GetDefaultNodeName() const
+{
+	return TEXT("Attribute Maths Op");
 }
 #endif // WITH_EDITOR
 
@@ -234,7 +241,7 @@ bool FPCGMetadataMathsElement::DoOperation(FOperationData& OperationData) const
 		using AttributeType = decltype(DummyOutValue);
 
 		// Need to remove types that would not compile
-		if constexpr (PCG::Private::MetadataTypes<AttributeType>::Id > (uint16)EPCGMetadataTypes::Vector4)
+		if constexpr (!PCG::Private::IsOfTypes<AttributeType, float, double, int32, int64, FVector2D, FVector, FVector4>())
 		{
 			return;
 		}
diff --git a/Engine/Plugins/Experimental/PCG/Source/PCG/Private/Elements/Metadata/PCGMetadataOpElementBase.cpp b/Engine/Plugins/Experimental/PCG/Source/PCG/Private/Elements/Metadata/PCGMetadataOpElementBase.cpp
--- a/Engine/Plugins/Experimental/PCG/Source/PCG/Private/Elements/Metadata/PCGMetadataOpElementBase.cpp
+++ b/Engine/Plugins/Experimental/PCG/Source/PCG/Private/Elements/Metadata/PCGMetadataOpElementBase.cpp
@@ -25,8 +25,6 @@ TArray<FPCGPinProperties> UPCGMetadataSettingsBase::InputPinProperties() const
 		}
 	}
 
-	PinProperties.Emplace(PCGPinConstants::DefaultParamsLabel, EPCGDataType::Param, /*bAllowMultipleConnections=*/false);
-
 	return PinProperties;
 }
 
@@ -67,6 +65,7 @@ bool FPCGMetadataElementBase::ExecuteInternal(FPCGContext* Context) const
 	const TArray<FPCGTaggedData>& Inputs = Context->InputData.TaggedData;
 	TArray<FPCGTaggedData>& Outputs = Context->OutputData.TaggedData;
 
+	// We disabled param pin, but we let the logic of attribute name override, if we need to re-enable it in the future.
 	const TArray<FPCGTaggedData>& ParamData = Context->InputData.GetParamsByPin(PCGPinConstants::DefaultParamsLabel);
 	UPCGParamData* Params = ParamData.IsEmpty() ? nullptr : Cast<UPCGParamData>(ParamData[0].Data);
 
diff --git a/Engine/Plugins/Experimental/PCG/Source/PCG/Private/Elements/Metadata/PCGMetadataRotatorOpElement.cpp b/Engine/Plugins/Experimental/PCG/Source/PCG/Private/Elements/Metadata/PCGMetadataRotatorOpElement.cpp
--- a/Engine/Plugins/Experimental/PCG/Source/PCG/Private/Elements/Metadata/PCGMetadataRotatorOpElement.cpp
+++ b/Engine/Plugins/Experimental/PCG/Source/PCG/Private/Elements/Metadata/PCGMetadataRotatorOpElement.cpp
@@ -137,17 +137,17 @@ bool UPCGMetadataRotatorSettings::IsSupportedInputType(uint16 TypeId, uint32 Inp
 	if (InputIndex == 2)
 	{
 		bHasSpecialRequirement = true;
-		return TypeId == (uint16)EPCGMetadataTypes::Float || TypeId == (uint16)EPCGMetadataTypes::Double;
+		return PCG::Private::IsOfTypes<float, double>(TypeId);
 	}
 	else if (InputIndex == 1 && PCGMetadataRotatorSettings::IsTransfromOp(Operation))
 	{
 		bHasSpecialRequirement = true;
-		return TypeId == (uint16)EPCGMetadataTypes::Transform;
+		return PCG::Private::IsOfTypes<FTransform>(TypeId);
 	}
 	else
 	{
 		bHasSpecialRequirement = false;
-		return TypeId == (uint16)EPCGMetadataTypes::Rotator || TypeId == (uint16)EPCGMetadataTypes::Quaternion;
+		return PCG::Private::IsOfTypes<FRotator, FQuat>(TypeId);
 	}
 }
 
@@ -166,15 +166,22 @@ FName UPCGMetadataRotatorSettings::GetInputAttributeNameWithOverride(uint32 Inde
 	}
 }
 
-#if WITH_EDITOR
-FName UPCGMetadataRotatorSettings::GetDefaultNodeName() const
+FName UPCGMetadataRotatorSettings::AdditionalTaskName() const
 {
-	if (const UEnum* EnumPtr = FindObject<UEnum>(nullptr, TEXT("EPCGMedadataRotatorOperation"), true))
-	{ 
-		return EnumPtr->GetNameByValue(static_cast<int>(Operation)); 
+	if (const UEnum* EnumPtr = FindObject<UEnum>(nullptr, TEXT("/Script/PCG.EPCGMedadataRotatorOperation"), true))
+	{
+		return FName(FString("Rotator: ") + EnumPtr->GetNameStringByValue(static_cast<int>(Operation)));
 	}
+	else
+	{
+		return NAME_None;
+	}
+}
 
-	return TEXT("Metadata Rotator Node");
+#if WITH_EDITOR
+FName UPCGMetadataRotatorSettings::GetDefaultNodeName() const
+{
+	return TEXT("Attribute Rotator Op");
 }
 #endif // WITH_EDITOR
 
@@ -193,8 +200,7 @@ bool FPCGMetadataRotatorElement::DoOperation(FOperationData& OperationData) cons
 	{
 		using AttributeType = decltype(DummyValue);
 
-		if constexpr (PCG::Private::MetadataTypes<AttributeType>::Id != (uint16)EPCGMetadataTypes::Quaternion &&
-					PCG::Private::MetadataTypes<AttributeType>::Id != (uint16)EPCGMetadataTypes::Rotator)
+		if constexpr (!PCG::Private::IsOfTypes<AttributeType, FQuat, FRotator>())
 		{
 			return;
 		}
diff --git a/Engine/Plugins/Experimental/PCG/Source/PCG/Private/Elements/Metadata/PCGMetadataTransformOpElement.cpp b/Engine/Plugins/Experimental/PCG/Source/PCG/Private/Elements/Metadata/PCGMetadataTransformOpElement.cpp
--- a/Engine/Plugins/Experimental/PCG/Source/PCG/Private/Elements/Metadata/PCGMetadataTransformOpElement.cpp
+++ b/Engine/Plugins/Experimental/PCG/Source/PCG/Private/Elements/Metadata/PCGMetadataTransformOpElement.cpp
@@ -111,12 +111,12 @@ bool UPCGMetadataTransformSettings::IsSupportedInputType(uint16 TypeId, uint32 I
 	if (InputIndex == 2)
 	{
 		bHasSpecialRequirement = true;
-		return TypeId == (uint16)EPCGMetadataTypes::Float || TypeId == (uint16)EPCGMetadataTypes::Double;
+		return PCG::Private::IsOfTypes<float, double>(TypeId);
 	}
 	else
 	{
 		bHasSpecialRequirement = false;
-		return TypeId == (uint16)EPCGMetadataTypes::Transform;
+		return PCG::Private::IsOfTypes<FTransform>(TypeId);
 	}
 }
 
@@ -135,15 +135,22 @@ FName UPCGMetadataTransformSettings::GetInputAttributeNameWithOverride(uint32 In
 	}
 }
 
-#if WITH_EDITOR
-FName UPCGMetadataTransformSettings::GetDefaultNodeName() const
+FName UPCGMetadataTransformSettings::AdditionalTaskName() const
 {
-	if (const UEnum* EnumPtr = FindObject<UEnum>(nullptr, TEXT("EPCGMedadataTransformOperation"), true))
-	{ 
-		return EnumPtr->GetNameByValue(static_cast<int>(Operation)); 
+	if (const UEnum* EnumPtr = FindObject<UEnum>(nullptr, TEXT("/Script/PCG.EPCGMedadataTransformOperation"), true))
+	{
+		return FName(FString("Transform: ") + EnumPtr->GetNameStringByValue(static_cast<int>(Operation)));
 	}
+	else
+	{
+		return NAME_None;
+	}
+}
 
-	return TEXT("Metadata Transform Node");
+#if WITH_EDITOR
+FName UPCGMetadataTransformSettings::GetDefaultNodeName() const
+{
+	return TEXT("Attribute Transform Op");
 }
 #endif // WITH_EDITOR
 
diff --git a/Engine/Plugins/Experimental/PCG/Source/PCG/Private/Elements/Metadata/PCGMetadataTrigOpElement.cpp b/Engine/Plugins/Experimental/PCG/Source/PCG/Private/Elements/Metadata/PCGMetadataTrigOpElement.cpp
--- a/Engine/Plugins/Experimental/PCG/Source/PCG/Private/Elements/Metadata/PCGMetadataTrigOpElement.cpp
+++ b/Engine/Plugins/Experimental/PCG/Source/PCG/Private/Elements/Metadata/PCGMetadataTrigOpElement.cpp
@@ -66,7 +66,7 @@ uint32 UPCGMetadataTrigSettings::GetInputPinNum() const
 bool UPCGMetadataTrigSettings::IsSupportedInputType(uint16 TypeId, uint32 InputIndex, bool& bHasSpecialRequirement) const
 {
 	bHasSpecialRequirement = false;
-	return TypeId <= (uint16)EPCGMetadataTypes::Integer64;
+	return PCG::Private::IsOfTypes<int32, int64, float, double>(TypeId);
 }
 
 FName UPCGMetadataTrigSettings::GetInputAttributeNameWithOverride(uint32 Index, UPCGParamData* Params) const
@@ -84,25 +84,25 @@ FName UPCGMetadataTrigSettings::GetInputAttributeNameWithOverride(uint32 Index,
 
 uint16 UPCGMetadataTrigSettings::GetOutputType(uint16 InputTypeId) const
 {
-	if (InputTypeId == (uint16)EPCGMetadataTypes::Integer32 || InputTypeId == (uint16)EPCGMetadataTypes::Integer64)
+	return (uint16)EPCGMetadataTypes::Double;
+}
+
+FName UPCGMetadataTrigSettings::AdditionalTaskName() const
+{
+	if (const UEnum* EnumPtr = FindObject<UEnum>(nullptr, TEXT("/Script/PCG.EPCGMedadataTrigOperation"), true))
 	{
-		return (uint16)EPCGMetadataTypes::Double;
+		return FName(FString("Trig: ") + EnumPtr->GetNameStringByValue(static_cast<int>(Operation)));
 	}
 	else
 	{
-		return InputTypeId;
+		return NAME_None;
 	}
 }
 
 #if WITH_EDITOR
 FName UPCGMetadataTrigSettings::GetDefaultNodeName() const
 {
-	if (const UEnum* EnumPtr = FindObject<UEnum>(nullptr, TEXT("EPCGMedadataTrigOperation"), true))
-	{ 
-		return EnumPtr->GetNameByValue(static_cast<int>(Operation)); 
-	}
-
-	return TEXT("Metadata Trig Node");
+	return TEXT("Attribute Trig Op");
 }
 #endif // WITH_EDITOR
 
diff --git a/Engine/Plugins/Experimental/PCG/Source/PCG/Private/Elements/Metadata/PCGMetadataVectorOpElement.cpp b/Engine/Plugins/Experimental/PCG/Source/PCG/Private/Elements/Metadata/PCGMetadataVectorOpElement.cpp
--- a/Engine/Plugins/Experimental/PCG/Source/PCG/Private/Elements/Metadata/PCGMetadataVectorOpElement.cpp
+++ b/Engine/Plugins/Experimental/PCG/Source/PCG/Private/Elements/Metadata/PCGMetadataVectorOpElement.cpp
@@ -35,12 +35,6 @@ namespace PCGMetadataVectorSettings
 		return (uint16)Operation >= (uint16)EPCGMedadataVectorOperation::TransformOp;
 	}
 
-	template <typename InType>
-	inline constexpr bool IsInvalidInType()
-	{
-		return PCG::Private::MetadataTypes<InType>::Id < (uint16)EPCGMetadataTypes::Vector2 || PCG::Private::MetadataTypes<InType>::Id > (uint16)EPCGMetadataTypes::Vector4;
-	}
-
 	// Mimic KismetMathLibrary
 	template <typename InType>
 	inline InType ApplyTransformOperation(const InType& Value, const FTransform& Transform, EPCGMedadataVectorOperation Operation)
@@ -276,16 +270,16 @@ bool UPCGMetadataVectorSettings::IsSupportedInputType(uint16 TypeId, uint32 Inpu
 	if (Operation == EPCGMedadataVectorOperation::RotateAroundAxis && InputIndex == 2)
 	{
 		bHasSpecialRequirement = true;
-		return (TypeId == (uint16)EPCGMetadataTypes::Double) || (TypeId == (uint16)EPCGMetadataTypes::Float);
+		return PCG::Private::IsOfTypes<double, float>(TypeId);
 	}
 	else if (PCGMetadataVectorSettings::IsTransformOp(Operation) && InputIndex == 1)
 	{
 		bHasSpecialRequirement = true;
-		return TypeId == (uint16)EPCGMetadataTypes::Transform;
+		return PCG::Private::IsOfTypes<FTransform>(TypeId);
 	}
 	else
 	{
-		return (TypeId == (uint16)EPCGMetadataTypes::Vector2) || (TypeId == (uint16)EPCGMetadataTypes::Vector) || (TypeId == (uint16)EPCGMetadataTypes::Vector4);
+		return PCG::Private::IsOfTypes<FVector2D, FVector, FVector4>(TypeId);
 	}
 }
 
@@ -304,15 +298,22 @@ FName UPCGMetadataVectorSettings::GetInputAttributeNameWithOverride(uint32 Index
 	}
 }
 
-#if WITH_EDITOR
-FName UPCGMetadataVectorSettings::GetDefaultNodeName() const
+FName UPCGMetadataVectorSettings::AdditionalTaskName() const
 {
-	if (const UEnum* EnumPtr = FindObject<UEnum>(nullptr, TEXT("EPCGMedadataVectorOperation"), true))
+	if (const UEnum* EnumPtr = FindObject<UEnum>(nullptr, TEXT("/Script/PCG.EPCGMedadataVectorOperation"), true))
+	{
+		return FName(FString("Vector: ") + EnumPtr->GetNameStringByValue(static_cast<int>(Operation)));
+	}
+	else
 	{
-		return EnumPtr->GetNameByValue(static_cast<int>(Operation));
+		return NAME_None;
 	}
+}
 
-	return TEXT("Metadata Vector Node");
+#if WITH_EDITOR
+FName UPCGMetadataVectorSettings::GetDefaultNodeName() const
+{
+	return TEXT("Attribute Vector Op");
 }
 #endif // WITH_EDITOR
 
@@ -348,7 +349,7 @@ bool FPCGMetadataVectorElement::DoOperation(FOperationData& OperationData) const
 	{
 		using AttributeType = decltype(DummyValue);
 
-		if constexpr (PCGMetadataVectorSettings::IsInvalidInType<AttributeType>())
+		if constexpr (!PCG::Private::IsOfTypes<AttributeType, FVector2D, FVector, FVector4>())
 		{
 			return;
 		}
@@ -358,7 +359,7 @@ bool FPCGMetadataVectorElement::DoOperation(FOperationData& OperationData) const
 			{
 				DoBinaryOp<AttributeType, FTransform>(OperationData, [Operation](const AttributeType& Value, const FTransform& Transform) -> AttributeType { return PCGMetadataVectorSettings::ApplyTransformOperation(Value, Transform, Operation);});
 			}
-			else if (OperationData.OutputType == (uint16)EPCGMetadataTypes::Double)
+			else if (PCG::Private::IsOfTypes<double>(OperationData.OutputType))
 			{
 				if (PCGMetadataVectorSettings::IsUnaryOp(Operation))
 				{
diff --git a/Engine/Plugins/Experimental/PCG/Source/PCG/Public/Elements/Metadata/PCGMetadataBitwiseOpElement.h b/Engine/Plugins/Experimental/PCG/Source/PCG/Public/Elements/Metadata/PCGMetadataBitwiseOpElement.h
--- a/Engine/Plugins/Experimental/PCG/Source/PCG/Public/Elements/Metadata/PCGMetadataBitwiseOpElement.h
+++ b/Engine/Plugins/Experimental/PCG/Source/PCG/Public/Elements/Metadata/PCGMetadataBitwiseOpElement.h
@@ -11,10 +11,10 @@
 UENUM()
 enum class EPCGMedadataBitwiseOperation : uint16
 {
-	BitwiseAnd,
-	BitwiseNot,
-	BitwiseOr,
-	BitwiseXor
+	And,
+	Not,
+	Or,
+	Xor
 };
 
 UCLASS(BlueprintType, ClassGroup = (Procedural))
@@ -27,6 +27,7 @@ public:
 #if WITH_EDITOR
 	virtual FName GetDefaultNodeName() const override;
 #endif
+	virtual FName AdditionalTaskName() const override;
 	//~End UPCGSettings interface
 
 	//~Begin UPCGMetadataSettingsBase interface
@@ -46,12 +47,12 @@ protected:
 
 public:
 	UPROPERTY(BlueprintReadWrite, EditAnywhere, Category = Settings)
-	EPCGMedadataBitwiseOperation Operation = EPCGMedadataBitwiseOperation::BitwiseAnd;
+	EPCGMedadataBitwiseOperation Operation = EPCGMedadataBitwiseOperation::And;
 
 	UPROPERTY(BlueprintReadWrite, EditAnywhere, Category = Input)
 	FName Input1AttributeName = NAME_None;
 
-	UPROPERTY(BlueprintReadWrite, EditAnywhere, Category = Input, meta = (EditCondition = "Operation != EPCGMedadataBitwiseOperation::BitwiseNot", EditConditionHides))
+	UPROPERTY(BlueprintReadWrite, EditAnywhere, Category = Input, meta = (EditCondition = "Operation != EPCGMedadataBitwiseOperation::Not", EditConditionHides))
 	FName Input2AttributeName = NAME_None;
 };
 
diff --git a/Engine/Plugins/Experimental/PCG/Source/PCG/Public/Elements/Metadata/PCGMetadataBooleanOpElement.h b/Engine/Plugins/Experimental/PCG/Source/PCG/Public/Elements/Metadata/PCGMetadataBooleanOpElement.h
--- a/Engine/Plugins/Experimental/PCG/Source/PCG/Public/Elements/Metadata/PCGMetadataBooleanOpElement.h
+++ b/Engine/Plugins/Experimental/PCG/Source/PCG/Public/Elements/Metadata/PCGMetadataBooleanOpElement.h
@@ -27,6 +27,7 @@ public:
 #if WITH_EDITOR
 	virtual FName GetDefaultNodeName() const override;
 #endif
+	virtual FName AdditionalTaskName() const override;
 	//~End UPCGSettings interface
 
 	//~Begin UPCGMetadataSettingsBase interface
diff --git a/Engine/Plugins/Experimental/PCG/Source/PCG/Public/Elements/Metadata/PCGMetadataBreakTransform.h b/Engine/Plugins/Experimental/PCG/Source/PCG/Public/Elements/Metadata/PCGMetadataBreakTransform.h
--- a/Engine/Plugins/Experimental/PCG/Source/PCG/Public/Elements/Metadata/PCGMetadataBreakTransform.h
+++ b/Engine/Plugins/Experimental/PCG/Source/PCG/Public/Elements/Metadata/PCGMetadataBreakTransform.h
@@ -18,7 +18,7 @@ class PCG_API UPCGMetadataBreakTransformSettings : public UPCGMetadataSettingsBa
 public:
 	//~Begin UPCGSettings interface
 #if WITH_EDITOR
-	virtual FName GetDefaultNodeName() const override { return FName(TEXT("MetadataBreakTransformNode")); }
+	virtual FName GetDefaultNodeName() const override;
 #endif
 
 	virtual FName GetInputAttributeNameWithOverride(uint32 Index, UPCGParamData* Params) const override;
diff --git a/Engine/Plugins/Experimental/PCG/Source/PCG/Public/Elements/Metadata/PCGMetadataBreakVector.h b/Engine/Plugins/Experimental/PCG/Source/PCG/Public/Elements/Metadata/PCGMetadataBreakVector.h
--- a/Engine/Plugins/Experimental/PCG/Source/PCG/Public/Elements/Metadata/PCGMetadataBreakVector.h
+++ b/Engine/Plugins/Experimental/PCG/Source/PCG/Public/Elements/Metadata/PCGMetadataBreakVector.h
@@ -25,7 +25,7 @@ class PCG_API UPCGMetadataBreakVectorSettings : public UPCGMetadataSettingsBase
 public:
 	//~Begin UPCGSettings interface
 #if WITH_EDITOR
-	virtual FName GetDefaultNodeName() const override { return FName(TEXT("MetadataBreakVectorNode")); }
+	virtual FName GetDefaultNodeName() const override;
 #endif
 
 	virtual FName GetInputAttributeNameWithOverride(uint32 Index, UPCGParamData* Params) const override;
diff --git a/Engine/Plugins/Experimental/PCG/Source/PCG/Public/Elements/Metadata/PCGMetadataCompareOpElement.h b/Engine/Plugins/Experimental/PCG/Source/PCG/Public/Elements/Metadata/PCGMetadataCompareOpElement.h
--- a/Engine/Plugins/Experimental/PCG/Source/PCG/Public/Elements/Metadata/PCGMetadataCompareOpElement.h
+++ b/Engine/Plugins/Experimental/PCG/Source/PCG/Public/Elements/Metadata/PCGMetadataCompareOpElement.h
@@ -29,6 +29,7 @@ public:
 #if WITH_EDITOR
 	virtual FName GetDefaultNodeName() const override;
 #endif
+	virtual FName AdditionalTaskName() const override;
 	//~End UPCGSettings interface
 
 	//~Begin UPCGMetadataSettingsBase interface
diff --git a/Engine/Plugins/Experimental/PCG/Source/PCG/Public/Elements/Metadata/PCGMetadataMakeTransform.h b/Engine/Plugins/Experimental/PCG/Source/PCG/Public/Elements/Metadata/PCGMetadataMakeTransform.h
--- a/Engine/Plugins/Experimental/PCG/Source/PCG/Public/Elements/Metadata/PCGMetadataMakeTransform.h
+++ b/Engine/Plugins/Experimental/PCG/Source/PCG/Public/Elements/Metadata/PCGMetadataMakeTransform.h
@@ -24,7 +24,7 @@ class PCG_API UPCGMetadataMakeTransformSettings : public UPCGMetadataSettingsBas
 public:
 #if WITH_EDITOR
 	//~Begin UPCGSettings interface
-	virtual FName GetDefaultNodeName() const override { return FName(TEXT("MetadataMakeTransformNode")); }
+	virtual FName GetDefaultNodeName() const override;
 	//~End UPCGSettings interface
 #endif
 
diff --git a/Engine/Plugins/Experimental/PCG/Source/PCG/Public/Elements/Metadata/PCGMetadataMakeVector.h b/Engine/Plugins/Experimental/PCG/Source/PCG/Public/Elements/Metadata/PCGMetadataMakeVector.h
--- a/Engine/Plugins/Experimental/PCG/Source/PCG/Public/Elements/Metadata/PCGMetadataMakeVector.h
+++ b/Engine/Plugins/Experimental/PCG/Source/PCG/Public/Elements/Metadata/PCGMetadataMakeVector.h
@@ -46,7 +46,7 @@ class PCG_API UPCGMetadataMakeVectorSettings : public UPCGMetadataSettingsBase
 public:
 #if WITH_EDITOR
 	//~Begin UPCGSettings interface
-	virtual FName GetDefaultNodeName() const override { return FName(TEXT("MetadataMakeVectorNode")); }
+	virtual FName GetDefaultNodeName() const override;
 	//~End UPCGSettings interface
 #endif
 
diff --git a/Engine/Plugins/Experimental/PCG/Source/PCG/Public/Elements/Metadata/PCGMetadataMathsOpElement.h b/Engine/Plugins/Experimental/PCG/Source/PCG/Public/Elements/Metadata/PCGMetadataMathsOpElement.h
--- a/Engine/Plugins/Experimental/PCG/Source/PCG/Public/Elements/Metadata/PCGMetadataMathsOpElement.h
+++ b/Engine/Plugins/Experimental/PCG/Source/PCG/Public/Elements/Metadata/PCGMetadataMathsOpElement.h
@@ -49,6 +49,7 @@ public:
 #if WITH_EDITOR
 	virtual FName GetDefaultNodeName() const override;
 #endif
+	virtual FName AdditionalTaskName() const override;
 	//~End UPCGSettings interface
 
 	//~Begin UPCGMetadataSettingsBase interface
diff --git a/Engine/Plugins/Experimental/PCG/Source/PCG/Public/Elements/Metadata/PCGMetadataRotatorOpElement.h b/Engine/Plugins/Experimental/PCG/Source/PCG/Public/Elements/Metadata/PCGMetadataRotatorOpElement.h
--- a/Engine/Plugins/Experimental/PCG/Source/PCG/Public/Elements/Metadata/PCGMetadataRotatorOpElement.h
+++ b/Engine/Plugins/Experimental/PCG/Source/PCG/Public/Elements/Metadata/PCGMetadataRotatorOpElement.h
@@ -37,6 +37,7 @@ public:
 #if WITH_EDITOR
 	virtual FName GetDefaultNodeName() const override;
 #endif
+	virtual FName AdditionalTaskName() const override;
 	//~End UPCGSettings interface
 
 	//~Begin UPCGMetadataSettingsBase interface
diff --git a/Engine/Plugins/Experimental/PCG/Source/PCG/Public/Elements/Metadata/PCGMetadataTransformOpElement.h b/Engine/Plugins/Experimental/PCG/Source/PCG/Public/Elements/Metadata/PCGMetadataTransformOpElement.h
--- a/Engine/Plugins/Experimental/PCG/Source/PCG/Public/Elements/Metadata/PCGMetadataTransformOpElement.h
+++ b/Engine/Plugins/Experimental/PCG/Source/PCG/Public/Elements/Metadata/PCGMetadataTransformOpElement.h
@@ -40,6 +40,7 @@ public:
 #if WITH_EDITOR
 	virtual FName GetDefaultNodeName() const override;
 #endif
+	virtual FName AdditionalTaskName() const override;
 	//~End UPCGSettings interface
 
 	//~Begin UPCGMetadataSettingsBase interface
diff --git a/Engine/Plugins/Experimental/PCG/Source/PCG/Public/Elements/Metadata/PCGMetadataTrigOpElement.h b/Engine/Plugins/Experimental/PCG/Source/PCG/Public/Elements/Metadata/PCGMetadataTrigOpElement.h
--- a/Engine/Plugins/Experimental/PCG/Source/PCG/Public/Elements/Metadata/PCGMetadataTrigOpElement.h
+++ b/Engine/Plugins/Experimental/PCG/Source/PCG/Public/Elements/Metadata/PCGMetadataTrigOpElement.h
@@ -32,6 +32,7 @@ public:
 #if WITH_EDITOR
 	virtual FName GetDefaultNodeName() const override;
 #endif
+	virtual FName AdditionalTaskName() const override;
 	//~End UPCGSettings interface
 
 	//~Begin UPCGMetadataSettingsBase interface
diff --git a/Engine/Plugins/Experimental/PCG/Source/PCG/Public/Elements/Metadata/PCGMetadataVectorOpElement.h b/Engine/Plugins/Experimental/PCG/Source/PCG/Public/Elements/Metadata/PCGMetadataVectorOpElement.h
--- a/Engine/Plugins/Experimental/PCG/Source/PCG/Public/Elements/Metadata/PCGMetadataVectorOpElement.h
+++ b/Engine/Plugins/Experimental/PCG/Source/PCG/Public/Elements/Metadata/PCGMetadataVectorOpElement.h
@@ -36,6 +36,7 @@ public:
 #if WITH_EDITOR
 	virtual FName GetDefaultNodeName() const override;
 #endif
+	virtual FName AdditionalTaskName() const override;
 	//~End UPCGSettings interface
 
 	//~Begin UPCGMetadataSettingsBase interface
