[PCG] Fixed a lockup were the attribute list view would keep old data around after refresh.

diff --git a/Engine/Plugins/Experimental/PCG/Source/PCGEditor/Private/SPCGEditorGraphAttributeListView.cpp b/Engine/Plugins/Experimental/PCG/Source/PCGEditor/Private/SPCGEditorGraphAttributeListView.cpp
--- a/Engine/Plugins/Experimental/PCG/Source/PCGEditor/Private/SPCGEditorGraphAttributeListView.cpp
+++ b/Engine/Plugins/Experimental/PCG/Source/PCGEditor/Private/SPCGEditorGraphAttributeListView.cpp
@@ -496,6 +496,10 @@ void SPCGEditorGraphAttributeListView::OnGenerateUpdated(UPCGComponent* /*InPCGC
 
 void SPCGEditorGraphAttributeListView::RefreshAttributeList()
 {
+	// Swapping to an empty item list to force a widget clear, otherwise the widgets will try to update during add column and access invalid data
+	static const TArray<PCGListviewItemPtr> EmptyList;
+	ListView->SetListItemsSource(EmptyList);
+	
 	ListViewItems.Empty();
 	ListViewHeader->ClearColumns();
 	MetadataInfos.Empty();
@@ -579,7 +583,7 @@ void SPCGEditorGraphAttributeListView::RefreshAttributeList()
 		}
 	}
 
-	ListView->RequestListRefresh();
+	ListView->SetListItemsSource(ListViewItems);
 }
 
 void SPCGEditorGraphAttributeListView::RefreshDataComboBox()
