[PCG] Fix Unittest assert in StaticSpawnerTest

#rnx

diff --git a/Engine/Plugins/Experimental/PCG/Source/PCG/Private/Tests/Elements/PCGStaticMeshSpawnerTest.cpp b/Engine/Plugins/Experimental/PCG/Source/PCG/Private/Tests/Elements/PCGStaticMeshSpawnerTest.cpp
--- a/Engine/Plugins/Experimental/PCG/Source/PCG/Private/Tests/Elements/PCGStaticMeshSpawnerTest.cpp
+++ b/Engine/Plugins/Experimental/PCG/Source/PCG/Private/Tests/Elements/PCGStaticMeshSpawnerTest.cpp
@@ -151,6 +151,7 @@ void TestMeshSelectorByAttribute(
 	// initialize and execute the StaticMeshSpawner
 	FPCGElementPtr StaticMeshSpawner = Settings->GetElement();
 	FPCGContext* Context = StaticMeshSpawner->Initialize(TestData.InputData, TestData.TestPCGComponent, nullptr);
+	Context->NumAvailableTasks = 1;
 
 	while (!StaticMeshSpawner->Execute(Context))
 	{}
@@ -279,6 +280,7 @@ void TestMeshSelectorWeighted(
 	// initialize and execute the StaticMeshSpawner
 	FPCGElementPtr StaticMeshSpawner = Settings->GetElement();
 	FPCGContext* Context = StaticMeshSpawner->Initialize(TestData.InputData, TestData.TestPCGComponent, nullptr);
+	Context->NumAvailableTasks = 1;
 
 	while (!StaticMeshSpawner->Execute(Context))
 	{}
@@ -371,6 +373,7 @@ void TestMeshSelectorWeightedByCategory(
 	// initialize and execute the StaticMeshSpawner
 	FPCGElementPtr StaticMeshSpawner = Settings->GetElement();
 	FPCGContext* Context = StaticMeshSpawner->Initialize(TestData.InputData, TestData.TestPCGComponent, nullptr);
+	Context->NumAvailableTasks = 1;
 
 	while (!StaticMeshSpawner->Execute(Context))
 	{
