[PCG] Re-factor Metadata element to allow property<->attribute to work even if the types doesn't match exactly.

*  Follow the rule to allow T -> U if we can construct U from T.

diff --git a/Engine/Plugins/Experimental/PCG/Source/PCG/Private/Elements/Metadata/PCGMetadataElement.cpp b/Engine/Plugins/Experimental/PCG/Source/PCG/Private/Elements/Metadata/PCGMetadataElement.cpp
--- a/Engine/Plugins/Experimental/PCG/Source/PCG/Private/Elements/Metadata/PCGMetadataElement.cpp
+++ b/Engine/Plugins/Experimental/PCG/Source/PCG/Private/Elements/Metadata/PCGMetadataElement.cpp
@@ -9,43 +9,81 @@
 
 namespace PCGMetadataOperations
 {
-	template<typename U, typename T>
-	bool SetValue(TArray<FPCGPoint>& InPoints, FPCGMetadataAttributeBase* AttributeBase, UPCGMetadata* Metadata, TFunctionRef<U(const FPCGPoint& InPoint)> PropGetter)
+	// Getter and setter for properties.
+	// Need to specify the property type and the getter/setter.
+	// Will work if AttributeType and PropertyType are the same, or if we can construct one from the other.
+	template <typename T>
+	bool SetValue(TArray<FPCGPoint>& InPoints, FPCGMetadataAttributeBase* AttributeBase, UPCGMetadata* Metadata, TFunctionRef<T(const FPCGPoint& InPoint)> PropGetter)
 	{
-		if (AttributeBase && AttributeBase->GetTypeId() == PCG::Private::MetadataTypes<T>::Id)
+		if (!AttributeBase)
 		{
-			FPCGMetadataAttribute<T>* Attribute = static_cast<FPCGMetadataAttribute<T>*>(AttributeBase);
+			return false;
+		}
+
+		auto Func = [&InPoints, AttributeBase, Metadata, &PropGetter](auto DummyValue) -> bool
+		{
+			using AttributeType = decltype(DummyValue);
+
+			FPCGMetadataAttribute<AttributeType>* Attribute = static_cast<FPCGMetadataAttribute<AttributeType>*>(AttributeBase);
+
 			for (FPCGPoint& Point : InPoints)
 			{
-				Metadata->InitializeOnSet(Point.MetadataEntry);
-				Attribute->SetValue(Point.MetadataEntry, PropGetter(Point));
+				if constexpr (std::is_same_v<AttributeType, T>)
+				{
+					Metadata->InitializeOnSet(Point.MetadataEntry);
+					Attribute->SetValue(Point.MetadataEntry, PropGetter(Point));
+				}
+				else if constexpr (std::is_constructible_v<AttributeType, T>)
+				{
+					Metadata->InitializeOnSet(Point.MetadataEntry);
+					Attribute->SetValue(Point.MetadataEntry, AttributeType(PropGetter(Point)));
+				}
+				else
+				{
+					return false;
+				}
 			}
 
 			return true;
-		}
-		else
+		};
+
+		return PCGMetadataAttribute::CallbackWithRightType(AttributeBase->GetTypeId(), Func);
+	}
+
+	template <typename T>
+	bool SetValue(const FPCGMetadataAttributeBase* AttributeBase, TArray<FPCGPoint>& InPoints, TFunctionRef<void(FPCGPoint& OutPoint, const T& InValue)> PropSetter)
+	{
+		if (!AttributeBase)
 		{
 			return false;
 		}
-	}
 
-	template<typename T, typename U>
-	bool SetValue(const FPCGMetadataAttributeBase* AttributeBase, TArray<FPCGPoint>& InPoints, TFunctionRef<void(FPCGPoint& OutPoint, const U& InValue)> PropSetter)
-	{
-		if (AttributeBase && AttributeBase->GetTypeId() == PCG::Private::MetadataTypes<T>::Id)
+		auto Func = [&InPoints, AttributeBase, &PropSetter](auto DummyValue) -> bool
 		{
-			const FPCGMetadataAttribute<T>* Attribute = static_cast<const FPCGMetadataAttribute<T>*>(AttributeBase);
+			using AttributeType = decltype(DummyValue);
+
+			const FPCGMetadataAttribute<AttributeType>* Attribute = static_cast<const FPCGMetadataAttribute<AttributeType>*>(AttributeBase);
+
 			for (FPCGPoint& Point : InPoints)
 			{
-				PropSetter(Point, U(Attribute->GetValueFromItemKey(Point.MetadataEntry)));
+				if constexpr (std::is_same_v<T, AttributeType>)
+				{
+					PropSetter(Point, Attribute->GetValueFromItemKey(Point.MetadataEntry));
+				}
+				else if constexpr (std::is_constructible_v<T, AttributeType>)
+				{
+					PropSetter(Point, T(Attribute->GetValueFromItemKey(Point.MetadataEntry)));
+				}
+				else
+				{
+					return false;
+				}
 			}
 
 			return true;
-		}
-		else
-		{
-			return false;
-		}
+		};
+
+		return PCGMetadataAttribute::CallbackWithRightType(AttributeBase->GetTypeId(), Func);
 	}
 }
 
@@ -133,9 +171,7 @@ bool FPCGMetadataOperationElement::ExecuteInternal(FPCGContext* Context) const
 				}
 
 				FPCGMetadataAttributeBase* AttributeBase = SampledData->Metadata->GetMutableAttribute(DestinationAttribute);
-				if(!PCGMetadataOperations::SetValue<float, float>(SampledPoints, AttributeBase, SampledData->Metadata, DensityGetter) &&
-					!PCGMetadataOperations::SetValue<float, FVector>(SampledPoints, AttributeBase, SampledData->Metadata, DensityGetter) &&
-					!PCGMetadataOperations::SetValue<float, FVector4>(SampledPoints, AttributeBase, SampledData->Metadata, DensityGetter))
+				if(!PCGMetadataOperations::SetValue<float>(SampledPoints, AttributeBase, SampledData->Metadata, DensityGetter))
 				{
 					PCGE_LOG(Error, "Attribute %s already exists but its type is not compatible", *DestinationAttribute.ToString());
 				}
@@ -150,8 +186,7 @@ bool FPCGMetadataOperationElement::ExecuteInternal(FPCGContext* Context) const
 				}
 
 				FPCGMetadataAttributeBase* AttributeBase = SampledData->Metadata->GetMutableAttribute(DestinationAttribute);
-				if (!PCGMetadataOperations::SetValue<FVector, FVector>(SampledPoints, AttributeBase, SampledData->Metadata, ExtentsGetter) &&
-					!PCGMetadataOperations::SetValue<FVector, FTransform>(SampledPoints, AttributeBase, SampledData->Metadata, ExtentsGetter))
+				if (!PCGMetadataOperations::SetValue<FVector>(SampledPoints, AttributeBase, SampledData->Metadata, ExtentsGetter))
 				{
 					PCGE_LOG(Error, "Attribute %s already exists but its type is not compatible", *DestinationAttribute.ToString());
 				}
@@ -166,8 +201,7 @@ bool FPCGMetadataOperationElement::ExecuteInternal(FPCGContext* Context) const
 				}
 
 				FPCGMetadataAttributeBase* AttributeBase = SampledData->Metadata->GetMutableAttribute(DestinationAttribute);
-				if (!PCGMetadataOperations::SetValue<FVector, FVector>(SampledPoints, AttributeBase, SampledData->Metadata, ExtentsGetter) &&
-					!PCGMetadataOperations::SetValue<FVector, FTransform>(SampledPoints, AttributeBase, SampledData->Metadata, ExtentsGetter))
+				if (!PCGMetadataOperations::SetValue<FVector>(SampledPoints, AttributeBase, SampledData->Metadata, ExtentsGetter))
 				{
 					PCGE_LOG(Error, "Attribute %s already exists but its type is not compatible", *DestinationAttribute.ToString());
 				}
@@ -182,8 +216,7 @@ bool FPCGMetadataOperationElement::ExecuteInternal(FPCGContext* Context) const
 				}
 
 				FPCGMetadataAttributeBase* AttributeBase = SampledData->Metadata->GetMutableAttribute(DestinationAttribute);
-				if(!PCGMetadataOperations::SetValue<FVector, FVector>(SampledPoints, AttributeBase, SampledData->Metadata, ExtentsGetter) &&
-					!PCGMetadataOperations::SetValue<FVector, FTransform>(SampledPoints, AttributeBase, SampledData->Metadata, ExtentsGetter))
+				if(!PCGMetadataOperations::SetValue<FVector>(SampledPoints, AttributeBase, SampledData->Metadata, ExtentsGetter))
 				{
 					PCGE_LOG(Error, "Attribute %s already exists but its type is not compatible", *DestinationAttribute.ToString());
 				}
@@ -198,7 +231,7 @@ bool FPCGMetadataOperationElement::ExecuteInternal(FPCGContext* Context) const
 				}
 
 				FPCGMetadataAttributeBase* AttributeBase = SampledData->Metadata->GetMutableAttribute(DestinationAttribute);
-				if(!PCGMetadataOperations::SetValue<FVector4, FVector4>(SampledPoints, AttributeBase, SampledData->Metadata, ColorGetter))
+				if(!PCGMetadataOperations::SetValue<FVector4>(SampledPoints, AttributeBase, SampledData->Metadata, ColorGetter))
 				{
 					PCGE_LOG(Error, "Attribute %s already exists but its type is not compatible", *DestinationAttribute.ToString());
 				}
@@ -213,8 +246,7 @@ bool FPCGMetadataOperationElement::ExecuteInternal(FPCGContext* Context) const
 				}
 
 				FPCGMetadataAttributeBase* AttributeBase = SampledData->Metadata->GetMutableAttribute(DestinationAttribute);
-				if (!PCGMetadataOperations::SetValue<FVector, FVector>(SampledPoints, AttributeBase, SampledData->Metadata, PositionGetter) &&
-					!PCGMetadataOperations::SetValue<FVector, FTransform>(SampledPoints, AttributeBase, SampledData->Metadata, PositionGetter))
+				if (!PCGMetadataOperations::SetValue<FVector>(SampledPoints, AttributeBase, SampledData->Metadata, PositionGetter))
 				{
 					PCGE_LOG(Error, "Attribute %s already exists but its type is not compatible", *DestinationAttribute.ToString());
 				}
@@ -229,7 +261,7 @@ bool FPCGMetadataOperationElement::ExecuteInternal(FPCGContext* Context) const
 				}
 
 				FPCGMetadataAttributeBase* AttributeBase = SampledData->Metadata->GetMutableAttribute(DestinationAttribute);
-				if (!PCGMetadataOperations::SetValue<FQuat, FQuat>(SampledPoints, AttributeBase, SampledData->Metadata, RotationGetter))
+				if (!PCGMetadataOperations::SetValue<FQuat>(SampledPoints, AttributeBase, SampledData->Metadata, RotationGetter))
 				{
 					PCGE_LOG(Error, "Attribute %s already exists but its type is not compatible", *DestinationAttribute.ToString());
 				}
@@ -244,7 +276,7 @@ bool FPCGMetadataOperationElement::ExecuteInternal(FPCGContext* Context) const
 				}
 
 				FPCGMetadataAttributeBase* AttributeBase = SampledData->Metadata->GetMutableAttribute(DestinationAttribute);
-				if (!PCGMetadataOperations::SetValue<FVector, FVector>(SampledPoints, AttributeBase, SampledData->Metadata, ScaleGetter))
+				if (!PCGMetadataOperations::SetValue<FVector>(SampledPoints, AttributeBase, SampledData->Metadata, ScaleGetter))
 				{
 					PCGE_LOG(Error, "Attribute %s already exists but its type is not compatible", *DestinationAttribute.ToString());
 				}
@@ -259,7 +291,7 @@ bool FPCGMetadataOperationElement::ExecuteInternal(FPCGContext* Context) const
 				}
 
 				FPCGMetadataAttributeBase* AttributeBase = SampledData->Metadata->GetMutableAttribute(DestinationAttribute);
-				if (!PCGMetadataOperations::SetValue<FTransform, FTransform>(SampledPoints, AttributeBase, SampledData->Metadata, TransformGetter))
+				if (!PCGMetadataOperations::SetValue<FTransform>(SampledPoints, AttributeBase, SampledData->Metadata, TransformGetter))
 				{
 					PCGE_LOG(Error, "Attribute %s already exists but its type is not compatible", *DestinationAttribute.ToString());
 				}
@@ -274,7 +306,7 @@ bool FPCGMetadataOperationElement::ExecuteInternal(FPCGContext* Context) const
 				}
 
 				FPCGMetadataAttributeBase* AttributeBase = SampledData->Metadata->GetMutableAttribute(DestinationAttribute);
-				if (!PCGMetadataOperations::SetValue<float, float>(SampledPoints, AttributeBase, SampledData->Metadata, SteepnessGetter))
+				if (!PCGMetadataOperations::SetValue<float>(SampledPoints, AttributeBase, SampledData->Metadata, SteepnessGetter))
 				{
 					PCGE_LOG(Error, "Attribute %s already exists but its type is not compatible", *DestinationAttribute.ToString());
 				}
@@ -287,7 +319,7 @@ bool FPCGMetadataOperationElement::ExecuteInternal(FPCGContext* Context) const
 				auto DensitySetter = [](FPCGPoint& InPoint, const float& InValue) { InPoint.Density = InValue; };
 
 				const FPCGMetadataAttributeBase* AttributeBase = SampledData->Metadata->GetConstAttribute(LocalSourceAttribute);
-				if(!PCGMetadataOperations::SetValue<float, float>(AttributeBase, SampledPoints, DensitySetter))
+				if(!PCGMetadataOperations::SetValue<float>(AttributeBase, SampledPoints, DensitySetter))
 				{
 					PCGE_LOG(Error, "Attribute %s already exists but its type is not compatible", *LocalSourceAttribute.ToString());
 				}
@@ -297,8 +329,7 @@ bool FPCGMetadataOperationElement::ExecuteInternal(FPCGContext* Context) const
 				auto ExtentsSetter = [](FPCGPoint& InPoint, const FVector& InValue) { InPoint.BoundsMin = InValue; };
 
 				const FPCGMetadataAttributeBase* AttributeBase = SampledData->Metadata->GetConstAttribute(LocalSourceAttribute);
-				if (!PCGMetadataOperations::SetValue<float, FVector>(AttributeBase, SampledPoints, ExtentsSetter) &&
-					!PCGMetadataOperations::SetValue<FVector, FVector>(AttributeBase, SampledPoints, ExtentsSetter))
+				if (!PCGMetadataOperations::SetValue<FVector>(AttributeBase, SampledPoints, ExtentsSetter))
 				{
 					PCGE_LOG(Error, "Attribute %s already exists but its type is not compatible", *LocalSourceAttribute.ToString());
 				}
@@ -308,8 +339,7 @@ bool FPCGMetadataOperationElement::ExecuteInternal(FPCGContext* Context) const
 				auto ExtentsSetter = [](FPCGPoint& InPoint, const FVector& InValue) { InPoint.BoundsMax = InValue; };
 
 				const FPCGMetadataAttributeBase* AttributeBase = SampledData->Metadata->GetConstAttribute(LocalSourceAttribute);
-				if (!PCGMetadataOperations::SetValue<float, FVector>(AttributeBase, SampledPoints, ExtentsSetter) &&
-					!PCGMetadataOperations::SetValue<FVector, FVector>(AttributeBase, SampledPoints, ExtentsSetter))
+				if (!PCGMetadataOperations::SetValue<FVector>(AttributeBase, SampledPoints, ExtentsSetter))
 				{
 					PCGE_LOG(Error, "Attribute %s already exists but its type is not compatible", *LocalSourceAttribute.ToString());
 				}
@@ -319,8 +349,7 @@ bool FPCGMetadataOperationElement::ExecuteInternal(FPCGContext* Context) const
 				auto ExtentsSetter = [](FPCGPoint& InPoint, const FVector& InValue) { InPoint.SetExtents(InValue); };
 
 				const FPCGMetadataAttributeBase* AttributeBase = SampledData->Metadata->GetConstAttribute(LocalSourceAttribute);
-				if(!PCGMetadataOperations::SetValue<float, FVector>(AttributeBase, SampledPoints, ExtentsSetter) &&
-					!PCGMetadataOperations::SetValue<FVector, FVector>(AttributeBase, SampledPoints, ExtentsSetter))
+				if(!PCGMetadataOperations::SetValue<FVector>(AttributeBase, SampledPoints, ExtentsSetter))
 				{
 					PCGE_LOG(Error, "Attribute %s already exists but its type is not compatible", *LocalSourceAttribute.ToString());
 				}
@@ -330,8 +359,7 @@ bool FPCGMetadataOperationElement::ExecuteInternal(FPCGContext* Context) const
 				auto ColorSetter = [](FPCGPoint& InPoint, const FVector4& InValue) { InPoint.Color = InValue; };
 
 				const FPCGMetadataAttributeBase* AttributeBase = SampledData->Metadata->GetConstAttribute(LocalSourceAttribute);
-				if(!PCGMetadataOperations::SetValue<float, FVector4>(AttributeBase, SampledPoints, ColorSetter) &&
-					!PCGMetadataOperations::SetValue<FVector4, FVector4>(AttributeBase, SampledPoints, ColorSetter))
+				if(!PCGMetadataOperations::SetValue<FVector4>(AttributeBase, SampledPoints, ColorSetter))
 				{
 					PCGE_LOG(Error, "Attribute %s already exists but its type is not compatible", *LocalSourceAttribute.ToString());
 				}
@@ -341,8 +369,7 @@ bool FPCGMetadataOperationElement::ExecuteInternal(FPCGContext* Context) const
 				auto PositionSetter = [](FPCGPoint& InPoint, const FVector& InValue) { InPoint.Transform.SetLocation(InValue); };
 
 				const FPCGMetadataAttributeBase* AttributeBase = SampledData->Metadata->GetConstAttribute(LocalSourceAttribute);
-				if (!PCGMetadataOperations::SetValue<float, FVector>(AttributeBase, SampledPoints, PositionSetter) &&
-					!PCGMetadataOperations::SetValue<FVector, FVector>(AttributeBase, SampledPoints, PositionSetter))
+				if (!PCGMetadataOperations::SetValue<FVector>(AttributeBase, SampledPoints, PositionSetter))
 				{
 					PCGE_LOG(Error, "Attribute %s already exists but its type is not compatible", *LocalSourceAttribute.ToString());
 				}
@@ -352,7 +379,7 @@ bool FPCGMetadataOperationElement::ExecuteInternal(FPCGContext* Context) const
 				auto RotationSetter = [](FPCGPoint& InPoint, const FQuat& InValue) { InPoint.Transform.SetRotation(InValue.GetNormalized()); };
 
 				const FPCGMetadataAttributeBase* AttributeBase = SampledData->Metadata->GetConstAttribute(LocalSourceAttribute);
-				if (!PCGMetadataOperations::SetValue<FQuat, FQuat>(AttributeBase, SampledPoints, RotationSetter))
+				if (!PCGMetadataOperations::SetValue<FQuat>(AttributeBase, SampledPoints, RotationSetter))
 				{
 					PCGE_LOG(Error, "Attribute %s already exists but its type is not compatible", *LocalSourceAttribute.ToString());
 				}
@@ -362,8 +389,7 @@ bool FPCGMetadataOperationElement::ExecuteInternal(FPCGContext* Context) const
 				auto ScaleSetter = [](FPCGPoint& InPoint, const FVector& InValue) { InPoint.Transform.SetScale3D(InValue); };
 
 				const FPCGMetadataAttributeBase* AttributeBase = SampledData->Metadata->GetConstAttribute(LocalSourceAttribute);
-				if (!PCGMetadataOperations::SetValue<float, FVector>(AttributeBase, SampledPoints, ScaleSetter) &&
-					!PCGMetadataOperations::SetValue<FVector, FVector>(AttributeBase, SampledPoints, ScaleSetter))
+				if (!PCGMetadataOperations::SetValue<FVector>(AttributeBase, SampledPoints, ScaleSetter))
 				{
 					PCGE_LOG(Error, "Attribute %s already exists but its type is not compatible", *LocalSourceAttribute.ToString());
 				}
@@ -373,8 +399,7 @@ bool FPCGMetadataOperationElement::ExecuteInternal(FPCGContext* Context) const
 				auto TransformSetter = [](FPCGPoint& InPoint, const FTransform& InValue) { InPoint.Transform = InValue; };
 
 				const FPCGMetadataAttributeBase* AttributeBase = SampledData->Metadata->GetConstAttribute(LocalSourceAttribute);
-				if (!PCGMetadataOperations::SetValue<FVector, FTransform>(AttributeBase, SampledPoints, TransformSetter) &&
-					!PCGMetadataOperations::SetValue<FTransform, FTransform>(AttributeBase, SampledPoints, TransformSetter))
+				if (!PCGMetadataOperations::SetValue<FTransform>(AttributeBase, SampledPoints, TransformSetter))
 				{
 					PCGE_LOG(Error, "Attribute %s already exists but its type is not compatible", *LocalSourceAttribute.ToString());
 				}
@@ -384,7 +409,7 @@ bool FPCGMetadataOperationElement::ExecuteInternal(FPCGContext* Context) const
 				auto SteepnessSetter = [](FPCGPoint& InPoint, const float& InValue) { InPoint.Steepness = InValue; };
 
 				const FPCGMetadataAttributeBase* AttributeBase = SampledData->Metadata->GetConstAttribute(LocalSourceAttribute);
-				if (!PCGMetadataOperations::SetValue<float, float>(AttributeBase, SampledPoints, SteepnessSetter))
+				if (!PCGMetadataOperations::SetValue<float>(AttributeBase, SampledPoints, SteepnessSetter))
 				{
 					PCGE_LOG(Error, "Attribute %s already exists but its type is not compatible", *LocalSourceAttribute.ToString());
 				}
