[PCG] Merge CreateParamData and CreateAttribute and deprecate CreateParamData.

diff --git a/Engine/Plugins/Experimental/PCG/Source/PCG/Private/Elements/PCGCreateAttribute.cpp b/Engine/Plugins/Experimental/PCG/Source/PCG/Private/Elements/PCGCreateAttribute.cpp
--- a/Engine/Plugins/Experimental/PCG/Source/PCG/Private/Elements/PCGCreateAttribute.cpp
+++ b/Engine/Plugins/Experimental/PCG/Source/PCG/Private/Elements/PCGCreateAttribute.cpp
@@ -56,36 +56,50 @@ bool FPCGCreateAttributeElement::ExecuteInternal(FPCGContext* Context) const
 
 	TArray<FPCGTaggedData> Inputs = Context->InputData.GetInputsByPin(PCGPinConstants::DefaultInputLabel);
 
+	// If the input is empty, we will create a new ParamData.
+	// We can re-use this newly object as the output
+	bool bCanReuseInputData = false;
 	if (Inputs.IsEmpty())
 	{
-		PCGE_LOG(Error, "No input connected");
-		return true;
+		FPCGTaggedData& NewData = Inputs.Emplace_GetRef();
+		NewData.Data = NewObject<UPCGParamData>();
+		NewData.Pin = PCGPinConstants::DefaultInputLabel;
+		bCanReuseInputData = true;
 	}
 
 	for (const FPCGTaggedData& InputTaggedData : Inputs)
 	{
 		const UPCGData* InputData = InputTaggedData.Data;
+		UPCGData* OutputData = nullptr;
 
-		const UPCGMetadata* ParentMetadata = nullptr;
+		UPCGMetadata* Metadata = NewObject<UPCGMetadata>();
+
+		bool bShouldAddNewEntry = false;
 
-		if (const UPCGPointData* InputPointData = Cast<UPCGPointData>(InputData))
+		if (const UPCGSpatialData* InputSpatialData = Cast<UPCGSpatialData>(InputData))
 		{
-			ParentMetadata = InputPointData->Metadata;
+			UPCGSpatialData* NewSpatialData = DuplicateObject<UPCGSpatialData>(const_cast<UPCGSpatialData*>(InputSpatialData), nullptr);
+			NewSpatialData->Metadata = Metadata;
+			NewSpatialData->InitializeFromData(InputSpatialData, /*InMetadataParentOverride=*/ nullptr, /*bInheritMetadata=*/ Settings->bKeepExistingAttributes);
+			
+			OutputData = NewSpatialData;
 		}
 		else if (const UPCGParamData* InputParamData = Cast<UPCGParamData>(InputData))
 		{
-			ParentMetadata = InputParamData->Metadata;
+			// If we can reuse input data, it is safe to const_cast, as it was created by ourselves above.
+			UPCGParamData* NewParamData = bCanReuseInputData ? const_cast<UPCGParamData*>(InputParamData) : NewObject<UPCGParamData>();
+			NewParamData->Metadata = Metadata;
+
+			Metadata->Initialize(Settings->bKeepExistingAttributes ? InputParamData->Metadata : nullptr);
+			OutputData = NewParamData;
+
+			// In case of param data, we want to add a new entry too
+			bShouldAddNewEntry = true;
 		}
 		else
 		{
-			PCGE_LOG(Error, "Invalid data as input. Only support points and params");
-			return true;
-		}
-
-		UPCGMetadata* Metadata = NewObject<UPCGMetadata>();
-		if (Settings->bKeepExistingAttributes)
-		{
-			Metadata->Initialize(ParentMetadata);
+			PCGE_LOG(Error, "Invalid data as input. Only support spatial and params");
+			continue;
 		}
 
 		FPCGMetadataAttributeBase* Attribute = Settings->ClearOrCreateAttribute(Metadata, ParamData);
@@ -93,37 +107,20 @@ bool FPCGCreateAttributeElement::ExecuteInternal(FPCGContext* Context) const
 		if (!Attribute)
 		{
 			PCGE_LOG(Error, "Error while creating attribute %s", *Settings->OutputAttributeName.ToString());
-			return true;
+			continue;
 		}
 
 		TArray<FPCGTaggedData>& Outputs = Context->OutputData.TaggedData;
 		FPCGTaggedData& Output = Outputs.Emplace_GetRef();
+		Output.Data = OutputData;
 
-		if (const UPCGPointData* InputPointData = Cast<UPCGPointData>(InputData))
+		// Add a new entry if it is a param data
+		if (bShouldAddNewEntry)
 		{
-			UPCGPointData* PointOutputData = NewObject<UPCGPointData>();
-			PointOutputData->Metadata = Metadata;
-			PointOutputData->InitializeFromData(InputPointData);
-
-			// Copy the points
-			TArray<FPCGPoint>& Points = PointOutputData->GetMutablePoints();
-			Points = InputPointData->GetPoints();
-
-			Settings->SetAttribute(Attribute, Metadata, Points, ParamData);
-
-			Output.Data = PointOutputData;
-		}
-		else if (const UPCGParamData* InputParamData = Cast<UPCGParamData>(InputData))
-		{
-			UPCGParamData* OutputParamData = NewObject<UPCGParamData>();
-			OutputParamData->Metadata = Metadata;
-
 			PCGMetadataEntryKey EntryKey = Metadata->AddEntry();
 			Settings->SetAttribute(Attribute, Metadata, EntryKey, ParamData);
-
-			Output.Data = OutputParamData;
 		}
 	}
 
 	return true;
-}
\ No newline at end of file
+}
diff --git a/Engine/Plugins/Experimental/PCG/Source/PCG/Public/Elements/PCGCreateAttribute.h b/Engine/Plugins/Experimental/PCG/Source/PCG/Public/Elements/PCGCreateAttribute.h
--- a/Engine/Plugins/Experimental/PCG/Source/PCG/Public/Elements/PCGCreateAttribute.h
+++ b/Engine/Plugins/Experimental/PCG/Source/PCG/Public/Elements/PCGCreateAttribute.h
@@ -34,6 +34,7 @@ protected:
 class FPCGCreateAttributeElement : public FSimplePCGElement
 {
 public:
+	virtual bool CanExecuteOnlyOnMainThread(const UPCGSettings* InSettings) const { return true; }
 	virtual bool IsCacheable(const UPCGSettings* InSettings) const override { return false; }
 
 protected:
diff --git a/Engine/Plugins/Experimental/PCG/Source/PCG/Public/Elements/PCGCreateParamData.h b/Engine/Plugins/Experimental/PCG/Source/PCG/Public/Elements/PCGCreateParamData.h
--- a/Engine/Plugins/Experimental/PCG/Source/PCG/Public/Elements/PCGCreateParamData.h
+++ b/Engine/Plugins/Experimental/PCG/Source/PCG/Public/Elements/PCGCreateParamData.h
@@ -7,7 +7,7 @@
 
 #include "PCGCreateParamData.generated.h"
 
-UCLASS(BlueprintType, ClassGroup = (Procedural))
+UCLASS(BlueprintType, ClassGroup = (Procedural), meta = (Deprecated = "5.2", DeprecationMessage = "Use CreateAttribute instead."))
 class PCG_API UPCGCreateParamDataSettings : public UPCGAttributeCreationBaseSettings
 {
 	GENERATED_BODY()
