[PCG] Fix wrong node name with boolean in CreateAttributeNode

#rnx

diff --git a/Engine/Plugins/Experimental/PCG/Source/PCG/Private/Elements/PCGCreateAttributeBase.cpp b/Engine/Plugins/Experimental/PCG/Source/PCG/Private/Elements/PCGCreateAttributeBase.cpp
--- a/Engine/Plugins/Experimental/PCG/Source/PCG/Private/Elements/PCGCreateAttributeBase.cpp
+++ b/Engine/Plugins/Experimental/PCG/Source/PCG/Private/Elements/PCGCreateAttributeBase.cpp
@@ -141,7 +141,7 @@ FName UPCGAttributeCreationBaseSettings::AdditionalTaskName() const
 	case EPCGMetadataTypes::Transform:
 		return FName(FString::Printf(TEXT("%s: Transform"), *Name));
 	case EPCGMetadataTypes::Boolean:
-		return FName(FString::Printf(TEXT("%s: %s"), *Name, (BoolValue ? "True" : "False")));
+		return FName(FString::Printf(TEXT("%s: %s"), *Name, (BoolValue ? TEXT("True") : TEXT("False"))));
 	default:
 		return NAME_None;
 	}
