[PCG] Add culling start and end distances to static mesh spawner node. These drive the corresponding settings on the generated ISMCs.

Side note: I could not test the start distance is working, which is the distance at which instances are supposed to start fading. I could not get this working even on normal ISMC components (without PCG involved). I assume some material flag is needed to allow this fading (I tried a few things but was not successful), or maybe its just broken at the moment. In any case this change simply makes PCG drive the cull start distance value on the ISMC and that works. I've kept this value exposed on the assumption that the right material flags will activate it.

diff --git a/Engine/Plugins/Experimental/PCG/Source/PCG/Private/Elements/PCGDebugElement.cpp b/Engine/Plugins/Experimental/PCG/Source/PCG/Private/Elements/PCGDebugElement.cpp
--- a/Engine/Plugins/Experimental/PCG/Source/PCG/Private/Elements/PCGDebugElement.cpp
+++ b/Engine/Plugins/Experimental/PCG/Source/PCG/Private/Elements/PCGDebugElement.cpp
@@ -145,6 +145,8 @@ namespace PCGDebugElement
 			Params.Mesh = Mesh;
 			Params.MaterialOverrides = Materials;
 			Params.CollisionProfile = UCollisionProfile::NoCollision_ProfileName;
+			// Note: In the future we may consider enabling culling for performance reasons, but for now culling disabled.
+			Params.CullStartDistance = Params.CullEndDistance = 0;
 
 			UInstancedStaticMeshComponent* ISMC = UPCGActorHelpers::GetOrCreateISMC(TargetActor, Context->SourceComponent.Get(), Params);
 			check(ISMC);
diff --git a/Engine/Plugins/Experimental/PCG/Source/PCG/Private/Elements/PCGStaticMeshSpawner.cpp b/Engine/Plugins/Experimental/PCG/Source/PCG/Private/Elements/PCGStaticMeshSpawner.cpp
--- a/Engine/Plugins/Experimental/PCG/Source/PCG/Private/Elements/PCGStaticMeshSpawner.cpp
+++ b/Engine/Plugins/Experimental/PCG/Source/PCG/Private/Elements/PCGStaticMeshSpawner.cpp
@@ -143,6 +143,9 @@ void FPCGStaticMeshSpawnerElement::SpawnStaticMeshInstances(FPCGContext* Context
 		Params.MaterialOverrides = InstanceList.MaterialOverrides;
 	}
 
+	Params.CullStartDistance = InstanceList.CullStartDistance;
+	Params.CullEndDistance = InstanceList.CullEndDistance;
+
 	Params.NumCustomDataFloats = PackedCustomData.NumCustomDataFloats;
 
 	UInstancedStaticMeshComponent* ISMC = UPCGActorHelpers::GetOrCreateISMC(TargetActor, Context->SourceComponent.Get(), Params);
diff --git a/Engine/Plugins/Experimental/PCG/Source/PCG/Private/Helpers/PCGActorHelpers.cpp b/Engine/Plugins/Experimental/PCG/Source/PCG/Private/Helpers/PCGActorHelpers.cpp
--- a/Engine/Plugins/Experimental/PCG/Source/PCG/Private/Helpers/PCGActorHelpers.cpp
+++ b/Engine/Plugins/Experimental/PCG/Source/PCG/Private/Helpers/PCGActorHelpers.cpp
@@ -64,7 +64,9 @@ UInstancedStaticMeshComponent* UPCGActorHelpers::GetOrCreateISMC(AActor* InTarge
 			// Check basic parameters
 			if (ISMC->Mobility != InParams.Mobility ||
 				(ISMC->bUseDefaultCollision && InParams.CollisionProfile != TEXT("Default")) ||
-				(!ISMC->bUseDefaultCollision && ISMC->GetCollisionProfileName() != InParams.CollisionProfile))
+				(!ISMC->bUseDefaultCollision && ISMC->GetCollisionProfileName() != InParams.CollisionProfile) ||
+				ISMC->InstanceStartCullDistance != InParams.CullStartDistance ||
+				ISMC->InstanceEndCullDistance != InParams.CullEndDistance)
 			{
 				continue;
 			}
@@ -140,6 +142,9 @@ UInstancedStaticMeshComponent* UPCGActorHelpers::GetOrCreateISMC(AActor* InTarge
 	{
 		ISMC->bUseDefaultCollision = true;
 	}
+
+	ISMC->InstanceStartCullDistance = InParams.CullStartDistance;
+	ISMC->InstanceEndCullDistance = InParams.CullEndDistance;
 	
 	ISMC->AttachToComponent(InTargetActor->GetRootComponent(), FAttachmentTransformRules::KeepWorldTransform);
 	ISMC->ComponentTags.Add(InSourceComponent->GetFName());
diff --git a/Engine/Plugins/Experimental/PCG/Source/PCG/Private/MeshSelectors/PCGMeshSelectorBase.cpp b/Engine/Plugins/Experimental/PCG/Source/PCG/Private/MeshSelectors/PCGMeshSelectorBase.cpp
--- a/Engine/Plugins/Experimental/PCG/Source/PCG/Private/MeshSelectors/PCGMeshSelectorBase.cpp
+++ b/Engine/Plugins/Experimental/PCG/Source/PCG/Private/MeshSelectors/PCGMeshSelectorBase.cpp
@@ -11,6 +11,8 @@ bool UPCGMeshSelectorBase::FindOrAddInstanceList(
 	const FCollisionProfileName& CollisionProfile,
 	bool bOverrideMaterials,
 	const TArray<UMaterialInterface*>& MaterialOverrides,
+	const float InCullStartDistance,
+	const float InCullEndDistance,
 	int32& OutIndex) const
 {
 	for (int Index = 0; Index < OutInstanceLists.Num(); ++Index)
@@ -20,6 +22,11 @@ bool UPCGMeshSelectorBase::FindOrAddInstanceList(
 			continue;
 		}
 
+		if (OutInstanceLists[Index].CullStartDistance != InCullStartDistance || OutInstanceLists[Index].CullEndDistance != InCullEndDistance)
+		{
+			continue;
+		}
+
 		if (OutInstanceLists[Index].bOverrideCollisionProfile && !OutInstanceLists[Index].CollisionProfile.Name.IsEqual(CollisionProfile.Name))
 		{
 			continue;
@@ -35,7 +42,7 @@ bool UPCGMeshSelectorBase::FindOrAddInstanceList(
 		return false;
 	}
 
-	OutIndex = OutInstanceLists.Emplace(Mesh, bOverrideCollisionProfile, CollisionProfile, bOverrideMaterials, MaterialOverrides);
+	OutIndex = OutInstanceLists.Emplace(Mesh, bOverrideCollisionProfile, CollisionProfile, bOverrideMaterials, MaterialOverrides, InCullStartDistance, InCullEndDistance);
 
 	return true;
 }
diff --git a/Engine/Plugins/Experimental/PCG/Source/PCG/Private/MeshSelectors/PCGMeshSelectorByAttribute.cpp b/Engine/Plugins/Experimental/PCG/Source/PCG/Private/MeshSelectors/PCGMeshSelectorByAttribute.cpp
--- a/Engine/Plugins/Experimental/PCG/Source/PCG/Private/MeshSelectors/PCGMeshSelectorByAttribute.cpp
+++ b/Engine/Plugins/Experimental/PCG/Source/PCG/Private/MeshSelectors/PCGMeshSelectorByAttribute.cpp
@@ -91,7 +91,7 @@ void UPCGMeshSelectorByAttribute::SelectInstances_Implementation(
 		}
 
 		int32 Index = INDEX_NONE;
-		FindOrAddInstanceList(OutMeshInstances, Mesh, bOverrideCollisionProfile, CollisionProfile, bOverrideMaterials, MaterialOverrides, Index);
+		FindOrAddInstanceList(OutMeshInstances, Mesh, bOverrideCollisionProfile, CollisionProfile, bOverrideMaterials, MaterialOverrides, CullStartDistance, CullEndDistance, Index);
 		OutMeshInstances[Index].Instances.Emplace(Point);
 	}
 }
diff --git a/Engine/Plugins/Experimental/PCG/Source/PCG/Private/MeshSelectors/PCGMeshSelectorWeighted.cpp b/Engine/Plugins/Experimental/PCG/Source/PCG/Private/MeshSelectors/PCGMeshSelectorWeighted.cpp
--- a/Engine/Plugins/Experimental/PCG/Source/PCG/Private/MeshSelectors/PCGMeshSelectorWeighted.cpp
+++ b/Engine/Plugins/Experimental/PCG/Source/PCG/Private/MeshSelectors/PCGMeshSelectorWeighted.cpp
@@ -33,7 +33,7 @@ void UPCGMeshSelectorWeighted::SelectInstances_Implementation(
 		}
 
 		int32 Index = INDEX_NONE;
-		FindOrAddInstanceList(OutMeshInstances, Entry.Mesh, Entry.bOverrideCollisionProfile, Entry.CollisionProfile, Entry.bOverrideMaterials, Entry.MaterialOverrides, Index);
+		FindOrAddInstanceList(OutMeshInstances, Entry.Mesh, Entry.bOverrideCollisionProfile, Entry.CollisionProfile, Entry.bOverrideMaterials, Entry.MaterialOverrides, Entry.CullStartDistance, Entry.CullEndDistance, Index);
 		TotalWeight += Entry.Weight;
 		CumulativeWeights.Add(TotalWeight);
 	}
diff --git a/Engine/Plugins/Experimental/PCG/Source/PCG/Private/MeshSelectors/PCGMeshSelectorWeightedByCategory.cpp b/Engine/Plugins/Experimental/PCG/Source/PCG/Private/MeshSelectors/PCGMeshSelectorWeightedByCategory.cpp
--- a/Engine/Plugins/Experimental/PCG/Source/PCG/Private/MeshSelectors/PCGMeshSelectorWeightedByCategory.cpp
+++ b/Engine/Plugins/Experimental/PCG/Source/PCG/Private/MeshSelectors/PCGMeshSelectorWeightedByCategory.cpp
@@ -112,7 +112,7 @@ void UPCGMeshSelectorWeightedByCategory::SelectInstances_Implementation(
 			}
 			
 			int32 Index = INDEX_NONE;
-			FindOrAddInstanceList(OutMeshInstances, WeightedEntry.Mesh, WeightedEntry.bOverrideCollisionProfile, WeightedEntry.CollisionProfile, WeightedEntry.bOverrideMaterials, WeightedEntry.MaterialOverrides, Index);
+			FindOrAddInstanceList(OutMeshInstances, WeightedEntry.Mesh, WeightedEntry.bOverrideCollisionProfile, WeightedEntry.CollisionProfile, WeightedEntry.bOverrideMaterials, WeightedEntry.MaterialOverrides, WeightedEntry.CullStartDistance, WeightedEntry.CullEndDistance, Index);
 
 			// precompute the weights
 			TotalWeight += WeightedEntry.Weight;
diff --git a/Engine/Plugins/Experimental/PCG/Source/PCG/Public/Helpers/PCGActorHelpers.h b/Engine/Plugins/Experimental/PCG/Source/PCG/Public/Helpers/PCGActorHelpers.h
--- a/Engine/Plugins/Experimental/PCG/Source/PCG/Public/Helpers/PCGActorHelpers.h
+++ b/Engine/Plugins/Experimental/PCG/Source/PCG/Public/Helpers/PCGActorHelpers.h
@@ -28,6 +28,8 @@ struct FPCGISMCBuilderParameters
 	EComponentMobility::Type Mobility = EComponentMobility::Static;
 	FName CollisionProfile = TEXT("Default");
 	int32 NumCustomDataFloats = 0;
+	float CullStartDistance = 0;
+	float CullEndDistance = 0;
 };
 
 UCLASS(BlueprintType)
diff --git a/Engine/Plugins/Experimental/PCG/Source/PCG/Public/MeshSelectors/PCGMeshSelectorBase.h b/Engine/Plugins/Experimental/PCG/Source/PCG/Public/MeshSelectors/PCGMeshSelectorBase.h
--- a/Engine/Plugins/Experimental/PCG/Source/PCG/Public/MeshSelectors/PCGMeshSelectorBase.h
+++ b/Engine/Plugins/Experimental/PCG/Source/PCG/Public/MeshSelectors/PCGMeshSelectorBase.h
@@ -21,8 +21,8 @@ struct FPCGMeshInstanceList
 
 	FPCGMeshInstanceList() = default;
 	
-	FPCGMeshInstanceList(const TSoftObjectPtr<UStaticMesh>& InMesh, bool bInOverrideCollisionProfile, const FCollisionProfileName& InCollisionProfile, bool bInOverrideMaterials, const TArray<UMaterialInterface*>& InMaterialOverrides)
-		: Mesh(InMesh), bOverrideCollisionProfile(bInOverrideCollisionProfile), CollisionProfile(InCollisionProfile), bOverrideMaterials(bInOverrideMaterials), MaterialOverrides(InMaterialOverrides)
+	FPCGMeshInstanceList(const TSoftObjectPtr<UStaticMesh>& InMesh, bool bInOverrideCollisionProfile, const FCollisionProfileName& InCollisionProfile, bool bInOverrideMaterials, const TArray<UMaterialInterface*>& InMaterialOverrides, const float InCullStartDistance, const float InCullEndDistance)
+		: Mesh(InMesh), bOverrideCollisionProfile(bInOverrideCollisionProfile), CollisionProfile(InCollisionProfile), bOverrideMaterials(bInOverrideMaterials), MaterialOverrides(InMaterialOverrides), CullStartDistance(InCullStartDistance), CullEndDistance(InCullEndDistance)
 	{}
 
 	UPROPERTY(BlueprintReadWrite, EditAnywhere, Category = Settings)
@@ -42,6 +42,14 @@ struct FPCGMeshInstanceList
 
 	UPROPERTY(BlueprintReadWrite, EditAnywhere, Category = Settings)
     TArray<FPCGPoint> Instances;
+
+	/** Distance at which instances begin to fade. */
+	UPROPERTY(BlueprintReadWrite, EditAnywhere, Category = Settings)
+	float CullStartDistance = 0;
+	
+	/** Distance at which instances are culled. Use 0 to disable. */
+	UPROPERTY(BlueprintReadWrite, EditAnywhere, Category = Settings)
+	float CullEndDistance = 0;
 };
 
 UCLASS(Abstract, BlueprintType, Blueprintable, ClassGroup = (Procedural))
@@ -74,5 +82,7 @@ public:
 		const FCollisionProfileName& CollisionProfile,
 		bool bOverrideMaterials,
 		const TArray<UMaterialInterface*>& MaterialOverrides,
+		const float InCullStartDistance,
+		const float InCullEndDistance,
 		int32& OutIndex) const;
 };
diff --git a/Engine/Plugins/Experimental/PCG/Source/PCG/Public/MeshSelectors/PCGMeshSelectorByAttribute.h b/Engine/Plugins/Experimental/PCG/Source/PCG/Public/MeshSelectors/PCGMeshSelectorByAttribute.h
--- a/Engine/Plugins/Experimental/PCG/Source/PCG/Public/MeshSelectors/PCGMeshSelectorByAttribute.h
+++ b/Engine/Plugins/Experimental/PCG/Source/PCG/Public/MeshSelectors/PCGMeshSelectorByAttribute.h
@@ -34,5 +34,13 @@ public:
 
 	UPROPERTY(BlueprintReadWrite, EditAnywhere, Category = Settings)
 	TArray<TObjectPtr<UMaterialInterface>> MaterialOverrides;
+
+	/** Distance at which instances begin to fade. */
+	UPROPERTY(BlueprintReadWrite, EditAnywhere, Category = Settings)
+	float CullStartDistance = 0;
+	
+	/** Distance at which instances are culled. Use 0 to disable. */
+	UPROPERTY(BlueprintReadWrite, EditAnywhere, Category = Settings)
+	float CullEndDistance = 0;
 };
 
diff --git a/Engine/Plugins/Experimental/PCG/Source/PCG/Public/MeshSelectors/PCGMeshSelectorWeighted.h b/Engine/Plugins/Experimental/PCG/Source/PCG/Public/MeshSelectors/PCGMeshSelectorWeighted.h
--- a/Engine/Plugins/Experimental/PCG/Source/PCG/Public/MeshSelectors/PCGMeshSelectorWeighted.h
+++ b/Engine/Plugins/Experimental/PCG/Source/PCG/Public/MeshSelectors/PCGMeshSelectorWeighted.h
@@ -32,6 +32,14 @@ struct PCG_API FPCGMeshSelectorWeightedEntry
 	UPROPERTY(BlueprintReadWrite, EditAnywhere, Category = Settings)
 	TArray<TObjectPtr<UMaterialInterface>> MaterialOverrides;
 
+	/** Distance at which instances begin to fade. */
+	UPROPERTY(BlueprintReadWrite, EditAnywhere, Category = Settings)
+	float CullStartDistance = 0;
+	
+	/** Distance at which instances are culled. Use 0 to disable. */
+	UPROPERTY(BlueprintReadWrite, EditAnywhere, Category = Settings)
+	float CullEndDistance = 0;
+
 	UPROPERTY(BlueprintReadWrite, EditAnywhere, Category = Settings, meta = (ClampMin = "0"))
 	int Weight = 1;
 };
