Improved the way we do point/volume to point sampling w.r.t bounds

diff --git a/Engine/Plugins/Experimental/PCG/Source/PCG/Private/Data/PCGPointData.cpp b/Engine/Plugins/Experimental/PCG/Source/PCG/Private/Data/PCGPointData.cpp
--- a/Engine/Plugins/Experimental/PCG/Source/PCG/Private/Data/PCGPointData.cpp
+++ b/Engine/Plugins/Experimental/PCG/Source/PCG/Private/Data/PCGPointData.cpp
@@ -66,18 +66,41 @@ namespace PCGPointHelpers
 		}
 	}
 
-	float VolumeOverlap(const FPCGPoint& InPoint, const FBox& InTransformedBounds)
+	/** Computes reasonable overlap ratio for point, 1d, 2d and volume overlaps, to be used as weights.
+	* Note that this assumes that either data set is homogeneous in its points dimension (either 0d, 1d, 2d, 3d)
+	* Otherwise there will be some artifacts from our assumption here (namely using a 1.0 value for the additional coordinates).
+	*/
+	float ComputeOverlapRatio(const FBox& Numerator, const FBox& Denominator)
 	{
-		const FBox PointTransformedBounds = InPoint.GetLocalBounds().TransformBy(InPoint.Transform);
-		const FBox Overlap = PointTransformedBounds.Overlap(InTransformedBounds);
-		if (Overlap.IsValid)
+		const FVector NumeratorExtent = Numerator.GetExtent();
+		const FVector DenominatorExtent = Denominator.GetExtent();
+
+		return (float)((DenominatorExtent.X > 0 ? NumeratorExtent.X / DenominatorExtent.X : 1.0) *
+			(DenominatorExtent.Y > 0 ? NumeratorExtent.Y / DenominatorExtent.Y : 1.0) *
+			(DenominatorExtent.Z > 0 ? NumeratorExtent.Z / DenominatorExtent.Z : 1.0));
+	}
+
+	float VolumeOverlap(const FPCGPoint& InPoint, const FBox& InBounds, const FTransform& InTransform)
+	{
+		// This is similar in idea to SAT considering we have two boxes - since we will test all 6 axes.
+		// However, there is some uncertainty due to rotation, and using the overlap value as-is is an overestimation, which might not be critical in this case
+		// TODO: investigate if we should do a 8-pt test instead (would be more precise, but significantly more costly).
+		const FBox PointBounds = InPoint.GetLocalBounds();
+		const FTransform& PointTransform = InPoint.Transform;
+
+		const FBox FirstOverlap = PointBounds.Overlap(InBounds.TransformBy(InTransform.GetRelativeTransform(PointTransform)));
+		if (!FirstOverlap.IsValid)
 		{
-			return Overlap.GetVolume();
+			return 0;
 		}
-		else
+
+		const FBox SecondOverlap = InBounds.Overlap(PointBounds.TransformBy(PointTransform.GetRelativeTransform(InTransform)));
+		if (!SecondOverlap.IsValid)
 		{
 			return 0;
 		}
+		
+		return FMath::Min(ComputeOverlapRatio(FirstOverlap, PointBounds), ComputeOverlapRatio(SecondOverlap, InBounds));
 	}
 
 	/** Helper function for additive blending of quaternions (copied from ControlRig) */
@@ -316,14 +339,18 @@ bool UPCGPointData::SamplePoint(const FTransform& InTransform, const FBox& InBou
 		const FVector InPosition = InTransform.GetLocation();
 		Octree.FindElementsWithBoundsTest(FBoxCenterAndExtent(InPosition, FVector::Zero()), [&InPosition, &Contributions](const FPCGPointRef& InPointRef) {
 			Contributions.Emplace(InPointRef.Point, PCGPointHelpers::InverseEuclidianDistance(*InPointRef.Point, InPosition));
-			});
+		});
 	}
 	else
 	{
 		FBox TransformedBounds = InBounds.TransformBy(InTransform);
-		Octree.FindElementsWithBoundsTest(FBoxCenterAndExtent(TransformedBounds.GetCenter(), TransformedBounds.GetExtent()), [&TransformedBounds, &Contributions](const FPCGPointRef& InPointRef) {
-			Contributions.Emplace(InPointRef.Point, PCGPointHelpers::VolumeOverlap(*InPointRef.Point, TransformedBounds));
-			});
+		Octree.FindElementsWithBoundsTest(FBoxCenterAndExtent(TransformedBounds.GetCenter(), TransformedBounds.GetExtent()), [&InBounds, &InTransform, &Contributions](const FPCGPointRef& InPointRef) {
+			float Contribution = PCGPointHelpers::VolumeOverlap(*InPointRef.Point, InBounds, InTransform);
+			if (Contribution > 0)
+			{
+				Contributions.Emplace(InPointRef.Point, Contribution);
+			}
+		});
 	}
 
 	float SumContributions = 0;
@@ -348,12 +375,6 @@ bool UPCGPointData::SamplePoint(const FTransform& InTransform, const FBox& InBou
 
 	TArray<TPair<const FPCGPoint*, float>, TInlineAllocator<4>> ContributionsForMetadata;
 
-	const FVector::FReal TransformedBoundsVolume = InBounds.TransformBy(InTransform).GetVolume();
-	auto ComputeDensityContribution = [TransformedBoundsVolume](float VolumeIntersection)
-	{
-		return (TransformedBoundsVolume > 0) ? VolumeIntersection / TransformedBoundsVolume : 1.0f;
-	};
-
 	// Computed weighted average of spatial properties
 	FVector WeightedPosition = FVector::ZeroVector;
 	FQuat WeightedQuat = FQuat::Identity;
@@ -379,7 +400,7 @@ bool UPCGPointData::SamplePoint(const FTransform& InTransform, const FBox& InBou
 		}
 		else
 		{
-			WeightedDensity += SourcePoint.Density * Weight * ComputeDensityContribution(Contribution.Value);
+			WeightedDensity += SourcePoint.Density * Weight * Contribution.Value;
 		}
 
 		WeightedBoundsMin += SourcePoint.BoundsMin * Weight;
