[PCG] Add PostGenerateFunction to PCGComponents

diff --git a/Engine/Plugins/Experimental/PCG/Source/PCG/Private/Graph/PCGGraphCompiler.cpp b/Engine/Plugins/Experimental/PCG/Source/PCG/Private/Graph/PCGGraphCompiler.cpp
--- a/Engine/Plugins/Experimental/PCG/Source/PCG/Private/Graph/PCGGraphCompiler.cpp
+++ b/Engine/Plugins/Experimental/PCG/Source/PCG/Private/Graph/PCGGraphCompiler.cpp
@@ -310,14 +310,18 @@ void FPCGGraphCompiler::CompileTopGraph(UPCGGraph* InGraph)
 	// Note: this works because there is a weak ordering on the tasks such that
 	// a successor task is always after its predecessors
 	TSet<FPCGTaskId> TasksWithSuccessors;
-	TArray<FPCGTaskId> EndNodes;
+	const UPCGNode* GraphOutputNode = InGraph->GetOutputNode();
 
 	for (int TaskIndex = TaskNum - 1; TaskIndex >= 0; --TaskIndex)
 	{
 		const FPCGGraphTask& Task = CompiledTasks[TaskIndex];
 		if (!TasksWithSuccessors.Contains(Task.NodeId))
 		{
-			PostExecuteTask.Inputs.Emplace(Task.NodeId, nullptr, nullptr);
+			// For the post task, only the output node will provide data. 
+			// It is necessary for any post generation task to get the content of the output node
+			// and only this content.
+			const bool bProvideData = Task.Node == GraphOutputNode;
+			PostExecuteTask.Inputs.Emplace(Task.NodeId, nullptr, nullptr, bProvideData);
 		}
 
 		for (const FPCGGraphTaskInput& Input : Task.Inputs)
diff --git a/Engine/Plugins/Experimental/PCG/Source/PCG/Private/Graph/PCGGraphExecutor.cpp b/Engine/Plugins/Experimental/PCG/Source/PCG/Private/Graph/PCGGraphExecutor.cpp
--- a/Engine/Plugins/Experimental/PCG/Source/PCG/Private/Graph/PCGGraphExecutor.cpp
+++ b/Engine/Plugins/Experimental/PCG/Source/PCG/Private/Graph/PCGGraphExecutor.cpp
@@ -106,7 +106,8 @@ FPCGTaskId FPCGGraphExecutor::Schedule(UPCGGraph* Graph, UPCGComponent* SourceCo
 		check(ScheduledTask.Tasks.Num() >= 2 && ScheduledTask.Tasks[ScheduledTask.Tasks.Num() - 2].Node == nullptr);
 		for (FPCGTaskId ExternalDependency : ExternalDependencies)
 		{
-			ScheduledTask.Tasks[ScheduledTask.Tasks.Num() - 2].Inputs.Emplace(ExternalDependency, nullptr, nullptr);
+			// For the pre-task, we don't consume any input
+			ScheduledTask.Tasks[ScheduledTask.Tasks.Num() - 2].Inputs.Emplace(ExternalDependency, nullptr, nullptr, /*bConsumeData=*/false);
 		}
 
 		ScheduleLock.Unlock();
@@ -116,12 +117,22 @@ FPCGTaskId FPCGGraphExecutor::Schedule(UPCGGraph* Graph, UPCGComponent* SourceCo
 }
 
 FPCGTaskId FPCGGraphExecutor::ScheduleGeneric(TFunction<bool()> InOperation, UPCGComponent* InSourceComponent, const TArray<FPCGTaskId>& TaskDependencies)
+{
+	// Since we have no context, the generic task will consume no input
+	constexpr bool bConsumeInputData = false;
+	return ScheduleGenericWithContext([Operation = MoveTemp(InOperation)](FPCGContext*) -> bool
+		{
+			return Operation();
+		}, InSourceComponent, TaskDependencies, bConsumeInputData);
+}
+
+FPCGTaskId FPCGGraphExecutor::ScheduleGenericWithContext(TFunction<bool(FPCGContext*)> InOperation, UPCGComponent* InSourceComponent, const TArray<FPCGTaskId>& TaskDependencies, bool bConsumeInputData)
 {
 	// Build task & element to hold the operation to perform
 	FPCGGraphTask Task;
 	for (FPCGTaskId TaskDependency : TaskDependencies)
 	{
-		Task.Inputs.Emplace(TaskDependency, nullptr, nullptr);
+		Task.Inputs.Emplace(TaskDependency, nullptr, nullptr, bConsumeInputData);
 	}
 
 	Task.SourceComponent = InSourceComponent;
@@ -601,6 +612,13 @@ void FPCGGraphExecutor::BuildTaskInput(const FPCGGraphTask& Task, FPCGDataCollec
 	for (const FPCGGraphTaskInput& Input : Task.Inputs)
 	{
 		check(OutputData.Contains(Input.TaskId));
+
+		// If the input does not provide any data, don't add it to the task input.
+		if (!Input.bProvideData)
+		{
+			continue;
+		}
+
 		const FPCGDataCollection& InputCollection = OutputData[Input.TaskId];
 
 		TaskInput.bCancelExecution |= InputCollection.bCancelExecution;
@@ -884,7 +902,7 @@ bool FPCGFetchInputElement::ExecuteInternal(FPCGContext* Context) const
 	return true;
 }
 
-FPCGGenericElement::FPCGGenericElement(TFunction<bool()> InOperation)
+FPCGGenericElement::FPCGGenericElement(TFunction<bool(FPCGContext*)> InOperation)
 	: Operation(InOperation)
 {
 }
@@ -892,5 +910,5 @@ FPCGGenericElement::FPCGGenericElement(TFunction<bool()> InOperation)
 bool FPCGGenericElement::ExecuteInternal(FPCGContext* Context) const
 {
 	TRACE_CPUPROFILER_EVENT_SCOPE(FPCGGenericElement::Execute);
-	return Operation();
+	return Operation(Context);
 }
diff --git a/Engine/Plugins/Experimental/PCG/Source/PCG/Private/Graph/PCGGraphExecutor.h b/Engine/Plugins/Experimental/PCG/Source/PCG/Private/Graph/PCGGraphExecutor.h
--- a/Engine/Plugins/Experimental/PCG/Source/PCG/Private/Graph/PCGGraphExecutor.h
+++ b/Engine/Plugins/Experimental/PCG/Source/PCG/Private/Graph/PCGGraphExecutor.h
@@ -22,16 +22,18 @@ class FPCGGraphCompiler;
 
 struct FPCGGraphTaskInput
 {
-	FPCGGraphTaskInput(FPCGTaskId InTaskId, const UPCGPin* InInboundPin, const UPCGPin* InOutboundPin)
+	FPCGGraphTaskInput(FPCGTaskId InTaskId, const UPCGPin* InInboundPin, const UPCGPin* InOutboundPin, bool bInProvideData = true)
 		: TaskId(InTaskId)
 		, InPin(InInboundPin)
 		, OutPin(InOutboundPin)
+		, bProvideData(bInProvideData)
 	{
 	}
 
 	FPCGTaskId TaskId;
 	const UPCGPin* InPin;
 	const UPCGPin* OutPin;
+	bool bProvideData;
 };
 
 struct FPCGGraphTask
@@ -72,9 +74,17 @@ public:
 	FPCGTaskId Schedule(UPCGComponent* InComponent, const TArray<FPCGTaskId>& TaskDependency = TArray<FPCGTaskId>());
 	FPCGTaskId Schedule(UPCGGraph* Graph, UPCGComponent* InSourceComponent, FPCGElementPtr InputElement, const TArray<FPCGTaskId>& TaskDependency);
 
-	/** General job scheduling, used to control loading/unloading */
+	// Back compatibility function. Use ScheduleGenericWithContext
 	FPCGTaskId ScheduleGeneric(TFunction<bool()> InOperation, UPCGComponent* InSourceComponent, const TArray<FPCGTaskId>& TaskDependencies);
 
+	/** General job scheduling
+	*  @param InOperation:       Callback that takes a Context as argument and returns true if the task is done, false otherwise
+	*  @param InSourceComponent: PCG component associated with this task. Can be null.
+	*  @param TaskDependencies:  List of all the dependencies for this task.
+	*  @param bConsumeInputData: If your task need a context, but don't need the input data, set this flag to false. Default is true.
+	*/
+	FPCGTaskId ScheduleGenericWithContext(TFunction<bool(FPCGContext*)> InOperation, UPCGComponent* InSourceComponent, const TArray<FPCGTaskId>& TaskDependencies, bool bConsumeInputData = true);
+
 	/** Gets data in the output results. Returns false if data is not ready. */
 	bool GetOutputData(FPCGTaskId InTaskId, FPCGDataCollection& OutData);
 
@@ -164,7 +174,7 @@ protected:
 class FPCGGenericElement : public FSimplePCGElement
 {
 public:
-	FPCGGenericElement(TFunction<bool()> InOperation);
+	FPCGGenericElement(TFunction<bool(FPCGContext*)> InOperation);
 	virtual bool IsCacheable(const UPCGSettings* InSettings) const override { return false; }
 	virtual bool CanExecuteOnlyOnMainThread(const UPCGSettings* InSettings) const override { return true; }
 
@@ -179,5 +189,5 @@ protected:
 #endif
 
 private:
-	TFunction<bool()> Operation;
+	TFunction<bool(FPCGContext*)> Operation;
 };
diff --git a/Engine/Plugins/Experimental/PCG/Source/PCG/Private/PCGComponent.cpp b/Engine/Plugins/Experimental/PCG/Source/PCG/Private/PCGComponent.cpp
--- a/Engine/Plugins/Experimental/PCG/Source/PCG/Private/PCGComponent.cpp
+++ b/Engine/Plugins/Experimental/PCG/Source/PCG/Private/PCGComponent.cpp
@@ -329,7 +329,7 @@ bool UPCGComponent::GetActorsFromTags(const TSet<FName>& InTags, TSet<TWeakObjec
 	return bHasValidTag;
 }
 
-void UPCGComponent::PostProcessGraph(const FBox& InNewBounds, bool bInGenerated)
+void UPCGComponent::PostProcessGraph(const FBox& InNewBounds, bool bInGenerated, FPCGContext* Context)
 {
 	LastGeneratedBounds = InNewBounds;
 
@@ -345,6 +345,72 @@ void UPCGComponent::PostProcessGraph(const FBox& InNewBounds, bool bInGenerated)
 		bDirtyGenerated = false;
 		OnPCGGraphGeneratedDelegate.Broadcast(this);
 #endif
+		// After a successful generation, we also want to call PostGenerateFunctions
+		// if we have any. We also need a context.
+
+		if (Context)
+		{
+			// If the original component is partitioned, local components have to forward
+			// their inputs, so that they can be gathered by the original component.
+			// We don't have the info on the original component here, so forward for all
+			// components.
+			Context->OutputData = Context->InputData;
+
+			CallPostGenerateFunctions(Context);
+		}
+	}
+}
+
+void UPCGComponent::CallPostGenerateFunctions(FPCGContext* Context) const
+{
+	check(Context);
+
+	if (AActor* Owner = GetOwner())
+	{
+		for (const FName& FunctionName : PostGenerateFunctionNames)
+		{
+			if (UFunction* PostGenerateFunc = Owner->GetClass()->FindFunctionByName(FunctionName))
+			{
+				// Validate that the function take the right number of arguments
+				if (PostGenerateFunc->NumParms != 1)
+				{
+					UE_LOG(LogPCG, Error, TEXT("[UPCGComponent] PostGenerateFunction \"%s\" from actor \"%s\" doesn't have exactly 1 parameter. Will skip the call."), *FunctionName.ToString(), *Owner->GetFName().ToString());
+					continue;
+				}
+
+				bool bIsValid = false;
+				TFieldIterator<FProperty> PropIterator(PostGenerateFunc);
+				while (PropIterator)
+				{
+					if (!!(PropIterator->PropertyFlags & CPF_Parm))
+					{
+						if (FStructProperty* Property = CastField<FStructProperty>(*PropIterator))
+						{
+							if (Property->Struct == FPCGDataCollection::StaticStruct())
+							{
+								bIsValid = true;
+								break;
+							}
+						}
+					}
+
+					++PropIterator;
+				}
+
+				if (bIsValid)
+				{
+					Owner->ProcessEvent(PostGenerateFunc, &Context->InputData);
+				}
+				else
+				{
+					UE_LOG(LogPCG, Error, TEXT("[UPCGComponent] PostGenerateFunction \"%s\" from actor \"%s\" parameter type is not PCGDataCollection. Will skip the call."), *FunctionName.ToString(), *Owner->GetFName().ToString());
+				}
+			}
+			else
+			{
+				UE_LOG(LogPCG, Error, TEXT("[UPCGComponent] PostGenerateFunction \"%s\" was not found in the component owner \"%s\"."), *FunctionName.ToString(), *Owner->GetFName().ToString());
+			}
+		}
 	}
 }
 
@@ -652,7 +718,7 @@ void UPCGComponent::BeginPlay()
 			const FBox NewBounds = GetGridBounds();
 			if (NewBounds.IsValid)
 			{
-				PostProcessGraph(NewBounds, true);
+				PostProcessGraph(NewBounds, true, nullptr);
 			}
 		}
 		else
diff --git a/Engine/Plugins/Experimental/PCG/Source/PCG/Private/PCGSubsystem.cpp b/Engine/Plugins/Experimental/PCG/Source/PCG/Private/PCGSubsystem.cpp
--- a/Engine/Plugins/Experimental/PCG/Source/PCG/Private/PCGSubsystem.cpp
+++ b/Engine/Plugins/Experimental/PCG/Source/PCG/Private/PCGSubsystem.cpp
@@ -247,7 +247,7 @@ FPCGTaskId UPCGSubsystem::ScheduleComponent(UPCGComponent* PCGComponent, bool bS
 	{
 		TWeakObjectPtr<UPCGComponent> ComponentPtr(PCGComponent);
 
-		return GraphExecutor->ScheduleGeneric([ComponentPtr]() {
+		return GraphExecutor->ScheduleGenericWithContext([ComponentPtr](FPCGContext* Context) {
 			if (UPCGComponent* Component = ComponentPtr.Get())
 			{
 				// If the component is not valid anymore, just early out.
@@ -257,7 +257,7 @@ FPCGTaskId UPCGSubsystem::ScheduleComponent(UPCGComponent* PCGComponent, bool bS
 				}
 
 				const FBox NewBounds = Component->GetGridBounds();
-				Component->PostProcessGraph(NewBounds, /*bGenerate=*/true);
+				Component->PostProcessGraph(NewBounds, /*bGenerate=*/true, Context);
 			}
 
 			return true;
@@ -390,6 +390,12 @@ FPCGTaskId UPCGSubsystem::ScheduleGeneric(TFunction<bool()> InOperation, UPCGCom
 	return GraphExecutor->ScheduleGeneric(InOperation, Component, TaskDependencies);
 }
 
+FPCGTaskId UPCGSubsystem::ScheduleGenericWithContext(TFunction<bool(FPCGContext*)> InOperation, UPCGComponent* Component, const TArray<FPCGTaskId>& TaskDependencies, bool bConsumeInputData)
+{
+	check(GraphExecutor);
+	return GraphExecutor->ScheduleGenericWithContext(InOperation, Component, TaskDependencies, bConsumeInputData);
+}
+
 TArray<FPCGTaskId> UPCGSubsystem::DispatchToRegisteredLocalComponents(UPCGComponent* OriginalComponent, const TFunction<FPCGTaskId(UPCGComponent*)>& InFunc) const
 {
 	TRACE_CPUPROFILER_EVENT_SCOPE(UPCGSubsystem::DispatchToRegisteredLocalComponents);
@@ -1169,15 +1175,15 @@ FPCGTaskId UPCGSubsystem::ProcessGraph(UPCGComponent* Component, const FBox& InP
 	// Finally, call PostProcessGraph if something happened
 	if (ProcessAllCellsTaskId != InvalidPCGTaskId)
 	{
-		auto PostProcessGraph = [ComponentPtr, InNewBounds, bGenerate]() {
+		auto PostProcessGraphTask = [ComponentPtr, InNewBounds, bGenerate](FPCGContext* Context) {
 			if (UPCGComponent* Component = ComponentPtr.Get())
 			{
-				Component->PostProcessGraph(InNewBounds, bGenerate);
+				Component->PostProcessGraph(InNewBounds, bGenerate, Context);
 			}
 			return true;
 		};
 
-		return GraphExecutor->ScheduleGeneric(PostProcessGraph, Component, { ProcessAllCellsTaskId });
+		return GraphExecutor->ScheduleGenericWithContext(PostProcessGraphTask, Component, { ProcessAllCellsTaskId });
 	}
 	else
 	{
diff --git a/Engine/Plugins/Experimental/PCG/Source/PCG/Public/PCGComponent.h b/Engine/Plugins/Experimental/PCG/Source/PCG/Public/PCGComponent.h
--- a/Engine/Plugins/Experimental/PCG/Source/PCG/Public/PCGComponent.h
+++ b/Engine/Plugins/Experimental/PCG/Source/PCG/Public/PCGComponent.h
@@ -9,6 +9,7 @@
 #include "PCGComponent.generated.h"
 
 class APCGPartitionActor;
+struct FPCGContext;
 class UPCGComponent;
 class UPCGGraph;
 class UPCGManagedResource;
@@ -177,6 +178,12 @@ public:
 	FOnPCGGraphCleaned OnPCGGraphCleanedDelegate;
 #endif
 
+	/** Can specify a list of functions from the owner of this component to be called when generation is done, in order. 
+	*   Need to take (and only take) a PCGDataCollection as parameter and with "CallInEditor" flag enabled.
+	*/
+	UPROPERTY(BlueprintReadWrite, EditAnywhere, AdvancedDisplay, Category = Properties)
+	TArray<FName> PostGenerateFunctionNames;
+
 	/** Return if we are currently generating the graph for this component */
 	bool IsGenerating() const { return CurrentGenerationTask != InvalidPCGTaskId; }
 	bool IsCleaningUp() const { return CurrentCleanupTask != InvalidPCGTaskId; }
@@ -240,7 +247,8 @@ private:
 	/* Internal call to create tasks to cleanup the component. If there is nothing to do, an invalid task id will be returned. Should only be used by the subsystem. */
 	FPCGTaskId CreateCleanupTask(bool bRemoveComponents, const TArray<FPCGTaskId>& Dependencies);
 
-	void PostProcessGraph(const FBox& InNewBounds, bool bInGenerated);
+	void PostProcessGraph(const FBox& InNewBounds, bool bInGenerated, FPCGContext* Context);
+	void CallPostGenerateFunctions(FPCGContext* Context) const;
 	void PostCleanupGraph();
 	void OnProcessGraphAborted();
 	void CleanupUnusedManagedResources();
diff --git a/Engine/Plugins/Experimental/PCG/Source/PCG/Public/PCGSubsystem.h b/Engine/Plugins/Experimental/PCG/Source/PCG/Public/PCGSubsystem.h
--- a/Engine/Plugins/Experimental/PCG/Source/PCG/Public/PCGSubsystem.h
+++ b/Engine/Plugins/Experimental/PCG/Source/PCG/Public/PCGSubsystem.h
@@ -70,8 +70,20 @@ public:
 	// Schedule graph (used internally for dynamic subgraph execution)
 	FPCGTaskId ScheduleGraph(UPCGComponent* SourceComponent, const TArray<FPCGTaskId>& Dependencies);
 
-	/** General job scheduling, used to control loading/unloading */
+	/** General job scheduling
+	*  @param InOperation:       Callback that returns true if the task is done, false otherwise.
+	*  @param InSourceComponent: PCG component associated with this task. Can be null.
+	*  @param TaskDependencies:  List of all the dependencies for this task.
+	*/
 	FPCGTaskId ScheduleGeneric(TFunction<bool()> InOperation, UPCGComponent* SourceComponent, const TArray<FPCGTaskId>& TaskDependencies);
+	
+	/** General job scheduling with context
+	*  @param InOperation:       Callback that takes a Context as argument and returns true if the task is done, false otherwise.
+	*  @param InSourceComponent: PCG component associated with this task. Can be null.
+	*  @param TaskDependencies:  List of all the dependencies for this task.
+	*  @param bConsumeInputData: If your task need a context, but don't need the input data, set this flag to false. Default is true.
+	*/
+	FPCGTaskId ScheduleGenericWithContext(TFunction<bool(FPCGContext*)> InOperation, UPCGComponent* SourceComponent, const TArray<FPCGTaskId>& TaskDependencies, bool bConsumeInputData = true);
 
 	/** Gets the output data for a given task */
 	bool GetOutputData(FPCGTaskId InTaskId, FPCGDataCollection& OutData);
