[PCG] Fix metadata values not being projected. The return value of WeightedSum() was being discarded and metadata after the projection would end up with default values.

diff --git a/Engine/Plugins/Experimental/PCG/Source/PCG/Public/Metadata/PCGMetadataAttributeTpl.h b/Engine/Plugins/Experimental/PCG/Source/PCG/Public/Metadata/PCGMetadataAttributeTpl.h
--- a/Engine/Plugins/Experimental/PCG/Source/PCG/Public/Metadata/PCGMetadataAttributeTpl.h
+++ b/Engine/Plugins/Experimental/PCG/Source/PCG/Public/Metadata/PCGMetadataAttributeTpl.h
@@ -411,7 +411,7 @@ protected:
 		IT Value = PCG::Private::MetadataTraits<IT>::ZeroValue();
 		for (const TPair<PCGMetadataEntryKey, float>& WeightedEntry : InWeightedKeys)
 		{
-			PCG::Private::MetadataTraits<IT>::WeightedSum(
+			Value = PCG::Private::MetadataTraits<IT>::WeightedSum(
 				Value,
 				static_cast<const FPCGMetadataAttribute<T>*>(InAttribute)->GetValueFromItemKey(WeightedEntry.Key),
 				WeightedEntry.Value);
