Disabled tighter SCC integration in our generation process temporarily because it was not able to distinguish between unloading & deleting. Will be revisited soon.

diff --git a/Engine/Plugins/Experimental/PCG/Source/PCG/Private/Helpers/PCGActorHelpers.cpp b/Engine/Plugins/Experimental/PCG/Source/PCG/Private/Helpers/PCGActorHelpers.cpp
--- a/Engine/Plugins/Experimental/PCG/Source/PCG/Private/Helpers/PCGActorHelpers.cpp
+++ b/Engine/Plugins/Experimental/PCG/Source/PCG/Private/Helpers/PCGActorHelpers.cpp
@@ -160,9 +160,7 @@ UInstancedStaticMeshComponent* UPCGActorHelpers::GetOrCreateISMC(AActor* InTarge
 
 bool UPCGActorHelpers::DeleteActors(UWorld* World, const TArray<TSoftObjectPtr<AActor>>& ActorsToDelete)
 {
-	check(World);
-
-	if (ActorsToDelete.Num() == 0)
+	if (!World || ActorsToDelete.Num() == 0)
 	{
 		return true;
 	}
@@ -174,59 +172,81 @@ bool UPCGActorHelpers::DeleteActors(UWorld* World, const TArray<TSoftObjectPtr<A
 		GEditor->ResetAllSelectionSets();
 	}*/
 
-	UWorldPartition* WorldPartition = World ? World->GetWorldPartition() : nullptr;
-
-	if (WorldPartition)
+	/** Note: the following code block is commented out because this is currently getting hit when partition actors & partitioned components are unloaded.
+	* Destroying the actors is fine, but trying to delete them from SCC is a major issue and it can hang UE.
+	* There are multiple code paths calling this method that would need to be revisited if we want to control this a bit better
+	* or we can centralize these kind of requirements in the subsystem - this could then be taken care of in a single place, knowing the actual setting.
+	*/
+	//UWorldPartition* WorldPartition = World->GetWorldPartition();
+	//if (!PCGHelpers::IsRuntimeOrPIE() && WorldPartition)
+	//{
+	//	TArray<FString> PackagesToDeleteFromSCC;
+	//	TSet<UPackage*> PackagesToCleanup;
+
+	//	for (const TSoftObjectPtr<AActor>& ManagedActor : ActorsToDelete)
+	//	{
+	//		// If actor is loaded, just remove from world and keep track of package to cleanup
+	//		if (AActor* Actor = ManagedActor.Get())
+	//		{
+	//			if (UPackage* ActorPackage = Actor->GetExternalPackage())
+	//			{
+	//				PackagesToCleanup.Emplace(ActorPackage);
+	//			}
+	//			
+	//			if (Actor->GetWorld() == World)
+	//			{
+	//				World->DestroyActor(Actor);
+	//			}
+	//			else
+	//			{
+	//				// If we're here and the world is null, then the actor has either been destroyed already or it will be picked up by GC by design.
+	//				// Otherwise, we have bigger issues, something is very wrong.
+	//				check(Actor->GetWorld() == nullptr);
+	//			}
+	//		}
+	//		// Otherwise, get from World Partition.
+	//		// Note that it is possible that some actors don't exist anymore, so a null here is not a critical condition
+	//		else if (const FWorldPartitionActorDesc* ActorDesc = WorldPartition->GetActorDesc(ManagedActor.ToSoftObjectPath()))
+	//		{
+	//			PackagesToDeleteFromSCC.Emplace(ActorDesc->GetActorPackage().ToString());
+	//			WorldPartition->RemoveActor(ActorDesc->GetGuid());
+	//		}
+	//	}
+
+	//	// Save currently loaded packages so they get deleted
+	//	if (PackagesToCleanup.Num() > 0)
+	//	{
+	//		ObjectTools::CleanupAfterSuccessfulDelete(PackagesToCleanup.Array(), /*bPerformReferenceCheck=*/true);
+	//	}
+
+	//	// Delete outstanding unloaded packages
+	//	if (PackagesToDeleteFromSCC.Num() > 0)
+	//	{
+	//		FPackageSourceControlHelper PackageHelper;
+	//		if (!PackageHelper.Delete(PackagesToDeleteFromSCC))
+	//		{
+	//			return false;
+	//		}
+	//	}
+	//}
+	//else
+#endif
 	{
-		TArray<FString> PackagesToDeleteFromSCC;
-		TSet<UPackage*> PackagesToCleanup;
-
+		// Not in editor, really unlikely to happen but might be slow
 		for (const TSoftObjectPtr<AActor>& ManagedActor : ActorsToDelete)
 		{
-			// If actor is loaded, just remove from world and keep track of package to cleanup
 			if (AActor* Actor = ManagedActor.Get())
 			{
-				if (UPackage* ActorPackage = Actor->GetExternalPackage())
+				if (Actor->GetWorld() == World)
 				{
-					PackagesToCleanup.Emplace(ActorPackage);
+					World->DestroyActor(Actor);
+				}
+				else
+				{
+					// If we're here and the world is null, then the actor has either been destroyed already or it will be picked up by GC by design.
+					// Otherwise, we have bigger issues, something is very wrong.
+					check(Actor->GetWorld() == nullptr);
 				}
-				
-				World->DestroyActor(Actor);
-			}
-			// Otherwise, get from World Partition.
-			// Note that it is possible that some actors don't exist anymore, so a null here is not a critical condition
-			else if (const FWorldPartitionActorDesc* ActorDesc = WorldPartition->GetActorDesc(ManagedActor.ToSoftObjectPath()))
-			{
-				PackagesToDeleteFromSCC.Emplace(ActorDesc->GetActorPackage().ToString());
-				WorldPartition->RemoveActor(ActorDesc->GetGuid());
-			}
-		}
-
-		// Save currently loaded packages so they get deleted
-		if (PackagesToCleanup.Num() > 0)
-		{
-			ObjectTools::CleanupAfterSuccessfulDelete(PackagesToCleanup.Array(), /*bPerformReferenceCheck=*/true);
-		}
-
-		// Delete outstanding unloaded packages
-		if (PackagesToDeleteFromSCC.Num() > 0)
-		{
-			FPackageSourceControlHelper PackageHelper;
-			if (!PackageHelper.Delete(PackagesToDeleteFromSCC))
-			{
-				return false;
-			}
-		}
-	}
-	else
-#endif
-	{
-		// Not in editor, really unlikely to happen but might be slow
-		for (const TSoftObjectPtr<AActor>& ManagedActor : ActorsToDelete)
-		{
-			if (ManagedActor.Get())
-			{
-				World->DestroyActor(ManagedActor.Get());
 			}
 		}
 	}
