Temporary fix for scaling issue when spawing non-merged actors. Will be superseded once this issue is fixed & we support multiple mesh support.

diff --git a/Engine/Plugins/Experimental/PCG/Source/PCG/Private/Elements/PCGSpawnActor.cpp b/Engine/Plugins/Experimental/PCG/Source/PCG/Private/Elements/PCGSpawnActor.cpp
--- a/Engine/Plugins/Experimental/PCG/Source/PCG/Private/Elements/PCGSpawnActor.cpp
+++ b/Engine/Plugins/Experimental/PCG/Source/PCG/Private/Elements/PCGSpawnActor.cpp
@@ -215,6 +215,8 @@ bool FPCGSpawnActorElement::ExecuteInternal(FPCGContext* Context) const
 					for (const FPCGPoint& Point : Points)
 					{
 						AActor* GeneratedActor = TargetActor->GetWorld()->SpawnActor(Settings->TemplateActorClass, &Point.Transform, SpawnParams);
+						// HACK: until UE-62747 is fixed, we have to force set the scale after spawning the actor
+						GeneratedActor->SetActorRelativeScale3D(Point.Transform.GetScale3D());
 						GeneratedActor->Tags.Append(NewActorTags);
 						GeneratedActor->AttachToActor(TargetActor, FAttachmentTransformRules::KeepWorldTransform);
 
