Added BP accessor to transformed point bounds

diff --git a/Engine/Plugins/Experimental/PCG/Source/PCG/Private/Helpers/PCGBlueprintHelpers.cpp b/Engine/Plugins/Experimental/PCG/Source/PCG/Private/Helpers/PCGBlueprintHelpers.cpp
--- a/Engine/Plugins/Experimental/PCG/Source/PCG/Private/Helpers/PCGBlueprintHelpers.cpp
+++ b/Engine/Plugins/Experimental/PCG/Source/PCG/Private/Helpers/PCGBlueprintHelpers.cpp
@@ -91,6 +91,11 @@ FVector UPCGBlueprintHelpers::GetLocalCenter(const FPCGPoint& InPoint)
 	return InPoint.GetLocalCenter();
 }
 
+FBox UPCGBlueprintHelpers::GetTransformedBounds(const FPCGPoint& InPoint)
+{
+	return FBox(InPoint.BoundsMin, InPoint.BoundsMax).TransformBy(InPoint.Transform);
+}
+
 FBox UPCGBlueprintHelpers::GetActorBoundsPCG(AActor* InActor, bool bIgnorePCGCreatedComponents)
 {
 	return PCGHelpers::GetActorBounds(InActor, bIgnorePCGCreatedComponents);
diff --git a/Engine/Plugins/Experimental/PCG/Source/PCG/Public/Helpers/PCGBlueprintHelpers.h b/Engine/Plugins/Experimental/PCG/Source/PCG/Public/Helpers/PCGBlueprintHelpers.h
--- a/Engine/Plugins/Experimental/PCG/Source/PCG/Public/Helpers/PCGBlueprintHelpers.h
+++ b/Engine/Plugins/Experimental/PCG/Source/PCG/Public/Helpers/PCGBlueprintHelpers.h
@@ -57,6 +57,9 @@ public:
 	UFUNCTION(BlueprintCallable, Category = "PCG|Points", meta = (ScriptMethod))
 	static FVector GetLocalCenter(const FPCGPoint& InPoint);
 
+	UFUNCTION(BlueprintCallable, Category = "PCG|Points", meta = (ScriptMethod))
+	static FBox GetTransformedBounds(const FPCGPoint& InPoint);
+
 	UFUNCTION(BlueprintCallable, Category = "PCG|Helpers", meta = (ScriptMethod))
 	static FBox GetActorBoundsPCG(AActor* InActor, bool bIgnorePCGCreatedComponents = true);
 
