Refactored PCG element execute to make time-slicing easier.

diff --git a/Engine/Plugins/Experimental/PCG/Source/PCG/Private/Elements/PCGExecuteBlueprint.cpp b/Engine/Plugins/Experimental/PCG/Source/PCG/Private/Elements/PCGExecuteBlueprint.cpp
--- a/Engine/Plugins/Experimental/PCG/Source/PCG/Private/Elements/PCGExecuteBlueprint.cpp
+++ b/Engine/Plugins/Experimental/PCG/Source/PCG/Private/Elements/PCGExecuteBlueprint.cpp
@@ -615,6 +615,15 @@ void UPCGBlueprintElement::LoopNTimes(FPCGContext& InContext, int64 NumIteration
 	});
 }
 
+FPCGBlueprintExecutionContext::~FPCGBlueprintExecutionContext()
+{
+	if (BlueprintElementInstance)
+	{
+		BlueprintElementInstance->RemoveFromRoot();
+		BlueprintElementInstance = nullptr;
+	}
+}
+
 FPCGContext* FPCGExecuteBlueprintElement::Initialize(const FPCGDataCollection& InputData, TWeakObjectPtr<UPCGComponent> SourceComponent, const UPCGNode* Node)
 {
 	TRACE_CPUPROFILER_EVENT_SCOPE(FPCGExecuteBlueprintElement::Initialize);
@@ -627,6 +636,7 @@ FPCGContext* FPCGExecuteBlueprintElement::Initialize(const FPCGDataCollection& I
 	if (Settings && Settings->BlueprintElementInstance)
 	{
 		Context->BlueprintElementInstance = CastChecked<UPCGBlueprintElement>(StaticDuplicateObject(Settings->BlueprintElementInstance, GetTransientPackage(), FName()));
+		Context->BlueprintElementInstance->AddToRoot();
 #if !WITH_EDITOR
 		if (SourceComponent.IsValid() && SourceComponent->GetOwner())
 		{
@@ -655,9 +665,10 @@ bool FPCGExecuteBlueprintElement::IsCacheable(const UPCGSettings* InSettings) co
 	}
 }
 
-bool FPCGExecuteBlueprintElement::CanExecuteOnlyOnMainThread(const UPCGSettings* InSettings) const
+bool FPCGExecuteBlueprintElement::CanExecuteOnlyOnMainThread(FPCGContext* Context) const
 {
-	const UPCGBlueprintSettings* BPSettings = Cast<const UPCGBlueprintSettings>(InSettings);
+	check(Context);
+	const UPCGBlueprintSettings* BPSettings = Context->GetInputSettings<UPCGBlueprintSettings>();
 	if (BPSettings && BPSettings->BlueprintElementInstance)
 	{
 		return !BPSettings->BlueprintElementInstance->bCanBeMultithreaded;
diff --git a/Engine/Plugins/Experimental/PCG/Source/PCG/Private/Elements/PCGStaticMeshSpawner.cpp b/Engine/Plugins/Experimental/PCG/Source/PCG/Private/Elements/PCGStaticMeshSpawner.cpp
--- a/Engine/Plugins/Experimental/PCG/Source/PCG/Private/Elements/PCGStaticMeshSpawner.cpp
+++ b/Engine/Plugins/Experimental/PCG/Source/PCG/Private/Elements/PCGStaticMeshSpawner.cpp
@@ -28,10 +28,11 @@ FPCGElementPtr UPCGStaticMeshSpawnerSettings::CreateElement() const
 	return MakeShared<FPCGStaticMeshSpawnerElement>();
 }
 
-bool FPCGStaticMeshSpawnerElement::ExecuteInternal(FPCGContext* Context) const
+bool FPCGStaticMeshSpawnerElement::PrepareDataInternal(FPCGContext* InContext) const
 {
-	TRACE_CPUPROFILER_EVENT_SCOPE(FPCGStaticMeshSpawnerElement::Execute);
+	TRACE_CPUPROFILER_EVENT_SCOPE(FPCGStaticMeshSpawnerElement::PrepareDataInternal);
 	// TODO : time-sliced implementation
+	FPCGStaticMeshSpawnerContext* Context = static_cast<FPCGStaticMeshSpawnerContext*>(InContext);
 	const UPCGStaticMeshSpawnerSettings* Settings = Context->GetInputSettings<UPCGStaticMeshSpawnerSettings>();
 	check(Settings);
 
@@ -42,11 +43,8 @@ bool FPCGStaticMeshSpawnerElement::ExecuteInternal(FPCGContext* Context) const
 	}
 
 	// perform mesh selection
-	TArray<FPCGMeshInstanceList> MeshInstances;
 	TArray<FPCGTaggedData> Inputs = Context->InputData.GetInputs();
 	TArray<FPCGTaggedData>& Outputs = Context->OutputData.TaggedData;
-	
-	FPCGPackedCustomData PackedCustomData;
 
 	// Forward any non-input data
 	Outputs.Append(Context->InputData.GetAllSettings());
@@ -75,8 +73,8 @@ bool FPCGStaticMeshSpawnerElement::ExecuteInternal(FPCGContext* Context) const
 
 		if (bOutputPinConnected || Settings->bForceConnectOutput)
 		{
-			FPCGTaggedData& Output = Outputs.Add_GetRef(Input); 
-			
+			FPCGTaggedData& Output = Outputs.Add_GetRef(Input);
+
 			OutputPointData = NewObject<UPCGPointData>();
 			OutputPointData->InitializeFromData(SpatialData);
 
@@ -91,27 +89,74 @@ bool FPCGStaticMeshSpawnerElement::ExecuteInternal(FPCGContext* Context) const
 			Output.Data = OutputPointData;
 		}
 
+		TArray<FPCGMeshInstanceList> MeshInstances;
 		Settings->MeshSelectorInstance->SelectInstances(*Context, Settings, SpatialData, MeshInstances, OutputPointData);
 
-		for (const FPCGMeshInstanceList& InstanceList : MeshInstances)
+		TArray<FPCGPackedCustomData> PackedCustomData;
+		PackedCustomData.SetNum(MeshInstances.Num());
+		if (Settings->InstancePackerInstance)
 		{
-			if (Settings->InstancePackerInstance)
+			for(int32 InstanceListIndex = 0; InstanceListIndex < MeshInstances.Num(); ++InstanceListIndex)
 			{
-				PackedCustomData.CustomData.Reset();
-				PackedCustomData.NumCustomDataFloats = 0;
-
-				Settings->InstancePackerInstance->PackInstances(*Context, SpatialData, InstanceList, PackedCustomData);
+				Settings->InstancePackerInstance->PackInstances(*Context, SpatialData, MeshInstances[InstanceListIndex], PackedCustomData[InstanceListIndex]);
 			}
-
-			SpawnStaticMeshInstances(Context, InstanceList, TargetActor, PackedCustomData);
 		}
 
-		MeshInstances.Reset();
+		FPCGStaticMeshSpawnerContext::FPackedInstanceListData& InstanceListData = Context->MeshInstancesData.Emplace_GetRef();
+		InstanceListData.SpatialData = SpatialData;
+		InstanceListData.MeshInstances = MoveTemp(MeshInstances);
+		InstanceListData.PackedCustomData = MoveTemp(PackedCustomData);
 	}
 
 	return true;
 }
 
+bool FPCGStaticMeshSpawnerElement::ExecuteInternal(FPCGContext* InContext) const
+{
+	TRACE_CPUPROFILER_EVENT_SCOPE(FPCGStaticMeshSpawnerElement::Execute);
+	FPCGStaticMeshSpawnerContext* Context = static_cast<FPCGStaticMeshSpawnerContext*>(InContext);
+	const UPCGStaticMeshSpawnerSettings* Settings = Context->GetInputSettings<UPCGStaticMeshSpawnerSettings>();
+	check(Settings);
+
+	while(!Context->MeshInstancesData.IsEmpty())
+	{
+		const FPCGStaticMeshSpawnerContext::FPackedInstanceListData& InstanceList = Context->MeshInstancesData.Last();
+		check(InstanceList.MeshInstances.Num() == InstanceList.PackedCustomData.Num());
+
+		if (InstanceList.SpatialData->TargetActor.IsValid())
+		{
+			for (int32 DataIndex = 0; DataIndex < InstanceList.MeshInstances.Num(); ++DataIndex)
+			{
+				SpawnStaticMeshInstances(Context, InstanceList.MeshInstances[DataIndex], InstanceList.SpatialData->TargetActor.Get(), InstanceList.PackedCustomData[DataIndex]);
+			}
+		}
+
+		Context->MeshInstancesData.RemoveAtSwap(Context->MeshInstancesData.Num() - 1);
+
+		if (Context->ShouldStop())
+		{
+			break;
+		}
+	}
+
+	return Context->MeshInstancesData.IsEmpty();
+}
+
+FPCGContext* FPCGStaticMeshSpawnerElement::Initialize(const FPCGDataCollection& InputData, TWeakObjectPtr<UPCGComponent> SourceComponent, const UPCGNode* Node)
+{
+	FPCGStaticMeshSpawnerContext* Context = new FPCGStaticMeshSpawnerContext();
+	Context->InputData = InputData;
+	Context->SourceComponent = SourceComponent;
+	Context->Node = Node;
+
+	return Context;
+}
+
+bool FPCGStaticMeshSpawnerElement::CanExecuteOnlyOnMainThread(FPCGContext* Context) const
+{
+	return Context->CurrentPhase == EPCGExecutionPhase::Execute;
+}
+
 void FPCGStaticMeshSpawnerElement::SpawnStaticMeshInstances(FPCGContext* Context, const FPCGMeshInstanceList& InstanceList, AActor* TargetActor, const FPCGPackedCustomData& PackedCustomData) const
 {
 	// Populate the (H)ISM from the previously prepared entries
diff --git a/Engine/Plugins/Experimental/PCG/Source/PCG/Private/Graph/PCGGraphExecutor.cpp b/Engine/Plugins/Experimental/PCG/Source/PCG/Private/Graph/PCGGraphExecutor.cpp
--- a/Engine/Plugins/Experimental/PCG/Source/PCG/Private/Graph/PCGGraphExecutor.cpp
+++ b/Engine/Plugins/Experimental/PCG/Source/PCG/Private/Graph/PCGGraphExecutor.cpp
@@ -85,7 +85,7 @@ FPCGTaskId FPCGGraphExecutor::Schedule(UPCGGraph* Graph, UPCGComponent* SourceCo
 
 	// Prepare scheduled task that will be promoted in the next Execute call.
 	if (CompiledTasks.Num() > 0)
-	{
+	{ 
 		check(CompiledTasks[0].Node == Graph->GetInputNode());
 
 		// Setup fetch task on input node
@@ -221,6 +221,7 @@ void FPCGGraphExecutor::Execute()
 	bool bAnyTaskEnded = false;
 
 	const double StartTime = FPlatformTime::Seconds();
+	const double EndTime = StartTime + (CVarTimePerFrame.GetValueOnAnyThread() / 1000.0);
 	const int32 MaxNumThreads = FMath::Max(0, FMath::Min(FPlatformMisc::NumberOfCoresIncludingHyperthreads() - 2, CVarMaxNumTasks.GetValueOnAnyThread() - 1));
 	const bool bAllowMultiDispatch = CVarGraphMultithreading.GetValueOnAnyThread();
 
@@ -231,7 +232,7 @@ void FPCGGraphExecutor::Execute()
 	while(ReadyTasks.Num() > 0 || ActiveTasks.Num() > 0 || SleepingTasks.Num() > 0)
 	{
 		// First: if we have free resources, move ready tasks to the active tasks
-		bool bMainThreadAvailable = (ActiveTasks.Num() == 0 || !ActiveTasks.Last().Element->CanExecuteOnlyOnMainThread(ActiveTasks.Last().Context->GetInputSettings<UPCGSettings>()));
+		bool bMainThreadAvailable = (ActiveTasks.Num() == 0 || !ActiveTasks[0].Element->CanExecuteOnlyOnMainThread(ActiveTasks[0].Context.Get()));
 		int32 NumAvailableThreads = FMath::Max(0, MaxNumThreads - CurrentlyUsedThreads);
 
 		const bool bMainThreadWasAvailable = bMainThreadAvailable;
@@ -279,18 +280,16 @@ void FPCGGraphExecutor::Execute()
 				}
 
 				// Validate that we can start this task now
-				const bool bIsMainThreadTask = SleepingTask.Element->CanExecuteOnlyOnMainThread(SleepingTask.Context->GetInputSettings<UPCGSettings>());
+				const bool bIsMainThreadTask = SleepingTask.Element->CanExecuteOnlyOnMainThread(SleepingTask.Context.Get());
 
 				if (!bIsMainThreadTask || bMainThreadAvailable)
 				{
-					const bool bInsertLast = (bIsMainThreadTask || bMainThreadAvailable || ActiveTasks.Num() == 0);
-					if (bInsertLast)
-					{
-						ActiveTasks.Emplace(MoveTemp(SleepingTask));
-					}
-					else
+					ActiveTasks.Emplace(MoveTemp(SleepingTask));
+
+					// Move task up front if it has to run on the main thread
+					if (bIsMainThreadTask && ActiveTasks.Num() > 1)
 					{
-						ActiveTasks.Insert(MoveTemp(SleepingTask), ActiveTasks.Num() - 1);
+						ActiveTasks.Swap(0, ActiveTasks.Num() - 1);
 					}
 
 					TaskDispatchBookKeeping(bIsMainThreadTask);
@@ -328,6 +327,7 @@ void FPCGGraphExecutor::Execute()
 				// At this point, if the task doesn't have an element, we will never be able to execute it, so we can drop it.
 				if (!Task.Element)
 				{
+					check(!Task.Context);
 					ReadyTasks.RemoveAtSwap(ReadyTaskIndex);
 					continue;
 				}
@@ -356,6 +356,7 @@ void FPCGGraphExecutor::Execute()
 					// Fast-forward cached result to stored results
 					FPCGTaskId SkippedTaskId = Task.NodeId;
 					StoreResults(SkippedTaskId, CachedOutput);
+					delete Task.Context;
 					ReadyTasks.RemoveAtSwap(ReadyTaskIndex);
 					QueueNextTasks(SkippedTaskId);
 					bAnyTaskEnded = true;
@@ -363,29 +364,37 @@ void FPCGGraphExecutor::Execute()
 					continue;
 				}
 
+				// Allocate context if not previously done
+				if (!Task.Context)
+				{
+					Task.Context = Task.Element->Initialize(TaskInput, Task.SourceComponent, Task.Node);
+					Task.Context->TaskId = Task.NodeId;
+				}
+
 				// Validate that we can start this task now
-				const bool bIsMainThreadTask = Task.Element->CanExecuteOnlyOnMainThread(TaskSettings);
+				const bool bIsMainThreadTask = Task.Element->CanExecuteOnlyOnMainThread(Task.Context);
 
 				if (!bIsMainThreadTask || bMainThreadAvailable)
 				{
-					const bool bInsertLast = (bIsMainThreadTask || bMainThreadAvailable || ActiveTasks.Num() == 0);
-					FPCGGraphActiveTask& ActiveTask = (bInsertLast ? ActiveTasks.Emplace_GetRef() : ActiveTasks.EmplaceAt_GetRef(ActiveTasks.Num() - 1));
-
+					FPCGGraphActiveTask& ActiveTask = ActiveTasks.Emplace_GetRef();
 					ActiveTask.Element = Task.Element;
 					ActiveTask.NodeId = Task.NodeId;
-
-					FPCGContext* Context = Task.Element->Initialize(TaskInput, Task.SourceComponent, Task.Node);
-					Context->TaskId = Task.NodeId;
-					ActiveTask.Context = TUniquePtr<FPCGContext>(Context);
+					ActiveTask.Context = TUniquePtr<FPCGContext>(Task.Context);
 
 #if WITH_EDITOR
 					if (bResultAlreadyInCache)
 					{
 						ActiveTask.bIsBypassed = true;
-						Context->OutputData = CachedOutput;
+						ActiveTask.Context->OutputData = CachedOutput;
 					}
 #endif
 
+					// Move the task up front if it needs to run on the main thread
+					if (bIsMainThreadTask && ActiveTasks.Num() > 1)
+					{
+						ActiveTasks.Swap(0, ActiveTasks.Num() - 1);
+					}
+
 					TaskDispatchBookKeeping(bIsMainThreadTask);
 
 					ReadyTasks.RemoveAtSwap(ReadyTaskIndex);
@@ -422,15 +431,19 @@ void FPCGGraphExecutor::Execute()
 
 		{
 			TRACE_CPUPROFILER_EVENT_SCOPE(FPCGGraphExecutor::Execute::StartFutures);
-			for(int32 ExecutionIndex = 0; ExecutionIndex < ActiveTasks.Num() - 1; ++ExecutionIndex)
+			for(int32 ExecutionIndex = 1; ExecutionIndex < ActiveTasks.Num(); ++ExecutionIndex)
 			{
 				FPCGGraphActiveTask& ActiveTask = ActiveTasks[ExecutionIndex];
-				check(!ActiveTask.Context || !ActiveTask.Context->bIsPaused);
+				check(!ActiveTask.Context->bIsPaused);
 
 	#if WITH_EDITOR
 				if(!ActiveTask.bIsBypassed)
 	#endif
 				{
+					check(!ActiveTask.Element->CanExecuteOnlyOnMainThread(ActiveTask.Context.Get()));
+					ActiveTask.Context->EndTime = EndTime;
+					ActiveTask.Context->bIsRunningOnMainThread = false;
+
 					Futures.Emplace(ExecutionIndex, Async(EAsyncExecution::ThreadPool, [&ActiveTask]()
 					{
 						return ActiveTask.Element->Execute(ActiveTask.Context.Get());
@@ -488,8 +501,11 @@ void FPCGGraphExecutor::Execute()
 			if (!ActiveTasks.IsEmpty())
 			{
 				TRACE_CPUPROFILER_EVENT_SCOPE(FPCGGraphExecutor::Execute::ExecuteTasks::MainThreadTask);
-				FPCGGraphActiveTask& MainThreadTask = ActiveTasks.Last();
-				check(!MainThreadTask.Context || !MainThreadTask.Context->bIsPaused);
+				FPCGGraphActiveTask& MainThreadTask = ActiveTasks[0];
+				check(!MainThreadTask.Context->bIsPaused);
+				MainThreadTask.Context->EndTime = EndTime;
+				MainThreadTask.Context->bIsRunningOnMainThread = true;
+
 #if WITH_EDITOR
 				if(MainThreadTask.bIsBypassed || MainThreadTask.Element->Execute(MainThreadTask.Context.Get()))
 #else
@@ -497,19 +513,14 @@ void FPCGGraphExecutor::Execute()
 #endif
 				{
 					bMainTaskDone = true;
-					PostTaskExecute(ActiveTasks.Num() - 1);
 				}
 			}
 
-			// If the main thread task is done, it was removed from ActiveTasks, and ActiveTasks should contains only futures (if there are some)
-			// Otherwise, ActiveTasks.Num() - 1 still point to the main thread tasks, so futures indexes are within [0, ActiveTasks.Num() - 2].
-			int32 FuturesLastIndex = bMainTaskDone ? ActiveTasks.Num() - 1 : ActiveTasks.Num() - 2;
-			if(FuturesLastIndex >= 0)
 			{
 				TRACE_CPUPROFILER_EVENT_SCOPE(FPCGGraphExecutor::Execute::ExecuteTasks::WaitForFutures);
 
-				// Then wait after all futures
-				for (int32 ExecutionIndex = FuturesLastIndex; ExecutionIndex >= 0; --ExecutionIndex)
+				// Then wait after all futures - start from the back so we can more easily manage the ActiveTasks array
+				for(int32 ExecutionIndex = ActiveTasks.Num() - 1; ExecutionIndex > 0; --ExecutionIndex)
 				{
 					bool bTaskDone = false;
 					// Wait on the future if any
@@ -525,6 +536,11 @@ void FPCGGraphExecutor::Execute()
 					}
 				}
 			}
+
+			if (bMainTaskDone)
+			{
+				PostTaskExecute(0);
+			}
 		}
 
 		// Any paused tasks at that point should relinquish their resources
@@ -533,9 +549,13 @@ void FPCGGraphExecutor::Execute()
 			TRACE_CPUPROFILER_EVENT_SCOPE(FPCGGraphExecutor::Execute::ExecuteTasks::CheckSleepingTasks);
 			for (int32 ActiveTaskIndex = ActiveTasks.Num() - 1; ActiveTaskIndex >= 0; --ActiveTaskIndex)
 			{
-				if (ActiveTasks[ActiveTaskIndex].Context->bIsPaused)
+				FPCGGraphActiveTask& ActiveTask = ActiveTasks[ActiveTaskIndex];
+				
+				// Any task that asks to be paused or now needs to run on the main thread but doesn't have that slot currently will be moved to the sleeping queue
+				const bool bTaskShouldBePutAside = (ActiveTask.Context->bIsPaused || (ActiveTaskIndex > 0 && ActiveTask.Element->CanExecuteOnlyOnMainThread(ActiveTask.Context.Get())));
+
+				if(bTaskShouldBePutAside)
 				{
-					FPCGGraphActiveTask& ActiveTask = ActiveTasks[ActiveTaskIndex];
 					check(ActiveTask.Context->NumAvailableTasks > 0);
 					CurrentlyUsedThreads -= ActiveTask.Context->NumAvailableTasks;
 					ActiveTask.Context->NumAvailableTasks = 0;
@@ -548,8 +568,7 @@ void FPCGGraphExecutor::Execute()
 
 		check(ActiveTasks.Num() > 0 || CurrentlyUsedThreads == 0);
 
-		const double EndLoopTime = FPlatformTime::Seconds();
-		if (static_cast<float>((EndLoopTime - StartTime) * 1000) > CVarTimePerFrame.GetValueOnAnyThread())
+		if (FPlatformTime::Seconds() >= EndTime)
 		{
 			break;
 		}
@@ -561,7 +580,7 @@ void FPCGGraphExecutor::Execute()
 		// TODO: this is fine and will make sure any intermediate data is properly
 		// garbage collected, however, this goes a bit against our goals if we want to
 		// keep a cache of intermediate results.
-		if (ReadyTasks.Num() == 0 && ActiveTasks.Num() == 0 && Tasks.Num() == 0)
+		if (ReadyTasks.Num() == 0 && ActiveTasks.Num() == 0 && SleepingTasks.Num() == 0 && Tasks.Num() == 0)
 		{
 			check(TaskSuccessors.IsEmpty());
 			ClearResults();
diff --git a/Engine/Plugins/Experimental/PCG/Source/PCG/Private/Graph/PCGGraphExecutor.h b/Engine/Plugins/Experimental/PCG/Source/PCG/Private/Graph/PCGGraphExecutor.h
--- a/Engine/Plugins/Experimental/PCG/Source/PCG/Private/Graph/PCGGraphExecutor.h
+++ b/Engine/Plugins/Experimental/PCG/Source/PCG/Private/Graph/PCGGraphExecutor.h
@@ -43,6 +43,7 @@ struct FPCGGraphTask
 	const UPCGNode* Node = nullptr;
 	TWeakObjectPtr<UPCGComponent> SourceComponent = nullptr;
 	FPCGElementPtr Element; // Added to have tasks that aren't node-bound
+	FPCGContext* Context = nullptr;
 	FPCGTaskId NodeId = InvalidPCGTaskId;
 };
 
@@ -164,7 +165,7 @@ class FPCGFetchInputElement : public FSimplePCGElement
 {
 public:
 	virtual bool IsCacheable(const UPCGSettings* InSettings) const override { return false; }
-	virtual bool CanExecuteOnlyOnMainThread(const UPCGSettings* InSettings) const override { return true; }
+	virtual bool CanExecuteOnlyOnMainThread(FPCGContext* Context) const override { return true; }
 
 protected:
 	virtual bool ExecuteInternal(FPCGContext* Context) const override;
@@ -176,7 +177,7 @@ class FPCGGenericElement : public FSimplePCGElement
 public:
 	FPCGGenericElement(TFunction<bool(FPCGContext*)> InOperation);
 	virtual bool IsCacheable(const UPCGSettings* InSettings) const override { return false; }
-	virtual bool CanExecuteOnlyOnMainThread(const UPCGSettings* InSettings) const override { return true; }
+	virtual bool CanExecuteOnlyOnMainThread(FPCGContext* Context) const override { return true; }
 
 protected:
 	// Important note: generic elements must always be run on the main thread
diff --git a/Engine/Plugins/Experimental/PCG/Source/PCG/Private/PCGContext.cpp b/Engine/Plugins/Experimental/PCG/Source/PCG/Private/PCGContext.cpp
--- a/Engine/Plugins/Experimental/PCG/Source/PCG/Private/PCGContext.cpp
+++ b/Engine/Plugins/Experimental/PCG/Source/PCG/Private/PCGContext.cpp
@@ -31,4 +31,9 @@ FString FPCGContext::GetTaskName() const
 FString FPCGContext::GetComponentName() const
 {
 	return SourceComponent.IsValid() && SourceComponent->GetOwner() ? SourceComponent->GetOwner()->GetFName().ToString() : TEXT("Non-PCG Component");
+}
+
+bool FPCGContext::ShouldStop() const
+{
+	return FPlatformTime::Seconds() > EndTime;
 }
\ No newline at end of file
diff --git a/Engine/Plugins/Experimental/PCG/Source/PCG/Private/PCGElement.cpp b/Engine/Plugins/Experimental/PCG/Source/PCG/Private/PCGElement.cpp
--- a/Engine/Plugins/Experimental/PCG/Source/PCG/Private/PCGElement.cpp
+++ b/Engine/Plugins/Experimental/PCG/Source/PCG/Private/PCGElement.cpp
@@ -12,8 +12,76 @@
 bool IPCGElement::Execute(FPCGContext* Context) const
 {
 	TRACE_CPUPROFILER_EVENT_SCOPE(IPCGElement::Execute);
-	check(Context && Context->NumAvailableTasks > 0);
+	check(Context && Context->NumAvailableTasks > 0 && Context->CurrentPhase < EPCGExecutionPhase::Done);
+	check(Context->bIsRunningOnMainThread || !CanExecuteOnlyOnMainThread(Context));
 
+	while (Context->CurrentPhase != EPCGExecutionPhase::Done)
+	{
+		bool bExecutionPostponed = false;
+
+		switch (Context->CurrentPhase)
+		{
+			case EPCGExecutionPhase::NotExecuted: // Fall-through
+			{
+				PreExecute(Context);
+				break;
+			}
+
+			case EPCGExecutionPhase::PrepareData:
+			{
+				FScopedCallTimer CallTimer(*this, Context);
+				if (PrepareDataInternal(Context))
+				{
+					Context->CurrentPhase = EPCGExecutionPhase::Execute;
+				}
+				else
+				{
+					bExecutionPostponed = true;
+				}
+				break;
+			}
+
+			case EPCGExecutionPhase::Execute:
+			{
+				FScopedCallTimer CallTimer(*this, Context);
+				if (ExecuteInternal(Context))
+				{
+					Context->CurrentPhase = EPCGExecutionPhase::PostExecute;
+				}
+				else
+				{
+					bExecutionPostponed = true;
+				}
+				break;
+			}
+
+			case EPCGExecutionPhase::PostExecute:
+			{
+				PostExecute(Context);
+				break;
+			}
+
+			default: // should not happen
+			{
+				check(0);
+				break;
+			}
+		}
+
+		if (bExecutionPostponed || 
+			Context->ShouldStop() ||
+			(!Context->bIsRunningOnMainThread && CanExecuteOnlyOnMainThread(Context))) // phase change might require access to main thread
+		{
+			break;
+		}
+	}
+
+	return Context->CurrentPhase == EPCGExecutionPhase::Done;
+}
+
+void IPCGElement::PreExecute(FPCGContext* Context) const
+{
+	// Check for early outs (task cancelled + node disabled)
 	// Early out to stop execution
 	if (Context->InputData.bCancelExecution || (!Context->SourceComponent.IsExplicitlyNull() && !Context->SourceComponent.IsValid()))
 	{
@@ -21,23 +89,32 @@ bool IPCGElement::Execute(FPCGContext* Context) const
 
 		if (IsCancellable())
 		{
-			return true;
+			// Skip task completely
+			Context->CurrentPhase = EPCGExecutionPhase::Done;
+			return;
 		}
 	}
 
+	// Prepare to move to prepare data phase
+	Context->CurrentPhase = EPCGExecutionPhase::PrepareData;
+
 	const UPCGSettings* Settings = Context->GetInputSettings<UPCGSettings>();
-	if (Settings && Settings->ExecutionMode == EPCGSettingsExecutionMode::Disabled)
+	if (!Settings)
+	{
+		return;
+	}
+
+	if (Settings->ExecutionMode == EPCGSettingsExecutionMode::Disabled)
 	{
-		//Pass-through
+		//Pass-through - no execution
 		Context->OutputData = Context->InputData;
-		CleanupAndValidateOutput(Context);
-		return true;
+		Context->CurrentPhase = EPCGExecutionPhase::PostExecute;
 	}
 	else
 	{
+		// Perform input filtering
 		/** TODO - Placeholder feature */
-		TArray<FPCGTaggedData> BypassedTaggedData;
-		if (Settings && !Settings->FilterOnTags.IsEmpty())
+		if (!Settings->FilterOnTags.IsEmpty())
 		{
 			// Move any of the inputs that don't have the tags to the outputs as a pass-through
 			// NOTE: this breaks a bit the ordering of inputs, however, there's no obvious way around it
@@ -49,7 +126,6 @@ bool IPCGElement::Execute(FPCGContext* Context) const
 					if (Settings->bPassThroughFilteredOutInputs)
 					{
 						Context->OutputData.TaggedData.Add(TaggedData);
-						BypassedTaggedData.Add(TaggedData);
 					}
 				}
 				else // input has the required tags
@@ -58,89 +134,57 @@ bool IPCGElement::Execute(FPCGContext* Context) const
 				}
 			}
 
-			Context->InputData.TaggedData = FilteredTaggedData;
+			Context->InputData.TaggedData = MoveTemp(FilteredTaggedData);
+			Context->BypassedOutputCount = Context->OutputData.TaggedData.Num();
 		}
+	}
+}
 
-		bool bDone = false;
-
-#if WITH_EDITOR
-		if (ShouldLog())
-		{
-			UE_LOG(LogPCG, Verbose, TEXT("---------------------------------------"));
-		}
-#endif
-
-#if WITH_EDITOR
-		const double StartTime = FPlatformTime::Seconds();
-#endif
+bool IPCGElement::PrepareDataInternal(FPCGContext* Context) const
+{
+	return true;
+}
 
-		bDone = ExecuteInternal(Context);
+void IPCGElement::PostExecute(FPCGContext* Context) const
+{
+	// Cleanup and validate output
+	CleanupAndValidateOutput(Context);
 
 #if WITH_EDITOR
-		const double EndTime = FPlatformTime::Seconds();
-		const double ElapsedTime = EndTime - StartTime;
-		Context->ElapsedTime += ElapsedTime;
-		Context->ExecutionCount++;
-
-		// TODO: Is hardcoded, should be configurable
-		constexpr int MaxNumberOfTrackedTimers = 100;
-		{
-			FScopeLock Lock(&TimersLock);
-			if (Timers.Num() < MaxNumberOfTrackedTimers)
-			{
-				Timers.Add(ElapsedTime);
-			}
-			else
-			{
-				Timers[CurrentTimerIndex] = ElapsedTime;
-			}
-			CurrentTimerIndex = (CurrentTimerIndex + 1) % MaxNumberOfTrackedTimers;
-		}
+	PCGE_LOG(Log, "Executed in (%f)s and (%d) call(s)", Context->ElapsedTime, Context->ExecutionCount);
 #endif
 
-		if (bDone)
-		{
-			CleanupAndValidateOutput(Context);
-
-#if WITH_EDITOR
-			PCGE_LOG(Log, "Executed in (%f)s and (%d) call(s)", Context->ElapsedTime, Context->ExecutionCount);
-#endif
-		}
+	const UPCGSettings* Settings = Context->GetInputSettings<UPCGSettings>();
 
-		/** TODO - Placeholder feature */
-		if (bDone && Settings && !Settings->TagsAppliedOnOutput.IsEmpty())
+	// Apply tags on output
+	/** TODO - Placeholder feature */
+	if (Settings && !Settings->TagsAppliedOnOutput.IsEmpty())
+	{
+		for (int32 TaggedDataIdx = Context->BypassedOutputCount; TaggedDataIdx < Context->OutputData.TaggedData.Num(); ++TaggedDataIdx)
 		{
-			for (FPCGTaggedData& TaggedData : Context->OutputData.TaggedData)
-			{
-				if (!BypassedTaggedData.Contains(TaggedData))
-				{
-					TaggedData.Tags.Append(Settings->TagsAppliedOnOutput);
-				}
-			}
+			Context->OutputData.TaggedData[TaggedDataIdx].Tags.Append(Settings->TagsAppliedOnOutput);
 		}
+	}
 
+	// Additional debug things (check for duplicates),
 #if WITH_EDITOR
-		if (bDone && Settings)
-		{
-			if (Settings->DebugSettings.bCheckForDuplicates)
-			{
-				FPCGDataCollection ElementInputs = Context->InputData;
-				FPCGDataCollection ElementOutputs = Context->OutputData;
+	if (Settings && Settings->DebugSettings.bCheckForDuplicates)
+	{
+		FPCGDataCollection ElementInputs = Context->InputData;
+		FPCGDataCollection ElementOutputs = Context->OutputData;
 
-				Context->InputData = ElementOutputs;
-				Context->OutputData = FPCGDataCollection();
+		Context->InputData = ElementOutputs;
+		Context->OutputData = FPCGDataCollection();
 
-				PCGE_LOG(Verbose, "Performing remove duplicate points test (perf warning)");
-				PCGSelfPruningElement::Execute(Context, EPCGSelfPruningType::RemoveDuplicates, 0.0f, false);
+		PCGE_LOG(Verbose, "Performing remove duplicate points test (perf warning)");
+		PCGSelfPruningElement::Execute(Context, EPCGSelfPruningType::RemoveDuplicates, 0.0f, false);
 
-				Context->InputData = ElementInputs;
-				Context->OutputData = ElementOutputs;
-			}
-		}
+		Context->InputData = ElementInputs;
+		Context->OutputData = ElementOutputs;
+	}
 #endif
 
-		return bDone;
-	}
+	Context->CurrentPhase = EPCGExecutionPhase::Done;
 }
 
 #if WITH_EDITOR
@@ -217,6 +261,36 @@ void IPCGElement::CleanupAndValidateOutput(FPCGContext* Context) const
 	}
 }
 
+#if WITH_EDITOR
+IPCGElement::FScopedCallTimer::FScopedCallTimer(const IPCGElement& InOwner, FPCGContext* InContext)
+	: Owner(InOwner), Context(InContext)
+{
+	StartTime = FPlatformTime::Seconds();
+}
+
+IPCGElement::FScopedCallTimer::~FScopedCallTimer()
+{
+	const double EndTime = FPlatformTime::Seconds();
+	const double ElapsedTime = EndTime - StartTime;
+	Context->ElapsedTime += ElapsedTime;
+	Context->ExecutionCount++;
+
+	constexpr int MaxNumberOfTrackedTimers = 100;
+	{
+		FScopeLock Lock(&Owner.TimersLock);
+		if (Owner.Timers.Num() < MaxNumberOfTrackedTimers)
+		{
+			Owner.Timers.Add(ElapsedTime);
+		}
+		else
+		{
+			Owner.Timers[Owner.CurrentTimerIndex] = ElapsedTime;
+		}
+		Owner.CurrentTimerIndex = (Owner.CurrentTimerIndex + 1) % MaxNumberOfTrackedTimers;
+	}
+}
+#endif // WITH_EDITOR
+
 FPCGContext* FSimplePCGElement::Initialize(const FPCGDataCollection& InputData, TWeakObjectPtr<UPCGComponent> SourceComponent, const UPCGNode* Node)
 {
 	FPCGContext* Context = new FPCGContext();
diff --git a/Engine/Plugins/Experimental/PCG/Source/PCG/Public/Elements/Metadata/PCGMetadataOpElementBase.h b/Engine/Plugins/Experimental/PCG/Source/PCG/Public/Elements/Metadata/PCGMetadataOpElementBase.h
--- a/Engine/Plugins/Experimental/PCG/Source/PCG/Public/Elements/Metadata/PCGMetadataOpElementBase.h
+++ b/Engine/Plugins/Experimental/PCG/Source/PCG/Public/Elements/Metadata/PCGMetadataOpElementBase.h
@@ -114,7 +114,7 @@ class FPCGMetadataElementBase : public FSimplePCGElement
 {
 public:
 	// FPCGMetadataElementBase relies on StaticDuplicateObject (when we duplicate tagged data), so we cannot run outside of the main thread
-	virtual bool CanExecuteOnlyOnMainThread(const UPCGSettings* InSettings) const override { return true; }
+	virtual bool CanExecuteOnlyOnMainThread(FPCGContext*) const override { return true; }
 
 	virtual bool IsCacheable(const UPCGSettings* InSettings) const override { return true; }
 
diff --git a/Engine/Plugins/Experimental/PCG/Source/PCG/Public/Elements/Metadata/PCGMetadataRenameElement.h b/Engine/Plugins/Experimental/PCG/Source/PCG/Public/Elements/Metadata/PCGMetadataRenameElement.h
--- a/Engine/Plugins/Experimental/PCG/Source/PCG/Public/Elements/Metadata/PCGMetadataRenameElement.h
+++ b/Engine/Plugins/Experimental/PCG/Source/PCG/Public/Elements/Metadata/PCGMetadataRenameElement.h
@@ -36,7 +36,7 @@ public:
 class FPCGMetadataRenameElement : public FSimplePCGElement
 {
 public:
-	virtual bool CanExecuteOnlyOnMainThread(const UPCGSettings* InSettings) const override { return true; }
+	virtual bool CanExecuteOnlyOnMainThread(FPCGContext* Context) const override { return true; }
 
 protected:
 	virtual bool ExecuteInternal(FPCGContext* Context) const override;
diff --git a/Engine/Plugins/Experimental/PCG/Source/PCG/Public/Elements/PCGCreateAttribute.h b/Engine/Plugins/Experimental/PCG/Source/PCG/Public/Elements/PCGCreateAttribute.h
--- a/Engine/Plugins/Experimental/PCG/Source/PCG/Public/Elements/PCGCreateAttribute.h
+++ b/Engine/Plugins/Experimental/PCG/Source/PCG/Public/Elements/PCGCreateAttribute.h
@@ -34,7 +34,7 @@ protected:
 class FPCGCreateAttributeElement : public FSimplePCGElement
 {
 public:
-	virtual bool CanExecuteOnlyOnMainThread(const UPCGSettings* InSettings) const { return true; }
+	virtual bool CanExecuteOnlyOnMainThread(FPCGContext* Context) const { return true; }
 	virtual bool IsCacheable(const UPCGSettings* InSettings) const override { return false; }
 
 protected:
diff --git a/Engine/Plugins/Experimental/PCG/Source/PCG/Public/Elements/PCGExecuteBlueprint.h b/Engine/Plugins/Experimental/PCG/Source/PCG/Public/Elements/PCGExecuteBlueprint.h
--- a/Engine/Plugins/Experimental/PCG/Source/PCG/Public/Elements/PCGExecuteBlueprint.h
+++ b/Engine/Plugins/Experimental/PCG/Source/PCG/Public/Elements/PCGExecuteBlueprint.h
@@ -229,13 +229,15 @@ protected:
 
 struct FPCGBlueprintExecutionContext : public FPCGContext
 {
+	virtual ~FPCGBlueprintExecutionContext();
+
 	UPCGBlueprintElement* BlueprintElementInstance = nullptr;
 };
 
 class FPCGExecuteBlueprintElement : public IPCGElement
 {
 public:
-	virtual bool CanExecuteOnlyOnMainThread(const UPCGSettings* InSettings) const override;
+	virtual bool CanExecuteOnlyOnMainThread(FPCGContext* Context) const override;
 	virtual bool IsCacheable(const UPCGSettings* InSettings) const override;
 
 protected:
diff --git a/Engine/Plugins/Experimental/PCG/Source/PCG/Public/Elements/PCGPropertyToParamData.h b/Engine/Plugins/Experimental/PCG/Source/PCG/Public/Elements/PCGPropertyToParamData.h
--- a/Engine/Plugins/Experimental/PCG/Source/PCG/Public/Elements/PCGPropertyToParamData.h
+++ b/Engine/Plugins/Experimental/PCG/Source/PCG/Public/Elements/PCGPropertyToParamData.h
@@ -88,7 +88,7 @@ public:
 class FPCGPropertyToParamDataElement : public FSimplePCGElement
 {
 public:
-	virtual bool CanExecuteOnlyOnMainThread(const UPCGSettings* InSettings) const override { return true; }
+	virtual bool CanExecuteOnlyOnMainThread(FPCGContext* Context) const override { return true; }
 	virtual bool IsCacheable(const UPCGSettings* InSettings) const override { return !CastChecked<UPCGPropertyToParamDataSettings>(InSettings)->bAlwaysRequeryActors; }
 protected:
 	virtual bool ExecuteInternal(FPCGContext* Context) const;
diff --git a/Engine/Plugins/Experimental/PCG/Source/PCG/Public/Elements/PCGSpawnActor.h b/Engine/Plugins/Experimental/PCG/Source/PCG/Public/Elements/PCGSpawnActor.h
--- a/Engine/Plugins/Experimental/PCG/Source/PCG/Public/Elements/PCGSpawnActor.h
+++ b/Engine/Plugins/Experimental/PCG/Source/PCG/Public/Elements/PCGSpawnActor.h
@@ -81,7 +81,7 @@ public:
 class FPCGSpawnActorElement : public FSimplePCGElement
 {
 public:
-	virtual bool CanExecuteOnlyOnMainThread(const UPCGSettings* InSettings) const override { return true; }
+	virtual bool CanExecuteOnlyOnMainThread(FPCGContext* Context) const override { return true; }
 	virtual bool IsCacheable(const UPCGSettings* InSettings) const override { return false; }
 	virtual bool IsPassthrough() const override { return true; }
 
diff --git a/Engine/Plugins/Experimental/PCG/Source/PCG/Public/Elements/PCGStaticMeshSpawner.h b/Engine/Plugins/Experimental/PCG/Source/PCG/Public/Elements/PCGStaticMeshSpawner.h
--- a/Engine/Plugins/Experimental/PCG/Source/PCG/Public/Elements/PCGStaticMeshSpawner.h
+++ b/Engine/Plugins/Experimental/PCG/Source/PCG/Public/Elements/PCGStaticMeshSpawner.h
@@ -95,13 +95,30 @@ protected:
 	void RefreshInstancePacker();
 };
 
-class FPCGStaticMeshSpawnerElement : public FSimplePCGElement
+USTRUCT()
+struct FPCGStaticMeshSpawnerContext : public FPCGContext
+{
+	GENERATED_BODY()
+
+	struct FPackedInstanceListData
+	{
+		const UPCGSpatialData* SpatialData;
+		TArray<FPCGMeshInstanceList> MeshInstances;
+		TArray<FPCGPackedCustomData> PackedCustomData;
+	};
+
+	TArray<FPackedInstanceListData> MeshInstancesData;
+};
+
+class FPCGStaticMeshSpawnerElement : public IPCGElement
 {
 public:
-	virtual bool CanExecuteOnlyOnMainThread(const UPCGSettings* InSettings) const { return true; }
+	virtual FPCGContext* Initialize(const FPCGDataCollection& InputData, TWeakObjectPtr<UPCGComponent> SourceComponent, const UPCGNode* Node) override;
+	virtual bool CanExecuteOnlyOnMainThread(FPCGContext* Context) const;
 	virtual bool IsCacheable(const UPCGSettings* InSettings) const override { return false; }
 
 protected:
+	virtual bool PrepareDataInternal(FPCGContext* Context) const override;
 	virtual bool ExecuteInternal(FPCGContext* Context) const override;	
 	void SpawnStaticMeshInstances(FPCGContext* Context, const FPCGMeshInstanceList& InstanceList, AActor* TargetActor, const FPCGPackedCustomData& PackedCustomData) const;
 };
diff --git a/Engine/Plugins/Experimental/PCG/Source/PCG/Public/PCGContext.h b/Engine/Plugins/Experimental/PCG/Source/PCG/Public/PCGContext.h
--- a/Engine/Plugins/Experimental/PCG/Source/PCG/Public/PCGContext.h
+++ b/Engine/Plugins/Experimental/PCG/Source/PCG/Public/PCGContext.h
@@ -27,6 +27,16 @@ namespace PCGContextHelpers
 	}
 }
 
+UENUM()
+enum class EPCGExecutionPhase : uint8
+{
+		NotExecuted = 0,
+		PrepareData,
+		Execute,
+		PostExecute,
+		Done
+};
+
 USTRUCT(BlueprintType)
 struct PCG_API FPCGContext
 {
@@ -45,6 +55,12 @@ struct PCG_API FPCGContext
 	FPCGTaskId TaskId = InvalidPCGTaskId;
 	bool bIsPaused = false;
 
+	EPCGExecutionPhase CurrentPhase = EPCGExecutionPhase::NotExecuted;
+	int32 BypassedOutputCount = 0;
+
+	double EndTime = 0.0;
+	bool bIsRunningOnMainThread = false;
+
 #if WITH_EDITOR
 	double ElapsedTime = 0.0;
 	int32 ExecutionCount = 0;
@@ -58,4 +74,5 @@ struct PCG_API FPCGContext
 
 	FString GetTaskName() const;
 	FString GetComponentName() const;
+	bool ShouldStop() const;
 };
diff --git a/Engine/Plugins/Experimental/PCG/Source/PCG/Public/PCGElement.h b/Engine/Plugins/Experimental/PCG/Source/PCG/Public/PCGElement.h
--- a/Engine/Plugins/Experimental/PCG/Source/PCG/Public/PCGElement.h
+++ b/Engine/Plugins/Experimental/PCG/Source/PCG/Public/PCGElement.h
@@ -34,11 +34,19 @@ class PCG_API IPCGElement
 {
 public:
 	virtual ~IPCGElement() = default;
+
+	/** Creates a custom context object paired to this element */
 	virtual FPCGContext* Initialize(const FPCGDataCollection& InputData, TWeakObjectPtr<UPCGComponent> SourceComponent, const UPCGNode* Node) = 0;
 
-	virtual bool CanExecuteOnlyOnMainThread(const UPCGSettings* InSettings) const { return false; }
+	/** Returns true if the element, in its current phase can be executed only from the main thread */
+	virtual bool CanExecuteOnlyOnMainThread(FPCGContext* Context) const { return false; }
+
+	/** Returns true if the node can be cached (e.g. does not create artifacts & does not depend on untracked data */
 	virtual bool IsCacheable(const UPCGSettings* InSettings) const { return true; }
 
+	/** Public function that executes the element on the appropriately created context.
+	* The caller should call the Execute function until it returns true.
+	*/
 	bool Execute(FPCGContext* Context) const;
 
 	/** Note: the following methods must be called from the main thread */
@@ -49,9 +57,20 @@ public:
 #endif
 
 protected:
+	/** This function will be called once and once only, at the beginning of an execution */
+	void PreExecute(FPCGContext* Context) const;
+	/** The prepare data phase is one where it is more likely to be able to multithread */
+	virtual bool PrepareDataInternal(FPCGContext* Context) const;
+	/** Core execution method for the given element. Will be called until it returns true. */
 	virtual bool ExecuteInternal(FPCGContext* Context) const = 0;
+	/** This function will be called once and once only, at the end of an execution */
+	void PostExecute(FPCGContext* Context) const;
+
+	/** Controls whether an element can skip its execution wholly when the input data has the cancelled tag */
 	virtual bool IsCancellable() const { return true; }
+	/** Used to specify that the element passes through the data without any manipulation - used to correct target pins, etc. */
 	virtual bool IsPassthrough() const { return false; }
+
 #if WITH_EDITOR
 	virtual bool ShouldLog() const { return true; }
 #endif
@@ -60,6 +79,16 @@ private:
 	void CleanupAndValidateOutput(FPCGContext* Context) const;
 
 #if WITH_EDITOR
+	struct FScopedCallTimer
+	{
+		FScopedCallTimer(const IPCGElement& InOwner, FPCGContext* InContext);
+		~FScopedCallTimer();
+
+		const IPCGElement& Owner;
+		FPCGContext* Context;
+		double StartTime;
+	};
+
 	// Set mutable because we need to modify them in the execute call, which is const
 	// TODO: Should be a map with PCG Components. We need a mechanism to make sure that this map is cleaned up when component doesn't exist anymore.
 	// For now, it will track all calls to execute (excluding call where the result is already in cache).
@@ -68,6 +97,11 @@ private:
 	// Perhaps overkill but there is a slight chance that we need to protect the timers array. If we call reset from the UI while an element is executing,
 	// it could crash while writing to the timers array.
 	mutable FCriticalSection TimersLock;
+#else // !WITH_EDITOR
+	struct FScopedCallTimer
+	{
+		FScopedCallTimer(const IPCGElement& InOwner, FPCGContext* InContext) {}
+	};
 #endif // WITH_EDITOR
 };
 
