[PCG] Fix potential uninitialized data issue - ensure all data in an 'Out' variable is set.

As an aside, in a local test OutPoint came in with default data so it may not affect current use cases.

diff --git a/Engine/Plugins/Experimental/PCG/Source/PCG/Private/Grid/PCGLandscapeCache.cpp b/Engine/Plugins/Experimental/PCG/Source/PCG/Private/Grid/PCGLandscapeCache.cpp
--- a/Engine/Plugins/Experimental/PCG/Source/PCG/Private/Grid/PCGLandscapeCache.cpp
+++ b/Engine/Plugins/Experimental/PCG/Source/PCG/Private/Grid/PCGLandscapeCache.cpp
@@ -108,10 +108,12 @@ void FPCGLandscapeCacheEntry::GetPoint(int32 PointIndex, FPCGPoint& OutPoint, UP
 	TangentX = FVector(Normal.Z, 0.f, -Normal.X);
 	TangentY = Normal ^ TangentX;
 
-	OutPoint.Transform = FTransform(TangentX.GetSafeNormal(), TangentY.GetSafeNormal(), Normal.GetSafeNormal(), Position);
+	const float Density = 1;
+	const int32 Seed = 1 + PointIndex;
+
+	new(&OutPoint) FPCGPoint(FTransform(TangentX.GetSafeNormal(), TangentY.GetSafeNormal(), Normal.GetSafeNormal(), Position), Density, Seed);
 	OutPoint.BoundsMin = -PointHalfSize;
 	OutPoint.BoundsMax = PointHalfSize;
-	OutPoint.Seed = 1 + PointIndex;
 	
 	if (OutMetadata && !LayerData.IsEmpty())
 	{
@@ -139,10 +141,12 @@ void FPCGLandscapeCacheEntry::GetPointHeightOnly(int32 PointIndex, FPCGPoint& Ou
 
 	const FVector& Position = PositionsAndNormals[2 * PointIndex];
 	
-	OutPoint.Transform = FTransform(Position);
+	const float Density = 1;
+	const int32 Seed = 1 + PointIndex;
+
+	new(&OutPoint) FPCGPoint(FTransform(Position), Density, Seed);
 	OutPoint.BoundsMin = -PointHalfSize;
 	OutPoint.BoundsMax = PointHalfSize;
-	OutPoint.Seed = 1 + PointIndex;
 }
 
 void FPCGLandscapeCacheEntry::GetInterpolatedPoint(const FVector2D& LocalPoint, FPCGPoint& OutPoint, UPCGMetadata* OutMetadata) const
@@ -181,10 +185,12 @@ void FPCGLandscapeCacheEntry::GetInterpolatedPoint(const FVector2D& LocalPoint,
 	TangentX = FVector(Normal.Z, 0.f, -Normal.X);
 	TangentY = Normal ^ TangentX;
 
-	OutPoint.Transform = FTransform(TangentX.GetSafeNormal(), TangentY.GetSafeNormal(), Normal.GetSafeNormal(), Position);
+	const int32 Seed = UPCGBlueprintHelpers::ComputeSeedFromPosition(Position);
+	const float Density = 1;
+
+	new(&OutPoint) FPCGPoint(FTransform(TangentX.GetSafeNormal(), TangentY.GetSafeNormal(), Normal.GetSafeNormal(), Position), Density, Seed);
 	OutPoint.BoundsMin = -PointHalfSize;
 	OutPoint.BoundsMax = PointHalfSize;
-	OutPoint.Seed = UPCGBlueprintHelpers::ComputeSeedFromPosition(Position);
 
 	if (OutMetadata && !LayerData.IsEmpty())
 	{
@@ -233,10 +239,12 @@ void FPCGLandscapeCacheEntry::GetInterpolatedPointHeightOnly(const FVector2D& Lo
 	const FVector LerpPositionY1 = FMath::Lerp(PositionX0Y1, PositionX1Y1, XFactor);
 	const FVector Position = FMath::Lerp(LerpPositionY0, LerpPositionY1, YFactor);
 
-	OutPoint.Transform = FTransform(Position);
+	const int32 Seed = UPCGBlueprintHelpers::ComputeSeedFromPosition(Position);
+	const float Density = 1;
+
+	new(&OutPoint) FPCGPoint(FTransform(Position), Density, Seed);
 	OutPoint.BoundsMin = -PointHalfSize;
 	OutPoint.BoundsMax = PointHalfSize;
-	OutPoint.Seed = UPCGBlueprintHelpers::ComputeSeedFromPosition(Position);
 }
 
 bool FPCGLandscapeCacheEntry::TouchAndLoad(int32 InTouch) const
