Utility function to query landscape layer weights from PCG landscape cache. This has uses outside of PCG (biome audio for example)

diff --git a/Engine/Plugins/Experimental/PCG/Source/PCG/Private/Grid/PCGLandscapeCache.cpp b/Engine/Plugins/Experimental/PCG/Source/PCG/Private/Grid/PCGLandscapeCache.cpp
--- a/Engine/Plugins/Experimental/PCG/Source/PCG/Private/Grid/PCGLandscapeCache.cpp
+++ b/Engine/Plugins/Experimental/PCG/Source/PCG/Private/Grid/PCGLandscapeCache.cpp
@@ -150,6 +150,42 @@ void FPCGLandscapeCacheEntry::GetPointHeightOnly(int32 PointIndex, FPCGPoint& Ou
 	OutPoint.BoundsMax = PointHalfSize;
 }
 
+void FPCGLandscapeCacheEntry::GetInterpolatedLayerWeights(const FVector2D& LocalPoint, TArray<FPCGLandscapeLayerWeight>& OutLayerWeights) const
+{
+	OutLayerWeights.SetNum(LayerData.Num());
+	
+	if (!bDataLoaded)
+	{
+		return;
+	}
+
+	const int32 X0Y0 = FMath::FloorToInt(LocalPoint.X) + FMath::FloorToInt(LocalPoint.Y) * Stride;
+	const int32 X1Y0 = X0Y0 + 1;
+	const int32 X0Y1 = X0Y0 + Stride;
+	const int32 X1Y1 = X0Y1 + 1;
+
+	const float XFactor = FMath::Fractional(LocalPoint.X);
+	const float YFactor = FMath::Fractional(LocalPoint.Y);
+
+	if (!LayerData.IsEmpty())
+	{
+		check(LayerData.Num() == LayerDataNames.Num());
+
+		for(int32 LayerIndex = 0; LayerIndex < LayerData.Num(); ++LayerIndex)
+		{
+			FPCGLandscapeLayerWeight& OutLayer = OutLayerWeights[LayerIndex];
+			OutLayer.Name = LayerDataNames[LayerIndex];
+
+			const TArray<uint8>& CurrentLayerData = LayerData[LayerIndex];
+
+			const float Y0Data = FMath::Lerp((float)CurrentLayerData[X0Y0] / 255.0f, (float)CurrentLayerData[X1Y0] / 255.0f, XFactor);
+			const float Y1Data = FMath::Lerp((float)CurrentLayerData[X0Y1] / 255.0f, (float)CurrentLayerData[X1Y1] / 255.0f, XFactor);
+
+			OutLayer.Weight = FMath::Lerp(Y0Data, Y1Data, YFactor);
+		}
+	}
+}
+
 void FPCGLandscapeCacheEntry::GetInterpolatedPoint(const FVector2D& LocalPoint, FPCGPoint& OutPoint, UPCGMetadata* OutMetadata) const
 {
 	check(bDataLoaded);
diff --git a/Engine/Plugins/Experimental/PCG/Source/PCG/Private/Helpers/PCGBlueprintHelpers.cpp b/Engine/Plugins/Experimental/PCG/Source/PCG/Private/Helpers/PCGBlueprintHelpers.cpp
--- a/Engine/Plugins/Experimental/PCG/Source/PCG/Private/Helpers/PCGBlueprintHelpers.cpp
+++ b/Engine/Plugins/Experimental/PCG/Source/PCG/Private/Helpers/PCGBlueprintHelpers.cpp
@@ -4,7 +4,15 @@
 #include "PCGComponent.h"
 #include "PCGHelpers.h"
 #include "PCGSettings.h"
+#include "PCGSubsystem.h"
+#include "Metadata/PCGMetadata.h"
+#include "Metadata/PCGMetadataAttributeTpl.h"
 #include "Grid/PCGPartitionActor.h"
+#include "Grid/PCGLandscapeCache.h"
+#include "Engine/World.h"
+#include "UObject/UObjectIterator.h"
+#include "LandscapeProxy.h"
+#include "LandscapeInfo.h"
 
 int UPCGBlueprintHelpers::ComputeSeedFromPosition(const FVector& InPosition)
 {
@@ -104,4 +112,65 @@ FBox UPCGBlueprintHelpers::GetActorBoundsPCG(AActor* InActor, bool bIgnorePCGCre
 FBox UPCGBlueprintHelpers::GetActorLocalBoundsPCG(AActor* InActor, bool bIgnorePCGCreatedComponents)
 {
 	return PCGHelpers::GetActorLocalBounds(InActor, bIgnorePCGCreatedComponents);
+}
+
+TArray<FPCGLandscapeLayerWeight> UPCGBlueprintHelpers::GetInterpolatedPCGLandscapeLayerWeights(UObject* WorldContextObject, const FVector& Location)
+{
+	UWorld* World = WorldContextObject ? WorldContextObject->GetWorld() : nullptr;
+	if (!World)
+		return {};
+
+	UPCGSubsystem* PCGSubSystem = UWorld::GetSubsystem<UPCGSubsystem>(World);
+	if (!PCGSubSystem)
+		return {};
+
+	FIntPoint ComponentKey{};
+
+	auto FindLandscapeComponent = [World, Location, &ComponentKey]() -> ULandscapeComponent* {
+		for (TObjectIterator<ALandscapeProxy> It; It; ++It)
+		{
+			if (It->GetWorld() == World)
+			{
+				FBox Box = It->GetComponentsBoundingBox();
+				if (Box.IsInsideOrOnXY(Location))
+				{
+					if (ULandscapeInfo* Info = It->GetLandscapeInfo())
+					{
+						const FVector ActorSpaceLocation = It->LandscapeActorToWorld().InverseTransformPosition(Location);
+						ComponentKey = FIntPoint(FMath::FloorToInt(ActorSpaceLocation.X / It->ComponentSizeQuads), FMath::FloorToInt(ActorSpaceLocation.Y / It->ComponentSizeQuads));
+
+						return Info->XYtoComponentMap.FindRef(ComponentKey);
+					}
+
+					return nullptr;
+				}
+			}
+		}
+
+		return nullptr;
+	};
+
+	ULandscapeComponent* LandscapeComponent = FindLandscapeComponent();
+	if (!LandscapeComponent)
+		return {};
+
+	const FVector ComponentSpaceLocation = LandscapeComponent->GetComponentToWorld().InverseTransformPosition(Location);
+
+	UPCGLandscapeCache* LandscapeCache = PCGSubSystem->GetLandscapeCache();
+	if (!LandscapeCache)
+		return {};
+
+	const FPCGLandscapeCacheEntry* CacheEntry = LandscapeCache->GetCacheEntry(LandscapeComponent, ComponentKey);
+	if (!CacheEntry)
+		return {};
+
+	TArray<FPCGLandscapeLayerWeight> Result;
+
+	CacheEntry->GetInterpolatedLayerWeights(FVector2D(ComponentSpaceLocation), Result);
+
+	Result.Sort([](const FPCGLandscapeLayerWeight& Lhs, const FPCGLandscapeLayerWeight& Rhs){
+		return Lhs.Weight > Rhs.Weight;
+	});
+					
+	return Result;
 }
\ No newline at end of file
diff --git a/Engine/Plugins/Experimental/PCG/Source/PCG/Public/Grid/PCGLandscapeCache.h b/Engine/Plugins/Experimental/PCG/Source/PCG/Public/Grid/PCGLandscapeCache.h
--- a/Engine/Plugins/Experimental/PCG/Source/PCG/Public/Grid/PCGLandscapeCache.h
+++ b/Engine/Plugins/Experimental/PCG/Source/PCG/Public/Grid/PCGLandscapeCache.h
@@ -14,6 +14,18 @@ class UPCGPointData;
 class UPCGMetadata;
 struct FPCGPoint;
 
+USTRUCT(BlueprintType)
+struct FPCGLandscapeLayerWeight
+{
+	GENERATED_BODY()
+
+	UPROPERTY(EditAnywhere, BlueprintReadWrite, Category = "Landscape Attribute")
+	FName Name;
+
+	UPROPERTY(EditAnywhere, BlueprintReadWrite, Category = "Landscape Attribute")
+	float Weight;
+};
+
 struct FPCGLandscapeCacheEntry
 {
 	friend class UPCGLandscapeCache;
@@ -23,6 +35,7 @@ public:
 	void GetPointHeightOnly(int32 PointIndex, FPCGPoint& OutPoint) const;
 	void GetInterpolatedPoint(const FVector2D& LocalPoint, FPCGPoint& OutPoint, UPCGMetadata* OutMetadata) const;
 	void GetInterpolatedPointHeightOnly(const FVector2D& LocalPoint, FPCGPoint& OutPoint) const;
+	void GetInterpolatedLayerWeights(const FVector2D& LocalPoint, TArray<FPCGLandscapeLayerWeight>& OutLayerWeights) const;
 
 private:
 	// Private API for UPCGLandscapeCache usage
diff --git a/Engine/Plugins/Experimental/PCG/Source/PCG/Public/Helpers/PCGBlueprintHelpers.h b/Engine/Plugins/Experimental/PCG/Source/PCG/Public/Helpers/PCGBlueprintHelpers.h
--- a/Engine/Plugins/Experimental/PCG/Source/PCG/Public/Helpers/PCGBlueprintHelpers.h
+++ b/Engine/Plugins/Experimental/PCG/Source/PCG/Public/Helpers/PCGBlueprintHelpers.h
@@ -65,4 +65,7 @@ public:
 
 	UFUNCTION(BlueprintCallable, Category = "PCG|Helpers", meta = (ScriptMethod))
 	static FBox GetActorLocalBoundsPCG(AActor* InActor, bool bIgnorePCGCreatedComponents = true);
+
+	UFUNCTION(BlueprintCallable, Category = "PCG|Helpers", meta = (ScriptMethod))
+	static TArray<FPCGLandscapeLayerWeight> GetInterpolatedPCGLandscapeLayerWeights(UObject* WorldContextObject, const FVector& Location);
 };
\ No newline at end of file
