Add custom version numbers for PCG. These open the door to additional powerful migration use cases in addition to the usual methods such as the '_DEPRECATED' postfix.

The initial use case of this will be projection node inputs migration - a single input is now split into projection source and target. It was not trivial to detect the split inputs from the deprecation code.

This kind of migration of node connections must happen on/after graph deserialization rather than node (settings) deserialization, because it is a graph-level fixup. And also the node pins/edges are not resolved until after Serialize(), so I could not simply apply fixup in Serialize() as would normally be the case. So this particular migration is not a good example of using the custom version, and is more complicated than normal migration. However the custom version implementation is powerful and is hopefully usable in an easier way for other migration needs in the future. The custom version can be used on any UObject in the PCG system using standard patterns used throughout the engine.

I have added an additional custom version that we will not touch internally, which external teams could use if they wish.

diff --git a/Engine/Plugins/Experimental/PCG/Source/PCG/Private/PCGCustomVersion.cpp b/Engine/Plugins/Experimental/PCG/Source/PCG/Private/PCGCustomVersion.cpp
new file mode 100644
--- /dev/null
+++ b/Engine/Plugins/Experimental/PCG/Source/PCG/Private/PCGCustomVersion.cpp
@@ -0,0 +1,9 @@
+// Copyright Epic Games, Inc. All Rights Reserved.
+
+#include "PCGCustomVersion.h"
+#include "Serialization/CustomVersion.h"
+
+const FGuid FPCGCustomVersion::GUID(0x2763920D, 0x0F784B39, 0x986E4BB3, 0xA88D666D);
+
+// Register the custom version with core
+FCustomVersionRegistration GRegisterPCGCustomVersion(FPCGCustomVersion::GUID, FPCGCustomVersion::LatestVersion, TEXT("PCGBaseVer"));
diff --git a/Engine/Plugins/Experimental/PCG/Source/PCG/Private/PCGNode.cpp b/Engine/Plugins/Experimental/PCG/Source/PCG/Private/PCGNode.cpp
--- a/Engine/Plugins/Experimental/PCG/Source/PCG/Private/PCGNode.cpp
+++ b/Engine/Plugins/Experimental/PCG/Source/PCG/Private/PCGNode.cpp
@@ -116,6 +116,11 @@ void UPCGNode::ApplyDeprecation()
 		}
 	}
 	OutboundEdges_DEPRECATED.Reset();
+
+	if (DefaultSettings)
+	{
+		DefaultSettings->ApplyDeprecation(this);
+	}
 }
 #endif
 
diff --git a/Engine/Plugins/Experimental/PCG/Source/PCG/Private/PCGSettings.cpp b/Engine/Plugins/Experimental/PCG/Source/PCG/Private/PCGSettings.cpp
--- a/Engine/Plugins/Experimental/PCG/Source/PCG/Private/PCGSettings.cpp
+++ b/Engine/Plugins/Experimental/PCG/Source/PCG/Private/PCGSettings.cpp
@@ -1,6 +1,7 @@
 // Copyright Epic Games, Inc. All Rights Reserved.
 
 #include "PCGSettings.h"
+#include "PCGCustomVersion.h"
 #include "PCGNode.h"
 #include "PCGSubsystem.h"
 #include "Serialization/ArchiveObjectCrc32.h"
@@ -44,6 +45,41 @@ uint32 UPCGSettings::GetCrc32() const
 	return Ar.Crc32(const_cast<UPCGSettings*>(this));
 }
 
+void UPCGSettings::Serialize(FArchive& Ar)
+{
+	Super::Serialize(Ar);
+
+	Ar.UsingCustomVersion(FPCGCustomVersion::GUID);
+
+#if WITH_EDITOR
+	if (Ar.IsLoading())
+	{
+		// Some data migration must happen after the graph is fully initialized, such as manipulating node connections, so we
+		// store off the loaded version number to be used later.
+		DataVersion = Ar.CustomVer(FPCGCustomVersion::GUID);
+	}
+#endif // WITH_EDITOR
+
+	// An additional custom version number that can be driven by external system users to track system modifications. To use a custom
+	// version in user settings objects, override the GetUserCustomVersionGuid() method.
+	const FGuid UserDataGuid = GetUserCustomVersionGuid();
+	const bool bUsingCustomUserVersion = UserDataGuid != FGuid();
+
+	if (bUsingCustomUserVersion)
+	{
+		Ar.UsingCustomVersion(UserDataGuid);
+
+#if WITH_EDITOR
+		if (Ar.IsLoading())
+		{
+			// Some data migration must happen after the graph is fully initialized, such as manipulating node connections, so we
+			// store off the loaded version number to be used later.
+			UserDataVersion = Ar.CustomVer(UserDataGuid);
+		}
+#endif // WITH_EDITOR
+	}
+}
+
 #if WITH_EDITOR
 UObject* UPCGSettings::GetJumpTargetForDoubleClick() const
 {
diff --git a/Engine/Plugins/Experimental/PCG/Source/PCG/Public/PCGCustomVersion.h b/Engine/Plugins/Experimental/PCG/Source/PCG/Public/PCGCustomVersion.h
new file mode 100644
--- /dev/null
+++ b/Engine/Plugins/Experimental/PCG/Source/PCG/Public/PCGCustomVersion.h
@@ -0,0 +1,26 @@
+// Copyright Epic Games, Inc. All Rights Reserved.
+
+#pragma once
+
+#include "CoreMinimal.h"
+#include "Misc/Guid.h"
+
+// Custom serialization version for assets/classes in the PCG plugin
+struct PCG_API FPCGCustomVersion
+{
+	enum Type
+	{
+		// Before any version changes were made in the plugin
+		BeforeCustomVersionWasAdded = 0,
+
+		// -----<new versions can be added above this line>-------------------------------------------------
+		VersionPlusOne,
+		LatestVersion = VersionPlusOne - 1
+	};
+
+	// The GUID for this custom version number
+	const static FGuid GUID;
+
+private:
+	FPCGCustomVersion() {}
+};
diff --git a/Engine/Plugins/Experimental/PCG/Source/PCG/Public/PCGSettings.h b/Engine/Plugins/Experimental/PCG/Source/PCG/Public/PCGSettings.h
--- a/Engine/Plugins/Experimental/PCG/Source/PCG/Public/PCGSettings.h
+++ b/Engine/Plugins/Experimental/PCG/Source/PCG/Public/PCGSettings.h
@@ -59,6 +59,10 @@ public:
 	virtual EPCGDataType GetDataType() const override { return EPCGDataType::Settings | Super::GetDataType(); }
 	// ~End UPCGData interface
 
+	//~Begin UObject interface
+	virtual void Serialize(FArchive& Ar) override;
+	//~End UObject interface
+
 	// TODO: check if we need this to be virtual, we don't really need if we're always caching
 	/*virtual*/ FPCGElementPtr GetElement() const;
 	virtual UPCGNode* CreateNode() const;
@@ -71,6 +75,8 @@ public:
 	bool UseSeed() const { return bUseSeed; }
 
 #if WITH_EDITOR
+	virtual void ApplyDeprecation(UPCGNode* InOutNode) {}
+
 	virtual FName GetDefaultNodeName() const { return NAME_None; }
 	virtual FLinearColor GetNodeTitleColor() const { return FLinearColor::White; }
 	virtual EPCGSettingsType GetType() const { return EPCGSettingsType::Generic; }
@@ -119,6 +125,9 @@ public:
 protected:
 	virtual FPCGElementPtr CreateElement() const PURE_VIRTUAL(UPCGSettings::CreateElement, return nullptr;);
 
+	/** An additional custom version number that external system users can use to track versions. This version will be serialized into the asset and will be provided by UserDataVersion after load. */
+	virtual FGuid GetUserCustomVersionGuid() { return FGuid(); }
+
 #if WITH_EDITOR
 	virtual void PostEditChangeProperty(FPropertyChangedEvent& PropertyChangedEvent) override;
 	virtual bool IsStructuralProperty(const FName& InPropertyName) const { return false; }
@@ -134,6 +143,15 @@ protected:
 	/** Methods to remove boilerplate code across settings */
 	TArray<FPCGPinProperties> DefaultPointOutputPinProperties() const;
 
+#if WITH_EDITOR
+protected:
+	/** The version number of the data after load and after any data migration. */
+	int32 DataVersion = -1;
+
+	/** If a custom version guid was provided through GetUserCustomVersionGuid(), this field will hold the version number after load and after any data migration. */
+	int32 UserDataVersion = -1;
+#endif
+
 private:
 	mutable FPCGElementPtr CachedElement;
 	mutable FCriticalSection CacheLock;
