[PCG] Refactor 'FindOrAdd' function in MeshSelectorBase to remove unnecessary boolean that indicated found vs added. This is more in line with other similar functions in the codebase.

diff --git a/Engine/Plugins/Experimental/PCG/Source/PCG/Private/MeshSelectors/PCGMeshSelectorBase.cpp b/Engine/Plugins/Experimental/PCG/Source/PCG/Private/MeshSelectors/PCGMeshSelectorBase.cpp
--- a/Engine/Plugins/Experimental/PCG/Source/PCG/Private/MeshSelectors/PCGMeshSelectorBase.cpp
+++ b/Engine/Plugins/Experimental/PCG/Source/PCG/Private/MeshSelectors/PCGMeshSelectorBase.cpp
@@ -4,7 +4,7 @@
 
 #include "Elements/PCGStaticMeshSpawner.h"
 
-bool UPCGMeshSelectorBase::FindOrAddInstanceList(
+int32 UPCGMeshSelectorBase::FindOrAddInstanceList(
 	TArray<FPCGMeshInstanceList>& OutInstanceLists,
 	const TSoftObjectPtr<UStaticMesh>& Mesh,
 	bool bOverrideCollisionProfile,
@@ -12,8 +12,7 @@ bool UPCGMeshSelectorBase::FindOrAddInstanceList(
 	bool bOverrideMaterials,
 	const TArray<UMaterialInterface*>& MaterialOverrides,
 	const float InCullStartDistance,
-	const float InCullEndDistance,
-	int32& OutIndex) const
+	const float InCullEndDistance) const
 {
 	for (int Index = 0; Index < OutInstanceLists.Num(); ++Index)
 	{
@@ -37,12 +36,8 @@ bool UPCGMeshSelectorBase::FindOrAddInstanceList(
 			continue;
 		}
 
-		OutIndex = Index;
-
-		return false;
+		return Index;
 	}
 
-	OutIndex = OutInstanceLists.Emplace(Mesh, bOverrideCollisionProfile, CollisionProfile, bOverrideMaterials, MaterialOverrides, InCullStartDistance, InCullEndDistance);
-
-	return true;
+	return OutInstanceLists.Emplace(Mesh, bOverrideCollisionProfile, CollisionProfile, bOverrideMaterials, MaterialOverrides, InCullStartDistance, InCullEndDistance);
 }
diff --git a/Engine/Plugins/Experimental/PCG/Source/PCG/Private/MeshSelectors/PCGMeshSelectorByAttribute.cpp b/Engine/Plugins/Experimental/PCG/Source/PCG/Private/MeshSelectors/PCGMeshSelectorByAttribute.cpp
--- a/Engine/Plugins/Experimental/PCG/Source/PCG/Private/MeshSelectors/PCGMeshSelectorByAttribute.cpp
+++ b/Engine/Plugins/Experimental/PCG/Source/PCG/Private/MeshSelectors/PCGMeshSelectorByAttribute.cpp
@@ -90,8 +90,9 @@ void UPCGMeshSelectorByAttribute::SelectInstances_Implementation(
 			continue;
 		}
 
-		int32 Index = INDEX_NONE;
-		FindOrAddInstanceList(OutMeshInstances, Mesh, bOverrideCollisionProfile, CollisionProfile, bOverrideMaterials, MaterialOverrides, CullStartDistance, CullEndDistance, Index);
+		const int32 Index = FindOrAddInstanceList(OutMeshInstances, Mesh, bOverrideCollisionProfile, CollisionProfile, bOverrideMaterials, MaterialOverrides, CullStartDistance, CullEndDistance);
+		check(Index != INDEX_NONE);
+		
 		OutMeshInstances[Index].Instances.Emplace(Point);
 	}
 }
diff --git a/Engine/Plugins/Experimental/PCG/Source/PCG/Private/MeshSelectors/PCGMeshSelectorWeighted.cpp b/Engine/Plugins/Experimental/PCG/Source/PCG/Private/MeshSelectors/PCGMeshSelectorWeighted.cpp
--- a/Engine/Plugins/Experimental/PCG/Source/PCG/Private/MeshSelectors/PCGMeshSelectorWeighted.cpp
+++ b/Engine/Plugins/Experimental/PCG/Source/PCG/Private/MeshSelectors/PCGMeshSelectorWeighted.cpp
@@ -32,8 +32,7 @@ void UPCGMeshSelectorWeighted::SelectInstances_Implementation(
 			continue;
 		}
 
-		int32 Index = INDEX_NONE;
-		FindOrAddInstanceList(OutMeshInstances, Entry.Mesh, Entry.bOverrideCollisionProfile, Entry.CollisionProfile, Entry.bOverrideMaterials, Entry.MaterialOverrides, Entry.CullStartDistance, Entry.CullEndDistance, Index);
+		FindOrAddInstanceList(OutMeshInstances, Entry.Mesh, Entry.bOverrideCollisionProfile, Entry.CollisionProfile, Entry.bOverrideMaterials, Entry.MaterialOverrides, Entry.CullStartDistance, Entry.CullEndDistance);
 		TotalWeight += Entry.Weight;
 		CumulativeWeights.Add(TotalWeight);
 	}
diff --git a/Engine/Plugins/Experimental/PCG/Source/PCG/Private/MeshSelectors/PCGMeshSelectorWeightedByCategory.cpp b/Engine/Plugins/Experimental/PCG/Source/PCG/Private/MeshSelectors/PCGMeshSelectorWeightedByCategory.cpp
--- a/Engine/Plugins/Experimental/PCG/Source/PCG/Private/MeshSelectors/PCGMeshSelectorWeightedByCategory.cpp
+++ b/Engine/Plugins/Experimental/PCG/Source/PCG/Private/MeshSelectors/PCGMeshSelectorWeightedByCategory.cpp
@@ -111,8 +111,8 @@ void UPCGMeshSelectorWeightedByCategory::SelectInstances_Implementation(
 				continue;
 			}
 			
-			int32 Index = INDEX_NONE;
-			FindOrAddInstanceList(OutMeshInstances, WeightedEntry.Mesh, WeightedEntry.bOverrideCollisionProfile, WeightedEntry.CollisionProfile, WeightedEntry.bOverrideMaterials, WeightedEntry.MaterialOverrides, WeightedEntry.CullStartDistance, WeightedEntry.CullEndDistance, Index);
+			const int32 Index = FindOrAddInstanceList(OutMeshInstances, WeightedEntry.Mesh, WeightedEntry.bOverrideCollisionProfile, WeightedEntry.CollisionProfile, WeightedEntry.bOverrideMaterials, WeightedEntry.MaterialOverrides, WeightedEntry.CullStartDistance, WeightedEntry.CullEndDistance);
+			check(Index != INDEX_NONE);
 
 			// precompute the weights
 			TotalWeight += WeightedEntry.Weight;
diff --git a/Engine/Plugins/Experimental/PCG/Source/PCG/Public/MeshSelectors/PCGMeshSelectorBase.h b/Engine/Plugins/Experimental/PCG/Source/PCG/Public/MeshSelectors/PCGMeshSelectorBase.h
--- a/Engine/Plugins/Experimental/PCG/Source/PCG/Public/MeshSelectors/PCGMeshSelectorBase.h
+++ b/Engine/Plugins/Experimental/PCG/Source/PCG/Public/MeshSelectors/PCGMeshSelectorBase.h
@@ -73,9 +73,9 @@ public:
 		TArray<FPCGMeshInstanceList>& OutMeshInstances,
 		UPCGPointData* OutPointData) const PURE_VIRTUAL(UPCGMeshSelectorBase::SelectInstances_Implementation);
 
-	/** Searches OutInstanceLists for an InstanceList matching the given parameters. If nothing is found, creates a new InstanceList and adds to OutInstanceLists. Returns true if added. */
+	/** Searches OutInstanceLists for an InstanceList matching the given parameters. If nothing is found, creates a new InstanceList and adds to OutInstanceLists. Returns index of the matching instance list. */
 	UFUNCTION(BlueprintCallable, Category = MeshSelection)
-	bool FindOrAddInstanceList(
+	int32 FindOrAddInstanceList(
 		TArray<FPCGMeshInstanceList>& OutInstanceLists,
 		const TSoftObjectPtr<UStaticMesh>& Mesh,
 		bool bOverrideCollisionProfile,
@@ -83,6 +83,5 @@ public:
 		bool bOverrideMaterials,
 		const TArray<UMaterialInterface*>& MaterialOverrides,
 		const float InCullStartDistance,
-		const float InCullEndDistance,
-		int32& OutIndex) const;
+		const float InCullEndDistance) const;
 };
