[PCG] Clean up graph if Nodes are removed from code.

diff --git a/Engine/Plugins/Experimental/PCG/Source/PCG/Private/PCGGraph.cpp b/Engine/Plugins/Experimental/PCG/Source/PCG/Private/PCGGraph.cpp
--- a/Engine/Plugins/Experimental/PCG/Source/PCG/Private/PCGGraph.cpp
+++ b/Engine/Plugins/Experimental/PCG/Source/PCG/Private/PCGGraph.cpp
@@ -112,9 +112,18 @@ void UPCGGraph::PostLoad()
 	InputNode->OnNodeChangedDelegate.AddUObject(this, &UPCGGraph::OnNodeChanged);
 	OutputNode->OnNodeChangedDelegate.AddUObject(this, &UPCGGraph::OnNodeChanged);
 
-	for (UPCGNode* Node : Nodes)
+	// Also, try to remove all nodes that are invalid (meaning that the settings are null)
+	// We remove it at the end, to let the nodes that have null settings to clean up their pins and edges.
+	for (int32 i = Nodes.Num() - 1; i >= 0; --i)
 	{
-		OnNodeAdded(Node);
+		if (!Nodes[i]->DefaultSettings)
+		{
+			Nodes.RemoveAtSwap(i);
+		}
+		else
+		{
+			OnNodeAdded(Nodes[i]);
+		}
 	}
 #endif
 }
diff --git a/Engine/Plugins/Experimental/PCG/Source/PCG/Private/PCGNode.cpp b/Engine/Plugins/Experimental/PCG/Source/PCG/Private/PCGNode.cpp
--- a/Engine/Plugins/Experimental/PCG/Source/PCG/Private/PCGNode.cpp
+++ b/Engine/Plugins/Experimental/PCG/Source/PCG/Private/PCGNode.cpp
@@ -427,6 +427,23 @@ bool UPCGNode::UpdatePins(TFunctionRef<UPCGPin*(UPCGNode*)> PinAllocator)
 			Modify();
 		}
 
+		// Clean up edges
+		for (UPCGPin* Pin : InputPins)
+		{
+			if (Pin)
+			{
+				Pin->BreakAllEdges();
+			}
+		}
+
+		for (UPCGPin* Pin : OutputPins)
+		{
+			if (Pin)
+			{
+				Pin->BreakAllEdges();
+			}
+		}
+
 		InputPins.Reset();
 		OutputPins.Reset();
 		return bChanged;
