[PCG] Attribute management elements
* FilterAttribute
* TransferAttribute
* Update CreateAttribute to support from SourceParam
* Remove CreateParamData

diff --git a/Engine/Plugins/Experimental/PCG/Source/PCG/Private/Elements/PCGAttributeFilterElement.cpp b/Engine/Plugins/Experimental/PCG/Source/PCG/Private/Elements/PCGAttributeFilterElement.cpp
new file mode 100644
--- /dev/null
+++ b/Engine/Plugins/Experimental/PCG/Source/PCG/Private/Elements/PCGAttributeFilterElement.cpp
@@ -0,0 +1,92 @@
+// Copyright Epic Games, Inc. All Rights Reserved.
+
+#include "Elements/PCGAttributeFilterElement.h"
+
+#include "PCGData.h"
+#include "PCGParamData.h"
+#include "Data/PCGPointData.h"
+#include "Metadata/PCGMetadata.h"
+#include "Metadata/PCGMetadataAttribute.h"
+
+#if WITH_EDITOR
+FName UPCGAttributeFilterSettings::GetDefaultNodeName() const
+{
+	return TEXT("FilterAttribute");
+}
+#endif
+
+TArray<FPCGPinProperties> UPCGAttributeFilterSettings::OutputPinProperties() const
+{
+	TArray<FPCGPinProperties> PinProperties;
+	PinProperties.Emplace(PCGPinConstants::DefaultOutputLabel, EPCGDataType::Any);
+
+	return PinProperties;
+}
+
+FPCGElementPtr UPCGAttributeFilterSettings::CreateElement() const
+{
+	return MakeShared<FPCGAttributeFilterElement>();
+}
+
+bool FPCGAttributeFilterElement::ExecuteInternal(FPCGContext* Context) const
+{
+	TRACE_CPUPROFILER_EVENT_SCOPE(FPCGAttributeFilterElement::Execute);
+
+	check(Context);
+
+	const UPCGAttributeFilterSettings* Settings = Context->GetInputSettings<UPCGAttributeFilterSettings>();
+
+	TArray<FPCGTaggedData> Inputs = Context->InputData.GetInputsByPin(PCGPinConstants::DefaultInputLabel);
+
+	for (const FPCGTaggedData& InputTaggedData : Inputs)
+	{
+		const UPCGData* InputData = InputTaggedData.Data;
+		UPCGData* OutputData = nullptr;
+
+		const UPCGMetadata* ParentMetadata = nullptr;
+		UPCGMetadata* Metadata = nullptr;
+
+		if (const UPCGSpatialData* InputSpatialData = Cast<UPCGSpatialData>(InputData))
+		{
+			Metadata = NewObject<UPCGMetadata>();
+			ParentMetadata = InputSpatialData->Metadata;
+
+			UPCGSpatialData* NewSpatialData = DuplicateObject<UPCGSpatialData>(const_cast<UPCGSpatialData*>(InputSpatialData), nullptr);
+			NewSpatialData->Metadata = Metadata;
+			NewSpatialData->Metadata->Initialize(ParentMetadata, /*bAddAttributesFromParent=*/false);
+
+			// No need to inherit metadata since we already initialized it.
+			NewSpatialData->InitializeFromData(InputSpatialData, /*InMetadataParentOverride=*/ nullptr, /*bInheritMetadata=*/ false);
+
+			OutputData = NewSpatialData;
+		}
+		else if (const UPCGParamData* InputParamData = Cast<UPCGParamData>(InputData))
+		{
+			Metadata = NewObject<UPCGMetadata>();
+			ParentMetadata = InputParamData->Metadata;
+
+			UPCGParamData* NewParamData = NewObject<UPCGParamData>();
+			NewParamData->Metadata = Metadata;
+
+			Metadata->Initialize(InputParamData->Metadata, /*bAddAttributesFromParent=*/false);
+			OutputData = NewParamData;
+		}
+		else
+		{
+			PCGE_LOG(Error, "Invalid data as input. Only support spatial and params");
+			continue;
+		}
+
+		// Then add each attribute in the list explicitly
+		for (const FName& AttributeName : Settings->AttributesToKeep)
+		{
+			Metadata->AddAttribute(ParentMetadata, AttributeName);
+		}
+
+		TArray<FPCGTaggedData>& Outputs = Context->OutputData.TaggedData;
+		FPCGTaggedData& Output = Outputs.Emplace_GetRef();
+		Output.Data = OutputData;
+	}
+
+	return true;
+}
diff --git a/Engine/Plugins/Experimental/PCG/Source/PCG/Private/Elements/PCGAttributeTransferElement.cpp b/Engine/Plugins/Experimental/PCG/Source/PCG/Private/Elements/PCGAttributeTransferElement.cpp
new file mode 100644
--- /dev/null
+++ b/Engine/Plugins/Experimental/PCG/Source/PCG/Private/Elements/PCGAttributeTransferElement.cpp
@@ -0,0 +1,181 @@
+// Copyright Epic Games, Inc. All Rights Reserved.
+
+#include "Elements/PCGAttributeTransferElement.h"
+
+#include "PCGData.h"
+#include "PCGParamData.h"
+#include "Data/PCGPointData.h"
+#include "Metadata/PCGMetadata.h"
+#include "Metadata/PCGMetadataAttribute.h"
+
+#if WITH_EDITOR
+FName UPCGAttributeTransferSettings::GetDefaultNodeName() const
+{
+	return PCGAttributeTransferConstants::NodeName;
+}
+#endif
+
+FName UPCGAttributeTransferSettings::AdditionalTaskName() const
+{
+	if (SourceAttributeName != TargetAttributeName && TargetAttributeName != NAME_None)
+	{
+		return FName(FString::Printf(TEXT("%s %s to %s"), *PCGAttributeTransferConstants::NodeName.ToString(), *SourceAttributeName.ToString(), *TargetAttributeName.ToString()));
+	}
+	else
+	{
+		return FName(FString::Printf(TEXT("%s %s"), *PCGAttributeTransferConstants::NodeName.ToString(), *SourceAttributeName.ToString()));
+	}
+}
+
+TArray<FPCGPinProperties> UPCGAttributeTransferSettings::InputPinProperties() const
+{
+	TArray<FPCGPinProperties> PinProperties;
+	PinProperties.Emplace(PCGAttributeTransferConstants::TargetLabel, EPCGDataType::Spatial, /*bInAllowMultipleConnections=*/ false);
+	PinProperties.Emplace(PCGAttributeTransferConstants::SourceLabel, EPCGDataType::Spatial, /*bInAllowMultipleConnections=*/ false);
+
+	return PinProperties;
+}
+
+FPCGElementPtr UPCGAttributeTransferSettings::CreateElement() const
+{
+	return MakeShared<FPCGAttributeTransferElement>();
+}
+
+bool FPCGAttributeTransferElement::ExecuteInternal(FPCGContext* Context) const
+{
+	TRACE_CPUPROFILER_EVENT_SCOPE(FPCGAttributeTransferElement::Execute);
+
+	check(Context);
+
+	const UPCGAttributeTransferSettings* Settings = Context->GetInputSettings<UPCGAttributeTransferSettings>();
+	check(Settings);
+
+	TArray<FPCGTaggedData> SourceInputs = Context->InputData.GetInputsByPin(PCGAttributeTransferConstants::SourceLabel);
+	TArray<FPCGTaggedData> TargetInputs = Context->InputData.GetInputsByPin(PCGAttributeTransferConstants::TargetLabel);
+
+	if (SourceInputs.Num() != 1 || TargetInputs.Num() != 1)
+	{
+		PCGE_LOG(Error, "Wrong number of inputs");
+		return true;
+	}
+
+	const UPCGSpatialData* SourceData = Cast<UPCGSpatialData>(SourceInputs[0].Data);
+	const UPCGSpatialData* TargetData = Cast<UPCGSpatialData>(TargetInputs[0].Data);
+
+	if (!SourceData || !TargetData || SourceData->IsA<UPCGPointData>() != TargetData->IsA<UPCGPointData>())
+	{
+		PCGE_LOG(Error, "Only support spatial to spatial data or point to point data");
+		return true;
+	}
+
+	if (!SourceData->Metadata)
+	{
+		PCGE_LOG(Error, "Source doesn't have metadata");
+		return true;
+	}
+
+	if (!SourceData->Metadata->HasAttribute(Settings->SourceAttributeName))
+	{
+		PCGE_LOG(Error, "Source doesn't have the attribute \"%s\"", *Settings->SourceAttributeName.ToString());
+		return true;
+	}
+
+	UPCGData* OutputData = nullptr;
+
+	// We only supports spatial to spatial or point to points, with the same number of points
+	if (SourceData->IsA<UPCGPointData>() && TargetData->IsA<UPCGPointData>())
+	{
+		OutputData = TransferPointToPoint(Context, CastChecked<UPCGPointData>(SourceData), CastChecked<UPCGPointData>(TargetData));
+	}
+	else
+	{
+		OutputData = TransferSpatialToSpatial(Context, SourceData, TargetData);
+	}
+
+	if (OutputData)
+	{
+		TArray<FPCGTaggedData>& Outputs = Context->OutputData.TaggedData;
+		FPCGTaggedData& Output = Outputs.Emplace_GetRef();
+		Output.Data = OutputData;
+	}
+
+	return true;
+}
+
+UPCGSpatialData* FPCGAttributeTransferElement::TransferSpatialToSpatial(FPCGContext* Context, const UPCGSpatialData* SourceData, const UPCGSpatialData* TargetData) const
+{
+	check(Context && SourceData && TargetData);
+
+	const UPCGAttributeTransferSettings* Settings = Context->GetInputSettings<UPCGAttributeTransferSettings>();
+	check(Settings);
+
+	// When transferring attribute from spatial to spatial, copy the attribute entries and values.
+	// They need to match their number of entries
+	if (SourceData->Metadata->GetItemCountForChild() != TargetData->Metadata->GetItemCountForChild())
+	{
+		PCGE_LOG(Error, "Source and target doesn't have the same number of metadata entries");
+		return nullptr;
+	}
+
+	UPCGSpatialData* NewSpatialData = DuplicateObject<UPCGSpatialData>(const_cast<UPCGSpatialData*>(TargetData), nullptr);
+	NewSpatialData->InitializeFromData(TargetData);
+
+	const FName TargetAttributeName = Settings->TargetAttributeName == NAME_None ? Settings->SourceAttributeName : Settings->TargetAttributeName;
+
+	// Making sure the target attribute doesn't exists in the target
+	if (NewSpatialData->Metadata->HasAttribute(TargetAttributeName))
+	{
+		NewSpatialData->Metadata->DeleteAttribute(TargetAttributeName);
+	}
+
+	NewSpatialData->Metadata->CopyAttribute(SourceData->Metadata, Settings->SourceAttributeName, TargetAttributeName);
+
+	return NewSpatialData;
+}
+
+UPCGPointData* FPCGAttributeTransferElement::TransferPointToPoint(FPCGContext* Context, const UPCGPointData* SourceData, const UPCGPointData* TargetData) const
+{
+	check(Context && SourceData && TargetData);
+
+	const UPCGAttributeTransferSettings* Settings = Context->GetInputSettings<UPCGAttributeTransferSettings>();
+	check(Settings);
+
+	// When transferring attribute point to point, iterate over the source point and set the attribute to the target point.
+	// They need to match their number of points
+	if (SourceData->GetPoints().Num() != TargetData->GetPoints().Num())
+	{
+		PCGE_LOG(Error, "Source and target doesn't have the same number of points");
+		return nullptr;
+	}
+
+	UPCGPointData* NewPointData = NewObject<UPCGPointData>();
+	NewPointData->InitializeFromData(TargetData);
+
+	const TArray<FPCGPoint>& SourcePoints = SourceData->GetPoints();
+	TArray<FPCGPoint>& TargetPoints = NewPointData->GetMutablePoints();
+
+	TargetPoints = TargetData->GetPoints();
+
+	const FName TargetAttributeName = Settings->TargetAttributeName == NAME_None ? Settings->SourceAttributeName : Settings->TargetAttributeName;
+
+	// Making sure the target attribute doesn't exist in the target
+	if (NewPointData->Metadata->HasAttribute(TargetAttributeName))
+	{
+		NewPointData->Metadata->DeleteAttribute(TargetAttributeName);
+	}
+
+	const FPCGMetadataAttributeBase* SourceAttribute = SourceData->Metadata->GetConstAttribute(Settings->SourceAttributeName);
+	FPCGMetadataAttributeBase* TargetAttribute = NewPointData->Metadata->CopyAttribute(SourceAttribute, TargetAttributeName, /*bKeepParent=*/ false, /*bCopyEntries=*/ false, /*bCopyValues=*/ true);
+
+	for (int32 i = 0; i < SourcePoints.Num(); ++i)
+	{
+		PCGMetadataEntryKey& TargetKey = TargetPoints[i].MetadataEntry;
+		NewPointData->Metadata->InitializeOnSet(TargetKey);
+
+		check(TargetKey != PCGInvalidEntryKey);
+
+		TargetAttribute->SetValueFromValueKey(TargetKey, SourceAttribute->GetValueKey(SourcePoints[i].MetadataEntry));
+	}
+
+	return NewPointData;
+}
\ No newline at end of file
diff --git a/Engine/Plugins/Experimental/PCG/Source/PCG/Private/Elements/PCGCreateAttribute.cpp b/Engine/Plugins/Experimental/PCG/Source/PCG/Private/Elements/PCGCreateAttribute.cpp
deleted file mode 100644
diff --git a/Engine/Plugins/Experimental/PCG/Source/PCG/Private/Elements/PCGCreateAttributeBase.cpp b/Engine/Plugins/Experimental/PCG/Source/PCG/Private/Elements/PCGCreateAttributeBase.cpp
--- a/Engine/Plugins/Experimental/PCG/Source/PCG/Private/Elements/PCGCreateAttributeBase.cpp
+++ b/Engine/Plugins/Experimental/PCG/Source/PCG/Private/Elements/PCGCreateAttributeBase.cpp
@@ -2,15 +2,24 @@
 
 #include "Elements/PCGCreateAttributeBase.h"
 
+#include "PCGData.h"
 #include "PCGParamData.h"
+#include "Data/PCGPointData.h"
 #include "Elements/Metadata/PCGMetadataElementCommon.h"
-#include "Metadata/PCGMetadata.h"
 #include "Helpers/PCGSettingsHelpers.h"
+#include "Metadata/PCGMetadata.h"
+#include "Metadata/PCGMetadataAttribute.h"
 
-namespace PCGAttributeCreationBaseSettings
+namespace PCGCreateAttributeConstants
+{
+	const FName NodeName = TEXT("CreateAttribute");
+	const FName SourceLabel = TEXT("Source");
+}
+
+namespace PCGCreateAttributeElement
 {
 	template <typename Func>
-	decltype(auto) Dispatcher(const UPCGAttributeCreationBaseSettings* Settings, UPCGParamData* Params, Func Callback)
+	decltype(auto) Dispatcher(const UPCGCreateAttributeSettings* Settings, const UPCGParamData* Params, Func Callback)
 	{
 		using ReturnType = decltype(Callback(double{}));
 
@@ -52,97 +61,239 @@ namespace PCGAttributeCreationBaseSettings
 	}
 }
 
-FPCGMetadataAttributeBase* UPCGAttributeCreationBaseSettings::ClearOrCreateAttribute(UPCGMetadata* Metadata, UPCGParamData* Params) const
+FName UPCGCreateAttributeSettings::AdditionalTaskName() const
 {
-	check(Metadata);
+	if (bFromSourceParam)
+	{
+		const FName NodeName = PCGCreateAttributeConstants::NodeName;
 
-	auto CreateAttribute = [this, Metadata](auto&& Value) -> FPCGMetadataAttributeBase*
+		if (OutputAttributeName == NAME_None && SourceParamAttributeName == NAME_None)
+		{
+			return NodeName;
+		}
+		else
+		{
+			const FString AttributeName = (OutputAttributeName == NAME_None ? SourceParamAttributeName : OutputAttributeName).ToString();
+			return FName(FString::Printf(TEXT("%s %s"), *NodeName.ToString(), *AttributeName));
+		}
+	}
+	else
 	{
-		return PCGMetadataElementCommon::ClearOrCreateAttribute(Metadata, OutputAttributeName, Value);
-	};
+		const FString Name = OutputAttributeName.ToString();
 
-	return PCGAttributeCreationBaseSettings::Dispatcher(this, Params, CreateAttribute);
+		switch (Type)
+		{
+		case EPCGMetadataTypes::Integer64:
+			return FName(FString::Printf(TEXT("%s: %lld"), *Name, IntValue));
+		case EPCGMetadataTypes::Double:
+			return FName(FString::Printf(TEXT("%s: %.2f"), *Name, DoubleValue));
+		case EPCGMetadataTypes::String:
+			return FName(FString::Printf(TEXT("%s: \"%s\""), *Name, *StringValue));
+		case EPCGMetadataTypes::Name:
+			return FName(FString::Printf(TEXT("%s: N(\"%s\")"), *Name, *NameValue.ToString()));
+		case EPCGMetadataTypes::Vector2:
+			return FName(FString::Printf(TEXT("%s: V(%.2f, %.2f)"), *Name, Vector2Value.X, Vector2Value.Y));
+		case EPCGMetadataTypes::Vector:
+			return FName(FString::Printf(TEXT("%s: V(%.2f, %.2f, %.2f)"), *Name, VectorValue.X, VectorValue.Y, VectorValue.Z));
+		case EPCGMetadataTypes::Vector4:
+			return FName(FString::Printf(TEXT("%s: V(%.2f, %.2f, %.2f, %.2f)"), *Name, Vector4Value.X, Vector4Value.Y, Vector4Value.Z, Vector4Value.W));
+		case EPCGMetadataTypes::Rotator:
+			return FName(FString::Printf(TEXT("%s: R(%.2f, %.2f, %.2f)"), *Name, RotatorValue.Roll, RotatorValue.Pitch, RotatorValue.Yaw));
+		case EPCGMetadataTypes::Quaternion:
+			return FName(FString::Printf(TEXT("%s: Q(%.2f, %.2f, %.2f, %.2f)"), *Name, QuatValue.X, QuatValue.Y, QuatValue.Z, QuatValue.W));
+		case EPCGMetadataTypes::Transform:
+			return FName(FString::Printf(TEXT("%s: Transform"), *Name));
+		case EPCGMetadataTypes::Boolean:
+			return FName(FString::Printf(TEXT("%s: %s"), *Name, (BoolValue ? TEXT("True") : TEXT("False"))));
+		default:
+			return NAME_None;
+		}
+	}
 }
 
-PCGMetadataEntryKey UPCGAttributeCreationBaseSettings::SetAttribute(FPCGMetadataAttributeBase* Attribute, UPCGMetadata* Metadata, PCGMetadataEntryKey EntryKey, UPCGParamData* Params) const
+#if WITH_EDITOR
+FName UPCGCreateAttributeSettings::GetDefaultNodeName() const
 {
-	check(Attribute && Metadata);
+	return PCGCreateAttributeConstants::NodeName;
+}
+#endif // WITH_EDITOR
 
-	auto SetAttribute = [this, Attribute, EntryKey, Metadata](auto&& Value) -> PCGMetadataEntryKey
-	{
-		using AttributeType = std::remove_reference_t<decltype(Value)>;
+TArray<FPCGPinProperties> UPCGCreateAttributeSettings::InputPinProperties() const
+{
+	TArray<FPCGPinProperties> PinProperties;
+	PinProperties.Emplace(PCGPinConstants::DefaultInputLabel, EPCGDataType::Any, /*bInAllowMultipleConnections=*/ true);
 
-		check(Attribute->GetTypeId() == PCG::Private::MetadataTypes<AttributeType>::Id);
+	if (bFromSourceParam)
+	{
+		PinProperties.Emplace(PCGCreateAttributeConstants::SourceLabel, EPCGDataType::Param, /*bInAllowMultipleConnections=*/ false);
+	}
 
-		const PCGMetadataEntryKey FinalKey = (EntryKey == PCGInvalidEntryKey) ? Metadata->AddEntry() : EntryKey;
+	return PinProperties;
+}
 
-		static_cast<FPCGMetadataAttribute<AttributeType>*>(Attribute)->SetValue(FinalKey, Value);
+TArray<FPCGPinProperties> UPCGCreateAttributeSettings::OutputPinProperties() const
+{
+	TArray<FPCGPinProperties> PinProperties;
+	PinProperties.Emplace(PCGPinConstants::DefaultOutputLabel, EPCGDataType::Any);
 
-		return FinalKey;
-	};
+	return PinProperties;
+}
 
-	return PCGAttributeCreationBaseSettings::Dispatcher(this, Params, SetAttribute);
+FPCGElementPtr UPCGCreateAttributeSettings::CreateElement() const
+{
+	return MakeShared<FPCGCreateAttributeElement>();
 }
 
-void UPCGAttributeCreationBaseSettings::SetAttribute(FPCGMetadataAttributeBase* Attribute, UPCGMetadata* Metadata, TArray<FPCGPoint>& Points, UPCGParamData* Params) const
+bool FPCGCreateAttributeElement::ExecuteInternal(FPCGContext* Context) const
 {
-	if (Points.IsEmpty())
+	TRACE_CPUPROFILER_EVENT_SCOPE(FPCGCreateAttributeElement::Execute);
+
+	check(Context);
+
+	const UPCGCreateAttributeSettings* Settings = Context->GetInputSettings<UPCGCreateAttributeSettings>();
+	check(Settings);
+
+	TArray<FPCGTaggedData> SourceParams = Context->InputData.GetInputsByPin(PCGCreateAttributeConstants::SourceLabel);
+	UPCGParamData* SourceParamData = nullptr;
+	FName SourceParamAttributeName = NAME_None;
+
+	if (Settings->bFromSourceParam)
 	{
-		return;
+		if (SourceParams.IsEmpty())
+		{
+			PCGE_LOG(Error, "Source param was not provided.");
+			return true;
+		}
+
+		SourceParamData = CastChecked<UPCGParamData>(SourceParams[0].Data);
+
+		if (!SourceParamData->Metadata)
+		{
+			PCGE_LOG(Error, "Source param data doesn't have metadata");
+			return true;
+		}
+
+		SourceParamAttributeName = Settings->SourceParamAttributeName == NAME_None ? SourceParamData->Metadata->GetLatestAttributeNameOrNone() : Settings->SourceParamAttributeName;
+
+		if (!SourceParamData->Metadata->HasAttribute(SourceParamAttributeName))
+		{
+			PCGE_LOG(Error, "Source param data doesn't have an attribute \"%s\"", *SourceParamAttributeName.ToString());
+			return true;
+		}
 	}
 
-	check(Attribute && Metadata);
+	TArray<FPCGTaggedData> Inputs = Context->InputData.GetInputsByPin(PCGPinConstants::DefaultInputLabel);
 
-	auto SetAttribute = [this, Attribute, Metadata, &Points](auto&& Value) -> void
+	// If the input is empty, we will create a new ParamData.
+	// We can re-use this newly object as the output
+	bool bCanReuseInputData = false;
+	if (Inputs.IsEmpty())
 	{
-		using AttributeType = std::remove_reference_t<decltype(Value)>;
+		FPCGTaggedData& NewData = Inputs.Emplace_GetRef();
+		NewData.Data = NewObject<UPCGParamData>();
+		NewData.Pin = PCGPinConstants::DefaultInputLabel;
+		bCanReuseInputData = true;
+	}
 
-		check(Attribute->GetTypeId() == PCG::Private::MetadataTypes<AttributeType>::Id);
+	for (const FPCGTaggedData& InputTaggedData : Inputs)
+	{
+		const UPCGData* InputData = InputTaggedData.Data;
+		UPCGData* OutputData = nullptr;
+
+		UPCGMetadata* Metadata = nullptr;
+
+		bool bShouldAddNewEntry = false;
+
+		if (const UPCGSpatialData* InputSpatialData = Cast<UPCGSpatialData>(InputData))
+		{
+			UPCGSpatialData* NewSpatialData = DuplicateObject<UPCGSpatialData>(const_cast<UPCGSpatialData*>(InputSpatialData), nullptr);
+			NewSpatialData->Metadata = NewObject<UPCGMetadata>();
+			NewSpatialData->InitializeFromData(InputSpatialData, /*InMetadataParentOverride=*/ nullptr, /*bInheritMetadata=*/ Settings->bKeepExistingAttributes);
+
+			OutputData = NewSpatialData;
+			Metadata = NewSpatialData->Metadata;
+		}
+		else if (const UPCGParamData* InputParamData = Cast<UPCGParamData>(InputData))
+		{
+			// If we can reuse input data, it is safe to const_cast, as it was created by ourselves above.
+			UPCGParamData* NewParamData = bCanReuseInputData ? const_cast<UPCGParamData*>(InputParamData) : NewObject<UPCGParamData>();
+			NewParamData->Metadata->Initialize((!bCanReuseInputData && Settings->bKeepExistingAttributes) ? InputParamData->Metadata : nullptr);
 
-		FPCGMetadataAttribute<AttributeType>* TypedAttribute = static_cast<FPCGMetadataAttribute<AttributeType>*>(Attribute);
+			OutputData = NewParamData;
+			Metadata = NewParamData->Metadata;
+
+			// In case of param data, we want to add a new entry too
+			bShouldAddNewEntry = true;
+		}
+		else
+		{
+			PCGE_LOG(Error, "Invalid data as input. Only support spatial and params");
+			continue;
+		}
+
+		const FName OutputAttributeName = (Settings->bFromSourceParam && Settings->OutputAttributeName == NAME_None) ? SourceParamAttributeName : Settings->OutputAttributeName;
+
+		FPCGMetadataAttributeBase* Attribute = nullptr;
+
+		if (Settings->bFromSourceParam)
+		{
+			const FPCGMetadataAttributeBase* SourceAttribute = SourceParamData->Metadata->GetConstAttribute(SourceParamAttributeName);
+			Attribute = Metadata->CopyAttribute(SourceAttribute, OutputAttributeName, /*bKeepParent=*/false, /*bCopyEntries=*/bShouldAddNewEntry, /*bCopyValues=*/bShouldAddNewEntry);
+		}
+		else
+		{
+			Attribute = ClearOrCreateAttribute(Settings, Metadata, nullptr, &OutputAttributeName);
+		}
 
-		for (FPCGPoint& Point : Points)
+		if (!Attribute)
 		{
-			PCGMetadataEntryKey EntryKey = (Point.MetadataEntry == PCGInvalidEntryKey) ? Metadata->AddEntry() : Point.MetadataEntry;
+			PCGE_LOG(Error, "Error while creating attribute %s", *OutputAttributeName.ToString());
+			continue;
+		}
 
-			TypedAttribute->SetValue(EntryKey, Value);
+		TArray<FPCGTaggedData>& Outputs = Context->OutputData.TaggedData;
+		FPCGTaggedData& Output = Outputs.Emplace_GetRef();
+		Output.Data = OutputData;
 
-			Point.MetadataEntry = EntryKey;
+		// Add a new entry if it is a param data and not from source (because entries are already copied)
+		if (bShouldAddNewEntry && !Settings->bFromSourceParam)
+		{
+			PCGMetadataEntryKey EntryKey = Metadata->AddEntry();
+			SetAttribute(Settings, Attribute, Metadata, EntryKey, nullptr);
 		}
+	}
+
+	return true;
+}
+
+FPCGMetadataAttributeBase* FPCGCreateAttributeElement::ClearOrCreateAttribute(const UPCGCreateAttributeSettings* Settings, UPCGMetadata* Metadata, const UPCGParamData* Params, const FName* OutputAttributeNameOverride) const
+{
+	check(Metadata);
+
+	auto CreateAttribute = [Settings, Metadata, OutputAttributeNameOverride](auto&& Value) -> FPCGMetadataAttributeBase*
+	{
+		return PCGMetadataElementCommon::ClearOrCreateAttribute(Metadata, OutputAttributeNameOverride ? *OutputAttributeNameOverride : Settings->OutputAttributeName, Value);
 	};
 
-	PCGAttributeCreationBaseSettings::Dispatcher(this, Params, SetAttribute);
+	return PCGCreateAttributeElement::Dispatcher(Settings, Params, CreateAttribute);
 }
 
-FName UPCGAttributeCreationBaseSettings::AdditionalTaskName() const
+PCGMetadataEntryKey FPCGCreateAttributeElement::SetAttribute(const UPCGCreateAttributeSettings* Settings, FPCGMetadataAttributeBase* Attribute, UPCGMetadata* Metadata, PCGMetadataEntryKey EntryKey, const UPCGParamData* Params) const
 {
-	const FString Name = OutputAttributeName.ToString();
+	check(Attribute && Metadata);
 
-	switch (Type)
+	auto SetAttribute = [Attribute, EntryKey, Metadata](auto&& Value) -> PCGMetadataEntryKey
 	{
-	case EPCGMetadataTypes::Integer64:
-		return FName(FString::Printf(TEXT("%s: %ll"), *Name, IntValue));
-	case EPCGMetadataTypes::Double:
-		return FName(FString::Printf(TEXT("%s: %.2f"), *Name, DoubleValue));
-	case EPCGMetadataTypes::String:
-		return FName(FString::Printf(TEXT("%s: \"%s\""), *Name, *StringValue));
-	case EPCGMetadataTypes::Name:
-		return FName(FString::Printf(TEXT("%s: N(\"%s\")"), *Name, *NameValue.ToString()));
-	case EPCGMetadataTypes::Vector2:
-		return FName(FString::Printf(TEXT("%s: V(%.2f, %.2f)"), *Name, Vector2Value.X, Vector2Value.Y));
-	case EPCGMetadataTypes::Vector:
-		return FName(FString::Printf(TEXT("%s: V(%.2f, %.2f, %.2f)"), *Name, VectorValue.X, VectorValue.Y, VectorValue.Z));
-	case EPCGMetadataTypes::Vector4:
-		return FName(FString::Printf(TEXT("%s: V(%.2f, %.2f, %.2f, %.2f)"), *Name, Vector4Value.X, Vector4Value.Y, Vector4Value.Z, Vector4Value.W));
-	case EPCGMetadataTypes::Rotator:
-		return FName(FString::Printf(TEXT("%s: R(%.2f, %.2f, %.2f)"), *Name, RotatorValue.Roll, RotatorValue.Pitch, RotatorValue.Yaw));
-	case EPCGMetadataTypes::Quaternion:
-		return FName(FString::Printf(TEXT("%s: Q(%.2f, %.2f, %.2f, %.2f)"), *Name, QuatValue.X, QuatValue.Y, QuatValue.Z, QuatValue.W));
-	case EPCGMetadataTypes::Transform:
-		return FName(FString::Printf(TEXT("%s: Transform"), *Name));
-	case EPCGMetadataTypes::Boolean:
-		return FName(FString::Printf(TEXT("%s: %s"), *Name, (BoolValue ? TEXT("True") : TEXT("False"))));
-	default:
-		return NAME_None;
-	}
-}
\ No newline at end of file
+		using AttributeType = std::remove_reference_t<decltype(Value)>;
+
+		check(Attribute->GetTypeId() == PCG::Private::MetadataTypes<AttributeType>::Id);
+
+		const PCGMetadataEntryKey FinalKey = (EntryKey == PCGInvalidEntryKey) ? Metadata->AddEntry() : EntryKey;
+
+		static_cast<FPCGMetadataAttribute<AttributeType>*>(Attribute)->SetValue(FinalKey, Value);
+
+		return FinalKey;
+	};
+
+	return PCGCreateAttributeElement::Dispatcher(Settings, Params, SetAttribute);
+}
diff --git a/Engine/Plugins/Experimental/PCG/Source/PCG/Private/Elements/PCGCreateParamData.cpp b/Engine/Plugins/Experimental/PCG/Source/PCG/Private/Elements/PCGCreateParamData.cpp
deleted file mode 100644
diff --git a/Engine/Plugins/Experimental/PCG/Source/PCG/Private/Metadata/PCGMetadata.cpp b/Engine/Plugins/Experimental/PCG/Source/PCG/Private/Metadata/PCGMetadata.cpp
--- a/Engine/Plugins/Experimental/PCG/Source/PCG/Private/Metadata/PCGMetadata.cpp
+++ b/Engine/Plugins/Experimental/PCG/Source/PCG/Private/Metadata/PCGMetadata.cpp
@@ -69,6 +69,11 @@ void UPCGMetadata::BeginDestroy()
 }
 
 void UPCGMetadata::Initialize(const UPCGMetadata* InParent)
+{
+	Initialize(InParent, /*bAddAttributesFromParent=*/true);
+}
+
+void UPCGMetadata::Initialize(const UPCGMetadata* InParent, bool bAddAttributesFromParent)
 {
 	if (Parent.IsValid() || Attributes.Num() != 0)
 	{
@@ -78,7 +83,11 @@ void UPCGMetadata::Initialize(const UPCGMetadata* InParent)
 
 	Parent = ((InParent != this) ? InParent : nullptr);
 	ItemKeyOffset = Parent.IsValid() ? Parent->GetItemCountForChild() : 0;
-	AddAttributes(InParent);
+
+	if (bAddAttributesFromParent)
+	{
+		AddAttributes(InParent);
+	}
 }
 
 void UPCGMetadata::InitializeAsCopy(const UPCGMetadata* InMetadataToCopy)
@@ -139,7 +148,7 @@ void UPCGMetadata::AddAttribute(const UPCGMetadata* InOther, FName AttributeName
 		return;
 	}
 
-	CopyAttribute(InOther->GetConstAttribute(AttributeName), AttributeName, /*bKeepParent=*/false, /*bCopyEntries=*/false, /*bCopyValues=*/false);
+	CopyAttribute(InOther->GetConstAttribute(AttributeName), AttributeName, /*bKeepParent=*/InOther == Parent, /*bCopyEntries=*/false, /*bCopyValues=*/false);
 	OtherParents.Add(InOther);
 }
 
diff --git a/Engine/Plugins/Experimental/PCG/Source/PCG/Public/Elements/PCGAttributeFilterElement.h b/Engine/Plugins/Experimental/PCG/Source/PCG/Public/Elements/PCGAttributeFilterElement.h
new file mode 100644
--- /dev/null
+++ b/Engine/Plugins/Experimental/PCG/Source/PCG/Public/Elements/PCGAttributeFilterElement.h
@@ -0,0 +1,48 @@
+// Copyright Epic Games, Inc. All Rights Reserved.
+
+#pragma once
+
+#include "PCGElement.h"
+#include "PCGSettings.h"
+
+#include "PCGAttributeFilterElement.generated.h"
+
+/**
+* Filter the attributes from a given input metadata.
+* Will remove all attributes that are not listed in AttributesToKeep.
+* If an attribute to keep is not in the original metadata, it won't be added.
+* 
+* The output will be the original data with the updated metadata.
+*/
+UCLASS(BlueprintType, ClassGroup = (Procedural))
+class PCG_API UPCGAttributeFilterSettings : public UPCGSettings
+{
+	GENERATED_BODY()
+
+public:
+	//~Begin UPCGSettings interface
+#if WITH_EDITOR
+	virtual FName GetDefaultNodeName() const override;
+	virtual EPCGSettingsType GetType() const override { return EPCGSettingsType::Metadata; }
+#endif
+
+	virtual TArray<FPCGPinProperties> OutputPinProperties() const override;
+	//~End UPCGSettings interface
+
+	UPROPERTY(BlueprintReadWrite, EditAnywhere, Category = Settings)
+	TArray<FName> AttributesToKeep;
+
+protected:
+	virtual FPCGElementPtr CreateElement() const override;
+};
+
+
+class FPCGAttributeFilterElement : public FSimplePCGElement
+{
+public:
+	virtual bool CanExecuteOnlyOnMainThread(FPCGContext* Context) const { return true; }
+	virtual bool IsCacheable(const UPCGSettings* InSettings) const override { return false; }
+
+protected:
+	virtual bool ExecuteInternal(FPCGContext* Context) const override;
+};
diff --git a/Engine/Plugins/Experimental/PCG/Source/PCG/Public/Elements/PCGAttributeTransferElement.h b/Engine/Plugins/Experimental/PCG/Source/PCG/Public/Elements/PCGAttributeTransferElement.h
new file mode 100644
--- /dev/null
+++ b/Engine/Plugins/Experimental/PCG/Source/PCG/Public/Elements/PCGAttributeTransferElement.h
@@ -0,0 +1,69 @@
+// Copyright Epic Games, Inc. All Rights Reserved.
+
+#pragma once
+
+#include "PCGElement.h"
+#include "PCGSettings.h"
+
+#include "PCGAttributeTransferElement.generated.h"
+
+class UPCGPointData;
+class UPCGSpatialData;
+
+namespace PCGAttributeTransferConstants
+{
+	const FName NodeName = TEXT("TransferAttribute");
+	const FName SourceLabel = TEXT("Source");
+	const FName TargetLabel = TEXT("Target");
+}
+
+/**
+* Transfer an attribute from a source metadata to a target data.
+* Only support Spatial to Spatial and Points to Points, and they need to match.
+*  - For Spatial data, number of entries in the metadata should be the same between source and target.
+*  - For Point data, number of points should be the same between source and target.
+* 
+* The output will be the target data with the updated metadata.
+* If the TargetAttributeName is None, it will use SourceAttributeName instead.
+* 
+* To do the same but with a Source param data, use CreateAttribute.
+*/
+UCLASS(BlueprintType, ClassGroup = (Procedural))
+class PCG_API UPCGAttributeTransferSettings : public UPCGSettings
+{
+	GENERATED_BODY()
+
+public:
+	//~Begin UPCGSettings interface
+#if WITH_EDITOR
+	virtual FName GetDefaultNodeName() const override;
+	virtual EPCGSettingsType GetType() const override { return EPCGSettingsType::Metadata; }
+#endif
+	virtual FName AdditionalTaskName() const override;
+
+	virtual TArray<FPCGPinProperties> InputPinProperties() const override;
+	//~End UPCGSettings interface
+
+	UPROPERTY(BlueprintReadWrite, EditAnywhere, Category = Settings)
+	FName SourceAttributeName;
+
+	UPROPERTY(BlueprintReadWrite, EditAnywhere, Category = Settings)
+	FName TargetAttributeName;
+
+protected:
+	virtual FPCGElementPtr CreateElement() const override;
+};
+
+
+class FPCGAttributeTransferElement : public FSimplePCGElement
+{
+public:
+	virtual bool CanExecuteOnlyOnMainThread(FPCGContext* Context) const { return true; }
+	virtual bool IsCacheable(const UPCGSettings* InSettings) const override { return false; }
+
+protected:
+	virtual bool ExecuteInternal(FPCGContext* Context) const override;
+
+	UPCGSpatialData* TransferSpatialToSpatial(FPCGContext* Context, const UPCGSpatialData* SourceData, const UPCGSpatialData* TargetData) const;
+	UPCGPointData* TransferPointToPoint(FPCGContext* Context, const UPCGPointData* SourceData, const UPCGPointData* TargetData) const;
+};
diff --git a/Engine/Plugins/Experimental/PCG/Source/PCG/Public/Elements/PCGCreateAttribute.h b/Engine/Plugins/Experimental/PCG/Source/PCG/Public/Elements/PCGCreateAttribute.h
deleted file mode 100644
diff --git a/Engine/Plugins/Experimental/PCG/Source/PCG/Public/Elements/PCGCreateAttributeBase.h b/Engine/Plugins/Experimental/PCG/Source/PCG/Public/Elements/PCGCreateAttributeBase.h
--- a/Engine/Plugins/Experimental/PCG/Source/PCG/Public/Elements/PCGCreateAttributeBase.h
+++ b/Engine/Plugins/Experimental/PCG/Source/PCG/Public/Elements/PCGCreateAttributeBase.h
@@ -13,67 +13,96 @@ class FPCGMetadataAttributeBase;
 class UPCGMetadata;
 class UPCGParamData;
 
-/*
-* Base settings class for any node that needs to create new attributes.
-* Contains all supported PCG types, a utility function to change the node name depending of the created attribute
-* and a utility function to create/set the attribute value depending on the type in the settings;
+/**
+* Node to create an attribute from a Spatial data or a ParamData.
+* If no input is provided, the node behave as a "CreateParamData" and will output a ParamData
+* with the given attribute set.
+* 
+* Note: This need to be updated if we ever add new types.
 */
-UCLASS(BlueprintType, ClassGroup = (Procedural), Abstract)
-class PCG_API UPCGAttributeCreationBaseSettings : public UPCGSettings
+UCLASS(BlueprintType, ClassGroup = (Procedural))
+class PCG_API UPCGCreateAttributeSettings : public UPCGSettings
 {
 	GENERATED_BODY()
-
 public:
 	//~Begin UPCGSettings interface
+#if WITH_EDITOR
+	virtual FName GetDefaultNodeName() const override;
+	virtual EPCGSettingsType GetType() const override { return EPCGSettingsType::Metadata; }
+#endif // WITH_EDITOR
+
+	virtual TArray<FPCGPinProperties> InputPinProperties() const override;
+	virtual TArray<FPCGPinProperties> OutputPinProperties() const override;
+
 	virtual FName AdditionalTaskName() const override;
 	//~End UPCGSettings interface
 
-	/* Create (or clear) an attribute named by OutputAttributeName and depending on the selected type. Value can be overridden by params. Default value will be set to the specified value. */
-	FPCGMetadataAttributeBase* ClearOrCreateAttribute(UPCGMetadata* Metadata, UPCGParamData* Params = nullptr) const;
+	UPROPERTY(BlueprintReadWrite, EditAnywhere, Category = Settings)
+	bool bKeepExistingAttributes = false;
 
-	/* Set an entry defined by EntryKey (or create one if it is invalid) in the Attribute depending on the selected type. Value can be overridden by params. */
-	PCGMetadataEntryKey SetAttribute(FPCGMetadataAttributeBase* Attribute, UPCGMetadata* Metadata, PCGMetadataEntryKey EntryKey, UPCGParamData* Params = nullptr) const;
+	UPROPERTY(BlueprintReadWrite, EditAnywhere, Category = Settings)
+	FName OutputAttributeName = NAME_None;
 
-	/* Set attributes for each point in the Attribute depending on the selected type. Value can be overridden by params. */
-	void SetAttribute(FPCGMetadataAttributeBase* Attribute, UPCGMetadata* Metadata, TArray<FPCGPoint>& Points, UPCGParamData* Params = nullptr) const;
+	UPROPERTY(BlueprintReadWrite, EditAnywhere, Category = Settings)
+	bool bFromSourceParam = false;
 
-	UPROPERTY(BlueprintReadWrite, EditAnywhere, Category = Settings, meta = (ValidEnumValues = "Double, Integer64, Vector2, Vector, Vector4, Quaternion, Transform, String, Boolean, Rotator, Name"))
-	EPCGMetadataTypes Type = EPCGMetadataTypes::Double;
+	UPROPERTY(BlueprintReadWrite, EditAnywhere, Category = Settings, meta = (EditCondition = "bFromSourceParam", EditConditionHides))
+	FName SourceParamAttributeName = NAME_None;
 
-	UPROPERTY(BlueprintReadWrite, EditAnywhere, Category = Settings)
-	FName OutputAttributeName = NAME_None;
+	UPROPERTY(BlueprintReadWrite, EditAnywhere, Category = Settings, meta = (ValidEnumValues = "Double, Integer64, Vector2, Vector, Vector4, Quaternion, Transform, String, Boolean, Rotator, Name", EditCondition = "!bFromSourceParam", EditConditionHides))
+	EPCGMetadataTypes Type = EPCGMetadataTypes::Double;
 
 	// All different types
-	UPROPERTY(BlueprintReadWrite, EditAnywhere, Category = Settings, meta = (EditCondition = "Type == EPCGMetadataTypes::Double", EditConditionHides))
+	UPROPERTY(BlueprintReadWrite, EditAnywhere, Category = Settings, meta = (EditCondition = "Type == EPCGMetadataTypes::Double && !bFromSourceParam", EditConditionHides))
 	double DoubleValue = 0.0;
 
-	UPROPERTY(BlueprintReadWrite, EditAnywhere, Category = Settings, meta = (EditCondition = "Type == EPCGMetadataTypes::Integer64", EditConditionHides))
+	UPROPERTY(BlueprintReadWrite, EditAnywhere, Category = Settings, meta = (EditCondition = "Type == EPCGMetadataTypes::Integer64 && !bFromSourceParam", EditConditionHides))
 	int64 IntValue = 0;
 
-	UPROPERTY(BlueprintReadWrite, EditAnywhere, Category = Settings, meta = (EditCondition = "Type == EPCGMetadataTypes::Vector2", EditConditionHides))
+	UPROPERTY(BlueprintReadWrite, EditAnywhere, Category = Settings, meta = (EditCondition = "Type == EPCGMetadataTypes::Vector2 && !bFromSourceParam", EditConditionHides))
 	FVector2D Vector2Value = FVector2D::ZeroVector;
 
-	UPROPERTY(BlueprintReadWrite, EditAnywhere, Category = Settings, meta = (EditCondition = "Type == EPCGMetadataTypes::Vector", EditConditionHides))
+	UPROPERTY(BlueprintReadWrite, EditAnywhere, Category = Settings, meta = (EditCondition = "Type == EPCGMetadataTypes::Vector && !bFromSourceParam", EditConditionHides))
 	FVector VectorValue = FVector::ZeroVector;
 
-	UPROPERTY(BlueprintReadWrite, EditAnywhere, Category = Settings, meta = (EditCondition = "Type == EPCGMetadataTypes::Vector4", EditConditionHides))
+	UPROPERTY(BlueprintReadWrite, EditAnywhere, Category = Settings, meta = (EditCondition = "Type == EPCGMetadataTypes::Vector4 && !bFromSourceParam", EditConditionHides))
 	FVector4 Vector4Value = FVector4::Zero();
 
-	UPROPERTY(BlueprintReadWrite, EditAnywhere, Category = Settings, meta = (EditCondition = "Type == EPCGMetadataTypes::Quaternion", EditConditionHides))
+	UPROPERTY(BlueprintReadWrite, EditAnywhere, Category = Settings, meta = (EditCondition = "Type == EPCGMetadataTypes::Quaternion && !bFromSourceParam", EditConditionHides))
 	FQuat QuatValue = FQuat::Identity;
 
-	UPROPERTY(BlueprintReadWrite, EditAnywhere, Category = Settings, meta = (EditCondition = "Type == EPCGMetadataTypes::Transform", EditConditionHides))
+	UPROPERTY(BlueprintReadWrite, EditAnywhere, Category = Settings, meta = (EditCondition = "Type == EPCGMetadataTypes::Transform && !bFromSourceParam", EditConditionHides))
 	FTransform TransformValue = FTransform::Identity;
 
-	UPROPERTY(BlueprintReadWrite, EditAnywhere, Category = Settings, meta = (EditCondition = "Type == EPCGMetadataTypes::String", EditConditionHides))
+	UPROPERTY(BlueprintReadWrite, EditAnywhere, Category = Settings, meta = (EditCondition = "Type == EPCGMetadataTypes::String && !bFromSourceParam", EditConditionHides))
 	FString StringValue = "";
 
-	UPROPERTY(BlueprintReadWrite, EditAnywhere, Category = Settings, meta = (EditCondition = "Type == EPCGMetadataTypes::Boolean", EditConditionHides))
+	UPROPERTY(BlueprintReadWrite, EditAnywhere, Category = Settings, meta = (EditCondition = "Type == EPCGMetadataTypes::Boolean && !bFromSourceParam", EditConditionHides))
 	bool BoolValue = false;
 
-	UPROPERTY(BlueprintReadWrite, EditAnywhere, Category = Settings, meta = (EditCondition = "Type == EPCGMetadataTypes::Rotator", EditConditionHides))
+	UPROPERTY(BlueprintReadWrite, EditAnywhere, Category = Settings, meta = (EditCondition = "Type == EPCGMetadataTypes::Rotator && !bFromSourceParam", EditConditionHides))
 	FRotator RotatorValue = FRotator::ZeroRotator;
 
-	UPROPERTY(BlueprintReadWrite, EditAnywhere, Category = Settings, meta = (EditCondition = "Type == EPCGMetadataTypes::Name", EditConditionHides))
+	UPROPERTY(BlueprintReadWrite, EditAnywhere, Category = Settings, meta = (EditCondition = "Type == EPCGMetadataTypes::Name && !bFromSourceParam", EditConditionHides))
 	FName NameValue = NAME_None;
+
+protected:
+	virtual FPCGElementPtr CreateElement() const override;
+};
+
+class FPCGCreateAttributeElement : public FSimplePCGElement
+{
+public:
+	virtual bool CanExecuteOnlyOnMainThread(FPCGContext* Context) const { return true; }
+	virtual bool IsCacheable(const UPCGSettings* InSettings) const override { return false; }
+
+protected:
+	virtual bool ExecuteInternal(FPCGContext* Context) const override;
+
+private:
+	/* Create (or clear) an attribute named by OutputAttributeName and depending on the selected type. Value can be overridden by params. Default value will be set to the specified value. */
+	FPCGMetadataAttributeBase* ClearOrCreateAttribute(const UPCGCreateAttributeSettings* Settings, UPCGMetadata* Metadata, const UPCGParamData* Params = nullptr, const FName* OutputAttributeNameOverride = nullptr) const;
+
+	/* Set an entry defined by EntryKey (or create one if it is invalid) in the Attribute depending on the selected type. Value can be overridden by params. */
+	PCGMetadataEntryKey SetAttribute(const UPCGCreateAttributeSettings* Settings, FPCGMetadataAttributeBase* Attribute, UPCGMetadata* Metadata, PCGMetadataEntryKey EntryKey, const UPCGParamData* Params = nullptr) const;
 };
\ No newline at end of file
diff --git a/Engine/Plugins/Experimental/PCG/Source/PCG/Public/Elements/PCGCreateParamData.h b/Engine/Plugins/Experimental/PCG/Source/PCG/Public/Elements/PCGCreateParamData.h
deleted file mode 100644
diff --git a/Engine/Plugins/Experimental/PCG/Source/PCG/Public/Helpers/PCGSettingsHelpers.h b/Engine/Plugins/Experimental/PCG/Source/PCG/Public/Helpers/PCGSettingsHelpers.h
--- a/Engine/Plugins/Experimental/PCG/Source/PCG/Public/Helpers/PCGSettingsHelpers.h
+++ b/Engine/Plugins/Experimental/PCG/Source/PCG/Public/Helpers/PCGSettingsHelpers.h
@@ -13,54 +13,48 @@ class UPCGSettings;
 
 namespace PCGSettingsHelpers
 {
-	// Utility function to return a value from an attribute that doesn't match the type
-	// Use if constexpr to remove un-necessary checks.
-	// Returns true if the conversion worked
-	template <typename T>
-	inline bool GetValueWithImplicitConversion(const FPCGMetadataAttributeBase* InAttribute, PCGMetadataEntryKey InKey, T& OutValue)
+	/** Utility function to get the value of type T from a param data or a default value
+	* @param InName - Attribute to get from the param
+	* @param InValue - Default value to return if the param doesn't have the given attribute
+	* @param InParams - ParamData to get the value from.
+	* @param InKey - Metadata Entry Key to get the value from.
+	*/
+	template<typename T, typename TEnableIf<!TIsEnumClass<T>::Value>::Type* = nullptr>
+	T GetValue(const FName& InName, const T& InValue, const UPCGParamData* InParams, PCGMetadataEntryKey InKey)
 	{
-#define ImplicitConversion(InType, OutType) if constexpr (std::is_same_v<T, OutType>) { if (InAttribute->GetTypeId() == PCG::Private::MetadataTypes<InType>::Id) { OutValue = static_cast<T>(static_cast<const FPCGMetadataAttribute<InType>*>(InAttribute)->GetValueFromItemKey(InKey)); return true; } }
-
-		ImplicitConversion(float, double);
-		ImplicitConversion(double, float);
-		ImplicitConversion(int32, int64);
-		ImplicitConversion(int64, int32);
 
-#undef ImplicitConversion
+		if (InParams && InParams->Metadata)
+		{
+			const FPCGMetadataAttributeBase* MatchingAttribute = InParams->Metadata->GetConstAttribute(InName);
 
-		return false;
-	}
+			if (!MatchingAttribute)
+			{
+				return InValue;
+			}
 
-	template<typename T, typename TEnableIf<!TIsEnumClass<T>::Value>::Type* = nullptr>
-	T GetValue(const FName& InName, const T& InValue, UPCGParamData* InParams, PCGMetadataEntryKey InKey, bool bAllowImplicitConversion = true)
-	{
-		if (InParams)
-		{
-			const FPCGMetadataAttributeBase* MatchingAttribute = InParams->Metadata ? InParams->Metadata->GetConstAttribute(InName) : nullptr;
-			if (MatchingAttribute)
+			auto GetTypedValue = [MatchingAttribute, &InValue, InKey](auto DummyValue) -> T
 			{
-				if (MatchingAttribute->GetTypeId() == PCG::Private::MetadataTypes<T>::Id)
+				using AttributeType = decltype(DummyValue);
+
+				const FPCGMetadataAttribute<AttributeType>* ParamAttribute = static_cast<const FPCGMetadataAttribute<AttributeType>*>(MatchingAttribute);
+
+				if constexpr (std::is_same_v<T, AttributeType>)
 				{
-					return static_cast<const FPCGMetadataAttribute<T>*>(MatchingAttribute)->GetValueFromItemKey(InKey);
+					return ParamAttribute->GetValueFromItemKey(InKey);
 				}
-
-				if (bAllowImplicitConversion)
+				else if constexpr (std::is_constructible_v<T, AttributeType>)
+				{
+					UE_LOG(LogPCG, Verbose, TEXT("[GetAttributeValue] Matching attribute was found, but not the right type. Implicit conversion done (%d vs %d)"), MatchingAttribute->GetTypeId(), PCG::Private::MetadataTypes<T>::Id);
+					return T(ParamAttribute->GetValueFromItemKey(InKey));
+				}
+				else
 				{
-					T OutValue{};
-					if (GetValueWithImplicitConversion<T>(MatchingAttribute, InKey, OutValue))
-					{
-						UE_LOG(LogPCG, Warning, TEXT("[GetAttributeValue] Matching attribute was found, but not the right type. Implicit conversion done (%d vs %d)"), MatchingAttribute->GetTypeId(), PCG::Private::MetadataTypes<T>::Id);
-						return OutValue;
-					}
+					UE_LOG(LogPCG, Error, TEXT("[GetAttributeValue] Matching attribute was found, but not the right type. %d vs %d"), MatchingAttribute->GetTypeId(), PCG::Private::MetadataTypes<T>::Id);
+					return InValue;
 				}
+			};
 
-				UE_LOG(LogPCG, Error, TEXT("[GetAttributeValue] Matching attribute was found, but not the right type. %d vs %d"), MatchingAttribute->GetTypeId(), PCG::Private::MetadataTypes<T>::Id);
-				return InValue;
-			}
-			else
-			{
-				return InValue;
-			}
+			return PCGMetadataAttribute::CallbackWithRightType(MatchingAttribute->GetTypeId(), GetTypedValue);
 		}
 		else
 		{
@@ -69,17 +63,17 @@ namespace PCGSettingsHelpers
 	}
 
 	template<typename T, typename TEnableIf<!TIsEnumClass<T>::Value>::Type* = nullptr>
-	T GetValue(const FName& InName, const T& InValue, UPCGParamData* InParams, bool bAllowImplicitConversion = true)
+	T GetValue(const FName& InName, const T& InValue, const UPCGParamData* InParams)
 	{
-		return GetValue(InName, InValue, InParams, 0, bAllowImplicitConversion);
+		return GetValue(InName, InValue, InParams, 0);
 	}
 
 	template<typename T, typename TEnableIf<!TIsEnumClass<T>::Value>::Type* = nullptr>
-	T GetValue(const FName& InName, const T& InValue, UPCGParamData* InParams, const FName& InParamName, bool bAllowImplicitConversion = true)
+	T GetValue(const FName& InName, const T& InValue, const UPCGParamData* InParams, const FName& InParamName)
 	{
 		if (InParams && InParamName != NAME_None)
 		{
-			return GetValue(InName, InValue, InParams, InParams->FindMetadataKey(InParamName), bAllowImplicitConversion);
+			return GetValue(InName, InValue, InParams, InParams->FindMetadataKey(InParamName));
 		}
 		else
 		{
@@ -89,40 +83,40 @@ namespace PCGSettingsHelpers
 
 	/** Specialized versions for enums */
 	template<typename T, typename TEnableIf<TIsEnumClass<T>::Value>::Type* = nullptr>
-	T GetValue(const FName& InName, const T& InValue, UPCGParamData* InParams, PCGMetadataEntryKey InKey, bool bAllowImplicitConversion = true)
+	T GetValue(const FName& InName, const T& InValue, const UPCGParamData* InParams, PCGMetadataEntryKey InKey)
 	{
-		return static_cast<T>(GetValue(InName, static_cast<__underlying_type(T)>(InValue), InParams, InKey, bAllowImplicitConversion));
+		return static_cast<T>(GetValue(InName, static_cast<__underlying_type(T)>(InValue), InParams, InKey));
 	}
 
 	template<typename T, typename TEnableIf<TIsEnumClass<T>::Value>::Type* = nullptr>
-	T GetValue(const FName& InName, const T& InValue, UPCGParamData* InParams, bool bAllowImplicitConversion = true)
+	T GetValue(const FName& InName, const T& InValue, const UPCGParamData* InParams)
 	{
-		return static_cast<T>(GetValue(InName, static_cast<__underlying_type(T)>(InValue), InParams, bAllowImplicitConversion));
+		return static_cast<T>(GetValue(InName, static_cast<__underlying_type(T)>(InValue), InParams));
 	}
 
 	template<typename T, typename TEnableIf<TIsEnumClass<T>::Value>::Type* = nullptr>
-	T GetValue(const FName& InName, const T& InValue, UPCGParamData* InParams, const FName& InParamName, bool bAllowImplicitConversion = true)
+	T GetValue(const FName& InName, const T& InValue, const UPCGParamData* InParams, const FName& InParamName)
 	{
-		return static_cast<T>(GetValue(InName, static_cast<__underlying_type(T)>(InValue), InParams, InParamName, bAllowImplicitConversion));
+		return static_cast<T>(GetValue(InName, static_cast<__underlying_type(T)>(InValue), InParams, InParamName));
 	}
 
 	/** Specialized version for names */
 	template<>
-	FORCEINLINE FName GetValue(const FName& InName, const FName& InValue, UPCGParamData* InParams, PCGMetadataEntryKey InKey, bool bAllowImplicitConversion)
+	FORCEINLINE FName GetValue(const FName& InName, const FName& InValue, const UPCGParamData* InParams, PCGMetadataEntryKey InKey)
 	{
-		return FName(GetValue(InName, InValue.ToString(), InParams, InKey, bAllowImplicitConversion));
+		return FName(GetValue(InName, InValue.ToString(), InParams, InKey));
 	}
 
 	template<>
-	FORCEINLINE FName GetValue(const FName& InName, const FName& InValue, UPCGParamData* InParams, bool bAllowImplicitConversion)
+	FORCEINLINE FName GetValue(const FName& InName, const FName& InValue, const UPCGParamData* InParams)
 	{
-		return FName(GetValue(InName, InValue.ToString(), InParams, bAllowImplicitConversion));
+		return FName(GetValue(InName, InValue.ToString(), InParams));
 	}
 
 	template<>
-	FORCEINLINE FName GetValue(const FName& InName, const FName& InValue, UPCGParamData* InParams, const FName& InParamName, bool bAllowImplicitConversion)
+	FORCEINLINE FName GetValue(const FName& InName, const FName& InValue, const UPCGParamData* InParams, const FName& InParamName)
 	{
-		return FName(GetValue(InName, InValue.ToString(), InParams, InParamName, bAllowImplicitConversion));
+		return FName(GetValue(InName, InValue.ToString(), InParams, InParamName));
 	}
 
 	/** Sets data from the params to a given property, matched on a name basis */
diff --git a/Engine/Plugins/Experimental/PCG/Source/PCG/Public/Metadata/PCGMetadata.h b/Engine/Plugins/Experimental/PCG/Source/PCG/Public/Metadata/PCGMetadata.h
--- a/Engine/Plugins/Experimental/PCG/Source/PCG/Public/Metadata/PCGMetadata.h
+++ b/Engine/Plugins/Experimental/PCG/Source/PCG/Public/Metadata/PCGMetadata.h
@@ -28,6 +28,10 @@ public:
 	UFUNCTION(BlueprintCallable, Category = "PCG|Metadata")
 	void Initialize(const UPCGMetadata* InParent);
 
+	/** Initializes the metadata from a parent metadata, if any (can be null) with the option to not add attributes from the parent. */
+	void Initialize(const UPCGMetadata* InParent, bool bAddAttributesFromParent);
+
+	/** Initializes the metadata from a parent metadata by copying all attributes to it. */
 	UFUNCTION(BlueprintCallable, Category = "PCG|Metadata")
 	void InitializeAsCopy(const UPCGMetadata* InMetadataToCopy);
 
