Fixed metadata warning when using a surface sampler on multiple landscapes

diff --git a/Engine/Plugins/Experimental/PCG/Source/PCG/Private/Data/PCGLandscapeData.cpp b/Engine/Plugins/Experimental/PCG/Source/PCG/Private/Data/PCGLandscapeData.cpp
--- a/Engine/Plugins/Experimental/PCG/Source/PCG/Private/Data/PCGLandscapeData.cpp
+++ b/Engine/Plugins/Experimental/PCG/Source/PCG/Private/Data/PCGLandscapeData.cpp
@@ -51,7 +51,10 @@ void UPCGLandscapeData::Initialize(const TArray<TWeakObjectPtr<ALandscapeProxy>>
 
 				for (const FName& Layer : Layers)
 				{
-					Metadata->CreateFloatAttribute(Layer, 0.0f, /*bAllowInterpolation=*/true);
+					if (!Metadata->HasAttribute(Layer))
+					{
+						Metadata->CreateFloatAttribute(Layer, 0.0f, /*bAllowInterpolation=*/true);
+					}
 				}
 			}
 		}
