[PCG] Follow up on Attribute management elements
* ReduceAttributes (AttributeToParamData)
* GetAttributeFromIndex
* SelectAttribute
* Remove support for FTransform AddSub in MetadataTraits.

diff --git a/Engine/Plugins/Experimental/PCG/Source/PCG/Private/Elements/PCGAttributeGetFromPointIndexElement.cpp b/Engine/Plugins/Experimental/PCG/Source/PCG/Private/Elements/PCGAttributeGetFromPointIndexElement.cpp
new file mode 100644
--- /dev/null
+++ b/Engine/Plugins/Experimental/PCG/Source/PCG/Private/Elements/PCGAttributeGetFromPointIndexElement.cpp
@@ -0,0 +1,121 @@
+// Copyright Epic Games, Inc. All Rights Reserved.
+
+#include "Elements/PCGAttributeGetFromPointIndexElement.h"
+
+#include "PCGData.h"
+#include "PCGHelpers.h"
+#include "PCGParamData.h"
+#include "Data/PCGPointData.h"
+#include "Metadata/PCGMetadata.h"
+#include "Metadata/PCGMetadataAttribute.h"
+#include "Helpers/PCGSettingsHelpers.h"
+
+#if WITH_EDITOR
+FName UPCGAttributeGetFromPointIndexSettings::GetDefaultNodeName() const
+{
+	return TEXT("GetAttributeFromPointIndex");
+}
+#endif
+
+TArray<FPCGPinProperties> UPCGAttributeGetFromPointIndexSettings::InputPinProperties() const
+{
+	TArray<FPCGPinProperties> PinProperties;
+	PinProperties.Emplace(PCGPinConstants::DefaultInputLabel, EPCGDataType::Point, /*bInAllowMultipleConnections=*/ false);
+	PinProperties.Emplace(PCGPinConstants::DefaultParamsLabel, EPCGDataType::Param, /*bInAllowMultipleConnections=*/ false);
+
+	return PinProperties;
+}
+
+TArray<FPCGPinProperties> UPCGAttributeGetFromPointIndexSettings::OutputPinProperties() const
+{
+	TArray<FPCGPinProperties> PinProperties;
+	PinProperties.Emplace(PCGAttributeGetFromPointIndexConstants::OutputAttributeLabel, EPCGDataType::Param);
+	PinProperties.Emplace(PCGAttributeGetFromPointIndexConstants::OutputPointLabel, EPCGDataType::Point);
+
+	return PinProperties;
+}
+
+FPCGElementPtr UPCGAttributeGetFromPointIndexSettings::CreateElement() const
+{
+	return MakeShared<FPCGAttributeGetFromPointIndexElement>();
+}
+
+bool FPCGAttributeGetFromPointIndexElement::ExecuteInternal(FPCGContext* Context) const
+{
+	TRACE_CPUPROFILER_EVENT_SCOPE(FPCGAttributeGetFromPointIndexElement::Execute);
+
+	check(Context);
+
+	const UPCGAttributeGetFromPointIndexSettings* Settings = Context->GetInputSettings<UPCGAttributeGetFromPointIndexSettings>();
+	check(Settings);
+
+	TArray<FPCGTaggedData> Inputs = Context->InputData.GetInputsByPin(PCGPinConstants::DefaultInputLabel);
+
+	if (Inputs.Num() != 1)
+	{
+		PCGE_LOG(Error, "Input pin doesn't have the right number of inputs.");
+		return true;
+	}
+
+	const UPCGPointData* PointData = Cast<UPCGPointData>(Inputs[0].Data);
+
+	if (!PointData)
+	{
+		PCGE_LOG(Error, "Input is not a point data.");
+		return true;
+	}
+
+	const UPCGParamData* ParamData = Context->InputData.GetParams();
+
+	const int32 Index = PCG_GET_OVERRIDEN_VALUE(Settings, Index, ParamData);
+
+	if (Index < 0 || Index >= PointData->GetPoints().Num())
+	{
+		PCGE_LOG(Error, "Index is out of bounds. Index: %d; Number of Points: %d", Index, PointData->GetPoints().Num());
+		return true;
+	}
+
+	TArray<FPCGTaggedData>& Outputs = Context->OutputData.TaggedData;
+	const FPCGPoint& Point = PointData->GetPoints()[Index];
+
+	// Only create a point if we are connected, or in editor with the component being inspected.
+	if (Context->IsOutputConnectedOrInspecting(PCGAttributeGetFromPointIndexConstants::OutputPointLabel))
+	{
+		UPCGPointData* OutputPointData = NewObject<UPCGPointData>();
+		OutputPointData->InitializeFromData(PointData);
+		OutputPointData->GetMutablePoints().Add(Point);
+
+		FPCGTaggedData& Output = Outputs.Emplace_GetRef();
+		Output.Data = OutputPointData;
+		Output.Pin = PCGAttributeGetFromPointIndexConstants::OutputPointLabel;
+	}
+
+	if (PointData->Metadata && PointData->Metadata->HasAttribute(Settings->InputAttributeName))
+	{
+		const FPCGMetadataAttributeBase* Attribute = PointData->Metadata->GetConstAttribute(Settings->InputAttributeName);
+		UPCGParamData* OutputParamData = NewObject<UPCGParamData>();
+
+		auto ExtractAttribute = [AttributeName = Settings->InputAttributeName, &Point, &OutputParamData, Attribute](auto DummyValue)
+		{
+			using AttributeType = decltype(DummyValue);
+
+			const FPCGMetadataAttribute<AttributeType>* TypedAttribute = static_cast<const FPCGMetadataAttribute<AttributeType>*>(Attribute);
+
+			FPCGMetadataAttributeBase* NewAttributeBase =
+				OutputParamData->Metadata->CreateAttribute<AttributeType>(AttributeName, TypedAttribute->GetValueFromItemKey(Point.MetadataEntry), TypedAttribute->AllowsInterpolation(), /*bOverrideParent=*/false);
+
+			check(NewAttributeBase);
+
+			PCGMetadataEntryKey EntryKey = OutputParamData->Metadata->AddEntry();
+			NewAttributeBase->SetValueFromValueKey(EntryKey, PCGDefaultValueKey);
+		};
+
+		PCGMetadataAttribute::CallbackWithRightType(Attribute->GetTypeId(), ExtractAttribute);
+
+		FPCGTaggedData& Output = Outputs.Emplace_GetRef();
+		Output.Data = OutputParamData;
+		Output.Pin = PCGAttributeGetFromPointIndexConstants::OutputAttributeLabel;
+	}
+
+	return true;
+}
\ No newline at end of file
diff --git a/Engine/Plugins/Experimental/PCG/Source/PCG/Private/Elements/PCGAttributeReduceElement.cpp b/Engine/Plugins/Experimental/PCG/Source/PCG/Private/Elements/PCGAttributeReduceElement.cpp
new file mode 100644
--- /dev/null
+++ b/Engine/Plugins/Experimental/PCG/Source/PCG/Private/Elements/PCGAttributeReduceElement.cpp
@@ -0,0 +1,235 @@
+// Copyright Epic Games, Inc. All Rights Reserved.
+
+#include "Elements/PCGAttributeReduceElement.h"
+
+#include "PCGData.h"
+#include "PCGParamData.h"
+#include "Data/PCGPointData.h"
+#include "Metadata/PCGMetadata.h"
+#include "Metadata/PCGMetadataAttribute.h"
+#include "Metadata/PCGMetadataEntryKeyIterator.h"
+#include "Helpers/PCGSettingsHelpers.h"
+
+namespace PCGAttributeReduceElement
+{
+	template <typename T>
+	bool Average(IPCGMetadataEntryIterator& Iterator, const FPCGMetadataAttribute<T>* InAttribute, T& OutValue)
+	{
+		if constexpr (!PCG::Private::MetadataTraits<T>::CanSubAdd || !PCG::Private::MetadataTraits<T>::CanInterpolate)
+		{
+			return false;
+		}
+		else
+		{
+			// Can't have a repeat iterator
+			check(!Iterator.IsRepeat());
+
+			int32 NumberOfEntries = 0;
+			OutValue = PCG::Private::MetadataTraits<T>::ZeroValue();
+			for (; !Iterator.IsEnd(); ++Iterator)
+			{
+				OutValue = PCG::Private::MetadataTraits<T>::Add(OutValue, InAttribute->GetValueFromItemKey(*Iterator));
+				++NumberOfEntries;
+			}
+
+			if (NumberOfEntries != 0)
+			{
+				OutValue = PCG::Private::MetadataTraits<T>::WeightedSum(PCG::Private::MetadataTraits<T>::ZeroValue(), OutValue, 1.0f / NumberOfEntries);
+			}
+
+			return true;
+		}
+	}
+
+	template <typename T, bool bIsMin>
+	bool MinMax(IPCGMetadataEntryIterator& Iterator, const FPCGMetadataAttribute<T>* InAttribute, T& OutValue)
+	{
+		if constexpr (!PCG::Private::MetadataTraits<T>::CanMinMax)
+		{
+			return false;
+		}
+		else
+		{
+			// Can't have a repeat iterator
+			check(!Iterator.IsRepeat());
+
+			int32 NumberOfEntries = 0;
+			for (; !Iterator.IsEnd(); ++Iterator)
+			{
+				T AttributeValue = InAttribute->GetValueFromItemKey(*Iterator);
+
+				if (NumberOfEntries == 0)
+				{
+					OutValue = AttributeValue;
+				}
+				else
+				{
+					if constexpr (bIsMin)
+					{
+						OutValue = PCG::Private::MetadataTraits<T>::Min(OutValue, AttributeValue);
+					}
+					else
+					{
+						OutValue = PCG::Private::MetadataTraits<T>::Max(OutValue, AttributeValue);
+					}
+				}
+				++NumberOfEntries;
+			}
+
+			return true;
+		}
+	}
+}
+
+#if WITH_EDITOR
+FName UPCGAttributeReduceSettings::GetDefaultNodeName() const
+{
+	return TEXT("AttributeReduce");
+}
+#endif
+
+FName UPCGAttributeReduceSettings::AdditionalTaskName() const
+{
+	if (const UEnum* EnumPtr = FindObject<UEnum>(nullptr, TEXT("/Script/PCG.EPCGAttributeReduceOperation"), true))
+	{
+		const FString OperationName = EnumPtr->GetNameStringByValue(static_cast<int>(Operation));
+
+		if (InputAttributeName != OutputAttributeName && OutputAttributeName != NAME_None)
+		{
+			return FName(FString::Printf(TEXT("Reduce %s to %s: %s"), *InputAttributeName.ToString(), *OutputAttributeName.ToString(), *OperationName));
+		}
+		else
+		{
+			return FName(FString::Printf(TEXT("Reduce %s: %s"), *InputAttributeName.ToString(), *OperationName));
+		}
+	}
+	else
+	{
+		return NAME_None;
+	}
+}
+
+TArray<FPCGPinProperties> UPCGAttributeReduceSettings::InputPinProperties() const
+{
+	TArray<FPCGPinProperties> PinProperties;
+	PinProperties.Emplace(PCGPinConstants::DefaultInputLabel, EPCGDataType::Spatial, /*bInAllowMultipleConnections=*/ false);
+
+	return PinProperties;
+}
+
+TArray<FPCGPinProperties> UPCGAttributeReduceSettings::OutputPinProperties() const
+{
+	TArray<FPCGPinProperties> PinProperties;
+	PinProperties.Emplace(PCGPinConstants::DefaultOutputLabel, EPCGDataType::Param);
+
+	return PinProperties;
+}
+
+FPCGElementPtr UPCGAttributeReduceSettings::CreateElement() const
+{
+	return MakeShared<FPCGAttributeReduceElement>();
+}
+
+bool FPCGAttributeReduceElement::ExecuteInternal(FPCGContext* Context) const
+{
+	TRACE_CPUPROFILER_EVENT_SCOPE(FPCGAttributeReduceElement::Execute);
+
+	check(Context);
+
+	const UPCGAttributeReduceSettings* Settings = Context->GetInputSettings<UPCGAttributeReduceSettings>();
+	check(Settings);
+
+	TArray<FPCGTaggedData> Inputs = Context->InputData.GetInputsByPin(PCGPinConstants::DefaultInputLabel);
+
+	if (Inputs.Num() != 1)
+	{
+		PCGE_LOG(Error, "Input pin doesn't have the right number of inputs.");
+		return true;
+	}
+
+	const UPCGSpatialData* SpatialData = Cast<UPCGSpatialData>(Inputs[0].Data);
+
+	if (!SpatialData)
+	{
+		PCGE_LOG(Error, "Input is not a spatial data.");
+		return true;
+	}
+
+	if (!SpatialData->Metadata)
+	{
+		PCGE_LOG(Error, "Input data doesn't have metadata");
+		return true;
+	}
+
+	if (!SpatialData->Metadata->HasAttribute(Settings->InputAttributeName))
+	{
+		PCGE_LOG(Error, "Input metadata doesn't have attribute \"%s\"", *Settings->InputAttributeName.ToString());
+		return true;
+	}
+
+	const UPCGPointData* PointData = Cast<UPCGPointData>(SpatialData);
+
+	const FPCGMetadataAttributeBase* InputAttribute = SpatialData->Metadata->GetConstAttribute(Settings->InputAttributeName);
+	const FName OutputAttributeName = Settings->OutputAttributeName == NAME_None ? Settings->InputAttributeName : Settings->OutputAttributeName;
+	UPCGParamData* OutputParamData = NewObject<UPCGParamData>();
+	
+	FPCGMetadataAttributeBase* OutputAttribute = OutputParamData->Metadata->CopyAttribute(InputAttribute, OutputAttributeName, /*bKeepParent=*/ false, /*bCopyEntries=*/ false, /*bCopyValues=*/ false);
+
+	auto DoOperation = [InputAttribute, OutputAttribute, Operation = Settings->Operation, SpatialData, PointData, OutputParamData](auto DummyValue) -> bool
+	{
+		using AttributeType = decltype(DummyValue);
+
+		bool bSuccess = false;
+		const bool bIsRepeat = false;
+
+		const FPCGMetadataAttribute<AttributeType>* TypedInputAttribute = static_cast<const FPCGMetadataAttribute<AttributeType>*>(InputAttribute);
+		FPCGMetadataAttribute<AttributeType>* TypedOutputAttribute = static_cast<FPCGMetadataAttribute<AttributeType>*>(OutputAttribute);
+
+		AttributeType OutputValue{};
+
+		TUniquePtr<IPCGMetadataEntryIterator> Iterator;
+		if (PointData)
+		{
+			Iterator = MakeUnique<FPCGMetadataEntryPointIterator>(PointData, bIsRepeat);
+		}
+		else
+		{
+			Iterator = MakeUnique<FPCGMetadataEntryAttributeIterator>(*InputAttribute, bIsRepeat);
+		}
+
+		switch (Operation)
+		{
+		case EPCGAttributeReduceOperation::Average:
+			bSuccess = PCGAttributeReduceElement::Average<AttributeType>(*Iterator, TypedInputAttribute, OutputValue);
+			break;
+		case EPCGAttributeReduceOperation::Max:
+			bSuccess = PCGAttributeReduceElement::MinMax<AttributeType, /*bIsMin*/false>(*Iterator, TypedInputAttribute, OutputValue);
+			break;
+		case EPCGAttributeReduceOperation::Min:
+			bSuccess = PCGAttributeReduceElement::MinMax<AttributeType, /*bIsMin*/true>(*Iterator, TypedInputAttribute, OutputValue);
+			break;
+		default:
+			break;
+		}
+
+		if (bSuccess)
+		{
+			TypedOutputAttribute->SetDefaultValue(OutputValue);
+			TypedOutputAttribute->SetValueFromValueKey(OutputParamData->Metadata->AddEntry(), PCGDefaultValueKey);
+		}
+
+		return bSuccess;
+	};
+
+	if (!PCGMetadataAttribute::CallbackWithRightType(InputAttribute->GetTypeId(), DoOperation))
+	{
+		PCGE_LOG(Error, "Operation was not compatible with the attribute type.");
+		return true;
+	}
+
+	TArray<FPCGTaggedData>& Outputs = Context->OutputData.TaggedData;
+	FPCGTaggedData& Output = Outputs.Emplace_GetRef();
+	Output.Data = OutputParamData;
+
+	return true;
+}
\ No newline at end of file
diff --git a/Engine/Plugins/Experimental/PCG/Source/PCG/Private/Elements/PCGAttributeSelectElement.cpp b/Engine/Plugins/Experimental/PCG/Source/PCG/Private/Elements/PCGAttributeSelectElement.cpp
new file mode 100644
--- /dev/null
+++ b/Engine/Plugins/Experimental/PCG/Source/PCG/Private/Elements/PCGAttributeSelectElement.cpp
@@ -0,0 +1,411 @@
+// Copyright Epic Games, Inc. All Rights Reserved.
+
+#include "Elements/PCGAttributeSelectElement.h"
+
+#include "PCGData.h"
+#include "PCGParamData.h"
+#include "Data/PCGPointData.h"
+#include "Metadata/PCGMetadata.h"
+#include "Metadata/PCGMetadataAttribute.h"
+#include "Metadata/PCGMetadataEntryKeyIterator.h"
+#include "Helpers/PCGSettingsHelpers.h"
+
+namespace PCGAttributeSelectElement
+{
+	// Of course, Dot doesn't exist for FVector4...
+	template <typename T>
+	inline double Dot(const T& A, const T& B)
+	{
+		if constexpr (std::is_same_v<FVector4, T>)
+		{
+			return Dot4(A, B);
+		}
+		else
+		{
+			return A.Dot(B);
+		}
+	}
+
+	// No need to do a dot product if the axis is either X, Y, Z or W.
+	// We still need to check the type for Z and W, since it needs to compile even if we
+	// will never call this function with the wrong type for a given axis.
+	// If T is a scalar, just return InValue. Use decltype(auto) to return an int if T is an int.
+	// It will allow to do comparison between int, instead of losing precision by converting it to double.
+	template <typename T, int Axis>
+	inline decltype(auto) Projection(const T& InValue, const T& InAxis)
+	{
+		if constexpr (PCG::Private::IsOfTypes<T, FVector2D, FVector, FVector4>())
+		{
+			if constexpr (Axis == 0)
+			{
+				return InValue.X;
+			}
+			else if constexpr (Axis == 1)
+			{
+				return InValue.Y;
+			}
+			else if constexpr (Axis == 2 && PCG::Private::IsOfTypes<T, FVector, FVector4>())
+			{
+				return InValue.Z;
+			}
+			else if constexpr (Axis == 3 && PCG::Private::IsOfTypes<T, FVector4>())
+			{
+				return InValue.W;
+			}
+			else
+			{
+				return Dot(InValue, InAxis);
+			}
+		}
+		else
+		{
+			// Use T constructor to force the compiler to match decltype(auto) to T.
+			// Otherwise, it would match to const T, which leads to issues.
+			return T(InValue);
+		}
+	}
+
+	template <typename T, bool bIsMin, int Axis>
+	bool MinMaxSelect(IPCGMetadataEntryIterator& Iterator, const FPCGMetadataAttribute<T>* InAttribute, const T& InAxis, T& OutValue, int32& OutIndex)
+	{
+		// Can't have a repeat iterator
+		check(!Iterator.IsRepeat());
+
+		// Will be int or double
+		using CompareType = decltype(Projection<T, Axis>(T{}, InAxis));
+
+		CompareType MinMaxValue{};
+
+		int32 NumberOfEntries = 0;
+		for (; !Iterator.IsEnd(); ++Iterator)
+		{
+			T AttributeValue = InAttribute->GetValueFromItemKey(*Iterator);
+			CompareType CurrentValue = Projection<T, Axis>(AttributeValue, InAxis);
+
+			if constexpr (bIsMin)
+			{
+				if (NumberOfEntries == 0 || CurrentValue < MinMaxValue)
+				{
+					MinMaxValue = CurrentValue;
+					OutValue = AttributeValue;
+					OutIndex = NumberOfEntries;
+				}
+			}
+			else
+			{
+				if (NumberOfEntries == 0 || CurrentValue > MinMaxValue)
+				{
+					MinMaxValue = CurrentValue;
+					OutValue = AttributeValue;
+					OutIndex = NumberOfEntries;
+				}
+			}
+
+			++NumberOfEntries;
+		}
+
+		return NumberOfEntries != 0;
+	}
+
+	template <typename T, int Axis>
+	bool MedianSelect(IPCGMetadataEntryIterator& Iterator, const FPCGMetadataAttribute<T>* InAttribute, const T& InAxis, T& OutValue, int32& OutIndex)
+	{
+		// Can't have a repeat iterator
+		check(!Iterator.IsRepeat());
+
+		// Will be int or double
+		using CompareType = decltype(Projection<T, Axis>(T{}, InAxis));
+
+		struct Item
+		{
+			Item(int32 InIndex, CompareType InCompareValue, T&& InAttributeValue)
+				: Index(InIndex)
+				, CompareValue(InCompareValue)
+				, AttributeValue(Forward<T>(InAttributeValue)) 
+			{}
+
+			int32 Index;
+			CompareType CompareValue;
+			T AttributeValue;
+		};
+
+		int32 NumberOfEntries = 0;
+		TArray<Item> Items;
+		for (; !Iterator.IsEnd(); ++Iterator)
+		{
+			T AttributeValue = InAttribute->GetValueFromItemKey(*Iterator);
+			CompareType CurrentValue = Projection<T, Axis>(AttributeValue, InAxis);
+
+			Items.Emplace(NumberOfEntries++, MoveTemp(CurrentValue), MoveTemp(AttributeValue));
+		}
+
+		if (NumberOfEntries == 0)
+		{
+			return false;
+		}
+
+		Algo::Sort(Items, [](const Item& A, const Item& B) -> bool { return A.CompareValue < B.CompareValue; });
+
+		// Since we need to return an index, we can't do the mean on 2 values if the number of entries is even, since it will yield a value that might not exist in the original dataset.
+		// In this case we arbitrarily chose one entry.
+		const int32 Index = NumberOfEntries / 2;
+		OutValue = Items[Index].AttributeValue;
+		OutIndex = Items[Index].Index;
+
+		return true;
+	}
+
+	template <typename T, int Axis>
+	inline bool DispatchOperation(IPCGMetadataEntryIterator& Iterator, const FPCGMetadataAttribute<T>* InAttribute, const T& InAxis, EPCGAttributeSelectOperation Operation, T& OutValue, int32& OutIndex)
+	{
+		switch (Operation)
+		{
+		case EPCGAttributeSelectOperation::Min:
+			return PCGAttributeSelectElement::MinMaxSelect<T, /*bIsMin=*/true, Axis>(Iterator, InAttribute, InAxis, OutValue, OutIndex);
+		case EPCGAttributeSelectOperation::Max:
+			return PCGAttributeSelectElement::MinMaxSelect<T, /*bIsMin=*/false, Axis>(Iterator, InAttribute, InAxis, OutValue, OutIndex);
+		case EPCGAttributeSelectOperation::Median:
+			return PCGAttributeSelectElement::MedianSelect<T, Axis>(Iterator, InAttribute, InAxis, OutValue, OutIndex);
+		default:
+			return false;
+		}
+	}
+}
+
+#if WITH_EDITOR
+FName UPCGAttributeSelectSettings::GetDefaultNodeName() const
+{
+	return TEXT("AttributeSelect");
+}
+#endif
+
+FName UPCGAttributeSelectSettings::AdditionalTaskName() const
+{
+	if (const UEnum* EnumOpPtr = FindObject<UEnum>(nullptr, TEXT("/Script/PCG.EPCGAttributeSelectOperation"), true))
+	{
+		if (const UEnum* EnumAxisPtr = FindObject<UEnum>(nullptr, TEXT("/Script/PCG.EPCGAttributeSelectAxis"), true))
+		{
+			const FString OperationName = EnumOpPtr->GetNameStringByValue(static_cast<int>(Operation));
+			FString AxisName;
+			if (Axis == EPCGAttributeSelectAxis::CustomAxis)
+			{
+				AxisName = FString::Printf(TEXT("(%.2f, %.2f, %.2f, %.2f)"), CustomAxis.X, CustomAxis.Y, CustomAxis.Z, CustomAxis.W);
+			}
+			else
+			{
+				AxisName = EnumAxisPtr->GetNameStringByValue(static_cast<int>(Axis));
+			}
+
+			if (InputAttributeName != OutputAttributeName && OutputAttributeName != NAME_None)
+			{
+				return FName(FString::Printf(TEXT("Select %s to %s: %s on %s"), *InputAttributeName.ToString(), *OutputAttributeName.ToString(), *OperationName, *AxisName));
+			}
+			else
+			{
+				return FName(FString::Printf(TEXT("Select %s: %s on %s"), *InputAttributeName.ToString(), *OperationName, *AxisName));
+			}
+		}
+	}
+
+	return NAME_None;
+}
+
+TArray<FPCGPinProperties> UPCGAttributeSelectSettings::InputPinProperties() const
+{
+	TArray<FPCGPinProperties> PinProperties;
+	PinProperties.Emplace(PCGPinConstants::DefaultInputLabel, EPCGDataType::Spatial, /*bInAllowMultipleConnections=*/ false);
+	PinProperties.Emplace(PCGPinConstants::DefaultParamsLabel, EPCGDataType::Param, /*bInAllowMultipleConnections=*/ false);
+
+	return PinProperties;
+}
+
+TArray<FPCGPinProperties> UPCGAttributeSelectSettings::OutputPinProperties() const
+{
+	TArray<FPCGPinProperties> PinProperties;
+	PinProperties.Emplace(PCGAttributeSelectConstants::OutputAttributeLabel, EPCGDataType::Param);
+	PinProperties.Emplace(PCGAttributeSelectConstants::OutputPointLabel, EPCGDataType::Point);
+
+	return PinProperties;
+}
+
+FPCGElementPtr UPCGAttributeSelectSettings::CreateElement() const
+{
+	return MakeShared<FPCGAttributeSelectElement>();
+}
+
+bool FPCGAttributeSelectElement::ExecuteInternal(FPCGContext* Context) const
+{
+	TRACE_CPUPROFILER_EVENT_SCOPE(FPCGAttributeSelectElement::Execute);
+
+	check(Context);
+
+	const UPCGAttributeSelectSettings* Settings = Context->GetInputSettings<UPCGAttributeSelectSettings>();
+	check(Settings);
+
+	TArray<FPCGTaggedData> Inputs = Context->InputData.GetInputsByPin(PCGPinConstants::DefaultInputLabel);
+
+	if (Inputs.Num() != 1)
+	{
+		PCGE_LOG(Error, "Input pin doesn't have the right number of inputs.");
+		return true;
+	}
+
+	const UPCGSpatialData* SpatialData = Cast<UPCGSpatialData>(Inputs[0].Data);
+
+	if (!SpatialData)
+	{
+		PCGE_LOG(Error, "Input is not a spatial data.");
+		return true;
+	}
+
+	const UPCGPointData* PointData = Cast<UPCGPointData>(SpatialData);
+	if (!PointData && Context->Node && Context->Node->IsOutputPinConnected(PCGAttributeSelectConstants::OutputPointLabel))
+	{
+		PCGE_LOG(Warning, "Not a point data as input, will output nothing in the %s output pin", *PCGAttributeSelectConstants::OutputPointLabel.ToString());
+	}
+
+	if (!SpatialData->Metadata)
+	{
+		PCGE_LOG(Error, "Input data doesn't have metadata");
+		return true;
+	}
+
+	if (!SpatialData->Metadata->HasAttribute(Settings->InputAttributeName))
+	{
+		PCGE_LOG(Error, "Input metadata doesn't have attribute \"%s\"", *Settings->InputAttributeName.ToString());
+		return true;
+	}
+
+	const FPCGMetadataAttributeBase* InputAttribute = SpatialData->Metadata->GetConstAttribute(Settings->InputAttributeName);
+	const FName OutputAttributeName = Settings->OutputAttributeName == NAME_None ? Settings->InputAttributeName : Settings->OutputAttributeName;
+	UPCGParamData* OutputParamData = NewObject<UPCGParamData>();
+
+	FPCGMetadataAttributeBase* OutputAttribute = OutputParamData->Metadata->CopyAttribute(InputAttribute, OutputAttributeName, /*bKeepParent=*/ false, /*bCopyEntries=*/ false, /*bCopyValues=*/ false);
+
+	auto DoOperation = [&](auto DummyValue) -> int32
+	{
+		using AttributeType = decltype(DummyValue);
+
+		if constexpr (!PCG::Private::IsOfTypes<AttributeType, int32, int64, float, double, FVector2D, FVector, FVector4>())
+		{
+			PCGE_LOG(Error, "Attribute type is not a Vector nor a scalar");
+			return -1;
+		}
+		else
+		{
+			const FPCGMetadataAttribute<AttributeType>* TypedInputAttribute = static_cast<const FPCGMetadataAttribute<AttributeType>*>(InputAttribute);
+			FPCGMetadataAttribute<AttributeType>* TypedOutputAttribute = static_cast<FPCGMetadataAttribute<AttributeType>*>(OutputAttribute);
+
+			bool bSuccess = false;
+			const bool bIsRepeat = false;
+			AttributeType OutputValue{};
+			AttributeType Axis{};
+			int32 OutputIndex = -1;
+
+			// First we need to verify if the axis we want to project on is valid for dimension of our vector type.
+			// If it is a scalar, we won't project anything.
+			if constexpr (PCG::Private::IsOfTypes<AttributeType, FVector2D, FVector, FVector4>())
+			{
+				bool bIsValid = false;
+
+				switch (Settings->Axis)
+				{
+				case EPCGAttributeSelectAxis::X:
+				case EPCGAttributeSelectAxis::Y:
+					bIsValid = true;
+					break;
+				case EPCGAttributeSelectAxis::Z:
+					bIsValid = PCG::Private::IsOfTypes<AttributeType, FVector, FVector4>();
+					break;
+				case EPCGAttributeSelectAxis::W:
+					bIsValid = PCG::Private::IsOfTypes<AttributeType, FVector4>();
+					break;
+				case EPCGAttributeSelectAxis::CustomAxis:
+					Axis = AttributeType(Settings->CustomAxis);
+					bIsValid = !Axis.Equals(AttributeType::Zero());
+					break;
+				default:
+					break;
+				}
+
+				if (!bIsValid)
+				{
+					PCGE_LOG(Error, "Invalid axis for attribute type.");
+					return -1;
+				}
+			}
+
+			// Then create our iterator depending if it is a point data or a spatial data
+			TUniquePtr<IPCGMetadataEntryIterator> Iterator;
+			if (PointData)
+			{
+				Iterator = MakeUnique<FPCGMetadataEntryPointIterator>(PointData, bIsRepeat);
+			}
+			else
+			{
+				Iterator = MakeUnique<FPCGMetadataEntryAttributeIterator>(*InputAttribute, bIsRepeat);
+			}
+
+			// Finally dispatch the operation depending on the axis.
+			// If the axis is X, Y, Z or W, the projection is overkill (Dot product), so use templates to 
+			// indicate which coordinate we should take. If the axis value is -1, it will do the projection with the custom axis, passed as parameter.
+			switch (Settings->Axis)
+			{
+			case EPCGAttributeSelectAxis::X:
+				bSuccess = PCGAttributeSelectElement::DispatchOperation<AttributeType, 0>(*Iterator, TypedInputAttribute, Axis, Settings->Operation, OutputValue, OutputIndex);
+				break;
+			case EPCGAttributeSelectAxis::Y:
+				bSuccess = PCGAttributeSelectElement::DispatchOperation<AttributeType, 1>(*Iterator, TypedInputAttribute, Axis, Settings->Operation, OutputValue, OutputIndex);
+				break;
+			case EPCGAttributeSelectAxis::Z:
+				bSuccess = PCGAttributeSelectElement::DispatchOperation<AttributeType, 2>(*Iterator, TypedInputAttribute, Axis, Settings->Operation, OutputValue, OutputIndex);
+				break;
+			case EPCGAttributeSelectAxis::W:
+				bSuccess = PCGAttributeSelectElement::DispatchOperation<AttributeType, 3>(*Iterator, TypedInputAttribute, Axis, Settings->Operation, OutputValue, OutputIndex);
+				break;
+			case EPCGAttributeSelectAxis::CustomAxis:
+				bSuccess = PCGAttributeSelectElement::DispatchOperation<AttributeType, -1>(*Iterator, TypedInputAttribute, Axis, Settings->Operation, OutputValue, OutputIndex);
+				break;
+			default:
+				break;
+			}
+
+			if (bSuccess)
+			{
+				TypedOutputAttribute->SetDefaultValue(OutputValue);
+				TypedOutputAttribute->SetValueFromValueKey(OutputParamData->Metadata->AddEntry(), PCGDefaultValueKey);
+			}
+			else
+			{
+				PCGE_LOG(Error, "Invalid axis for attribute type.");
+				OutputIndex = -1;
+			}
+
+			return OutputIndex;
+		}
+	};
+
+	int32 OutputIndex = PCGMetadataAttribute::CallbackWithRightType(InputAttribute->GetTypeId(), DoOperation);
+	if (OutputIndex < 0)
+	{
+		return true;
+	}
+
+	TArray<FPCGTaggedData>& Outputs = Context->OutputData.TaggedData;
+	FPCGTaggedData& Output = Outputs.Emplace_GetRef();
+	Output.Data = OutputParamData;
+	Output.Pin = PCGAttributeSelectConstants::OutputAttributeLabel;
+
+	if (PointData && Context->IsOutputConnectedOrInspecting(PCGAttributeSelectConstants::OutputPointLabel))
+	{
+		UPCGPointData* OutputPointData = NewObject<UPCGPointData>();
+		OutputPointData->InitializeFromData(PointData);
+		OutputPointData->GetMutablePoints().Add(PointData->GetPoint(OutputIndex));
+
+		FPCGTaggedData& PointOutput = Outputs.Emplace_GetRef();
+		PointOutput.Data = OutputPointData;
+		PointOutput.Pin = PCGAttributeSelectConstants::OutputPointLabel;
+	}
+
+	return true;
+}
\ No newline at end of file
diff --git a/Engine/Plugins/Experimental/PCG/Source/PCG/Private/PCGContext.cpp b/Engine/Plugins/Experimental/PCG/Source/PCG/Private/PCGContext.cpp
--- a/Engine/Plugins/Experimental/PCG/Source/PCG/Private/PCGContext.cpp
+++ b/Engine/Plugins/Experimental/PCG/Source/PCG/Private/PCGContext.cpp
@@ -2,6 +2,7 @@
 
 #include "PCGContext.h"
 #include "PCGComponent.h"
+#include "PCGHelpers.h"
 #include "GameFramework/Actor.h"
 
 FString FPCGContext::GetTaskName() const
@@ -36,4 +37,14 @@ FString FPCGContext::GetComponentName() const
 bool FPCGContext::ShouldStop() const
 {
 	return FPlatformTime::Seconds() > EndTime;
+}
+
+bool FPCGContext::IsOutputConnectedOrInspecting(FName PinLabel) const
+{
+	const bool bOutputConnected = Node && Node->IsOutputPinConnected(PinLabel);
+#if WITH_EDITOR
+	return bOutputConnected || (!PCGHelpers::IsRuntimeOrPIE() && SourceComponent.IsValid() && SourceComponent->IsInspecting());
+#else
+	return bOutputConnected;
+#endif // WITH_EDITOR
 }
\ No newline at end of file
diff --git a/Engine/Plugins/Experimental/PCG/Source/PCG/Public/Elements/PCGAttributeGetFromPointIndexElement.h b/Engine/Plugins/Experimental/PCG/Source/PCG/Public/Elements/PCGAttributeGetFromPointIndexElement.h
new file mode 100644
--- /dev/null
+++ b/Engine/Plugins/Experimental/PCG/Source/PCG/Public/Elements/PCGAttributeGetFromPointIndexElement.h
@@ -0,0 +1,57 @@
+// Copyright Epic Games, Inc. All Rights Reserved.
+
+#pragma once
+
+#include "PCGElement.h"
+#include "PCGSettings.h"
+
+#include "PCGAttributeGetFromPointIndexElement.generated.h"
+
+namespace PCGAttributeGetFromPointIndexConstants
+{
+	const FName OutputAttributeLabel = TEXT("Attribute");
+	const FName OutputPointLabel = TEXT("Point");
+}
+
+/**
+* Get the attribute of a point given its index. The result will be in a ParamData.
+* There is also a second output that will output the selected point. This point will be output
+* even if the attribute doesn't exist.
+* 
+* The Index can be overridden by a second Params input.
+*/
+UCLASS(BlueprintType, ClassGroup = (Procedural))
+class PCG_API UPCGAttributeGetFromPointIndexSettings : public UPCGSettings
+{
+	GENERATED_BODY()
+
+public:
+	//~Begin UPCGSettings interface
+#if WITH_EDITOR
+	virtual FName GetDefaultNodeName() const override;
+	virtual EPCGSettingsType GetType() const override { return EPCGSettingsType::Metadata; }
+#endif
+
+	virtual TArray<FPCGPinProperties> InputPinProperties() const override;
+	virtual TArray<FPCGPinProperties> OutputPinProperties() const override;
+	//~End UPCGSettings interface
+
+	UPROPERTY(BlueprintReadWrite, EditAnywhere, Category = Settings)
+	FName InputAttributeName = NAME_None;
+
+	UPROPERTY(BlueprintReadWrite, EditAnywhere, Category = Settings)
+	int32 Index = 0;
+
+protected:
+	virtual FPCGElementPtr CreateElement() const override;
+};
+
+
+class FPCGAttributeGetFromPointIndexElement : public FSimplePCGElement
+{
+public:
+	virtual bool IsCacheable(const UPCGSettings* InSettings) const override { return false; }
+
+protected:
+	virtual bool ExecuteInternal(FPCGContext* Context) const override;
+};
diff --git a/Engine/Plugins/Experimental/PCG/Source/PCG/Public/Elements/PCGAttributeReduceElement.h b/Engine/Plugins/Experimental/PCG/Source/PCG/Public/Elements/PCGAttributeReduceElement.h
new file mode 100644
--- /dev/null
+++ b/Engine/Plugins/Experimental/PCG/Source/PCG/Public/Elements/PCGAttributeReduceElement.h
@@ -0,0 +1,63 @@
+// Copyright Epic Games, Inc. All Rights Reserved.
+
+#pragma once
+
+#include "PCGElement.h"
+#include "PCGSettings.h"
+
+#include "PCGAttributeReduceElement.generated.h"
+
+UENUM()
+enum class EPCGAttributeReduceOperation
+{
+	Average,
+	Max,
+	Min
+};
+
+/**
+* Take all the entries/points from the input and perform a reduce operation on the given attribute
+* and output the result into a ParamData.
+* 
+* If the OutputAttributeName is None, we will use InputAttributeName
+*/
+UCLASS(BlueprintType, ClassGroup = (Procedural))
+class PCG_API UPCGAttributeReduceSettings : public UPCGSettings
+{
+	GENERATED_BODY()
+
+public:
+	//~Begin UPCGSettings interface
+#if WITH_EDITOR
+	virtual FName GetDefaultNodeName() const override;
+	virtual EPCGSettingsType GetType() const override { return EPCGSettingsType::Metadata; }
+#endif
+
+	virtual FName AdditionalTaskName() const override;
+
+	virtual TArray<FPCGPinProperties> InputPinProperties() const override;
+	virtual TArray<FPCGPinProperties> OutputPinProperties() const override;
+	//~End UPCGSettings interface
+
+	UPROPERTY(BlueprintReadWrite, EditAnywhere, Category = Settings)
+	FName InputAttributeName = NAME_None;
+
+	UPROPERTY(BlueprintReadWrite, EditAnywhere, Category = Settings)
+	FName OutputAttributeName = NAME_None;
+
+	UPROPERTY(BlueprintReadWrite, EditAnywhere, Category = Settings)
+	EPCGAttributeReduceOperation Operation = EPCGAttributeReduceOperation::Average;
+
+protected:
+	virtual FPCGElementPtr CreateElement() const override;
+};
+
+
+class FPCGAttributeReduceElement : public FSimplePCGElement
+{
+public:
+	virtual bool IsCacheable(const UPCGSettings* InSettings) const override { return false; }
+
+protected:
+	virtual bool ExecuteInternal(FPCGContext* Context) const override;
+};
\ No newline at end of file
diff --git a/Engine/Plugins/Experimental/PCG/Source/PCG/Public/Elements/PCGAttributeSelectElement.h b/Engine/Plugins/Experimental/PCG/Source/PCG/Public/Elements/PCGAttributeSelectElement.h
new file mode 100644
--- /dev/null
+++ b/Engine/Plugins/Experimental/PCG/Source/PCG/Public/Elements/PCGAttributeSelectElement.h
@@ -0,0 +1,92 @@
+// Copyright Epic Games, Inc. All Rights Reserved.
+
+#pragma once
+
+#include "PCGElement.h"
+#include "PCGSettings.h"
+
+#include "PCGAttributeSelectElement.generated.h"
+
+namespace PCGAttributeSelectConstants
+{
+	const FName OutputAttributeLabel = TEXT("Attribute");
+	const FName OutputPointLabel = TEXT("Point");
+}
+
+UENUM()
+enum class EPCGAttributeSelectOperation
+{
+	Min,
+	Max,
+	Median
+};
+
+UENUM()
+enum class EPCGAttributeSelectAxis
+{
+	X,
+	Y,
+	Z,
+	W,
+	CustomAxis
+};
+
+/**
+* Take all the entries/points from the input and perform a select operation on the given attribute on the given axis 
+* (if the attribute is a vector) and output the result into a ParamData.
+* It will also output the selected point if the input is a PointData.
+* 
+* Only support vector attributes and scalar attributes.
+* 
+* CustomAxis is overridable.
+* 
+* In case of the median operation, and the number of elements is even, we arbitrarily chose a point (Index = Num / 2)
+*
+* If the OutputAttributeName is None, we will use InputAttributeName.
+*/
+UCLASS(BlueprintType, ClassGroup = (Procedural))
+class PCG_API UPCGAttributeSelectSettings : public UPCGSettings
+{
+	GENERATED_BODY()
+
+public:
+	//~Begin UPCGSettings interface
+#if WITH_EDITOR
+	virtual FName GetDefaultNodeName() const override;
+	virtual EPCGSettingsType GetType() const override { return EPCGSettingsType::Metadata; }
+#endif
+
+	virtual FName AdditionalTaskName() const override;
+
+	virtual TArray<FPCGPinProperties> InputPinProperties() const override;
+	virtual TArray<FPCGPinProperties> OutputPinProperties() const override;
+	//~End UPCGSettings interface
+
+	UPROPERTY(BlueprintReadWrite, EditAnywhere, Category = Settings)
+	FName InputAttributeName = NAME_None;
+
+	UPROPERTY(BlueprintReadWrite, EditAnywhere, Category = Settings)
+	FName OutputAttributeName = NAME_None;
+
+	UPROPERTY(BlueprintReadWrite, EditAnywhere, Category = Settings)
+	EPCGAttributeSelectOperation Operation = EPCGAttributeSelectOperation::Min;
+
+	UPROPERTY(BlueprintReadWrite, EditAnywhere, Category = Settings)
+	EPCGAttributeSelectAxis Axis = EPCGAttributeSelectAxis::X;
+
+	UPROPERTY(BlueprintReadWrite, EditAnywhere, Category = Settings, meta = (EditCondition = "Axis == EPCGAttributeSelectAxis::CustomAxis", EditConditionHides))
+	FVector4 CustomAxis = FVector4::Zero();
+
+protected:
+	virtual FPCGElementPtr CreateElement() const override;
+};
+
+
+class FPCGAttributeSelectElement : public FSimplePCGElement
+{
+public:
+	virtual bool IsCacheable(const UPCGSettings* InSettings) const override { return false; }
+
+protected:
+	virtual bool ExecuteInternal(FPCGContext* Context) const override;
+};
\ No newline at end of file
diff --git a/Engine/Plugins/Experimental/PCG/Source/PCG/Public/Metadata/PCGMetadataAttributeTraits.h b/Engine/Plugins/Experimental/PCG/Source/PCG/Public/Metadata/PCGMetadataAttributeTraits.h
--- a/Engine/Plugins/Experimental/PCG/Source/PCG/Public/Metadata/PCGMetadataAttributeTraits.h
+++ b/Engine/Plugins/Experimental/PCG/Source/PCG/Public/Metadata/PCGMetadataAttributeTraits.h
@@ -397,7 +397,7 @@ namespace PCG
 		{
 			enum { CompressData = false };
 			enum { CanMinMax = false };
-			enum { CanSubAdd = true };
+			enum { CanSubAdd = false };
 			enum { CanMulDiv = true };
 			enum { CanInterpolate = true };
 
@@ -408,16 +408,6 @@ namespace PCG
 					A.GetScale3D() == B.GetScale3D();
 			}
 
-			static FTransform Add(const FTransform& A, const FTransform& B)
-			{
-				return A * B;
-			}
-
-			static FTransform Sub(const FTransform& A, const FTransform& B)
-			{
-				return A * B.Inverse();
-			}
-
 			static FTransform Mul(const FTransform& A, const FTransform& B)
 			{
 				return A * B;
diff --git a/Engine/Plugins/Experimental/PCG/Source/PCG/Public/PCGContext.h b/Engine/Plugins/Experimental/PCG/Source/PCG/Public/PCGContext.h
--- a/Engine/Plugins/Experimental/PCG/Source/PCG/Public/PCGContext.h
+++ b/Engine/Plugins/Experimental/PCG/Source/PCG/Public/PCGContext.h
@@ -75,4 +75,9 @@ struct PCG_API FPCGContext
 	FString GetTaskName() const;
 	FString GetComponentName() const;
 	bool ShouldStop() const;
+
+	/** Helper function to return if the output is connected, or if we need to force the connection because the component is being inspected. 
+	*   Useful to avoid creating output when it is not necessary.
+	*/
+	bool IsOutputConnectedOrInspecting(FName PinLabel) const;
 };
