Fixed issue with debug display following pivot changes in CL 22749377

diff --git a/Engine/Plugins/Experimental/PCG/Source/PCG/Private/Elements/PCGDebugElement.cpp b/Engine/Plugins/Experimental/PCG/Source/PCG/Private/Elements/PCGDebugElement.cpp
--- a/Engine/Plugins/Experimental/PCG/Source/PCG/Private/Elements/PCGDebugElement.cpp
+++ b/Engine/Plugins/Experimental/PCG/Source/PCG/Private/Elements/PCGDebugElement.cpp
@@ -154,7 +154,7 @@ namespace PCGDebugElement
 			ISMC->ComponentTags.AddUnique(PCGHelpers::DefaultPCGDebugTag);
 			ISMC->NumCustomDataFloats = NumCustomData;
 			const int32 PreExistingInstanceCount = ISMC->GetInstanceCount();
-			ISMC->AddInstances(Instances, false);
+			ISMC->AddInstances(Instances, /*bShouldReturnIndices=*/false, /*bWorldSpace=*/true);
 
 			// Then get & assign custom data
 			for (int32 PointIndex = 0; PointIndex < Points.Num(); ++PointIndex)
