Root temporary results when they are present in the data collection

diff --git a/Engine/Plugins/Experimental/PCG/Source/PCG/Private/Graph/PCGGraphExecutor.cpp b/Engine/Plugins/Experimental/PCG/Source/PCG/Private/Graph/PCGGraphExecutor.cpp
--- a/Engine/Plugins/Experimental/PCG/Source/PCG/Private/Graph/PCGGraphExecutor.cpp
+++ b/Engine/Plugins/Experimental/PCG/Source/PCG/Private/Graph/PCGGraphExecutor.cpp
@@ -443,6 +443,8 @@ void FPCGGraphExecutor::Execute()
 					check(!ActiveTask.Element->CanExecuteOnlyOnMainThread(ActiveTask.Context.Get()));
 					ActiveTask.Context->EndTime = EndTime;
 					ActiveTask.Context->bIsRunningOnMainThread = false;
+					// Remove the precreated data so we properly count in the root set
+					ActiveTask.Context->OutputData.RemoveFromRootSet(DataRootSet);
 
 					Futures.Emplace(ExecutionIndex, Async(EAsyncExecution::ThreadPool, [&ActiveTask]()
 					{
@@ -505,6 +507,8 @@ void FPCGGraphExecutor::Execute()
 				check(!MainThreadTask.Context->bIsPaused);
 				MainThreadTask.Context->EndTime = EndTime;
 				MainThreadTask.Context->bIsRunningOnMainThread = true;
+				// Remove the precreated data so we properly count in the root set
+				MainThreadTask.Context->OutputData.RemoveFromRootSet(DataRootSet);
 
 #if WITH_EDITOR
 				if(MainThreadTask.bIsBypassed || MainThreadTask.Element->Execute(MainThreadTask.Context.Get()))
@@ -534,6 +538,11 @@ void FPCGGraphExecutor::Execute()
 					{
 						PostTaskExecute(ExecutionIndex);
 					}
+					else
+					{
+						// Task isn't done, but we need to make sure the data isn't garbage collected
+						ActiveTasks[ExecutionIndex].Context->OutputData.AddToRootSet(DataRootSet);
+					}
 				}
 			}
 
@@ -541,6 +550,11 @@ void FPCGGraphExecutor::Execute()
 			{
 				PostTaskExecute(0);
 			}
+			else if(!ActiveTasks.IsEmpty())
+			{
+				// Task isn't done, make sure we don't garbage collected the precreated data
+				ActiveTasks[0].Context->OutputData.AddToRootSet(DataRootSet);
+			}
 		}
 
 		// Any paused tasks at that point should relinquish their resources
