Added PCG generation cancellation feature. Further usage of this will be done in other CLs.

diff --git a/Engine/Plugins/Experimental/PCG/Source/PCG/Private/Graph/PCGGraphExecutor.cpp b/Engine/Plugins/Experimental/PCG/Source/PCG/Private/Graph/PCGGraphExecutor.cpp
--- a/Engine/Plugins/Experimental/PCG/Source/PCG/Private/Graph/PCGGraphExecutor.cpp
+++ b/Engine/Plugins/Experimental/PCG/Source/PCG/Private/Graph/PCGGraphExecutor.cpp
@@ -53,7 +53,9 @@ FPCGGraphExecutor::~FPCGGraphExecutor()
 	DataRootSet.Clear();
 
 #if WITH_EDITOR
-	EndGenerationNotification();
+	// Cleanup + clear notification
+	ClearAllTasks();
+	UpdateGenerationNotification();
 #endif
 }
 
@@ -95,6 +97,7 @@ FPCGTaskId FPCGGraphExecutor::Schedule(UPCGGraph* Graph, UPCGComponent* SourceCo
 
 		FPCGGraphScheduleTask& ScheduledTask = ScheduledTasks.Emplace_GetRef();
 		ScheduledTask.Tasks = MoveTemp(CompiledTasks);
+		ScheduledTask.SourceComponent = SourceComponent;
 
 		// Offset task node ids
 		FPCGGraphCompiler::OffsetNodeIds(ScheduledTask.Tasks, NextTaskId);
@@ -116,6 +119,116 @@ FPCGTaskId FPCGGraphExecutor::Schedule(UPCGGraph* Graph, UPCGComponent* SourceCo
 	return ScheduledId;
 }
 
+void FPCGGraphExecutor::Cancel(UPCGComponent* InComponent)
+{
+	auto CancelIfSameComponent = [InComponent](TWeakObjectPtr<UPCGComponent> Component)
+	{
+		return Component.IsValid() && InComponent == Component;
+	};
+
+	return Cancel(CancelIfSameComponent);
+}
+
+TArray<UPCGComponent*> FPCGGraphExecutor::Cancel(UPCGGraph* InGraph)
+{
+	TSet<UPCGComponent*> CancelledComponents;
+
+	auto CancelIfComponentUsesGraph = [InGraph, &CancelledComponents](TWeakObjectPtr<UPCGComponent> Component)
+	{
+		if (Component.IsValid() && Component->GetGraph() == InGraph)
+		{
+			CancelledComponents.Add(Component.Get());
+			return true;
+		}
+		else
+		{
+			return false;
+		}
+	};
+
+	Cancel(CancelIfComponentUsesGraph);
+	return CancelledComponents.Array();
+}
+
+TArray<UPCGComponent*> FPCGGraphExecutor::CancelAll()
+{
+	TSet<UPCGComponent*> CancelledComponents;
+
+	auto CancelAllGeneration = [&CancelledComponents](TWeakObjectPtr<UPCGComponent> Component)
+	{
+		if (Component.IsValid())
+		{
+			CancelledComponents.Add(Component.Get());
+			return true;
+		}
+		else
+		{
+			return false;
+		}
+	};
+
+	Cancel(CancelAllGeneration);
+	return CancelledComponents.Array();
+}
+
+void FPCGGraphExecutor::Cancel(TFunctionRef<bool(TWeakObjectPtr<UPCGComponent>)> CancelFilter)
+{
+	// Remove from scheduled tasks
+	ScheduleLock.Lock();
+	for (int32 ScheduledTaskIndex = ScheduledTasks.Num() - 1; ScheduledTaskIndex >= 0; --ScheduledTaskIndex)
+	{
+		if(CancelFilter(ScheduledTasks[ScheduledTaskIndex].SourceComponent))
+		{
+			ScheduledTasks.RemoveAtSwap(ScheduledTaskIndex);
+		}
+	}
+	ScheduleLock.Unlock();
+
+	// Remove from ready tasks
+	for (int32 ReadyTaskIndex = ReadyTasks.Num() - 1; ReadyTaskIndex >= 0; --ReadyTaskIndex)
+	{
+		FPCGGraphTask& Task = ReadyTasks[ReadyTaskIndex];
+
+		if(CancelFilter(Task.SourceComponent))
+		{
+			FPCGTaskId CancelledTaskId = Task.NodeId;
+			delete Task.Context;
+			ReadyTasks.RemoveAtSwap(ReadyTaskIndex);
+			CancelNextTasks(CancelledTaskId);
+		}
+	}
+
+	// Remove from active tasks
+	for (int32 ActiveTaskIndex = ActiveTasks.Num() - 1; ActiveTaskIndex >= 0; --ActiveTaskIndex)
+	{
+		FPCGGraphActiveTask& Task = ActiveTasks[ActiveTaskIndex];
+		if(Task.Context && CancelFilter(Task.Context->SourceComponent))
+		{
+			CurrentlyUsedThreads -= Task.Context->NumAvailableTasks;
+			FPCGTaskId CancelledTaskId = Task.NodeId;
+			ActiveTasks.RemoveAtSwap(ActiveTaskIndex);
+			CancelNextTasks(CancelledTaskId);
+		}
+	}
+
+	// Remove from sleeping tasks
+	for (int32 SleepingTaskIndex = SleepingTasks.Num() - 1; SleepingTaskIndex >= 0; --SleepingTaskIndex)
+	{
+		FPCGGraphActiveTask& Task = SleepingTasks[SleepingTaskIndex];
+		if(Task.Context && CancelFilter(Task.Context->SourceComponent))
+		{
+			FPCGTaskId CancelledTaskId = Task.NodeId;
+			SleepingTasks.RemoveAtSwap(CancelledTaskId);
+			CancelNextTasks(CancelledTaskId);
+		}
+	}
+
+	// Finally, update the notification so it shows the new information
+#if WITH_EDITOR
+	UpdateGenerationNotification();
+#endif
+}
+
 FPCGTaskId FPCGGraphExecutor::ScheduleGeneric(TFunction<bool()> InOperation, UPCGComponent* InSourceComponent, const TArray<FPCGTaskId>& TaskDependencies)
 {
 	// Since we have no context, the generic task will consume no input
@@ -145,6 +258,7 @@ FPCGTaskId FPCGGraphExecutor::ScheduleGenericWithContext(TFunction<bool(FPCGCont
 
 	FPCGGraphScheduleTask& ScheduledTask = ScheduledTasks.Emplace_GetRef();
 	ScheduledTask.Tasks.Add(Task);
+	ScheduledTask.SourceComponent = InSourceComponent;
 
 	ScheduleLock.Unlock();
 
@@ -214,7 +328,7 @@ void FPCGGraphExecutor::Execute()
 	if (ReadyTasks.Num() == 0 && ActiveTasks.Num() == 0 && SleepingTasks.Num() == 0 && Tasks.Num() > 0)
 	{
 		UE_LOG(LogPCG, Error, TEXT("PCG Graph executor error: tasks are in a deadlocked state. Will drop all tasks."));
-		Tasks.Reset();
+		ClearAllTasks();
 	}
 
 	// TODO: change this if we support tasks that are not framebound
@@ -226,7 +340,7 @@ void FPCGGraphExecutor::Execute()
 	const bool bAllowMultiDispatch = CVarGraphMultithreading.GetValueOnAnyThread();
 
 #if WITH_EDITOR
-	UpdateGenerationNotification(Tasks.Num() + ReadyTasks.Num() + ActiveTasks.Num());
+	UpdateGenerationNotification();
 #endif
 
 	while(ReadyTasks.Num() > 0 || ActiveTasks.Num() > 0 || SleepingTasks.Num() > 0)
@@ -600,7 +714,8 @@ void FPCGGraphExecutor::Execute()
 			ClearResults();
 
 #if WITH_EDITOR
-			EndGenerationNotification();
+			// Call the notification update here to prevent it from sticking around - needed because we early out before this
+			UpdateGenerationNotification();
 #endif
 		}
 
@@ -612,6 +727,21 @@ void FPCGGraphExecutor::Execute()
 	}
 }
 
+void FPCGGraphExecutor::ClearAllTasks()
+{
+	Tasks.Reset();
+
+	// Make sure we don't leak preallocated contexts
+	for (FPCGGraphTask& ReadyTask : ReadyTasks)
+	{
+		delete ReadyTask.Context;
+	}
+
+	ReadyTasks.Reset();
+	ActiveTasks.Reset();
+	SleepingTasks.Reset();
+}
+
 void FPCGGraphExecutor::QueueNextTasks(FPCGTaskId FinishedTask)
 {
 	TRACE_CPUPROFILER_EVENT_SCOPE(FPCGGraphExecutor::QueueNextTasks);
@@ -639,6 +769,20 @@ void FPCGGraphExecutor::QueueNextTasks(FPCGTaskId FinishedTask)
 	}
 }
 
+void FPCGGraphExecutor::CancelNextTasks(FPCGTaskId CancelledTask)
+{
+	if (TSet<FPCGTaskId>* Successors = TaskSuccessors.Find(CancelledTask))
+	{
+		for (FPCGTaskId Successor : *Successors)
+		{
+			Tasks.Remove(Successor);
+			CancelNextTasks(Successor);
+		}
+
+		TaskSuccessors.Remove(CancelledTask);
+	}
+}
+
 void FPCGGraphExecutor::BuildTaskInput(const FPCGGraphTask& Task, FPCGDataCollection& TaskInput)
 {
 	TRACE_CPUPROFILER_EVENT_SCOPE(FPCGGraphExecutor::BuildTaskInput);
@@ -820,19 +964,20 @@ void FPCGGraphExecutor::NotifyGraphChanged(UPCGGraph* InGraph)
 	}
 }
 
-void FPCGGraphExecutor::UpdateGenerationNotification(int32 RemainingTaskNum)
+void FPCGGraphExecutor::UpdateGenerationNotification()
 {
-	if(RemainingTaskNum > 0)
+	const int32 RemainingTaskNum = Tasks.Num() + ReadyTasks.Num() + ActiveTasks.Num() + SleepingTasks.Num();
+
+	if (RemainingTaskNum > 0)
 	{
 		GenerationProgressNotification.Update(RemainingTaskNum);
 	}
-}
-
-void FPCGGraphExecutor::EndGenerationNotification()
-{
-	GenerationProgressNotification.Update(0);
-	// To reset the UI notification
-	GenerationProgressNotification = FAsyncCompilationNotification(GetNotificationTextFormat());
+	else
+	{
+		GenerationProgressNotification.Update(0);
+		// To reset the UI notification
+		GenerationProgressNotification = FAsyncCompilationNotification(GetNotificationTextFormat());
+	}
 }
 
 FTextFormat FPCGGraphExecutor::GetNotificationTextFormat()
diff --git a/Engine/Plugins/Experimental/PCG/Source/PCG/Private/Graph/PCGGraphExecutor.h b/Engine/Plugins/Experimental/PCG/Source/PCG/Private/Graph/PCGGraphExecutor.h
--- a/Engine/Plugins/Experimental/PCG/Source/PCG/Private/Graph/PCGGraphExecutor.h
+++ b/Engine/Plugins/Experimental/PCG/Source/PCG/Private/Graph/PCGGraphExecutor.h
@@ -50,6 +50,7 @@ struct FPCGGraphTask
 struct FPCGGraphScheduleTask
 {
 	TArray<FPCGGraphTask> Tasks;
+	TWeakObjectPtr<UPCGComponent> SourceComponent = nullptr;
 };
 
 struct FPCGGraphActiveTask
@@ -75,6 +76,15 @@ public:
 	FPCGTaskId Schedule(UPCGComponent* InComponent, const TArray<FPCGTaskId>& TaskDependency = TArray<FPCGTaskId>());
 	FPCGTaskId Schedule(UPCGGraph* Graph, UPCGComponent* InSourceComponent, FPCGElementPtr InputElement, const TArray<FPCGTaskId>& TaskDependency);
 
+	/** Cancels all tasks originating from the given component */
+	void Cancel(UPCGComponent* InComponent);
+	
+	/** Cancels all tasks running a given graph */
+	TArray<UPCGComponent*> Cancel(UPCGGraph* InGraph);
+
+	/** Cancels all tasks */
+	TArray<UPCGComponent*> CancelAll();
+
 	// Back compatibility function. Use ScheduleGenericWithContext
 	FPCGTaskId ScheduleGeneric(TFunction<bool()> InOperation, UPCGComponent* InSourceComponent, const TArray<FPCGTaskId>& TaskDependencies);
 
@@ -106,7 +116,10 @@ public:
 	FPCGGraphCache& GetCache() { return GraphCache; }
 
 private:
+	void Cancel(TFunctionRef<bool(TWeakObjectPtr<UPCGComponent>)> CancelFilter);
+	void ClearAllTasks();
 	void QueueNextTasks(FPCGTaskId FinishedTask);
+	void CancelNextTasks(FPCGTaskId CancelledTask);
 	void BuildTaskInput(const FPCGGraphTask& Task, FPCGDataCollection& TaskInput);
 	void StoreResults(FPCGTaskId InTaskId, const FPCGDataCollection& InTaskOutput);
 	void ClearResults();
@@ -117,8 +130,7 @@ private:
 	void SaveDirtyActors();
 	void ReleaseUnusedActors();
 
-	void UpdateGenerationNotification(int32 RemainingTaskNum);
-	void EndGenerationNotification();
+	void UpdateGenerationNotification();
 	static FTextFormat GetNotificationTextFormat();
 #endif
 
diff --git a/Engine/Plugins/Experimental/PCG/Source/PCG/Private/PCGComponent.cpp b/Engine/Plugins/Experimental/PCG/Source/PCG/Private/PCGComponent.cpp
--- a/Engine/Plugins/Experimental/PCG/Source/PCG/Private/PCGComponent.cpp
+++ b/Engine/Plugins/Experimental/PCG/Source/PCG/Private/PCGComponent.cpp
@@ -427,13 +427,17 @@ void UPCGComponent::PostCleanupGraph()
 #endif
 }
 
-void UPCGComponent::OnProcessGraphAborted()
+void UPCGComponent::OnProcessGraphAborted(bool bQuiet)
 {
-	UE_LOG(LogPCG, Warning, TEXT("Process Graph was called but aborted, check for errors in log if you expected a result."));
+	if (!bQuiet)
+	{
+		UE_LOG(LogPCG, Warning, TEXT("Process Graph was called but aborted, check for errors in log if you expected a result."));
+	}
 
 	CleanupUnusedManagedResources();
 
 	CurrentGenerationTask = InvalidPCGTaskId;
+	CurrentCleanupTask = InvalidPCGTaskId; // this is needed to support cancellation
 
 #if WITH_EDITOR
 	bDirtyGenerated = false;
@@ -442,7 +446,7 @@ void UPCGComponent::OnProcessGraphAborted()
 
 void UPCGComponent::Cleanup()
 {
-	if ((!bGenerated && !IsGenerating()) || !GetSubsystem() || IsCleaningUp())
+	if (!GetSubsystem() || IsCleaningUp())
 	{
 		return;
 	}
@@ -466,7 +470,7 @@ void UPCGComponent::CleanupLocal(bool bRemoveComponents, bool bSave)
 
 FPCGTaskId UPCGComponent::CleanupInternal(bool bRemoveComponents, bool bSave, const TArray<FPCGTaskId>& Dependencies)
 {
-	if ((!bGenerated && !IsGenerating()) || !GetSubsystem() || IsCleaningUp())
+	if (!GetSubsystem() || IsCleaningUp())
 	{
 		return InvalidPCGTaskId;
 	}
@@ -477,6 +481,14 @@ FPCGTaskId UPCGComponent::CleanupInternal(bool bRemoveComponents, bool bSave, co
 	return CurrentCleanupTask;
 }
 
+void UPCGComponent::CancelGeneration()
+{
+	if (CurrentGenerationTask != InvalidPCGTaskId)
+	{
+		GetSubsystem()->CancelGeneration(this);
+	}
+}
+
 AActor* UPCGComponent::ClearPCGLink(UClass* TemplateActor)
 {
 	if (!bGenerated || !GetOwner() || !GetWorld())
diff --git a/Engine/Plugins/Experimental/PCG/Source/PCG/Private/PCGSubsystem.cpp b/Engine/Plugins/Experimental/PCG/Source/PCG/Private/PCGSubsystem.cpp
--- a/Engine/Plugins/Experimental/PCG/Source/PCG/Private/PCGSubsystem.cpp
+++ b/Engine/Plugins/Experimental/PCG/Source/PCG/Private/PCGSubsystem.cpp
@@ -396,6 +396,61 @@ FPCGTaskId UPCGSubsystem::ScheduleGenericWithContext(TFunction<bool(FPCGContext*
 	return GraphExecutor->ScheduleGenericWithContext(InOperation, Component, TaskDependencies, bConsumeInputData);
 }
 
+void UPCGSubsystem::CancelGeneration(UPCGComponent* Component)
+{
+	check(GraphExecutor);
+
+	if (!Component)
+	{
+		return;
+	}
+
+	if (Component->IsPartitioned())
+	{
+		auto LocalCancel = [this](UPCGComponent* LocalComponent)
+		{
+			if (LocalComponent->IsGenerating())
+			{
+				CancelGeneration(LocalComponent);
+			}
+
+			return InvalidPCGTaskId;
+		};
+
+		DispatchToRegisteredLocalComponents(Component, LocalCancel);
+	}
+
+	GraphExecutor->Cancel(Component);
+	Component->OnProcessGraphAborted(/*bQuiet=*/true);
+}
+
+void UPCGSubsystem::CancelGeneration(UPCGGraph* Graph)
+{
+	check(GraphExecutor);
+
+	if (!Graph)
+	{
+		return;
+	}
+
+	TArray<UPCGComponent*> CancelledComponents = GraphExecutor->Cancel(Graph);
+	for (UPCGComponent* Component : CancelledComponents)
+	{
+		Component->OnProcessGraphAborted(/*bQuiet=*/true);
+	}
+}
+
+void UPCGSubsystem::CancelAllGeneration()
+{
+	check(GraphExecutor);
+
+	TArray<UPCGComponent*> CancelledComponents = GraphExecutor->CancelAll();
+	for (UPCGComponent* Component : CancelledComponents)
+	{
+		Component->OnProcessGraphAborted(/*bQuiet=*/true);
+	}
+}
+
 TArray<FPCGTaskId> UPCGSubsystem::DispatchToRegisteredLocalComponents(UPCGComponent* OriginalComponent, const TFunction<FPCGTaskId(UPCGComponent*)>& InFunc) const
 {
 	TRACE_CPUPROFILER_EVENT_SCOPE(UPCGSubsystem::DispatchToRegisteredLocalComponents);
diff --git a/Engine/Plugins/Experimental/PCG/Source/PCG/Public/PCGComponent.h b/Engine/Plugins/Experimental/PCG/Source/PCG/Public/PCGComponent.h
--- a/Engine/Plugins/Experimental/PCG/Source/PCG/Public/PCGComponent.h
+++ b/Engine/Plugins/Experimental/PCG/Source/PCG/Public/PCGComponent.h
@@ -130,6 +130,9 @@ public:
 	UFUNCTION(BlueprintCallable, NetMulticast, Reliable, Category = PCG)
 	void Cleanup(bool bRemoveComponents, bool bSave = false);
 
+	/** Cancels in-progress generation */
+	void CancelGeneration();
+
 	/** Move all generated resources under a new actor, following a template (AActor if not provided), clearing all link to this PCG component. Returns the new actor.*/
 	UFUNCTION(BlueprintCallable, Category = PCG)
 	AActor* ClearPCGLink(UClass* TemplateActor = nullptr);
@@ -250,7 +253,7 @@ private:
 	void PostProcessGraph(const FBox& InNewBounds, bool bInGenerated, FPCGContext* Context);
 	void CallPostGenerateFunctions(FPCGContext* Context) const;
 	void PostCleanupGraph();
-	void OnProcessGraphAborted();
+	void OnProcessGraphAborted(bool bQuiet = false);
 	void CleanupUnusedManagedResources();
 	bool MoveResourcesToNewActor(AActor* InNewActor, bool bCreateChild);
 
diff --git a/Engine/Plugins/Experimental/PCG/Source/PCG/Public/PCGSubsystem.h b/Engine/Plugins/Experimental/PCG/Source/PCG/Public/PCGSubsystem.h
--- a/Engine/Plugins/Experimental/PCG/Source/PCG/Public/PCGSubsystem.h
+++ b/Engine/Plugins/Experimental/PCG/Source/PCG/Public/PCGSubsystem.h
@@ -85,6 +85,15 @@ public:
 	*/
 	FPCGTaskId ScheduleGenericWithContext(TFunction<bool(FPCGContext*)> InOperation, UPCGComponent* SourceComponent, const TArray<FPCGTaskId>& TaskDependencies, bool bConsumeInputData = true);
 
+	/** Cancels currently running generation */
+	void CancelGeneration(UPCGComponent* Component);
+
+	/** Cancels currently running generation on given graph */
+	void CancelGeneration(UPCGGraph* Graph);
+
+	/** Cancels everything running */
+	void CancelAllGeneration();
+
 	/** Gets the output data for a given task */
 	bool GetOutputData(FPCGTaskId InTaskId, FPCGDataCollection& OutData);
 
diff --git a/Engine/Plugins/Experimental/PCG/Source/PCGEditor/Private/PCGComponentDetails.cpp b/Engine/Plugins/Experimental/PCG/Source/PCGEditor/Private/PCGComponentDetails.cpp
--- a/Engine/Plugins/Experimental/PCG/Source/PCGEditor/Private/PCGComponentDetails.cpp
+++ b/Engine/Plugins/Experimental/PCG/Source/PCGEditor/Private/PCGComponentDetails.cpp
@@ -68,6 +68,7 @@ void FPCGComponentDetails::CustomizeDetails(IDetailLayoutBuilder& DetailBuilder)
 			[
 				SNew(SButton)
 				.OnClicked(this, &FPCGComponentDetails::OnGenerateClicked)
+				.Visibility(this, &FPCGComponentDetails::GenerateButtonVisible)
 				[
 					SNew(STextBlock)
 					.Font(IDetailLayoutBuilder::GetDetailFont())
@@ -79,10 +80,23 @@ void FPCGComponentDetails::CustomizeDetails(IDetailLayoutBuilder& DetailBuilder)
 			.Padding(4.0f)
 			[
 				SNew(SButton)
-				.OnClicked(this, &FPCGComponentDetails::OnCleanupClicked)
+				.OnClicked(this, &FPCGComponentDetails::OnCancelClicked)
+				.Visibility(this, &FPCGComponentDetails::CancelButtonVisible)
 				[
 					SNew(STextBlock)
 					.Font(IDetailLayoutBuilder::GetDetailFont())
+					.Text(LOCTEXT("CancelButton", "Cancel"))
+				]
+			]
+			+ SHorizontalBox::Slot()
+			.AutoWidth()
+			.Padding(4.0f)
+			[
+				SNew(SButton)
+				.OnClicked(this, &FPCGComponentDetails::OnCleanupClicked)				
+				[
+					SNew(STextBlock)					
+					.Font(IDetailLayoutBuilder::GetDetailFont())
 					.Text(LOCTEXT("CleanupButton", "Cleanup"))
 				]
 			]
@@ -124,6 +138,32 @@ void FPCGComponentDetails::PendingDelete()
 	}
 }
 
+EVisibility FPCGComponentDetails::GenerateButtonVisible() const
+{
+	for (const TWeakObjectPtr<UPCGComponent>& Component : SelectedComponents)
+	{
+		if (Component.IsValid() && !Component->IsGenerating())
+		{
+			return EVisibility::Visible;
+		}
+	}
+
+	return EVisibility::Collapsed;
+}
+
+EVisibility FPCGComponentDetails::CancelButtonVisible() const
+{
+	for (const TWeakObjectPtr<UPCGComponent>& Component : SelectedComponents)
+	{
+		if (Component.IsValid() && Component->IsGenerating())
+		{
+			return EVisibility::Visible;
+		}
+	}
+
+	return EVisibility::Collapsed;
+}
+
 FReply FPCGComponentDetails::OnGenerateClicked()
 {
 	for (TWeakObjectPtr<UPCGComponent>& Component : SelectedComponents)
@@ -137,6 +177,19 @@ FReply FPCGComponentDetails::OnGenerateClicked()
 	return FReply::Handled();
 }
 
+FReply FPCGComponentDetails::OnCancelClicked()
+{
+	for (TWeakObjectPtr<UPCGComponent>& Component : SelectedComponents)
+	{
+		if (Component.IsValid())
+		{
+			Component.Get()->CancelGeneration();
+		}
+	}
+
+	return FReply::Handled();
+}
+
 FReply FPCGComponentDetails::OnClearPCGLinkClicked()
 {
 	for (TWeakObjectPtr<UPCGComponent>& Component : SelectedComponents)
diff --git a/Engine/Plugins/Experimental/PCG/Source/PCGEditor/Private/PCGEditor.cpp b/Engine/Plugins/Experimental/PCG/Source/PCGEditor/Private/PCGEditor.cpp
--- a/Engine/Plugins/Experimental/PCG/Source/PCGEditor/Private/PCGEditor.cpp
+++ b/Engine/Plugins/Experimental/PCG/Source/PCGEditor/Private/PCGEditor.cpp
@@ -9,6 +9,7 @@
 #include "PCGInputOutputSettings.h"
 #include "PCGPin.h"
 #include "PCGSubgraph.h"
+#include "PCGSubsystem.h"
 #include "Tests/Determinism/PCGDeterminismNativeTests.h"
 #include "Tests/Determinism/PCGDeterminismTestBlueprintBase.h"
 #include "Tests/Determinism/PCGDeterminismTestsCommon.h"
@@ -311,6 +312,12 @@ void FPCGEditor::RegisterToolbar() const
 			TAttribute<FText>(),
 			FSlateIcon(FAppStyle::GetAppStyleSetName(), "SourceControl.Actions.Refresh")));
 
+		Section.AddEntry(FToolMenuEntry::InitToolBarButton(
+			FPCGEditorCommands::Get().CancelExecution,
+			TAttribute<FText>(),
+			TAttribute<FText>(),
+			FSlateIcon(FAppStyle::GetAppStyleSetName(), "Icons.Toolbar.Stop")));
+
 		Section.AddSeparator(NAME_None);
 		Section.AddDynamicEntry("Debugging", FNewToolMenuSectionDelegate::CreateLambda([](FToolMenuSection& InSection)
 		{
@@ -355,6 +362,10 @@ void FPCGEditor::BindCommands()
 		PCGEditorCommands.ForceGraphRegeneration,
 		FExecuteAction::CreateSP(this, &FPCGEditor::OnForceGraphRegeneration_Clicked));
 
+	ToolkitCommands->MapAction(
+		PCGEditorCommands.CancelExecution,
+		FExecuteAction::CreateSP(this, &FPCGEditor::OnCancelExecution_Clicked));
+
 	ToolkitCommands->MapAction(
 		PCGEditorCommands.RunDeterminismGraphTest,
 		FExecuteAction::CreateSP(this, &FPCGEditor::OnDeterminismGraphTest),
@@ -441,6 +452,14 @@ void FPCGEditor::OnForceGraphRegeneration_Clicked()
 	}
 }
 
+void FPCGEditor::OnCancelExecution_Clicked()
+{
+	if (PCGEditorGraph)
+	{
+		GEditor->GetEditorWorldContext().World()->GetSubsystem<UPCGSubsystem>()->CancelGeneration(PCGEditorGraph->GetPCGGraph());
+	}
+}
+
 bool FPCGEditor::CanRunDeterminismNodeTest() const
 {
 	check(GraphEditorWidget.IsValid());
diff --git a/Engine/Plugins/Experimental/PCG/Source/PCGEditor/Private/PCGEditor.h b/Engine/Plugins/Experimental/PCG/Source/PCGEditor/Private/PCGEditor.h
--- a/Engine/Plugins/Experimental/PCG/Source/PCGEditor/Private/PCGEditor.h
+++ b/Engine/Plugins/Experimental/PCG/Source/PCGEditor/Private/PCGEditor.h
@@ -109,6 +109,9 @@ private:
 	/** Stop inspecting the current inspected node */
 	void OnStopInspectNode();
 
+	/** Cancels the current execution of the selected graph */
+	void OnCancelExecution_Clicked();
+
 	/** Can determinism be tested on the selected node(s) */
 	bool CanRunDeterminismNodeTest() const;
 	/** Run the determinism test on the selected node(s) */
diff --git a/Engine/Plugins/Experimental/PCG/Source/PCGEditor/Private/PCGEditorCommands.cpp b/Engine/Plugins/Experimental/PCG/Source/PCGEditor/Private/PCGEditorCommands.cpp
--- a/Engine/Plugins/Experimental/PCG/Source/PCGEditor/Private/PCGEditorCommands.cpp
+++ b/Engine/Plugins/Experimental/PCG/Source/PCGEditor/Private/PCGEditorCommands.cpp
@@ -31,6 +31,7 @@ void FPCGEditorCommands::RegisterCommands()
 	UI_COMMAND(ExecutionModeDebug, "Execution Mode Debug", "Set the execution mode for this node to Debug.", EUserInterfaceActionType::RadioButton, FInputChord());
 	UI_COMMAND(ExecutionModeIsolated, "Execution Mode Isolated", "Set the execution mode for this node to Isolated.", EUserInterfaceActionType::RadioButton, FInputChord());
 	UI_COMMAND(ExecutionModeDisabled, "Execution Mode Disabled", "Set the execution mode for this node to Disabled.", EUserInterfaceActionType::RadioButton, FInputChord());
+	UI_COMMAND(CancelExecution, "Cancel Execution", "Cancels the execution of the current graph", EUserInterfaceActionType::Button, FInputChord(EModifierKey::Control, EKeys::Escape));
 }
 
 #undef LOCTEXT_NAMESPACE
diff --git a/Engine/Plugins/Experimental/PCG/Source/PCGEditor/Private/PCGEditorCommands.h b/Engine/Plugins/Experimental/PCG/Source/PCGEditor/Private/PCGEditorCommands.h
--- a/Engine/Plugins/Experimental/PCG/Source/PCGEditor/Private/PCGEditorCommands.h
+++ b/Engine/Plugins/Experimental/PCG/Source/PCGEditor/Private/PCGEditorCommands.h
@@ -27,4 +27,5 @@ public:
 	TSharedPtr<FUICommandInfo> ExecutionModeDebug;
 	TSharedPtr<FUICommandInfo> ExecutionModeIsolated;
 	TSharedPtr<FUICommandInfo> ExecutionModeDisabled;
+	TSharedPtr<FUICommandInfo> CancelExecution;
 };
diff --git a/Engine/Plugins/Experimental/PCG/Source/PCGEditor/Private/PCGEditorModule.cpp b/Engine/Plugins/Experimental/PCG/Source/PCGEditor/Private/PCGEditorModule.cpp
--- a/Engine/Plugins/Experimental/PCG/Source/PCGEditor/Private/PCGEditorModule.cpp
+++ b/Engine/Plugins/Experimental/PCG/Source/PCGEditor/Private/PCGEditorModule.cpp
@@ -247,6 +247,25 @@ void FPCGEditorModule::PopulateMenuActions(FMenuBuilder& MenuBuilder)
 				}
 				})),
 		NAME_None);
+
+	MenuBuilder.AddMenuEntry(
+		LOCTEXT("CancelAllGeneration", "Cancel all PCG tasks"),
+		LOCTEXT("CancelAllGeneration_Tooltip", "Cancels all PCG tasks running"),
+		FSlateIcon(),
+		FUIAction(
+			FExecuteAction::CreateLambda([]() {
+				if (GEditor)
+				{
+					if (UWorld* World = GEditor->GetEditorWorldContext().World())
+					{
+						if (UPCGSubsystem* Subsystem = World->GetSubsystem<UPCGSubsystem>())
+						{
+							Subsystem->CancelAllGeneration();
+						}
+					}
+				}
+				})),
+		NAME_None);
 }
 
 void FPCGEditorModule::RegisterSettings()
diff --git a/Engine/Plugins/Experimental/PCG/Source/PCGEditor/Public/PCGComponentDetails.h b/Engine/Plugins/Experimental/PCG/Source/PCGEditor/Public/PCGComponentDetails.h
--- a/Engine/Plugins/Experimental/PCG/Source/PCGEditor/Public/PCGComponentDetails.h
+++ b/Engine/Plugins/Experimental/PCG/Source/PCGEditor/Public/PCGComponentDetails.h
@@ -4,6 +4,7 @@
 
 #include "UObject/WeakObjectPtr.h"
 #include "Input/Reply.h"
+#include "Layout/Visibility.h"
 #include "IDetailCustomization.h"
 
 class UPCGComponent;
@@ -23,7 +24,10 @@ protected:
 	virtual bool AddDefaultProperties() const { return true; }
 
 private:
+	EVisibility GenerateButtonVisible() const;
+	EVisibility CancelButtonVisible() const;
 	FReply OnGenerateClicked();
+	FReply OnCancelClicked();
 	FReply OnCleanupClicked();
 	FReply OnClearPCGLinkClicked();
 	void OnGraphChanged(UPCGComponent* InComponent);
