[PCG] Tentative fix for static analysis

#rnx

diff --git a/Engine/Plugins/Experimental/PCG/Source/PCG/Private/Elements/PCGAttributeReduceElement.cpp b/Engine/Plugins/Experimental/PCG/Source/PCG/Private/Elements/PCGAttributeReduceElement.cpp
--- a/Engine/Plugins/Experimental/PCG/Source/PCG/Private/Elements/PCGAttributeReduceElement.cpp
+++ b/Engine/Plugins/Experimental/PCG/Source/PCG/Private/Elements/PCGAttributeReduceElement.cpp
@@ -170,7 +170,7 @@ bool FPCGAttributeReduceElement::ExecuteInternal(FPCGContext* Context) const
 	const UPCGPointData* PointData = Cast<UPCGPointData>(SpatialData);
 
 	const FPCGMetadataAttributeBase* InputAttribute = SpatialData->Metadata->GetConstAttribute(Settings->InputAttributeName);
-	const FName OutputAttributeName = Settings->OutputAttributeName == NAME_None ? Settings->InputAttributeName : Settings->OutputAttributeName;
+	const FName OutputAttributeName = (Settings->OutputAttributeName == NAME_None) ? Settings->InputAttributeName : Settings->OutputAttributeName;
 	UPCGParamData* OutputParamData = NewObject<UPCGParamData>();
 	
 	FPCGMetadataAttributeBase* OutputAttribute = OutputParamData->Metadata->CopyAttribute(InputAttribute, OutputAttributeName, /*bKeepParent=*/ false, /*bCopyEntries=*/ false, /*bCopyValues=*/ false);
diff --git a/Engine/Plugins/Experimental/PCG/Source/PCG/Private/Elements/PCGAttributeSelectElement.cpp b/Engine/Plugins/Experimental/PCG/Source/PCG/Private/Elements/PCGAttributeSelectElement.cpp
--- a/Engine/Plugins/Experimental/PCG/Source/PCG/Private/Elements/PCGAttributeSelectElement.cpp
+++ b/Engine/Plugins/Experimental/PCG/Source/PCG/Private/Elements/PCGAttributeSelectElement.cpp
@@ -277,7 +277,7 @@ bool FPCGAttributeSelectElement::ExecuteInternal(FPCGContext* Context) const
 	}
 
 	const FPCGMetadataAttributeBase* InputAttribute = SpatialData->Metadata->GetConstAttribute(Settings->InputAttributeName);
-	const FName OutputAttributeName = Settings->OutputAttributeName == NAME_None ? Settings->InputAttributeName : Settings->OutputAttributeName;
+	const FName OutputAttributeName = (Settings->OutputAttributeName == NAME_None) ? Settings->InputAttributeName : Settings->OutputAttributeName;
 	UPCGParamData* OutputParamData = NewObject<UPCGParamData>();
 
 	FPCGMetadataAttributeBase* OutputAttribute = OutputParamData->Metadata->CopyAttribute(InputAttribute, OutputAttributeName, /*bKeepParent=*/ false, /*bCopyEntries=*/ false, /*bCopyValues=*/ false);
