[PCG] Tentative fix static analysis + rename file

#rnx

diff --git a/Engine/Plugins/Experimental/PCG/Source/PCG/Private/Elements/PCGAttributeTransferElement.cpp b/Engine/Plugins/Experimental/PCG/Source/PCG/Private/Elements/PCGAttributeTransferElement.cpp
--- a/Engine/Plugins/Experimental/PCG/Source/PCG/Private/Elements/PCGAttributeTransferElement.cpp
+++ b/Engine/Plugins/Experimental/PCG/Source/PCG/Private/Elements/PCGAttributeTransferElement.cpp
@@ -120,7 +120,7 @@ UPCGSpatialData* FPCGAttributeTransferElement::TransferSpatialToSpatial(FPCGCont
 	UPCGSpatialData* NewSpatialData = DuplicateObject<UPCGSpatialData>(const_cast<UPCGSpatialData*>(TargetData), nullptr);
 	NewSpatialData->InitializeFromData(TargetData);
 
-	const FName TargetAttributeName = Settings->TargetAttributeName == NAME_None ? Settings->SourceAttributeName : Settings->TargetAttributeName;
+	const FName TargetAttributeName = (Settings->TargetAttributeName == NAME_None) ? Settings->SourceAttributeName : Settings->TargetAttributeName;
 
 	// Making sure the target attribute doesn't exists in the target
 	if (NewSpatialData->Metadata->HasAttribute(TargetAttributeName))
@@ -156,7 +156,7 @@ UPCGPointData* FPCGAttributeTransferElement::TransferPointToPoint(FPCGContext* C
 
 	TargetPoints = TargetData->GetPoints();
 
-	const FName TargetAttributeName = Settings->TargetAttributeName == NAME_None ? Settings->SourceAttributeName : Settings->TargetAttributeName;
+	const FName TargetAttributeName = (Settings->TargetAttributeName == NAME_None) ? Settings->SourceAttributeName : Settings->TargetAttributeName;
 
 	// Making sure the target attribute doesn't exist in the target
 	if (NewPointData->Metadata->HasAttribute(TargetAttributeName))
diff --git a/Engine/Plugins/Experimental/PCG/Source/PCG/Private/Elements/PCGCreateAttribute.cpp b/Engine/Plugins/Experimental/PCG/Source/PCG/Private/Elements/PCGCreateAttribute.cpp
new file mode 100644
--- /dev/null
+++ b/Engine/Plugins/Experimental/PCG/Source/PCG/Private/Elements/PCGCreateAttribute.cpp
@@ -0,0 +1,299 @@
+// Copyright Epic Games, Inc. All Rights Reserved.
+
+#include "Elements/PCGCreateAttribute.h"
+
+#include "PCGData.h"
+#include "PCGParamData.h"
+#include "Data/PCGPointData.h"
+#include "Elements/Metadata/PCGMetadataElementCommon.h"
+#include "Helpers/PCGSettingsHelpers.h"
+#include "Metadata/PCGMetadata.h"
+#include "Metadata/PCGMetadataAttribute.h"
+
+namespace PCGCreateAttributeConstants
+{
+	const FName NodeName = TEXT("CreateAttribute");
+	const FName SourceLabel = TEXT("Source");
+}
+
+namespace PCGCreateAttributeElement
+{
+	template <typename Func>
+	decltype(auto) Dispatcher(const UPCGCreateAttributeSettings* Settings, const UPCGParamData* Params, Func Callback)
+	{
+		using ReturnType = decltype(Callback(double{}));
+
+		switch (Settings->Type)
+		{
+		case EPCGMetadataTypes::Integer64:
+			return Callback(PCG_GET_OVERRIDEN_VALUE(Settings, IntValue, Params));
+		case EPCGMetadataTypes::Double:
+			return Callback(PCG_GET_OVERRIDEN_VALUE(Settings, DoubleValue, Params));
+		case EPCGMetadataTypes::Vector2:
+			return Callback(PCG_GET_OVERRIDEN_VALUE(Settings, Vector2Value, Params));
+		case EPCGMetadataTypes::Vector:
+			return Callback(PCG_GET_OVERRIDEN_VALUE(Settings, VectorValue, Params));
+		case EPCGMetadataTypes::Vector4:
+			return Callback(PCG_GET_OVERRIDEN_VALUE(Settings, Vector4Value, Params));
+		case EPCGMetadataTypes::Quaternion:
+			return Callback(PCG_GET_OVERRIDEN_VALUE(Settings, QuatValue, Params));
+		case EPCGMetadataTypes::Transform:
+			return Callback(PCG_GET_OVERRIDEN_VALUE(Settings, TransformValue, Params));
+		case EPCGMetadataTypes::String:
+			return Callback(PCG_GET_OVERRIDEN_VALUE(Settings, StringValue, Params));
+		case EPCGMetadataTypes::Boolean:
+			return Callback(PCG_GET_OVERRIDEN_VALUE(Settings, BoolValue, Params));
+		case EPCGMetadataTypes::Rotator:
+			return Callback(PCG_GET_OVERRIDEN_VALUE(Settings, RotatorValue, Params));
+		case EPCGMetadataTypes::Name:
+			return Callback(PCG_GET_OVERRIDEN_VALUE(Settings, NameValue, Params));
+		default:
+			// ReturnType{} is invalid if ReturnType is void
+			if constexpr (std::is_same_v<ReturnType, void>)
+			{
+				return;
+			}
+			else
+			{
+				return ReturnType{};
+			}
+		}
+	}
+}
+
+FName UPCGCreateAttributeSettings::AdditionalTaskName() const
+{
+	if (bFromSourceParam)
+	{
+		const FName NodeName = PCGCreateAttributeConstants::NodeName;
+
+		if (OutputAttributeName == NAME_None && SourceParamAttributeName == NAME_None)
+		{
+			return NodeName;
+		}
+		else
+		{
+			const FString AttributeName = ((OutputAttributeName == NAME_None) ? SourceParamAttributeName : OutputAttributeName).ToString();
+			return FName(FString::Printf(TEXT("%s %s"), *NodeName.ToString(), *AttributeName));
+		}
+	}
+	else
+	{
+		const FString Name = OutputAttributeName.ToString();
+
+		switch (Type)
+		{
+		case EPCGMetadataTypes::Integer64:
+			return FName(FString::Printf(TEXT("%s: %lld"), *Name, IntValue));
+		case EPCGMetadataTypes::Double:
+			return FName(FString::Printf(TEXT("%s: %.2f"), *Name, DoubleValue));
+		case EPCGMetadataTypes::String:
+			return FName(FString::Printf(TEXT("%s: \"%s\""), *Name, *StringValue));
+		case EPCGMetadataTypes::Name:
+			return FName(FString::Printf(TEXT("%s: N(\"%s\")"), *Name, *NameValue.ToString()));
+		case EPCGMetadataTypes::Vector2:
+			return FName(FString::Printf(TEXT("%s: V(%.2f, %.2f)"), *Name, Vector2Value.X, Vector2Value.Y));
+		case EPCGMetadataTypes::Vector:
+			return FName(FString::Printf(TEXT("%s: V(%.2f, %.2f, %.2f)"), *Name, VectorValue.X, VectorValue.Y, VectorValue.Z));
+		case EPCGMetadataTypes::Vector4:
+			return FName(FString::Printf(TEXT("%s: V(%.2f, %.2f, %.2f, %.2f)"), *Name, Vector4Value.X, Vector4Value.Y, Vector4Value.Z, Vector4Value.W));
+		case EPCGMetadataTypes::Rotator:
+			return FName(FString::Printf(TEXT("%s: R(%.2f, %.2f, %.2f)"), *Name, RotatorValue.Roll, RotatorValue.Pitch, RotatorValue.Yaw));
+		case EPCGMetadataTypes::Quaternion:
+			return FName(FString::Printf(TEXT("%s: Q(%.2f, %.2f, %.2f, %.2f)"), *Name, QuatValue.X, QuatValue.Y, QuatValue.Z, QuatValue.W));
+		case EPCGMetadataTypes::Transform:
+			return FName(FString::Printf(TEXT("%s: Transform"), *Name));
+		case EPCGMetadataTypes::Boolean:
+			return FName(FString::Printf(TEXT("%s: %s"), *Name, (BoolValue ? TEXT("True") : TEXT("False"))));
+		default:
+			return NAME_None;
+		}
+	}
+}
+
+#if WITH_EDITOR
+FName UPCGCreateAttributeSettings::GetDefaultNodeName() const
+{
+	return PCGCreateAttributeConstants::NodeName;
+}
+#endif // WITH_EDITOR
+
+TArray<FPCGPinProperties> UPCGCreateAttributeSettings::InputPinProperties() const
+{
+	TArray<FPCGPinProperties> PinProperties;
+	PinProperties.Emplace(PCGPinConstants::DefaultInputLabel, EPCGDataType::Any, /*bInAllowMultipleConnections=*/ true);
+
+	if (bFromSourceParam)
+	{
+		PinProperties.Emplace(PCGCreateAttributeConstants::SourceLabel, EPCGDataType::Param, /*bInAllowMultipleConnections=*/ false);
+	}
+
+	return PinProperties;
+}
+
+TArray<FPCGPinProperties> UPCGCreateAttributeSettings::OutputPinProperties() const
+{
+	TArray<FPCGPinProperties> PinProperties;
+	PinProperties.Emplace(PCGPinConstants::DefaultOutputLabel, EPCGDataType::Any);
+
+	return PinProperties;
+}
+
+FPCGElementPtr UPCGCreateAttributeSettings::CreateElement() const
+{
+	return MakeShared<FPCGCreateAttributeElement>();
+}
+
+bool FPCGCreateAttributeElement::ExecuteInternal(FPCGContext* Context) const
+{
+	TRACE_CPUPROFILER_EVENT_SCOPE(FPCGCreateAttributeElement::Execute);
+
+	check(Context);
+
+	const UPCGCreateAttributeSettings* Settings = Context->GetInputSettings<UPCGCreateAttributeSettings>();
+	check(Settings);
+
+	TArray<FPCGTaggedData> SourceParams = Context->InputData.GetInputsByPin(PCGCreateAttributeConstants::SourceLabel);
+	UPCGParamData* SourceParamData = nullptr;
+	FName SourceParamAttributeName = NAME_None;
+
+	if (Settings->bFromSourceParam)
+	{
+		if (SourceParams.IsEmpty())
+		{
+			PCGE_LOG(Error, "Source param was not provided.");
+			return true;
+		}
+
+		SourceParamData = CastChecked<UPCGParamData>(SourceParams[0].Data);
+
+		if (!SourceParamData->Metadata)
+		{
+			PCGE_LOG(Error, "Source param data doesn't have metadata");
+			return true;
+		}
+
+		SourceParamAttributeName = (Settings->SourceParamAttributeName == NAME_None) ? SourceParamData->Metadata->GetLatestAttributeNameOrNone() : Settings->SourceParamAttributeName;
+
+		if (!SourceParamData->Metadata->HasAttribute(SourceParamAttributeName))
+		{
+			PCGE_LOG(Error, "Source param data doesn't have an attribute \"%s\"", *SourceParamAttributeName.ToString());
+			return true;
+		}
+	}
+
+	TArray<FPCGTaggedData> Inputs = Context->InputData.GetInputsByPin(PCGPinConstants::DefaultInputLabel);
+
+	// If the input is empty, we will create a new ParamData.
+	// We can re-use this newly object as the output
+	bool bCanReuseInputData = false;
+	if (Inputs.IsEmpty())
+	{
+		FPCGTaggedData& NewData = Inputs.Emplace_GetRef();
+		NewData.Data = NewObject<UPCGParamData>();
+		NewData.Pin = PCGPinConstants::DefaultInputLabel;
+		bCanReuseInputData = true;
+	}
+
+	for (const FPCGTaggedData& InputTaggedData : Inputs)
+	{
+		const UPCGData* InputData = InputTaggedData.Data;
+		UPCGData* OutputData = nullptr;
+
+		UPCGMetadata* Metadata = nullptr;
+
+		bool bShouldAddNewEntry = false;
+
+		if (const UPCGSpatialData* InputSpatialData = Cast<UPCGSpatialData>(InputData))
+		{
+			UPCGSpatialData* NewSpatialData = DuplicateObject<UPCGSpatialData>(const_cast<UPCGSpatialData*>(InputSpatialData), nullptr);
+			NewSpatialData->Metadata = NewObject<UPCGMetadata>();
+			NewSpatialData->InitializeFromData(InputSpatialData, /*InMetadataParentOverride=*/ nullptr, /*bInheritMetadata=*/ Settings->bKeepExistingAttributes);
+
+			OutputData = NewSpatialData;
+			Metadata = NewSpatialData->Metadata;
+		}
+		else if (const UPCGParamData* InputParamData = Cast<UPCGParamData>(InputData))
+		{
+			// If we can reuse input data, it is safe to const_cast, as it was created by ourselves above.
+			UPCGParamData* NewParamData = bCanReuseInputData ? const_cast<UPCGParamData*>(InputParamData) : NewObject<UPCGParamData>();
+			NewParamData->Metadata->Initialize((!bCanReuseInputData && Settings->bKeepExistingAttributes) ? InputParamData->Metadata : nullptr);
+
+			OutputData = NewParamData;
+			Metadata = NewParamData->Metadata;
+
+			// In case of param data, we want to add a new entry too
+			bShouldAddNewEntry = true;
+		}
+		else
+		{
+			PCGE_LOG(Error, "Invalid data as input. Only support spatial and params");
+			continue;
+		}
+
+		const FName OutputAttributeName = (Settings->bFromSourceParam && Settings->OutputAttributeName == NAME_None) ? SourceParamAttributeName : Settings->OutputAttributeName;
+
+		FPCGMetadataAttributeBase* Attribute = nullptr;
+
+		if (Settings->bFromSourceParam)
+		{
+			const FPCGMetadataAttributeBase* SourceAttribute = SourceParamData->Metadata->GetConstAttribute(SourceParamAttributeName);
+			Attribute = Metadata->CopyAttribute(SourceAttribute, OutputAttributeName, /*bKeepParent=*/false, /*bCopyEntries=*/bShouldAddNewEntry, /*bCopyValues=*/bShouldAddNewEntry);
+		}
+		else
+		{
+			Attribute = ClearOrCreateAttribute(Settings, Metadata, nullptr, &OutputAttributeName);
+		}
+
+		if (!Attribute)
+		{
+			PCGE_LOG(Error, "Error while creating attribute %s", *OutputAttributeName.ToString());
+			continue;
+		}
+
+		TArray<FPCGTaggedData>& Outputs = Context->OutputData.TaggedData;
+		FPCGTaggedData& Output = Outputs.Emplace_GetRef();
+		Output.Data = OutputData;
+
+		// Add a new entry if it is a param data and not from source (because entries are already copied)
+		if (bShouldAddNewEntry && !Settings->bFromSourceParam)
+		{
+			PCGMetadataEntryKey EntryKey = Metadata->AddEntry();
+			SetAttribute(Settings, Attribute, Metadata, EntryKey, nullptr);
+		}
+	}
+
+	return true;
+}
+
+FPCGMetadataAttributeBase* FPCGCreateAttributeElement::ClearOrCreateAttribute(const UPCGCreateAttributeSettings* Settings, UPCGMetadata* Metadata, const UPCGParamData* Params, const FName* OutputAttributeNameOverride) const
+{
+	check(Metadata);
+
+	auto CreateAttribute = [Settings, Metadata, OutputAttributeNameOverride](auto&& Value) -> FPCGMetadataAttributeBase*
+	{
+		return PCGMetadataElementCommon::ClearOrCreateAttribute(Metadata, OutputAttributeNameOverride ? *OutputAttributeNameOverride : Settings->OutputAttributeName, Value);
+	};
+
+	return PCGCreateAttributeElement::Dispatcher(Settings, Params, CreateAttribute);
+}
+
+PCGMetadataEntryKey FPCGCreateAttributeElement::SetAttribute(const UPCGCreateAttributeSettings* Settings, FPCGMetadataAttributeBase* Attribute, UPCGMetadata* Metadata, PCGMetadataEntryKey EntryKey, const UPCGParamData* Params) const
+{
+	check(Attribute && Metadata);
+
+	auto SetAttribute = [Attribute, EntryKey, Metadata](auto&& Value) -> PCGMetadataEntryKey
+	{
+		using AttributeType = std::remove_reference_t<decltype(Value)>;
+
+		check(Attribute->GetTypeId() == PCG::Private::MetadataTypes<AttributeType>::Id);
+
+		const PCGMetadataEntryKey FinalKey = (EntryKey == PCGInvalidEntryKey) ? Metadata->AddEntry() : EntryKey;
+
+		static_cast<FPCGMetadataAttribute<AttributeType>*>(Attribute)->SetValue(FinalKey, Value);
+
+		return FinalKey;
+	};
+
+	return PCGCreateAttributeElement::Dispatcher(Settings, Params, SetAttribute);
+}
diff --git a/Engine/Plugins/Experimental/PCG/Source/PCG/Private/Elements/PCGCreateAttributeBase.cpp b/Engine/Plugins/Experimental/PCG/Source/PCG/Private/Elements/PCGCreateAttributeBase.cpp
deleted file mode 100644
diff --git a/Engine/Plugins/Experimental/PCG/Source/PCG/Public/Elements/PCGCreateAttribute.h b/Engine/Plugins/Experimental/PCG/Source/PCG/Public/Elements/PCGCreateAttribute.h
new file mode 100644
--- /dev/null
+++ b/Engine/Plugins/Experimental/PCG/Source/PCG/Public/Elements/PCGCreateAttribute.h
@@ -0,0 +1,108 @@
+// Copyright Epic Games, Inc. All Rights Reserved.
+
+#pragma once
+
+#include "PCGPoint.h"
+#include "PCGSettings.h"
+#include "Metadata/PCGMetadataCommon.h"
+#include "Metadata/PCGMetadataAttributeTpl.h"
+
+#include "PCGCreateAttribute.generated.h"
+
+class FPCGMetadataAttributeBase;
+class UPCGMetadata;
+class UPCGParamData;
+
+/**
+* Node to create an attribute from a Spatial data or a ParamData.
+* If no input is provided, the node behave as a "CreateParamData" and will output a ParamData
+* with the given attribute set.
+* 
+* Note: This need to be updated if we ever add new types.
+*/
+UCLASS(BlueprintType, ClassGroup = (Procedural))
+class PCG_API UPCGCreateAttributeSettings : public UPCGSettings
+{
+	GENERATED_BODY()
+public:
+	//~Begin UPCGSettings interface
+#if WITH_EDITOR
+	virtual FName GetDefaultNodeName() const override;
+	virtual EPCGSettingsType GetType() const override { return EPCGSettingsType::Metadata; }
+#endif // WITH_EDITOR
+
+	virtual TArray<FPCGPinProperties> InputPinProperties() const override;
+	virtual TArray<FPCGPinProperties> OutputPinProperties() const override;
+
+	virtual FName AdditionalTaskName() const override;
+	//~End UPCGSettings interface
+
+	UPROPERTY(BlueprintReadWrite, EditAnywhere, Category = Settings)
+	bool bKeepExistingAttributes = false;
+
+	UPROPERTY(BlueprintReadWrite, EditAnywhere, Category = Settings)
+	FName OutputAttributeName = NAME_None;
+
+	UPROPERTY(BlueprintReadWrite, EditAnywhere, Category = Settings)
+	bool bFromSourceParam = false;
+
+	UPROPERTY(BlueprintReadWrite, EditAnywhere, Category = Settings, meta = (EditCondition = "bFromSourceParam", EditConditionHides))
+	FName SourceParamAttributeName = NAME_None;
+
+	UPROPERTY(BlueprintReadWrite, EditAnywhere, Category = Settings, meta = (ValidEnumValues = "Double, Integer64, Vector2, Vector, Vector4, Quaternion, Transform, String, Boolean, Rotator, Name", EditCondition = "!bFromSourceParam", EditConditionHides))
+	EPCGMetadataTypes Type = EPCGMetadataTypes::Double;
+
+	// All different types
+	UPROPERTY(BlueprintReadWrite, EditAnywhere, Category = Settings, meta = (EditCondition = "Type == EPCGMetadataTypes::Double && !bFromSourceParam", EditConditionHides))
+	double DoubleValue = 0.0;
+
+	UPROPERTY(BlueprintReadWrite, EditAnywhere, Category = Settings, meta = (EditCondition = "Type == EPCGMetadataTypes::Integer64 && !bFromSourceParam", EditConditionHides))
+	int64 IntValue = 0;
+
+	UPROPERTY(BlueprintReadWrite, EditAnywhere, Category = Settings, meta = (EditCondition = "Type == EPCGMetadataTypes::Vector2 && !bFromSourceParam", EditConditionHides))
+	FVector2D Vector2Value = FVector2D::ZeroVector;
+
+	UPROPERTY(BlueprintReadWrite, EditAnywhere, Category = Settings, meta = (EditCondition = "Type == EPCGMetadataTypes::Vector && !bFromSourceParam", EditConditionHides))
+	FVector VectorValue = FVector::ZeroVector;
+
+	UPROPERTY(BlueprintReadWrite, EditAnywhere, Category = Settings, meta = (EditCondition = "Type == EPCGMetadataTypes::Vector4 && !bFromSourceParam", EditConditionHides))
+	FVector4 Vector4Value = FVector4::Zero();
+
+	UPROPERTY(BlueprintReadWrite, EditAnywhere, Category = Settings, meta = (EditCondition = "Type == EPCGMetadataTypes::Quaternion && !bFromSourceParam", EditConditionHides))
+	FQuat QuatValue = FQuat::Identity;
+
+	UPROPERTY(BlueprintReadWrite, EditAnywhere, Category = Settings, meta = (EditCondition = "Type == EPCGMetadataTypes::Transform && !bFromSourceParam", EditConditionHides))
+	FTransform TransformValue = FTransform::Identity;
+
+	UPROPERTY(BlueprintReadWrite, EditAnywhere, Category = Settings, meta = (EditCondition = "Type == EPCGMetadataTypes::String && !bFromSourceParam", EditConditionHides))
+	FString StringValue = "";
+
+	UPROPERTY(BlueprintReadWrite, EditAnywhere, Category = Settings, meta = (EditCondition = "Type == EPCGMetadataTypes::Boolean && !bFromSourceParam", EditConditionHides))
+	bool BoolValue = false;
+
+	UPROPERTY(BlueprintReadWrite, EditAnywhere, Category = Settings, meta = (EditCondition = "Type == EPCGMetadataTypes::Rotator && !bFromSourceParam", EditConditionHides))
+	FRotator RotatorValue = FRotator::ZeroRotator;
+
+	UPROPERTY(BlueprintReadWrite, EditAnywhere, Category = Settings, meta = (EditCondition = "Type == EPCGMetadataTypes::Name && !bFromSourceParam", EditConditionHides))
+	FName NameValue = NAME_None;
+
+protected:
+	virtual FPCGElementPtr CreateElement() const override;
+};
+
+class FPCGCreateAttributeElement : public FSimplePCGElement
+{
+public:
+	virtual bool CanExecuteOnlyOnMainThread(FPCGContext* Context) const { return true; }
+	virtual bool IsCacheable(const UPCGSettings* InSettings) const override { return false; }
+
+protected:
+	virtual bool ExecuteInternal(FPCGContext* Context) const override;
+
+private:
+	/* Create (or clear) an attribute named by OutputAttributeName and depending on the selected type. Value can be overridden by params. Default value will be set to the specified value. */
+	FPCGMetadataAttributeBase* ClearOrCreateAttribute(const UPCGCreateAttributeSettings* Settings, UPCGMetadata* Metadata, const UPCGParamData* Params = nullptr, const FName* OutputAttributeNameOverride = nullptr) const;
+
+	/* Set an entry defined by EntryKey (or create one if it is invalid) in the Attribute depending on the selected type. Value can be overridden by params. */
+	PCGMetadataEntryKey SetAttribute(const UPCGCreateAttributeSettings* Settings, FPCGMetadataAttributeBase* Attribute, UPCGMetadata* Metadata, PCGMetadataEntryKey EntryKey, const UPCGParamData* Params = nullptr) const;
+};
\ No newline at end of file
diff --git a/Engine/Plugins/Experimental/PCG/Source/PCG/Public/Elements/PCGCreateAttributeBase.h b/Engine/Plugins/Experimental/PCG/Source/PCG/Public/Elements/PCGCreateAttributeBase.h
deleted file mode 100644
