[PCG] Always try to destroy actor when cleaning up.

* Should remove the assert (and then crash) when we exit PIE, in a client/server configuration.

diff --git a/Engine/Plugins/Experimental/PCG/Source/PCG/Private/Helpers/PCGActorHelpers.cpp b/Engine/Plugins/Experimental/PCG/Source/PCG/Private/Helpers/PCGActorHelpers.cpp
--- a/Engine/Plugins/Experimental/PCG/Source/PCG/Private/Helpers/PCGActorHelpers.cpp
+++ b/Engine/Plugins/Experimental/PCG/Source/PCG/Private/Helpers/PCGActorHelpers.cpp
@@ -237,15 +237,9 @@ bool UPCGActorHelpers::DeleteActors(UWorld* World, const TArray<TSoftObjectPtr<A
 		{
 			if (AActor* Actor = ManagedActor.Get())
 			{
-				if (Actor->GetWorld() == World)
+				if (!ensure(World->DestroyActor(Actor)))
 				{
-					World->DestroyActor(Actor);
-				}
-				else
-				{
-					// If we're here and the world is null, then the actor has either been destroyed already or it will be picked up by GC by design.
-					// Otherwise, we have bigger issues, something is very wrong.
-					check(Actor->GetWorld() == nullptr);
+					UE_LOG(LogPCG, Warning, TEXT("Actor %s failed to be destroyed."), *Actor->GetPathName());
 				}
 			}
 		}
