Fixed warning for uninitialized variables

diff --git a/Engine/Plugins/Experimental/PCG/Source/PCG/Public/Grid/PCGLandscapeCache.h b/Engine/Plugins/Experimental/PCG/Source/PCG/Public/Grid/PCGLandscapeCache.h
--- a/Engine/Plugins/Experimental/PCG/Source/PCG/Public/Grid/PCGLandscapeCache.h
+++ b/Engine/Plugins/Experimental/PCG/Source/PCG/Public/Grid/PCGLandscapeCache.h
@@ -20,10 +20,10 @@ struct FPCGLandscapeLayerWeight
 	GENERATED_BODY()
 
 	UPROPERTY(EditAnywhere, BlueprintReadWrite, Category = "Landscape Attribute")
-	FName Name;
+	FName Name = NAME_None;
 
 	UPROPERTY(EditAnywhere, BlueprintReadWrite, Category = "Landscape Attribute")
-	float Weight;
+	float Weight = 0.0f;
 };
 
 struct FPCGLandscapeCacheEntry
