Add actor-to-PCG blueprint methods

diff --git a/Engine/Plugins/Experimental/PCG/Source/PCG/Private/Elements/PCGActorSelector.cpp b/Engine/Plugins/Experimental/PCG/Source/PCG/Private/Elements/PCGActorSelector.cpp
new file mode 100644
--- /dev/null
+++ b/Engine/Plugins/Experimental/PCG/Source/PCG/Private/Elements/PCGActorSelector.cpp
@@ -0,0 +1,161 @@
+// Copyright Epic Games, Inc. All Rights Reserved.
+
+#include "Elements/PCGActorSelector.h"
+#include "Helpers/PCGActorHelpers.h"
+
+#include "Kismet/GameplayStatics.h"
+
+namespace PCGActorSelector
+{
+	// Need to pass a pointer of pointer to the found actor. The lambda will capture this pointer and modify its value when an actor is found.
+	TFunction<bool(AActor*)> GetFilteringFunction(const FPCGActorSelectorSettings& InSettings, AActor*& InFoundActor)
+	{
+		switch (InSettings.ActorSelection)
+		{
+		case EPCGActorSelection::ByTag:
+			return[ActorSelectionTag = InSettings.ActorSelectionTag, &InFoundActor](AActor* Actor) -> bool
+			{
+				if (Actor->ActorHasTag(ActorSelectionTag))
+				{
+					InFoundActor = Actor;
+					return false;
+				}
+
+				return true;
+			};
+
+		case EPCGActorSelection::ByName:
+			return[ActorSelectionName = InSettings.ActorSelectionName, &InFoundActor](AActor* Actor) -> bool
+			{
+				if (Actor->GetFName().IsEqual(ActorSelectionName, ENameCase::IgnoreCase, /*bCompareNumber=*/ false))
+				{
+					InFoundActor = Actor;
+					return false;
+				}
+
+				return true;
+			};
+
+		case EPCGActorSelection::ByClass:
+			return[ActorSelectionClass = InSettings.ActorSelectionClass, &InFoundActor](AActor* Actor) -> bool
+			{
+				if (Actor->IsA(ActorSelectionClass))
+				{
+					InFoundActor = Actor;
+					return false;
+				}
+
+				return true;
+			};
+
+		default:
+			break;
+		}
+
+		return [](AActor* Actor) -> bool { return false; };
+	}
+
+	AActor* FindActor(const FPCGActorSelectorSettings& Settings, UWorld* World, AActor* Self)
+	{
+		TRACE_CPUPROFILER_EVENT_SCOPE(PCGActorSelector::FindActor);
+
+		AActor* FoundActor = nullptr;
+
+		if (!World)
+		{
+			return FoundActor;
+		}
+
+		// Early out if we have not the information necessary
+		if ((Settings.ActorSelection == EPCGActorSelection::ByTag && Settings.ActorSelectionTag == NAME_None) ||
+			(Settings.ActorSelection == EPCGActorSelection::ByName && Settings.ActorSelectionName == NAME_None) ||
+			(Settings.ActorSelection == EPCGActorSelection::ByClass && !Settings.ActorSelectionClass))
+		{
+			return FoundActor;
+		}
+
+		// We pass FoundActor ref, that will be captured by the FilteringFunction
+		// It will modify the FoundActor pointer to the found actor, if found.
+		TFunction<bool(AActor*)> FilteringFunction = PCGActorSelector::GetFilteringFunction(Settings, FoundActor);
+
+		// In case of iterating over all actors in the world, call our filtering function and get out.
+		if (Settings.ActorFilter == EPCGActorFilter::AllWorldActors)
+		{
+			UPCGActorHelpers::ForEachActorInWorld<AActor>(World, FilteringFunction);
+
+			// FoundActor is set by the FilteringFunction (captured)
+			return FoundActor;
+		}
+
+		// Otherwise, gather all the actors we need to check
+		TArray<AActor*> ActorsToCheck;
+		switch (Settings.ActorFilter)
+		{
+		case EPCGActorFilter::Self:
+			if (Self)
+			{
+				ActorsToCheck.Add(Self);
+			}
+			break;
+
+		case EPCGActorFilter::Parent:
+			if (Self)
+			{
+				if (AActor* Parent = Self->GetParentActor())
+				{
+					ActorsToCheck.Add(Parent);
+				}
+				else
+				{
+					// If there is no parent, set the owner as the parent.
+					ActorsToCheck.Add(Self);
+				}
+			}
+			break;
+
+		case EPCGActorFilter::Root:
+		{
+			AActor* Current = Self;
+			while (Current != nullptr)
+			{
+				AActor* Parent = Current->GetParentActor();
+				if (Parent == nullptr)
+				{
+					ActorsToCheck.Add(Current);
+					break;
+				}
+				Current = Parent;
+			}
+
+			break;
+		}
+
+		//case EPCGActorFilter::TrackedActors:
+			//	//TODO
+			//	break;
+
+		default:
+			break;
+		}
+
+		if (Settings.bIncludeChildren)
+		{
+			const int32 InitialCount = ActorsToCheck.Num();
+			for (int32 i = 0; i < InitialCount; ++i)
+			{
+				ActorsToCheck[i]->GetAttachedActors(ActorsToCheck, /*bResetArray=*/ false, /*bRecursivelyIncludeAttachedActors=*/ true);
+			}
+		}
+
+		for (AActor* Actor : ActorsToCheck)
+		{
+			// FoundActor is set by the FilteringFunction (captured)
+			if (!FilteringFunction(Actor))
+			{
+				break;
+			}
+		}
+
+		return FoundActor;
+	}
+}
\ No newline at end of file
diff --git a/Engine/Plugins/Experimental/PCG/Source/PCG/Private/Elements/PCGDataFromActor.cpp b/Engine/Plugins/Experimental/PCG/Source/PCG/Private/Elements/PCGDataFromActor.cpp
new file mode 100644
--- /dev/null
+++ b/Engine/Plugins/Experimental/PCG/Source/PCG/Private/Elements/PCGDataFromActor.cpp
@@ -0,0 +1,34 @@
+// Copyright Epic Games, Inc. All Rights Reserved.
+
+#include "Elements/PCGDataFromActor.h"
+
+#include "PCGComponent.h"
+
+FPCGElementPtr UPCGDataFromActorSettings::CreateElement() const
+{
+	return MakeShared<FPCGDataFromActorElement>();
+}
+
+bool FPCGDataFromActorElement::ExecuteInternal(FPCGContext* Context) const
+{
+	TRACE_CPUPROFILER_EVENT_SCOPE(FPCGDataFromActorElement::Execute);
+
+	check(Context);
+	const UPCGDataFromActorSettings* Settings = Context->GetInputSettings<UPCGDataFromActorSettings>();
+	check(Settings);
+
+	UWorld* World = Context->SourceComponent.IsValid() ? Context->SourceComponent->GetWorld() : nullptr;
+	AActor* FoundActor = PCGActorSelector::FindActor(Settings->ActorSelector, World, nullptr);
+
+	if (!FoundActor)
+	{
+		PCGE_LOG(Warning, "No matching actor was found.");
+		return true;
+	}
+
+	TArray<FPCGTaggedData>& Outputs = Context->OutputData.TaggedData;
+	FPCGTaggedData& Output = Outputs.Emplace_GetRef();
+	Output.Data = UPCGComponent::CreateActorPCGData(FoundActor, Context->SourceComponent.Get(), /*bParseActor=*/true);
+
+	return true;
+}
\ No newline at end of file
diff --git a/Engine/Plugins/Experimental/PCG/Source/PCG/Private/Elements/PCGPropertyToParamData.cpp b/Engine/Plugins/Experimental/PCG/Source/PCG/Private/Elements/PCGPropertyToParamData.cpp
--- a/Engine/Plugins/Experimental/PCG/Source/PCG/Private/Elements/PCGPropertyToParamData.cpp
+++ b/Engine/Plugins/Experimental/PCG/Source/PCG/Private/Elements/PCGPropertyToParamData.cpp
@@ -7,179 +7,34 @@
 #include "PCGComponent.h"
 #include "PCGParamData.h"
 #include "Helpers/PCGBlueprintHelpers.h"
-#include "Helpers/PCGActorHelpers.h"
 #include "Metadata/PCGMetadata.h"
 #include "Metadata/PCGMetadataAttributeTraits.h"
 
 #include "EngineUtils.h"
-#include "Kismet/GameplayStatics.h"
 
-namespace PCGPropertyToParamDataHelpers
+void UPCGPropertyToParamDataSettings::PostLoad()
 {
-	// Need to pass a pointer of pointer to the found actor. The lambda will capture this pointer and modify its value when an actor is found.
-	TFunction<bool(AActor*)> GetFilteringFunction(const UPCGPropertyToParamDataSettings* InSettings, AActor*& InFoundActor)
+	// Migrate deprecated actor selection settings to struct if needed
+	if (ActorSelection_DEPRECATED != EPCGActorSelection::ByTag ||
+		ActorSelectionTag_DEPRECATED != NAME_None ||
+		ActorSelectionName_DEPRECATED != NAME_None ||
+		ActorSelectionClass_DEPRECATED != TSubclassOf<AActor>{} ||
+		ActorFilter_DEPRECATED != EPCGActorFilter::Self ||
+		bIncludeChildren_DEPRECATED != false)
 	{
-		check(InSettings);
-
-		switch (InSettings->ActorSelection)
-		{
-		case EPCGActorSelection::ByTag:
-			return [ActorSelectionTag = InSettings->ActorSelectionTag, &InFoundActor](AActor* Actor) -> bool
-				{
-					if (Actor->ActorHasTag(ActorSelectionTag))
-					{
-						InFoundActor = Actor;
-						return false;
-					}
-
-					return true;
-				};
-
-		case EPCGActorSelection::ByName:
-			return [ActorSelectionName = InSettings->ActorSelectionName, &InFoundActor](AActor* Actor) -> bool
-				{
-					if (Actor->GetFName().IsEqual(ActorSelectionName, ENameCase::IgnoreCase, /*bCompareNumber=*/ false))
-					{
-						InFoundActor = Actor;
-						return false;
-					}
-
-					return true;
-				};
-
-		case EPCGActorSelection::ByClass:
-			return [ActorSelectionClass = InSettings->ActorSelectionClass, &InFoundActor](AActor* Actor) -> bool
-				{
-					if (Actor->IsA(ActorSelectionClass))
-					{
-						InFoundActor = Actor;
-						return false;
-					}
-
-					return true;
-				};
-
-		default:
-			break;
-		}
-
-		return [](AActor* Actor) -> bool { return false; };
-	}
-
-	AActor* FindActor(FPCGContext& InContext)
-	{
-		TRACE_CPUPROFILER_EVENT_SCOPE(PCGPropertyToParamDataHelpers::FindActor);
-
-		AActor* FoundActor = nullptr;
-
-		if (!InContext.SourceComponent.IsValid())
-		{
-			return FoundActor;
-		}
-
-		UPCGComponent* OriginalComponent = UPCGBlueprintHelpers::GetOriginalComponent(InContext);
-		check(OriginalComponent);
-
-		UWorld* World = OriginalComponent->GetWorld();
-		if (!World)
-		{
-			return FoundActor;
-		}
-
-		const UPCGPropertyToParamDataSettings* Settings = InContext.GetInputSettings<UPCGPropertyToParamDataSettings>();
-		check(Settings);
-
-		// Early out if we have not the information necessary
-		if ((Settings->ActorSelection == EPCGActorSelection::ByTag && Settings->ActorSelectionTag == NAME_None) ||
-			(Settings->ActorSelection == EPCGActorSelection::ByName && Settings->ActorSelectionName == NAME_None) ||
-			(Settings->ActorSelection == EPCGActorSelection::ByClass && !Settings->ActorSelectionClass))
-		{
-			return FoundActor;
-		}
-
-		// We pass FoundActor ref, that will be captured by the FilteringFunction
-		// It will modify the FoundActor pointer to the found actor, if found.
-		TFunction<bool(AActor*)> FilteringFunction = PCGPropertyToParamDataHelpers::GetFilteringFunction(Settings, FoundActor);
-
-		// In case of iterating over all actors in the world, call our filtering function and get out.
-		if (Settings->ActorFilter == EPCGActorFilter::AllWorldActors)
-		{
-			UPCGActorHelpers::ForEachActorInWorld<AActor>(World, FilteringFunction);
-
-			// FoundActor is set by the FilteringFunction (captured)
-			return FoundActor;
-		}
-
-		// Otherwise, gather all the actors we need to check
-		TArray<AActor*> ActorsToCheck;
-		switch (Settings->ActorFilter)
-		{
-		case EPCGActorFilter::Self:
-			if (AActor* Owner = OriginalComponent->GetOwner())
-			{
-				ActorsToCheck.Add(Owner);
-			}
-			break;
-
-		case EPCGActorFilter::Parent:
-			if (AActor* Owner = OriginalComponent->GetOwner())
-			{
-				if (AActor* Parent = Owner->GetParentActor())
-				{
-					ActorsToCheck.Add(Parent);
-				}
-				else
-				{
-					// If there is no parent, set the owner as the parent.
-					ActorsToCheck.Add(Owner);
-				}
-			}
-			break;
-
-		case EPCGActorFilter::Root:
-		{
-			AActor* Current = OriginalComponent->GetOwner();
-			while (Current != nullptr)
-			{
-				AActor* Parent = Current->GetParentActor();
-				if (Parent == nullptr)
-				{
-					ActorsToCheck.Add(Current);
-					break;
-				}
-				Current = Parent;
-			}
-
-			break;
-		}
-
-		//case EPCGActorFilter::TrackedActors:
-			//	//TODO
-			//	break;
-
-		default:
-			break;
-		}
-
-		if (Settings->bIncludeChildren)
-		{
-			int32 InitialCount = ActorsToCheck.Num();
-			for (int32 i = 0; i < InitialCount; ++i)
-			{
-				ActorsToCheck[i]->GetAttachedActors(ActorsToCheck, /*bResetArray=*/ false, /*bRecursivelyIncludeAttachedActors=*/ true);
-			}
-		}
-
-		for (AActor* Actor : ActorsToCheck)
-		{
-			// FoundActor is set by the FilteringFunction (captured)
-			if (!FilteringFunction(Actor))
-			{
-				break;
-			}
-		}
-
-		return FoundActor;
+		ActorSelector.ActorSelection = ActorSelection_DEPRECATED;
+		ActorSelector.ActorSelectionTag = ActorSelectionTag_DEPRECATED;
+		ActorSelector.ActorSelectionName = ActorSelectionName_DEPRECATED;
+		ActorSelector.ActorSelectionClass = ActorSelectionClass_DEPRECATED;
+		ActorSelector.ActorFilter = ActorFilter_DEPRECATED;
+		ActorSelector.bIncludeChildren = bIncludeChildren_DEPRECATED;
+
+		ActorSelection_DEPRECATED = EPCGActorSelection::ByTag;
+		ActorSelectionTag_DEPRECATED = NAME_None;
+		ActorSelectionName_DEPRECATED = NAME_None;
+		ActorSelectionClass_DEPRECATED = TSubclassOf<AActor>{};
+		ActorFilter_DEPRECATED = EPCGActorFilter::Self;
+		bIncludeChildren_DEPRECATED = false;
 	}
 }
 
@@ -196,7 +51,6 @@ FPCGElementPtr UPCGPropertyToParamDataSettings::CreateElement() const
 	return MakeShared<FPCGPropertyToParamDataElement>();
 }
 
-
 bool FPCGPropertyToParamDataElement::ExecuteInternal(FPCGContext* Context) const
 {
 	TRACE_CPUPROFILER_EVENT_SCOPE(FPCGPropertyToParamDataElement::Execute);
@@ -221,7 +75,8 @@ bool FPCGPropertyToParamDataElement::ExecuteInternal(FPCGContext* Context) const
 	}
 
 	// First find the actor depending on the selection
-	AActor* FoundActor = PCGPropertyToParamDataHelpers::FindActor(*Context);
+	UPCGComponent* OriginalComponent = UPCGBlueprintHelpers::GetOriginalComponent(*Context);
+	AActor* FoundActor = PCGActorSelector::FindActor(Settings->ActorSelector, OriginalComponent ? OriginalComponent->GetWorld() : nullptr, OriginalComponent ? OriginalComponent->GetOwner() : nullptr);
 
 	if (!FoundActor)
 	{
diff --git a/Engine/Plugins/Experimental/PCG/Source/PCG/Private/Helpers/PCGBlueprintHelpers.cpp b/Engine/Plugins/Experimental/PCG/Source/PCG/Private/Helpers/PCGBlueprintHelpers.cpp
--- a/Engine/Plugins/Experimental/PCG/Source/PCG/Private/Helpers/PCGBlueprintHelpers.cpp
+++ b/Engine/Plugins/Experimental/PCG/Source/PCG/Private/Helpers/PCGBlueprintHelpers.cpp
@@ -114,6 +114,11 @@ FBox UPCGBlueprintHelpers::GetActorLocalBoundsPCG(AActor* InActor, bool bIgnoreP
 	return PCGHelpers::GetActorLocalBounds(InActor, bIgnorePCGCreatedComponents);
 }
 
+UPCGData* UPCGBlueprintHelpers::CreatePCGDataFromActor(AActor* InActor, bool bParseActor)
+{
+	return UPCGComponent::CreateActorPCGData(InActor, nullptr, bParseActor);
+}
+
 TArray<FPCGLandscapeLayerWeight> UPCGBlueprintHelpers::GetInterpolatedPCGLandscapeLayerWeights(UObject* WorldContextObject, const FVector& Location)
 {
 	UWorld* World = WorldContextObject ? WorldContextObject->GetWorld() : nullptr;
diff --git a/Engine/Plugins/Experimental/PCG/Source/PCG/Private/PCGComponent.cpp b/Engine/Plugins/Experimental/PCG/Source/PCG/Private/PCGComponent.cpp
--- a/Engine/Plugins/Experimental/PCG/Source/PCG/Private/PCGComponent.cpp
+++ b/Engine/Plugins/Experimental/PCG/Source/PCG/Private/PCGComponent.cpp
@@ -1709,6 +1709,11 @@ UPCGData* UPCGComponent::CreateActorPCGData()
 }
 
 UPCGData* UPCGComponent::CreateActorPCGData(AActor* Actor, bool bParseActor)
+{
+	return CreateActorPCGData(Actor, this, bParseActor);
+}
+
+UPCGData* UPCGComponent::CreateActorPCGData(AActor* Actor, const UPCGComponent* Component, bool bParseActor)
 {
 	TRACE_CPUPROFILER_EVENT_SCOPE(UPCGComponent::CreateActorPCGData);
 
@@ -1721,27 +1726,28 @@ UPCGData* UPCGComponent::CreateActorPCGData(AActor* Actor, bool bParseActor)
 	// TODO: add factory for extensibility
 	if (APCGPartitionActor* PartitionActor = Cast<APCGPartitionActor>(Actor))
 	{
-		check(GetOwner() == Actor); // Invalid processing otherwise because of the this usage
-		if (UPCGComponent* OriginalComponent = PartitionActor->GetOriginalComponent(this))
+		check(!Component || Component->GetOwner() == Actor); // Invalid processing otherwise because of the this usage
+
+		UPCGVolumeData* Data = NewObject<UPCGVolumeData>();
+		Data->Initialize(PartitionActor->GetFixedBounds(), PartitionActor);
+
+		UPCGComponent* OriginalComponent = Component ? PartitionActor->GetOriginalComponent(Component) : nullptr;
+		if (OriginalComponent)
 		{
-			check(OriginalComponent->IsPartitioned());
-			// TODO: cache/share the original component's actor pcg data
 			if (const UPCGSpatialData* OriginalComponentSpatialData = Cast<const UPCGSpatialData>(OriginalComponent->GetActorPCGData()))
 			{
-				UPCGVolumeData* Data = NewObject<UPCGVolumeData>();
-				Data->Initialize(PartitionActor->GetFixedBounds(), PartitionActor);
-
 				return Data->IntersectWith(OriginalComponentSpatialData);
 			}
 		}
 
-		// TODO: review this once we support non-spatial data?
-		return nullptr;
+		return Data;
 	}
 	else if (ALandscapeProxy* Landscape = Cast<ALandscapeProxy>(Actor))
 	{
 		UPCGLandscapeData* Data = NewObject<UPCGLandscapeData>();
-		Data->Initialize({ Landscape }, GetGridBounds(Actor), /*bHeightOnly=*/false, /*bUseMetadata=*/Graph && Graph->bLandscapeUsesMetadata);
+		const bool bUseLandscapeMetadata = (!Component || !Component->Graph || (Component->Graph && Component->Graph->bLandscapeUsesMetadata));
+
+		Data->Initialize({ Landscape }, PCGHelpers::GetGridBounds(Actor, Component), /*bHeightOnly=*/false, bUseLandscapeMetadata);
 
 		return Data;
 	}
@@ -2011,40 +2017,12 @@ bool UPCGComponent::IsLandscapeCachedDataDirty(const UPCGData* Data) const
 
 FBox UPCGComponent::GetGridBounds() const
 {
-	return GetGridBounds(GetOwner());
+	return PCGHelpers::GetGridBounds(GetOwner(), this);
 }
 
 FBox UPCGComponent::GetGridBounds(AActor* Actor) const
 {
-	check(Actor);
-
-	FBox Bounds(EForceInit::ForceInit);
-
-	if (APCGPartitionActor* PartitionActor = Cast<APCGPartitionActor>(Actor))
-	{
-		// First, get the bounds from the partition actor
-		Bounds = PartitionActor->GetFixedBounds();
-
-		// Then intersect with the original component's bounds.
-		if (const UPCGComponent* OriginalComponent = PartitionActor->GetOriginalComponent(this))
-		{
-			if (OriginalComponent->GetOwner() != PartitionActor)
-			{
-				Bounds = Bounds.Overlap(OriginalComponent->GetGridBounds());
-			}
-		}
-	}
-	// TODO: verify this works as expected in non-editor builds
-	else if (ALandscapeProxy* LandscapeActor = Cast<ALandscape>(Actor))
-	{
-		Bounds = PCGHelpers::GetLandscapeBounds(LandscapeActor);
-	}
-	else
-	{
-		Bounds = PCGHelpers::GetActorBounds(Actor);
-	}
-
-	return Bounds;
+	return PCGHelpers::GetGridBounds(Actor, this);
 }
 
 UPCGSubsystem* UPCGComponent::GetSubsystem() const
diff --git a/Engine/Plugins/Experimental/PCG/Source/PCG/Private/PCGHelpers.cpp b/Engine/Plugins/Experimental/PCG/Source/PCG/Private/PCGHelpers.cpp
--- a/Engine/Plugins/Experimental/PCG/Source/PCG/Private/PCGHelpers.cpp
+++ b/Engine/Plugins/Experimental/PCG/Source/PCG/Private/PCGHelpers.cpp
@@ -2,8 +2,10 @@
 
 #include "PCGHelpers.h"
 
+#include "PCGComponent.h"
 #include "PCGSubsystem.h"
 #include "PCGWorldActor.h"
+#include "Grid/PCGPartitionActor.h"
 
 #include "GameFramework/Actor.h"
 #include "Landscape.h"
@@ -56,6 +58,39 @@ namespace PCGHelpers
 		}
 	}
 
+	FBox GetGridBounds(AActor* Actor, const UPCGComponent* Component)
+	{
+		check(Actor);
+
+		FBox Bounds(EForceInit::ForceInit);
+
+		if (APCGPartitionActor* PartitionActor = Cast<APCGPartitionActor>(Actor))
+		{
+			// First, get the bounds from the partition actor
+			Bounds = PartitionActor->GetFixedBounds();
+
+			const UPCGComponent* OriginalComponent = Component ? PartitionActor->GetOriginalComponent(Component) : nullptr;
+			if (OriginalComponent)
+			{
+				if (OriginalComponent->GetOwner() != PartitionActor)
+				{
+					Bounds = Bounds.Overlap(OriginalComponent->GetGridBounds());
+				}
+			}
+		}
+		// TODO: verify this works as expected in non-editor builds
+		else if (ALandscapeProxy* LandscapeActor = Cast<ALandscape>(Actor))
+		{
+			Bounds = GetLandscapeBounds(LandscapeActor);
+		}
+		else
+		{
+			Bounds = GetActorBounds(Actor);
+		}
+
+		return Bounds;
+	}
+
 	FBox GetActorBounds(AActor* InActor, bool bIgnorePCGCreatedComponents)
 	{
 		// Specialized version of GetComponentsBoundingBox that skips over PCG generated components
diff --git a/Engine/Plugins/Experimental/PCG/Source/PCG/Private/PCGHelpers.h b/Engine/Plugins/Experimental/PCG/Source/PCG/Private/PCGHelpers.h
--- a/Engine/Plugins/Experimental/PCG/Source/PCG/Private/PCGHelpers.h
+++ b/Engine/Plugins/Experimental/PCG/Source/PCG/Private/PCGHelpers.h
@@ -9,6 +9,7 @@ class AActor;
 class APCGWorldActor;
 class ALandscape;
 class ALandscapeProxy;
+class UPCGComponent;
 class UWorld;
 
 namespace PCGHelpers
@@ -27,6 +28,9 @@ namespace PCGHelpers
 
 	FBox OverlapBounds(const FBox& InBoxA, const FBox& InBoxB);
 
+	/** Returns the bounds of InActor, intersected with the component if InActor is a partition actor */
+	FBox GetGridBounds(AActor* InActor, const UPCGComponent* InComponent);
+
 	FBox GetActorBounds(AActor* InActor, bool bIgnorePCGCreatedComponents = true);
 	FBox GetActorLocalBounds(AActor* InActor, bool bIgnorePCGCreatedComponents = true);
 	FBox GetLandscapeBounds(ALandscapeProxy* InLandscape);
diff --git a/Engine/Plugins/Experimental/PCG/Source/PCG/Private/Tests/Elements/PCGPropertyToParamDataElementTest.cpp b/Engine/Plugins/Experimental/PCG/Source/PCG/Private/Tests/Elements/PCGPropertyToParamDataElementTest.cpp
--- a/Engine/Plugins/Experimental/PCG/Source/PCG/Private/Tests/Elements/PCGPropertyToParamDataElementTest.cpp
+++ b/Engine/Plugins/Experimental/PCG/Source/PCG/Private/Tests/Elements/PCGPropertyToParamDataElementTest.cpp
@@ -108,9 +108,9 @@ bool FPCGPropertyToParamDataPropertyTypeTest::RunTest(const FString& Parameters)
 	bool bSuccess = true;
 
 	UPCGPropertyToParamDataSettings* Settings = NewObject<UPCGPropertyToParamDataSettings>();
-	Settings->ActorSelection = EPCGActorSelection::ByClass;
-	Settings->ActorSelectionClass = APCGUnitTestDummyActor::StaticClass();
-	Settings->ActorFilter = EPCGActorFilter::Self;
+	Settings->ActorSelector.ActorSelection = EPCGActorSelection::ByClass;
+	Settings->ActorSelector.ActorSelectionClass = APCGUnitTestDummyActor::StaticClass();
+	Settings->ActorSelector.ActorFilter = EPCGActorFilter::Self;
 
 	static constexpr int32 Seed = 42;
 	const FString ExtraTestWhat = "PropertyToParamDataPropertyTypeTest";
@@ -209,15 +209,15 @@ bool FPCGPropertyToParamDataActorFindTest::RunTest(const FString& Parameters)
 	static const FString Name = TEXT("MyPCGUnitTestDummyActor");
 
 	UPCGPropertyToParamDataSettings* Settings = NewObject<UPCGPropertyToParamDataSettings>();
-	Settings->ActorSelectionClass = APCGUnitTestDummyActor::StaticClass();
-	Settings->ActorSelectionTag = Tag;
-	Settings->ActorSelectionName = APCGUnitTestDummyActor::StaticClass()->GetFName();
+	Settings->ActorSelector.ActorSelectionClass = APCGUnitTestDummyActor::StaticClass();
+	Settings->ActorSelector.ActorSelectionTag = Tag;
+	Settings->ActorSelector.ActorSelectionName = APCGUnitTestDummyActor::StaticClass()->GetFName();
 	Settings->ComponentClass = UPCGUnitTestDummyComponent::StaticClass();
 
 	// Self by class
 	{
-		Settings->ActorSelection = EPCGActorSelection::ByClass;
-		Settings->ActorFilter = EPCGActorFilter::Self;
+		Settings->ActorSelector.ActorSelection = EPCGActorSelection::ByClass;
+		Settings->ActorSelector.ActorFilter = EPCGActorFilter::Self;
 
 		PCGTestsCommon::FTestData TestData(Seed, Settings, APCGUnitTestDummyActor::StaticClass());
 		Cast<APCGUnitTestDummyActor>(TestData.TestActor)->IntProperty = 42;
@@ -227,8 +227,8 @@ bool FPCGPropertyToParamDataActorFindTest::RunTest(const FString& Parameters)
 
 	// Self by tag
 	{
-		Settings->ActorSelection = EPCGActorSelection::ByTag;
-		Settings->ActorFilter = EPCGActorFilter::Self;
+		Settings->ActorSelector.ActorSelection = EPCGActorSelection::ByTag;
+		Settings->ActorSelector.ActorFilter = EPCGActorFilter::Self;
 
 		PCGTestsCommon::FTestData TestData(Seed, Settings, APCGUnitTestDummyActor::StaticClass());
 		TestData.TestActor->Tags.Add(Tag);
@@ -239,8 +239,8 @@ bool FPCGPropertyToParamDataActorFindTest::RunTest(const FString& Parameters)
 
 	// Self by name
 	{
-		Settings->ActorSelection = EPCGActorSelection::ByName;
-		Settings->ActorFilter = EPCGActorFilter::Self;
+		Settings->ActorSelector.ActorSelection = EPCGActorSelection::ByName;
+		Settings->ActorSelector.ActorFilter = EPCGActorFilter::Self;
 
 		PCGTestsCommon::FTestData TestData(Seed, Settings, APCGUnitTestDummyActor::StaticClass());
 		Cast<APCGUnitTestDummyActor>(TestData.TestActor)->IntProperty = 42;
@@ -294,8 +294,8 @@ bool FPCGPropertyToParamDataActorFindTest::RunTest(const FString& Parameters)
 
 	// Self by class including component
 	{
-		Settings->ActorSelection = EPCGActorSelection::ByTag;
-		Settings->ActorFilter = EPCGActorFilter::Self;
+		Settings->ActorSelector.ActorSelection = EPCGActorSelection::ByTag;
+		Settings->ActorSelector.ActorFilter = EPCGActorFilter::Self;
 		Settings->bSelectComponent = true;
 		
 		PCGTestsCommon::FTestData TestData(Seed, Settings, AActor::StaticClass());
@@ -309,8 +309,8 @@ bool FPCGPropertyToParamDataActorFindTest::RunTest(const FString& Parameters)
 
 	// World by class
 	{
-		Settings->ActorSelection = EPCGActorSelection::ByClass;
-		Settings->ActorFilter = EPCGActorFilter::AllWorldActors;
+		Settings->ActorSelector.ActorSelection = EPCGActorSelection::ByClass;
+		Settings->ActorSelector.ActorFilter = EPCGActorFilter::AllWorldActors;
 
 		PCGTestsCommon::FTestData TestData(Seed, Settings, APCGUnitTestDummyActor::StaticClass());
 		Cast<APCGUnitTestDummyActor>(TestData.TestActor)->IntProperty = 42;
@@ -320,8 +320,8 @@ bool FPCGPropertyToParamDataActorFindTest::RunTest(const FString& Parameters)
 
 	// World by tag
 	{
-		Settings->ActorSelection = EPCGActorSelection::ByTag;
-		Settings->ActorFilter = EPCGActorFilter::AllWorldActors;
+		Settings->ActorSelector.ActorSelection = EPCGActorSelection::ByTag;
+		Settings->ActorSelector.ActorFilter = EPCGActorFilter::AllWorldActors;
 
 		PCGTestsCommon::FTestData TestData(Seed, Settings, APCGUnitTestDummyActor::StaticClass());
 		TestData.TestActor->Tags.Add(Tag);
@@ -332,8 +332,8 @@ bool FPCGPropertyToParamDataActorFindTest::RunTest(const FString& Parameters)
 
 	// World by name
 	{
-		Settings->ActorSelection = EPCGActorSelection::ByName;
-		Settings->ActorFilter = EPCGActorFilter::AllWorldActors;
+		Settings->ActorSelector.ActorSelection = EPCGActorSelection::ByName;
+		Settings->ActorSelector.ActorFilter = EPCGActorFilter::AllWorldActors;
 
 		PCGTestsCommon::FTestData TestData(Seed, Settings, APCGUnitTestDummyActor::StaticClass());
 		Cast<APCGUnitTestDummyActor>(TestData.TestActor)->IntProperty = 42;
diff --git a/Engine/Plugins/Experimental/PCG/Source/PCG/Public/Elements/PCGActorSelector.h b/Engine/Plugins/Experimental/PCG/Source/PCG/Public/Elements/PCGActorSelector.h
new file mode 100644
--- /dev/null
+++ b/Engine/Plugins/Experimental/PCG/Source/PCG/Public/Elements/PCGActorSelector.h
@@ -0,0 +1,58 @@
+// Copyright Epic Games, Inc. All Rights Reserved.
+
+#pragma once
+
+#include "Templates/SubclassOf.h"
+
+#include "PCGActorSelector.generated.h"
+
+class AActor;
+class UWorld;
+
+UENUM()
+enum class EPCGActorSelection : uint8
+{
+	ByTag,
+	ByName,
+	ByClass
+};
+
+UENUM()
+enum class EPCGActorFilter : uint8
+{
+	Self,
+	Parent,
+	Root,
+	AllWorldActors
+	// TODO
+	// TrackedActors
+};
+
+USTRUCT(BlueprintType)
+struct FPCGActorSelectorSettings
+{
+	GENERATED_BODY()
+
+	UPROPERTY(BlueprintReadWrite, EditAnywhere, Category = Settings)
+	EPCGActorSelection ActorSelection = EPCGActorSelection::ByTag;
+
+	UPROPERTY(BlueprintReadWrite, EditAnywhere, Category = Settings, meta = (EditCondition = "ActorSelection==EPCGActorSelection::ByTag", EditConditionHides))
+	FName ActorSelectionTag;
+
+	UPROPERTY(BlueprintReadWrite, EditAnywhere, Category = Settings, meta = (EditCondition = "ActorSelection==EPCGActorSelection::ByName", EditConditionHides))
+	FName ActorSelectionName;
+
+	UPROPERTY(BlueprintReadWrite, EditAnywhere, Category = Settings, meta = (EditCondition = "ActorSelection==EPCGActorSelection::ByClass", EditConditionHides))
+	TSubclassOf<AActor> ActorSelectionClass;
+
+	UPROPERTY(BlueprintReadWrite, EditAnywhere, Category = Settings)
+	EPCGActorFilter ActorFilter = EPCGActorFilter::Self;
+
+	UPROPERTY(BlueprintReadWrite, EditAnywhere, Category = Settings, meta = (EditCondition = "ActorFilter!=EPCGActorFilter::AllWorldActors", EditConditionHides))
+	bool bIncludeChildren = false;
+};
+
+namespace PCGActorSelector
+{
+	AActor* FindActor(const FPCGActorSelectorSettings& Settings, UWorld* World, AActor* Self);
+}
\ No newline at end of file
diff --git a/Engine/Plugins/Experimental/PCG/Source/PCG/Public/Elements/PCGDataFromActor.h b/Engine/Plugins/Experimental/PCG/Source/PCG/Public/Elements/PCGDataFromActor.h
new file mode 100644
--- /dev/null
+++ b/Engine/Plugins/Experimental/PCG/Source/PCG/Public/Elements/PCGDataFromActor.h
@@ -0,0 +1,41 @@
+// Copyright Epic Games, Inc. All Rights Reserved.
+
+#pragma once
+
+#include "PCGElement.h"
+#include "PCGSettings.h"
+#include "Elements/PCGActorSelector.h"
+
+#include "PCGDataFromActor.generated.h"
+
+UCLASS(BlueprintType, ClassGroup = (Procedural))
+class PCG_API UPCGDataFromActorSettings : public UPCGSettings
+{
+	GENERATED_BODY()
+
+public:
+	//~Begin UPCGSettings interface
+#if WITH_EDITOR
+	virtual FName GetDefaultNodeName() const override { return FName(TEXT("DataFromActorNode")); }
+	virtual EPCGSettingsType GetType() const override { return EPCGSettingsType::Spatial; }
+#endif
+
+	virtual TArray<FPCGPinProperties> InputPinProperties() const override { return TArray<FPCGPinProperties>(); }
+
+protected:
+	virtual FPCGElementPtr CreateElement() const override;
+	//~End UPCGSettings
+
+public:
+	UPROPERTY(BlueprintReadWrite, EditAnywhere, Category = Settings, meta = (ShowOnlyInnerProperties))
+	FPCGActorSelectorSettings ActorSelector;
+};
+
+class FPCGDataFromActorElement : public FSimplePCGElement
+{
+public:
+	virtual bool CanExecuteOnlyOnMainThread(FPCGContext* Context) const override { return true; }
+	virtual bool IsCacheable(const UPCGSettings* InSettings) const override { return false; }
+protected:
+	virtual bool ExecuteInternal(FPCGContext* Context) const;
+};
\ No newline at end of file
diff --git a/Engine/Plugins/Experimental/PCG/Source/PCG/Public/Elements/PCGPropertyToParamData.h b/Engine/Plugins/Experimental/PCG/Source/PCG/Public/Elements/PCGPropertyToParamData.h
--- a/Engine/Plugins/Experimental/PCG/Source/PCG/Public/Elements/PCGPropertyToParamData.h
+++ b/Engine/Plugins/Experimental/PCG/Source/PCG/Public/Elements/PCGPropertyToParamData.h
@@ -4,43 +4,26 @@
 
 #include "PCGElement.h"
 #include "PCGSettings.h"
-
-#include "Templates/SubclassOf.h"
+#include "Elements/PCGActorSelector.h"
 
 #include "PCGPropertyToParamData.generated.h"
 
-class AActor;
 class UActorComponent;
 
-UENUM()
-enum class EPCGActorSelection : uint8
-{
-	ByTag,
-	ByName,
-	ByClass
-};
-
-UENUM()
-enum class EPCGActorFilter : uint8
-{
-	Self,
-	Parent,
-	Root,
-	AllWorldActors
-	// TODO
-	// TrackedActors
-};
-
 UCLASS(BlueprintType, ClassGroup = (Procedural))
 class PCG_API UPCGPropertyToParamDataSettings : public UPCGSettings
 {
 	GENERATED_BODY()
 
 public:
+	//~Begin UObject interface
+	virtual void PostLoad() override;
+	//~End UObject interface
+
 	//~Begin UPCGSettings interface
 #if WITH_EDITOR
 	virtual FName GetDefaultNodeName() const override { return FName(TEXT("PropertyToParamDataNode")); }
-	virtual EPCGSettingsType GetType() const override { return EPCGSettingsType::Metadata; }
+	virtual EPCGSettingsType GetType() const override { return EPCGSettingsType::Param; }
 #endif
 
 	virtual TArray<FPCGPinProperties> OutputPinProperties() const override;
@@ -51,23 +34,8 @@ protected:
 	//~End UPCGSettings interface
 
 public:
-	UPROPERTY(BlueprintReadWrite, EditAnywhere, Category = Settings)
-	EPCGActorSelection ActorSelection;
-
-	UPROPERTY(BlueprintReadWrite, EditAnywhere, Category = Settings, meta = (EditCondition = "ActorSelection==EPCGActorSelection::ByTag", EditConditionHides))
-	FName ActorSelectionTag;
-
-	UPROPERTY(BlueprintReadWrite, EditAnywhere, Category = Settings, meta = (EditCondition = "ActorSelection==EPCGActorSelection::ByName", EditConditionHides))
-	FName ActorSelectionName;
-
-	UPROPERTY(BlueprintReadWrite, EditAnywhere, Category = Settings, meta = (EditCondition = "ActorSelection==EPCGActorSelection::ByClass", EditConditionHides))
-	TSubclassOf<AActor> ActorSelectionClass;
-
-	UPROPERTY(BlueprintReadWrite, EditAnywhere, Category = Settings)
-	EPCGActorFilter ActorFilter = EPCGActorFilter::Self;
-
-	UPROPERTY(BlueprintReadWrite, EditAnywhere, Category = Settings, meta = (EditCondition = "ActorFilter!=EPCGActorFilter::AllWorldActors", EditConditionHides))
-	bool bIncludeChildren = false;
+	UPROPERTY(BlueprintReadWrite, EditAnywhere, Category = Settings, meta = (ShowOnlyInnerProperties))
+	FPCGActorSelectorSettings ActorSelector;
 
 	UPROPERTY(BlueprintReadWrite, EditAnywhere, Category = Settings)
 	bool bSelectComponent = false;
@@ -83,6 +51,25 @@ public:
 
 	UPROPERTY(BlueprintReadWrite, EditAnywhere, Category = Settings)
 	bool bAlwaysRequeryActors = true;
+
+private:
+	UPROPERTY()
+	EPCGActorSelection ActorSelection_DEPRECATED;
+
+	UPROPERTY()
+	FName ActorSelectionTag_DEPRECATED;
+
+	UPROPERTY()
+	FName ActorSelectionName_DEPRECATED;
+
+	UPROPERTY()
+	TSubclassOf<AActor> ActorSelectionClass_DEPRECATED;
+
+	UPROPERTY()
+	EPCGActorFilter ActorFilter_DEPRECATED = EPCGActorFilter::Self;
+
+	UPROPERTY()
+	bool bIncludeChildren_DEPRECATED = false;
 };
 
 class FPCGPropertyToParamDataElement : public FSimplePCGElement
diff --git a/Engine/Plugins/Experimental/PCG/Source/PCG/Public/Helpers/PCGBlueprintHelpers.h b/Engine/Plugins/Experimental/PCG/Source/PCG/Public/Helpers/PCGBlueprintHelpers.h
--- a/Engine/Plugins/Experimental/PCG/Source/PCG/Public/Helpers/PCGBlueprintHelpers.h
+++ b/Engine/Plugins/Experimental/PCG/Source/PCG/Public/Helpers/PCGBlueprintHelpers.h
@@ -66,6 +66,9 @@ public:
 	UFUNCTION(BlueprintCallable, Category = "PCG|Helpers", meta = (ScriptMethod))
 	static FBox GetActorLocalBoundsPCG(AActor* InActor, bool bIgnorePCGCreatedComponents = true);
 
+	UFUNCTION(BlueprintCallable, Category = "PCG|Helpers", meta = (ScriptMethod))
+	static UPCGData* CreatePCGDataFromActor(AActor* InActor, bool bParseActor = true);
+
 	UFUNCTION(BlueprintCallable, Category = "PCG|Helpers", meta = (ScriptMethod))
 	static TArray<FPCGLandscapeLayerWeight> GetInterpolatedPCGLandscapeLayerWeights(UObject* WorldContextObject, const FVector& Location);
 };
\ No newline at end of file
diff --git a/Engine/Plugins/Experimental/PCG/Source/PCG/Public/PCGComponent.h b/Engine/Plugins/Experimental/PCG/Source/PCG/Public/PCGComponent.h
--- a/Engine/Plugins/Experimental/PCG/Source/PCG/Public/PCGComponent.h
+++ b/Engine/Plugins/Experimental/PCG/Source/PCG/Public/PCGComponent.h
@@ -225,6 +225,9 @@ public:
 
 	FBox GetGridBounds() const;
 
+	/** Builds the canonical PCG data from a given actor and its PCG component if any. */
+	static UPCGData* CreateActorPCGData(AActor* Actor, const UPCGComponent* Component, bool bParseActor = true);
+
 protected:
 	UPROPERTY(BlueprintReadWrite, EditAnywhere, Category = Properties)
 	TObjectPtr<UPCGGraph> Graph;
