Fixed non-uniform scaling issue but kept translation relative to prevent rendering issues

diff --git a/Engine/Plugins/Experimental/PCG/Source/PCG/Private/Helpers/PCGActorHelpers.cpp b/Engine/Plugins/Experimental/PCG/Source/PCG/Private/Helpers/PCGActorHelpers.cpp
--- a/Engine/Plugins/Experimental/PCG/Source/PCG/Private/Helpers/PCGActorHelpers.cpp
+++ b/Engine/Plugins/Experimental/PCG/Source/PCG/Private/Helpers/PCGActorHelpers.cpp
@@ -145,8 +145,9 @@ UInstancedStaticMeshComponent* UPCGActorHelpers::GetOrCreateISMC(AActor* InTarge
 
 	ISMC->InstanceStartCullDistance = InParams.CullStartDistance;
 	ISMC->InstanceEndCullDistance = InParams.CullEndDistance;
-	
-	ISMC->AttachToComponent(InTargetActor->GetRootComponent(), FAttachmentTransformRules::KeepRelativeTransform);
+
+	// Implementation note: we use the relative for position only to prevent issues at the rendering level, but world for scale & rotation otherwise we run into issues when we have non-uniform scales esp. in volumes
+	ISMC->AttachToComponent(InTargetActor->GetRootComponent(), FAttachmentTransformRules(EAttachmentRule::KeepRelative, EAttachmentRule::KeepWorld, EAttachmentRule::KeepWorld, false));
 	ISMC->ComponentTags.Add(InSourceComponent->GetFName());
 	ISMC->ComponentTags.Add(PCGHelpers::DefaultPCGTag);
 
