Tentative fix for NoPCG error in PCG.

#rnx

diff --git a/Engine/Plugins/Experimental/PCG/Source/PCG/Private/PCGHelpers.cpp b/Engine/Plugins/Experimental/PCG/Source/PCG/Private/PCGHelpers.cpp
--- a/Engine/Plugins/Experimental/PCG/Source/PCG/Private/PCGHelpers.cpp
+++ b/Engine/Plugins/Experimental/PCG/Source/PCG/Private/PCGHelpers.cpp
@@ -7,6 +7,7 @@
 #include "PCGWorldActor.h"
 #include "Grid/PCGPartitionActor.h"
 
+#include "Engine/World.h"
 #include "GameFramework/Actor.h"
 #include "Landscape.h"
 #include "UObject/UObjectIterator.h"
