[PCG] Projection node renovations - first batch of changes for UE-164445

* Break out source and target data inputs - create new input pin for the projection target. Migration added that moves the *last* input edge to the new projection target pin. This is not exactly the same as the previous behaviour (which daisy-chained projections over all inputs), but it will produce the same behaviour in the case of 2 input edges and its probably going to be fine for users.
* Select which intrinsic data is projected (position, rotation, scale, color..)
* Filter string field for metadata attributes, as described in UE-164445.

diff --git a/Engine/Plugins/Experimental/PCG/Source/PCG/Private/Data/PCGPointData.cpp b/Engine/Plugins/Experimental/PCG/Source/PCG/Private/Data/PCGPointData.cpp
--- a/Engine/Plugins/Experimental/PCG/Source/PCG/Private/Data/PCGPointData.cpp
+++ b/Engine/Plugins/Experimental/PCG/Source/PCG/Private/Data/PCGPointData.cpp
@@ -426,6 +426,7 @@ bool UPCGPointData::SamplePoint(const FTransform& InTransform, const FBox& InBou
 	if (OutMetadata)
 	{
 		//TRACE_CPUPROFILER_EVENT_SCOPE(UPCGPointData::SamplePoint::SetupMetadata);
+		// Initialise metadata entry for this temporary point
 		UPCGMetadataAccessorHelpers::InitializeMetadataWithParent(OutPoint, OutMetadata, *MaxContributor, Metadata);
 
 		if (ContributionsForMetadata.Num() > 1)
diff --git a/Engine/Plugins/Experimental/PCG/Source/PCG/Private/Data/PCGProjectionData.cpp b/Engine/Plugins/Experimental/PCG/Source/PCG/Private/Data/PCGProjectionData.cpp
--- a/Engine/Plugins/Experimental/PCG/Source/PCG/Private/Data/PCGProjectionData.cpp
+++ b/Engine/Plugins/Experimental/PCG/Source/PCG/Private/Data/PCGProjectionData.cpp
@@ -4,7 +4,7 @@
 #include "Helpers/PCGAsync.h"
 #include "Metadata/PCGMetadataAccessor.h"
 
-void UPCGProjectionData::Initialize(const UPCGSpatialData* InSource, const UPCGSpatialData* InTarget)
+void UPCGProjectionData::Initialize(const UPCGSpatialData* InSource, const UPCGSpatialData* InTarget, const FPCGProjectionParams& InParams)
 {
 	check(InSource && InTarget);
 	// TODO: improve support for higher-dimension projection.
@@ -14,6 +14,8 @@ void UPCGProjectionData::Initialize(const UPCGSpatialData* InSource, const UPCGS
 	Target = InTarget;
 	TargetActor = InSource->TargetActor;
 
+	Params = InParams;
+
 	CachedBounds = ProjectBounds(Source->GetBounds());
 	CachedStrictBounds = ProjectBounds(Source->GetStrictBounds());
 }
@@ -73,32 +75,13 @@ FBox UPCGProjectionData::ProjectBounds(const FBox& InBounds) const
 
 bool UPCGProjectionData::SamplePoint(const FTransform& InTransform, const FBox& InBounds, FPCGPoint& OutPoint, UPCGMetadata* OutMetadata) const
 {
-	//TRACE_CPUPROFILER_EVENT_SCOPE(UPCGProjectionData::SamplePoint);
-	FPCGPoint PointFromSource;
-	if (!Source->SamplePoint(InTransform, InBounds, PointFromSource, OutMetadata))
-	{
-		return false;
-	}
-
-	FPCGPoint PointFromTarget;
-	if (!Target->SamplePoint(PointFromSource.Transform, PointFromSource.GetLocalBounds(), PointFromTarget, OutMetadata))
-	{
-		return false;
-	}
-
-	// Merge points into a single point
-	OutPoint = PointFromSource;
-	OutPoint.Transform = PointFromTarget.Transform;
-	OutPoint.Density *= PointFromTarget.Density;
-	OutPoint.Color *= PointFromTarget.Color;
+	// Detecting if a point is in a projection is a non-trivial. Projection is not in general a bijection and we cannot simply unproject
+	// the point from the Target and check if it is in the Source. However we can approximate the image of the projection and check
+	// if the query point is in the image.
 	
-	if (OutMetadata && PointFromTarget.MetadataEntry != PCGInvalidEntryKey)
-	{
-		//METADATA TODO Review op
-		OutMetadata->MergePointAttributesSubset(PointFromSource, OutMetadata, Source->Metadata, PointFromTarget, OutMetadata, Target->Metadata, OutPoint, EPCGMetadataOp::Max);
-	}
-
-	return true;
+	// Passing nullptr for the context means the operation will execute single threaded which is not ideal. To mitigate this we
+	// prewarm the point cache when this projection data is constructed in the projection element.
+	return ToPointData(nullptr)->SamplePoint(InTransform, InBounds, OutPoint, OutMetadata);
 }
 
 bool UPCGProjectionData::HasNonTrivialTransform() const
@@ -115,17 +98,35 @@ const UPCGPointData* UPCGProjectionData::CreatePointData(FPCGContext* Context) c
 	const TArray<FPCGPoint>& SourcePoints = SourcePointData->GetPoints();
 
 	UPCGPointData* PointData = NewObject<UPCGPointData>();
+
+	// Copy metadata attributes from source point including values
 	PointData->InitializeFromData(this, SourcePointData->Metadata);
 	check(PointData->Metadata);
-	PointData->Metadata->AddAttributes(Target->Metadata);
 
+	// The projection operation will write into this temporary metadata
 	UPCGMetadata* TempTargetMetadata = nullptr;
 	if (Target->Metadata)
 	{
 		TempTargetMetadata = NewObject<UPCGMetadata>();
-		TempTargetMetadata->Initialize(Target->Metadata);
+
+		// We achieve filtering of metadata attributes by manipulating this temporary metadata, which works because
+		// the projection operation operates on the attributes in this metadata.
+
+		// Behavior modes:
+		// * An excluded attribute that exists on source data will be kept and unchanged
+		// * An excluded attribute that does not exist on source data will not be kept in result
+		// * Included attributes are the only attributes that can be changed during projection
+		// * Included attributes are the only attributes that will be added from target data
+
+		TSet<FName> AttributesFilter;
+		GetIncludeExcludeAttributeNames(AttributesFilter);
+
+		TempTargetMetadata->InitializeWithAttributeFilter(Target->Metadata, AttributesFilter, Params.AttributeMode);
 	}
 
+	// Add any attributes from filtered metadata to produce final list
+	PointData->Metadata->AddAttributes(TempTargetMetadata);
+
 	TArray<FPCGPoint>& Points = PointData->GetMutablePoints();
 
 	FPCGAsync::AsyncPointProcessing(Context, SourcePoints.Num(), Points, [this, SourcePointData, PointData, TempTargetMetadata, &SourcePoints](int32 Index, FPCGPoint& OutPoint)
@@ -144,15 +145,15 @@ const UPCGPointData* UPCGProjectionData::CreatePointData(FPCGContext* Context) c
 
 		// Merge points into a single point
 		OutPoint = SourcePoint;
+		// Assign metadata key and copy values
 		UPCGMetadataAccessorHelpers::InitializeMetadata(OutPoint, PointData->Metadata, SourcePoint);
-		OutPoint.Transform = PointFromTarget.Transform;
-		OutPoint.Density *= PointFromTarget.Density;
-		OutPoint.Color *= PointFromTarget.Color;
+
+		ApplyProjectionResult(PointFromTarget, OutPoint);
 
 		if (PointData->Metadata && TempTargetMetadata && PointFromTarget.MetadataEntry != PCGInvalidEntryKey)
 		{
-			//METADATA TODO review op
-			PointData->Metadata->MergePointAttributesSubset(SourcePoint, SourcePointData->Metadata, SourcePointData->Metadata, PointFromTarget, TempTargetMetadata, TempTargetMetadata, OutPoint, EPCGMetadataOp::Max);
+			// Merge metadata to produce final attribute values
+			PointData->Metadata->MergePointAttributesSubset(SourcePoint, SourcePointData->Metadata, SourcePointData->Metadata, PointFromTarget, TempTargetMetadata, TempTargetMetadata, OutPoint, Params.AttributeMergeOperation);
 		}
 
 		return true;
@@ -161,4 +162,45 @@ const UPCGPointData* UPCGProjectionData::CreatePointData(FPCGContext* Context) c
 	UE_LOG(LogPCG, Verbose, TEXT("Projection generated %d points from %d source points"), Points.Num(), SourcePoints.Num());
 
 	return PointData;
-}
\ No newline at end of file
+}
+
+void UPCGProjectionData::ApplyProjectionResult(const FPCGPoint& InTargetPoint, FPCGPoint& InOutProjected) const
+{
+	if (Params.bProjectPositions)
+	{
+		InOutProjected.Transform.SetLocation(InTargetPoint.Transform.GetLocation());
+	}
+
+	if (Params.bProjectRotations)
+	{
+		InOutProjected.Transform.SetRotation(InTargetPoint.Transform.GetRotation());
+	}
+
+	if (Params.bProjectScales)
+	{
+		InOutProjected.Transform.SetScale3D(InTargetPoint.Transform.GetScale3D());
+	}
+
+	if (Params.bProjectColors)
+	{
+		InOutProjected.Color *= InTargetPoint.Color;
+	}
+
+	InOutProjected.Density *= InTargetPoint.Density;
+}
+
+void UPCGProjectionData::GetIncludeExcludeAttributeNames(TSet<FName>& OutAttributeNames) const
+{
+	if (Params.AttributeList.IsEmpty())
+	{
+		return;
+	}
+
+	TArray<FString> AttributeNameStrings;
+	Params.AttributeList.ParseIntoArray(AttributeNameStrings, TEXT(","), true);
+
+	for (const FString& Attribute : AttributeNameStrings)
+	{
+		OutAttributeNames.Add(FName(*Attribute));
+	}
+}
diff --git a/Engine/Plugins/Experimental/PCG/Source/PCG/Private/Data/PCGSpatialData.cpp b/Engine/Plugins/Experimental/PCG/Source/PCG/Private/Data/PCGSpatialData.cpp
--- a/Engine/Plugins/Experimental/PCG/Source/PCG/Private/Data/PCGSpatialData.cpp
+++ b/Engine/Plugins/Experimental/PCG/Source/PCG/Private/Data/PCGSpatialData.cpp
@@ -90,10 +90,10 @@ UPCGIntersectionData* UPCGSpatialData::IntersectWith(const UPCGSpatialData* InOt
 	return IntersectionData;
 }
 
-UPCGProjectionData* UPCGSpatialData::ProjectOn(const UPCGSpatialData* InOther) const
+UPCGProjectionData* UPCGSpatialData::ProjectOn(const UPCGSpatialData* InOther, const FPCGProjectionParams& InParams) const
 {
 	UPCGProjectionData* ProjectionData = NewObject<UPCGProjectionData>();
-	ProjectionData->Initialize(this, InOther);
+	ProjectionData->Initialize(this, InOther, InParams);
 
 	return ProjectionData;
 }
diff --git a/Engine/Plugins/Experimental/PCG/Source/PCG/Private/Data/PCGSplineData.cpp b/Engine/Plugins/Experimental/PCG/Source/PCG/Private/Data/PCGSplineData.cpp
--- a/Engine/Plugins/Experimental/PCG/Source/PCG/Private/Data/PCGSplineData.cpp
+++ b/Engine/Plugins/Experimental/PCG/Source/PCG/Private/Data/PCGSplineData.cpp
@@ -104,17 +104,17 @@ bool UPCGSplineData::SamplePoint(const FTransform& InTransform, const FBox& InBo
 	}
 }
 
-UPCGProjectionData* UPCGSplineData::ProjectOn(const UPCGSpatialData* InOther) const
+UPCGProjectionData* UPCGSplineData::ProjectOn(const UPCGSpatialData* InOther, const FPCGProjectionParams& InParams) const
 {
 	if(InOther->GetDimension() == 2)
 	{
 		UPCGSplineProjectionData* SplineProjectionData = NewObject<UPCGSplineProjectionData>(const_cast<UPCGSplineData*>(this));
-		SplineProjectionData->Initialize(this, InOther);
+		SplineProjectionData->Initialize(this, InOther, InParams);
 		return SplineProjectionData;
 	}
 	else
 	{
-		return Super::ProjectOn(InOther);
+		return Super::ProjectOn(InOther, InParams);
 	}
 }
 
@@ -152,9 +152,9 @@ FVector2D UPCGSplineProjectionData::Project(const FVector& InVector) const
 	return Projection2D;
 }
 
-void UPCGSplineProjectionData::Initialize(const UPCGSplineData* InSourceSpline, const UPCGSpatialData* InTargetSurface)
+void UPCGSplineProjectionData::Initialize(const UPCGSplineData* InSourceSpline, const UPCGSpatialData* InTargetSurface, const FPCGProjectionParams& InParams)
 {
-	Super::Initialize(InSourceSpline, InTargetSurface);
+	Super::Initialize(InSourceSpline, InTargetSurface, InParams);
 
 	const USplineComponent* Spline = GetSpline()->Spline.Get();
 	const FVector& SurfaceNormal = GetSurface()->GetNormal();
@@ -183,60 +183,6 @@ void UPCGSplineProjectionData::Initialize(const UPCGSplineData* InSourceSpline,
 	}
 }
 
-bool UPCGSplineProjectionData::SamplePoint(const FTransform& InTransform, const FBox& InBounds, FPCGPoint& OutPoint, UPCGMetadata* OutMetadata) const
-{
-	// TODO: support metadata
-	// TODO: support bounds
-
-	// Find nearest point on projected spline
-	const FVector InPosition = InTransform.GetLocation();
-	const USplineComponent* Spline = GetSpline()->Spline.Get();
-	const FVector& SurfaceNormal = GetSurface()->GetNormal();
-
-	// Project to 2D space
-	const FTransform LocalTransform = InTransform * Spline->GetComponentTransform().Inverse();
-	FVector2D LocalPosition2D = Project(LocalTransform.GetLocation());
-	float Dummy;
-	// Find nearest key on 2D spline
-	float NearestInputKey = ProjectedPosition.InaccurateFindNearest(LocalPosition2D, Dummy);
-	// TODO: if we didn't want to hand off density computation to the spline and do it here instead, we could do it in 2D space.
-	// Find point on original spline using the previously found key
-	// Note: this is an approximation that might not hold true since we are changing the curve length
-	const FVector NearestPointOnSpline = Spline->GetLocationAtSplineInputKey(NearestInputKey, ESplineCoordinateSpace::World);
-	const FVector PointOnLine = FMath::ClosestPointOnInfiniteLine(InPosition, InPosition + SurfaceNormal, NearestPointOnSpline);
-
-	// TODO: this is super inefficient
-	FPCGPoint SplinePoint;
-	if (GetSpline()->SamplePoint(FTransform(PointOnLine), InBounds, SplinePoint, OutMetadata))
-	{
-		FPCGPoint SurfacePoint;
-		if (GetSurface()->SamplePoint(SplinePoint.Transform, InBounds, SurfacePoint, OutMetadata))
-		{
-			OutPoint = SplinePoint;
-			OutPoint.Transform = SurfacePoint.Transform;
-			OutPoint.Density *= SurfacePoint.Density;
-			OutPoint.Color *= SurfacePoint.Color;
-
-			if (OutMetadata)
-			{
-				if (SplinePoint.MetadataEntry != PCGInvalidEntryKey && SurfacePoint.MetadataEntry)
-				{
-					//TODO review op
-					OutMetadata->MergePointAttributesSubset(SplinePoint, OutMetadata, GetSpline()->Metadata, SurfacePoint, OutMetadata, GetSurface()->Metadata, OutPoint, EPCGMetadataOp::Max);
-				}
-				else if (SurfacePoint.MetadataEntry != PCGInvalidEntryKey)
-				{
-					OutPoint.MetadataEntry = SurfacePoint.MetadataEntry;
-				}
-			}
-
-			return true;
-		}
-	}
-
-	return false;
-}
-
 const UPCGSplineData* UPCGSplineProjectionData::GetSpline() const
 {
 	return Cast<const UPCGSplineData>(Source);
diff --git a/Engine/Plugins/Experimental/PCG/Source/PCG/Private/Elements/PCGProjectionElement.cpp b/Engine/Plugins/Experimental/PCG/Source/PCG/Private/Elements/PCGProjectionElement.cpp
--- a/Engine/Plugins/Experimental/PCG/Source/PCG/Private/Elements/PCGProjectionElement.cpp
+++ b/Engine/Plugins/Experimental/PCG/Source/PCG/Private/Elements/PCGProjectionElement.cpp
@@ -1,9 +1,21 @@
 // Copyright Epic Games, Inc. All Rights Reserved.
 
 #include "Elements/PCGProjectionElement.h"
-#include "Data/PCGProjectionData.h"
+
+#include "PCGCustomVersion.h"
+#include "PCGEdge.h"
 #include "Helpers/PCGSettingsHelpers.h"
 
+TArray<FPCGPinProperties> UPCGProjectionSettings::InputPinProperties() const
+{
+	TArray<FPCGPinProperties> PinProperties;
+	PinProperties.Emplace(PCGPinConstants::DefaultInputLabel, EPCGDataType::Spatial);
+	PinProperties.Emplace(PCGProjectionConstants::ProjectionTargetLabel, EPCGDataType::Spatial, /*bAllowMultipleConnections=*/false);
+	PinProperties.Emplace(PCGPinConstants::DefaultParamsLabel, EPCGDataType::Param, /*bInAllowMultipleConnections=*/false);
+
+	return PinProperties;
+}
+
 FPCGElementPtr UPCGProjectionSettings::CreateElement() const
 {
 	return MakeShared<FPCGProjectionElement>();
@@ -12,13 +24,35 @@ FPCGElementPtr UPCGProjectionSettings::CreateElement() const
 bool FPCGProjectionElement::ExecuteInternal(FPCGContext* Context) const
 {
 	TRACE_CPUPROFILER_EVENT_SCOPE(FPCGProjectionElement::Execute);
+	check(Context);
 
 	const UPCGProjectionSettings* Settings = Context->GetInputSettings<UPCGProjectionSettings>();
 	check(Settings);
 
-	TArray<FPCGTaggedData> Inputs = Context->InputData.GetInputs();
+	TArray<FPCGTaggedData> Sources = Context->InputData.GetInputsByPin(PCGPinConstants::DefaultInputLabel);
+	TArray<FPCGTaggedData> Targets = Context->InputData.GetInputsByPin(PCGProjectionConstants::ProjectionTargetLabel);
+
+	// If there are no sources or no targets, then nothing to do.
+	if (Sources.Num() == 0 || Targets.Num() != 1)
+	{
+		return true;
+	}
+
+	// Ensure we have spatial data to project onto
+	UPCGSpatialData* ProjectionTarget = Cast<UPCGSpatialData>(Targets[0].Data);
+	if (!ProjectionTarget)
+	{
+		return true;
+	}
+
 	UPCGParamData* Params = Context->InputData.GetParams();
 
+	FPCGProjectionParams ProjectionParams = Settings->Params;
+	ProjectionParams.bProjectPositions = PCGSettingsHelpers::GetValue(GET_MEMBER_NAME_CHECKED(FPCGProjectionParams, bProjectPositions), ProjectionParams.bProjectPositions, Params);
+	ProjectionParams.bProjectRotations = PCGSettingsHelpers::GetValue(GET_MEMBER_NAME_CHECKED(FPCGProjectionParams, bProjectRotations), ProjectionParams.bProjectRotations, Params);
+	ProjectionParams.bProjectScales = PCGSettingsHelpers::GetValue(GET_MEMBER_NAME_CHECKED(FPCGProjectionParams, bProjectScales), ProjectionParams.bProjectScales, Params);
+	ProjectionParams.bProjectColors = PCGSettingsHelpers::GetValue(GET_MEMBER_NAME_CHECKED(FPCGProjectionParams, bProjectColors), ProjectionParams.bProjectColors, Params);
+
 	TArray<FPCGTaggedData>& Outputs = Context->OutputData.TaggedData;
 
 #if WITH_EDITORONLY_DATA
@@ -27,45 +61,70 @@ bool FPCGProjectionElement::ExecuteInternal(FPCGContext* Context) const
 	const bool bKeepZeroDensityPoints = false;
 #endif
 
-	const UPCGSpatialData* FirstSpatialData = nullptr;
-	UPCGProjectionData* ProjectionData = nullptr;
-	int32 ProjectionTaggedDataIndex = -1;
-
-	// TODO: it might not make sense to perform the projection if the first
-	// data isn't a spatial data, otherwise, what would it really mean?
-	for (FPCGTaggedData& Input : Inputs)
+	for (FPCGTaggedData& Source : Sources)
 	{
-		const UPCGSpatialData* SpatialData = Cast<UPCGSpatialData>(Input.Data);
+		UPCGSpatialData* ProjectionSource = Cast<UPCGSpatialData>(Source.Data);
 
-		if (!SpatialData)
+		if (!ProjectionSource)
 		{
-			Outputs.Add(Input);
+			PCGE_LOG(Error, "Invalid projection source data input found (non-spatial data). Input will be ignored.");
 			continue;
 		}
 
-		if (!FirstSpatialData)
-		{
-			FirstSpatialData = SpatialData;
-			ProjectionTaggedDataIndex = Outputs.Num();
-			Outputs.Add(Input);
-			continue;
-		}
-
-		// Create a new projection
-		ProjectionData = (ProjectionData ? ProjectionData : FirstSpatialData)->ProjectOn(SpatialData);
-
+		UPCGProjectionData* ProjectionData = ProjectionSource->ProjectOn(ProjectionTarget, ProjectionParams);
 #if WITH_EDITORONLY_DATA
 		ProjectionData->bKeepZeroDensityPoints = bKeepZeroDensityPoints;
 #endif
-		
-		// Update the tagged data
-		FPCGTaggedData& ProjectionTaggedData = Outputs[ProjectionTaggedDataIndex];
+
+		if (ProjectionData->RequiresCollapseToSample())
+		{
+			// Calling ToPointData will populate the point cache. Doing so here means we can pass in the Context object, which
+			// means the operation will be multi-threaded. This primes the cache in the most efficient way.
+			ProjectionData->ToPointData(Context);
+		}
+
+		FPCGTaggedData& ProjectionTaggedData = Outputs.Emplace_GetRef(Source);
 		ProjectionTaggedData.Data = ProjectionData;
-		ProjectionTaggedData.Tags.Append(Input.Tags);
+		ProjectionTaggedData.Tags.Append(Targets[0].Tags);
 	}
 
 	// Pass-through exclusions/settings
 	Outputs.Append(Context->InputData.GetAllSettings());
 
 	return true;
-}
\ No newline at end of file
+}
+
+#if WITH_EDITOR
+void UPCGProjectionSettings::ApplyDeprecation(UPCGNode* InOutNode)
+{
+	Super::ApplyDeprecation(InOutNode);
+	check(InOutNode);
+
+	if (DataVersion < FPCGCustomVersion::SplitProjectionNodeInputs)
+	{
+		// Split first pin inputs across two pins. The last edge connected to the first pin becomes the projection target.
+		check(InOutNode->GetInputPins().Num() == 2);
+
+		UPCGPin* SourcePin = InOutNode->GetInputPins()[0];
+		check(SourcePin);
+
+		if (SourcePin->EdgeCount() > 1)
+		{
+			UPCGPin* TargetPin = InOutNode->GetInputPins()[1];
+			check(TargetPin);
+
+			UPCGEdge* ProjectionTargetEdge = SourcePin->Edges.Last();
+			check(ProjectionTargetEdge);
+
+			UPCGPin* UpstreamPin = ProjectionTargetEdge->InputPin;
+			check(UpstreamPin);
+
+			UpstreamPin->BreakEdgeTo(SourcePin);
+			UpstreamPin->AddEdgeTo(TargetPin);
+		}
+	}
+
+	// Signal that data is up to date at end of deprecation code.
+	DataVersion = FPCGCustomVersion::LatestVersion;
+}
+#endif // WITH_EDITOR
diff --git a/Engine/Plugins/Experimental/PCG/Source/PCG/Private/Metadata/PCGMetadata.cpp b/Engine/Plugins/Experimental/PCG/Source/PCG/Private/Metadata/PCGMetadata.cpp
--- a/Engine/Plugins/Experimental/PCG/Source/PCG/Private/Metadata/PCGMetadata.cpp
+++ b/Engine/Plugins/Experimental/PCG/Source/PCG/Private/Metadata/PCGMetadata.cpp
@@ -68,12 +68,15 @@ void UPCGMetadata::BeginDestroy()
 	Super::BeginDestroy();
 }
 
-void UPCGMetadata::Initialize(const UPCGMetadata* InParent)
+void UPCGMetadata::Initialize(const UPCGMetadata* InParent, bool bAddAttributesFromParent)
 {
-	Initialize(InParent, /*bAddAttributesFromParent=*/true);
+	// If we are adding attributes from parent, then we use exclude filter with empty list so
+	// that all parameters added. Otherwise use include filter with empty list so none are added.
+	const EPCGMetadataFilterMode bFilter = bAddAttributesFromParent ? EPCGMetadataFilterMode::ExcludeAttributes : EPCGMetadataFilterMode::IncludeAttributes;
+	InitializeWithAttributeFilter(InParent, TSet<FName>(), bFilter);
 }
 
-void UPCGMetadata::Initialize(const UPCGMetadata* InParent, bool bAddAttributesFromParent)
+void UPCGMetadata::InitializeWithAttributeFilter(const UPCGMetadata* InParent, const TSet<FName>& InFilteredAttributes, EPCGMetadataFilterMode InFilterMode)
 {
 	if (Parent.IsValid() || Attributes.Num() != 0)
 	{
@@ -84,9 +87,11 @@ void UPCGMetadata::Initialize(const UPCGMetadata* InParent, bool bAddAttributesF
 	Parent = ((InParent != this) ? InParent : nullptr);
 	ItemKeyOffset = Parent.IsValid() ? Parent->GetItemCountForChild() : 0;
 
-	if (bAddAttributesFromParent)
+	// If we have been given an include list which is empty, then don't bother adding any attributes
+	const bool bSkipAddingAttributesFromParent = (InFilterMode == EPCGMetadataFilterMode::IncludeAttributes) && (InFilteredAttributes.Num() == 0);
+	if (!bSkipAddingAttributesFromParent)
 	{
-		AddAttributes(InParent);
+		AddAttributesFiltered(InParent, InFilteredAttributes, InFilterMode);
 	}
 }
 
@@ -116,7 +121,7 @@ void UPCGMetadata::InitializeAsCopy(const UPCGMetadata* InMetadataToCopy)
 	}
 }
 
-void UPCGMetadata::AddAttributes(const UPCGMetadata* InOther)
+void UPCGMetadata::AddAttributesFiltered(const UPCGMetadata* InOther, const TSet<FName>& InFilteredAttributes, EPCGMetadataFilterMode InFilterMode)
 {
 	if (!InOther)
 	{
@@ -125,7 +130,12 @@ void UPCGMetadata::AddAttributes(const UPCGMetadata* InOther)
 
 	for (const TPair<FName, FPCGMetadataAttributeBase*> OtherAttribute : InOther->Attributes)
 	{
-		if (HasAttribute(OtherAttribute.Key))
+		// Skip this attribute if it is in an exclude list, or if it is not in an include list
+		const bool bAttributeInFilterList = InFilteredAttributes.Contains(OtherAttribute.Key);
+		const bool bSkipAttributesInFilterList = InFilterMode == EPCGMetadataFilterMode::ExcludeAttributes;
+		const bool bSkipThisAttribute = bSkipAttributesInFilterList == bAttributeInFilterList;
+
+		if (bSkipThisAttribute || HasAttribute(OtherAttribute.Key))
 		{
 			continue;
 		}
@@ -135,7 +145,7 @@ void UPCGMetadata::AddAttributes(const UPCGMetadata* InOther)
 		}
 	}
 
-	if(InOther != Parent)
+	if (InOther != Parent)
 	{
 		OtherParents.Add(InOther);
 	}
diff --git a/Engine/Plugins/Experimental/PCG/Source/PCG/Public/Data/PCGProjectionData.h b/Engine/Plugins/Experimental/PCG/Source/PCG/Public/Data/PCGProjectionData.h
--- a/Engine/Plugins/Experimental/PCG/Source/PCG/Public/Data/PCGProjectionData.h
+++ b/Engine/Plugins/Experimental/PCG/Source/PCG/Public/Data/PCGProjectionData.h
@@ -2,8 +2,9 @@
 
 #pragma once
 
-#include "PCGSpatialData.h"
 #include "PCGPointData.h"
+#include "PCGSpatialData.h"
+#include "Elements/PCGProjectionParams.h"
 
 #include "PCGProjectionData.generated.h"
 
@@ -15,7 +16,7 @@ class PCG_API UPCGProjectionData : public UPCGSpatialDataWithPointCache
 {
 	GENERATED_BODY()
 public:
-	void Initialize(const UPCGSpatialData* InSource, const UPCGSpatialData* InTarget);
+	void Initialize(const UPCGSpatialData* InSource, const UPCGSpatialData* InTarget, const FPCGProjectionParams& InParams);
 
 	//~Begin UPCGSpatialData interface
 	virtual int GetDimension() const override;
@@ -24,6 +25,7 @@ public:
 	virtual FVector GetNormal() const override;
 	virtual bool SamplePoint(const FTransform& Transform, const FBox& Bounds, FPCGPoint& OutPoint, UPCGMetadata* OutMetadata) const override;
 	virtual bool HasNonTrivialTransform() const override;
+	virtual bool RequiresCollapseToSample() const override { return true; }
 	//~End UPCGSpatialData interface
 
 	//~Begin UPCGSpatialDataWithPointCache interface
@@ -33,6 +35,11 @@ public:
 protected:
 	FBox ProjectBounds(const FBox& InBounds) const;
 
+	/** Applies data from target point to projected point, conditionally according to the projection params. */
+	void ApplyProjectionResult(const FPCGPoint& InTargetPoint, FPCGPoint& InOutProjected) const;
+
+	void GetIncludeExcludeAttributeNames(TSet<FName>& OutAttributeNames) const;
+
 	UPROPERTY(BlueprintReadOnly, VisibleAnywhere, Category = SpatialData)
 	TObjectPtr<const UPCGSpatialData> Source = nullptr;
 
@@ -44,4 +51,7 @@ protected:
 
 	UPROPERTY()
 	FBox CachedStrictBounds = FBox(EForceInit::ForceInit);
+
+	UPROPERTY(BlueprintReadwrite, VisibleAnywhere, Category = SpatialData)
+	FPCGProjectionParams Params;
 };
diff --git a/Engine/Plugins/Experimental/PCG/Source/PCG/Public/Data/PCGSpatialData.h b/Engine/Plugins/Experimental/PCG/Source/PCG/Public/Data/PCGSpatialData.h
--- a/Engine/Plugins/Experimental/PCG/Source/PCG/Public/Data/PCGSpatialData.h
+++ b/Engine/Plugins/Experimental/PCG/Source/PCG/Public/Data/PCGSpatialData.h
@@ -6,6 +6,7 @@
 #include "PCGContext.h"
 #include "PCGData.h"
 #include "PCGPoint.h"
+#include "Elements/PCGProjectionParams.h"
 #include "Metadata/PCGMetadata.h"
 
 #include "PCGSpatialData.generated.h"
@@ -84,7 +85,7 @@ public:
 
 	/** Returns a specialized data to project this on another data of equal or higher dimension */
 	UFUNCTION(BlueprintCallable, Category = SpatialData)
-	virtual UPCGProjectionData* ProjectOn(const UPCGSpatialData* InOther) const;
+	virtual UPCGProjectionData* ProjectOn(const UPCGSpatialData* InOther, const FPCGProjectionParams& InParams = FPCGProjectionParams()) const;
 
 	/** Returns a specialized data to union this with another data */
 	UFUNCTION(BlueprintCallable, Category = SpatialData)
@@ -105,6 +106,9 @@ public:
 	UFUNCTION(BlueprintCallable, Category = SpatialData)
 	void InitializeFromData(const UPCGSpatialData* InSource, const UPCGMetadata* InMetadataParentOverride = nullptr, bool bInheritMetadata = true);
 
+	/** True if this operation does not have an inverse and cannot be queried analytically/implicitly, and therefore must be collapsed to an explicit point representation. */
+	virtual bool RequiresCollapseToSample() const { return false; }
+
 	UPROPERTY(Transient, BlueprintReadWrite, EditAnywhere, Category = Data)
 	TWeakObjectPtr<AActor> TargetActor = nullptr;
 
diff --git a/Engine/Plugins/Experimental/PCG/Source/PCG/Public/Data/PCGSplineData.h b/Engine/Plugins/Experimental/PCG/Source/PCG/Public/Data/PCGSplineData.h
--- a/Engine/Plugins/Experimental/PCG/Source/PCG/Public/Data/PCGSplineData.h
+++ b/Engine/Plugins/Experimental/PCG/Source/PCG/Public/Data/PCGSplineData.h
@@ -2,9 +2,12 @@
 
 #pragma once
 
-#include "CoreMinimal.h"
 #include "PCGPolyLineData.h"
 #include "PCGProjectionData.h"
+#include "Elements/PCGProjectionParams.h"
+
+#include "CoreMinimal.h"
+
 #include "PCGSplineData.generated.h"
 
 class USplineComponent;
@@ -36,7 +39,7 @@ public:
 	//~Begin UPCGSpatialData interface
 	virtual FBox GetBounds() const override;
 	virtual bool SamplePoint(const FTransform& Transform, const FBox& Bounds, FPCGPoint& OutPoint, UPCGMetadata* OutMetadata) const override;
-	virtual UPCGProjectionData* ProjectOn(const UPCGSpatialData* InOther) const override;
+	virtual UPCGProjectionData* ProjectOn(const UPCGSpatialData* InOther, const FPCGProjectionParams& InParams = FPCGProjectionParams()) const override;
 	//~End 
 
 	UPROPERTY(BlueprintReadWrite, EditAnywhere, Category = SourceData)
@@ -47,16 +50,13 @@ protected:
 	FBox CachedBounds = FBox(EForceInit::ForceInit);
 };
 
+/* The projection of a spline onto a surface. */
 UCLASS(BlueprintType, ClassGroup=(Procedural))
 class PCG_API UPCGSplineProjectionData : public UPCGProjectionData
 {
 	GENERATED_BODY()
 public:
-	void Initialize(const UPCGSplineData* InSourceSpline, const UPCGSpatialData* InTargetSurface);
-
-	//~Begin UPCGSpatialData interface
-	virtual bool SamplePoint(const FTransform& Transform, const FBox& Bounds, FPCGPoint& OutPoint, UPCGMetadata* OutMetadata) const override;
-	//~End UPCGSpatialData interface
+	void Initialize(const UPCGSplineData* InSourceSpline, const UPCGSpatialData* InTargetSurface, const FPCGProjectionParams& InParams);
 
 	const UPCGSplineData* GetSpline() const;
 	const UPCGSpatialData* GetSurface() const;
diff --git a/Engine/Plugins/Experimental/PCG/Source/PCG/Public/Elements/PCGProjectionElement.h b/Engine/Plugins/Experimental/PCG/Source/PCG/Public/Elements/PCGProjectionElement.h
--- a/Engine/Plugins/Experimental/PCG/Source/PCG/Public/Elements/PCGProjectionElement.h
+++ b/Engine/Plugins/Experimental/PCG/Source/PCG/Public/Elements/PCGProjectionElement.h
@@ -2,12 +2,19 @@
 
 #pragma once
 
-#include "CoreMinimal.h"
 #include "PCGElement.h"
 #include "PCGSettings.h"
+#include "Data/PCGProjectionData.h"
+
+#include "CoreMinimal.h"
 
 #include "PCGProjectionElement.generated.h"
 
+namespace PCGProjectionConstants
+{
+	const FName ProjectionTargetLabel = TEXT("Projection Target");
+}
+
 UCLASS(BlueprintType, ClassGroup = (Procedural))
 class PCG_API UPCGProjectionSettings : public UPCGSettings
 {
@@ -20,11 +27,20 @@ public:
 	virtual EPCGSettingsType GetType() const override { return EPCGSettingsType::Spatial; }
 #endif
 
+	virtual TArray<FPCGPinProperties> InputPinProperties() const override;
+
+	UPROPERTY(BlueprintReadWrite, EditAnywhere, Category = Settings, meta = (ShowOnlyInnerProperties))
+	FPCGProjectionParams Params;
+
 #if WITH_EDITORONLY_DATA
 	UPROPERTY(Transient, BlueprintReadWrite, EditAnywhere, Category = Debug)
 	bool bKeepZeroDensityPoints = false;
 #endif
 
+#if WITH_EDITOR
+	virtual void ApplyDeprecation(UPCGNode* InOutNode) override;
+#endif
+
 protected:
 	virtual FPCGElementPtr CreateElement() const override;
 	//~End UPCGSettings interface
diff --git a/Engine/Plugins/Experimental/PCG/Source/PCG/Public/Elements/PCGProjectionParams.h b/Engine/Plugins/Experimental/PCG/Source/PCG/Public/Elements/PCGProjectionParams.h
new file mode 100644
--- /dev/null
+++ b/Engine/Plugins/Experimental/PCG/Source/PCG/Public/Elements/PCGProjectionParams.h
@@ -0,0 +1,44 @@
+// Copyright Epic Games, Inc. All Rights Reserved.
+
+#pragma once
+
+#include "Metadata/PCGMetadataCommon.h"
+
+#include "CoreMinimal.h"
+
+#include "PCGProjectionParams.generated.h"
+
+/** Parameters that control projection behaviour. */
+USTRUCT(BlueprintType)
+struct PCG_API FPCGProjectionParams
+{
+	GENERATED_BODY()
+
+	/** Project positions. */
+	UPROPERTY(BlueprintReadWrite, EditAnywhere, Category = "Apply Data")
+	bool bProjectPositions = true;
+
+	/** Project rotations. */
+	UPROPERTY(BlueprintReadWrite, EditAnywhere, Category = "Apply Data")
+	bool bProjectRotations = true;
+
+	/** Project scales. */
+	UPROPERTY(BlueprintReadWrite, EditAnywhere, Category = "Apply Data")
+	bool bProjectScales = false;
+
+	/** Project colors. */
+	UPROPERTY(BlueprintReadWrite, EditAnywhere, Category = "Apply Data")
+	bool bProjectColors = false;
+
+	/** Attributes to either explicitly exclude or include in the projection operation, depending on the Attribute Mode setting. Leave empty to gather all attributes and their values. Format is comma separated list like: Attribute1,Attribute2 .*/
+	UPROPERTY(BlueprintReadWrite, EditAnywhere, Category = "Apply Data")
+	FString AttributeList = TEXT("");
+
+	/** How the attribute list is used. Exclude Attributes will ignore these attributes and their values on the projection target. */
+	UPROPERTY(BlueprintReadWrite, EditAnywhere, Category = "Apply Data")
+	EPCGMetadataFilterMode AttributeMode = EPCGMetadataFilterMode::ExcludeAttributes;
+
+	/** Operation to use to combine attributes that reside on both source and target data. */
+	UPROPERTY(BlueprintReadWrite, EditAnywhere, Category = "Apply Data")
+	EPCGMetadataOp AttributeMergeOperation = EPCGMetadataOp::TargetValue;
+};
diff --git a/Engine/Plugins/Experimental/PCG/Source/PCG/Public/Metadata/PCGMetadata.h b/Engine/Plugins/Experimental/PCG/Source/PCG/Public/Metadata/PCGMetadata.h
--- a/Engine/Plugins/Experimental/PCG/Source/PCG/Public/Metadata/PCGMetadata.h
+++ b/Engine/Plugins/Experimental/PCG/Source/PCG/Public/Metadata/PCGMetadata.h
@@ -24,20 +24,38 @@ public:
 	virtual void BeginDestroy() override;
 	//~ End UObject interface
 
-	/** Initializes the metadata from a parent metadata, if any (can be null) */
+	/** Initializes the metadata from a parent metadata, if any (can be null). Copies attributes and values. */
 	UFUNCTION(BlueprintCallable, Category = "PCG|Metadata")
-	void Initialize(const UPCGMetadata* InParent);
+	void Initialize(const UPCGMetadata* InParent) { InitializeWithAttributeFilter(InParent, TSet<FName>()); }
 
 	/** Initializes the metadata from a parent metadata, if any (can be null) with the option to not add attributes from the parent. */
 	void Initialize(const UPCGMetadata* InParent, bool bAddAttributesFromParent);
 
+	/**
+	 * Initializes the metadata from a parent metadata. Copies attributes and values.
+	 * @param InParent The parent metadata to use as a template, if any (can be null).
+	 * @param InFilteredAttributes Optional list of attributes to exclude or include when adding the attributes from the parent.
+	 * @param InFilterMode Defines attribute filter operation.
+	 */
+	UFUNCTION(BlueprintCallable, Category = "PCG|Metadata")
+	void InitializeWithAttributeFilter(const UPCGMetadata* InParent, const TSet<FName>& InFilteredAttributes, EPCGMetadataFilterMode InFilterMode = EPCGMetadataFilterMode::ExcludeAttributes);
+
 	/** Initializes the metadata from a parent metadata by copying all attributes to it. */
 	UFUNCTION(BlueprintCallable, Category = "PCG|Metadata")
 	void InitializeAsCopy(const UPCGMetadata* InMetadataToCopy);
 
 	/** Creates missing attributes from another metadata if they are not currently present - note that this does not copy values */
 	UFUNCTION(BlueprintCallable, Category = "PCG|Metadata")
-	void AddAttributes(const UPCGMetadata* InOther);
+	void AddAttributes(const UPCGMetadata* InOther) { AddAttributesFiltered(InOther, TSet<FName>()); }
+
+	/**
+	 * Creates missing attributes from another metadata if they are not currently present - note that this does not copy values.
+	 * @param InOther The other metadata to obtain a list of attributes from.
+	 * @param InFilteredAttributes Optional list of attributes to exclude or include when adding the attributes.
+	 * @param InFilterMode Defines attribute filter operation.
+	 */
+	UFUNCTION(BlueprintCallable, Category = "PCG|Metadata")
+	void AddAttributesFiltered(const UPCGMetadata* InOther, const TSet<FName>& InFilteredAttributes, EPCGMetadataFilterMode InFilterMode = EPCGMetadataFilterMode::ExcludeAttributes);
 
 	/** Creates missing attribute from another metadata if it is not currently present - note that this does not copy values */
 	UFUNCTION(BlueprintCallable, Category = "PCG|Metadata")
diff --git a/Engine/Plugins/Experimental/PCG/Source/PCG/Public/Metadata/PCGMetadataAttributeTpl.h b/Engine/Plugins/Experimental/PCG/Source/PCG/Public/Metadata/PCGMetadataAttributeTpl.h
--- a/Engine/Plugins/Experimental/PCG/Source/PCG/Public/Metadata/PCGMetadataAttributeTpl.h
+++ b/Engine/Plugins/Experimental/PCG/Source/PCG/Public/Metadata/PCGMetadataAttributeTpl.h
@@ -108,8 +108,37 @@ public:
 		check(ItemKey != PCGInvalidEntryKey);
 		bool bAppliedValue = false;
 
-		if (InAttributeA && InAttributeB && bAllowsInterpolation)
+		if (Op == EPCGMetadataOp::TargetValue && InAttributeB)
 		{
+			// Take value of second attribute.
+			if (InAttributeB == this)
+			{
+				SetValueFromValueKey(ItemKey, GetValueKey(InEntryKeyB));
+			}
+			else
+			{
+				SetValue(ItemKey, static_cast<const FPCGMetadataAttribute<T>*>(InAttributeB)->GetValueFromItemKey(InEntryKeyB));
+			}
+
+			bAppliedValue = true;
+		}
+		else if (Op == EPCGMetadataOp::SourceValue && InAttributeA)
+		{
+			// Take value of first attribute.
+			if (InAttributeA == this)
+			{
+				SetValueFromValueKey(ItemKey, GetValueKey(InEntryKeyA));
+			}
+			else
+			{
+				SetValue(ItemKey, static_cast<const FPCGMetadataAttribute<T>*>(InAttributeA)->GetValueFromItemKey(InEntryKeyA));
+			}
+
+			bAppliedValue = true;
+		}
+		else if (InAttributeA && InAttributeB && bAllowsInterpolation)
+		{
+			// Combine attributes using specified operation.
 			if (Op == EPCGMetadataOp::Min)
 			{
 				bAppliedValue = SetMin(ItemKey, InAttributeA, InEntryKeyA, InAttributeB, InEntryKeyB);
diff --git a/Engine/Plugins/Experimental/PCG/Source/PCG/Public/Metadata/PCGMetadataCommon.h b/Engine/Plugins/Experimental/PCG/Source/PCG/Public/Metadata/PCGMetadataCommon.h
--- a/Engine/Plugins/Experimental/PCG/Source/PCG/Public/Metadata/PCGMetadataCommon.h
+++ b/Engine/Plugins/Experimental/PCG/Source/PCG/Public/Metadata/PCGMetadataCommon.h
@@ -16,10 +16,29 @@ const PCGMetadataValueKey PCGDefaultValueKey = -1;
 UENUM()
 enum class EPCGMetadataOp : uint8
 {
+	/** Take the minimum value. */
 	Min,
+	/** Take the maximum value. */
 	Max,
+	/** Subtract the values. */
 	Sub,
+	/** Add the values. */
 	Add,
+	/** Multiply the values. */
 	Mul,
-	Div
-};
\ No newline at end of file
+	/** Divide the values. */
+	Div,
+	/** Pick the source (first) value. */
+	SourceValue,
+	/** Pick the target (second) value. */
+	TargetValue
+};
+
+UENUM()
+enum class EPCGMetadataFilterMode : uint8
+{
+	/** The listed attributes will be unchanged by the projection and will not be added from the target data. */
+	ExcludeAttributes,
+	/** Only the listed attributes will be changed by the projection or added from the target data. */
+	IncludeAttributes,
+};
diff --git a/Engine/Plugins/Experimental/PCG/Source/PCG/Public/PCGCustomVersion.h b/Engine/Plugins/Experimental/PCG/Source/PCG/Public/PCGCustomVersion.h
--- a/Engine/Plugins/Experimental/PCG/Source/PCG/Public/PCGCustomVersion.h
+++ b/Engine/Plugins/Experimental/PCG/Source/PCG/Public/PCGCustomVersion.h
@@ -13,6 +13,9 @@ struct PCG_API FPCGCustomVersion
 		// Before any version changes were made in the plugin
 		BeforeCustomVersionWasAdded = 0,
 
+		// Split projection nodes inputs to separate source edges and target edge
+		SplitProjectionNodeInputs = 1,
+		
 		// -----<new versions can be added above this line>-------------------------------------------------
 		VersionPlusOne,
 		LatestVersion = VersionPlusOne - 1
