Submitted on behalf of adrien.logut
[PCG] Fix metadata parenting in some element nodes + add attribute filtered in node name

diff --git a/Engine/Plugins/Experimental/PCG/Source/PCG/Private/Elements/PCGAttributeFilterElement.cpp b/Engine/Plugins/Experimental/PCG/Source/PCG/Private/Elements/PCGAttributeFilterElement.cpp
--- a/Engine/Plugins/Experimental/PCG/Source/PCG/Private/Elements/PCGAttributeFilterElement.cpp
+++ b/Engine/Plugins/Experimental/PCG/Source/PCG/Private/Elements/PCGAttributeFilterElement.cpp
@@ -8,13 +8,31 @@
 #include "Metadata/PCGMetadata.h"
 #include "Metadata/PCGMetadataAttribute.h"
 
+namespace PCGAttributeFilterConstants
+{
+	const FName NodeName = TEXT("FilterAttribute");
+}
+
 #if WITH_EDITOR
 FName UPCGAttributeFilterSettings::GetDefaultNodeName() const
 {
-	return TEXT("FilterAttribute");
+	return PCGAttributeFilterConstants::NodeName;
 }
 #endif
 
+FName UPCGAttributeFilterSettings::AdditionalTaskName() const
+{
+	// If we filter only one attribute, show its name
+	if (AttributesToKeep.Num() == 1)
+	{
+		return FName(FString::Printf(TEXT("%s: %s"), *PCGAttributeFilterConstants::NodeName.ToString(), *AttributesToKeep[0].ToString()));
+	}
+	else
+	{
+		return NAME_None;
+	}
+}
+
 TArray<FPCGPinProperties> UPCGAttributeFilterSettings::OutputPinProperties() const
 {
 	TArray<FPCGPinProperties> PinProperties;
@@ -53,6 +71,8 @@ bool FPCGAttributeFilterElement::ExecuteInternal(FPCGContext* Context) const
 
 			UPCGSpatialData* NewSpatialData = DuplicateObject<UPCGSpatialData>(const_cast<UPCGSpatialData*>(InputSpatialData), nullptr);
 			NewSpatialData->Metadata = Metadata;
+			// We also need set the Metadata outer to this new data.
+			Metadata->Rename(nullptr, NewSpatialData);
 			NewSpatialData->Metadata->Initialize(ParentMetadata, /*bAddAttributesFromParent=*/false);
 
 			// No need to inherit metadata since we already initialized it.
@@ -62,11 +82,10 @@ bool FPCGAttributeFilterElement::ExecuteInternal(FPCGContext* Context) const
 		}
 		else if (const UPCGParamData* InputParamData = Cast<UPCGParamData>(InputData))
 		{
-			Metadata = NewObject<UPCGMetadata>();
 			ParentMetadata = InputParamData->Metadata;
 
 			UPCGParamData* NewParamData = NewObject<UPCGParamData>();
-			NewParamData->Metadata = Metadata;
+			Metadata = NewParamData->Metadata;
 
 			Metadata->Initialize(InputParamData->Metadata, /*bAddAttributesFromParent=*/false);
 			OutputData = NewParamData;
diff --git a/Engine/Plugins/Experimental/PCG/Source/PCG/Private/Elements/PCGCreateAttribute.cpp b/Engine/Plugins/Experimental/PCG/Source/PCG/Private/Elements/PCGCreateAttribute.cpp
--- a/Engine/Plugins/Experimental/PCG/Source/PCG/Private/Elements/PCGCreateAttribute.cpp
+++ b/Engine/Plugins/Experimental/PCG/Source/PCG/Private/Elements/PCGCreateAttribute.cpp
@@ -206,8 +206,8 @@ bool FPCGCreateAttributeElement::ExecuteInternal(FPCGContext* Context) const
 
 		if (const UPCGSpatialData* InputSpatialData = Cast<UPCGSpatialData>(InputData))
 		{
-			UPCGSpatialData* NewSpatialData = DuplicateObject<UPCGSpatialData>(const_cast<UPCGSpatialData*>(InputSpatialData), nullptr);
-			NewSpatialData->Metadata = NewObject<UPCGMetadata>();
+			UPCGSpatialData* NewSpatialData = DuplicateObject<UPCGSpatialData>(InputSpatialData, nullptr);
+			NewSpatialData->Metadata = NewObject<UPCGMetadata>(NewSpatialData);
 			NewSpatialData->InitializeFromData(InputSpatialData, /*InMetadataParentOverride=*/ nullptr, /*bInheritMetadata=*/ Settings->bKeepExistingAttributes);
 
 			OutputData = NewSpatialData;
diff --git a/Engine/Plugins/Experimental/PCG/Source/PCG/Public/Elements/PCGAttributeFilterElement.h b/Engine/Plugins/Experimental/PCG/Source/PCG/Public/Elements/PCGAttributeFilterElement.h
--- a/Engine/Plugins/Experimental/PCG/Source/PCG/Public/Elements/PCGAttributeFilterElement.h
+++ b/Engine/Plugins/Experimental/PCG/Source/PCG/Public/Elements/PCGAttributeFilterElement.h
@@ -27,6 +27,7 @@ public:
 #endif
 
 	virtual TArray<FPCGPinProperties> OutputPinProperties() const override;
+	virtual FName AdditionalTaskName() const override;
 	//~End UPCGSettings interface
 
 	UPROPERTY(BlueprintReadWrite, EditAnywhere, Category = Settings)
