Fixed issue when we reconnect a single pin.
Minor API cleanup.

diff --git a/Engine/Plugins/Experimental/PCG/Source/PCG/Private/PCGGraph.cpp b/Engine/Plugins/Experimental/PCG/Source/PCG/Private/PCGGraph.cpp
--- a/Engine/Plugins/Experimental/PCG/Source/PCG/Private/PCGGraph.cpp
+++ b/Engine/Plugins/Experimental/PCG/Source/PCG/Private/PCGGraph.cpp
@@ -253,17 +253,18 @@ void UPCGGraph::OnNodeRemoved(UPCGNode* InNode)
 #endif
 }
 
-UPCGNode* UPCGGraph::AddEdge(UPCGNode* From, UPCGNode* To)
+UPCGNode* UPCGGraph::AddEdge(UPCGNode* From, const FName& InboundLabel, UPCGNode* To, const FName& OutboundLabel)
 {
-	return AddLabeledEdge(From, NAME_None, To, NAME_None);
+	AddLabeledEdge(From, InboundLabel, To, OutboundLabel);
+	return To;
 }
 
-UPCGNode* UPCGGraph::AddLabeledEdge(UPCGNode* From, const FName& InboundLabel, UPCGNode* To, const FName& OutboundLabel)
+bool UPCGGraph::AddLabeledEdge(UPCGNode* From, const FName& InboundLabel, UPCGNode* To, const FName& OutboundLabel)
 {
 	if (!From || !To)
 	{
 		UE_LOG(LogPCG, Error, TEXT("Invalid edge nodes"));
-		return To;
+		return false;
 	}
 
 	UPCGPin* FromPin = From->GetOutputPin(InboundLabel);
@@ -271,7 +272,7 @@ UPCGNode* UPCGGraph::AddLabeledEdge(UPCGNode* From, const FName& InboundLabel, U
 	if(!FromPin)
 	{
 		UE_LOG(LogPCG, Error, TEXT("From node %s does not have the %s label"), *From->GetName(), *InboundLabel.ToString());
-		return To;
+		return false;
 	}
 
 	UPCGPin* ToPin = To->GetInputPin(OutboundLabel);
@@ -279,17 +280,25 @@ UPCGNode* UPCGGraph::AddLabeledEdge(UPCGNode* From, const FName& InboundLabel, U
 	if(!ToPin)
 	{
 		UE_LOG(LogPCG, Error, TEXT("To node %s does not have the %s label"), *To->GetName(), *OutboundLabel.ToString());
-		return To;
+		return false;
 	}
 
 	// Create edge
 	FromPin->AddEdgeTo(ToPin);
+
+	bool bToPinBrokeOtherEdges = false;
+	
+	// Add an edge to a pin that doesn't allow multiple connections requires to do some cleanup
+	if (!ToPin->Properties.bAllowMultipleConnections)
+	{
+		bToPinBrokeOtherEdges = ToPin->BreakAllIncompatibleEdges();
+	}
 	
 #if WITH_EDITOR
 	NotifyGraphChanged(EPCGChangeType::Structural);
 #endif
 
-	return To;
+	return bToPinBrokeOtherEdges;
 }
 
 TObjectPtr<UPCGNode> UPCGGraph::ReconstructNewNode(const UPCGNode* InNode)
diff --git a/Engine/Plugins/Experimental/PCG/Source/PCG/Private/PCGNode.cpp b/Engine/Plugins/Experimental/PCG/Source/PCG/Private/PCGNode.cpp
--- a/Engine/Plugins/Experimental/PCG/Source/PCG/Private/PCGNode.cpp
+++ b/Engine/Plugins/Experimental/PCG/Source/PCG/Private/PCGNode.cpp
@@ -189,7 +189,7 @@ UPCGNode* UPCGNode::AddEdgeTo(FName InboundName, UPCGNode* To, FName OutboundNam
 	check(GetGraph());
 	if (UPCGGraph* Graph = GetGraph())
 	{
-		return Graph->AddLabeledEdge(this, InboundName, To, OutboundName);
+		return Graph->AddEdge(this, InboundName, To, OutboundName);
 	}
 	else
 	{
diff --git a/Engine/Plugins/Experimental/PCG/Source/PCG/Public/PCGGraph.h b/Engine/Plugins/Experimental/PCG/Source/PCG/Public/PCGGraph.h
--- a/Engine/Plugins/Experimental/PCG/Source/PCG/Public/PCGGraph.h
+++ b/Engine/Plugins/Experimental/PCG/Source/PCG/Public/PCGGraph.h
@@ -57,10 +57,7 @@ public:
 
 	/** Adds a directed edge in the graph. Returns the "To" node for easy chaining */
 	UFUNCTION(BlueprintCallable, Category = Graph)
-	UPCGNode* AddEdge(UPCGNode* From, UPCGNode* To);
-
-	UFUNCTION(BlueprintCallable, Category = Graph)
-	UPCGNode* AddLabeledEdge(UPCGNode* From, const FName& InboundLabel, UPCGNode* To, const FName& OutboundLabel);
+	UPCGNode* AddEdge(UPCGNode* From, const FName& InboundLabel, UPCGNode* To, const FName& OutboundLabel);
 
 	/** Returns the graph input node */
 	UFUNCTION(BlueprintCallable, Category = Graph)
@@ -73,6 +70,9 @@ public:
 	/** Duplicate a given node by creating a new node with the same settings and properties, but without any edges and add it to the graph */
 	TObjectPtr<UPCGNode> ReconstructNewNode(const UPCGNode* InNode);
 
+	/** Creates an edge between two nodes/pins based on the labels. Returns true if the To node has removed other edges (happens with single pins) */
+	bool AddLabeledEdge(UPCGNode* From, const FName& InboundLabel, UPCGNode* To, const FName& OutboundLabel);
+
 	bool Contains(UPCGNode* Node) const;
 	const TArray<UPCGNode*>& GetNodes() const { return Nodes; }
 	void AddNode(UPCGNode* InNode);
diff --git a/Engine/Plugins/Experimental/PCG/Source/PCGEditor/Private/PCGEditorGraphSchema.cpp b/Engine/Plugins/Experimental/PCG/Source/PCGEditor/Private/PCGEditorGraphSchema.cpp
--- a/Engine/Plugins/Experimental/PCG/Source/PCGEditor/Private/PCGEditorGraphSchema.cpp
+++ b/Engine/Plugins/Experimental/PCG/Source/PCGEditor/Private/PCGEditorGraphSchema.cpp
@@ -133,19 +133,12 @@ bool UPCGEditorGraphSchema::TryCreateConnection(UEdGraphPin* InA, UEdGraphPin* I
 		UPCGGraph* PCGGraph = PCGNodeA->GetGraph();
 		check(PCGGraph);
 
-		PCGGraph->AddLabeledEdge(PCGNodeA, A->PinName, PCGNodeB, B->PinName);
+		const bool bReconstructNodeB = PCGGraph->AddLabeledEdge(PCGNodeA, A->PinName, PCGNodeB, B->PinName);
 
-		// TODO: unclear if that kind of behavior should be down the code hierarchy or not,
-		// Since we really want to do cleanup only on manual interaction
-		if (UPCGPin* InputPin = PCGNodeB->GetInputPin(B->PinName))
+		// If AddLabeledEdge return true, then it means that the ToPin triggered a change that requires a fuil node reconstruct
+		if (bReconstructNodeB)
 		{
-			if (!InputPin->Properties.bAllowMultipleConnections)
-			{
-				if (InputPin->BreakAllIncompatibleEdges())
-				{
-					PCGGraphNodeB->ReconstructNode();
-				}
-			}
+			PCGGraphNodeB->ReconstructNode();
 		}
 	}
 
