[PCG] AttributeListView added zoom to item functionality to doubleclick

diff --git a/Engine/Plugins/Experimental/PCG/Source/PCGEditor/Private/SPCGEditorGraphAttributeListView.cpp b/Engine/Plugins/Experimental/PCG/Source/PCGEditor/Private/SPCGEditorGraphAttributeListView.cpp
--- a/Engine/Plugins/Experimental/PCG/Source/PCGEditor/Private/SPCGEditorGraphAttributeListView.cpp
+++ b/Engine/Plugins/Experimental/PCG/Source/PCGEditor/Private/SPCGEditorGraphAttributeListView.cpp
@@ -397,6 +397,7 @@ void SPCGEditorGraphAttributeListView::Construct(const FArguments& InArgs, TShar
 		.ListItemsSource(&ListViewItems)
 		.HeaderRow(ListViewHeader)
 		.OnGenerateRow(this, &SPCGEditorGraphAttributeListView::OnGenerateRow)
+		.OnMouseButtonDoubleClick(this, &SPCGEditorGraphAttributeListView::OnItemDoubleClicked)
 		.AllowOverscroll(EAllowOverscroll::No)
 		.ExternalScrollbar(VerticalScrollBar)
 		.ConsumeMouseWheel(EConsumeMouseWheel::Always);
@@ -738,6 +739,17 @@ TSharedRef<ITableRow> SPCGEditorGraphAttributeListView::OnGenerateRow(PCGListvie
 	return SNew(SPCGListViewItemRow, OwnerTable, Item);
 }
 
+void SPCGEditorGraphAttributeListView::OnItemDoubleClicked(PCGListviewItemPtr Item) const
+{
+	check(Item);
+
+	if (const FPCGPoint* Point = Item->PCGPoint)
+	{
+		const FBox BoundingBox = Point->GetLocalBounds().TransformBy(Point->Transform.ToMatrixWithScale());
+		GEditor->MoveViewportCamerasToBox(BoundingBox, true, 2.5f);	
+	}
+}
+
 void SPCGEditorGraphAttributeListView::AddColumn(const FName& InColumnID, const FText& ColumnLabel, float ColumnWidth, EHorizontalAlignment HeaderHAlign, EHorizontalAlignment CellHAlign)
 {
 	SHeaderRow::FColumn::FArguments Arguments;
diff --git a/Engine/Plugins/Experimental/PCG/Source/PCGEditor/Private/SPCGEditorGraphAttributeListView.h b/Engine/Plugins/Experimental/PCG/Source/PCGEditor/Private/SPCGEditorGraphAttributeListView.h
--- a/Engine/Plugins/Experimental/PCG/Source/PCGEditor/Private/SPCGEditorGraphAttributeListView.h
+++ b/Engine/Plugins/Experimental/PCG/Source/PCGEditor/Private/SPCGEditorGraphAttributeListView.h
@@ -78,6 +78,7 @@ private:
 	void GenerateColumnsFromMetadata(const UPCGMetadata* PCGMetadata);
 
 	TSharedRef<ITableRow> OnGenerateRow(PCGListviewItemPtr Item, const TSharedRef<STableViewBase>& OwnerTable) const;
+	void OnItemDoubleClicked(PCGListviewItemPtr Item) const;
 
 	void AddColumn(const FName& InColumnID, const FText& ColumnLabel, float ColumnWidth, EHorizontalAlignment HeaderHAlign = HAlign_Center, EHorizontalAlignment CellHAlign = HAlign_Right);
 	void RemoveColumn(const FName& InColumnID);
