[PCG] Added PCG graph as an asset reference for PCG Volumes, this gives the ability to browse to and edit the asset directly from the actor right click menues and by commands (ctrl+b and ctrl+e)

diff --git a/Engine/Plugins/Experimental/PCG/Source/PCG/Private/PCGVolume.cpp b/Engine/Plugins/Experimental/PCG/Source/PCG/Private/PCGVolume.cpp
--- a/Engine/Plugins/Experimental/PCG/Source/PCG/Private/PCGVolume.cpp
+++ b/Engine/Plugins/Experimental/PCG/Source/PCG/Private/PCGVolume.cpp
@@ -2,10 +2,26 @@
 
 #include "PCGVolume.h"
 #include "PCGComponent.h"
-#include "Engine/World.h"
+#include "PCGGraph.h"
 
 APCGVolume::APCGVolume(const FObjectInitializer& ObjectInitializer)
 	: Super(ObjectInitializer)
 {
 	PCGComponent = ObjectInitializer.CreateDefaultSubobject<UPCGComponent>(this, TEXT("PCG Component"));
-}
\ No newline at end of file
+}
+
+#if WITH_EDITOR
+bool APCGVolume::GetReferencedContentObjects(TArray<UObject*>& Objects) const
+{
+	Super::GetReferencedContentObjects(Objects);
+
+	if (PCGComponent)
+    {
+    	if (UPCGGraph* PCGGraph = PCGComponent->GetGraph())
+    	{
+    		Objects.Add(PCGGraph);
+    	}
+    }
+    return true;
+}
+#endif // WITH_EDITOR
diff --git a/Engine/Plugins/Experimental/PCG/Source/PCG/Public/PCGVolume.h b/Engine/Plugins/Experimental/PCG/Source/PCG/Public/PCGVolume.h
--- a/Engine/Plugins/Experimental/PCG/Source/PCG/Public/PCGVolume.h
+++ b/Engine/Plugins/Experimental/PCG/Source/PCG/Public/PCGVolume.h
@@ -18,6 +18,12 @@ class PCG_API APCGVolume : public AVolume
 public:
 	APCGVolume(const FObjectInitializer& ObjectInitializer);
 
+	//~ Begin AActor Interface
+#if WITH_EDITOR
+	virtual bool GetReferencedContentObjects(TArray<UObject*>& Objects) const override;
+#endif // WITH_EDITOR
+	//~ End AActor Interface
+
 	UPROPERTY(BlueprintReadWrite, EditAnywhere, Category = PCG)
 	TObjectPtr<UPCGComponent> PCGComponent;
 };
