Added some guardrails when trying to call async processing inside of async processing in PCG

diff --git a/Engine/Plugins/Experimental/PCG/Source/PCG/Private/Helpers/PCGAsync.cpp b/Engine/Plugins/Experimental/PCG/Source/PCG/Private/Helpers/PCGAsync.cpp
--- a/Engine/Plugins/Experimental/PCG/Source/PCG/Private/Helpers/PCGAsync.cpp
+++ b/Engine/Plugins/Experimental/PCG/Source/PCG/Private/Helpers/PCGAsync.cpp
@@ -87,7 +87,17 @@ namespace FPCGAsync
 	void AsyncPointProcessing(FPCGContext* Context, int32 NumIterations, TArray<FPCGPoint>& OutPoints, const TFunction<bool(int32, FPCGPoint&)>& PointFunc)
 	{
 		const int32 MinIterationsPerTask = 256;
-		AsyncPointProcessing(Context ? Context->NumAvailableTasks : 1, MinIterationsPerTask, NumIterations, OutPoints, PointFunc);
+		
+		if (Context && !Context->bIsRunningAsyncCall)
+		{
+			Context->bIsRunningAsyncCall = true;
+			AsyncPointProcessing(Context->NumAvailableTasks, MinIterationsPerTask, NumIterations, OutPoints, PointFunc);
+			Context->bIsRunningAsyncCall = false;
+		}
+		else
+		{
+			AsyncPointProcessing(1, MinIterationsPerTask, NumIterations, OutPoints, PointFunc);
+		}
 	}
 
 	void AsyncPointProcessing(FPCGContext* Context, const TArray<FPCGPoint>& InPoints, TArray<FPCGPoint>& OutPoints, const TFunction<bool(const FPCGPoint&, FPCGPoint&)>& PointFunc)
@@ -110,7 +120,16 @@ namespace FPCGAsync
 			return NumPointsWritten;
 		};
 
-		AsyncPointProcessing(Context ? Context->NumAvailableTasks : 1, MinIterationsPerTask, InPoints.Num(), OutPoints, IterationInnerLoop);
+		if (Context && !Context->bIsRunningAsyncCall)
+		{
+			Context->bIsRunningAsyncCall = true;
+			AsyncPointProcessing(Context->NumAvailableTasks, MinIterationsPerTask, InPoints.Num(), OutPoints, IterationInnerLoop);
+			Context->bIsRunningAsyncCall = false;
+		}
+		else
+		{
+			AsyncPointProcessing(1, MinIterationsPerTask, InPoints.Num(), OutPoints, IterationInnerLoop);
+		}
 	}
 
 	void AsyncPointProcessing(int32 NumAvailableTasks, int32 MinIterationsPerTask, int32 NumIterations, TArray<FPCGPoint>& OutPoints, const TFunction<bool(int32, FPCGPoint&)>& PointFunc)
@@ -137,7 +156,17 @@ namespace FPCGAsync
 	void AsyncPointFilterProcessing(FPCGContext* Context, int32 NumIterations, TArray<FPCGPoint>& InFilterPoints, TArray<FPCGPoint>& OutFilterPoints, const TFunction<bool(int32, FPCGPoint&, FPCGPoint&)>& PointFunc)
 	{
 		const int32 MinIterationsPerTask = 256;
-		AsyncPointFilterProcessing(Context ? Context->NumAvailableTasks : 1, MinIterationsPerTask, NumIterations, InFilterPoints, OutFilterPoints, PointFunc);
+
+		if (Context && !Context->bIsRunningAsyncCall)
+		{
+			Context->bIsRunningAsyncCall = true;
+			AsyncPointFilterProcessing(Context->NumAvailableTasks, MinIterationsPerTask, NumIterations, InFilterPoints, OutFilterPoints, PointFunc);
+			Context->bIsRunningAsyncCall = false;
+		}
+		else
+		{
+			AsyncPointFilterProcessing(1, MinIterationsPerTask, NumIterations, InFilterPoints, OutFilterPoints, PointFunc);
+		}
 	}
 
 	void AsyncPointFilterProcessing(int32 NumAvailableTasks, int32 MinIterationsPerTask, int32 NumIterations, TArray<FPCGPoint>& InFilterPoints, TArray<FPCGPoint>& OutFilterPoints, const TFunction<bool(int32, FPCGPoint&, FPCGPoint&)>& PointFunc)
@@ -266,7 +295,16 @@ namespace FPCGAsync
 	void AsyncMultiPointProcessing(FPCGContext* Context, int32 NumIterations, TArray<FPCGPoint>& OutPoints, const TFunction<TArray<FPCGPoint>(int32)>& PointFunc)
 	{
 		const int32 MinIterationsPerTask = 256;
-		AsyncMultiPointProcessing(Context ? Context->NumAvailableTasks : 1, MinIterationsPerTask, NumIterations, OutPoints, PointFunc);
+		if (Context && !Context->bIsRunningAsyncCall)
+		{
+			Context->bIsRunningAsyncCall = true;
+			AsyncMultiPointProcessing(Context->NumAvailableTasks, MinIterationsPerTask, NumIterations, OutPoints, PointFunc);
+			Context->bIsRunningAsyncCall = false;
+		}
+		else
+		{
+			AsyncMultiPointProcessing(1, MinIterationsPerTask, NumIterations, OutPoints, PointFunc);
+		}
 	}
 
 	void AsyncMultiPointProcessing(int32 NumAvailableTasks, int32 MinIterationsPerTask, int32 NumIterations, TArray<FPCGPoint>& OutPoints, const TFunction<TArray<FPCGPoint>(int32)>& PointFunc)
diff --git a/Engine/Plugins/Experimental/PCG/Source/PCG/Public/PCGContext.h b/Engine/Plugins/Experimental/PCG/Source/PCG/Public/PCGContext.h
--- a/Engine/Plugins/Experimental/PCG/Source/PCG/Public/PCGContext.h
+++ b/Engine/Plugins/Experimental/PCG/Source/PCG/Public/PCGContext.h
@@ -61,6 +61,9 @@ struct PCG_API FPCGContext
 	double EndTime = 0.0;
 	bool bIsRunningOnMainThread = false;
 
+	/** True if currently inside a PCGAsync scope - will prevent further async processing */
+	bool bIsRunningAsyncCall = false;
+
 #if WITH_EDITOR
 	double ElapsedTime = 0.0;
 	int32 ExecutionCount = 0;
