Added a flag to decide condition to call Generate on spawned actors containing a PCG component. Used only in non-merged cases

diff --git a/Engine/Plugins/Experimental/PCG/Source/PCG/Private/Elements/PCGSpawnActor.cpp b/Engine/Plugins/Experimental/PCG/Source/PCG/Private/Elements/PCGSpawnActor.cpp
--- a/Engine/Plugins/Experimental/PCG/Source/PCG/Private/Elements/PCGSpawnActor.cpp
+++ b/Engine/Plugins/Experimental/PCG/Source/PCG/Private/Elements/PCGSpawnActor.cpp
@@ -213,6 +213,14 @@ bool FPCGSpawnActorElement::ExecuteInternal(FPCGContext* Context) const
 					// Create managed resource for actor tracking
 					UPCGManagedActors* ManagedActors = NewObject<UPCGManagedActors>(Context->SourceComponent.Get());
 
+					const bool bForceCallGenerate = Settings->bGenerationTrigger == EPCGSpawnActorGenerationTrigger::ForceGenerate;
+#if WITH_EDITOR
+					const bool bOnLoadCallGenerate = Settings->bGenerationTrigger == EPCGSpawnActorGenerationTrigger::Default;
+#else
+					const bool bOnLoadCallGenerate = (Settings->bGenerationTrigger == EPCGSpawnActorGenerationTrigger::Default ||
+						Settings->bGenerationTrigger == EPCGSpawnActorGenerationTrigger::DoNotGenerateInEditor);
+#endif
+
 					for (const FPCGPoint& Point : Points)
 					{
 						AActor* GeneratedActor = TargetActor->GetWorld()->SpawnActor(Settings->TemplateActorClass, &Point.Transform, SpawnParams);
@@ -241,7 +249,10 @@ bool FPCGSpawnActorElement::ExecuteInternal(FPCGContext* Context) const
 									PCGComponent->bParseActorComponents = false;
 								}
 
-								PCGComponent->Generate();
+								if(bForceCallGenerate || (bOnLoadCallGenerate && PCGComponent->GenerationTrigger == EPCGComponentGenerationTrigger::GenerateOnLoad))
+								{
+									PCGComponent->Generate();
+								}
 							}
 							else
 							{
diff --git a/Engine/Plugins/Experimental/PCG/Source/PCG/Public/Elements/PCGSpawnActor.h b/Engine/Plugins/Experimental/PCG/Source/PCG/Public/Elements/PCGSpawnActor.h
--- a/Engine/Plugins/Experimental/PCG/Source/PCG/Public/Elements/PCGSpawnActor.h
+++ b/Engine/Plugins/Experimental/PCG/Source/PCG/Public/Elements/PCGSpawnActor.h
@@ -20,6 +20,15 @@ enum class EPCGSpawnActorOption : uint8
 	NoMerging
 };
 
+UENUM()
+enum class EPCGSpawnActorGenerationTrigger : uint8
+{
+	Default, // Generate if the component has "Generate On Load"
+	ForceGenerate, // Generate in all cases
+	DoNotGenerateInEditor, // Does not call generate in editor, but decays to Default otherwise
+	DoNotGenerate // Does not call generate
+};
+
 UCLASS(BlueprintType, ClassGroup = (Procedural))
 class PCG_API UPCGSpawnActorSettings : public UPCGBaseSubgraphSettings
 {
@@ -39,6 +48,9 @@ public:
 	UPROPERTY(BlueprintReadWrite, EditAnywhere, Category = Settings, meta = (EditCondition = "Option==EPCGSpawnActorOption::NoMerging"))
 	bool bForceDisableActorParsing = true;
 
+	UPROPERTY(BlueprintReadWrite, EditAnywhere, Category = Settings, meta = (EditCondition = "Option==EPCGSpawnActorOption::NoMerging"))
+	EPCGSpawnActorGenerationTrigger bGenerationTrigger = EPCGSpawnActorGenerationTrigger::Default;
+
 	/** Warning: inheriting parent actor tags work only in non-collapsed actor hierarchies */
 	UPROPERTY(BlueprintReadWrite, EditAnywhere, Category = Settings, meta = (EditCondition = "Option!=EPCGSpawnActorOption::CollapseActors"))
 	bool bInheritActorTags = false;
