Added UPCGMetadata::SetAttributeFromDataProperty, to be able to create a property from a reflected struct.
Added 'PCGDataTableRowToParamData' node

diff --git a/Engine/Plugins/Experimental/PCG/Source/PCG/Private/Elements/PCGDataTableRowToParamData.cpp b/Engine/Plugins/Experimental/PCG/Source/PCG/Private/Elements/PCGDataTableRowToParamData.cpp
new file mode 100644
--- /dev/null
+++ b/Engine/Plugins/Experimental/PCG/Source/PCG/Private/Elements/PCGDataTableRowToParamData.cpp
@@ -0,0 +1,71 @@
+// Copyright Epic Games, Inc. All Rights Reserved.
+
+#include "Elements/PCGDataTableRowToParamData.h"
+#include "Helpers/PCGAsync.h"
+#include "Helpers/PCGSettingsHelpers.h"
+#include "Metadata/PCGMetadata.h"
+#include "Engine/DataTable.h"
+
+FPCGElementPtr UPCGDataTableRowToParamDataSettings::CreateElement() const
+{
+	return MakeShared<FPCGDataTableRowToParamData>();
+}
+
+TArray<FPCGPinProperties> UPCGDataTableRowToParamDataSettings::OutputPinProperties() const
+{
+	TArray<FPCGPinProperties> PinProperties;
+	PinProperties.Emplace(PCGPinConstants::DefaultOutputLabel, EPCGDataType::Param);
+
+	return PinProperties;
+}
+
+FName UPCGDataTableRowToParamDataSettings::AdditionalTaskName() const
+{
+	return FName(FString::Printf(TEXT("%s[ %s ]"), DataTable ? *DataTable->GetFName().ToString() : TEXT("None"), *RowName.ToString()));
+}
+
+bool FPCGDataTableRowToParamData::ExecuteInternal(FPCGContext* Context) const
+{
+	TRACE_CPUPROFILER_EVENT_SCOPE(FPCGDataTableRowToParamData::Execute);
+
+	const UPCGDataTableRowToParamDataSettings* Settings = Context->GetInputSettings<UPCGDataTableRowToParamDataSettings>();
+	check(Settings);
+
+	const UDataTable* DataTable = Settings->DataTable.Get();
+	if (!DataTable)
+	{
+		PCGE_LOG(Error, TEXT("FPCGDataTableRowToParamData: Data table is invalid"));
+		return true;
+	}
+
+	const uint8* const* RowDataItr = DataTable->GetRowMap().Find(Settings->RowName);
+	const uint8* RowDataPtr = RowDataItr ? *RowDataItr : nullptr;
+	if (!RowDataPtr)
+	{
+		PCGE_LOG(Error, TEXT("FPCGDataTableRowToParamData: Data table '%s' has no row named '%s'"), *DataTable->GetPathName(), *Settings->RowName.ToString());
+		return true;
+	}
+
+	// From there, we should be able to create the data.
+	UPCGParamData* ParamData = NewObject<UPCGParamData>();
+	UPCGMetadata* Metadata = ParamData->MutableMetadata();
+	check(Metadata);
+	PCGMetadataEntryKey EntryKey = Metadata->AddEntry();
+
+	for (TFieldIterator<const FProperty> FieldIt(DataTable->GetRowStruct(), EFieldIterationFlags::IncludeSuper); FieldIt; ++FieldIt)
+	{
+		const FString FieldName = DataTable->GetRowStruct()->GetAuthoredNameForField(*FieldIt);
+		const FName AttributeName(FieldName);
+
+		if (!Metadata->SetAttributeFromDataProperty(AttributeName, EntryKey, RowDataPtr, *FieldIt, /*bCreate=*/ true))
+		{
+			PCGE_LOG(Warning, "Error while creating a attribute '%s'. Either the property type is not supported by PCG or attribute creation failed.", *AttributeName.ToString());
+		}
+	}
+	
+	TArray<FPCGTaggedData>& Outputs = Context->OutputData.TaggedData;
+	FPCGTaggedData& Output = Outputs.Emplace_GetRef();
+	Output.Data = ParamData;
+
+	return true;
+}
diff --git a/Engine/Plugins/Experimental/PCG/Source/PCG/Private/Metadata/PCGMetadata.cpp b/Engine/Plugins/Experimental/PCG/Source/PCG/Private/Metadata/PCGMetadata.cpp
--- a/Engine/Plugins/Experimental/PCG/Source/PCG/Private/Metadata/PCGMetadata.cpp
+++ b/Engine/Plugins/Experimental/PCG/Source/PCG/Private/Metadata/PCGMetadata.cpp
@@ -403,60 +403,74 @@ void UPCGMetadata::CreateStringAttribute(FName AttributeName, FString DefaultVal
 	CreateAttribute<FString>(AttributeName, DefaultValue, bAllowsInterpolation, bOverrideParent);
 }
 
-bool UPCGMetadata::SetAttributeFromProperty(FName AttributeName, PCGMetadataEntryKey& EntryKey, const UObject* Object, const FProperty* InProperty, bool bCreate)
+namespace PCGMetadata
 {
-	if (!InProperty || !Object)
-	{
-		return false;
-	}
-
-	// Check if an attribute already exists or not if we ask to create a new one
-	if (!bCreate && !HasAttribute(AttributeName))
-	{
-		return false;
-	}
-
-	auto CreateAttributeAndSet = [&AttributeName, this, bCreate, &EntryKey](auto&& PropertyValue) -> bool
+	template<typename DataType>
+	bool SetAttributeFromPropertyHelper(UPCGMetadata& Metadata, FName AttributeName, PCGMetadataEntryKey& EntryKey, const DataType* DataPtr, const FProperty* InProperty, bool bCreate)
 	{
-		using PropertyType = std::remove_const_t<std::remove_reference_t<decltype(PropertyValue)>>;
-
-		FPCGMetadataAttributeBase* BaseAttribute = GetMutableAttribute(AttributeName);
-
-		if (!BaseAttribute && bCreate)
+		if (!InProperty || !DataPtr)
 		{
-			// Interpolation is disabled and no parent override.
-			BaseAttribute = CreateAttribute<PropertyType>(AttributeName, PropertyValue, false, false);
+			return false;
 		}
 
-		if (!BaseAttribute)
+		// Check if an attribute already exists or not if we ask to create a new one
+		if (!bCreate && !Metadata.HasAttribute(AttributeName))
 		{
 			return false;
 		}
 
-		// Allow to set the value if both type matches or if we can construct AttributeType from PropertyType.
-		return PCGMetadataAttribute::CallbackWithRightType(BaseAttribute->GetTypeId(), [&](auto AttributeValue) -> bool
+		auto CreateAttributeAndSet = [AttributeName, &Metadata, bCreate, &EntryKey](auto&& PropertyValue) -> bool
+		{
+			using PropertyType = std::remove_const_t<std::remove_reference_t<decltype(PropertyValue)>>;
+
+			FPCGMetadataAttributeBase* BaseAttribute = Metadata.GetMutableAttribute(AttributeName);
+
+			if (!BaseAttribute && bCreate)
 			{
-				using AttributeType = decltype(AttributeValue);
-				FPCGMetadataAttribute<AttributeType>* Attribute = static_cast<FPCGMetadataAttribute<AttributeType>*>(BaseAttribute);
+				// Interpolation is disabled and no parent override.
+				BaseAttribute = Metadata.CreateAttribute<PropertyType>(AttributeName, PropertyValue, false, false);
+			}
 
-				if constexpr (std::is_same_v<AttributeType, PropertyType>)
-				{
-					Attribute->SetValue(EntryKey, PropertyValue);
-					return true;
-				}
-				else if constexpr (std::is_constructible_v<AttributeType, PropertyType>)
-				{
-					Attribute->SetValue(EntryKey, AttributeType(PropertyValue));
-					return true;
-				}
-				else
+			if (!BaseAttribute)
+			{
+				return false;
+			}
+
+			// Allow to set the value if both type matches or if we can construct AttributeType from PropertyType.
+			return PCGMetadataAttribute::CallbackWithRightType(BaseAttribute->GetTypeId(), [&](auto AttributeValue) -> bool
 				{
-					return false;
-				}
-			});
-	};
+					using AttributeType = decltype(AttributeValue);
+					FPCGMetadataAttribute<AttributeType>* Attribute = static_cast<FPCGMetadataAttribute<AttributeType>*>(BaseAttribute);
+
+					if constexpr (std::is_same_v<AttributeType, PropertyType>)
+					{
+						Attribute->SetValue(EntryKey, PropertyValue);
+						return true;
+					}
+					else if constexpr (std::is_constructible_v<AttributeType, PropertyType>)
+					{
+						Attribute->SetValue(EntryKey, AttributeType(PropertyValue));
+						return true;
+					}
+					else
+					{
+						return false;
+					}
+				});
+		};
+
+		return PCGSettingsHelpers::GetPropertyValueWithCallback(DataPtr, InProperty, CreateAttributeAndSet);
+	}
+}
 
-	return PCGSettingsHelpers::GetPropertyValueWithCallback(Object, InProperty, CreateAttributeAndSet);
+bool UPCGMetadata::SetAttributeFromProperty(FName AttributeName, PCGMetadataEntryKey& EntryKey, const UObject* Object, const FProperty* InProperty, bool bCreate)
+{
+	return PCGMetadata::SetAttributeFromPropertyHelper<UObject>(*this, AttributeName, EntryKey, Object, InProperty, bCreate);
+}
+
+bool UPCGMetadata::SetAttributeFromDataProperty(FName AttributeName, PCGMetadataEntryKey& EntryKey, const void* Data, const FProperty* InProperty, bool bCreate)
+{
+	return PCGMetadata::SetAttributeFromPropertyHelper<void>(*this, AttributeName, EntryKey, Data, InProperty, bCreate);
 }
 
 void UPCGMetadata::CopyExistingAttribute(FName AttributeToCopy, FName NewAttributeName, bool bKeepParent)
diff --git a/Engine/Plugins/Experimental/PCG/Source/PCG/Public/Elements/PCGDataTableRowToParamData.h b/Engine/Plugins/Experimental/PCG/Source/PCG/Public/Elements/PCGDataTableRowToParamData.h
new file mode 100644
--- /dev/null
+++ b/Engine/Plugins/Experimental/PCG/Source/PCG/Public/Elements/PCGDataTableRowToParamData.h
@@ -0,0 +1,53 @@
+// Copyright Epic Games, Inc. All Rights Reserved.
+
+#pragma once
+
+#include "PCGElement.h"
+#include "CoreMinimal.h"
+#include "PCGSettings.h"
+#include "UObject/ObjectMacros.h"
+
+#include "PCGDataTableRowToParamData.generated.h"
+
+class UDataTable;
+
+UCLASS(BlueprintType, ClassGroup = (Procedural))
+class UPCGDataTableRowToParamDataSettings : public UPCGSettings
+{
+	GENERATED_BODY()
+
+public:
+#if WITH_EDITOR
+	//~Begin UPCGSettings interface
+	virtual FName GetDefaultNodeName() const override { return FName(TEXT("DataTableRowToParamData")); }
+	virtual EPCGSettingsType GetType() const override { return EPCGSettingsType::Metadata; }
+#endif
+
+	virtual FName AdditionalTaskName() const override;
+
+	virtual TArray<FPCGPinProperties> OutputPinProperties() const override;
+	virtual TArray<FPCGPinProperties> InputPinProperties() const override { return TArray<FPCGPinProperties>(); }
+
+	// The name of the row to copy from
+	UPROPERTY(BlueprintReadWrite, EditAnywhere, Category = Settings)
+	FName RowName = NAME_None;	
+
+	// the data table to copy from
+	UPROPERTY(BlueprintReadWrite, EditAnywhere, Category = Settings)
+	TObjectPtr<UDataTable> DataTable;
+
+protected:
+	virtual FPCGElementPtr CreateElement() const override;
+	// ~End UPCGSettings interface
+};
+
+class FPCGDataTableRowToParamData : public FSimplePCGElement
+{
+public:
+	virtual bool CanExecuteOnlyOnMainThread(FPCGContext* Context) const override { return true; }
+	virtual bool IsCacheable(const UPCGSettings* InSettings) const { return false; }
+
+protected:
+	virtual bool ExecuteInternal(FPCGContext* Context) const override;
+};
+
diff --git a/Engine/Plugins/Experimental/PCG/Source/PCG/Public/Helpers/PCGSettingsHelpers.h b/Engine/Plugins/Experimental/PCG/Source/PCG/Public/Helpers/PCGSettingsHelpers.h
--- a/Engine/Plugins/Experimental/PCG/Source/PCG/Public/Helpers/PCGSettingsHelpers.h
+++ b/Engine/Plugins/Experimental/PCG/Source/PCG/Public/Helpers/PCGSettingsHelpers.h
@@ -126,8 +126,8 @@ namespace PCGSettingsHelpers
 	* Validate that the InProperty is a PCGData supported type and call InFunc with the value of this property (with the right type).
 	* The function returns the result of InFunc (or default value of the return type of InFunc if it failed).
 	*/
-	template <typename Func>
-	inline decltype(auto) GetPropertyValueWithCallback(const UObject* InObject, const FProperty* InProperty, Func InFunc)
+	template <typename ObjectType, typename Func>
+	inline decltype(auto) GetPropertyValueWithCallback(const ObjectType* InObject, const FProperty* InProperty, Func InFunc)
 	{
 		// Double property is supported by PCG, we use a dummy double to deduce the return type of InFunc.
 		using ReturnType = decltype(InFunc(0.0));
diff --git a/Engine/Plugins/Experimental/PCG/Source/PCG/Public/Metadata/PCGMetadata.h b/Engine/Plugins/Experimental/PCG/Source/PCG/Public/Metadata/PCGMetadata.h
--- a/Engine/Plugins/Experimental/PCG/Source/PCG/Public/Metadata/PCGMetadata.h
+++ b/Engine/Plugins/Experimental/PCG/Source/PCG/Public/Metadata/PCGMetadata.h
@@ -118,6 +118,16 @@ public:
 	*/
 	bool SetAttributeFromProperty(FName AttributeName, PCGMetadataEntryKey& EntryKey, const UObject* Object, const FProperty* Property, bool bCreate);
 
+	/** Set an attribute given a property and its value.
+	* @param AttributeName: Target attribute to set the property's value to
+	* @param EntryKey: Metadata entry key to set the value to
+	* @param Data: Data pointer to get the property value from
+	* @param Property: The property to set from
+	* @param bCreate: If true and the attribute doesn't exists, it will create an attribute based on the property type
+	* @returns true if the attribute creation (if required) and the value set succeeded
+	*/
+	bool SetAttributeFromDataProperty(FName AttributeName, PCGMetadataEntryKey& EntryKey, const void* Data, const FProperty* Property, bool bCreate);
+
 	/** Get attributes */
 	FPCGMetadataAttributeBase* GetMutableAttribute(FName AttributeName);
 	const FPCGMetadataAttributeBase* GetConstAttribute(FName AttributeName) const;
