[PCG] Changed inspected pcg component from raw to weak pointer to avoid accessing deleted data in the attribute list.

diff --git a/Engine/Plugins/Experimental/PCG/Source/PCGEditor/Private/PCGEditor.cpp b/Engine/Plugins/Experimental/PCG/Source/PCGEditor/Private/PCGEditor.cpp
--- a/Engine/Plugins/Experimental/PCG/Source/PCGEditor/Private/PCGEditor.cpp
+++ b/Engine/Plugins/Experimental/PCG/Source/PCGEditor/Private/PCGEditor.cpp
@@ -157,9 +157,9 @@ void FPCGEditor::SetPCGComponentBeingDebugged(UPCGComponent* InPCGComponent)
 	{
 		PCGComponentBeingDebugged = InPCGComponent;
 
-		OnDebugObjectChangedDelegate.Broadcast(PCGComponentBeingDebugged); 
+		OnDebugObjectChangedDelegate.Broadcast(PCGComponentBeingDebugged.Get());
 
-		if (PCGComponentBeingDebugged)
+		if (PCGComponentBeingDebugged.IsValid())
 		{
 			// We need to force generation so that we create debug data
 			PCGComponentBeingDebugged->GenerateLocal(/*bForce=*/true);
@@ -172,7 +172,7 @@ void FPCGEditor::SetPCGNodeBeingInspected(UPCGNode* InPCGNode)
 	if (GetPCGNodeBeingInspected() != InPCGNode)
 	{
 		PCGNodeBeingInspected = InPCGNode;
-		OnInspectedNodeChangedDelegate.Broadcast(PCGNodeBeingInspected);	
+		OnInspectedNodeChangedDelegate.Broadcast(PCGNodeBeingInspected);
 	}
 }
 
@@ -187,7 +187,7 @@ void FPCGEditor::JumpToNode(const UEdGraphNode* InNode)
 void FPCGEditor::RegisterTabSpawners(const TSharedRef<FTabManager>& InTabManager)
 {
 	WorkspaceMenuCategory = InTabManager->AddLocalWorkspaceMenuCategory(LOCTEXT("WorkspaceMenu_PCGEditor", "PCG Editor"));
-	auto WorkspaceMenuCategoryRef = WorkspaceMenuCategory.ToSharedRef();
+	const TSharedRef<FWorkspaceItem>& WorkspaceMenuCategoryRef = WorkspaceMenuCategory.ToSharedRef();
 
 	FAssetEditorToolkit::RegisterTabSpawners(InTabManager);
 
@@ -575,14 +575,14 @@ void FPCGEditor::OnDeterminismNodeTest()
 
 bool FPCGEditor::CanRunDeterminismGraphTest() const
 {
-	return PCGEditorGraph && PCGComponentBeingDebugged;
+	return PCGEditorGraph && PCGComponentBeingDebugged.IsValid();
 }
 
 void FPCGEditor::OnDeterminismGraphTest()
 {
 	check(GraphEditorWidget.IsValid());
 
-	if (!DeterminismWidget.IsValid() || !DeterminismWidget->WidgetIsConstructed() || !PCGGraphBeingEdited || !PCGComponentBeingDebugged)
+	if (!DeterminismWidget.IsValid() || !DeterminismWidget->WidgetIsConstructed() || !PCGGraphBeingEdited || !PCGComponentBeingDebugged.IsValid())
 	{
 		return;
 	}
@@ -604,7 +604,7 @@ void FPCGEditor::OnDeterminismGraphTest()
 	TestResult->TestResultName = PCGGraphBeingEdited->GetName();
 	TestResult->Seed = PCGComponentBeingDebugged->Seed;
 
-	PCGDeterminismTests::RunDeterminismTest(PCGGraphBeingEdited, PCGComponentBeingDebugged, *TestResult);
+	PCGDeterminismTests::RunDeterminismTest(PCGGraphBeingEdited, PCGComponentBeingDebugged.Get(), *TestResult);
 
 	DeterminismWidget->AddItem(TestResult);
 	DeterminismWidget->AddDetailsColumn();
@@ -1678,10 +1678,10 @@ void FPCGEditor::JumpToDefinition(const UClass* Class) const
 				{
 					const FScopedTransaction Transaction(LOCTEXT("ChangeEnableNavigateToNativeNodes", "Change Enable Navigate to Native Nodes Setting"));
 
-					UPCGEditorSettings* MutableEditorSetings = GetMutableDefault<UPCGEditorSettings>();
-					MutableEditorSetings->Modify();
-					MutableEditorSetings->bEnableNavigateToNativeNodes = (NewState == ECheckBoxState::Checked) ? true : false;
-					MutableEditorSetings->SaveConfig();
+					UPCGEditorSettings* MutableEditorSettings = GetMutableDefault<UPCGEditorSettings>();
+					MutableEditorSettings->Modify();
+					MutableEditorSettings->bEnableNavigateToNativeNodes = (NewState == ECheckBoxState::Checked) ? true : false;
+					MutableEditorSettings->SaveConfig();
 				}
 			);
 			Info.CheckBoxText = LOCTEXT("EnableNavigationToNative", "Enable Navigate to Native Nodes?");
diff --git a/Engine/Plugins/Experimental/PCG/Source/PCGEditor/Private/PCGEditor.h b/Engine/Plugins/Experimental/PCG/Source/PCGEditor/Private/PCGEditor.h
--- a/Engine/Plugins/Experimental/PCG/Source/PCGEditor/Private/PCGEditor.h
+++ b/Engine/Plugins/Experimental/PCG/Source/PCGEditor/Private/PCGEditor.h
@@ -38,7 +38,7 @@ public:
 	void SetPCGComponentBeingDebugged(UPCGComponent* InPCGComponent);
 
 	/** Gets the PCG component we are debugging */
-	UPCGComponent* GetPCGComponentBeingDebugged() const { return PCGComponentBeingDebugged; }
+	UPCGComponent* GetPCGComponentBeingDebugged() const { return PCGComponentBeingDebugged.Get(); }
 
 	/** Sets the PCG node we want to inspect */
 	void SetPCGNodeBeingInspected(UPCGNode* InPCGNode);
@@ -243,7 +243,7 @@ private:
 	UPCGGraph* PCGGraphBeingEdited = nullptr;
 	UPCGEditorGraph* PCGEditorGraph = nullptr;
 
-	UPCGComponent* PCGComponentBeingDebugged = nullptr;
+	TWeakObjectPtr<UPCGComponent> PCGComponentBeingDebugged;
 	UPCGNode* PCGNodeBeingInspected = nullptr;
 	UPCGEditorGraphNodeBase* PCGGraphNodeBeingInspected = nullptr;
 };
diff --git a/Engine/Plugins/Experimental/PCG/Source/PCGEditor/Private/SPCGEditorGraphAttributeListView.cpp b/Engine/Plugins/Experimental/PCG/Source/PCGEditor/Private/SPCGEditorGraphAttributeListView.cpp
--- a/Engine/Plugins/Experimental/PCG/Source/PCGEditor/Private/SPCGEditorGraphAttributeListView.cpp
+++ b/Engine/Plugins/Experimental/PCG/Source/PCGEditor/Private/SPCGEditorGraphAttributeListView.cpp
@@ -461,7 +461,7 @@ TSharedRef<SHeaderRow> SPCGEditorGraphAttributeListView::CreateHeaderRowWidget()
 
 void SPCGEditorGraphAttributeListView::OnDebugObjectChanged(UPCGComponent* InPCGComponent)
 {
-	if (PCGComponent)
+	if (PCGComponent.IsValid())
 	{
 		PCGComponent->OnPCGGraphGeneratedDelegate.RemoveAll(this);
 		PCGComponent->OnPCGGraphCleanedDelegate.RemoveAll(this);
@@ -470,7 +470,7 @@ void SPCGEditorGraphAttributeListView::OnDebugObjectChanged(UPCGComponent* InPCG
 
 	PCGComponent = InPCGComponent;
 
-	if (PCGComponent)
+	if (PCGComponent.IsValid())
 	{
 		PCGComponent->EnableInspection();
 		PCGComponent->OnPCGGraphGeneratedDelegate.AddSP(this, &SPCGEditorGraphAttributeListView::OnGenerateUpdated);
@@ -511,7 +511,7 @@ void SPCGEditorGraphAttributeListView::RefreshAttributeList()
 		return;
 	}
 
-	if (!PCGComponent)
+	if (!PCGComponent.IsValid())
 	{
 		return;
 	}
@@ -593,7 +593,7 @@ void SPCGEditorGraphAttributeListView::RefreshDataComboBox()
 	DataComboBox->ClearSelection();
 	DataComboBox->RefreshOptions();
 
-	if (!PCGComponent)
+	if (!PCGComponent.IsValid())
 	{
 		return;
 	}
diff --git a/Engine/Plugins/Experimental/PCG/Source/PCGEditor/Private/SPCGEditorGraphAttributeListView.h b/Engine/Plugins/Experimental/PCG/Source/PCGEditor/Private/SPCGEditorGraphAttributeListView.h
--- a/Engine/Plugins/Experimental/PCG/Source/PCGEditor/Private/SPCGEditorGraphAttributeListView.h
+++ b/Engine/Plugins/Experimental/PCG/Source/PCGEditor/Private/SPCGEditorGraphAttributeListView.h
@@ -96,7 +96,7 @@ private:
 	TWeakPtr<FPCGEditor> PCGEditorPtr;
 
 	/** Cached PCGComponent being viewed */
-	UPCGComponent* PCGComponent = nullptr;
+	TWeakObjectPtr<UPCGComponent> PCGComponent;
 
 	TSharedPtr<SHeaderRow> ListViewHeader;
 	TSharedPtr<SListView<PCGListviewItemPtr>> ListView;
