Fixed issue when not inheriting metadata attributes but preserving points with allocated metadata entry keys - since the parent hierarchy was broken, it would cause issues downstream.
Reused parenting-without-attributes introduced recently to keep this clean.

Also added Floor & Ceil (based on antoine.guilbaud's implementation).
Fixed minor UHT issue.

diff --git a/Engine/Plugins/Experimental/PCG/Source/PCG/Private/Data/PCGSpatialData.cpp b/Engine/Plugins/Experimental/PCG/Source/PCG/Private/Data/PCGSpatialData.cpp
--- a/Engine/Plugins/Experimental/PCG/Source/PCG/Private/Data/PCGSpatialData.cpp
+++ b/Engine/Plugins/Experimental/PCG/Source/PCG/Private/Data/PCGSpatialData.cpp
@@ -126,7 +126,7 @@ UPCGMetadata* UPCGSpatialData::CreateEmptyMetadata()
 	return Metadata;
 }
 
-void UPCGSpatialData::InitializeFromData(const UPCGSpatialData* InSource, const UPCGMetadata* InMetadataParentOverride, bool bInheritMetadata)
+void UPCGSpatialData::InitializeFromData(const UPCGSpatialData* InSource, const UPCGMetadata* InMetadataParentOverride, bool bInheritAttributes)
 {
 	if (InSource && TargetActor.IsExplicitlyNull())
 	{
@@ -138,11 +138,10 @@ void UPCGSpatialData::InitializeFromData(const UPCGSpatialData* InSource, const
 		Metadata = NewObject<UPCGMetadata>(this);
 	}
 
-	if (!bInheritMetadata || InMetadataParentOverride || InSource)
+	if (!bInheritAttributes || InMetadataParentOverride || InSource)
 	{
-		const UPCGMetadata* ParentMetadata = bInheritMetadata ? (InMetadataParentOverride ? InMetadataParentOverride : InSource->Metadata) : nullptr;
-
-		Metadata->Initialize(ParentMetadata);
+		const UPCGMetadata* ParentMetadata = (InMetadataParentOverride ? InMetadataParentOverride : InSource->Metadata);
+		Metadata->Initialize(ParentMetadata, bInheritAttributes);
 	}
 	else
 	{
diff --git a/Engine/Plugins/Experimental/PCG/Source/PCG/Private/Elements/Metadata/PCGMetadataMaths.inl b/Engine/Plugins/Experimental/PCG/Source/PCG/Private/Elements/Metadata/PCGMetadataMaths.inl
--- a/Engine/Plugins/Experimental/PCG/Source/PCG/Private/Elements/Metadata/PCGMetadataMaths.inl
+++ b/Engine/Plugins/Experimental/PCG/Source/PCG/Private/Elements/Metadata/PCGMetadataMaths.inl
@@ -36,7 +36,7 @@ namespace PCGMetadataMaths
 	template <typename T>
 	inline T Frac(const T& Value)
 	{
-		return static_cast<T>(FMath::Frac((double)Value));
+		return static_cast<T>(FMath::Frac(static_cast<double>(Value)));
 	}
 
 	template<>
@@ -63,7 +63,7 @@ namespace PCGMetadataMaths
 	template <typename T>
 	inline T Truncate(const T& Value)
 	{
-		return static_cast<T>(FMath::TruncToDouble((double)Value));
+		return static_cast<T>(FMath::TruncToDouble(static_cast<double>(Value)));
 	}
 
 	template<>
@@ -90,7 +90,7 @@ namespace PCGMetadataMaths
 	template <typename T>
 	inline T Round(const T& Value)
 	{
-		return static_cast<T>(FMath::RoundToDouble((double)Value));
+		return static_cast<T>(FMath::RoundToDouble(static_cast<double>(Value)));
 	}
 
 	template<>
@@ -111,13 +111,67 @@ namespace PCGMetadataMaths
 		return FVector4(Round(Value.X), Round(Value.Y), Round(Value.Z), Round(Value.W));
 	}
 
+	////////////////////
+	// Floor
+	////////////////////
+	template <typename T>
+	inline T Floor(const T& Value)
+	{
+		return static_cast<T>(FMath::FloorToDouble(static_cast<double>(Value)));
+	}
+
+	template<>
+	inline FVector2D Floor(const FVector2D& Value)
+	{
+		return FVector2D(Floor(Value.X), Floor(Value.Y));
+	}
+
+	template<>
+	inline FVector Floor(const FVector& Value)
+	{
+		return FVector(Floor(Value.X), Floor(Value.Y), Floor(Value.Z));
+	}
+
+	template<>
+	inline FVector4 Floor(const FVector4& Value)
+	{
+		return FVector4(Floor(Value.X), Floor(Value.Y), Floor(Value.Z), Floor(Value.W));
+	}
+
+	////////////////////
+	// Ceil
+	////////////////////
+	template <typename T>
+	inline T Ceil(const T& Value)
+	{
+		return static_cast<T>(FMath::CeilToDouble(static_cast<double>(Value)));
+	}
+
+	template<>
+	inline FVector2D Ceil(const FVector2D& Value)
+	{
+		return FVector2D(Ceil(Value.X), Ceil(Value.Y));
+	}
+
+	template<>
+	inline FVector Ceil(const FVector& Value)
+	{
+		return FVector(Ceil(Value.X), Ceil(Value.Y), Ceil(Value.Z));
+	}
+
+	template<>
+	inline FVector4 Ceil(const FVector4& Value)
+	{
+		return FVector4(Ceil(Value.X), Ceil(Value.Y), Ceil(Value.Z), Ceil(Value.W));
+	}
+
 	////////////////////
 	// Sqrt
 	////////////////////
 	template <typename T>
 	inline T Sqrt(const T& Value)
 	{
-		return static_cast<T>(FMath::Sqrt((double)Value));
+		return static_cast<T>(FMath::Sqrt(static_cast<double>(Value)));
 	}
 
 	template<>
@@ -225,7 +279,7 @@ namespace PCGMetadataMaths
 	template <typename T>
 	inline T Pow(const T& Value, const T& Power)
 	{
-		return static_cast<T>(FMath::Pow((double)Value, (double)Power));
+		return static_cast<T>(FMath::Pow(static_cast<double>(Value), static_cast<double>(Power)));
 	}
 
 	template<>
diff --git a/Engine/Plugins/Experimental/PCG/Source/PCG/Private/Elements/Metadata/PCGMetadataMathsOpElement.cpp b/Engine/Plugins/Experimental/PCG/Source/PCG/Private/Elements/Metadata/PCGMetadataMathsOpElement.cpp
--- a/Engine/Plugins/Experimental/PCG/Source/PCG/Private/Elements/Metadata/PCGMetadataMathsOpElement.cpp
+++ b/Engine/Plugins/Experimental/PCG/Source/PCG/Private/Elements/Metadata/PCGMetadataMathsOpElement.cpp
@@ -101,6 +101,10 @@ namespace PCGMetadataMathsSettings
 			return PCGMetadataMaths::Sqrt(Value);
 		case EPCGMedadataMathsOperation::Abs:
 			return PCGMetadataMaths::Abs(Value);
+		case EPCGMedadataMathsOperation::Floor:
+			return PCGMetadataMaths::Floor(Value);
+		case EPCGMedadataMathsOperation::Ceil:
+			return PCGMetadataMaths::Ceil(Value);
 		default:
 			return T{};
 		}
diff --git a/Engine/Plugins/Experimental/PCG/Source/PCG/Private/Elements/PCGCreateAttribute.cpp b/Engine/Plugins/Experimental/PCG/Source/PCG/Private/Elements/PCGCreateAttribute.cpp
--- a/Engine/Plugins/Experimental/PCG/Source/PCG/Private/Elements/PCGCreateAttribute.cpp
+++ b/Engine/Plugins/Experimental/PCG/Source/PCG/Private/Elements/PCGCreateAttribute.cpp
@@ -217,7 +217,7 @@ bool FPCGCreateAttributeElement::ExecuteInternal(FPCGContext* Context) const
 		{
 			// If we can reuse input data, it is safe to const_cast, as it was created by ourselves above.
 			UPCGParamData* NewParamData = bCanReuseInputData ? const_cast<UPCGParamData*>(InputParamData) : NewObject<UPCGParamData>();
-			NewParamData->Metadata->Initialize((!bCanReuseInputData && Settings->bKeepExistingAttributes) ? InputParamData->Metadata : nullptr);
+			NewParamData->Metadata->Initialize(bCanReuseInputData ? nullptr : InputParamData->Metadata, Settings->bKeepExistingAttributes);
 
 			OutputData = NewParamData;
 			Metadata = NewParamData->Metadata;
diff --git a/Engine/Plugins/Experimental/PCG/Source/PCG/Public/Data/PCGSpatialData.h b/Engine/Plugins/Experimental/PCG/Source/PCG/Public/Data/PCGSpatialData.h
--- a/Engine/Plugins/Experimental/PCG/Source/PCG/Public/Data/PCGSpatialData.h
+++ b/Engine/Plugins/Experimental/PCG/Source/PCG/Public/Data/PCGSpatialData.h
@@ -104,7 +104,7 @@ public:
 	UPCGMetadata* CreateEmptyMetadata();
 
 	UFUNCTION(BlueprintCallable, Category = SpatialData)
-	void InitializeFromData(const UPCGSpatialData* InSource, const UPCGMetadata* InMetadataParentOverride = nullptr, bool bInheritMetadata = true);
+	void InitializeFromData(const UPCGSpatialData* InSource, const UPCGMetadata* InMetadataParentOverride = nullptr, UPARAM(DisplayName="InheritAttributes") bool bInheritMetadata = true);
 
 	/** True if this operation does not have an inverse and cannot be queried analytically/implicitly, and therefore must be collapsed to an explicit point representation. */
 	virtual bool RequiresCollapseToSample() const { return false; }
diff --git a/Engine/Plugins/Experimental/PCG/Source/PCG/Public/Elements/Metadata/PCGMetadataMathsOpElement.h b/Engine/Plugins/Experimental/PCG/Source/PCG/Public/Elements/Metadata/PCGMetadataMathsOpElement.h
--- a/Engine/Plugins/Experimental/PCG/Source/PCG/Public/Elements/Metadata/PCGMetadataMathsOpElement.h
+++ b/Engine/Plugins/Experimental/PCG/Source/PCG/Public/Elements/Metadata/PCGMetadataMathsOpElement.h
@@ -19,6 +19,8 @@ enum class EPCGMedadataMathsOperation : uint16
 	Round,
 	Sqrt,
 	Abs,
+	Floor,
+	Ceil,
 
 	// Binary op
 	BinaryOp = 1 << 11 UMETA(Hidden),
@@ -73,7 +75,7 @@ public:
 	UPROPERTY(BlueprintReadWrite, EditAnywhere, Category = Input)
 	FName Input1AttributeName = NAME_None;
 
-	UPROPERTY(BlueprintReadWrite, EditAnywhere, Category = Input, meta = (EditCondition="(Operation & '/Script/PCG.EPCGMedadataMathsOperation::Binary') || (Operation & '/Script/PCG.EPCGMedadataMathsOperation::TernaryOp')", EditConditionHides))
+	UPROPERTY(BlueprintReadWrite, EditAnywhere, Category = Input, meta = (EditCondition="(Operation & '/Script/PCG.EPCGMedadataMathsOperation::BinaryOp') || (Operation & '/Script/PCG.EPCGMedadataMathsOperation::TernaryOp')", EditConditionHides))
 	FName Input2AttributeName = NAME_None;
 
 	UPROPERTY(BlueprintReadWrite, EditAnywhere, Category = Input, meta = (EditCondition = "Operation & '/Script/PCG.EPCGMedadataMathsOperation::TernaryOp'", EditConditionHides))
