Cleaned up landscape cache serialize to prevent issue with memeport

diff --git a/Engine/Plugins/Experimental/PCG/Source/PCG/Private/Grid/PCGLandscapeCache.cpp b/Engine/Plugins/Experimental/PCG/Source/PCG/Private/Grid/PCGLandscapeCache.cpp
--- a/Engine/Plugins/Experimental/PCG/Source/PCG/Private/Grid/PCGLandscapeCache.cpp
+++ b/Engine/Plugins/Experimental/PCG/Source/PCG/Private/Grid/PCGLandscapeCache.cpp
@@ -390,8 +390,6 @@ void FPCGLandscapeCacheEntry::SerializeFromBulkData() const
 
 void FPCGLandscapeCacheEntry::Serialize(FArchive& Archive, UObject* Owner, int32 Index)
 {
-	check(bDataLoaded || Archive.IsLoading());
-
 	if (bDataLoaded && Archive.IsSaving() && Archive.IsCooking())
 	{
 		SerializeToBulkData();
