[PCG] Fixed usage of deprecated asset object paths in the pcg graph schema

diff --git a/Engine/Plugins/Experimental/PCG/Source/PCGEditor/Private/PCGEditorGraphSchema.cpp b/Engine/Plugins/Experimental/PCG/Source/PCGEditor/Private/PCGEditorGraphSchema.cpp
--- a/Engine/Plugins/Experimental/PCG/Source/PCGEditor/Private/PCGEditorGraphSchema.cpp
+++ b/Engine/Plugins/Experimental/PCG/Source/PCGEditor/Private/PCGEditorGraphSchema.cpp
@@ -110,13 +110,13 @@ const FPinConnectionResponse UPCGEditorGraphSchema::CanCreateConnection(const UE
 bool UPCGEditorGraphSchema::TryCreateConnection(UEdGraphPin* InA, UEdGraphPin* InB) const
 {
 	// TODO: check if we need to verify connectivity first
-	bool bModified = Super::TryCreateConnection(InA, InB);
+	const bool bModified = Super::TryCreateConnection(InA, InB);
 
 	if (bModified)
 	{
 		check(InA && InB);
-		UEdGraphPin* A = (InA->Direction == EGPD_Output) ? InA : InB;
-		UEdGraphPin* B = (InA->Direction == EGPD_Input) ? InA : InB;
+		const UEdGraphPin* A = (InA->Direction == EGPD_Output) ? InA : InB;
+		const UEdGraphPin* B = (InA->Direction == EGPD_Input) ? InA : InB;
 		
 		check(A->Direction == EGPD_Output && B->Direction == EGPD_Input);
 
@@ -135,7 +135,7 @@ bool UPCGEditorGraphSchema::TryCreateConnection(UEdGraphPin* InA, UEdGraphPin* I
 
 		const bool bReconstructNodeB = PCGGraph->AddLabeledEdge(PCGNodeA, A->PinName, PCGNodeB, B->PinName);
 
-		// If AddLabeledEdge return true, then it means that the ToPin triggered a change that requires a fuil node reconstruct
+		// If AddLabeledEdge return true, then it means that the ToPin triggered a change that requires a full node reconstruct
 		if (bReconstructNodeB)
 		{
 			PCGGraphNodeB->ReconstructNode();
@@ -222,7 +222,7 @@ void UPCGEditorGraphSchema::GetNativeElementActions(FGraphActionMenuBuilder& Act
 
 void UPCGEditorGraphSchema::GetBlueprintElementActions(FGraphActionMenuBuilder& ActionMenuBuilder) const
 {
-	FAssetRegistryModule& AssetRegistryModule = FModuleManager::LoadModuleChecked<FAssetRegistryModule>(TEXT("AssetRegistry"));
+	const FAssetRegistryModule& AssetRegistryModule = FModuleManager::LoadModuleChecked<FAssetRegistryModule>(TEXT("AssetRegistry"));
 
 	FARFilter Filter;
 	Filter.ClassPaths.Add(UBlueprint::StaticClass()->GetClassPathName());
@@ -252,7 +252,7 @@ void UPCGEditorGraphSchema::GetBlueprintElementActions(FGraphActionMenuBuilder&
 
 void UPCGEditorGraphSchema::GetSubgraphElementActions(FGraphActionMenuBuilder& ActionMenuBuilder) const
 {
-	FAssetRegistryModule& AssetRegistryModule = FModuleManager::LoadModuleChecked<FAssetRegistryModule>(TEXT("AssetRegistry"));
+	const FAssetRegistryModule& AssetRegistryModule = FModuleManager::LoadModuleChecked<FAssetRegistryModule>(TEXT("AssetRegistry"));
 
 	TArray<FAssetData> AssetDataList;
 	AssetRegistryModule.Get().GetAssetsByClass(UPCGGraph::StaticClass()->GetClassPathName(), AssetDataList);
@@ -267,9 +267,7 @@ void UPCGEditorGraphSchema::GetSubgraphElementActions(FGraphActionMenuBuilder& A
 			const FText Description = AssetData.GetTagValueRef<FText>(TEXT("Description"));
 
 			TSharedPtr<FPCGEditorGraphSchemaAction_NewSubgraphElement> NewSubgraphAction(new FPCGEditorGraphSchemaAction_NewSubgraphElement(Category, MenuDesc, Description, 0));
-PRAGMA_DISABLE_DEPRECATION_WARNINGS
-			NewSubgraphAction->SubgraphObjectPath = AssetData.ObjectPath;
-PRAGMA_ENABLE_DEPRECATION_WARNINGS
+			NewSubgraphAction->SubgraphObjectPath = AssetData.GetSoftObjectPath();
 			ActionMenuBuilder.AddAction(NewSubgraphAction);
 		}
 	}
@@ -282,7 +280,7 @@ void UPCGEditorGraphSchema::GetExtraElementActions(FGraphActionMenuBuilder& Acti
 	const FText Category;
 	const FText Description = LOCTEXT("PCGAddCommentTooltip", "Create a resizable comment box.");
 
-	TSharedPtr<FPCGEditorGraphSchemaAction_NewComment> NewCommentAction(new FPCGEditorGraphSchemaAction_NewComment(Category, MenuDesc, Description, 0));
+	const TSharedPtr<FPCGEditorGraphSchemaAction_NewComment> NewCommentAction(new FPCGEditorGraphSchemaAction_NewComment(Category, MenuDesc, Description, 0));
 	ActionMenuBuilder.AddAction(NewCommentAction);
 }
 
@@ -299,9 +297,7 @@ void UPCGEditorGraphSchema::DroppedAssetsOnGraph(const TArray<FAssetData>& Asset
 			if (Asset->IsA<UPCGGraph>())
 			{
 				FPCGEditorGraphSchemaAction_NewSubgraphElement NewSubgraphAction;
-PRAGMA_DISABLE_DEPRECATION_WARNINGS
-				NewSubgraphAction.SubgraphObjectPath = AssetData.ObjectPath;
-PRAGMA_ENABLE_DEPRECATION_WARNINGS
+				NewSubgraphAction.SubgraphObjectPath = AssetData.GetSoftObjectPath();
 				NewSubgraphAction.PerformAction(Graph, NullFromPin, GraphPositionOffset);
 				GraphPositionOffset.Y += PositionOffsetIncrementY;
 			}
