Tentative fix for static analysis

diff --git a/Engine/Plugins/Experimental/PCG/Source/PCG/Private/Data/PCGSpatialData.cpp b/Engine/Plugins/Experimental/PCG/Source/PCG/Private/Data/PCGSpatialData.cpp
--- a/Engine/Plugins/Experimental/PCG/Source/PCG/Private/Data/PCGSpatialData.cpp
+++ b/Engine/Plugins/Experimental/PCG/Source/PCG/Private/Data/PCGSpatialData.cpp
@@ -140,7 +140,7 @@ void UPCGSpatialData::InitializeFromData(const UPCGSpatialData* InSource, const
 
 	if (!bInheritAttributes || InMetadataParentOverride || InSource)
 	{
-		const UPCGMetadata* ParentMetadata = (InMetadataParentOverride ? InMetadataParentOverride : InSource->Metadata);
+		const UPCGMetadata* ParentMetadata = (InMetadataParentOverride ? InMetadataParentOverride : (InSource ? InSource->Metadata : nullptr));
 		Metadata->Initialize(ParentMetadata, bInheritAttributes);
 	}
 	else
