Review spatial data changes since they have side effects

#rnx

diff --git a/Engine/Plugins/Experimental/PCG/Source/PCG/Private/Data/PCGSpatialData.cpp b/Engine/Plugins/Experimental/PCG/Source/PCG/Private/Data/PCGSpatialData.cpp
--- a/Engine/Plugins/Experimental/PCG/Source/PCG/Private/Data/PCGSpatialData.cpp
+++ b/Engine/Plugins/Experimental/PCG/Source/PCG/Private/Data/PCGSpatialData.cpp
@@ -126,7 +126,7 @@ UPCGMetadata* UPCGSpatialData::CreateEmptyMetadata()
 	return Metadata;
 }
 
-void UPCGSpatialData::InitializeFromData(const UPCGSpatialData* InSource, const UPCGMetadata* InMetadataParentOverride, bool bInheritAttributes)
+void UPCGSpatialData::InitializeFromData(const UPCGSpatialData* InSource, const UPCGMetadata* InMetadataParentOverride, bool bInheritMetadata, bool bInheritAttributes)
 {
 	if (InSource && TargetActor.IsExplicitlyNull())
 	{
@@ -138,9 +138,9 @@ void UPCGSpatialData::InitializeFromData(const UPCGSpatialData* InSource, const
 		Metadata = NewObject<UPCGMetadata>(this);
 	}
 
-	if (!bInheritAttributes || InMetadataParentOverride || InSource)
+	if (!bInheritMetadata || InMetadataParentOverride || InSource)
 	{
-		const UPCGMetadata* ParentMetadata = (InMetadataParentOverride ? InMetadataParentOverride : (InSource ? InSource->Metadata : nullptr));
+		const UPCGMetadata* ParentMetadata = bInheritMetadata ? (InMetadataParentOverride ? InMetadataParentOverride : (InSource ? InSource->Metadata : nullptr)) : nullptr;
 		Metadata->Initialize(ParentMetadata, bInheritAttributes);
 	}
 	else
diff --git a/Engine/Plugins/Experimental/PCG/Source/PCG/Private/Elements/PCGCreateAttribute.cpp b/Engine/Plugins/Experimental/PCG/Source/PCG/Private/Elements/PCGCreateAttribute.cpp
--- a/Engine/Plugins/Experimental/PCG/Source/PCG/Private/Elements/PCGCreateAttribute.cpp
+++ b/Engine/Plugins/Experimental/PCG/Source/PCG/Private/Elements/PCGCreateAttribute.cpp
@@ -208,7 +208,7 @@ bool FPCGCreateAttributeElement::ExecuteInternal(FPCGContext* Context) const
 		{
 			UPCGSpatialData* NewSpatialData = DuplicateObject<UPCGSpatialData>(InputSpatialData, nullptr);
 			NewSpatialData->Metadata = NewObject<UPCGMetadata>(NewSpatialData);
-			NewSpatialData->InitializeFromData(InputSpatialData, /*InMetadataParentOverride=*/ nullptr, /*bInheritMetadata=*/ Settings->bKeepExistingAttributes);
+			NewSpatialData->InitializeFromData(InputSpatialData, /*InMetadataParentOverride=*/ nullptr, /*bInheritMetadata=*/true, /*bInheritMetadata=*/ Settings->bKeepExistingAttributes);
 
 			OutputData = NewSpatialData;
 			Metadata = NewSpatialData->Metadata;
diff --git a/Engine/Plugins/Experimental/PCG/Source/PCG/Public/Data/PCGSpatialData.h b/Engine/Plugins/Experimental/PCG/Source/PCG/Public/Data/PCGSpatialData.h
--- a/Engine/Plugins/Experimental/PCG/Source/PCG/Public/Data/PCGSpatialData.h
+++ b/Engine/Plugins/Experimental/PCG/Source/PCG/Public/Data/PCGSpatialData.h
@@ -104,7 +104,7 @@ public:
 	UPCGMetadata* CreateEmptyMetadata();
 
 	UFUNCTION(BlueprintCallable, Category = SpatialData)
-	void InitializeFromData(const UPCGSpatialData* InSource, const UPCGMetadata* InMetadataParentOverride = nullptr, UPARAM(DisplayName="InheritAttributes") bool bInheritMetadata = true);
+	void InitializeFromData(const UPCGSpatialData* InSource, const UPCGMetadata* InMetadataParentOverride = nullptr, bool bInheritMetadata = true, bool bInheritAttributes = true);
 
 	/** True if this operation does not have an inverse and cannot be queried analytically/implicitly, and therefore must be collapsed to an explicit point representation. */
 	virtual bool RequiresCollapseToSample() const { return false; }
