Fixed landscape data query issue with landscape proxy bounds

#rnx

diff --git a/Engine/Plugins/Experimental/PCG/Source/PCG/Private/Data/PCGLandscapeData.cpp b/Engine/Plugins/Experimental/PCG/Source/PCG/Private/Data/PCGLandscapeData.cpp
--- a/Engine/Plugins/Experimental/PCG/Source/PCG/Private/Data/PCGLandscapeData.cpp
+++ b/Engine/Plugins/Experimental/PCG/Source/PCG/Private/Data/PCGLandscapeData.cpp
@@ -277,7 +277,7 @@ const ULandscapeInfo* UPCGLandscapeData::GetLandscapeInfo(const FVector& InPosit
 	// TODO: we could maybe improve on this if we find the "nearest" landscape on a Z perspective, but this might still lead to issues
 	for (const TPair<FBox, ULandscapeInfo*>& LandscapeInfoPair : LandscapeInfos)
 	{
-		if (LandscapeInfoPair.Key.IsInside(InPosition))
+		if (LandscapeInfoPair.Key.IsInsideXY(InPosition))
 		{
 			return LandscapeInfoPair.Value;
 		}
