Added PCGPointExtentsModifier to convert one of Jimmy's BP nodes to native. This allows for adjustment of the BoundsMin and BoundsMax values.

diff --git a/Engine/Plugins/Experimental/PCG/Source/PCG/Private/Elements/PCGPointExtentsModifier.cpp b/Engine/Plugins/Experimental/PCG/Source/PCG/Private/Elements/PCGPointExtentsModifier.cpp
new file mode 100644
--- /dev/null
+++ b/Engine/Plugins/Experimental/PCG/Source/PCG/Private/Elements/PCGPointExtentsModifier.cpp
@@ -0,0 +1,88 @@
+// Copyright Epic Games, Inc. All Rights Reserved.
+
+#include "Elements/PCGPointExtentsModifier.h"
+
+#include "PCGHelpers.h"
+#include "Helpers/PCGSettingsHelpers.h"
+
+FPCGElementPtr UPCGPointExtentsModifierSettings::CreateElement() const
+{
+	return MakeShared<FPCGPointExtentsModifier>();
+}
+
+TArray<FPCGPinProperties> UPCGPointExtentsModifierSettings::InputPinProperties() const
+{
+	TArray<FPCGPinProperties> PinProperties;
+	PinProperties.Emplace(PCGPinConstants::DefaultInputLabel, EPCGDataType::Spatial);	
+	PinProperties.Emplace(PCGPinConstants::DefaultParamsLabel, EPCGDataType::Param);
+	return PinProperties;
+}
+
+TArray<FPCGPinProperties> UPCGPointExtentsModifierSettings::OutputPinProperties() const
+{
+	TArray<FPCGPinProperties> PinProperties;
+	PinProperties.Emplace(PCGPinConstants::DefaultOutputLabel, EPCGDataType::Spatial);
+	return PinProperties;
+}
+
+bool FPCGPointExtentsModifier::ExecuteInternal(FPCGContext* Context) const
+{
+	TRACE_CPUPROFILER_EVENT_SCOPE(FPCGPointExtentsModifier::Execute);
+
+	const UPCGPointExtentsModifierSettings* Settings = Context->GetInputSettings<UPCGPointExtentsModifierSettings>();
+	check(Settings);
+
+	TArray<FPCGTaggedData> Inputs = Context->InputData.GetInputsByPin(PCGPinConstants::DefaultInputLabel);
+	TArray<FPCGTaggedData>& Outputs = Context->OutputData.TaggedData;
+	UPCGParamData* Params = Context->InputData.GetParams();
+
+	const EPCGPointExtentsModifierMode Mode = PCG_GET_OVERRIDEN_VALUE(Settings, Mode, Params);
+	const FVector Extents = PCG_GET_OVERRIDEN_VALUE(Settings, Extents, Params);
+
+	switch (Mode)
+	{
+	case EPCGPointExtentsModifierMode::Minimum:
+		ProcessPoints(Context, Inputs, Outputs, [Extents](const FPCGPoint& InPoint, FPCGPoint& OutPoint){
+			OutPoint = InPoint;
+			OutPoint.SetExtents(FVector::Min(Extents, InPoint.GetExtents()));
+			return true;
+		});
+		break;
+
+	case EPCGPointExtentsModifierMode::Maximum:
+		ProcessPoints(Context, Inputs, Outputs, [Extents](const FPCGPoint& InPoint, FPCGPoint& OutPoint){
+			OutPoint = InPoint;
+			OutPoint.SetExtents(FVector::Max(Extents, InPoint.GetExtents()));
+			return true;
+		});
+		break;
+
+	case EPCGPointExtentsModifierMode::Add:
+		ProcessPoints(Context, Inputs, Outputs, [Extents](const FPCGPoint& InPoint, FPCGPoint& OutPoint){
+			OutPoint = InPoint;
+			OutPoint.SetExtents(Extents + InPoint.GetExtents());
+			return true;
+		});
+		break;
+
+	case EPCGPointExtentsModifierMode::Multiply:
+		ProcessPoints(Context, Inputs, Outputs, [Extents](const FPCGPoint& InPoint, FPCGPoint& OutPoint){
+			OutPoint = InPoint;
+			OutPoint.SetExtents(Extents * InPoint.GetExtents());
+			return true;
+		});
+		break;
+
+	case EPCGPointExtentsModifierMode::Set:
+		ProcessPoints(Context, Inputs, Outputs, [Extents](const FPCGPoint& InPoint, FPCGPoint& OutPoint){
+			OutPoint = InPoint;
+			OutPoint.SetExtents(Extents);
+			return true;
+		});
+		break;
+	}
+
+	Outputs.Append(Context->InputData.GetAllSettings());
+	
+	return true;
+}
diff --git a/Engine/Plugins/Experimental/PCG/Source/PCG/Private/Tests/Elements/PCGPointExtentsModifierTest.cpp b/Engine/Plugins/Experimental/PCG/Source/PCG/Private/Tests/Elements/PCGPointExtentsModifierTest.cpp
new file mode 100644
--- /dev/null
+++ b/Engine/Plugins/Experimental/PCG/Source/PCG/Private/Tests/Elements/PCGPointExtentsModifierTest.cpp
@@ -0,0 +1,188 @@
+// Copyright Epic Games, Inc. All Rights Reserved.
+
+#include "Tests/Determinism/PCGDeterminismTestsCommon.h"
+#include "Tests/PCGTestsCommon.h"
+
+#include "PCGComponent.h"
+#include "PCGData.h"
+#include "PCGHelpers.h"
+
+#include "Data/PCGPointData.h"
+#include "Elements/PCGPointExtentsModifier.h"
+
+IMPLEMENT_CUSTOM_SIMPLE_AUTOMATION_TEST(FPCGPointsBoundsModiferTest, FPCGTestBaseClass, "pcg.tests.PointExtentsModifier.Basic", PCGTestsCommon::TestFlags)
+
+bool FPCGPointsBoundsModiferTest::RunTest(const FString& Parameters)
+{
+	PCGTestsCommon::FTestData TestData(PCGDeterminismTests::Defaults::Seed);
+	PCGDeterminismTests::GenerateSettings<UPCGPointExtentsModifierSettings>(TestData);
+	UPCGPointExtentsModifierSettings* Settings = CastChecked<UPCGPointExtentsModifierSettings>(TestData.Settings);
+	FPCGElementPtr BoundsModifierElement = TestData.Settings->GetElement();
+
+	TObjectPtr<UPCGPointData> PointData = PCGTestsCommon::CreateEmptyPointData();
+	TArray<FPCGPoint>& Points = PointData->GetMutablePoints();
+
+	static const int PointCount = 4;
+	for (int I = 0; I < PointCount; ++I)
+	{
+		FPCGPoint& Point = Points.Emplace_GetRef(FTransform(), 1, I);
+		const float Value = float(I+1);
+
+		Point.SetExtents(FVector(Value, Value, Value));
+
+		if (I == PointCount-1)
+		{
+			// make the last one off-center, so (-3,-3,-3),(5,5,5)
+			Point.SetLocalCenter(FVector(1.0f, 1.0f, 1.0f));
+		}
+	}
+
+	FPCGTaggedData& TaggedData = TestData.InputData.TaggedData.Emplace_GetRef(FPCGTaggedData());
+	TaggedData.Pin = PCGPinConstants::DefaultInputLabel;
+	TaggedData.Data = PointData;
+
+	auto ValidateBounds = [this, &TestData, BoundsModifierElement, Settings](TArray<FVector> ExpectedOutput) -> bool
+	{
+		TUniquePtr<FPCGContext> Context = MakeUnique<FPCGContext>(*BoundsModifierElement->Initialize(TestData.InputData, TestData.TestPCGComponent, nullptr));
+		Context->NumAvailableTasks = 1;
+
+		while (!BoundsModifierElement->Execute(Context.Get()))
+		{}
+
+		const TArray<FPCGTaggedData>& Inputs = Context->InputData.GetInputsByPin(PCGPinConstants::DefaultInputLabel);
+		const TArray<FPCGTaggedData>& Outputs = Context->OutputData.GetInputs();
+
+		if (!TestEqual("Valid number of Outputs", Inputs.Num(), Outputs.Num()))
+		{
+			return false;
+		}
+
+		if (!TestEqual("Has Outputs", Outputs.Num() >= 1, true))
+		{
+			return false;
+		}
+
+		bool bTestPassed = true;
+
+		for (int DataIndex = 0; DataIndex < Inputs.Num(); ++DataIndex)
+		{
+			const FPCGTaggedData& Input = Inputs[DataIndex];
+			const FPCGTaggedData& Output = Outputs[DataIndex];
+
+			const UPCGSpatialData* InSpatialData = Cast<UPCGSpatialData>(Input.Data);
+			check(InSpatialData);
+
+			const UPCGPointData* InPointData = InSpatialData->ToPointData(Context.Get());
+			check(InPointData);
+
+			const UPCGSpatialData* OutSpatialData = Cast<UPCGSpatialData>(Output.Data);
+
+			if (!TestNotNull("Valid output SpatialData", OutSpatialData))
+			{
+				bTestPassed = false;
+				continue;
+			}
+
+			const UPCGPointData* OutPointData = OutSpatialData->ToPointData(Context.Get());
+
+			if (!TestNotNull("Valid output PointData", OutPointData))
+			{
+				bTestPassed = false;
+				continue;
+			}
+
+			const TArray<FPCGPoint>& InPoints = InPointData->GetPoints();
+			const TArray<FPCGPoint>& OutPoints = OutPointData->GetPoints();
+
+			if (!TestEqual("Input and output point counts match", InPoints.Num(), OutPoints.Num()))
+			{ 
+				bTestPassed = false;
+				continue;
+			}
+
+			if (!TestEqual("Test data has enough points", PointCount, OutPoints.Num()))
+			{
+				bTestPassed = false;
+				continue;
+			}
+
+			for (int PointIndex = 0; PointIndex < ExpectedOutput.Num()/2; ++PointIndex)
+			{
+				bTestPassed &= TestEqual("Correct BoundsMin", OutPoints[PointIndex].BoundsMin, ExpectedOutput[PointIndex*2+0]);
+				bTestPassed &= TestEqual("Correct BoundsMax", OutPoints[PointIndex].BoundsMax, ExpectedOutput[PointIndex*2+1]);
+			}
+		}
+
+		return bTestPassed;
+	};
+
+	bool bTestPassed = true;
+
+	// Test expected values for Set
+	{
+		Settings->Extents = FVector(0.5f, 0.5f, 0.5f);
+		Settings->Mode = EPCGPointExtentsModifierMode::Set;
+		bTestPassed &= ValidateBounds({
+			{-0.5f, -0.5f, -0.5f}, {0.5f, 0.5f, 0.5f},
+			{-0.5f, -0.5f, -0.5f}, {0.5f, 0.5f, 0.5f},
+			{-0.5f, -0.5f, -0.5f}, {0.5f, 0.5f, 0.5f},
+			{0.5f, 0.5f, 0.5f}, {1.5f, 1.5f, 1.5f},
+		});
+	}
+
+
+	// Test expected values for Minimum
+	{
+		Settings->Extents = FVector(2.0f, 2.0f, 2.0f);
+		Settings->Mode = EPCGPointExtentsModifierMode::Minimum;
+
+		bTestPassed &= ValidateBounds({
+			{-1.0f, -1.0f, -1.0f}, {1.0f, 1.0f, 1.0f},
+			{-2.0f, -2.0f, -2.0f}, {2.0f, 2.0f, 2.0f},
+			{-2.0f, -2.0f, -2.0f}, {2.0f, 2.0f, 2.0f},
+			{-1.0f, -1.0f, -1.0f}, {3.0f, 3.0f, 3.0f},
+		});
+	}
+
+	// Test expected values for Maximum
+	{
+		Settings->Extents = FVector(2.0f, 2.0f, 2.0f);
+		Settings->Mode = EPCGPointExtentsModifierMode::Maximum;
+
+		bTestPassed &= ValidateBounds({
+			{-2.0f, -2.0f, -2.0f}, {2.0f, 2.0f, 2.0f},
+			{-2.0f, -2.0f, -2.0f}, {2.0f, 2.0f, 2.0f},
+			{-3.0f, -3.0f, -3.0f}, {3.0f, 3.0f, 3.0f},
+			{-3.0f, -3.0f, -3.0f}, {5.0f, 5.0f, 5.0f},
+		});
+	}
+
+	// Test expected values for Add
+	{
+		Settings->Extents = FVector(2.0f, 2.0f, 2.0f);
+		Settings->Mode = EPCGPointExtentsModifierMode::Add;
+
+		bTestPassed &= ValidateBounds({
+			{-3.0f, -3.0f, -3.0f}, {3.0f, 3.0f, 3.0f},
+			{-4.0f, -4.0f, -4.0f}, {4.0f, 4.0f, 4.0f},
+			{-5.0f, -5.0f, -5.0f}, {5.0f, 5.0f, 5.0f},
+			{-5.0f, -5.0f, -5.0f}, {7.0f, 7.0f, 7.0f},
+		});
+	}
+
+
+	// Test expected values for Multiply
+	{
+		Settings->Extents = FVector(2.0f, 2.0f, 2.0f);
+		Settings->Mode = EPCGPointExtentsModifierMode::Multiply;
+
+		bTestPassed &= ValidateBounds({
+			{-2.0f, -2.0f, -2.0f}, {2.0f, 2.0f, 2.0f},
+			{-4.0f, -4.0f, -4.0f}, {4.0f, 4.0f, 4.0f},
+			{-6.0f, -6.0f, -6.0f}, {6.0f, 6.0f, 6.0f},
+			{-7.0f, -7.0f, -7.0f}, {9.0f, 9.0f, 9.0f},
+		});
+	}
+
+	return bTestPassed;
+}
diff --git a/Engine/Plugins/Experimental/PCG/Source/PCG/Public/Elements/PCGPointExtentsModifier.h b/Engine/Plugins/Experimental/PCG/Source/PCG/Public/Elements/PCGPointExtentsModifier.h
new file mode 100644
--- /dev/null
+++ b/Engine/Plugins/Experimental/PCG/Source/PCG/Public/Elements/PCGPointExtentsModifier.h
@@ -0,0 +1,56 @@
+// Copyright Epic Games, Inc. All Rights Reserved.
+
+#pragma once
+
+#include "CoreMinimal.h"
+#include "PCGSettings.h"
+#include "PCGElement.h"
+#include "PCGNode.h"
+
+#include "Elements/PCGPointProcessingElementBase.h"
+
+#include "PCGPointExtentsModifier.generated.h"
+
+UENUM()
+enum class EPCGPointExtentsModifierMode : uint8
+{
+	Set,
+	Minimum,
+	Maximum,
+	Add,
+	Multiply
+};
+
+UCLASS(BlueprintType, ClassGroup = (Procedural))
+class PCG_API UPCGPointExtentsModifierSettings : public UPCGSettings
+{
+	GENERATED_BODY()
+
+public:
+
+	// ~Begin UPCGSettings interface
+#if WITH_EDITOR
+	virtual FName GetDefaultNodeName() const override { return FName(TEXT("ExtentsModifierNode")); }
+	virtual EPCGSettingsType GetType() const override { return EPCGSettingsType::Spatial; }
+#endif
+
+	virtual TArray<FPCGPinProperties> InputPinProperties() const;
+	virtual TArray<FPCGPinProperties> OutputPinProperties() const;
+
+protected:
+	virtual FPCGElementPtr CreateElement() const override;
+	//~End UPCGSettings interface
+
+public:
+	UPROPERTY(BlueprintReadWrite, EditAnywhere, Category = Settings)
+	FVector Extents = FVector::One();
+
+	UPROPERTY(BlueprintReadWrite, EditAnywhere, Category = Settings)
+	EPCGPointExtentsModifierMode Mode = EPCGPointExtentsModifierMode::Set;
+};
+
+class FPCGPointExtentsModifier : public FPCGPointProcessingElementBase
+{
+protected:
+	virtual bool ExecuteInternal(FPCGContext* Context) const;
+};
