[PCG] Constrain the Params pin to only take a single input if they use the code path that returns the first params input (only).

diff --git a/Engine/Plugins/Experimental/PCG/Source/PCG/Private/Elements/PCGCopyPoints.cpp b/Engine/Plugins/Experimental/PCG/Source/PCG/Private/Elements/PCGCopyPoints.cpp
--- a/Engine/Plugins/Experimental/PCG/Source/PCG/Private/Elements/PCGCopyPoints.cpp
+++ b/Engine/Plugins/Experimental/PCG/Source/PCG/Private/Elements/PCGCopyPoints.cpp
@@ -11,7 +11,7 @@
 TArray<FPCGPinProperties> UPCGCopyPointsSettings::InputPinProperties() const
 {
 	TArray<FPCGPinProperties> PinProperties;
-	PinProperties.Emplace(PCGCopyPointsConstants::ParamsLabel, EPCGDataType::Param);
+	PinProperties.Emplace(PCGCopyPointsConstants::ParamsLabel, EPCGDataType::Param, /*bAllowMultipleConnections=*/false);
 	PinProperties.Emplace(PCGCopyPointsConstants::SourcePointsLabel, EPCGDataType::Spatial, /*bAllowMultipleConnections=*/false);
 	PinProperties.Emplace(PCGCopyPointsConstants::TargetPointsLabel, EPCGDataType::Spatial, /*bAllowMultipleConnections=*/false);
 
diff --git a/Engine/Plugins/Experimental/PCG/Source/PCG/Private/Elements/PCGTextureSampler.cpp b/Engine/Plugins/Experimental/PCG/Source/PCG/Private/Elements/PCGTextureSampler.cpp
--- a/Engine/Plugins/Experimental/PCG/Source/PCG/Private/Elements/PCGTextureSampler.cpp
+++ b/Engine/Plugins/Experimental/PCG/Source/PCG/Private/Elements/PCGTextureSampler.cpp
@@ -17,7 +17,7 @@ FPCGElementPtr UPCGTextureSamplerSettings::CreateElement() const
 TArray<FPCGPinProperties> UPCGTextureSamplerSettings::InputPinProperties() const
 {
 	TArray<FPCGPinProperties> PinProperties;
-	PinProperties.Emplace(PCGPinConstants::DefaultParamsLabel, EPCGDataType::Param);
+	PinProperties.Emplace(PCGPinConstants::DefaultParamsLabel, EPCGDataType::Param, /*bInAllowMultipleConnections=*/false);
 
 	return PinProperties;
 }
diff --git a/Engine/Plugins/Experimental/PCG/Source/PCG/Private/Elements/PCGTransformPoints.cpp b/Engine/Plugins/Experimental/PCG/Source/PCG/Private/Elements/PCGTransformPoints.cpp
--- a/Engine/Plugins/Experimental/PCG/Source/PCG/Private/Elements/PCGTransformPoints.cpp
+++ b/Engine/Plugins/Experimental/PCG/Source/PCG/Private/Elements/PCGTransformPoints.cpp
@@ -13,7 +13,7 @@ UPCGTransformPointsSettings::UPCGTransformPointsSettings()
 TArray<FPCGPinProperties> UPCGTransformPointsSettings::InputPinProperties() const
 {
 	TArray<FPCGPinProperties> PinProperties;
-	PinProperties.Emplace(TEXT("Params"), EPCGDataType::Param);
+	PinProperties.Emplace(TEXT("Params"), EPCGDataType::Param, /*bInAllowMultipleConnections=*/false);
 	PinProperties.Emplace(PCGPinConstants::DefaultInputLabel, EPCGDataType::Spatial);
 
 	return PinProperties;
