Makign LogPCG export from the DLL to solve some linking issues if you call APIs that use them.

diff --git a/Engine/Plugins/Experimental/PCG/Source/PCG/Private/PCGModule.cpp b/Engine/Plugins/Experimental/PCG/Source/PCG/Private/PCGModule.cpp
--- a/Engine/Plugins/Experimental/PCG/Source/PCG/Private/PCGModule.cpp
+++ b/Engine/Plugins/Experimental/PCG/Source/PCG/Private/PCGModule.cpp
@@ -69,4 +69,4 @@ void FPCGModule::DeregisterNativeElementDeterminismTests()
 
 IMPLEMENT_MODULE(FPCGModule, PCG);
 
-DEFINE_LOG_CATEGORY(LogPCG);
\ No newline at end of file
+PCG_API DEFINE_LOG_CATEGORY(LogPCG);
\ No newline at end of file
diff --git a/Engine/Plugins/Experimental/PCG/Source/PCG/Public/PCGModule.h b/Engine/Plugins/Experimental/PCG/Source/PCG/Public/PCGModule.h
--- a/Engine/Plugins/Experimental/PCG/Source/PCG/Public/PCGModule.h
+++ b/Engine/Plugins/Experimental/PCG/Source/PCG/Public/PCGModule.h
@@ -7,7 +7,7 @@
 #include "Stats/Stats.h"
 
 // Logs
-DECLARE_LOG_CATEGORY_EXTERN(LogPCG, Log, All);
+PCG_API DECLARE_LOG_CATEGORY_EXTERN(LogPCG, Log, All);
 
 // Stats
 DECLARE_STATS_GROUP(TEXT("PCG"), STATGROUP_PCG, STATCAT_Advanced)
