[PCG] Localized node tooltips. Demonstrated for projection node.

If a tooltip is specified, it will show this, otherwise it will just show the node name. I saw this behaviour on a couple of blueprint nodes so went for this.

Pin tooltips are something different, not covered by this change. I think there is utility in pin tooltips but I did not figure out how to get these stood up in my initial time box and they might take more work. I could look at them in a follow up CL if we want that too.

diff --git a/Engine/Plugins/Experimental/PCG/Source/PCG/Private/Elements/PCGProjectionElement.cpp b/Engine/Plugins/Experimental/PCG/Source/PCG/Private/Elements/PCGProjectionElement.cpp
--- a/Engine/Plugins/Experimental/PCG/Source/PCG/Private/Elements/PCGProjectionElement.cpp
+++ b/Engine/Plugins/Experimental/PCG/Source/PCG/Private/Elements/PCGProjectionElement.cpp
@@ -6,6 +6,8 @@
 #include "PCGEdge.h"
 #include "Helpers/PCGSettingsHelpers.h"
 
+#define LOCTEXT_NAMESPACE "PCGProjectionElement"
+
 TArray<FPCGPinProperties> UPCGProjectionSettings::InputPinProperties() const
 {
 	TArray<FPCGPinProperties> PinProperties;
@@ -16,6 +18,13 @@ TArray<FPCGPinProperties> UPCGProjectionSettings::InputPinProperties() const
 	return PinProperties;
 }
 
+#if WITH_EDITOR
+FText UPCGProjectionSettings::GetNodeTooltipText() const
+{
+	return LOCTEXT("ProjectionNodeTooltip", "Projects each of the inputs connected to In onto the Projection Target and concatenates all of the results to Out. Overrides to the projection settings can be specified using the Params.");
+}
+#endif
+
 FPCGElementPtr UPCGProjectionSettings::CreateElement() const
 {
 	return MakeShared<FPCGProjectionElement>();
@@ -131,3 +140,5 @@ void UPCGProjectionSettings::ApplyDeprecation(UPCGNode* InOutNode)
 	DataVersion = FPCGCustomVersion::LatestVersion;
 }
 #endif // WITH_EDITOR
+
+#undef LOCTEXT_NAMESPACE
diff --git a/Engine/Plugins/Experimental/PCG/Source/PCG/Private/PCGNode.cpp b/Engine/Plugins/Experimental/PCG/Source/PCG/Private/PCGNode.cpp
--- a/Engine/Plugins/Experimental/PCG/Source/PCG/Private/PCGNode.cpp
+++ b/Engine/Plugins/Experimental/PCG/Source/PCG/Private/PCGNode.cpp
@@ -220,6 +220,13 @@ FName UPCGNode::GetNodeTitle() const
 	return TEXT("Unnamed node");
 }
 
+#if WITH_EDITOR
+FText UPCGNode::GetNodeTooltipText() const
+{
+	return DefaultSettings ? DefaultSettings->GetNodeTooltipText() : FText::GetEmpty();
+}
+#endif
+
 TArray<FPCGPinProperties> UPCGNode::InputPinProperties() const
 {
 	TArray<FPCGPinProperties> PinProperties;
diff --git a/Engine/Plugins/Experimental/PCG/Source/PCG/Public/Elements/PCGProjectionElement.h b/Engine/Plugins/Experimental/PCG/Source/PCG/Public/Elements/PCGProjectionElement.h
--- a/Engine/Plugins/Experimental/PCG/Source/PCG/Public/Elements/PCGProjectionElement.h
+++ b/Engine/Plugins/Experimental/PCG/Source/PCG/Public/Elements/PCGProjectionElement.h
@@ -24,6 +24,7 @@ public:
 	//~Begin UPCGSettings interface
 #if WITH_EDITOR
 	virtual FName GetDefaultNodeName() const override { return FName(TEXT("ProjectionNode")); }
+	virtual FText GetNodeTooltipText() const override;
 	virtual EPCGSettingsType GetType() const override { return EPCGSettingsType::Spatial; }
 #endif
 
diff --git a/Engine/Plugins/Experimental/PCG/Source/PCG/Public/PCGNode.h b/Engine/Plugins/Experimental/PCG/Source/PCG/Public/PCGNode.h
--- a/Engine/Plugins/Experimental/PCG/Source/PCG/Public/PCGNode.h
+++ b/Engine/Plugins/Experimental/PCG/Source/PCG/Public/PCGNode.h
@@ -53,6 +53,11 @@ public:
 	/** Returns the node title, based either on the current node label, or defaulted to its settings */
 	FName GetNodeTitle() const;
 
+#if WITH_EDITOR
+	/** Tooltip that describes node functionality and other information. */
+	FText GetNodeTooltipText() const;
+#endif
+
 	/** Returns all the input pin properties */
 	TArray<FPCGPinProperties> InputPinProperties() const;
 
diff --git a/Engine/Plugins/Experimental/PCG/Source/PCG/Public/PCGSettings.h b/Engine/Plugins/Experimental/PCG/Source/PCG/Public/PCGSettings.h
--- a/Engine/Plugins/Experimental/PCG/Source/PCG/Public/PCGSettings.h
+++ b/Engine/Plugins/Experimental/PCG/Source/PCG/Public/PCGSettings.h
@@ -78,6 +78,7 @@ public:
 	virtual void ApplyDeprecation(UPCGNode* InOutNode) {}
 
 	virtual FName GetDefaultNodeName() const { return NAME_None; }
+	virtual FText GetNodeTooltipText() const { return FText::GetEmpty(); }
 	virtual FLinearColor GetNodeTitleColor() const { return FLinearColor::White; }
 	virtual EPCGSettingsType GetType() const { return EPCGSettingsType::Generic; }
 	/** Derived classes must implement this to communicate dependencies on external actors */
diff --git a/Engine/Plugins/Experimental/PCG/Source/PCGEditor/Private/PCGEditorGraphNodeBase.cpp b/Engine/Plugins/Experimental/PCG/Source/PCGEditor/Private/PCGEditorGraphNodeBase.cpp
--- a/Engine/Plugins/Experimental/PCG/Source/PCGEditor/Private/PCGEditorGraphNodeBase.cpp
+++ b/Engine/Plugins/Experimental/PCG/Source/PCGEditor/Private/PCGEditorGraphNodeBase.cpp
@@ -424,10 +424,13 @@ FEdGraphPinType UPCGEditorGraphNodeBase::GetPinType(const UPCGPin* InPin)
 
 FText UPCGEditorGraphNodeBase::GetTooltipText() const
 {
-	return FText::Format(LOCTEXT("NodeTooltip", "{0}\n{1} - Node index {2}"),
-		GetNodeTitle(ENodeTitleType::FullTitle),
-		PCGNode ? FText::FromName(PCGNode->GetFName()) :FText(LOCTEXT("InvalidNodeName", "Unbound node")),
-		PCGNode ? FText::AsNumber(PCGNode->GetGraph()->GetNodes().IndexOfByKey(PCGNode)) : FText(LOCTEXT("InvalidNodeIndex", "Invalid index")));
+	// Either use specified tooltip for description, or fall back to node name if none given.
+	const FText Description = (PCGNode && !PCGNode->GetNodeTooltipText().IsEmpty()) ? PCGNode->GetNodeTooltipText() : GetNodeTitle(ENodeTitleType::FullTitle);
+
+	return FText::Format(LOCTEXT("NodeTooltip", "{0}\n\n{1} - Node index {2}"),
+		Description,
+		PCGNode ? FText::FromName(PCGNode->GetFName()) : LOCTEXT("InvalidNodeName", "Unbound node"),
+		PCGNode ? FText::AsNumber(PCGNode->GetGraph()->GetNodes().IndexOfByKey(PCGNode)) : LOCTEXT("InvalidNodeIndex", "Invalid index"));
 }
 
 void UPCGEditorGraphNodeBase::GetPinHoverText(const UEdGraphPin& Pin, FString& HoverTextOut) const
