[PCG] Fix issue with CopyAttribute when we have a parent attribute.

* From now on, if we copy the attribute and we don't keep the parent, we copy "entries" and "values" from the attribute and all its ancestors.

diff --git a/Engine/Plugins/Experimental/PCG/Source/PCG/Private/Metadata/PCGMetadata.cpp b/Engine/Plugins/Experimental/PCG/Source/PCG/Private/Metadata/PCGMetadata.cpp
--- a/Engine/Plugins/Experimental/PCG/Source/PCG/Private/Metadata/PCGMetadata.cpp
+++ b/Engine/Plugins/Experimental/PCG/Source/PCG/Private/Metadata/PCGMetadata.cpp
@@ -159,7 +159,11 @@ void UPCGMetadata::AddAttribute(const UPCGMetadata* InOther, FName AttributeName
 	}
 
 	CopyAttribute(InOther->GetConstAttribute(AttributeName), AttributeName, /*bKeepParent=*/InOther == Parent, /*bCopyEntries=*/false, /*bCopyValues=*/false);
-	OtherParents.Add(InOther);
+
+	if (InOther != Parent)
+	{
+		OtherParents.Add(InOther);
+	}
 }
 
 void UPCGMetadata::CopyAttributes(const UPCGMetadata* InOther)
diff --git a/Engine/Plugins/Experimental/PCG/Source/PCG/Public/Metadata/PCGMetadataAttributeTpl.h b/Engine/Plugins/Experimental/PCG/Source/PCG/Public/Metadata/PCGMetadataAttributeTpl.h
--- a/Engine/Plugins/Experimental/PCG/Source/PCG/Public/Metadata/PCGMetadataAttributeTpl.h
+++ b/Engine/Plugins/Experimental/PCG/Source/PCG/Public/Metadata/PCGMetadataAttributeTpl.h
@@ -50,23 +50,62 @@ public:
 
 	virtual FPCGMetadataAttributeBase* Copy(FName NewName, UPCGMetadata* InMetadata, bool bKeepParent, bool bCopyEntries = true, bool bCopyValues = true) const override
 	{
-		// this copies to a new attribute
+		// If we copy an attribute where we don't want to keep the parent, while copying entries and/or values, we'll lose data.
+		// In that case, we will copy all the data from this attribute and all its ancestors.
+
+		// We can't keep the parent if we don't have the same root.
 		check(!bKeepParent || PCGMetadataHelpers::HasSameRoot(Metadata, InMetadata));
+		
+		// This copies to a new attribute.
 		FPCGMetadataAttribute<T>* AttributeCopy = new FPCGMetadataAttribute<T>(InMetadata, NewName, bKeepParent ? this : nullptr, DefaultValue, bAllowsInterpolation);
 
+		// Gather the chain of parents if we don't keep the parent and we want to copy entries/values.
+		// We always have at least one item, "this".
+		TArray<const FPCGMetadataAttribute<T>*, TInlineAllocator<2>> Parents = { this };
+		if (!bKeepParent && (bCopyEntries || bCopyValues))
+		{
+			const FPCGMetadataAttribute<T>* Current = this;
+			while(Current->Parent)
+			{
+				Current = static_cast<const FPCGMetadataAttribute<T>*>(Current->Parent);
+				Parents.Add(Current);
+			}
+		}
+
 		if (bCopyEntries)
 		{
-			EntryMapLock.ReadLock();
-			AttributeCopy->EntryToValueKeyMap = EntryToValueKeyMap;
-			EntryMapLock.ReadUnlock();
+			// We go backwards, since we need to preserve order (root -> this)
+			// Latest entry in our Parents array is the root.
+			for (int32 i = Parents.Num() - 1; i >= 0; --i)
+			{
+				const FPCGMetadataAttribute<T>* Current = Parents[i];
+				
+				Current->EntryMapLock.ReadLock();
+				AttributeCopy->EntryToValueKeyMap.Append(Current->EntryToValueKeyMap);
+				Current->EntryMapLock.ReadUnlock();
+			}
 		}
 
 		if (bCopyValues)
 		{
-			ValueLock.ReadLock();
-			AttributeCopy->Values = Values;
-			AttributeCopy->ValueKeyOffset = ValueKeyOffset;
-			ValueLock.ReadUnlock();
+			// We go backwards, since we need to preserve order (root -> this)
+			// Latest entry in our Parents array is the root.
+			for (int32 i = Parents.Num() - 1; i >= 0; --i)
+			{
+				const FPCGMetadataAttribute<T>* Current = Parents[i];
+				
+				Current->ValueLock.ReadLock();
+				AttributeCopy->Values.Append(Current->Values);
+				// The expected value key offset is the one for this attribute (i == 0), and only if we
+				// keep the parent. Otherwise we don't have any parent, so offset should be kept at 0.
+				if (i == 0 && bKeepParent)
+				{
+					AttributeCopy->ValueKeyOffset = Current->ValueKeyOffset;
+				}
+				Current->ValueLock.ReadUnlock();
+			}
+			
+			
 		}
 
 		return AttributeCopy;
@@ -604,4 +643,4 @@ namespace PCGMetadataAttribute
 			return CallbackWithRightType(InAttribute->GetTypeId(), Func);
 		}
 	}
-}
\ No newline at end of file
+}
