[PCG] Added default keybinds to change node execution mode ('E' Enabled, 'D' Debug)

diff --git a/Engine/Plugins/Experimental/PCG/Source/PCGEditor/Private/PCGEditorCommands.cpp b/Engine/Plugins/Experimental/PCG/Source/PCGEditor/Private/PCGEditorCommands.cpp
--- a/Engine/Plugins/Experimental/PCG/Source/PCGEditor/Private/PCGEditorCommands.cpp
+++ b/Engine/Plugins/Experimental/PCG/Source/PCGEditor/Private/PCGEditorCommands.cpp
@@ -27,8 +27,8 @@ void FPCGEditorCommands::RegisterCommands()
 	UI_COMMAND(EditClassDefaults, "Class Defaults", "Edit the initial values of your class.", EUserInterfaceActionType::ToggleButton, FInputChord());
 	UI_COMMAND(StartInspectNode, "Start Inspecting Node", "Start Inspecting Node", EUserInterfaceActionType::Button, FInputChord(EKeys::I));
 	UI_COMMAND(StopInspectNode, "Stop Inspecting Node", "Stop Inspecting Node", EUserInterfaceActionType::Button, FInputChord(EModifierKey::Shift, EKeys::I));
-	UI_COMMAND(ExecutionModeEnabled, "Execution Mode Enabled", "Set the execution mode for this node to Enabled.", EUserInterfaceActionType::RadioButton, FInputChord());
-	UI_COMMAND(ExecutionModeDebug, "Execution Mode Debug", "Set the execution mode for this node to Debug.", EUserInterfaceActionType::RadioButton, FInputChord());
+	UI_COMMAND(ExecutionModeEnabled, "Execution Mode Enabled", "Set the execution mode for this node to Enabled.", EUserInterfaceActionType::RadioButton, FInputChord(EKeys::E));
+	UI_COMMAND(ExecutionModeDebug, "Execution Mode Debug", "Set the execution mode for this node to Debug.", EUserInterfaceActionType::RadioButton, FInputChord(EKeys::D));
 	UI_COMMAND(ExecutionModeIsolated, "Execution Mode Isolated", "Set the execution mode for this node to Isolated.", EUserInterfaceActionType::RadioButton, FInputChord());
 	UI_COMMAND(ExecutionModeDisabled, "Execution Mode Disabled", "Set the execution mode for this node to Disabled.", EUserInterfaceActionType::RadioButton, FInputChord());
 	UI_COMMAND(CancelExecution, "Cancel Execution", "Cancels the execution of the current graph", EUserInterfaceActionType::Button, FInputChord(EModifierKey::Control, EKeys::Escape));
