Added world queries (volume/ray)
Refactored the surface sampler in a separate namespace so we can reuse it for the world ray hits

diff --git a/Engine/Plugins/Experimental/PCG/Source/PCG/Private/Data/PCGWorldData.cpp b/Engine/Plugins/Experimental/PCG/Source/PCG/Private/Data/PCGWorldData.cpp
new file mode 100644
--- /dev/null
+++ b/Engine/Plugins/Experimental/PCG/Source/PCG/Private/Data/PCGWorldData.cpp
@@ -0,0 +1,235 @@
+// Copyright Epic Games, Inc. All Rights Reserved.
+
+#include "Data/PCGWorldData.h"
+
+#include "PCGComponent.h"
+#include "PCGHelpers.h"
+#include "Data/PCGPointData.h"
+#include "Elements/PCGSurfaceSampler.h"
+#include "Elements/PCGVolumeSampler.h"
+#include "Helpers/PCGBlueprintHelpers.h"
+
+#include "Components/BrushComponent.h"
+#include "Engine/World.h"
+
+void UPCGWorldVolumetricData::Initialize(UWorld* InWorld, const FBox& InBounds)
+{
+	Super::Initialize(InBounds, nullptr);
+	TargetActor = nullptr;
+	World = InWorld;
+	check(World.IsValid());
+}
+
+bool UPCGWorldVolumetricData::SamplePoint(const FTransform& InTransform, const FBox& InBounds, FPCGPoint& OutPoint, UPCGMetadata* OutMetadata) const
+{
+	check(World.IsValid());
+
+	FCollisionObjectQueryParams ObjectQueryParams(QueryParams.CollisionChannel);
+	FCollisionShape CollisionShape = FCollisionShape::MakeBox(InBounds.GetExtent() * InTransform.GetScale3D());
+	FCollisionQueryParams Params; // TODO: apply properties from the settings when/if they exist
+
+	TArray<FOverlapResult> Overlaps;
+	/*bool bOverlaps =*/ World->OverlapMultiByObjectType(Overlaps, InTransform.TransformPosition(InBounds.GetCenter()), InTransform.GetRotation(), ObjectQueryParams, CollisionShape, Params);
+
+	for (const FOverlapResult& Overlap : Overlaps)
+	{
+		// Skip invisible walls / triggers / volumes
+		const UPrimitiveComponent* OverlappedComponent = Overlap.GetComponent();
+		if (OverlappedComponent->IsA<UBrushComponent>())
+		{
+			continue;
+		}
+
+		// Skip "no collision" type actors
+		if (!OverlappedComponent->IsQueryCollisionEnabled() || OverlappedComponent->GetCollisionResponseToChannel(QueryParams.CollisionChannel) != ECR_Block)
+		{
+			continue;
+		}
+
+		// Skip PCG-created objects optionally
+		if (QueryParams.bIgnorePCGHits && OverlappedComponent->ComponentTags.Contains(PCGHelpers::DefaultPCGTag))
+		{
+			continue;
+		}
+
+		// Skip self-generated PCG objects optionally
+		if (QueryParams.bIgnoreSelfHits && OriginatingComponent.IsValid() && OverlappedComponent->ComponentTags.Contains(OriginatingComponent->GetFName()))
+		{
+			continue;
+		}
+
+		// TODO: additional filtering?
+
+		if (QueryParams.bSearchForOverlap)
+		{
+			OutPoint = FPCGPoint(InTransform, 1.0f, 0);
+			UPCGBlueprintHelpers::SetSeedFromPosition(OutPoint);
+			OutPoint.SetLocalBounds(InBounds);
+			return true;
+		}
+		else
+		{
+			return false;
+		}
+	}
+
+	// No valid hits found
+	if (!QueryParams.bSearchForOverlap)
+	{
+		OutPoint = FPCGPoint(InTransform, 1.0f, 0);
+		UPCGBlueprintHelpers::SetSeedFromPosition(OutPoint);
+		OutPoint.SetLocalBounds(InBounds);
+		return true;
+	}
+	else
+	{
+		return false;
+	}
+}
+
+const UPCGPointData* UPCGWorldVolumetricData::CreatePointData(FPCGContext* Context, const FBox& InBounds) const
+{
+	TRACE_CPUPROFILER_EVENT_SCOPE(UPCGWorldVolumetricData::CreatePointData);
+
+	UPCGPointData* Data = NewObject<UPCGPointData>();
+	Data->InitializeFromData(this);
+	TArray<FPCGPoint>& Points = Data->GetMutablePoints();
+
+	FBox EffectiveBounds = Bounds;
+	if (InBounds.IsValid)
+	{
+		if (Bounds.IsValid)
+		{
+			EffectiveBounds = Bounds.Overlap(InBounds);
+		}
+		else
+		{
+			EffectiveBounds = InBounds;
+		}
+	}
+	
+	// Early out
+	if (!EffectiveBounds.IsValid)
+	{
+		if (!Bounds.IsValid && !InBounds.IsValid)
+		{
+			UE_LOG(LogPCG, Error, TEXT("PCG World Volumetric Data cannot generate if there are no framing bounds"));
+		}
+		
+		return Data;
+	}
+	
+	PCGVolumeSampler::FVolumeSamplerSettings SamplerSettings;
+	SamplerSettings.VoxelSize = VoxelSize;
+
+	PCGVolumeSampler::SampleVolume(Context, this, SamplerSettings, Data, EffectiveBounds);
+	UE_LOG(LogPCG, Verbose, TEXT("Volumetric world extracted %d points"), Data->GetPoints().Num());
+	
+	return Data;
+}
+
+/** World Ray Hit data implementation */
+void UPCGWorldRayHitData::Initialize(UWorld* InWorld, const FBox& InBounds)
+{
+	World = InWorld;
+	Bounds = InBounds;
+}
+
+bool UPCGWorldRayHitData::SamplePoint(const FTransform& InTransform, const FBox& InBounds, FPCGPoint& OutPoint, UPCGMetadata* OutMetadata) const
+{
+	check(World.IsValid());
+
+	// Todo: consider prebuilding this
+	FCollisionObjectQueryParams ObjectQueryParams(QueryParams.CollisionChannel);
+	FCollisionQueryParams Params; // TODO: apply properties from the settings when/if they exist
+
+	// Project the InTransform location on the ray origin plane
+	const FVector PointLocation = InTransform.GetLocation();
+	FVector RayStart = PointLocation - ((PointLocation - QueryParams.RayOrigin) | QueryParams.RayDirection) * QueryParams.RayDirection;
+	FVector RayEnd = RayStart + QueryParams.RayDirection * QueryParams.RayLength;
+
+	TArray<FHitResult> Hits;
+	World->LineTraceMultiByObjectType(Hits, RayStart, RayEnd, ObjectQueryParams, Params);
+
+	for (const FHitResult& Hit : Hits)
+	{
+		// Skip invisible walls / triggers / volumes
+		const UPrimitiveComponent* HitComponent = Hit.GetComponent();
+		if (HitComponent->IsA<UBrushComponent>())
+		{
+			continue;
+		}
+
+		// Skip "No collision" type actors
+		if (!HitComponent->IsQueryCollisionEnabled() || HitComponent->GetCollisionResponseToChannel(QueryParams.CollisionChannel) != ECR_Block)
+		{
+			continue;
+		}
+		
+		// Skip PCG-created objects optionally
+		if (QueryParams.bIgnorePCGHits && HitComponent->ComponentTags.Contains(PCGHelpers::DefaultPCGTag))
+		{
+			continue;
+		}
+
+		// Skip self-generated PCG objects optionally
+		if (QueryParams.bIgnoreSelfHits && OriginatingComponent.IsValid() && HitComponent->ComponentTags.Contains(OriginatingComponent->GetFName()))
+		{
+			continue;
+		}
+
+		// TODO: additional filtering?
+
+		// Finally, fill in OutPoint - we're done
+		// TODO: use normal to orient point?
+		OutPoint = FPCGPoint(FTransform(Hit.ImpactNormal.Rotation(), Hit.ImpactPoint), 1.0f, 0);
+		UPCGBlueprintHelpers::SetSeedFromPosition(OutPoint);
+
+		// TODO: apply metadata
+		// e.g. if hit on landscape -> metadata on landscape at that position
+		// maybe material, etc.
+
+		return true;
+	}
+
+	return false;
+}
+
+const UPCGPointData* UPCGWorldRayHitData::CreatePointData(FPCGContext* Context, const FBox& InBounds) const
+{
+	UPCGPointData* Data = NewObject<UPCGPointData>();
+	Data->InitializeFromData(this);
+	TArray<FPCGPoint>& Points = Data->GetMutablePoints();
+
+	FBox EffectiveBounds = Bounds;
+	if (InBounds.IsValid)
+	{
+		if (Bounds.IsValid)
+		{
+			EffectiveBounds = Bounds.Overlap(InBounds);
+		}
+		else
+		{
+			EffectiveBounds = InBounds;
+		}
+	}
+
+	// Early out
+	if (!EffectiveBounds.IsValid)
+	{
+		if (!Bounds.IsValid && !InBounds.IsValid)
+		{
+			UE_LOG(LogPCG, Error, TEXT("PCG World Volumetric Data cannot generate if there are no framing bounds"));
+		}
+
+		return Data;
+	}
+
+	PCGSurfaceSampler::FSurfaceSamplerSettings SamplerSettings;
+	if (SamplerSettings.Initialize(nullptr, Context, EffectiveBounds))
+	{
+		PCGSurfaceSampler::SampleSurface(Context, this, SamplerSettings, Data);
+	}
+
+	return Data;
+}
\ No newline at end of file
diff --git a/Engine/Plugins/Experimental/PCG/Source/PCG/Private/Elements/PCGSurfaceSampler.cpp b/Engine/Plugins/Experimental/PCG/Source/PCG/Private/Elements/PCGSurfaceSampler.cpp
--- a/Engine/Plugins/Experimental/PCG/Source/PCG/Private/Elements/PCGSurfaceSampler.cpp
+++ b/Engine/Plugins/Experimental/PCG/Source/PCG/Private/Elements/PCGSurfaceSampler.cpp
@@ -1,57 +1,42 @@
 // Copyright Epic Games, Inc. All Rights Reserved.
 
 #include "Elements/PCGSurfaceSampler.h"
+
+#include "PCGComponent.h"
+#include "PCGHelpers.h"
 #include "Data/PCGPointData.h"
 #include "Data/PCGSpatialData.h"
 #include "Helpers/PCGAsync.h"
 #include "Helpers/PCGSettingsHelpers.h"
-#include "PCGHelpers.h"
+
 #include "Math/RandomStream.h"
 
-struct FPCGSurfaceSamplerLoopData
+namespace PCGSurfaceSampler
 {
-	const UPCGSurfaceSamplerSettings* Settings;
-
-	float PointsPerSquaredMeter;
-	FVector PointExtents;
-	float Looseness;
-	bool bApplyDensityToPoints;
-	float PointSteepness;
-	bool bKeepZeroDensityPoints = false;
-
-	FVector InterstitialDistance;
-	FVector InnerCellSize;
-	FVector CellSize;
-
-	int32 CellMinX;
-	int32 CellMaxX;
-	int32 CellMinY;
-	int32 CellMaxY;
-
-	int32 CellCount;
-
-	int64 TargetPointCount;
-
-	float Ratio;
-
-	int Seed;
-
-	bool Initialize(const UPCGSurfaceSamplerSettings* InSettings, FPCGContext* Context, const FBox& InputBounds)
+	bool FSurfaceSamplerSettings::Initialize(const UPCGSurfaceSamplerSettings* InSettings, FPCGContext* Context, const FBox& InputBounds)
 	{
-		UPCGParamData* Params = Context->InputData.GetParams();
 		Settings = InSettings;
 
-		// Compute used values
-		PointsPerSquaredMeter = PCGSettingsHelpers::GetValue(GET_MEMBER_NAME_CHECKED(UPCGSurfaceSamplerSettings, PointsPerSquaredMeter), Settings->PointsPerSquaredMeter, Params);
-		PointExtents = PCGSettingsHelpers::GetValue(GET_MEMBER_NAME_CHECKED(UPCGSurfaceSamplerSettings, PointExtents), Settings->PointExtents, Params);
-		Looseness = PCGSettingsHelpers::GetValue(GET_MEMBER_NAME_CHECKED(UPCGSurfaceSamplerSettings, Looseness), Settings->Looseness, Params);
-		bApplyDensityToPoints = PCGSettingsHelpers::GetValue(GET_MEMBER_NAME_CHECKED(UPCGSurfaceSamplerSettings, bApplyDensityToPoints), Settings->bApplyDensityToPoints, Params);
-		PointSteepness = PCGSettingsHelpers::GetValue(GET_MEMBER_NAME_CHECKED(UPCGSurfaceSamplerSettings, PointSteepness), Settings->PointSteepness, Params);
+		if (Settings)
+		{
+			UPCGParamData* Params = Context ? Context->InputData.GetParams() : nullptr;
+			
+			// Compute used values
+			PointsPerSquaredMeter = PCGSettingsHelpers::GetValue(GET_MEMBER_NAME_CHECKED(UPCGSurfaceSamplerSettings, PointsPerSquaredMeter), Settings->PointsPerSquaredMeter, Params);
+			PointExtents = PCGSettingsHelpers::GetValue(GET_MEMBER_NAME_CHECKED(UPCGSurfaceSamplerSettings, PointExtents), Settings->PointExtents, Params);
+			Looseness = PCGSettingsHelpers::GetValue(GET_MEMBER_NAME_CHECKED(UPCGSurfaceSamplerSettings, Looseness), Settings->Looseness, Params);
+			bApplyDensityToPoints = PCGSettingsHelpers::GetValue(GET_MEMBER_NAME_CHECKED(UPCGSurfaceSamplerSettings, bApplyDensityToPoints), Settings->bApplyDensityToPoints, Params);
+			PointSteepness = PCGSettingsHelpers::GetValue(GET_MEMBER_NAME_CHECKED(UPCGSurfaceSamplerSettings, PointSteepness), Settings->PointSteepness, Params);
 #if WITH_EDITORONLY_DATA
-		bKeepZeroDensityPoints = PCGSettingsHelpers::GetValue(GET_MEMBER_NAME_CHECKED(UPCGSurfaceSamplerSettings, bKeepZeroDensityPoints), Settings->bKeepZeroDensityPoints, Params);
+			bKeepZeroDensityPoints = PCGSettingsHelpers::GetValue(GET_MEMBER_NAME_CHECKED(UPCGSurfaceSamplerSettings, bKeepZeroDensityPoints), Settings->bKeepZeroDensityPoints, Params);
 #endif
 
-		Seed = PCGSettingsHelpers::ComputeSeedWithOverride(InSettings, Context->SourceComponent, Params);
+			Seed = PCGSettingsHelpers::ComputeSeedWithOverride(InSettings, Context->SourceComponent, Params);
+		}
+		else
+		{
+			Seed = (Context && Context->SourceComponent.IsValid()) ? Context->SourceComponent->Seed : 42;
+		}
 
 		// Conceptually, we will break down the surface bounds in a N x M grid
 		InterstitialDistance = PointExtents * 2;
@@ -92,7 +77,7 @@ struct FPCGSurfaceSamplerLoopData
 		return true;
 	}
 
-	void ComputeCellIndices(int32 Index, int32& CellX, int32& CellY)
+	void FSurfaceSamplerSettings::ComputeCellIndices(int32 Index, int32& CellX, int32& CellY) const
 	{
 		check(Index >= 0 && Index < CellCount);
 		const int32 CellCountX = 1 + CellMaxX - CellMinX;
@@ -100,7 +85,68 @@ struct FPCGSurfaceSamplerLoopData
 		CellX = CellMinX + (Index % CellCountX);
 		CellY = CellMinY + (Index / CellCountX);
 	}
-};
+
+	UPCGPointData* SampleSurface(FPCGContext* Context, const UPCGSpatialData* SpatialInput, const FSurfaceSamplerSettings& LoopData)
+	{
+		UPCGPointData* SampledData = NewObject<UPCGPointData>();
+		SampledData->InitializeFromData(SpatialInput);
+
+		SampleSurface(Context, SpatialInput, LoopData, SampledData);
+
+		return SampledData;
+	}
+
+	void SampleSurface(FPCGContext* Context, const UPCGSpatialData* SpatialInput, const FSurfaceSamplerSettings& LoopData, UPCGPointData* SampledData)
+	{
+		TArray<FPCGPoint>& SampledPoints = SampledData->GetMutablePoints();
+
+		FPCGAsync::AsyncPointProcessing(Context, LoopData.CellCount, SampledPoints, [&LoopData, SampledData, SpatialInput](int32 Index, FPCGPoint& OutPoint)
+		{
+			int32 CellX;
+			int32 CellY;
+			LoopData.ComputeCellIndices(Index, CellX, CellY);
+
+			const FVector::FReal CurrentX = CellX * LoopData.CellSize.X;
+			const FVector::FReal CurrentY = CellY * LoopData.CellSize.Y;
+			const FVector InnerCellSize = LoopData.InnerCellSize;
+
+			FRandomStream RandomSource(PCGHelpers::ComputeSeed(LoopData.Seed, CellX, CellY));
+			float Chance = RandomSource.FRand();
+
+			const float Ratio = LoopData.Ratio;
+
+			if (Chance < Ratio)
+			{
+				const float RandX = RandomSource.FRand();
+				const float RandY = RandomSource.FRand();
+				 
+				const FVector TentativeLocation = FVector(CurrentX + RandX * InnerCellSize.X, CurrentY + RandY * InnerCellSize.Y, 0.0f);
+				const FBox LocalBound(-LoopData.PointExtents, LoopData.PointExtents);
+
+				if (SpatialInput->SamplePoint(FTransform(TentativeLocation), LocalBound, OutPoint, SampledData->Metadata) || LoopData.bKeepZeroDensityPoints)
+				{
+					// Apply final parameters on the point
+					OutPoint.SetExtents(LoopData.PointExtents);
+					OutPoint.Density *= (LoopData.bApplyDensityToPoints ? ((Ratio - Chance) / Ratio) : 1.0f);
+					OutPoint.Steepness = LoopData.PointSteepness;
+					OutPoint.Seed = RandomSource.GetCurrentSeed();
+
+					return true;
+				}
+				else
+				{
+					return false;
+				}
+			}
+			else
+			{
+				return false;
+			}
+		});
+
+		PCGE_LOG_C(Verbose, Context, "Generated %d points in %d cells", SampledPoints.Num(), LoopData.CellCount);
+	}
+}
 
 UPCGSurfaceSamplerSettings::UPCGSurfaceSamplerSettings()
 {
@@ -163,7 +209,7 @@ bool FPCGSurfaceSamplerElement::ExecuteInternal(FPCGContext* Context) const
 			continue;
 		}
 
-		FPCGSurfaceSamplerLoopData LoopData;
+		PCGSurfaceSampler::FSurfaceSamplerSettings LoopData;
 		if (!LoopData.Initialize(Settings, Context, InputBounds))
 		{
 			continue;
@@ -172,58 +218,7 @@ bool FPCGSurfaceSamplerElement::ExecuteInternal(FPCGContext* Context) const
 		// Finally, create data
 		FPCGTaggedData& Output = Outputs.Emplace_GetRef();
 		Output = Input;
-
-		UPCGPointData* SampledData = NewObject<UPCGPointData>();
-		SampledData->InitializeFromData(SpatialInput);
-		Output.Data = SampledData;
-
-		TArray<FPCGPoint>& SampledPoints = SampledData->GetMutablePoints();
-
-		FPCGAsync::AsyncPointProcessing(Context, LoopData.CellCount, SampledPoints, [&LoopData, SampledData, SpatialInput](int32 Index, FPCGPoint& OutPoint)
-		{
-			int32 CellX;
-			int32 CellY;
-			LoopData.ComputeCellIndices(Index, CellX, CellY);
-
-			const FVector::FReal CurrentX = CellX * LoopData.CellSize.X;
-			const FVector::FReal CurrentY = CellY * LoopData.CellSize.Y;
-			const FVector InnerCellSize = LoopData.InnerCellSize;
-
-			FRandomStream RandomSource(PCGHelpers::ComputeSeed(LoopData.Seed, CellX, CellY));
-			float Chance = RandomSource.FRand();
-
-			const float Ratio = LoopData.Ratio;
-
-			if (Chance < Ratio)
-			{
-				const float RandX = RandomSource.FRand();
-				const float RandY = RandomSource.FRand();
-				 
-				const FVector TentativeLocation = FVector(CurrentX + RandX * InnerCellSize.X, CurrentY + RandY * InnerCellSize.Y, 0.0f);
-				const FBox LocalBound(-LoopData.PointExtents, LoopData.PointExtents);
-
-				if (SpatialInput->SamplePoint(FTransform(TentativeLocation), LocalBound, OutPoint, SampledData->Metadata) || LoopData.bKeepZeroDensityPoints)
-				{
-					// Apply final parameters on the point
-					OutPoint.SetExtents(LoopData.PointExtents);
-					OutPoint.Density *= (LoopData.bApplyDensityToPoints ? ((Ratio - Chance) / Ratio) : 1.0f);
-					OutPoint.Steepness = LoopData.PointSteepness;
-					OutPoint.Seed = RandomSource.GetCurrentSeed();
-
-					return true;
-				}
-				else
-				{
-					return false;
-				}
-			}
-			else
-			{
-				return false;
-			}
-		});
-
-		PCGE_LOG(Verbose, "Generated %d points in %d cells", SampledPoints.Num(), LoopData.CellCount);
+		Output.Data = PCGSurfaceSampler::SampleSurface(Context, SpatialInput, LoopData);
 	}
 
 	// Finally, forward any exclusions/settings
diff --git a/Engine/Plugins/Experimental/PCG/Source/PCG/Private/Elements/PCGVolumeSampler.cpp b/Engine/Plugins/Experimental/PCG/Source/PCG/Private/Elements/PCGVolumeSampler.cpp
--- a/Engine/Plugins/Experimental/PCG/Source/PCG/Private/Elements/PCGVolumeSampler.cpp
+++ b/Engine/Plugins/Experimental/PCG/Source/PCG/Private/Elements/PCGVolumeSampler.cpp
@@ -13,10 +13,25 @@ namespace PCGVolumeSampler
 	{
 		UPCGPointData* Data = NewObject<UPCGPointData>();
 		Data->InitializeFromData(SpatialData);
-		TArray<FPCGPoint>& Points = Data->GetMutablePoints();
+		const FBox Bounds = SpatialData->GetBounds();
+
+		SampleVolume(Context, SpatialData, SamplerSettings, Data, Bounds);
+
+		return Data;
+	}
+
+	void SampleVolume(FPCGContext* Context, const UPCGSpatialData* SpatialData, const FVolumeSamplerSettings& SamplerSettings, UPCGPointData* OutputData, const FBox& Bounds)
+	{
+		check(SpatialData && OutputData);
 
+		// Early out
+		if (!Bounds.IsValid)
+		{
+			return;
+		}
+
+		TArray<FPCGPoint>& Points = OutputData->GetMutablePoints();
 		const FVector& VoxelSize = SamplerSettings.VoxelSize;
-		const FBox Bounds = SpatialData->GetBounds();
 
 		const int32 MinX = FMath::CeilToInt(Bounds.Min.X / VoxelSize.X);
 		const int32 MaxX = FMath::FloorToInt(Bounds.Max.X / VoxelSize.X);
@@ -45,8 +60,6 @@ namespace PCGVolumeSampler
 				return false;
 			}
 		});
-
-		return Data;
 	}
 }
 
diff --git a/Engine/Plugins/Experimental/PCG/Source/PCG/Private/Elements/PCGWorldQuery.cpp b/Engine/Plugins/Experimental/PCG/Source/PCG/Private/Elements/PCGWorldQuery.cpp
new file mode 100644
--- /dev/null
+++ b/Engine/Plugins/Experimental/PCG/Source/PCG/Private/Elements/PCGWorldQuery.cpp
@@ -0,0 +1,111 @@
+// Copyright Epic Games, Inc. All Rights Reserved.
+
+#include "Elements/PCGWorldQuery.h"
+
+#include "PCGComponent.h"
+#include "PCGHelpers.h"
+
+#include "GameFramework/Actor.h"
+
+FPCGElementPtr UPCGWorldQuerySettings::CreateElement() const
+{
+	return MakeShared<FPCGWorldVolumetricQueryElement>();
+}
+
+bool FPCGWorldVolumetricQueryElement::ExecuteInternal(FPCGContext* Context) const
+{
+	TRACE_CPUPROFILER_EVENT_SCOPE(FPCGWorldVolumetricQueryElement::Execute);
+
+	const UPCGWorldQuerySettings* Settings = Context->GetInputSettings<UPCGWorldQuerySettings>();
+	check(Settings);
+
+	FPCGWorldVolumetricQueryParams QueryParams = Settings->QueryParams;
+	
+	// TODO: Add params pin + Apply param data overrides
+
+	check(Context->SourceComponent.IsValid());
+	UWorld* World = Context->SourceComponent->GetWorld();
+	// TODO: Not strictly required but will be useful until we change the usage of this
+	FBox SourceBounds = PCGHelpers::GetActorBounds(Context->SourceComponent->GetOwner());
+
+	UPCGWorldVolumetricData* Data = NewObject<UPCGWorldVolumetricData>();
+	Data->Initialize(World, SourceBounds);
+	Data->QueryParams = QueryParams;
+	Data->OriginatingComponent = Context->SourceComponent;
+	Data->TargetActor = Context->SourceComponent->GetOwner();
+
+	FPCGTaggedData& Output = Context->OutputData.TaggedData.Emplace_GetRef();
+	Output.Data = Data;
+
+	// Pass-through settings & exclusions
+	Context->OutputData.TaggedData.Append(Context->InputData.GetAllSettings());
+
+	return true;
+}
+
+FPCGElementPtr UPCGWorldRayHitSettings::CreateElement() const
+{
+	return MakeShared<FPCGWorldRayHitQueryElement>();
+}
+
+bool FPCGWorldRayHitQueryElement::ExecuteInternal(FPCGContext* Context) const
+{
+	TRACE_CPUPROFILER_EVENT_SCOPE(FPCGWorldRayHitQueryElement::Execute);
+
+	const UPCGWorldRayHitSettings* Settings = Context->GetInputSettings<UPCGWorldRayHitSettings>();
+	check(Settings);
+
+	FPCGWorldRayHitQueryParams QueryParams = Settings->QueryParams;
+	
+	// TODO: Support params pin + Apply param data
+
+	// Compute default parameters based on owner component - raycast down local Z axis
+	if (!QueryParams.bOverrideDefaultParams)
+	{
+		AActor* Owner = Context->SourceComponent->GetOwner();
+		const FTransform& Transform = Owner->GetTransform();
+
+		const FBox LocalBounds = PCGHelpers::GetActorLocalBounds(Owner);
+		const FVector RayOrigin = Transform.TransformPosition(FVector(0, 0, LocalBounds.Max.Z));
+		const FVector RayEnd = Transform.TransformPosition(FVector(0, 0, LocalBounds.Min.Z));
+
+		const FVector::FReal RayLength = (RayEnd - RayOrigin).Length();
+		const FVector RayDirection = (RayLength > UE_SMALL_NUMBER ? (RayEnd - RayOrigin) / RayLength : FVector(0, 0, -1.0));
+
+		QueryParams.RayOrigin = RayOrigin;
+		QueryParams.RayDirection = RayDirection;
+		QueryParams.RayLength = RayLength;
+	}
+	else // user provided ray parameters
+	{
+		const FVector::FReal RayDirectionLength = QueryParams.RayDirection.Length();
+		if (RayDirectionLength > UE_SMALL_NUMBER)
+		{
+			QueryParams.RayDirection = QueryParams.RayDirection / RayDirectionLength;
+			QueryParams.RayLength *= RayDirectionLength;
+		}
+		else
+		{
+			QueryParams.RayDirection = FVector(0, 0, -1.0);
+		}
+	}
+
+	check(Context->SourceComponent.IsValid());
+	UWorld* World = Context->SourceComponent->GetWorld();
+	// TODO: Not strictly required but will be useful until we change the usage of this
+	FBox SourceBounds = PCGHelpers::GetActorBounds(Context->SourceComponent->GetOwner());
+
+	UPCGWorldRayHitData* Data = NewObject<UPCGWorldRayHitData>();
+	Data->Initialize(World, SourceBounds);
+	Data->QueryParams = QueryParams;
+	Data->OriginatingComponent = Context->SourceComponent;
+	Data->TargetActor = Context->SourceComponent->GetOwner();
+
+	FPCGTaggedData& Output = Context->OutputData.TaggedData.Emplace_GetRef();
+	Output.Data = Data;
+
+	// Pass-through settings & exclusions
+	Context->OutputData.TaggedData.Append(Context->InputData.GetAllSettings());
+
+	return true;
+}
\ No newline at end of file
diff --git a/Engine/Plugins/Experimental/PCG/Source/PCG/Public/Data/PCGWorldData.h b/Engine/Plugins/Experimental/PCG/Source/PCG/Public/Data/PCGWorldData.h
new file mode 100644
--- /dev/null
+++ b/Engine/Plugins/Experimental/PCG/Source/PCG/Public/Data/PCGWorldData.h
@@ -0,0 +1,131 @@
+// Copyright Epic Games, Inc. All Rights Reserved.
+
+#pragma once
+
+#include "PCGSpatialData.h"
+#include "PCGVolumeData.h"
+#include "PCGSurfaceData.h"
+
+#include "Engine/EngineTypes.h"
+
+#include "PCGWorldData.generated.h"
+
+class UWorld;
+class UPCGMetadata;
+class UPCGComponent;
+
+USTRUCT(BlueprintType)
+struct FPCGWorldVolumetricQueryParams
+{
+	GENERATED_BODY()
+
+	UPROPERTY(BlueprintReadWrite, EditAnywhere, Category = Data)
+	bool bSearchForOverlap = true;
+
+	UPROPERTY(BlueprintReadWrite, EditAnywhere, Category = Data)
+	bool bIgnorePCGHits = false;
+
+	UPROPERTY(BlueprintReadWrite, EditAnywhere, Category = Data)
+	bool bIgnoreSelfHits = true;
+
+	UPROPERTY(BlueprintReadWrite, EditAnywhere, Category = "Data|Advanced")
+	TEnumAsByte<ECollisionChannel> CollisionChannel = ECC_WorldStatic;
+};
+
+UCLASS(BlueprintType, ClassGroup = (Procedural))
+class PCG_API UPCGWorldVolumetricData : public UPCGVolumeData
+{
+	GENERATED_BODY()
+
+public:
+	void Initialize(UWorld* InWorld, const FBox& InBounds = FBox(EForceInit::ForceInit));
+
+	//~Begin UPCGSpatialData interface
+	virtual bool SamplePoint(const FTransform& Transform, const FBox& Bounds, FPCGPoint& OutPoint, UPCGMetadata* OutMetadata) const override;
+	//~End UPCGSpatialData interface
+
+	//~Begin UPCGSpatialDataWithPointCache
+	virtual bool SupportsBoundedPointData() const { return true; }
+	virtual const UPCGPointData* CreatePointData(FPCGContext* Context) const override { return CreatePointData(Context, FBox(EForceInit::ForceInit)); }
+	virtual const UPCGPointData* CreatePointData(FPCGContext* Context, const FBox& InBounds) const override;
+	//~End UPCGSpatialDataWithPointCache
+
+	UPROPERTY()
+	TWeakObjectPtr<UWorld> World = nullptr;
+
+	UPROPERTY()
+	TWeakObjectPtr<UPCGComponent> OriginatingComponent = nullptr;
+
+	UPROPERTY(BlueprintReadWrite, EditAnywhere, Category = Data, meta = (ShowOnlyInnerProperties))
+	FPCGWorldVolumetricQueryParams QueryParams;
+};
+
+USTRUCT(BlueprintType)
+struct FPCGWorldRayHitQueryParams
+{
+	GENERATED_BODY()
+
+	UPROPERTY(BlueprintReadWrite, EditAnywhere, Category = Data)
+	bool bOverrideDefaultParams = false;
+
+	UPROPERTY(BlueprintReadWrite, EditAnywhere, Category = Data, meta = (EditCondition = "bOverrideDefaultParams"))
+	FVector RayOrigin = FVector::ZeroVector;
+
+	UPROPERTY(BlueprintReadWrite, EditAnywhere, Category = Data, meta = (EditCondition = "bOverrideDefaultParams"))
+	FVector RayDirection = FVector(0.0, 0.0, -1.0);
+
+	UPROPERTY(BlueprintReadWrite, EditAnywhere, Category = Data, meta = (EditCondition = "bOverrideDefaultParams"))
+	double RayLength = 1.0e+5; // 100m
+
+	UPROPERTY(BlueprintReadWrite, EditAnywhere, Category = Data)
+	bool bIgnorePCGHits = false;
+
+	UPROPERTY(BlueprintReadWrite, EditAnywhere, Category = Data)
+	bool bIgnoreSelfHits = true;
+
+	UPROPERTY(BlueprintReadWrite, EditAnywhere, Category = "Data|Advanced")
+	TEnumAsByte<ECollisionChannel> CollisionChannel = ECC_WorldStatic;
+
+	// TODO: see in FCollisionQueryParams if there are some flags we want to expose
+	// examples: bTraceComplex, bReturnFaceIndex, bReturnPhysicalMaterial, some ignore patterns
+
+	//TODO UPROPERTY()
+	//bool bUseMetadataFromLandscape = true;
+};
+
+UCLASS(BlueprintType, ClassGroup = (Procedural))
+class PCG_API UPCGWorldRayHitData : public UPCGSurfaceData
+{
+	GENERATED_BODY()
+
+public:
+	void Initialize(UWorld* InWorld, const FBox& InBounds = FBox(EForceInit::ForceInit));
+
+	//~Begin UPCGData interface
+	virtual EPCGDataType GetDataType() const override { return Super::GetDataType(); }
+	//~End UPCGData interface
+
+	//~Begin UPCGSpatialData interface
+	virtual FBox GetBounds() const override { return Bounds; }
+	virtual FBox GetStrictBounds() const override { return Bounds; }
+	virtual bool SamplePoint(const FTransform& Transform, const FBox& Bounds, FPCGPoint& OutPoint, UPCGMetadata* OutMetadata) const override;
+	virtual bool HasNonTrivialTransform() const override { return true; }
+
+	// ~Begin UPCGSpatialDataWithPointCache interface
+	virtual bool SupportsBoundedPointData() const { return true; }
+	virtual const UPCGPointData* CreatePointData(FPCGContext* Context) const override { return CreatePointData(Context, FBox(EForceInit::ForceInit)); }
+	virtual const UPCGPointData* CreatePointData(FPCGContext* Context, const FBox& InBounds) const override;
+	// ~End UPCGConcreteDataWithPointCache interface
+
+	UPROPERTY()
+	TWeakObjectPtr<UWorld> World = nullptr;
+
+	UPROPERTY()
+	TWeakObjectPtr<UPCGComponent> OriginatingComponent = nullptr;
+
+	UPROPERTY()
+	FBox Bounds = FBox(EForceInit::ForceInit);
+
+	UPROPERTY(BlueprintReadWrite, EditAnywhere, Category = Data, meta = (ShowOnlyInnerProperties))
+	FPCGWorldRayHitQueryParams QueryParams;
+};
\ No newline at end of file
diff --git a/Engine/Plugins/Experimental/PCG/Source/PCG/Public/Elements/PCGSurfaceSampler.h b/Engine/Plugins/Experimental/PCG/Source/PCG/Public/Elements/PCGSurfaceSampler.h
--- a/Engine/Plugins/Experimental/PCG/Source/PCG/Public/Elements/PCGSurfaceSampler.h
+++ b/Engine/Plugins/Experimental/PCG/Source/PCG/Public/Elements/PCGSurfaceSampler.h
@@ -10,6 +10,8 @@
 #include "PCGSurfaceSampler.generated.h"
 
 class FPCGSurfaceSamplerElement;
+class UPCGPointData;
+class UPCGSpatialData;
 
 UCLASS(BlueprintType, ClassGroup = (Procedural))
 class PCG_API UPCGSurfaceSamplerSettings : public UPCGSettings
@@ -64,3 +66,40 @@ class FPCGSurfaceSamplerElement : public FSimplePCGElement
 protected:
 	virtual bool ExecuteInternal(FPCGContext* Context) const override;
 };
+
+namespace PCGSurfaceSampler
+{
+	struct FSurfaceSamplerSettings
+	{
+		const UPCGSurfaceSamplerSettings* Settings = nullptr;
+
+		float PointsPerSquaredMeter = 1.0f;
+		FVector PointExtents = FVector::One() * 0.5f;
+		float Looseness = 0.0f;
+		bool bApplyDensityToPoints = false;
+		float PointSteepness = 0.0f;
+#if WITH_EDITORONLY_DATA
+		bool bKeepZeroDensityPoints = false;
+#endif
+
+		bool Initialize(const UPCGSurfaceSamplerSettings* InSettings, FPCGContext* Context, const FBox& InputBounds);
+		void ComputeCellIndices(int32 Index, int32& CellX, int32& CellY) const;
+
+		/** Computed values **/
+		FVector InterstitialDistance;
+		FVector InnerCellSize;
+		FVector CellSize;
+
+		int32 CellMinX;
+		int32 CellMaxX;
+		int32 CellMinY;
+		int32 CellMaxY;
+		int32 CellCount;
+		int64 TargetPointCount;
+		float Ratio;
+		int Seed;
+	};
+
+	UPCGPointData* SampleSurface(FPCGContext* Context, const UPCGSpatialData* SpatialInput, const FSurfaceSamplerSettings& LoopData);
+	void SampleSurface(FPCGContext* Context, const UPCGSpatialData* SpatialInput, const FSurfaceSamplerSettings& LoopData, UPCGPointData* SampledData);
+}
\ No newline at end of file
diff --git a/Engine/Plugins/Experimental/PCG/Source/PCG/Public/Elements/PCGVolumeSampler.h b/Engine/Plugins/Experimental/PCG/Source/PCG/Public/Elements/PCGVolumeSampler.h
--- a/Engine/Plugins/Experimental/PCG/Source/PCG/Public/Elements/PCGVolumeSampler.h
+++ b/Engine/Plugins/Experimental/PCG/Source/PCG/Public/Elements/PCGVolumeSampler.h
@@ -18,6 +18,7 @@ namespace PCGVolumeSampler
 	};
 
 	UPCGPointData* SampleVolume(FPCGContext* Context, const UPCGSpatialData* SpatialData, const FVolumeSamplerSettings& SamplerSettings);
+	void SampleVolume(FPCGContext* Context, const UPCGSpatialData* SpatialData, const FVolumeSamplerSettings& SamplerSettings, UPCGPointData* OutputData, const FBox& Bounds);
 }
 
 UCLASS(BlueprintType, ClassGroup = (Procedural))
diff --git a/Engine/Plugins/Experimental/PCG/Source/PCG/Public/Elements/PCGWorldQuery.h b/Engine/Plugins/Experimental/PCG/Source/PCG/Public/Elements/PCGWorldQuery.h
new file mode 100644
--- /dev/null
+++ b/Engine/Plugins/Experimental/PCG/Source/PCG/Public/Elements/PCGWorldQuery.h
@@ -0,0 +1,71 @@
+// Copyright Epic Games, Inc. All Rights Reserved.
+
+#pragma once
+
+#include "PCGElement.h"
+#include "PCGSettings.h"
+#include "Data/PCGWorldData.h"
+
+#include "PCGWorldQuery.generated.h"
+
+UCLASS(BlueprintType, ClassGroup = (Procedural))
+class PCG_API UPCGWorldQuerySettings : public UPCGSettings
+{
+	GENERATED_BODY()
+
+public:
+	//~Begin UPCGSettings interface
+#if WITH_EDITOR
+	virtual FName GetDefaultNodeName() const override { return FName(TEXT("WorldVolumetricNode")); }
+	virtual EPCGSettingsType GetType() const override { return EPCGSettingsType::Spatial; }
+#endif
+
+	TArray<FPCGPinProperties> InputPinProperties() const { return TArray<FPCGPinProperties>(); }
+
+protected:
+	virtual FPCGElementPtr CreateElement() const override;
+	//~End UPCGSettings interface
+
+public:
+	UPROPERTY(BlueprintReadWrite, EditAnywhere, Category = Data, meta = (ShowOnlyInnerProperties))
+	FPCGWorldVolumetricQueryParams QueryParams;
+};
+
+class FPCGWorldVolumetricQueryElement : public FSimplePCGElement
+{
+public:
+	virtual bool IsCacheable(const UPCGSettings* InSettings) const { return false; }
+protected:
+	virtual bool ExecuteInternal(FPCGContext* Context) const;
+};
+
+UCLASS(BlueprintType, ClassGroup = (Procedural))
+class PCG_API UPCGWorldRayHitSettings : public UPCGSettings
+{
+	GENERATED_BODY()
+
+public:
+	//~Begin UPCGSettings interface
+#if WITH_EDITOR
+	virtual FName GetDefaultNodeName() const override { return FName(TEXT("WorldRayHitNode")); }
+	virtual EPCGSettingsType GetType() const override { return EPCGSettingsType::Spatial; }
+#endif
+
+	TArray<FPCGPinProperties> InputPinProperties() const { return TArray<FPCGPinProperties>(); }
+
+protected:
+	virtual FPCGElementPtr CreateElement() const override;
+	//~End UPCGSettings interface
+
+public:
+	UPROPERTY(BlueprintReadWrite, EditAnywhere, Category = Data, meta = (ShowOnlyInnerProperties))
+	FPCGWorldRayHitQueryParams QueryParams;
+};
+
+class FPCGWorldRayHitQueryElement : public FSimplePCGElement
+{
+public:
+	virtual bool IsCacheable(const UPCGSettings* InSettings) const { return false; }
+protected:
+	virtual bool ExecuteInternal(FPCGContext* Context) const;
+};
\ No newline at end of file
