Add option to show additional metadata entry key information on points

diff --git a/Engine/Plugins/Experimental/PCG/Source/PCGEditor/Private/SPCGEditorGraphAttributeListView.cpp b/Engine/Plugins/Experimental/PCG/Source/PCGEditor/Private/SPCGEditorGraphAttributeListView.cpp
--- a/Engine/Plugins/Experimental/PCG/Source/PCGEditor/Private/SPCGEditorGraphAttributeListView.cpp
+++ b/Engine/Plugins/Experimental/PCG/Source/PCGEditor/Private/SPCGEditorGraphAttributeListView.cpp
@@ -20,6 +20,11 @@
 
 #define LOCTEXT_NAMESPACE "SPCGEditorGraphAttributeListView"
 
+static TAutoConsoleVariable<bool> CVarShowAdvancedAttributesFields(
+	TEXT("pcg.graph.ShowAdvancedAttributes"),
+	false,
+	TEXT("Control whether advanced attributes/properties are shown in the PCG graph editor"));
+
 namespace PCGEditorGraphAttributeListView
 {
 	const FText NoDataAvailableText = LOCTEXT("NoDataAvailableText", "No data available");
@@ -48,6 +53,8 @@ namespace PCGEditorGraphAttributeListView
 	const FName NAME_PointDensity = FName(TEXT("PointDensity"));
 	const FName NAME_PointSteepness = FName(TEXT("PointSteepness"));
 	const FName NAME_PointSeed = FName(TEXT("PointSeed"));
+	const FName NAME_PointMetadataEntry = FName(TEXT("PointMetadataEntry"));
+	const FName NAME_PointMetadataEntryParent = FName(TEXT("PointMetadataEntryParent"));
 
 	/** Labels of the columns */
 	const FText TEXT_IndexLabel = LOCTEXT("IndexLabel", "Index");
@@ -73,6 +80,8 @@ namespace PCGEditorGraphAttributeListView
 	const FText TEXT_PointDensityLabel = LOCTEXT("PointDensityLabel", "Density");
 	const FText TEXT_PointSteepnessLabel = LOCTEXT("PointSteepnessLabel", "Steepness");
 	const FText TEXT_PointSeedLabel = LOCTEXT("PointSeedLabel", "Seed");
+	const FText TEXT_PointMetadataEntryLabel = LOCTEXT("PointMetadataEntryLabel", "Entry Key");
+	const FText TEXT_PointMetadataEntryParentLabel = LOCTEXT("PointMetadataEntryParentLabel", "Parent Key");
 }
 
 void SPCGListViewItemRow::Construct(const FArguments& InArgs, const TSharedRef<STableViewBase>& InOwnerTableView, const PCGListviewItemPtr& Item)
@@ -218,6 +227,22 @@ FText SPCGListViewItemRow::ConvertPointDataToText(const FPCGPoint* PCGPoint, con
 		const int32 Seed = PCGPoint->Seed;
 		return FText::AsNumber(Seed);
 	}
+	else if (ColumnId == PCGEditorGraphAttributeListView::NAME_PointMetadataEntry)
+	{
+		const int64 MetadataEntryKey = PCGPoint->MetadataEntry;
+		return FText::AsNumber(MetadataEntryKey);
+	}
+	else if (ColumnId == PCGEditorGraphAttributeListView::NAME_PointMetadataEntryParent)
+	{
+		if(InternalItem->PCGMetadata)
+		{
+			return FText::AsNumber(InternalItem->PCGMetadata->GetParentKey(PCGPoint->MetadataEntry));
+		}
+		else
+		{
+			return FText::AsNumber(PCGInvalidEntryKey);
+		}
+	}
 	// None of the default point columns were caught, see if its metadata
 	else if (const UPCGMetadata* PCGMetadata = InternalItem->PCGMetadata)
 	{
@@ -801,6 +826,12 @@ void SPCGEditorGraphAttributeListView::AddPointDataColumns()
 	AddColumn(PCGEditorGraphAttributeListView::NAME_PointDensity, PCGEditorGraphAttributeListView::TEXT_PointDensityLabel, 54);
 	AddColumn(PCGEditorGraphAttributeListView::NAME_PointSteepness, PCGEditorGraphAttributeListView::TEXT_PointSteepnessLabel, 73);
 	AddColumn(PCGEditorGraphAttributeListView::NAME_PointSeed, PCGEditorGraphAttributeListView::TEXT_PointSeedLabel, 88);
+
+	if (CVarShowAdvancedAttributesFields.GetValueOnAnyThread())
+	{
+		AddColumn(PCGEditorGraphAttributeListView::NAME_PointMetadataEntry, PCGEditorGraphAttributeListView::TEXT_PointMetadataEntryLabel, 50);
+		AddColumn(PCGEditorGraphAttributeListView::NAME_PointMetadataEntryParent, PCGEditorGraphAttributeListView::TEXT_PointMetadataEntryParentLabel, 50);
+	}	
 }
 
 void SPCGEditorGraphAttributeListView::RemovePointDataColumns()
@@ -828,6 +859,9 @@ void SPCGEditorGraphAttributeListView::RemovePointDataColumns()
 	RemoveColumn(PCGEditorGraphAttributeListView::NAME_PointDensity);
 	RemoveColumn(PCGEditorGraphAttributeListView::NAME_PointSteepness);
 	RemoveColumn(PCGEditorGraphAttributeListView::NAME_PointSeed);
+
+	RemoveColumn(PCGEditorGraphAttributeListView::NAME_PointMetadataEntry);
+	RemoveColumn(PCGEditorGraphAttributeListView::NAME_PointMetadataEntryParent);
 }
 
 void SPCGEditorGraphAttributeListView::AddMetadataColumn(const FName& InColumnId, const FName& InMetadataId, const int8 InValueIndex, const TCHAR* PostFix)
