Compilation fix

#rnx

diff --git a/Engine/Plugins/Experimental/PCG/Source/PCG/Private/Elements/PCGSurfaceSampler.cpp b/Engine/Plugins/Experimental/PCG/Source/PCG/Private/Elements/PCGSurfaceSampler.cpp
--- a/Engine/Plugins/Experimental/PCG/Source/PCG/Private/Elements/PCGSurfaceSampler.cpp
+++ b/Engine/Plugins/Experimental/PCG/Source/PCG/Private/Elements/PCGSurfaceSampler.cpp
@@ -123,7 +123,11 @@ namespace PCGSurfaceSampler
 				const FVector TentativeLocation = FVector(CurrentX + RandX * InnerCellSize.X, CurrentY + RandY * InnerCellSize.Y, 0.0f);
 				const FBox LocalBound(-LoopData.PointExtents, LoopData.PointExtents);
 
+#if WITH_EDITORONLY_DATA
 				if (SpatialInput->SamplePoint(FTransform(TentativeLocation), LocalBound, OutPoint, SampledData->Metadata) || LoopData.bKeepZeroDensityPoints)
+#else
+				if (SpatialInput->SamplePoint(FTransform(TentativeLocation), LocalBound, OutPoint, SampledData->Metadata))
+#endif
 				{
 					// Apply final parameters on the point
 					OutPoint.SetExtents(LoopData.PointExtents);
