[PCG] GetLastAttributeNameOrNone now uses attribute ids to know the latest attribute added.

* Also fix issue with serialization (and therefore deep copy), where NextAttributeId didn't have the right value (leading to collisions)

diff --git a/Engine/Plugins/Experimental/PCG/Source/PCG/Private/Metadata/PCGMetadata.cpp b/Engine/Plugins/Experimental/PCG/Source/PCG/Private/Metadata/PCGMetadata.cpp
--- a/Engine/Plugins/Experimental/PCG/Source/PCG/Private/Metadata/PCGMetadata.cpp
+++ b/Engine/Plugins/Experimental/PCG/Source/PCG/Private/Metadata/PCGMetadata.cpp
@@ -10,6 +10,8 @@ void UPCGMetadata::Serialize(FArchive& InArchive)
 	Super::Serialize(InArchive);
 
 	int32 NumAttributes = (InArchive.IsLoading() ? 0 : Attributes.Num());
+	// We need to keep track of the max attribute Id, since it won't necessary be equal to the number of attributes + 1.
+	int64 MaxAttributeId = -1;
 
 	InArchive << NumAttributes;
 
@@ -29,6 +31,7 @@ void UPCGMetadata::Serialize(FArchive& InArchive)
 				SerializedAttribute->Name = AttributeName;
 				SerializedAttribute->Serialize(this, InArchive);
 				Attributes.Add(AttributeName, SerializedAttribute);
+				MaxAttributeId = FMath::Max(SerializedAttribute->AttributeId, MaxAttributeId);
 			}
 		}
 	}
@@ -50,7 +53,10 @@ void UPCGMetadata::Serialize(FArchive& InArchive)
 	// Finally, initialize non-serialized members
 	if (InArchive.IsLoading())
 	{
-		NextAttributeId = Attributes.Num();
+		// The next attribute id need to be bigger than the max attribute id of all attributes (or we could have collisions).
+		// Therefore by construction, it should never be less than the number of attributes (but can be greater).
+		NextAttributeId = MaxAttributeId + 1;
+		check(NextAttributeId >= Attributes.Num());
 		ItemKeyOffset = (Parent.IsValid() ? Parent->GetItemCountForChild() : 0);
 	}
 }
@@ -340,12 +346,16 @@ void UPCGMetadata::GetAttributes(TArray<FName>& AttributeNames, TArray<EPCGMetad
 FName UPCGMetadata::GetLatestAttributeNameOrNone() const
 {
 	FName LatestAttributeName = NAME_None;
+	int64 MaxAttributeId = -1;
+	
 	AttributeLock.ReadLock();
-	if (!Attributes.IsEmpty())
+	for (const TPair<FName, FPCGMetadataAttributeBase*>& It : Attributes)
 	{
-		TArray<FName> Keys;
-		Attributes.GenerateKeyArray(Keys);
-		LatestAttributeName = Keys.Last();
+		if (It.Value && (It.Value->AttributeId > MaxAttributeId))
+		{
+			MaxAttributeId = It.Value->AttributeId;
+			LatestAttributeName = It.Key;
+		}
 	}
 	AttributeLock.ReadUnlock();
 
