[PCG] Missing PostLoad in UPCGPropertyToParamDataSettings

#rnx

diff --git a/Engine/Plugins/Experimental/PCG/Source/PCG/Private/Elements/PCGPropertyToParamData.cpp b/Engine/Plugins/Experimental/PCG/Source/PCG/Private/Elements/PCGPropertyToParamData.cpp
--- a/Engine/Plugins/Experimental/PCG/Source/PCG/Private/Elements/PCGPropertyToParamData.cpp
+++ b/Engine/Plugins/Experimental/PCG/Source/PCG/Private/Elements/PCGPropertyToParamData.cpp
@@ -14,6 +14,8 @@
 
 void UPCGPropertyToParamDataSettings::PostLoad()
 {
+	Super::PostLoad();
+
 	// Migrate deprecated actor selection settings to struct if needed
 	if (ActorSelection_DEPRECATED != EPCGActorSelection::ByTag ||
 		ActorSelectionTag_DEPRECATED != NAME_None ||
