[PCG] Added default keybind to disable nodes in a PCG graph (Shift+E)

diff --git a/Engine/Plugins/Experimental/PCG/Source/PCGEditor/Private/PCGEditorCommands.cpp b/Engine/Plugins/Experimental/PCG/Source/PCGEditor/Private/PCGEditorCommands.cpp
--- a/Engine/Plugins/Experimental/PCG/Source/PCGEditor/Private/PCGEditorCommands.cpp
+++ b/Engine/Plugins/Experimental/PCG/Source/PCGEditor/Private/PCGEditorCommands.cpp
@@ -30,7 +30,7 @@ void FPCGEditorCommands::RegisterCommands()
 	UI_COMMAND(ExecutionModeEnabled, "Execution Mode Enabled", "Set the execution mode for this node to Enabled.", EUserInterfaceActionType::RadioButton, FInputChord(EKeys::E));
 	UI_COMMAND(ExecutionModeDebug, "Execution Mode Debug", "Set the execution mode for this node to Debug.", EUserInterfaceActionType::RadioButton, FInputChord(EKeys::D));
 	UI_COMMAND(ExecutionModeIsolated, "Execution Mode Isolated", "Set the execution mode for this node to Isolated.", EUserInterfaceActionType::RadioButton, FInputChord());
-	UI_COMMAND(ExecutionModeDisabled, "Execution Mode Disabled", "Set the execution mode for this node to Disabled.", EUserInterfaceActionType::RadioButton, FInputChord());
+	UI_COMMAND(ExecutionModeDisabled, "Execution Mode Disabled", "Set the execution mode for this node to Disabled.", EUserInterfaceActionType::RadioButton, FInputChord(EModifierKey::Shift, EKeys::E));
 	UI_COMMAND(CancelExecution, "Cancel Execution", "Cancels the execution of the current graph", EUserInterfaceActionType::Button, FInputChord(EModifierKey::Control, EKeys::Escape));
 }
 
