API improvement for bp graph deprecation
Clarified node api in graph (copy vs. instance - future feature)
Added drag & drop for settings
Added settings (presets) palette
Added export command for nodes (to settings)

diff --git a/Engine/Plugins/Experimental/PCG/Source/PCG/Private/PCGGraph.cpp b/Engine/Plugins/Experimental/PCG/Source/PCG/Private/PCGGraph.cpp
--- a/Engine/Plugins/Experimental/PCG/Source/PCG/Private/PCGGraph.cpp
+++ b/Engine/Plugins/Experimental/PCG/Source/PCG/Private/PCGGraph.cpp
@@ -235,6 +235,20 @@ UPCGNode* UPCGGraph::AddNode(UPCGSettings* InSettings)
 	return Node;
 }
 
+UPCGNode* UPCGGraph::AddNodeCopy(UPCGSettings* InSettings, UPCGSettings*& DefaultNodeSettings)
+{
+	UPCGSettings* SettingsCopy = DuplicateObject(InSettings, nullptr);
+	UPCGNode* NewNode = AddNode(SettingsCopy);
+
+	if (SettingsCopy)
+	{
+		SettingsCopy->Rename(nullptr, NewNode);
+	}	
+
+	DefaultNodeSettings = SettingsCopy;
+	return NewNode;
+}
+
 void UPCGGraph::OnNodeAdded(UPCGNode* InNode)
 {
 #if WITH_EDITOR
@@ -254,13 +268,13 @@ void UPCGGraph::OnNodeRemoved(UPCGNode* InNode)
 #endif
 }
 
-UPCGNode* UPCGGraph::AddEdge(UPCGNode* From, const FName& InboundLabel, UPCGNode* To, const FName& OutboundLabel)
+UPCGNode* UPCGGraph::AddEdge(UPCGNode* From, const FName& FromPinLabel, UPCGNode* To, const FName& ToPinLabel)
 {
-	AddLabeledEdge(From, InboundLabel, To, OutboundLabel);
+	AddLabeledEdge(From, FromPinLabel, To, ToPinLabel);
 	return To;
 }
 
-bool UPCGGraph::AddLabeledEdge(UPCGNode* From, const FName& InboundLabel, UPCGNode* To, const FName& OutboundLabel)
+bool UPCGGraph::AddLabeledEdge(UPCGNode* From, const FName& FromPinLabel, UPCGNode* To, const FName& ToPinLabel)
 {
 	if (!From || !To)
 	{
@@ -268,19 +282,19 @@ bool UPCGGraph::AddLabeledEdge(UPCGNode* From, const FName& InboundLabel, UPCGNo
 		return false;
 	}
 
-	UPCGPin* FromPin = From->GetOutputPin(InboundLabel);
+	UPCGPin* FromPin = From->GetOutputPin(FromPinLabel);
 
 	if(!FromPin)
 	{
-		UE_LOG(LogPCG, Error, TEXT("From node %s does not have the %s label"), *From->GetName(), *InboundLabel.ToString());
+		UE_LOG(LogPCG, Error, TEXT("From node %s does not have the %s label"), *From->GetName(), *FromPinLabel.ToString());
 		return false;
 	}
 
-	UPCGPin* ToPin = To->GetInputPin(OutboundLabel);
+	UPCGPin* ToPin = To->GetInputPin(ToPinLabel);
 
 	if(!ToPin)
 	{
-		UE_LOG(LogPCG, Error, TEXT("To node %s does not have the %s label"), *To->GetName(), *OutboundLabel.ToString());
+		UE_LOG(LogPCG, Error, TEXT("To node %s does not have the %s label"), *To->GetName(), *ToPinLabel.ToString());
 		return false;
 	}
 
@@ -304,9 +318,8 @@ bool UPCGGraph::AddLabeledEdge(UPCGNode* From, const FName& InboundLabel, UPCGNo
 
 TObjectPtr<UPCGNode> UPCGGraph::ReconstructNewNode(const UPCGNode* InNode)
 {
-	TObjectPtr<UPCGSettings> NewSettings = DuplicateObject(InNode->DefaultSettings, nullptr);
-	TObjectPtr<UPCGNode> NewNode = AddNode(NewSettings);
-	NewSettings->Rename(nullptr, NewNode);
+	UPCGSettings* NewSettings = nullptr;
+	TObjectPtr<UPCGNode> NewNode = AddNodeCopy(InNode->DefaultSettings, NewSettings);
 
 #if WITH_EDITOR
 	InNode->TransferEditorProperties(NewNode);
diff --git a/Engine/Plugins/Experimental/PCG/Source/PCG/Private/PCGNode.cpp b/Engine/Plugins/Experimental/PCG/Source/PCG/Private/PCGNode.cpp
--- a/Engine/Plugins/Experimental/PCG/Source/PCG/Private/PCGNode.cpp
+++ b/Engine/Plugins/Experimental/PCG/Source/PCG/Private/PCGNode.cpp
@@ -184,12 +184,11 @@ UPCGGraph* UPCGNode::GetGraph() const
 	return Cast<UPCGGraph>(GetOuter());
 }
 
-UPCGNode* UPCGNode::AddEdgeTo(FName InboundName, UPCGNode* To, FName OutboundName)
+UPCGNode* UPCGNode::AddEdgeTo(FName FromPinLabel, UPCGNode* To, FName ToPinLabel)
 {
-	check(GetGraph());
 	if (UPCGGraph* Graph = GetGraph())
 	{
-		return Graph->AddEdge(this, InboundName, To, OutboundName);
+		return Graph->AddEdge(this, FromPinLabel, To, ToPinLabel);
 	}
 	else
 	{
@@ -197,6 +196,18 @@ UPCGNode* UPCGNode::AddEdgeTo(FName InboundName, UPCGNode* To, FName OutboundNam
 	}
 }
 
+bool UPCGNode::RemoveEdgeTo(FName FromPinLabel, UPCGNode* To, FName ToPinLabel)
+{
+	if (UPCGGraph* Graph = GetGraph())
+	{
+		return Graph->RemoveEdge(this, FromPinLabel, To, ToPinLabel);
+	}
+	else
+	{
+		return false;
+	}
+}
+
 FName UPCGNode::GetNodeTitle() const
 {
 	if (NodeTitle != NAME_None)
@@ -536,3 +547,19 @@ bool UPCGNode::UpdatePins(TFunctionRef<UPCGPin*(UPCGNode*)> PinAllocator)
 
 	return bChanged;
 }
+
+#if WITH_EDITOR
+
+void UPCGNode::GetNodePosition(int32& OutPositionX, int32& OutPositionY) const
+{
+	OutPositionX = PositionX;
+	OutPositionY = PositionY;
+}
+
+void UPCGNode::SetNodePosition(int32 InPositionX, int32 InPositionY)
+{
+	PositionX = InPositionX;
+	PositionY = InPositionY;
+}
+
+#endif
\ No newline at end of file
diff --git a/Engine/Plugins/Experimental/PCG/Source/PCG/Public/Elements/PCGExecuteBlueprint.h b/Engine/Plugins/Experimental/PCG/Source/PCG/Public/Elements/PCGExecuteBlueprint.h
--- a/Engine/Plugins/Experimental/PCG/Source/PCG/Public/Elements/PCGExecuteBlueprint.h
+++ b/Engine/Plugins/Experimental/PCG/Source/PCG/Public/Elements/PCGExecuteBlueprint.h
@@ -190,9 +190,12 @@ public:
 	// ~End UObject interface
 #endif
 
-	UFUNCTION(BlueprintCallable, Category = Settings, meta=(DeterminesOutputType="InElementType", DynamicOutputParam = "ElementInstance"))
+	UFUNCTION(BlueprintCallable, Category = "Settings|Template", meta=(DeterminesOutputType="InElementType", DynamicOutputParam = "ElementInstance"))
 	void SetElementType(TSubclassOf<UPCGBlueprintElement> InElementType, UPCGBlueprintElement*& ElementInstance);
 
+	UFUNCTION(BlueprintCallable, Category = "Settings|Template")
+	TSubclassOf<UPCGBlueprintElement> GetElementType() const { return BlueprintElementType; }
+
 protected:
 	UPROPERTY()
 	TSubclassOf<UPCGBlueprintElement> BlueprintElement_DEPRECATED;
@@ -200,7 +203,7 @@ protected:
 	UPROPERTY(BlueprintReadOnly, EditAnywhere, Category = Template)
 	TSubclassOf<UPCGBlueprintElement> BlueprintElementType;
 
-	UPROPERTY(BlueprintReadOnly, VisibleAnywhere, Instanced, Category = Settings, meta = (ShowOnlyInnerProperties))
+	UPROPERTY(BlueprintReadOnly, VisibleAnywhere, Instanced, Category = "Instance", meta = (ShowOnlyInnerProperties))
 	TObjectPtr<UPCGBlueprintElement> BlueprintElementInstance;
 
 #if WITH_EDITORONLY_DATA
diff --git a/Engine/Plugins/Experimental/PCG/Source/PCG/Public/PCGGraph.h b/Engine/Plugins/Experimental/PCG/Source/PCG/Public/PCGGraph.h
--- a/Engine/Plugins/Experimental/PCG/Source/PCG/Public/PCGGraph.h
+++ b/Engine/Plugins/Experimental/PCG/Source/PCG/Public/PCGGraph.h
@@ -51,13 +51,25 @@ public:
 	UFUNCTION(BlueprintCallable, Category = Graph, meta=(DeterminesOutputType = "InSettingsClass", DynamicOutputParam = "DefaultNodeSettings"))
 	UPCGNode* AddNodeOfType(TSubclassOf<class UPCGSettings> InSettingsClass, UPCGSettings*& DefaultNodeSettings);
 
-	/** Creates a node and assigns it in the input settings. Returns the created node. */
+	/** Creates a node and assigns it in the input settings (settings are not owned by node). Returns the created node. */
 	UFUNCTION(BlueprintCallable, Category = Graph)
 	UPCGNode* AddNode(UPCGSettings* InSettings);
 
+	/** Creates a node and copies the input settings. Returns the created node. */
+	UFUNCTION(BlueprintCallable, Category = Graph, meta = (DeterminesOutputType = "InSettings", DynamicOutputParam = "OutCopiedSettings"))
+	UPCGNode* AddNodeCopy(UPCGSettings* InSettings, UPCGSettings*& DefaultNodeSettings);
+
+	/** Removes a node from the graph. */
+	UFUNCTION(BlueprintCallable, Category = Graph)
+	void RemoveNode(UPCGNode* InNode);
+
 	/** Adds a directed edge in the graph. Returns the "To" node for easy chaining */
 	UFUNCTION(BlueprintCallable, Category = Graph)
-	UPCGNode* AddEdge(UPCGNode* From, const FName& InboundLabel, UPCGNode* To, const FName& OutboundLabel);
+	UPCGNode* AddEdge(UPCGNode* From, const FName& FromPinLabel, UPCGNode* To, const FName& ToPinLabel);
+
+	/** Removes an edge in the graph. Returns true if an edge was removed. */
+	UFUNCTION(BlueprintCallable, Category = Graph)
+	bool RemoveEdge(UPCGNode* From, const FName& FromLabel, UPCGNode* To, const FName& ToLabel);
 
 	/** Returns the graph input node */
 	UFUNCTION(BlueprintCallable, Category = Graph)
@@ -76,8 +88,7 @@ public:
 	bool Contains(UPCGNode* Node) const;
 	const TArray<UPCGNode*>& GetNodes() const { return Nodes; }
 	void AddNode(UPCGNode* InNode);
-	void RemoveNode(UPCGNode* InNode);
-	bool RemoveEdge(UPCGNode* From, const FName& FromLabel, UPCGNode* To, const FName& ToLabel);
+
 	bool RemoveAllInboundEdges(UPCGNode* InNode);
 	bool RemoveAllOutboundEdges(UPCGNode* InNode);
 	bool RemoveInboundEdges(UPCGNode* InNode, const FName& InboundLabel);
@@ -88,7 +99,10 @@ public:
 	void EnableNotificationsForEditor();
 	void ToggleUserPausedNotificationsForEditor();
 	bool NotificationsForEditorArePausedByUser() const { return bUserPausedNotificationsInGraphEditor; }
+
+	UFUNCTION(BlueprintCallable, Category = "Graph|Advanded")
 	void ForceNotificationForEditor();
+
 	void PreNodeUndo(UPCGNode* InPCGNode);
 	void PostNodeUndo(UPCGNode* InPCGNode);
 
diff --git a/Engine/Plugins/Experimental/PCG/Source/PCG/Public/PCGNode.h b/Engine/Plugins/Experimental/PCG/Source/PCG/Public/PCGNode.h
--- a/Engine/Plugins/Experimental/PCG/Source/PCG/Public/PCGNode.h
+++ b/Engine/Plugins/Experimental/PCG/Source/PCG/Public/PCGNode.h
@@ -48,7 +48,11 @@ public:
 
 	/** Adds an edge in the owning graph to the given "To" node. */
 	UFUNCTION(BlueprintCallable, Category = Node)
-	UPCGNode* AddEdgeTo(FName InboundName, UPCGNode* To, FName OutboundName);
+	UPCGNode* AddEdgeTo(FName FromPinLabel, UPCGNode* To, FName ToPinLabel);
+
+	/** Removes an edge originating from this node */
+	UFUNCTION(BlueprintCallable, Category = Node)
+	bool RemoveEdgeTo(FName FromPinLable, UPCGNode* To, FName ToPinLabel);
 
 	/** Returns the node title, based either on the current node label, or defaulted to its settings */
 	FName GetNodeTitle() const;
@@ -90,6 +94,14 @@ public:
 	void TransferEditorProperties(UPCGNode* OtherNode) const;
 #endif // WITH_EDITOR
 
+#if WITH_EDITOR
+	UFUNCTION(BlueprintCallable, Category = Node)
+	void GetNodePosition(int32& OutPositionX, int32& OutPositionY) const;
+
+	UFUNCTION(BlueprintCallable, Category = Node)
+	void SetNodePosition(int32 InPositionX, int32 InPositionY);
+#endif
+
 	/** Note: do not set this property directly from code, use SetDefaultSettings instead */
 	UPROPERTY(BlueprintReadWrite, EditAnywhere, Category = Node, meta=(EditInline))
 	TObjectPtr<UPCGSettings> DefaultSettings;
diff --git a/Engine/Plugins/Experimental/PCG/Source/PCG/Public/PCGPin.h b/Engine/Plugins/Experimental/PCG/Source/PCG/Public/PCGPin.h
--- a/Engine/Plugins/Experimental/PCG/Source/PCG/Public/PCGPin.h
+++ b/Engine/Plugins/Experimental/PCG/Source/PCG/Public/PCGPin.h
@@ -42,16 +42,16 @@ public:
 	virtual void PostLoad() override;
 	// ~End UObject interface
 
-	UPROPERTY()
+	UPROPERTY(BlueprintReadOnly, Category = Properties)
 	TObjectPtr<UPCGNode> Node = nullptr;
 
 	UPROPERTY()
 	FName Label_DEPRECATED = NAME_None;
 
-	UPROPERTY(TextExportTransient)
+	UPROPERTY(BlueprintReadOnly, TextExportTransient, Category = Properties)
 	TArray<TObjectPtr<UPCGEdge>> Edges;
 
-	UPROPERTY(EditAnywhere, Category = Settings, meta = (ShowOnlyInnerProperties))
+	UPROPERTY(BlueprintReadOnly, EditAnywhere, Category = Settings, meta = (ShowOnlyInnerProperties))
 	FPCGPinProperties Properties;
 
 	bool IsCompatible(const UPCGPin* OtherPin) const;
diff --git a/Engine/Plugins/Experimental/PCG/Source/PCG/Public/PCGSettings.h b/Engine/Plugins/Experimental/PCG/Source/PCG/Public/PCGSettings.h
--- a/Engine/Plugins/Experimental/PCG/Source/PCG/Public/PCGSettings.h
+++ b/Engine/Plugins/Experimental/PCG/Source/PCG/Public/PCGSettings.h
@@ -123,6 +123,14 @@ public:
 	FOnPCGSettingsChanged OnSettingsChangedDelegate;
 #endif
 
+#if WITH_EDITORONLY_DATA
+	UPROPERTY(EditAnywhere, Category = AssetInfo, AssetRegistrySearchable)
+	FText Category;
+
+	UPROPERTY(EditAnywhere, Category = AssetInfo, AssetRegistrySearchable)
+	FText Description;
+#endif
+
 protected:
 	virtual FPCGElementPtr CreateElement() const PURE_VIRTUAL(UPCGSettings::CreateElement, return nullptr;);
 
diff --git a/Engine/Plugins/Experimental/PCG/Source/PCGEditor/Private/PCGEditor.cpp b/Engine/Plugins/Experimental/PCG/Source/PCGEditor/Private/PCGEditor.cpp
--- a/Engine/Plugins/Experimental/PCG/Source/PCGEditor/Private/PCGEditor.cpp
+++ b/Engine/Plugins/Experimental/PCG/Source/PCGEditor/Private/PCGEditor.cpp
@@ -382,6 +382,11 @@ void FPCGEditor::BindCommands()
 		FExecuteAction::CreateSP(this, &FPCGEditor::OnCollapseNodesInSubgraph),
 		FCanExecuteAction::CreateSP(this, &FPCGEditor::CanCollapseNodesInSubgraph));
 
+	GraphEditorCommands->MapAction(
+		PCGEditorCommands.ExportNodes,
+		FExecuteAction::CreateSP(this, &FPCGEditor::OnExportNodes),
+		FCanExecuteAction::CreateSP(this, &FPCGEditor::CanExportNodes));
+
 	GraphEditorCommands->MapAction(
 		PCGEditorCommands.StartInspectNode,
 		FExecuteAction::CreateSP(this, &FPCGEditor::OnStartInspectNode));
@@ -1004,6 +1009,82 @@ void FPCGEditor::OnCollapseNodesInSubgraph()
 	NewPCGGraph->NotifyGraphChanged(EPCGChangeType::Structural);
 }
 
+bool FPCGEditor::CanExportNodes() const
+{
+	for (const UObject* Object : GraphEditorWidget->GetSelectedNodes())
+	{
+		check(Object);
+
+		// Exclude input and output nodes from the subgraph.
+		if (Object->IsA<UPCGEditorGraphNodeInput>() || Object->IsA<UPCGEditorGraphNodeOutput>())
+		{
+			continue;
+		}
+
+		if (Object->IsA<UPCGEditorGraphNodeBase>())
+		{
+			return true;
+		}
+	}
+
+	return false;
+}
+
+void FPCGEditor::OnExportNodes()
+{
+	if (!GraphEditorWidget.IsValid() || PCGEditorGraph == nullptr)
+	{
+		UE_LOG(LogPCGEditor, Error, TEXT("GraphEditorWidget or PCGEditorGraph is null, aborting"));
+		return;
+	}
+
+	if (PCGGraphBeingEdited == nullptr)
+	{
+		UE_LOG(LogPCGEditor, Error, TEXT("Editor has no graph loaded, aborting"));
+		return;
+	}
+
+	IAssetTools& AssetTools = FModuleManager::Get().LoadModuleChecked<FAssetToolsModule>("AssetTools").Get();
+
+	for (UObject* Object : GraphEditorWidget->GetSelectedNodes())
+	{
+		check(Object);
+
+		// Exclude input and output nodes from the subgraph.
+		if (Object->IsA<UPCGEditorGraphNodeInput>() || Object->IsA<UPCGEditorGraphNodeOutput>())
+		{
+			continue;
+		}
+
+		UPCGSettings* Settings = nullptr;
+
+		if (UPCGEditorGraphNodeBase* PCGEditorGraphNode = Cast<UPCGEditorGraphNodeBase>(Object))
+		{
+			UPCGNode* PCGNode = PCGEditorGraphNode->GetPCGNode();
+			check(PCGNode);
+			Settings = PCGNode->DefaultSettings;
+		}
+
+		if (!Settings)
+		{
+			continue;
+		}
+
+		// Create new settings asset
+		FString NewPackageName;
+		FString NewAssetName;
+		PCGEditorUtils::GetParentPackagePathAndUniqueName(PCGGraphBeingEdited, LOCTEXT("NewPCGSettingsAsset", "NewPCGSettings").ToString(), NewPackageName, NewAssetName);
+
+		UObject* NewSettings = AssetTools.DuplicateAssetWithDialogAndTitle(NewAssetName, NewPackageName, Settings, NSLOCTEXT("PCGEditor_ExportNodes", "PCGEditor_ExportNodesTitle", "Export Settings As..."));
+
+		if (NewSettings == nullptr)
+		{
+			UE_LOG(LogPCGEditor, Warning, TEXT("Settings asset creation was aborted or failed, abort."));
+			return;
+		}
+	}
+}
+
 void FPCGEditor::OnStartInspectNode()
 {
 	if (!GraphEditorWidget.IsValid())
diff --git a/Engine/Plugins/Experimental/PCG/Source/PCGEditor/Private/PCGEditor.h b/Engine/Plugins/Experimental/PCG/Source/PCGEditor/Private/PCGEditor.h
--- a/Engine/Plugins/Experimental/PCG/Source/PCGEditor/Private/PCGEditor.h
+++ b/Engine/Plugins/Experimental/PCG/Source/PCGEditor/Private/PCGEditor.h
@@ -166,6 +166,12 @@ private:
 	/** Whether we are able to duplicate the currently selected nodes */
 	bool CanDuplicateNodes() const;
 
+	/** Exports node settings to assets */
+	void OnExportNodes();
+
+	/** Whether we are able to export the currently selected nodes */
+	bool CanExportNodes() const;
+
 	/** Collapse the currently selected nodes in a subgraph */
 	void OnCollapseNodesInSubgraph();
 	/** Whether we can collapse nodes in a subgraph */
diff --git a/Engine/Plugins/Experimental/PCG/Source/PCGEditor/Private/PCGEditorCommands.cpp b/Engine/Plugins/Experimental/PCG/Source/PCGEditor/Private/PCGEditorCommands.cpp
--- a/Engine/Plugins/Experimental/PCG/Source/PCGEditor/Private/PCGEditorCommands.cpp
+++ b/Engine/Plugins/Experimental/PCG/Source/PCGEditor/Private/PCGEditorCommands.cpp
@@ -19,6 +19,7 @@ FPCGEditorCommands::FPCGEditorCommands()
 void FPCGEditorCommands::RegisterCommands()
 {
 	UI_COMMAND(CollapseNodes, "Collapse into Subgraph", "Collapse selected nodes into a separate PCGGraph asset.", EUserInterfaceActionType::Button, FInputChord(EModifierKey::Alt, EKeys::J));
+	UI_COMMAND(ExportNodes, "Export nodes to AssetData", "Exports selected nodes to separate and reusable PCGSettings assets.", EUserInterfaceActionType::Button, FInputChord());
 	UI_COMMAND(Find, "Find", "Finds PCG nodes and comments in the current graph", EUserInterfaceActionType::Button, FInputChord(EModifierKey::Control, EKeys::F));
 	UI_COMMAND(PauseAutoRegeneration, "Pause Regen", "Pause automatic regeneration of the current graph.", EUserInterfaceActionType::ToggleButton, FInputChord(EModifierKey::Alt, EKeys::R));
 	UI_COMMAND(ForceGraphRegeneration, "Force Regen", "Manually force a regeneration of the current graph.", EUserInterfaceActionType::Button, FInputChord());
diff --git a/Engine/Plugins/Experimental/PCG/Source/PCGEditor/Private/PCGEditorCommands.h b/Engine/Plugins/Experimental/PCG/Source/PCGEditor/Private/PCGEditorCommands.h
--- a/Engine/Plugins/Experimental/PCG/Source/PCGEditor/Private/PCGEditorCommands.h
+++ b/Engine/Plugins/Experimental/PCG/Source/PCGEditor/Private/PCGEditorCommands.h
@@ -15,6 +15,7 @@ public:
 	// ~End TCommands<> interface
 
 	TSharedPtr<FUICommandInfo> CollapseNodes;
+	TSharedPtr<FUICommandInfo> ExportNodes;
 	TSharedPtr<FUICommandInfo> Find;
 	TSharedPtr<FUICommandInfo> PauseAutoRegeneration;
 	TSharedPtr<FUICommandInfo> ForceGraphRegeneration;
diff --git a/Engine/Plugins/Experimental/PCG/Source/PCGEditor/Private/PCGEditorCommon.h b/Engine/Plugins/Experimental/PCG/Source/PCGEditor/Private/PCGEditorCommon.h
--- a/Engine/Plugins/Experimental/PCG/Source/PCGEditor/Private/PCGEditorCommon.h
+++ b/Engine/Plugins/Experimental/PCG/Source/PCGEditor/Private/PCGEditorCommon.h
@@ -10,8 +10,9 @@ enum class EPCGElementType : uint8
 	Native = 1 << 0,
 	Blueprint = 1 << 1,
 	Subgraph = 1 << 2,
-	Other = 1 << 3,
-	All = (1 << 4) - 1
+	Settings = 1 << 3,
+	Other = 1 << 4,
+	All = (1 << 5) - 1
 };
 ENUM_CLASS_FLAGS(EPCGElementType);
 
diff --git a/Engine/Plugins/Experimental/PCG/Source/PCGEditor/Private/PCGEditorGraphNodeBase.cpp b/Engine/Plugins/Experimental/PCG/Source/PCGEditor/Private/PCGEditorGraphNodeBase.cpp
--- a/Engine/Plugins/Experimental/PCG/Source/PCGEditor/Private/PCGEditorGraphNodeBase.cpp
+++ b/Engine/Plugins/Experimental/PCG/Source/PCGEditor/Private/PCGEditorGraphNodeBase.cpp
@@ -91,6 +91,8 @@ void UPCGEditorGraphNodeBase::GetNodeContextMenuActions(UToolMenu* Menu, class U
 			SubSection.AddMenuEntry(FPCGEditorCommands::Get().ExecutionModeDisabled, LOCTEXT("ExecutionModeDisabledLabel", "Disabled"));
 			SubSection.AddMenuEntry(FPCGEditorCommands::Get().ExecutionModeIsolated, LOCTEXT("ExecutionModeIsolatedLabel", "Isolated"));
 		}));
+
+		Section.AddMenuEntry(FPCGEditorCommands::Get().ExportNodes);
 	}
 
 	{
diff --git a/Engine/Plugins/Experimental/PCG/Source/PCGEditor/Private/PCGEditorGraphSchema.cpp b/Engine/Plugins/Experimental/PCG/Source/PCGEditor/Private/PCGEditorGraphSchema.cpp
--- a/Engine/Plugins/Experimental/PCG/Source/PCGEditor/Private/PCGEditorGraphSchema.cpp
+++ b/Engine/Plugins/Experimental/PCG/Source/PCGEditor/Private/PCGEditorGraphSchema.cpp
@@ -34,6 +34,10 @@ void UPCGEditorGraphSchema::GetPaletteActions(FGraphActionMenuBuilder& ActionMen
 	{
 		GetBlueprintElementActions(ActionMenuBuilder);
 	}
+	if (!!(InPCGElementTypeFilter & EPCGElementType::Settings))
+	{
+		GetSettingsElementActions(ActionMenuBuilder);
+	}
 	if (!!(InPCGElementTypeFilter & EPCGElementType::Other))
 	{
 		GetExtraElementActions(ActionMenuBuilder);
@@ -47,6 +51,7 @@ void UPCGEditorGraphSchema::GetGraphContextActions(FGraphContextMenuBuilder& Con
 	GetNativeElementActions(ContextMenuBuilder);
 	GetSubgraphElementActions(ContextMenuBuilder);
 	GetBlueprintElementActions(ContextMenuBuilder);
+	GetSettingsElementActions(ContextMenuBuilder);
 	GetExtraElementActions(ContextMenuBuilder);
 }
 
@@ -250,6 +255,25 @@ void UPCGEditorGraphSchema::GetBlueprintElementActions(FGraphActionMenuBuilder&
 	}
 }
 
+void UPCGEditorGraphSchema::GetSettingsElementActions(FGraphActionMenuBuilder& ActionMenuBuilder) const
+{
+	const FAssetRegistryModule& AssetRegistryModule = FModuleManager::LoadModuleChecked<FAssetRegistryModule>(TEXT("AssetRegistry"));
+
+	TArray<FAssetData> AssetDataList;
+	AssetRegistryModule.Get().GetAssetsByClass(UPCGSettings::StaticClass()->GetClassPathName(), AssetDataList, /*bSearchSubClasses=*/true);
+
+	for (const FAssetData& AssetData : AssetDataList)
+	{
+		const FText MenuDesc = FText::FromName(AssetData.AssetName);
+		const FText Category = AssetData.GetTagValueRef<FText>(TEXT("Category"));
+		const FText Description = AssetData.GetTagValueRef<FText>(TEXT("Description"));
+
+		TSharedPtr<FPCGEditorGraphSchemaAction_NewSettingsElement> NewSettingsAction(new FPCGEditorGraphSchemaAction_NewSettingsElement(Category, MenuDesc, Description, 0));
+		NewSettingsAction->SettingsObjectPath = AssetData.GetSoftObjectPath();
+		ActionMenuBuilder.AddAction(NewSettingsAction);
+	}
+}
+
 void UPCGEditorGraphSchema::GetSubgraphElementActions(FGraphActionMenuBuilder& ActionMenuBuilder) const
 {
 	const FAssetRegistryModule& AssetRegistryModule = FModuleManager::LoadModuleChecked<FAssetRegistryModule>(TEXT("AssetRegistry"));
@@ -310,6 +334,13 @@ void UPCGEditorGraphSchema::DroppedAssetsOnGraph(const TArray<FAssetData>& Asset
 				NewBlueprintAction.PerformAction(Graph, NullFromPin, GraphPositionOffset);
 				GraphPositionOffset.Y += PositionOffsetIncrementY;
 			}
+			else if (Asset->IsA<UPCGSettings>())
+			{
+				FPCGEditorGraphSchemaAction_NewSettingsElement NewNodeAction;
+				NewNodeAction.SettingsObjectPath = AssetData.GetSoftObjectPath();
+				NewNodeAction.PerformAction(Graph, NullFromPin, GraphPositionOffset);
+				GraphPositionOffset.Y += PositionOffsetIncrementY;
+			}
 		}
 	}
 }
@@ -320,7 +351,7 @@ void UPCGEditorGraphSchema::GetAssetsGraphHoverMessage(const TArray<FAssetData>&
 	{
 		if (const UObject* Asset = AssetData.GetAsset())
 		{
-			if (Asset->IsA<UPCGGraph>() || PCGEditorUtils::IsAssetPCGBlueprint(AssetData))
+			if (Asset->IsA<UPCGGraph>() || Asset->IsA<UPCGSettings>() || PCGEditorUtils::IsAssetPCGBlueprint(AssetData))
 			{
 				OutOkIcon = true;
 				return;
diff --git a/Engine/Plugins/Experimental/PCG/Source/PCGEditor/Private/PCGEditorGraphSchema.h b/Engine/Plugins/Experimental/PCG/Source/PCGEditor/Private/PCGEditorGraphSchema.h
--- a/Engine/Plugins/Experimental/PCG/Source/PCGEditor/Private/PCGEditorGraphSchema.h
+++ b/Engine/Plugins/Experimental/PCG/Source/PCGEditor/Private/PCGEditorGraphSchema.h
@@ -34,6 +34,7 @@ private:
 	void GetNativeElementActions(FGraphActionMenuBuilder& ActionMenuBuilder) const;
 	void GetBlueprintElementActions(FGraphActionMenuBuilder& ActionMenuBuilder) const;
 	void GetSubgraphElementActions(FGraphActionMenuBuilder& ActionMenuBuilder) const;
+	void GetSettingsElementActions(FGraphActionMenuBuilder& ActionMenuBuilder) const;
 	void GetExtraElementActions(FGraphActionMenuBuilder& ActionMenuBuilder) const;
 };
 
diff --git a/Engine/Plugins/Experimental/PCG/Source/PCGEditor/Private/PCGEditorGraphSchemaActions.cpp b/Engine/Plugins/Experimental/PCG/Source/PCGEditor/Private/PCGEditorGraphSchemaActions.cpp
--- a/Engine/Plugins/Experimental/PCG/Source/PCGEditor/Private/PCGEditorGraphSchemaActions.cpp
+++ b/Engine/Plugins/Experimental/PCG/Source/PCGEditor/Private/PCGEditorGraphSchemaActions.cpp
@@ -59,6 +59,60 @@ UEdGraphNode* FPCGEditorGraphSchemaAction_NewNativeElement::PerformAction(UEdGra
 	return NewNode;
 }
 
+UEdGraphNode* FPCGEditorGraphSchemaAction_NewSettingsElement::PerformAction(UEdGraph* ParentGraph, UEdGraphPin* FromPin, const FVector2D Location, bool bSelectNewNode)
+{
+	UPCGEditorGraph* EditorGraph = Cast<UPCGEditorGraph>(ParentGraph);
+	if (!EditorGraph)
+	{
+		UE_LOG(LogPCGEditor, Error, TEXT("Invalid EditorGraph"));
+		return nullptr;
+	}
+
+	UPCGGraph* PCGGraph = EditorGraph->GetPCGGraph();
+	if (!PCGGraph)
+	{
+		UE_LOG(LogPCGEditor, Error, TEXT("Invalid PCGGraph"));
+		return nullptr;
+	}
+
+	UPCGSettings* Settings = CastChecked<UPCGSettings>(SettingsObjectPath.TryLoad());
+	if (!Settings)
+	{
+		UE_LOG(LogPCGEditor, Error, TEXT("Invalid settings"));
+		return nullptr;
+	}
+
+	const FScopedTransaction Transaction(*FPCGEditorCommon::ContextIdentifier, LOCTEXT("PCGEditorNewSettingsElement", "PCG Editor: New Settings Element"), nullptr);
+	EditorGraph->Modify();
+
+	UPCGNode* NewPCGNode = nullptr;
+	UPCGSettings* NewSettings = nullptr;
+	NewPCGNode = PCGGraph->AddNodeCopy(Settings, NewSettings);
+
+	if (!NewPCGNode)
+	{
+		UE_LOG(LogPCGEditor, Error, TEXT("Unable to create node"));
+		return nullptr;
+	}
+
+	FGraphNodeCreator<UPCGEditorGraphNode> NodeCreator(*EditorGraph);
+	UPCGEditorGraphNode* NewNode = NodeCreator.CreateUserInvokedNode(bSelectNewNode);
+	NewNode->Construct(NewPCGNode);
+	NewNode->NodePosX = Location.X;
+	NewNode->NodePosY = Location.Y;
+	NodeCreator.Finalize();
+
+	NewPCGNode->PositionX = Location.X;
+	NewPCGNode->PositionY = Location.Y;
+
+	if (FromPin)
+	{
+		NewNode->AutowireNewNode(FromPin);
+	}
+
+	return NewNode;
+}
+
 UEdGraphNode* FPCGEditorGraphSchemaAction_NewBlueprintElement::PerformAction(UEdGraph* ParentGraph, UEdGraphPin* FromPin, const FVector2D Location, bool bSelectNewNode)
 {
 	UPCGEditorGraph* EditorGraph = Cast<UPCGEditorGraph>(ParentGraph);
diff --git a/Engine/Plugins/Experimental/PCG/Source/PCGEditor/Private/PCGEditorGraphSchemaActions.h b/Engine/Plugins/Experimental/PCG/Source/PCGEditor/Private/PCGEditorGraphSchemaActions.h
--- a/Engine/Plugins/Experimental/PCG/Source/PCGEditor/Private/PCGEditorGraphSchemaActions.h
+++ b/Engine/Plugins/Experimental/PCG/Source/PCGEditor/Private/PCGEditorGraphSchemaActions.h
@@ -29,6 +29,22 @@ struct FPCGEditorGraphSchemaAction_NewNativeElement : public FEdGraphSchemaActio
 	//~ End FEdGraphSchemaAction Interface
 };
 
+USTRUCT()
+struct FPCGEditorGraphSchemaAction_NewSettingsElement : public FEdGraphSchemaAction
+{
+	GENERATED_BODY()
+
+	// Inherit the base class's constructors
+	using FEdGraphSchemaAction::FEdGraphSchemaAction;
+
+	UPROPERTY()
+	FSoftObjectPath SettingsObjectPath;
+
+	//~ Begin FEdGraphSchemaAction Interface
+	virtual UEdGraphNode* PerformAction(UEdGraph* ParentGraph, UEdGraphPin* FromPin, const FVector2D Location, bool bSelectNewNode = true) override;
+	//~ End FEdGraphSchemaAction Interface
+};
+
 USTRUCT()
 struct FPCGEditorGraphSchemaAction_NewBlueprintElement : public FEdGraphSchemaAction
 {
diff --git a/Engine/Plugins/Experimental/PCG/Source/PCGEditor/Private/SPCGEditorGraphNodePalette.cpp b/Engine/Plugins/Experimental/PCG/Source/PCGEditor/Private/SPCGEditorGraphNodePalette.cpp
--- a/Engine/Plugins/Experimental/PCG/Source/PCGEditor/Private/SPCGEditorGraphNodePalette.cpp
+++ b/Engine/Plugins/Experimental/PCG/Source/PCGEditor/Private/SPCGEditorGraphNodePalette.cpp
@@ -113,7 +113,9 @@ void SPCGEditorGraphNodePalette::CollectAllActions(FGraphActionListBuilderBase&
 
 void SPCGEditorGraphNodePalette::OnAssetChanged(const FAssetData& InAssetData)
 {
-	if (InAssetData.AssetClassPath == UPCGGraph::StaticClass()->GetClassPathName() || PCGEditorUtils::IsAssetPCGBlueprint(InAssetData))
+	if (InAssetData.AssetClassPath == UPCGGraph::StaticClass()->GetClassPathName() ||
+		InAssetData.AssetClassPath == UPCGSettings::StaticClass()->GetClassPathName() ||
+		PCGEditorUtils::IsAssetPCGBlueprint(InAssetData))
 	{
 		RefreshActionsList(true);
 	}
