Tentative fix for static analysis warnings

#rnx

diff --git a/Engine/Plugins/Experimental/PCG/Source/PCG/Private/Elements/PCGSurfaceSampler.cpp b/Engine/Plugins/Experimental/PCG/Source/PCG/Private/Elements/PCGSurfaceSampler.cpp
--- a/Engine/Plugins/Experimental/PCG/Source/PCG/Private/Elements/PCGSurfaceSampler.cpp
+++ b/Engine/Plugins/Experimental/PCG/Source/PCG/Private/Elements/PCGSurfaceSampler.cpp
@@ -31,7 +31,7 @@ namespace PCGSurfaceSampler
 			bKeepZeroDensityPoints = PCGSettingsHelpers::GetValue(GET_MEMBER_NAME_CHECKED(UPCGSurfaceSamplerSettings, bKeepZeroDensityPoints), Settings->bKeepZeroDensityPoints, Params);
 #endif
 
-			Seed = PCGSettingsHelpers::ComputeSeedWithOverride(InSettings, Context->SourceComponent, Params);
+			Seed = PCGSettingsHelpers::ComputeSeedWithOverride(InSettings, Context ? Context->SourceComponent : nullptr, Params);
 		}
 		else
 		{
@@ -52,7 +52,11 @@ namespace PCGSurfaceSampler
 
 		if (CellMinX > CellMaxX || CellMinY > CellMaxY)
 		{
-			PCGE_LOG_C(Verbose, Context, "Skipped - invalid cell bounds");
+			if (Context)
+			{
+				PCGE_LOG_C(Verbose, Context, "Skipped - invalid cell bounds");
+			}
+			
 			return false;
 		}
 
@@ -64,7 +68,11 @@ namespace PCGSurfaceSampler
 
 		if (TargetPointCount == 0)
 		{
-			PCGE_LOG_C(Verbose, Context, "Skipped - density yields no points");
+			if (Context)
+			{
+				PCGE_LOG_C(Verbose, Context, "Skipped - density yields no points");
+			}
+			
 			return false;
 		}
 		else if (TargetPointCount > CellCount)
@@ -148,7 +156,10 @@ namespace PCGSurfaceSampler
 			}
 		});
 
-		PCGE_LOG_C(Verbose, Context, "Generated %d points in %d cells", SampledPoints.Num(), LoopData.CellCount);
+		if (Context)
+		{
+			PCGE_LOG_C(Verbose, Context, "Generated %d points in %d cells", SampledPoints.Num(), LoopData.CellCount);
+		}
 	}
 }
 
@@ -177,6 +188,7 @@ bool FPCGSurfaceSamplerElement::ExecuteInternal(FPCGContext* Context) const
 {
 	TRACE_CPUPROFILER_EVENT_SCOPE(FPCGSurfaceSamplerElement::Execute);
 	// TODO: time-sliced implementation
+	check(Context);
 	const UPCGSurfaceSamplerSettings* Settings = Context->GetInputSettings<UPCGSurfaceSamplerSettings>();
 	check(Settings);
 
