PCGComponents now store their output result (serialized - shallow copy only will not work with complex types yet)
Cleaned up some outering of runtime data to prevent GC issues.
Added some mechanism on the PCG Data From Actor node to get data on a PCG component and ability to get a collection from an arbitrary collection.

diff --git a/Engine/Plugins/Experimental/PCG/Source/PCG/Private/Data/PCGDifferenceData.cpp b/Engine/Plugins/Experimental/PCG/Source/PCG/Private/Data/PCGDifferenceData.cpp
--- a/Engine/Plugins/Experimental/PCG/Source/PCG/Private/Data/PCGDifferenceData.cpp
+++ b/Engine/Plugins/Experimental/PCG/Source/PCG/Private/Data/PCGDifferenceData.cpp
@@ -55,7 +55,7 @@ void UPCGDifferenceData::AddDifference(const UPCGSpatialData* InDifference)
 	{
 		if (!DifferencesUnion)
 		{
-			DifferencesUnion = NewObject<UPCGUnionData>(this);
+			DifferencesUnion = NewObject<UPCGUnionData>();
 			DifferencesUnion->AddData(Difference);
 			DifferencesUnion->SetDensityFunction(PCGDifferenceDataUtils::ToUnionDensityFunction(DensityFunction));
 			Difference = DifferencesUnion;
diff --git a/Engine/Plugins/Experimental/PCG/Source/PCG/Private/Data/PCGSpatialData.cpp b/Engine/Plugins/Experimental/PCG/Source/PCG/Private/Data/PCGSpatialData.cpp
--- a/Engine/Plugins/Experimental/PCG/Source/PCG/Private/Data/PCGSpatialData.cpp
+++ b/Engine/Plugins/Experimental/PCG/Source/PCG/Private/Data/PCGSpatialData.cpp
@@ -84,7 +84,7 @@ FVector UPCGSpatialData::TransformPosition(const FVector& InPosition) const
 
 UPCGIntersectionData* UPCGSpatialData::IntersectWith(const UPCGSpatialData* InOther) const
 {
-	UPCGIntersectionData* IntersectionData = NewObject<UPCGIntersectionData>(const_cast<UPCGSpatialData*>(this));
+	UPCGIntersectionData* IntersectionData = NewObject<UPCGIntersectionData>();
 	IntersectionData->Initialize(this, InOther);
 
 	return IntersectionData;
@@ -100,7 +100,7 @@ UPCGProjectionData* UPCGSpatialData::ProjectOn(const UPCGSpatialData* InOther, c
 
 UPCGUnionData* UPCGSpatialData::UnionWith(const UPCGSpatialData* InOther) const
 {
-	UPCGUnionData* UnionData = NewObject<UPCGUnionData>(const_cast<UPCGSpatialData*>(this));
+	UPCGUnionData* UnionData = NewObject<UPCGUnionData>();
 	UnionData->Initialize(this, InOther);
 
 	return UnionData;
@@ -108,7 +108,7 @@ UPCGUnionData* UPCGSpatialData::UnionWith(const UPCGSpatialData* InOther) const
 
 UPCGDifferenceData* UPCGSpatialData::Subtract(const UPCGSpatialData* InOther) const
 {
-	UPCGDifferenceData* DifferenceData = NewObject<UPCGDifferenceData>(const_cast<UPCGSpatialData*>(this));
+	UPCGDifferenceData* DifferenceData = NewObject<UPCGDifferenceData>();
 	DifferenceData->Initialize(this);
 	DifferenceData->AddDifference(InOther);
 
diff --git a/Engine/Plugins/Experimental/PCG/Source/PCG/Private/Data/PCGSplineData.cpp b/Engine/Plugins/Experimental/PCG/Source/PCG/Private/Data/PCGSplineData.cpp
--- a/Engine/Plugins/Experimental/PCG/Source/PCG/Private/Data/PCGSplineData.cpp
+++ b/Engine/Plugins/Experimental/PCG/Source/PCG/Private/Data/PCGSplineData.cpp
@@ -108,7 +108,7 @@ UPCGProjectionData* UPCGSplineData::ProjectOn(const UPCGSpatialData* InOther, co
 {
 	if(InOther->GetDimension() == 2)
 	{
-		UPCGSplineProjectionData* SplineProjectionData = NewObject<UPCGSplineProjectionData>(const_cast<UPCGSplineData*>(this));
+		UPCGSplineProjectionData* SplineProjectionData = NewObject<UPCGSplineProjectionData>();
 		SplineProjectionData->Initialize(this, InOther, InParams);
 		return SplineProjectionData;
 	}
diff --git a/Engine/Plugins/Experimental/PCG/Source/PCG/Private/Data/PCGTextureData.cpp b/Engine/Plugins/Experimental/PCG/Source/PCG/Private/Data/PCGTextureData.cpp
--- a/Engine/Plugins/Experimental/PCG/Source/PCG/Private/Data/PCGTextureData.cpp
+++ b/Engine/Plugins/Experimental/PCG/Source/PCG/Private/Data/PCGTextureData.cpp
@@ -245,7 +245,7 @@ void UPCGTextureData::Initialize(UTexture2D* InTexture, const FTransform& InTran
 	Width = 0;
 	Height = 0;
 
-	if (Texture)
+	if (InTexture)
 	{
 		FTexturePlatformData* PlatformData = Texture->GetPlatformData();
 		if (PlatformData && PlatformData->Mips.Num() > 0)
diff --git a/Engine/Plugins/Experimental/PCG/Source/PCG/Private/Data/PCGVolumeData.cpp b/Engine/Plugins/Experimental/PCG/Source/PCG/Private/Data/PCGVolumeData.cpp
--- a/Engine/Plugins/Experimental/PCG/Source/PCG/Private/Data/PCGVolumeData.cpp
+++ b/Engine/Plugins/Experimental/PCG/Source/PCG/Private/Data/PCGVolumeData.cpp
@@ -65,7 +65,7 @@ bool UPCGVolumeData::SamplePoint(const FTransform& InTransform, const FBox& InBo
 	{
 		float PointDensity = 0.0f;
 
-		if (!Volume || PCGHelpers::IsInsideBounds(GetStrictBounds(), InPosition))
+		if (!Volume.IsValid() || PCGHelpers::IsInsideBounds(GetStrictBounds(), InPosition))
 		{
 			PointDensity = 1.0f;
 		}
diff --git a/Engine/Plugins/Experimental/PCG/Source/PCG/Private/Elements/PCGDataFromActor.cpp b/Engine/Plugins/Experimental/PCG/Source/PCG/Private/Elements/PCGDataFromActor.cpp
--- a/Engine/Plugins/Experimental/PCG/Source/PCG/Private/Elements/PCGDataFromActor.cpp
+++ b/Engine/Plugins/Experimental/PCG/Source/PCG/Private/Elements/PCGDataFromActor.cpp
@@ -9,6 +9,21 @@ FPCGElementPtr UPCGDataFromActorSettings::CreateElement() const
 	return MakeShared<FPCGDataFromActorElement>();
 }
 
+TArray<FPCGPinProperties> UPCGDataFromActorSettings::OutputPinProperties() const
+{
+	TArray<FPCGPinProperties> Pins = Super::OutputPinProperties();
+
+	if (Mode == EPCGGetDataFromActorMode::GetDataFromPCGComponent || Mode == EPCGGetDataFromActorMode::GetDataFromPCGComponentOrParseComponents)
+	{
+		for (const FName& Pin : ExpectedPins)
+		{
+			Pins.Emplace(Pin);
+		}
+	}
+
+	return Pins;
+}
+
 bool FPCGDataFromActorElement::ExecuteInternal(FPCGContext* Context) const
 {
 	TRACE_CPUPROFILER_EVENT_SCOPE(FPCGDataFromActorElement::Execute);
@@ -26,9 +41,88 @@ bool FPCGDataFromActorElement::ExecuteInternal(FPCGContext* Context) const
 		return true;
 	}
 
+	TArray<UPCGComponent*> PCGComponents;
+	bool bHasGeneratedPCGData = false;
+	FProperty* FoundProperty = nullptr;
+
+	if (Settings->Mode == EPCGGetDataFromActorMode::GetDataFromPCGComponent || Settings->Mode == EPCGGetDataFromActorMode::GetDataFromPCGComponentOrParseComponents)
+	{
+		FoundActor->GetComponents(PCGComponents);
+
+		for (UPCGComponent* Component : PCGComponents)
+		{
+			bHasGeneratedPCGData |= !Component->GetGeneratedGraphOutput().TaggedData.IsEmpty();
+		}
+	}
+	else if (Settings->Mode == EPCGGetDataFromActorMode::GetDataFromProperty)
+	{
+		if (Settings->PropertyName != NAME_None)
+		{
+			FoundProperty = FindFProperty<FProperty>(FoundActor->GetClass(), Settings->PropertyName);
+		}
+	}
+
+	// Some additional validation
+	if (Settings->Mode == EPCGGetDataFromActorMode::GetDataFromPCGComponent && !bHasGeneratedPCGData)
+	{
+		PCGE_LOG(Warning, "Actor (%s) does not have any previously generated data.", *FoundActor->GetFName().ToString());
+		return true;
+	}
+	else if (Settings->Mode == EPCGGetDataFromActorMode::GetDataFromProperty && !FoundProperty)
+	{
+		PCGE_LOG(Warning, "Actor (%s) does not have a property name (%s)", *FoundActor->GetFName().ToString(), *Settings->PropertyName.ToString());
+		return true;
+	}
+
 	TArray<FPCGTaggedData>& Outputs = Context->OutputData.TaggedData;
-	FPCGTaggedData& Output = Outputs.Emplace_GetRef();
-	Output.Data = UPCGComponent::CreateActorPCGData(FoundActor, Context->SourceComponent.Get(), /*bParseActor=*/true);
+
+	if (bHasGeneratedPCGData)
+	{
+		for (UPCGComponent* Component : PCGComponents)
+		{
+			// TODO - Temporary behavior
+			// At the moment, intersections that reside in the transient package can hold on to a reference on this data
+			// which prevents proper garbage collection on map change, hence why we duplicate here. Normally, we would expect
+			// this not to be a problem, as these intersections should be garbage collected, but this requires more investigation
+			for (const FPCGTaggedData& TaggedData : Component->GetGeneratedGraphOutput().TaggedData)
+			{
+				FPCGTaggedData& DuplicatedTaggedData = Outputs.Add_GetRef(TaggedData);
+				DuplicatedTaggedData.Data = Cast<UPCGData>(StaticDuplicateObject(TaggedData.Data, GetTransientPackage()));
+			}
+			//Outputs.Append(Component->GetGeneratedGraphOutput().TaggedData);
+		}
+	}
+	else if (FoundProperty)
+	{
+		bool bAbleToGetProperty = false;
+		const void* PropertyAddressData = FoundProperty->ContainerPtrToValuePtr<void>(FoundActor);
+		// TODO: support more property types here
+		// Pointer to UPCGData
+		// Soft object pointer to UPCGData
+		// Array of pcg data -> all on the default pin
+		// Map of pcg data -> use key for pin? might not be robust
+		// FPCGDataCollection
+		if (const FStructProperty* StructProperty = CastField<FStructProperty>(FoundProperty))
+		{
+			if (StructProperty->Struct == FPCGDataCollection::StaticStruct())
+			{
+				const FPCGDataCollection* CollectionInProperty = reinterpret_cast<const FPCGDataCollection*>(PropertyAddressData);
+				Outputs.Append(CollectionInProperty->TaggedData);
+
+				bAbleToGetProperty = true;
+			}
+		}
+
+		if (!bAbleToGetProperty)
+		{
+			PCGE_LOG(Warning, "Actor (%s) property (%s) does not have a supported type", *FoundActor->GetFName().ToString(), *Settings->PropertyName.ToString());
+		}
+	}
+	else
+	{
+		FPCGTaggedData& Output = Outputs.Emplace_GetRef();
+		Output.Data = UPCGComponent::CreateActorPCGData(FoundActor, Context->SourceComponent.Get(), /*bParseActor=*/Settings->Mode != EPCGGetDataFromActorMode::GetSinglePoint);
+	}
 
 	return true;
 }
\ No newline at end of file
diff --git a/Engine/Plugins/Experimental/PCG/Source/PCG/Private/Graph/PCGGraphExecutor.cpp b/Engine/Plugins/Experimental/PCG/Source/PCG/Private/Graph/PCGGraphExecutor.cpp
--- a/Engine/Plugins/Experimental/PCG/Source/PCG/Private/Graph/PCGGraphExecutor.cpp
+++ b/Engine/Plugins/Experimental/PCG/Source/PCG/Private/Graph/PCGGraphExecutor.cpp
@@ -50,6 +50,7 @@ FPCGGraphExecutor::~FPCGGraphExecutor()
 	// We don't really need to do this here (it would be done in the destructor of these both)
 	// but this is to clarify/ensure the order in which this happens
 	GraphCache.ClearCache();
+	ensure(DataRootSet.RootSet.IsEmpty());
 	DataRootSet.Clear();
 
 #if WITH_EDITOR
diff --git a/Engine/Plugins/Experimental/PCG/Source/PCG/Private/PCGComponent.cpp b/Engine/Plugins/Experimental/PCG/Source/PCG/Private/PCGComponent.cpp
--- a/Engine/Plugins/Experimental/PCG/Source/PCG/Private/PCGComponent.cpp
+++ b/Engine/Plugins/Experimental/PCG/Source/PCG/Private/PCGComponent.cpp
@@ -337,6 +337,8 @@ void UPCGComponent::PostProcessGraph(const FBox& InNewBounds, bool bInGenerated,
 
 	CleanupUnusedManagedResources();
 
+	GeneratedGraphOutput.Reset();
+
 	if (bInGenerated)
 	{
 		bGenerated = true;
@@ -352,6 +354,18 @@ void UPCGComponent::PostProcessGraph(const FBox& InNewBounds, bool bInGenerated,
 
 		if (Context)
 		{
+			// TODO: should we filter based on supported serialized types here?
+			// TOOD: should reouter the contained data to this component
+			// .. and also remove it from the rootset information in the graph executor
+			//GeneratedGraphOutput = Context->InputData;
+			for (const FPCGTaggedData& TaggedData : Context->InputData.TaggedData)
+			{
+				FPCGTaggedData& DuplicatedTaggedData = GeneratedGraphOutput.TaggedData.Add_GetRef(TaggedData);
+				// TODO: outering the first layer might not be sufficient here - might need to expose
+				// some methods in the data to traverse all the data to outer everything for serialization
+				DuplicatedTaggedData.Data = Cast<UPCGData>(StaticDuplicateObject(TaggedData.Data, this));
+			}
+
 			// If the original component is partitioned, local components have to forward
 			// their inputs, so that they can be gathered by the original component.
 			// We don't have the info on the original component here, so forward for all
@@ -420,6 +434,7 @@ void UPCGComponent::PostCleanupGraph()
 {
 	bGenerated = false;
 	CurrentCleanupTask = InvalidPCGTaskId;
+	GeneratedGraphOutput.Reset();
 
 #if WITH_EDITOR
 	OnPCGGraphCleanedDelegate.Broadcast(this);
diff --git a/Engine/Plugins/Experimental/PCG/Source/PCG/Private/PCGData.cpp b/Engine/Plugins/Experimental/PCG/Source/PCG/Private/PCGData.cpp
--- a/Engine/Plugins/Experimental/PCG/Source/PCG/Private/PCGData.cpp
+++ b/Engine/Plugins/Experimental/PCG/Source/PCG/Private/PCGData.cpp
@@ -22,7 +22,7 @@ void FPCGRootSet::Add(UObject* InObject)
 	{
 		(*Found)++;
 	}
-	else if(!InObject->IsRooted())
+	else if(!InObject->IsRooted() && InObject->GetPackage() == GetTransientPackage())
 	{
 		InObject->AddToRoot();
 		RootSet.Emplace(InObject, 1);
@@ -30,7 +30,12 @@ void FPCGRootSet::Add(UObject* InObject)
 }
 void FPCGRootSet::Remove(UObject* InObject)
 {
-	check(InObject);
+	if (!InObject)
+	{
+		UE_LOG(LogPCG, Warning, TEXT("Trying to remove a null object from the rootset"));
+		return;
+	}
+	
 	if (int32* Found = RootSet.Find(InObject))
 	{
 		check(InObject->IsRooted());
@@ -185,6 +190,14 @@ void FPCGDataCollection::RemoveFromRootSet(FPCGRootSet& RootSet) const
 	}
 }
 
+void FPCGDataCollection::Reset()
+{
+	// Implementation note: We are assuming that there is no need to remove the data from the rootset here.
+	TaggedData.Reset();
+	bCancelExecutionOnEmpty = false;
+	bCancelExecution = false;
+}
+
 TArray<FPCGTaggedData> UPCGDataFunctionLibrary::GetInputs(const FPCGDataCollection& InCollection)
 {
 	return InCollection.GetInputs();
diff --git a/Engine/Plugins/Experimental/PCG/Source/PCG/Public/Data/PCGTextureData.h b/Engine/Plugins/Experimental/PCG/Source/PCG/Public/Data/PCGTextureData.h
--- a/Engine/Plugins/Experimental/PCG/Source/PCG/Public/Data/PCGTextureData.h
+++ b/Engine/Plugins/Experimental/PCG/Source/PCG/Public/Data/PCGTextureData.h
@@ -104,5 +104,5 @@ public:
 	void Initialize(UTexture2D* InTexture, const FTransform& InTransform);
 
 	UPROPERTY(BlueprintReadOnly, VisibleAnywhere, Category = Properties)
-	TObjectPtr<UTexture2D> Texture = nullptr;
+	TWeakObjectPtr<UTexture2D> Texture = nullptr;
 };
diff --git a/Engine/Plugins/Experimental/PCG/Source/PCG/Public/Data/PCGVolumeData.h b/Engine/Plugins/Experimental/PCG/Source/PCG/Public/Data/PCGVolumeData.h
--- a/Engine/Plugins/Experimental/PCG/Source/PCG/Public/Data/PCGVolumeData.h
+++ b/Engine/Plugins/Experimental/PCG/Source/PCG/Public/Data/PCGVolumeData.h
@@ -38,7 +38,7 @@ public:
 
 protected:
 	UPROPERTY(BlueprintReadWrite, EditAnywhere, Category = SourceData)
-	TObjectPtr<AVolume> Volume = nullptr;
+	TWeakObjectPtr<AVolume> Volume = nullptr;
 
 	UPROPERTY()
 	FBox Bounds = FBox(EForceInit::ForceInit);
diff --git a/Engine/Plugins/Experimental/PCG/Source/PCG/Public/Elements/PCGDataFromActor.h b/Engine/Plugins/Experimental/PCG/Source/PCG/Public/Elements/PCGDataFromActor.h
--- a/Engine/Plugins/Experimental/PCG/Source/PCG/Public/Elements/PCGDataFromActor.h
+++ b/Engine/Plugins/Experimental/PCG/Source/PCG/Public/Elements/PCGDataFromActor.h
@@ -8,6 +8,16 @@
 
 #include "PCGDataFromActor.generated.h"
 
+UENUM()
+enum class EPCGGetDataFromActorMode : uint8
+{
+	ParseActorComponents,
+	GetSinglePoint,
+	GetDataFromProperty,
+	GetDataFromPCGComponent,
+	GetDataFromPCGComponentOrParseComponents
+};
+
 UCLASS(BlueprintType, ClassGroup = (Procedural))
 class PCG_API UPCGDataFromActorSettings : public UPCGSettings
 {
@@ -21,6 +31,7 @@ public:
 #endif
 
 	virtual TArray<FPCGPinProperties> InputPinProperties() const override { return TArray<FPCGPinProperties>(); }
+	virtual TArray<FPCGPinProperties> OutputPinProperties() const override;
 
 protected:
 	virtual FPCGElementPtr CreateElement() const override;
@@ -29,6 +40,15 @@ protected:
 public:
 	UPROPERTY(BlueprintReadWrite, EditAnywhere, Category = Settings, meta = (ShowOnlyInnerProperties))
 	FPCGActorSelectorSettings ActorSelector;
+
+	UPROPERTY(BlueprintReadWrite, EditAnywhere, Category = Settings)
+	EPCGGetDataFromActorMode Mode = EPCGGetDataFromActorMode::ParseActorComponents;
+
+	UPROPERTY(BlueprintReadWrite, EditAnywhere, Category = Settings, meta = (EditCondition = "Mode == EPCGGetDataFromActorMode::GetDataFromPCGComponent || Mode == EPCGGetDataFromActorMode::GetDataFromPCGComponentOrParseComponents"))
+	TArray<FName> ExpectedPins;
+
+	UPROPERTY(BlueprintReadWrite, EditAnywhere, Category = Settings, meta = (EditCondition = "Mode == EPCGGetDataFromActorMode::GetDataFromProperty"))
+	FName PropertyName = NAME_None;
 };
 
 class FPCGDataFromActorElement : public FSimplePCGElement
diff --git a/Engine/Plugins/Experimental/PCG/Source/PCG/Public/PCGComponent.h b/Engine/Plugins/Experimental/PCG/Source/PCG/Public/PCGComponent.h
--- a/Engine/Plugins/Experimental/PCG/Source/PCG/Public/PCGComponent.h
+++ b/Engine/Plugins/Experimental/PCG/Source/PCG/Public/PCGComponent.h
@@ -133,6 +133,10 @@ public:
 	/** Cancels in-progress generation */
 	void CancelGeneration();
 
+	/** Retrieves generated data */
+	UFUNCTION(BlueprintCallable, Category = PCG)
+	const FPCGDataCollection& GetGeneratedGraphOutput() const { return GeneratedGraphOutput; }
+
 	/** Move all generated resources under a new actor, following a template (AActor if not provided), clearing all link to this PCG component. Returns the new actor.*/
 	UFUNCTION(BlueprintCallable, Category = PCG)
 	AActor* ClearPCGLink(UClass* TemplateActor = nullptr);
@@ -344,6 +348,9 @@ private:
 	UPROPERTY()
 	FBox LastGeneratedBounds = FBox(EForceInit::ForceInit);
 
+	UPROPERTY()
+	FPCGDataCollection GeneratedGraphOutput;
+
 	FPCGTaskId CurrentGenerationTask = InvalidPCGTaskId;
 	FPCGTaskId CurrentCleanupTask = InvalidPCGTaskId;
 
diff --git a/Engine/Plugins/Experimental/PCG/Source/PCG/Public/PCGData.h b/Engine/Plugins/Experimental/PCG/Source/PCG/Public/PCGData.h
--- a/Engine/Plugins/Experimental/PCG/Source/PCG/Public/PCGData.h
+++ b/Engine/Plugins/Experimental/PCG/Source/PCG/Public/PCGData.h
@@ -79,6 +79,9 @@ struct PCG_API FPCGDataCollection
 	void AddToRootSet(FPCGRootSet& RootSet) const;
 	void RemoveFromRootSet(FPCGRootSet& RootSet) const;
 
+	/** Cleans up the collection, but does not unroot any previously rooted data. */
+	void Reset();
+
 	UPROPERTY(BlueprintReadWrite, EditAnywhere, Category = Data)
 	TArray<FPCGTaggedData> TaggedData;
 
