Fixed missing includes

#rnx

diff --git a/Engine/Plugins/Experimental/PCG/Source/PCG/Private/Elements/PCGDataFromActor.cpp b/Engine/Plugins/Experimental/PCG/Source/PCG/Private/Elements/PCGDataFromActor.cpp
--- a/Engine/Plugins/Experimental/PCG/Source/PCG/Private/Elements/PCGDataFromActor.cpp
+++ b/Engine/Plugins/Experimental/PCG/Source/PCG/Private/Elements/PCGDataFromActor.cpp
@@ -4,6 +4,9 @@
 
 #include "PCGComponent.h"
 
+#include "GameFramework/Actor.h"
+#include "UObject/Package.h"
+
 FPCGElementPtr UPCGDataFromActorSettings::CreateElement() const
 {
 	return MakeShared<FPCGDataFromActorElement>();
