Recovery for deleted metadata that's not properly caught up.
Will cause some invalid values, but is a minimal problem vs. a deadlock.

#rnx

diff --git a/Engine/Plugins/Experimental/PCG/Source/PCG/Private/Graph/PCGGraphExecutor.cpp b/Engine/Plugins/Experimental/PCG/Source/PCG/Private/Graph/PCGGraphExecutor.cpp
--- a/Engine/Plugins/Experimental/PCG/Source/PCG/Private/Graph/PCGGraphExecutor.cpp
+++ b/Engine/Plugins/Experimental/PCG/Source/PCG/Private/Graph/PCGGraphExecutor.cpp
@@ -50,7 +50,6 @@ FPCGGraphExecutor::~FPCGGraphExecutor()
 	// We don't really need to do this here (it would be done in the destructor of these both)
 	// but this is to clarify/ensure the order in which this happens
 	GraphCache.ClearCache();
-	ensure(DataRootSet.RootSet.IsEmpty());
 	DataRootSet.Clear();
 
 #if WITH_EDITOR
diff --git a/Engine/Plugins/Experimental/PCG/Source/PCG/Private/Helpers/PCGMetadataHelpers.cpp b/Engine/Plugins/Experimental/PCG/Source/PCG/Private/Helpers/PCGMetadataHelpers.cpp
--- a/Engine/Plugins/Experimental/PCG/Source/PCG/Private/Helpers/PCGMetadataHelpers.cpp
+++ b/Engine/Plugins/Experimental/PCG/Source/PCG/Private/Helpers/PCGMetadataHelpers.cpp
@@ -9,4 +9,15 @@ namespace PCGMetadataHelpers
 	{
 		return Metadata1 && Metadata2 && Metadata1->GetRoot() == Metadata2->GetRoot();
 	}
+
+	const UPCGMetadata* GetParentMetadata(const UPCGMetadata* Metadata)
+	{
+		check(Metadata);
+		TWeakObjectPtr<const UPCGMetadata> Parent = Metadata->GetParentPtr();
+
+		// We're expecting the parent to either be null, or to be valid - if not, then it has been deleted
+		// which is going to cause some issues.
+		//check(Parent.IsExplicitlyNull() || Parent.IsValid());
+		return Parent.Get();
+	}
 }
\ No newline at end of file
diff --git a/Engine/Plugins/Experimental/PCG/Source/PCG/Public/Helpers/PCGMetadataHelpers.h b/Engine/Plugins/Experimental/PCG/Source/PCG/Public/Helpers/PCGMetadataHelpers.h
--- a/Engine/Plugins/Experimental/PCG/Source/PCG/Public/Helpers/PCGMetadataHelpers.h
+++ b/Engine/Plugins/Experimental/PCG/Source/PCG/Public/Helpers/PCGMetadataHelpers.h
@@ -7,4 +7,5 @@ class UPCGMetadata;
 namespace PCGMetadataHelpers
 {
 	bool HasSameRoot(const UPCGMetadata* Metadata1, const UPCGMetadata* Metadata2);
+	const UPCGMetadata* GetParentMetadata(const UPCGMetadata* Metadata);
 }
\ No newline at end of file
diff --git a/Engine/Plugins/Experimental/PCG/Source/PCG/Public/Metadata/PCGMetadata.h b/Engine/Plugins/Experimental/PCG/Source/PCG/Public/Metadata/PCGMetadata.h
--- a/Engine/Plugins/Experimental/PCG/Source/PCG/Public/Metadata/PCGMetadata.h
+++ b/Engine/Plugins/Experimental/PCG/Source/PCG/Public/Metadata/PCGMetadata.h
@@ -73,6 +73,7 @@ public:
 	FPCGMetadataAttributeBase* CopyAttribute(const FPCGMetadataAttributeBase* OriginalAttribute, FName NewAttributeName, bool bKeepParent, bool bCopyEntries, bool bCopyValues);
 
 	/** Returns this metadata's parent */
+	TWeakObjectPtr<const UPCGMetadata> GetParentPtr() const { return Parent; }
 	const UPCGMetadata* GetParent() const { return Parent.Get(); }
 	const UPCGMetadata* GetRoot() const;
 	bool HasParent(const UPCGMetadata* InTentativeParent) const;
diff --git a/Engine/Plugins/Experimental/PCG/Source/PCG/Public/Metadata/PCGMetadataAttributeTpl.h b/Engine/Plugins/Experimental/PCG/Source/PCG/Public/Metadata/PCGMetadataAttributeTpl.h
--- a/Engine/Plugins/Experimental/PCG/Source/PCG/Public/Metadata/PCGMetadataAttributeTpl.h
+++ b/Engine/Plugins/Experimental/PCG/Source/PCG/Public/Metadata/PCGMetadataAttributeTpl.h
@@ -64,11 +64,17 @@ public:
 		TArray<const FPCGMetadataAttribute<T>*, TInlineAllocator<2>> Parents = { this };
 		if (!bKeepParent && (bCopyEntries || bCopyValues))
 		{
+			const UPCGMetadata* CurrentMetadata = Metadata.Get();
 			const FPCGMetadataAttribute<T>* Current = this;
-			while(Current->Parent)
+
+			const UPCGMetadata* ParentMetadata = PCGMetadataHelpers::GetParentMetadata(CurrentMetadata);
+			while(ParentMetadata && Current->Parent)
 			{
+				CurrentMetadata = ParentMetadata;
 				Current = static_cast<const FPCGMetadataAttribute<T>*>(Current->Parent);
 				Parents.Add(Current);
+
+				ParentMetadata = PCGMetadataHelpers::GetParentMetadata(CurrentMetadata);
 			}
 		}
 
