[PCG] Add get/set bool and int32 in BP + can create int32/float attribute in PCG

diff --git a/Engine/Plugins/Experimental/PCG/Source/PCG/Private/Elements/PCGCreateAttribute.cpp b/Engine/Plugins/Experimental/PCG/Source/PCG/Private/Elements/PCGCreateAttribute.cpp
--- a/Engine/Plugins/Experimental/PCG/Source/PCG/Private/Elements/PCGCreateAttribute.cpp
+++ b/Engine/Plugins/Experimental/PCG/Source/PCG/Private/Elements/PCGCreateAttribute.cpp
@@ -27,6 +27,10 @@ namespace PCGCreateAttributeElement
 		{
 		case EPCGMetadataTypes::Integer64:
 			return Callback(PCG_GET_OVERRIDEN_VALUE(Settings, IntValue, Params));
+		case EPCGMetadataTypes::Integer32:
+			return Callback(PCG_GET_OVERRIDEN_VALUE(Settings, Int32Value, Params));
+		case EPCGMetadataTypes::Float:
+			return Callback(PCG_GET_OVERRIDEN_VALUE(Settings, FloatValue, Params));
 		case EPCGMetadataTypes::Double:
 			return Callback(PCG_GET_OVERRIDEN_VALUE(Settings, DoubleValue, Params));
 		case EPCGMetadataTypes::Vector2:
@@ -83,6 +87,10 @@ FName UPCGCreateAttributeSettings::AdditionalTaskName() const
 
 		switch (Type)
 		{
+		case EPCGMetadataTypes::Integer32:
+			return FName(FString::Printf(TEXT("%s: %d"), *Name, Int32Value));
+		case EPCGMetadataTypes::Float:
+			return FName(FString::Printf(TEXT("%s: %.2f"), *Name, FloatValue));
 		case EPCGMetadataTypes::Integer64:
 			return FName(FString::Printf(TEXT("%s: %lld"), *Name, IntValue));
 		case EPCGMetadataTypes::Double:
diff --git a/Engine/Plugins/Experimental/PCG/Source/PCG/Private/Metadata/PCGMetadata.cpp b/Engine/Plugins/Experimental/PCG/Source/PCG/Private/Metadata/PCGMetadata.cpp
--- a/Engine/Plugins/Experimental/PCG/Source/PCG/Private/Metadata/PCGMetadata.cpp
+++ b/Engine/Plugins/Experimental/PCG/Source/PCG/Private/Metadata/PCGMetadata.cpp
@@ -367,6 +367,11 @@ bool UPCGMetadata::ParentHasAttribute(FName AttributeName) const
 	return Parent.IsValid() && Parent->HasAttribute(AttributeName);
 }
 
+void UPCGMetadata::CreateInteger32Attribute(FName AttributeName, int32 DefaultValue, bool bAllowsInterpolation, bool bOverrideParent)
+{
+	CreateAttribute<int32>(AttributeName, DefaultValue, bAllowsInterpolation, bOverrideParent);
+}
+
 void UPCGMetadata::CreateInteger64Attribute(FName AttributeName, int64 DefaultValue, bool bAllowsInterpolation, bool bOverrideParent)
 {
 	CreateAttribute<int64>(AttributeName, DefaultValue, bAllowsInterpolation, bOverrideParent);
@@ -417,6 +422,11 @@ void UPCGMetadata::CreateStringAttribute(FName AttributeName, FString DefaultVal
 	CreateAttribute<FString>(AttributeName, DefaultValue, bAllowsInterpolation, bOverrideParent);
 }
 
+void UPCGMetadata::CreateBoolAttribute(FName AttributeName, bool DefaultValue, bool bAllowsInterpolation, bool bOverrideParent)
+{
+	CreateAttribute<bool>(AttributeName, DefaultValue, bAllowsInterpolation, bOverrideParent);
+}
+
 namespace PCGMetadata
 {
 	template<typename DataType>
diff --git a/Engine/Plugins/Experimental/PCG/Source/PCG/Private/Metadata/PCGMetadataAccessor.cpp b/Engine/Plugins/Experimental/PCG/Source/PCG/Private/Metadata/PCGMetadataAccessor.cpp
--- a/Engine/Plugins/Experimental/PCG/Source/PCG/Private/Metadata/PCGMetadataAccessor.cpp
+++ b/Engine/Plugins/Experimental/PCG/Source/PCG/Private/Metadata/PCGMetadataAccessor.cpp
@@ -89,6 +89,16 @@ bool UPCGMetadataAccessorHelpers::HasAttributeSetByMetadataKey(PCGMetadataEntryK
 	}
 }
 
+int32 UPCGMetadataAccessorHelpers::GetInteger32AttributeByMetadataKey(int64 Key, const UPCGMetadata* Metadata, FName AttributeName)
+{
+	return GetAttribute<int32>(Key, Metadata, AttributeName);
+}
+
+void UPCGMetadataAccessorHelpers::SetInteger32AttributeByMetadataKey(int64& Key, UPCGMetadata* Metadata, FName AttributeName, int32 Value)
+{
+	SetAttribute(Key, Metadata, AttributeName, Value);
+}
+
 int64 UPCGMetadataAccessorHelpers::GetInteger64AttributeByMetadataKey(int64 Key, const UPCGMetadata* Metadata, FName AttributeName)
 {
 	return GetAttribute<int64>(Key, Metadata, AttributeName);
@@ -189,6 +199,16 @@ void UPCGMetadataAccessorHelpers::SetStringAttributeByMetadataKey(int64& Key, UP
 	SetAttribute(Key, Metadata, AttributeName, Value);
 }
 
+bool UPCGMetadataAccessorHelpers::GetBoolAttributeByMetadataKey(int64 Key, const UPCGMetadata* Metadata, FName AttributeName)
+{
+	return GetAttribute<bool>(Key, Metadata, AttributeName);
+}
+
+void UPCGMetadataAccessorHelpers::SetBoolAttributeByMetadataKey(int64& Key, UPCGMetadata* Metadata, FName AttributeName, bool Value)
+{
+	SetAttribute(Key, Metadata, AttributeName, Value);
+}
+
 bool UPCGMetadataAccessorHelpers::SetAttributeFromPropertyByMetadataKey(int64& Key, UPCGMetadata* Metadata, FName AttributeName, const UObject* Object, FName PropertyName)
 {
 	if (!Object)
@@ -254,6 +274,16 @@ void UPCGMetadataAccessorHelpers::InitializeMetadataWithParent(FPCGPoint& Point,
 	Point.MetadataEntry = Metadata ? (Metadata->HasParent(ParentMetadata) ? Metadata->AddEntry(ParentPoint.MetadataEntry) : Metadata->AddEntry()) : PCGInvalidEntryKey;
 }
 
+int32 UPCGMetadataAccessorHelpers::GetInteger32Attribute(const FPCGPoint& Point, const UPCGMetadata* Metadata, FName AttributeName)
+{
+	return GetAttribute<int32>(Point.MetadataEntry, Metadata, AttributeName);
+}
+
+void UPCGMetadataAccessorHelpers::SetInteger32Attribute(FPCGPoint& Point, UPCGMetadata* Metadata, FName AttributeName, int32 Value)
+{
+	SetAttribute(Point.MetadataEntry, Metadata, AttributeName, Value);
+}
+
 int64 UPCGMetadataAccessorHelpers::GetInteger64Attribute(const FPCGPoint& Point, const UPCGMetadata* Metadata, FName AttributeName)
 {
 	return GetAttribute<int64>(Point.MetadataEntry, Metadata, AttributeName);
@@ -354,6 +384,16 @@ void UPCGMetadataAccessorHelpers::SetStringAttribute(FPCGPoint& Point, UPCGMetad
 	SetAttribute(Point.MetadataEntry, Metadata, AttributeName, Value);
 }
 
+bool UPCGMetadataAccessorHelpers::GetBoolAttribute(const FPCGPoint& Point, const UPCGMetadata* Metadata, FName AttributeName)
+{
+	return GetAttribute<bool>(Point.MetadataEntry, Metadata, AttributeName);
+}
+
+void UPCGMetadataAccessorHelpers::SetBoolAttribute(FPCGPoint& Point, UPCGMetadata* Metadata, FName AttributeName, bool Value)
+{
+	SetAttribute(Point.MetadataEntry, Metadata, AttributeName, Value);
+}
+
 bool UPCGMetadataAccessorHelpers::HasAttributeSet(const FPCGPoint& Point, const UPCGMetadata* Metadata, FName AttributeName)
 {
 	return HasAttributeSetByMetadataKey(Point.MetadataEntry, Metadata, AttributeName);
diff --git a/Engine/Plugins/Experimental/PCG/Source/PCG/Public/Elements/PCGCreateAttribute.h b/Engine/Plugins/Experimental/PCG/Source/PCG/Public/Elements/PCGCreateAttribute.h
--- a/Engine/Plugins/Experimental/PCG/Source/PCG/Public/Elements/PCGCreateAttribute.h
+++ b/Engine/Plugins/Experimental/PCG/Source/PCG/Public/Elements/PCGCreateAttribute.h
@@ -49,10 +49,16 @@ public:
 	UPROPERTY(BlueprintReadWrite, EditAnywhere, Category = Settings, meta = (EditCondition = "bFromSourceParam", EditConditionHides))
 	FName SourceParamAttributeName = NAME_None;
 
-	UPROPERTY(BlueprintReadWrite, EditAnywhere, Category = Settings, meta = (ValidEnumValues = "Double, Integer64, Vector2, Vector, Vector4, Quaternion, Transform, String, Boolean, Rotator, Name", EditCondition = "!bFromSourceParam", EditConditionHides))
+	UPROPERTY(BlueprintReadWrite, EditAnywhere, Category = Settings, meta = (ValidEnumValues = "Float, Double, Integer32, Integer64, Vector2, Vector, Vector4, Quaternion, Transform, String, Boolean, Rotator, Name", EditCondition = "!bFromSourceParam", EditConditionHides))
 	EPCGMetadataTypes Type = EPCGMetadataTypes::Double;
 
 	// All different types
+	UPROPERTY(BlueprintReadWrite, EditAnywhere, Category = Settings, meta = (EditCondition = "Type == EPCGMetadataTypes::Float && !bFromSourceParam", EditConditionHides))
+	float FloatValue = 0.0;
+
+	UPROPERTY(BlueprintReadWrite, EditAnywhere, Category = Settings, meta = (EditCondition = "Type == EPCGMetadataTypes::Integer32 && !bFromSourceParam", EditConditionHides))
+	int32 Int32Value = 0;
+
 	UPROPERTY(BlueprintReadWrite, EditAnywhere, Category = Settings, meta = (EditCondition = "Type == EPCGMetadataTypes::Double && !bFromSourceParam", EditConditionHides))
 	double DoubleValue = 0.0;
 
diff --git a/Engine/Plugins/Experimental/PCG/Source/PCG/Public/Metadata/PCGMetadata.h b/Engine/Plugins/Experimental/PCG/Source/PCG/Public/Metadata/PCGMetadata.h
--- a/Engine/Plugins/Experimental/PCG/Source/PCG/Public/Metadata/PCGMetadata.h
+++ b/Engine/Plugins/Experimental/PCG/Source/PCG/Public/Metadata/PCGMetadata.h
@@ -79,6 +79,9 @@ public:
 	bool HasParent(const UPCGMetadata* InTentativeParent) const;
 
 	/** Create new streams */
+	UFUNCTION(BlueprintCallable, Category = "PCG|Metadata")
+	void CreateInteger32Attribute(FName AttributeName, int32 DefaultValue, bool bAllowsInterpolation, bool bOverrideParent = true);
+
 	UFUNCTION(BlueprintCallable, Category = "PCG|Metadata")
 	void CreateInteger64Attribute(FName AttributeName, int64 DefaultValue, bool bAllowsInterpolation, bool bOverrideParent = true);
 
@@ -109,6 +112,9 @@ public:
 	UFUNCTION(BlueprintCallable, Category = "PCG|Metadata")
 	void CreateStringAttribute(FName AttributeName, FString DefaultValue, bool bAllowsInterpolation, bool bOverrideParent = true);
 
+	UFUNCTION(BlueprintCallable, Category = "PCG|Metadata")
+	void CreateBoolAttribute(FName AttributeName, bool DefaultValue, bool bAllowsInterpolation, bool bOverrideParent = true);
+
 	/** Set an attribute given a property and its value.
 	* @param AttributeName: Target attribute to set the property's value to
 	* @param EntryKey: Metadata entry key to set the value to
diff --git a/Engine/Plugins/Experimental/PCG/Source/PCG/Public/Metadata/PCGMetadataAccessor.h b/Engine/Plugins/Experimental/PCG/Source/PCG/Public/Metadata/PCGMetadataAccessor.h
--- a/Engine/Plugins/Experimental/PCG/Source/PCG/Public/Metadata/PCGMetadataAccessor.h
+++ b/Engine/Plugins/Experimental/PCG/Source/PCG/Public/Metadata/PCGMetadataAccessor.h
@@ -18,6 +18,12 @@ class PCG_API UPCGMetadataAccessorHelpers : public UBlueprintFunctionLibrary
 
 public:
 	/** Id-based metadata functions */
+	UFUNCTION(BlueprintCallable, Category = "PCG|Metadata")
+	static int32 GetInteger32AttributeByMetadataKey(int64 Key, const UPCGMetadata* Metadata, FName AttributeName);
+
+	UFUNCTION(BlueprintCallable, Category = "PCG|Metadata")
+	static void SetInteger32AttributeByMetadataKey(UPARAM(ref) int64& Key, UPCGMetadata* Metadata, FName AttributeName, int32 Value);
+
 	UFUNCTION(BlueprintCallable, Category = "PCG|Metadata")
 	static int64 GetInteger64AttributeByMetadataKey(int64 Key, const UPCGMetadata* Metadata, FName AttributeName);
 
@@ -78,6 +84,12 @@ public:
 	UFUNCTION(BlueprintCallable, Category = "PCG|Metadata")
 	static void SetStringAttributeByMetadataKey(UPARAM(ref) int64& Key, UPCGMetadata* Metadata, FName AttributeName, const FString& Value);
 
+	UFUNCTION(BlueprintCallable, Category = "PCG|Metadata")
+	static bool GetBoolAttributeByMetadataKey(int64 Key, const UPCGMetadata* Metadata, FName AttributeName);
+
+	UFUNCTION(BlueprintCallable, Category = "PCG|Metadata")
+	static void SetBoolAttributeByMetadataKey(UPARAM(ref) int64& Key, UPCGMetadata* Metadata, FName AttributeName, bool Value);
+
 	UFUNCTION(BlueprintCallable, Category = "PCG|Metadata")
 	static bool SetAttributeFromPropertyByMetadataKey(UPARAM(ref) int64& Key, UPCGMetadata* Metadata, FName AttributeName, const UObject* Object, FName PropertyName);
 
@@ -96,6 +108,12 @@ public:
 	UFUNCTION(BlueprintCallable, Category = "PCG|Metadata", meta = (ScriptMethod))
 	static void InitializeMetadata(UPARAM(ref) FPCGPoint& Point, UPCGMetadata* Metadata, const FPCGPoint& ParentPoint, const UPCGMetadata* ParentMetadata = nullptr);
 
+	UFUNCTION(BlueprintCallable, Category = "PCG|Metadata", meta = (ScriptMethod))
+	static int32 GetInteger32Attribute(const FPCGPoint& Point, const UPCGMetadata* Metadata, FName AttributeName);
+
+	UFUNCTION(BlueprintCallable, Category = "PCG|Metadata", meta = (ScriptMethod))
+	static void SetInteger32Attribute(UPARAM(ref) FPCGPoint& Point, UPCGMetadata* Metadata, FName AttributeName, int32 Value);
+
 	UFUNCTION(BlueprintCallable, Category = "PCG|Metadata", meta = (ScriptMethod))
 	static int64 GetInteger64Attribute(const FPCGPoint& Point, const UPCGMetadata* Metadata, FName AttributeName);
 
@@ -156,6 +174,12 @@ public:
 	UFUNCTION(BlueprintCallable, Category = "PCG|Metadata", meta = (ScriptMethod))
 	static void SetStringAttribute(UPARAM(ref) FPCGPoint& Point, UPCGMetadata* Metadata, FName AttributeName, const FString& Value);
 
+	UFUNCTION(BlueprintCallable, Category = "PCG|Metadata", meta = (ScriptMethod))
+	static bool GetBoolAttribute(const FPCGPoint& Point, const UPCGMetadata* Metadata, FName AttributeName);
+
+	UFUNCTION(BlueprintCallable, Category = "PCG|Metadata", meta = (ScriptMethod))
+	static void SetBoolAttribute(UPARAM(ref) FPCGPoint& Point, UPCGMetadata* Metadata, FName AttributeName, bool Value);
+
 	UFUNCTION(BlueprintCallable, Category = "PCG|Metadata", meta = (ScriptMethod))
 	static bool HasAttributeSet(const FPCGPoint& Point, const UPCGMetadata* Metadata, FName AttributeName);
 
