[PCG] Add Copy function to remove DuplicateObject, which can be very costly for spatial data

* Duplicating the data should be better everywhere
* Re-enable multi-threading on all node that were depending on DuplicateObject
* Re-enable cache for metadata operation

diff --git a/Engine/Plugins/Experimental/PCG/Source/PCG/Private/Data/PCGDifferenceData.cpp b/Engine/Plugins/Experimental/PCG/Source/PCG/Private/Data/PCGDifferenceData.cpp
--- a/Engine/Plugins/Experimental/PCG/Source/PCG/Private/Data/PCGDifferenceData.cpp
+++ b/Engine/Plugins/Experimental/PCG/Source/PCG/Private/Data/PCGDifferenceData.cpp
@@ -211,4 +211,19 @@ const UPCGPointData* UPCGDifferenceData::CreatePointData(FPCGContext* Context) c
 	UE_LOG(LogPCG, Verbose, TEXT("Difference generated %d points from %d source points"), TargetPoints.Num(), SourcePointData->GetPoints().Num());
 
 	return Data;
+}
+
+UPCGSpatialData* UPCGDifferenceData::CopyInternal() const
+{
+	UPCGDifferenceData* NewDifferenceData = NewObject<UPCGDifferenceData>();
+
+	NewDifferenceData->Source = Source;
+	NewDifferenceData->Difference = Difference;
+	NewDifferenceData->DensityFunction = DensityFunction;
+	if (DifferencesUnion)
+	{
+		NewDifferenceData->DifferencesUnion = static_cast<UPCGUnionData*>(DifferencesUnion->DuplicateData());
+	}
+
+	return NewDifferenceData;
 }
\ No newline at end of file
diff --git a/Engine/Plugins/Experimental/PCG/Source/PCG/Private/Data/PCGIntersectionData.cpp b/Engine/Plugins/Experimental/PCG/Source/PCG/Private/Data/PCGIntersectionData.cpp
--- a/Engine/Plugins/Experimental/PCG/Source/PCG/Private/Data/PCGIntersectionData.cpp
+++ b/Engine/Plugins/Experimental/PCG/Source/PCG/Private/Data/PCGIntersectionData.cpp
@@ -195,4 +195,17 @@ UPCGPointData* UPCGIntersectionData::CreateAndFilterPointData(FPCGContext* Conte
 	UE_LOG(LogPCG, Verbose, TEXT("Intersection generated %d points from %d source points"), TargetPoints.Num(), SourcePoints.Num());
 
 	return Data;
+}
+
+UPCGSpatialData* UPCGIntersectionData::CopyInternal() const
+{
+	UPCGIntersectionData* NewIntersectionData = NewObject<UPCGIntersectionData>();
+
+	NewIntersectionData->DensityFunction = DensityFunction;
+	NewIntersectionData->A = A;
+	NewIntersectionData->B = B;
+	NewIntersectionData->CachedBounds = CachedBounds;
+	NewIntersectionData->CachedStrictBounds = CachedStrictBounds;
+
+	return NewIntersectionData;
 }
\ No newline at end of file
diff --git a/Engine/Plugins/Experimental/PCG/Source/PCG/Private/Data/PCGLandscapeData.cpp b/Engine/Plugins/Experimental/PCG/Source/PCG/Private/Data/PCGLandscapeData.cpp
--- a/Engine/Plugins/Experimental/PCG/Source/PCG/Private/Data/PCGLandscapeData.cpp
+++ b/Engine/Plugins/Experimental/PCG/Source/PCG/Private/Data/PCGLandscapeData.cpp
@@ -284,4 +284,20 @@ const ULandscapeInfo* UPCGLandscapeData::GetLandscapeInfo(const FVector& InPosit
 	}
 
 	return nullptr;
+}
+
+UPCGSpatialData* UPCGLandscapeData::CopyInternal() const
+{
+	UPCGLandscapeData* NewLandsacapeData = NewObject<UPCGLandscapeData>();
+
+	CopyBaseSurfaceData(NewLandsacapeData);
+
+	NewLandsacapeData->Landscapes = Landscapes;
+	NewLandsacapeData->Bounds = Bounds;
+	NewLandsacapeData->bHeightOnly = bHeightOnly;
+	NewLandsacapeData->bUseMetadata = bUseMetadata;
+	NewLandsacapeData->LandscapeInfos = LandscapeInfos;
+	NewLandsacapeData->LandscapeCache = LandscapeCache;
+
+	return NewLandsacapeData;
 }
\ No newline at end of file
diff --git a/Engine/Plugins/Experimental/PCG/Source/PCG/Private/Data/PCGLandscapeSplineData.cpp b/Engine/Plugins/Experimental/PCG/Source/PCG/Private/Data/PCGLandscapeSplineData.cpp
--- a/Engine/Plugins/Experimental/PCG/Source/PCG/Private/Data/PCGLandscapeSplineData.cpp
+++ b/Engine/Plugins/Experimental/PCG/Source/PCG/Private/Data/PCGLandscapeSplineData.cpp
@@ -217,4 +217,13 @@ bool UPCGLandscapeSplineData::SamplePoint(const FTransform& InTransform, const F
 
 	OutPoint.Density = PointDensity;
 	return OutPoint.Density > 0;
+}
+
+UPCGSpatialData* UPCGLandscapeSplineData::CopyInternal() const
+{
+	UPCGLandscapeSplineData* NewLandscapeSplineData = NewObject<UPCGLandscapeSplineData>();
+
+	NewLandscapeSplineData->Spline = Spline;
+
+	return NewLandscapeSplineData;
 }
\ No newline at end of file
diff --git a/Engine/Plugins/Experimental/PCG/Source/PCG/Private/Data/PCGPointData.cpp b/Engine/Plugins/Experimental/PCG/Source/PCG/Private/Data/PCGPointData.cpp
--- a/Engine/Plugins/Experimental/PCG/Source/PCG/Private/Data/PCGPointData.cpp
+++ b/Engine/Plugins/Experimental/PCG/Source/PCG/Private/Data/PCGPointData.cpp
@@ -460,4 +460,12 @@ void UPCGPointData::RebuildOctree() const
 
 	Octree = NewOctree;
 	bOctreeIsDirty = false;
+}
+
+UPCGSpatialData* UPCGPointData::CopyInternal() const
+{
+	UPCGPointData* NewPointData = NewObject<UPCGPointData>();
+	NewPointData->GetMutablePoints() = GetPoints();
+
+	return NewPointData;
 }
\ No newline at end of file
diff --git a/Engine/Plugins/Experimental/PCG/Source/PCG/Private/Data/PCGPrimitiveData.cpp b/Engine/Plugins/Experimental/PCG/Source/PCG/Private/Data/PCGPrimitiveData.cpp
--- a/Engine/Plugins/Experimental/PCG/Source/PCG/Private/Data/PCGPrimitiveData.cpp
+++ b/Engine/Plugins/Experimental/PCG/Source/PCG/Private/Data/PCGPrimitiveData.cpp
@@ -50,4 +50,16 @@ const UPCGPointData* UPCGPrimitiveData::CreatePointData(FPCGContext* Context) co
 	}
 
 	return Data;
+}
+
+UPCGSpatialData* UPCGPrimitiveData::CopyInternal() const
+{
+	UPCGPrimitiveData* NewPrimitiveData = NewObject<UPCGPrimitiveData>();
+
+	NewPrimitiveData->VoxelSize = VoxelSize;
+	NewPrimitiveData->Primitive = Primitive;
+	NewPrimitiveData->CachedBounds = CachedBounds;
+	NewPrimitiveData->CachedStrictBounds = CachedStrictBounds;
+
+	return NewPrimitiveData;
 }
\ No newline at end of file
diff --git a/Engine/Plugins/Experimental/PCG/Source/PCG/Private/Data/PCGProjectionData.cpp b/Engine/Plugins/Experimental/PCG/Source/PCG/Private/Data/PCGProjectionData.cpp
--- a/Engine/Plugins/Experimental/PCG/Source/PCG/Private/Data/PCGProjectionData.cpp
+++ b/Engine/Plugins/Experimental/PCG/Source/PCG/Private/Data/PCGProjectionData.cpp
@@ -204,3 +204,21 @@ void UPCGProjectionData::GetIncludeExcludeAttributeNames(TSet<FName>& OutAttribu
 		OutAttributeNames.Add(FName(*Attribute));
 	}
 }
+
+inline void UPCGProjectionData::CopyBaseProjectionClass(UPCGProjectionData* NewProjectionData) const
+{
+	NewProjectionData->Source = Source;
+	NewProjectionData->Target = Target;
+	NewProjectionData->CachedBounds = CachedBounds;
+	NewProjectionData->CachedStrictBounds = CachedStrictBounds;
+	NewProjectionData->Params = Params;
+}
+
+UPCGSpatialData* UPCGProjectionData::CopyInternal() const
+{
+	UPCGProjectionData* NewProjectionData = NewObject<UPCGProjectionData>();
+
+	CopyBaseProjectionClass(NewProjectionData);
+
+	return NewProjectionData;
+}
\ No newline at end of file
diff --git a/Engine/Plugins/Experimental/PCG/Source/PCG/Private/Data/PCGRenderTargetData.cpp b/Engine/Plugins/Experimental/PCG/Source/PCG/Private/Data/PCGRenderTargetData.cpp
--- a/Engine/Plugins/Experimental/PCG/Source/PCG/Private/Data/PCGRenderTargetData.cpp
+++ b/Engine/Plugins/Experimental/PCG/Source/PCG/Private/Data/PCGRenderTargetData.cpp
@@ -38,4 +38,15 @@ void UPCGRenderTargetData::Initialize(UTextureRenderTarget2D* InRenderTarget, co
 	Bounds += FVector(-1.0f, -1.0f, 0.0f);
 	Bounds += FVector(1.0f, 1.0f, 0.0f);
 	Bounds = Bounds.TransformBy(Transform);
+}
+
+UPCGSpatialData* UPCGRenderTargetData::CopyInternal() const
+{
+	UPCGRenderTargetData* NewRenderTargetData = NewObject<UPCGRenderTargetData>();
+
+	CopyBaseTextureData(NewRenderTargetData);
+
+	NewRenderTargetData->RenderTarget = RenderTarget;
+
+	return NewRenderTargetData;
 }
\ No newline at end of file
diff --git a/Engine/Plugins/Experimental/PCG/Source/PCG/Private/Data/PCGSpatialData.cpp b/Engine/Plugins/Experimental/PCG/Source/PCG/Private/Data/PCGSpatialData.cpp
--- a/Engine/Plugins/Experimental/PCG/Source/PCG/Private/Data/PCGSpatialData.cpp
+++ b/Engine/Plugins/Experimental/PCG/Source/PCG/Private/Data/PCGSpatialData.cpp
@@ -147,4 +147,17 @@ void UPCGSpatialData::InitializeFromData(const UPCGSpatialData* InSource, const
 	{
 		UE_LOG(LogPCG, Warning, TEXT("InitializeFromData has both no source and no metadata override"));
 	}
+}
+
+UPCGSpatialData* UPCGSpatialData::DuplicateData(const bool bInitializeMetadata) const
+{
+	UPCGSpatialData* NewSpatialData = CopyInternal();
+	check(NewSpatialData);
+
+	if (bInitializeMetadata)
+	{
+		NewSpatialData->InitializeFromData(this);
+	}
+
+	return NewSpatialData;
 }
\ No newline at end of file
diff --git a/Engine/Plugins/Experimental/PCG/Source/PCG/Private/Data/PCGSplineData.cpp b/Engine/Plugins/Experimental/PCG/Source/PCG/Private/Data/PCGSplineData.cpp
--- a/Engine/Plugins/Experimental/PCG/Source/PCG/Private/Data/PCGSplineData.cpp
+++ b/Engine/Plugins/Experimental/PCG/Source/PCG/Private/Data/PCGSplineData.cpp
@@ -118,6 +118,16 @@ UPCGProjectionData* UPCGSplineData::ProjectOn(const UPCGSpatialData* InOther, co
 	}
 }
 
+UPCGSpatialData* UPCGSplineData::CopyInternal() const
+{
+	UPCGSplineData* NewSplineData = NewObject<UPCGSplineData>();
+
+	NewSplineData->Spline = Spline;
+	NewSplineData->CachedBounds = CachedBounds;
+
+	return NewSplineData;
+}
+
 FVector2D UPCGSplineProjectionData::Project(const FVector& InVector) const
 {
 	const FVector& SurfaceNormal = GetSurface()->GetNormal();
@@ -191,4 +201,15 @@ const UPCGSplineData* UPCGSplineProjectionData::GetSpline() const
 const UPCGSpatialData* UPCGSplineProjectionData::GetSurface() const
 {
 	return Target;
+}
+
+UPCGSpatialData* UPCGSplineProjectionData::CopyInternal() const
+{
+	UPCGSplineProjectionData* NewProjectionData = NewObject<UPCGSplineProjectionData>();
+
+	CopyBaseProjectionClass(NewProjectionData);
+
+	NewProjectionData->ProjectedPosition = ProjectedPosition;
+
+	return NewProjectionData;
 }
\ No newline at end of file
diff --git a/Engine/Plugins/Experimental/PCG/Source/PCG/Private/Data/PCGSurfaceData.cpp b/Engine/Plugins/Experimental/PCG/Source/PCG/Private/Data/PCGSurfaceData.cpp
--- a/Engine/Plugins/Experimental/PCG/Source/PCG/Private/Data/PCGSurfaceData.cpp
+++ b/Engine/Plugins/Experimental/PCG/Source/PCG/Private/Data/PCGSurfaceData.cpp
@@ -1,3 +1,8 @@
 // Copyright Epic Games, Inc. All Rights Reserved.
 
 #include "Data/PCGSurfaceData.h"
+
+void UPCGSurfaceData::CopyBaseSurfaceData(UPCGSurfaceData* NewSurfaceData) const
+{
+	NewSurfaceData->Transform = Transform;
+}
\ No newline at end of file
diff --git a/Engine/Plugins/Experimental/PCG/Source/PCG/Private/Data/PCGTextureData.cpp b/Engine/Plugins/Experimental/PCG/Source/PCG/Private/Data/PCGTextureData.cpp
--- a/Engine/Plugins/Experimental/PCG/Source/PCG/Private/Data/PCGTextureData.cpp
+++ b/Engine/Plugins/Experimental/PCG/Source/PCG/Private/Data/PCGTextureData.cpp
@@ -238,6 +238,25 @@ bool UPCGBaseTextureData::IsValid() const
 	return Height > 0 && Width > 0;
 }
 
+void UPCGBaseTextureData::CopyBaseTextureData(UPCGBaseTextureData* NewTextureData) const
+{
+	CopyBaseSurfaceData(NewTextureData);
+
+	NewTextureData->DensityFunction = DensityFunction;
+	NewTextureData->ColorChannel = ColorChannel;
+	NewTextureData->TexelSize = TexelSize;
+	NewTextureData->bUseAdvancedTiling = bUseAdvancedTiling;
+	NewTextureData->Tiling = Tiling;
+	NewTextureData->CenterOffset = CenterOffset;
+	NewTextureData->Rotation = Rotation;
+	NewTextureData->bUseTileBounds = bUseTileBounds;
+	NewTextureData->TileBounds = TileBounds;
+	NewTextureData->ColorData = ColorData;
+	NewTextureData->Bounds = Bounds;
+	NewTextureData->Height = Height;
+	NewTextureData->Width = Width;
+}
+
 void UPCGTextureData::Initialize(UTexture2D* InTexture, const FTransform& InTransform)
 {
 	Texture = InTexture;
@@ -308,4 +327,15 @@ void UPCGTextureData::Initialize(UTexture2D* InTexture, const FTransform& InTran
 	Bounds += FVector(-1.0f, -1.0f, 0.0f);
 	Bounds += FVector(1.0f, 1.0f, 0.0f);
 	Bounds = Bounds.TransformBy(Transform);
+}
+
+UPCGSpatialData* UPCGTextureData::CopyInternal() const
+{
+	UPCGTextureData* NewTextureData = NewObject<UPCGTextureData>();
+
+	CopyBaseTextureData(NewTextureData);
+
+	NewTextureData->Texture = Texture;
+
+	return NewTextureData;
 }
\ No newline at end of file
diff --git a/Engine/Plugins/Experimental/PCG/Source/PCG/Private/Data/PCGUnionData.cpp b/Engine/Plugins/Experimental/PCG/Source/PCG/Private/Data/PCGUnionData.cpp
--- a/Engine/Plugins/Experimental/PCG/Source/PCG/Private/Data/PCGUnionData.cpp
+++ b/Engine/Plugins/Experimental/PCG/Source/PCG/Private/Data/PCGUnionData.cpp
@@ -328,4 +328,19 @@ void UPCGUnionData::CreateSequentialPointData(FPCGContext* Context, UPCGPointDat
 		TargetPoints += SelectedDataPoints;
 		SelectedDataPoints.Reset();
 	}
-}
\ No newline at end of file
+}
+
+UPCGSpatialData* UPCGUnionData::CopyInternal() const
+{
+	UPCGUnionData* NewUnionData = NewObject<UPCGUnionData>();
+
+	NewUnionData->Data = Data;
+	NewUnionData->FirstNonTrivialTransformData = FirstNonTrivialTransformData;
+	NewUnionData->UnionType = UnionType;
+	NewUnionData->DensityFunction = DensityFunction;
+	NewUnionData->CachedBounds = CachedBounds;
+	NewUnionData->CachedStrictBounds = CachedStrictBounds;
+	NewUnionData->CachedDimension = CachedDimension;
+
+	return NewUnionData;
+}
diff --git a/Engine/Plugins/Experimental/PCG/Source/PCG/Private/Data/PCGVolumeData.cpp b/Engine/Plugins/Experimental/PCG/Source/PCG/Private/Data/PCGVolumeData.cpp
--- a/Engine/Plugins/Experimental/PCG/Source/PCG/Private/Data/PCGVolumeData.cpp
+++ b/Engine/Plugins/Experimental/PCG/Source/PCG/Private/Data/PCGVolumeData.cpp
@@ -84,4 +84,21 @@ bool UPCGVolumeData::SamplePoint(const FTransform& InTransform, const FBox& InBo
 	{
 		return false;
 	}
+}
+
+inline void UPCGVolumeData::CopyBaseVolumeData(UPCGVolumeData* NewVolumeData) const
+{
+	NewVolumeData->VoxelSize = VoxelSize;
+	NewVolumeData->Volume = Volume;
+	NewVolumeData->Bounds = Bounds;
+	NewVolumeData->StrictBounds = StrictBounds;
+}
+
+UPCGSpatialData* UPCGVolumeData::CopyInternal() const
+{
+	UPCGVolumeData* NewVolumeData = NewObject<UPCGVolumeData>();
+
+	CopyBaseVolumeData(NewVolumeData);
+
+	return NewVolumeData;
 }
\ No newline at end of file
diff --git a/Engine/Plugins/Experimental/PCG/Source/PCG/Private/Data/PCGWorldData.cpp b/Engine/Plugins/Experimental/PCG/Source/PCG/Private/Data/PCGWorldData.cpp
--- a/Engine/Plugins/Experimental/PCG/Source/PCG/Private/Data/PCGWorldData.cpp
+++ b/Engine/Plugins/Experimental/PCG/Source/PCG/Private/Data/PCGWorldData.cpp
@@ -128,6 +128,19 @@ const UPCGPointData* UPCGWorldVolumetricData::CreatePointData(FPCGContext* Conte
 	return Data;
 }
 
+UPCGSpatialData* UPCGWorldVolumetricData::CopyInternal() const
+{
+	UPCGWorldVolumetricData* NewVolumetricData = NewObject<UPCGWorldVolumetricData>();
+
+	CopyBaseVolumeData(NewVolumetricData);
+
+	NewVolumetricData->World = World;
+	NewVolumetricData->OriginatingComponent = OriginatingComponent;
+	NewVolumetricData->QueryParams = QueryParams;
+
+	return NewVolumetricData;
+}
+
 /** World Ray Hit data implementation */
 void UPCGWorldRayHitData::Initialize(UWorld* InWorld, const FBox& InBounds)
 {
@@ -232,4 +245,18 @@ const UPCGPointData* UPCGWorldRayHitData::CreatePointData(FPCGContext* Context,
 	}
 
 	return Data;
+}
+
+UPCGSpatialData* UPCGWorldRayHitData::CopyInternal() const
+{
+	UPCGWorldRayHitData* NewData = NewObject<UPCGWorldRayHitData>();
+
+	CopyBaseSurfaceData(NewData);
+
+	NewData->World = World;
+	NewData->OriginatingComponent = OriginatingComponent;
+	NewData->Bounds = Bounds;
+	NewData->QueryParams = QueryParams;
+
+	return NewData;
 }
\ No newline at end of file
diff --git a/Engine/Plugins/Experimental/PCG/Source/PCG/Private/Elements/Metadata/PCGMetadataElementCommon.cpp b/Engine/Plugins/Experimental/PCG/Source/PCG/Private/Elements/Metadata/PCGMetadataElementCommon.cpp
--- a/Engine/Plugins/Experimental/PCG/Source/PCG/Private/Elements/Metadata/PCGMetadataElementCommon.cpp
+++ b/Engine/Plugins/Experimental/PCG/Source/PCG/Private/Elements/Metadata/PCGMetadataElementCommon.cpp
@@ -10,26 +10,30 @@ namespace PCGMetadataElementCommon
 {
 	void DuplicateTaggedData(const FPCGTaggedData& InTaggedData, FPCGTaggedData& OutTaggedData, UPCGMetadata*& OutMetadata)
 	{
+		TRACE_CPUPROFILER_EVENT_SCOPE(PCGMetadataElementCommonDuplicateTaggedData);
+
 		if (const UPCGSpatialData* SpatialInput = Cast<const UPCGSpatialData>(InTaggedData.Data))
 		{
-			UPCGSpatialData* NewSpatialData = Cast<UPCGSpatialData>(StaticDuplicateObject(SpatialInput, const_cast<UPCGSpatialData*>(SpatialInput), FName()));
-			NewSpatialData->InitializeFromData(SpatialInput);
+			UPCGSpatialData* NewSpatialData = SpatialInput->DuplicateData();
 			OutTaggedData.Data = NewSpatialData;
 
 			OutMetadata = NewSpatialData->Metadata;
 		}
 		else if (const UPCGParamData* ParamsInput = Cast<const UPCGParamData>(InTaggedData.Data))
 		{
-			UPCGParamData* NewParamData = Cast<UPCGParamData>(StaticDuplicateObject(ParamsInput, const_cast<UPCGParamData*>(ParamsInput), FName()));
+			UPCGParamData* NewParamData = NewObject<UPCGParamData>();
+			NewParamData->Metadata->Initialize(ParamsInput->Metadata);
+
 			OutTaggedData.Data = NewParamData;
 
-			NewParamData->Metadata->Initialize(ParamsInput->Metadata);
 			OutMetadata = NewParamData->Metadata;
 		}
 	}
 
 	void CopyEntryToValueKeyMap(const UPCGMetadata* MetadataToCopy, const FPCGMetadataAttributeBase* AttributeToCopy, FPCGMetadataAttributeBase* OutAttribute)
 	{
+		TRACE_CPUPROFILER_EVENT_SCOPE(PCGMetadataElementCommon::CopyEntryToValueKeyMap);
+
 		if (!OutAttribute)
 		{
 			UE_LOG(LogPCG, Error, TEXT("Failed to create output attribute"));
diff --git a/Engine/Plugins/Experimental/PCG/Source/PCG/Private/Elements/Metadata/PCGMetadataOpElementBase.cpp b/Engine/Plugins/Experimental/PCG/Source/PCG/Private/Elements/Metadata/PCGMetadataOpElementBase.cpp
--- a/Engine/Plugins/Experimental/PCG/Source/PCG/Private/Elements/Metadata/PCGMetadataOpElementBase.cpp
+++ b/Engine/Plugins/Experimental/PCG/Source/PCG/Private/Elements/Metadata/PCGMetadataOpElementBase.cpp
@@ -213,7 +213,6 @@ bool FPCGMetadataElementBase::ExecuteInternal(FPCGContext* Context) const
 
 		PCGMetadataElementCommon::DuplicateTaggedData(InputTaggedData[0], OutputData, OutMetadata);
 		OperationData.OutputAttributes[OutputIndex] = PCGMetadataElementCommon::ClearOrCreateAttribute(OutMetadata, OutputAttributeFinalName, DefaultValue);
-		PCGMetadataElementCommon::CopyEntryToValueKeyMap(SourceMetadata[0], OperationData.SourceAttributes[0], OperationData.OutputAttributes[OutputIndex]);
 
 		return OperationData.OutputAttributes[OutputIndex] != nullptr;
 	};
diff --git a/Engine/Plugins/Experimental/PCG/Source/PCG/Private/Elements/PCGAttributeFilterElement.cpp b/Engine/Plugins/Experimental/PCG/Source/PCG/Private/Elements/PCGAttributeFilterElement.cpp
--- a/Engine/Plugins/Experimental/PCG/Source/PCG/Private/Elements/PCGAttributeFilterElement.cpp
+++ b/Engine/Plugins/Experimental/PCG/Source/PCG/Private/Elements/PCGAttributeFilterElement.cpp
@@ -66,13 +66,10 @@ bool FPCGAttributeFilterElement::ExecuteInternal(FPCGContext* Context) const
 
 		if (const UPCGSpatialData* InputSpatialData = Cast<UPCGSpatialData>(InputData))
 		{
-			Metadata = NewObject<UPCGMetadata>();
 			ParentMetadata = InputSpatialData->Metadata;
 
-			UPCGSpatialData* NewSpatialData = DuplicateObject<UPCGSpatialData>(const_cast<UPCGSpatialData*>(InputSpatialData), nullptr);
-			NewSpatialData->Metadata = Metadata;
-			// We also need set the Metadata outer to this new data.
-			Metadata->Rename(nullptr, NewSpatialData);
+			UPCGSpatialData* NewSpatialData = InputSpatialData->DuplicateData(/*bInitializeFromThisData=*/false);
+			Metadata = NewSpatialData->Metadata;
 			NewSpatialData->Metadata->Initialize(ParentMetadata, /*bAddAttributesFromParent=*/false);
 
 			// No need to inherit metadata since we already initialized it.
diff --git a/Engine/Plugins/Experimental/PCG/Source/PCG/Private/Elements/PCGAttributeTransferElement.cpp b/Engine/Plugins/Experimental/PCG/Source/PCG/Private/Elements/PCGAttributeTransferElement.cpp
--- a/Engine/Plugins/Experimental/PCG/Source/PCG/Private/Elements/PCGAttributeTransferElement.cpp
+++ b/Engine/Plugins/Experimental/PCG/Source/PCG/Private/Elements/PCGAttributeTransferElement.cpp
@@ -117,8 +117,7 @@ UPCGSpatialData* FPCGAttributeTransferElement::TransferSpatialToSpatial(FPCGCont
 		return nullptr;
 	}
 
-	UPCGSpatialData* NewSpatialData = DuplicateObject<UPCGSpatialData>(const_cast<UPCGSpatialData*>(TargetData), nullptr);
-	NewSpatialData->InitializeFromData(TargetData);
+	UPCGSpatialData* NewSpatialData = TargetData->DuplicateData();
 
 	const FName TargetAttributeName = (Settings->TargetAttributeName == NAME_None) ? Settings->SourceAttributeName : Settings->TargetAttributeName;
 
@@ -148,14 +147,11 @@ UPCGPointData* FPCGAttributeTransferElement::TransferPointToPoint(FPCGContext* C
 		return nullptr;
 	}
 
-	UPCGPointData* NewPointData = NewObject<UPCGPointData>();
-	NewPointData->InitializeFromData(TargetData);
+	UPCGPointData* NewPointData = Cast<UPCGPointData>(TargetData->DuplicateData());
 
 	const TArray<FPCGPoint>& SourcePoints = SourceData->GetPoints();
 	TArray<FPCGPoint>& TargetPoints = NewPointData->GetMutablePoints();
 
-	TargetPoints = TargetData->GetPoints();
-
 	const FName TargetAttributeName = (Settings->TargetAttributeName == NAME_None) ? Settings->SourceAttributeName : Settings->TargetAttributeName;
 
 	// Making sure the target attribute doesn't exist in the target
diff --git a/Engine/Plugins/Experimental/PCG/Source/PCG/Private/Elements/PCGCreateAttribute.cpp b/Engine/Plugins/Experimental/PCG/Source/PCG/Private/Elements/PCGCreateAttribute.cpp
--- a/Engine/Plugins/Experimental/PCG/Source/PCG/Private/Elements/PCGCreateAttribute.cpp
+++ b/Engine/Plugins/Experimental/PCG/Source/PCG/Private/Elements/PCGCreateAttribute.cpp
@@ -214,8 +214,7 @@ bool FPCGCreateAttributeElement::ExecuteInternal(FPCGContext* Context) const
 
 		if (const UPCGSpatialData* InputSpatialData = Cast<UPCGSpatialData>(InputData))
 		{
-			UPCGSpatialData* NewSpatialData = DuplicateObject<UPCGSpatialData>(InputSpatialData, nullptr);
-			NewSpatialData->Metadata = NewObject<UPCGMetadata>(NewSpatialData);
+			UPCGSpatialData* NewSpatialData = InputSpatialData->DuplicateData(/*bInitializeFromData=*/false);
 			NewSpatialData->InitializeFromData(InputSpatialData, /*InMetadataParentOverride=*/ nullptr, /*bInheritMetadata=*/true, /*bInheritMetadata=*/ Settings->bKeepExistingAttributes);
 
 			OutputData = NewSpatialData;
diff --git a/Engine/Plugins/Experimental/PCG/Source/PCG/Public/Data/PCGDifferenceData.h b/Engine/Plugins/Experimental/PCG/Source/PCG/Public/Data/PCGDifferenceData.h
--- a/Engine/Plugins/Experimental/PCG/Source/PCG/Public/Data/PCGDifferenceData.h
+++ b/Engine/Plugins/Experimental/PCG/Source/PCG/Public/Data/PCGDifferenceData.h
@@ -52,8 +52,11 @@ public:
 	virtual FBox GetStrictBounds() const override;
 	virtual bool SamplePoint(const FTransform& Transform, const FBox& Bounds, FPCGPoint& OutPoint, UPCGMetadata* OutMetadata) const override;
 	virtual bool HasNonTrivialTransform() const override;
+protected:
+	virtual UPCGSpatialData* CopyInternal() const override;
 	//~End UPCGSpatialData interface
 
+public:
 	//~Begin UPCGSpatialDataWithPointCache interface
 	virtual const UPCGPointData* CreatePointData(FPCGContext* Context) const override;
 	//~End UPCGSpatialDataWithPointCache interface
diff --git a/Engine/Plugins/Experimental/PCG/Source/PCG/Public/Data/PCGIntersectionData.h b/Engine/Plugins/Experimental/PCG/Source/PCG/Public/Data/PCGIntersectionData.h
--- a/Engine/Plugins/Experimental/PCG/Source/PCG/Public/Data/PCGIntersectionData.h
+++ b/Engine/Plugins/Experimental/PCG/Source/PCG/Public/Data/PCGIntersectionData.h
@@ -30,8 +30,11 @@ public:
 	virtual FBox GetStrictBounds() const override;
 	virtual bool SamplePoint(const FTransform& Transform, const FBox& Bounds, FPCGPoint& OutPoint, UPCGMetadata* OutMetadata) const override;
 	virtual bool HasNonTrivialTransform() const override;
+protected:
+	virtual UPCGSpatialData* CopyInternal() const override;
 	//~End UPCGSpatialData interface
 
+public:
 	//~Begin UPCGSpatialDataWithPointCache interface
 	virtual const UPCGPointData* CreatePointData(FPCGContext* Context) const override;
 	//~End UPCGSpatialDataWithPointCache interface
diff --git a/Engine/Plugins/Experimental/PCG/Source/PCG/Public/Data/PCGLandscapeData.h b/Engine/Plugins/Experimental/PCG/Source/PCG/Public/Data/PCGLandscapeData.h
--- a/Engine/Plugins/Experimental/PCG/Source/PCG/Public/Data/PCGLandscapeData.h
+++ b/Engine/Plugins/Experimental/PCG/Source/PCG/Public/Data/PCGLandscapeData.h
@@ -35,13 +35,16 @@ public:
 	virtual FBox GetStrictBounds() const override;
 	virtual bool SamplePoint(const FTransform& Transform, const FBox& Bounds, FPCGPoint& OutPoint, UPCGMetadata* OutMetadata) const override;
 	virtual bool HasNonTrivialTransform() const override { return true; }
-	// ~End UPGCConcreteData interface
+protected:
+	virtual UPCGSpatialData* CopyInternal() const override;
+	//~End UPCGSpatialData interface
 
+public:
 	// ~Begin UPCGSpatialDataWithPointCache interface
 	virtual bool SupportsBoundedPointData() const { return true; }
 	virtual const UPCGPointData* CreatePointData(FPCGContext* Context) const override { return CreatePointData(Context, FBox(EForceInit::ForceInit)); }
 	virtual const UPCGPointData* CreatePointData(FPCGContext* Context, const FBox& InBounds) const override;
-	// ~End UPCGConcreteDataWithPointCache interface
+	// ~End UPCGSpatialDataWithPointCache interface
 
 	// TODO: add on property changed to clear cached data. This is used to populate the LandscapeInfos array.
 	UPROPERTY(BlueprintReadWrite, EditAnywhere, Category = SourceData)
diff --git a/Engine/Plugins/Experimental/PCG/Source/PCG/Public/Data/PCGLandscapeSplineData.h b/Engine/Plugins/Experimental/PCG/Source/PCG/Public/Data/PCGLandscapeSplineData.h
--- a/Engine/Plugins/Experimental/PCG/Source/PCG/Public/Data/PCGLandscapeSplineData.h
+++ b/Engine/Plugins/Experimental/PCG/Source/PCG/Public/Data/PCGLandscapeSplineData.h
@@ -33,7 +33,9 @@ public:
 	//~Begin UPCGSpatialData interface
 	virtual FBox GetBounds() const override;
 	virtual bool SamplePoint(const FTransform& Transform, const FBox& Bounds, FPCGPoint& OutPoint, UPCGMetadata* OutMetadata) const override;
-	//~End
+protected:
+	virtual UPCGSpatialData* CopyInternal() const override;
+	//~End UPCGSpatialData interface
 
 protected:
 	UPROPERTY(BlueprintReadWrite, EditAnywhere, Category = SourceData)
diff --git a/Engine/Plugins/Experimental/PCG/Source/PCG/Public/Data/PCGPointData.h b/Engine/Plugins/Experimental/PCG/Source/PCG/Public/Data/PCGPointData.h
--- a/Engine/Plugins/Experimental/PCG/Source/PCG/Public/Data/PCGPointData.h
+++ b/Engine/Plugins/Experimental/PCG/Source/PCG/Public/Data/PCGPointData.h
@@ -75,8 +75,11 @@ public:
 	virtual FBox GetBounds() const override;
 	virtual const UPCGPointData* ToPointData(FPCGContext* Context, const FBox& InBounds = FBox(EForceInit::ForceInit)) const { return this; }
 	virtual bool SamplePoint(const FTransform& Transform, const FBox& Bounds, FPCGPoint& OutPoint, UPCGMetadata* OutMetadata) const override;
+protected:
+	virtual UPCGSpatialData* CopyInternal() const override;
 	// ~End UPCGSpatialData interface
 
+public:
 	void InitializeFromActor(AActor* InActor);
 
 	UFUNCTION(BlueprintCallable, Category = SpatialData)
diff --git a/Engine/Plugins/Experimental/PCG/Source/PCG/Public/Data/PCGPrimitiveData.h b/Engine/Plugins/Experimental/PCG/Source/PCG/Public/Data/PCGPrimitiveData.h
--- a/Engine/Plugins/Experimental/PCG/Source/PCG/Public/Data/PCGPrimitiveData.h
+++ b/Engine/Plugins/Experimental/PCG/Source/PCG/Public/Data/PCGPrimitiveData.h
@@ -23,8 +23,11 @@ public:
 	virtual FBox GetBounds() const override { return CachedBounds; }
 	virtual FBox GetStrictBounds() const override { return CachedStrictBounds; }
 	virtual bool SamplePoint(const FTransform& Transform, const FBox& Bounds, FPCGPoint& OutPoint, UPCGMetadata* OutMetadata) const override;
-	// ~End UPCGSpatialData interface
+protected:
+	virtual UPCGSpatialData* CopyInternal() const override;
+	//~End UPCGSpatialData interface
 
+public:
 	// ~Begin UPCGSpatialDataWithPointCache implementation
 	virtual const UPCGPointData* CreatePointData(FPCGContext* Context) const override;
 	// ~End UPCGSpatialDataWithPointCache implementation
diff --git a/Engine/Plugins/Experimental/PCG/Source/PCG/Public/Data/PCGProjectionData.h b/Engine/Plugins/Experimental/PCG/Source/PCG/Public/Data/PCGProjectionData.h
--- a/Engine/Plugins/Experimental/PCG/Source/PCG/Public/Data/PCGProjectionData.h
+++ b/Engine/Plugins/Experimental/PCG/Source/PCG/Public/Data/PCGProjectionData.h
@@ -26,13 +26,18 @@ public:
 	virtual bool SamplePoint(const FTransform& Transform, const FBox& Bounds, FPCGPoint& OutPoint, UPCGMetadata* OutMetadata) const override;
 	virtual bool HasNonTrivialTransform() const override;
 	virtual bool RequiresCollapseToSample() const override { return true; }
+protected:
+	virtual UPCGSpatialData* CopyInternal() const override;
 	//~End UPCGSpatialData interface
 
+public:
 	//~Begin UPCGSpatialDataWithPointCache interface
 	virtual const UPCGPointData* CreatePointData(FPCGContext* Context) const override;
 	//~End UPCGSpatialDataWithPointCache interface
 
 protected:
+	void CopyBaseProjectionClass(UPCGProjectionData* NewProjectionData) const;
+
 	FBox ProjectBounds(const FBox& InBounds) const;
 
 	/** Applies data from target point to projected point, conditionally according to the projection params. */
diff --git a/Engine/Plugins/Experimental/PCG/Source/PCG/Public/Data/PCGRenderTargetData.h b/Engine/Plugins/Experimental/PCG/Source/PCG/Public/Data/PCGRenderTargetData.h
--- a/Engine/Plugins/Experimental/PCG/Source/PCG/Public/Data/PCGRenderTargetData.h
+++ b/Engine/Plugins/Experimental/PCG/Source/PCG/Public/Data/PCGRenderTargetData.h
@@ -19,6 +19,13 @@ public:
 	virtual EPCGDataType GetDataType() const override { return EPCGDataType::RenderTarget | Super::GetDataType(); }
 	// ~End UPCGData interface
 
+
+	//~Begin UPCGSpatialData interface
+protected:
+	virtual UPCGSpatialData* CopyInternal() const override;
+	//~End UPCGSpatialData interface
+
+public:
 	UFUNCTION(BlueprintCallable, Category = RenderTarget)
 	void Initialize(UTextureRenderTarget2D* InRenderTarget, const FTransform& InTransform);
 
diff --git a/Engine/Plugins/Experimental/PCG/Source/PCG/Public/Data/PCGSpatialData.h b/Engine/Plugins/Experimental/PCG/Source/PCG/Public/Data/PCGSpatialData.h
--- a/Engine/Plugins/Experimental/PCG/Source/PCG/Public/Data/PCGSpatialData.h
+++ b/Engine/Plugins/Experimental/PCG/Source/PCG/Public/Data/PCGSpatialData.h
@@ -40,6 +40,16 @@ public:
 	virtual EPCGDataType GetDataType() const override { return EPCGDataType::Spatial | Super::GetDataType(); }
 	// ~End UPCGData interface
 
+	/** Virtual call to allocate a new spacial data object, duplicate this spatial data into
+	*   and parent the new metadata with this class metadata (if asked).
+	*   Should be way cheaper than DuplicateObject, since we avoid duplicating metadata.
+	*   It will not deep copy references.
+	*   Some data are marked mutable and therefore are not threadsafe to copy, so they are not copied.
+	*   They are mainly cached values (and octree for points).
+	*   TODO: If we want to also copy those values (can be an optimization), we need to guard the copy.
+	*/
+	UPCGSpatialData* DuplicateData(const bool bInitializeMetadata = true) const;
+
 	/** Returns the dimension of the data type, which has nothing to do with the dimension of its points */
 	UFUNCTION(BlueprintCallable, Category = SpatialData)
 	virtual int GetDimension() const PURE_VIRTUAL(UPCGSpatialData::GetDimension, return 0;);
@@ -120,6 +130,9 @@ public:
 	// Not accessible through blueprint to make sure the constness is preserved
 	UPROPERTY(VisibleAnywhere, Category = Metadata)
 	TObjectPtr<UPCGMetadata> Metadata = nullptr;
+
+protected:
+	virtual UPCGSpatialData* CopyInternal() const PURE_VIRTUAL(UPCGSpatialData::CopyInternal, return nullptr;);
 };
 
 UCLASS(Abstract, ClassGroup = (Procedural))
diff --git a/Engine/Plugins/Experimental/PCG/Source/PCG/Public/Data/PCGSplineData.h b/Engine/Plugins/Experimental/PCG/Source/PCG/Public/Data/PCGSplineData.h
--- a/Engine/Plugins/Experimental/PCG/Source/PCG/Public/Data/PCGSplineData.h
+++ b/Engine/Plugins/Experimental/PCG/Source/PCG/Public/Data/PCGSplineData.h
@@ -40,8 +40,11 @@ public:
 	virtual FBox GetBounds() const override;
 	virtual bool SamplePoint(const FTransform& Transform, const FBox& Bounds, FPCGPoint& OutPoint, UPCGMetadata* OutMetadata) const override;
 	virtual UPCGProjectionData* ProjectOn(const UPCGSpatialData* InOther, const FPCGProjectionParams& InParams = FPCGProjectionParams()) const override;
-	//~End 
+protected:
+	virtual UPCGSpatialData* CopyInternal() const override;
+	//~End UPCGSpatialData interface
 
+public:
 	UPROPERTY(BlueprintReadWrite, EditAnywhere, Category = SourceData)
 	TSoftObjectPtr<USplineComponent> Spline;
 
@@ -66,4 +69,8 @@ public:
 
 protected:
 	FVector2D Project(const FVector& InVector) const;
+
+	//~Begin UPCGSpatialData interface
+	virtual UPCGSpatialData* CopyInternal() const override;
+	//~End UPCGSpatialData interface
 };
diff --git a/Engine/Plugins/Experimental/PCG/Source/PCG/Public/Data/PCGSurfaceData.h b/Engine/Plugins/Experimental/PCG/Source/PCG/Public/Data/PCGSurfaceData.h
--- a/Engine/Plugins/Experimental/PCG/Source/PCG/Public/Data/PCGSurfaceData.h
+++ b/Engine/Plugins/Experimental/PCG/Source/PCG/Public/Data/PCGSurfaceData.h
@@ -23,6 +23,8 @@ public:
 	//~End UPCGSpatialData interface
 
 protected:
+	void CopyBaseSurfaceData(UPCGSurfaceData* NewSurfaceData) const;
+
 	UPROPERTY(BlueprintReadOnly, VisibleAnywhere, Category = SpatialData)
 	FTransform Transform;
 };
diff --git a/Engine/Plugins/Experimental/PCG/Source/PCG/Public/Data/PCGTextureData.h b/Engine/Plugins/Experimental/PCG/Source/PCG/Public/Data/PCGTextureData.h
--- a/Engine/Plugins/Experimental/PCG/Source/PCG/Public/Data/PCGTextureData.h
+++ b/Engine/Plugins/Experimental/PCG/Source/PCG/Public/Data/PCGTextureData.h
@@ -88,6 +88,8 @@ protected:
 
 	UPROPERTY(BlueprintReadOnly, VisibleAnywhere, Category = SpatialData)
 	int32 Width = 0;
+
+	void CopyBaseTextureData(UPCGBaseTextureData* NewTextureData) const;
 };
 
 UCLASS(BlueprintType, ClassGroup=(Procedural))
@@ -100,6 +102,12 @@ public:
 	virtual EPCGDataType GetDataType() const override { return EPCGDataType::Texture | Super::GetDataType(); }
 	// ~End UPCGData interface
 
+	//~Begin UPCGSpatialData interface
+protected:
+	virtual UPCGSpatialData* CopyInternal() const override;
+	//~End UPCGSpatialData interface
+
+public:
 	UFUNCTION(BlueprintCallable, Category = Texture)
 	void Initialize(UTexture2D* InTexture, const FTransform& InTransform);
 
diff --git a/Engine/Plugins/Experimental/PCG/Source/PCG/Public/Data/PCGUnionData.h b/Engine/Plugins/Experimental/PCG/Source/PCG/Public/Data/PCGUnionData.h
--- a/Engine/Plugins/Experimental/PCG/Source/PCG/Public/Data/PCGUnionData.h
+++ b/Engine/Plugins/Experimental/PCG/Source/PCG/Public/Data/PCGUnionData.h
@@ -43,8 +43,11 @@ public:
 	virtual FBox GetStrictBounds() const override;
 	virtual bool SamplePoint(const FTransform& Transform, const FBox& Bounds, FPCGPoint& OutPoint, UPCGMetadata* OutMetadata) const override;
 	virtual bool HasNonTrivialTransform() const override;
+protected:
+	virtual UPCGSpatialData* CopyInternal() const override;
 	//~End UPCGSpatialData interface
 
+public:
 	//~Begin UPCGSpatialDataWithPointCache interface
 	virtual const UPCGPointData* CreatePointData(FPCGContext* Context) const override;
 	//~End UPCGSpatialDataWithPointCache interface
diff --git a/Engine/Plugins/Experimental/PCG/Source/PCG/Public/Data/PCGVolumeData.h b/Engine/Plugins/Experimental/PCG/Source/PCG/Public/Data/PCGVolumeData.h
--- a/Engine/Plugins/Experimental/PCG/Source/PCG/Public/Data/PCGVolumeData.h
+++ b/Engine/Plugins/Experimental/PCG/Source/PCG/Public/Data/PCGVolumeData.h
@@ -27,8 +27,11 @@ public:
 	virtual FBox GetBounds() const override;
 	virtual FBox GetStrictBounds() const override;
 	virtual bool SamplePoint(const FTransform& Transform, const FBox& Bounds, FPCGPoint& OutPoint, UPCGMetadata* OutMetadata) const override;
-	// ~End UPGCSpatialData interface
+protected:
+	virtual UPCGSpatialData* CopyInternal() const override;
+	//~End UPCGSpatialData interface
 
+public:
 	// ~Begin UPCGSpatialDataWithPointCache interface
 	virtual const UPCGPointData* CreatePointData(FPCGContext* Context) const override;
 	// ~End UPCGSpatialDataWithPointCache interface
@@ -37,6 +40,8 @@ public:
 	FVector VoxelSize = FVector(100.0, 100.0, 100.0);
 
 protected:
+	void CopyBaseVolumeData(UPCGVolumeData* NewVolumeData) const;
+
 	UPROPERTY(BlueprintReadWrite, EditAnywhere, Category = SourceData)
 	TWeakObjectPtr<AVolume> Volume = nullptr;
 
diff --git a/Engine/Plugins/Experimental/PCG/Source/PCG/Public/Data/PCGWorldData.h b/Engine/Plugins/Experimental/PCG/Source/PCG/Public/Data/PCGWorldData.h
--- a/Engine/Plugins/Experimental/PCG/Source/PCG/Public/Data/PCGWorldData.h
+++ b/Engine/Plugins/Experimental/PCG/Source/PCG/Public/Data/PCGWorldData.h
@@ -42,8 +42,11 @@ public:
 
 	//~Begin UPCGSpatialData interface
 	virtual bool SamplePoint(const FTransform& Transform, const FBox& Bounds, FPCGPoint& OutPoint, UPCGMetadata* OutMetadata) const override;
+protected:
+	virtual UPCGSpatialData* CopyInternal() const override;
 	//~End UPCGSpatialData interface
 
+public:
 	//~Begin UPCGSpatialDataWithPointCache
 	virtual bool SupportsBoundedPointData() const { return true; }
 	virtual const UPCGPointData* CreatePointData(FPCGContext* Context) const override { return CreatePointData(Context, FBox(EForceInit::ForceInit)); }
@@ -110,7 +113,11 @@ public:
 	virtual FBox GetStrictBounds() const override { return Bounds; }
 	virtual bool SamplePoint(const FTransform& Transform, const FBox& Bounds, FPCGPoint& OutPoint, UPCGMetadata* OutMetadata) const override;
 	virtual bool HasNonTrivialTransform() const override { return true; }
+protected:
+	virtual UPCGSpatialData* CopyInternal() const override;
+	//~End UPCGSpatialData interface
 
+public:
 	// ~Begin UPCGSpatialDataWithPointCache interface
 	virtual bool SupportsBoundedPointData() const { return true; }
 	virtual const UPCGPointData* CreatePointData(FPCGContext* Context) const override { return CreatePointData(Context, FBox(EForceInit::ForceInit)); }
diff --git a/Engine/Plugins/Experimental/PCG/Source/PCG/Public/Elements/Metadata/PCGMetadataElementCommon.h b/Engine/Plugins/Experimental/PCG/Source/PCG/Public/Elements/Metadata/PCGMetadataElementCommon.h
--- a/Engine/Plugins/Experimental/PCG/Source/PCG/Public/Elements/Metadata/PCGMetadataElementCommon.h
+++ b/Engine/Plugins/Experimental/PCG/Source/PCG/Public/Elements/Metadata/PCGMetadataElementCommon.h
@@ -20,6 +20,8 @@ namespace PCGMetadataElementCommon
 	template<typename T>
 	FPCGMetadataAttribute<T>* ClearOrCreateAttribute(UPCGMetadata* Metadata, const FName& DestinationAttribute, T DefaultValue)
 	{
+		TRACE_CPUPROFILER_EVENT_SCOPE(PCGMetadataElementCommon::ClearOrCreateAttribute);
+
 		if (!Metadata)
 		{
 			UE_LOG(LogPCG, Error, TEXT("Failed to create metadata"));
@@ -28,7 +30,7 @@ namespace PCGMetadataElementCommon
 
 		if (Metadata->HasAttribute(DestinationAttribute))
 		{
-			UE_LOG(LogPCG, Warning, TEXT("Attribute %s already exists and has been overwritten"), *DestinationAttribute.ToString());
+			UE_LOG(LogPCG, Verbose, TEXT("Attribute %s already exists and has been overwritten"), *DestinationAttribute.ToString());
 			Metadata->DeleteAttribute(DestinationAttribute);
 		}
 
diff --git a/Engine/Plugins/Experimental/PCG/Source/PCG/Public/Elements/Metadata/PCGMetadataOpElementBase.h b/Engine/Plugins/Experimental/PCG/Source/PCG/Public/Elements/Metadata/PCGMetadataOpElementBase.h
--- a/Engine/Plugins/Experimental/PCG/Source/PCG/Public/Elements/Metadata/PCGMetadataOpElementBase.h
+++ b/Engine/Plugins/Experimental/PCG/Source/PCG/Public/Elements/Metadata/PCGMetadataOpElementBase.h
@@ -113,11 +113,6 @@ public:
 class FPCGMetadataElementBase : public FSimplePCGElement
 {
 public:
-	// FPCGMetadataElementBase relies on StaticDuplicateObject (when we duplicate tagged data), so we cannot run outside of the main thread
-	virtual bool CanExecuteOnlyOnMainThread(FPCGContext*) const override { return true; }
-
-	virtual bool IsCacheable(const UPCGSettings* InSettings) const override { return true; }
-
 	struct FOperationData
 	{
 		TArray<const FPCGMetadataAttributeBase*> SourceAttributes;
diff --git a/Engine/Plugins/Experimental/PCG/Source/PCG/Public/Elements/PCGAttributeFilterElement.h b/Engine/Plugins/Experimental/PCG/Source/PCG/Public/Elements/PCGAttributeFilterElement.h
--- a/Engine/Plugins/Experimental/PCG/Source/PCG/Public/Elements/PCGAttributeFilterElement.h
+++ b/Engine/Plugins/Experimental/PCG/Source/PCG/Public/Elements/PCGAttributeFilterElement.h
@@ -40,10 +40,6 @@ protected:
 
 class FPCGAttributeFilterElement : public FSimplePCGElement
 {
-public:
-	virtual bool CanExecuteOnlyOnMainThread(FPCGContext* Context) const { return true; }
-	virtual bool IsCacheable(const UPCGSettings* InSettings) const override { return false; }
-
 protected:
 	virtual bool ExecuteInternal(FPCGContext* Context) const override;
 };
diff --git a/Engine/Plugins/Experimental/PCG/Source/PCG/Public/Elements/PCGAttributeReduceElement.h b/Engine/Plugins/Experimental/PCG/Source/PCG/Public/Elements/PCGAttributeReduceElement.h
--- a/Engine/Plugins/Experimental/PCG/Source/PCG/Public/Elements/PCGAttributeReduceElement.h
+++ b/Engine/Plugins/Experimental/PCG/Source/PCG/Public/Elements/PCGAttributeReduceElement.h
@@ -55,9 +55,6 @@ protected:
 
 class FPCGAttributeReduceElement : public FSimplePCGElement
 {
-public:
-	virtual bool IsCacheable(const UPCGSettings* InSettings) const override { return false; }
-
 protected:
 	virtual bool ExecuteInternal(FPCGContext* Context) const override;
 };
\ No newline at end of file
diff --git a/Engine/Plugins/Experimental/PCG/Source/PCG/Public/Elements/PCGAttributeSelectElement.h b/Engine/Plugins/Experimental/PCG/Source/PCG/Public/Elements/PCGAttributeSelectElement.h
--- a/Engine/Plugins/Experimental/PCG/Source/PCG/Public/Elements/PCGAttributeSelectElement.h
+++ b/Engine/Plugins/Experimental/PCG/Source/PCG/Public/Elements/PCGAttributeSelectElement.h
@@ -84,9 +84,6 @@ protected:
 
 class FPCGAttributeSelectElement : public FSimplePCGElement
 {
-public:
-	virtual bool IsCacheable(const UPCGSettings* InSettings) const override { return false; }
-
 protected:
 	virtual bool ExecuteInternal(FPCGContext* Context) const override;
 };
\ No newline at end of file
diff --git a/Engine/Plugins/Experimental/PCG/Source/PCG/Public/Elements/PCGAttributeTransferElement.h b/Engine/Plugins/Experimental/PCG/Source/PCG/Public/Elements/PCGAttributeTransferElement.h
--- a/Engine/Plugins/Experimental/PCG/Source/PCG/Public/Elements/PCGAttributeTransferElement.h
+++ b/Engine/Plugins/Experimental/PCG/Source/PCG/Public/Elements/PCGAttributeTransferElement.h
@@ -57,10 +57,6 @@ protected:
 
 class FPCGAttributeTransferElement : public FSimplePCGElement
 {
-public:
-	virtual bool CanExecuteOnlyOnMainThread(FPCGContext* Context) const { return true; }
-	virtual bool IsCacheable(const UPCGSettings* InSettings) const override { return false; }
-
 protected:
 	virtual bool ExecuteInternal(FPCGContext* Context) const override;
 
diff --git a/Engine/Plugins/Experimental/PCG/Source/PCG/Public/Elements/PCGCreateAttribute.h b/Engine/Plugins/Experimental/PCG/Source/PCG/Public/Elements/PCGCreateAttribute.h
--- a/Engine/Plugins/Experimental/PCG/Source/PCG/Public/Elements/PCGCreateAttribute.h
+++ b/Engine/Plugins/Experimental/PCG/Source/PCG/Public/Elements/PCGCreateAttribute.h
@@ -98,10 +98,6 @@ protected:
 
 class FPCGCreateAttributeElement : public FSimplePCGElement
 {
-public:
-	virtual bool CanExecuteOnlyOnMainThread(FPCGContext* Context) const { return true; }
-	virtual bool IsCacheable(const UPCGSettings* InSettings) const override { return false; }
-
 protected:
 	virtual bool ExecuteInternal(FPCGContext* Context) const override;
 
