Optimizations to self pruning node:
- reduce bounds over-calculation
- making it be more cache coherent as it doesn't need to dereference the Point ptr during octree calculations
- changing the TSet to a bit array to make them O(1) cost.

One particular self pruning node in the project went from 2.05 to 1.25 seconds, a 40% savings.

diff --git a/Engine/Plugins/Experimental/PCG/Source/PCG/Private/Elements/PCGSelfPruning.cpp b/Engine/Plugins/Experimental/PCG/Source/PCG/Private/Elements/PCGSelfPruning.cpp
--- a/Engine/Plugins/Experimental/PCG/Source/PCG/Private/Elements/PCGSelfPruning.cpp
+++ b/Engine/Plugins/Experimental/PCG/Source/PCG/Private/Elements/PCGSelfPruning.cpp
@@ -9,20 +9,20 @@
 
 namespace PCGSelfPruningAlgorithms
 {
-	bool RandomSort(const FPCGPoint* A, const FPCGPoint* B)
+	bool RandomSort(const FPCGPointRef& A, const FPCGPointRef& B)
 	{
-		return A->Seed < B->Seed;
+		return A.Point->Seed < B.Point->Seed;
 	}
 
-	bool SortSmallToLargeNoRandom(const FPCGPoint* A, const FPCGPoint* B, FVector::FReal SquaredRadiusEquality)
+	bool SortSmallToLargeNoRandom(const FPCGPointRef& A, const FPCGPointRef& B, FVector::FReal SquaredRadiusEquality)
 	{
-		return A->GetDensityBounds().BoxExtent.SquaredLength() * SquaredRadiusEquality < B->GetDensityBounds().BoxExtent.SquaredLength();
+		return A.Bounds.BoxExtent.SquaredLength() * SquaredRadiusEquality < B.Bounds.BoxExtent.SquaredLength();
 	}
 
-	bool SortSmallToLargeWithRandom(const FPCGPoint* A, const FPCGPoint* B, FVector::FReal SquaredRadiusEquality)
+	bool SortSmallToLargeWithRandom(const FPCGPointRef& A, const FPCGPointRef& B, FVector::FReal SquaredRadiusEquality)
 	{
-		const FVector::FReal SqrLenA = A->GetDensityBounds().BoxExtent.SquaredLength();
-		const FVector::FReal SqrLenB = B->GetDensityBounds().BoxExtent.SquaredLength();
+		const FVector::FReal SqrLenA = A.Bounds.BoxExtent.SquaredLength();
+		const FVector::FReal SqrLenB = B.Bounds.BoxExtent.SquaredLength();
 		if (SqrLenA * SquaredRadiusEquality < SqrLenB)
 		{
 			return true;
@@ -40,22 +40,52 @@ namespace PCGSelfPruningAlgorithms
 
 namespace PCGSelfPruningElement
 {
-	void DensityBoundsExclusion(const TArray<const FPCGPoint*>& SortedPoints, const UPCGPointData::PointOctree& Octree, TSet<const FPCGPoint*>& ExclusionPoints)
+	// all points are in contiguous memory so this will create a bitset to just look up the index to set/test a bit
+	// this should reduce the 'set' cost to O[1]
+	struct FPointBitSet
 	{
-		TSet<const FPCGPoint*> ExcludedPoints;
-		ExcludedPoints.Reserve(SortedPoints.Num());
+		TArray<uint32> Bits;
+		const FPCGPoint* FirstPoint = nullptr;
+
+		FPointBitSet(const TArray<FPCGPoint>& Points)
+		{
+			Bits.SetNumZeroed(Points.Num()/32 + 1);
+			FirstPoint = Points.Num() > 0 ? &Points[0] : nullptr;
+		}
+
+		void Add(const FPCGPoint* Point)
+		{
+			const uint32 Index = Point-FirstPoint;
+			const uint32 BitsIndex = Index / 32;
+			const uint32 Bit = Index % 32;
+
+			Bits[BitsIndex] |= (1<<Bit);
+		}
+
+		bool Contains(const FPCGPoint* Point)
+		{
+			const uint32 Index = Point-FirstPoint;
+			const uint32 BitsIndex = Index / 32;
+			const uint32 Bit = Index % 32;
+
+			return (Bits[BitsIndex] & (1<<Bit)) != 0;
+		}
+	};
+
+	void DensityBoundsExclusion(const TArray<FPCGPoint>& Points, const TArray<FPCGPointRef>& SortedPoints, const UPCGPointData::PointOctree& Octree, FPointBitSet& ExclusionPoints)
+	{
+		FPointBitSet ExcludedPoints(Points);
 		
-		for (const FPCGPoint* Point : SortedPoints)
+		for (const FPCGPointRef& PointRef : SortedPoints)
 		{
-			if (ExcludedPoints.Contains(Point))
+			if (ExcludedPoints.Contains(PointRef.Point))
 			{
 				continue;
 			}
 
-			ExclusionPoints.Add(Point);
+			ExclusionPoints.Add(PointRef.Point);
 
-			const FBoxSphereBounds PointBounds = Point->GetDensityBounds();
-			Octree.FindElementsWithBoundsTest(FBoxCenterAndExtent(PointBounds.Origin, PointBounds.BoxExtent), [&ExclusionPoints, &ExcludedPoints](const FPCGPointRef& InPointRef)
+			Octree.FindElementsWithBoundsTest(FBoxCenterAndExtent(PointRef.Bounds.Origin, PointRef.Bounds.BoxExtent), [&ExclusionPoints, &ExcludedPoints](const FPCGPointRef& InPointRef)
 			{
 				// TODO: check on an oriented-box basis?
 				if (!ExclusionPoints.Contains(InPointRef.Point))
@@ -66,23 +96,22 @@ namespace PCGSelfPruningElement
 		}
 	}
 
-	void DuplicatePointsExclusion(const TArray<const FPCGPoint*>& SortedPoints, const UPCGPointData::PointOctree& Octree, TSet<const FPCGPoint*>& ExclusionPoints)
+	void DuplicatePointsExclusion(const TArray<FPCGPoint>& Points, const TArray<FPCGPointRef>& SortedPoints, const UPCGPointData::PointOctree& Octree, FPointBitSet& ExclusionPoints)
 	{
-		TSet<const FPCGPoint*> ExcludedPoints;
-		ExcludedPoints.Reserve(SortedPoints.Num());
+		FPointBitSet ExcludedPoints(Points);
 
-		for (const FPCGPoint* Point : SortedPoints)
+		for (const FPCGPointRef& PointRef : SortedPoints)
 		{
-			if (ExcludedPoints.Contains(Point))
+			if (ExcludedPoints.Contains(PointRef.Point))
 			{
 				continue;
 			}
 
-			ExclusionPoints.Add(Point);
+			ExclusionPoints.Add(PointRef.Point);
 
-			Octree.FindElementsWithBoundsTest(FBoxCenterAndExtent(Point->Transform.TransformPosition(Point->GetLocalCenter()), FVector::Zero()), [&ExclusionPoints, &ExcludedPoints, Point](const FPCGPointRef& InPointRef)
+			Octree.FindElementsWithBoundsTest(FBoxCenterAndExtent(PointRef.Point->Transform.TransformPosition(PointRef.Point->GetLocalCenter()), FVector::Zero()), [&ExclusionPoints, &ExcludedPoints, PointRef](const FPCGPointRef& InPointRef)
 			{
-				if ((Point->Transform.GetLocation() - InPointRef.Point->Transform.GetLocation()).SquaredLength() <= SMALL_NUMBER &&
+				if ((PointRef.Point->Transform.GetLocation() - InPointRef.Point->Transform.GetLocation()).SquaredLength() <= SMALL_NUMBER &&
 					!ExclusionPoints.Contains(InPointRef.Point))
 				{
 					ExcludedPoints.Add(InPointRef.Point);
@@ -127,11 +156,11 @@ namespace PCGSelfPruningElement
 			//  if in its vicinity, there is >=1 non-rejected point with a radius significantly larger
 			//  or in its range + has a randomly assigned index -> we'll look at its seed
 			//  then remove this point
-			TArray<const FPCGPoint*> SortedPoints;
+			TArray<FPCGPointRef> SortedPoints;
 			SortedPoints.Reserve(Points.Num());
 			for (const FPCGPoint& Point : Points)
 			{
-				SortedPoints.Add(&Point);
+				SortedPoints.Add(FPCGPointRef(Point));
 			}
 
 			// Apply proper sort algorithm
@@ -139,22 +168,22 @@ namespace PCGSelfPruningElement
 			{
 				if (bRandomizedPruning)
 				{
-					Algo::Sort(SortedPoints, [SquaredRadiusEquality](const FPCGPoint* A, const FPCGPoint* B) { return !PCGSelfPruningAlgorithms::SortSmallToLargeWithRandom(A, B, SquaredRadiusEquality); });
+					Algo::Sort(SortedPoints, [SquaredRadiusEquality](const FPCGPointRef& A, const FPCGPointRef& B) { return !PCGSelfPruningAlgorithms::SortSmallToLargeWithRandom(A, B, SquaredRadiusEquality); });
 				}
 				else
 				{
-					Algo::Sort(SortedPoints, [SquaredRadiusEquality](const FPCGPoint* A, const FPCGPoint* B) { return !PCGSelfPruningAlgorithms::SortSmallToLargeNoRandom(A, B, SquaredRadiusEquality); });
+					Algo::Sort(SortedPoints, [SquaredRadiusEquality](const FPCGPointRef& A, const FPCGPointRef& B) { return !PCGSelfPruningAlgorithms::SortSmallToLargeNoRandom(A, B, SquaredRadiusEquality); });
 				}
 			}
 			else if (PruningType == EPCGSelfPruningType::SmallToLarge)
 			{
 				if (bRandomizedPruning)
 				{
-					Algo::Sort(SortedPoints, [SquaredRadiusEquality](const FPCGPoint* A, const FPCGPoint* B) { return PCGSelfPruningAlgorithms::SortSmallToLargeWithRandom(A, B, SquaredRadiusEquality); });
+					Algo::Sort(SortedPoints, [SquaredRadiusEquality](const FPCGPointRef& A, const FPCGPointRef& B) { return PCGSelfPruningAlgorithms::SortSmallToLargeWithRandom(A, B, SquaredRadiusEquality); });
 				}
 				else
 				{
-					Algo::Sort(SortedPoints, [SquaredRadiusEquality](const FPCGPoint* A, const FPCGPoint* B) { return PCGSelfPruningAlgorithms::SortSmallToLargeNoRandom(A, B, SquaredRadiusEquality); });
+					Algo::Sort(SortedPoints, [SquaredRadiusEquality](const FPCGPointRef& A, const FPCGPointRef& B) { return PCGSelfPruningAlgorithms::SortSmallToLargeNoRandom(A, B, SquaredRadiusEquality); });
 				}
 			}
 			else
@@ -165,18 +194,17 @@ namespace PCGSelfPruningElement
 				}
 			}
 
-			TSet<const FPCGPoint*> ExclusionPoints;
-			ExclusionPoints.Reserve(Points.Num());
+			FPointBitSet ExclusionPoints(Points);
 
 			const bool bIsDuplicateTest = (PruningType == EPCGSelfPruningType::RemoveDuplicates);
 
 			if (bIsDuplicateTest)
 			{
-				PCGSelfPruningElement::DuplicatePointsExclusion(SortedPoints, Octree, ExclusionPoints);
+				PCGSelfPruningElement::DuplicatePointsExclusion(Points, SortedPoints, Octree, ExclusionPoints);
 			}
 			else
 			{
-				PCGSelfPruningElement::DensityBoundsExclusion(SortedPoints, Octree, ExclusionPoints);
+				PCGSelfPruningElement::DensityBoundsExclusion(Points, SortedPoints, Octree, ExclusionPoints);
 			}
 
 			// Finally, output all points that are present in the ExclusionPoints.
@@ -188,12 +216,14 @@ namespace PCGSelfPruningElement
 			Output.Data = PrunedData;
 
 			TArray<FPCGPoint>& OutputPoints = PrunedData->GetMutablePoints();
-			OutputPoints.Reserve(ExclusionPoints.Num());
 
-			for (const FPCGPoint* Point : ExclusionPoints)
+			for (const FPCGPoint& Point : Points)
 			{
-				OutputPoints.Add(*Point);
-			}
+				if (ExclusionPoints.Contains(&Point))
+				{
+					OutputPoints.Add(Point);
+				}
+			}				
 
 			if (bIsDuplicateTest)
 			{
