Set target actor to default to source component owner when creating data from the data from actor node

#rnx

diff --git a/Engine/Plugins/Experimental/PCG/Source/PCG/Private/Elements/PCGDataFromActor.cpp b/Engine/Plugins/Experimental/PCG/Source/PCG/Private/Elements/PCGDataFromActor.cpp
--- a/Engine/Plugins/Experimental/PCG/Source/PCG/Private/Elements/PCGDataFromActor.cpp
+++ b/Engine/Plugins/Experimental/PCG/Source/PCG/Private/Elements/PCGDataFromActor.cpp
@@ -3,6 +3,7 @@
 #include "Elements/PCGDataFromActor.h"
 
 #include "PCGComponent.h"
+#include "Data/PCGSpatialData.h"
 
 #include "GameFramework/Actor.h"
 #include "UObject/Package.h"
@@ -127,5 +128,16 @@ bool FPCGDataFromActorElement::ExecuteInternal(FPCGContext* Context) const
 		Output.Data = UPCGComponent::CreateActorPCGData(FoundActor, Context->SourceComponent.Get(), /*bParseActor=*/Settings->Mode != EPCGGetDataFromActorMode::GetSinglePoint);
 	}
 
+	if (Context->SourceComponent.IsValid())
+	{
+		for (FPCGTaggedData& Output : Outputs)
+		{
+			if (UPCGSpatialData* SpatialData = Cast<UPCGSpatialData>(Output.Data))
+			{
+				SpatialData->TargetActor = Context->SourceComponent->GetOwner();
+			}
+		}
+	}
+
 	return true;
 }
\ No newline at end of file
