Added info to Insights to find what PCG BP is running and removed some log spam
Fixed some bugs in the profiling tab view (it was only reporting the last 100 calls, which wasn't accurate over long timesliced tasks)

diff --git a/Engine/Plugins/Experimental/PCG/Source/PCG/Private/Elements/PCGExecuteBlueprint.cpp b/Engine/Plugins/Experimental/PCG/Source/PCG/Private/Elements/PCGExecuteBlueprint.cpp
--- a/Engine/Plugins/Experimental/PCG/Source/PCG/Private/Elements/PCGExecuteBlueprint.cpp
+++ b/Engine/Plugins/Experimental/PCG/Source/PCG/Private/Elements/PCGExecuteBlueprint.cpp
@@ -436,13 +436,14 @@ FPCGElementPtr UPCGBlueprintSettings::CreateElement() const
 
 bool FPCGExecuteBlueprintElement::ExecuteInternal(FPCGContext* InContext) const
 {
-	TRACE_CPUPROFILER_EVENT_SCOPE(FPCGExecuteBlueprintElement::Execute);
 	FPCGBlueprintExecutionContext* Context = static_cast<FPCGBlueprintExecutionContext*>(InContext);
 
 	if (Context && Context->BlueprintElementInstance)
 	{
 		UClass* BPClass = Context->BlueprintElementInstance->GetClass();
 
+		TRACE_CPUPROFILER_EVENT_SCOPE_TEXT(*FString::Printf(TEXT("FPCGExecuteBlueprintElement::Execute (%s)"), BPClass ? *BPClass->GetFName().ToString() : TEXT("")));
+
 #if WITH_EDITOR
 		/** Check if the blueprint has been successfully compiled */
 		if (UBlueprint* Blueprint = Cast<UBlueprint>(BPClass ? BPClass->ClassGeneratedBy : nullptr))
diff --git a/Engine/Plugins/Experimental/PCG/Source/PCG/Private/PCGElement.cpp b/Engine/Plugins/Experimental/PCG/Source/PCG/Private/PCGElement.cpp
--- a/Engine/Plugins/Experimental/PCG/Source/PCG/Private/PCGElement.cpp
+++ b/Engine/Plugins/Experimental/PCG/Source/PCG/Private/PCGElement.cpp
@@ -151,7 +151,10 @@ void IPCGElement::PostExecute(FPCGContext* Context) const
 	CleanupAndValidateOutput(Context);
 
 #if WITH_EDITOR
-	PCGE_LOG(Log, "Executed in (%f)s and (%d) call(s)", Context->ElapsedTime, Context->ExecutionCount);
+	{
+		FScopeLock Lock(&TimersLock);
+		PCGE_LOG(Verbose, "Executed in (%f)s and (%d) frames(s)", Timers[CurrentTimerIndex].ExecutionTime, Timers[CurrentTimerIndex].ExecutionFrameCount);
+	}
 #endif
 
 	const UPCGSettings* Settings = Context->GetInputSettings<UPCGSettings>();
@@ -217,7 +220,6 @@ void IPCGElement::ResetTimers()
 	FScopeLock Lock(&TimersLock);
 	Timers.Empty();
 	CurrentTimerIndex = 0;
-
 }
 #endif // WITH_EDITOR
 
@@ -263,30 +265,50 @@ void IPCGElement::CleanupAndValidateOutput(FPCGContext* Context) const
 
 #if WITH_EDITOR
 IPCGElement::FScopedCallTimer::FScopedCallTimer(const IPCGElement& InOwner, FPCGContext* InContext)
-	: Owner(InOwner), Context(InContext)
+	: Owner(InOwner), Context(InContext), Phase(InContext->CurrentPhase)
 {
 	StartTime = FPlatformTime::Seconds();
 }
 
 IPCGElement::FScopedCallTimer::~FScopedCallTimer()
 {
-	const double EndTime = FPlatformTime::Seconds();
-	const double ElapsedTime = EndTime - StartTime;
-	Context->ElapsedTime += ElapsedTime;
-	Context->ExecutionCount++;
+	const double ThisFrameTime = FPlatformTime::Seconds() - StartTime;
 
-	constexpr int MaxNumberOfTrackedTimers = 100;
+	FScopeLock Lock(&Owner.TimersLock);
+
+	constexpr int32 MaxNumberOfTrackedTimers = 100;
+	if (Phase == EPCGExecutionPhase::PrepareData)
 	{
-		FScopeLock Lock(&Owner.TimersLock);
 		if (Owner.Timers.Num() < MaxNumberOfTrackedTimers)
 		{
-			Owner.Timers.Add(ElapsedTime);
+			// first time here
+			Owner.Timers.Add({});
+			Owner.CurrentTimerIndex = Owner.Timers.Num()-1;
 		}
 		else
 		{
-			Owner.Timers[Owner.CurrentTimerIndex] = ElapsedTime;
+			Owner.CurrentTimerIndex = (Owner.CurrentTimerIndex+1) % MaxNumberOfTrackedTimers;
+			Owner.Timers[Owner.CurrentTimerIndex] = FCallTime();
 		}
-		Owner.CurrentTimerIndex = (Owner.CurrentTimerIndex + 1) % MaxNumberOfTrackedTimers;
+
+		FCallTime& Time = Owner.Timers[Owner.CurrentTimerIndex];
+		Time.PrepareDataTime = ThisFrameTime;
+	}
+	else if (Phase == EPCGExecutionPhase::Execute)
+	{
+		FCallTime& Time = Owner.Timers[Owner.CurrentTimerIndex];
+
+		Time.ExecutionTime += ThisFrameTime;
+		Time.ExecutionFrameCount++;
+
+		Time.MaxExecutionFrameTime = FMath::Max(Time.MaxExecutionFrameTime, ThisFrameTime);
+		Time.MinExecutionFrameTime = FMath::Min(Time.MinExecutionFrameTime, ThisFrameTime);
+	}
+	else if (Phase == EPCGExecutionPhase::PostExecute)
+	{
+		FCallTime& Time = Owner.Timers[Owner.CurrentTimerIndex];
+
+		Time.PostExecuteTime = ThisFrameTime;
 	}
 }
 #endif // WITH_EDITOR
diff --git a/Engine/Plugins/Experimental/PCG/Source/PCG/Public/PCGContext.h b/Engine/Plugins/Experimental/PCG/Source/PCG/Public/PCGContext.h
--- a/Engine/Plugins/Experimental/PCG/Source/PCG/Public/PCGContext.h
+++ b/Engine/Plugins/Experimental/PCG/Source/PCG/Public/PCGContext.h
@@ -64,11 +64,6 @@ struct PCG_API FPCGContext
 	/** True if currently inside a PCGAsync scope - will prevent further async processing */
 	bool bIsRunningAsyncCall = false;
 
-#if WITH_EDITOR
-	double ElapsedTime = 0.0;
-	int32 ExecutionCount = 0;
-#endif
-
 	template<typename SettingsType>
 	const SettingsType* GetInputSettings() const
 	{
diff --git a/Engine/Plugins/Experimental/PCG/Source/PCG/Public/PCGElement.h b/Engine/Plugins/Experimental/PCG/Source/PCG/Public/PCGElement.h
--- a/Engine/Plugins/Experimental/PCG/Source/PCG/Public/PCGElement.h
+++ b/Engine/Plugins/Experimental/PCG/Source/PCG/Public/PCGElement.h
@@ -51,8 +51,20 @@ public:
 
 	/** Note: the following methods must be called from the main thread */
 #if WITH_EDITOR
+	struct FCallTime
+	{
+		// sum of all frames
+		double ExecutionTime = 0;
+		// how many frames
+		int32 ExecutionFrameCount = 0;
+		double MinExecutionFrameTime = std::numeric_limits<double>::max();
+		double MaxExecutionFrameTime = 0.0;
+		double PrepareDataTime = 0.0;
+		double PostExecuteTime = 0.0;
+	};
+
 	void DebugDisplay(FPCGContext* Context) const;
-	const TArray<double>& GetTimers() const { return Timers; }
+	const TArray<FCallTime>& GetTimers() const { return Timers; }
 	void ResetTimers();
 #endif
 
@@ -87,13 +99,15 @@ private:
 		const IPCGElement& Owner;
 		FPCGContext* Context;
 		double StartTime;
+		EPCGExecutionPhase Phase;
 	};
 
 	// Set mutable because we need to modify them in the execute call, which is const
 	// TODO: Should be a map with PCG Components. We need a mechanism to make sure that this map is cleaned up when component doesn't exist anymore.
 	// For now, it will track all calls to execute (excluding call where the result is already in cache).
-	mutable TArray<double> Timers;
-	mutable int CurrentTimerIndex = 0;
+	mutable TArray<FCallTime> Timers;
+	mutable int32 CurrentTimerIndex = 0;
+
 	// Perhaps overkill but there is a slight chance that we need to protect the timers array. If we call reset from the UI while an element is executing,
 	// it could crash while writing to the timers array.
 	mutable FCriticalSection TimersLock;
@@ -101,6 +115,7 @@ private:
 	struct FScopedCallTimer
 	{
 		FScopedCallTimer(const IPCGElement& InOwner, FPCGContext* InContext) {}
+		void AdvanceTimerIndex() {}
 	};
 #endif // WITH_EDITOR
 };
diff --git a/Engine/Plugins/Experimental/PCG/Source/PCGEditor/Private/SPCGEditorGraphProfilingView.cpp b/Engine/Plugins/Experimental/PCG/Source/PCGEditor/Private/SPCGEditorGraphProfilingView.cpp
--- a/Engine/Plugins/Experimental/PCG/Source/PCGEditor/Private/SPCGEditorGraphProfilingView.cpp
+++ b/Engine/Plugins/Experimental/PCG/Source/PCGEditor/Private/SPCGEditorGraphProfilingView.cpp
@@ -22,21 +22,33 @@ namespace PCGEditorGraphProfilingView
 	
 	/** Names of the columns in the attribute list */
 	const FName NAME_Node = FName(TEXT("Node"));
-	const FName NAME_AvgTime = FName(TEXT("AvgTime"));
-	const FName NAME_MinTime = FName(TEXT("MinTime"));
-	const FName NAME_MaxTime = FName(TEXT("MaxTime"));
-	const FName NAME_StdTime = FName(TEXT("StdTime"));
-	const FName NAME_TotalTime = FName(TEXT("TotalTime"));
+	const FName NAME_PrepareDataTime = FName(TEXT("PrepareDataTime"));
+	const FName NAME_AvgExecutionTime = FName(TEXT("AvgExecutionTime"));
+	const FName NAME_MinExecutionTime = FName(TEXT("MinExecutionTime"));
+	const FName NAME_MaxExecutionTime = FName(TEXT("MaxExecutionTime"));
+	const FName NAME_MinExecutionFrameTime = FName(TEXT("MinFrameTime"));
+	const FName NAME_MaxExecutionFrameTime = FName(TEXT("MaxFrameTime"));
+	const FName NAME_MinNbExecutionFrames = FName(TEXT("MinNbExecutionFrames"));
+	const FName NAME_MaxNbExecutionFrames = FName(TEXT("MaxNbExecutionFrames"));
+	const FName NAME_StdExecutionTime = FName(TEXT("StdExecutionTime"));
+	const FName NAME_TotalExecutionTime = FName(TEXT("TotalExecutionTime"));
 	const FName NAME_NbCalls = FName(TEXT("NbCalls"));
+	const FName NAME_NbExecutionFrames = FName(TEXT("NbExecutionFrames"));
+	const FName NAME_PostExecuteTime = FName(TEXT("PostExecuteTime"));
 
 	/** Labels of the columns */
 	const FText TEXT_NodeLabel = LOCTEXT("NodeLabel", "Node");
-	const FText TEXT_AvgTimeLabel = LOCTEXT("AvgTimeLabel", "Avg Time(ms)");
-	const FText TEXT_MinTimeLabel = LOCTEXT("MinTimeLabel", "Min Time(ms)");
-	const FText TEXT_MaxTimeLabel = LOCTEXT("MaxTimeLabel", "Max Time(ms)");
-	const FText TEXT_StdTimeLabel = LOCTEXT("StdTimeLabel", "Std(ms)");
-	const FText TEXT_TotalTimeLabel = LOCTEXT("TotalTimeLabel", "Total time(s)");
-	const FText TEXT_NbCallsLabel = LOCTEXT("NbCallsLabel", "Number of calls");
+	const FText TEXT_PrepareDataTimeLabel = LOCTEXT("PrepareDataTimeLabel", "PrepareData (ms)");
+	const FText TEXT_PostExecuteTimeLabel = LOCTEXT("PrepareDataTimeLabel", "PrepareData (ms)");
+	const FText TEXT_AvgExecutionTimeLabel = LOCTEXT("AvgExecutionTimeLabel", "Avg Time(ms)");
+	const FText TEXT_MinExecutionTimeLabel = LOCTEXT("MinExecutionTimeLabel", "Min Time(ms)");
+	const FText TEXT_MaxExecutionTimeLabel = LOCTEXT("MaxExecutionTimeLabel", "Max Time(ms)");
+	const FText TEXT_MinExecutionFrameTimeLabel = LOCTEXT("MinExecutionFrameTimeLabel", "Min Frame Time(ms)");
+	const FText TEXT_MaxExecutionFrameTimeLabel = LOCTEXT("MaxExecutionFrameTimeLabel", "Max Frame Time(ms)");
+	const FText TEXT_StdExecutionTimeLabel = LOCTEXT("StdExecutionTimeLabel", "Std(ms)");
+	const FText TEXT_TotalExecutionTimeLabel = LOCTEXT("TotalExecutionTimeLabel", "Total time(s)");
+	const FText TEXT_NbCallsLabel = LOCTEXT("NbCallsLabel", "Calls");
+	const FText TEXT_NbExecutionFramesLabel = LOCTEXT("NbExecutionFramesLabel", "Exec frames");
 }
 
 void SPCGProfilingListViewItemRow::Construct(const FArguments& InArgs, const TSharedRef<STableViewBase>& InOwnerTableView, const PCGProfilingListViewItemPtr& Item)
@@ -62,35 +74,65 @@ TSharedRef<SWidget> SPCGProfilingListViewItemRow::GenerateWidgetForColumn(const
 		{
 			ColumnData = FText::AsNumber(InternalItem->NbCalls);
 		}
+		else if (ColumnId == PCGEditorGraphProfilingView::NAME_NbExecutionFrames)
+		{
+			ColumnData = FText::AsNumber(InternalItem->NbExecutionFrames);
+		}
+		else if (ColumnId == PCGEditorGraphProfilingView::NAME_MinNbExecutionFrames)
+		{
+			ColumnData = FText::AsNumber(InternalItem->MinNbExecutionFrames);
+		}
+		else if (ColumnId == PCGEditorGraphProfilingView::NAME_MaxNbExecutionFrames)
+		{
+			ColumnData = FText::AsNumber(InternalItem->MaxNbExecutionFrames);
+		}
 		// For all other values, if we don't have data, just write "N/A"
 		else if (!InternalItem->HasData)
 		{
 			ColumnData = PCGEditorGraphProfilingView::NoDataAvailableText;
 		}
-		else if (ColumnId == PCGEditorGraphProfilingView::NAME_AvgTime)
+		else if (ColumnId == PCGEditorGraphProfilingView::NAME_AvgExecutionTime)
+		{
+			// In ms
+			ColumnData = FText::AsNumber(InternalItem->AvgExecutionTime * 1000.0);
+		}
+		else if (ColumnId == PCGEditorGraphProfilingView::NAME_MinExecutionTime)
+		{
+			// In ms
+			ColumnData = FText::AsNumber(InternalItem->MinExecutionTime * 1000.0);
+		}
+		else if (ColumnId == PCGEditorGraphProfilingView::NAME_MaxExecutionTime)
 		{
 			// In ms
-			ColumnData = FText::AsNumber(InternalItem->AvgTime * 1000.0);
+			ColumnData = FText::AsNumber(InternalItem->MaxExecutionTime * 1000.0);
 		}
-		else if (ColumnId == PCGEditorGraphProfilingView::NAME_MinTime)
+		else if (ColumnId == PCGEditorGraphProfilingView::NAME_MinExecutionFrameTime)
 		{
 			// In ms
-			ColumnData = FText::AsNumber(InternalItem->MinTime * 1000.0);
+			ColumnData = FText::AsNumber(InternalItem->MinExecutionFrameTime * 1000.0);
 		}
-		else if (ColumnId == PCGEditorGraphProfilingView::NAME_MaxTime)
+		else if (ColumnId == PCGEditorGraphProfilingView::NAME_MaxExecutionFrameTime)
 		{
 			// In ms
-			ColumnData = FText::AsNumber(InternalItem->MaxTime * 1000.0);
+			ColumnData = FText::AsNumber(InternalItem->MaxExecutionFrameTime * 1000.0);
 		}
-		else if (ColumnId == PCGEditorGraphProfilingView::NAME_StdTime)
+		else if (ColumnId == PCGEditorGraphProfilingView::NAME_StdExecutionTime)
 		{
 			// In ms
-			ColumnData = FText::AsNumber(InternalItem->StdTime * 1000.0);
+			ColumnData = FText::AsNumber(InternalItem->StdExecutionTime * 1000.0);
 		}
-		else if (ColumnId == PCGEditorGraphProfilingView::NAME_TotalTime)
+		else if (ColumnId == PCGEditorGraphProfilingView::NAME_TotalExecutionTime)
 		{
 			// In s
-			ColumnData = FText::AsNumber(InternalItem->TotalTime);
+			ColumnData = FText::AsNumber(InternalItem->ExecutionTime);
+		}
+		else if (ColumnId == PCGEditorGraphProfilingView::NAME_PrepareDataTime)
+		{
+			ColumnData = FText::AsNumber(InternalItem->PrepareDataTime * 1000.0);
+		}
+		else if (ColumnId == PCGEditorGraphProfilingView::NAME_PostExecuteTime)
+		{
+			ColumnData = FText::AsNumber(InternalItem->PostExecuteTime * 1000.0);
 		}
 	}
 
@@ -211,47 +253,75 @@ TSharedRef<SHeaderRow> SPCGEditorGraphProfilingView::CreateHeaderRowWidget()
 		.HAlignCell(HAlign_Left)
 		.SortMode(this, &SPCGEditorGraphProfilingView::GetColumnSortMode, PCGEditorGraphProfilingView::NAME_Node)
 		.OnSort(this, &SPCGEditorGraphProfilingView::OnSortColumnHeader)
-		+ SHeaderRow::Column(PCGEditorGraphProfilingView::NAME_AvgTime)
-		.ManualWidth(100)
-		.DefaultLabel(PCGEditorGraphProfilingView::TEXT_AvgTimeLabel)
+		+ SHeaderRow::Column(PCGEditorGraphProfilingView::NAME_NbCalls)
+		.ManualWidth(80)
+		.DefaultLabel(PCGEditorGraphProfilingView::TEXT_NbCallsLabel)
+		.HAlignHeader(HAlign_Center)
+		.HAlignCell(HAlign_Right)
+		.SortMode(this, &SPCGEditorGraphProfilingView::GetColumnSortMode, PCGEditorGraphProfilingView::NAME_NbCalls)
+		.OnSort(this, &SPCGEditorGraphProfilingView::OnSortColumnHeader)
+		+ SHeaderRow::Column(PCGEditorGraphProfilingView::NAME_PrepareDataTime)
+		.ManualWidth(125)
+		.DefaultLabel(PCGEditorGraphProfilingView::TEXT_PrepareDataTimeLabel)
+		.HAlignHeader(HAlign_Center)
+		.HAlignCell(HAlign_Right)
+		.SortMode(this, &SPCGEditorGraphProfilingView::GetColumnSortMode, PCGEditorGraphProfilingView::NAME_PrepareDataTime)
+		.OnSort(this, &SPCGEditorGraphProfilingView::OnSortColumnHeader)
+		+ SHeaderRow::Column(PCGEditorGraphProfilingView::NAME_NbExecutionFrames)
+		.ManualWidth(80)
+		.DefaultLabel(PCGEditorGraphProfilingView::TEXT_NbExecutionFramesLabel)
+		.HAlignHeader(HAlign_Center)
+		.HAlignCell(HAlign_Right)
+		.SortMode(this, &SPCGEditorGraphProfilingView::GetColumnSortMode, PCGEditorGraphProfilingView::NAME_NbExecutionFrames)
+		.OnSort(this, &SPCGEditorGraphProfilingView::OnSortColumnHeader)
+		+ SHeaderRow::Column(PCGEditorGraphProfilingView::NAME_MinExecutionFrameTime)
+		.ManualWidth(130)
+		.DefaultLabel(PCGEditorGraphProfilingView::TEXT_MinExecutionFrameTimeLabel)
+		.HAlignHeader(HAlign_Center)
+		.HAlignCell(HAlign_Right)
+		.SortMode(this, &SPCGEditorGraphProfilingView::GetColumnSortMode, PCGEditorGraphProfilingView::NAME_MinExecutionFrameTime)
+		.OnSort(this, &SPCGEditorGraphProfilingView::OnSortColumnHeader)
+		+ SHeaderRow::Column(PCGEditorGraphProfilingView::NAME_MaxExecutionFrameTime)
+		.ManualWidth(130)
+		.DefaultLabel(PCGEditorGraphProfilingView::TEXT_MaxExecutionFrameTimeLabel)
 		.HAlignHeader(HAlign_Center)
 		.HAlignCell(HAlign_Right)
-		.SortMode(this, &SPCGEditorGraphProfilingView::GetColumnSortMode, PCGEditorGraphProfilingView::NAME_AvgTime)
+		.SortMode(this, &SPCGEditorGraphProfilingView::GetColumnSortMode, PCGEditorGraphProfilingView::NAME_MaxExecutionFrameTime)
 		.OnSort(this, &SPCGEditorGraphProfilingView::OnSortColumnHeader)
-		+ SHeaderRow::Column(PCGEditorGraphProfilingView::NAME_MinTime)
+		+ SHeaderRow::Column(PCGEditorGraphProfilingView::NAME_AvgExecutionTime)
 		.ManualWidth(100)
-		.DefaultLabel(PCGEditorGraphProfilingView::TEXT_MinTimeLabel)
+		.DefaultLabel(PCGEditorGraphProfilingView::TEXT_AvgExecutionTimeLabel)
 		.HAlignHeader(HAlign_Center)
 		.HAlignCell(HAlign_Right)
-		.SortMode(this, &SPCGEditorGraphProfilingView::GetColumnSortMode, PCGEditorGraphProfilingView::NAME_MinTime)
+		.SortMode(this, &SPCGEditorGraphProfilingView::GetColumnSortMode, PCGEditorGraphProfilingView::NAME_AvgExecutionTime)
 		.OnSort(this, &SPCGEditorGraphProfilingView::OnSortColumnHeader)
-		+ SHeaderRow::Column(PCGEditorGraphProfilingView::NAME_MaxTime)
+		+ SHeaderRow::Column(PCGEditorGraphProfilingView::NAME_MinExecutionTime)
 		.ManualWidth(100)
-		.DefaultLabel(PCGEditorGraphProfilingView::TEXT_MaxTimeLabel)
+		.DefaultLabel(PCGEditorGraphProfilingView::TEXT_MinExecutionTimeLabel)
 		.HAlignHeader(HAlign_Center)
 		.HAlignCell(HAlign_Right)
-		.SortMode(this, &SPCGEditorGraphProfilingView::GetColumnSortMode, PCGEditorGraphProfilingView::NAME_MaxTime)
+		.SortMode(this, &SPCGEditorGraphProfilingView::GetColumnSortMode, PCGEditorGraphProfilingView::NAME_MinExecutionTime)
 		.OnSort(this, &SPCGEditorGraphProfilingView::OnSortColumnHeader)
-		+ SHeaderRow::Column(PCGEditorGraphProfilingView::NAME_StdTime)
+		+ SHeaderRow::Column(PCGEditorGraphProfilingView::NAME_MaxExecutionTime)
 		.ManualWidth(100)
-		.DefaultLabel(PCGEditorGraphProfilingView::TEXT_StdTimeLabel)
+		.DefaultLabel(PCGEditorGraphProfilingView::TEXT_MaxExecutionTimeLabel)
 		.HAlignHeader(HAlign_Center)
 		.HAlignCell(HAlign_Right)
-		.SortMode(this, &SPCGEditorGraphProfilingView::GetColumnSortMode, PCGEditorGraphProfilingView::NAME_StdTime)
+		.SortMode(this, &SPCGEditorGraphProfilingView::GetColumnSortMode, PCGEditorGraphProfilingView::NAME_MaxExecutionTime)
 		.OnSort(this, &SPCGEditorGraphProfilingView::OnSortColumnHeader)
-		+ SHeaderRow::Column(PCGEditorGraphProfilingView::NAME_TotalTime)
+		+ SHeaderRow::Column(PCGEditorGraphProfilingView::NAME_StdExecutionTime)
 		.ManualWidth(100)
-		.DefaultLabel(PCGEditorGraphProfilingView::TEXT_TotalTimeLabel)
+		.DefaultLabel(PCGEditorGraphProfilingView::TEXT_StdExecutionTimeLabel)
 		.HAlignHeader(HAlign_Center)
 		.HAlignCell(HAlign_Right)
-		.SortMode(this, &SPCGEditorGraphProfilingView::GetColumnSortMode, PCGEditorGraphProfilingView::NAME_TotalTime)
+		.SortMode(this, &SPCGEditorGraphProfilingView::GetColumnSortMode, PCGEditorGraphProfilingView::NAME_StdExecutionTime)
 		.OnSort(this, &SPCGEditorGraphProfilingView::OnSortColumnHeader)
-		+ SHeaderRow::Column(PCGEditorGraphProfilingView::NAME_NbCalls)
-		.ManualWidth(125)
-		.DefaultLabel(PCGEditorGraphProfilingView::TEXT_NbCallsLabel)
+		+ SHeaderRow::Column(PCGEditorGraphProfilingView::NAME_TotalExecutionTime)
+		.ManualWidth(100)
+		.DefaultLabel(PCGEditorGraphProfilingView::TEXT_TotalExecutionTimeLabel)
 		.HAlignHeader(HAlign_Center)
 		.HAlignCell(HAlign_Right)
-		.SortMode(this, &SPCGEditorGraphProfilingView::GetColumnSortMode, PCGEditorGraphProfilingView::NAME_NbCalls)
+		.SortMode(this, &SPCGEditorGraphProfilingView::GetColumnSortMode, PCGEditorGraphProfilingView::NAME_TotalExecutionTime)
 		.OnSort(this, &SPCGEditorGraphProfilingView::OnSortColumnHeader);
 }
 
@@ -355,37 +425,43 @@ FReply SPCGEditorGraphProfilingView::Refresh()
 				{
 					if (const IPCGElement* Element = Settings->GetElement().Get())
 					{
-						const TArray<double>& Timers = Element->GetTimers();
-						if (!Timers.IsEmpty())
+						const TArray<IPCGElement::FCallTime>& Timers = Element->GetTimers();
+						if (Timers.IsEmpty())
+						{
+							continue;
+						}
+
+						ListViewItem->MinExecutionTime = std::numeric_limits<double>::max();
+						ListViewItem->MaxExecutionTime = std::numeric_limits<double>::min();
+						ListViewItem->MinExecutionFrameTime = std::numeric_limits<double>::max();
+						ListViewItem->MaxExecutionFrameTime = std::numeric_limits<double>::min();
+
+						for(const IPCGElement::FCallTime& Timer : Timers)
 						{
-							ListViewItem->MinTime = std::numeric_limits<double>::max();
-							ListViewItem->MaxTime = std::numeric_limits<double>::min();
-							for (double v : Timers)
-							{
-								ListViewItem->TotalTime += v;
-								if (v < ListViewItem->MinTime)
-								{
-									ListViewItem->MinTime = v;
-								}
-
-								if (v > ListViewItem->MaxTime)
-								{
-									ListViewItem->MaxTime = v;
-								}
-							}
-
-							ListViewItem->AvgTime = ListViewItem->TotalTime / Timers.Num();
-
-							for (double v : Timers)
-							{
-								ListViewItem->StdTime += (ListViewItem->AvgTime - v) * (ListViewItem->AvgTime - v);
-							}
-
-							ListViewItem->StdTime = FMath::Sqrt(ListViewItem->StdTime / Timers.Num());
-
-							ListViewItem->HasData = true;
-							ListViewItem->NbCalls = Timers.Num();
+							ListViewItem->PrepareDataTime += Timer.PrepareDataTime;
+
+							ListViewItem->MinExecutionFrameTime = FMath::Min(ListViewItem->MinExecutionFrameTime, Timer.MinExecutionFrameTime);
+							ListViewItem->MaxExecutionFrameTime = FMath::Max(ListViewItem->MaxExecutionFrameTime, Timer.MaxExecutionFrameTime);
+
+							ListViewItem->ExecutionTime += Timer.ExecutionTime;
+							ListViewItem->MinExecutionTime = FMath::Min(ListViewItem->MinExecutionTime, Timer.ExecutionTime);
+							ListViewItem->MaxExecutionTime = FMath::Max(ListViewItem->MaxExecutionTime, Timer.ExecutionTime);
+
+							ListViewItem->NbExecutionFrames += Timer.ExecutionFrameCount;
+							ListViewItem->MaxNbExecutionFrames = FMath::Max(ListViewItem->MaxNbExecutionFrames, Timer.ExecutionFrameCount);
+							ListViewItem->MinNbExecutionFrames = FMath::Min(ListViewItem->MinNbExecutionFrames, Timer.ExecutionFrameCount);
 						}
+
+						ListViewItem->NbCalls = Timers.Num();
+						ListViewItem->AvgExecutionTime = ListViewItem->ExecutionTime / ListViewItem->NbCalls;
+
+						for (const IPCGElement::FCallTime& Timer : Timers)
+						{
+							ListViewItem->StdExecutionTime += FMath::Square(ListViewItem->AvgExecutionTime - Timer.ExecutionTime);
+						}
+
+						ListViewItem->StdExecutionTime = FMath::Sqrt(ListViewItem->StdExecutionTime / ListViewItem->NbCalls);
+						ListViewItem->HasData = true;
 					}
 				}
 			}
@@ -403,30 +479,54 @@ FReply SPCGEditorGraphProfilingView::Refresh()
 				{
 					isLess = A->Name.FastLess(B->Name);
 				}
-				else if (SortingColumn == PCGEditorGraphProfilingView::NAME_AvgTime)
+				else if (SortingColumn == PCGEditorGraphProfilingView::NAME_PrepareDataTime)
+				{
+					isLess = A->PrepareDataTime < B->PrepareDataTime;
+				}
+				else if (SortingColumn == PCGEditorGraphProfilingView::NAME_AvgExecutionTime)
+				{
+					isLess = A->AvgExecutionTime < B->AvgExecutionTime;
+				}
+				else if (SortingColumn == PCGEditorGraphProfilingView::NAME_MinExecutionTime)
 				{
-					isLess = A->AvgTime < B->AvgTime;
+					isLess = A->MinExecutionTime < B->MinExecutionTime;
 				}
-				else if (SortingColumn == PCGEditorGraphProfilingView::NAME_MinTime)
+				else if (SortingColumn == PCGEditorGraphProfilingView::NAME_MaxExecutionTime)
 				{
-					isLess = A->MinTime < B->MinTime;
+					isLess = A->MaxExecutionTime < B->MaxExecutionTime;
 				}
-				else if (SortingColumn == PCGEditorGraphProfilingView::NAME_MaxTime)
+				else if (SortingColumn == PCGEditorGraphProfilingView::NAME_MinExecutionFrameTime)
 				{
-					isLess = A->MaxTime < B->MaxTime;
+					isLess = A->MinExecutionFrameTime < B->MinExecutionFrameTime;
 				}
-				else if (SortingColumn == PCGEditorGraphProfilingView::NAME_StdTime)
+				else if (SortingColumn == PCGEditorGraphProfilingView::NAME_MaxExecutionFrameTime)
 				{
-					isLess = A->StdTime < B->StdTime;
+					isLess = A->MaxExecutionFrameTime < B->MaxExecutionFrameTime;
 				}
-				else if (SortingColumn == PCGEditorGraphProfilingView::NAME_TotalTime)
+				else if (SortingColumn == PCGEditorGraphProfilingView::NAME_StdExecutionTime)
 				{
-					isLess = A->TotalTime < B->TotalTime;
+					isLess = A->StdExecutionTime < B->StdExecutionTime;
+				}
+				else if (SortingColumn == PCGEditorGraphProfilingView::NAME_TotalExecutionTime)
+				{
+					isLess = A->ExecutionTime < B->ExecutionTime;
 				}
 				else if (SortingColumn == PCGEditorGraphProfilingView::NAME_NbCalls)
 				{
 					isLess = A->NbCalls < B->NbCalls;
 				}
+				else if (SortingColumn == PCGEditorGraphProfilingView::NAME_NbExecutionFrames)
+				{
+					isLess = A->NbExecutionFrames < B->NbExecutionFrames;
+				}
+				else if (SortingColumn == PCGEditorGraphProfilingView::NAME_MaxNbExecutionFrames)
+				{
+					isLess = A->MaxNbExecutionFrames < B->MaxNbExecutionFrames;
+				}
+				else if (SortingColumn == PCGEditorGraphProfilingView::NAME_MinNbExecutionFrames)
+				{
+					isLess = A->MinNbExecutionFrames < B->MinNbExecutionFrames;
+				}
 
 				return SortMode == EColumnSortMode::Ascending ? isLess : !isLess;
 			});
diff --git a/Engine/Plugins/Experimental/PCG/Source/PCGEditor/Private/SPCGEditorGraphProfilingView.h b/Engine/Plugins/Experimental/PCG/Source/PCGEditor/Private/SPCGEditorGraphProfilingView.h
--- a/Engine/Plugins/Experimental/PCG/Source/PCGEditor/Private/SPCGEditorGraphProfilingView.h
+++ b/Engine/Plugins/Experimental/PCG/Source/PCGEditor/Private/SPCGEditorGraphProfilingView.h
@@ -21,12 +21,19 @@ struct FPCGProfilingListViewItem
 	const UPCGEditorGraphNode* EditorNode = nullptr;
 
 	FName Name = NAME_None;
-	double AvgTime = 0.0;
-	double MinTime = 0.0;
-	double MaxTime = 0.0;
-	double StdTime = 0.0;
-	double TotalTime = 0.0;
+	double PrepareDataTime = 0.0;
+	double ExecutionTime = 0.0;
+	double AvgExecutionTime = 0.0;
+	double MinExecutionTime = 0.0;
+	double MaxExecutionTime = 0.0;
+	double MinExecutionFrameTime = 0.0;
+	double MaxExecutionFrameTime = 0.0;
+	double StdExecutionTime = 0.0;
+	double PostExecuteTime = 0.0;
 	int32 NbCalls = 0;
+	int32 NbExecutionFrames = 0;
+	int32 MaxNbExecutionFrames = 0;
+	int32 MinNbExecutionFrames = 0;
 	bool HasData = false;
 };
 
