[Backout] - CL23163457 to resolve compile error
Original CL Desc
-----------------------------------------------------------------
Added info to Insights to find what PCG BP is running and removed some log spam
Fixed some bugs in the profiling tab view (it was only reporting the last 100 calls, which wasn't accurate over long timesliced tasks)

diff --git a/Engine/Plugins/Experimental/PCG/Source/PCG/Private/Elements/PCGExecuteBlueprint.cpp b/Engine/Plugins/Experimental/PCG/Source/PCG/Private/Elements/PCGExecuteBlueprint.cpp
--- a/Engine/Plugins/Experimental/PCG/Source/PCG/Private/Elements/PCGExecuteBlueprint.cpp
+++ b/Engine/Plugins/Experimental/PCG/Source/PCG/Private/Elements/PCGExecuteBlueprint.cpp
@@ -436,14 +436,13 @@ FPCGElementPtr UPCGBlueprintSettings::CreateElement() const
 
 bool FPCGExecuteBlueprintElement::ExecuteInternal(FPCGContext* InContext) const
 {
+	TRACE_CPUPROFILER_EVENT_SCOPE(FPCGExecuteBlueprintElement::Execute);
 	FPCGBlueprintExecutionContext* Context = static_cast<FPCGBlueprintExecutionContext*>(InContext);
 
 	if (Context && Context->BlueprintElementInstance)
 	{
 		UClass* BPClass = Context->BlueprintElementInstance->GetClass();
 
-		TRACE_CPUPROFILER_EVENT_SCOPE_TEXT(*FString::Printf(TEXT("FPCGExecuteBlueprintElement::Execute (%s)"), BPClass ? *BPClass->GetFName().ToString() : TEXT("")));
-
 #if WITH_EDITOR
 		/** Check if the blueprint has been successfully compiled */
 		if (UBlueprint* Blueprint = Cast<UBlueprint>(BPClass ? BPClass->ClassGeneratedBy : nullptr))
diff --git a/Engine/Plugins/Experimental/PCG/Source/PCG/Private/PCGElement.cpp b/Engine/Plugins/Experimental/PCG/Source/PCG/Private/PCGElement.cpp
--- a/Engine/Plugins/Experimental/PCG/Source/PCG/Private/PCGElement.cpp
+++ b/Engine/Plugins/Experimental/PCG/Source/PCG/Private/PCGElement.cpp
@@ -151,10 +151,7 @@ void IPCGElement::PostExecute(FPCGContext* Context) const
 	CleanupAndValidateOutput(Context);
 
 #if WITH_EDITOR
-	{
-		FScopeLock Lock(&TimersLock);
-		PCGE_LOG(Verbose, "Executed in (%f)s and (%d) frames(s)", Timers[CurrentTimerIndex].ExecutionTime, Timers[CurrentTimerIndex].ExecutionFrameCount);
-	}
+	PCGE_LOG(Log, "Executed in (%f)s and (%d) call(s)", Context->ElapsedTime, Context->ExecutionCount);
 #endif
 
 	const UPCGSettings* Settings = Context->GetInputSettings<UPCGSettings>();
@@ -220,6 +217,7 @@ void IPCGElement::ResetTimers()
 	FScopeLock Lock(&TimersLock);
 	Timers.Empty();
 	CurrentTimerIndex = 0;
+
 }
 #endif // WITH_EDITOR
 
@@ -265,50 +263,30 @@ void IPCGElement::CleanupAndValidateOutput(FPCGContext* Context) const
 
 #if WITH_EDITOR
 IPCGElement::FScopedCallTimer::FScopedCallTimer(const IPCGElement& InOwner, FPCGContext* InContext)
-	: Owner(InOwner), Context(InContext), Phase(InContext->CurrentPhase)
+	: Owner(InOwner), Context(InContext)
 {
 	StartTime = FPlatformTime::Seconds();
 }
 
 IPCGElement::FScopedCallTimer::~FScopedCallTimer()
 {
-	const double ThisFrameTime = FPlatformTime::Seconds() - StartTime;
+	const double EndTime = FPlatformTime::Seconds();
+	const double ElapsedTime = EndTime - StartTime;
+	Context->ElapsedTime += ElapsedTime;
+	Context->ExecutionCount++;
 
-	FScopeLock Lock(&Owner.TimersLock);
-
-	constexpr int32 MaxNumberOfTrackedTimers = 100;
-	if (Phase == EPCGExecutionPhase::PrepareData)
+	constexpr int MaxNumberOfTrackedTimers = 100;
 	{
+		FScopeLock Lock(&Owner.TimersLock);
 		if (Owner.Timers.Num() < MaxNumberOfTrackedTimers)
 		{
-			// first time here
-			Owner.Timers.Add({});
-			Owner.CurrentTimerIndex = Owner.Timers.Num()-1;
+			Owner.Timers.Add(ElapsedTime);
 		}
 		else
 		{
-			Owner.CurrentTimerIndex = (Owner.CurrentTimerIndex+1) % MaxNumberOfTrackedTimers;
-			Owner.Timers[Owner.CurrentTimerIndex] = FCallTime();
+			Owner.Timers[Owner.CurrentTimerIndex] = ElapsedTime;
 		}
-
-		FCallTime& Time = Owner.Timers[Owner.CurrentTimerIndex];
-		Time.PrepareDataTime = ThisFrameTime;
-	}
-	else if (Phase == EPCGExecutionPhase::Execute)
-	{
-		FCallTime& Time = Owner.Timers[Owner.CurrentTimerIndex];
-
-		Time.ExecutionTime += ThisFrameTime;
-		Time.ExecutionFrameCount++;
-
-		Time.MaxExecutionFrameTime = FMath::Max(Time.MaxExecutionFrameTime, ThisFrameTime);
-		Time.MinExecutionFrameTime = FMath::Min(Time.MinExecutionFrameTime, ThisFrameTime);
-	}
-	else if (Phase == EPCGExecutionPhase::PostExecute)
-	{
-		FCallTime& Time = Owner.Timers[Owner.CurrentTimerIndex];
-
-		Time.PostExecuteTime = ThisFrameTime;
+		Owner.CurrentTimerIndex = (Owner.CurrentTimerIndex + 1) % MaxNumberOfTrackedTimers;
 	}
 }
 #endif // WITH_EDITOR
diff --git a/Engine/Plugins/Experimental/PCG/Source/PCG/Public/PCGContext.h b/Engine/Plugins/Experimental/PCG/Source/PCG/Public/PCGContext.h
--- a/Engine/Plugins/Experimental/PCG/Source/PCG/Public/PCGContext.h
+++ b/Engine/Plugins/Experimental/PCG/Source/PCG/Public/PCGContext.h
@@ -64,6 +64,11 @@ struct PCG_API FPCGContext
 	/** True if currently inside a PCGAsync scope - will prevent further async processing */
 	bool bIsRunningAsyncCall = false;
 
+#if WITH_EDITOR
+	double ElapsedTime = 0.0;
+	int32 ExecutionCount = 0;
+#endif
+
 	template<typename SettingsType>
 	const SettingsType* GetInputSettings() const
 	{
diff --git a/Engine/Plugins/Experimental/PCG/Source/PCG/Public/PCGElement.h b/Engine/Plugins/Experimental/PCG/Source/PCG/Public/PCGElement.h
--- a/Engine/Plugins/Experimental/PCG/Source/PCG/Public/PCGElement.h
+++ b/Engine/Plugins/Experimental/PCG/Source/PCG/Public/PCGElement.h
@@ -51,20 +51,8 @@ public:
 
 	/** Note: the following methods must be called from the main thread */
 #if WITH_EDITOR
-	struct FCallTime
-	{
-		// sum of all frames
-		double ExecutionTime = 0;
-		// how many frames
-		int32 ExecutionFrameCount = 0;
-		double MinExecutionFrameTime = std::numeric_limits<double>::max();
-		double MaxExecutionFrameTime = 0.0;
-		double PrepareDataTime = 0.0;
-		double PostExecuteTime = 0.0;
-	};
-
 	void DebugDisplay(FPCGContext* Context) const;
-	const TArray<FCallTime>& GetTimers() const { return Timers; }
+	const TArray<double>& GetTimers() const { return Timers; }
 	void ResetTimers();
 #endif
 
@@ -99,15 +87,13 @@ private:
 		const IPCGElement& Owner;
 		FPCGContext* Context;
 		double StartTime;
-		EPCGExecutionPhase Phase;
 	};
 
 	// Set mutable because we need to modify them in the execute call, which is const
 	// TODO: Should be a map with PCG Components. We need a mechanism to make sure that this map is cleaned up when component doesn't exist anymore.
 	// For now, it will track all calls to execute (excluding call where the result is already in cache).
-	mutable TArray<FCallTime> Timers;
-	mutable int32 CurrentTimerIndex = 0;
-
+	mutable TArray<double> Timers;
+	mutable int CurrentTimerIndex = 0;
 	// Perhaps overkill but there is a slight chance that we need to protect the timers array. If we call reset from the UI while an element is executing,
 	// it could crash while writing to the timers array.
 	mutable FCriticalSection TimersLock;
@@ -115,7 +101,6 @@ private:
 	struct FScopedCallTimer
 	{
 		FScopedCallTimer(const IPCGElement& InOwner, FPCGContext* InContext) {}
-		void AdvanceTimerIndex() {}
 	};
 #endif // WITH_EDITOR
 };
diff --git a/Engine/Plugins/Experimental/PCG/Source/PCGEditor/Private/SPCGEditorGraphProfilingView.cpp b/Engine/Plugins/Experimental/PCG/Source/PCGEditor/Private/SPCGEditorGraphProfilingView.cpp
--- a/Engine/Plugins/Experimental/PCG/Source/PCGEditor/Private/SPCGEditorGraphProfilingView.cpp
+++ b/Engine/Plugins/Experimental/PCG/Source/PCGEditor/Private/SPCGEditorGraphProfilingView.cpp
@@ -22,33 +22,21 @@ namespace PCGEditorGraphProfilingView
 	
 	/** Names of the columns in the attribute list */
 	const FName NAME_Node = FName(TEXT("Node"));
-	const FName NAME_PrepareDataTime = FName(TEXT("PrepareDataTime"));
-	const FName NAME_AvgExecutionTime = FName(TEXT("AvgExecutionTime"));
-	const FName NAME_MinExecutionTime = FName(TEXT("MinExecutionTime"));
-	const FName NAME_MaxExecutionTime = FName(TEXT("MaxExecutionTime"));
-	const FName NAME_MinExecutionFrameTime = FName(TEXT("MinFrameTime"));
-	const FName NAME_MaxExecutionFrameTime = FName(TEXT("MaxFrameTime"));
-	const FName NAME_MinNbExecutionFrames = FName(TEXT("MinNbExecutionFrames"));
-	const FName NAME_MaxNbExecutionFrames = FName(TEXT("MaxNbExecutionFrames"));
-	const FName NAME_StdExecutionTime = FName(TEXT("StdExecutionTime"));
-	const FName NAME_TotalExecutionTime = FName(TEXT("TotalExecutionTime"));
+	const FName NAME_AvgTime = FName(TEXT("AvgTime"));
+	const FName NAME_MinTime = FName(TEXT("MinTime"));
+	const FName NAME_MaxTime = FName(TEXT("MaxTime"));
+	const FName NAME_StdTime = FName(TEXT("StdTime"));
+	const FName NAME_TotalTime = FName(TEXT("TotalTime"));
 	const FName NAME_NbCalls = FName(TEXT("NbCalls"));
-	const FName NAME_NbExecutionFrames = FName(TEXT("NbExecutionFrames"));
-	const FName NAME_PostExecuteTime = FName(TEXT("PostExecuteTime"));
 
 	/** Labels of the columns */
 	const FText TEXT_NodeLabel = LOCTEXT("NodeLabel", "Node");
-	const FText TEXT_PrepareDataTimeLabel = LOCTEXT("PrepareDataTimeLabel", "PrepareData (ms)");
-	const FText TEXT_PostExecuteTimeLabel = LOCTEXT("PrepareDataTimeLabel", "PrepareData (ms)");
-	const FText TEXT_AvgExecutionTimeLabel = LOCTEXT("AvgExecutionTimeLabel", "Avg Time(ms)");
-	const FText TEXT_MinExecutionTimeLabel = LOCTEXT("MinExecutionTimeLabel", "Min Time(ms)");
-	const FText TEXT_MaxExecutionTimeLabel = LOCTEXT("MaxExecutionTimeLabel", "Max Time(ms)");
-	const FText TEXT_MinExecutionFrameTimeLabel = LOCTEXT("MinExecutionFrameTimeLabel", "Min Frame Time(ms)");
-	const FText TEXT_MaxExecutionFrameTimeLabel = LOCTEXT("MaxExecutionFrameTimeLabel", "Max Frame Time(ms)");
-	const FText TEXT_StdExecutionTimeLabel = LOCTEXT("StdExecutionTimeLabel", "Std(ms)");
-	const FText TEXT_TotalExecutionTimeLabel = LOCTEXT("TotalExecutionTimeLabel", "Total time(s)");
-	const FText TEXT_NbCallsLabel = LOCTEXT("NbCallsLabel", "Calls");
-	const FText TEXT_NbExecutionFramesLabel = LOCTEXT("NbExecutionFramesLabel", "Exec frames");
+	const FText TEXT_AvgTimeLabel = LOCTEXT("AvgTimeLabel", "Avg Time(ms)");
+	const FText TEXT_MinTimeLabel = LOCTEXT("MinTimeLabel", "Min Time(ms)");
+	const FText TEXT_MaxTimeLabel = LOCTEXT("MaxTimeLabel", "Max Time(ms)");
+	const FText TEXT_StdTimeLabel = LOCTEXT("StdTimeLabel", "Std(ms)");
+	const FText TEXT_TotalTimeLabel = LOCTEXT("TotalTimeLabel", "Total time(s)");
+	const FText TEXT_NbCallsLabel = LOCTEXT("NbCallsLabel", "Number of calls");
 }
 
 void SPCGProfilingListViewItemRow::Construct(const FArguments& InArgs, const TSharedRef<STableViewBase>& InOwnerTableView, const PCGProfilingListViewItemPtr& Item)
@@ -74,65 +62,35 @@ TSharedRef<SWidget> SPCGProfilingListViewItemRow::GenerateWidgetForColumn(const
 		{
 			ColumnData = FText::AsNumber(InternalItem->NbCalls);
 		}
-		else if (ColumnId == PCGEditorGraphProfilingView::NAME_NbExecutionFrames)
-		{
-			ColumnData = FText::AsNumber(InternalItem->NbExecutionFrames);
-		}
-		else if (ColumnId == PCGEditorGraphProfilingView::NAME_MinNbExecutionFrames)
-		{
-			ColumnData = FText::AsNumber(InternalItem->MinNbExecutionFrames);
-		}
-		else if (ColumnId == PCGEditorGraphProfilingView::NAME_MaxNbExecutionFrames)
-		{
-			ColumnData = FText::AsNumber(InternalItem->MaxNbExecutionFrames);
-		}
 		// For all other values, if we don't have data, just write "N/A"
 		else if (!InternalItem->HasData)
 		{
 			ColumnData = PCGEditorGraphProfilingView::NoDataAvailableText;
 		}
-		else if (ColumnId == PCGEditorGraphProfilingView::NAME_AvgExecutionTime)
-		{
-			// In ms
-			ColumnData = FText::AsNumber(InternalItem->AvgExecutionTime * 1000.0);
-		}
-		else if (ColumnId == PCGEditorGraphProfilingView::NAME_MinExecutionTime)
-		{
-			// In ms
-			ColumnData = FText::AsNumber(InternalItem->MinExecutionTime * 1000.0);
-		}
-		else if (ColumnId == PCGEditorGraphProfilingView::NAME_MaxExecutionTime)
+		else if (ColumnId == PCGEditorGraphProfilingView::NAME_AvgTime)
 		{
 			// In ms
-			ColumnData = FText::AsNumber(InternalItem->MaxExecutionTime * 1000.0);
+			ColumnData = FText::AsNumber(InternalItem->AvgTime * 1000.0);
 		}
-		else if (ColumnId == PCGEditorGraphProfilingView::NAME_MinExecutionFrameTime)
+		else if (ColumnId == PCGEditorGraphProfilingView::NAME_MinTime)
 		{
 			// In ms
-			ColumnData = FText::AsNumber(InternalItem->MinExecutionFrameTime * 1000.0);
+			ColumnData = FText::AsNumber(InternalItem->MinTime * 1000.0);
 		}
-		else if (ColumnId == PCGEditorGraphProfilingView::NAME_MaxExecutionFrameTime)
+		else if (ColumnId == PCGEditorGraphProfilingView::NAME_MaxTime)
 		{
 			// In ms
-			ColumnData = FText::AsNumber(InternalItem->MaxExecutionFrameTime * 1000.0);
+			ColumnData = FText::AsNumber(InternalItem->MaxTime * 1000.0);
 		}
-		else if (ColumnId == PCGEditorGraphProfilingView::NAME_StdExecutionTime)
+		else if (ColumnId == PCGEditorGraphProfilingView::NAME_StdTime)
 		{
 			// In ms
-			ColumnData = FText::AsNumber(InternalItem->StdExecutionTime * 1000.0);
+			ColumnData = FText::AsNumber(InternalItem->StdTime * 1000.0);
 		}
-		else if (ColumnId == PCGEditorGraphProfilingView::NAME_TotalExecutionTime)
+		else if (ColumnId == PCGEditorGraphProfilingView::NAME_TotalTime)
 		{
 			// In s
-			ColumnData = FText::AsNumber(InternalItem->ExecutionTime);
-		}
-		else if (ColumnId == PCGEditorGraphProfilingView::NAME_PrepareDataTime)
-		{
-			ColumnData = FText::AsNumber(InternalItem->PrepareDataTime * 1000.0);
-		}
-		else if (ColumnId == PCGEditorGraphProfilingView::NAME_PostExecuteTime)
-		{
-			ColumnData = FText::AsNumber(InternalItem->PostExecuteTime * 1000.0);
+			ColumnData = FText::AsNumber(InternalItem->TotalTime);
 		}
 	}
 
@@ -253,75 +211,47 @@ TSharedRef<SHeaderRow> SPCGEditorGraphProfilingView::CreateHeaderRowWidget()
 		.HAlignCell(HAlign_Left)
 		.SortMode(this, &SPCGEditorGraphProfilingView::GetColumnSortMode, PCGEditorGraphProfilingView::NAME_Node)
 		.OnSort(this, &SPCGEditorGraphProfilingView::OnSortColumnHeader)
-		+ SHeaderRow::Column(PCGEditorGraphProfilingView::NAME_NbCalls)
-		.ManualWidth(80)
-		.DefaultLabel(PCGEditorGraphProfilingView::TEXT_NbCallsLabel)
-		.HAlignHeader(HAlign_Center)
-		.HAlignCell(HAlign_Right)
-		.SortMode(this, &SPCGEditorGraphProfilingView::GetColumnSortMode, PCGEditorGraphProfilingView::NAME_NbCalls)
-		.OnSort(this, &SPCGEditorGraphProfilingView::OnSortColumnHeader)
-		+ SHeaderRow::Column(PCGEditorGraphProfilingView::NAME_PrepareDataTime)
-		.ManualWidth(125)
-		.DefaultLabel(PCGEditorGraphProfilingView::TEXT_PrepareDataTimeLabel)
-		.HAlignHeader(HAlign_Center)
-		.HAlignCell(HAlign_Right)
-		.SortMode(this, &SPCGEditorGraphProfilingView::GetColumnSortMode, PCGEditorGraphProfilingView::NAME_PrepareDataTime)
-		.OnSort(this, &SPCGEditorGraphProfilingView::OnSortColumnHeader)
-		+ SHeaderRow::Column(PCGEditorGraphProfilingView::NAME_NbExecutionFrames)
-		.ManualWidth(80)
-		.DefaultLabel(PCGEditorGraphProfilingView::TEXT_NbExecutionFramesLabel)
-		.HAlignHeader(HAlign_Center)
-		.HAlignCell(HAlign_Right)
-		.SortMode(this, &SPCGEditorGraphProfilingView::GetColumnSortMode, PCGEditorGraphProfilingView::NAME_NbExecutionFrames)
-		.OnSort(this, &SPCGEditorGraphProfilingView::OnSortColumnHeader)
-		+ SHeaderRow::Column(PCGEditorGraphProfilingView::NAME_MinExecutionFrameTime)
-		.ManualWidth(130)
-		.DefaultLabel(PCGEditorGraphProfilingView::TEXT_MinExecutionFrameTimeLabel)
-		.HAlignHeader(HAlign_Center)
-		.HAlignCell(HAlign_Right)
-		.SortMode(this, &SPCGEditorGraphProfilingView::GetColumnSortMode, PCGEditorGraphProfilingView::NAME_MinExecutionFrameTime)
-		.OnSort(this, &SPCGEditorGraphProfilingView::OnSortColumnHeader)
-		+ SHeaderRow::Column(PCGEditorGraphProfilingView::NAME_MaxExecutionFrameTime)
-		.ManualWidth(130)
-		.DefaultLabel(PCGEditorGraphProfilingView::TEXT_MaxExecutionFrameTimeLabel)
+		+ SHeaderRow::Column(PCGEditorGraphProfilingView::NAME_AvgTime)
+		.ManualWidth(100)
+		.DefaultLabel(PCGEditorGraphProfilingView::TEXT_AvgTimeLabel)
 		.HAlignHeader(HAlign_Center)
 		.HAlignCell(HAlign_Right)
-		.SortMode(this, &SPCGEditorGraphProfilingView::GetColumnSortMode, PCGEditorGraphProfilingView::NAME_MaxExecutionFrameTime)
+		.SortMode(this, &SPCGEditorGraphProfilingView::GetColumnSortMode, PCGEditorGraphProfilingView::NAME_AvgTime)
 		.OnSort(this, &SPCGEditorGraphProfilingView::OnSortColumnHeader)
-		+ SHeaderRow::Column(PCGEditorGraphProfilingView::NAME_AvgExecutionTime)
+		+ SHeaderRow::Column(PCGEditorGraphProfilingView::NAME_MinTime)
 		.ManualWidth(100)
-		.DefaultLabel(PCGEditorGraphProfilingView::TEXT_AvgExecutionTimeLabel)
+		.DefaultLabel(PCGEditorGraphProfilingView::TEXT_MinTimeLabel)
 		.HAlignHeader(HAlign_Center)
 		.HAlignCell(HAlign_Right)
-		.SortMode(this, &SPCGEditorGraphProfilingView::GetColumnSortMode, PCGEditorGraphProfilingView::NAME_AvgExecutionTime)
+		.SortMode(this, &SPCGEditorGraphProfilingView::GetColumnSortMode, PCGEditorGraphProfilingView::NAME_MinTime)
 		.OnSort(this, &SPCGEditorGraphProfilingView::OnSortColumnHeader)
-		+ SHeaderRow::Column(PCGEditorGraphProfilingView::NAME_MinExecutionTime)
+		+ SHeaderRow::Column(PCGEditorGraphProfilingView::NAME_MaxTime)
 		.ManualWidth(100)
-		.DefaultLabel(PCGEditorGraphProfilingView::TEXT_MinExecutionTimeLabel)
+		.DefaultLabel(PCGEditorGraphProfilingView::TEXT_MaxTimeLabel)
 		.HAlignHeader(HAlign_Center)
 		.HAlignCell(HAlign_Right)
-		.SortMode(this, &SPCGEditorGraphProfilingView::GetColumnSortMode, PCGEditorGraphProfilingView::NAME_MinExecutionTime)
+		.SortMode(this, &SPCGEditorGraphProfilingView::GetColumnSortMode, PCGEditorGraphProfilingView::NAME_MaxTime)
 		.OnSort(this, &SPCGEditorGraphProfilingView::OnSortColumnHeader)
-		+ SHeaderRow::Column(PCGEditorGraphProfilingView::NAME_MaxExecutionTime)
+		+ SHeaderRow::Column(PCGEditorGraphProfilingView::NAME_StdTime)
 		.ManualWidth(100)
-		.DefaultLabel(PCGEditorGraphProfilingView::TEXT_MaxExecutionTimeLabel)
+		.DefaultLabel(PCGEditorGraphProfilingView::TEXT_StdTimeLabel)
 		.HAlignHeader(HAlign_Center)
 		.HAlignCell(HAlign_Right)
-		.SortMode(this, &SPCGEditorGraphProfilingView::GetColumnSortMode, PCGEditorGraphProfilingView::NAME_MaxExecutionTime)
+		.SortMode(this, &SPCGEditorGraphProfilingView::GetColumnSortMode, PCGEditorGraphProfilingView::NAME_StdTime)
 		.OnSort(this, &SPCGEditorGraphProfilingView::OnSortColumnHeader)
-		+ SHeaderRow::Column(PCGEditorGraphProfilingView::NAME_StdExecutionTime)
+		+ SHeaderRow::Column(PCGEditorGraphProfilingView::NAME_TotalTime)
 		.ManualWidth(100)
-		.DefaultLabel(PCGEditorGraphProfilingView::TEXT_StdExecutionTimeLabel)
+		.DefaultLabel(PCGEditorGraphProfilingView::TEXT_TotalTimeLabel)
 		.HAlignHeader(HAlign_Center)
 		.HAlignCell(HAlign_Right)
-		.SortMode(this, &SPCGEditorGraphProfilingView::GetColumnSortMode, PCGEditorGraphProfilingView::NAME_StdExecutionTime)
+		.SortMode(this, &SPCGEditorGraphProfilingView::GetColumnSortMode, PCGEditorGraphProfilingView::NAME_TotalTime)
 		.OnSort(this, &SPCGEditorGraphProfilingView::OnSortColumnHeader)
-		+ SHeaderRow::Column(PCGEditorGraphProfilingView::NAME_TotalExecutionTime)
-		.ManualWidth(100)
-		.DefaultLabel(PCGEditorGraphProfilingView::TEXT_TotalExecutionTimeLabel)
+		+ SHeaderRow::Column(PCGEditorGraphProfilingView::NAME_NbCalls)
+		.ManualWidth(125)
+		.DefaultLabel(PCGEditorGraphProfilingView::TEXT_NbCallsLabel)
 		.HAlignHeader(HAlign_Center)
 		.HAlignCell(HAlign_Right)
-		.SortMode(this, &SPCGEditorGraphProfilingView::GetColumnSortMode, PCGEditorGraphProfilingView::NAME_TotalExecutionTime)
+		.SortMode(this, &SPCGEditorGraphProfilingView::GetColumnSortMode, PCGEditorGraphProfilingView::NAME_NbCalls)
 		.OnSort(this, &SPCGEditorGraphProfilingView::OnSortColumnHeader);
 }
 
@@ -425,43 +355,37 @@ FReply SPCGEditorGraphProfilingView::Refresh()
 				{
 					if (const IPCGElement* Element = Settings->GetElement().Get())
 					{
-						const TArray<IPCGElement::FCallTime>& Timers = Element->GetTimers();
-						if (Timers.IsEmpty())
-						{
-							continue;
-						}
-
-						ListViewItem->MinExecutionTime = std::numeric_limits<double>::max();
-						ListViewItem->MaxExecutionTime = std::numeric_limits<double>::min();
-						ListViewItem->MinExecutionFrameTime = std::numeric_limits<double>::max();
-						ListViewItem->MaxExecutionFrameTime = std::numeric_limits<double>::min();
-
-						for(const IPCGElement::FCallTime& Timer : Timers)
+						const TArray<double>& Timers = Element->GetTimers();
+						if (!Timers.IsEmpty())
 						{
-							ListViewItem->PrepareDataTime += Timer.PrepareDataTime;
-
-							ListViewItem->MinExecutionFrameTime = FMath::Min(ListViewItem->MinExecutionFrameTime, Timer.MinExecutionFrameTime);
-							ListViewItem->MaxExecutionFrameTime = FMath::Max(ListViewItem->MaxExecutionFrameTime, Timer.MaxExecutionFrameTime);
-
-							ListViewItem->ExecutionTime += Timer.ExecutionTime;
-							ListViewItem->MinExecutionTime = FMath::Min(ListViewItem->MinExecutionTime, Timer.ExecutionTime);
-							ListViewItem->MaxExecutionTime = FMath::Max(ListViewItem->MaxExecutionTime, Timer.ExecutionTime);
-
-							ListViewItem->NbExecutionFrames += Timer.ExecutionFrameCount;
-							ListViewItem->MaxNbExecutionFrames = FMath::Max(ListViewItem->MaxNbExecutionFrames, Timer.ExecutionFrameCount);
-							ListViewItem->MinNbExecutionFrames = FMath::Min(ListViewItem->MinNbExecutionFrames, Timer.ExecutionFrameCount);
+							ListViewItem->MinTime = std::numeric_limits<double>::max();
+							ListViewItem->MaxTime = std::numeric_limits<double>::min();
+							for (double v : Timers)
+							{
+								ListViewItem->TotalTime += v;
+								if (v < ListViewItem->MinTime)
+								{
+									ListViewItem->MinTime = v;
+								}
+
+								if (v > ListViewItem->MaxTime)
+								{
+									ListViewItem->MaxTime = v;
+								}
+							}
+
+							ListViewItem->AvgTime = ListViewItem->TotalTime / Timers.Num();
+
+							for (double v : Timers)
+							{
+								ListViewItem->StdTime += (ListViewItem->AvgTime - v) * (ListViewItem->AvgTime - v);
+							}
+
+							ListViewItem->StdTime = FMath::Sqrt(ListViewItem->StdTime / Timers.Num());
+
+							ListViewItem->HasData = true;
+							ListViewItem->NbCalls = Timers.Num();
 						}
-
-						ListViewItem->NbCalls = Timers.Num();
-						ListViewItem->AvgExecutionTime = ListViewItem->ExecutionTime / ListViewItem->NbCalls;
-
-						for (const IPCGElement::FCallTime& Timer : Timers)
-						{
-							ListViewItem->StdExecutionTime += FMath::Square(ListViewItem->AvgExecutionTime - Timer.ExecutionTime);
-						}
-
-						ListViewItem->StdExecutionTime = FMath::Sqrt(ListViewItem->StdExecutionTime / ListViewItem->NbCalls);
-						ListViewItem->HasData = true;
 					}
 				}
 			}
@@ -479,54 +403,30 @@ FReply SPCGEditorGraphProfilingView::Refresh()
 				{
 					isLess = A->Name.FastLess(B->Name);
 				}
-				else if (SortingColumn == PCGEditorGraphProfilingView::NAME_PrepareDataTime)
-				{
-					isLess = A->PrepareDataTime < B->PrepareDataTime;
-				}
-				else if (SortingColumn == PCGEditorGraphProfilingView::NAME_AvgExecutionTime)
-				{
-					isLess = A->AvgExecutionTime < B->AvgExecutionTime;
-				}
-				else if (SortingColumn == PCGEditorGraphProfilingView::NAME_MinExecutionTime)
+				else if (SortingColumn == PCGEditorGraphProfilingView::NAME_AvgTime)
 				{
-					isLess = A->MinExecutionTime < B->MinExecutionTime;
+					isLess = A->AvgTime < B->AvgTime;
 				}
-				else if (SortingColumn == PCGEditorGraphProfilingView::NAME_MaxExecutionTime)
+				else if (SortingColumn == PCGEditorGraphProfilingView::NAME_MinTime)
 				{
-					isLess = A->MaxExecutionTime < B->MaxExecutionTime;
+					isLess = A->MinTime < B->MinTime;
 				}
-				else if (SortingColumn == PCGEditorGraphProfilingView::NAME_MinExecutionFrameTime)
+				else if (SortingColumn == PCGEditorGraphProfilingView::NAME_MaxTime)
 				{
-					isLess = A->MinExecutionFrameTime < B->MinExecutionFrameTime;
+					isLess = A->MaxTime < B->MaxTime;
 				}
-				else if (SortingColumn == PCGEditorGraphProfilingView::NAME_MaxExecutionFrameTime)
+				else if (SortingColumn == PCGEditorGraphProfilingView::NAME_StdTime)
 				{
-					isLess = A->MaxExecutionFrameTime < B->MaxExecutionFrameTime;
+					isLess = A->StdTime < B->StdTime;
 				}
-				else if (SortingColumn == PCGEditorGraphProfilingView::NAME_StdExecutionTime)
+				else if (SortingColumn == PCGEditorGraphProfilingView::NAME_TotalTime)
 				{
-					isLess = A->StdExecutionTime < B->StdExecutionTime;
-				}
-				else if (SortingColumn == PCGEditorGraphProfilingView::NAME_TotalExecutionTime)
-				{
-					isLess = A->ExecutionTime < B->ExecutionTime;
+					isLess = A->TotalTime < B->TotalTime;
 				}
 				else if (SortingColumn == PCGEditorGraphProfilingView::NAME_NbCalls)
 				{
 					isLess = A->NbCalls < B->NbCalls;
 				}
-				else if (SortingColumn == PCGEditorGraphProfilingView::NAME_NbExecutionFrames)
-				{
-					isLess = A->NbExecutionFrames < B->NbExecutionFrames;
-				}
-				else if (SortingColumn == PCGEditorGraphProfilingView::NAME_MaxNbExecutionFrames)
-				{
-					isLess = A->MaxNbExecutionFrames < B->MaxNbExecutionFrames;
-				}
-				else if (SortingColumn == PCGEditorGraphProfilingView::NAME_MinNbExecutionFrames)
-				{
-					isLess = A->MinNbExecutionFrames < B->MinNbExecutionFrames;
-				}
 
 				return SortMode == EColumnSortMode::Ascending ? isLess : !isLess;
 			});
diff --git a/Engine/Plugins/Experimental/PCG/Source/PCGEditor/Private/SPCGEditorGraphProfilingView.h b/Engine/Plugins/Experimental/PCG/Source/PCGEditor/Private/SPCGEditorGraphProfilingView.h
--- a/Engine/Plugins/Experimental/PCG/Source/PCGEditor/Private/SPCGEditorGraphProfilingView.h
+++ b/Engine/Plugins/Experimental/PCG/Source/PCGEditor/Private/SPCGEditorGraphProfilingView.h
@@ -21,19 +21,12 @@ struct FPCGProfilingListViewItem
 	const UPCGEditorGraphNode* EditorNode = nullptr;
 
 	FName Name = NAME_None;
-	double PrepareDataTime = 0.0;
-	double ExecutionTime = 0.0;
-	double AvgExecutionTime = 0.0;
-	double MinExecutionTime = 0.0;
-	double MaxExecutionTime = 0.0;
-	double MinExecutionFrameTime = 0.0;
-	double MaxExecutionFrameTime = 0.0;
-	double StdExecutionTime = 0.0;
-	double PostExecuteTime = 0.0;
+	double AvgTime = 0.0;
+	double MinTime = 0.0;
+	double MaxTime = 0.0;
+	double StdTime = 0.0;
+	double TotalTime = 0.0;
 	int32 NbCalls = 0;
-	int32 NbExecutionFrames = 0;
-	int32 MaxNbExecutionFrames = 0;
-	int32 MinNbExecutionFrames = 0;
 	bool HasData = false;
 };
 
