[PCG] Added missing collision profile include.

diff --git a/Engine/Plugins/Experimental/PCG/Source/PCG/Public/Helpers/PCGActorHelpers.h b/Engine/Plugins/Experimental/PCG/Source/PCG/Public/Helpers/PCGActorHelpers.h
--- a/Engine/Plugins/Experimental/PCG/Source/PCG/Public/Helpers/PCGActorHelpers.h
+++ b/Engine/Plugins/Experimental/PCG/Source/PCG/Public/Helpers/PCGActorHelpers.h
@@ -6,6 +6,7 @@
 #include "Kismet/BlueprintFunctionLibrary.h"
 #include "Templates/SubclassOf.h"
 #include "Engine/EngineTypes.h"
+#include "Engine/CollisionProfile.h"
 
 #include <type_traits>
 
