Added settings instances support. More UI changes to come in later CLs.

diff --git a/Engine/Plugins/Experimental/PCG/Source/PCG/Private/Elements/PCGDebugElement.cpp b/Engine/Plugins/Experimental/PCG/Source/PCG/Private/Elements/PCGDebugElement.cpp
--- a/Engine/Plugins/Experimental/PCG/Source/PCG/Private/Elements/PCGDebugElement.cpp
+++ b/Engine/Plugins/Experimental/PCG/Source/PCG/Private/Elements/PCGDebugElement.cpp
@@ -23,14 +23,14 @@ namespace PCGDebugElement
 			return;
 		}
 
-		const UPCGSettings* Settings = Context->GetInputSettings<UPCGSettings>();
+		const UPCGSettingsInterface* SettingsInterface = Context->GetInputSettingsInterface();
 
-		if (!Settings)
+		if (!SettingsInterface)
 		{
 			return;
 		}
 
-		const FPCGDebugVisualizationSettings& DebugSettings = Settings->DebugSettings;
+		const FPCGDebugVisualizationSettings& DebugSettings = SettingsInterface->DebugSettings;
 
 		UStaticMesh* Mesh = DebugSettings.PointMesh.LoadSynchronous();
 
diff --git a/Engine/Plugins/Experimental/PCG/Source/PCG/Private/Elements/PCGSpawnActor.cpp b/Engine/Plugins/Experimental/PCG/Source/PCG/Private/Elements/PCGSpawnActor.cpp
--- a/Engine/Plugins/Experimental/PCG/Source/PCG/Private/Elements/PCGSpawnActor.cpp
+++ b/Engine/Plugins/Experimental/PCG/Source/PCG/Private/Elements/PCGSpawnActor.cpp
@@ -67,7 +67,7 @@ bool UPCGSpawnActorSettings::IsStructuralProperty(const FName& InPropertyName) c
 
 TObjectPtr<UPCGGraph> UPCGSpawnActorNode::GetSubgraph() const
 {
-	TObjectPtr<UPCGSpawnActorSettings> Settings = Cast<UPCGSpawnActorSettings>(DefaultSettings);
+	TObjectPtr<UPCGSpawnActorSettings> Settings = Cast<UPCGSpawnActorSettings>(GetSettings());
 	return (Settings && Settings->Option != EPCGSpawnActorOption::NoMerging) ? Settings->GetSubgraph() : nullptr;
 }
 
diff --git a/Engine/Plugins/Experimental/PCG/Source/PCG/Private/Graph/PCGGraphExecutor.cpp b/Engine/Plugins/Experimental/PCG/Source/PCG/Private/Graph/PCGGraphExecutor.cpp
--- a/Engine/Plugins/Experimental/PCG/Source/PCG/Private/Graph/PCGGraphExecutor.cpp
+++ b/Engine/Plugins/Experimental/PCG/Source/PCG/Private/Graph/PCGGraphExecutor.cpp
@@ -430,7 +430,7 @@ void FPCGGraphExecutor::Execute()
 				{
 					// Get appropriate settings
 					check(Task.Node);
-					const UPCGSettings* Settings = TaskInput.GetSettings(Task.Node->DefaultSettings);
+					const UPCGSettings* Settings = TaskInput.GetSettings(Task.Node->GetSettings());
 
 					if (Settings)
 					{
@@ -448,11 +448,12 @@ void FPCGGraphExecutor::Execute()
 
 				// If a task is cacheable and has been cached, then we don't need to create an active task for it unless
 				// there is an execution mode that would prevent us from doing so.
-				const UPCGSettings* TaskSettings = PCGContextHelpers::GetInputSettings<UPCGSettings>(Task.Node, TaskInput);
+				const UPCGSettingsInterface* TaskSettingsInterface = TaskInput.GetSettingsInterface(Task.Node ? Task.Node->GetSettingsInterface() : nullptr);
+				const UPCGSettings* TaskSettings = TaskSettingsInterface ? TaskSettingsInterface->GetSettings() : nullptr;
 				FPCGDataCollection CachedOutput;
-				const bool bResultAlreadyInCache = Task.Element->IsCacheable(TaskSettings) && GraphCache.GetFromCache(Task.Element.Get(), TaskInput, TaskSettings, Task.SourceComponent.Get(), CachedOutput);
+				const bool bResultAlreadyInCache = Task.Element->IsCacheableInstance(TaskSettingsInterface) && GraphCache.GetFromCache(Task.Element.Get(), TaskInput, TaskSettings, Task.SourceComponent.Get(), CachedOutput);
 #if WITH_EDITOR
-				const bool bNeedsToCreateActiveTask = !bResultAlreadyInCache || TaskSettings->ExecutionMode == EPCGSettingsExecutionMode::Debug || TaskSettings->ExecutionMode == EPCGSettingsExecutionMode::Isolated;
+				const bool bNeedsToCreateActiveTask = !bResultAlreadyInCache || TaskSettingsInterface->ExecutionMode == EPCGSettingsExecutionMode::Debug || TaskSettingsInterface->ExecutionMode == EPCGSettingsExecutionMode::Isolated;
 #else
 				const bool bNeedsToCreateActiveTask = !bResultAlreadyInCache;
 #endif
@@ -577,9 +578,10 @@ void FPCGGraphExecutor::Execute()
 #endif
 			{
 				// Store result in cache as needed - done here because it needs to be done on the main thread
-				const UPCGSettings* ActiveTaskSettings = ActiveTask.Context->GetInputSettings<UPCGSettings>();
-				if (ActiveTaskSettings && ActiveTask.Element->IsCacheable(ActiveTaskSettings))
+				const UPCGSettingsInterface* ActiveTaskSettingsInterface = ActiveTask.Context->GetInputSettingsInterface();
+				if (ActiveTaskSettingsInterface && ActiveTask.Element->IsCacheableInstance(ActiveTaskSettingsInterface))
 				{
+					const UPCGSettings* ActiveTaskSettings = ActiveTaskSettingsInterface ? ActiveTaskSettingsInterface->GetSettings() : nullptr;
 					GraphCache.StoreInCache(ActiveTask.Element.Get(), ActiveTask.Context->InputData, ActiveTaskSettings, ActiveTask.Context->SourceComponent.Get(), ActiveTask.Context->OutputData);
 				}
 			}
diff --git a/Engine/Plugins/Experimental/PCG/Source/PCG/Private/PCGContext.cpp b/Engine/Plugins/Experimental/PCG/Source/PCG/Private/PCGContext.cpp
--- a/Engine/Plugins/Experimental/PCG/Source/PCG/Private/PCGContext.cpp
+++ b/Engine/Plugins/Experimental/PCG/Source/PCG/Private/PCGContext.cpp
@@ -47,4 +47,16 @@ bool FPCGContext::IsOutputConnectedOrInspecting(FName PinLabel) const
 #else
 	return bOutputConnected;
 #endif // WITH_EDITOR
+}
+
+const UPCGSettingsInterface* FPCGContext::GetInputSettingsInterface() const
+{
+	if (Node)
+	{
+		return InputData.GetSettingsInterface(Node->GetSettingsInterface());
+	}
+	else
+	{
+		return InputData.GetSettingsInterface();
+	}
 }
\ No newline at end of file
diff --git a/Engine/Plugins/Experimental/PCG/Source/PCG/Private/PCGData.cpp b/Engine/Plugins/Experimental/PCG/Source/PCG/Private/PCGData.cpp
--- a/Engine/Plugins/Experimental/PCG/Source/PCG/Private/PCGData.cpp
+++ b/Engine/Plugins/Experimental/PCG/Source/PCG/Private/PCGData.cpp
@@ -133,7 +133,7 @@ const UPCGSettings* FPCGDataCollection::GetSettings(const UPCGSettings* InDefaul
 	else
 	{
 		const FPCGTaggedData* MatchingData = TaggedData.FindByPredicate([InDefaultSettings](const FPCGTaggedData& Data) {
-			return Data.Data && 
+			return Data.Data &&
 				(Data.Data->GetClass() == InDefaultSettings->GetClass() ||
 					Data.Data->GetClass()->IsChildOf(InDefaultSettings->GetClass()));
 			});
@@ -142,6 +142,34 @@ const UPCGSettings* FPCGDataCollection::GetSettings(const UPCGSettings* InDefaul
 	}
 }
 
+const UPCGSettingsInterface* FPCGDataCollection::GetSettingsInterface() const
+{
+	return GetSettings<UPCGSettingsInterface>();
+}
+
+const UPCGSettingsInterface* FPCGDataCollection::GetSettingsInterface(const UPCGSettingsInterface* InDefaultSettingsInterface) const
+{
+	if (!InDefaultSettingsInterface || InDefaultSettingsInterface->GetSettings() == nullptr)
+	{
+		return GetSettingsInterface();
+	}
+	else
+	{
+		const FPCGTaggedData* MatchingData = TaggedData.FindByPredicate([InDefaultSettingsInterface](const FPCGTaggedData& Data) {
+			if (UPCGSettingsInterface* DataSettingsInterface = Cast<UPCGSettingsInterface>(Data.Data))
+			{
+				// Compare settings classes
+				return DataSettingsInterface->GetSettings()->GetClass() == InDefaultSettingsInterface->GetSettings()->GetClass() ||
+					DataSettingsInterface->GetSettings()->GetClass()->IsChildOf(InDefaultSettingsInterface->GetSettings()->GetClass());
+			}
+
+			return false;
+		});
+
+		return MatchingData ? Cast<const UPCGSettingsInterface>(MatchingData->Data) : InDefaultSettingsInterface;
+	}
+}
+
 bool FPCGDataCollection::operator==(const FPCGDataCollection& Other) const
 {
 	if (bCancelExecution != Other.bCancelExecution || TaggedData.Num() != Other.TaggedData.Num())
diff --git a/Engine/Plugins/Experimental/PCG/Source/PCG/Private/PCGElement.cpp b/Engine/Plugins/Experimental/PCG/Source/PCG/Private/PCGElement.cpp
--- a/Engine/Plugins/Experimental/PCG/Source/PCG/Private/PCGElement.cpp
+++ b/Engine/Plugins/Experimental/PCG/Source/PCG/Private/PCGElement.cpp
@@ -98,13 +98,15 @@ void IPCGElement::PreExecute(FPCGContext* Context) const
 	// Prepare to move to prepare data phase
 	Context->CurrentPhase = EPCGExecutionPhase::PrepareData;
 
-	const UPCGSettings* Settings = Context->GetInputSettings<UPCGSettings>();
-	if (!Settings)
+	const UPCGSettingsInterface* SettingsInterface = Context->GetInputSettingsInterface();
+	const UPCGSettings* Settings = SettingsInterface ? SettingsInterface->GetSettings() : nullptr;
+
+	if (!SettingsInterface || !Settings)
 	{
 		return;
 	}
 
-	if (Settings->ExecutionMode == EPCGSettingsExecutionMode::Disabled)
+	if (SettingsInterface->ExecutionMode == EPCGSettingsExecutionMode::Disabled)
 	{
 		//Pass-through - no execution
 		Context->OutputData = Context->InputData;
@@ -154,7 +156,8 @@ void IPCGElement::PostExecute(FPCGContext* Context) const
 	PCGE_LOG(Log, "Executed in (%f)s and (%d) call(s)", Context->ElapsedTime, Context->ExecutionCount);
 #endif
 
-	const UPCGSettings* Settings = Context->GetInputSettings<UPCGSettings>();
+	const UPCGSettingsInterface* SettingsInterface = Context->GetInputSettingsInterface();
+	const UPCGSettings* Settings = SettingsInterface ? SettingsInterface->GetSettings() : nullptr;
 
 	// Apply tags on output
 	/** TODO - Placeholder feature */
@@ -168,7 +171,7 @@ void IPCGElement::PostExecute(FPCGContext* Context) const
 
 	// Additional debug things (check for duplicates),
 #if WITH_EDITOR
-	if (Settings && Settings->DebugSettings.bCheckForDuplicates)
+	if (SettingsInterface && SettingsInterface->DebugSettings.bCheckForDuplicates)
 	{
 		FPCGDataCollection ElementInputs = Context->InputData;
 		FPCGDataCollection ElementOutputs = Context->OutputData;
@@ -190,8 +193,8 @@ void IPCGElement::PostExecute(FPCGContext* Context) const
 #if WITH_EDITOR
 void IPCGElement::DebugDisplay(FPCGContext* Context) const
 {
-	const UPCGSettings* Settings = Context->GetInputSettings<UPCGSettings>();
-	if (Settings && (Settings->ExecutionMode == EPCGSettingsExecutionMode::Debug || Settings->ExecutionMode == EPCGSettingsExecutionMode::Isolated))
+	const UPCGSettingsInterface* SettingsInterface = Context->GetInputSettingsInterface();
+	if (SettingsInterface && (SettingsInterface->ExecutionMode == EPCGSettingsExecutionMode::Debug || SettingsInterface->ExecutionMode == EPCGSettingsExecutionMode::Isolated))
 	{
 		FPCGDataCollection ElementInputs = Context->InputData;
 		FPCGDataCollection ElementOutputs = Context->OutputData;
@@ -205,7 +208,7 @@ void IPCGElement::DebugDisplay(FPCGContext* Context) const
 		Context->OutputData = ElementOutputs;
 
 		// Null out the output if this node is executed in isolation
-		if (Settings->ExecutionMode == EPCGSettingsExecutionMode::Isolated)
+		if (SettingsInterface->ExecutionMode == EPCGSettingsExecutionMode::Isolated)
 		{
 			Context->OutputData.bCancelExecution = true;
 		}
@@ -224,7 +227,8 @@ void IPCGElement::ResetTimers()
 void IPCGElement::CleanupAndValidateOutput(FPCGContext* Context) const
 {
 	check(Context);
-	const UPCGSettings* Settings = Context->GetInputSettings<UPCGSettings>();
+	const UPCGSettingsInterface* SettingsInterface = Context->GetInputSettingsInterface();
+	const UPCGSettings* Settings = SettingsInterface ? SettingsInterface->GetSettings() : nullptr;
 
 	if (!IsPassthrough() && Settings)
 	{
@@ -241,7 +245,7 @@ void IPCGElement::CleanupAndValidateOutput(FPCGContext* Context) const
 
 		// Validate all out data for errors in labels
 #if WITH_EDITOR
-		if (Settings->ExecutionMode != EPCGSettingsExecutionMode::Disabled)
+		if (SettingsInterface->ExecutionMode != EPCGSettingsExecutionMode::Disabled)
 		{
 			for (FPCGTaggedData& TaggedData : Context->OutputData.TaggedData)
 			{
@@ -261,6 +265,25 @@ void IPCGElement::CleanupAndValidateOutput(FPCGContext* Context) const
 	}
 }
 
+bool IPCGElement::IsCacheableInstance(const UPCGSettingsInterface* InSettingsInterface) const
+{
+	if (InSettingsInterface)
+	{
+		if (InSettingsInterface->ExecutionMode == EPCGSettingsExecutionMode::Disabled)
+		{
+			return false;
+		}
+		else
+		{
+			return IsCacheable(InSettingsInterface->GetSettings());
+		}
+	}
+	else
+	{
+		return false;
+	}
+}
+
 #if WITH_EDITOR
 IPCGElement::FScopedCallTimer::FScopedCallTimer(const IPCGElement& InOwner, FPCGContext* InContext)
 	: Owner(InOwner), Context(InContext)
diff --git a/Engine/Plugins/Experimental/PCG/Source/PCG/Private/PCGGraph.cpp b/Engine/Plugins/Experimental/PCG/Source/PCG/Private/PCGGraph.cpp
--- a/Engine/Plugins/Experimental/PCG/Source/PCG/Private/PCGGraph.cpp
+++ b/Engine/Plugins/Experimental/PCG/Source/PCG/Private/PCGGraph.cpp
@@ -28,7 +28,7 @@ UPCGGraph::UPCGGraph(const FObjectInitializer& ObjectInitializer)
 
 	UPCGGraphInputOutputSettings* InputSettings = ObjectInitializer.CreateDefaultSubobject<UPCGGraphInputOutputSettings>(this, TEXT("DefaultInputNodeSettings"));
 	InputSettings->SetInput(true);
-	InputNode->SetDefaultSettings(InputSettings, /*bUpdatePins=*/false);
+	InputNode->SetSettingsInterface(InputSettings, /*bUpdatePins=*/false);
 	InputNode->UpdatePins(PinAllocator);
 	
 	OutputNode = ObjectInitializer.CreateDefaultSubobject<UPCGNode>(this, TEXT("DefaultOutputNode"));
@@ -36,7 +36,7 @@ UPCGGraph::UPCGGraph(const FObjectInitializer& ObjectInitializer)
 
 	UPCGGraphInputOutputSettings* OutputSettings = ObjectInitializer.CreateDefaultSubobject<UPCGGraphInputOutputSettings>(this, TEXT("DefaultOutputNodeSettings"));
 	OutputSettings->SetInput(false);
-	OutputNode->SetDefaultSettings(OutputSettings, /*bUpdatePins=*/false);
+	OutputNode->SetSettingsInterface(OutputSettings, /*bUpdatePins=*/false);
 	OutputNode->UpdatePins(PinAllocator);
 	
 #if WITH_EDITORONLY_DATA
@@ -64,21 +64,21 @@ void UPCGGraph::PostLoad()
 	// Deprecation
 	InputNode->ConditionalPostLoad();
 
-	if (!Cast<UPCGGraphInputOutputSettings>(InputNode->DefaultSettings))
+	if (!Cast<UPCGGraphInputOutputSettings>(InputNode->GetSettings()))
 	{
-		InputNode->DefaultSettings = NewObject<UPCGGraphInputOutputSettings>(this, TEXT("DefaultInputNodeSettings"));
+		InputNode->SetSettingsInterface(NewObject<UPCGGraphInputOutputSettings>(this, TEXT("DefaultInputNodeSettings")));
 	}
 
-	Cast<UPCGGraphInputOutputSettings>(InputNode->DefaultSettings)->SetInput(true);
+	Cast<UPCGGraphInputOutputSettings>(InputNode->GetSettings())->SetInput(true);
 
 	OutputNode->ConditionalPostLoad();
 
-	if (!Cast<UPCGGraphInputOutputSettings>(OutputNode->DefaultSettings))
+	if (!Cast<UPCGGraphInputOutputSettings>(OutputNode->GetSettings()))
 	{
-		OutputNode->DefaultSettings = NewObject<UPCGGraphInputOutputSettings>(this, TEXT("DefaultOutputNodeSettings"));
+		OutputNode->SetSettingsInterface(NewObject<UPCGGraphInputOutputSettings>(this, TEXT("DefaultOutputNodeSettings")));
 	}
 
-	Cast<UPCGGraphInputOutputSettings>(OutputNode->DefaultSettings)->SetInput(false);
+	Cast<UPCGGraphInputOutputSettings>(OutputNode->GetSettings())->SetInput(false);
 
 	// Ensure that all nodes are loaded (& updated their deprecated data)
 	for (UPCGNode* Node : Nodes)
@@ -117,7 +117,7 @@ void UPCGGraph::PostLoad()
 	// We remove it at the end, to let the nodes that have null settings to clean up their pins and edges.
 	for (int32 i = Nodes.Num() - 1; i >= 0; --i)
 	{
-		if (!Nodes[i]->DefaultSettings)
+		if (!Nodes[i]->GetSettings())
 		{
 			Nodes.RemoveAtSwap(i);
 		}
@@ -165,47 +165,26 @@ UPCGNode* UPCGGraph::AddNodeOfType(TSubclassOf<class UPCGSettings> InSettingsCla
 		return nullptr;
 	}
 
-	UPCGNode* Node = Settings->CreateNode();
+	UPCGNode* Node = AddNode(Settings);
 
 	if (Node)
 	{
-		Node->SetFlags(RF_Transactional);
-
-		Modify();
-
-		// Assign settings to node
-		Node->SetDefaultSettings(Settings);
 		Settings->Rename(nullptr, Node);
 		Settings->SetFlags(RF_Transactional);
-
-		// Reparent node to this graph
-		Node->Rename(nullptr, this);
-
-#if WITH_EDITOR
-		const FName DefaultNodeName = Settings->GetDefaultNodeName();
-		if (DefaultNodeName != NAME_None)
-		{
-			FName NodeName = MakeUniqueObjectName(this, UPCGNode::StaticClass(), DefaultNodeName);
-			Node->Rename(*NodeName.ToString());
-		}
-#endif
-
-		Nodes.Add(Node);
-		OnNodeAdded(Node);
 	}
 
 	OutDefaultNodeSettings = Settings;
 	return Node;
 }
 
-UPCGNode* UPCGGraph::AddNode(UPCGSettings* InSettings)
+UPCGNode* UPCGGraph::AddNode(UPCGSettingsInterface* InSettingsInterface)
 {
-	if (!InSettings)
+	if (!InSettingsInterface || !InSettingsInterface->GetSettings())
 	{
 		return nullptr;
 	}
 
-	UPCGNode* Node = InSettings->CreateNode();
+	UPCGNode* Node = InSettingsInterface->GetSettings()->CreateNode();
 
 	if (Node)
 	{
@@ -214,13 +193,13 @@ UPCGNode* UPCGGraph::AddNode(UPCGSettings* InSettings)
 		Modify();
 
 		// Assign settings to node & reparent
-		Node->SetDefaultSettings(InSettings);
+		Node->SetSettingsInterface(InSettingsInterface);
 
 		// Reparent node to this graph
 		Node->Rename(nullptr, this);
 
 #if WITH_EDITOR
-		const FName DefaultNodeName = InSettings->GetDefaultNodeName();
+		const FName DefaultNodeName = InSettingsInterface->GetSettings()->GetDefaultNodeName();
 		if (DefaultNodeName != NAME_None)
 		{
 			FName NodeName = MakeUniqueObjectName(this, UPCGNode::StaticClass(), DefaultNodeName);
@@ -235,15 +214,42 @@ UPCGNode* UPCGGraph::AddNode(UPCGSettings* InSettings)
 	return Node;
 }
 
+UPCGNode* UPCGGraph::AddNodeInstance(UPCGSettings* InSettings)
+{
+	if (!InSettings)
+	{
+		return nullptr;
+	}
+
+	UPCGSettingsInstance* SettingsInstance = NewObject<UPCGSettingsInstance>();
+	SettingsInstance->SetSettings(InSettings);
+
+	UPCGNode* Node = AddNode(SettingsInstance);
+
+	if (Node)
+	{
+		SettingsInstance->Rename(nullptr, Node);
+		SettingsInstance->SetFlags(RF_Transactional);
+	}
+
+	return Node;
+}
+
 UPCGNode* UPCGGraph::AddNodeCopy(UPCGSettings* InSettings, UPCGSettings*& DefaultNodeSettings)
 {
+	if (!InSettings)
+	{
+		return nullptr;
+	}
+
 	UPCGSettings* SettingsCopy = DuplicateObject(InSettings, nullptr);
 	UPCGNode* NewNode = AddNode(SettingsCopy);
 
 	if (SettingsCopy)
 	{
 		SettingsCopy->Rename(nullptr, NewNode);
-	}	
+		SettingsCopy->SetFlags(RF_Transactional);
+	}
 
 	DefaultNodeSettings = SettingsCopy;
 	return NewNode;
@@ -319,7 +325,7 @@ bool UPCGGraph::AddLabeledEdge(UPCGNode* From, const FName& FromPinLabel, UPCGNo
 TObjectPtr<UPCGNode> UPCGGraph::ReconstructNewNode(const UPCGNode* InNode)
 {
 	UPCGSettings* NewSettings = nullptr;
-	TObjectPtr<UPCGNode> NewNode = AddNodeCopy(InNode->DefaultSettings, NewSettings);
+	TObjectPtr<UPCGNode> NewNode = AddNodeCopy(InNode->GetSettings(), NewSettings);
 
 #if WITH_EDITOR
 	InNode->TransferEditorProperties(NewNode);
@@ -342,7 +348,7 @@ void UPCGGraph::AddNode(UPCGNode* InNode)
 	InNode->Rename(nullptr, this);
 
 #if WITH_EDITOR
-	const FName DefaultNodeName = InNode->DefaultSettings->GetDefaultNodeName();
+	const FName DefaultNodeName = InNode->GetSettings()->GetDefaultNodeName();
 	if (DefaultNodeName != NAME_None)
 	{
 		FName NodeName = MakeUniqueObjectName(this, UPCGNode::StaticClass(), DefaultNodeName);
@@ -569,9 +575,9 @@ void UPCGGraph::GetTrackedTagsToSettings(FPCGTagToSettingsMap& OutTagsToSettings
 
 	for (UPCGNode* Node : Nodes)
 	{
-		if (Node && Node->DefaultSettings)
+		if (Node && Node->GetSettings())
 		{
-			Node->DefaultSettings->GetTrackedActorTags(OutTagsToSettings, OutVisitedGraphs);
+			Node->GetSettings()->GetTrackedActorTags(OutTagsToSettings, OutVisitedGraphs);
 		}
 	}
 }
diff --git a/Engine/Plugins/Experimental/PCG/Source/PCG/Private/PCGNode.cpp b/Engine/Plugins/Experimental/PCG/Source/PCG/Private/PCGNode.cpp
--- a/Engine/Plugins/Experimental/PCG/Source/PCG/Private/PCGNode.cpp
+++ b/Engine/Plugins/Experimental/PCG/Source/PCG/Private/PCGNode.cpp
@@ -11,7 +11,7 @@
 UPCGNode::UPCGNode(const FObjectInitializer& ObjectInitializer)
 	: Super(ObjectInitializer)
 {
-	DefaultSettings = ObjectInitializer.CreateDefaultSubobject<UPCGTrivialSettings>(this, TEXT("DefaultNodeSettings"));
+	SettingsInterface = ObjectInitializer.CreateDefaultSubobject<UPCGTrivialSettings>(this, TEXT("DefaultNodeSettings"));
 }
 
 void UPCGNode::PostLoad()
@@ -19,10 +19,16 @@ void UPCGNode::PostLoad()
 	Super::PostLoad();
 
 #if WITH_EDITOR
-	if (DefaultSettings)
+	if (DefaultSettings_DEPRECATED)
 	{
-		DefaultSettings->OnSettingsChangedDelegate.AddUObject(this, &UPCGNode::OnSettingsChanged);
-		DefaultSettings->ConditionalPostLoad();
+		SettingsInterface = DefaultSettings_DEPRECATED;
+		DefaultSettings_DEPRECATED = nullptr;
+	}
+
+	if (SettingsInterface)
+	{
+		SettingsInterface->OnSettingsChangedDelegate.AddUObject(this, &UPCGNode::OnSettingsChanged);
+		SettingsInterface->ConditionalPostLoad();
 	}
 
 	// Make sure legacy nodes support transactions.
@@ -117,9 +123,9 @@ void UPCGNode::ApplyDeprecation()
 	}
 	OutboundEdges_DEPRECATED.Reset();
 
-	if (DefaultSettings)
+	if (UPCGSettings* Settings = GetSettings())
 	{
-		DefaultSettings->ApplyDeprecation(this);
+		Settings->ApplyDeprecation(this);
 	}
 }
 #endif
@@ -128,17 +134,17 @@ void UPCGNode::ApplyDeprecation()
 void UPCGNode::PostEditImport()
 {
 	Super::PostEditImport();
-	if (DefaultSettings)
+	if (SettingsInterface)
 	{
-		DefaultSettings->OnSettingsChangedDelegate.AddUObject(this, &UPCGNode::OnSettingsChanged);
+		SettingsInterface->OnSettingsChangedDelegate.AddUObject(this, &UPCGNode::OnSettingsChanged);
 	}
 }
 
 void UPCGNode::PreEditUndo()
 {
-	if (DefaultSettings)
+	if (SettingsInterface)
 	{
-		DefaultSettings->OnSettingsChangedDelegate.RemoveAll(this);
+		SettingsInterface->OnSettingsChangedDelegate.RemoveAll(this);
 	}
 
 	UObject* Outer = GetOuter();
@@ -154,9 +160,9 @@ void UPCGNode::PostEditUndo()
 {
 	Super::PostEditUndo();
 
-	if (DefaultSettings)
+	if (SettingsInterface)
 	{
-		DefaultSettings->OnSettingsChangedDelegate.AddUObject(this, &UPCGNode::OnSettingsChanged);
+		SettingsInterface->OnSettingsChangedDelegate.AddUObject(this, &UPCGNode::OnSettingsChanged);
 	}
 
  	UObject* Outer = GetOuter();
@@ -170,9 +176,9 @@ void UPCGNode::PostEditUndo()
 void UPCGNode::BeginDestroy()
 {
 #if WITH_EDITOR
-	if (DefaultSettings)
+	if (SettingsInterface)
 	{
-		DefaultSettings->OnSettingsChangedDelegate.RemoveAll(this);
+		SettingsInterface->OnSettingsChangedDelegate.RemoveAll(this);
 	}
 #endif
 
@@ -214,16 +220,16 @@ FName UPCGNode::GetNodeTitle() const
 	{
 		return NodeTitle;
 	}
-	else if (DefaultSettings)
+	else if (UPCGSettings* Settings = GetSettings())
 	{
-		if (DefaultSettings->AdditionalTaskName() != NAME_None)
+		if (Settings->AdditionalTaskName() != NAME_None)
 		{
-			return DefaultSettings->AdditionalTaskName();
+			return Settings->AdditionalTaskName();
 		}
 #if WITH_EDITOR
 		else
 		{
-			return DefaultSettings->GetDefaultNodeName();
+			return Settings->GetDefaultNodeName();
 		}
 #endif
 	}
@@ -234,10 +240,22 @@ FName UPCGNode::GetNodeTitle() const
 #if WITH_EDITOR
 FText UPCGNode::GetNodeTooltipText() const
 {
-	return DefaultSettings ? DefaultSettings->GetNodeTooltipText() : FText::GetEmpty();
+	if (UPCGSettings* Settings = GetSettings())
+	{
+		return Settings->GetNodeTooltipText();
+	}
+	else
+	{
+		return FText::GetEmpty();
+	}
 }
 #endif
 
+bool UPCGNode::IsInstance() const
+{
+	return SettingsInterface && SettingsInterface->IsInstance();
+}
+
 TArray<FPCGPinProperties> UPCGNode::InputPinProperties() const
 {
 	TArray<FPCGPinProperties> PinProperties;
@@ -344,22 +362,23 @@ bool UPCGNode::HasInboundEdges() const
 	return false;
 }
 
-void UPCGNode::SetDefaultSettings(TObjectPtr<UPCGSettings> InSettings, bool bUpdatePins)
+void UPCGNode::SetSettingsInterface(UPCGSettingsInterface* InSettingsInterface, bool bUpdatePins)
 {
 #if WITH_EDITOR
-	const bool bDifferentSettings = (DefaultSettings != InSettings);
-	if (bDifferentSettings && DefaultSettings)
+	const bool bDifferentInterface = (SettingsInterface.Get() != InSettingsInterface);
+	if (bDifferentInterface && SettingsInterface)
 	{
-		DefaultSettings->OnSettingsChangedDelegate.RemoveAll(this);
+		SettingsInterface->OnSettingsChangedDelegate.RemoveAll(this);
 	}
 #endif
 
-	DefaultSettings = InSettings;
+	SettingsInterface = InSettingsInterface;
 
 #if WITH_EDITOR
-	if (bDifferentSettings && DefaultSettings)
+	if (bDifferentInterface && SettingsInterface)
 	{
-		DefaultSettings->OnSettingsChangedDelegate.AddUObject(this, &UPCGNode::OnSettingsChanged);
+		check(SettingsInterface->GetSettings());
+		SettingsInterface->OnSettingsChangedDelegate.AddUObject(this, &UPCGNode::OnSettingsChanged);
 	}
 #endif
 
@@ -369,44 +388,22 @@ void UPCGNode::SetDefaultSettings(TObjectPtr<UPCGSettings> InSettings, bool bUpd
 	}
 }
 
-#if WITH_EDITOR
-
-void UPCGNode::PreEditChange(FProperty* PropertyAboutToChange)
+UPCGSettings* UPCGNode::GetSettings() const
 {
-	// To properly clean up old callbacks during paste we clear during null call since the Settings property doesn't get a specific call.
-	if (!PropertyAboutToChange || (PropertyAboutToChange && PropertyAboutToChange->GetFName() == GET_MEMBER_NAME_CHECKED(UPCGNode, DefaultSettings)))
+	if (SettingsInterface)
 	{
-		if (DefaultSettings)
-		{
-			DefaultSettings->OnSettingsChangedDelegate.RemoveAll(this);
-		}
+		return SettingsInterface->GetSettings();
 	}
-
-	Super::PreEditChange(PropertyAboutToChange);
-}
-
-void UPCGNode::PostEditChangeProperty(struct FPropertyChangedEvent& PropertyChangedEvent)
-{
-	Super::PostEditChangeProperty(PropertyChangedEvent);
-
-	if (PropertyChangedEvent.Property && PropertyChangedEvent.Property->GetFName() == GET_MEMBER_NAME_CHECKED(UPCGNode, DefaultSettings))
-	{
-		if (DefaultSettings)
-		{
-			DefaultSettings->OnSettingsChangedDelegate.AddUObject(this, &UPCGNode::OnSettingsChanged);
-			UpdatePins();
-			OnNodeChangedDelegate.Broadcast(this, ((Cast<UPCGBaseSubgraphSettings>(DefaultSettings) ? EPCGChangeType::Structural : EPCGChangeType::None) | EPCGChangeType::Settings));
-		}
-	}
-	else if (PropertyChangedEvent.Property && PropertyChangedEvent.Property->GetFName() == GET_MEMBER_NAME_CHECKED(UPCGNode, NodeTitle))
+	else
 	{
-		OnNodeChangedDelegate.Broadcast(this, EPCGChangeType::Cosmetic);
+		return nullptr;
 	}
 }
 
+#if WITH_EDITOR
 void UPCGNode::OnSettingsChanged(UPCGSettings* InSettings, EPCGChangeType ChangeType)
 {
-	if (InSettings == DefaultSettings)
+	if (InSettings == GetSettings())
 	{
 		const bool bUpdatedPins = UpdatePins();
 		OnNodeChangedDelegate.Broadcast(this, ((bUpdatedPins ? EPCGChangeType::Edge : EPCGChangeType::None) | ChangeType));
@@ -436,7 +433,7 @@ bool UPCGNode::UpdatePins()
 
 bool UPCGNode::UpdatePins(TFunctionRef<UPCGPin*(UPCGNode*)> PinAllocator)
 {
-	if (!DefaultSettings)
+	if (!GetSettings())
 	{
 		bool bChanged = !InputPins.IsEmpty() || !OutputPins.IsEmpty();
 
@@ -467,8 +464,9 @@ bool UPCGNode::UpdatePins(TFunctionRef<UPCGPin*(UPCGNode*)> PinAllocator)
 		return bChanged;
 	}
 
-	TArray<FPCGPinProperties> InboundPinProperties = DefaultSettings->InputPinProperties();
-	TArray<FPCGPinProperties> OutboundPinProperties = DefaultSettings->OutputPinProperties();
+	UPCGSettings* Settings = GetSettings();
+	TArray<FPCGPinProperties> InboundPinProperties = Settings->InputPinProperties();
+	TArray<FPCGPinProperties> OutboundPinProperties = Settings->OutputPinProperties();
 
 	auto UpdatePins = [this, &PinAllocator](TArray<UPCGPin*>& Pins, const TArray<FPCGPinProperties>& PinProperties)
 	{
diff --git a/Engine/Plugins/Experimental/PCG/Source/PCG/Private/PCGSettings.cpp b/Engine/Plugins/Experimental/PCG/Source/PCG/Private/PCGSettings.cpp
--- a/Engine/Plugins/Experimental/PCG/Source/PCG/Private/PCGSettings.cpp
+++ b/Engine/Plugins/Experimental/PCG/Source/PCG/Private/PCGSettings.cpp
@@ -5,6 +5,7 @@
 #include "PCGNode.h"
 #include "PCGSubsystem.h"
 #include "Serialization/ArchiveObjectCrc32.h"
+#include "UObject/ObjectSaveContext.h"
 
 #if WITH_EDITOR
 #include "Editor.h"
@@ -24,6 +25,11 @@ public:
 	}
 };
 
+bool UPCGSettingsInterface::IsInstance() const
+{
+	return this != GetSettings();
+}
+
 bool UPCGSettings::operator==(const UPCGSettings& Other) const
 {
 	if (this == &Other)
@@ -45,6 +51,19 @@ uint32 UPCGSettings::GetCrc32() const
 	return Ar.Crc32(const_cast<UPCGSettings*>(this));
 }
 
+void UPCGSettings::PostLoad()
+{
+	Super::PostLoad();
+
+#if WITH_EDITORONLY_DATA
+	if (ExecutionMode_DEPRECATED != EPCGSettingsExecutionMode::Enabled)
+	{
+		ExecutionMode = ExecutionMode_DEPRECATED;
+		ExecutionMode_DEPRECATED = EPCGSettingsExecutionMode::Enabled;
+	}
+#endif
+}
+
 void UPCGSettings::Serialize(FArchive& Ar)
 {
 	Super::Serialize(Ar);
@@ -80,6 +99,17 @@ void UPCGSettings::Serialize(FArchive& Ar)
 	}
 }
 
+void UPCGSettings::PostSaveRoot(FObjectPostSaveRootContext ObjectSaveContext)
+{
+	Super::PostSaveRoot(ObjectSaveContext);
+
+#if WITH_EDITOR
+	// This will get called when an external settings gets saved;
+	// This is to trigger generation on save, if we've called changed properties from a blueprint
+	OnSettingsChangedDelegate.Broadcast(this, EPCGChangeType::Structural);
+#endif
+}
+
 #if WITH_EDITOR
 UObject* UPCGSettings::GetJumpTargetForDoubleClick() const
 {
@@ -158,6 +188,75 @@ TArray<FPCGPinProperties> UPCGSettings::DefaultPointOutputPinProperties() const
 	return Properties;
 }
 
+void UPCGSettingsInstance::PostLoad()
+{
+	Super::PostLoad();
+
+#if WITH_EDITOR
+	if (Settings)
+	{
+		Settings->OnSettingsChangedDelegate.AddUObject(this, &UPCGSettingsInstance::OnSettingsChanged);
+		Settings->ConditionalPostLoad();
+	}
+#endif
+
+#if WITH_EDITORONLY_DATA
+	OriginalSettings = Settings;
+#endif
+}
+
+void UPCGSettingsInstance::BeginDestroy()
+{
+#if WITH_EDITOR
+	if (Settings)
+	{
+		Settings->OnSettingsChangedDelegate.RemoveAll(this);
+	}
+#endif
+
+	Super::BeginDestroy();
+}
+
+void UPCGSettingsInstance::SetSettings(UPCGSettings* InSettings)
+{
+#if WITH_EDITOR
+	if (Settings)
+	{
+		Settings->OnSettingsChangedDelegate.RemoveAll(this);
+	}
+#endif
+
+	Settings = InSettings;
+#if WITH_EDITORONLY_DATA
+	OriginalSettings = Settings;
+#endif
+
+#if WITH_EDITOR
+	if (Settings)
+	{
+		Settings->OnSettingsChangedDelegate.AddUObject(this, &UPCGSettingsInstance::OnSettingsChanged);
+	}
+#endif
+}
+
+#if WITH_EDITOR
+void UPCGSettingsInstance::PostEditChangeProperty(FPropertyChangedEvent& PropertyChangedEvent)
+{
+	Super::PostEditChangeProperty(PropertyChangedEvent);
+	// Some setting in the instance has changed. We don't have a flag for that yet (to add if needed)
+	// However, we can make it behave like a standard change
+	OnSettingsChangedDelegate.Broadcast(Settings, EPCGChangeType::Settings);
+}
+
+void UPCGSettingsInstance::OnSettingsChanged(UPCGSettings* InSettings, EPCGChangeType ChangeType)
+{
+	if (InSettings == Settings)
+	{
+		OnSettingsChangedDelegate.Broadcast(InSettings, ChangeType);
+	}
+}
+#endif
+
 FPCGElementPtr UPCGTrivialSettings::CreateElement() const
 {
 	return MakeShared<FPCGTrivialElement>();
diff --git a/Engine/Plugins/Experimental/PCG/Source/PCG/Private/PCGSubgraph.cpp b/Engine/Plugins/Experimental/PCG/Source/PCG/Private/PCGSubgraph.cpp
--- a/Engine/Plugins/Experimental/PCG/Source/PCG/Private/PCGSubgraph.cpp
+++ b/Engine/Plugins/Experimental/PCG/Source/PCG/Private/PCGSubgraph.cpp
@@ -135,7 +135,7 @@ FPCGElementPtr UPCGSubgraphSettings::CreateElement() const
 
 TObjectPtr<UPCGGraph> UPCGSubgraphNode::GetSubgraph() const
 {
-	TObjectPtr<UPCGSubgraphSettings> Settings = Cast<UPCGSubgraphSettings>(DefaultSettings);
+	TObjectPtr<UPCGSubgraphSettings> Settings = Cast<UPCGSubgraphSettings>(GetSettings());
 	return Settings ? Settings->Subgraph : nullptr;
 }
 
diff --git a/Engine/Plugins/Experimental/PCG/Source/PCG/Private/Tests/Determinism/PCGDeterminismTestBlueprintBase.cpp b/Engine/Plugins/Experimental/PCG/Source/PCG/Private/Tests/Determinism/PCGDeterminismTestBlueprintBase.cpp
--- a/Engine/Plugins/Experimental/PCG/Source/PCG/Private/Tests/Determinism/PCGDeterminismTestBlueprintBase.cpp
+++ b/Engine/Plugins/Experimental/PCG/Source/PCG/Private/Tests/Determinism/PCGDeterminismTestBlueprintBase.cpp
@@ -7,6 +7,6 @@ void UPCGDeterminismTestBlueprintBase::ExecuteTest_Implementation(const UPCGNode
 	// The default case if the user tries to run the blueprint test unimplemented
 	InOutTestResult.TestResultTitle = InPCGNode->GetNodeTitle();
 	InOutTestResult.TestResultName = InPCGNode->GetName();
-	InOutTestResult.Seed = InPCGNode->DefaultSettings->Seed;
+	InOutTestResult.Seed = InPCGNode->GetSettings()->Seed;
 	PCGDeterminismTests::RunDeterminismTest(InPCGNode, InOutTestResult, PCGDeterminismTests::Defaults::DeterminismBasicTestInfo);
 }
\ No newline at end of file
diff --git a/Engine/Plugins/Experimental/PCG/Source/PCG/Private/Tests/Determinism/PCGDeterminismTestsCommon.cpp b/Engine/Plugins/Experimental/PCG/Source/PCG/Private/Tests/Determinism/PCGDeterminismTestsCommon.cpp
--- a/Engine/Plugins/Experimental/PCG/Source/PCG/Private/Tests/Determinism/PCGDeterminismTestsCommon.cpp
+++ b/Engine/Plugins/Experimental/PCG/Source/PCG/Private/Tests/Determinism/PCGDeterminismTestsCommon.cpp
@@ -45,7 +45,7 @@ namespace PCGDeterminismTests
 	{
 		TRACE_CPUPROFILER_EVENT_SCOPE(PCGDeterminismTests::RunDeterminismTest::Node);
 
-		if (!InPCGNode || !InPCGNode->DefaultSettings)
+		if (!InPCGNode || !InPCGNode->GetSettings())
 		{
 			OutResult.DataTypesTested = EPCGDataType::None;
 			OutResult.bFlagRaised = true;
@@ -279,14 +279,14 @@ namespace PCGDeterminismTests
 		const UPCGNode* PCGNode = NodeAndOptions.PCGNode;
 		const TArray<TArray<EPCGDataType>>& BaseOptionsByPin = NodeAndOptions.BaseOptionsByPin;
 
-		PCGTestsCommon::FTestData TestData(NodeAndOptions.Seed, PCGNode->DefaultSettings);
+		PCGTestsCommon::FTestData TestData(NodeAndOptions.Seed, PCGNode->GetSettings());
 
 		// For all permutations
 		const int32 NumPermutations = GetNumPermutations(BaseOptionsByPin);
 		for (int32 PermutationIndex = 0; PermutationIndex < NumPermutations; ++PermutationIndex)
 		{
 			// Reset test data
-			TestData.Reset(PCGNode->DefaultSettings);
+			TestData.Reset(PCGNode->GetSettings());
 
 			// Add an input for each pin
 			for (int32 PinIndex = 0; PinIndex < BaseOptionsByPin.Num(); ++PinIndex)
@@ -309,16 +309,16 @@ namespace PCGDeterminismTests
 		const UPCGNode* PCGNode = NodeAndOptions.PCGNode;
 		const TArray<TArray<EPCGDataType>>& BaseOptionsByPin = NodeAndOptions.BaseOptionsByPin;
 
-		PCGTestsCommon::FTestData FirstTestData(NodeAndOptions.Seed, PCGNode->DefaultSettings);
-		PCGTestsCommon::FTestData SecondTestData(NodeAndOptions.Seed, PCGNode->DefaultSettings);
+		PCGTestsCommon::FTestData FirstTestData(NodeAndOptions.Seed, PCGNode->GetSettings());
+		PCGTestsCommon::FTestData SecondTestData(NodeAndOptions.Seed, PCGNode->GetSettings());
 
 		const int32 NumPermutations = GetNumPermutations(BaseOptionsByPin);
 		// For all permutations
 		for (int32 PermutationIndex = 0; PermutationIndex < NumPermutations; ++PermutationIndex)
 		{
 			// Reset test data
-			FirstTestData.Reset(PCGNode->DefaultSettings);
-			SecondTestData.Reset(PCGNode->DefaultSettings);
+			FirstTestData.Reset(PCGNode->GetSettings());
+			SecondTestData.Reset(PCGNode->GetSettings());
 
 			// Add an input for each pin
 			for (int32 PinIndex = 0; PinIndex < BaseOptionsByPin.Num(); ++PinIndex)
@@ -346,8 +346,8 @@ namespace PCGDeterminismTests
 		const UPCGNode* PCGNode = NodeAndOptions.PCGNode;
 		const TArray<TArray<EPCGDataType>>& BaseOptionsByPin = NodeAndOptions.BaseOptionsByPin;
 
-		PCGTestsCommon::FTestData FirstTestData(NodeAndOptions.Seed, PCGNode->DefaultSettings);
-		PCGTestsCommon::FTestData SecondTestData(NodeAndOptions.Seed, PCGNode->DefaultSettings);
+		PCGTestsCommon::FTestData FirstTestData(NodeAndOptions.Seed, PCGNode->GetSettings());
+		PCGTestsCommon::FTestData SecondTestData(NodeAndOptions.Seed, PCGNode->GetSettings());
 
 		int32 NumInputs = NodeAndOptions.bMultipleOptionsPerPin ? NumInputsPerPin : 1;
 		const int32 NumPermutations = GetNumPermutations(BaseOptionsByPin);
@@ -358,8 +358,8 @@ namespace PCGDeterminismTests
 			for (int32 PermutationIndex = 0; PermutationIndex < NumPermutations;)
 			{
 				// Reset test data
-				FirstTestData.Reset(PCGNode->DefaultSettings);
-				SecondTestData.Reset(PCGNode->DefaultSettings);
+				FirstTestData.Reset(PCGNode->GetSettings());
+				SecondTestData.Reset(PCGNode->GetSettings());
 
 				// Add an input for each pin
 				for (int32 PinIndex = 0; PinIndex < BaseOptionsByPin.Num(); ++PinIndex)
@@ -424,8 +424,8 @@ namespace PCGDeterminismTests
 		const UPCGNode* PCGNode = NodeAndOptions.PCGNode;
 		const TArray<TArray<EPCGDataType>>& BaseOptionsByPin = NodeAndOptions.BaseOptionsByPin;
 
-		PCGTestsCommon::FTestData FirstTestData(NodeAndOptions.Seed, PCGNode->DefaultSettings);
-		PCGTestsCommon::FTestData SecondTestData(NodeAndOptions.Seed, PCGNode->DefaultSettings);
+		PCGTestsCommon::FTestData FirstTestData(NodeAndOptions.Seed, PCGNode->GetSettings());
+		PCGTestsCommon::FTestData SecondTestData(NodeAndOptions.Seed, PCGNode->GetSettings());
 
 		int32 NumInputs = NodeAndOptions.bMultipleOptionsPerPin ? NumInputsPerPin : 1;
 
diff --git a/Engine/Plugins/Experimental/PCG/Source/PCG/Public/Elements/PCGDensityFilter.h b/Engine/Plugins/Experimental/PCG/Source/PCG/Public/Elements/PCGDensityFilter.h
--- a/Engine/Plugins/Experimental/PCG/Source/PCG/Public/Elements/PCGDensityFilter.h
+++ b/Engine/Plugins/Experimental/PCG/Source/PCG/Public/Elements/PCGDensityFilter.h
@@ -38,7 +38,7 @@ public:
 	bool bInvertFilter = false;
 
 #if WITH_EDITORONLY_DATA
-	UPROPERTY(Transient, BlueprintReadWrite, EditAnywhere, Category = Debug)
+	UPROPERTY(Transient, BlueprintReadWrite, EditAnywhere, Category = "Settings|Debug")
 	bool bKeepZeroDensityPoints = false;
 #endif
 };
diff --git a/Engine/Plugins/Experimental/PCG/Source/PCG/Public/Elements/PCGDifferenceElement.h b/Engine/Plugins/Experimental/PCG/Source/PCG/Public/Elements/PCGDifferenceElement.h
--- a/Engine/Plugins/Experimental/PCG/Source/PCG/Public/Elements/PCGDifferenceElement.h
+++ b/Engine/Plugins/Experimental/PCG/Source/PCG/Public/Elements/PCGDifferenceElement.h
@@ -44,7 +44,7 @@ public:
 	bool bDiffMetadata = true;
 
 #if WITH_EDITORONLY_DATA
-	UPROPERTY(Transient, BlueprintReadWrite, EditAnywhere, Category = Debug)
+	UPROPERTY(Transient, BlueprintReadWrite, EditAnywhere, Category = "Settings|Debug")
 	bool bKeepZeroDensityPoints = false;
 #endif
 };
diff --git a/Engine/Plugins/Experimental/PCG/Source/PCG/Public/Elements/PCGIntersectionElement.h b/Engine/Plugins/Experimental/PCG/Source/PCG/Public/Elements/PCGIntersectionElement.h
--- a/Engine/Plugins/Experimental/PCG/Source/PCG/Public/Elements/PCGIntersectionElement.h
+++ b/Engine/Plugins/Experimental/PCG/Source/PCG/Public/Elements/PCGIntersectionElement.h
@@ -30,7 +30,7 @@ public:
 	EPCGIntersectionDensityFunction DensityFunction = EPCGIntersectionDensityFunction::Multiply;
 
 #if WITH_EDITORONLY_DATA
-	UPROPERTY(Transient, BlueprintReadWrite, EditAnywhere, Category = Debug)
+	UPROPERTY(Transient, BlueprintReadWrite, EditAnywhere, Category = "Settings|Debug")
 	bool bKeepZeroDensityPoints = false;
 #endif
 };
diff --git a/Engine/Plugins/Experimental/PCG/Source/PCG/Public/Elements/PCGPointSampler.h b/Engine/Plugins/Experimental/PCG/Source/PCG/Public/Elements/PCGPointSampler.h
--- a/Engine/Plugins/Experimental/PCG/Source/PCG/Public/Elements/PCGPointSampler.h
+++ b/Engine/Plugins/Experimental/PCG/Source/PCG/Public/Elements/PCGPointSampler.h
@@ -34,7 +34,7 @@ public:
 	float Ratio = 0.1f;
 
 #if WITH_EDITORONLY_DATA
-	UPROPERTY(Transient, BlueprintReadWrite, EditAnywhere, Category = Debug)
+	UPROPERTY(Transient, BlueprintReadWrite, EditAnywhere, Category = "Settings|Debug")
 	bool bKeepZeroDensityPoints = false;
 #endif
 };
diff --git a/Engine/Plugins/Experimental/PCG/Source/PCG/Public/Elements/PCGProjectionElement.h b/Engine/Plugins/Experimental/PCG/Source/PCG/Public/Elements/PCGProjectionElement.h
--- a/Engine/Plugins/Experimental/PCG/Source/PCG/Public/Elements/PCGProjectionElement.h
+++ b/Engine/Plugins/Experimental/PCG/Source/PCG/Public/Elements/PCGProjectionElement.h
@@ -34,7 +34,7 @@ public:
 	FPCGProjectionParams Params;
 
 #if WITH_EDITORONLY_DATA
-	UPROPERTY(Transient, BlueprintReadWrite, EditAnywhere, Category = Debug)
+	UPROPERTY(Transient, BlueprintReadWrite, EditAnywhere, Category = "Settings|Debug")
 	bool bKeepZeroDensityPoints = false;
 #endif
 
diff --git a/Engine/Plugins/Experimental/PCG/Source/PCG/Public/Elements/PCGSurfaceSampler.h b/Engine/Plugins/Experimental/PCG/Source/PCG/Public/Elements/PCGSurfaceSampler.h
--- a/Engine/Plugins/Experimental/PCG/Source/PCG/Public/Elements/PCGSurfaceSampler.h
+++ b/Engine/Plugins/Experimental/PCG/Source/PCG/Public/Elements/PCGSurfaceSampler.h
@@ -40,7 +40,7 @@ public:
 	float PointSteepness = 0.5f;
 
 #if WITH_EDITORONLY_DATA
-	UPROPERTY(Transient, BlueprintReadWrite, EditAnywhere, Category = "Debug")
+	UPROPERTY(Transient, BlueprintReadWrite, EditAnywhere, Category = "Settings|Debug")
 	bool bKeepZeroDensityPoints = false;
 #endif
 
diff --git a/Engine/Plugins/Experimental/PCG/Source/PCG/Public/PCGContext.h b/Engine/Plugins/Experimental/PCG/Source/PCG/Public/PCGContext.h
--- a/Engine/Plugins/Experimental/PCG/Source/PCG/Public/PCGContext.h
+++ b/Engine/Plugins/Experimental/PCG/Source/PCG/Public/PCGContext.h
@@ -9,6 +9,7 @@
 #include "PCGContext.generated.h"
 
 class UPCGComponent;
+class UPCGSettingsInterface;
 struct FPCGGraphCache;
 
 namespace PCGContextHelpers
@@ -16,9 +17,9 @@ namespace PCGContextHelpers
 	template<typename SettingsType>
 	const SettingsType* GetInputSettings(const UPCGNode* Node, const FPCGDataCollection& InputData)
 	{
-		if (Node && Node->DefaultSettings)
+		if (Node && Node->GetSettings())
 		{
-			return Cast<SettingsType>(InputData.GetSettings(Node->DefaultSettings));
+			return Cast<SettingsType>(InputData.GetSettings(Node->GetSettings()));
 		}
 		else
 		{
@@ -69,6 +70,8 @@ struct PCG_API FPCGContext
 	int32 ExecutionCount = 0;
 #endif
 
+	const UPCGSettingsInterface* GetInputSettingsInterface() const;
+
 	template<typename SettingsType>
 	const SettingsType* GetInputSettings() const
 	{
diff --git a/Engine/Plugins/Experimental/PCG/Source/PCG/Public/PCGData.h b/Engine/Plugins/Experimental/PCG/Source/PCG/Public/PCGData.h
--- a/Engine/Plugins/Experimental/PCG/Source/PCG/Public/PCGData.h
+++ b/Engine/Plugins/Experimental/PCG/Source/PCG/Public/PCGData.h
@@ -11,6 +11,7 @@
 #include "PCGData.generated.h"
 
 class UPCGSettings;
+class UPCGSettingsInterface;
 class UPCGParamData;
 
 /**
@@ -69,6 +70,9 @@ struct PCG_API FPCGDataCollection
 	TArray<FPCGTaggedData> GetParamsByPin(const FName& InPinLabel) const;
 	UPCGParamData* GetParams() const;
 
+	const UPCGSettingsInterface* GetSettingsInterface() const;
+	const UPCGSettingsInterface* GetSettingsInterface(const UPCGSettingsInterface* InDefaultSettingsInterface) const;
+
 	template<typename SettingsType>
 	const SettingsType* GetSettings() const;
 
diff --git a/Engine/Plugins/Experimental/PCG/Source/PCG/Public/PCGElement.h b/Engine/Plugins/Experimental/PCG/Source/PCG/Public/PCGElement.h
--- a/Engine/Plugins/Experimental/PCG/Source/PCG/Public/PCGElement.h
+++ b/Engine/Plugins/Experimental/PCG/Source/PCG/Public/PCGElement.h
@@ -41,6 +41,9 @@ public:
 	/** Returns true if the element, in its current phase can be executed only from the main thread */
 	virtual bool CanExecuteOnlyOnMainThread(FPCGContext* Context) const { return false; }
 
+	/** Returns true if the node can be cached - also checks for instance flags, if any. */
+	bool IsCacheableInstance(const UPCGSettingsInterface* InSettingsInterface) const;
+
 	/** Returns true if the node can be cached (e.g. does not create artifacts & does not depend on untracked data */
 	virtual bool IsCacheable(const UPCGSettings* InSettings) const { return true; }
 
diff --git a/Engine/Plugins/Experimental/PCG/Source/PCG/Public/PCGGraph.h b/Engine/Plugins/Experimental/PCG/Source/PCG/Public/PCGGraph.h
--- a/Engine/Plugins/Experimental/PCG/Source/PCG/Public/PCGGraph.h
+++ b/Engine/Plugins/Experimental/PCG/Source/PCG/Public/PCGGraph.h
@@ -47,13 +47,16 @@ public:
 	UPROPERTY(EditAnywhere, Category = Settings)
 	bool bLandscapeUsesMetadata = true;
 
+	/** Creates a node using the given settings interface. Does not manage ownership - done outside of this method. */
+	UPCGNode* AddNode(UPCGSettingsInterface* InSettings);
+
 	/** Creates a default node based on the settings class wanted. Returns the newly created node. */
 	UFUNCTION(BlueprintCallable, Category = Graph, meta=(DeterminesOutputType = "InSettingsClass", DynamicOutputParam = "DefaultNodeSettings"))
 	UPCGNode* AddNodeOfType(TSubclassOf<class UPCGSettings> InSettingsClass, UPCGSettings*& DefaultNodeSettings);
 
-	/** Creates a node and assigns it in the input settings (settings are not owned by node). Returns the created node. */
+	/** Creates a node containing an instance to the given settings. Returns the created node. */
 	UFUNCTION(BlueprintCallable, Category = Graph)
-	UPCGNode* AddNode(UPCGSettings* InSettings);
+	UPCGNode* AddNodeInstance(UPCGSettings* InSettings);
 
 	/** Creates a node and copies the input settings. Returns the created node. */
 	UFUNCTION(BlueprintCallable, Category = Graph, meta = (DeterminesOutputType = "InSettings", DynamicOutputParam = "OutCopiedSettings"))
diff --git a/Engine/Plugins/Experimental/PCG/Source/PCG/Public/PCGNode.h b/Engine/Plugins/Experimental/PCG/Source/PCG/Public/PCGNode.h
--- a/Engine/Plugins/Experimental/PCG/Source/PCG/Public/PCGNode.h
+++ b/Engine/Plugins/Experimental/PCG/Source/PCG/Public/PCGNode.h
@@ -8,6 +8,7 @@
 #include "PCGNode.generated.h"
 
 class UPCGSettings;
+class UPCGSettingsInterface;
 class UPCGGraph;
 class UPCGEdge;
 class IPCGElement;
@@ -74,8 +75,18 @@ public:
 	/** Returns true if the output pin is connected */
 	bool IsOutputPinConnected(const FName& Label) const;
 
+	/** Returns true if the node has an instance of the settings (e.g. does not own the settings) */
+	bool IsInstance() const;
+
+	/** Returns the settings interface (settings or instance of settings) on this node */
+	UPCGSettingsInterface* GetSettingsInterface() const { return SettingsInterface.Get(); }
+
 	/** Changes the default settings in the node */
-	void SetDefaultSettings(TObjectPtr<UPCGSettings> InSettings, bool bUpdatePins = true);
+	void SetSettingsInterface(UPCGSettingsInterface* InSettingsInterface, bool bUpdatePins = true);
+
+	/** Returns the settings this node holds (either directly or through an instance) */
+	UFUNCTION(BlueprintCallable, Category = Node)
+	UPCGSettings* GetSettings() const;
 
 	/** Triggers some uppdates after creating a new node and changing its settings */
 	void UpdateAfterSettingsChangeDuringCreation();
@@ -102,9 +113,10 @@ public:
 	void SetNodePosition(int32 InPositionX, int32 InPositionY);
 #endif
 
-	/** Note: do not set this property directly from code, use SetDefaultSettings instead */
-	UPROPERTY(BlueprintReadWrite, EditAnywhere, Category = Node, meta=(EditInline))
-	TObjectPtr<UPCGSettings> DefaultSettings;
+#if WITH_EDITORONLY_DATA
+	UPROPERTY()
+	TObjectPtr<UPCGSettings> DefaultSettings_DEPRECATED; 
+#endif
 
 	UPROPERTY(BlueprintReadWrite, EditAnywhere, Category = Node)
 	FName NodeTitle = NAME_None;
@@ -140,11 +152,13 @@ protected:
 	bool UpdatePins(TFunctionRef<UPCGPin* (UPCGNode*)> PinAllocator);
 
 #if WITH_EDITOR
-	virtual void PreEditChange(FProperty* PropertyAboutToChange) override;
-	virtual void PostEditChangeProperty(struct FPropertyChangedEvent& PropertyChangedEvent) override;
 	void OnSettingsChanged(UPCGSettings* InSettings, EPCGChangeType ChangeType);
 #endif
 
+	/** Note: do not set this property directly from code, use SetSettingsInterface instead */
+	UPROPERTY(BlueprintReadOnly, VisibleAnywhere, Category = Node)
+	TObjectPtr<UPCGSettingsInterface> SettingsInterface;
+
 	UPROPERTY()
 	TArray<TObjectPtr<UPCGNode>> OutboundNodes_DEPRECATED;
 
diff --git a/Engine/Plugins/Experimental/PCG/Source/PCG/Public/PCGSettings.h b/Engine/Plugins/Experimental/PCG/Source/PCG/Public/PCGSettings.h
--- a/Engine/Plugins/Experimental/PCG/Source/PCG/Public/PCGSettings.h
+++ b/Engine/Plugins/Experimental/PCG/Source/PCG/Public/PCGSettings.h
@@ -46,11 +46,35 @@ enum class EPCGSettingsType : uint8
 	DECLARE_MULTICAST_DELEGATE_TwoParams(FOnPCGSettingsChanged, UPCGSettings*, EPCGChangeType);
 #endif
 
+UCLASS(Abstract, MinimalAPI)
+class UPCGSettingsInterface : public UPCGData
+{
+	GENERATED_BODY()
+
+public:
+	virtual UPCGSettings* GetSettings() PURE_VIRTUAL(UPCGSettingsInterface::GetSettings, return nullptr;);
+	virtual const UPCGSettings* GetSettings() const PURE_VIRTUAL(UPCGSettingsInterface::GetSettings, return nullptr;);
+
+	bool IsInstance() const;
+
+#if WITH_EDITOR
+	FOnPCGSettingsChanged OnSettingsChangedDelegate;
+#endif
+
+	UPROPERTY(EditAnywhere, BlueprintReadWrite, Category = Debug)
+	EPCGSettingsExecutionMode ExecutionMode = EPCGSettingsExecutionMode::Enabled;
+
+#if WITH_EDITORONLY_DATA
+	UPROPERTY(EditAnywhere, BlueprintReadWrite, Category = Debug, meta = (ShowOnlyInnerProperties))
+	FPCGDebugVisualizationSettings DebugSettings;
+#endif
+};
+
 /**
 * Base class for settings-as-data in the PCG framework
 */
 UCLASS(Abstract, BlueprintType, ClassGroup = (Procedural))
-class PCG_API UPCGSettings : public UPCGData
+class PCG_API UPCGSettings : public UPCGSettingsInterface
 {
 	GENERATED_BODY()
 
@@ -59,8 +83,15 @@ public:
 	virtual EPCGDataType GetDataType() const override { return EPCGDataType::Settings | Super::GetDataType(); }
 	// ~End UPCGData interface
 
+	// ~Begin UPCGSettingsInterface interface
+	virtual UPCGSettings* GetSettings() { return this; }
+	virtual const UPCGSettings* GetSettings() const { return this; }
+	// ~End UPCGSettingsInterface interface
+
 	//~Begin UObject interface
+	virtual void PostLoad() override;
 	virtual void Serialize(FArchive& Ar) override;
+	virtual void PostSaveRoot(FObjectPostSaveRootContext ObjectSaveContext) override;
 	//~End UObject interface
 
 	// TODO: check if we need this to be virtual, we don't really need if we're always caching
@@ -108,21 +139,14 @@ public:
 	UPROPERTY(EditAnywhere, BlueprintReadWrite, Category = "Settings|Tags")
 	TSet<FString> TagsAppliedOnOutput;
 
-	UPROPERTY(EditAnywhere, BlueprintReadWrite, Category = Debug)
-	EPCGSettingsExecutionMode ExecutionMode = EPCGSettingsExecutionMode::Enabled;
-
 #if WITH_EDITORONLY_DATA
-	UPROPERTY(EditAnywhere, BlueprintReadWrite, Category = Debug, meta = (ShowOnlyInnerProperties))
-	FPCGDebugVisualizationSettings DebugSettings;
+	UPROPERTY()
+	EPCGSettingsExecutionMode ExecutionMode_DEPRECATED = EPCGSettingsExecutionMode::Enabled;
 
 	UPROPERTY(EditAnywhere, BlueprintReadWrite, Category = Determinism, meta = (ShowOnlyInnerProperties))
 	FPCGDeterminismSettings DeterminismSettings;
 #endif
 
-#if WITH_EDITOR
-	FOnPCGSettingsChanged OnSettingsChangedDelegate;
-#endif
-
 #if WITH_EDITORONLY_DATA
 	UPROPERTY(EditAnywhere, Category = AssetInfo, AssetRegistrySearchable)
 	FText Category;
@@ -166,6 +190,41 @@ private:
 	mutable FCriticalSection CacheLock;
 };
 
+UCLASS(BlueprintType, ClassGroup = (Procedural))
+class PCG_API UPCGSettingsInstance : public UPCGSettingsInterface
+{
+	GENERATED_BODY()
+
+public:
+	// ~Begin UPCGSettingsInterface
+	virtual UPCGSettings* GetSettings() { return Settings.Get(); }
+	virtual const UPCGSettings* GetSettings() const { return Settings.Get(); }
+	// ~End UPCGSettingsInterface
+
+	// ~Begin UObject interface
+	virtual void PostLoad() override;
+	virtual void BeginDestroy() override;
+
+#if WITH_EDITOR
+	virtual void PostEditChangeProperty(FPropertyChangedEvent& PropertyChangedEvent) override;
+#endif
+	// ~End UObject interface
+
+	void SetSettings(UPCGSettings* InSettings);
+
+#if WITH_EDITOR
+	void OnSettingsChanged(UPCGSettings* InSettings, EPCGChangeType ChangeType);
+#endif
+
+#if WITH_EDITORONLY_DATA
+	UPROPERTY(Transient, VisibleAnywhere, Category = Instance)
+	TObjectPtr<UPCGSettings> OriginalSettings = nullptr; // Transient just for display
+#endif
+
+	UPROPERTY(VisibleAnywhere, BlueprintReadOnly, Category = Instance, meta = (EditInline))
+	TObjectPtr<UPCGSettings> Settings = nullptr;
+};
+
 /** Trivial / Pass-through settings used for input/output nodes */
 UCLASS(BlueprintType, ClassGroup = (Procedural))
 class PCG_API UPCGTrivialSettings : public UPCGSettings
diff --git a/Engine/Plugins/Experimental/PCG/Source/PCGEditor/Private/PCGEditor.cpp b/Engine/Plugins/Experimental/PCG/Source/PCGEditor/Private/PCGEditor.cpp
--- a/Engine/Plugins/Experimental/PCG/Source/PCGEditor/Private/PCGEditor.cpp
+++ b/Engine/Plugins/Experimental/PCG/Source/PCGEditor/Private/PCGEditor.cpp
@@ -90,6 +90,10 @@ void FPCGEditor::Initialize(const EToolkitMode::Type InMode, const TSharedPtr<cl
 	PropertyDetailsWidget = PropertyEditorModule.CreateDetailView(DetailsViewArgs);
 	PropertyDetailsWidget->SetObject(PCGGraphBeingEdited);
 
+	IDetailsView* RawPropertyDetailsPtr = PropertyDetailsWidget.Get();
+	PropertyDetailsWidget->SetIsPropertyReadOnlyDelegate(FIsPropertyReadOnly::CreateRaw(this, &FPCGEditor::IsReadOnlyProperty, RawPropertyDetailsPtr));
+	PropertyDetailsWidget->SetIsPropertyVisibleDelegate(FIsPropertyVisible::CreateRaw(this, &FPCGEditor::IsVisibleProperty, RawPropertyDetailsPtr));
+
 	GraphEditorWidget = CreateGraphEditorWidget();
 	PaletteWidget = CreatePaletteWidget();
 	FindWidget = CreateFindWidget();
@@ -513,17 +517,17 @@ void FPCGEditor::OnDeterminismNodeTest()
 			if (const UPCGEditorGraphNodeBase* PCGEditorGraphNode = Cast<const UPCGEditorGraphNodeBase>(Object))
 			{
 				const UPCGNode* PCGNode = PCGEditorGraphNode->GetPCGNode();
-				check(PCGNode && PCGNode->DefaultSettings);
+				check(PCGNode && PCGNode->GetSettings());
 
 				TSharedPtr<FDeterminismTestResult> NodeResult = MakeShared<FDeterminismTestResult>();
 				NodeResult->TestResultTitle = PCGNode->GetNodeTitle();
 				NodeResult->TestResultName = PCGNode->GetName();
-				NodeResult->Seed = PCGNode->DefaultSettings->Seed;
+				NodeResult->Seed = PCGNode->GetSettings()->Seed;
 
-				if (PCGNode->DefaultSettings->DeterminismSettings.bNativeTests)
+				if (PCGNode->GetSettings()->DeterminismSettings.bNativeTests)
 				{
 					// If the settings has a native test suite
-					if (TFunction<bool()> NativeTestSuite = PCGDeterminismTests::FNativeTestRegistry::GetNativeTestFunction(PCGNode->DefaultSettings))
+					if (TFunction<bool()> NativeTestSuite = PCGDeterminismTests::FNativeTestRegistry::GetNativeTestFunction(PCGNode->GetSettings()))
 					{
 						FName NodeName(PCGNode->GetName());
 
@@ -548,9 +552,9 @@ void FPCGEditor::OnDeterminismNodeTest()
 				}
 
 				// Custom tests
-				if (PCGNode->DefaultSettings->DeterminismSettings.bUseBlueprintDeterminismTest)
+				if (PCGNode->GetSettings()->DeterminismSettings.bUseBlueprintDeterminismTest)
 				{
-					TSubclassOf<UPCGDeterminismTestBlueprintBase> Blueprint = PCGNode->DefaultSettings->DeterminismSettings.DeterminismTestBlueprint;
+					TSubclassOf<UPCGDeterminismTestBlueprintBase> Blueprint = PCGNode->GetSettings()->DeterminismSettings.DeterminismTestBlueprint;
 					Blueprint.GetDefaultObject()->ExecuteTest(PCGNode, *NodeResult);
 					FName BlueprintName(Blueprint->GetName());
 
@@ -805,10 +809,10 @@ void FPCGEditor::OnCollapseNodesInSubgraph()
 		TObjectPtr<UPCGNode> InputOutputNode = bIsInput ? PCGGraph->GetInputNode() : PCGGraph->GetOutputNode();
 		TObjectPtr<UPCGNode> NewInputOutputNode = bIsInput ? NewPCGGraph->GetInputNode() : NewPCGGraph->GetOutputNode();
 
-		TObjectPtr<UPCGGraphInputOutputSettings> InputOutputSettings = CastChecked<UPCGGraphInputOutputSettings>(InputOutputNode->DefaultSettings);
+		TObjectPtr<UPCGGraphInputOutputSettings> InputOutputSettings = CastChecked<UPCGGraphInputOutputSettings>(InputOutputNode->GetSettings());
 		if (InputOutputSettings->IsPinAdvanced(Pin))
 		{
-			TObjectPtr<UPCGGraphInputOutputSettings> NewInputOutputSettings = CastChecked<UPCGGraphInputOutputSettings>(NewInputOutputNode->DefaultSettings);
+			TObjectPtr<UPCGGraphInputOutputSettings> NewInputOutputSettings = CastChecked<UPCGGraphInputOutputSettings>(NewInputOutputNode->GetSettings());
 			NewInputOutputSettings->SetShowAdvancedPins(true);
 			NewInputOutputNode->UpdateAfterSettingsChangeDuringCreation();
 
@@ -825,7 +829,7 @@ void FPCGEditor::OnCollapseNodesInSubgraph()
 	auto CreateNewCustomPin = [&NewPCGGraph, &NameCollisionMapping](UPCGPin* Pin, bool bIsInput)
 	{
 		TObjectPtr<UPCGNode> NewInputOutputNode = bIsInput ? NewPCGGraph->GetInputNode() : NewPCGGraph->GetOutputNode();
-		TObjectPtr<UPCGGraphInputOutputSettings> NewInputOutputSettings = CastChecked<UPCGGraphInputOutputSettings>(NewInputOutputNode->DefaultSettings);
+		TObjectPtr<UPCGGraphInputOutputSettings> NewInputOutputSettings = CastChecked<UPCGGraphInputOutputSettings>(NewInputOutputNode->GetSettings());
 		FString NewName = Pin->Node->GetNodeTitle().ToString() + " " + Pin->Properties.Label.ToString();
 		if (NameCollisionMapping.Contains(NewName))
 		{
@@ -1062,7 +1066,7 @@ void FPCGEditor::OnExportNodes()
 		{
 			UPCGNode* PCGNode = PCGEditorGraphNode->GetPCGNode();
 			check(PCGNode);
-			Settings = PCGNode->DefaultSettings;
+			Settings = PCGNode->GetSettings();
 		}
 
 		if (!Settings)
@@ -1147,13 +1151,13 @@ bool FPCGEditor::IsExecutionModeActive(EPCGSettingsExecutionMode InExecutionMode
 				continue;
 			}
 
-			const UPCGSettings* PCGSettings = PCGNode->DefaultSettings;
-			if (!PCGSettings)
+			const UPCGSettingsInterface* PCGSettingsInterface = PCGNode->GetSettingsInterface();
+			if (!PCGSettingsInterface)
 			{
 				continue;
 			}
 
-			if (PCGSettings->ExecutionMode == InExecutionMode)
+			if (PCGSettingsInterface->ExecutionMode == InExecutionMode)
 			{
 				return true;
 			}
@@ -1184,17 +1188,17 @@ void FPCGEditor::OnSetExecutionMode(EPCGSettingsExecutionMode InExecutionMode)
 				continue;
 			}
 
-			UPCGSettings* PCGSettings = PCGNode->DefaultSettings;
-			if (!PCGSettings)
+			UPCGSettingsInterface* PCGSettingsInterface = PCGNode->GetSettingsInterface();
+			if (!PCGSettingsInterface)
 			{
 				continue;
 			}
 
-			if (PCGSettings->ExecutionMode != InExecutionMode)
+			if (PCGSettingsInterface->ExecutionMode != InExecutionMode)
 			{
-				PCGSettings->Modify();
-				PCGSettings->ExecutionMode = InExecutionMode;
-				PCGSettings->OnSettingsChangedDelegate.Broadcast(PCGSettings, EPCGChangeType::Settings);
+				PCGSettingsInterface->Modify();
+				PCGSettingsInterface->ExecutionMode = InExecutionMode;
+				PCGNode->OnNodeChangedDelegate.Broadcast(PCGNode, EPCGChangeType::Settings);
 			}
 		}
 	}
@@ -1696,7 +1700,14 @@ void FPCGEditor::OnSelectedNodesChanged(const TSet<UObject*>& NewSelection)
 		{
 			if (UPCGNode* PCGNode = PCGGraphNode->GetPCGNode())
 			{
-				SelectedObjects.Add(PCGNode->DefaultSettings);
+				if (PCGNode->IsInstance())
+				{
+					SelectedObjects.Add(Cast<UPCGSettingsInstance>(PCGNode->GetSettingsInterface()));
+				}
+				else
+				{
+					SelectedObjects.Add(PCGNode->GetSettings());
+				}
 			}
 		}
 		else if (UEdGraphNode* GraphNode = Cast<UEdGraphNode>(Object))
@@ -1772,6 +1783,54 @@ void FPCGEditor::JumpToDefinition(const UClass* Class) const
 	}
 }
 
+bool FPCGEditor::IsReadOnlyProperty(const FPropertyAndParent& InPropertyAndParent, IDetailsView* InDetailsView) const
+{
+	// Everything is writeable when not in an instance
+	if (!InDetailsView || 
+		InPropertyAndParent.ParentProperties.IsEmpty() || 
+		InPropertyAndParent.ParentProperties.Last()->GetFName() != GET_MEMBER_NAME_CHECKED(UPCGSettingsInstance, Settings))
+	{
+		return false;
+	}
+
+	const TArray<TWeakObjectPtr<UObject>>& SelectedObjects = InDetailsView->GetSelectedObjects();
+	for (const TWeakObjectPtr<UObject>& SelectedObject : SelectedObjects)
+	{
+		if (!SelectedObject.IsValid())
+		{
+			continue;
+		}
+		
+		if (UPCGSettingsInstance* Instance = Cast<UPCGSettingsInstance>(SelectedObject.Get()))
+		{
+			return true;
+		}
+	}
+
+	return false;
+}
+
+bool FPCGEditor::IsVisibleProperty(const FPropertyAndParent& InPropertyAndParent, IDetailsView* InDetailsView) const
+{
+	// Everything is visible when not in an instance
+	if (!InDetailsView ||
+		InPropertyAndParent.ParentProperties.IsEmpty() ||
+		InPropertyAndParent.ParentProperties.Last()->GetFName() != GET_MEMBER_NAME_CHECKED(UPCGSettingsInstance, Settings))
+	{
+		return true;
+	}
+
+	// Hide debug settings from the setting when showing the instance settings.
+	if (InPropertyAndParent.Property.GetFName() == GET_MEMBER_NAME_CHECKED(UPCGSettings, ExecutionMode) ||
+		(InPropertyAndParent.ParentProperties.Num() >= 2 && 
+			InPropertyAndParent.ParentProperties[InPropertyAndParent.ParentProperties.Num() - 2]->GetFName() == GET_MEMBER_NAME_CHECKED(UPCGSettings, DebugSettings)))
+	{
+		return false;
+	}
+
+	return true;
+}
+
 TSharedRef<SDockTab> FPCGEditor::SpawnTab_GraphEditor(const FSpawnTabArgs& Args)
 {
 	return SNew(SDockTab)
diff --git a/Engine/Plugins/Experimental/PCG/Source/PCGEditor/Private/PCGEditor.h b/Engine/Plugins/Experimental/PCG/Source/PCGEditor/Private/PCGEditor.h
--- a/Engine/Plugins/Experimental/PCG/Source/PCGEditor/Private/PCGEditor.h
+++ b/Engine/Plugins/Experimental/PCG/Source/PCGEditor/Private/PCGEditor.h
@@ -78,6 +78,9 @@ public:
 	virtual void InitToolMenuContext(FToolMenuContext& MenuContext) override;
 	// ~End FAssetEditorToolkit interface
 
+	// Methods used by schema actions
+	UEdGraphNode* AddNode(UPCGSettings* InSettings, bool bIsInstance, const FVector2D& InLocation);
+
 	FOnDebugObjectChanged OnDebugObjectChangedDelegate;
 	FOnInspectedNodeChanged OnInspectedNodeChangedDelegate;
 
@@ -229,6 +232,12 @@ private:
 	/** To be called everytime we need to replicate our extra nodes to the underlying PCGGraph */
 	void ReplicateExtraNodes() const;
 
+	/** Returns whether a property should be readonly (used for instances) */
+	bool IsReadOnlyProperty(const FPropertyAndParent& InPropertyAndParent, IDetailsView* InDetailsView) const;
+
+	/** Returns whether a property should be visible (used for instance vs. settings properties) */
+	bool IsVisibleProperty(const FPropertyAndParent& InPropertyAndParent, IDetailsView* InDetailsView) const;
+
 	TSharedRef<SDockTab> SpawnTab_GraphEditor(const FSpawnTabArgs& Args);
 	TSharedRef<SDockTab> SpawnTab_PropertyDetails(const FSpawnTabArgs& Args);
 	TSharedRef<SDockTab> SpawnTab_Palette(const FSpawnTabArgs& Args);
diff --git a/Engine/Plugins/Experimental/PCG/Source/PCGEditor/Private/PCGEditorGraphNodeBase.cpp b/Engine/Plugins/Experimental/PCG/Source/PCGEditor/Private/PCGEditorGraphNodeBase.cpp
--- a/Engine/Plugins/Experimental/PCG/Source/PCGEditor/Private/PCGEditorGraphNodeBase.cpp
+++ b/Engine/Plugins/Experimental/PCG/Source/PCGEditor/Private/PCGEditorGraphNodeBase.cpp
@@ -36,9 +36,9 @@ void UPCGEditorGraphNodeBase::Construct(UPCGNode* InPCGNode)
 	bCommentBubblePinned = InPCGNode->bCommentBubblePinned;
 	bCommentBubbleVisible = InPCGNode->bCommentBubbleVisible;
 
-	if (const UPCGSettings* PCGSettings = InPCGNode->DefaultSettings)
+	if (const UPCGSettingsInterface* PCGSettingsInterface = InPCGNode->GetSettingsInterface())
 	{
-		const ENodeEnabledState NewEnabledState = (PCGSettings->ExecutionMode == EPCGSettingsExecutionMode::Disabled) ? ENodeEnabledState::Disabled : ENodeEnabledState::Enabled;
+		const ENodeEnabledState NewEnabledState = (PCGSettingsInterface->ExecutionMode == EPCGSettingsExecutionMode::Disabled) ? ENodeEnabledState::Disabled : ENodeEnabledState::Enabled;
 		SetEnabledState(NewEnabledState);
 	}
 }
@@ -250,9 +250,9 @@ void UPCGEditorGraphNodeBase::OnNodeChanged(UPCGNode* InNode, EPCGChangeType Cha
 	{
 		if (!!(ChangeType & EPCGChangeType::Settings))
 		{
-			if (const UPCGSettings* PCGSettings = InNode->DefaultSettings)
+			if (const UPCGSettingsInterface* PCGSettingsInterface = InNode->GetSettingsInterface())
 			{
-				const ENodeEnabledState NewEnabledState = (PCGSettings->ExecutionMode == EPCGSettingsExecutionMode::Disabled) ? ENodeEnabledState::Disabled : ENodeEnabledState::Enabled;
+				const ENodeEnabledState NewEnabledState = (PCGSettingsInterface->ExecutionMode == EPCGSettingsExecutionMode::Disabled) ? ENodeEnabledState::Disabled : ENodeEnabledState::Enabled;
 				if (NewEnabledState != GetDesiredEnabledState())
 				{
 					SetEnabledState(NewEnabledState);
@@ -322,8 +322,10 @@ FLinearColor UPCGEditorGraphNodeBase::GetNodeTitleColor() const
 {
 	if (PCGNode)
 	{
-		const UPCGSettings* PCGSettings = PCGNode->DefaultSettings;
-		if (PCGSettings && PCGSettings->ExecutionMode == EPCGSettingsExecutionMode::Isolated)
+		const UPCGSettingsInterface* PCGSettingsInterface = PCGNode->GetSettingsInterface();
+		const UPCGSettings* PCGSettings = PCGSettingsInterface->GetSettings();
+
+		if (PCGSettingsInterface && PCGSettingsInterface->ExecutionMode == EPCGSettingsExecutionMode::Isolated)
 		{
 			return GetDefault<UPCGEditorSettings>()->IsolatedNodeColor;
 		}
@@ -334,10 +336,10 @@ FLinearColor UPCGEditorGraphNodeBase::GetNodeTitleColor() const
 		}
 		else if (PCGSettings)
 		{
-			FLinearColor SettingsColor = PCGNode->DefaultSettings->GetNodeTitleColor();
+			FLinearColor SettingsColor = PCGNode->GetSettings()->GetNodeTitleColor();
 			if (SettingsColor == FLinearColor::White)
 			{
-				SettingsColor = GetDefault<UPCGEditorSettings>()->GetColor(PCGNode->DefaultSettings);
+				SettingsColor = GetDefault<UPCGEditorSettings>()->GetColor(PCGNode->GetSettings());
 			}
 
 			if (SettingsColor != FLinearColor::White)
@@ -354,8 +356,8 @@ FLinearColor UPCGEditorGraphNodeBase::GetNodeBodyTintColor() const
 {
 	if (PCGNode)
 	{
-		const UPCGSettings* PCGSettings = PCGNode->DefaultSettings;
-		if (PCGSettings && PCGSettings->ExecutionMode == EPCGSettingsExecutionMode::Isolated)
+		const UPCGSettingsInterface* PCGSettingsInterface = PCGNode->GetSettingsInterface();
+		if (PCGSettingsInterface && PCGSettingsInterface->ExecutionMode == EPCGSettingsExecutionMode::Isolated)
 		{
 			return GetDefault<UPCGEditorSettings>()->IsolatedNodeColor;
 		}
@@ -500,7 +502,7 @@ UObject* UPCGEditorGraphNodeBase::GetJumpTargetForDoubleClick() const
 {
 	if (PCGNode)
 	{
-		if (UPCGSettings* Settings = PCGNode->DefaultSettings)
+		if (UPCGSettings* Settings = PCGNode->GetSettings())
 		{
 			return Settings->GetJumpTargetForDoubleClick();
 		}
diff --git a/Engine/Plugins/Experimental/PCG/Source/PCGEditor/Private/PCGEditorGraphSchema.cpp b/Engine/Plugins/Experimental/PCG/Source/PCGEditor/Private/PCGEditorGraphSchema.cpp
--- a/Engine/Plugins/Experimental/PCG/Source/PCGEditor/Private/PCGEditorGraphSchema.cpp
+++ b/Engine/Plugins/Experimental/PCG/Source/PCGEditor/Private/PCGEditorGraphSchema.cpp
@@ -3,22 +3,24 @@
 #include "PCGEditorGraphSchema.h"
 
 #include "Elements/PCGExecuteBlueprint.h"
+#include "PCGGraph.h"
+#include "PCGSettings.h"
+
 #include "PCGEditorCommon.h"
 #include "PCGEditorGraph.h"
 #include "PCGEditorGraphNodeBase.h"
 #include "PCGEditorGraphSchemaActions.h"
 #include "PCGEditorSettings.h"
-#include "PCGGraph.h"
-#include "PCGSettings.h"
+#include "PCGEditorUtils.h"
 
 #include "AssetRegistry/AssetRegistryModule.h"
 #include "Engine/Blueprint.h"
 #include "ScopedTransaction.h"
 #include "UObject/UObjectIterator.h"
-#include "PCGEditorUtils.h"
 
-#define LOCTEXT_NAMESPACE "PCGEditorGraphSchema"
+#include "SGraphPanel.h"
 
+#define LOCTEXT_NAMESPACE "PCGEditorGraphSchema"
 
 void UPCGEditorGraphSchema::GetPaletteActions(FGraphActionMenuBuilder& ActionMenuBuilder, const EPCGElementType InPCGElementTypeFilter) const
 {
@@ -36,7 +38,7 @@ void UPCGEditorGraphSchema::GetPaletteActions(FGraphActionMenuBuilder& ActionMen
 	}
 	if (!!(InPCGElementTypeFilter & EPCGElementType::Settings))
 	{
-		GetSettingsElementActions(ActionMenuBuilder);
+		GetSettingsElementActions(ActionMenuBuilder, /*bIsContextual=*/false);
 	}
 	if (!!(InPCGElementTypeFilter & EPCGElementType::Other))
 	{
@@ -51,7 +53,7 @@ void UPCGEditorGraphSchema::GetGraphContextActions(FGraphContextMenuBuilder& Con
 	GetNativeElementActions(ContextMenuBuilder);
 	GetSubgraphElementActions(ContextMenuBuilder);
 	GetBlueprintElementActions(ContextMenuBuilder);
-	GetSettingsElementActions(ContextMenuBuilder);
+	GetSettingsElementActions(ContextMenuBuilder, /*bIsContextual=*/true);
 	GetExtraElementActions(ContextMenuBuilder);
 }
 
@@ -255,7 +257,7 @@ void UPCGEditorGraphSchema::GetBlueprintElementActions(FGraphActionMenuBuilder&
 	}
 }
 
-void UPCGEditorGraphSchema::GetSettingsElementActions(FGraphActionMenuBuilder& ActionMenuBuilder) const
+void UPCGEditorGraphSchema::GetSettingsElementActions(FGraphActionMenuBuilder& ActionMenuBuilder, bool bIsContextual) const
 {
 	const FAssetRegistryModule& AssetRegistryModule = FModuleManager::LoadModuleChecked<FAssetRegistryModule>(TEXT("AssetRegistry"));
 
@@ -268,9 +270,26 @@ void UPCGEditorGraphSchema::GetSettingsElementActions(FGraphActionMenuBuilder& A
 		const FText Category = AssetData.GetTagValueRef<FText>(TEXT("Category"));
 		const FText Description = AssetData.GetTagValueRef<FText>(TEXT("Description"));
 
-		TSharedPtr<FPCGEditorGraphSchemaAction_NewSettingsElement> NewSettingsAction(new FPCGEditorGraphSchemaAction_NewSettingsElement(Category, MenuDesc, Description, 0));
-		NewSettingsAction->SettingsObjectPath = AssetData.GetSoftObjectPath();
-		ActionMenuBuilder.AddAction(NewSettingsAction);
+		if (!bIsContextual)
+		{
+			TSharedPtr<FPCGEditorGraphSchemaAction_NewSettingsElement> NewSettingsAction(new FPCGEditorGraphSchemaAction_NewSettingsElement(Category, MenuDesc, Description, 0));
+			NewSettingsAction->SettingsObjectPath = AssetData.GetSoftObjectPath();
+			ActionMenuBuilder.AddAction(NewSettingsAction);
+		}
+		else
+		{
+			const FText MenuAndSubCategory = FText::Join(LOCTEXT("MenuDelimiter", "|"), Category, MenuDesc);
+
+			TSharedPtr<FPCGEditorGraphSchemaAction_NewSettingsElement> NewSettingsActionCopy(new FPCGEditorGraphSchemaAction_NewSettingsElement(MenuAndSubCategory, LOCTEXT("ContextMenuCopySettings", "Copy"), Description, 0));
+			NewSettingsActionCopy->SettingsObjectPath = AssetData.GetSoftObjectPath();
+			NewSettingsActionCopy->Behavior = EPCGEditorNewSettingsBehavior::ForceCopy;
+			ActionMenuBuilder.AddAction(NewSettingsActionCopy);
+
+			TSharedPtr<FPCGEditorGraphSchemaAction_NewSettingsElement> NewSettingsActionInstance(new FPCGEditorGraphSchemaAction_NewSettingsElement(MenuAndSubCategory, LOCTEXT("ContextMenuCopySettings", "Instance"), Description, 0));
+			NewSettingsActionInstance->SettingsObjectPath = AssetData.GetSoftObjectPath();
+			NewSettingsActionInstance->Behavior = EPCGEditorNewSettingsBehavior::ForceInstance;
+			ActionMenuBuilder.AddAction(NewSettingsActionInstance);
+		}
 	}
 }
 
@@ -314,6 +333,9 @@ void UPCGEditorGraphSchema::DroppedAssetsOnGraph(const TArray<FAssetData>& Asset
 	constexpr float PositionOffsetIncrementY = 50.f;
 	UEdGraphPin* NullFromPin = nullptr;
 
+	TArray<FSoftObjectPath> SettingsPaths;
+	TArray<FVector2D> GraphPositions;
+
 	for (const FAssetData& AssetData : Assets)
 	{
 		if (const UObject* Asset = AssetData.GetAsset())
@@ -336,13 +358,25 @@ void UPCGEditorGraphSchema::DroppedAssetsOnGraph(const TArray<FAssetData>& Asset
 			}
 			else if (Asset->IsA<UPCGSettings>())
 			{
-				FPCGEditorGraphSchemaAction_NewSettingsElement NewNodeAction;
-				NewNodeAction.SettingsObjectPath = AssetData.GetSoftObjectPath();
-				NewNodeAction.PerformAction(Graph, NullFromPin, GraphPositionOffset);
+				// Delay creation so we can open a menu, once, if needed.
+				SettingsPaths.Add(AssetData.GetSoftObjectPath());
+				GraphPositions.Add(GraphPositionOffset);
 				GraphPositionOffset.Y += PositionOffsetIncrementY;
 			}
 		}
 	}
+
+	// If we've dragged settings assets, we might want to open a menu (ergo this call)
+	if (!SettingsPaths.IsEmpty())
+	{
+		UPCGEditorGraph* EditorGraph = CastChecked<UPCGEditorGraph>(Graph);
+		
+		TSharedPtr<SGraphEditor> GraphEditor = SGraphEditor::FindGraphEditorForGraph(EditorGraph);
+		const FVector2D MouseCursorLocation = FSlateApplication::Get().GetCursorPos();
+
+		check(SettingsPaths.Num() == GraphPositions.Num());
+		FPCGEditorGraphSchemaAction_NewSettingsElement::MakeSettingsNodesOrContextualMenu(GraphEditor->GetGraphPanel()->AsShared(), MouseCursorLocation, Graph, SettingsPaths, GraphPositions, /*bSelectNewNodes=*/true);
+	}
 }
 
 void UPCGEditorGraphSchema::GetAssetsGraphHoverMessage(const TArray<FAssetData>& Assets, const UEdGraph* /*HoverGraph*/, FString& OutTooltipText, bool& OutOkIcon) const
diff --git a/Engine/Plugins/Experimental/PCG/Source/PCGEditor/Private/PCGEditorGraphSchema.h b/Engine/Plugins/Experimental/PCG/Source/PCGEditor/Private/PCGEditorGraphSchema.h
--- a/Engine/Plugins/Experimental/PCG/Source/PCGEditor/Private/PCGEditorGraphSchema.h
+++ b/Engine/Plugins/Experimental/PCG/Source/PCGEditor/Private/PCGEditorGraphSchema.h
@@ -34,7 +34,7 @@ private:
 	void GetNativeElementActions(FGraphActionMenuBuilder& ActionMenuBuilder) const;
 	void GetBlueprintElementActions(FGraphActionMenuBuilder& ActionMenuBuilder) const;
 	void GetSubgraphElementActions(FGraphActionMenuBuilder& ActionMenuBuilder) const;
-	void GetSettingsElementActions(FGraphActionMenuBuilder& ActionMenuBuilder) const;
+	void GetSettingsElementActions(FGraphActionMenuBuilder& ActionMenuBuilder, bool bIsContextual) const;
 	void GetExtraElementActions(FGraphActionMenuBuilder& ActionMenuBuilder) const;
 };
 
diff --git a/Engine/Plugins/Experimental/PCG/Source/PCGEditor/Private/PCGEditorGraphSchemaActions.cpp b/Engine/Plugins/Experimental/PCG/Source/PCGEditor/Private/PCGEditorGraphSchemaActions.cpp
--- a/Engine/Plugins/Experimental/PCG/Source/PCGEditor/Private/PCGEditorGraphSchemaActions.cpp
+++ b/Engine/Plugins/Experimental/PCG/Source/PCGEditor/Private/PCGEditorGraphSchemaActions.cpp
@@ -82,12 +82,128 @@ UEdGraphNode* FPCGEditorGraphSchemaAction_NewSettingsElement::PerformAction(UEdG
 		return nullptr;
 	}
 
+	bool bCreateInstance = false;
+
+	if(Behavior != EPCGEditorNewSettingsBehavior::Normal)
+	{
+		bCreateInstance = (Behavior == EPCGEditorNewSettingsBehavior::ForceInstance);
+	}
+	else
+	{
+		FModifierKeysState ModifierKeys = FSlateApplication::Get().GetModifierKeys();
+		bCreateInstance = ModifierKeys.IsAltDown();
+	}
+
+	return MakeSettingsNode(EditorGraph, FromPin, Settings, Location, bSelectNewNode, bCreateInstance);
+}
+
+void FPCGEditorGraphSchemaAction_NewSettingsElement::MakeSettingsNodesOrContextualMenu(const TSharedRef<class SWidget>& InPanel, FVector2D InScreenPosition, UEdGraph* InGraph, const TArray<FSoftObjectPath>& InSettingsPaths, const TArray<FVector2D>& InLocations, bool bInSelectNewNodes)
+{
+	UPCGEditorGraph* EditorGraph = Cast<UPCGEditorGraph>(InGraph);
+
+	if (!EditorGraph)
+	{
+		return;
+	}
+
+	TArray<UPCGSettings*> Settings;
+	TArray<FVector2D> SettingsLocations;
+	check(InSettingsPaths.Num() == InLocations.Num());
+
+	for(int32 PathIndex = 0; PathIndex < InSettingsPaths.Num(); ++PathIndex)
+	{
+		const FSoftObjectPath& SettingsPath = InSettingsPaths[PathIndex];
+		if (UPCGSettings* LoadedSettings = Cast<UPCGSettings>(SettingsPath.TryLoad()))
+		{
+			Settings.Add(LoadedSettings);
+			SettingsLocations.Add(InLocations[PathIndex]);
+		}
+	}
+
+	FModifierKeysState ModifierKeys = FSlateApplication::Get().GetModifierKeys();
+	const bool bModifiedKeysActive = ModifierKeys.IsControlDown() || ModifierKeys.IsAltDown();
+
+	if (bModifiedKeysActive)
+	{
+		MakeSettingsNodes(EditorGraph, Settings, SettingsLocations, bInSelectNewNodes, ModifierKeys.IsAltDown());
+	}
+	else if(!Settings.IsEmpty())
+	{
+		FMenuBuilder MenuBuilder(true, nullptr);
+		const FText SettingsTextName = ((Settings.Num() == 1) ? FText::FromName(Settings[0]->GetFName()) : LOCTEXT("ManySettings", "Settings"));
+
+		MenuBuilder.BeginSection("SettingsDroppedOn", SettingsTextName);
+
+		MenuBuilder.AddMenuEntry(
+			FText::Format(LOCTEXT("CopySettings", "Copy {0}"), SettingsTextName),
+			FText::Format(LOCTEXT("CopySettingsToolTip", "Copy the settings asset {0}, and keeps no reference to the original\n(Ctrl-drop to automatically create a copy)"), SettingsTextName),
+			FSlateIcon(),
+			FUIAction(
+				FExecuteAction::CreateStatic(&FPCGEditorGraphSchemaAction_NewSettingsElement::MakeSettingsNodes, EditorGraph, Settings, SettingsLocations, bInSelectNewNodes, false), 
+				FCanExecuteAction()
+			)
+		);
+
+		MenuBuilder.AddMenuEntry(
+			FText::Format(LOCTEXT("InstanceSettings", "Instance {0}"), SettingsTextName),
+			FText::Format(LOCTEXT("InstanceSettingsToolTip", "Instance the settings asset {0}\n(Alt-drop to automatically create an instance)"), SettingsTextName),
+			FSlateIcon(),
+			FUIAction(
+				FExecuteAction::CreateStatic(&FPCGEditorGraphSchemaAction_NewSettingsElement::MakeSettingsNodes, EditorGraph, Settings, SettingsLocations, bInSelectNewNodes, true), 
+				FCanExecuteAction()
+			)
+		);
+
+		TSharedRef<SWidget> PanelWidget = InPanel;
+		// Show dialog to choose getter vs setter
+		FSlateApplication::Get().PushMenu(
+			PanelWidget,
+			FWidgetPath(),
+			MenuBuilder.MakeWidget(),
+			InScreenPosition,
+			FPopupTransitionEffect(FPopupTransitionEffect::ContextMenu)
+		);
+
+		MenuBuilder.EndSection();
+	}
+}
+
+void FPCGEditorGraphSchemaAction_NewSettingsElement::MakeSettingsNodes(UPCGEditorGraph* InEditorGraph, TArray<UPCGSettings*> InSettings, TArray<FVector2D> InSettingsLocations, bool bInSelectNewNodes, bool bInCreateInstances)
+{
+	check(InSettings.Num() == InSettingsLocations.Num());
+	for (int32 SettingsIndex = 0; SettingsIndex < InSettings.Num(); ++SettingsIndex)
+	{
+		FPCGEditorGraphSchemaAction_NewSettingsElement::MakeSettingsNode(InEditorGraph, nullptr, InSettings[SettingsIndex], InSettingsLocations[SettingsIndex], bInSelectNewNodes, bInCreateInstances);
+	}
+}
+
+UEdGraphNode* FPCGEditorGraphSchemaAction_NewSettingsElement::MakeSettingsNode(UPCGEditorGraph* InEditorGraph, UEdGraphPin* InFromPin, UPCGSettings* InSettings, FVector2D InLocation, bool bInSelectNewNode, bool bInCreateInstance)
+{
+	if (!InEditorGraph || !InSettings)
+	{
+		return nullptr;
+	}
+
+	UPCGGraph* PCGGraph = InEditorGraph->GetPCGGraph();
+	if (!PCGGraph)
+	{
+		return nullptr;
+	}
+
 	const FScopedTransaction Transaction(*FPCGEditorCommon::ContextIdentifier, LOCTEXT("PCGEditorNewSettingsElement", "PCG Editor: New Settings Element"), nullptr);
-	EditorGraph->Modify();
+	InEditorGraph->Modify();
 
 	UPCGNode* NewPCGNode = nullptr;
-	UPCGSettings* NewSettings = nullptr;
-	NewPCGNode = PCGGraph->AddNodeCopy(Settings, NewSettings);
+
+	if (bInCreateInstance)
+	{
+		NewPCGNode = PCGGraph->AddNodeInstance(InSettings);
+	}
+	else
+	{
+		UPCGSettings* NewSettings = nullptr;
+		NewPCGNode = PCGGraph->AddNodeCopy(InSettings, NewSettings);
+	}
 
 	if (!NewPCGNode)
 	{
@@ -95,19 +211,19 @@ UEdGraphNode* FPCGEditorGraphSchemaAction_NewSettingsElement::PerformAction(UEdG
 		return nullptr;
 	}
 
-	FGraphNodeCreator<UPCGEditorGraphNode> NodeCreator(*EditorGraph);
-	UPCGEditorGraphNode* NewNode = NodeCreator.CreateUserInvokedNode(bSelectNewNode);
+	FGraphNodeCreator<UPCGEditorGraphNode> NodeCreator(*InEditorGraph);
+	UPCGEditorGraphNode* NewNode = NodeCreator.CreateUserInvokedNode(bInSelectNewNode);
 	NewNode->Construct(NewPCGNode);
-	NewNode->NodePosX = Location.X;
-	NewNode->NodePosY = Location.Y;
+	NewNode->NodePosX = InLocation.X;
+	NewNode->NodePosY = InLocation.Y;
 	NodeCreator.Finalize();
 
-	NewPCGNode->PositionX = Location.X;
-	NewPCGNode->PositionY = Location.Y;
+	NewPCGNode->PositionX = InLocation.X;
+	NewPCGNode->PositionY = InLocation.Y;
 
-	if (FromPin)
+	if (InFromPin)
 	{
-		NewNode->AutowireNewNode(FromPin);
+		NewNode->AutowireNewNode(InFromPin);
 	}
 
 	return NewNode;
diff --git a/Engine/Plugins/Experimental/PCG/Source/PCGEditor/Private/PCGEditorGraphSchemaActions.h b/Engine/Plugins/Experimental/PCG/Source/PCGEditor/Private/PCGEditorGraphSchemaActions.h
--- a/Engine/Plugins/Experimental/PCG/Source/PCGEditor/Private/PCGEditorGraphSchemaActions.h
+++ b/Engine/Plugins/Experimental/PCG/Source/PCGEditor/Private/PCGEditorGraphSchemaActions.h
@@ -12,6 +12,15 @@ class UEdGraph;
 class UEdGraphPin;
 class UPCGBlueprintElement;
 class UPCGSettings;
+class UPCGEditorGraph;
+
+UENUM()
+enum class EPCGEditorNewSettingsBehavior : uint8
+{
+	Normal = 0,
+	ForceCopy,
+	ForceInstance
+};
 
 USTRUCT()
 struct FPCGEditorGraphSchemaAction_NewNativeElement : public FEdGraphSchemaAction
@@ -40,9 +49,16 @@ struct FPCGEditorGraphSchemaAction_NewSettingsElement : public FEdGraphSchemaAct
 	UPROPERTY()
 	FSoftObjectPath SettingsObjectPath;
 
+	UPROPERTY()
+	EPCGEditorNewSettingsBehavior Behavior = EPCGEditorNewSettingsBehavior::Normal;
+
 	//~ Begin FEdGraphSchemaAction Interface
 	virtual UEdGraphNode* PerformAction(UEdGraph* ParentGraph, UEdGraphPin* FromPin, const FVector2D Location, bool bSelectNewNode = true) override;
 	//~ End FEdGraphSchemaAction Interface
+
+	static void MakeSettingsNodesOrContextualMenu(const TSharedRef<class SWidget>& InPanel, FVector2D InScreenPosition, UEdGraph* InGraph, const TArray<FSoftObjectPath>& InSettingsPaths, const TArray<FVector2D>& InLocations, bool bInSelectNewNodes);
+	static void MakeSettingsNodes(UPCGEditorGraph* InEditorGraph, TArray<UPCGSettings*> InSettings, TArray<FVector2D> InSettingsLocations, bool bInSelectNewNodes, bool bInCreateInstances);
+	static UEdGraphNode* MakeSettingsNode(UPCGEditorGraph* InEditorGraph, UEdGraphPin* InFromPin, UPCGSettings* InSettings, FVector2D InLocation, bool bInSelectNewNode, bool bInCreateInstance);
 };
 
 USTRUCT()
diff --git a/Engine/Plugins/Experimental/PCG/Source/PCGEditor/Private/PCGSettingsDragDropAction.cpp b/Engine/Plugins/Experimental/PCG/Source/PCGEditor/Private/PCGSettingsDragDropAction.cpp
new file mode 100644
--- /dev/null
+++ b/Engine/Plugins/Experimental/PCG/Source/PCGEditor/Private/PCGSettingsDragDropAction.cpp
@@ -0,0 +1,25 @@
+// Copyright Epic Games, Inc. All Rights Reserved.
+
+#include "PCGSettingsDragDropAction.h"
+
+#include "PCGEditorGraph.h"
+#include "PCGEditorGraphSchema.h"
+#include "PCGEditorGraphSchemaActions.h"
+
+FReply FPCGSettingsDragDropAction::DroppedOnPanel(const TSharedRef<class SWidget>& Panel, FVector2D ScreenPosition, FVector2D GraphPosition, UEdGraph& Graph)
+{
+	if (!Graph.GetSchema()->IsA<UPCGEditorGraphSchema>())
+	{
+		return FReply::Unhandled();
+	}
+
+	UPCGEditorGraph* EditorGraph = Cast<UPCGEditorGraph>(&Graph);
+	if (!ensure(EditorGraph))
+	{
+		return FReply::Unhandled();
+	}
+
+	FPCGEditorGraphSchemaAction_NewSettingsElement::MakeSettingsNodesOrContextualMenu(Panel, ScreenPosition, &Graph, { SettingsObjectPath }, { GraphPosition }, /*bSelectNewNodes=*/true);
+
+	return FReply::Handled();
+}
\ No newline at end of file
diff --git a/Engine/Plugins/Experimental/PCG/Source/PCGEditor/Private/PCGSettingsDragDropAction.h b/Engine/Plugins/Experimental/PCG/Source/PCGEditor/Private/PCGSettingsDragDropAction.h
new file mode 100644
--- /dev/null
+++ b/Engine/Plugins/Experimental/PCG/Source/PCGEditor/Private/PCGSettingsDragDropAction.h
@@ -0,0 +1,34 @@
+// Copyright Epic Games, Inc. All Rights Reserved.
+
+#pragma once
+
+#include "GraphEditorDragDropAction.h"
+#include "Input/DragAndDrop.h"
+#include "Input/Reply.h"
+
+class UPCGEditorGraph;
+class UPCGSettings;
+
+class FPCGSettingsDragDropAction : public FGraphSchemaActionDragDropAction
+{
+public:
+	DRAG_DROP_OPERATOR_TYPE(FPCGSettingsDragDropAction, FGraphSchemaActionDragDropAction)
+
+	// ~Begin FGraphEditorDragDropAction interface
+	virtual FReply DroppedOnPanel(const TSharedRef<class SWidget>& Panel, FVector2D ScreenPosition, FVector2D GraphPosition, UEdGraph& Graph) override;
+	// ~End FGraphEditorGraphDropAction interface
+
+	static TSharedRef<FPCGSettingsDragDropAction> New(TSharedPtr<FEdGraphSchemaAction> InAction, const FSoftObjectPath& InSettingsObjectPath)
+	{
+		TSharedRef<FPCGSettingsDragDropAction> Operation = MakeShareable(new FPCGSettingsDragDropAction);
+		Operation->SourceAction = InAction;
+		Operation->SettingsObjectPath = InSettingsObjectPath;
+		Operation->Construct();
+		return Operation;
+	}
+
+protected:
+	FPCGSettingsDragDropAction() = default;
+
+	FSoftObjectPath SettingsObjectPath;
+};
\ No newline at end of file
diff --git a/Engine/Plugins/Experimental/PCG/Source/PCGEditor/Private/SPCGEditorGraphNode.cpp b/Engine/Plugins/Experimental/PCG/Source/PCGEditor/Private/SPCGEditorGraphNode.cpp
--- a/Engine/Plugins/Experimental/PCG/Source/PCGEditor/Private/SPCGEditorGraphNode.cpp
+++ b/Engine/Plugins/Experimental/PCG/Source/PCGEditor/Private/SPCGEditorGraphNode.cpp
@@ -57,7 +57,7 @@ void SPCGEditorGraphNode::GetOverlayBrushes(bool bSelected, const FVector2D Widg
 
 	if (const UPCGNode* PCGNode = PCGEditorGraphNode->GetPCGNode())
 	{
-		if (PCGNode->DefaultSettings && PCGNode->DefaultSettings->ExecutionMode == EPCGSettingsExecutionMode::Debug)
+		if (PCGNode->GetSettingsInterface() && PCGNode->GetSettingsInterface()->ExecutionMode == EPCGSettingsExecutionMode::Debug)
 		{
 			FOverlayBrushInfo BrushInfo;
 			BrushInfo.Brush = DebugBrush;
diff --git a/Engine/Plugins/Experimental/PCG/Source/PCGEditor/Private/SPCGEditorGraphNodePalette.cpp b/Engine/Plugins/Experimental/PCG/Source/PCGEditor/Private/SPCGEditorGraphNodePalette.cpp
--- a/Engine/Plugins/Experimental/PCG/Source/PCGEditor/Private/SPCGEditorGraphNodePalette.cpp
+++ b/Engine/Plugins/Experimental/PCG/Source/PCGEditor/Private/SPCGEditorGraphNodePalette.cpp
@@ -2,15 +2,18 @@
 
 #include "SPCGEditorGraphNodePalette.h"
 
+#include "PCGGraph.h"
 #include "Elements/PCGExecuteBlueprint.h"
+
 #include "PCGEditorGraphSchema.h"
+#include "PCGEditorGraphSchemaActions.h"
 #include "PCGEditorUtils.h"
-#include "PCGGraph.h"
+#include "PCGSettingsDragDropAction.h"
 
-#include "AssetRegistry/AssetRegistryModule.h"
-#include "Modules/ModuleManager.h"
 #include "SEnumCombo.h"
 #include "SGraphActionMenu.h"
+#include "AssetRegistry/AssetRegistryModule.h"
+#include "Modules/ModuleManager.h"
 
 #define LOCTEXT_NAMESPACE "SPCGEditorGraphNodePalette"
 
@@ -102,6 +105,21 @@ TSharedRef<SWidget> SPCGEditorGraphNodePalette::OnCreateWidgetForAction(FCreateW
 	return SNew(SPCGEditorGraphNodePaletteItem, InCreateData);
 }
 
+FReply SPCGEditorGraphNodePalette::OnActionDragged(const TArray<TSharedPtr<FEdGraphSchemaAction>>& InActions, const FPointerEvent& MouseEvent)
+{
+	if (InActions.Num() > 0 && InActions[0].IsValid())
+	{
+		TSharedPtr<FEdGraphSchemaAction> InAction = InActions[0];
+		if (InAction->GetTypeId() == FPCGEditorGraphSchemaAction_NewSettingsElement::StaticGetTypeId())
+		{
+			FPCGEditorGraphSchemaAction_NewSettingsElement* SettingsAction = static_cast<FPCGEditorGraphSchemaAction_NewSettingsElement*>(InAction.Get());
+			return FReply::Handled().BeginDragDrop(FPCGSettingsDragDropAction::New(InAction, SettingsAction->SettingsObjectPath));
+		}
+	}
+
+	return SGraphPalette::OnActionDragged(InActions, MouseEvent);
+}
+
 void SPCGEditorGraphNodePalette::CollectAllActions(FGraphActionListBuilderBase& OutAllActions)
 {
 	const UPCGEditorGraphSchema* PCGSchema = GetDefault<UPCGEditorGraphSchema>();
diff --git a/Engine/Plugins/Experimental/PCG/Source/PCGEditor/Private/SPCGEditorGraphNodePalette.h b/Engine/Plugins/Experimental/PCG/Source/PCGEditor/Private/SPCGEditorGraphNodePalette.h
--- a/Engine/Plugins/Experimental/PCG/Source/PCGEditor/Private/SPCGEditorGraphNodePalette.h
+++ b/Engine/Plugins/Experimental/PCG/Source/PCGEditor/Private/SPCGEditorGraphNodePalette.h
@@ -33,6 +33,7 @@ public:
 protected:
 	//~ Begin SGraphPalette Interface
 	virtual TSharedRef<SWidget> OnCreateWidgetForAction(FCreateWidgetForActionData* const InCreateData) override;
+	virtual FReply OnActionDragged(const TArray<TSharedPtr<FEdGraphSchemaAction>>& InActions, const FPointerEvent& MouseEvent) override;
 	virtual void CollectAllActions(FGraphActionListBuilderBase& OutAllActions) override;
 	//~ End SGraphPalette Interface
 
diff --git a/Engine/Plugins/Experimental/PCG/Source/PCGEditor/Private/SPCGEditorGraphProfilingView.cpp b/Engine/Plugins/Experimental/PCG/Source/PCGEditor/Private/SPCGEditorGraphProfilingView.cpp
--- a/Engine/Plugins/Experimental/PCG/Source/PCGEditor/Private/SPCGEditorGraphProfilingView.cpp
+++ b/Engine/Plugins/Experimental/PCG/Source/PCGEditor/Private/SPCGEditorGraphProfilingView.cpp
@@ -303,7 +303,7 @@ FReply SPCGEditorGraphProfilingView::ResetTimers()
 		{
 			if (UPCGNode* PCGNode = PCGEditorNode->GetPCGNode())
 			{
-				if (UPCGSettings* Settings = PCGNode->DefaultSettings)
+				if (UPCGSettings* Settings = PCGNode->GetSettings())
 				{
 					if (IPCGElement* Element = Settings->GetElement().Get())
 					{
@@ -351,7 +351,7 @@ FReply SPCGEditorGraphProfilingView::Refresh()
 			{
 				ListViewItem->Name = ListViewItem->PCGNode->GetNodeTitle();
 
-				if (const UPCGSettings* Settings = ListViewItem->PCGNode->DefaultSettings)
+				if (const UPCGSettings* Settings = ListViewItem->PCGNode->GetSettings())
 				{
 					if (const IPCGElement* Element = Settings->GetElement().Get())
 					{
