Tentative fix for non-unity build

#rnx

diff --git a/Engine/Plugins/Experimental/PCG/Source/PCGEditor/Private/PCGEditor.cpp b/Engine/Plugins/Experimental/PCG/Source/PCGEditor/Private/PCGEditor.cpp
--- a/Engine/Plugins/Experimental/PCG/Source/PCGEditor/Private/PCGEditor.cpp
+++ b/Engine/Plugins/Experimental/PCG/Source/PCGEditor/Private/PCGEditor.cpp
@@ -51,6 +51,7 @@
 #include "IDetailsView.h"
 #include "Modules/ModuleManager.h"
 #include "Preferences/UnrealEdOptions.h"
+#include "PropertyEditorDelegates.h"
 #include "PropertyEditorModule.h"
 #include "SNodePanel.h"
 #include "Misc/ITransaction.h"
diff --git a/Engine/Plugins/Experimental/PCG/Source/PCGEditor/Private/PCGEditor.h b/Engine/Plugins/Experimental/PCG/Source/PCGEditor/Private/PCGEditor.h
--- a/Engine/Plugins/Experimental/PCG/Source/PCGEditor/Private/PCGEditor.h
+++ b/Engine/Plugins/Experimental/PCG/Source/PCGEditor/Private/PCGEditor.h
@@ -8,6 +8,7 @@
 #include "EditorUndoClient.h"
 #include "Toolkits/AssetEditorToolkit.h"
 
+struct FPropertyAndParent;
 class FUICommandList;
 class IDetailsView;
 class SGraphEditor;
diff --git a/Engine/Plugins/Experimental/PCG/Source/PCGEditor/Private/PCGEditorGraphSchemaActions.cpp b/Engine/Plugins/Experimental/PCG/Source/PCGEditor/Private/PCGEditorGraphSchemaActions.cpp
--- a/Engine/Plugins/Experimental/PCG/Source/PCGEditor/Private/PCGEditorGraphSchemaActions.cpp
+++ b/Engine/Plugins/Experimental/PCG/Source/PCGEditor/Private/PCGEditorGraphSchemaActions.cpp
@@ -11,11 +11,15 @@
 #include "PCGSettings.h"
 #include "PCGSubgraph.h"
 
-#include "EdGraph/EdGraph.h"
-#include "Engine/Blueprint.h"
-#include "ScopedTransaction.h"
 #include "EdGraphNode_Comment.h"
 #include "GraphEditor.h"
+#include "ScopedTransaction.h"
+#include "EdGraph/EdGraph.h"
+#include "Engine/Blueprint.h"
+#include "Framework/Application/MenuStack.h"
+#include "Framework/Application/SlateApplication.h"
+#include "Framework/MultiBox/MultiBoxBuilder.h"
+#include "Layout/WidgetPath.h"
 
 #define LOCTEXT_NAMESPACE "PCGEditorGraphSchemaActions"
 
