Tentative fix for non-unity build take 2 (missing include)

#rnx

diff --git a/Engine/Plugins/Experimental/PCG/Source/PCGEditor/Private/PCGEditorGraphSchema.cpp b/Engine/Plugins/Experimental/PCG/Source/PCGEditor/Private/PCGEditorGraphSchema.cpp
--- a/Engine/Plugins/Experimental/PCG/Source/PCGEditor/Private/PCGEditorGraphSchema.cpp
+++ b/Engine/Plugins/Experimental/PCG/Source/PCGEditor/Private/PCGEditorGraphSchema.cpp
@@ -15,6 +15,7 @@
 
 #include "AssetRegistry/AssetRegistryModule.h"
 #include "Engine/Blueprint.h"
+#include "Framework/Application/SlateApplication.h"
 #include "ScopedTransaction.h"
 #include "UObject/UObjectIterator.h"
 
