[PCG] Fixed a crash caused by missed type information in pcg schema actions.

diff --git a/Engine/Plugins/Experimental/PCG/Source/PCGEditor/Private/PCGEditorGraphSchemaActions.h b/Engine/Plugins/Experimental/PCG/Source/PCGEditor/Private/PCGEditorGraphSchemaActions.h
--- a/Engine/Plugins/Experimental/PCG/Source/PCGEditor/Private/PCGEditorGraphSchemaActions.h
+++ b/Engine/Plugins/Experimental/PCG/Source/PCGEditor/Private/PCGEditorGraphSchemaActions.h
@@ -30,10 +30,18 @@ struct FPCGEditorGraphSchemaAction_NewNativeElement : public FEdGraphSchemaActio
 	// Inherit the base class's constructors
 	using FEdGraphSchemaAction::FEdGraphSchemaAction;
 
+	// Simple type info
+	static FName StaticGetTypeId()
+	{
+		static FName Type("FPCGEditorGraphSchemaAction_NewNativeElement");
+		return Type;
+	}
+	
 	UPROPERTY()
 	TSubclassOf<UPCGSettings> SettingsClass;
 
 	//~ Begin FEdGraphSchemaAction Interface
+	virtual FName GetTypeId() const override { return StaticGetTypeId(); }
 	virtual UEdGraphNode* PerformAction(UEdGraph* ParentGraph, UEdGraphPin* FromPin, const FVector2D Location, bool bSelectNewNode = true) override;
 	//~ End FEdGraphSchemaAction Interface
 };
@@ -45,6 +53,13 @@ struct FPCGEditorGraphSchemaAction_NewSettingsElement : public FEdGraphSchemaAct
 
 	// Inherit the base class's constructors
 	using FEdGraphSchemaAction::FEdGraphSchemaAction;
+	
+	// Simple type info
+	static FName StaticGetTypeId()
+	{
+		static FName Type("FPCGEditorGraphSchemaAction_NewSettingsElement");
+		return Type;
+	}
 
 	UPROPERTY()
 	FSoftObjectPath SettingsObjectPath;
@@ -53,6 +68,7 @@ struct FPCGEditorGraphSchemaAction_NewSettingsElement : public FEdGraphSchemaAct
 	EPCGEditorNewSettingsBehavior Behavior = EPCGEditorNewSettingsBehavior::Normal;
 
 	//~ Begin FEdGraphSchemaAction Interface
+	virtual FName GetTypeId() const override { return StaticGetTypeId(); }
 	virtual UEdGraphNode* PerformAction(UEdGraph* ParentGraph, UEdGraphPin* FromPin, const FVector2D Location, bool bSelectNewNode = true) override;
 	//~ End FEdGraphSchemaAction Interface
 
@@ -68,11 +84,19 @@ struct FPCGEditorGraphSchemaAction_NewBlueprintElement : public FEdGraphSchemaAc
 
 	// Inherit the base class's constructors
 	using FEdGraphSchemaAction::FEdGraphSchemaAction;
+	
+	// Simple type info
+	static FName StaticGetTypeId()
+	{
+		static FName Type("FPCGEditorGraphSchemaAction_NewBlueprintElement");
+		return Type;
+	}
 
 	UPROPERTY()
 	FSoftClassPath BlueprintClassPath;
 
 	//~ Begin FEdGraphSchemaAction Interface
+	virtual FName GetTypeId() const override { return StaticGetTypeId(); }
 	virtual UEdGraphNode* PerformAction(UEdGraph* ParentGraph, UEdGraphPin* FromPin, const FVector2D Location, bool bSelectNewNode = true) override;
 	//~ End FEdGraphSchemaAction Interface
 };
@@ -85,10 +109,18 @@ struct FPCGEditorGraphSchemaAction_NewSubgraphElement : public FEdGraphSchemaAct
 	// Inherit the base class's constructors
 	using FEdGraphSchemaAction::FEdGraphSchemaAction;
 
+	// Simple type info
+	static FName StaticGetTypeId()
+	{
+		static FName Type("FPCGEditorGraphSchemaAction_NewSubgraphElement");
+		return Type;
+	}
+
 	UPROPERTY()
 	FSoftObjectPath SubgraphObjectPath;
 
 	//~ Begin FEdGraphSchemaAction Interface
+	virtual FName GetTypeId() const override { return StaticGetTypeId(); }
 	virtual UEdGraphNode* PerformAction(UEdGraph* ParentGraph, UEdGraphPin* FromPin, const FVector2D Location, bool bSelectNewNode = true) override;
 	//~ End FEdGraphSchemaAction Interface
 };
@@ -103,10 +135,14 @@ struct FPCGEditorGraphSchemaAction_NewComment : public FEdGraphSchemaAction
 	using FEdGraphSchemaAction::FEdGraphSchemaAction;
 
 	// Simple type info
-	static FName StaticGetTypeId() { static FName Type("FPCGEditorGraphSchemaAction_NewComment"); return Type; }
-	virtual FName GetTypeId() const override { return StaticGetTypeId(); }
+	static FName StaticGetTypeId()
+	{
+		static FName Type("FPCGEditorGraphSchemaAction_NewComment");
+		return Type;
+	}
 
 	// FEdGraphSchemaAction interface
+	virtual FName GetTypeId() const override { return StaticGetTypeId(); }
 	virtual UEdGraphNode* PerformAction(class UEdGraph* ParentGraph, UEdGraphPin* FromPin, const FVector2D Location, bool bSelectNewNode = true) override;
 	// End of FEdGraphSchemaAction interface
 };
