Fixed static analysis warning

#rnx

diff --git a/Engine/Plugins/Experimental/PCG/Source/PCGEditor/Private/PCGEditorGraphNodeBase.cpp b/Engine/Plugins/Experimental/PCG/Source/PCGEditor/Private/PCGEditorGraphNodeBase.cpp
--- a/Engine/Plugins/Experimental/PCG/Source/PCGEditor/Private/PCGEditorGraphNodeBase.cpp
+++ b/Engine/Plugins/Experimental/PCG/Source/PCGEditor/Private/PCGEditorGraphNodeBase.cpp
@@ -323,7 +323,7 @@ FLinearColor UPCGEditorGraphNodeBase::GetNodeTitleColor() const
 	if (PCGNode)
 	{
 		const UPCGSettingsInterface* PCGSettingsInterface = PCGNode->GetSettingsInterface();
-		const UPCGSettings* PCGSettings = PCGSettingsInterface->GetSettings();
+		const UPCGSettings* PCGSettings = PCGSettingsInterface ? PCGSettingsInterface->GetSettings() : nullptr;
 
 		if (PCGSettingsInterface && PCGSettingsInterface->ExecutionMode == EPCGSettingsExecutionMode::Isolated)
 		{
