Changed visual node representation (tint & text) when using instanced settings
Renamed class defaults button to Graph settings
Added asset-accessible param on settings to determine whether they should be exposed or not in the palette/menu
Hidden asset info related members in the graph editor

diff --git a/Engine/Plugins/Experimental/PCG/Source/PCG/Private/Elements/PCGExecuteBlueprint.cpp b/Engine/Plugins/Experimental/PCG/Source/PCG/Private/Elements/PCGExecuteBlueprint.cpp
--- a/Engine/Plugins/Experimental/PCG/Source/PCG/Private/Elements/PCGExecuteBlueprint.cpp
+++ b/Engine/Plugins/Experimental/PCG/Source/PCG/Private/Elements/PCGExecuteBlueprint.cpp
@@ -41,6 +41,7 @@ void UPCGBlueprintElement::PostLoad()
 	Super::PostLoad();
 	Initialize();
 
+#if WITH_EDITORONLY_DATA
 	if (!InputPinLabels_DEPRECATED.IsEmpty())
 	{
 		for (const FName& Label : InputPinLabels_DEPRECATED)
@@ -60,6 +61,7 @@ void UPCGBlueprintElement::PostLoad()
 
 		OutputPinLabels_DEPRECATED.Reset();
 	}
+#endif
 }
 
 void UPCGBlueprintElement::BeginDestroy()
@@ -208,11 +210,13 @@ void UPCGBlueprintSettings::PostLoad()
 {
 	Super::PostLoad();
 
+#if WITH_EDITORONLY_DATA
 	if (BlueprintElement_DEPRECATED && !BlueprintElementType)
 	{
 		BlueprintElementType = BlueprintElement_DEPRECATED;
 		BlueprintElement_DEPRECATED = nullptr;
 	}
+#endif
 
 	SetupBlueprintEvent();
 
@@ -228,12 +232,16 @@ void UPCGBlueprintSettings::PostLoad()
 	if (BlueprintElementInstance)
 	{
 		BlueprintElementInstance->ConditionalPostLoad();
+#if WITH_EDITORONLY_DATA
 		BlueprintElementInstance->bCreatesArtifacts |= bCreatesArtifacts_DEPRECATED;
 		BlueprintElementInstance->bCanBeMultithreaded |= bCanBeMultithreaded_DEPRECATED;
+#endif
 	}
 
+#if WITH_EDITORONLY_DATA
 	bCreatesArtifacts_DEPRECATED = false;
 	bCanBeMultithreaded_DEPRECATED = false;
+#endif
 }
 
 void UPCGBlueprintSettings::BeginDestroy()
diff --git a/Engine/Plugins/Experimental/PCG/Source/PCG/Private/PCGNode.cpp b/Engine/Plugins/Experimental/PCG/Source/PCG/Private/PCGNode.cpp
--- a/Engine/Plugins/Experimental/PCG/Source/PCG/Private/PCGNode.cpp
+++ b/Engine/Plugins/Experimental/PCG/Source/PCG/Private/PCGNode.cpp
@@ -401,6 +401,16 @@ UPCGSettings* UPCGNode::GetSettings() const
 }
 
 #if WITH_EDITOR
+void UPCGNode::PostEditChangeProperty(struct FPropertyChangedEvent& PropertyChangedEvent)
+{
+	Super::PostEditChangeProperty(PropertyChangedEvent);
+
+	if (PropertyChangedEvent.Property && PropertyChangedEvent.Property->GetFName() == GET_MEMBER_NAME_CHECKED(UPCGNode, NodeTitle))
+	{
+		OnNodeChangedDelegate.Broadcast(this, EPCGChangeType::Cosmetic);
+	}
+}
+
 void UPCGNode::OnSettingsChanged(UPCGSettings* InSettings, EPCGChangeType ChangeType)
 {
 	if (InSettings == GetSettings())
diff --git a/Engine/Plugins/Experimental/PCG/Source/PCG/Public/Elements/PCGExecuteBlueprint.h b/Engine/Plugins/Experimental/PCG/Source/PCG/Public/Elements/PCGExecuteBlueprint.h
--- a/Engine/Plugins/Experimental/PCG/Source/PCG/Public/Elements/PCGExecuteBlueprint.h
+++ b/Engine/Plugins/Experimental/PCG/Source/PCG/Public/Elements/PCGExecuteBlueprint.h
@@ -104,41 +104,43 @@ public:
 	FOnPCGBlueprintChanged OnBlueprintChangedDelegate;
 #endif
 
-	UPROPERTY(BlueprintReadWrite, EditAnywhere, Category = Settings)
+	UPROPERTY(BlueprintReadWrite, EditDefaultsOnly, Category = Settings)
 	bool bCreatesArtifacts = false;
 
-	UPROPERTY(BlueprintReadWrite, EditAnywhere, Category = Settings)
+	UPROPERTY(BlueprintReadWrite, EditDefaultsOnly, Category = Settings)
 	bool bCanBeMultithreaded = false;
 
+#if WITH_EDITORONLY_DATA
 	UPROPERTY(BlueprintGetter=InputLabels, Category = "Settings|Input & Output", meta = (DeprecatedProperty, DeprecatedMessage = "Input Pin Labels are deprecated - use Input Labels instead."))
 	TSet<FName> InputPinLabels_DEPRECATED;
 
 	UPROPERTY(BlueprintGetter=OutputLabels, Category = "Settings|Input & Output")
 	TSet<FName> OutputPinLabels_DEPRECATED;
+#endif
 
-	UPROPERTY(BlueprintReadWrite, EditAnywhere, Category = "Settings|Input & Output")
+	UPROPERTY(BlueprintReadWrite, EditDefaultsOnly, Category = "Settings|Input & Output")
 	TArray<FPCGPinProperties> CustomInputPins;
 
-	UPROPERTY(BlueprintReadWrite, EditAnywhere, Category = "Settings|Input & Output")
+	UPROPERTY(BlueprintReadWrite, EditDefaultsOnly, Category = "Settings|Input & Output")
 	TArray<FPCGPinProperties> CustomOutputPins;
 
-	UPROPERTY(BlueprintReadWrite, EditAnywhere, Category = "Settings|Input & Output")
+	UPROPERTY(BlueprintReadWrite, EditDefaultsOnly, Category = "Settings|Input & Output")
 	bool bHasDefaultInPin = true;
 
-	UPROPERTY(BlueprintReadWrite, EditAnywhere, Category = "Settings|Input & Output")
+	UPROPERTY(BlueprintReadWrite, EditDefaultsOnly, Category = "Settings|Input & Output")
 	bool bHasDefaultOutPin = true;
 
 #if WITH_EDITORONLY_DATA
-	UPROPERTY(EditAnywhere, Category = AssetInfo, AssetRegistrySearchable)
+	UPROPERTY(BlueprintReadWrite, EditDefaultsOnly, Category = AssetInfo, AssetRegistrySearchable)
 	bool bExposeToLibrary = false;
 
-	UPROPERTY(EditAnywhere, Category = AssetInfo, AssetRegistrySearchable)
+	UPROPERTY(BlueprintReadWrite, EditDefaultsOnly, Category = AssetInfo, AssetRegistrySearchable)
 	FText Category;
 
-	UPROPERTY(EditAnywhere, Category = AssetInfo, AssetRegistrySearchable)
+	UPROPERTY(BlueprintReadWrite, EditDefaultsOnly, Category = AssetInfo, AssetRegistrySearchable)
 	FText Description;
 
-	UPROPERTY(BlueprintReadOnly, EditAnywhere, Category = "Settings|Advanced")
+	UPROPERTY(BlueprintReadWrite, EditDefaultsOnly, Category = "Settings|Advanced")
 	int32 DependencyParsingDepth = 1;
 #endif
 
@@ -197,8 +199,10 @@ public:
 	TSubclassOf<UPCGBlueprintElement> GetElementType() const { return BlueprintElementType; }
 
 protected:
+#if WITH_EDITORONLY_DATA
 	UPROPERTY()
 	TSubclassOf<UPCGBlueprintElement> BlueprintElement_DEPRECATED;
+#endif
 
 	UPROPERTY(BlueprintReadOnly, EditAnywhere, Category = Template)
 	TSubclassOf<UPCGBlueprintElement> BlueprintElementType;
@@ -209,13 +213,13 @@ protected:
 #if WITH_EDITORONLY_DATA
 	UPROPERTY(BlueprintReadWrite, EditAnywhere, Category = Settings)
 	TArray<FName> TrackedActorTags;
-#endif
 
 	UPROPERTY()
 	bool bCreatesArtifacts_DEPRECATED = false;
 
 	UPROPERTY()
 	bool bCanBeMultithreaded_DEPRECATED = false;
+#endif
 
 protected:
 #if WITH_EDITOR
diff --git a/Engine/Plugins/Experimental/PCG/Source/PCG/Public/PCGGraph.h b/Engine/Plugins/Experimental/PCG/Source/PCG/Public/PCGGraph.h
--- a/Engine/Plugins/Experimental/PCG/Source/PCG/Public/PCGGraph.h
+++ b/Engine/Plugins/Experimental/PCG/Source/PCG/Public/PCGGraph.h
@@ -34,13 +34,13 @@ public:
 	/** ~End UObject interface */
 
 #if WITH_EDITORONLY_DATA
-	UPROPERTY(EditAnywhere, Category = AssetInfo, AssetRegistrySearchable)
+	UPROPERTY(BlueprintReadWrite, EditAnywhere, Category = AssetInfo, AssetRegistrySearchable)
 	bool bExposeToLibrary = false;
 
-	UPROPERTY(EditAnywhere, Category = AssetInfo, AssetRegistrySearchable)
+	UPROPERTY(BlueprintReadWrite, EditAnywhere, Category = AssetInfo, AssetRegistrySearchable)
 	FText Category;
 
-	UPROPERTY(EditAnywhere, Category = AssetInfo, AssetRegistrySearchable)
+	UPROPERTY(BlueprintReadWrite, EditAnywhere, Category = AssetInfo, AssetRegistrySearchable)
 	FText Description;
 #endif
 
diff --git a/Engine/Plugins/Experimental/PCG/Source/PCG/Public/PCGNode.h b/Engine/Plugins/Experimental/PCG/Source/PCG/Public/PCGNode.h
--- a/Engine/Plugins/Experimental/PCG/Source/PCG/Public/PCGNode.h
+++ b/Engine/Plugins/Experimental/PCG/Source/PCG/Public/PCGNode.h
@@ -152,6 +152,7 @@ protected:
 	bool UpdatePins(TFunctionRef<UPCGPin* (UPCGNode*)> PinAllocator);
 
 #if WITH_EDITOR
+	virtual void PostEditChangeProperty(struct FPropertyChangedEvent& PropertyChangedEvent) override;
 	void OnSettingsChanged(UPCGSettings* InSettings, EPCGChangeType ChangeType);
 #endif
 
diff --git a/Engine/Plugins/Experimental/PCG/Source/PCG/Public/PCGSettings.h b/Engine/Plugins/Experimental/PCG/Source/PCG/Public/PCGSettings.h
--- a/Engine/Plugins/Experimental/PCG/Source/PCG/Public/PCGSettings.h
+++ b/Engine/Plugins/Experimental/PCG/Source/PCG/Public/PCGSettings.h
@@ -46,8 +46,8 @@ enum class EPCGSettingsType : uint8
 	DECLARE_MULTICAST_DELEGATE_TwoParams(FOnPCGSettingsChanged, UPCGSettings*, EPCGChangeType);
 #endif
 
-UCLASS(Abstract, MinimalAPI)
-class UPCGSettingsInterface : public UPCGData
+UCLASS(Abstract)
+class PCG_API UPCGSettingsInterface : public UPCGData
 {
 	GENERATED_BODY()
 
@@ -148,10 +148,13 @@ public:
 #endif
 
 #if WITH_EDITORONLY_DATA
-	UPROPERTY(EditAnywhere, Category = AssetInfo, AssetRegistrySearchable)
+	UPROPERTY(BlueprintReadWrite, EditAnywhere, Category = AssetInfo, AssetRegistrySearchable)
+	bool bExposeToLibrary = false;
+
+	UPROPERTY(BlueprintReadWrite, EditAnywhere, Category = AssetInfo, AssetRegistrySearchable)
 	FText Category;
 
-	UPROPERTY(EditAnywhere, Category = AssetInfo, AssetRegistrySearchable)
+	UPROPERTY(BlueprintReadWrite, EditAnywhere, Category = AssetInfo, AssetRegistrySearchable)
 	FText Description;
 #endif
 
diff --git a/Engine/Plugins/Experimental/PCG/Source/PCGEditor/Private/PCGEditor.cpp b/Engine/Plugins/Experimental/PCG/Source/PCGEditor/Private/PCGEditor.cpp
--- a/Engine/Plugins/Experimental/PCG/Source/PCGEditor/Private/PCGEditor.cpp
+++ b/Engine/Plugins/Experimental/PCG/Source/PCGEditor/Private/PCGEditor.cpp
@@ -88,6 +88,7 @@ void FPCGEditor::Initialize(const EToolkitMode::Type InMode, const TSharedPtr<cl
 	FDetailsViewArgs DetailsViewArgs;
 	DetailsViewArgs.NameAreaSettings = FDetailsViewArgs::HideNameArea;
 	DetailsViewArgs.bHideSelectionTip = true;
+	DetailsViewArgs.DefaultsOnlyVisibility = EEditDefaultsOnlyNodeVisibility::Automatic;
 	PropertyDetailsWidget = PropertyEditorModule.CreateDetailView(DetailsViewArgs);
 	PropertyDetailsWidget->SetObject(PCGGraphBeingEdited);
 
@@ -342,7 +343,7 @@ void FPCGEditor::RegisterToolbar() const
 
 		Section.AddSeparator(NAME_None);
 		Section.AddEntry(FToolMenuEntry::InitToolBarButton(
-			 PCGEditorCommands.EditClassDefaults,
+			 PCGEditorCommands.EditGraphSettings,
 			 TAttribute<FText>(),
 			 TAttribute<FText>(),
 			 FSlateIcon(FAppStyle::GetAppStyleSetName(), "FullBlueprintEditor.EditClassDefaults")));
@@ -377,10 +378,10 @@ void FPCGEditor::BindCommands()
 		FCanExecuteAction::CreateSP(this, &FPCGEditor::CanRunDeterminismGraphTest));
 
 	ToolkitCommands->MapAction(
-		PCGEditorCommands.EditClassDefaults,
-		FExecuteAction::CreateSP(this, &FPCGEditor::OnEditClassDefaults),
+		PCGEditorCommands.EditGraphSettings,
+		FExecuteAction::CreateSP(this, &FPCGEditor::OnEditGraphSettings),
 		FCanExecuteAction(),
-		FIsActionChecked::CreateSP(this, &FPCGEditor::IsEditClassDefaultsToggled));
+		FIsActionChecked::CreateSP(this, &FPCGEditor::IsEditGraphSettingsToggled));
 	
 	GraphEditorCommands->MapAction(
 		PCGEditorCommands.CollapseNodes,
@@ -392,6 +393,11 @@ void FPCGEditor::BindCommands()
 		FExecuteAction::CreateSP(this, &FPCGEditor::OnExportNodes),
 		FCanExecuteAction::CreateSP(this, &FPCGEditor::CanExportNodes));
 
+	GraphEditorCommands->MapAction(
+		PCGEditorCommands.ConvertToStandaloneNodes,
+		FExecuteAction::CreateSP(this, &FPCGEditor::OnConvertToStandaloneNodes),
+		FCanExecuteAction::CreateSP(this, &FPCGEditor::CanConvertToStandaloneNodes));
+
 	GraphEditorCommands->MapAction(
 		PCGEditorCommands.StartInspectNode,
 		FExecuteAction::CreateSP(this, &FPCGEditor::OnStartInspectNode));
@@ -627,12 +633,12 @@ void FPCGEditor::OnDeterminismGraphTest()
 	}
 }
 
-void FPCGEditor::OnEditClassDefaults() const
+void FPCGEditor::OnEditGraphSettings() const
 {
 	PropertyDetailsWidget->SetObject(PCGGraphBeingEdited);
 }
 
-bool FPCGEditor::IsEditClassDefaultsToggled() const
+bool FPCGEditor::IsEditGraphSettingsToggled() const
 {
 	const TArray<TWeakObjectPtr<UObject>>& SelectedObjects = PropertyDetailsWidget->GetSelectedObjects();
 	return SelectedObjects.Num() == 1 && SelectedObjects[0] == PCGGraphBeingEdited;
@@ -1090,6 +1096,75 @@ void FPCGEditor::OnExportNodes()
 	}
 }
 
+void FPCGEditor::OnConvertToStandaloneNodes()
+{
+	const FScopedTransaction Transaction(*FPCGEditorCommon::ContextIdentifier, LOCTEXT("PCGEditorCollapseInSubgraphMessage", "PCG Editor: Converting instanced nodes to standalone"), nullptr);
+
+	for (UObject* Object : GraphEditorWidget->GetSelectedNodes())
+	{
+		check(Object);
+		// Exclude input and output nodes from the subgraph.
+		if (Object->IsA<UPCGEditorGraphNodeInput>() || Object->IsA<UPCGEditorGraphNodeOutput>())
+		{
+			continue;
+		}
+
+		if (UPCGEditorGraphNodeBase* Node = Cast<UPCGEditorGraphNodeBase>(Object))
+		{
+			if (UPCGNode* PCGNode = Node->GetPCGNode())
+			{
+				if (PCGNode->IsInstance())
+				{
+					PCGNode->Modify();
+
+					UPCGSettings* SourceSettings = PCGNode->GetSettings();
+					check(SourceSettings);
+
+					UPCGSettings* SettingsCopy = DuplicateObject(SourceSettings, PCGNode);
+					SettingsCopy->SetFlags(RF_Transactional);
+
+					PCGNode->SetSettingsInterface(SettingsCopy);
+				}
+			}
+
+			Node->ReconstructNode();
+		}
+	}
+
+	// Notify the widget
+	if (GraphEditorWidget)
+	{
+		GraphEditorWidget->NotifyGraphChanged();
+	}
+}
+
+bool FPCGEditor::CanConvertToStandaloneNodes() const
+{
+	for (const UObject* Object : GraphEditorWidget->GetSelectedNodes())
+	{
+		check(Object);
+
+		// Exclude input and output nodes from the subgraph.
+		if (Object->IsA<UPCGEditorGraphNodeInput>() || Object->IsA<UPCGEditorGraphNodeOutput>())
+		{
+			continue;
+		}
+
+		if (const UPCGEditorGraphNodeBase* Node = Cast<const UPCGEditorGraphNodeBase>(Object))
+		{
+			if (const UPCGNode* PCGNode = Node->GetPCGNode())
+			{
+				if (PCGNode->IsInstance())
+				{
+					return true;
+				}
+			}
+		}
+	}
+
+	return false;
+}
+
 void FPCGEditor::OnStartInspectNode()
 {
 	if (!GraphEditorWidget.IsValid())
@@ -1813,9 +1888,20 @@ bool FPCGEditor::IsReadOnlyProperty(const FPropertyAndParent& InPropertyAndParen
 
 bool FPCGEditor::IsVisibleProperty(const FPropertyAndParent& InPropertyAndParent, IDetailsView* InDetailsView) const
 {
-	// Everything is visible when not in an instance
-	if (!InDetailsView ||
-		InPropertyAndParent.ParentProperties.IsEmpty() ||
+	if (!InDetailsView)
+	{
+		return true;
+	}
+
+	// Always hide asset info information
+	if (InPropertyAndParent.Property.HasMetaData(TEXT("Category")) &&
+		InPropertyAndParent.Property.GetMetaData(TEXT("Category")) == TEXT("AssetInfo"))
+	{
+		return false;
+	}
+
+	// Otherwise, everything is visible when not in an instance
+	if(	InPropertyAndParent.ParentProperties.IsEmpty() ||
 		InPropertyAndParent.ParentProperties.Last()->GetFName() != GET_MEMBER_NAME_CHECKED(UPCGSettingsInstance, Settings))
 	{
 		return true;
diff --git a/Engine/Plugins/Experimental/PCG/Source/PCGEditor/Private/PCGEditor.h b/Engine/Plugins/Experimental/PCG/Source/PCGEditor/Private/PCGEditor.h
--- a/Engine/Plugins/Experimental/PCG/Source/PCGEditor/Private/PCGEditor.h
+++ b/Engine/Plugins/Experimental/PCG/Source/PCGEditor/Private/PCGEditor.h
@@ -128,9 +128,9 @@ private:
 	void OnDeterminismGraphTest();
 
 	/** Open details view for the PCG object being edited */
-	void OnEditClassDefaults() const;
+	void OnEditGraphSettings() const;
 	/** Whether the PCG object being edited is opened in details view or not */
-	bool IsEditClassDefaultsToggled() const;
+	bool IsEditGraphSettingsToggled() const;
 
 	/** Whether or not an execution mode is active for the selected nodes */
 	bool IsExecutionModeActive(EPCGSettingsExecutionMode InExecutionMode) const;
@@ -176,6 +176,12 @@ private:
 	/** Whether we are able to export the currently selected nodes */
 	bool CanExportNodes() const;
 
+	/** Converts instanced nodes to independent nodes */
+	void OnConvertToStandaloneNodes();
+
+	/** Whether we are able to convert the selected nodes to standalone */
+	bool CanConvertToStandaloneNodes() const;
+
 	/** Collapse the currently selected nodes in a subgraph */
 	void OnCollapseNodesInSubgraph();
 	/** Whether we can collapse nodes in a subgraph */
diff --git a/Engine/Plugins/Experimental/PCG/Source/PCGEditor/Private/PCGEditorCommands.cpp b/Engine/Plugins/Experimental/PCG/Source/PCGEditor/Private/PCGEditorCommands.cpp
--- a/Engine/Plugins/Experimental/PCG/Source/PCGEditor/Private/PCGEditorCommands.cpp
+++ b/Engine/Plugins/Experimental/PCG/Source/PCGEditor/Private/PCGEditorCommands.cpp
@@ -20,12 +20,13 @@ void FPCGEditorCommands::RegisterCommands()
 {
 	UI_COMMAND(CollapseNodes, "Collapse into Subgraph", "Collapse selected nodes into a separate PCGGraph asset.", EUserInterfaceActionType::Button, FInputChord(EModifierKey::Alt, EKeys::J));
 	UI_COMMAND(ExportNodes, "Export nodes to AssetData", "Exports selected nodes to separate and reusable PCGSettings assets.", EUserInterfaceActionType::Button, FInputChord());
+	UI_COMMAND(ConvertToStandaloneNodes, "Convert to standalone Nodes", "Converts instanced nodes to standalone nodes.", EUserInterfaceActionType::Button, FInputChord());
 	UI_COMMAND(Find, "Find", "Finds PCG nodes and comments in the current graph", EUserInterfaceActionType::Button, FInputChord(EModifierKey::Control, EKeys::F));
 	UI_COMMAND(PauseAutoRegeneration, "Pause Regen", "Pause automatic regeneration of the current graph.", EUserInterfaceActionType::ToggleButton, FInputChord(EModifierKey::Alt, EKeys::R));
 	UI_COMMAND(ForceGraphRegeneration, "Force Regen", "Manually force a regeneration of the current graph.", EUserInterfaceActionType::Button, FInputChord());
 	UI_COMMAND(RunDeterminismNodeTest, "Run Determinism Test on Node", "Evaluate the current node for determinism.", EUserInterfaceActionType::Button, FInputChord(EModifierKey::Alt, EKeys::D));
 	UI_COMMAND(RunDeterminismGraphTest, "Graph Determinism Test", "Evaluate the current graph for determinism.", EUserInterfaceActionType::Button, FInputChord());
-	UI_COMMAND(EditClassDefaults, "Class Defaults", "Edit the initial values of your class.", EUserInterfaceActionType::ToggleButton, FInputChord());
+	UI_COMMAND(EditGraphSettings, "Graph Settings", "Edit the graph settings.", EUserInterfaceActionType::ToggleButton, FInputChord());
 	UI_COMMAND(StartInspectNode, "Start Inspecting Node", "Start Inspecting Node", EUserInterfaceActionType::Button, FInputChord(EKeys::I));
 	UI_COMMAND(StopInspectNode, "Stop Inspecting Node", "Stop Inspecting Node", EUserInterfaceActionType::Button, FInputChord(EModifierKey::Shift, EKeys::I));
 	UI_COMMAND(ExecutionModeEnabled, "Execution Mode Enabled", "Set the execution mode for this node to Enabled.", EUserInterfaceActionType::RadioButton, FInputChord(EKeys::E));
diff --git a/Engine/Plugins/Experimental/PCG/Source/PCGEditor/Private/PCGEditorCommands.h b/Engine/Plugins/Experimental/PCG/Source/PCGEditor/Private/PCGEditorCommands.h
--- a/Engine/Plugins/Experimental/PCG/Source/PCGEditor/Private/PCGEditorCommands.h
+++ b/Engine/Plugins/Experimental/PCG/Source/PCGEditor/Private/PCGEditorCommands.h
@@ -16,12 +16,13 @@ public:
 
 	TSharedPtr<FUICommandInfo> CollapseNodes;
 	TSharedPtr<FUICommandInfo> ExportNodes;
+	TSharedPtr<FUICommandInfo> ConvertToStandaloneNodes;
 	TSharedPtr<FUICommandInfo> Find;
 	TSharedPtr<FUICommandInfo> PauseAutoRegeneration;
 	TSharedPtr<FUICommandInfo> ForceGraphRegeneration;
 	TSharedPtr<FUICommandInfo> RunDeterminismNodeTest;
 	TSharedPtr<FUICommandInfo> RunDeterminismGraphTest;
-	TSharedPtr<FUICommandInfo> EditClassDefaults;
+	TSharedPtr<FUICommandInfo> EditGraphSettings;
 	TSharedPtr<FUICommandInfo> StartInspectNode;
 	TSharedPtr<FUICommandInfo> StopInspectNode;
 	TSharedPtr<FUICommandInfo> ExecutionModeEnabled;
diff --git a/Engine/Plugins/Experimental/PCG/Source/PCGEditor/Private/PCGEditorGraphNodeBase.cpp b/Engine/Plugins/Experimental/PCG/Source/PCGEditor/Private/PCGEditorGraphNodeBase.cpp
--- a/Engine/Plugins/Experimental/PCG/Source/PCGEditor/Private/PCGEditorGraphNodeBase.cpp
+++ b/Engine/Plugins/Experimental/PCG/Source/PCGEditor/Private/PCGEditorGraphNodeBase.cpp
@@ -93,6 +93,7 @@ void UPCGEditorGraphNodeBase::GetNodeContextMenuActions(UToolMenu* Menu, class U
 		}));
 
 		Section.AddMenuEntry(FPCGEditorCommands::Get().ExportNodes);
+		Section.AddMenuEntry(FPCGEditorCommands::Get().ConvertToStandaloneNodes);
 	}
 
 	{
@@ -357,9 +358,16 @@ FLinearColor UPCGEditorGraphNodeBase::GetNodeBodyTintColor() const
 	if (PCGNode)
 	{
 		const UPCGSettingsInterface* PCGSettingsInterface = PCGNode->GetSettingsInterface();
-		if (PCGSettingsInterface && PCGSettingsInterface->ExecutionMode == EPCGSettingsExecutionMode::Isolated)
+		if (PCGSettingsInterface)
 		{
-			return GetDefault<UPCGEditorSettings>()->IsolatedNodeColor;
+			if (PCGSettingsInterface->IsInstance())
+			{
+				return GetDefault<UPCGEditorSettings>()->InstancedNodeBodyTintColor;
+			}
+			else if (PCGSettingsInterface->ExecutionMode == EPCGSettingsExecutionMode::Isolated)
+			{
+				return GetDefault<UPCGEditorSettings>()->IsolatedNodeColor;
+			}
 		}
 	}
 
diff --git a/Engine/Plugins/Experimental/PCG/Source/PCGEditor/Private/PCGEditorGraphSchema.cpp b/Engine/Plugins/Experimental/PCG/Source/PCGEditor/Private/PCGEditorGraphSchema.cpp
--- a/Engine/Plugins/Experimental/PCG/Source/PCGEditor/Private/PCGEditorGraphSchema.cpp
+++ b/Engine/Plugins/Experimental/PCG/Source/PCGEditor/Private/PCGEditorGraphSchema.cpp
@@ -267,29 +267,33 @@ void UPCGEditorGraphSchema::GetSettingsElementActions(FGraphActionMenuBuilder& A
 
 	for (const FAssetData& AssetData : AssetDataList)
 	{
-		const FText MenuDesc = FText::FromName(AssetData.AssetName);
-		const FText Category = AssetData.GetTagValueRef<FText>(TEXT("Category"));
-		const FText Description = AssetData.GetTagValueRef<FText>(TEXT("Description"));
-
-		if (!bIsContextual)
-		{
-			TSharedPtr<FPCGEditorGraphSchemaAction_NewSettingsElement> NewSettingsAction(new FPCGEditorGraphSchemaAction_NewSettingsElement(Category, MenuDesc, Description, 0));
-			NewSettingsAction->SettingsObjectPath = AssetData.GetSoftObjectPath();
-			ActionMenuBuilder.AddAction(NewSettingsAction);
-		}
-		else
+		const bool bExposeToLibrary = AssetData.GetTagValueRef<bool>(TEXT("bExposeToLibrary"));
+		if (bExposeToLibrary)
 		{
-			const FText MenuAndSubCategory = FText::Join(LOCTEXT("MenuDelimiter", "|"), Category, MenuDesc);
+			const FText MenuDesc = FText::FromName(AssetData.AssetName);
+			const FText Category = AssetData.GetTagValueRef<FText>(TEXT("Category"));
+			const FText Description = AssetData.GetTagValueRef<FText>(TEXT("Description"));
+
+			if (!bIsContextual)
+			{
+				TSharedPtr<FPCGEditorGraphSchemaAction_NewSettingsElement> NewSettingsAction(new FPCGEditorGraphSchemaAction_NewSettingsElement(Category, MenuDesc, Description, 0));
+				NewSettingsAction->SettingsObjectPath = AssetData.GetSoftObjectPath();
+				ActionMenuBuilder.AddAction(NewSettingsAction);
+			}
+			else
+			{
+				const FText MenuAndSubCategory = FText::Join(LOCTEXT("MenuDelimiter", "|"), Category, MenuDesc);
 
-			TSharedPtr<FPCGEditorGraphSchemaAction_NewSettingsElement> NewSettingsActionCopy(new FPCGEditorGraphSchemaAction_NewSettingsElement(MenuAndSubCategory, LOCTEXT("ContextMenuCopySettings", "Copy"), Description, 0));
-			NewSettingsActionCopy->SettingsObjectPath = AssetData.GetSoftObjectPath();
-			NewSettingsActionCopy->Behavior = EPCGEditorNewSettingsBehavior::ForceCopy;
-			ActionMenuBuilder.AddAction(NewSettingsActionCopy);
+				TSharedPtr<FPCGEditorGraphSchemaAction_NewSettingsElement> NewSettingsActionCopy(new FPCGEditorGraphSchemaAction_NewSettingsElement(MenuAndSubCategory, LOCTEXT("ContextMenuCopySettings", "Copy"), Description, 0));
+				NewSettingsActionCopy->SettingsObjectPath = AssetData.GetSoftObjectPath();
+				NewSettingsActionCopy->Behavior = EPCGEditorNewSettingsBehavior::ForceCopy;
+				ActionMenuBuilder.AddAction(NewSettingsActionCopy);
 
-			TSharedPtr<FPCGEditorGraphSchemaAction_NewSettingsElement> NewSettingsActionInstance(new FPCGEditorGraphSchemaAction_NewSettingsElement(MenuAndSubCategory, LOCTEXT("ContextMenuCopySettings", "Instance"), Description, 0));
-			NewSettingsActionInstance->SettingsObjectPath = AssetData.GetSoftObjectPath();
-			NewSettingsActionInstance->Behavior = EPCGEditorNewSettingsBehavior::ForceInstance;
-			ActionMenuBuilder.AddAction(NewSettingsActionInstance);
+				TSharedPtr<FPCGEditorGraphSchemaAction_NewSettingsElement> NewSettingsActionInstance(new FPCGEditorGraphSchemaAction_NewSettingsElement(MenuAndSubCategory, LOCTEXT("ContextMenuCopySettings", "Instance"), Description, 0));
+				NewSettingsActionInstance->SettingsObjectPath = AssetData.GetSoftObjectPath();
+				NewSettingsActionInstance->Behavior = EPCGEditorNewSettingsBehavior::ForceInstance;
+				ActionMenuBuilder.AddAction(NewSettingsActionInstance);
+			}
 		}
 	}
 }
diff --git a/Engine/Plugins/Experimental/PCG/Source/PCGEditor/Private/PCGEditorSettings.cpp b/Engine/Plugins/Experimental/PCG/Source/PCGEditor/Private/PCGEditorSettings.cpp
--- a/Engine/Plugins/Experimental/PCG/Source/PCGEditor/Private/PCGEditorSettings.cpp
+++ b/Engine/Plugins/Experimental/PCG/Source/PCGEditor/Private/PCGEditorSettings.cpp
@@ -9,6 +9,7 @@ UPCGEditorSettings::UPCGEditorSettings(const FObjectInitializer& ObjectInitializ
 	: Super(ObjectInitializer)
 {
 	DefaultNodeColor = FLinearColor(0.4f, 0.62f, 1.0f);
+	InstancedNodeBodyTintColor = FLinearColor(0.5f, 0.5f, 0.5f);
 	IsolatedNodeColor = FLinearColor(0.0f, 0.0f, 1.0f);
 	InputOutputNodeColor = FLinearColor(1.0f, 0.0f, 0.0f);
 	SetOperationNodeColor = FLinearColor(0.8f, 0.2f, 0.8f);
diff --git a/Engine/Plugins/Experimental/PCG/Source/PCGEditor/Private/PCGEditorStyle.cpp b/Engine/Plugins/Experimental/PCG/Source/PCGEditor/Private/PCGEditorStyle.cpp
--- a/Engine/Plugins/Experimental/PCG/Source/PCGEditor/Private/PCGEditorStyle.cpp
+++ b/Engine/Plugins/Experimental/PCG/Source/PCGEditor/Private/PCGEditorStyle.cpp
@@ -3,6 +3,7 @@
 #include "PCGEditorStyle.h"
 
 #include "Misc/Paths.h"
+#include "Styling/SlateStyle.h"
 #include "Styling/SlateStyleMacros.h"
 #include "Styling/SlateStyleRegistry.h"
 #include "Styling/StyleColors.h"
@@ -25,6 +26,14 @@ FPCGEditorStyle::FPCGEditorStyle() : FSlateStyleSet("PCGEditorStyle")
 
 	Set("PCG.NodeOverlay.Debug", new CORE_IMAGE_BRUSH_SVG("Starship/Blueprints/Breakpoint_Valid", Icon20x20, FSlateColor(FColor::Cyan)));
 	Set("PCG.NodeOverlay.Inspect", new CORE_IMAGE_BRUSH_SVG("Starship/Blueprints/Breakpoint_Valid", Icon20x20, FSlateColor(FColor::Orange)));
+
+	FInlineEditableTextBlockStyle NodeTitleStyle = FInlineEditableTextBlockStyle(FAppStyle::Get().GetWidgetStyle<FInlineEditableTextBlockStyle>("Graph.Node.NodeTitleInlineEditableText"));
+	FTextBlockStyle GraphNodeItalicTitle = FTextBlockStyle(FAppStyle::Get().GetWidgetStyle<FTextBlockStyle>("Graph.Node.NodeTitle"))
+		.SetFont(DEFAULT_FONT("BoldCondensedItalic", 10));
+
+	Set("PCG.Node.NodeTitleInlineEditableText", NodeTitleStyle);
+	Set("PCG.Node.InstancedNodeTitleInlineEditableText", FInlineEditableTextBlockStyle(NodeTitleStyle)
+		.SetTextStyle(GraphNodeItalicTitle));
 }
 
 const FPCGEditorStyle& FPCGEditorStyle::Get()
diff --git a/Engine/Plugins/Experimental/PCG/Source/PCGEditor/Private/SPCGEditorGraphNode.cpp b/Engine/Plugins/Experimental/PCG/Source/PCGEditor/Private/SPCGEditorGraphNode.cpp
--- a/Engine/Plugins/Experimental/PCG/Source/PCGEditor/Private/SPCGEditorGraphNode.cpp
+++ b/Engine/Plugins/Experimental/PCG/Source/PCGEditor/Private/SPCGEditorGraphNode.cpp
@@ -10,6 +10,7 @@
 
 #include "SGraphPin.h"
 #include "Styling/SlateBrush.h"
+#include "Widgets/Text/SInlineEditableTextBlock.h"
 
 void SPCGEditorGraphNode::Construct(const FArguments& InArgs, UPCGEditorGraphNodeBase* InNode)
 {
@@ -24,6 +25,35 @@ void SPCGEditorGraphNode::Construct(const FArguments& InArgs, UPCGEditorGraphNod
 	UpdateGraphNode();
 }
 
+const FSlateBrush* SPCGEditorGraphNode::GetNodeBodyBrush() const
+{
+	if (PCGEditorGraphNode && PCGEditorGraphNode->GetPCGNode() && PCGEditorGraphNode->GetPCGNode()->IsInstance())
+	{
+		return FAppStyle::GetBrush("Graph.Node.TintedBody");
+	}
+	else
+	{
+		return FAppStyle::GetBrush("Graph.Node.Body");
+	}
+}
+
+TSharedRef<SWidget> SPCGEditorGraphNode::CreateTitleWidget(TSharedPtr<SNodeTitle> InNodeTitle)
+{
+	// Reimplementation of the SGraphNode::CreateTitleWidget so we can control the style
+	const bool bIsInstanceNode = (PCGEditorGraphNode && PCGEditorGraphNode->GetPCGNode() && PCGEditorGraphNode->GetPCGNode()->IsInstance());
+
+	SAssignNew(InlineEditableText, SInlineEditableTextBlock)
+		.Style(FPCGEditorStyle::Get(), bIsInstanceNode ? "PCG.Node.InstancedNodeTitleInlineEditableText" : "PCG.Node.NodeTitleInlineEditableText")
+		.Text(InNodeTitle.Get(), &SNodeTitle::GetHeadTitle)
+		.OnVerifyTextChanged(this, &SPCGEditorGraphNode::OnVerifyNameTextChanged)
+		.OnTextCommitted(this, &SPCGEditorGraphNode::OnNameTextCommited)
+		.IsReadOnly(this, &SPCGEditorGraphNode::IsNameReadOnly)
+		.IsSelected(this, &SPCGEditorGraphNode::IsSelectedExclusively);
+	InlineEditableText->SetColorAndOpacity(TAttribute<FLinearColor>::Create(TAttribute<FLinearColor>::FGetter::CreateSP(this, &SPCGEditorGraphNode::GetNodeTitleTextColor)));
+
+	return InlineEditableText.ToSharedRef();
+}
+
 void SPCGEditorGraphNode::AddPin(const TSharedRef<SGraphPin>& PinToAdd)
 {
 	check(PCGEditorGraphNode);
diff --git a/Engine/Plugins/Experimental/PCG/Source/PCGEditor/Private/SPCGEditorGraphNode.h b/Engine/Plugins/Experimental/PCG/Source/PCGEditor/Private/SPCGEditorGraphNode.h
--- a/Engine/Plugins/Experimental/PCG/Source/PCGEditor/Private/SPCGEditorGraphNode.h
+++ b/Engine/Plugins/Experimental/PCG/Source/PCGEditor/Private/SPCGEditorGraphNode.h
@@ -16,8 +16,10 @@ public:
 	void Construct(const FArguments& InArgs, UPCGEditorGraphNodeBase* InNode);
 
 	//~ Begin SGraphNode Interface
+	virtual const FSlateBrush* GetNodeBodyBrush() const override;
 	virtual void RequestRenameOnSpawn() override { /* Empty to avoid the default behavior to rename on node spawn */ }
 	virtual void AddPin(const TSharedRef<SGraphPin>& PinToAdd) override;
+	virtual TSharedRef<SWidget> CreateTitleWidget(TSharedPtr<SNodeTitle> InNodeTitle) override;
 	//~ End SGraphNode Interface
 
 	//~ Begin SNodePanel::SNode Interface
diff --git a/Engine/Plugins/Experimental/PCG/Source/PCGEditor/Public/PCGEditorSettings.h b/Engine/Plugins/Experimental/PCG/Source/PCGEditor/Public/PCGEditorSettings.h
--- a/Engine/Plugins/Experimental/PCG/Source/PCGEditor/Public/PCGEditorSettings.h
+++ b/Engine/Plugins/Experimental/PCG/Source/PCGEditor/Public/PCGEditorSettings.h
@@ -25,6 +25,10 @@ public:
 	UPROPERTY(EditAnywhere, config, Category = Node, meta = (HideAlphaChannel))
 	FLinearColor DefaultNodeColor;
 
+	/** Instanced node body tint color */
+	UPROPERTY(EditAnywhere, config, Category = Node, meta = (HideAlphaChannel))
+	FLinearColor InstancedNodeBodyTintColor;
+
 	/** Isolated node color */
 	UPROPERTY(EditAnywhere, config, Category = Node, meta = (HideAlphaChannel))
 	FLinearColor IsolatedNodeColor;
