Tentative non-unity build fix wrt editor style

#rnx

diff --git a/Engine/Plugins/Experimental/PCG/Source/PCGEditor/Private/PCGEditorStyle.cpp b/Engine/Plugins/Experimental/PCG/Source/PCGEditor/Private/PCGEditorStyle.cpp
--- a/Engine/Plugins/Experimental/PCG/Source/PCGEditor/Private/PCGEditorStyle.cpp
+++ b/Engine/Plugins/Experimental/PCG/Source/PCGEditor/Private/PCGEditorStyle.cpp
@@ -3,9 +3,11 @@
 #include "PCGEditorStyle.h"
 
 #include "Misc/Paths.h"
+#include "Styling/AppStyle.h"
 #include "Styling/SlateStyle.h"
 #include "Styling/SlateStyleMacros.h"
 #include "Styling/SlateStyleRegistry.h"
+#include "Styling/SlateTypes.h"
 #include "Styling/StyleColors.h"
 
 void FPCGEditorStyle::Register()
