Added multi actor support in data from actor node

diff --git a/Engine/Plugins/Experimental/PCG/Source/PCG/Private/Elements/PCGActorSelector.cpp b/Engine/Plugins/Experimental/PCG/Source/PCG/Private/Elements/PCGActorSelector.cpp
--- a/Engine/Plugins/Experimental/PCG/Source/PCG/Private/Elements/PCGActorSelector.cpp
+++ b/Engine/Plugins/Experimental/PCG/Source/PCG/Private/Elements/PCGActorSelector.cpp
@@ -8,41 +8,43 @@
 namespace PCGActorSelector
 {
 	// Need to pass a pointer of pointer to the found actor. The lambda will capture this pointer and modify its value when an actor is found.
-	TFunction<bool(AActor*)> GetFilteringFunction(const FPCGActorSelectorSettings& InSettings, AActor*& InFoundActor)
+	TFunction<bool(AActor*)> GetFilteringFunction(const FPCGActorSelectorSettings& InSettings, TArray<AActor*>& InFoundActors)
 	{
+		const bool bMultiSelect = InSettings.bSelectMultiple;
+
 		switch (InSettings.ActorSelection)
 		{
 		case EPCGActorSelection::ByTag:
-			return[ActorSelectionTag = InSettings.ActorSelectionTag, &InFoundActor](AActor* Actor) -> bool
+			return[ActorSelectionTag = InSettings.ActorSelectionTag, &InFoundActors, bMultiSelect](AActor* Actor) -> bool
 			{
 				if (Actor->ActorHasTag(ActorSelectionTag))
 				{
-					InFoundActor = Actor;
-					return false;
+					InFoundActors.Add(Actor);
+					return bMultiSelect;
 				}
 
 				return true;
 			};
 
 		case EPCGActorSelection::ByName:
-			return[ActorSelectionName = InSettings.ActorSelectionName, &InFoundActor](AActor* Actor) -> bool
+			return[ActorSelectionName = InSettings.ActorSelectionName, &InFoundActors, bMultiSelect](AActor* Actor) -> bool
 			{
 				if (Actor->GetFName().IsEqual(ActorSelectionName, ENameCase::IgnoreCase, /*bCompareNumber=*/ false))
 				{
-					InFoundActor = Actor;
-					return false;
+					InFoundActors.Add(Actor);
+					return bMultiSelect;
 				}
 
 				return true;
 			};
 
 		case EPCGActorSelection::ByClass:
-			return[ActorSelectionClass = InSettings.ActorSelectionClass, &InFoundActor](AActor* Actor) -> bool
+			return[ActorSelectionClass = InSettings.ActorSelectionClass, &InFoundActors, bMultiSelect](AActor* Actor) -> bool
 			{
 				if (Actor->IsA(ActorSelectionClass))
 				{
-					InFoundActor = Actor;
-					return false;
+					InFoundActors.Add(Actor);
+					return bMultiSelect;
 				}
 
 				return true;
@@ -55,15 +57,15 @@ namespace PCGActorSelector
 		return [](AActor* Actor) -> bool { return false; };
 	}
 
-	AActor* FindActor(const FPCGActorSelectorSettings& Settings, UWorld* World, AActor* Self)
+	TArray<AActor*> FindActors(const FPCGActorSelectorSettings& Settings, UWorld* World, AActor* Self)
 	{
 		TRACE_CPUPROFILER_EVENT_SCOPE(PCGActorSelector::FindActor);
 
-		AActor* FoundActor = nullptr;
+		TArray<AActor*> FoundActors;
 
 		if (!World)
 		{
-			return FoundActor;
+			return FoundActors;
 		}
 
 		// Early out if we have not the information necessary
@@ -71,12 +73,12 @@ namespace PCGActorSelector
 			(Settings.ActorSelection == EPCGActorSelection::ByName && Settings.ActorSelectionName == NAME_None) ||
 			(Settings.ActorSelection == EPCGActorSelection::ByClass && !Settings.ActorSelectionClass))
 		{
-			return FoundActor;
+			return FoundActors;
 		}
 
 		// We pass FoundActor ref, that will be captured by the FilteringFunction
 		// It will modify the FoundActor pointer to the found actor, if found.
-		TFunction<bool(AActor*)> FilteringFunction = PCGActorSelector::GetFilteringFunction(Settings, FoundActor);
+		TFunction<bool(AActor*)> FilteringFunction = PCGActorSelector::GetFilteringFunction(Settings, FoundActors);
 
 		// In case of iterating over all actors in the world, call our filtering function and get out.
 		if (Settings.ActorFilter == EPCGActorFilter::AllWorldActors)
@@ -84,7 +86,7 @@ namespace PCGActorSelector
 			UPCGActorHelpers::ForEachActorInWorld<AActor>(World, FilteringFunction);
 
 			// FoundActor is set by the FilteringFunction (captured)
-			return FoundActor;
+			return FoundActors;
 		}
 
 		// Otherwise, gather all the actors we need to check
@@ -156,6 +158,16 @@ namespace PCGActorSelector
 			}
 		}
 
-		return FoundActor;
+		return FoundActors;
+	}
+
+	AActor* FindActor(const FPCGActorSelectorSettings& InSettings, UWorld* World, AActor* Self)
+	{
+		// In order to make sure we don't try to select multiple, we'll do a copy of the settings here.
+		FPCGActorSelectorSettings Settings = InSettings;
+		Settings.bSelectMultiple = false;
+
+		TArray<AActor*> Actors = FindActors(Settings, World, Self);
+		return Actors.IsEmpty() ? nullptr : Actors[0];
 	}
 }
\ No newline at end of file
diff --git a/Engine/Plugins/Experimental/PCG/Source/PCG/Private/Elements/PCGDataFromActor.cpp b/Engine/Plugins/Experimental/PCG/Source/PCG/Private/Elements/PCGDataFromActor.cpp
--- a/Engine/Plugins/Experimental/PCG/Source/PCG/Private/Elements/PCGDataFromActor.cpp
+++ b/Engine/Plugins/Experimental/PCG/Source/PCG/Private/Elements/PCGDataFromActor.cpp
@@ -37,14 +37,32 @@ bool FPCGDataFromActorElement::ExecuteInternal(FPCGContext* Context) const
 	check(Settings);
 
 	UWorld* World = Context->SourceComponent.IsValid() ? Context->SourceComponent->GetWorld() : nullptr;
-	AActor* FoundActor = PCGActorSelector::FindActor(Settings->ActorSelector, World, nullptr);
+	TArray<AActor*> FoundActors = PCGActorSelector::FindActors(Settings->ActorSelector, World, nullptr);
 
-	if (!FoundActor)
+	if (FoundActors.IsEmpty())
 	{
 		PCGE_LOG(Warning, "No matching actor was found.");
 		return true;
 	}
 
+	for (AActor* Actor : FoundActors)
+	{
+		ProcessActor(Context, Settings, Actor);
+	}
+
+	return true;
+}
+
+void FPCGDataFromActorElement::ProcessActor(FPCGContext* Context, const UPCGDataFromActorSettings* Settings, AActor* FoundActor) const
+{
+	check(Context);
+	check(Settings);
+
+	if (!FoundActor)
+	{
+		return;
+	}
+
 	TArray<UPCGComponent*> PCGComponents;
 	bool bHasGeneratedPCGData = false;
 	FProperty* FoundProperty = nullptr;
@@ -70,12 +88,12 @@ bool FPCGDataFromActorElement::ExecuteInternal(FPCGContext* Context) const
 	if (Settings->Mode == EPCGGetDataFromActorMode::GetDataFromPCGComponent && !bHasGeneratedPCGData)
 	{
 		PCGE_LOG(Warning, "Actor (%s) does not have any previously generated data.", *FoundActor->GetFName().ToString());
-		return true;
+		return;
 	}
 	else if (Settings->Mode == EPCGGetDataFromActorMode::GetDataFromProperty && !FoundProperty)
 	{
 		PCGE_LOG(Warning, "Actor (%s) does not have a property name (%s)", *FoundActor->GetFName().ToString(), *Settings->PropertyName.ToString());
-		return true;
+		return;
 	}
 
 	TArray<FPCGTaggedData>& Outputs = Context->OutputData.TaggedData;
@@ -138,6 +156,4 @@ bool FPCGDataFromActorElement::ExecuteInternal(FPCGContext* Context) const
 			}
 		}
 	}
-
-	return true;
 }
\ No newline at end of file
diff --git a/Engine/Plugins/Experimental/PCG/Source/PCG/Public/Elements/PCGActorSelector.h b/Engine/Plugins/Experimental/PCG/Source/PCG/Public/Elements/PCGActorSelector.h
--- a/Engine/Plugins/Experimental/PCG/Source/PCG/Public/Elements/PCGActorSelector.h
+++ b/Engine/Plugins/Experimental/PCG/Source/PCG/Public/Elements/PCGActorSelector.h
@@ -50,9 +50,13 @@ struct FPCGActorSelectorSettings
 
 	UPROPERTY(BlueprintReadWrite, EditAnywhere, Category = Settings, meta = (EditCondition = "ActorFilter!=EPCGActorFilter::AllWorldActors", EditConditionHides))
 	bool bIncludeChildren = false;
+
+	UPROPERTY(BlueprintReadWrite, EditAnywhere, Category = Settings, meta = (EditCondition = "ActorSelection!=EPCGActorSelection::ByName"))
+	bool bSelectMultiple = false;
 };
 
 namespace PCGActorSelector
 {
+	TArray<AActor*> FindActors(const FPCGActorSelectorSettings& Settings, UWorld* World, AActor* Self);
 	AActor* FindActor(const FPCGActorSelectorSettings& Settings, UWorld* World, AActor* Self);
 }
\ No newline at end of file
diff --git a/Engine/Plugins/Experimental/PCG/Source/PCG/Public/Elements/PCGDataFromActor.h b/Engine/Plugins/Experimental/PCG/Source/PCG/Public/Elements/PCGDataFromActor.h
--- a/Engine/Plugins/Experimental/PCG/Source/PCG/Public/Elements/PCGDataFromActor.h
+++ b/Engine/Plugins/Experimental/PCG/Source/PCG/Public/Elements/PCGDataFromActor.h
@@ -58,4 +58,5 @@ public:
 	virtual bool IsCacheable(const UPCGSettings* InSettings) const override { return false; }
 protected:
 	virtual bool ExecuteInternal(FPCGContext* Context) const;
+	void ProcessActor(FPCGContext* Context, const UPCGDataFromActorSettings* Settings, AActor* FoundActor) const;
 };
\ No newline at end of file
