Fixed localization warnings

#rnx

diff --git a/Engine/Plugins/Experimental/PCG/Source/PCGEditor/Private/PCGEditor.cpp b/Engine/Plugins/Experimental/PCG/Source/PCGEditor/Private/PCGEditor.cpp
--- a/Engine/Plugins/Experimental/PCG/Source/PCGEditor/Private/PCGEditor.cpp
+++ b/Engine/Plugins/Experimental/PCG/Source/PCGEditor/Private/PCGEditor.cpp
@@ -1098,7 +1098,7 @@ void FPCGEditor::OnExportNodes()
 
 void FPCGEditor::OnConvertToStandaloneNodes()
 {
-	const FScopedTransaction Transaction(*FPCGEditorCommon::ContextIdentifier, LOCTEXT("PCGEditorCollapseInSubgraphMessage", "PCG Editor: Converting instanced nodes to standalone"), nullptr);
+	const FScopedTransaction Transaction(*FPCGEditorCommon::ContextIdentifier, LOCTEXT("PCGEditorConvertToStandaloneMessage", "PCG Editor: Converting instanced nodes to standalone"), nullptr);
 
 	for (UObject* Object : GraphEditorWidget->GetSelectedNodes())
 	{
diff --git a/Engine/Plugins/Experimental/PCG/Source/PCGEditor/Private/PCGEditorGraphSchema.cpp b/Engine/Plugins/Experimental/PCG/Source/PCGEditor/Private/PCGEditorGraphSchema.cpp
--- a/Engine/Plugins/Experimental/PCG/Source/PCGEditor/Private/PCGEditorGraphSchema.cpp
+++ b/Engine/Plugins/Experimental/PCG/Source/PCGEditor/Private/PCGEditorGraphSchema.cpp
@@ -289,7 +289,7 @@ void UPCGEditorGraphSchema::GetSettingsElementActions(FGraphActionMenuBuilder& A
 				NewSettingsActionCopy->Behavior = EPCGEditorNewSettingsBehavior::ForceCopy;
 				ActionMenuBuilder.AddAction(NewSettingsActionCopy);
 
-				TSharedPtr<FPCGEditorGraphSchemaAction_NewSettingsElement> NewSettingsActionInstance(new FPCGEditorGraphSchemaAction_NewSettingsElement(MenuAndSubCategory, LOCTEXT("ContextMenuCopySettings", "Instance"), Description, 0));
+				TSharedPtr<FPCGEditorGraphSchemaAction_NewSettingsElement> NewSettingsActionInstance(new FPCGEditorGraphSchemaAction_NewSettingsElement(MenuAndSubCategory, LOCTEXT("ContextMenuInstanceSettings", "Instance"), Description, 0));
 				NewSettingsActionInstance->SettingsObjectPath = AssetData.GetSoftObjectPath();
 				NewSettingsActionInstance->Behavior = EPCGEditorNewSettingsBehavior::ForceInstance;
 				ActionMenuBuilder.AddAction(NewSettingsActionInstance);
