[PCG] Fix tests

#rnx

diff --git a/Engine/Plugins/Experimental/PCG/Source/PCG/Private/Tests/Elements/PCGCopyPointsTest.cpp b/Engine/Plugins/Experimental/PCG/Source/PCG/Private/Tests/Elements/PCGCopyPointsTest.cpp
--- a/Engine/Plugins/Experimental/PCG/Source/PCG/Private/Tests/Elements/PCGCopyPointsTest.cpp
+++ b/Engine/Plugins/Experimental/PCG/Source/PCG/Private/Tests/Elements/PCGCopyPointsTest.cpp
@@ -68,7 +68,7 @@ bool FPCGCopyPointsTest::RunTest(const FString& Parameters)
 	// Test only supports float attributes for simplicity
 	auto ValidateCopyPointsPoints = [this, &TestData, CopyPointsElement, Settings]() -> bool
 	{
-		TUniquePtr<FPCGContext> Context = MakeUnique<FPCGContext>(*CopyPointsElement->Initialize(TestData.InputData, TestData.TestPCGComponent, nullptr));
+		TUniquePtr<FPCGContext> Context(CopyPointsElement->Initialize(TestData.InputData, TestData.TestPCGComponent, nullptr));
 		Context->NumAvailableTasks = 1;
 
 		while (!CopyPointsElement->Execute(Context.Get()))
diff --git a/Engine/Plugins/Experimental/PCG/Source/PCG/Private/Tests/Elements/PCGDensityNoiseTest.cpp b/Engine/Plugins/Experimental/PCG/Source/PCG/Private/Tests/Elements/PCGDensityNoiseTest.cpp
--- a/Engine/Plugins/Experimental/PCG/Source/PCG/Private/Tests/Elements/PCGDensityNoiseTest.cpp
+++ b/Engine/Plugins/Experimental/PCG/Source/PCG/Private/Tests/Elements/PCGDensityNoiseTest.cpp
@@ -35,7 +35,7 @@ bool FPCGDensityNoiseTest::RunTest(const FString& Parameters)
 
 	auto ValidateDensityNoise = [this, &TestData, DensityNoiseElement, Settings](TArray<float> ExpectedOutput) -> bool
 	{
-		TUniquePtr<FPCGContext> Context = MakeUnique<FPCGContext>(*DensityNoiseElement->Initialize(TestData.InputData, TestData.TestPCGComponent, nullptr));
+		TUniquePtr<FPCGContext> Context(DensityNoiseElement->Initialize(TestData.InputData, TestData.TestPCGComponent, nullptr));
 		Context->NumAvailableTasks = 1;
 
 		while (!DensityNoiseElement->Execute(Context.Get()))
diff --git a/Engine/Plugins/Experimental/PCG/Source/PCG/Private/Tests/Elements/PCGDensityRemapTest.cpp b/Engine/Plugins/Experimental/PCG/Source/PCG/Private/Tests/Elements/PCGDensityRemapTest.cpp
--- a/Engine/Plugins/Experimental/PCG/Source/PCG/Private/Tests/Elements/PCGDensityRemapTest.cpp
+++ b/Engine/Plugins/Experimental/PCG/Source/PCG/Private/Tests/Elements/PCGDensityRemapTest.cpp
@@ -35,7 +35,7 @@ bool FPCGDensityRemapTest::RunTest(const FString& Parameters)
 
 	auto ValidateDensityRemap = [this, &TestData, DensityRemapElement, Settings](TArray<float> CorrectDensities) -> bool
 	{
-		TUniquePtr<FPCGContext> Context = MakeUnique<FPCGContext>(*DensityRemapElement->Initialize(TestData.InputData, TestData.TestPCGComponent, nullptr));
+		TUniquePtr<FPCGContext> Context(DensityRemapElement->Initialize(TestData.InputData, TestData.TestPCGComponent, nullptr));
 		Context->NumAvailableTasks = 1;
 
 		while (!DensityRemapElement->Execute(Context.Get()))
diff --git a/Engine/Plugins/Experimental/PCG/Source/PCG/Private/Tests/Elements/PCGInstancePackerTest.cpp b/Engine/Plugins/Experimental/PCG/Source/PCG/Private/Tests/Elements/PCGInstancePackerTest.cpp
--- a/Engine/Plugins/Experimental/PCG/Source/PCG/Private/Tests/Elements/PCGInstancePackerTest.cpp
+++ b/Engine/Plugins/Experimental/PCG/Source/PCG/Private/Tests/Elements/PCGInstancePackerTest.cpp
@@ -57,7 +57,7 @@ namespace
 		TArray<float> ExpectedCustomData, 
 		int NumCustomDataFloats)
 	{
-		TUniquePtr<FPCGContext> Context = MakeUnique<FPCGContext>(*StaticMeshSpawnerElement->Initialize(TestData.InputData, TestData.TestPCGComponent, nullptr));
+		TUniquePtr<FPCGContext> Context(StaticMeshSpawnerElement->Initialize(TestData.InputData, TestData.TestPCGComponent, nullptr));
 		Context->NumAvailableTasks = 1;
 
 		while (!StaticMeshSpawnerElement->Execute(Context.Get()))
diff --git a/Engine/Plugins/Experimental/PCG/Source/PCG/Private/Tests/Elements/PCGPointExtentsModifierTest.cpp b/Engine/Plugins/Experimental/PCG/Source/PCG/Private/Tests/Elements/PCGPointExtentsModifierTest.cpp
--- a/Engine/Plugins/Experimental/PCG/Source/PCG/Private/Tests/Elements/PCGPointExtentsModifierTest.cpp
+++ b/Engine/Plugins/Experimental/PCG/Source/PCG/Private/Tests/Elements/PCGPointExtentsModifierTest.cpp
@@ -43,7 +43,7 @@ bool FPCGPointsBoundsModiferTest::RunTest(const FString& Parameters)
 
 	auto ValidateBounds = [this, &TestData, BoundsModifierElement, Settings](TArray<FVector> ExpectedOutput) -> bool
 	{
-		TUniquePtr<FPCGContext> Context = MakeUnique<FPCGContext>(*BoundsModifierElement->Initialize(TestData.InputData, TestData.TestPCGComponent, nullptr));
+		TUniquePtr<FPCGContext> Context(BoundsModifierElement->Initialize(TestData.InputData, TestData.TestPCGComponent, nullptr));
 		Context->NumAvailableTasks = 1;
 
 		while (!BoundsModifierElement->Execute(Context.Get()))
diff --git a/Engine/Plugins/Experimental/PCG/Source/PCG/Private/Tests/Elements/PCGTransformPointsTest.cpp b/Engine/Plugins/Experimental/PCG/Source/PCG/Private/Tests/Elements/PCGTransformPointsTest.cpp
--- a/Engine/Plugins/Experimental/PCG/Source/PCG/Private/Tests/Elements/PCGTransformPointsTest.cpp
+++ b/Engine/Plugins/Experimental/PCG/Source/PCG/Private/Tests/Elements/PCGTransformPointsTest.cpp
@@ -60,7 +60,7 @@ bool FPCGTransformPointsTest::RunTest(const FString& Parameters)
 
 	auto ValidateTransformPoints = [this, &TestData, TransformPointsElement, Settings]() -> bool
 	{
-		TUniquePtr<FPCGContext> Context = MakeUnique<FPCGContext>(*TransformPointsElement->Initialize(TestData.InputData, TestData.TestPCGComponent, nullptr));
+		TUniquePtr<FPCGContext> Context(TransformPointsElement->Initialize(TestData.InputData, TestData.TestPCGComponent, nullptr));
 		Context->NumAvailableTasks = 1;
 
 		while (!TransformPointsElement->Execute(Context.Get()))
diff --git a/Engine/Plugins/Experimental/PCG/Source/PCG/Public/PCGContext.h b/Engine/Plugins/Experimental/PCG/Source/PCG/Public/PCGContext.h
--- a/Engine/Plugins/Experimental/PCG/Source/PCG/Public/PCGContext.h
+++ b/Engine/Plugins/Experimental/PCG/Source/PCG/Public/PCGContext.h
@@ -60,7 +60,7 @@ struct PCG_API FPCGContext
 	int32 BypassedOutputCount = 0;
 
 	double EndTime = 0.0;
-	bool bIsRunningOnMainThread = false;
+	bool bIsRunningOnMainThread = true;
 
 	/** True if currently inside a PCGAsync scope - will prevent further async processing */
 	bool bIsRunningAsyncCall = false;
