[PCG] Remove "Node" in node titles.

diff --git a/Engine/Plugins/Experimental/PCG/Source/PCG/Public/Elements/Metadata/PCGMetadataElement.h b/Engine/Plugins/Experimental/PCG/Source/PCG/Public/Elements/Metadata/PCGMetadataElement.h
--- a/Engine/Plugins/Experimental/PCG/Source/PCG/Public/Elements/Metadata/PCGMetadataElement.h
+++ b/Engine/Plugins/Experimental/PCG/Source/PCG/Public/Elements/Metadata/PCGMetadataElement.h
@@ -24,7 +24,7 @@ class PCG_API UPCGMetadataOperationSettings : public UPCGSettings
 public:
 	//~Begin UPCGSettings interface
 #if WITH_EDITOR
-	virtual FName GetDefaultNodeName() const override { return FName(TEXT("MetadataOperationNode")); }
+	virtual FName GetDefaultNodeName() const override { return FName(TEXT("MetadataOperation")); }
 	virtual EPCGSettingsType GetType() const override { return EPCGSettingsType::Metadata; }
 #endif
 
diff --git a/Engine/Plugins/Experimental/PCG/Source/PCG/Public/Elements/Metadata/PCGMetadataPartition.h b/Engine/Plugins/Experimental/PCG/Source/PCG/Public/Elements/Metadata/PCGMetadataPartition.h
--- a/Engine/Plugins/Experimental/PCG/Source/PCG/Public/Elements/Metadata/PCGMetadataPartition.h
+++ b/Engine/Plugins/Experimental/PCG/Source/PCG/Public/Elements/Metadata/PCGMetadataPartition.h
@@ -15,7 +15,7 @@ class PCG_API UPCGMetadataPartitionSettings : public UPCGSettings
 public:
 	//~Begin UPCGSettings interface
 #if WITH_EDITOR
-	virtual FName GetDefaultNodeName() const override { return FName(TEXT("MetadataPartitionNode")); }
+	virtual FName GetDefaultNodeName() const override { return FName(TEXT("MetadataPartition")); }
 	virtual EPCGSettingsType GetType() const override { return EPCGSettingsType::Metadata; }
 #endif
 
diff --git a/Engine/Plugins/Experimental/PCG/Source/PCG/Public/Elements/Metadata/PCGMetadataRenameElement.h b/Engine/Plugins/Experimental/PCG/Source/PCG/Public/Elements/Metadata/PCGMetadataRenameElement.h
--- a/Engine/Plugins/Experimental/PCG/Source/PCG/Public/Elements/Metadata/PCGMetadataRenameElement.h
+++ b/Engine/Plugins/Experimental/PCG/Source/PCG/Public/Elements/Metadata/PCGMetadataRenameElement.h
@@ -15,7 +15,7 @@ class PCG_API UPCGMetadataRenameSettings : public UPCGSettings
 public:
 	//~Begin UPCGSettings interface
 #if WITH_EDITOR
-	virtual FName GetDefaultNodeName() const override { return FName(TEXT("MetadataRenameNode")); }
+	virtual FName GetDefaultNodeName() const override { return FName(TEXT("AttributeRename")); }
 	virtual EPCGSettingsType GetType() const override { return EPCGSettingsType::Metadata; }
 #endif
 
diff --git a/Engine/Plugins/Experimental/PCG/Source/PCG/Public/Elements/PCGCollapseElement.h b/Engine/Plugins/Experimental/PCG/Source/PCG/Public/Elements/PCGCollapseElement.h
--- a/Engine/Plugins/Experimental/PCG/Source/PCG/Public/Elements/PCGCollapseElement.h
+++ b/Engine/Plugins/Experimental/PCG/Source/PCG/Public/Elements/PCGCollapseElement.h
@@ -15,7 +15,7 @@ class PCG_API UPCGCollapseSettings : public UPCGSettings
 public:
 	//~Begin UPCGSettings interface
 #if WITH_EDITOR
-	virtual FName GetDefaultNodeName() const override { return FName(TEXT("CollapseNode")); }
+	virtual FName GetDefaultNodeName() const override { return FName(TEXT("Collapse")); }
 	virtual EPCGSettingsType GetType() const override { return EPCGSettingsType::Spatial; }
 #endif
 
diff --git a/Engine/Plugins/Experimental/PCG/Source/PCG/Public/Elements/PCGCopyPoints.h b/Engine/Plugins/Experimental/PCG/Source/PCG/Public/Elements/PCGCopyPoints.h
--- a/Engine/Plugins/Experimental/PCG/Source/PCG/Public/Elements/PCGCopyPoints.h
+++ b/Engine/Plugins/Experimental/PCG/Source/PCG/Public/Elements/PCGCopyPoints.h
@@ -37,7 +37,7 @@ class PCG_API UPCGCopyPointsSettings : public UPCGSettings
 public:
 	//~Begin UPCGSettings interface
 #if WITH_EDITOR
-	virtual FName GetDefaultNodeName() const override { return FName(TEXT("CopyPointsNode")); }
+	virtual FName GetDefaultNodeName() const override { return FName(TEXT("CopyPoints")); }
 	virtual EPCGSettingsType GetType() const override { return EPCGSettingsType::Sampler; }
 #endif
 
diff --git a/Engine/Plugins/Experimental/PCG/Source/PCG/Public/Elements/PCGDataFromActor.h b/Engine/Plugins/Experimental/PCG/Source/PCG/Public/Elements/PCGDataFromActor.h
--- a/Engine/Plugins/Experimental/PCG/Source/PCG/Public/Elements/PCGDataFromActor.h
+++ b/Engine/Plugins/Experimental/PCG/Source/PCG/Public/Elements/PCGDataFromActor.h
@@ -26,7 +26,7 @@ class PCG_API UPCGDataFromActorSettings : public UPCGSettings
 public:
 	//~Begin UPCGSettings interface
 #if WITH_EDITOR
-	virtual FName GetDefaultNodeName() const override { return FName(TEXT("DataFromActorNode")); }
+	virtual FName GetDefaultNodeName() const override { return FName(TEXT("DataFromActor")); }
 	virtual EPCGSettingsType GetType() const override { return EPCGSettingsType::Spatial; }
 #endif
 
diff --git a/Engine/Plugins/Experimental/PCG/Source/PCG/Public/Elements/PCGDebugElement.h b/Engine/Plugins/Experimental/PCG/Source/PCG/Public/Elements/PCGDebugElement.h
--- a/Engine/Plugins/Experimental/PCG/Source/PCG/Public/Elements/PCGDebugElement.h
+++ b/Engine/Plugins/Experimental/PCG/Source/PCG/Public/Elements/PCGDebugElement.h
@@ -20,7 +20,7 @@ class PCG_API UPCGDebugSettings : public UPCGSettings
 public:
 #if WITH_EDITOR
 	// ~Begin UPCGSettings interface
-	virtual FName GetDefaultNodeName() const override { return FName(TEXT("DebugNode")); }
+	virtual FName GetDefaultNodeName() const override { return FName(TEXT("Debug")); }
 	virtual EPCGSettingsType GetType() const override { return EPCGSettingsType::Debug; }
 #endif
 
diff --git a/Engine/Plugins/Experimental/PCG/Source/PCG/Public/Elements/PCGDensityNoise.h b/Engine/Plugins/Experimental/PCG/Source/PCG/Public/Elements/PCGDensityNoise.h
--- a/Engine/Plugins/Experimental/PCG/Source/PCG/Public/Elements/PCGDensityNoise.h
+++ b/Engine/Plugins/Experimental/PCG/Source/PCG/Public/Elements/PCGDensityNoise.h
@@ -31,7 +31,7 @@ public:
 
 	// ~Begin UPCGSettings interface
 #if WITH_EDITOR
-	virtual FName GetDefaultNodeName() const override { return FName(TEXT("DensityNoiseNode")); }
+	virtual FName GetDefaultNodeName() const override { return FName(TEXT("DensityNoise")); }
 	virtual EPCGSettingsType GetType() const override { return EPCGSettingsType::Density; }
 #endif
 
diff --git a/Engine/Plugins/Experimental/PCG/Source/PCG/Public/Elements/PCGDensityRemapElement.h b/Engine/Plugins/Experimental/PCG/Source/PCG/Public/Elements/PCGDensityRemapElement.h
--- a/Engine/Plugins/Experimental/PCG/Source/PCG/Public/Elements/PCGDensityRemapElement.h
+++ b/Engine/Plugins/Experimental/PCG/Source/PCG/Public/Elements/PCGDensityRemapElement.h
@@ -19,7 +19,7 @@ class PCG_API UPCGDensityRemapSettings : public UPCGSettings
 public:
 	// ~Begin UPCGSettings interface
 #if WITH_EDITOR
-	virtual FName GetDefaultNodeName() const override { return FName(TEXT("DensityRemapNode")); }
+	virtual FName GetDefaultNodeName() const override { return FName(TEXT("DensityRemap")); }
 	virtual EPCGSettingsType GetType() const override { return EPCGSettingsType::Density; }
 #endif
 
diff --git a/Engine/Plugins/Experimental/PCG/Source/PCG/Public/Elements/PCGDifferenceElement.h b/Engine/Plugins/Experimental/PCG/Source/PCG/Public/Elements/PCGDifferenceElement.h
--- a/Engine/Plugins/Experimental/PCG/Source/PCG/Public/Elements/PCGDifferenceElement.h
+++ b/Engine/Plugins/Experimental/PCG/Source/PCG/Public/Elements/PCGDifferenceElement.h
@@ -23,7 +23,7 @@ class PCG_API UPCGDifferenceSettings : public UPCGSettings
 public:
 	//~Begin UPCGSettings interface
 #if WITH_EDITOR
-	virtual FName GetDefaultNodeName() const override { return FName(TEXT("DifferenceNode")); }
+	virtual FName GetDefaultNodeName() const override { return FName(TEXT("Difference")); }
 	virtual EPCGSettingsType GetType() const override { return EPCGSettingsType::Spatial; }
 #endif
 
diff --git a/Engine/Plugins/Experimental/PCG/Source/PCG/Public/Elements/PCGExecuteBlueprint.h b/Engine/Plugins/Experimental/PCG/Source/PCG/Public/Elements/PCGExecuteBlueprint.h
--- a/Engine/Plugins/Experimental/PCG/Source/PCG/Public/Elements/PCGExecuteBlueprint.h
+++ b/Engine/Plugins/Experimental/PCG/Source/PCG/Public/Elements/PCGExecuteBlueprint.h
@@ -167,7 +167,7 @@ public:
 
 	// ~Begin UPCGSettings interface
 #if WITH_EDITOR
-	virtual FName GetDefaultNodeName() const override { return FName(TEXT("BlueprintNode")); }
+	virtual FName GetDefaultNodeName() const override { return FName(TEXT("ExecuteBlueprint")); }
 	virtual FLinearColor GetNodeTitleColor() const override;
 	virtual EPCGSettingsType GetType() const override;
 	virtual void GetTrackedActorTags(FPCGTagToSettingsMap& OutTagToSettings, TArray<TObjectPtr<const UPCGGraph>>& OutVisitedGraphs) const override;
diff --git a/Engine/Plugins/Experimental/PCG/Source/PCG/Public/Elements/PCGIntersectionElement.h b/Engine/Plugins/Experimental/PCG/Source/PCG/Public/Elements/PCGIntersectionElement.h
--- a/Engine/Plugins/Experimental/PCG/Source/PCG/Public/Elements/PCGIntersectionElement.h
+++ b/Engine/Plugins/Experimental/PCG/Source/PCG/Public/Elements/PCGIntersectionElement.h
@@ -17,7 +17,7 @@ class PCG_API UPCGIntersectionSettings : public UPCGSettings
 public:
 	//~Begin UPCGSettings interface
 #if WITH_EDITOR
-	virtual FName GetDefaultNodeName() const override { return FName(TEXT("IntersectionNode")); }
+	virtual FName GetDefaultNodeName() const override { return FName(TEXT("Intersection")); }
 	virtual EPCGSettingsType GetType() const override { return EPCGSettingsType::Spatial; }
 #endif
 
diff --git a/Engine/Plugins/Experimental/PCG/Source/PCG/Public/Elements/PCGPointExtentsModifier.h b/Engine/Plugins/Experimental/PCG/Source/PCG/Public/Elements/PCGPointExtentsModifier.h
--- a/Engine/Plugins/Experimental/PCG/Source/PCG/Public/Elements/PCGPointExtentsModifier.h
+++ b/Engine/Plugins/Experimental/PCG/Source/PCG/Public/Elements/PCGPointExtentsModifier.h
@@ -30,7 +30,7 @@ public:
 
 	// ~Begin UPCGSettings interface
 #if WITH_EDITOR
-	virtual FName GetDefaultNodeName() const override { return FName(TEXT("ExtentsModifierNode")); }
+	virtual FName GetDefaultNodeName() const override { return FName(TEXT("ExtentsModifier")); }
 	virtual EPCGSettingsType GetType() const override { return EPCGSettingsType::Spatial; }
 #endif
 
diff --git a/Engine/Plugins/Experimental/PCG/Source/PCG/Public/Elements/PCGPointFilter.h b/Engine/Plugins/Experimental/PCG/Source/PCG/Public/Elements/PCGPointFilter.h
--- a/Engine/Plugins/Experimental/PCG/Source/PCG/Public/Elements/PCGPointFilter.h
+++ b/Engine/Plugins/Experimental/PCG/Source/PCG/Public/Elements/PCGPointFilter.h
@@ -66,7 +66,7 @@ class PCG_API UPCGPointFilterSettings : public UPCGSettings
 public:
 	//~Begin UPCGSettings interface
 #if WITH_EDITOR
-	virtual FName GetDefaultNodeName() const override { return FName(TEXT("PointFilterNode")); }
+	virtual FName GetDefaultNodeName() const override { return FName(TEXT("PointFilter")); }
 	virtual EPCGSettingsType GetType() const override { return EPCGSettingsType::Filter; }
 #endif
 
diff --git a/Engine/Plugins/Experimental/PCG/Source/PCG/Public/Elements/PCGPointSampler.h b/Engine/Plugins/Experimental/PCG/Source/PCG/Public/Elements/PCGPointSampler.h
--- a/Engine/Plugins/Experimental/PCG/Source/PCG/Public/Elements/PCGPointSampler.h
+++ b/Engine/Plugins/Experimental/PCG/Source/PCG/Public/Elements/PCGPointSampler.h
@@ -19,7 +19,7 @@ public:
 
 	//~Begin UPCGSettings interface
 #if WITH_EDITOR
-	virtual FName GetDefaultNodeName() const override { return FName(TEXT("PointSamplerNode")); }
+	virtual FName GetDefaultNodeName() const override { return FName(TEXT("PointSampler")); }
 	virtual EPCGSettingsType GetType() const override { return EPCGSettingsType::Sampler; }
 #endif
 
diff --git a/Engine/Plugins/Experimental/PCG/Source/PCG/Public/Elements/PCGProjectionElement.h b/Engine/Plugins/Experimental/PCG/Source/PCG/Public/Elements/PCGProjectionElement.h
--- a/Engine/Plugins/Experimental/PCG/Source/PCG/Public/Elements/PCGProjectionElement.h
+++ b/Engine/Plugins/Experimental/PCG/Source/PCG/Public/Elements/PCGProjectionElement.h
@@ -23,7 +23,7 @@ class PCG_API UPCGProjectionSettings : public UPCGSettings
 public:
 	//~Begin UPCGSettings interface
 #if WITH_EDITOR
-	virtual FName GetDefaultNodeName() const override { return FName(TEXT("ProjectionNode")); }
+	virtual FName GetDefaultNodeName() const override { return FName(TEXT("Projection")); }
 	virtual FText GetNodeTooltipText() const override;
 	virtual EPCGSettingsType GetType() const override { return EPCGSettingsType::Spatial; }
 #endif
diff --git a/Engine/Plugins/Experimental/PCG/Source/PCG/Public/Elements/PCGPropertyToParamData.h b/Engine/Plugins/Experimental/PCG/Source/PCG/Public/Elements/PCGPropertyToParamData.h
--- a/Engine/Plugins/Experimental/PCG/Source/PCG/Public/Elements/PCGPropertyToParamData.h
+++ b/Engine/Plugins/Experimental/PCG/Source/PCG/Public/Elements/PCGPropertyToParamData.h
@@ -22,7 +22,7 @@ public:
 
 	//~Begin UPCGSettings interface
 #if WITH_EDITOR
-	virtual FName GetDefaultNodeName() const override { return FName(TEXT("PropertyToParamDataNode")); }
+	virtual FName GetDefaultNodeName() const override { return FName(TEXT("PropertyToParamData")); }
 	virtual EPCGSettingsType GetType() const override { return EPCGSettingsType::Param; }
 #endif
 
diff --git a/Engine/Plugins/Experimental/PCG/Source/PCG/Public/Elements/PCGSelfPruning.h b/Engine/Plugins/Experimental/PCG/Source/PCG/Public/Elements/PCGSelfPruning.h
--- a/Engine/Plugins/Experimental/PCG/Source/PCG/Public/Elements/PCGSelfPruning.h
+++ b/Engine/Plugins/Experimental/PCG/Source/PCG/Public/Elements/PCGSelfPruning.h
@@ -41,7 +41,7 @@ public:
 
 #if WITH_EDITOR
 	//~Begin UPCGSettings interface
-	virtual FName GetDefaultNodeName() const override { return FName(TEXT("SelfPruningNode")); }
+	virtual FName GetDefaultNodeName() const override { return FName(TEXT("SelfPruning")); }
 	virtual EPCGSettingsType GetType() const override { return EPCGSettingsType::Filter; }
 #endif
 
diff --git a/Engine/Plugins/Experimental/PCG/Source/PCG/Public/Elements/PCGSpawnActor.h b/Engine/Plugins/Experimental/PCG/Source/PCG/Public/Elements/PCGSpawnActor.h
--- a/Engine/Plugins/Experimental/PCG/Source/PCG/Public/Elements/PCGSpawnActor.h
+++ b/Engine/Plugins/Experimental/PCG/Source/PCG/Public/Elements/PCGSpawnActor.h
@@ -62,7 +62,7 @@ public:
 	virtual UPCGNode* CreateNode() const override;
 
 #if WITH_EDITOR	
-	virtual FName GetDefaultNodeName() const override { return FName(TEXT("SpawnActorNode")); }
+	virtual FName GetDefaultNodeName() const override { return FName(TEXT("SpawnActor")); }
 	virtual EPCGSettingsType GetType() const override;
 #endif
 protected:
diff --git a/Engine/Plugins/Experimental/PCG/Source/PCG/Public/Elements/PCGSplineSampler.h b/Engine/Plugins/Experimental/PCG/Source/PCG/Public/Elements/PCGSplineSampler.h
--- a/Engine/Plugins/Experimental/PCG/Source/PCG/Public/Elements/PCGSplineSampler.h
+++ b/Engine/Plugins/Experimental/PCG/Source/PCG/Public/Elements/PCGSplineSampler.h
@@ -116,7 +116,7 @@ class PCG_API UPCGSplineSamplerSettings : public UPCGSettings
 public:
 	// ~Begin UPCGSettings interface
 #if WITH_EDITOR
-	virtual FName GetDefaultNodeName() const override { return FName(TEXT("SplineSamplerNode")); }
+	virtual FName GetDefaultNodeName() const override { return FName(TEXT("SplineSampler")); }
 	virtual EPCGSettingsType GetType() const override { return EPCGSettingsType::Sampler; }
 #endif
 
diff --git a/Engine/Plugins/Experimental/PCG/Source/PCG/Public/Elements/PCGStaticMeshSpawner.h b/Engine/Plugins/Experimental/PCG/Source/PCG/Public/Elements/PCGStaticMeshSpawner.h
--- a/Engine/Plugins/Experimental/PCG/Source/PCG/Public/Elements/PCGStaticMeshSpawner.h
+++ b/Engine/Plugins/Experimental/PCG/Source/PCG/Public/Elements/PCGStaticMeshSpawner.h
@@ -42,7 +42,7 @@ public:
 
 #if WITH_EDITOR
 	// ~Begin UPCGSettings interface
-	virtual FName GetDefaultNodeName() const override { return FName(TEXT("StaticMeshSpawnerNode")); }
+	virtual FName GetDefaultNodeName() const override { return FName(TEXT("StaticMeshSpawner")); }
 	virtual EPCGSettingsType GetType() const override { return EPCGSettingsType::Spawner; }
 #endif
 
diff --git a/Engine/Plugins/Experimental/PCG/Source/PCG/Public/Elements/PCGSurfaceSampler.h b/Engine/Plugins/Experimental/PCG/Source/PCG/Public/Elements/PCGSurfaceSampler.h
--- a/Engine/Plugins/Experimental/PCG/Source/PCG/Public/Elements/PCGSurfaceSampler.h
+++ b/Engine/Plugins/Experimental/PCG/Source/PCG/Public/Elements/PCGSurfaceSampler.h
@@ -50,7 +50,7 @@ public:
 
 #if WITH_EDITOR
 	//~Begin UPCGSettings interface
-	virtual FName GetDefaultNodeName() const override { return FName(TEXT("SurfaceSamplerNode")); }
+	virtual FName GetDefaultNodeName() const override { return FName(TEXT("SurfaceSampler")); }
 	virtual EPCGSettingsType GetType() const override { return EPCGSettingsType::Sampler; }
 #endif
 
diff --git a/Engine/Plugins/Experimental/PCG/Source/PCG/Public/Elements/PCGTextureSampler.h b/Engine/Plugins/Experimental/PCG/Source/PCG/Public/Elements/PCGTextureSampler.h
--- a/Engine/Plugins/Experimental/PCG/Source/PCG/Public/Elements/PCGTextureSampler.h
+++ b/Engine/Plugins/Experimental/PCG/Source/PCG/Public/Elements/PCGTextureSampler.h
@@ -16,7 +16,7 @@ class PCG_API UPCGTextureSamplerSettings : public UPCGSettings
 public:
 	//~Begin UPCGSettings interface
 #if WITH_EDITOR
-	virtual FName GetDefaultNodeName() const override { return FName(TEXT("TextureSamplerNode")); }
+	virtual FName GetDefaultNodeName() const override { return FName(TEXT("TextureSampler")); }
 	virtual EPCGSettingsType GetType() const override { return EPCGSettingsType::Sampler; }
 #endif
 
diff --git a/Engine/Plugins/Experimental/PCG/Source/PCG/Public/Elements/PCGTransformPoints.h b/Engine/Plugins/Experimental/PCG/Source/PCG/Public/Elements/PCGTransformPoints.h
--- a/Engine/Plugins/Experimental/PCG/Source/PCG/Public/Elements/PCGTransformPoints.h
+++ b/Engine/Plugins/Experimental/PCG/Source/PCG/Public/Elements/PCGTransformPoints.h
@@ -19,7 +19,7 @@ public:
 
 	//~Begin UPCGSettings interface
 #if WITH_EDITOR
-	virtual FName GetDefaultNodeName() const override { return FName(TEXT("TransformPointsNode")); }
+	virtual FName GetDefaultNodeName() const override { return FName(TEXT("TransformPoints")); }
 	virtual EPCGSettingsType GetType() const override { return EPCGSettingsType::Spatial; }
 #endif
 
diff --git a/Engine/Plugins/Experimental/PCG/Source/PCG/Public/Elements/PCGUnionElement.h b/Engine/Plugins/Experimental/PCG/Source/PCG/Public/Elements/PCGUnionElement.h
--- a/Engine/Plugins/Experimental/PCG/Source/PCG/Public/Elements/PCGUnionElement.h
+++ b/Engine/Plugins/Experimental/PCG/Source/PCG/Public/Elements/PCGUnionElement.h
@@ -17,7 +17,7 @@ class PCG_API UPCGUnionSettings : public UPCGSettings
 public:
 	//~Begin UPCGSettings interface
 #if WITH_EDITOR
-	virtual FName GetDefaultNodeName() const override { return FName(TEXT("UnionNode")); }
+	virtual FName GetDefaultNodeName() const override { return FName(TEXT("Union")); }
 	virtual EPCGSettingsType GetType() const override { return EPCGSettingsType::Spatial; }
 #endif
 
diff --git a/Engine/Plugins/Experimental/PCG/Source/PCG/Public/Elements/PCGVolumeSampler.h b/Engine/Plugins/Experimental/PCG/Source/PCG/Public/Elements/PCGVolumeSampler.h
--- a/Engine/Plugins/Experimental/PCG/Source/PCG/Public/Elements/PCGVolumeSampler.h
+++ b/Engine/Plugins/Experimental/PCG/Source/PCG/Public/Elements/PCGVolumeSampler.h
@@ -32,7 +32,7 @@ public:
 
 	//~Begin UPCGSettings interface
 #if WITH_EDITOR
-	virtual FName GetDefaultNodeName() const override { return FName(TEXT("VolumeSamplerNode")); }
+	virtual FName GetDefaultNodeName() const override { return FName(TEXT("VolumeSampler")); }
 	virtual EPCGSettingsType GetType() const override { return EPCGSettingsType::Sampler; }
 #endif
 
diff --git a/Engine/Plugins/Experimental/PCG/Source/PCG/Public/Elements/PCGWorldQuery.h b/Engine/Plugins/Experimental/PCG/Source/PCG/Public/Elements/PCGWorldQuery.h
--- a/Engine/Plugins/Experimental/PCG/Source/PCG/Public/Elements/PCGWorldQuery.h
+++ b/Engine/Plugins/Experimental/PCG/Source/PCG/Public/Elements/PCGWorldQuery.h
@@ -16,7 +16,7 @@ class PCG_API UPCGWorldQuerySettings : public UPCGSettings
 public:
 	//~Begin UPCGSettings interface
 #if WITH_EDITOR
-	virtual FName GetDefaultNodeName() const override { return FName(TEXT("WorldVolumetricNode")); }
+	virtual FName GetDefaultNodeName() const override { return FName(TEXT("WorldVolumetric")); }
 	virtual EPCGSettingsType GetType() const override { return EPCGSettingsType::Spatial; }
 #endif
 
@@ -47,7 +47,7 @@ class PCG_API UPCGWorldRayHitSettings : public UPCGSettings
 public:
 	//~Begin UPCGSettings interface
 #if WITH_EDITOR
-	virtual FName GetDefaultNodeName() const override { return FName(TEXT("WorldRayHitNode")); }
+	virtual FName GetDefaultNodeName() const override { return FName(TEXT("WorldRayHit")); }
 	virtual EPCGSettingsType GetType() const override { return EPCGSettingsType::Spatial; }
 #endif
 
diff --git a/Engine/Plugins/Experimental/PCG/Source/PCG/Public/PCGSubgraph.h b/Engine/Plugins/Experimental/PCG/Source/PCG/Public/PCGSubgraph.h
--- a/Engine/Plugins/Experimental/PCG/Source/PCG/Public/PCGSubgraph.h
+++ b/Engine/Plugins/Experimental/PCG/Source/PCG/Public/PCGSubgraph.h
@@ -48,7 +48,7 @@ public:
 	virtual UPCGNode* CreateNode() const override;
 
 #if WITH_EDITOR
-	virtual FName GetDefaultNodeName() const override { return FName(TEXT("SubgraphNode")); }
+	virtual FName GetDefaultNodeName() const override { return FName(TEXT("Subgraph")); }
 	virtual EPCGSettingsType GetType() const override { return EPCGSettingsType::Subgraph; }
 	virtual UObject* GetJumpTargetForDoubleClick() const override;
 #endif
