Added methods to be able to get the current settings from the context
Added methods to get the seed/random stream in a PCG BP element

diff --git a/Engine/Plugins/Experimental/PCG/Source/PCG/Private/Elements/PCGExecuteBlueprint.cpp b/Engine/Plugins/Experimental/PCG/Source/PCG/Private/Elements/PCGExecuteBlueprint.cpp
--- a/Engine/Plugins/Experimental/PCG/Source/PCG/Private/Elements/PCGExecuteBlueprint.cpp
+++ b/Engine/Plugins/Experimental/PCG/Source/PCG/Private/Elements/PCGExecuteBlueprint.cpp
@@ -155,6 +155,34 @@ TSet<FName> UPCGBlueprintElement::OutputLabels() const
 	return Labels;
 }
 
+int UPCGBlueprintElement::GetSeed(FPCGContext& InContext) const 
+{
+	const UPCGComponent* SourceComponent = InContext.SourceComponent.Get();
+	const UPCGBlueprintSettings* Settings = InContext.GetInputSettings<UPCGBlueprintSettings>();
+
+	if (SourceComponent && Settings)
+	{
+		return PCGHelpers::ComputeSeed(Settings->Seed, SourceComponent->Seed);
+	}
+	else if (Settings)
+	{
+		return Settings->Seed;
+	}
+	else if (SourceComponent)
+	{
+		return SourceComponent->Seed;
+	}
+	else
+	{
+		return 42;
+	}
+}
+
+FRandomStream UPCGBlueprintElement::GetRandomStream(FPCGContext& InContext) const
+{
+	return FRandomStream(GetSeed(InContext));
+}
+
 UPCGBlueprintSettings::UPCGBlueprintSettings()
 {
 	bUseSeed = true;
diff --git a/Engine/Plugins/Experimental/PCG/Source/PCG/Private/Helpers/PCGBlueprintHelpers.cpp b/Engine/Plugins/Experimental/PCG/Source/PCG/Private/Helpers/PCGBlueprintHelpers.cpp
--- a/Engine/Plugins/Experimental/PCG/Source/PCG/Private/Helpers/PCGBlueprintHelpers.cpp
+++ b/Engine/Plugins/Experimental/PCG/Source/PCG/Private/Helpers/PCGBlueprintHelpers.cpp
@@ -46,6 +46,11 @@ FRandomStream UPCGBlueprintHelpers::GetRandomStream(const FPCGPoint& InPoint, co
 	return FRandomStream(Seed);
 }
 
+const UPCGSettings* UPCGBlueprintHelpers::GetSettings(FPCGContext& Context)
+{
+	return Context.GetInputSettings<UPCGSettings>();
+}
+
 UPCGData* UPCGBlueprintHelpers::GetActorData(FPCGContext& Context)
 {
 	return Context.SourceComponent.IsValid() ? Context.SourceComponent->GetActorPCGData() : nullptr;
diff --git a/Engine/Plugins/Experimental/PCG/Source/PCG/Public/Elements/PCGExecuteBlueprint.h b/Engine/Plugins/Experimental/PCG/Source/PCG/Public/Elements/PCGExecuteBlueprint.h
--- a/Engine/Plugins/Experimental/PCG/Source/PCG/Public/Elements/PCGExecuteBlueprint.h
+++ b/Engine/Plugins/Experimental/PCG/Source/PCG/Public/Elements/PCGExecuteBlueprint.h
@@ -8,6 +8,7 @@
 #include "Data/PCGPointData.h"
 #include "PCGPoint.h"
 
+#include "Math/RandomStream.h"
 #include "Templates/SubclassOf.h"
 
 #include "PCGExecuteBlueprint.generated.h"
@@ -81,6 +82,14 @@ public:
 	UFUNCTION(BlueprintCallable, Category = "Input & Output")
 	TSet<FName> OutputLabels() const;
 
+	/** Gets the seed from the associated settings & source component */
+	UFUNCTION(BlueprintCallable, Category = "PCG|Random")
+	int GetSeed(UPARAM(ref) FPCGContext& InContext) const;
+
+	/** Creates a random stream from the settings & source component */
+	UFUNCTION(BlueprintCallable, BlueprintPure = false, Category = "PCG|Random")
+	FRandomStream GetRandomStream(UPARAM(ref) FPCGContext& InContext) const;
+
 	/** Called after object creation to setup the object callbacks */
 	void Initialize();
 
diff --git a/Engine/Plugins/Experimental/PCG/Source/PCG/Public/Helpers/PCGBlueprintHelpers.h b/Engine/Plugins/Experimental/PCG/Source/PCG/Public/Helpers/PCGBlueprintHelpers.h
--- a/Engine/Plugins/Experimental/PCG/Source/PCG/Public/Helpers/PCGBlueprintHelpers.h
+++ b/Engine/Plugins/Experimental/PCG/Source/PCG/Public/Helpers/PCGBlueprintHelpers.h
@@ -30,6 +30,9 @@ public:
 	UFUNCTION(BlueprintCallable, Category = "PCG|Helpers", meta = (ScriptMethod))
 	static FRandomStream GetRandomStream(const FPCGPoint& InPoint, const UPCGSettings* OptionalSettings = nullptr, const UPCGComponent* OptionalComponent = nullptr);
 
+	UFUNCTION(BlueprintCallable, Category = "PCG|Helpers", meta = (ScriptMethod))
+	static const UPCGSettings* GetSettings(UPARAM(ref) FPCGContext& Context);
+
 	UFUNCTION(BlueprintCallable, Category = "PCG|Temporary", meta = (ScriptMethod))
 	static UPCGData* GetActorData(UPARAM(ref) FPCGContext& Context);
 
