[PCG] Make the use of properties/attributes transparent for the user for Metadata operations.

* Add a wrapper around properties and attributes for getter and setter
* First pass, only modify the basic metadata operation (Maths, Vector, etc...)

diff --git a/Engine/Plugins/Experimental/PCG/Source/PCG/Private/Data/PCGProjectionData.cpp b/Engine/Plugins/Experimental/PCG/Source/PCG/Private/Data/PCGProjectionData.cpp
--- a/Engine/Plugins/Experimental/PCG/Source/PCG/Private/Data/PCGProjectionData.cpp
+++ b/Engine/Plugins/Experimental/PCG/Source/PCG/Private/Data/PCGProjectionData.cpp
@@ -205,7 +205,7 @@ void UPCGProjectionData::GetIncludeExcludeAttributeNames(TSet<FName>& OutAttribu
 	}
 }
 
-inline void UPCGProjectionData::CopyBaseProjectionClass(UPCGProjectionData* NewProjectionData) const
+void UPCGProjectionData::CopyBaseProjectionClass(UPCGProjectionData* NewProjectionData) const
 {
 	NewProjectionData->Source = Source;
 	NewProjectionData->Target = Target;
diff --git a/Engine/Plugins/Experimental/PCG/Source/PCG/Private/Data/PCGVolumeData.cpp b/Engine/Plugins/Experimental/PCG/Source/PCG/Private/Data/PCGVolumeData.cpp
--- a/Engine/Plugins/Experimental/PCG/Source/PCG/Private/Data/PCGVolumeData.cpp
+++ b/Engine/Plugins/Experimental/PCG/Source/PCG/Private/Data/PCGVolumeData.cpp
@@ -86,7 +86,7 @@ bool UPCGVolumeData::SamplePoint(const FTransform& InTransform, const FBox& InBo
 	}
 }
 
-inline void UPCGVolumeData::CopyBaseVolumeData(UPCGVolumeData* NewVolumeData) const
+void UPCGVolumeData::CopyBaseVolumeData(UPCGVolumeData* NewVolumeData) const
 {
 	NewVolumeData->VoxelSize = VoxelSize;
 	NewVolumeData->Volume = Volume;
diff --git a/Engine/Plugins/Experimental/PCG/Source/PCG/Private/Elements/Metadata/PCGMetadataBitwiseOpElement.cpp b/Engine/Plugins/Experimental/PCG/Source/PCG/Private/Elements/Metadata/PCGMetadataBitwiseOpElement.cpp
--- a/Engine/Plugins/Experimental/PCG/Source/PCG/Private/Elements/Metadata/PCGMetadataBitwiseOpElement.cpp
+++ b/Engine/Plugins/Experimental/PCG/Source/PCG/Private/Elements/Metadata/PCGMetadataBitwiseOpElement.cpp
@@ -35,6 +35,27 @@ namespace PCGMetadataBitwiseSettings
 	}
 }
 
+void UPCGMetadataBitwiseSettings::PostLoad()
+{
+	Super::PostLoad();
+
+#if WITH_EDITOR
+	if (Input1AttributeName_DEPRECATED != NAME_None)
+	{
+		InputSource1.Selection = EPCGAttributePropertySelection::Attribute;
+		InputSource1.AttributeName = Input1AttributeName_DEPRECATED;
+		Input1AttributeName_DEPRECATED = NAME_None;
+	}
+
+	if (Input2AttributeName_DEPRECATED != NAME_None)
+	{
+		InputSource2.Selection = EPCGAttributePropertySelection::Attribute;
+		InputSource2.AttributeName = Input2AttributeName_DEPRECATED;
+		Input2AttributeName_DEPRECATED = NAME_None;
+	}
+#endif // WITH_EDITOR
+}
+
 FName UPCGMetadataBitwiseSettings::GetInputPinLabel(uint32 Index) const
 {
 	switch (Index)
@@ -59,16 +80,16 @@ bool UPCGMetadataBitwiseSettings::IsSupportedInputType(uint16 TypeId, uint32 Inp
 	return PCG::Private::IsOfTypes<int32, int64>(TypeId);
 }
 
-FName UPCGMetadataBitwiseSettings::GetInputAttributeNameWithOverride(uint32 Index, UPCGParamData* Params) const
+FPCGAttributePropertySelector UPCGMetadataBitwiseSettings::GetInputSource(uint32 Index) const
 {
 	switch (Index)
 	{
 	case 0:
-		return PCG_GET_OVERRIDEN_VALUE(this, Input1AttributeName, Params);
+		return InputSource1;
 	case 1:
-		return PCG_GET_OVERRIDEN_VALUE(this, Input2AttributeName, Params);
+		return InputSource2;
 	default:
-		return NAME_None;
+		return FPCGAttributePropertySelector();
 	}
 }
 
diff --git a/Engine/Plugins/Experimental/PCG/Source/PCG/Private/Elements/Metadata/PCGMetadataBooleanOpElement.cpp b/Engine/Plugins/Experimental/PCG/Source/PCG/Private/Elements/Metadata/PCGMetadataBooleanOpElement.cpp
--- a/Engine/Plugins/Experimental/PCG/Source/PCG/Private/Elements/Metadata/PCGMetadataBooleanOpElement.cpp
+++ b/Engine/Plugins/Experimental/PCG/Source/PCG/Private/Elements/Metadata/PCGMetadataBooleanOpElement.cpp
@@ -35,6 +35,27 @@ namespace PCGMetadataBooleanSettings
 	}
 }
 
+void UPCGMetadataBooleanSettings::PostLoad()
+{
+	Super::PostLoad();
+
+#if WITH_EDITOR
+	if (Input1AttributeName_DEPRECATED != NAME_None)
+	{
+		InputSource1.Selection = EPCGAttributePropertySelection::Attribute;
+		InputSource1.AttributeName = Input1AttributeName_DEPRECATED;
+		Input1AttributeName_DEPRECATED = NAME_None;
+	}
+
+	if (Input2AttributeName_DEPRECATED != NAME_None)
+	{
+		InputSource2.Selection = EPCGAttributePropertySelection::Attribute;
+		InputSource2.AttributeName = Input2AttributeName_DEPRECATED;
+		Input2AttributeName_DEPRECATED = NAME_None;
+	}
+#endif // WITH_EDITOR
+}
+
 FName UPCGMetadataBooleanSettings::GetInputPinLabel(uint32 Index) const
 {
 	switch (Index)
@@ -59,16 +80,16 @@ bool UPCGMetadataBooleanSettings::IsSupportedInputType(uint16 TypeId, uint32 Inp
 	return PCG::Private::IsOfTypes<bool>(TypeId);
 }
 
-FName UPCGMetadataBooleanSettings::GetInputAttributeNameWithOverride(uint32 Index, UPCGParamData* Params) const
+FPCGAttributePropertySelector UPCGMetadataBooleanSettings::GetInputSource(uint32 Index) const
 {
 	switch (Index)
 	{
 	case 0:
-		return PCG_GET_OVERRIDEN_VALUE(this, Input1AttributeName, Params);
+		return InputSource1;
 	case 1:
-		return PCG_GET_OVERRIDEN_VALUE(this, Input2AttributeName, Params);
+		return InputSource2;
 	default:
-		return NAME_None;
+		return FPCGAttributePropertySelector();
 	}
 }
 
diff --git a/Engine/Plugins/Experimental/PCG/Source/PCG/Private/Elements/Metadata/PCGMetadataBreakTransform.cpp b/Engine/Plugins/Experimental/PCG/Source/PCG/Private/Elements/Metadata/PCGMetadataBreakTransform.cpp
--- a/Engine/Plugins/Experimental/PCG/Source/PCG/Private/Elements/Metadata/PCGMetadataBreakTransform.cpp
+++ b/Engine/Plugins/Experimental/PCG/Source/PCG/Private/Elements/Metadata/PCGMetadataBreakTransform.cpp
@@ -28,9 +28,23 @@ namespace PCGMetadataBreakTransformSettings
 	}
 }
 
-FName UPCGMetadataBreakTransformSettings::GetInputAttributeNameWithOverride(uint32 Index, UPCGParamData* Params) const
+void UPCGMetadataBreakTransformSettings::PostLoad()
 {
-	return PCG_GET_OVERRIDEN_VALUE(this, InputAttributeName, Params);
+	Super::PostLoad();
+
+#if WITH_EDITOR
+	if (InputAttributeName_DEPRECATED != NAME_None)
+	{
+		InputSource.Selection = EPCGAttributePropertySelection::Attribute;
+		InputSource.AttributeName = InputAttributeName_DEPRECATED;
+		InputAttributeName_DEPRECATED = NAME_None;
+	}
+#endif // WITH_EDITOR
+}
+
+FPCGAttributePropertySelector UPCGMetadataBreakTransformSettings::GetInputSource(uint32 Index) const
+{
+	return InputSource;
 }
 
 FName UPCGMetadataBreakTransformSettings::GetOutputPinLabel(uint32 Index) const
diff --git a/Engine/Plugins/Experimental/PCG/Source/PCG/Private/Elements/Metadata/PCGMetadataBreakVector.cpp b/Engine/Plugins/Experimental/PCG/Source/PCG/Private/Elements/Metadata/PCGMetadataBreakVector.cpp
--- a/Engine/Plugins/Experimental/PCG/Source/PCG/Private/Elements/Metadata/PCGMetadataBreakVector.cpp
+++ b/Engine/Plugins/Experimental/PCG/Source/PCG/Private/Elements/Metadata/PCGMetadataBreakVector.cpp
@@ -78,9 +78,23 @@ namespace PCGMetadataBreakVectorSettings
 	}
 }
 
-FName UPCGMetadataBreakVectorSettings::GetInputAttributeNameWithOverride(uint32 Index, UPCGParamData* Params) const
+void UPCGMetadataBreakVectorSettings::PostLoad()
 {
-	return PCG_GET_OVERRIDEN_VALUE(this, InputAttributeName, Params);
+	Super::PostLoad();
+
+#if WITH_EDITOR
+	if (InputAttributeName_DEPRECATED != NAME_None)
+	{
+		InputSource.Selection = EPCGAttributePropertySelection::Attribute;
+		InputSource.AttributeName = InputAttributeName_DEPRECATED;
+		InputAttributeName_DEPRECATED = NAME_None;
+	}
+#endif // WITH_EDITOR
+}
+
+FPCGAttributePropertySelector UPCGMetadataBreakVectorSettings::GetInputSource(uint32 Index) const
+{
+	return InputSource;
 }
 
 FName UPCGMetadataBreakVectorSettings::GetOutputPinLabel(uint32 Index) const 
diff --git a/Engine/Plugins/Experimental/PCG/Source/PCG/Private/Elements/Metadata/PCGMetadataCompareOpElement.cpp b/Engine/Plugins/Experimental/PCG/Source/PCG/Private/Elements/Metadata/PCGMetadataCompareOpElement.cpp
--- a/Engine/Plugins/Experimental/PCG/Source/PCG/Private/Elements/Metadata/PCGMetadataCompareOpElement.cpp
+++ b/Engine/Plugins/Experimental/PCG/Source/PCG/Private/Elements/Metadata/PCGMetadataCompareOpElement.cpp
@@ -52,6 +52,27 @@ namespace PCGMetadataCompareSettings
 	}
 }
 
+void UPCGMetadataCompareSettings::PostLoad()
+{
+	Super::PostLoad();
+
+#if WITH_EDITOR
+	if (Input1AttributeName_DEPRECATED != NAME_None)
+	{
+		InputSource1.Selection = EPCGAttributePropertySelection::Attribute;
+		InputSource1.AttributeName = Input1AttributeName_DEPRECATED;
+		Input1AttributeName_DEPRECATED = NAME_None;
+	}
+
+	if (Input2AttributeName_DEPRECATED != NAME_None)
+	{
+		InputSource2.Selection = EPCGAttributePropertySelection::Attribute;
+		InputSource2.AttributeName = Input2AttributeName_DEPRECATED;
+		Input2AttributeName_DEPRECATED = NAME_None;
+	}
+#endif // WITH_EDITOR
+}
+
 FName UPCGMetadataCompareSettings::GetInputPinLabel(uint32 Index) const
 {
 	switch (Index)
@@ -76,16 +97,16 @@ bool UPCGMetadataCompareSettings::IsSupportedInputType(uint16 TypeId, uint32 Inp
 	return PCG::Private::IsOfTypes<int32, int64, float, double>(TypeId);
 }
 
-FName UPCGMetadataCompareSettings::GetInputAttributeNameWithOverride(uint32 Index, UPCGParamData* Params) const
+FPCGAttributePropertySelector UPCGMetadataCompareSettings::GetInputSource(uint32 Index) const
 {
 	switch (Index)
 	{
 	case 0:
-		return PCG_GET_OVERRIDEN_VALUE(this, Input1AttributeName, Params);
+		return InputSource1;
 	case 1:
-		return PCG_GET_OVERRIDEN_VALUE(this, Input2AttributeName, Params);
+		return InputSource2;
 	default:
-		return NAME_None;
+		return FPCGAttributePropertySelector();
 	}
 }
 
diff --git a/Engine/Plugins/Experimental/PCG/Source/PCG/Private/Elements/Metadata/PCGMetadataElementCommon.cpp b/Engine/Plugins/Experimental/PCG/Source/PCG/Private/Elements/Metadata/PCGMetadataElementCommon.cpp
--- a/Engine/Plugins/Experimental/PCG/Source/PCG/Private/Elements/Metadata/PCGMetadataElementCommon.cpp
+++ b/Engine/Plugins/Experimental/PCG/Source/PCG/Private/Elements/Metadata/PCGMetadataElementCommon.cpp
@@ -10,7 +10,7 @@ namespace PCGMetadataElementCommon
 {
 	void DuplicateTaggedData(const FPCGTaggedData& InTaggedData, FPCGTaggedData& OutTaggedData, UPCGMetadata*& OutMetadata)
 	{
-		TRACE_CPUPROFILER_EVENT_SCOPE(PCGMetadataElementCommonDuplicateTaggedData);
+		TRACE_CPUPROFILER_EVENT_SCOPE(PCGMetadataElementCommon::DuplicateTaggedData);
 
 		if (const UPCGSpatialData* SpatialInput = Cast<const UPCGSpatialData>(InTaggedData.Data))
 		{
diff --git a/Engine/Plugins/Experimental/PCG/Source/PCG/Private/Elements/Metadata/PCGMetadataMakeTransform.cpp b/Engine/Plugins/Experimental/PCG/Source/PCG/Private/Elements/Metadata/PCGMetadataMakeTransform.cpp
--- a/Engine/Plugins/Experimental/PCG/Source/PCG/Private/Elements/Metadata/PCGMetadataMakeTransform.cpp
+++ b/Engine/Plugins/Experimental/PCG/Source/PCG/Private/Elements/Metadata/PCGMetadataMakeTransform.cpp
@@ -27,6 +27,34 @@ namespace PCGMetadataMakeTransformSettings
 	}
 }
 
+void UPCGMetadataMakeTransformSettings::PostLoad()
+{
+	Super::PostLoad();
+
+#if WITH_EDITOR
+	if (Input1AttributeName_DEPRECATED != NAME_None)
+	{
+		InputSource1.Selection = EPCGAttributePropertySelection::Attribute;
+		InputSource1.AttributeName = Input1AttributeName_DEPRECATED;
+		Input1AttributeName_DEPRECATED = NAME_None;
+	}
+
+	if (Input2AttributeName_DEPRECATED != NAME_None)
+	{
+		InputSource2.Selection = EPCGAttributePropertySelection::Attribute;
+		InputSource2.AttributeName = Input2AttributeName_DEPRECATED;
+		Input2AttributeName_DEPRECATED = NAME_None;
+	}
+
+	if (Input3AttributeName_DEPRECATED != NAME_None)
+	{
+		InputSource3.Selection = EPCGAttributePropertySelection::Attribute;
+		InputSource3.AttributeName = Input3AttributeName_DEPRECATED;
+		Input3AttributeName_DEPRECATED = NAME_None;
+	}
+#endif // WITH_EDITOR
+}
+
 FName UPCGMetadataMakeTransformSettings::GetInputPinLabel(uint32 Index) const
 {
 	switch (Index)
@@ -61,18 +89,18 @@ bool UPCGMetadataMakeTransformSettings::IsSupportedInputType(uint16 TypeId, uint
 	}
 }
 
-FName UPCGMetadataMakeTransformSettings::GetInputAttributeNameWithOverride(uint32 Index, UPCGParamData* Params) const
+FPCGAttributePropertySelector UPCGMetadataMakeTransformSettings::GetInputSource(uint32 Index) const
 {
 	switch (Index)
 	{
 	case 0:
-		return PCG_GET_OVERRIDEN_VALUE(this, Input1AttributeName, Params);
+		return InputSource1;
 	case 1:
-		return PCG_GET_OVERRIDEN_VALUE(this, Input2AttributeName, Params);
+		return InputSource2;
 	case 2:
-		return PCG_GET_OVERRIDEN_VALUE(this, Input3AttributeName, Params);
+		return InputSource3;
 	default:
-		return NAME_None;
+		return FPCGAttributePropertySelector();
 	}
 }
 
diff --git a/Engine/Plugins/Experimental/PCG/Source/PCG/Private/Elements/Metadata/PCGMetadataMakeVector.cpp b/Engine/Plugins/Experimental/PCG/Source/PCG/Private/Elements/Metadata/PCGMetadataMakeVector.cpp
--- a/Engine/Plugins/Experimental/PCG/Source/PCG/Private/Elements/Metadata/PCGMetadataMakeVector.cpp
+++ b/Engine/Plugins/Experimental/PCG/Source/PCG/Private/Elements/Metadata/PCGMetadataMakeVector.cpp
@@ -44,6 +44,41 @@ namespace PCGMetadataMakeVectorSettings
 	}
 }
 
+void UPCGMetadataMakeVectorSettings::PostLoad()
+{
+	Super::PostLoad();
+
+#if WITH_EDITOR
+	if (Input1AttributeName_DEPRECATED != NAME_None)
+	{
+		InputSource1.Selection = EPCGAttributePropertySelection::Attribute;
+		InputSource1.AttributeName = Input1AttributeName_DEPRECATED;
+		Input1AttributeName_DEPRECATED = NAME_None;
+	}
+
+	if (Input2AttributeName_DEPRECATED != NAME_None)
+	{
+		InputSource2.Selection = EPCGAttributePropertySelection::Attribute;
+		InputSource2.AttributeName = Input2AttributeName_DEPRECATED;
+		Input2AttributeName_DEPRECATED = NAME_None;
+	}
+
+	if (Input3AttributeName_DEPRECATED != NAME_None)
+	{
+		InputSource3.Selection = EPCGAttributePropertySelection::Attribute;
+		InputSource3.AttributeName = Input3AttributeName_DEPRECATED;
+		Input3AttributeName_DEPRECATED = NAME_None;
+	}
+
+	if (Input4AttributeName_DEPRECATED != NAME_None)
+	{
+		InputSource4.Selection = EPCGAttributePropertySelection::Attribute;
+		InputSource4.AttributeName = Input4AttributeName_DEPRECATED;
+		Input4AttributeName_DEPRECATED = NAME_None;
+	}
+#endif // WITH_EDITOR
+}
+
 FName UPCGMetadataMakeVectorSettings::GetInputPinLabel(uint32 Index) const
 {
 	static const FName DefaultNames[4] = {
@@ -129,20 +164,20 @@ bool UPCGMetadataMakeVectorSettings::IsSupportedInputType(uint16 TypeId, uint32
 	}
 }
 
-FName UPCGMetadataMakeVectorSettings::GetInputAttributeNameWithOverride(uint32 Index, UPCGParamData* Params) const
+FPCGAttributePropertySelector UPCGMetadataMakeVectorSettings::GetInputSource(uint32 Index) const
 {
 	switch (Index)
 	{
 	case 0:
-		return PCG_GET_OVERRIDEN_VALUE(this, Input1AttributeName, Params);
+		return InputSource1;
 	case 1:
-		return PCG_GET_OVERRIDEN_VALUE(this, Input2AttributeName, Params);
+		return InputSource2;
 	case 2:
-		return PCG_GET_OVERRIDEN_VALUE(this, Input3AttributeName, Params);
+		return InputSource3;
 	case 3:
-		return PCG_GET_OVERRIDEN_VALUE(this, Input4AttributeName, Params);
+		return InputSource4;
 	default:
-		return NAME_None;
+		return FPCGAttributePropertySelector();
 	}
 }
 
diff --git a/Engine/Plugins/Experimental/PCG/Source/PCG/Private/Elements/Metadata/PCGMetadataMathsOpElement.cpp b/Engine/Plugins/Experimental/PCG/Source/PCG/Private/Elements/Metadata/PCGMetadataMathsOpElement.cpp
--- a/Engine/Plugins/Experimental/PCG/Source/PCG/Private/Elements/Metadata/PCGMetadataMathsOpElement.cpp
+++ b/Engine/Plugins/Experimental/PCG/Source/PCG/Private/Elements/Metadata/PCGMetadataMathsOpElement.cpp
@@ -153,6 +153,34 @@ namespace PCGMetadataMathsSettings
 	}
 }
 
+void UPCGMetadataMathsSettings::PostLoad()
+{
+	Super::PostLoad();
+
+#if WITH_EDITOR
+	if (Input1AttributeName_DEPRECATED != NAME_None)
+	{
+		InputSource1.Selection = EPCGAttributePropertySelection::Attribute;
+		InputSource1.AttributeName = Input1AttributeName_DEPRECATED;
+		Input1AttributeName_DEPRECATED = NAME_None;
+	}
+
+	if (Input2AttributeName_DEPRECATED != NAME_None)
+	{
+		InputSource2.Selection = EPCGAttributePropertySelection::Attribute;
+		InputSource2.AttributeName = Input2AttributeName_DEPRECATED;
+		Input2AttributeName_DEPRECATED = NAME_None;
+	}
+
+	if (Input3AttributeName_DEPRECATED != NAME_None)
+	{
+		InputSource3.Selection = EPCGAttributePropertySelection::Attribute;
+		InputSource3.AttributeName = Input3AttributeName_DEPRECATED;
+		Input3AttributeName_DEPRECATED = NAME_None;
+	}
+#endif // WITH_EDITOR
+}
+
 FName UPCGMetadataMathsSettings::GetInputPinLabel(uint32 Index) const
 {
 	switch (Index)
@@ -195,18 +223,18 @@ bool UPCGMetadataMathsSettings::IsSupportedInputType(uint16 TypeId, uint32 Input
 	return PCG::Private::IsOfTypes<float, double, int32, int64, FVector2D, FVector, FVector4>(TypeId);
 }
 
-FName UPCGMetadataMathsSettings::GetInputAttributeNameWithOverride(uint32 Index, UPCGParamData* Params) const
+FPCGAttributePropertySelector UPCGMetadataMathsSettings::GetInputSource(uint32 Index) const
 {
 	switch (Index)
 	{
 	case 0:
-		return PCG_GET_OVERRIDEN_VALUE(this, Input1AttributeName, Params);
+		return InputSource1;
 	case 1:
-		return PCG_GET_OVERRIDEN_VALUE(this, Input2AttributeName, Params);
+		return InputSource2;
 	case 2:
-		return PCG_GET_OVERRIDEN_VALUE(this, Input3AttributeName, Params);
+		return InputSource3;
 	default:
-		return NAME_None;
+		return FPCGAttributePropertySelector();
 	}
 }
 
diff --git a/Engine/Plugins/Experimental/PCG/Source/PCG/Private/Elements/Metadata/PCGMetadataOpElementBase.cpp b/Engine/Plugins/Experimental/PCG/Source/PCG/Private/Elements/Metadata/PCGMetadataOpElementBase.cpp
--- a/Engine/Plugins/Experimental/PCG/Source/PCG/Private/Elements/Metadata/PCGMetadataOpElementBase.cpp
+++ b/Engine/Plugins/Experimental/PCG/Source/PCG/Private/Elements/Metadata/PCGMetadataOpElementBase.cpp
@@ -3,14 +3,25 @@
 #include "Elements/Metadata/PCGMetadataOpElementBase.h"
 
 #include "PCGParamData.h"
-#include "Data/PCGSpatialData.h"
 #include "Data/PCGPointData.h"
+#include "Data/PCGSpatialData.h"
 #include "Elements/Metadata/PCGMetadataElementCommon.h"
 #include "Helpers/PCGSettingsHelpers.h"
 #include "Metadata/PCGMetadata.h"
-#include "Metadata/PCGMetadataAttribute.h"
-#include "Metadata/PCGMetadataAttributeTpl.h"
-#include "Metadata/PCGMetadataEntryKeyIterator.h"
+
+void UPCGMetadataSettingsBase::PostLoad()
+{
+	Super::PostLoad();
+
+#if WITH_EDITOR
+	if (OutputAttributeName_DEPRECATED != NAME_None)
+	{
+		OutputTarget.Selection = EPCGAttributePropertySelection::Attribute;
+		OutputTarget.AttributeName = OutputAttributeName_DEPRECATED;
+		OutputAttributeName_DEPRECATED = NAME_None;
+	}
+#endif // WITH_EDITOR
+}
 
 TArray<FPCGPinProperties> UPCGMetadataSettingsBase::InputPinProperties() const
 {
@@ -65,12 +76,6 @@ bool FPCGMetadataElementBase::ExecuteInternal(FPCGContext* Context) const
 	const TArray<FPCGTaggedData>& Inputs = Context->InputData.TaggedData;
 	TArray<FPCGTaggedData>& Outputs = Context->OutputData.TaggedData;
 
-	// We disabled param pin, but we let the logic of attribute name override, if we need to re-enable it in the future.
-	const TArray<FPCGTaggedData>& ParamData = Context->InputData.GetParamsByPin(PCGPinConstants::DefaultParamsLabel);
-	UPCGParamData* Params = ParamData.IsEmpty() ? nullptr : Cast<UPCGParamData>(ParamData[0].Data);
-
-	FName OutputAttributeName = PCG_GET_OVERRIDEN_VALUE(Settings, OutputAttributeName, Params);
-
 	// Gathering all the inputs metadata
 	TArray<const UPCGMetadata*> SourceMetadata;
 	TArray<FPCGTaggedData> InputTaggedData;
@@ -106,115 +111,144 @@ bool FPCGMetadataElementBase::ExecuteInternal(FPCGContext* Context) const
 	}
 
 	FOperationData OperationData;
-	OperationData.Iterators.SetNum(NumberOfInputs);
+	OperationData.InputIterators.SetNum(NumberOfInputs);
 	TArray<int32> NumberOfElements;
 	NumberOfElements.SetNum(NumberOfInputs);
 
-	OperationData.SourceAttributes.SetNum(NumberOfInputs);
-
 	OperationData.MostComplexInputType = (uint16)EPCGMetadataTypes::Unknown;
 	OperationData.NumberOfElementsToProcess = -1;
 
 	for (uint32 i = 0; i < NumberOfInputs; ++i)
 	{
-		FName SourceAttributeName = Settings->GetInputAttributeNameWithOverride(i, Params);
-		if (SourceAttributeName == NAME_None)
+		// First we verify if the input data match the first one.
+		if (i != 0 &&
+			InputTaggedData[0].Data->GetClass() != InputTaggedData[i].Data->GetClass() &&
+			!InputTaggedData[i].Data->IsA<UPCGParamData>())
+		{
+			PCGE_LOG(Error, "Input %d is not of the same type than input 0 and is not a param data. Not supported.", i);
+			return true;
+		}
+
+		FPCGAttributePropertySelector InputSource = Settings->GetInputSource(i);
+
+		if (InputSource.Selection == EPCGAttributePropertySelection::Attribute && InputSource.AttributeName == NAME_None)
 		{
-			SourceAttributeName = SourceMetadata[i]->GetLatestAttributeNameOrNone();
+			InputSource.AttributeName = SourceMetadata[i]->GetLatestAttributeNameOrNone();
 		}
 
-		OperationData.SourceAttributes[i] = SourceMetadata[i]->GetConstAttribute(SourceAttributeName);
+		OperationData.InputIterators[i] = PCGMetadataAttributeWrapper::CreateIteratorWrapper(InputTaggedData[i].Data, InputSource);
 
-		if (OperationData.SourceAttributes[i] == nullptr)
+		if (!OperationData.InputIterators[i].IsValid())
 		{
-			PCGE_LOG(Error, "Attribute %s does not exist for input %d", *SourceAttributeName.ToString(), i);
+			PCGE_LOG(Error, "Attribute/Property %s does not exist for input %d", *InputSource.GetName().ToString(), i);
 			return true;
 		}
 
+		uint16 AttributeTypeId = OperationData.InputIterators[i].GetType();
+
 		// Then verify that the type is OK
 		bool bHasSpecialRequirement = false;
-		if (!Settings->IsSupportedInputType(OperationData.SourceAttributes[i]->GetTypeId(), i, bHasSpecialRequirement))
+		if (!Settings->IsSupportedInputType(AttributeTypeId, i, bHasSpecialRequirement))
 		{
-			PCGE_LOG(Error, "Attribute %s is not a supported type for input %d", *SourceAttributeName.ToString(), i);
+			PCGE_LOG(Error, "Attribute/Property %s is not a supported type for input %d", *InputSource.GetName().ToString(), i);
 			return true;
 		}
 
 		if (!bHasSpecialRequirement)
 		{
 			// In this case, check if we have a more complex type, or if we can broadcast to the most complex type.
-			if (OperationData.MostComplexInputType == (uint16)EPCGMetadataTypes::Unknown || Settings->IsMoreComplexType(OperationData.SourceAttributes[i]->GetTypeId(), OperationData.MostComplexInputType))
+			if (OperationData.MostComplexInputType == (uint16)EPCGMetadataTypes::Unknown || Settings->IsMoreComplexType(AttributeTypeId, OperationData.MostComplexInputType))
 			{
-				OperationData.MostComplexInputType = OperationData.SourceAttributes[i]->GetTypeId();
+				OperationData.MostComplexInputType = AttributeTypeId;
 			}
-			else if (OperationData.MostComplexInputType != OperationData.SourceAttributes[i]->GetTypeId() && !PCG::Private::IsBroadcastable(OperationData.SourceAttributes[i]->GetTypeId(), OperationData.MostComplexInputType))
+			else if (OperationData.MostComplexInputType != AttributeTypeId && !PCG::Private::IsBroadcastable(AttributeTypeId, OperationData.MostComplexInputType))
 			{
-				PCGE_LOG(Error, "Attribute %s cannot be broadcasted to match types for input %d", *SourceAttributeName.ToString(), i);
+				PCGE_LOG(Error, "Attribute %s cannot be broadcasted to match types for input %d", *InputSource.GetName().ToString(), i);
 				return true;
 			}
 		}
 
-		// Finally check that we have the right number of inputs, depending on the source
-		if (const UPCGPointData* PointData = Cast<UPCGPointData>(InputTaggedData[i].Data))
-		{
-			NumberOfElements[i] = PointData->GetPoints().Num();
+		NumberOfElements[i] = OperationData.InputIterators[i].Num();
 
-			// If we are not the first input, we only have an iterator if we have a single point and we broadcast.
-			bool bShouldBroadcast = i != 0 && NumberOfElements[i] == 1 && Settings->Mode == EPCGMetadataSettingsBaseMode::Broadcast;
-			if (i == 0 || NumberOfElements[i] == NumberOfElements[0] || bShouldBroadcast)
-			{
-				OperationData.Iterators[i] = MakeUnique<FPCGMetadataEntryPointIterator>(PointData, bShouldBroadcast);
-			}
-		}
-		else
+		if (OperationData.NumberOfElementsToProcess == -1)
 		{
-			NumberOfElements[i] = OperationData.SourceAttributes[i]->GetNumberOfEntriesWithParents();
-
-			// No element means only dealing with the default value.
-			if (NumberOfElements[i] == 0)
-			{
-				OperationData.Iterators[i] = MakeUnique<FPCGMetadataEntryConstantIterator>(PCGInvalidEntryKey, /*bRepeat=*/ true);
-			}
-			else
-			{
-				// Broadcast only if we have a single element and we are in broadcast mode, or inferred mode and we are param data.
-				bool bShouldBroadcast = i != 0 && 
-					NumberOfElements[i] == 1 && 
-					(Settings->Mode == EPCGMetadataSettingsBaseMode::Broadcast || (Settings->Mode == EPCGMetadataSettingsBaseMode::Inferred && InputTaggedData[i].Data->IsA<UPCGParamData>()));
-
-				if (i == 0 || NumberOfElements[i] == NumberOfElements[0] || bShouldBroadcast)
-				{
-					OperationData.Iterators[i] = MakeUnique<FPCGMetadataEntryAttributeIterator>(*OperationData.SourceAttributes[i], bShouldBroadcast);
-				}
-			}
+			OperationData.NumberOfElementsToProcess = NumberOfElements[i];
 		}
 
-		if (OperationData.NumberOfElementsToProcess == -1)
+		// Verify that the number of elements makes sense
+		if (OperationData.NumberOfElementsToProcess % NumberOfElements[i] != 0)
 		{
-			OperationData.NumberOfElementsToProcess = NumberOfElements[i];
+			PCGE_LOG(Error, "Mismatch between the number of elements in input 0 (%d) and in input %d (%d).", OperationData.NumberOfElementsToProcess, i, NumberOfElements[i]);
+			return true;
 		}
 	}
 
 	// At this point, we verified everything, so we can go forward with the computation, depending on the most complex type
 	// So first forward outputs and create the attribute
-	OperationData.OutputAttributes.SetNum(Settings->GetOutputPinNum());
+	OperationData.OutputIterators.SetNum(Settings->GetOutputPinNum());
+
+	FPCGAttributePropertySelector OutputTarget = Settings->OutputTarget;
+	if (OutputTarget.Selection == EPCGAttributePropertySelection::Attribute && OutputTarget.AttributeName == NAME_None)
+	{
+		OutputTarget.AttributeName = OperationData.InputIterators[0].GetWrapper().GetName();
+	}
 
 	auto CreateAttribute = [&](uint32 OutputIndex, auto DummyOutValue) -> bool
 	{
 		using AttributeType = decltype(DummyOutValue);
 
-		AttributeType DefaultValue = PCGMetadataAttribute::GetValueWithBroadcast<AttributeType>(OperationData.SourceAttributes[0], PCGDefaultValueKey);
-
 		FPCGTaggedData& OutputData = Outputs.Add_GetRef(InputTaggedData[0]);
 		OutputData.Pin = Settings->GetOutputPinLabel(OutputIndex);
 
-		FName OutputAttributeFinalName = Settings->GetOutputAttributeName(OutputAttributeName, OutputIndex);
-
 		UPCGMetadata* OutMetadata = nullptr;
 
+		FName OutputName = OutputTarget.GetName();
+
+		bool bShouldClearOrCreateAttribute = false;
+		bool bShouldCopyEntryMap = false;
+
+		if (UPCGPointData* PointData = Cast<UPCGPointData>(InputTaggedData[0].Data))
+		{
+			int16 PropertyType = -1;
+			if (PCGMetadataAttributeWrapper::IsPropertyWithType(InputTaggedData[0].Data, OutputName, &PropertyType))
+			{
+				// We matched a property, check if the output type is valid
+				if (!PCG::Private::IsBroadcastable(PCG::Private::MetadataTypes<AttributeType>::Id, PropertyType))
+				{
+					PCGE_LOG(Error, "Property %s cannot be broadcasted to match types for input", *OutputName.ToString());
+					return false;
+				}
+			}
+			else
+			{
+				// Otherwise it is an attribute, we will create a clean one, so output type will match.
+				bShouldClearOrCreateAttribute = true;
+			}
+		}
+		else
+		{
+			// If it is not a point data, we will only deal with attributes, and the entry map will be necessary.
+			bShouldClearOrCreateAttribute = true;
+			bShouldCopyEntryMap = true;
+		}
+
 		PCGMetadataElementCommon::DuplicateTaggedData(InputTaggedData[0], OutputData, OutMetadata);
-		OperationData.OutputAttributes[OutputIndex] = PCGMetadataElementCommon::ClearOrCreateAttribute(OutMetadata, OutputAttributeFinalName, DefaultValue);
+		if (bShouldClearOrCreateAttribute)
+		{
+			AttributeType DefaultValue{};
+			OperationData.InputIterators[0].GetWrapper().Get<AttributeType>(PCGInvalidEntryKey, DefaultValue);
+
+			FPCGMetadataAttributeBase* OutputAttribute = PCGMetadataElementCommon::ClearOrCreateAttribute(OutMetadata, OutputName, DefaultValue);
+			if (bShouldCopyEntryMap)
+			{
+				PCGMetadataElementCommon::CopyEntryToValueKeyMap(SourceMetadata[0], OperationData.InputIterators[0].GetWrapper().GetAttribute(), OutputAttribute);
+			}
+		}
+
+		// Cast to use the non-const version of this method.
+		OperationData.OutputIterators[OutputIndex] = PCGMetadataAttributeWrapper::CreateIteratorWrapper(Cast<UPCGData>(OutputData.Data), OutputTarget);
 
-		return OperationData.OutputAttributes[OutputIndex] != nullptr;
+		return OperationData.OutputIterators[OutputIndex].IsValid();
 	};
 
 	auto CreateAllSameAttributes = [&](auto DummyOutValue) -> bool
diff --git a/Engine/Plugins/Experimental/PCG/Source/PCG/Private/Elements/Metadata/PCGMetadataRotatorOpElement.cpp b/Engine/Plugins/Experimental/PCG/Source/PCG/Private/Elements/Metadata/PCGMetadataRotatorOpElement.cpp
--- a/Engine/Plugins/Experimental/PCG/Source/PCG/Private/Elements/Metadata/PCGMetadataRotatorOpElement.cpp
+++ b/Engine/Plugins/Experimental/PCG/Source/PCG/Private/Elements/Metadata/PCGMetadataRotatorOpElement.cpp
@@ -94,6 +94,34 @@ namespace PCGMetadataRotatorSettings
 	}
 }
 
+void UPCGMetadataRotatorSettings::PostLoad()
+{
+	Super::PostLoad();
+
+#if WITH_EDITOR
+	if (Input1AttributeName_DEPRECATED != NAME_None)
+	{
+		InputSource1.Selection = EPCGAttributePropertySelection::Attribute;
+		InputSource1.AttributeName = Input1AttributeName_DEPRECATED;
+		Input1AttributeName_DEPRECATED = NAME_None;
+	}
+
+	if (Input2AttributeName_DEPRECATED != NAME_None)
+	{
+		InputSource2.Selection = EPCGAttributePropertySelection::Attribute;
+		InputSource2.AttributeName = Input2AttributeName_DEPRECATED;
+		Input2AttributeName_DEPRECATED = NAME_None;
+	}
+
+	if (Input3AttributeName_DEPRECATED != NAME_None)
+	{
+		InputSource3.Selection = EPCGAttributePropertySelection::Attribute;
+		InputSource3.AttributeName = Input3AttributeName_DEPRECATED;
+		Input3AttributeName_DEPRECATED = NAME_None;
+	}
+#endif // WITH_EDITOR
+}
+
 FName UPCGMetadataRotatorSettings::GetInputPinLabel(uint32 Index) const
 {
 	switch (Index)
@@ -151,18 +179,18 @@ bool UPCGMetadataRotatorSettings::IsSupportedInputType(uint16 TypeId, uint32 Inp
 	}
 }
 
-FName UPCGMetadataRotatorSettings::GetInputAttributeNameWithOverride(uint32 Index, UPCGParamData* Params) const
+FPCGAttributePropertySelector UPCGMetadataRotatorSettings::GetInputSource(uint32 Index) const
 {
 	switch (Index)
 	{
 	case 0:
-		return PCG_GET_OVERRIDEN_VALUE(this, Input1AttributeName, Params);
+		return InputSource1;
 	case 1:
-		return PCG_GET_OVERRIDEN_VALUE(this, Input2AttributeName, Params);
+		return InputSource2;
 	case 2:
-		return PCG_GET_OVERRIDEN_VALUE(this, Input3AttributeName, Params);
+		return InputSource3;
 	default:
-		return NAME_None;
+		return FPCGAttributePropertySelector();
 	}
 }
 
diff --git a/Engine/Plugins/Experimental/PCG/Source/PCG/Private/Elements/Metadata/PCGMetadataTransformOpElement.cpp b/Engine/Plugins/Experimental/PCG/Source/PCG/Private/Elements/Metadata/PCGMetadataTransformOpElement.cpp
--- a/Engine/Plugins/Experimental/PCG/Source/PCG/Private/Elements/Metadata/PCGMetadataTransformOpElement.cpp
+++ b/Engine/Plugins/Experimental/PCG/Source/PCG/Private/Elements/Metadata/PCGMetadataTransformOpElement.cpp
@@ -75,6 +75,34 @@ namespace PCGMetadataTransfromSettings
 	}
 }
 
+void UPCGMetadataTransformSettings::PostLoad()
+{
+	Super::PostLoad();
+
+#if WITH_EDITOR
+	if (Input1AttributeName_DEPRECATED != NAME_None)
+	{
+		InputSource1.Selection = EPCGAttributePropertySelection::Attribute;
+		InputSource1.AttributeName = Input1AttributeName_DEPRECATED;
+		Input1AttributeName_DEPRECATED = NAME_None;
+	}
+
+	if (Input2AttributeName_DEPRECATED != NAME_None)
+	{
+		InputSource2.Selection = EPCGAttributePropertySelection::Attribute;
+		InputSource2.AttributeName = Input2AttributeName_DEPRECATED;
+		Input2AttributeName_DEPRECATED = NAME_None;
+	}
+
+	if (Input3AttributeName_DEPRECATED != NAME_None)
+	{
+		InputSource3.Selection = EPCGAttributePropertySelection::Attribute;
+		InputSource3.AttributeName = Input3AttributeName_DEPRECATED;
+		Input3AttributeName_DEPRECATED = NAME_None;
+	}
+#endif // WITH_EDITOR
+}
+
 FName UPCGMetadataTransformSettings::GetInputPinLabel(uint32 Index) const
 {
 	switch (Index)
@@ -120,18 +148,18 @@ bool UPCGMetadataTransformSettings::IsSupportedInputType(uint16 TypeId, uint32 I
 	}
 }
 
-FName UPCGMetadataTransformSettings::GetInputAttributeNameWithOverride(uint32 Index, UPCGParamData* Params) const
+FPCGAttributePropertySelector UPCGMetadataTransformSettings::GetInputSource(uint32 Index) const
 {
 	switch (Index)
 	{
 	case 0:
-		return PCG_GET_OVERRIDEN_VALUE(this, Input1AttributeName, Params);
+		return InputSource1;
 	case 1:
-		return PCG_GET_OVERRIDEN_VALUE(this, Input2AttributeName, Params);
+		return InputSource2;
 	case 2:
-		return PCG_GET_OVERRIDEN_VALUE(this, Input3AttributeName, Params);
+		return InputSource3;
 	default:
-		return NAME_None;
+		return FPCGAttributePropertySelector();
 	}
 }
 
diff --git a/Engine/Plugins/Experimental/PCG/Source/PCG/Private/Elements/Metadata/PCGMetadataTrigOpElement.cpp b/Engine/Plugins/Experimental/PCG/Source/PCG/Private/Elements/Metadata/PCGMetadataTrigOpElement.cpp
--- a/Engine/Plugins/Experimental/PCG/Source/PCG/Private/Elements/Metadata/PCGMetadataTrigOpElement.cpp
+++ b/Engine/Plugins/Experimental/PCG/Source/PCG/Private/Elements/Metadata/PCGMetadataTrigOpElement.cpp
@@ -45,6 +45,27 @@ namespace PCGMetadataTrigSettings
 	}
 }
 
+void UPCGMetadataTrigSettings::PostLoad()
+{
+	Super::PostLoad();
+
+#if WITH_EDITOR
+	if (Input1AttributeName_DEPRECATED != NAME_None)
+	{
+		InputSource1.Selection = EPCGAttributePropertySelection::Attribute;
+		InputSource1.AttributeName = Input1AttributeName_DEPRECATED;
+		Input1AttributeName_DEPRECATED = NAME_None;
+	}
+
+	if (Input2AttributeName_DEPRECATED != NAME_None)
+	{
+		InputSource2.Selection = EPCGAttributePropertySelection::Attribute;
+		InputSource2.AttributeName = Input2AttributeName_DEPRECATED;
+		Input2AttributeName_DEPRECATED = NAME_None;
+	}
+#endif // WITH_EDITOR
+}
+
 FName UPCGMetadataTrigSettings::GetInputPinLabel(uint32 Index) const
 {
 	switch (Index)
@@ -69,16 +90,16 @@ bool UPCGMetadataTrigSettings::IsSupportedInputType(uint16 TypeId, uint32 InputI
 	return PCG::Private::IsOfTypes<int32, int64, float, double>(TypeId);
 }
 
-FName UPCGMetadataTrigSettings::GetInputAttributeNameWithOverride(uint32 Index, UPCGParamData* Params) const
+FPCGAttributePropertySelector UPCGMetadataTrigSettings::GetInputSource(uint32 Index) const
 {
 	switch (Index)
 	{
 	case 0:
-		return PCG_GET_OVERRIDEN_VALUE(this, Input1AttributeName, Params);
+		return InputSource1;
 	case 1:
-		return PCG_GET_OVERRIDEN_VALUE(this, Input2AttributeName, Params);
+		return InputSource2;
 	default:
-		return NAME_None;
+		return FPCGAttributePropertySelector();
 	}
 }
 
diff --git a/Engine/Plugins/Experimental/PCG/Source/PCG/Private/Elements/Metadata/PCGMetadataVectorOpElement.cpp b/Engine/Plugins/Experimental/PCG/Source/PCG/Private/Elements/Metadata/PCGMetadataVectorOpElement.cpp
--- a/Engine/Plugins/Experimental/PCG/Source/PCG/Private/Elements/Metadata/PCGMetadataVectorOpElement.cpp
+++ b/Engine/Plugins/Experimental/PCG/Source/PCG/Private/Elements/Metadata/PCGMetadataVectorOpElement.cpp
@@ -208,6 +208,34 @@ namespace PCGMetadataVectorSettings
 	};
 }
 
+void UPCGMetadataVectorSettings::PostLoad()
+{
+	Super::PostLoad();
+
+#if WITH_EDITOR
+	if (Input1AttributeName_DEPRECATED != NAME_None)
+	{
+		InputSource1.Selection = EPCGAttributePropertySelection::Attribute;
+		InputSource1.AttributeName = Input1AttributeName_DEPRECATED;
+		Input1AttributeName_DEPRECATED = NAME_None;
+	}
+
+	if (Input2AttributeName_DEPRECATED != NAME_None)
+	{
+		InputSource2.Selection = EPCGAttributePropertySelection::Attribute;
+		InputSource2.AttributeName = Input2AttributeName_DEPRECATED;
+		Input2AttributeName_DEPRECATED = NAME_None;
+	}
+
+	if (Input3AttributeName_DEPRECATED != NAME_None)
+	{
+		InputSource3.Selection = EPCGAttributePropertySelection::Attribute;
+		InputSource3.AttributeName = Input3AttributeName_DEPRECATED;
+		Input3AttributeName_DEPRECATED = NAME_None;
+	}
+#endif // WITH_EDITOR
+}
+
 FName UPCGMetadataVectorSettings::GetInputPinLabel(uint32 Index) const
 {
 	switch (Index)
@@ -283,18 +311,18 @@ bool UPCGMetadataVectorSettings::IsSupportedInputType(uint16 TypeId, uint32 Inpu
 	}
 }
 
-FName UPCGMetadataVectorSettings::GetInputAttributeNameWithOverride(uint32 Index, UPCGParamData* Params) const
+FPCGAttributePropertySelector UPCGMetadataVectorSettings::GetInputSource(uint32 Index) const
 {
 	switch (Index)
 	{
 	case 0:
-		return PCG_GET_OVERRIDEN_VALUE(this, Input1AttributeName, Params);
+		return InputSource1;
 	case 1:
-		return PCG_GET_OVERRIDEN_VALUE(this, Input2AttributeName, Params);
+		return InputSource2;
 	case 2:
-		return PCG_GET_OVERRIDEN_VALUE(this, Input3AttributeName, Params);
+		return InputSource3;
 	default:
-		return NAME_None;
+		return FPCGAttributePropertySelector();
 	}
 }
 
diff --git a/Engine/Plugins/Experimental/PCG/Source/PCG/Private/Metadata/PCGMetadata.cpp b/Engine/Plugins/Experimental/PCG/Source/PCG/Private/Metadata/PCGMetadata.cpp
--- a/Engine/Plugins/Experimental/PCG/Source/PCG/Private/Metadata/PCGMetadata.cpp
+++ b/Engine/Plugins/Experimental/PCG/Source/PCG/Private/Metadata/PCGMetadata.cpp
@@ -3,7 +3,7 @@
 #include "Metadata/PCGMetadata.h"
 
 #include "Misc/ScopeRWLock.h"
-#include "Helpers/PCGSettingsHelpers.h"
+#include "Metadata/PCGMetadataAttributeWrapper.h"
 
 void UPCGMetadata::Serialize(FArchive& InArchive)
 {
@@ -483,7 +483,7 @@ namespace PCGMetadata
 				});
 		};
 
-		return PCGSettingsHelpers::GetPropertyValueWithCallback(DataPtr, InProperty, CreateAttributeAndSet);
+		return PCGMetadataAttributeWrapper::GetPropertyValueWithCallback(DataPtr, InProperty, CreateAttributeAndSet);
 	}
 }
 
diff --git a/Engine/Plugins/Experimental/PCG/Source/PCG/Private/Metadata/PCGMetadataAttributeWrapper.cpp b/Engine/Plugins/Experimental/PCG/Source/PCG/Private/Metadata/PCGMetadataAttributeWrapper.cpp
new file mode 100644
--- /dev/null
+++ b/Engine/Plugins/Experimental/PCG/Source/PCG/Private/Metadata/PCGMetadataAttributeWrapper.cpp
@@ -0,0 +1,283 @@
+// Copyright Epic Games, Inc. All Rights Reserved.
+
+#include "Metadata/PCGMetadataAttributeWrapper.h"
+
+#include "Data/PCGSpatialData.h"
+#include "Data/PCGPointData.h"
+#include "PCGParamData.h"
+
+EPCGMetadataTypes PCGMetadataAttributeWrapper::GetMetadataTypeFromProperty(const FProperty* InProperty)
+{
+	if (!InProperty)
+	{
+		return EPCGMetadataTypes::Unknown;
+	}
+	else if (const FNumericProperty* NumericProperty = CastField<FNumericProperty>(InProperty))
+	{
+		if (NumericProperty->IsFloatingPoint())
+		{
+			return EPCGMetadataTypes::Double;
+		}
+		else if (NumericProperty->IsInteger())
+		{
+			return EPCGMetadataTypes::Integer64;
+		}
+	}
+	else if (const FBoolProperty* BoolProperty = CastField<FBoolProperty>(InProperty))
+	{
+		return EPCGMetadataTypes::Boolean;
+	}
+	else if (const FStrProperty* StringProperty = CastField<FStrProperty>(InProperty))
+	{
+		return EPCGMetadataTypes::String;
+	}
+	else if (const FNameProperty* NameProperty = CastField<FNameProperty>(InProperty))
+	{
+		return EPCGMetadataTypes::Name;
+	}
+	else if (const FEnumProperty* EnumProperty = CastField<FEnumProperty>(InProperty))
+	{
+		return EPCGMetadataTypes::Integer64;
+	}
+	else if (const FStructProperty* StructProperty = CastField<FStructProperty>(InProperty))
+	{
+		if (StructProperty->Struct == TBaseStructure<FVector>::Get())
+		{
+			return EPCGMetadataTypes::Vector;
+		}
+		else if (StructProperty->Struct == TBaseStructure<FVector4>::Get())
+		{
+			return EPCGMetadataTypes::Vector4;
+		}
+		else if (StructProperty->Struct == TBaseStructure<FVector2D>::Get())
+		{
+			return EPCGMetadataTypes::Vector2;
+		}
+		else if (StructProperty->Struct == TBaseStructure<FQuat>::Get())
+		{
+			return EPCGMetadataTypes::Quaternion;
+		}
+		else if (StructProperty->Struct == TBaseStructure<FTransform>::Get())
+		{
+			return EPCGMetadataTypes::Transform;
+		}
+		else if (StructProperty->Struct == TBaseStructure<FRotator>::Get())
+		{
+			return EPCGMetadataTypes::Rotator;
+		}
+		else if (StructProperty->Struct == TBaseStructure<FSoftObjectPath>::Get() || StructProperty->Struct == TBaseStructure<FSoftClassPath>::Get())
+		{
+			// Soft object path are transformed from strings
+			return EPCGMetadataTypes::String;
+		}
+		//else if (StructProperty->Struct == TBaseStructure<FColor>::Get())
+		//else if (StructProperty->Struct == TBaseStructure<FLinearColor>::Get())
+	}
+	else if (const FObjectPropertyBase* ObjectProperty = CastField<FObjectPropertyBase>(InProperty))
+	{
+		// Object are transformed into their soft path name (as a string attribute)
+		return EPCGMetadataTypes::String;
+	}
+
+	return EPCGMetadataTypes::Unknown;
+}
+
+bool PCGMetadataAttributeWrapper::IsPropertyWithType(const UPCGData* InData, FName InName, int16* OutType)
+{
+	if (const UPCGPointData* PointData = Cast<const UPCGPointData>(InData))
+	{
+		if (const FProperty* Property = FPCGPoint::StaticStruct()->FindPropertyByName(InName))
+		{
+			if (OutType)
+			{
+				*OutType = (int16)GetMetadataTypeFromProperty(Property);
+			}
+
+			return true;
+		}
+		else if (FPCGPoint::HasCustomPropertyGetterSetter(InName))
+		{
+			if (OutType)
+			{
+				*OutType = FPCGPoint::CreateCustomPropertyGetterSetter(InName).GetType();
+			}
+			return true;
+		}
+	}
+	
+	return false;
+}
+
+FPCGPropertyAttributeWrapper PCGMetadataAttributeWrapper::CreateWrapper(UPCGData* InData, const FPCGAttributePropertySelector& InSelector)
+{
+	constexpr bool bIsReadOnly = false;
+
+	FName Name = InSelector.GetName();
+
+	if (const UPCGPointData* PointData = Cast<const UPCGPointData>(InData))
+	{
+		if (InSelector.Selection == EPCGAttributePropertySelection::PointProperty)
+		{
+			if (FPCGPoint::HasCustomPropertyGetterSetter(Name))
+			{
+				return FPCGPropertyAttributeWrapper(FPCGPoint::CreateCustomPropertyGetterSetter(Name), bIsReadOnly);
+			}
+			else if (FProperty* Property = FPCGPoint::StaticStruct()->FindPropertyByName(Name))
+			{
+				return FPCGPropertyAttributeWrapper(Property, bIsReadOnly);
+			}
+		}
+	}
+
+	// At this point, it is not a point data or we didn't find a property.
+	// We can't continue if it is a property wanted.
+	if (InSelector.Selection == EPCGAttributePropertySelection::PointProperty)
+	{
+		UE_LOG(LogPCG, Error, TEXT("[PCGMetadataAttributeWrapper::CreateWrapper] Expected to select a property but the data doesn't support this property."));
+		return FPCGPropertyAttributeWrapper();
+	}
+
+	UPCGMetadata* Metadata = nullptr;
+
+	if (UPCGSpatialData* SpatialData = Cast<UPCGSpatialData>(InData))
+	{
+		Metadata = SpatialData->Metadata;
+	}
+	else if (UPCGParamData* ParamData = Cast<UPCGParamData>(InData))
+	{
+		Metadata = ParamData->Metadata;
+	}
+
+	if (!Metadata)
+	{
+		return FPCGPropertyAttributeWrapper();
+	}
+
+	FPCGMetadataAttributeBase* Attribute = Metadata->GetMutableAttribute(Name);
+
+	if (Attribute)
+	{
+		return FPCGPropertyAttributeWrapper(Attribute, Metadata);
+	}
+	else
+	{
+		return FPCGPropertyAttributeWrapper();
+	}
+}
+
+FPCGPropertyAttributeWrapper PCGMetadataAttributeWrapper::CreateWrapper(const UPCGData* InData, const FPCGAttributePropertySelector& InSelector)
+{
+	constexpr bool bIsReadOnly = true;
+
+	FName Name = InSelector.GetName();
+
+	if (const UPCGPointData* PointData = Cast<const UPCGPointData>(InData))
+	{
+		if (InSelector.Selection == EPCGAttributePropertySelection::PointProperty)
+		{
+			if (FPCGPoint::HasCustomPropertyGetterSetter(Name))
+			{
+				return FPCGPropertyAttributeWrapper(FPCGPoint::CreateCustomPropertyGetterSetter(Name), bIsReadOnly);
+			}
+			else if (FProperty* Property = FPCGPoint::StaticStruct()->FindPropertyByName(Name))
+			{
+				return FPCGPropertyAttributeWrapper(Property, bIsReadOnly);
+			}
+		}
+	}
+
+	// At this point, it is not a point data or we didn't find a property.
+	// We can't continue if it is a property wanted.
+	if (InSelector.Selection == EPCGAttributePropertySelection::PointProperty)
+	{
+		UE_LOG(LogPCG, Error, TEXT("[PCGMetadataAttributeWrapper::CreateWrapper] Expected to select a property but the data doesn't support this property."));
+		return FPCGPropertyAttributeWrapper();
+	}
+
+	const UPCGMetadata* Metadata = nullptr;
+
+	if (const UPCGSpatialData* SpatialData = Cast<const UPCGSpatialData>(InData))
+	{
+		Metadata = SpatialData->Metadata;
+	}
+	else if (const UPCGParamData* ParamData = Cast<const UPCGParamData>(InData))
+	{
+		Metadata = ParamData->Metadata;
+	}
+
+	if (!Metadata)
+	{
+		return FPCGPropertyAttributeWrapper();
+	}
+
+	const FPCGMetadataAttributeBase* Attribute = Metadata->GetConstAttribute(Name);
+
+	if (Attribute)
+	{
+		return FPCGPropertyAttributeWrapper(Attribute, Metadata);
+	}
+	else
+	{
+		return FPCGPropertyAttributeWrapper();
+	}
+}
+
+FPCGPropertyAttributeIterator PCGMetadataAttributeWrapper::CreateIteratorWrapper(UPCGData* InData, const FPCGAttributePropertySelector& InSelector)
+{
+	FPCGPropertyAttributeWrapper Wrapper = PCGMetadataAttributeWrapper::CreateWrapper(InData, InSelector);
+
+	if (!Wrapper.IsValid())
+	{
+		return FPCGPropertyAttributeIterator();
+	}
+
+	if (UPCGPointData* PointData = Cast<UPCGPointData>(InData))
+	{
+		TArray<FPCGPoint>& Points = PointData->GetMutablePoints();
+		return FPCGPropertyAttributeIterator(Wrapper, TArrayView<FPCGPoint>(Points.GetData(), Points.Num()));
+	}
+	else
+	{
+		return FPCGPropertyAttributeIterator(Wrapper);
+	}
+}
+
+FPCGPropertyAttributeIterator PCGMetadataAttributeWrapper::CreateIteratorWrapper(const UPCGData* InData, const FPCGAttributePropertySelector& InSelector)
+{
+	FPCGPropertyAttributeWrapper Wrapper = PCGMetadataAttributeWrapper::CreateWrapper(InData, InSelector);
+
+	if (!Wrapper.IsValid())
+	{
+		return FPCGPropertyAttributeIterator();
+	}
+
+	if (const UPCGPointData* PointData = Cast<const UPCGPointData>(InData))
+	{
+		return FPCGPropertyAttributeIterator(Wrapper, PointData->GetPoints());
+	}
+	else
+	{
+		return FPCGPropertyAttributeIterator(Wrapper);
+	}
+}
+
+FName FPCGAttributePropertySelector::GetName() const
+{
+	switch (Selection)
+	{
+	case EPCGAttributePropertySelection::PointProperty:
+	{
+		if (const UEnum* EnumPtr = FindObject<UEnum>(nullptr, TEXT("/Script/PCG.EPCGPointProperties"), true))
+		{
+			// Need to use the string version and not the name version, because the name verison has "EPCGPointProperties::" as a prefix.
+			return FName(EnumPtr->GetNameStringByValue((int64)PointProperty));
+		}
+		else
+		{
+			return NAME_None;
+		}
+	}
+	default:
+		return AttributeName;
+	}
+}
\ No newline at end of file
diff --git a/Engine/Plugins/Experimental/PCG/Source/PCG/Private/PCGPoint.cpp b/Engine/Plugins/Experimental/PCG/Source/PCG/Private/PCGPoint.cpp
--- a/Engine/Plugins/Experimental/PCG/Source/PCG/Private/PCGPoint.cpp
+++ b/Engine/Plugins/Experimental/PCG/Source/PCG/Private/PCGPoint.cpp
@@ -2,6 +2,15 @@
 
 #include "PCGPoint.h"
 
+namespace PCGPointCustomPropertyNames
+{
+	const FName ExtentsName = TEXT("Extents");
+	const FName LocalCenterName = TEXT("LocalCenter");
+	const FName PositionName = TEXT("Position");
+	const FName RotationName = TEXT("Rotation");
+	const FName ScaleName = TEXT("Scale");
+}
+
 FPCGPoint::FPCGPoint(const FTransform& InTransform, float InDensity, int32 InSeed)
 	: Transform(InTransform)
 	, Density(InDensity)
@@ -23,4 +32,71 @@ void FPCGPoint::SetLocalBounds(const FBox& InBounds)
 FBoxSphereBounds FPCGPoint::GetDensityBounds() const
 {
 	return FBoxSphereBounds(FBox((2 - Steepness) * BoundsMin, (2 - Steepness) * BoundsMax).TransformBy(Transform));
+}
+
+FPCGPoint::PointCustomPropertyGetterSetter::PointCustomPropertyGetterSetter(const FPCGPoint::PointCustomPropertyGetter& InGetter, const FPCGPoint::PointCustomPropertySetter& InSetter, int16 InType, FName InName)
+	: Getter(InGetter)
+	, Setter(InSetter)
+	, Type(InType)
+	, Name(InName)
+{}
+
+bool FPCGPoint::HasCustomPropertyGetterSetter(FName Name)
+{
+	return Name == PCGPointCustomPropertyNames::ExtentsName ||
+		Name == PCGPointCustomPropertyNames::LocalCenterName ||
+		Name == PCGPointCustomPropertyNames::PositionName ||
+		Name == PCGPointCustomPropertyNames::RotationName ||
+		Name == PCGPointCustomPropertyNames::ScaleName;
+}
+
+FPCGPoint::PointCustomPropertyGetterSetter FPCGPoint::CreateCustomPropertyGetterSetter(FName Name)
+{
+	if (Name == PCGPointCustomPropertyNames::ExtentsName)
+	{
+		return FPCGPoint::PointCustomPropertyGetterSetter(
+			[](const FPCGPoint& Point, void* OutValue) { *reinterpret_cast<FVector*>(OutValue) = Point.GetExtents(); return true; },
+			[](FPCGPoint& Point, const void* InValue) { Point.SetExtents(*reinterpret_cast<const FVector*>(InValue)); return true; },
+			(int16)PCG::Private::MetadataTypes<FVector>::Id,
+			PCGPointCustomPropertyNames::ExtentsName
+		);
+	}
+	else if (Name == PCGPointCustomPropertyNames::LocalCenterName)
+	{
+		return FPCGPoint::PointCustomPropertyGetterSetter(
+			[](const FPCGPoint& Point, void* OutValue) { *reinterpret_cast<FVector*>(OutValue) = Point.GetLocalCenter(); return true; },
+			[](FPCGPoint& Point, const void* InValue) { Point.SetLocalCenter(*reinterpret_cast<const FVector*>(InValue)); return true; },
+			(int16)PCG::Private::MetadataTypes<FVector>::Id,
+			PCGPointCustomPropertyNames::LocalCenterName
+		);
+	}
+	else if (Name == PCGPointCustomPropertyNames::PositionName)
+	{
+		return FPCGPoint::PointCustomPropertyGetterSetter(
+			[](const FPCGPoint& Point, void* OutValue) { *reinterpret_cast<FVector*>(OutValue) = Point.Transform.GetLocation(); return true; },
+			[](FPCGPoint& Point, const void* InValue) { Point.Transform.SetLocation(*reinterpret_cast<const FVector*>(InValue)); return true; },
+			(int16)PCG::Private::MetadataTypes<FVector>::Id,
+			PCGPointCustomPropertyNames::PositionName
+		);
+	}
+	else if (Name == PCGPointCustomPropertyNames::RotationName)
+	{
+		return FPCGPoint::PointCustomPropertyGetterSetter(
+			[](const FPCGPoint& Point, void* OutValue) { *reinterpret_cast<FQuat*>(OutValue) = Point.Transform.GetRotation(); return true; },
+			[](FPCGPoint& Point, const void* InValue) { Point.Transform.SetRotation(*reinterpret_cast<const FQuat*>(InValue)); return true; },
+			(int16)PCG::Private::MetadataTypes<FQuat>::Id,
+			PCGPointCustomPropertyNames::RotationName
+		);
+	}
+	else if (Name == PCGPointCustomPropertyNames::ScaleName)
+	{
+		return FPCGPoint::PointCustomPropertyGetterSetter(
+			[](const FPCGPoint& Point, void* OutValue) { *reinterpret_cast<FVector*>(OutValue) = Point.Transform.GetScale3D(); return true; },
+			[](FPCGPoint& Point, const void* InValue) { Point.Transform.SetScale3D(*reinterpret_cast<const FVector*>(InValue)); return true; },
+			(int16)PCG::Private::MetadataTypes<FVector>::Id,
+			PCGPointCustomPropertyNames::ScaleName
+		);
+	}
+
+	return FPCGPoint::PointCustomPropertyGetterSetter();
 }
\ No newline at end of file
diff --git a/Engine/Plugins/Experimental/PCG/Source/PCG/Private/Tests/Elements/Metadata/PCGMetadataBooleanOpTest.cpp b/Engine/Plugins/Experimental/PCG/Source/PCG/Private/Tests/Elements/Metadata/PCGMetadataBooleanOpTest.cpp
--- a/Engine/Plugins/Experimental/PCG/Source/PCG/Private/Tests/Elements/Metadata/PCGMetadataBooleanOpTest.cpp
+++ b/Engine/Plugins/Experimental/PCG/Source/PCG/Private/Tests/Elements/Metadata/PCGMetadataBooleanOpTest.cpp
@@ -37,7 +37,8 @@ bool FPCGMetadataBooleanOpTest::RunTest(const FString& Parameters)
 	const FName InvalidAttribute = TEXT("Invalid");
 	const FName OutputAttributeName = TEXT("Output");
 
-	Settings->OutputAttributeName = OutputAttributeName;
+	Settings->OutputTarget.Selection = EPCGAttributePropertySelection::Attribute;
+	Settings->OutputTarget.AttributeName = OutputAttributeName;
 	Settings->ForceOutputConnections[0] = true;
 
 	const bool bAllowInterpolation = false;
@@ -87,7 +88,7 @@ bool FPCGMetadataBooleanOpTest::RunTest(const FString& Parameters)
 			return false;
 		}
 
-		const FPCGMetadataAttributeBase* OutAttributeBase = OutMetadata->GetConstAttribute(Settings->OutputAttributeName);
+		const FPCGMetadataAttributeBase* OutAttributeBase = OutMetadata->GetConstAttribute(Settings->OutputTarget.AttributeName);
 		if (!TestNotNull(TEXT("Output attribute exists"), OutAttributeBase))
 		{
 			return false;
@@ -118,13 +119,16 @@ bool FPCGMetadataBooleanOpTest::RunTest(const FString& Parameters)
 	ParamTaggedData1.Pin = PCGPinConstants::DefaultInputLabel;
 
 	Settings->Operation = EPCGMedadataBooleanOperation::Not;
+	Settings->InputSource1.Selection = EPCGAttributePropertySelection::Attribute;
+	Settings->InputSource2.Selection = EPCGAttributePropertySelection::Attribute;
+
 	{
-		Settings->Input1AttributeName = TrueAttribute;
+		Settings->InputSource1.AttributeName = TrueAttribute;
 		bTestPassed &= ValidateOp(/*bExpectedResult=*/ false);
 	}
 
 	{
-		Settings->Input1AttributeName = FalseAttribute;
+		Settings->InputSource1.AttributeName = FalseAttribute;
 		bTestPassed &= ValidateOp(true);
 	}
 
@@ -137,48 +141,48 @@ bool FPCGMetadataBooleanOpTest::RunTest(const FString& Parameters)
 	ParamTaggedData2.Pin = PCGMetadataSettingsBaseConstants::DoubleInputSecondLabel;
 	Settings->Operation = EPCGMedadataBooleanOperation::And;
 
-	AddExpectedError(TEXT("Attribute Invalid is not a supported type for input"), EAutomationExpectedErrorFlags::Contains, 3);
+	AddExpectedError(TEXT("Attribute/Property Invalid is not a supported type for input"), EAutomationExpectedErrorFlags::Contains, 3);
 
 	{
-		Settings->Input1AttributeName = InvalidAttribute;
-		Settings->Input2AttributeName = FalseAttribute;
+		Settings->InputSource1.AttributeName = InvalidAttribute;
+		Settings->InputSource2.AttributeName = FalseAttribute;
 		bTestPassed &= ValidateOp(/*bExpectedResult=*/false, /*bIsValid=*/ false);
 	}
 
 	{
-		Settings->Input1AttributeName = FalseAttribute;
-		Settings->Input2AttributeName = InvalidAttribute;
+		Settings->InputSource1.AttributeName = FalseAttribute;
+		Settings->InputSource2.AttributeName = InvalidAttribute;
 		bTestPassed &= ValidateOp(/*bExpectedResult=*/false, /*bIsValid=*/ false);
 	}
 
 	{
-		Settings->Input1AttributeName = InvalidAttribute;
-		Settings->Input2AttributeName = InvalidAttribute;
+		Settings->InputSource1.AttributeName = InvalidAttribute;
+		Settings->InputSource2.AttributeName = InvalidAttribute;
 		bTestPassed &= ValidateOp(/*bExpectedResult=*/false, /*bIsValid=*/ false);
 	}
 
 	// Valid tests
 	{
-		Settings->Input1AttributeName = TrueAttribute;
-		Settings->Input2AttributeName = TrueAttribute;
+		Settings->InputSource1.AttributeName = TrueAttribute;
+		Settings->InputSource2.AttributeName = TrueAttribute;
 		bTestPassed &= ValidateOp(/*bExpectedResult=*/true);
 	}
 
 	{
-		Settings->Input1AttributeName = FalseAttribute;
-		Settings->Input2AttributeName = TrueAttribute;
+		Settings->InputSource1.AttributeName = FalseAttribute;
+		Settings->InputSource2.AttributeName = TrueAttribute;
 		bTestPassed &= ValidateOp(/*bExpectedResult=*/false);
 	}
 
 	{
-		Settings->Input1AttributeName = TrueAttribute;
-		Settings->Input2AttributeName = FalseAttribute;
+		Settings->InputSource1.AttributeName = TrueAttribute;
+		Settings->InputSource2.AttributeName = FalseAttribute;
 		bTestPassed &= ValidateOp(/*bExpectedResult=*/false);
 	}
 
 	{
-		Settings->Input1AttributeName = FalseAttribute;
-		Settings->Input2AttributeName = FalseAttribute;
+		Settings->InputSource1.AttributeName = FalseAttribute;
+		Settings->InputSource2.AttributeName = FalseAttribute;
 		bTestPassed &= ValidateOp(/*bExpectedResult=*/false);
 	}
 
@@ -188,26 +192,26 @@ bool FPCGMetadataBooleanOpTest::RunTest(const FString& Parameters)
 	/////////////////////////////////////////
 	Settings->Operation = EPCGMedadataBooleanOperation::Or;
 	{
-		Settings->Input1AttributeName = TrueAttribute;
-		Settings->Input2AttributeName = TrueAttribute;
+		Settings->InputSource1.AttributeName = TrueAttribute;
+		Settings->InputSource2.AttributeName = TrueAttribute;
 		bTestPassed &= ValidateOp(/*bExpectedResult=*/true);
 	}
 
 	{
-		Settings->Input1AttributeName = FalseAttribute;
-		Settings->Input2AttributeName = TrueAttribute;
+		Settings->InputSource1.AttributeName = FalseAttribute;
+		Settings->InputSource2.AttributeName = TrueAttribute;
 		bTestPassed &= ValidateOp(/*bExpectedResult=*/true);
 	}
 
 	{
-		Settings->Input1AttributeName = TrueAttribute;
-		Settings->Input2AttributeName = FalseAttribute;
+		Settings->InputSource1.AttributeName = TrueAttribute;
+		Settings->InputSource2.AttributeName = FalseAttribute;
 		bTestPassed &= ValidateOp(/*bExpectedResult=*/true);
 	}
 
 	{
-		Settings->Input1AttributeName = FalseAttribute;
-		Settings->Input2AttributeName = FalseAttribute;
+		Settings->InputSource1.AttributeName = FalseAttribute;
+		Settings->InputSource2.AttributeName = FalseAttribute;
 		bTestPassed &= ValidateOp(/*bExpectedResult=*/false);
 	}
 
@@ -216,26 +220,26 @@ bool FPCGMetadataBooleanOpTest::RunTest(const FString& Parameters)
 	/////////////////////////////////////////
 	Settings->Operation = EPCGMedadataBooleanOperation::Xor;
 	{
-		Settings->Input1AttributeName = TrueAttribute;
-		Settings->Input2AttributeName = TrueAttribute;
+		Settings->InputSource1.AttributeName = TrueAttribute;
+		Settings->InputSource2.AttributeName = TrueAttribute;
 		bTestPassed &= ValidateOp(/*bExpectedResult=*/false);
 	}
 
 	{
-		Settings->Input1AttributeName = FalseAttribute;
-		Settings->Input2AttributeName = TrueAttribute;
+		Settings->InputSource1.AttributeName = FalseAttribute;
+		Settings->InputSource2.AttributeName = TrueAttribute;
 		bTestPassed &= ValidateOp(/*bExpectedResult=*/true);
 	}
 
 	{
-		Settings->Input1AttributeName = TrueAttribute;
-		Settings->Input2AttributeName = FalseAttribute;
+		Settings->InputSource1.AttributeName = TrueAttribute;
+		Settings->InputSource2.AttributeName = FalseAttribute;
 		bTestPassed &= ValidateOp(/*bExpectedResult=*/true);
 	}
 
 	{
-		Settings->Input1AttributeName = FalseAttribute;
-		Settings->Input2AttributeName = FalseAttribute;
+		Settings->InputSource1.AttributeName = FalseAttribute;
+		Settings->InputSource2.AttributeName = FalseAttribute;
 		bTestPassed &= ValidateOp(/*bExpectedResult=*/false);
 	}
 
diff --git a/Engine/Plugins/Experimental/PCG/Source/PCG/Private/Tests/Elements/Metadata/PCGMetadataBreakTransformTest.cpp b/Engine/Plugins/Experimental/PCG/Source/PCG/Private/Tests/Elements/Metadata/PCGMetadataBreakTransformTest.cpp
--- a/Engine/Plugins/Experimental/PCG/Source/PCG/Private/Tests/Elements/Metadata/PCGMetadataBreakTransformTest.cpp
+++ b/Engine/Plugins/Experimental/PCG/Source/PCG/Private/Tests/Elements/Metadata/PCGMetadataBreakTransformTest.cpp
@@ -88,6 +88,31 @@ namespace PCGBreakTransformTest
 			Attribute->SetValue(EntryKey, RandomSource.FRand());
 		}
 	}
+
+	TArray<PCGMetadataEntryKey> GatherEntryKeys(const UPCGData* InData, const FPCGMetadataAttributeBase* InAttribute)
+	{
+		TArray<PCGMetadataEntryKey> Keys;
+		if (const UPCGPointData* PointOutput = Cast<const UPCGPointData>(InData))
+		{
+			for (const FPCGPoint& Point : PointOutput->GetPoints())
+			{
+				Keys.Add(Point.MetadataEntry);
+			}
+		}
+		else
+		{
+			const FPCGMetadataAttributeBase* Current = InAttribute;
+			while (Current)
+			{
+				TArray<PCGMetadataEntryKey> Temp;
+				Current->GetEntryToValueKeyMap_NotThreadSafe().GenerateKeyArray(Temp);
+				Keys.Append(Temp);
+				Current = Current->GetParent();
+			}
+		}
+
+		return Keys;
+	}
 }
 
 
@@ -98,7 +123,8 @@ bool FPCGMetadataBreakTransformTest::RunTest(const FString& Parameters)
 		const FName OutputAttributeName = TEXT("Output");
 
 		UPCGMetadataBreakTransformSettings* Settings = CastChecked<UPCGMetadataBreakTransformSettings>(TestData.Settings);
-		Settings->OutputAttributeName = OutputAttributeName;
+		Settings->OutputTarget.Selection = EPCGAttributePropertySelection::Attribute;
+		Settings->OutputTarget.AttributeName = OutputAttributeName;
 		FPCGElementPtr MetadataBreakTransformElement = TestData.Settings->GetElement();
 
 		TUniquePtr<FPCGContext> Context = TUniquePtr<FPCGContext>(MetadataBreakTransformElement->Initialize(TestData.InputData, TestData.TestPCGComponent, nullptr));
@@ -140,9 +166,11 @@ bool FPCGMetadataBreakTransformTest::RunTest(const FString& Parameters)
 
 		check(SourceMetadata);
 
-		const FPCGMetadataAttributeBase* SourceAttributeBase = SourceMetadata->GetConstAttribute(Settings->InputAttributeName);
+		const FPCGMetadataAttributeBase* SourceAttributeBase = SourceMetadata->GetConstAttribute(Settings->InputSource.GetName());
 		check(SourceAttributeBase);
 
+		TArray<PCGMetadataEntryKey> InputKeys = PCGBreakTransformTest::GatherEntryKeys(Input.Data, SourceAttributeBase);
+
 		auto ValidateComponentOutput = [&](const FPCGTaggedData& Output, const FName& OutAttributeName, PCGBreakTransformTest::EPCGComponentToCheck ComponentToCheck, auto DummyValue)
 		{
 			using OutType = decltype(DummyValue);
@@ -154,13 +182,13 @@ bool FPCGMetadataBreakTransformTest::RunTest(const FString& Parameters)
 			}
 
 			UPCGMetadata* OutMetadata = nullptr;
-			if (const UPCGSpatialData* SpatialInput = Cast<const UPCGSpatialData>(Output.Data))
+			if (const UPCGSpatialData* SpatialOutput = Cast<const UPCGSpatialData>(Output.Data))
 			{
-				OutMetadata = SpatialInput->Metadata;
+				OutMetadata = SpatialOutput->Metadata;
 			}
-			else if (const UPCGParamData* ParamsInput = Cast<const UPCGParamData>(Output.Data))
+			else if (const UPCGParamData* ParamsOutput = Cast<const UPCGParamData>(Output.Data))
 			{
-				OutMetadata = ParamsInput->Metadata;
+				OutMetadata = ParamsOutput->Metadata;
 			}
 			else
 			{
@@ -190,30 +218,18 @@ bool FPCGMetadataBreakTransformTest::RunTest(const FString& Parameters)
 			const FPCGMetadataAttribute<OutType>* OutAttribute = static_cast<const FPCGMetadataAttribute<OutType>*>(OutAttributeBase);
 			check(OutAttribute);
 
-			const PCGMetadataEntryKey InEntryKeyCount = SourceMetadata->GetItemCountForChild();
-			const PCGMetadataEntryKey OutEntryKeyCount = OutMetadata->GetItemCountForChild();
+			TArray<PCGMetadataEntryKey> OutputKeys = PCGBreakTransformTest::GatherEntryKeys(Output.Data, OutAttribute);
 
-			if (!TestEqual("Identical EntryKey counts", InEntryKeyCount, OutEntryKeyCount))
+			if (!TestEqual("Identical EntryKeys count", InputKeys.Num(), OutputKeys.Num()))
 			{
 				bTestPassed = false;
 				return;
 			}
 
-			const PCGMetadataValueKey InValueKeyCount = SourceAttributeBase->GetValueKeyOffsetForChild();
-			const PCGMetadataValueKey OutValueKeyCount = OutAttribute->GetValueKeyOffsetForChild();
-
-			if (!TestEqual("Identical ValueKey counts", InValueKeyCount, OutValueKeyCount))
+			for (int32 i = 0; i < OutputKeys.Num(); ++i)
 			{
-				bTestPassed = false;
-				return;
-			}
-
-			for (PCGMetadataEntryKey EntryKey = 0; EntryKey < InEntryKeyCount; ++EntryKey)
-			{
-				const PCGMetadataValueKey InValueKey = SourceAttributeBase->GetValueKey(EntryKey);
-				const PCGMetadataValueKey OutValueKey = OutAttribute->GetValueKey(EntryKey);
-
-				bTestPassed &= TestEqual("Identical value keys", InValueKey, OutValueKey);
+				const PCGMetadataValueKey InValueKey = SourceAttributeBase->GetValueKey(InputKeys[i]);
+				const PCGMetadataValueKey OutValueKey = OutAttribute->GetValueKey(OutputKeys[i]);
 
 				FTransform SourceValue = static_cast<const FPCGMetadataAttribute<FTransform>*>(SourceAttributeBase)->GetValue(InValueKey);
 
@@ -252,11 +268,11 @@ bool FPCGMetadataBreakTransformTest::RunTest(const FString& Parameters)
 		const TArray<FPCGTaggedData> OutputsRotation = Context->OutputData.GetInputsByPin(PCGMetadataTransformConstants::Rotation);
 		const TArray<FPCGTaggedData> OutputsScale = Context->OutputData.GetInputsByPin(PCGMetadataTransformConstants::Scale);
 
-		const FName DestinationAttributeForTranslation = Settings->GetOutputAttributeName(Settings->OutputAttributeName, 0);
-		const FName DestinationAttributeForRotation = Settings->GetOutputAttributeName(Settings->OutputAttributeName, 1);
-		const FName DestinationAttributeForScale = Settings->GetOutputAttributeName(Settings->OutputAttributeName, 2);
+		const FName DestinationAttributeForTranslation = Settings->OutputTarget.AttributeName; //Settings->GetOutputAttributeName(Settings->OutputName.AttributeName, 0);
+		const FName DestinationAttributeForRotation = Settings->OutputTarget.AttributeName; //Settings->GetOutputAttributeName(Settings->OutputName.AttributeName, 1);
+		const FName DestinationAttributeForScale = Settings->OutputTarget.AttributeName; //Settings->GetOutputAttributeName(Settings->OutputName.AttributeName, 2);
 
-		if (TestEqual("Appropriate number of outputs generated fox Translation", OutputsTranslation.Num(), 1))
+		if (TestEqual("Appropriate number of outputs generated for Translation", OutputsTranslation.Num(), 1))
 		{
 			ValidateComponentOutput(OutputsTranslation[0], DestinationAttributeForTranslation, PCGBreakTransformTest::EPCGComponentToCheck::Translation, FVector{});
 		}
@@ -303,7 +319,7 @@ bool FPCGMetadataBreakTransformTest::RunTest(const FString& Parameters)
 	AllTestData[1] = PairWhatData("Testing with param data as input", &TestDataParams);
 
 	// Setup error catching. This error should happen only twice (invalid type)
-	AddExpectedError(TEXT("Attribute Float is not a supported type for input 0"), EAutomationExpectedErrorFlags::Contains, 2);
+	AddExpectedError(TEXT("Attribute/Property Float is not a supported type for input 0"), EAutomationExpectedErrorFlags::Contains, 2);
 
 	for (PairWhatData& PairTestData : AllTestData)
 	{
@@ -316,15 +332,17 @@ bool FPCGMetadataBreakTransformTest::RunTest(const FString& Parameters)
 		Settings->ForceOutputConnections[1] = true;
 		Settings->ForceOutputConnections[2] = true;
 
+		Settings->InputSource.Selection = EPCGAttributePropertySelection::Attribute;
+
 		{
 			AddInfo(TEXT("Test with Transform as input attribute"));
-			Settings->InputAttributeName = PCGBreakTransformTest::TransformAttribute;
+			Settings->InputSource.AttributeName = PCGBreakTransformTest::TransformAttribute;
 			bTestPassed &= ValidateMetadataBreakVector(*TestData);
 		}
 
 		{
 			AddInfo(TEXT("Test with Float as input attribute (invalid type)"));
-			Settings->InputAttributeName = PCGBreakTransformTest::InvalidAttribute;
+			Settings->InputSource.AttributeName = PCGBreakTransformTest::InvalidAttribute;
 
 			bTestPassed &= ValidateMetadataBreakVector(*TestData, /*bIsValid=*/false);
 		}
diff --git a/Engine/Plugins/Experimental/PCG/Source/PCG/Private/Tests/Elements/Metadata/PCGMetadataBreakVectorTest.cpp b/Engine/Plugins/Experimental/PCG/Source/PCG/Private/Tests/Elements/Metadata/PCGMetadataBreakVectorTest.cpp
--- a/Engine/Plugins/Experimental/PCG/Source/PCG/Private/Tests/Elements/Metadata/PCGMetadataBreakVectorTest.cpp
+++ b/Engine/Plugins/Experimental/PCG/Source/PCG/Private/Tests/Elements/Metadata/PCGMetadataBreakVectorTest.cpp
@@ -117,7 +117,8 @@ bool FPCGMetadataBreakVectorTest::RunTest(const FString& Parameters)
 		const FName OutputAttributeName = TEXT("Output");
 
 		UPCGMetadataBreakVectorSettings* Settings = CastChecked<UPCGMetadataBreakVectorSettings>(TestData.Settings);
-		Settings->OutputAttributeName = OutputAttributeName;
+		Settings->OutputTarget.Selection = EPCGAttributePropertySelection::Attribute;
+		Settings->OutputTarget.AttributeName = OutputAttributeName;
 		FPCGElementPtr MetadataBreakVectorElement = TestData.Settings->GetElement();
 
 		TUniquePtr<FPCGContext> Context = TUniquePtr<FPCGContext>(MetadataBreakVectorElement->Initialize(TestData.InputData, TestData.TestPCGComponent, nullptr));
@@ -159,7 +160,7 @@ bool FPCGMetadataBreakVectorTest::RunTest(const FString& Parameters)
 
 		check(SourceMetadata);
 
-		const FPCGMetadataAttributeBase* SourceAttributeBase = SourceMetadata->GetConstAttribute(Settings->InputAttributeName);
+		const FPCGMetadataAttributeBase* SourceAttributeBase = SourceMetadata->GetConstAttribute(Settings->InputSource.GetName());
 		check(SourceAttributeBase);
 
 		auto ValidateComponentOutput = [&](const FPCGTaggedData& Output, const FName& OutAttributeName, PCGBreakVectorTest::EPCGComponentToCheck ComponentToCheck)
@@ -280,10 +281,10 @@ bool FPCGMetadataBreakVectorTest::RunTest(const FString& Parameters)
 		const TArray<FPCGTaggedData> OutputsZ = Context->OutputData.GetInputsByPin(PCGMetadataBreakVectorConstants::ZLabel);
 		const TArray<FPCGTaggedData> OutputsW = Context->OutputData.GetInputsByPin(PCGMetadataBreakVectorConstants::WLabel);
 
-		const FName DestinationAttributeForX = Settings->GetOutputAttributeName(Settings->OutputAttributeName, 0);
-		const FName DestinationAttributeForY = Settings->GetOutputAttributeName(Settings->OutputAttributeName, 1);
-		const FName DestinationAttributeForZ = Settings->GetOutputAttributeName(Settings->OutputAttributeName, 2);
-		const FName DestinationAttributeForW = Settings->GetOutputAttributeName(Settings->OutputAttributeName, 3);
+		const FName DestinationAttributeForX = Settings->OutputTarget.AttributeName; //Settings->GetOutputAttributeName(Settings->OutputName.AttributeName, 0);
+		const FName DestinationAttributeForY = Settings->OutputTarget.AttributeName; //Settings->GetOutputAttributeName(Settings->OutputName.AttributeName, 1);
+		const FName DestinationAttributeForZ = Settings->OutputTarget.AttributeName; //Settings->GetOutputAttributeName(Settings->OutputName.AttributeName, 2);
+		const FName DestinationAttributeForW = Settings->OutputTarget.AttributeName; //Settings->GetOutputAttributeName(Settings->OutputName.AttributeName, 3);
 
 		if (Settings->ForceOutputConnections[0])
 		{
@@ -353,7 +354,7 @@ bool FPCGMetadataBreakVectorTest::RunTest(const FString& Parameters)
 	AllTestData[1] = PairWhatData("Testing with param data as input", &TestDataParams);
 
 	// Setup error catching. This error should happen only twice (invalid type)
-	AddExpectedError(TEXT("Attribute Float is not a supported type for input 0"), EAutomationExpectedErrorFlags::Contains, 2);
+	AddExpectedError(TEXT("Attribute/Property Float is not a supported type for input 0"), EAutomationExpectedErrorFlags::Contains, 2);
 
 	for (PairWhatData& PairTestData : AllTestData)
 	{
@@ -367,9 +368,11 @@ bool FPCGMetadataBreakVectorTest::RunTest(const FString& Parameters)
 		Settings->ForceOutputConnections[2] = false;
 		Settings->ForceOutputConnections[3] = false;
 
+		Settings->InputSource.Selection = EPCGAttributePropertySelection::Attribute;
+
 		{
 			AddInfo(TEXT("Test with FVector2D as input attribute"));
-			Settings->InputAttributeName = PCGBreakVectorTest::Vec2Attribute;
+			Settings->InputSource.AttributeName = PCGBreakVectorTest::Vec2Attribute;
 			bTestPassed &= ValidateMetadataBreakVector(*TestData);
 		}
 
@@ -377,13 +380,13 @@ bool FPCGMetadataBreakVectorTest::RunTest(const FString& Parameters)
 
 		{
 			AddInfo(TEXT("Test with FVector as input attribute"));
-			Settings->InputAttributeName = PCGBreakVectorTest::Vec3Attribute;
+			Settings->InputSource.AttributeName = PCGBreakVectorTest::Vec3Attribute;
 			bTestPassed &= ValidateMetadataBreakVector(*TestData);
 		}
 
 		{
 			AddInfo(TEXT("Test with FVector4 as input attribute"));
-			Settings->InputAttributeName = PCGBreakVectorTest::Vec4Attribute;
+			Settings->InputSource.AttributeName = PCGBreakVectorTest::Vec4Attribute;
 			Settings->ForceOutputConnections[3] = true;
 			bTestPassed &= ValidateMetadataBreakVector(*TestData);
 			Settings->ForceOutputConnections[3] = false;
@@ -391,13 +394,13 @@ bool FPCGMetadataBreakVectorTest::RunTest(const FString& Parameters)
 
 		{
 			AddInfo(TEXT("Test with FRotator as input attribute"));
-			Settings->InputAttributeName = PCGBreakVectorTest::RotatorAttribute;
+			Settings->InputSource.AttributeName = PCGBreakVectorTest::RotatorAttribute;
 			bTestPassed &= ValidateMetadataBreakVector(*TestData);
 		}
 
 		{
 			AddInfo(TEXT("Test with Float as input attribute (invalid type)"));
-			Settings->InputAttributeName = PCGBreakVectorTest::FloatAttribute;
+			Settings->InputSource.AttributeName = PCGBreakVectorTest::FloatAttribute;
 
 			bTestPassed &= ValidateMetadataBreakVector(*TestData, /*bIsValid=*/false);
 		}
diff --git a/Engine/Plugins/Experimental/PCG/Source/PCG/Private/Tests/Elements/Metadata/PCGMetadataMakeTransformTest.cpp b/Engine/Plugins/Experimental/PCG/Source/PCG/Private/Tests/Elements/Metadata/PCGMetadataMakeTransformTest.cpp
--- a/Engine/Plugins/Experimental/PCG/Source/PCG/Private/Tests/Elements/Metadata/PCGMetadataMakeTransformTest.cpp
+++ b/Engine/Plugins/Experimental/PCG/Source/PCG/Private/Tests/Elements/Metadata/PCGMetadataMakeTransformTest.cpp
@@ -95,7 +95,8 @@ bool FPCGMetadataMakeTransformTest::RunTest(const FString& Parameters)
 		const FName OutputAttributeName = TEXT("Output");
 
 		UPCGMetadataMakeTransformSettings* Settings = CastChecked<UPCGMetadataMakeTransformSettings>(TestData.Settings);
-		Settings->OutputAttributeName = OutputAttributeName;
+		Settings->OutputTarget.Selection = EPCGAttributePropertySelection::Attribute;
+		Settings->OutputTarget.AttributeName = OutputAttributeName;
 		FPCGElementPtr MetadataBreakTransformElement = TestData.Settings->GetElement();
 
 		TUniquePtr<FPCGContext> Context = TUniquePtr<FPCGContext>(MetadataBreakTransformElement->Initialize(TestData.InputData, TestData.TestPCGComponent, nullptr));
@@ -130,13 +131,13 @@ bool FPCGMetadataMakeTransformTest::RunTest(const FString& Parameters)
 
 		check(SourceMetadata);
 
-		const FPCGMetadataAttributeBase* TranslationAttributeBase = SourceMetadata->GetConstAttribute(Settings->Input1AttributeName);
+		const FPCGMetadataAttributeBase* TranslationAttributeBase = SourceMetadata->GetConstAttribute(Settings->InputSource1.GetName());
 		check(TranslationAttributeBase);
 
-		const FPCGMetadataAttributeBase* RotationAttributeBase = SourceMetadata->GetConstAttribute(Settings->Input2AttributeName);
+		const FPCGMetadataAttributeBase* RotationAttributeBase = SourceMetadata->GetConstAttribute(Settings->InputSource2.GetName());
 		check(RotationAttributeBase);
 
-		const FPCGMetadataAttributeBase* ScaleAttributeBase = SourceMetadata->GetConstAttribute(Settings->Input3AttributeName);
+		const FPCGMetadataAttributeBase* ScaleAttributeBase = SourceMetadata->GetConstAttribute(Settings->InputSource3.GetName());
 		check(ScaleAttributeBase);
 
 		const TArray<FPCGTaggedData> Outputs = Context->OutputData.GetInputsByPin(PCGPinConstants::DefaultOutputLabel);
@@ -171,7 +172,7 @@ bool FPCGMetadataMakeTransformTest::RunTest(const FString& Parameters)
 			return false;
 		}
 
-		const FPCGMetadataAttributeBase* OutAttributeBase = OutMetadata->GetConstAttribute(Settings->OutputAttributeName);
+		const FPCGMetadataAttributeBase* OutAttributeBase = OutMetadata->GetConstAttribute(Settings->OutputTarget.GetName());
 
 		if (!TestNotNull("Valid output attribute", OutAttributeBase))
 		{
@@ -256,9 +257,14 @@ bool FPCGMetadataMakeTransformTest::RunTest(const FString& Parameters)
 		UPCGMetadataMakeTransformSettings* Settings = CastChecked<UPCGMetadataMakeTransformSettings>(TestData->Settings);
 
 		Settings->ForceOutputConnections[0] = true;
-		Settings->Input1AttributeName = Vec3Attribute;
-		Settings->Input2AttributeName = QuatAttribute;
-		Settings->Input3AttributeName = Vec3Attribute;
+		Settings->InputSource1.Selection = EPCGAttributePropertySelection::Attribute;
+		Settings->InputSource1.AttributeName = Vec3Attribute;
+
+		Settings->InputSource2.Selection = EPCGAttributePropertySelection::Attribute;
+		Settings->InputSource2.AttributeName = QuatAttribute;
+
+		Settings->InputSource3.Selection = EPCGAttributePropertySelection::Attribute;
+		Settings->InputSource3.AttributeName = Vec3Attribute;
 
 		bTestPassed &= ValidateMetadataMakeVector(*TestData);
 	}
diff --git a/Engine/Plugins/Experimental/PCG/Source/PCG/Public/Elements/Metadata/PCGMetadataBitwiseOpElement.h b/Engine/Plugins/Experimental/PCG/Source/PCG/Public/Elements/Metadata/PCGMetadataBitwiseOpElement.h
--- a/Engine/Plugins/Experimental/PCG/Source/PCG/Public/Elements/Metadata/PCGMetadataBitwiseOpElement.h
+++ b/Engine/Plugins/Experimental/PCG/Source/PCG/Public/Elements/Metadata/PCGMetadataBitwiseOpElement.h
@@ -23,6 +23,10 @@ class PCG_API UPCGMetadataBitwiseSettings : public UPCGMetadataSettingsBase
 	GENERATED_BODY()
 
 public:
+	// ~Begin UObject interface
+	virtual void PostLoad() override;
+	// ~End UObject interface
+
 	//~Begin UPCGSettings interface
 #if WITH_EDITOR
 	virtual FName GetDefaultNodeName() const override;
@@ -31,7 +35,7 @@ public:
 	//~End UPCGSettings interface
 
 	//~Begin UPCGMetadataSettingsBase interface
-	virtual FName GetInputAttributeNameWithOverride(uint32 Index, UPCGParamData* Params) const override;
+	FPCGAttributePropertySelector GetInputSource(uint32 Index) const override;
 
 	virtual FName GetInputPinLabel(uint32 Index) const override;
 	virtual uint32 GetInputPinNum() const override;
@@ -50,10 +54,18 @@ public:
 	EPCGMedadataBitwiseOperation Operation = EPCGMedadataBitwiseOperation::And;
 
 	UPROPERTY(BlueprintReadWrite, EditAnywhere, Category = Input)
-	FName Input1AttributeName = NAME_None;
+	FPCGAttributePropertySelector InputSource1;
 
 	UPROPERTY(BlueprintReadWrite, EditAnywhere, Category = Input, meta = (EditCondition = "Operation != EPCGMedadataBitwiseOperation::Not", EditConditionHides))
-	FName Input2AttributeName = NAME_None;
+	FPCGAttributePropertySelector InputSource2;
+	
+#if WITH_EDITORONLY_DATA
+	UPROPERTY()
+	FName Input1AttributeName_DEPRECATED = NAME_None;
+
+	UPROPERTY()
+	FName Input2AttributeName_DEPRECATED = NAME_None;
+#endif
 };
 
 class FPCGMetadataBitwiseElement : public FPCGMetadataElementBase
diff --git a/Engine/Plugins/Experimental/PCG/Source/PCG/Public/Elements/Metadata/PCGMetadataBooleanOpElement.h b/Engine/Plugins/Experimental/PCG/Source/PCG/Public/Elements/Metadata/PCGMetadataBooleanOpElement.h
--- a/Engine/Plugins/Experimental/PCG/Source/PCG/Public/Elements/Metadata/PCGMetadataBooleanOpElement.h
+++ b/Engine/Plugins/Experimental/PCG/Source/PCG/Public/Elements/Metadata/PCGMetadataBooleanOpElement.h
@@ -23,6 +23,10 @@ class PCG_API UPCGMetadataBooleanSettings : public UPCGMetadataSettingsBase
 	GENERATED_BODY()
 
 public:
+	// ~Begin UObject interface
+	virtual void PostLoad() override;
+	// ~End UObject interface
+	// 
 	//~Begin UPCGSettings interface
 #if WITH_EDITOR
 	virtual FName GetDefaultNodeName() const override;
@@ -31,7 +35,7 @@ public:
 	//~End UPCGSettings interface
 
 	//~Begin UPCGMetadataSettingsBase interface
-	virtual FName GetInputAttributeNameWithOverride(uint32 Index, UPCGParamData* Params) const override;
+	FPCGAttributePropertySelector GetInputSource(uint32 Index) const override;
 
 	virtual FName GetInputPinLabel(uint32 Index) const override;
 	virtual uint32 GetInputPinNum() const override;
@@ -50,10 +54,18 @@ public:
 	EPCGMedadataBooleanOperation Operation = EPCGMedadataBooleanOperation::And;
 
 	UPROPERTY(BlueprintReadWrite, EditAnywhere, Category = Input)
-	FName Input1AttributeName = NAME_None;
+	FPCGAttributePropertySelector InputSource1;
 
-	UPROPERTY(BlueprintReadWrite, EditAnywhere, Category = Input, meta = (EditCondition="Operation != EPCGMedadataBooleanOperation::Not", EditConditionHides))
-	FName Input2AttributeName = NAME_None;
+	UPROPERTY(BlueprintReadWrite, EditAnywhere, Category = Input, meta = (EditCondition = "Operation != EPCGMedadataBooleanOperation::Not", EditConditionHides))
+	FPCGAttributePropertySelector InputSource2;
+
+#if WITH_EDITORONLY_DATA
+	UPROPERTY()
+	FName Input1AttributeName_DEPRECATED = NAME_None;
+
+	UPROPERTY()
+	FName Input2AttributeName_DEPRECATED = NAME_None;
+#endif
 };
 
 class FPCGMetadataBooleanElement : public FPCGMetadataElementBase
diff --git a/Engine/Plugins/Experimental/PCG/Source/PCG/Public/Elements/Metadata/PCGMetadataBreakTransform.h b/Engine/Plugins/Experimental/PCG/Source/PCG/Public/Elements/Metadata/PCGMetadataBreakTransform.h
--- a/Engine/Plugins/Experimental/PCG/Source/PCG/Public/Elements/Metadata/PCGMetadataBreakTransform.h
+++ b/Engine/Plugins/Experimental/PCG/Source/PCG/Public/Elements/Metadata/PCGMetadataBreakTransform.h
@@ -16,12 +16,16 @@ class PCG_API UPCGMetadataBreakTransformSettings : public UPCGMetadataSettingsBa
 	GENERATED_BODY()
 
 public:
+	// ~Begin UObject interface
+	virtual void PostLoad() override;
+	// ~End UObject interface
+	// 
 	//~Begin UPCGSettings interface
 #if WITH_EDITOR
 	virtual FName GetDefaultNodeName() const override;
 #endif
 
-	virtual FName GetInputAttributeNameWithOverride(uint32 Index, UPCGParamData* Params) const override;
+	FPCGAttributePropertySelector GetInputSource(uint32 Index) const override;
 
 	virtual FName GetOutputPinLabel(uint32 Index) const override;
 	virtual uint32 GetOutputPinNum() const override;
@@ -37,8 +41,13 @@ protected:
 	//~End UPCGSettings interface
 
 public:
-	UPROPERTY(BlueprintReadWrite, EditAnywhere, Category = Settings)
-	FName InputAttributeName = NAME_None;
+	UPROPERTY(BlueprintReadWrite, EditAnywhere, Category = Input)
+	FPCGAttributePropertySelector InputSource;
+
+#if WITH_EDITORONLY_DATA
+	UPROPERTY()
+	FName InputAttributeName_DEPRECATED = NAME_None;
+#endif
 };
 
 class FPCGMetadataBreakTransformElement : public FPCGMetadataElementBase
diff --git a/Engine/Plugins/Experimental/PCG/Source/PCG/Public/Elements/Metadata/PCGMetadataBreakVector.h b/Engine/Plugins/Experimental/PCG/Source/PCG/Public/Elements/Metadata/PCGMetadataBreakVector.h
--- a/Engine/Plugins/Experimental/PCG/Source/PCG/Public/Elements/Metadata/PCGMetadataBreakVector.h
+++ b/Engine/Plugins/Experimental/PCG/Source/PCG/Public/Elements/Metadata/PCGMetadataBreakVector.h
@@ -23,12 +23,16 @@ class PCG_API UPCGMetadataBreakVectorSettings : public UPCGMetadataSettingsBase
 	GENERATED_BODY()
 
 public:
+	// ~Begin UObject interface
+	virtual void PostLoad() override;
+	// ~End UObject interface
+
 	//~Begin UPCGSettings interface
 #if WITH_EDITOR
 	virtual FName GetDefaultNodeName() const override;
 #endif
 
-	virtual FName GetInputAttributeNameWithOverride(uint32 Index, UPCGParamData* Params) const override;
+	FPCGAttributePropertySelector GetInputSource(uint32 Index) const override;
 
 	virtual FName GetOutputPinLabel(uint32 Index) const override;
 	virtual uint32 GetOutputPinNum() const override;
@@ -42,8 +46,13 @@ protected:
 	//~End UPCGSettings interface
 
 public:
-	UPROPERTY(BlueprintReadWrite, EditAnywhere, Category = Settings)
-	FName InputAttributeName = NAME_None;
+	UPROPERTY(BlueprintReadWrite, EditAnywhere, Category = Input)
+	FPCGAttributePropertySelector InputSource;
+
+#if WITH_EDITORONLY_DATA
+	UPROPERTY()
+	FName InputAttributeName_DEPRECATED = NAME_None;
+#endif
 };
 
 class FPCGMetadataBreakVectorElement : public FPCGMetadataElementBase
diff --git a/Engine/Plugins/Experimental/PCG/Source/PCG/Public/Elements/Metadata/PCGMetadataCompareOpElement.h b/Engine/Plugins/Experimental/PCG/Source/PCG/Public/Elements/Metadata/PCGMetadataCompareOpElement.h
--- a/Engine/Plugins/Experimental/PCG/Source/PCG/Public/Elements/Metadata/PCGMetadataCompareOpElement.h
+++ b/Engine/Plugins/Experimental/PCG/Source/PCG/Public/Elements/Metadata/PCGMetadataCompareOpElement.h
@@ -25,6 +25,10 @@ class PCG_API UPCGMetadataCompareSettings : public UPCGMetadataSettingsBase
 	GENERATED_BODY()
 
 public:
+	// ~Begin UObject interface
+	virtual void PostLoad() override;
+	// ~End UObject interface
+
 	//~Begin UPCGSettings interface
 #if WITH_EDITOR
 	virtual FName GetDefaultNodeName() const override;
@@ -33,7 +37,7 @@ public:
 	//~End UPCGSettings interface
 
 	//~Begin UPCGMetadataSettingsBase interface
-	virtual FName GetInputAttributeNameWithOverride(uint32 Index, UPCGParamData* Params) const override;
+	FPCGAttributePropertySelector GetInputSource(uint32 Index) const override;
 
 	virtual FName GetInputPinLabel(uint32 Index) const override;
 	virtual uint32 GetInputPinNum() const override;
@@ -52,13 +56,21 @@ public:
 	EPCGMedadataCompareOperation Operation = EPCGMedadataCompareOperation::Equal;
 
 	UPROPERTY(BlueprintReadWrite, EditAnywhere, Category = Input)
-	FName Input1AttributeName = NAME_None;
+	FPCGAttributePropertySelector InputSource1;
 
 	UPROPERTY(BlueprintReadWrite, EditAnywhere, Category = Input)
-	FName Input2AttributeName = NAME_None;
+	FPCGAttributePropertySelector InputSource2;
 
 	UPROPERTY(BlueprintReadWrite, EditAnywhere, Category = Settings, meta = (EditCondition = "Operation == EPCGMedadataCompareOperation::Equal || Operation == EPCGMedadataCompareOperation::NotEqual", EditConditionHides))
 	double Tolerance = UE_DOUBLE_SMALL_NUMBER;
+
+#if WITH_EDITORONLY_DATA
+	UPROPERTY()
+	FName Input1AttributeName_DEPRECATED = NAME_None;
+
+	UPROPERTY()
+	FName Input2AttributeName_DEPRECATED = NAME_None;
+#endif
 };
 
 class FPCGMetadataCompareElement : public FPCGMetadataElementBase
diff --git a/Engine/Plugins/Experimental/PCG/Source/PCG/Public/Elements/Metadata/PCGMetadataMakeTransform.h b/Engine/Plugins/Experimental/PCG/Source/PCG/Public/Elements/Metadata/PCGMetadataMakeTransform.h
--- a/Engine/Plugins/Experimental/PCG/Source/PCG/Public/Elements/Metadata/PCGMetadataMakeTransform.h
+++ b/Engine/Plugins/Experimental/PCG/Source/PCG/Public/Elements/Metadata/PCGMetadataMakeTransform.h
@@ -22,13 +22,17 @@ class PCG_API UPCGMetadataMakeTransformSettings : public UPCGMetadataSettingsBas
 	GENERATED_BODY()
 
 public:
+	// ~Begin UObject interface
+	virtual void PostLoad() override;
+	// ~End UObject interface
+
 #if WITH_EDITOR
 	//~Begin UPCGSettings interface
 	virtual FName GetDefaultNodeName() const override;
 	//~End UPCGSettings interface
 #endif
 
-	virtual FName GetInputAttributeNameWithOverride(uint32 Index, UPCGParamData* Params) const override;
+	FPCGAttributePropertySelector GetInputSource(uint32 Index) const override;
 
 	virtual FName GetInputPinLabel(uint32 Index) const override;
 	virtual uint32 GetInputPinNum() const override;
@@ -43,13 +47,24 @@ protected:
 
 public:
 	UPROPERTY(BlueprintReadWrite, EditAnywhere, Category = Input)
-	FName Input1AttributeName = NAME_None;
+	FPCGAttributePropertySelector InputSource1;
 
 	UPROPERTY(BlueprintReadWrite, EditAnywhere, Category = Input)
-	FName Input2AttributeName = NAME_None;
-
+	FPCGAttributePropertySelector InputSource2;
+	
 	UPROPERTY(BlueprintReadWrite, EditAnywhere, Category = Input)
-	FName Input3AttributeName = NAME_None;
+	FPCGAttributePropertySelector InputSource3;
+
+#if WITH_EDITORONLY_DATA
+	UPROPERTY()
+	FName Input1AttributeName_DEPRECATED = NAME_None;
+
+	UPROPERTY()
+	FName Input2AttributeName_DEPRECATED = NAME_None;
+
+	UPROPERTY()
+	FName Input3AttributeName_DEPRECATED = NAME_None;
+#endif
 };
 
 class FPCGMetadataMakeTransformElement : public FPCGMetadataElementBase
diff --git a/Engine/Plugins/Experimental/PCG/Source/PCG/Public/Elements/Metadata/PCGMetadataMakeVector.h b/Engine/Plugins/Experimental/PCG/Source/PCG/Public/Elements/Metadata/PCGMetadataMakeVector.h
--- a/Engine/Plugins/Experimental/PCG/Source/PCG/Public/Elements/Metadata/PCGMetadataMakeVector.h
+++ b/Engine/Plugins/Experimental/PCG/Source/PCG/Public/Elements/Metadata/PCGMetadataMakeVector.h
@@ -44,13 +44,17 @@ class PCG_API UPCGMetadataMakeVectorSettings : public UPCGMetadataSettingsBase
 	GENERATED_BODY()
 
 public:
+	// ~Begin UObject interface
+	virtual void PostLoad() override;
+	// ~End UObject interface
+
 #if WITH_EDITOR
 	//~Begin UPCGSettings interface
 	virtual FName GetDefaultNodeName() const override;
 	//~End UPCGSettings interface
 #endif
 
-	virtual FName GetInputAttributeNameWithOverride(uint32 Index, UPCGParamData* Params) const override;
+	FPCGAttributePropertySelector GetInputSource(uint32 Index) const override;
 
 	virtual FName GetInputPinLabel(uint32 Index) const override;
 	virtual uint32 GetInputPinNum() const override;
@@ -65,16 +69,16 @@ protected:
 
 public:
 	UPROPERTY(BlueprintReadWrite, EditAnywhere, Category = Input)
-	FName Input1AttributeName = NAME_None;
+	FPCGAttributePropertySelector InputSource1;
 
 	UPROPERTY(BlueprintReadWrite, EditAnywhere, Category = Input)
-	FName Input2AttributeName = NAME_None;
+	FPCGAttributePropertySelector InputSource2;
 
 	UPROPERTY(BlueprintReadWrite, EditAnywhere, Category = Input)
-	FName Input3AttributeName = NAME_None;
+	FPCGAttributePropertySelector InputSource3;
 
 	UPROPERTY(BlueprintReadWrite, EditAnywhere, Category = Input)
-	FName Input4AttributeName = NAME_None;
+	FPCGAttributePropertySelector InputSource4;
 
 	UPROPERTY(BlueprintReadWrite, EditAnywhere, Category = Settings, meta = (ValidEnumValues = "Vector2, Vector, Vector4"))
 	EPCGMetadataTypes OutputType = EPCGMetadataTypes::Vector2;
@@ -84,6 +88,20 @@ public:
 
 	UPROPERTY(BlueprintReadWrite, EditAnywhere, Category = Settings, meta = (EditCondition = "OutputType == EPCGMetadataTypes::Vector4", EditConditionHides))
 	EPCGMetadataMakeVector4 MakeVector4Op = EPCGMetadataMakeVector4::FourValues;
+
+#if WITH_EDITORONLY_DATA
+	UPROPERTY()
+	FName Input1AttributeName_DEPRECATED = NAME_None;
+
+	UPROPERTY()
+	FName Input2AttributeName_DEPRECATED = NAME_None;
+
+	UPROPERTY()
+	FName Input3AttributeName_DEPRECATED = NAME_None;
+
+	UPROPERTY()
+	FName Input4AttributeName_DEPRECATED = NAME_None;
+#endif
 };
 
 class FPCGMetadataMakeVectorElement : public FPCGMetadataElementBase
diff --git a/Engine/Plugins/Experimental/PCG/Source/PCG/Public/Elements/Metadata/PCGMetadataMathsOpElement.h b/Engine/Plugins/Experimental/PCG/Source/PCG/Public/Elements/Metadata/PCGMetadataMathsOpElement.h
--- a/Engine/Plugins/Experimental/PCG/Source/PCG/Public/Elements/Metadata/PCGMetadataMathsOpElement.h
+++ b/Engine/Plugins/Experimental/PCG/Source/PCG/Public/Elements/Metadata/PCGMetadataMathsOpElement.h
@@ -47,6 +47,10 @@ class PCG_API UPCGMetadataMathsSettings : public UPCGMetadataSettingsBase
 	GENERATED_BODY()
 
 public:
+	// ~Begin UObject interface
+	virtual void PostLoad() override;
+	// ~End UObject interface
+
 	//~Begin UPCGSettings interface
 #if WITH_EDITOR
 	virtual FName GetDefaultNodeName() const override;
@@ -55,7 +59,7 @@ public:
 	//~End UPCGSettings interface
 
 	//~Begin UPCGMetadataSettingsBase interface
-	virtual FName GetInputAttributeNameWithOverride(uint32 Index, UPCGParamData* Params) const override;
+	FPCGAttributePropertySelector GetInputSource(uint32 Index) const override;
 
 	virtual FName GetInputPinLabel(uint32 Index) const override;
 	virtual uint32 GetInputPinNum() const override;
@@ -73,13 +77,24 @@ public:
 	EPCGMedadataMathsOperation Operation = EPCGMedadataMathsOperation::Add;
 
 	UPROPERTY(BlueprintReadWrite, EditAnywhere, Category = Input)
-	FName Input1AttributeName = NAME_None;
+	FPCGAttributePropertySelector InputSource1;
 
-	UPROPERTY(BlueprintReadWrite, EditAnywhere, Category = Input, meta = (EditCondition="(Operation & '/Script/PCG.EPCGMedadataMathsOperation::BinaryOp') || (Operation & '/Script/PCG.EPCGMedadataMathsOperation::TernaryOp')", EditConditionHides))
-	FName Input2AttributeName = NAME_None;
+	UPROPERTY(BlueprintReadWrite, EditAnywhere, Category = Input, meta = (EditCondition = "(Operation & '/Script/PCG.EPCGMedadataMathsOperation::BinaryOp') || (Operation & '/Script/PCG.EPCGMedadataMathsOperation::TernaryOp')", EditConditionHides))
+	FPCGAttributePropertySelector InputSource2;
 
 	UPROPERTY(BlueprintReadWrite, EditAnywhere, Category = Input, meta = (EditCondition = "Operation & '/Script/PCG.EPCGMedadataMathsOperation::TernaryOp'", EditConditionHides))
-	FName Input3AttributeName = NAME_None;
+	FPCGAttributePropertySelector InputSource3;
+
+#if WITH_EDITORONLY_DATA
+	UPROPERTY()
+	FName Input1AttributeName_DEPRECATED = NAME_None;
+
+	UPROPERTY()
+	FName Input2AttributeName_DEPRECATED = NAME_None;
+
+	UPROPERTY()
+	FName Input3AttributeName_DEPRECATED = NAME_None;
+#endif
 };
 
 class FPCGMetadataMathsElement : public FPCGMetadataElementBase
diff --git a/Engine/Plugins/Experimental/PCG/Source/PCG/Public/Elements/Metadata/PCGMetadataOpElementBase.h b/Engine/Plugins/Experimental/PCG/Source/PCG/Public/Elements/Metadata/PCGMetadataOpElementBase.h
--- a/Engine/Plugins/Experimental/PCG/Source/PCG/Public/Elements/Metadata/PCGMetadataOpElementBase.h
+++ b/Engine/Plugins/Experimental/PCG/Source/PCG/Public/Elements/Metadata/PCGMetadataOpElementBase.h
@@ -7,14 +7,13 @@
 
 #include "Metadata/PCGMetadataAttribute.h"
 #include "Metadata/PCGMetadataAttributeTpl.h"
-#include "Metadata/PCGMetadataEntryKeyIterator.h"
+#include "Metadata/PCGMetadataAttributeWrapper.h"
 
 #include "Containers/StaticArray.h"
 
 #include "PCGMetadataOpElementBase.generated.h"
 
 class FPCGMetadataAttributeBase;
-class IPCGMetadataEntryIterator;
 
 namespace PCGMetadataSettingsBaseConstants
 {
@@ -44,15 +43,19 @@ enum class EPCGMetadataSettingsBaseTypes
 };
 
 /**
- * Base class for all Metadata operations
- * A metadata operation can work on 2 different type of inputs: ParamData and SpatialData
- * Each of those inputs can have some metadata.
- * The output will contain the metadata of the first input (all its attributes) + the result of the operation (in a separate attribute)
+ * Base class for all Metadata operations.
+ * Metadata operation can work with attributes or properties. For example you could compute the addition between all points density and a constant from
+ * a param data.
+ * The output will be the duplication of the first input, with the same metadata + the result of the operation (either in an attribute or a property)
  * The new attribute can collide with one of the attributes in the incoming metadata. In this case, the attribute value will be overridden by the result
  * of the operation. It will also override the type of the attribute if it doesn't match the original.
  * 
- * You can specify the name of the attribute for each input and for the output. If they are None, they will take the default attribute.
- * Attribute names can also be overridden by ParamData, just connect the Param pin with some param data that matches exactly the name of the property you want to override.
+ * We only support operations between points and between spatial data. They all need to match (or be a param data)
+ * For example, if input 0 is a point data and input 1 is a spatial data, we fail.
+ * 
+ * You can specify the name of the attribute for each input and for the output.
+ * If the input name is None, it will take the lastest attribute in the input metadata.
+ * If the output name is None, it will take the input name.
  * 
  * Each operation has some requirements for the input types, and can broadcast some values into others (example Vector + Float -> Vector).
  * For example, if the op only accept booleans, all other value types will throw an error.
@@ -69,6 +72,10 @@ class PCG_API UPCGMetadataSettingsBase : public UPCGSettings
 	GENERATED_BODY()
 
 public:
+	// ~Begin UObject interface
+	virtual void PostLoad() override;
+	// ~End UObject interface
+
 	//~Begin UPCGSettings interface
 #if WITH_EDITOR
 	virtual EPCGSettingsType GetType() const override { return EPCGSettingsType::Metadata; }
@@ -78,7 +85,7 @@ public:
 	virtual TArray<FPCGPinProperties> OutputPinProperties() const override;
 	//~End UPCGSettings interface
 
-	virtual FName GetInputAttributeNameWithOverride(uint32 Index, UPCGParamData* Params) const { return NAME_None; };
+	virtual FPCGAttributePropertySelector GetInputSource(uint32 Index) const { return FPCGAttributePropertySelector(); };
 
 	virtual FName GetInputPinLabel(uint32 Index) const { return PCGPinConstants::DefaultInputLabel; }
 	virtual uint32 GetInputPinNum() const { return 1; };
@@ -95,11 +102,16 @@ public:
 
 	bool IsMoreComplexType(uint16 FirstType, uint16 SecondType) const;
 
-	UPROPERTY(BlueprintReadWrite, EditAnywhere, Category = Output)
-	FName OutputAttributeName = NAME_None;
+	UPROPERTY(BlueprintReadWrite, EditAnywhere, Category = "Output")
+	FPCGAttributePropertySelector OutputTarget;
 
-	UPROPERTY(BlueprintReadWrite, EditAnywhere, Category = Settings)
-	EPCGMetadataSettingsBaseMode Mode = EPCGMetadataSettingsBaseMode::Inferred;
+#if WITH_EDITORONLY_DATA
+	UPROPERTY()
+	FName OutputAttributeName_DEPRECATED = NAME_None;
+
+	UPROPERTY()
+	EPCGMetadataSettingsBaseMode Mode_DEPRECATED = EPCGMetadataSettingsBaseMode::Inferred;
+#endif // WITH_EDITORONLY_DATA
 
 	static constexpr uint32 MaxNumberOfOutputs = 4;
 
@@ -115,13 +127,13 @@ class FPCGMetadataElementBase : public FSimplePCGElement
 public:
 	struct FOperationData
 	{
-		TArray<const FPCGMetadataAttributeBase*> SourceAttributes;
 		int32 NumberOfElementsToProcess = -1;
 		uint16 MostComplexInputType;
 		uint16 OutputType;
-		TArray<FPCGMetadataAttributeBase*> OutputAttributes;
 		const UPCGMetadataSettingsBase* Settings = nullptr;
-		TArray<TUniquePtr<IPCGMetadataEntryIterator>> Iterators;
+
+		TArray<FPCGPropertyAttributeIterator> InputIterators;
+		TArray<FPCGPropertyAttributeIterator> OutputIterators;
 	};
 
 protected:
@@ -149,58 +161,45 @@ protected:
 template <typename InType, typename... Callbacks>
 inline bool FPCGMetadataElementBase::DoUnaryOp(FOperationData& OperationData, Callbacks&& ...InCallbacks) const
 {
-	check(OperationData.SourceAttributes[0]);
+	check(OperationData.InputIterators[0].IsValid());
 
-	const uint32 NumOutputs = OperationData.OutputAttributes.Num();
+	const uint32 NumOutputs = OperationData.OutputIterators.Num();
 	constexpr uint32 NumCallbacks = (uint32)sizeof...(InCallbacks);
 	check(NumOutputs == NumCallbacks);
 
-	// All values in OperationData.Iterators are UniquePtr. Dereference them here to make the syntax less cumbersome.
-	check(OperationData.Iterators[0]);
-	IPCGMetadataEntryIterator& Iterator1 = *OperationData.Iterators[0];
-
-	InType DefaultValue = PCGMetadataAttribute::GetValueWithBroadcast<InType>(OperationData.SourceAttributes[0], PCGInvalidEntryKey);
-
-	// Fold expressions. It will un-roll the packed argument "InCallbacks", and we can use "InCallbacks" as the current unrolled argument.
-	// In this case, we have multiple callbacks that all take "const InType&" as input, but can return different outputs.
-	// We use this return value to know in which type we need to cast our attribute.
-	uint32 j = 0;
-	([&]
+	InType DefaultValue{};
+	if (OperationData.InputIterators[0].GetWrapper().Get<InType>(PCGInvalidEntryKey, DefaultValue))
 	{
-		check(j < NumOutputs);
-		auto OutValue = InCallbacks(DefaultValue);
-		typedef decltype(OutValue) OutType;
-
-		if (FPCGMetadataAttribute<OutType>* OutputAttribute = static_cast<FPCGMetadataAttribute<OutType>*>(OperationData.OutputAttributes[j++]))
+		// Fold expressions. It will un-roll the packed argument "InCallbacks", and we can use "InCallbacks" as the current unrolled argument.
+		// In this case, we have multiple callbacks that all take "const InType&" as input, but can return different outputs.
+		// We use this return value to know in which type we need to cast our attribute.
+		uint32 j = 0;
+		([&]
 		{
-			OutputAttribute->SetDefaultValue(OutValue);
-		}
-	} (), ...);
+			check(j < NumOutputs);
+			auto OutValue = InCallbacks(DefaultValue);
+			typedef decltype(OutValue) OutType;
+
+			OperationData.OutputIterators[j++].GetWrapper().Set<OutType>(PCGInvalidEntryKey, OutValue);
+		} (), ...);
+	}
+
+	InType Value{};
 
 	for (int32 i = 0; i < OperationData.NumberOfElementsToProcess; ++i)
 	{
-		PCGMetadataEntryKey EntryKey = *Iterator1;
-
-		// If the entry key is invalid, nothing to do
-		if (EntryKey != PCGInvalidEntryKey)
+		if (OperationData.InputIterators[0].GetAndAdvance<InType>(Value))
 		{
-			InType Value = PCGMetadataAttribute::GetValueWithBroadcast<InType>(OperationData.SourceAttributes[0], EntryKey);
-
-			j = 0;
+			uint32 j = 0;
 			([&]
 			{
 				check(j < NumOutputs);
 				auto OutValue = InCallbacks(Value);
 				typedef decltype(OutValue) OutType;
 
-				if (FPCGMetadataAttribute<OutType>* OutputAttribute = static_cast<FPCGMetadataAttribute<OutType>*>(OperationData.OutputAttributes[j++]))
-				{
-					OutputAttribute->SetValue(EntryKey, OutValue);
-				}
+				OperationData.OutputIterators[j++].SetAndAdvance<OutType>(OutValue);
 			} (), ...);
 		}
-
-		++Iterator1;
 	}
 
 	return true;
@@ -209,71 +208,49 @@ inline bool FPCGMetadataElementBase::DoUnaryOp(FOperationData& OperationData, Ca
 template <typename InType1, typename InType2, typename... Callbacks>
 inline bool FPCGMetadataElementBase::DoBinaryOp(FOperationData& OperationData, Callbacks&& ...InCallbacks) const
 {
-	check(OperationData.SourceAttributes[0]);
-	check(OperationData.SourceAttributes[1]);
+	check(OperationData.InputIterators[0].IsValid());
+	check(OperationData.InputIterators[1].IsValid());
 
-	const uint32 NumOutputs = OperationData.OutputAttributes.Num();
+	const uint32 NumOutputs = OperationData.OutputIterators.Num();
 	constexpr uint32 NumCallbacks = (uint32)sizeof...(InCallbacks);
 	check(NumOutputs == NumCallbacks);
 
-	// All values in OperationData.Iterators are UniquePtr. Dereference them here to make the syntax less cumbersome.
-	// Also some iterators can be null, it means they need to use the first iterator (that should never be null).
-	// We just need to make sure to not increment the first iterator multiple times in one loop.
-	check(OperationData.Iterators[0]);
-
-	bool bShouldIncrementIterator2 = OperationData.Iterators.Num() >= 2 && OperationData.Iterators[1];
+	InType1 DefaultValue1{};
+	InType2 DefaultValue2{};
 
-	IPCGMetadataEntryIterator& Iterator1 = *OperationData.Iterators[0];
-	IPCGMetadataEntryIterator& Iterator2 = bShouldIncrementIterator2 ? *OperationData.Iterators[1] : Iterator1;
-
-	InType1 DefaultValue1 = PCGMetadataAttribute::GetValueWithBroadcast<InType1>(OperationData.SourceAttributes[0], PCGInvalidEntryKey);
-	InType2 DefaultValue2 = PCGMetadataAttribute::GetValueWithBroadcast<InType2>(OperationData.SourceAttributes[1], PCGInvalidEntryKey);
-	
-	// Fold expression, cf. Unary op
-	uint32 j = 0;
-	([&]
+	if (OperationData.InputIterators[0].GetWrapper().Get<InType1>(PCGInvalidEntryKey, DefaultValue1) &&
+		OperationData.InputIterators[1].GetWrapper().Get<InType2>(PCGInvalidEntryKey, DefaultValue2))
 	{
-		check(j < NumOutputs);
-		auto OutValue = InCallbacks(DefaultValue1, DefaultValue2);
-		typedef decltype(OutValue) OutType;
-
-		if (FPCGMetadataAttribute<OutType>* OutputAttribute = static_cast<FPCGMetadataAttribute<OutType>*>(OperationData.OutputAttributes[j++]))
+		// Fold expression, cf. Unary op
+		uint32 j = 0;
+		([&]
 		{
-			OutputAttribute->SetDefaultValue(OutValue);
-		}
-	} (), ...);
+			check(j < NumOutputs);
+			auto OutValue = InCallbacks(DefaultValue1, DefaultValue2);
+			typedef decltype(OutValue) OutType;
+
+			OperationData.OutputIterators[j++].GetWrapper().Set<OutType>(PCGInvalidEntryKey, OutValue);
+		} (), ...);
+	}
+
+	InType1 Value1{};
+	InType2 Value2{};
 
 	for (int32 i = 0; i < OperationData.NumberOfElementsToProcess; ++i)
 	{
-		PCGMetadataEntryKey EntryKey1 = *Iterator1;
-
-		// If the entry key is invalid, nothing to do
-		if (EntryKey1 != PCGInvalidEntryKey)
+		if (OperationData.InputIterators[0].GetAndAdvance<InType1>(Value1) &&
+			OperationData.InputIterators[1].GetAndAdvance<InType2>(Value2))
 		{
-			PCGMetadataEntryKey EntryKey2 = *Iterator2;
-
-			InType1 Value1 = PCGMetadataAttribute::GetValueWithBroadcast<InType1>(OperationData.SourceAttributes[0], EntryKey1);
-			InType2 Value2 = PCGMetadataAttribute::GetValueWithBroadcast<InType2>(OperationData.SourceAttributes[1], EntryKey2);
-
-			j = 0;
+			uint32 j = 0;
 			([&]
 			{
 				check(j < NumOutputs);
 				auto OutValue = InCallbacks(Value1, Value2);
 				typedef decltype(OutValue) OutType;
 
-				if (FPCGMetadataAttribute<OutType>* OutputAttribute = static_cast<FPCGMetadataAttribute<OutType>*>(OperationData.OutputAttributes[j++]))
-				{
-					OutputAttribute->SetValue(EntryKey1, OutValue);
-				}
+				OperationData.OutputIterators[j++].SetAndAdvance<OutType>(OutValue);
 			} (), ...);
 		}
-
-		++Iterator1;
-		if (bShouldIncrementIterator2)
-		{
-			++Iterator2;
-		}
 	}
 
 	return true;
@@ -282,82 +259,54 @@ inline bool FPCGMetadataElementBase::DoBinaryOp(FOperationData& OperationData, C
 template <typename InType1, typename InType2, typename InType3, typename... Callbacks>
 inline bool FPCGMetadataElementBase::DoTernaryOp(FOperationData& OperationData, Callbacks&& ...InCallbacks) const
 {
-	check(OperationData.SourceAttributes[0]);
-	check(OperationData.SourceAttributes[1]);
-	check(OperationData.SourceAttributes[2]);
+	check(OperationData.InputIterators[0].IsValid());
+	check(OperationData.InputIterators[1].IsValid());
+	check(OperationData.InputIterators[2].IsValid());
 
-	const uint32 NumOutputs = OperationData.OutputAttributes.Num();
+	const uint32 NumOutputs = OperationData.OutputIterators.Num();
 	constexpr uint32 NumCallbacks = (uint32)sizeof...(InCallbacks);
 	check(NumOutputs == NumCallbacks);
 
-	// All values in OperationData.Iterators are UniquePtr. Dereference them here to make the syntax less cumbersome.
-	// Also some iterators can be null, it means they need to use the first iterator (that should never be null).
-	// We just need to make sure to not increment the first iterator multiple times in one loop.
-	check(OperationData.Iterators[0]);
+	InType1 DefaultValue1{};
+	InType2 DefaultValue2{};
+	InType3 DefaultValue3{};
 
-	bool bShouldIncrementIterator2 = OperationData.Iterators.Num() >= 2 && OperationData.Iterators[1];
-	bool bShouldIncrementIterator3 = OperationData.Iterators.Num() >= 3 && OperationData.Iterators[2];
-
-	IPCGMetadataEntryIterator& Iterator1 = *OperationData.Iterators[0];
-	IPCGMetadataEntryIterator& Iterator2 = bShouldIncrementIterator2 ? *OperationData.Iterators[1] : Iterator1;
-	IPCGMetadataEntryIterator& Iterator3 = bShouldIncrementIterator3 ? *OperationData.Iterators[2] : Iterator1;
-
-	InType1 DefaultValue1 = PCGMetadataAttribute::GetValueWithBroadcast<InType1>(OperationData.SourceAttributes[0], PCGInvalidEntryKey);
-	InType2 DefaultValue2 = PCGMetadataAttribute::GetValueWithBroadcast<InType2>(OperationData.SourceAttributes[1], PCGInvalidEntryKey);
-	InType3 DefaultValue3 = PCGMetadataAttribute::GetValueWithBroadcast<InType3>(OperationData.SourceAttributes[2], PCGInvalidEntryKey);
-	
-	// Fold expression, cf. Unary op
-	uint32 j = 0;
-	([&]
+	if (OperationData.InputIterators[0].GetWrapper().Get<InType1>(PCGInvalidEntryKey, DefaultValue1) &&
+		OperationData.InputIterators[1].GetWrapper().Get<InType2>(PCGInvalidEntryKey, DefaultValue2) &&
+		OperationData.InputIterators[2].GetWrapper().Get<InType3>(PCGInvalidEntryKey, DefaultValue3))
 	{
-		check(j < NumOutputs);
-		auto OutValue = InCallbacks(DefaultValue1, DefaultValue2, DefaultValue3);
-		typedef decltype(OutValue) OutType;
-
-		if (FPCGMetadataAttribute<OutType>* OutputAttribute = static_cast<FPCGMetadataAttribute<OutType>*>(OperationData.OutputAttributes[j++]))
+		// Fold expression, cf. Unary op
+		uint32 j = 0;
+		([&]
 		{
-			OutputAttribute->SetDefaultValue(OutValue);
-		}
-	} (), ...);
+			check(j < NumOutputs);
+			auto OutValue = InCallbacks(DefaultValue1, DefaultValue2, DefaultValue3);
+			typedef decltype(OutValue) OutType;
+
+			OperationData.OutputIterators[j++].GetWrapper().Set<OutType>(PCGInvalidEntryKey, OutValue);
+		} (), ...);
+	}
+
+	InType1 Value1{};
+	InType2 Value2{};
+	InType3 Value3{};
 
 	for (int32 i = 0; i < OperationData.NumberOfElementsToProcess; ++i)
 	{
-		PCGMetadataEntryKey EntryKey1 = *Iterator1;
-
-		// If the entry key is invalid, nothing to do
-		if (EntryKey1 != PCGInvalidEntryKey)
+		if (OperationData.InputIterators[0].GetAndAdvance<InType1>(Value1) &&
+			OperationData.InputIterators[1].GetAndAdvance<InType2>(Value2) &&
+			OperationData.InputIterators[2].GetAndAdvance<InType3>(Value3))
 		{
-			PCGMetadataEntryKey EntryKey2 = *Iterator2;
-			PCGMetadataEntryKey EntryKey3 = *Iterator3;
-
-			InType1 Value1 = PCGMetadataAttribute::GetValueWithBroadcast<InType1>(OperationData.SourceAttributes[0], EntryKey1);
-			InType2 Value2 = PCGMetadataAttribute::GetValueWithBroadcast<InType2>(OperationData.SourceAttributes[1], EntryKey2);
-			InType3 Value3 = PCGMetadataAttribute::GetValueWithBroadcast<InType3>(OperationData.SourceAttributes[2], EntryKey3);
-
-			j = 0;
+			uint32 j = 0;
 			([&]
 			{
 				check(j < NumOutputs);
 				auto OutValue = InCallbacks(Value1, Value2, Value3);
 				typedef decltype(OutValue) OutType;
 
-				if (FPCGMetadataAttribute<OutType>* OutputAttribute = static_cast<FPCGMetadataAttribute<OutType>*>(OperationData.OutputAttributes[j++]))
-				{
-					OutputAttribute->SetValue(EntryKey1, OutValue);
-				}
+				OperationData.OutputIterators[j++].SetAndAdvance<OutType>(OutValue);
 			} (), ...);
 		}
-
-		++Iterator1;
-		if (bShouldIncrementIterator2)
-		{
-			++Iterator2;
-		}
-
-		if (bShouldIncrementIterator3)
-		{
-			++Iterator3;
-		}
 	}
 
 	return true;
@@ -366,93 +315,59 @@ inline bool FPCGMetadataElementBase::DoTernaryOp(FOperationData& OperationData,
 template <typename InType1, typename InType2, typename InType3, typename InType4, typename... Callbacks>
 inline bool FPCGMetadataElementBase::DoQuaternaryOp(FOperationData& OperationData, Callbacks&& ...InCallbacks) const
 {
-	check(OperationData.SourceAttributes[0]);
-	check(OperationData.SourceAttributes[1]);
-	check(OperationData.SourceAttributes[2]);
-	check(OperationData.SourceAttributes[3]);
+	check(OperationData.InputIterators[0].IsValid());
+	check(OperationData.InputIterators[1].IsValid());
+	check(OperationData.InputIterators[2].IsValid());
+	check(OperationData.InputIterators[3].IsValid());
 
-	const uint32 NumOutputs = OperationData.OutputAttributes.Num();
+	const uint32 NumOutputs = OperationData.OutputIterators.Num();
 	constexpr uint32 NumCallbacks = (uint32)sizeof...(InCallbacks);
 	check(NumOutputs == NumCallbacks);
 
-	// All values in OperationData.Iterators are UniquePtr. Dereference them here to make the syntax less cumbersome.
-	// Also some iterators can be null, it means they need to use the first iterator (that should never be null).
-	// We just need to make sure to not increment the first iterator multiple times in one loop.
-	check(OperationData.Iterators[0]);
-
-	bool bShouldIncrementIterator2 = OperationData.Iterators.Num() >= 2 && OperationData.Iterators[1];
-	bool bShouldIncrementIterator3 = OperationData.Iterators.Num() >= 3 && OperationData.Iterators[2];
-	bool bShouldIncrementIterator4 = OperationData.Iterators.Num() >= 4 && OperationData.Iterators[3];
-
-	IPCGMetadataEntryIterator& Iterator1 = *OperationData.Iterators[0];
-	IPCGMetadataEntryIterator& Iterator2 = bShouldIncrementIterator2 ? *OperationData.Iterators[1] : Iterator1;
-	IPCGMetadataEntryIterator& Iterator3 = bShouldIncrementIterator3 ? *OperationData.Iterators[2] : Iterator1;
-	IPCGMetadataEntryIterator& Iterator4 = bShouldIncrementIterator4 ? *OperationData.Iterators[3] : Iterator1;
-
-	InType1 DefaultValue1 = PCGMetadataAttribute::GetValueWithBroadcast<InType1>(OperationData.SourceAttributes[0], PCGInvalidEntryKey);
-	InType2 DefaultValue2 = PCGMetadataAttribute::GetValueWithBroadcast<InType2>(OperationData.SourceAttributes[1], PCGInvalidEntryKey);
-	InType3 DefaultValue3 = PCGMetadataAttribute::GetValueWithBroadcast<InType3>(OperationData.SourceAttributes[2], PCGInvalidEntryKey);
-	InType3 DefaultValue4 = PCGMetadataAttribute::GetValueWithBroadcast<InType3>(OperationData.SourceAttributes[3], PCGInvalidEntryKey);
+	InType1 DefaultValue1{};
+	InType2 DefaultValue2{};
+	InType3 DefaultValue3{};
+	InType4 DefaultValue4{};
 	
-	// Fold expression, cf. Unary op
-	uint32 j = 0;
-	([&]
+	if (OperationData.InputIterators[0].GetWrapper().Get<InType1>(PCGInvalidEntryKey, DefaultValue1) &&
+		OperationData.InputIterators[1].GetWrapper().Get<InType2>(PCGInvalidEntryKey, DefaultValue2) &&
+		OperationData.InputIterators[2].GetWrapper().Get<InType3>(PCGInvalidEntryKey, DefaultValue3) &&
+		OperationData.InputIterators[3].GetWrapper().Get<InType4>(PCGInvalidEntryKey, DefaultValue4))
 	{
-		check(j < NumOutputs);
-		auto OutValue = InCallbacks(DefaultValue1, DefaultValue2, DefaultValue3, DefaultValue4);
-		typedef decltype(OutValue) OutType;
-
-		if (FPCGMetadataAttribute<OutType>* OutputAttribute = static_cast<FPCGMetadataAttribute<OutType>*>(OperationData.OutputAttributes[j++]))
+		// Fold expression, cf. Unary op
+		uint32 j = 0;
+		([&]
 		{
-			OutputAttribute->SetDefaultValue(OutValue);
-		}
-	} (), ...);
+			check(j < NumOutputs);
+			auto OutValue = InCallbacks(DefaultValue1, DefaultValue2, DefaultValue3, DefaultValue4);
+			typedef decltype(OutValue) OutType;
+
+			OperationData.OutputIterators[j++].GetWrapper().Set<OutType>(PCGInvalidEntryKey, OutValue);
+		} (), ...);
+	}
+
+	InType1 Value1{};
+	InType2 Value2{};
+	InType3 Value3{};
+	InType4 Value4{};
 
 	for (int32 i = 0; i < OperationData.NumberOfElementsToProcess; ++i)
 	{
-		PCGMetadataEntryKey EntryKey1 = *Iterator1;
-
-		// If the entry key is invalid, nothing to do
-		if (EntryKey1 != PCGInvalidEntryKey)
+		if (OperationData.InputIterators[0].GetAndAdvance<InType1>(Value1) &&
+			OperationData.InputIterators[1].GetAndAdvance<InType2>(Value2) &&
+			OperationData.InputIterators[2].GetAndAdvance<InType3>(Value3) &&
+			OperationData.InputIterators[3].GetAndAdvance<InType4>(Value4))
 		{
-			PCGMetadataEntryKey EntryKey2 = *Iterator2;
-			PCGMetadataEntryKey EntryKey3 = *Iterator3;
-			PCGMetadataEntryKey EntryKey4 = *Iterator4;
-
-			InType1 Value1 = PCGMetadataAttribute::GetValueWithBroadcast<InType1>(OperationData.SourceAttributes[0], EntryKey1);
-			InType2 Value2 = PCGMetadataAttribute::GetValueWithBroadcast<InType2>(OperationData.SourceAttributes[1], EntryKey2);
-			InType3 Value3 = PCGMetadataAttribute::GetValueWithBroadcast<InType3>(OperationData.SourceAttributes[2], EntryKey3);
-			InType4 Value4 = PCGMetadataAttribute::GetValueWithBroadcast<InType4>(OperationData.SourceAttributes[3], EntryKey4);
-
-			j = 0;
+			uint32 j = 0;
 			([&]
 			{
 				check(j < NumOutputs);
 				auto OutValue = InCallbacks(Value1, Value2, Value3, Value4);
 				typedef decltype(OutValue) OutType;
 
-				if (FPCGMetadataAttribute<OutType>* OutputAttribute = static_cast<FPCGMetadataAttribute<OutType>*>(OperationData.OutputAttributes[j++]))
-				{
-					OutputAttribute->SetValue(EntryKey1, OutValue);
-				}
+				OperationData.OutputIterators[j++].SetAndAdvance<OutType>(OutValue);
 			} (), ...);
 		}
-
-		++Iterator1;
-		if (bShouldIncrementIterator2)
-		{
-			++Iterator2;
-		}
-
-		if (bShouldIncrementIterator3)
-		{
-			++Iterator3;
-		}
-
-		if (bShouldIncrementIterator4)
-		{
-			++Iterator4;
-		}
 	}
 
 	return true;
diff --git a/Engine/Plugins/Experimental/PCG/Source/PCG/Public/Elements/Metadata/PCGMetadataRotatorOpElement.h b/Engine/Plugins/Experimental/PCG/Source/PCG/Public/Elements/Metadata/PCGMetadataRotatorOpElement.h
--- a/Engine/Plugins/Experimental/PCG/Source/PCG/Public/Elements/Metadata/PCGMetadataRotatorOpElement.h
+++ b/Engine/Plugins/Experimental/PCG/Source/PCG/Public/Elements/Metadata/PCGMetadataRotatorOpElement.h
@@ -33,15 +33,20 @@ class PCG_API UPCGMetadataRotatorSettings : public UPCGMetadataSettingsBase
 	GENERATED_BODY()
 
 public:
+	// ~Begin UObject interface
+	virtual void PostLoad() override;
+	// ~End UObject interface
+
 	//~Begin UPCGSettings interface
 #if WITH_EDITOR
 	virtual FName GetDefaultNodeName() const override;
 #endif
+
 	virtual FName AdditionalTaskName() const override;
 	//~End UPCGSettings interface
 
 	//~Begin UPCGMetadataSettingsBase interface
-	virtual FName GetInputAttributeNameWithOverride(uint32 Index, UPCGParamData* Params) const override;
+	FPCGAttributePropertySelector GetInputSource(uint32 Index) const override;
 
 	virtual FName GetInputPinLabel(uint32 Index) const override;
 	virtual uint32 GetInputPinNum() const override;
@@ -59,13 +64,24 @@ public:
 	EPCGMedadataRotatorOperation Operation = EPCGMedadataRotatorOperation::Combine;
 
 	UPROPERTY(BlueprintReadWrite, EditAnywhere, Category = Input)
-	FName Input1AttributeName = NAME_None;
+	FPCGAttributePropertySelector InputSource1;
 
 	UPROPERTY(BlueprintReadWrite, EditAnywhere, Category = Input, meta = (EditCondition = "Operation != EPCGMedadataTrigOperation::Invert", EditConditionHides))
-	FName Input2AttributeName = NAME_None;
+	FPCGAttributePropertySelector InputSource2;
 
 	UPROPERTY(BlueprintReadWrite, EditAnywhere, Category = Input, meta = (EditCondition = "Operation == EPCGMedadataTrigOperation::Lerp", EditConditionHides))
-	FName Input3AttributeName = NAME_None;
+	FPCGAttributePropertySelector InputSource3;
+
+#if WITH_EDITORONLY_DATA
+	UPROPERTY()
+	FName Input1AttributeName_DEPRECATED = NAME_None;
+
+	UPROPERTY()
+	FName Input2AttributeName_DEPRECATED = NAME_None;
+
+	UPROPERTY()
+	FName Input3AttributeName_DEPRECATED = NAME_None;
+#endif
 };
 
 class FPCGMetadataRotatorElement : public FPCGMetadataElementBase
diff --git a/Engine/Plugins/Experimental/PCG/Source/PCG/Public/Elements/Metadata/PCGMetadataTransformOpElement.h b/Engine/Plugins/Experimental/PCG/Source/PCG/Public/Elements/Metadata/PCGMetadataTransformOpElement.h
--- a/Engine/Plugins/Experimental/PCG/Source/PCG/Public/Elements/Metadata/PCGMetadataTransformOpElement.h
+++ b/Engine/Plugins/Experimental/PCG/Source/PCG/Public/Elements/Metadata/PCGMetadataTransformOpElement.h
@@ -36,6 +36,10 @@ class PCG_API UPCGMetadataTransformSettings : public UPCGMetadataSettingsBase
 	GENERATED_BODY()
 
 public:
+	// ~Begin UObject interface
+	virtual void PostLoad() override;
+	// ~End UObject interface
+
 	//~Begin UPCGSettings interface
 #if WITH_EDITOR
 	virtual FName GetDefaultNodeName() const override;
@@ -44,7 +48,7 @@ public:
 	//~End UPCGSettings interface
 
 	//~Begin UPCGMetadataSettingsBase interface
-	virtual FName GetInputAttributeNameWithOverride(uint32 Index, UPCGParamData* Params) const override;
+	FPCGAttributePropertySelector GetInputSource(uint32 Index) const override;
 
 	virtual FName GetInputPinLabel(uint32 Index) const override;
 	virtual uint32 GetInputPinNum() const override;
@@ -65,13 +69,24 @@ public:
 	EPCGTransformLerpMode TransformLerpMode = EPCGTransformLerpMode::QuatInterp;
 
 	UPROPERTY(BlueprintReadWrite, EditAnywhere, Category = Input)
-	FName Input1AttributeName = NAME_None;
+	FPCGAttributePropertySelector InputSource1;
 
 	UPROPERTY(BlueprintReadWrite, EditAnywhere, Category = Input, meta = (EditCondition = "Operation != EPCGMedadataTransformOperation::Invert", EditConditionHides))
-	FName Input2AttributeName = NAME_None;
+	FPCGAttributePropertySelector InputSource2;
 
 	UPROPERTY(BlueprintReadWrite, EditAnywhere, Category = Input, meta = (EditCondition = "Operation == EPCGMedadataTransformOperation::Lerp", EditConditionHides))
-	FName Input3AttributeName = NAME_None;
+	FPCGAttributePropertySelector InputSource3;
+
+#if WITH_EDITORONLY_DATA
+	UPROPERTY()
+	FName Input1AttributeName_DEPRECATED = NAME_None;
+
+	UPROPERTY()
+	FName Input2AttributeName_DEPRECATED = NAME_None;
+
+	UPROPERTY()
+	FName Input3AttributeName_DEPRECATED = NAME_None;
+#endif
 };
 
 class FPCGMetadataTransformElement : public FPCGMetadataElementBase
diff --git a/Engine/Plugins/Experimental/PCG/Source/PCG/Public/Elements/Metadata/PCGMetadataTrigOpElement.h b/Engine/Plugins/Experimental/PCG/Source/PCG/Public/Elements/Metadata/PCGMetadataTrigOpElement.h
--- a/Engine/Plugins/Experimental/PCG/Source/PCG/Public/Elements/Metadata/PCGMetadataTrigOpElement.h
+++ b/Engine/Plugins/Experimental/PCG/Source/PCG/Public/Elements/Metadata/PCGMetadataTrigOpElement.h
@@ -28,6 +28,10 @@ class PCG_API UPCGMetadataTrigSettings : public UPCGMetadataSettingsBase
 	GENERATED_BODY()
 
 public:
+	// ~Begin UObject interface
+	virtual void PostLoad() override;
+	// ~End UObject interface
+
 	//~Begin UPCGSettings interface
 #if WITH_EDITOR
 	virtual FName GetDefaultNodeName() const override;
@@ -36,7 +40,7 @@ public:
 	//~End UPCGSettings interface
 
 	//~Begin UPCGMetadataSettingsBase interface
-	virtual FName GetInputAttributeNameWithOverride(uint32 Index, UPCGParamData* Params) const override;
+	FPCGAttributePropertySelector GetInputSource(uint32 Index) const override;
 
 	virtual FName GetInputPinLabel(uint32 Index) const override;
 	virtual uint32 GetInputPinNum() const override;
@@ -55,10 +59,18 @@ public:
 	EPCGMedadataTrigOperation Operation = EPCGMedadataTrigOperation::Acos;
 
 	UPROPERTY(BlueprintReadWrite, EditAnywhere, Category = Input)
-	FName Input1AttributeName = NAME_None;
+	FPCGAttributePropertySelector InputSource1;
+
+	UPROPERTY(BlueprintReadWrite, EditAnywhere, Category = Input, meta = (EditCondition = "Operation == EPCGMedadataTrigOperation::Atan2", EditConditionHides))
+	FPCGAttributePropertySelector InputSource2;
 
-	UPROPERTY(BlueprintReadWrite, EditAnywhere, Category = Input, meta = (EditCondition="Operation == EPCGMedadataTrigOperation::Atan2", EditConditionHides))
-	FName Input2AttributeName = NAME_None;
+#if WITH_EDITORONLY_DATA
+	UPROPERTY()
+	FName Input1AttributeName_DEPRECATED = NAME_None;
+
+	UPROPERTY()
+	FName Input2AttributeName_DEPRECATED = NAME_None;
+#endif
 };
 
 class FPCGMetadataTrigElement : public FPCGMetadataElementBase
diff --git a/Engine/Plugins/Experimental/PCG/Source/PCG/Public/Elements/Metadata/PCGMetadataVectorOpElement.h b/Engine/Plugins/Experimental/PCG/Source/PCG/Public/Elements/Metadata/PCGMetadataVectorOpElement.h
--- a/Engine/Plugins/Experimental/PCG/Source/PCG/Public/Elements/Metadata/PCGMetadataVectorOpElement.h
+++ b/Engine/Plugins/Experimental/PCG/Source/PCG/Public/Elements/Metadata/PCGMetadataVectorOpElement.h
@@ -32,6 +32,10 @@ class PCG_API UPCGMetadataVectorSettings : public UPCGMetadataSettingsBase
 	GENERATED_BODY()
 
 public:
+	// ~Begin UObject interface
+	virtual void PostLoad() override;
+	// ~End UObject interface
+
 	//~Begin UPCGSettings interface
 #if WITH_EDITOR
 	virtual FName GetDefaultNodeName() const override;
@@ -40,7 +44,7 @@ public:
 	//~End UPCGSettings interface
 
 	//~Begin UPCGMetadataSettingsBase interface
-	virtual FName GetInputAttributeNameWithOverride(uint32 Index, UPCGParamData* Params) const override;
+	FPCGAttributePropertySelector GetInputSource(uint32 Index) const override;
 
 	virtual FName GetInputPinLabel(uint32 Index) const override;
 	virtual uint32 GetInputPinNum() const override;
@@ -59,13 +63,24 @@ public:
 	EPCGMedadataVectorOperation Operation = EPCGMedadataVectorOperation::Cross;
 
 	UPROPERTY(BlueprintReadWrite, EditAnywhere, Category = Input)
-	FName Input1AttributeName = NAME_None;
+	FPCGAttributePropertySelector InputSource1;
 
 	UPROPERTY(BlueprintReadWrite, EditAnywhere, Category = Input, meta = (EditCondition = "Operation != EPCGMedadataVectorOperation::Normalize && Operation != EPCGMedadataVectorOperation::Length", EditConditionHides))
-	FName Input2AttributeName = NAME_None;
+	FPCGAttributePropertySelector InputSource2;
 
 	UPROPERTY(BlueprintReadWrite, EditAnywhere, Category = Input, meta = (EditCondition = "Operation == EPCGMedadataVectorOperation::RotateAroundAxis", EditConditionHides))
-	FName Input3AttributeName = NAME_None;
+	FPCGAttributePropertySelector InputSource3;
+
+#if WITH_EDITORONLY_DATA
+	UPROPERTY()
+	FName Input1AttributeName_DEPRECATED = NAME_None;
+
+	UPROPERTY()
+	FName Input2AttributeName_DEPRECATED = NAME_None;
+
+	UPROPERTY()
+	FName Input3AttributeName_DEPRECATED = NAME_None;
+#endif
 };
 
 class FPCGMetadataVectorElement : public FPCGMetadataElementBase
diff --git a/Engine/Plugins/Experimental/PCG/Source/PCG/Public/Elements/PCGAttributeGetFromPointIndexElement.h b/Engine/Plugins/Experimental/PCG/Source/PCG/Public/Elements/PCGAttributeGetFromPointIndexElement.h
--- a/Engine/Plugins/Experimental/PCG/Source/PCG/Public/Elements/PCGAttributeGetFromPointIndexElement.h
+++ b/Engine/Plugins/Experimental/PCG/Source/PCG/Public/Elements/PCGAttributeGetFromPointIndexElement.h
@@ -49,9 +49,6 @@ protected:
 
 class FPCGAttributeGetFromPointIndexElement : public FSimplePCGElement
 {
-public:
-	virtual bool IsCacheable(const UPCGSettings* InSettings) const override { return false; }
-
 protected:
 	virtual bool ExecuteInternal(FPCGContext* Context) const override;
 };
diff --git a/Engine/Plugins/Experimental/PCG/Source/PCG/Public/Helpers/PCGSettingsHelpers.h b/Engine/Plugins/Experimental/PCG/Source/PCG/Public/Helpers/PCGSettingsHelpers.h
--- a/Engine/Plugins/Experimental/PCG/Source/PCG/Public/Helpers/PCGSettingsHelpers.h
+++ b/Engine/Plugins/Experimental/PCG/Source/PCG/Public/Helpers/PCGSettingsHelpers.h
@@ -122,97 +122,6 @@ namespace PCGSettingsHelpers
 	/** Sets data from the params to a given property, matched on a name basis */
 	void SetValue(UPCGParamData* Params, UObject* Object, FProperty* Property);
 
-	/**
-	* Validate that the InProperty is a PCGData supported type and call InFunc with the value of this property (with the right type).
-	* The function returns the result of InFunc (or default value of the return type of InFunc if it failed).
-	*/
-	template <typename ObjectType, typename Func>
-	inline decltype(auto) GetPropertyValueWithCallback(const ObjectType* InObject, const FProperty* InProperty, Func InFunc)
-	{
-		// Double property is supported by PCG, we use a dummy double to deduce the return type of InFunc.
-		using ReturnType = decltype(InFunc(0.0));
-
-		if (!InObject || !InProperty)
-		{
-			return ReturnType();
-		}
-
-		const void* PropertyAddressData = InProperty->ContainerPtrToValuePtr<void>(InObject);
-
-		if (const FNumericProperty* NumericProperty = CastField<FNumericProperty>(InProperty))
-		{
-			if (NumericProperty->IsFloatingPoint())
-			{
-				return InFunc(NumericProperty->GetFloatingPointPropertyValue(PropertyAddressData));
-			}
-			else if (NumericProperty->IsInteger())
-			{
-				return InFunc(NumericProperty->GetSignedIntPropertyValue(PropertyAddressData));
-			}
-		}
-		else if (const FBoolProperty* BoolProperty = CastField<FBoolProperty>(InProperty))
-		{
-			return InFunc(BoolProperty->GetPropertyValue(PropertyAddressData));
-		}
-		else if (const FStrProperty* StringProperty = CastField<FStrProperty>(InProperty))
-		{
-			return InFunc(StringProperty->GetPropertyValue(PropertyAddressData));
-		}
-		else if (const FNameProperty* NameProperty = CastField<FNameProperty>(InProperty))
-		{
-			return InFunc(NameProperty->GetPropertyValue(PropertyAddressData));
-		}
-		else if (const FEnumProperty* EnumProperty = CastField<FEnumProperty>(InProperty))
-		{
-			return InFunc(EnumProperty->GetUnderlyingProperty()->GetSignedIntPropertyValue(PropertyAddressData));
-		}
-		else if (const FStructProperty* StructProperty = CastField<FStructProperty>(InProperty))
-		{
-			if (StructProperty->Struct == TBaseStructure<FVector>::Get())
-			{
-				return InFunc(*reinterpret_cast<const FVector*>(PropertyAddressData));
-			}
-			else if (StructProperty->Struct == TBaseStructure<FVector4>::Get())
-			{
-				return InFunc(*reinterpret_cast<const FVector4*>(PropertyAddressData));
-			}
-			else if (StructProperty->Struct == TBaseStructure<FQuat>::Get())
-			{
-				return InFunc(*reinterpret_cast<const FQuat*>(PropertyAddressData));
-			}
-			else if (StructProperty->Struct == TBaseStructure<FTransform>::Get())
-			{
-				return InFunc(*reinterpret_cast<const FTransform*>(PropertyAddressData));
-			}
-			else if (StructProperty->Struct == TBaseStructure<FRotator>::Get())
-			{
-				return InFunc(*reinterpret_cast<const FRotator*>(PropertyAddressData));
-			}
-			else if (StructProperty->Struct == TBaseStructure<FSoftObjectPath>::Get())
-			{
-				// Soft object path are transformed to strings
-				return InFunc(reinterpret_cast<const FSoftObjectPath*>(PropertyAddressData)->ToString());
-			}
-			else if (StructProperty->Struct == TBaseStructure<FSoftClassPath>::Get())
-			{
-				// Soft class path are transformed to strings
-				return InFunc(reinterpret_cast<const FSoftClassPath*>(PropertyAddressData)->ToString());
-			}
-			//else if (StructProperty->Struct == TBaseStructure<FColor>::Get())
-			//else if (StructProperty->Struct == TBaseStructure<FLinearColor>::Get())
-		}
-		else if (const FObjectPropertyBase* ObjectProperty = CastField<FObjectPropertyBase>(InProperty))
-		{
-			if (const UObject* Object = ObjectProperty->GetObjectPropertyValue(PropertyAddressData))
-			{
-				// Object are transformed into their soft path name (as a string attribute)
-				return InFunc(Object->GetPathName());
-			}
-		}
-
-		return ReturnType();
-	}
-
 	int ComputeSeedWithOverride(const UPCGSettings* InSettings, const UPCGComponent* InComponent, UPCGParamData* InParams);
 
 	FORCEINLINE int ComputeSeedWithOverride(const UPCGSettings* InSettings, TWeakObjectPtr<UPCGComponent> InComponent, UPCGParamData* InParams)
diff --git a/Engine/Plugins/Experimental/PCG/Source/PCG/Public/Metadata/PCGMetadataAttributeTpl.h b/Engine/Plugins/Experimental/PCG/Source/PCG/Public/Metadata/PCGMetadataAttributeTpl.h
--- a/Engine/Plugins/Experimental/PCG/Source/PCG/Public/Metadata/PCGMetadataAttributeTpl.h
+++ b/Engine/Plugins/Experimental/PCG/Source/PCG/Public/Metadata/PCGMetadataAttributeTpl.h
@@ -583,70 +583,4 @@ namespace PCGMetadataAttribute
 		}
 		}
 	}
-
-	template <typename OutType>
-	inline OutType GetValueWithBroadcast(const FPCGMetadataAttributeBase* InAttribute, PCGMetadataEntryKey InKey)
-	{
-		auto Func = [InAttribute, InKey](auto DummyInType) -> OutType
-		{
-			using InType = decltype(DummyInType);
-			InType Value = static_cast<const FPCGMetadataAttribute<InType>*>(InAttribute)->GetValueFromItemKey(InKey);
-
-			if constexpr (std::is_same_v<OutType, InType>)
-			{
-				return Value;
-			}
-			else
-			{
-				if constexpr (!PCG::Private::IsBroadcastable(PCG::Private::MetadataTypes<InType>::Id, PCG::Private::MetadataTypes<OutType>::Id))
-				{
-					return OutType{};
-				}
-				else
-				{
-					if constexpr (std::is_same_v<OutType, FVector4>)
-					{
-						if constexpr (std::is_same_v<InType, FVector> || std::is_same_v<InType, FVector2D>)
-						{
-							// TODO: Should it be 0? 1? Something else? Depending on operation?
-							// For now it is too ambiguous, so don't support it
-							return FVector4();
-						}
-						else
-						{
-							return FVector4(Value, Value, Value, Value);
-						}
-					}
-					else
-					{
-						// Seems like the && condition is not evaluated correctly on Linux, so we cut the condition in two `if constexpr`.
-						if constexpr (std::is_same_v<OutType, FVector>)
-						{
-							if constexpr (std::is_same_v<InType, FVector2D>)
-							{
-								return FVector(Value, 0.0);
-							}
-							else
-							{
-								return OutType(Value);
-							}
-						}
-						else
-						{
-							return OutType(Value);
-						}
-					}
-				}
-			}
-		};
-
-		if (PCG::Private::MetadataTypes<OutType>::Id == InAttribute->GetTypeId())
-		{
-			return Func(OutType{});
-		}
-		else
-		{
-			return CallbackWithRightType(InAttribute->GetTypeId(), Func);
-		}
-	}
 }
diff --git a/Engine/Plugins/Experimental/PCG/Source/PCG/Public/Metadata/PCGMetadataAttributeTraits.h b/Engine/Plugins/Experimental/PCG/Source/PCG/Public/Metadata/PCGMetadataAttributeTraits.h
--- a/Engine/Plugins/Experimental/PCG/Source/PCG/Public/Metadata/PCGMetadataAttributeTraits.h
+++ b/Engine/Plugins/Experimental/PCG/Source/PCG/Public/Metadata/PCGMetadataAttributeTraits.h
@@ -110,6 +110,9 @@ namespace PCG
 				PCGMetadataBroadcastable(Integer64, Vector4);
 
 				PCGMetadataBroadcastable(Vector2, Vector);
+				PCGMetadataBroadcastable(Vector2, Vector4);
+
+				PCGMetadataBroadcastable(Vector, Vector4);
 
 				PCGMetadataBroadcastable(Quaternion, Rotator);
 				PCGMetadataBroadcastable(Rotator, Quaternion);
@@ -135,6 +138,77 @@ namespace PCG
 			return BroadcastableTypes[FirstType][SecondType];
 		}
 
+		template <typename FirstType, typename SecondType>
+		constexpr inline bool IsBroadcastable()
+		{
+			return IsBroadcastable(MetadataTypes<FirstType>::Id, MetadataTypes<SecondType>::Id);
+		}
+
+		/**
+		* Generic function to broadcast an InType to an OutType.
+		* Supports only PCG types
+		* @param InValue - The Value to convert
+		* @param OutValue - The converted value
+		* @returns true if the conversion worked, false otherwise.
+		*/
+		template <typename InType, typename OutType>
+		inline bool GetValueWithBroadcast(const InType& InValue, OutType& OutValue)
+		{
+			if constexpr (std::is_same_v<OutType, InType>)
+			{
+				OutValue = InValue;
+				return true;
+			}
+			else
+			{
+				if constexpr (!IsBroadcastable<InType, OutType>())
+				{
+					return false;
+				}
+				else
+				{
+					if constexpr (std::is_same_v<OutType, FVector4>)
+					{
+						// TODO: Should W be 0? 1? Something else? Depending on operation?
+						// For now, we set Z to 0 (for vec 2) and we set W to 1.
+						if constexpr (std::is_same_v<InType, FVector>)
+						{
+							OutValue = FVector4(InValue, 1.0);
+						}
+						else if constexpr (std::is_same_v<InType, FVector2D>)
+						{
+							OutValue = FVector4(InValue.X, InValue.Y, 0.0, 1.0);
+						}
+						else
+						{
+							OutValue = FVector4(InValue, InValue, InValue, InValue);
+						}
+					}
+					else
+					{
+						// Seems like the && condition is not evaluated correctly on Linux, so we cut the condition in two `if constexpr`.
+						if constexpr (std::is_same_v<OutType, FVector>)
+						{
+							if constexpr (std::is_same_v<InType, FVector2D>)
+							{
+								OutValue = FVector(InValue, 0.0);
+							}
+							else
+							{
+								OutValue = OutType(InValue);
+							}
+						}
+						else
+						{
+							OutValue = OutType(InValue);
+						}
+					}
+
+					return true;
+				}
+			}
+		}
+
 		template<typename T>
 		struct DefaultOperationTraits
 		{
diff --git a/Engine/Plugins/Experimental/PCG/Source/PCG/Public/Metadata/PCGMetadataAttributeWrapper.h b/Engine/Plugins/Experimental/PCG/Source/PCG/Public/Metadata/PCGMetadataAttributeWrapper.h
new file mode 100644
--- /dev/null
+++ b/Engine/Plugins/Experimental/PCG/Source/PCG/Public/Metadata/PCGMetadataAttributeWrapper.h
@@ -0,0 +1,876 @@
+// Copyright Epic Games, Inc. All Rights Reserved.
+
+#pragma once
+
+#include "PCGPoint.h"
+
+#include "PCGMetadata.h"
+#include "PCGMetadataCommon.h"
+#include "PCGMetadataAttribute.h"
+#include "PCGMetadataAttributeTpl.h"
+#include "PCGMetadataAttributeTraits.h"
+
+#include "PCGMetadataAttributeWrapper.generated.h"
+
+class FPCGPropertyAttributeWrapper;
+class FPCGPropertyAttributeIterator;
+class UPCGData;
+struct FPCGAttributePropertySelector;
+
+namespace PCGMetadataAttributeWrapper
+{
+	/**
+	* Attempt to get a T value from an attribute. Will only work if the attribute is of type T, or we can broadcast
+	* the attribute type to T
+	* @param Attribute - The attribute to get the value from
+	* @param Key - Entry key to get the value from
+	* @param OutValue - Result of the get
+	* @returns true if it succeeded, false otherwise.
+	*/
+	template <typename T>
+	bool GetValueFromAttribute(const FPCGMetadataAttributeBase* Attribute, PCGMetadataEntryKey Key, T& OutValue);
+
+	/**
+	* Attempt to set a T value to an attribute. Will only work if the attribute is of type T, or we can broadcast
+	* T to the attribute type
+	* @param Attribute - The attribute to set the value to
+	* @param Key - Entry key to set the value to. If it is an invalid key, it will override the default value.
+	* @param InValue - Value to set
+	* @returns true if it succeeded, false otherwise.
+	*/
+	template <typename T>
+	bool SetValueToAttribute(FPCGMetadataAttributeBase* Attribute, PCGMetadataEntryKey Key, const T& InValue);
+
+	/**
+	* Get a property value and pass it as a parameter to a callback function.
+	* @param InObject - The object to read from
+	* @param InProperty - The property to look for
+	* @param InFunc - Callback function that can return anything, and should have a single templated argument, where the property will be.
+	* @returns Forward the result of the callback.
+	*/
+	template <typename ObjectType, typename Func>
+	decltype(auto) GetPropertyValueWithCallback(const ObjectType* InObject, const FProperty* InProperty, Func InFunc);
+
+	/**
+	* Set a property value given by a callback function.
+	* @param InObject - The object to write to
+	* @param InProperty - The property to look for
+	* @param InFunc - Callback function that take a reference to a templated type. It will set the property with this value. returns true if we should set, false otherwise.
+	* @returns Forward the result of the callback
+	*/
+	template <typename ObjectType, typename Func>
+	bool SetPropertyValueFromCallback(ObjectType* InObject, const FProperty* InProperty, Func InFunc);
+
+	/**
+	* Attempt to get a T value from a property. Will only work if the property is of type T, or we can broadcast
+	* the property type to T
+	* @param InObject - The object to read from
+	* @param InProperty - The property to look for
+	* @param OutValue - The value of the property
+	* @returns true if it succeeded, false otherwise.
+	*/
+	template <typename T, typename ObjectType>
+	bool GetValueFromProperty(const ObjectType* InObject, const FProperty* InProperty, T& OutValue);
+
+	/**
+	* Attempt to set a T value to a property. Will only work if the property is of type T, or we can broadcast
+	* T to the property type
+	* @param InObject - The object to write to
+	* @param InProperty - The property to look for
+	* @param InValue - Value to set
+	* @returns true if it succeeded, false otherwise.
+	*/
+	template <typename T, typename ObjectType>
+	bool SetValueToProperty(ObjectType* InObject, const FProperty* InProperty, const T& InValue);
+
+	/**
+	* Conversion between property type and PCG type.
+	* @param InProperty - The property to look for
+	* @returns PCG type if the property is supported, Unknown otherwise.
+	*/
+	PCG_API EPCGMetadataTypes GetMetadataTypeFromProperty(const FProperty* InProperty);
+
+	/**
+	* Returns if the name for this data matches a property, and optional type
+	*/
+	PCG_API bool IsPropertyWithType(const UPCGData* InData, FName InName, int16* OutType = nullptr);
+
+	/**
+	* Create a wrapper around some data with a given name.
+	* cf FPCGPropertyAttributeWrapper for more information
+	* If it is a point data and the name matches a property, it will wrap a property
+	* Otherwise it will wrap an attribute.
+	* Note: It is not threadsafe for properties and won't work if the selection is invalid (example: Property for a param data)
+	* @param InData - The data to extract the attribute/metadata or the point property
+	* @param InSelection - Select a property or an attribute name. 
+	* @return The resulting wrapper. Can check "IsValid" to know if the wrapper is valid.
+	*/
+	PCG_API FPCGPropertyAttributeWrapper CreateWrapper(UPCGData* InData, const FPCGAttributePropertySelector& InSelector);
+	PCG_API FPCGPropertyAttributeWrapper CreateWrapper(const UPCGData* InData, const FPCGAttributePropertySelector& InSelector);
+
+	/**
+	* Create an iterator wrapper around some data with a given name.
+	* cf. FPCGPropertyAttributeIterator for more information.
+	* If it is a point data it will iterate over points, otherwise it will iterate over the metadata entry to value key map.
+	* cf. CreateWrapper for the wrapping of the attribute/property.
+	* Note: It is not threadsafe for properties and won't work if the selection is invalid (example: Property for a param data)
+	* @param InData - The data to extract the list of points or the entry key map.
+	* @param InSelection - Select a property or an attribute name. 
+	* @return The resulting iterator. Can check "IsValid" to know if the iterator is valid.
+	*/
+	PCG_API FPCGPropertyAttributeIterator CreateIteratorWrapper(UPCGData* InData, const FPCGAttributePropertySelector& InSelector);
+	PCG_API FPCGPropertyAttributeIterator CreateIteratorWrapper(const UPCGData* InData, const FPCGAttributePropertySelector& InSelector);
+}
+
+UENUM()
+enum class EPCGAttributePropertySelection
+{
+	Attribute,
+	PointProperty,
+};
+
+USTRUCT(BlueprintType)
+struct PCG_API FPCGAttributePropertySelector
+{
+	GENERATED_BODY()
+
+	UPROPERTY(EditAnywhere, Category = "Settings")
+	EPCGAttributePropertySelection Selection = EPCGAttributePropertySelection::Attribute;
+
+	UPROPERTY(EditAnywhere, Category = "Settings", meta = (EditCondition = "Selection == EPCGAttributePropertySelection::Attribute", EditConditionHides))
+	FName AttributeName = NAME_None;
+
+	UPROPERTY(EditAnywhere, Category = "Settings", meta = (EditCondition = "Selection == EPCGAttributePropertySelection::PointProperty", EditConditionHides))
+	EPCGPointProperties PointProperty = EPCGPointProperties::Position;
+
+	FName GetName() const;
+};
+
+/**
+* Wrapping class around an attribute, property or some custom getter/setters. Allow to design generic algorithms that doesn't need to know
+* how to data under is stored. (Similar to a Visitor pattern).
+* Get/Set are not threadsafe for properties.
+* Set is only available if the wrapper is not read only.
+*/
+class PCG_API FPCGPropertyAttributeWrapper
+{
+public:
+	friend FPCGPropertyAttributeIterator;
+
+	FPCGPropertyAttributeWrapper() = default;
+
+	FPCGPropertyAttributeWrapper(FPCGMetadataAttributeBase* InAttribute, UPCGMetadata* InMetadata)
+		: Attribute(InAttribute)
+		, Metadata(InMetadata)
+		, Property(nullptr)
+		, CustomPointGetterSetter()
+		, bReadOnly(false)
+	{
+		Type = InAttribute ? InAttribute->GetTypeId() : (int16)EPCGMetadataTypes::Unknown;
+		Name = InAttribute ? InAttribute->Name : NAME_None;
+	}
+
+	FPCGPropertyAttributeWrapper(const FPCGMetadataAttributeBase* InAttribute, const UPCGMetadata* InMetadata)
+		: Attribute(const_cast<FPCGMetadataAttributeBase*>(InAttribute))
+		, Metadata(const_cast<UPCGMetadata*>(InMetadata))
+		, Property(nullptr)
+		, CustomPointGetterSetter()
+		, bReadOnly(true)
+	{
+		Type = InAttribute ? InAttribute->GetTypeId() : (int16)EPCGMetadataTypes::Unknown;
+		Name = InAttribute ? InAttribute->Name : NAME_None;
+	}
+
+	FPCGPropertyAttributeWrapper(const FProperty* InProperty, bool bInReadOnly)
+		: Attribute(nullptr)
+		, Metadata(nullptr)
+		, Property(InProperty)
+		, CustomPointGetterSetter()
+		, bReadOnly(bInReadOnly)
+	{
+		Type = (int16)PCGMetadataAttributeWrapper::GetMetadataTypeFromProperty(InProperty);
+		Name = InProperty ? InProperty->GetFName() : NAME_None;
+	}
+
+	FPCGPropertyAttributeWrapper(FPCGPoint::PointCustomPropertyGetterSetter&& InCustomGetterSetter, bool bInReadOnly)
+		: Attribute(nullptr)
+		, Metadata(nullptr)
+		, Property(nullptr)
+		, CustomPointGetterSetter(std::forward<FPCGPoint::PointCustomPropertyGetterSetter>(InCustomGetterSetter))
+		, bReadOnly(bInReadOnly)
+	{
+		Type = CustomPointGetterSetter.GetType();
+		Name = CustomPointGetterSetter.GetName();
+	}
+
+	/**
+	* Attempt to get a value of type T for a given point.
+	* If T doesn't match the attribute/property type, we will try to broadcast it.
+	* @param Point - The point to look for our value
+	* @param OutValue - Output value
+	* @returns true if the get succeeded, false otherwise.
+	*/
+	template <typename T>
+	bool Get(const FPCGPoint& Point, T& OutValue) const;
+
+	/**
+	* Attempt to get a value of type T for a given entry key.
+	* If T doesn't match the attribute type, we will try to broadcast it.
+	* @param Key - The entry key to get our value
+	* @param OutValue - Output value
+	* @returns true if the get succeeded, false otherwise.
+	*/
+	template <typename T>
+	bool Get(PCGMetadataEntryKey Key, T& OutValue) const;
+
+	/**
+	* Attempt to set a value of type T for a given point.
+	* If T doesn't match the attribute/property type, we will try to broadcast it.
+	* The wrapper needs to not be read only, otherwise it will return false.
+	* @param Point - The point to ser our value to.
+	* @param InValue - Value to set
+	* @returns true if the set succeeded, false otherwise.
+	*/
+	template <typename T>
+	bool Set(FPCGPoint& Point, const T& InValue);
+
+	/**
+	* Attempt to set a value of type T for a given entry key.
+	* If T doesn't match the attribute type, we will try to broadcast it.
+	* The wrapper needs to not be read only, otherwise it will return false.
+	* @param Key - The entry key to set our value
+	* @param InValue - Value to set
+	* @returns true if the set succeeded, false otherwise.
+	*/
+	template <typename T>
+	bool Set(PCGMetadataEntryKey Key, const T& InValue);
+
+	int16 GetType() const { return Type; }
+
+	bool IsValid() const { return (Attribute && Metadata) || (Property) || (CustomPointGetterSetter.IsValid()); }
+
+	const FPCGMetadataAttributeBase* GetAttribute() const { return Attribute; }
+	const FProperty* GetProperty() const { return Property; }
+	const FPCGPoint::PointCustomPropertyGetterSetter& GetCustomPointGetterSetter() const { return CustomPointGetterSetter; }
+	FName GetName() const { return Name; }
+
+private:
+	FPCGMetadataAttributeBase* Attribute = nullptr;
+	UPCGMetadata* Metadata = nullptr;
+	const FProperty* Property = nullptr;
+	FPCGPoint::PointCustomPropertyGetterSetter CustomPointGetterSetter;
+
+	int16 Type = (int16)EPCGMetadataTypes::Unknown;
+	bool bReadOnly = true;
+	FName Name = NAME_None;
+};
+
+/**
+* Iterator over a property/attribute wrapper. Will allow to iterate over points or entry map for attributes.
+* Note that it is not threadsafe for properties.
+* If the iterable is an array of const points, the wrapper needs to be readonly.
+*/
+class PCG_API FPCGPropertyAttributeIterator
+{
+public:
+	FPCGPropertyAttributeIterator() = default;
+
+	FPCGPropertyAttributeIterator(const FPCGPropertyAttributeWrapper& InWrapper, const TArrayView<FPCGPoint>& InPoints)
+		: Wrapper(InWrapper)
+		, Points(InPoints)
+		, Entries()
+		, Index(0)
+	{}
+
+	FPCGPropertyAttributeIterator(const FPCGPropertyAttributeWrapper& InWrapper, const TArrayView<const FPCGPoint>& InPoints)
+		: Wrapper(InWrapper)
+		, Points(const_cast<FPCGPoint*>(InPoints.GetData()), InPoints.Num())
+		, Entries()
+		, Index(0)
+	{
+		check(InWrapper.bReadOnly);
+	}
+
+	FPCGPropertyAttributeIterator(const FPCGPropertyAttributeWrapper& InWrapper)
+		: Wrapper(InWrapper)
+		, Points()
+		, Entries()
+		, Index(0)
+	{
+		const FPCGMetadataAttributeBase* Current = Wrapper.Attribute;
+		while (Current)
+		{
+			TArray<PCGMetadataEntryKey> Temp;
+			Current->GetEntryToValueKeyMap_NotThreadSafe().GenerateKeyArray(Temp);
+			Entries.Append(Temp);
+			Current = Current->GetParent();
+		}
+
+		// If the attribute doesn't have any entry, we will always take the default value.
+		if (Entries.IsEmpty())
+		{
+			Entries.Add(PCGInvalidEntryKey);
+		}
+	}
+
+	FPCGPropertyAttributeWrapper& GetWrapper() { return Wrapper; }
+	const FPCGPropertyAttributeWrapper& GetWrapper() const { return Wrapper; }
+
+	/**
+	* Attempt to get a value of type T for our current item.
+	* If T doesn't match the attribute/property type, we will try to broadcast it.
+	* If it succeeded, we advance our iterator.
+	* @param OutValue - Output value
+	* @returns true if the get succeeded, false otherwise.
+	*/
+	template <typename T>
+	bool GetAndAdvance(T& OutValue);
+
+	/**
+	* Attempt to set a value of type T for our current item.
+	* If T doesn't match the attribute/property type, we will try to broadcast it.
+	* If it succeeded, we advance our iterator.
+	* @param InValue - The value to set
+	* @returns true if the set succeeded, false otherwise.
+	*/
+	template <typename T>
+	bool SetAndAdvance(const T& InValue);
+
+	/**
+	* Advance the iterator.
+	* If we reach the end of our iterable, we re-start from index 0.
+	*/
+	void Advance();
+
+	bool IsValid() const { return Wrapper.IsValid() && (!Points.IsEmpty() || !Entries.IsEmpty()); }
+
+	int16 GetType() const { return Wrapper.GetType(); }
+
+	int32 Num() const { return Points.Num() > 0 ? Points.Num() : Entries.Num(); }
+
+private:
+	FPCGPropertyAttributeWrapper Wrapper;
+	TArrayView<FPCGPoint> Points;
+	TArray<PCGMetadataEntryKey> Entries;
+	int32 Index = 0;
+};
+
+//////
+/// PCGMetadataAttributeWrapper Implementation
+//////
+
+template <typename T>
+inline bool PCGMetadataAttributeWrapper::GetValueFromAttribute(const FPCGMetadataAttributeBase* Attribute, PCGMetadataEntryKey Key, T& OutValue)
+{
+	auto Getter = [Attribute, Key, &OutValue](auto&& DummyValue) -> bool
+	{
+		using InType = std::decay_t<decltype(DummyValue)>;
+
+		if constexpr (std::is_same_v<InType, T>)
+		{
+			OutValue = static_cast<const FPCGMetadataAttribute<InType>*>(Attribute)->GetValueFromItemKey(Key);
+			return true;
+		}
+		else
+		{
+			InType InValue = static_cast<const FPCGMetadataAttribute<InType>*>(Attribute)->GetValueFromItemKey(Key);
+			return PCG::Private::GetValueWithBroadcast<InType, T>(InValue, OutValue);
+		}
+	};
+
+	return PCGMetadataAttribute::CallbackWithRightType(Attribute->GetTypeId(), Getter);
+}
+
+template <typename T>
+inline bool PCGMetadataAttributeWrapper::SetValueToAttribute(FPCGMetadataAttributeBase* Attribute, PCGMetadataEntryKey Key, const T& InValue)
+{
+	auto Setter = [Attribute, Key, &InValue](auto&& DummyValue) -> bool
+	{
+		using OutType = std::decay_t<decltype(DummyValue)>;
+		OutType OutValue{};
+		if (PCG::Private::GetValueWithBroadcast<T, OutType>(InValue, OutValue))
+		{
+			if (Key == PCGInvalidEntryKey)
+			{
+				static_cast<FPCGMetadataAttribute<OutType>*>(Attribute)->SetDefaultValue(OutValue);
+			}
+			else
+			{
+				static_cast<FPCGMetadataAttribute<OutType>*>(Attribute)->SetValue(Key, OutValue);
+			}
+
+			return true;
+		}
+		else
+		{
+			return false;
+		}
+	};
+
+	return PCGMetadataAttribute::CallbackWithRightType(Attribute->GetTypeId(), Setter);
+}
+
+// Func signature : auto(auto&&)
+template <typename ObjectType, typename Func>
+inline decltype(auto) PCGMetadataAttributeWrapper::GetPropertyValueWithCallback(const ObjectType* InObject, const FProperty* InProperty, Func InFunc)
+{
+	// Double property is supported by PCG, we use a dummy double to deduce the return type of InFunc.
+	using ReturnType = decltype(InFunc(0.0));
+
+	auto DefaultValue = []() -> ReturnType
+	{
+		if constexpr (std::is_same_v<ReturnType, void>)
+		{
+			return;
+		}
+		else
+		{
+			return ReturnType{};
+		}
+	};
+
+	if (!InObject || !InProperty)
+	{
+		return DefaultValue();
+	}
+
+	const void* PropertyAddressData = InProperty->ContainerPtrToValuePtr<void>(InObject);
+
+	if (const FNumericProperty* NumericProperty = CastField<FNumericProperty>(InProperty))
+	{
+		if (NumericProperty->IsFloatingPoint())
+		{
+			return InFunc(NumericProperty->GetFloatingPointPropertyValue(PropertyAddressData));
+		}
+		else if (NumericProperty->IsInteger())
+		{
+			return InFunc(NumericProperty->GetSignedIntPropertyValue(PropertyAddressData));
+		}
+	}
+	else if (const FBoolProperty* BoolProperty = CastField<FBoolProperty>(InProperty))
+	{
+		return InFunc(BoolProperty->GetPropertyValue(PropertyAddressData));
+	}
+	else if (const FStrProperty* StringProperty = CastField<FStrProperty>(InProperty))
+	{
+		return InFunc(StringProperty->GetPropertyValue(PropertyAddressData));
+	}
+	else if (const FNameProperty* NameProperty = CastField<FNameProperty>(InProperty))
+	{
+		return InFunc(NameProperty->GetPropertyValue(PropertyAddressData));
+	}
+	else if (const FEnumProperty* EnumProperty = CastField<FEnumProperty>(InProperty))
+	{
+		return InFunc(EnumProperty->GetUnderlyingProperty()->GetSignedIntPropertyValue(PropertyAddressData));
+	}
+	else if (const FStructProperty* StructProperty = CastField<FStructProperty>(InProperty))
+	{
+		if (StructProperty->Struct == TBaseStructure<FVector>::Get())
+		{
+			return InFunc(*reinterpret_cast<const FVector*>(PropertyAddressData));
+		}
+		else if (StructProperty->Struct == TBaseStructure<FVector4>::Get())
+		{
+			return InFunc(*reinterpret_cast<const FVector4*>(PropertyAddressData));
+		}
+		else if (StructProperty->Struct == TBaseStructure<FQuat>::Get())
+		{
+			return InFunc(*reinterpret_cast<const FQuat*>(PropertyAddressData));
+		}
+		else if (StructProperty->Struct == TBaseStructure<FTransform>::Get())
+		{
+			return InFunc(*reinterpret_cast<const FTransform*>(PropertyAddressData));
+		}
+		else if (StructProperty->Struct == TBaseStructure<FRotator>::Get())
+		{
+			return InFunc(*reinterpret_cast<const FRotator*>(PropertyAddressData));
+		}
+		else if (StructProperty->Struct == TBaseStructure<FSoftObjectPath>::Get())
+		{
+			// Soft object path are transformed to strings
+			return InFunc(reinterpret_cast<const FSoftObjectPath*>(PropertyAddressData)->ToString());
+		}
+		else if (StructProperty->Struct == TBaseStructure<FSoftClassPath>::Get())
+		{
+			// Soft class path are transformed to strings
+			return InFunc(reinterpret_cast<const FSoftClassPath*>(PropertyAddressData)->ToString());
+		}
+		//else if (StructProperty->Struct == TBaseStructure<FColor>::Get())
+		//else if (StructProperty->Struct == TBaseStructure<FLinearColor>::Get())
+	}
+	else if (const FObjectPropertyBase* ObjectProperty = CastField<FObjectPropertyBase>(InProperty))
+	{
+		if (const UObject* Object = ObjectProperty->GetObjectPropertyValue(PropertyAddressData))
+		{
+			// Object are transformed into their soft path name (as a string attribute)
+			return InFunc(Object->GetPathName());
+		}
+	}
+
+	return DefaultValue();
+}
+
+// Func signature : bool(auto&)
+// Will have property value in first arg, and boolean return if we should set the property after.
+// Returns true if the set succeeded
+template <typename ObjectType, typename Func>
+inline bool PCGMetadataAttributeWrapper::SetPropertyValueFromCallback(ObjectType* InObject, const FProperty* InProperty, Func InFunc)
+{
+	if (!InObject || !InProperty)
+	{
+		return false;
+	}
+
+	bool bSuccess = false;
+	void* PropertyAddressData = InProperty->ContainerPtrToValuePtr<void>(InObject);
+
+	if (const FNumericProperty* NumericProperty = CastField<FNumericProperty>(InProperty))
+	{
+		if (NumericProperty->IsFloatingPoint())
+		{
+			double DoubleValue{};
+			bSuccess = InFunc(DoubleValue);
+			if (bSuccess)
+			{
+				NumericProperty->SetFloatingPointPropertyValue(PropertyAddressData, DoubleValue);
+			}
+		}
+		else if (NumericProperty->IsInteger())
+		{
+			int64 IntValue{};
+			bSuccess = InFunc(IntValue);
+			if (bSuccess)
+			{
+				NumericProperty->SetIntPropertyValue(PropertyAddressData, IntValue);
+			}
+		}
+	}
+	else if (const FBoolProperty* BoolProperty = CastField<FBoolProperty>(InProperty))
+	{
+		bool BoolValue{};
+		bSuccess = InFunc(BoolValue);
+		if (bSuccess)
+		{
+			BoolProperty->SetPropertyValue(PropertyAddressData, BoolValue);
+		}
+	}
+	else if (const FStrProperty* StringProperty = CastField<FStrProperty>(InProperty))
+	{
+		FString StringValue;
+		bSuccess = InFunc(StringValue);
+		if (bSuccess)
+		{
+			StringProperty->SetPropertyValue(PropertyAddressData, StringValue);
+		}
+	}
+	else if (const FNameProperty* NameProperty = CastField<FNameProperty>(InProperty))
+	{
+		FName NameValue;
+		bSuccess = InFunc(NameValue);
+		if (bSuccess)
+		{
+			NameProperty->SetPropertyValue(PropertyAddressData, NameValue);
+		}
+	}
+	else if (const FEnumProperty* EnumProperty = CastField<FEnumProperty>(InProperty))
+	{
+		int64 IntValue{};
+		bSuccess = InFunc(IntValue);
+		if (bSuccess)
+		{
+			EnumProperty->GetUnderlyingProperty()->SetIntPropertyValue(PropertyAddressData, IntValue);
+		}
+	}
+	else if (const FStructProperty* StructProperty = CastField<FStructProperty>(InProperty))
+	{
+		if (StructProperty->Struct == TBaseStructure<FVector>::Get())
+		{
+			FVector VectorValue;
+			bSuccess = InFunc(VectorValue);
+			if (bSuccess)
+			{
+				*reinterpret_cast<FVector*>(PropertyAddressData) = VectorValue;
+			}
+		}
+		else if (StructProperty->Struct == TBaseStructure<FVector4>::Get())
+		{
+			FVector4 Vector4Value;
+			bSuccess = InFunc(Vector4Value);
+			if (bSuccess)
+			{
+				*reinterpret_cast<FVector4*>(PropertyAddressData) = Vector4Value;
+			}
+		}
+		else if (StructProperty->Struct == TBaseStructure<FVector2D>::Get())
+		{
+			FVector2D Vector2Value;
+			bSuccess = InFunc(Vector2Value);
+			if (bSuccess)
+			{
+				*reinterpret_cast<FVector2D*>(PropertyAddressData) = Vector2Value;
+			}
+		}
+		else if (StructProperty->Struct == TBaseStructure<FQuat>::Get())
+		{
+			FQuat QuatValue;
+			bSuccess = InFunc(QuatValue);
+			if (bSuccess)
+			{
+				*reinterpret_cast<FQuat*>(PropertyAddressData) = QuatValue;
+			}
+		}
+		else if (StructProperty->Struct == TBaseStructure<FTransform>::Get())
+		{
+			FTransform TransformValue;
+			bSuccess = InFunc(TransformValue);
+			if (bSuccess)
+			{
+				*reinterpret_cast<FTransform*>(PropertyAddressData) = TransformValue;
+			}
+		}
+		else if (StructProperty->Struct == TBaseStructure<FRotator>::Get())
+		{
+			FRotator RotatorValue;
+			bSuccess = InFunc(RotatorValue);
+			if (bSuccess)
+			{
+				*reinterpret_cast<FRotator*>(PropertyAddressData) = RotatorValue;
+			}
+		}
+		//	else if (StructProperty->Struct == TBaseStructure<FSoftObjectPath>::Get())
+		//	{
+		//		// Soft object path are transformed from strings
+		//		reinterpret_cast<const FSoftObjectPath*>(PropertyAddressData)->SetPath(Converter(std::forward<T>(Value), FString{}));
+		//	}
+		//	else if (StructProperty->Struct == TBaseStructure<FSoftClassPath>::Get())
+		//	{
+		//		// Soft class path are transformed to strings
+		//		reinterpret_cast<const FSoftClassPath*>(PropertyAddressData)->SetPath(Converter(std::forward<T>(Value), FString{}));
+		//	}
+		//	//else if (StructProperty->Struct == TBaseStructure<FColor>::Get())
+		//	//else if (StructProperty->Struct == TBaseStructure<FLinearColor>::Get())
+	}
+	//else if (const FObjectPropertyBase* ObjectProperty = CastField<FObjectPropertyBase>(InProperty))
+	//{
+	//	if (const UObject* Object = ObjectProperty->GetObjectPropertyValue(PropertyAddressData))
+	//	{
+	//		// TODO
+	//		// Object are transformed into their soft path name (as a string attribute)
+	//		//return Converter(Object->GetPathName());
+	//	}
+	//}
+
+	return bSuccess;
+}
+
+template <typename T, typename ObjectType>
+inline bool PCGMetadataAttributeWrapper::GetValueFromProperty(const ObjectType* InObject, const FProperty* InProperty, T& OutValue)
+{
+	auto Converter = [&OutValue](auto&& PropertyValue) -> bool
+	{
+		using PropertyType = std::decay_t<decltype(PropertyValue)>;
+		return PCG::Private::GetValueWithBroadcast<PropertyType, T>(PropertyValue, OutValue);
+	};
+
+	return GetPropertyValueWithCallback(InObject, InProperty, Converter);
+}
+
+template <typename T, typename ObjectType>
+inline bool PCGMetadataAttributeWrapper::SetValueToProperty(ObjectType* InObject, const FProperty* InProperty, const T& InValue)
+{
+	auto Converter = [&InValue](auto& PropertyValue) -> bool
+	{
+		using PropertyType = std::decay_t<decltype(PropertyValue)>;
+		return PCG::Private::GetValueWithBroadcast<T, PropertyType>(InValue, PropertyValue);
+	};
+
+	return SetPropertyValueFromCallback(InObject, InProperty, Converter);
+}
+
+//////
+/// FPCGPropertyAttributeWrapper Implementation
+//////
+
+template <typename T>
+inline bool FPCGPropertyAttributeWrapper::Get(const FPCGPoint& Point, T& OutValue) const
+{
+	if (CustomPointGetterSetter.IsValid())
+	{
+		auto Getter = [this, &Point, &OutValue](auto&& DummyValue) -> bool
+		{
+			using InType = std::decay_t<decltype(DummyValue)>;
+
+			if constexpr (std::is_same_v<InType, T>)
+			{
+				CustomPointGetterSetter.Get<InType>(Point, OutValue);
+				return true;
+			}
+			else
+			{
+				InType InValue{};
+				if (CustomPointGetterSetter.Get<InType>(Point, InValue))
+				{
+					return PCG::Private::GetValueWithBroadcast<InType, T>(InValue, OutValue);
+				}
+				else
+				{
+					return false;
+				}
+			}
+		};
+
+		return PCGMetadataAttribute::CallbackWithRightType(Type, Getter);
+	}
+	else if (Attribute)
+	{
+		return PCGMetadataAttributeWrapper::GetValueFromAttribute(Attribute, Point.MetadataEntry, OutValue);
+	}
+	else if (Property)
+	{
+		return PCGMetadataAttributeWrapper::GetValueFromProperty(&Point, Property, OutValue);
+	}
+	else
+	{
+		return false;
+	}
+}
+
+template <typename T>
+inline bool FPCGPropertyAttributeWrapper::Get(PCGMetadataEntryKey Key, T& OutValue) const
+{
+	if (Attribute)
+	{
+		return PCGMetadataAttributeWrapper::GetValueFromAttribute(Attribute, Key, OutValue);
+	}
+	else
+	{
+		return false;
+	}
+}
+
+template <typename T>
+inline bool FPCGPropertyAttributeWrapper::Set(FPCGPoint& Point, const T& InValue)
+{
+	if (bReadOnly)
+	{
+		return false;
+	}
+
+	if (CustomPointGetterSetter.IsValid())
+	{
+		auto Setter = [this, &Point, &InValue](auto&& DummyValue) -> bool
+		{
+			using OutType = std::decay_t<decltype(DummyValue)>;
+
+			if constexpr (std::is_same_v<OutType, T>)
+			{
+				CustomPointGetterSetter.Set<OutType>(Point, InValue);
+				return true;
+			}
+			else
+			{
+				OutType OutValue{};
+				if (PCG::Private::GetValueWithBroadcast<T, OutType>(InValue, OutValue))
+				{
+					return CustomPointGetterSetter.Set<OutType>(Point, OutValue);
+				}
+				else
+				{
+					return false;
+				}
+			}
+		};
+
+		return PCGMetadataAttribute::CallbackWithRightType(Type, Setter);
+	}
+	else if (Attribute && Metadata)
+	{
+		Metadata->InitializeOnSet(Point.MetadataEntry);
+		return PCGMetadataAttributeWrapper::SetValueToAttribute(Attribute, Point.MetadataEntry, InValue);
+	}
+	else if (Property)
+	{
+		return PCGMetadataAttributeWrapper::SetValueToProperty(&Point, Property, InValue);
+	}
+	else
+	{
+		return false;
+	}
+}
+
+template <typename T>
+inline bool FPCGPropertyAttributeWrapper::Set(PCGMetadataEntryKey Key, const T& InValue)
+{
+	if (Attribute && !bReadOnly)
+	{
+		return PCGMetadataAttributeWrapper::SetValueToAttribute(Attribute, Key, InValue);
+	}
+	else
+	{
+		return false;
+	}
+}
+
+////////////////
+// FPCGPropertyAttributeIterator implementation
+////////////////
+
+template <typename T>
+inline bool FPCGPropertyAttributeIterator::GetAndAdvance(T& OutValue)
+{
+	bool bSuccess = false;
+	if (Points.Num() != 0)
+	{
+		bSuccess = Wrapper.Get<T>(Points[Index], OutValue);
+	}
+	else if (!Entries.IsEmpty())
+	{
+		bSuccess = Wrapper.Get<T>(Entries[Index], OutValue);
+	}
+
+	if (bSuccess)
+	{
+		Advance();
+	}
+
+	return bSuccess;
+}
+
+template <typename T>
+inline bool FPCGPropertyAttributeIterator::SetAndAdvance(const T& InValue)
+{
+	bool bSuccess = false;
+	if (Points.Num() != 0)
+	{
+		bSuccess = Wrapper.Set<T>(Points[Index], InValue);
+	}
+	else if (!Entries.IsEmpty())
+	{
+		bSuccess = Wrapper.Set<T>(Entries[Index], InValue);
+	}
+
+	if (bSuccess)
+	{
+		Advance();
+	}
+
+	return bSuccess;
+}
+
+inline void FPCGPropertyAttributeIterator::Advance()
+{
+	if (Points.Num() != 0)
+	{
+		if (++Index == Points.Num())
+		{
+			Index = 0;
+		}
+	}
+	else if (!Entries.IsEmpty())
+	{
+		if (++Index == Entries.Num())
+		{
+			Index = 0;
+		}
+	}
+}
\ No newline at end of file
diff --git a/Engine/Plugins/Experimental/PCG/Source/PCG/Public/PCGPoint.h b/Engine/Plugins/Experimental/PCG/Source/PCG/Public/PCGPoint.h
--- a/Engine/Plugins/Experimental/PCG/Source/PCG/Public/PCGPoint.h
+++ b/Engine/Plugins/Experimental/PCG/Source/PCG/Public/PCGPoint.h
@@ -3,6 +3,7 @@
 #pragma once
 
 #include "CoreMinimal.h"
+#include "Metadata/PCGMetadataAttributeTraits.h"
 
 #include "PCGPoint.generated.h"
 
@@ -18,7 +19,8 @@ enum class EPCGPointProperties : uint8
 	Rotation,
 	Scale,
 	Transform,
-	Steepness
+	Steepness,
+	LocalCenter
 };
 
 USTRUCT(BlueprintType)
@@ -72,4 +74,54 @@ public:
 		BoundsMin += Delta;
 		BoundsMax += Delta;
 	}
+
+	using PointCustomPropertyGetter = TFunction<bool(const FPCGPoint&, void*)>;
+	using PointCustomPropertySetter = TFunction<bool(FPCGPoint&, const void*)>;
+
+	struct PointCustomPropertyGetterSetter
+	{
+	public:
+		PointCustomPropertyGetterSetter() = default;
+
+		PointCustomPropertyGetterSetter(const PointCustomPropertyGetter& InGetter, const PointCustomPropertySetter& InSetter, int16 InType, FName InName);
+
+		template<typename T>
+		bool Get(const FPCGPoint& Point, T& OutValue) const
+		{
+			if (PCG::Private::IsOfTypes<T>(Type))
+			{
+				return Getter(Point, &OutValue);
+			}
+			else
+			{
+				return false;
+			}
+		}
+
+		template<typename T>
+		bool Set(FPCGPoint& Point, const T& InValue)
+		{
+			if (PCG::Private::IsOfTypes<T>(Type))
+			{
+				return Setter(Point, &InValue);
+			}
+			else
+			{
+				return false;
+			}
+		}
+
+		bool IsValid() const { return Type >= 0 && Type < (int16)EPCGMetadataTypes::Count; }
+		int16 GetType() const { return Type; }
+		FName GetName() const { return Name; }
+
+	private:
+		PointCustomPropertyGetter Getter = [](const FPCGPoint&, void*) { return false; };
+		PointCustomPropertySetter Setter = [](FPCGPoint&, const void*) { return false; };
+		int16 Type = -1;
+		FName Name = NAME_None;
+	};
+
+	static bool HasCustomPropertyGetterSetter(FName Name);
+	static PointCustomPropertyGetterSetter CreateCustomPropertyGetterSetter(FName Name);
 };
