Updated to use UE_INLINE_GENERATED_CPP_BY_NAME.

diff --git a/Engine/Plugins/Experimental/PCG/Source/PCG/Private/InstancePackers/PCGInstancePackerByRegex.cpp b/Engine/Plugins/Experimental/PCG/Source/PCG/Private/InstancePackers/PCGInstancePackerByRegex.cpp
--- a/Engine/Plugins/Experimental/PCG/Source/PCG/Private/InstancePackers/PCGInstancePackerByRegex.cpp
+++ b/Engine/Plugins/Experimental/PCG/Source/PCG/Private/InstancePackers/PCGInstancePackerByRegex.cpp
@@ -4,6 +4,8 @@
 #include "Metadata/PCGMetadataAttribute.h"
 #include "Internationalization/Regex.h"
 
+#include UE_INLINE_GENERATED_CPP_BY_NAME(PCGInstancePackerByRegex)
+
 void UPCGInstancePackerByRegex::PackInstances_Implementation(FPCGContext& Context, const UPCGSpatialData* InSpatialData, const FPCGMeshInstanceList& InstanceList, FPCGPackedCustomData& OutPackedCustomData) const
 {
 	if (!InSpatialData || !InSpatialData->Metadata)
@@ -57,3 +59,4 @@ void UPCGInstancePackerByRegex::PackInstances_Implementation(FPCGContext& Contex
 	
 	PackCustomDataFromAttributes(InstanceList, SelectedAttributes, OutPackedCustomData);
 }
+
diff --git a/Engine/Plugins/Experimental/PCG/Source/PCG/Private/MeshSelectors/PCGMeshSelectorBase.cpp b/Engine/Plugins/Experimental/PCG/Source/PCG/Private/MeshSelectors/PCGMeshSelectorBase.cpp
--- a/Engine/Plugins/Experimental/PCG/Source/PCG/Private/MeshSelectors/PCGMeshSelectorBase.cpp
+++ b/Engine/Plugins/Experimental/PCG/Source/PCG/Private/MeshSelectors/PCGMeshSelectorBase.cpp
@@ -4,6 +4,8 @@
 
 #include "Elements/PCGStaticMeshSpawner.h"
 
+#include UE_INLINE_GENERATED_CPP_BY_NAME(PCGMeshSelectorBase)
+
 int32 UPCGMeshSelectorBase::FindOrAddInstanceList(
 	TArray<FPCGMeshInstanceList>& OutInstanceLists,
 	const TSoftObjectPtr<UStaticMesh>& Mesh,
@@ -41,3 +43,4 @@ int32 UPCGMeshSelectorBase::FindOrAddInstanceList(
 
 	return OutInstanceLists.Emplace(Mesh, bOverrideCollisionProfile, CollisionProfile, bOverrideMaterials, MaterialOverrides, InCullStartDistance, InCullEndDistance);
 }
+
diff --git a/Engine/Plugins/Experimental/PCG/Source/PCG/Private/MeshSelectors/PCGMeshSelectorByAttribute.cpp b/Engine/Plugins/Experimental/PCG/Source/PCG/Private/MeshSelectors/PCGMeshSelectorByAttribute.cpp
--- a/Engine/Plugins/Experimental/PCG/Source/PCG/Private/MeshSelectors/PCGMeshSelectorByAttribute.cpp
+++ b/Engine/Plugins/Experimental/PCG/Source/PCG/Private/MeshSelectors/PCGMeshSelectorByAttribute.cpp
@@ -9,6 +9,8 @@
 #include "Metadata/PCGMetadataAttributeTraits.h"
 #include "Metadata/PCGMetadataAttributeTpl.h"
 
+#include UE_INLINE_GENERATED_CPP_BY_NAME(PCGMeshSelectorByAttribute)
+
 void UPCGMeshSelectorByAttribute::SelectInstances_Implementation(
 	FPCGContext& Context, 
 	const UPCGStaticMeshSpawnerSettings* Settings, 
@@ -96,3 +98,4 @@ void UPCGMeshSelectorByAttribute::SelectInstances_Implementation(
 		OutMeshInstances[Index].Instances.Emplace(Point);
 	}
 }
+
diff --git a/Engine/Plugins/Experimental/PCG/Source/PCG/Private/MeshSelectors/PCGMeshSelectorWeighted.cpp b/Engine/Plugins/Experimental/PCG/Source/PCG/Private/MeshSelectors/PCGMeshSelectorWeighted.cpp
--- a/Engine/Plugins/Experimental/PCG/Source/PCG/Private/MeshSelectors/PCGMeshSelectorWeighted.cpp
+++ b/Engine/Plugins/Experimental/PCG/Source/PCG/Private/MeshSelectors/PCGMeshSelectorWeighted.cpp
@@ -13,6 +13,8 @@
 #include "Engine/StaticMesh.h"
 #include "Math/RandomStream.h"
 
+#include UE_INLINE_GENERATED_CPP_BY_NAME(PCGMeshSelectorWeighted)
+
 void UPCGMeshSelectorWeighted::SelectInstances_Implementation(
 	FPCGContext& Context, 
 	const UPCGStaticMeshSpawnerSettings* Settings, 
@@ -125,3 +127,4 @@ void UPCGMeshSelectorWeighted::SelectInstances_Implementation(
 		}
 	}
 }
+
diff --git a/Engine/Plugins/Experimental/PCG/Source/PCG/Private/MeshSelectors/PCGMeshSelectorWeightedByCategory.cpp b/Engine/Plugins/Experimental/PCG/Source/PCG/Private/MeshSelectors/PCGMeshSelectorWeightedByCategory.cpp
--- a/Engine/Plugins/Experimental/PCG/Source/PCG/Private/MeshSelectors/PCGMeshSelectorWeightedByCategory.cpp
+++ b/Engine/Plugins/Experimental/PCG/Source/PCG/Private/MeshSelectors/PCGMeshSelectorWeightedByCategory.cpp
@@ -13,6 +13,8 @@
 #include "Math/RandomStream.h"
 #include "Engine/StaticMesh.h"
 
+#include UE_INLINE_GENERATED_CPP_BY_NAME(PCGMeshSelectorWeightedByCategory)
+
 struct FPCGInstancesAndWeights
 {
 	TArray<int> InstanceListIndices;
@@ -220,3 +222,4 @@ void UPCGMeshSelectorWeightedByCategory::SelectInstances_Implementation(
 		}
 	}
 }
+
diff --git a/Engine/Plugins/Experimental/PCG/Source/PCG/Private/Metadata/PCGMetadata.cpp b/Engine/Plugins/Experimental/PCG/Source/PCG/Private/Metadata/PCGMetadata.cpp
--- a/Engine/Plugins/Experimental/PCG/Source/PCG/Private/Metadata/PCGMetadata.cpp
+++ b/Engine/Plugins/Experimental/PCG/Source/PCG/Private/Metadata/PCGMetadata.cpp
@@ -5,6 +5,8 @@
 #include "Misc/ScopeRWLock.h"
 #include "Metadata/PCGMetadataAttributeWrapper.h"
 
+#include UE_INLINE_GENERATED_CPP_BY_NAME(PCGMetadata)
+
 void UPCGMetadata::Serialize(FArchive& InArchive)
 {
 	Super::Serialize(InArchive);
@@ -1017,3 +1019,4 @@ void UPCGMetadata::AccumulatePointWeightedAttributes(const FPCGPoint& InPoint, c
 {
 	AccumulateWeightedAttributes(InPoint.MetadataEntry, InMetadata, Weight, bSetNonInterpolableAttributes, OutPoint.MetadataEntry);
 }
+
diff --git a/Engine/Plugins/Experimental/PCG/Source/PCG/Private/Metadata/PCGMetadataAccessor.cpp b/Engine/Plugins/Experimental/PCG/Source/PCG/Private/Metadata/PCGMetadataAccessor.cpp
--- a/Engine/Plugins/Experimental/PCG/Source/PCG/Private/Metadata/PCGMetadataAccessor.cpp
+++ b/Engine/Plugins/Experimental/PCG/Source/PCG/Private/Metadata/PCGMetadataAccessor.cpp
@@ -5,6 +5,8 @@
 #include "Metadata/PCGMetadataAttributeTpl.h"
 #include "PCGModule.h"
 
+#include UE_INLINE_GENERATED_CPP_BY_NAME(PCGMetadataAccessor)
+
 /** Key-based implmentations */
 template<typename T>
 T UPCGMetadataAccessorHelpers::GetAttribute(PCGMetadataEntryKey Key, const UPCGMetadata* Metadata, FName AttributeName)
@@ -397,4 +399,4 @@ void UPCGMetadataAccessorHelpers::SetBoolAttribute(FPCGPoint& Point, UPCGMetadat
 bool UPCGMetadataAccessorHelpers::HasAttributeSet(const FPCGPoint& Point, const UPCGMetadata* Metadata, FName AttributeName)
 {
 	return HasAttributeSetByMetadataKey(Point.MetadataEntry, Metadata, AttributeName);
-}
\ No newline at end of file
+}
diff --git a/Engine/Plugins/Experimental/PCG/Source/PCG/Private/PCGComponent.cpp b/Engine/Plugins/Experimental/PCG/Source/PCG/Private/PCGComponent.cpp
--- a/Engine/Plugins/Experimental/PCG/Source/PCG/Private/PCGComponent.cpp
+++ b/Engine/Plugins/Experimental/PCG/Source/PCG/Private/PCGComponent.cpp
@@ -34,6 +34,8 @@
 #include "Misc/ScopeLock.h"
 #include "WorldPartition/WorldPartition.h"
 
+#include UE_INLINE_GENERATED_CPP_BY_NAME(PCGComponent)
+
 #if WITH_EDITOR
 #include "ScopedTransaction.h"
 #endif
diff --git a/Engine/Plugins/Experimental/PCG/Source/PCG/Private/PCGContext.cpp b/Engine/Plugins/Experimental/PCG/Source/PCG/Private/PCGContext.cpp
--- a/Engine/Plugins/Experimental/PCG/Source/PCG/Private/PCGContext.cpp
+++ b/Engine/Plugins/Experimental/PCG/Source/PCG/Private/PCGContext.cpp
@@ -5,6 +5,8 @@
 #include "PCGHelpers.h"
 #include "GameFramework/Actor.h"
 
+#include UE_INLINE_GENERATED_CPP_BY_NAME(PCGContext)
+
 FString FPCGContext::GetTaskName() const
 {
 	if (Node)
@@ -59,4 +61,4 @@ const UPCGSettingsInterface* FPCGContext::GetInputSettingsInterface() const
 	{
 		return InputData.GetSettingsInterface();
 	}
-}
\ No newline at end of file
+}
diff --git a/Engine/Plugins/Experimental/PCG/Source/PCG/Private/PCGData.cpp b/Engine/Plugins/Experimental/PCG/Source/PCG/Private/PCGData.cpp
--- a/Engine/Plugins/Experimental/PCG/Source/PCG/Private/PCGData.cpp
+++ b/Engine/Plugins/Experimental/PCG/Source/PCG/Private/PCGData.cpp
@@ -5,6 +5,8 @@
 #include "PCGParamData.h"
 #include "Data/PCGSpatialData.h"
 
+#include UE_INLINE_GENERATED_CPP_BY_NAME(PCGData)
+
 void FPCGRootSet::Clear()
 {
 	for (TPair<UObject*, int32>& Entry : RootSet)
@@ -259,4 +261,4 @@ TArray<FPCGTaggedData> UPCGDataFunctionLibrary::GetTaggedParams(const FPCGDataCo
 TArray<FPCGTaggedData> UPCGDataFunctionLibrary::GetAllSettings(const FPCGDataCollection& InCollection)
 {
 	return InCollection.GetAllSettings();
-}
\ No newline at end of file
+}
diff --git a/Engine/Plugins/Experimental/PCG/Source/PCG/Private/PCGDebug.cpp b/Engine/Plugins/Experimental/PCG/Source/PCG/Private/PCGDebug.cpp
--- a/Engine/Plugins/Experimental/PCG/Source/PCG/Private/PCGDebug.cpp
+++ b/Engine/Plugins/Experimental/PCG/Source/PCG/Private/PCGDebug.cpp
@@ -3,6 +3,8 @@
 #include "PCGDebug.h"
 #include "Materials/MaterialInterface.h"
 
+#include UE_INLINE_GENERATED_CPP_BY_NAME(PCGDebug)
+
 namespace PCGDebugVisConstants
 {
 	const FSoftObjectPath DefaultPointMesh = FSoftObjectPath(TEXT("/PCG/DebugObjects/PCG_Cube.PCG_Cube"));
@@ -24,4 +26,4 @@ TSoftObjectPtr<UMaterialInterface> FPCGDebugVisualizationSettings::GetMaterial()
 	{
 		return MaterialOverride;
 	}
-}
\ No newline at end of file
+}
diff --git a/Engine/Plugins/Experimental/PCG/Source/PCG/Private/PCGEdge.cpp b/Engine/Plugins/Experimental/PCG/Source/PCG/Private/PCGEdge.cpp
--- a/Engine/Plugins/Experimental/PCG/Source/PCG/Private/PCGEdge.cpp
+++ b/Engine/Plugins/Experimental/PCG/Source/PCG/Private/PCGEdge.cpp
@@ -3,6 +3,8 @@
 #include "PCGEdge.h"
 #include "PCGNode.h"
 
+#include UE_INLINE_GENERATED_CPP_BY_NAME(PCGEdge)
+
 bool UPCGEdge::IsValid() const
 {
 	return InputPin.Get() && OutputPin.Get();
@@ -18,4 +20,4 @@ const UPCGPin* UPCGEdge::GetOtherPin(const UPCGPin* Pin) const
 {
 	check(Pin == InputPin || Pin == OutputPin);
 	return Pin == InputPin ? OutputPin : InputPin;
-}
\ No newline at end of file
+}
diff --git a/Engine/Plugins/Experimental/PCG/Source/PCG/Private/PCGGraph.cpp b/Engine/Plugins/Experimental/PCG/Source/PCG/Private/PCGGraph.cpp
--- a/Engine/Plugins/Experimental/PCG/Source/PCG/Private/PCGGraph.cpp
+++ b/Engine/Plugins/Experimental/PCG/Source/PCG/Private/PCGGraph.cpp
@@ -7,6 +7,8 @@
 #include "PCGSubsystem.h"
 #include "UObject/Package.h"
 
+#include UE_INLINE_GENERATED_CPP_BY_NAME(PCGGraph)
+
 #if WITH_EDITOR
 #include "Editor.h"
 #endif
@@ -634,3 +636,4 @@ void UPCGGraph::PostEditChangeProperty(FPropertyChangedEvent& PropertyChangedEve
 	}
 }
 #endif // WITH_EDITOR
+
diff --git a/Engine/Plugins/Experimental/PCG/Source/PCG/Private/PCGInputOutputSettings.cpp b/Engine/Plugins/Experimental/PCG/Source/PCG/Private/PCGInputOutputSettings.cpp
--- a/Engine/Plugins/Experimental/PCG/Source/PCG/Private/PCGInputOutputSettings.cpp
+++ b/Engine/Plugins/Experimental/PCG/Source/PCG/Private/PCGInputOutputSettings.cpp
@@ -5,6 +5,8 @@
 
 #include "Algo/Transform.h"
 
+#include UE_INLINE_GENERATED_CPP_BY_NAME(PCGInputOutputSettings)
+
 bool FPCGInputOutputElement::ExecuteInternal(FPCGContext* Context) const
 {
 	// Essentially a pass-through element
@@ -87,4 +89,4 @@ void UPCGGraphInputOutputSettings::AddCustomPin(const FPCGPinProperties& NewCust
 {
 	Modify();
 	CustomPins.Add(NewCustomPinProperties);
-}
\ No newline at end of file
+}
diff --git a/Engine/Plugins/Experimental/PCG/Source/PCG/Private/PCGManagedResource.cpp b/Engine/Plugins/Experimental/PCG/Source/PCG/Private/PCGManagedResource.cpp
--- a/Engine/Plugins/Experimental/PCG/Source/PCG/Private/PCGManagedResource.cpp
+++ b/Engine/Plugins/Experimental/PCG/Source/PCG/Private/PCGManagedResource.cpp
@@ -6,6 +6,8 @@
 #include "Components/InstancedStaticMeshComponent.h"
 #include "Components/SceneComponent.h"
 
+#include UE_INLINE_GENERATED_CPP_BY_NAME(PCGManagedResource)
+
 void UPCGManagedResource::PostApplyToComponent()
 {
 	PostEditImport();
@@ -259,4 +261,4 @@ void UPCGManagedISMComponent::ResetComponent()
 UInstancedStaticMeshComponent* UPCGManagedISMComponent::GetComponent() const
 {
 	return Cast<UInstancedStaticMeshComponent>(GeneratedComponent.Get());
-}
\ No newline at end of file
+}
diff --git a/Engine/Plugins/Experimental/PCG/Source/PCG/Private/PCGNode.cpp b/Engine/Plugins/Experimental/PCG/Source/PCG/Private/PCGNode.cpp
--- a/Engine/Plugins/Experimental/PCG/Source/PCG/Private/PCGNode.cpp
+++ b/Engine/Plugins/Experimental/PCG/Source/PCG/Private/PCGNode.cpp
@@ -8,6 +8,8 @@
 
 #include "Algo/Transform.h"
 
+#include UE_INLINE_GENERATED_CPP_BY_NAME(PCGNode)
+
 UPCGNode::UPCGNode(const FObjectInitializer& ObjectInitializer)
 	: Super(ObjectInitializer)
 {
@@ -570,4 +572,4 @@ void UPCGNode::SetNodePosition(int32 InPositionX, int32 InPositionY)
 	PositionY = InPositionY;
 }
 
-#endif
\ No newline at end of file
+#endif
diff --git a/Engine/Plugins/Experimental/PCG/Source/PCG/Private/PCGParamData.cpp b/Engine/Plugins/Experimental/PCG/Source/PCG/Private/PCGParamData.cpp
--- a/Engine/Plugins/Experimental/PCG/Source/PCG/Private/PCGParamData.cpp
+++ b/Engine/Plugins/Experimental/PCG/Source/PCG/Private/PCGParamData.cpp
@@ -2,6 +2,8 @@
 
 #include "PCGParamData.h"
 
+#include UE_INLINE_GENERATED_CPP_BY_NAME(PCGParamData)
+
 UPCGParamData::UPCGParamData(const FObjectInitializer& ObjectInitializer)
 	: Super(ObjectInitializer)
 {
@@ -66,4 +68,4 @@ UPCGParamData* UPCGParamData::FilterParamsByKey(int64 InKey) const
 	}
 
 	return NewParams;
-}
\ No newline at end of file
+}
diff --git a/Engine/Plugins/Experimental/PCG/Source/PCG/Private/PCGPin.cpp b/Engine/Plugins/Experimental/PCG/Source/PCG/Private/PCGPin.cpp
--- a/Engine/Plugins/Experimental/PCG/Source/PCG/Private/PCGPin.cpp
+++ b/Engine/Plugins/Experimental/PCG/Source/PCG/Private/PCGPin.cpp
@@ -4,6 +4,8 @@
 #include "PCGEdge.h"
 #include "PCGNode.h"
 
+#include UE_INLINE_GENERATED_CPP_BY_NAME(PCGPin)
+
 FPCGPinProperties::FPCGPinProperties(const FName& InLabel, EPCGDataType InAllowedTypes, bool bInAllowMultipleConnections)
 	: Label(InLabel), AllowedTypes(InAllowedTypes), bAllowMultipleConnections(bInAllowMultipleConnections)
 {}
@@ -172,4 +174,4 @@ bool UPCGPin::IsCompatible(const UPCGPin* OtherPin) const
 bool UPCGPin::CanConnect(const UPCGPin* OtherPin) const
 {
 	return OtherPin && (Properties.bAllowMultipleConnections || Edges.IsEmpty());
-}
\ No newline at end of file
+}
diff --git a/Engine/Plugins/Experimental/PCG/Source/PCG/Private/PCGPoint.cpp b/Engine/Plugins/Experimental/PCG/Source/PCG/Private/PCGPoint.cpp
--- a/Engine/Plugins/Experimental/PCG/Source/PCG/Private/PCGPoint.cpp
+++ b/Engine/Plugins/Experimental/PCG/Source/PCG/Private/PCGPoint.cpp
@@ -2,6 +2,8 @@
 
 #include "PCGPoint.h"
 
+#include UE_INLINE_GENERATED_CPP_BY_NAME(PCGPoint)
+
 namespace PCGPointCustomPropertyNames
 {
 	const FName ExtentsName = TEXT("Extents");
@@ -99,4 +101,4 @@ FPCGPoint::PointCustomPropertyGetterSetter FPCGPoint::CreateCustomPropertyGetter
 	}
 
 	return FPCGPoint::PointCustomPropertyGetterSetter();
-}
\ No newline at end of file
+}
diff --git a/Engine/Plugins/Experimental/PCG/Source/PCG/Private/PCGSettings.cpp b/Engine/Plugins/Experimental/PCG/Source/PCG/Private/PCGSettings.cpp
--- a/Engine/Plugins/Experimental/PCG/Source/PCG/Private/PCGSettings.cpp
+++ b/Engine/Plugins/Experimental/PCG/Source/PCG/Private/PCGSettings.cpp
@@ -7,6 +7,8 @@
 #include "Serialization/ArchiveObjectCrc32.h"
 #include "UObject/ObjectSaveContext.h"
 
+#include UE_INLINE_GENERATED_CPP_BY_NAME(PCGSettings)
+
 #if WITH_EDITOR
 #include "Editor.h"
 #endif
@@ -267,4 +269,4 @@ bool FPCGTrivialElement::ExecuteInternal(FPCGContext* Context) const
 	// Pass-through
 	Context->OutputData = Context->InputData;
 	return true;
-}
\ No newline at end of file
+}
diff --git a/Engine/Plugins/Experimental/PCG/Source/PCG/Private/PCGSubgraph.cpp b/Engine/Plugins/Experimental/PCG/Source/PCG/Private/PCGSubgraph.cpp
--- a/Engine/Plugins/Experimental/PCG/Source/PCG/Private/PCGSubgraph.cpp
+++ b/Engine/Plugins/Experimental/PCG/Source/PCG/Private/PCGSubgraph.cpp
@@ -5,6 +5,8 @@
 #include "PCGComponent.h"
 #include "PCGSubsystem.h"
 
+#include UE_INLINE_GENERATED_CPP_BY_NAME(PCGSubgraph)
+
 void UPCGBaseSubgraphSettings::PostLoad()
 {
 	Super::PostLoad();
@@ -232,4 +234,4 @@ bool FPCGInputForwardingElement::ExecuteInternal(FPCGContext* Context) const
 {
 	Context->OutputData = Input;
 	return true;
-}
\ No newline at end of file
+}
diff --git a/Engine/Plugins/Experimental/PCG/Source/PCG/Private/PCGVolume.cpp b/Engine/Plugins/Experimental/PCG/Source/PCG/Private/PCGVolume.cpp
--- a/Engine/Plugins/Experimental/PCG/Source/PCG/Private/PCGVolume.cpp
+++ b/Engine/Plugins/Experimental/PCG/Source/PCG/Private/PCGVolume.cpp
@@ -4,6 +4,8 @@
 #include "PCGComponent.h"
 #include "PCGGraph.h"
 
+#include UE_INLINE_GENERATED_CPP_BY_NAME(PCGVolume)
+
 APCGVolume::APCGVolume(const FObjectInitializer& ObjectInitializer)
 	: Super(ObjectInitializer)
 {
diff --git a/Engine/Plugins/Experimental/PCG/Source/PCG/Private/PCGWorldActor.cpp b/Engine/Plugins/Experimental/PCG/Source/PCG/Private/PCGWorldActor.cpp
--- a/Engine/Plugins/Experimental/PCG/Source/PCG/Private/PCGWorldActor.cpp
+++ b/Engine/Plugins/Experimental/PCG/Source/PCG/Private/PCGWorldActor.cpp
@@ -9,6 +9,8 @@
 
 #include "Engine/World.h"
 
+#include UE_INLINE_GENERATED_CPP_BY_NAME(PCGWorldActor)
+
 APCGWorldActor::APCGWorldActor(const FObjectInitializer& ObjectInitializer)
 	: Super(ObjectInitializer)
 {
@@ -134,4 +136,4 @@ void APCGWorldActor::PostEditChangeProperty(FPropertyChangedEvent& PropertyChang
 		OnPartitionGridSizeChanged();
 	}
 }
-#endif // WITH_EDITOR
\ No newline at end of file
+#endif // WITH_EDITOR
diff --git a/Engine/Plugins/Experimental/PCG/Source/PCG/Private/Tests/Determinism/PCGDeterminismTestBlueprintBase.cpp b/Engine/Plugins/Experimental/PCG/Source/PCG/Private/Tests/Determinism/PCGDeterminismTestBlueprintBase.cpp
--- a/Engine/Plugins/Experimental/PCG/Source/PCG/Private/Tests/Determinism/PCGDeterminismTestBlueprintBase.cpp
+++ b/Engine/Plugins/Experimental/PCG/Source/PCG/Private/Tests/Determinism/PCGDeterminismTestBlueprintBase.cpp
@@ -2,6 +2,8 @@
 
 #include "Tests/Determinism/PCGDeterminismTestBlueprintBase.h"
 
+#include UE_INLINE_GENERATED_CPP_BY_NAME(PCGDeterminismTestBlueprintBase)
+
 void UPCGDeterminismTestBlueprintBase::ExecuteTest_Implementation(const UPCGNode* InPCGNode, FDeterminismTestResult& InOutTestResult)
 {
 	// The default case if the user tries to run the blueprint test unimplemented
@@ -9,4 +11,4 @@ void UPCGDeterminismTestBlueprintBase::ExecuteTest_Implementation(const UPCGNode
 	InOutTestResult.TestResultName = InPCGNode->GetName();
 	InOutTestResult.Seed = InPCGNode->GetSettings()->Seed;
 	PCGDeterminismTests::RunDeterminismTest(InPCGNode, InOutTestResult, PCGDeterminismTests::Defaults::DeterminismBasicTestInfo);
-}
\ No newline at end of file
+}
diff --git a/Engine/Plugins/Experimental/PCG/Source/PCG/Private/Tests/Determinism/PCGDeterminismTestsCommon.cpp b/Engine/Plugins/Experimental/PCG/Source/PCG/Private/Tests/Determinism/PCGDeterminismTestsCommon.cpp
--- a/Engine/Plugins/Experimental/PCG/Source/PCG/Private/Tests/Determinism/PCGDeterminismTestsCommon.cpp
+++ b/Engine/Plugins/Experimental/PCG/Source/PCG/Private/Tests/Determinism/PCGDeterminismTestsCommon.cpp
@@ -25,6 +25,8 @@
 #include "GameFramework/Actor.h"
 #include "UObject/Package.h"
 
+#include UE_INLINE_GENERATED_CPP_BY_NAME(PCGDeterminismTestsCommon)
+
 #define LOCTEXT_NAMESPACE "PCGDeterminism"
 
 static TAutoConsoleVariable<int32> CVarDeterminismPermutationLimit(
@@ -1572,3 +1574,4 @@ namespace PCGDeterminismTests
 }
 
 #undef LOCTEXT_NAMESPACE
+
diff --git a/Engine/Plugins/Experimental/PCG/Source/PCG/Private/Tests/Elements/PCGPropertyToParamDataElementTest.cpp b/Engine/Plugins/Experimental/PCG/Source/PCG/Private/Tests/Elements/PCGPropertyToParamDataElementTest.cpp
--- a/Engine/Plugins/Experimental/PCG/Source/PCG/Private/Tests/Elements/PCGPropertyToParamDataElementTest.cpp
+++ b/Engine/Plugins/Experimental/PCG/Source/PCG/Private/Tests/Elements/PCGPropertyToParamDataElementTest.cpp
@@ -10,6 +10,8 @@
 #include "Metadata/PCGMetadataAttributeTpl.h"
 #include "Metadata/PCGMetadataAttributeTraits.h"
 
+#include UE_INLINE_GENERATED_CPP_BY_NAME(PCGPropertyToParamDataElementTest)
+
 // To run: automation runtest pcg.tests.PropertyToParamData
 
 #if WITH_EDITOR
@@ -344,4 +346,4 @@ bool FPCGPropertyToParamDataActorFindTest::RunTest(const FString& Parameters)
 	return bSuccess;
 }
 
-#endif // WITH_EDITOR
\ No newline at end of file
+#endif // WITH_EDITOR
diff --git a/Engine/Plugins/Experimental/UIFramework/Source/Private/Types/UIFParentWidget.cpp b/Engine/Plugins/Experimental/UIFramework/Source/Private/Types/UIFParentWidget.cpp
--- a/Engine/Plugins/Experimental/UIFramework/Source/Private/Types/UIFParentWidget.cpp
+++ b/Engine/Plugins/Experimental/UIFramework/Source/Private/Types/UIFParentWidget.cpp
@@ -5,6 +5,8 @@
 #include "UIFWidget.h"
 #include "UIFPlayerComponent.h"
 
+#include UE_INLINE_GENERATED_CPP_BY_NAME(UIFParentWidget)
+
 /**
  *
  */
diff --git a/Engine/Plugins/Experimental/UIFramework/Source/Private/Types/UIFSlotBase.cpp b/Engine/Plugins/Experimental/UIFramework/Source/Private/Types/UIFSlotBase.cpp
--- a/Engine/Plugins/Experimental/UIFramework/Source/Private/Types/UIFSlotBase.cpp
+++ b/Engine/Plugins/Experimental/UIFramework/Source/Private/Types/UIFSlotBase.cpp
@@ -12,6 +12,8 @@
 #include "GameFramework/PlayerController.h"
 #include "Net/UnrealNetwork.h"
 
+#include UE_INLINE_GENERATED_CPP_BY_NAME(UIFSlotBase)
+
 
 /**
  *
diff --git a/Engine/Plugins/Experimental/UIFramework/Source/Private/Types/UIFWidgetId.cpp b/Engine/Plugins/Experimental/UIFramework/Source/Private/Types/UIFWidgetId.cpp
--- a/Engine/Plugins/Experimental/UIFramework/Source/Private/Types/UIFWidgetId.cpp
+++ b/Engine/Plugins/Experimental/UIFramework/Source/Private/Types/UIFWidgetId.cpp
@@ -3,6 +3,8 @@
 #include "Types/UIFWidgetId.h"
 #include "UIFWidget.h"
 
+#include UE_INLINE_GENERATED_CPP_BY_NAME(UIFWidgetId)
+
 /**
  *
  */
diff --git a/Engine/Plugins/Experimental/UIFramework/Source/Private/Types/UIFWidgetTree.cpp b/Engine/Plugins/Experimental/UIFramework/Source/Private/Types/UIFWidgetTree.cpp
--- a/Engine/Plugins/Experimental/UIFramework/Source/Private/Types/UIFWidgetTree.cpp
+++ b/Engine/Plugins/Experimental/UIFramework/Source/Private/Types/UIFWidgetTree.cpp
@@ -12,6 +12,8 @@
 #include "GameFramework/PlayerController.h"
 #include "Net/UnrealNetwork.h"
 
+#include UE_INLINE_GENERATED_CPP_BY_NAME(UIFWidgetTree)
+
 #if UE_UIFRAMEWORK_WITH_DEBUG
 namespace UE::UIFramework::Private
 {
@@ -414,4 +416,4 @@ void FUIFrameworkWidgetTree::AuthorityTest() const
 		
 	}
 }
-#endif
\ No newline at end of file
+#endif
diff --git a/Engine/Plugins/Experimental/UIFramework/Source/Private/UIFPlayerComponent.cpp b/Engine/Plugins/Experimental/UIFramework/Source/Private/UIFPlayerComponent.cpp
--- a/Engine/Plugins/Experimental/UIFramework/Source/Private/UIFPlayerComponent.cpp
+++ b/Engine/Plugins/Experimental/UIFramework/Source/Private/UIFPlayerComponent.cpp
@@ -16,6 +16,8 @@
 #include "GameFramework/PlayerController.h"
 #include "Net/UnrealNetwork.h"
 
+#include UE_INLINE_GENERATED_CPP_BY_NAME(UIFPlayerComponent)
+
 
 /**
  *
diff --git a/Engine/Plugins/Experimental/UIFramework/Source/Private/UIFPresenter.cpp b/Engine/Plugins/Experimental/UIFramework/Source/Private/UIFPresenter.cpp
--- a/Engine/Plugins/Experimental/UIFramework/Source/Private/UIFPresenter.cpp
+++ b/Engine/Plugins/Experimental/UIFramework/Source/Private/UIFPresenter.cpp
@@ -7,6 +7,8 @@
 #include "Components/Widget.h"
 #include "GameFramework/PlayerController.h"
 
+#include UE_INLINE_GENERATED_CPP_BY_NAME(UIFPresenter)
+
 UUIFrameworkGameViewportPresenter::FWidgetPair::FWidgetPair(UWidget* InWidget, FUIFrameworkWidgetId InWidgetId)
 	: UMGWidget(InWidget)
 	, WidgetId(InWidgetId)
@@ -63,4 +65,4 @@ void UUIFrameworkGameViewportPresenter::BeginDestroy()
 	Widgets.Reset();
 
 	Super::BeginDestroy();
-}
\ No newline at end of file
+}
diff --git a/Engine/Plugins/Experimental/UIFramework/Source/Private/UIFWidget.cpp b/Engine/Plugins/Experimental/UIFramework/Source/Private/UIFWidget.cpp
--- a/Engine/Plugins/Experimental/UIFramework/Source/Private/UIFWidget.cpp
+++ b/Engine/Plugins/Experimental/UIFramework/Source/Private/UIFWidget.cpp
@@ -17,6 +17,8 @@
 #include "Net/UnrealNetwork.h"
 #include "Types/UIFWidgetTree.h"
 
+#include UE_INLINE_GENERATED_CPP_BY_NAME(UIFWidget)
+
 
 /**
  *
diff --git a/Engine/Plugins/Experimental/UIFramework/Source/Private/Widgets/UIFButton.cpp b/Engine/Plugins/Experimental/UIFramework/Source/Private/Widgets/UIFButton.cpp
--- a/Engine/Plugins/Experimental/UIFramework/Source/Private/Widgets/UIFButton.cpp
+++ b/Engine/Plugins/Experimental/UIFramework/Source/Private/Widgets/UIFButton.cpp
@@ -16,6 +16,8 @@
 #include "GameFramework/PlayerController.h"
 #include "Net/UnrealNetwork.h"
 
+#include UE_INLINE_GENERATED_CPP_BY_NAME(UIFButton)
+
 
 /**
  * 
diff --git a/Engine/Plugins/Experimental/UIFramework/Source/Private/Widgets/UIFCanvasBox.cpp b/Engine/Plugins/Experimental/UIFramework/Source/Private/Widgets/UIFCanvasBox.cpp
--- a/Engine/Plugins/Experimental/UIFramework/Source/Private/Widgets/UIFCanvasBox.cpp
+++ b/Engine/Plugins/Experimental/UIFramework/Source/Private/Widgets/UIFCanvasBox.cpp
@@ -16,6 +16,8 @@
 #include "GameFramework/PlayerController.h"
 #include "Net/UnrealNetwork.h"
 
+#include UE_INLINE_GENERATED_CPP_BY_NAME(UIFCanvasBox)
+
 
 /**
  *
@@ -198,4 +200,4 @@ bool UUIFrameworkCanvasBox::RemoveEntry(UUIFrameworkWidget* Widget)
 FUIFrameworkCanvasBoxSlot* UUIFrameworkCanvasBox::FindEntry(FUIFrameworkWidgetId WidgetId)
 {
 	return ReplicatedSlotList.FindEntry(WidgetId);
-}
\ No newline at end of file
+}
diff --git a/Engine/Plugins/Experimental/UIFramework/Source/Private/Widgets/UIFSimpleButton.cpp b/Engine/Plugins/Experimental/UIFramework/Source/Private/Widgets/UIFSimpleButton.cpp
--- a/Engine/Plugins/Experimental/UIFramework/Source/Private/Widgets/UIFSimpleButton.cpp
+++ b/Engine/Plugins/Experimental/UIFramework/Source/Private/Widgets/UIFSimpleButton.cpp
@@ -10,6 +10,8 @@
 
 #include "Net/UnrealNetwork.h"
 
+#include UE_INLINE_GENERATED_CPP_BY_NAME(UIFSimpleButton)
+
 
 /**
  * 
diff --git a/Engine/Plugins/Experimental/UIFramework/Source/Private/Widgets/UIFStackBox.cpp b/Engine/Plugins/Experimental/UIFramework/Source/Private/Widgets/UIFStackBox.cpp
--- a/Engine/Plugins/Experimental/UIFramework/Source/Private/Widgets/UIFStackBox.cpp
+++ b/Engine/Plugins/Experimental/UIFramework/Source/Private/Widgets/UIFStackBox.cpp
@@ -18,6 +18,8 @@
 #include "Net/UnrealNetwork.h"
 #include "UObject/Package.h"
 
+#include UE_INLINE_GENERATED_CPP_BY_NAME(UIFStackBox)
+
 
 /**
  *
@@ -259,4 +261,4 @@ bool UUIFrameworkStackBox::RemoveEntry(UUIFrameworkWidget* Widget)
 FUIFrameworkStackBoxSlot* UUIFrameworkStackBox::FindEntry(FUIFrameworkWidgetId WidgetId)
 {
 	return ReplicatedSlotList.FindEntry(WidgetId);
-}
\ No newline at end of file
+}
diff --git a/Engine/Plugins/Experimental/UIFramework/Source/Private/Widgets/UIFTextBlock.cpp b/Engine/Plugins/Experimental/UIFramework/Source/Private/Widgets/UIFTextBlock.cpp
--- a/Engine/Plugins/Experimental/UIFramework/Source/Private/Widgets/UIFTextBlock.cpp
+++ b/Engine/Plugins/Experimental/UIFramework/Source/Private/Widgets/UIFTextBlock.cpp
@@ -7,6 +7,8 @@
 #include "Net/Core/PushModel/PushModel.h"
 #include "Net/UnrealNetwork.h"
 
+#include UE_INLINE_GENERATED_CPP_BY_NAME(UIFTextBlock)
+
 
 /**
  *
