Disabled blueprint triggering refreshes on some low-info dependency events

#rnx

diff --git a/Engine/Plugins/Experimental/PCG/Source/PCG/Private/Elements/PCGExecuteBlueprint.cpp b/Engine/Plugins/Experimental/PCG/Source/PCG/Private/Elements/PCGExecuteBlueprint.cpp
--- a/Engine/Plugins/Experimental/PCG/Source/PCG/Private/Elements/PCGExecuteBlueprint.cpp
+++ b/Engine/Plugins/Experimental/PCG/Source/PCG/Private/Elements/PCGExecuteBlueprint.cpp
@@ -106,6 +106,12 @@ void UPCGBlueprintElement::OnDependencyChanged(UObject* Object, FPropertyChanged
 		return;
 	}
 
+	// There are many engine notifications that aren't needed for us, esp. wrt to compilation
+	if (PropertyChangedEvent.Property == nullptr && PropertyChangedEvent.ChangeType == EPropertyChangeType::Unspecified)
+	{
+		return;
+	}
+
 	if (!DataDependencies.Contains(Object))
 	{
 		return;
