Early validation of landscape info to prevent issue when we might try to get bounds of unloaded landscapes

#rnx

diff --git a/Engine/Plugins/Experimental/PCG/Source/PCG/Private/PCGHelpers.cpp b/Engine/Plugins/Experimental/PCG/Source/PCG/Private/PCGHelpers.cpp
--- a/Engine/Plugins/Experimental/PCG/Source/PCG/Private/PCGHelpers.cpp
+++ b/Engine/Plugins/Experimental/PCG/Source/PCG/Private/PCGHelpers.cpp
@@ -10,6 +10,7 @@
 #include "Engine/World.h"
 #include "GameFramework/Actor.h"
 #include "Landscape.h"
+#include "LandscapeInfo.h"
 #include "UObject/UObjectIterator.h"
 #include "WorldPartition/WorldPartition.h"
 
@@ -169,6 +170,12 @@ namespace PCGHelpers
 
 		if (ALandscape* Landscape = Cast<ALandscape>(InLandscape))
 		{
+			// If the landscape isn't done being loaded, we're very unlikely to want to interact with it
+			if (Landscape->GetLandscapeInfo() == nullptr)
+			{
+				return FBox(EForceInit::ForceInit);
+			}
+
 #if WITH_EDITOR
 			if (!IsRuntimeOrPIE())
 			{
