Fixed GC related crash when showing debug data on cached nodes

#rnx

diff --git a/Engine/Plugins/Experimental/PCG/Source/PCG/Private/Graph/PCGGraphExecutor.cpp b/Engine/Plugins/Experimental/PCG/Source/PCG/Private/Graph/PCGGraphExecutor.cpp
--- a/Engine/Plugins/Experimental/PCG/Source/PCG/Private/Graph/PCGGraphExecutor.cpp
+++ b/Engine/Plugins/Experimental/PCG/Source/PCG/Private/Graph/PCGGraphExecutor.cpp
@@ -501,6 +501,8 @@ void FPCGGraphExecutor::Execute()
 					{
 						ActiveTask.bIsBypassed = true;
 						ActiveTask.Context->OutputData = CachedOutput;
+						// Since we've copied the output data, we need to make sure to add count ref in the root set since it'll be removed once the task is executed
+						ActiveTask.Context->OutputData.AddToRootSet(DataRootSet);
 					}
 #endif
 
