Made metadata partition work for non-keyed attributes

diff --git a/Engine/Plugins/Experimental/PCG/Source/PCG/Private/Elements/Metadata/PCGMetadataPartition.cpp b/Engine/Plugins/Experimental/PCG/Source/PCG/Private/Elements/Metadata/PCGMetadataPartition.cpp
--- a/Engine/Plugins/Experimental/PCG/Source/PCG/Private/Elements/Metadata/PCGMetadataPartition.cpp
+++ b/Engine/Plugins/Experimental/PCG/Source/PCG/Private/Elements/Metadata/PCGMetadataPartition.cpp
@@ -5,6 +5,8 @@
 #include "Data/PCGPointData.h"
 #include "Helpers/PCGSettingsHelpers.h"
 
+#include "Algo/Find.h"
+
 #include UE_INLINE_GENERATED_CPP_BY_NAME(PCGMetadataPartition)
 
 FPCGElementPtr UPCGMetadataPartitionSettings::CreateElement() const
@@ -72,12 +74,23 @@ bool FPCGMetadataPartitionElement::ExecuteInternal(FPCGContext* Context) const
 		TArray<PCGMetadataValueKey> ValueKeyMapping;
 		ValueKeyMapping.Reserve(MetadataValueKeyCount);
 
+		const bool bUsesValueKeys = AttributeBase->UsesValueKeys();
+
 		for (PCGMetadataValueKey ValueKey = 0; ValueKey < MetadataValueKeyCount; ++ValueKey)
 		{
 			if (AttributeBase->IsEqualToDefaultValue(ValueKey))
 			{
 				ValueKeyMapping.Add(-1);
 			}
+			else if (bUsesValueKeys)
+			{
+				PCGMetadataValueKey* MatchingVK = Algo::FindByPredicate(ValueKeyMapping, [ValueKey, AttributeBase](const PCGMetadataValueKey& Key)
+				{
+					return AttributeBase->AreValuesEqual(ValueKey, Key);
+				});
+
+				ValueKeyMapping.Add(MatchingVK ? *MatchingVK : ValueKey);
+			}
 			else
 			{
 				ValueKeyMapping.Add(ValueKey);
@@ -125,4 +138,4 @@ bool FPCGMetadataPartitionElement::ExecuteInternal(FPCGContext* Context) const
 	}
 
 	return true;
-}
+}
\ No newline at end of file
diff --git a/Engine/Plugins/Experimental/PCG/Source/PCG/Public/Metadata/PCGMetadataAttribute.h b/Engine/Plugins/Experimental/PCG/Source/PCG/Public/Metadata/PCGMetadataAttribute.h
--- a/Engine/Plugins/Experimental/PCG/Source/PCG/Public/Metadata/PCGMetadataAttribute.h
+++ b/Engine/Plugins/Experimental/PCG/Source/PCG/Public/Metadata/PCGMetadataAttribute.h
@@ -32,6 +32,10 @@ public:
 	virtual void SetValue(PCGMetadataEntryKey ItemKey, const FPCGMetadataAttributeBase* InAttributeA, PCGMetadataEntryKey InEntryKeyA, const FPCGMetadataAttributeBase* InAttributeB, PCGMetadataEntryKey InEntryKeyB, EPCGMetadataOp Op) = 0;
 	virtual bool IsEqualToDefaultValue(PCGMetadataValueKey ValueKey) const = 0;
 
+	virtual bool UsesValueKeys() const = 0;
+	virtual bool AreValuesEqualForEntryKeys(PCGMetadataEntryKey EntryKey1, PCGMetadataEntryKey EntryKey2) const = 0;
+	virtual bool AreValuesEqual(PCGMetadataValueKey ValueKey1, PCGMetadataValueKey ValueKey2) const = 0;
+
 	void SetValueFromValueKey(PCGMetadataEntryKey EntryKey, PCGMetadataValueKey ValueKey);
 	PCGMetadataValueKey GetValueKey(PCGMetadataEntryKey EntryKey) const;
 	bool HasNonDefaultValue(PCGMetadataEntryKey EntryKey) const;
diff --git a/Engine/Plugins/Experimental/PCG/Source/PCG/Public/Metadata/PCGMetadataAttributeTpl.h b/Engine/Plugins/Experimental/PCG/Source/PCG/Public/Metadata/PCGMetadataAttributeTpl.h
--- a/Engine/Plugins/Experimental/PCG/Source/PCG/Public/Metadata/PCGMetadataAttributeTpl.h
+++ b/Engine/Plugins/Experimental/PCG/Source/PCG/Public/Metadata/PCGMetadataAttributeTpl.h
@@ -243,7 +243,35 @@ public:
 		}
 	}
 
-	virtual bool IsEqualToDefaultValue(PCGMetadataValueKey ValueKey) const
+	virtual bool AreValuesEqualForEntryKeys(PCGMetadataEntryKey EntryKey1, PCGMetadataEntryKey EntryKey2) const override
+	{
+		return AreValuesEqual(GetValueKey(EntryKey1), GetValueKey(EntryKey2));
+	}
+
+	virtual bool AreValuesEqual(PCGMetadataValueKey ValueKey1, PCGMetadataValueKey ValueKey2) const override
+	{
+		if constexpr (PCG::Private::MetadataTraits<T>::CompressData)
+		{
+			if (ValueKey1 == PCGInvalidEntryKey)
+			{
+				return PCG::Private::MetadataTraits<T>::Equal(DefaultValue, GetValue(ValueKey2));
+			}
+			else if (ValueKey2 == PCGInvalidEntryKey)
+			{
+				return PCG::Private::MetadataTraits<T>::Equal(GetValue(ValueKey1), DefaultValue);
+			}
+			else
+			{
+				return ValueKey1 == ValueKey2;
+			}
+		}
+		else
+		{
+			return ValueKey1 == ValueKey2 || PCG::Private::MetadataTraits<T>::Equal(GetValue(ValueKey1), GetValue(ValueKey2));
+		}
+	}
+
+	virtual bool IsEqualToDefaultValue(PCGMetadataValueKey ValueKey) const override
 	{
 		return PCG::Private::MetadataTraits<T>::Equal(GetValue(ValueKey), DefaultValue);
 	}
@@ -310,6 +338,11 @@ public:
 	}
 
 	/** Code related to finding values / compressing data */
+	virtual bool UsesValueKeys() const override
+	{
+		return PCG::Private::MetadataTraits<T>::CompressData;
+	}
+
 	template<typename IT = T, typename TEnableIf<PCG::Private::MetadataTraits<IT>::CompressData>::Type* = nullptr>
 	PCGMetadataValueKey FindValue(const T& InValue) const
 	{
