[PCG] Prevent the creation of attribute with invalid names

diff --git a/Engine/Plugins/Experimental/PCG/Source/PCG/Private/Elements/Metadata/PCGMetadataOpElementBase.cpp b/Engine/Plugins/Experimental/PCG/Source/PCG/Private/Elements/Metadata/PCGMetadataOpElementBase.cpp
--- a/Engine/Plugins/Experimental/PCG/Source/PCG/Private/Elements/Metadata/PCGMetadataOpElementBase.cpp
+++ b/Engine/Plugins/Experimental/PCG/Source/PCG/Private/Elements/Metadata/PCGMetadataOpElementBase.cpp
@@ -241,6 +241,12 @@ bool FPCGMetadataElementBase::ExecuteInternal(FPCGContext* Context) const
 			OperationData.InputIterators[0].GetWrapper().Get<AttributeType>(PCGInvalidEntryKey, DefaultValue);
 
 			FPCGMetadataAttributeBase* OutputAttribute = PCGMetadataElementCommon::ClearOrCreateAttribute(OutMetadata, OutputName, DefaultValue);
+
+			if (!OutputAttribute)
+			{
+				return false;
+			}
+
 			if (bShouldCopyEntryMap)
 			{
 				PCGMetadataElementCommon::CopyEntryToValueKeyMap(SourceMetadata[0], OperationData.InputIterators[0].GetWrapper().GetAttribute(), OutputAttribute);
diff --git a/Engine/Plugins/Experimental/PCG/Source/PCG/Private/Elements/PCGAttributeGetFromPointIndexElement.cpp b/Engine/Plugins/Experimental/PCG/Source/PCG/Private/Elements/PCGAttributeGetFromPointIndexElement.cpp
--- a/Engine/Plugins/Experimental/PCG/Source/PCG/Private/Elements/PCGAttributeGetFromPointIndexElement.cpp
+++ b/Engine/Plugins/Experimental/PCG/Source/PCG/Private/Elements/PCGAttributeGetFromPointIndexElement.cpp
@@ -97,7 +97,7 @@ bool FPCGAttributeGetFromPointIndexElement::ExecuteInternal(FPCGContext* Context
 		const FPCGMetadataAttributeBase* Attribute = PointData->Metadata->GetConstAttribute(Settings->InputAttributeName);
 		UPCGParamData* OutputParamData = NewObject<UPCGParamData>();
 
-		auto ExtractAttribute = [AttributeName = Settings->InputAttributeName, &Point, &OutputParamData, Attribute](auto DummyValue)
+		auto ExtractAttribute = [this, Context, AttributeName = Settings->InputAttributeName, &Point, &OutputParamData, Attribute](auto DummyValue) -> bool
 		{
 			using AttributeType = decltype(DummyValue);
 
@@ -106,17 +106,24 @@ bool FPCGAttributeGetFromPointIndexElement::ExecuteInternal(FPCGContext* Context
 			FPCGMetadataAttributeBase* NewAttributeBase =
 				OutputParamData->Metadata->CreateAttribute<AttributeType>(AttributeName, TypedAttribute->GetValueFromItemKey(Point.MetadataEntry), TypedAttribute->AllowsInterpolation(), /*bOverrideParent=*/false);
 
-			check(NewAttributeBase);
+			if (!NewAttributeBase)
+			{
+				PCGE_LOG(Error, "Error while creating target attribute %s", *AttributeName.ToString());
+				return false;
+			}
 
 			PCGMetadataEntryKey EntryKey = OutputParamData->Metadata->AddEntry();
 			NewAttributeBase->SetValueFromValueKey(EntryKey, PCGDefaultValueKey);
-		};
 
-		PCGMetadataAttribute::CallbackWithRightType(Attribute->GetTypeId(), ExtractAttribute);
+			return true;
+		};
 
-		FPCGTaggedData& Output = Outputs.Emplace_GetRef();
-		Output.Data = OutputParamData;
-		Output.Pin = PCGAttributeGetFromPointIndexConstants::OutputAttributeLabel;
+		if (PCGMetadataAttribute::CallbackWithRightType(Attribute->GetTypeId(), ExtractAttribute))
+		{
+			FPCGTaggedData& Output = Outputs.Emplace_GetRef();
+			Output.Data = OutputParamData;
+			Output.Pin = PCGAttributeGetFromPointIndexConstants::OutputAttributeLabel;
+		}
 	}
 
 	return true;
diff --git a/Engine/Plugins/Experimental/PCG/Source/PCG/Private/Elements/PCGAttributeReduceElement.cpp b/Engine/Plugins/Experimental/PCG/Source/PCG/Private/Elements/PCGAttributeReduceElement.cpp
--- a/Engine/Plugins/Experimental/PCG/Source/PCG/Private/Elements/PCGAttributeReduceElement.cpp
+++ b/Engine/Plugins/Experimental/PCG/Source/PCG/Private/Elements/PCGAttributeReduceElement.cpp
@@ -177,6 +177,12 @@ bool FPCGAttributeReduceElement::ExecuteInternal(FPCGContext* Context) const
 	
 	FPCGMetadataAttributeBase* OutputAttribute = OutputParamData->Metadata->CopyAttribute(InputAttribute, OutputAttributeName, /*bKeepParent=*/ false, /*bCopyEntries=*/ false, /*bCopyValues=*/ false);
 
+	if (!OutputAttribute)
+	{
+		PCGE_LOG(Error, "Error while creating target attribute %s", *OutputAttributeName.ToString());
+		return true;
+	}
+
 	auto DoOperation = [InputAttribute, OutputAttribute, Operation = Settings->Operation, SpatialData, PointData, OutputParamData](auto DummyValue) -> bool
 	{
 		using AttributeType = decltype(DummyValue);
diff --git a/Engine/Plugins/Experimental/PCG/Source/PCG/Private/Elements/PCGAttributeSelectElement.cpp b/Engine/Plugins/Experimental/PCG/Source/PCG/Private/Elements/PCGAttributeSelectElement.cpp
--- a/Engine/Plugins/Experimental/PCG/Source/PCG/Private/Elements/PCGAttributeSelectElement.cpp
+++ b/Engine/Plugins/Experimental/PCG/Source/PCG/Private/Elements/PCGAttributeSelectElement.cpp
@@ -284,6 +284,12 @@ bool FPCGAttributeSelectElement::ExecuteInternal(FPCGContext* Context) const
 
 	FPCGMetadataAttributeBase* OutputAttribute = OutputParamData->Metadata->CopyAttribute(InputAttribute, OutputAttributeName, /*bKeepParent=*/ false, /*bCopyEntries=*/ false, /*bCopyValues=*/ false);
 
+	if (!OutputAttribute)
+	{
+		PCGE_LOG(Error, "Error while creating target attribute %s", *OutputAttributeName.ToString());
+		return true;
+	}
+
 	auto DoOperation = [&](auto DummyValue) -> int32
 	{
 		using AttributeType = decltype(DummyValue);
diff --git a/Engine/Plugins/Experimental/PCG/Source/PCG/Private/Elements/PCGAttributeTransferElement.cpp b/Engine/Plugins/Experimental/PCG/Source/PCG/Private/Elements/PCGAttributeTransferElement.cpp
--- a/Engine/Plugins/Experimental/PCG/Source/PCG/Private/Elements/PCGAttributeTransferElement.cpp
+++ b/Engine/Plugins/Experimental/PCG/Source/PCG/Private/Elements/PCGAttributeTransferElement.cpp
@@ -131,6 +131,12 @@ UPCGSpatialData* FPCGAttributeTransferElement::TransferSpatialToSpatial(FPCGCont
 
 	NewSpatialData->Metadata->CopyAttribute(SourceData->Metadata, Settings->SourceAttributeName, TargetAttributeName);
 
+	if (!NewSpatialData->Metadata->HasAttribute(TargetAttributeName))
+	{
+		PCGE_LOG(Error, "Error while creating target attribute %s", *TargetAttributeName.ToString());
+		return nullptr;
+	}
+
 	return NewSpatialData;
 }
 
@@ -165,6 +171,12 @@ UPCGPointData* FPCGAttributeTransferElement::TransferPointToPoint(FPCGContext* C
 	const FPCGMetadataAttributeBase* SourceAttribute = SourceData->Metadata->GetConstAttribute(Settings->SourceAttributeName);
 	FPCGMetadataAttributeBase* TargetAttribute = NewPointData->Metadata->CopyAttribute(SourceAttribute, TargetAttributeName, /*bKeepParent=*/ false, /*bCopyEntries=*/ false, /*bCopyValues=*/ true);
 
+	if (!TargetAttribute)
+	{
+		PCGE_LOG(Error, "Error while creating target attribute %s", *TargetAttributeName.ToString());
+		return nullptr;
+	}
+
 	for (int32 i = 0; i < SourcePoints.Num(); ++i)
 	{
 		PCGMetadataEntryKey& TargetKey = TargetPoints[i].MetadataEntry;
diff --git a/Engine/Plugins/Experimental/PCG/Source/PCG/Private/Elements/PCGCopyPoints.cpp b/Engine/Plugins/Experimental/PCG/Source/PCG/Private/Elements/PCGCopyPoints.cpp
--- a/Engine/Plugins/Experimental/PCG/Source/PCG/Private/Elements/PCGCopyPoints.cpp
+++ b/Engine/Plugins/Experimental/PCG/Source/PCG/Private/Elements/PCGCopyPoints.cpp
@@ -117,8 +117,10 @@ bool FPCGCopyPointsElement::ExecuteInternal(FPCGContext* Context) const
 		if (!OutPointData->Metadata->HasAttribute(AttributeName))
 		{
 			const FPCGMetadataAttributeBase* Attribute = NonRootMetadata->GetConstAttribute(AttributeName);
-			FPCGMetadataAttributeBase* NewAttribute = OutPointData->Metadata->CopyAttribute(Attribute, AttributeName, /*bKeepRoot=*/false, /*bCopyEntries=*/false, /*bCopyValues=*/true);
-			NonRootAttributes.Add(NewAttribute);
+			if (FPCGMetadataAttributeBase* NewAttribute = OutPointData->Metadata->CopyAttribute(Attribute, AttributeName, /*bKeepRoot=*/false, /*bCopyEntries=*/false, /*bCopyValues=*/true))
+			{
+				NonRootAttributes.Add(NewAttribute);
+			}
 		}
 	}
 
diff --git a/Engine/Plugins/Experimental/PCG/Source/PCG/Private/Metadata/PCGMetadata.cpp b/Engine/Plugins/Experimental/PCG/Source/PCG/Private/Metadata/PCGMetadata.cpp
--- a/Engine/Plugins/Experimental/PCG/Source/PCG/Private/Metadata/PCGMetadata.cpp
+++ b/Engine/Plugins/Experimental/PCG/Source/PCG/Private/Metadata/PCGMetadata.cpp
@@ -136,6 +136,8 @@ void UPCGMetadata::AddAttributesFiltered(const UPCGMetadata* InOther, const TSet
 		return;
 	}
 
+	bool bAttributeAdded = false;
+
 	for (const TPair<FName, FPCGMetadataAttributeBase*> OtherAttribute : InOther->Attributes)
 	{
 		// Skip this attribute if it is in an exclude list, or if it is not in an include list
@@ -149,11 +151,14 @@ void UPCGMetadata::AddAttributesFiltered(const UPCGMetadata* InOther, const TSet
 		}
 		else
 		{
-			CopyAttribute(OtherAttribute.Value, OtherAttribute.Key, /*bKeepParent=*/InOther == Parent, /*bCopyEntries=*/false, /*bCopyValues=*/false);
+			if (CopyAttribute(OtherAttribute.Value, OtherAttribute.Key, /*bKeepParent=*/InOther == Parent, /*bCopyEntries=*/false, /*bCopyValues=*/false))
+			{
+				bAttributeAdded = true;
+			}
 		}
 	}
 
-	if (InOther != Parent)
+	if (InOther != Parent && bAttributeAdded)
 	{
 		OtherParents.Add(InOther);
 	}
@@ -166,9 +171,9 @@ void UPCGMetadata::AddAttribute(const UPCGMetadata* InOther, FName AttributeName
 		return;
 	}
 
-	CopyAttribute(InOther->GetConstAttribute(AttributeName), AttributeName, /*bKeepParent=*/InOther == Parent, /*bCopyEntries=*/false, /*bCopyValues=*/false);
+	const bool bAttributeAdded = CopyAttribute(InOther->GetConstAttribute(AttributeName), AttributeName, /*bKeepParent=*/InOther == Parent, /*bCopyEntries=*/false, /*bCopyValues=*/false) != nullptr;
 
-	if (InOther != Parent)
+	if (InOther != Parent && bAttributeAdded)
 	{
 		OtherParents.Add(InOther);
 	}
@@ -535,16 +540,25 @@ FPCGMetadataAttributeBase* UPCGMetadata::CopyAttribute(const FPCGMetadataAttribu
 	check(OriginalAttribute->GetMetadata()->GetRoot() == GetRoot() || !bKeepParent);
 	FPCGMetadataAttributeBase* NewAttribute = OriginalAttribute->Copy(NewAttributeName, this, bKeepParent, bCopyEntries, bCopyValues);
 
-	AttributeLock.WriteLock();
-	NewAttribute->AttributeId = NextAttributeId++;
-	AddAttributeInternal(NewAttributeName, NewAttribute);
-	AttributeLock.WriteUnlock();
+	if (NewAttribute)
+	{
+		AttributeLock.WriteLock();
+		NewAttribute->AttributeId = NextAttributeId++;
+		AddAttributeInternal(NewAttributeName, NewAttribute);
+		AttributeLock.WriteUnlock();
+	}
 
 	return NewAttribute;
 }
 
 void UPCGMetadata::RenameAttribute(FName AttributeToRename, FName NewAttributeName)
 {
+	if (!FPCGMetadataAttributeBase::IsValidName(NewAttributeName))
+	{
+		UE_LOG(LogPCG, Error, TEXT("New attribute name %s is not valid"), *NewAttributeName.ToString());
+		return;
+	}
+
 	bool bRenamed = false;
 	AttributeLock.WriteLock();
 	if (FPCGMetadataAttributeBase** AttributeFound = Attributes.Find(AttributeToRename))
diff --git a/Engine/Plugins/Experimental/PCG/Source/PCG/Private/Metadata/PCGMetadataAttribute.cpp b/Engine/Plugins/Experimental/PCG/Source/PCG/Private/Metadata/PCGMetadataAttribute.cpp
--- a/Engine/Plugins/Experimental/PCG/Source/PCG/Private/Metadata/PCGMetadataAttribute.cpp
+++ b/Engine/Plugins/Experimental/PCG/Source/PCG/Private/Metadata/PCGMetadataAttribute.cpp
@@ -72,4 +72,42 @@ bool FPCGMetadataAttributeBase::HasNonDefaultValue(PCGMetadataEntryKey EntryKey)
 void FPCGMetadataAttributeBase::ClearEntries()
 {
 	EntryToValueKeyMap.Reset();
+}
+
+bool FPCGMetadataAttributeBase::IsValidName(const FString& Name)
+{
+	// A valid name is alphanumeric with some special characters allowed.
+	static const FString AllowedSpecialCharacters = TEXT("_.-");
+
+	for (int32 i = 0; i < Name.Len(); ++i)
+	{
+		if (FChar::IsAlpha(Name[i]) || FChar::IsDigit(Name[i]))
+		{
+			continue;
+		}
+
+		bool bAllowedSpecialCharacterFound = false;
+
+		for (int32 j = 0; j < AllowedSpecialCharacters.Len(); ++j)
+		{
+			if (Name[i] == AllowedSpecialCharacters[j])
+			{
+				bAllowedSpecialCharacterFound = true;
+				break;
+			}
+		}
+
+		if (!bAllowedSpecialCharacterFound)
+		{
+			return false;
+		}
+	}
+
+	return true;
+}
+
+bool FPCGMetadataAttributeBase::IsValidName(const FName& Name)
+{
+	// Early out on None
+	return (Name == NAME_None) || IsValidName(Name.ToString());
 }
\ No newline at end of file
diff --git a/Engine/Plugins/Experimental/PCG/Source/PCG/Public/Metadata/PCGMetadata.h b/Engine/Plugins/Experimental/PCG/Source/PCG/Public/Metadata/PCGMetadata.h
--- a/Engine/Plugins/Experimental/PCG/Source/PCG/Public/Metadata/PCGMetadata.h
+++ b/Engine/Plugins/Experimental/PCG/Source/PCG/Public/Metadata/PCGMetadata.h
@@ -268,6 +268,12 @@ protected:
 template<typename T>
 FPCGMetadataAttributeBase* UPCGMetadata::CreateAttribute(FName AttributeName, const T& DefaultValue, bool bAllowsInterpolation, bool bOverrideParent)
 {
+	if (!FPCGMetadataAttributeBase::IsValidName(AttributeName))
+	{
+		UE_LOG(LogPCG, Error, TEXT("Attribute name %s is invalid"), *AttributeName.ToString());
+		return nullptr;
+	}
+
 	const FPCGMetadataAttributeBase* ParentAttribute = nullptr;
 
 	if (bOverrideParent && Parent.IsValid())
diff --git a/Engine/Plugins/Experimental/PCG/Source/PCG/Public/Metadata/PCGMetadataAttribute.h b/Engine/Plugins/Experimental/PCG/Source/PCG/Public/Metadata/PCGMetadataAttribute.h
--- a/Engine/Plugins/Experimental/PCG/Source/PCG/Public/Metadata/PCGMetadataAttribute.h
+++ b/Engine/Plugins/Experimental/PCG/Source/PCG/Public/Metadata/PCGMetadataAttribute.h
@@ -51,6 +51,9 @@ public:
 
 	const FPCGMetadataAttributeBase* GetParent() const { return Parent; }
 
+	static bool IsValidName(const FString& Name);
+	static bool IsValidName(const FName& Name);
+
 protected:
 	TMap<PCGMetadataEntryKey, PCGMetadataValueKey> EntryToValueKeyMap;
 	mutable FRWLock EntryMapLock;
diff --git a/Engine/Plugins/Experimental/PCG/Source/PCG/Public/Metadata/PCGMetadataAttributeTpl.h b/Engine/Plugins/Experimental/PCG/Source/PCG/Public/Metadata/PCGMetadataAttributeTpl.h
--- a/Engine/Plugins/Experimental/PCG/Source/PCG/Public/Metadata/PCGMetadataAttributeTpl.h
+++ b/Engine/Plugins/Experimental/PCG/Source/PCG/Public/Metadata/PCGMetadataAttributeTpl.h
@@ -55,6 +55,13 @@ public:
 
 		// We can't keep the parent if we don't have the same root.
 		check(!bKeepParent || PCGMetadataHelpers::HasSameRoot(Metadata, InMetadata));
+
+		// Validate that the new name is valid
+		if (!IsValidName(NewName))
+		{
+			UE_LOG(LogPCG, Error, TEXT("Try to create a new attribute with an invalid name: %s"), *NewName.ToString());
+			return nullptr;
+		}
 		
 		// This copies to a new attribute.
 		FPCGMetadataAttribute<T>* AttributeCopy = new FPCGMetadataAttribute<T>(InMetadata, NewName, bKeepParent ? this : nullptr, DefaultValue, bAllowsInterpolation);
