More general solution for up-to-one refresh per frame, fixes some issues with graph manipulation

diff --git a/Engine/Plugins/Experimental/PCG/Source/PCG/Private/PCGComponent.cpp b/Engine/Plugins/Experimental/PCG/Source/PCG/Private/PCGComponent.cpp
--- a/Engine/Plugins/Experimental/PCG/Source/PCG/Private/PCGComponent.cpp
+++ b/Engine/Plugins/Experimental/PCG/Source/PCG/Private/PCGComponent.cpp
@@ -1557,29 +1557,57 @@ void UPCGComponent::Refresh()
 	}
 
 	// Discard any refresh if have already one scheduled.
-	if (CurrentRefreshTask != InvalidPCGTaskId)
+	if (UPCGSubsystem* Subsystem = GetSubsystem())
 	{
-		return;
+		if (CurrentRefreshTask == InvalidPCGTaskId)
+		{
+			CurrentRefreshTask = Subsystem->ScheduleRefresh(this);
+		}
 	}
+}
+
+void UPCGComponent::OnRefresh()
+{
+	// Mark the refresh task invalid to allow re-triggering refreshes
+	CurrentRefreshTask = InvalidPCGTaskId;
 
 	// Before doing a refresh, update the component to the subsystem if we are partitioned
 	// Only redo the mapping if we are generated
-	if (IsPartitioned() && GetSubsystem())
+	UPCGSubsystem* Subsystem = GetSubsystem();
+	bool bWasGenerated = bGenerated;
+
+	if (IsPartitioned())
+	{
+		// If we are partitioned but we have resources, we need to force a cleanup
+		if (!GeneratedResources.IsEmpty())
+		{
+			CleanupLocalImmediate(true);
+		}
+	}
+
+	if (Subsystem)
 	{
-		GetSubsystem()->RegisterOrUpdatePCGComponent(this, /*bDoActorMapping=*/ bGenerated);
+		if (IsPartitioned())
+		{
+			Subsystem->RegisterOrUpdatePCGComponent(this, /*bDoActorMapping=*/ bWasGenerated);
+		}
+		else
+		{
+			Subsystem->UnregisterPCGComponent(this);
+		}
 	}
 
 	// Following a change in some properties or in some spatial information related to this component,
 	// We need to regenerate/cleanup the graph, depending of the state in the editor.
 	if (!bActivated)
 	{
-		bool bWasGenerated = bGenerated;
 		CleanupLocalImmediate(/*bRemoveComponents=*/true);
 		bGenerated = bWasGenerated;
 	}
 	else
 	{
-		if (bGenerated && bRegenerateInEditor)
+		// If we just cleaned up resources, call back generate
+		if (bWasGenerated && (!bGenerated || bRegenerateInEditor))
 		{
 			GenerateLocal(/*bForce=*/false);
 		}
@@ -2293,64 +2321,9 @@ void FPCGComponentInstanceData::ApplyToComponent(UActorComponent* Component, con
 		// Finally, start a delayed refresh task (if there is not one already), in editor only
 		// It is important to be delayed, because we cannot spawn Partition Actors within this scope,
 		// because we are in a construction script.
-		DelayedRefresh(PCGComponent);
+		PCGComponent->Refresh();
 #endif // WITH_EDITOR
 	}
 }
 
-#if WITH_EDITOR
-void FPCGComponentInstanceData::DelayedRefresh(UPCGComponent* PCGComponent)
-{
-	if (UPCGSubsystem* Subsystem = PCGComponent->GetSubsystem())
-	{
-		if (PCGComponent->CurrentRefreshTask == InvalidPCGTaskId)
-		{
-			TWeakObjectPtr<UPCGComponent> ComponentPtr(PCGComponent);
-			PCGComponent->CurrentRefreshTask = Subsystem->ScheduleGeneric([ComponentPtr]()
-				{
-					UPCGComponent* Component = ComponentPtr.Get();
-					if (Component && IsValid(Component))
-					{
-						bool bWasGenerated = Component->bGenerated;
-						if (UPCGSubsystem* Subsystem = Component->GetSubsystem())
-						{
-							// Register/Unregister depending on the partition flag.
-							if (Component->IsPartitioned())
-							{
-								// If we are partitioned but we have resources, we need to force a cleanup
-								if (!Component->GeneratedResources.IsEmpty())
-								{
-									Component->CleanupLocalImmediate(true);
-								}
-
-								Subsystem->RegisterOrUpdatePCGComponent(Component, /*bDoActorMapping=*/bWasGenerated);
-							}
-							else
-							{
-								Subsystem->UnregisterPCGComponent(Component);
-							}
-						}
-
-						// Mark the refresh task invalid (since we check in Refresh if there is already a task scheduled)
-						Component->CurrentRefreshTask = InvalidPCGTaskId;
-
-						// If we had a cleanup, just generate.
-						if (bWasGenerated && !Component->bGenerated)
-						{
-							Component->GenerateLocal(/*bForce=*/false);
-						}
-						// Otherwise, only refresh if we are dirty.
-						else if (Component->bDirtyGenerated)
-						{
-							Component->Refresh();
-						}
-					}
-
-					return true;
-				}, PCGComponent, {});
-		}
-	}
-}
-#endif // WITH_EDITOR
-
 #undef LOCTEXT_NAMESPACE
diff --git a/Engine/Plugins/Experimental/PCG/Source/PCG/Private/PCGSubsystem.cpp b/Engine/Plugins/Experimental/PCG/Source/PCG/Private/PCGSubsystem.cpp
--- a/Engine/Plugins/Experimental/PCG/Source/PCG/Private/PCGSubsystem.cpp
+++ b/Engine/Plugins/Experimental/PCG/Source/PCG/Private/PCGSubsystem.cpp
@@ -1253,6 +1253,21 @@ FPCGTaskId UPCGSubsystem::ProcessGraph(UPCGComponent* Component, const FBox& InP
 	}
 }
 
+FPCGTaskId UPCGSubsystem::ScheduleRefresh(UPCGComponent* Component)
+{
+	TWeakObjectPtr<UPCGComponent> ComponentPtr(Component);
+
+	auto RefreshTask = [ComponentPtr]() {
+		if (UPCGComponent* Component = ComponentPtr.Get())
+		{
+			Component->OnRefresh();
+		}
+		return true;
+	};
+
+	return GraphExecutor->ScheduleGeneric(RefreshTask, Component, {});
+}
+
 FPCGTaskId UPCGSubsystem::CleanupGraph(UPCGComponent* Component, const FBox& InBounds, bool bRemoveComponents, bool bSave)
 {
 	TWeakObjectPtr<UPCGComponent> ComponentPtr(Component);
diff --git a/Engine/Plugins/Experimental/PCG/Source/PCG/Public/PCGComponent.h b/Engine/Plugins/Experimental/PCG/Source/PCG/Public/PCGComponent.h
--- a/Engine/Plugins/Experimental/PCG/Source/PCG/Public/PCGComponent.h
+++ b/Engine/Plugins/Experimental/PCG/Source/PCG/Public/PCGComponent.h
@@ -197,6 +197,8 @@ public:
 
 #if WITH_EDITOR
 	void Refresh();
+	void OnRefresh();
+
 	void DirtyGenerated(EPCGComponentDirtyFlag DataToDirtyFlag = EPCGComponentDirtyFlag::None);
 
 	/** Reset last generated bounds to force PCGPartitionActor creation on next refresh */
@@ -406,10 +408,6 @@ protected:
 	virtual bool ContainsData() const override;
 	virtual void ApplyToComponent(UActorComponent* Component, const ECacheApplyPhase CacheApplyPhase) override;
 
-#if WITH_EDITOR
-	void DelayedRefresh(UPCGComponent* PCGComponent);
-#endif // WITH_EDITOR
-
 	UPROPERTY()
 	TArray<TObjectPtr<UPCGManagedResource>> GeneratedResources;
 
diff --git a/Engine/Plugins/Experimental/PCG/Source/PCG/Public/PCGSubsystem.h b/Engine/Plugins/Experimental/PCG/Source/PCG/Public/PCGSubsystem.h
--- a/Engine/Plugins/Experimental/PCG/Source/PCG/Public/PCGSubsystem.h
+++ b/Engine/Plugins/Experimental/PCG/Source/PCG/Public/PCGSubsystem.h
@@ -119,6 +119,9 @@ public:
 
 #if WITH_EDITOR
 public:
+	/** Schedule refresh on the current or next frame */
+	FPCGTaskId ScheduleRefresh(UPCGComponent* SourceComponent);
+
 	/** Schedules an operation to cleanup the graph in the given bounds */
 	FPCGTaskId CleanupGraph(UPCGComponent* Component, const FBox& InBounds, bool bRemoveComponents, bool bSave);
 
