[PCG] Added undo support for comment nodes

diff --git a/Engine/Plugins/Experimental/PCG/Source/PCGEditor/Private/PCGEditorGraphSchemaActions.cpp b/Engine/Plugins/Experimental/PCG/Source/PCGEditor/Private/PCGEditorGraphSchemaActions.cpp
--- a/Engine/Plugins/Experimental/PCG/Source/PCGEditor/Private/PCGEditorGraphSchemaActions.cpp
+++ b/Engine/Plugins/Experimental/PCG/Source/PCGEditor/Private/PCGEditorGraphSchemaActions.cpp
@@ -352,7 +352,10 @@ UEdGraphNode* FPCGEditorGraphSchemaAction_NewComment::PerformAction(class UEdGra
 		SpawnLocation.X = CommentTemplate->NodePosX;
 		SpawnLocation.Y = CommentTemplate->NodePosY;
 	}
-
+	
+	const FScopedTransaction Transaction(*FPCGEditorCommon::ContextIdentifier, LOCTEXT("PCGEditorGraphSchemaAction_NewComment", "PCG Editor: New Comment"), nullptr);
+	EditorGraph->Modify();
+	
 	UEdGraphNode_Comment* NewNode = FEdGraphSchemaAction_NewNode::SpawnNodeFromTemplate<UEdGraphNode_Comment>(ParentGraph, CommentTemplate, SpawnLocation, bSelectNewNode);
 
 	return NewNode;
