Added editor method to change PCG blueprint variable visilibilty in a semi-smart fashion.

#rnx

diff --git a/Engine/Plugins/Experimental/PCG/Source/PCGEditor/Private/PCGEditorGraphSchema.cpp b/Engine/Plugins/Experimental/PCG/Source/PCGEditor/Private/PCGEditorGraphSchema.cpp
--- a/Engine/Plugins/Experimental/PCG/Source/PCGEditor/Private/PCGEditorGraphSchema.cpp
+++ b/Engine/Plugins/Experimental/PCG/Source/PCGEditor/Private/PCGEditorGraphSchema.cpp
@@ -13,7 +13,7 @@
 #include "PCGEditorSettings.h"
 #include "PCGEditorUtils.h"
 
-#include "AssetRegistry/AssetRegistryModule.h"
+#include "AssetRegistry/AssetData.h"
 #include "Engine/Blueprint.h"
 #include "Framework/Application/SlateApplication.h"
 #include "ScopedTransaction.h"
@@ -230,17 +230,7 @@ void UPCGEditorGraphSchema::GetNativeElementActions(FGraphActionMenuBuilder& Act
 
 void UPCGEditorGraphSchema::GetBlueprintElementActions(FGraphActionMenuBuilder& ActionMenuBuilder) const
 {
-	const FAssetRegistryModule& AssetRegistryModule = FModuleManager::LoadModuleChecked<FAssetRegistryModule>(TEXT("AssetRegistry"));
-
-	FARFilter Filter;
-	Filter.ClassPaths.Add(UBlueprint::StaticClass()->GetClassPathName());
-	Filter.bRecursiveClasses = true;
-	Filter.TagsAndValues.Add(FBlueprintTags::NativeParentClassPath, UPCGBlueprintElement::GetParentClassName());
-
-	TArray<FAssetData> BlueprintElementAssets;
-	AssetRegistryModule.Get().GetAssets(Filter, BlueprintElementAssets);
-
-	for (const FAssetData& AssetData : BlueprintElementAssets)
+	PCGEditorUtils::ForEachPCGBlueprintAssetData([&ActionMenuBuilder](const FAssetData& AssetData)
 	{
 		const bool bExposeToLibrary = AssetData.GetTagValueRef<bool>(TEXT("bExposeToLibrary"));
 		if (bExposeToLibrary)
@@ -255,17 +245,14 @@ void UPCGEditorGraphSchema::GetBlueprintElementActions(FGraphActionMenuBuilder&
 			NewBlueprintAction->BlueprintClassPath = FSoftClassPath(GeneratedClass);
 			ActionMenuBuilder.AddAction(NewBlueprintAction);
 		}
-	}
+
+		return true;
+	});
 }
 
 void UPCGEditorGraphSchema::GetSettingsElementActions(FGraphActionMenuBuilder& ActionMenuBuilder, bool bIsContextual) const
 {
-	const FAssetRegistryModule& AssetRegistryModule = FModuleManager::LoadModuleChecked<FAssetRegistryModule>(TEXT("AssetRegistry"));
-
-	TArray<FAssetData> AssetDataList;
-	AssetRegistryModule.Get().GetAssetsByClass(UPCGSettings::StaticClass()->GetClassPathName(), AssetDataList, /*bSearchSubClasses=*/true);
-
-	for (const FAssetData& AssetData : AssetDataList)
+	PCGEditorUtils::ForEachPCGSettingsAssetData([&ActionMenuBuilder, bIsContextual](const FAssetData& AssetData)
 	{
 		const bool bExposeToLibrary = AssetData.GetTagValueRef<bool>(TEXT("bExposeToLibrary"));
 		if (bExposeToLibrary)
@@ -295,17 +282,14 @@ void UPCGEditorGraphSchema::GetSettingsElementActions(FGraphActionMenuBuilder& A
 				ActionMenuBuilder.AddAction(NewSettingsActionInstance);
 			}
 		}
-	}
+
+		return true;
+	});
 }
 
 void UPCGEditorGraphSchema::GetSubgraphElementActions(FGraphActionMenuBuilder& ActionMenuBuilder) const
 {
-	const FAssetRegistryModule& AssetRegistryModule = FModuleManager::LoadModuleChecked<FAssetRegistryModule>(TEXT("AssetRegistry"));
-
-	TArray<FAssetData> AssetDataList;
-	AssetRegistryModule.Get().GetAssetsByClass(UPCGGraph::StaticClass()->GetClassPathName(), AssetDataList);
-
-	for (const FAssetData& AssetData : AssetDataList)
+	PCGEditorUtils::ForEachPCGGraphAssetData([&ActionMenuBuilder](const FAssetData& AssetData)
 	{
 		const bool bExposeToLibrary = AssetData.GetTagValueRef<bool>(TEXT("bExposeToLibrary"));
 		if (bExposeToLibrary)
@@ -318,7 +302,9 @@ void UPCGEditorGraphSchema::GetSubgraphElementActions(FGraphActionMenuBuilder& A
 			NewSubgraphAction->SubgraphObjectPath = AssetData.GetSoftObjectPath();
 			ActionMenuBuilder.AddAction(NewSubgraphAction);
 		}
-	}
+
+		return true;
+	});
 }
 
 void UPCGEditorGraphSchema::GetExtraElementActions(FGraphActionMenuBuilder& ActionMenuBuilder) const
diff --git a/Engine/Plugins/Experimental/PCG/Source/PCGEditor/Private/PCGEditorModule.cpp b/Engine/Plugins/Experimental/PCG/Source/PCGEditor/Private/PCGEditorModule.cpp
--- a/Engine/Plugins/Experimental/PCG/Source/PCGEditor/Private/PCGEditorModule.cpp
+++ b/Engine/Plugins/Experimental/PCG/Source/PCGEditor/Private/PCGEditorModule.cpp
@@ -9,6 +9,7 @@
 #include "PCGEditorGraphNodeFactory.h"
 #include "PCGEditorSettings.h"
 #include "PCGEditorStyle.h"
+#include "PCGEditorUtils.h"
 #include "PCGGraphDetails.h"
 #include "PCGSubsystem.h"
 #include "PCGVolumeDetails.h"
@@ -266,6 +267,16 @@ void FPCGEditorModule::PopulateMenuActions(FMenuBuilder& MenuBuilder)
 				}
 				})),
 		NAME_None);
+
+	MenuBuilder.AddMenuEntry(
+		LOCTEXT("UpdatePCGBlueprintVariableVisibility", "Make all PCG blueprint variables visible to instances"),
+		LOCTEXT("UpdatePCGBlueprintVariableVisibility_Tooltip", "Will visit all PCG blueprints, update their Instance editable flag, unless there is already one variable that is visible"),
+		FSlateIcon(),
+		FUIAction(
+			FExecuteAction::CreateLambda([]() {
+				PCGEditorUtils::ForcePCGBlueprintVariableVisibility();
+				})),
+		NAME_None);
 }
 
 void FPCGEditorModule::RegisterSettings()
diff --git a/Engine/Plugins/Experimental/PCG/Source/PCGEditor/Private/PCGEditorUtils.cpp b/Engine/Plugins/Experimental/PCG/Source/PCGEditor/Private/PCGEditorUtils.cpp
--- a/Engine/Plugins/Experimental/PCG/Source/PCGEditor/Private/PCGEditorUtils.cpp
+++ b/Engine/Plugins/Experimental/PCG/Source/PCGEditor/Private/PCGEditorUtils.cpp
@@ -2,14 +2,22 @@
 
 #include "PCGEditorUtils.h"
 
-#include "AssetTypeCategories.h"
-#include "UObject/NoExportTypes.h"
+#include "PCGGraph.h"
+#include "PCGSettings.h"
 #include "Elements/PCGExecuteBlueprint.h"
-#include "Engine/Blueprint.h"
+
 #include "AssetToolsModule.h"
-#include "Modules/ModuleManager.h"
+#include "AssetTypeCategories.h"
+#include "AssetRegistry/ARFilter.h"
+#include "AssetRegistry/AssetData.h"
+#include "AssetRegistry/AssetRegistryModule.h"
+#include "Engine/Blueprint.h"
+#include "Kismet2/BlueprintEditorUtils.h"
+#include "Kismet2/KismetEditorUtilities.h"
 #include "Misc/PackageName.h"
+#include "Modules/ModuleManager.h"
 #include "ObjectTools.h"
+#include "UObject/NoExportTypes.h"
 
 bool PCGEditorUtils::IsAssetPCGBlueprint(const FAssetData& InAssetData)
 {
@@ -37,3 +45,82 @@ void PCGEditorUtils::GetParentPackagePathAndUniqueName(const UObject* OriginalOb
 	AssetTools.CreateUniqueAssetName(OutPackagePath, NewAssetTentativeName, DummyPackageName, OutUniqueName);
 }
 
+void PCGEditorUtils::ForEachAssetData(const FARFilter& InFilter, TFunctionRef<bool(const FAssetData&)> InFunc)
+{
+	const FAssetRegistryModule& AssetRegistryModule = FModuleManager::LoadModuleChecked<FAssetRegistryModule>(TEXT("AssetRegistry"));
+
+	TArray<FAssetData> AssetDataList;
+	AssetRegistryModule.Get().GetAssets(InFilter, AssetDataList);
+
+	for (const FAssetData& AssetData : AssetDataList)
+	{
+		if (!InFunc(AssetData))
+		{
+			break;
+		}
+	}
+}
+
+void PCGEditorUtils::ForEachPCGBlueprintAssetData(TFunctionRef<bool(const FAssetData&)> InFunc)
+{
+	FARFilter Filter;
+	Filter.ClassPaths.Add(UBlueprint::StaticClass()->GetClassPathName());
+	Filter.bRecursiveClasses = true;
+	Filter.TagsAndValues.Add(FBlueprintTags::NativeParentClassPath, UPCGBlueprintElement::GetParentClassName());
+
+	ForEachAssetData(Filter, InFunc);
+}
+
+void PCGEditorUtils::ForEachPCGSettingsAssetData(TFunctionRef<bool(const FAssetData&)> InFunc)
+{
+	FARFilter Filter;
+	Filter.ClassPaths.Add(UPCGSettings::StaticClass()->GetClassPathName());
+	Filter.bRecursiveClasses = true;
+
+	ForEachAssetData(Filter, InFunc);
+}
+
+void PCGEditorUtils::ForEachPCGGraphAssetData(TFunctionRef<bool(const FAssetData&)> InFunc)
+{
+	FARFilter Filter;
+	Filter.ClassPaths.Add(UPCGGraph::StaticClass()->GetClassPathName());
+	Filter.bRecursiveClasses = true;
+
+	ForEachAssetData(Filter, InFunc);
+}
+
+void PCGEditorUtils::ForcePCGBlueprintVariableVisibility()
+{
+	ForEachPCGBlueprintAssetData([](const FAssetData& AssetData)
+	{
+		const FString GeneratedClass = AssetData.GetTagValueRef<FString>(FBlueprintTags::GeneratedClassPath);
+		FSoftClassPath BlueprintClassPath = FSoftClassPath(GeneratedClass);
+		TSubclassOf<UPCGBlueprintElement> BlueprintClass = BlueprintClassPath.TryLoadClass<UPCGBlueprintElement>();
+		if (BlueprintClass)
+		{
+			if (UBlueprint* Blueprint = Cast<UBlueprint>(BlueprintClass->ClassGeneratedBy))
+			{
+				if (Blueprint->NewVariables.IsEmpty())
+				{
+					return true;
+				}
+
+				const bool bHasEditOnInstanceVariables = (Blueprint->NewVariables.FindByPredicate([](const FBPVariableDescription& VarDesc) { return !(VarDesc.PropertyFlags & CPF_DisableEditOnInstance); }) != nullptr);
+				if (!bHasEditOnInstanceVariables)
+				{
+					Blueprint->Modify();
+
+					for (FBPVariableDescription& VarDesc : Blueprint->NewVariables)
+					{
+						VarDesc.PropertyFlags &= ~CPF_DisableEditOnInstance;
+					}
+
+					FBlueprintEditorUtils::MarkBlueprintAsModified(Blueprint);
+					FKismetEditorUtilities::CompileBlueprint(Blueprint, EBlueprintCompileOptions::SkipGarbageCollection);
+				}
+			}
+		}
+
+		return true;
+	});
+}
\ No newline at end of file
diff --git a/Engine/Plugins/Experimental/PCG/Source/PCGEditor/Private/PCGEditorUtils.h b/Engine/Plugins/Experimental/PCG/Source/PCGEditor/Private/PCGEditorUtils.h
--- a/Engine/Plugins/Experimental/PCG/Source/PCGEditor/Private/PCGEditorUtils.h
+++ b/Engine/Plugins/Experimental/PCG/Source/PCGEditor/Private/PCGEditorUtils.h
@@ -2,7 +2,10 @@
 
 #pragma once
 
+#include "Templates/Function.h"
+
 struct FAssetData;
+struct FARFilter;
 class UObject;
 class FString;
 
@@ -17,4 +20,13 @@ namespace PCGEditorUtils
 	*/
 	void GetParentPackagePathAndUniqueName(const UObject* OriginalObject, const FString& NewAssetTentativeName, FString& OutPackagePath, FString& OutUniqueName);
 
+	/** Methods related to visiting asset data. Ends early if InFunc returns false */
+	void ForEachAssetData(const FARFilter& Filter, TFunctionRef<bool(const FAssetData&)> InFunc);
+
+	void ForEachPCGBlueprintAssetData(TFunctionRef<bool(const FAssetData&)> InFunc);
+	void ForEachPCGSettingsAssetData(TFunctionRef<bool(const FAssetData&)> InFunc);
+	void ForEachPCGGraphAssetData(TFunctionRef<bool(const FAssetData&)> InFunc);
+
+	/** Asset deprecation methods */
+	void ForcePCGBlueprintVariableVisibility();
 }
\ No newline at end of file
