[PCG] Fix unnecessary collapse of spline projected onto landscape.

In this case we can implement a SamplePoint to see if something is in the projection (provided the direction is straight down, things gets murky-er for projections in other directions such as a shallow angle onto the terrain which results in 'shadows' - it would take more work & computation to decide if a point is in the projection or in a shadow).

Also fixes bug in the Project() function that derived the projection axis from the projected point coordinates by looking for smallest coordinate. As an example, the origin (0,0,0) would project to (0,0,0) but this does not tell us which axis was used for projection. Use projection direction directly instead which is not ambiguous.

diff --git a/Engine/Plugins/Experimental/PCG/Source/PCG/Private/Data/PCGSplineData.cpp b/Engine/Plugins/Experimental/PCG/Source/PCG/Private/Data/PCGSplineData.cpp
--- a/Engine/Plugins/Experimental/PCG/Source/PCG/Private/Data/PCGSplineData.cpp
+++ b/Engine/Plugins/Experimental/PCG/Source/PCG/Private/Data/PCGSplineData.cpp
@@ -130,32 +130,87 @@ UPCGSpatialData* UPCGSplineData::CopyInternal() const
 	return NewSplineData;
 }
 
+bool UPCGSplineProjectionData::SamplePoint(const FTransform& InTransform, const FBox& InBounds, FPCGPoint& OutPoint, UPCGMetadata* OutMetadata) const
+{
+	// TODO: support metadata - we don't currently have a good preseentation of what metadata entries mean for non-point data
+	// TODO: use InBounds when sampling spline (sample in area rather than at closest point)
+
+	// Find nearest point on projected spline
+	const FVector InPosition = InTransform.GetLocation();
+	const USplineComponent* Spline = GetSpline()->Spline.Get();
+	const FVector& SurfaceNormal = GetSurface()->GetNormal();
+
+	// Project to 2D space
+	const FTransform LocalTransform = InTransform * Spline->GetComponentTransform().Inverse();
+	FVector2D LocalPosition2D = Project(LocalTransform.GetLocation());
+	float Dummy;
+	// Find nearest key on 2D spline
+	float NearestInputKey = ProjectedPosition.InaccurateFindNearest(LocalPosition2D, Dummy);
+	// TODO: if we didn't want to hand off density computation to the spline and do it here instead, we could do it in 2D space.
+	// Find point on original spline using the previously found key. Note this is an approximation that might not hold true since
+	// we are changing the curve length. Also, to support surface orientations that are not axis aligned, the project function
+	// probably needs to construct into a coordinate space and project onto it rather than discarding an axis, otherwise project
+	// coordinates may be non-uniformly scaled.
+	const FVector NearestPointOnSpline = Spline->GetLocationAtSplineInputKey(NearestInputKey, ESplineCoordinateSpace::World);
+	const FVector PointOnLine = FMath::ClosestPointOnInfiniteLine(InPosition, InPosition + SurfaceNormal, NearestPointOnSpline);
+
+	// TODO: this is super inefficient, could be done in 2D if we duplicate the sampling code
+	FPCGPoint SplinePoint;
+	if (GetSpline()->SamplePoint(FTransform(PointOnLine), InBounds, SplinePoint, OutMetadata))
+	{
+		FPCGPoint SurfacePoint;
+		if (GetSurface()->SamplePoint(SplinePoint.Transform, InBounds, SurfacePoint, OutMetadata))
+		{
+			OutPoint = SplinePoint;
+			OutPoint.Transform = SurfacePoint.Transform;
+			OutPoint.Density *= SurfacePoint.Density;
+			OutPoint.Color *= SurfacePoint.Color;
+
+			if (OutMetadata)
+			{
+				if (SplinePoint.MetadataEntry != PCGInvalidEntryKey && SurfacePoint.MetadataEntry)
+				{
+					// TODO use metadata op from parameters once we have metadata support
+					OutMetadata->MergePointAttributesSubset(SplinePoint, OutMetadata, GetSpline()->Metadata, SurfacePoint, OutMetadata, GetSurface()->Metadata, OutPoint, EPCGMetadataOp::Max);
+				}
+				else if (SurfacePoint.MetadataEntry != PCGInvalidEntryKey)
+				{
+					OutPoint.MetadataEntry = SurfacePoint.MetadataEntry;
+				}
+			}
+
+			return true;
+		}
+	}
+
+	return false;
+}
+
 FVector2D UPCGSplineProjectionData::Project(const FVector& InVector) const
 {
 	const FVector& SurfaceNormal = GetSurface()->GetNormal();
 	FVector Projection = InVector - InVector.ProjectOnToNormal(SurfaceNormal);
 
-	// Ignore smallest absolute coordinate value.
-	// Normally, one should be zero, but because of numerical precision,
-	// We'll make sure we're doing it right
-	FVector::FReal SmallestCoordinate = TNumericLimits<FVector::FReal>::Max();
-	int SmallestCoordinateAxis = -1;
+	// Find the largest coordinate of the normal and use as the projection axis
+	int BiggestCoordinateAxis = 0;
+	FVector::FReal BiggestCoordinate = FMath::Abs(SurfaceNormal[BiggestCoordinateAxis]);
 
-	for (int Axis = 0; Axis < 3; ++Axis)
+	for (int Axis = 1; Axis < 3; ++Axis)
 	{
-		FVector::FReal AbsoluteCoordinateValue = FMath::Abs(Projection[Axis]);
-		if (AbsoluteCoordinateValue < SmallestCoordinate)
+		FVector::FReal AbsoluteCoordinateValue = FMath::Abs(SurfaceNormal[Axis]);
+		if (AbsoluteCoordinateValue > BiggestCoordinate)
 		{
-			SmallestCoordinate = AbsoluteCoordinateValue;
-			SmallestCoordinateAxis = Axis;
+			BiggestCoordinate = AbsoluteCoordinateValue;
+			BiggestCoordinateAxis = Axis;
 		}
 	}
 
+	// Discard the projection axis coordinate
 	FVector2D Projection2D;
 	int AxisIndex = 0;
 	for (int Axis = 0; Axis < 3; ++Axis)
 	{
-		if (Axis != SmallestCoordinateAxis)
+		if (Axis != BiggestCoordinateAxis)
 		{
 			Projection2D[AxisIndex++] = Projection[Axis];
 		}
diff --git a/Engine/Plugins/Experimental/PCG/Source/PCG/Public/Data/PCGSplineData.h b/Engine/Plugins/Experimental/PCG/Source/PCG/Public/Data/PCGSplineData.h
--- a/Engine/Plugins/Experimental/PCG/Source/PCG/Public/Data/PCGSplineData.h
+++ b/Engine/Plugins/Experimental/PCG/Source/PCG/Public/Data/PCGSplineData.h
@@ -64,6 +64,15 @@ public:
 	const UPCGSplineData* GetSpline() const;
 	const UPCGSpatialData* GetSurface() const;
 
+	//~Begin UPCGSpatialData interface
+	virtual bool SamplePoint(const FTransform& Transform, const FBox& Bounds, FPCGPoint& OutPoint, UPCGMetadata* OutMetadata) const override;
+
+	// It is easy to determine if a point lies on a top-down projection onto a surface. No need to convert to points.
+	// NOTE: It will be less easy if the projection is not straight downwards, as the landscape will 'shadow' the projection. This could be detected
+	// here.
+	virtual bool RequiresCollapseToSample() const override { return false; }
+	//~End UPCGSpatialData interface
+
 	UPROPERTY(BlueprintReadWrite, EditAnywhere, Category = SpatialData)
 	FInterpCurveVector2D ProjectedPosition;
 
