[PCG] Make UPCGPin::AddEdgeTo robust to calling in both directions (upstream to downstream or downstream to upstream).

Without this I called them in the wrong order and crashed the graph executor.

diff --git a/Engine/Plugins/Experimental/PCG/Source/PCG/Private/PCGPin.cpp b/Engine/Plugins/Experimental/PCG/Source/PCG/Private/PCGPin.cpp
--- a/Engine/Plugins/Experimental/PCG/Source/PCG/Private/PCGPin.cpp
+++ b/Engine/Plugins/Experimental/PCG/Source/PCG/Private/PCGPin.cpp
@@ -47,14 +47,24 @@ bool UPCGPin::AddEdgeTo(UPCGPin* OtherPin)
 		}
 	}
 
+	// This pin is upstream if the pin is an output pin
+	const bool bThisPinIsUpstream = Node->GetOutputPin(Properties.Label) == this;
+
+	// Pins should not both be upstream or both be downstream..
+	const bool bOtherPinIsUpstream = OtherPin->Node->GetOutputPin(OtherPin->Properties.Label) == OtherPin;
+	if (!ensure(bThisPinIsUpstream != bOtherPinIsUpstream))
+	{
+		return false;
+	}
+
 	Modify();
 	OtherPin->Modify();
 
 	UPCGEdge* NewEdge = Edges.Add_GetRef(NewObject<UPCGEdge>(this));
 	OtherPin->Edges.Add(NewEdge);
 
-	NewEdge->InputPin = this;
-	NewEdge->OutputPin = OtherPin;
+	NewEdge->InputPin = bThisPinIsUpstream ? this : OtherPin;
+	NewEdge->OutputPin = bThisPinIsUpstream ? OtherPin : this;
 
 	return true;
 }
