Added custom visualization for PCG blueprint settings to bubble up any in-editor calls that can be applied on their held PCG blueprint element. Also works on instances.

diff --git a/Engine/Plugins/Experimental/PCG/Source/PCG/Public/Elements/PCGExecuteBlueprint.h b/Engine/Plugins/Experimental/PCG/Source/PCG/Public/Elements/PCGExecuteBlueprint.h
--- a/Engine/Plugins/Experimental/PCG/Source/PCG/Public/Elements/PCGExecuteBlueprint.h
+++ b/Engine/Plugins/Experimental/PCG/Source/PCG/Public/Elements/PCGExecuteBlueprint.h
@@ -207,6 +207,10 @@ public:
 	UFUNCTION(BlueprintCallable, Category = "Settings|Template")
 	TSubclassOf<UPCGBlueprintElement> GetElementType() const { return BlueprintElementType; }
 
+#if WITH_EDITOR
+	TObjectPtr<UPCGBlueprintElement> GetElementInstance() const { return BlueprintElementInstance; }
+#endif
+
 protected:
 #if WITH_EDITORONLY_DATA
 	UPROPERTY()
diff --git a/Engine/Plugins/Experimental/PCG/Source/PCGEditor/PCGEditor.Build.cs b/Engine/Plugins/Experimental/PCG/Source/PCGEditor/PCGEditor.Build.cs
--- a/Engine/Plugins/Experimental/PCG/Source/PCGEditor/PCGEditor.Build.cs
+++ b/Engine/Plugins/Experimental/PCG/Source/PCGEditor/PCGEditor.Build.cs
@@ -27,7 +27,9 @@ namespace UnrealBuildTool.Rules
 					"AppFramework",
 					"ApplicationCore",
 					"AssetTools",
+					"BlueprintGraph",
 					"DesktopWidgets",
+					"DetailCustomizations",
 					"EditorScriptingUtilities",
 					"EditorStyle",
 					"EditorSubsystem",
diff --git a/Engine/Plugins/Experimental/PCG/Source/PCGEditor/Private/PCGBlueprintSettingsDetails.cpp b/Engine/Plugins/Experimental/PCG/Source/PCGEditor/Private/PCGBlueprintSettingsDetails.cpp
new file mode 100644
--- /dev/null
+++ b/Engine/Plugins/Experimental/PCG/Source/PCGEditor/Private/PCGBlueprintSettingsDetails.cpp
@@ -0,0 +1,224 @@
+// Copyright Epic Games, Inc. All Rights Reserved.
+
+#include "PCGBlueprintSettingsDetails.h"
+
+#include "DetailLayoutBuilder.h"
+//#include "ObjectDetailsTools.h"
+
+#include "Containers/ContainerAllocationPolicies.h"
+#include "Containers/UnrealString.h"
+#include "Delegates/Delegate.h"
+#include "DetailCategoryBuilder.h"
+#include "DetailWidgetRow.h"
+#include "EdGraphSchema_K2.h"
+#include "Engine/Blueprint.h"
+#include "Framework/SlateDelegates.h"
+#include "HAL/Platform.h"
+#include "HAL/PlatformCrt.h"
+#include "Internationalization/Internationalization.h"
+#include "Internationalization/Text.h"
+#include "Kismet2/BlueprintEditorUtils.h"
+#include "Layout/Margin.h"
+#include "Math/NumericLimits.h"
+#include "Misc/Attribute.h"
+#include "Misc/CString.h"
+#include "ObjectEditorUtils.h"
+#include "SWarningOrErrorBox.h"
+#include "ScopedTransaction.h"
+#include "SlotBase.h"
+#include "Templates/Casts.h"
+#include "UObject/Class.h"
+#include "UObject/NameTypes.h"
+#include "UObject/Object.h"
+#include "UObject/ObjectMacros.h"
+#include "UObject/Script.h"
+#include "UObject/UnrealNames.h"
+#include "UObject/UnrealType.h"
+#include "Widgets/DeclarativeSyntaxSupport.h"
+#include "Widgets/Input/SButton.h"
+#include "Widgets/Layout/SBox.h"
+#include "Widgets/Layout/SWrapBox.h"
+
+#define LOCTEXT_NAMESPACE "PCGSettingsDetails"
+
+TSharedRef<IDetailCustomization> FPCGBlueprintSettingsDetails::MakeInstance()
+{
+	return MakeShareable(new FPCGBlueprintSettingsDetails());
+}
+
+void FPCGBlueprintSettingsDetails::CustomizeDetails(IDetailLayoutBuilder& DetailBuilder)
+{
+	TArray<TWeakObjectPtr<UObject>> ObjectsBeingCustomized;
+	DetailBuilder.GetObjectsBeingCustomized(ObjectsBeingCustomized);
+
+	if (ObjectsBeingCustomized.Num() != 1)
+	{
+		return;
+	}
+
+	UPCGSettingsInterface* SettingsInterface = Cast<UPCGSettingsInterface>(ObjectsBeingCustomized[0]);
+
+	if (!SettingsInterface)
+	{
+		return;
+	}
+
+	UPCGBlueprintSettings* BlueprintSettings = Cast<UPCGBlueprintSettings>(SettingsInterface->GetSettings());
+
+	if (!BlueprintSettings)
+	{
+		return;
+	}
+
+	// Go through all the properties, if we find properties that have CallInEditor,
+	// create buttons for them to be executed
+	UPCGBlueprintElement* BlueprintElement = BlueprintSettings->GetElementInstance();
+	if (BlueprintElement)
+	{
+		SelectedObject = TWeakObjectPtr<UPCGBlueprintElement>(BlueprintElement);
+		AddCallInEditorMethods(DetailBuilder);
+		//FObjectDetailsTools::AddCallInEditorMethods(
+		//	DetailBuilder,
+		//	BlueprintElement->GetClass(),
+		//	[this](TWeakObjectPtr<UFunction> WeakFunctionPtr) { return FOnClicked::CreateSP(this, &FPCGBlueprintSettingsDetails::OnExecuteCallInEditorFunction, WeakFunctionPtr); });
+	}
+}
+
+void FPCGBlueprintSettingsDetails::AddCallInEditorMethods(IDetailLayoutBuilder& DetailBuilder)
+{
+	check(SelectedObject.IsValid());
+	// NOTE: this is an adaptation of FObjectDetails::AddCallInEditorMethods
+	// metadata tag for defining sort order of function buttons within a Category
+	static const FName NAME_DisplayPriority("DisplayPriority");
+
+	TArray<UFunction*, TInlineAllocator<8>> CallInEditorFunctions;
+
+	for (TFieldIterator<UFunction> FunctionIter(SelectedObject->GetClass(), EFieldIteratorFlags::IncludeSuper); FunctionIter; ++FunctionIter)
+	{
+		UFunction* Function = *FunctionIter;
+		if (Function->GetBoolMetaData(FBlueprintMetadata::MD_CallInEditor) && (Function->ParmsSize == 0))
+		{
+			const FName FunctionName = Function->GetFName();
+			if (!CallInEditorFunctions.FindByPredicate([&FunctionName](const UFunction* Func) { return Func->GetFName() == FunctionName; }))
+			{
+				CallInEditorFunctions.Add(*FunctionIter);
+			}
+		}
+	}
+
+	if (CallInEditorFunctions.Num() > 0)
+	{
+		// Sort the functions by category and then by DisplayPriority meta tag, and then by name
+		CallInEditorFunctions.Sort([](UFunction& A, UFunction& B)
+		{
+			const int32 CategorySort = A.GetMetaData(FBlueprintMetadata::MD_FunctionCategory).Compare(B.GetMetaData(FBlueprintMetadata::MD_FunctionCategory));
+			if (CategorySort != 0)
+			{
+				return (CategorySort <= 0);
+			}
+			else
+			{
+				FString DisplayPriorityAStr = A.GetMetaData(NAME_DisplayPriority);
+				int32 DisplayPriorityA = (DisplayPriorityAStr.IsEmpty() ? MAX_int32 : FCString::Atoi(*DisplayPriorityAStr));
+				if (DisplayPriorityA == 0 && !FCString::IsNumeric(*DisplayPriorityAStr))
+				{
+					DisplayPriorityA = MAX_int32;
+				}
+
+				FString DisplayPriorityBStr = B.GetMetaData(NAME_DisplayPriority);
+				int32 DisplayPriorityB = (DisplayPriorityBStr.IsEmpty() ? MAX_int32 : FCString::Atoi(*DisplayPriorityBStr));
+				if (DisplayPriorityB == 0 && !FCString::IsNumeric(*DisplayPriorityBStr))
+				{
+					DisplayPriorityB = MAX_int32;
+				}
+
+				return (DisplayPriorityA == DisplayPriorityB) ? (A.GetName() <= B.GetName()) : (DisplayPriorityA <= DisplayPriorityB);
+			}
+		});
+
+		struct FCategoryEntry
+		{
+			FName CategoryName;
+			FName RowTag;
+			TSharedPtr<SWrapBox> WrapBox;
+			FTextBuilder FunctionSearchText;
+
+			FCategoryEntry(FName InCategoryName)
+				: CategoryName(InCategoryName)
+			{
+				WrapBox = SNew(SWrapBox).UseAllottedSize(true);
+			}
+		};
+
+		// Build up a set of functions for each category, accumulating search text and buttons in a wrap box
+		FName ActiveCategory;
+		TArray<FCategoryEntry, TInlineAllocator<8>> CategoryList;
+		for (UFunction* Function : CallInEditorFunctions)
+		{
+			FName FunctionCategoryName(NAME_Default);
+			if (Function->HasMetaData(FBlueprintMetadata::FBlueprintMetadata::MD_FunctionCategory))
+			{
+				FunctionCategoryName = FName(*Function->GetMetaData(FBlueprintMetadata::MD_FunctionCategory));
+			}
+
+			if (FunctionCategoryName != ActiveCategory)
+			{
+				ActiveCategory = FunctionCategoryName;
+				CategoryList.Emplace(FunctionCategoryName);
+			}
+			FCategoryEntry& CategoryEntry = CategoryList.Last();
+
+			//@TODO: Expose the code in UK2Node_CallFunction::GetUserFacingFunctionName / etc...
+			const FText ButtonCaption = FText::FromString(FName::NameToDisplayString(*Function->GetName(), false));
+			FText FunctionTooltip = Function->GetToolTipText();
+			if (FunctionTooltip.IsEmpty())
+			{
+				FunctionTooltip = FText::FromString(Function->GetName());
+			}
+
+			TWeakObjectPtr<UFunction> WeakFunctionPtr(Function);
+			CategoryEntry.WrapBox->AddSlot()
+				.Padding(0.0f, 0.0f, 5.0f, 3.0f)
+				[
+					SNew(SButton)
+					.Text(ButtonCaption)
+					.OnClicked(FOnClicked::CreateSP(this, &FPCGBlueprintSettingsDetails::OnExecuteCallInEditorFunction, WeakFunctionPtr))
+					.ToolTipText(FunctionTooltip.IsEmptyOrWhitespace() ? LOCTEXT("CallInEditorTooltip", "Call an event on the selected object(s)") : FunctionTooltip)
+				];
+
+			CategoryEntry.RowTag = Function->GetFName();
+			CategoryEntry.FunctionSearchText.AppendLine(ButtonCaption);
+			CategoryEntry.FunctionSearchText.AppendLine(FunctionTooltip);
+		}
+
+		// Now edit the categories, adding the button strips to the details panel
+		for (FCategoryEntry& CategoryEntry : CategoryList)
+		{
+			IDetailCategoryBuilder& CategoryBuilder = DetailBuilder.EditCategory(CategoryEntry.CategoryName);
+			CategoryBuilder.AddCustomRow(CategoryEntry.FunctionSearchText.ToText())
+				.RowTag(CategoryEntry.RowTag)
+				[
+					CategoryEntry.WrapBox.ToSharedRef()
+				];
+		}
+	}
+}
+
+FReply FPCGBlueprintSettingsDetails::OnExecuteCallInEditorFunction(TWeakObjectPtr<UFunction> WeakFunctionPtr)
+{
+	if (UFunction* Function = WeakFunctionPtr.Get())
+	{
+		//@TODO: Consider naming the transaction scope after the fully qualified function name for better UX
+		FScopedTransaction Transaction(LOCTEXT("ExecuteCallInEditorMethod", "Call In Editor Action"));
+
+		FEditorScriptExecutionGuard ScriptGuard;
+		if (UObject* Object = SelectedObject.Get())
+		{
+			Object->ProcessEvent(Function, nullptr);
+		}
+	}
+
+	return FReply::Handled();
+}
+
+#undef LOCTEXT_NAMESPACE
\ No newline at end of file
diff --git a/Engine/Plugins/Experimental/PCG/Source/PCGEditor/Private/PCGEditorModule.cpp b/Engine/Plugins/Experimental/PCG/Source/PCGEditor/Private/PCGEditorModule.cpp
--- a/Engine/Plugins/Experimental/PCG/Source/PCGEditor/Private/PCGEditorModule.cpp
+++ b/Engine/Plugins/Experimental/PCG/Source/PCGEditor/Private/PCGEditorModule.cpp
@@ -11,6 +11,7 @@
 #include "PCGEditorStyle.h"
 #include "PCGEditorUtils.h"
 #include "PCGGraphDetails.h"
+#include "PCGBlueprintSettingsDetails.h"
 #include "PCGSubsystem.h"
 #include "PCGVolumeDetails.h"
 #include "PCGVolumeFactory.h"
@@ -76,6 +77,7 @@ void FPCGEditorModule::RegisterDetailsCustomizations()
 	FPropertyEditorModule& PropertyEditor = FModuleManager::LoadModuleChecked<FPropertyEditorModule>("PropertyEditor");
 	PropertyEditor.RegisterCustomClassLayout("PCGComponent", FOnGetDetailCustomizationInstance::CreateStatic(&FPCGComponentDetails::MakeInstance));
 	PropertyEditor.RegisterCustomClassLayout("PCGGraph", FOnGetDetailCustomizationInstance::CreateStatic(&FPCGGraphDetails::MakeInstance));
+	PropertyEditor.RegisterCustomClassLayout("PCGBlueprintSettings", FOnGetDetailCustomizationInstance::CreateStatic(&FPCGBlueprintSettingsDetails::MakeInstance));
 	PropertyEditor.RegisterCustomClassLayout("PCGVolume", FOnGetDetailCustomizationInstance::CreateStatic(&FPCGVolumeDetails::MakeInstance));
 }
 
@@ -86,6 +88,7 @@ void FPCGEditorModule::UnregisterDetailsCustomizations()
 		FPropertyEditorModule& PropertyModule = FModuleManager::LoadModuleChecked<FPropertyEditorModule>("PropertyEditor");
 		PropertyModule.UnregisterCustomClassLayout("PCGComponent");
 		PropertyModule.UnregisterCustomClassLayout("PCGGraph");
+		PropertyModule.UnregisterCustomClassLayout("PCGBlueprintSettings");
 		PropertyModule.UnregisterCustomClassLayout("PCGVolume");
 		PropertyModule.NotifyCustomizationModuleChanged();
 	}
diff --git a/Engine/Plugins/Experimental/PCG/Source/PCGEditor/Public/PCGBlueprintSettingsDetails.h b/Engine/Plugins/Experimental/PCG/Source/PCGEditor/Public/PCGBlueprintSettingsDetails.h
new file mode 100644
--- /dev/null
+++ b/Engine/Plugins/Experimental/PCG/Source/PCGEditor/Public/PCGBlueprintSettingsDetails.h
@@ -0,0 +1,37 @@
+// Copyright Epic Games, Inc. All Rights Reserved.
+
+#pragma once
+
+#include "IDetailCustomization.h"
+#include "IPropertyTypeCustomization.h"
+#include "Input/Reply.h"
+#include "Templates/SharedPointer.h"
+#include "UObject/WeakObjectPtr.h"
+#include "UObject/WeakObjectPtrTemplates.h"
+
+#include "PCGSettings.h"
+#include "Elements/PCGExecuteBlueprint.h"
+
+class IDetailLayoutBuilder;
+class UFunction;
+class UPCGBlueprintElement;
+
+class FPCGBlueprintSettingsDetails : public IDetailCustomization
+{
+public:
+	static TSharedRef<IDetailCustomization> MakeInstance();
+
+	/** ~Begin IDetailCustomization interface */
+	virtual void CustomizeDetails(IDetailLayoutBuilder& DetailBuilder) override;
+	/** ~End IDetailCustomization interface */
+
+private:
+	// Creates a button strip in each category that contains reflected functions marked as CallInEditor
+	void AddCallInEditorMethods(IDetailLayoutBuilder& DetailBuilder);
+
+	// Executes the specified method on the selected objects
+	FReply OnExecuteCallInEditorFunction(TWeakObjectPtr<UFunction> WeakFunctionPtr);
+
+private:
+	TWeakObjectPtr<UPCGBlueprintElement> SelectedObject;
+};
\ No newline at end of file
