Fixed pcg world ray hit orientation

#rnx

diff --git a/Engine/Plugins/Experimental/PCG/Source/PCG/Private/Data/PCGWorldData.cpp b/Engine/Plugins/Experimental/PCG/Source/PCG/Private/Data/PCGWorldData.cpp
--- a/Engine/Plugins/Experimental/PCG/Source/PCG/Private/Data/PCGWorldData.cpp
+++ b/Engine/Plugins/Experimental/PCG/Source/PCG/Private/Data/PCGWorldData.cpp
@@ -196,8 +196,13 @@ bool UPCGWorldRayHitData::SamplePoint(const FTransform& InTransform, const FBox&
 		// TODO: additional filtering?
 
 		// Finally, fill in OutPoint - we're done
-		// TODO: use normal to orient point?
-		OutPoint = FPCGPoint(FTransform(Hit.ImpactNormal.Rotation(), Hit.ImpactPoint), 1.0f, 0);
+		// Implementation note: this uses the same orthonormalization process as the landscape cache
+		ensure(Hit.ImpactNormal.IsNormalized());
+		const FVector ArbitraryVector = (FMath::Abs(Hit.ImpactNormal.Y) < (1.f - UE_KINDA_SMALL_NUMBER) ? FVector::YAxisVector : FVector::ZAxisVector);
+		const FVector XAxis = (ArbitraryVector ^ Hit.ImpactNormal).GetSafeNormal();
+		const FVector YAxis = (Hit.ImpactNormal ^ XAxis);
+
+		OutPoint = FPCGPoint(FTransform(XAxis, YAxis, Hit.ImpactNormal, Hit.ImpactPoint), 1.0f, 0);
 		UPCGBlueprintHelpers::SetSeedFromPosition(OutPoint);
 
 		// TODO: apply metadata
