Exposed PCG pins properties so they can be built in blueprints

#rnx

diff --git a/Engine/Plugins/Experimental/PCG/Source/PCG/Public/PCGPin.h b/Engine/Plugins/Experimental/PCG/Source/PCG/Public/PCGPin.h
--- a/Engine/Plugins/Experimental/PCG/Source/PCG/Public/PCGPin.h
+++ b/Engine/Plugins/Experimental/PCG/Source/PCG/Public/PCGPin.h
@@ -20,13 +20,13 @@ struct PCG_API FPCGPinProperties
 
 	bool operator==(const FPCGPinProperties& Other) const;
 
-	UPROPERTY(EditAnywhere, Category = Settings)
+	UPROPERTY(BlueprintReadWrite, EditAnywhere, Category = Settings)
 	FName Label = NAME_None;
 
-	UPROPERTY(EditAnywhere, Category = Settings)
+	UPROPERTY(BlueprintReadWrite, EditAnywhere, Category = Settings)
 	EPCGDataType AllowedTypes = EPCGDataType::Any;
 
-	UPROPERTY(EditAnywhere, Category = Settings)
+	UPROPERTY(BlueprintReadWrite, EditAnywhere, Category = Settings)
 	bool bAllowMultipleConnections = true;
 };
 
