Replaced node execution mode with enabled / debug flag.
Removed isolate feature.
Added toggle commands for node enabled / debug.

diff --git a/Engine/Plugins/Experimental/PCG/Source/PCG/Private/Graph/PCGGraphExecutor.cpp b/Engine/Plugins/Experimental/PCG/Source/PCG/Private/Graph/PCGGraphExecutor.cpp
--- a/Engine/Plugins/Experimental/PCG/Source/PCG/Private/Graph/PCGGraphExecutor.cpp
+++ b/Engine/Plugins/Experimental/PCG/Source/PCG/Private/Graph/PCGGraphExecutor.cpp
@@ -453,7 +453,7 @@ void FPCGGraphExecutor::Execute()
 				FPCGDataCollection CachedOutput;
 				const bool bResultAlreadyInCache = Task.Element->IsCacheableInstance(TaskSettingsInterface) && GraphCache.GetFromCache(Task.Element.Get(), TaskInput, TaskSettings, Task.SourceComponent.Get(), CachedOutput);
 #if WITH_EDITOR
-				const bool bNeedsToCreateActiveTask = !bResultAlreadyInCache || TaskSettingsInterface->ExecutionMode == EPCGSettingsExecutionMode::Debug || TaskSettingsInterface->ExecutionMode == EPCGSettingsExecutionMode::Isolated;
+				const bool bNeedsToCreateActiveTask = !bResultAlreadyInCache || TaskSettingsInterface->bDebug;
 #else
 				const bool bNeedsToCreateActiveTask = !bResultAlreadyInCache;
 #endif
diff --git a/Engine/Plugins/Experimental/PCG/Source/PCG/Private/PCGElement.cpp b/Engine/Plugins/Experimental/PCG/Source/PCG/Private/PCGElement.cpp
--- a/Engine/Plugins/Experimental/PCG/Source/PCG/Private/PCGElement.cpp
+++ b/Engine/Plugins/Experimental/PCG/Source/PCG/Private/PCGElement.cpp
@@ -106,7 +106,7 @@ void IPCGElement::PreExecute(FPCGContext* Context) const
 		return;
 	}
 
-	if (SettingsInterface->ExecutionMode == EPCGSettingsExecutionMode::Disabled)
+	if (!SettingsInterface->bEnabled)
 	{
 		//Pass-through - no execution
 		Context->OutputData = Context->InputData;
@@ -197,7 +197,7 @@ void IPCGElement::PostExecute(FPCGContext* Context) const
 void IPCGElement::DebugDisplay(FPCGContext* Context) const
 {
 	const UPCGSettingsInterface* SettingsInterface = Context->GetInputSettingsInterface();
-	if (SettingsInterface && (SettingsInterface->ExecutionMode == EPCGSettingsExecutionMode::Debug || SettingsInterface->ExecutionMode == EPCGSettingsExecutionMode::Isolated))
+	if (SettingsInterface && SettingsInterface->bDebug)
 	{
 		FPCGDataCollection ElementInputs = Context->InputData;
 		FPCGDataCollection ElementOutputs = Context->OutputData;
@@ -209,12 +209,6 @@ void IPCGElement::DebugDisplay(FPCGContext* Context) const
 
 		Context->InputData = ElementInputs;
 		Context->OutputData = ElementOutputs;
-
-		// Null out the output if this node is executed in isolation
-		if (SettingsInterface->ExecutionMode == EPCGSettingsExecutionMode::Isolated)
-		{
-			Context->OutputData.bCancelExecution = true;
-		}
 	}
 }
 
@@ -247,11 +241,11 @@ void IPCGElement::CleanupAndValidateOutput(FPCGContext* Context) const
 
 		// Validate all out data for errors in labels
 #if WITH_EDITOR
-		if (SettingsInterface->ExecutionMode != EPCGSettingsExecutionMode::Disabled)
+		if (SettingsInterface->bEnabled)
 		{
 			for (FPCGTaggedData& TaggedData : Context->OutputData.TaggedData)
 			{
-				int32 MatchIndex = OutputPinProperties.IndexOfByPredicate([&TaggedData](const FPCGPinProperties& InProp) { return TaggedData.Pin == InProp.Label; });
+				const int32 MatchIndex = OutputPinProperties.IndexOfByPredicate([&TaggedData](const FPCGPinProperties& InProp) { return TaggedData.Pin == InProp.Label; });
 				if (MatchIndex == INDEX_NONE)
 				{
 					PCGE_LOG(Warning, "Output generated for pin %s but cannot be routed", *TaggedData.Pin.ToString());
@@ -271,7 +265,7 @@ bool IPCGElement::IsCacheableInstance(const UPCGSettingsInterface* InSettingsInt
 {
 	if (InSettingsInterface)
 	{
-		if (InSettingsInterface->ExecutionMode == EPCGSettingsExecutionMode::Disabled)
+		if (!InSettingsInterface->bEnabled)
 		{
 			return false;
 		}
diff --git a/Engine/Plugins/Experimental/PCG/Source/PCG/Private/PCGSettings.cpp b/Engine/Plugins/Experimental/PCG/Source/PCG/Private/PCGSettings.cpp
--- a/Engine/Plugins/Experimental/PCG/Source/PCG/Private/PCGSettings.cpp
+++ b/Engine/Plugins/Experimental/PCG/Source/PCG/Private/PCGSettings.cpp
@@ -60,7 +60,8 @@ void UPCGSettings::PostLoad()
 #if WITH_EDITORONLY_DATA
 	if (ExecutionMode_DEPRECATED != EPCGSettingsExecutionMode::Enabled)
 	{
-		ExecutionMode = ExecutionMode_DEPRECATED;
+		bEnabled = ExecutionMode_DEPRECATED != EPCGSettingsExecutionMode::Disabled;
+		bDebug = ExecutionMode_DEPRECATED == EPCGSettingsExecutionMode::Debug;
 		ExecutionMode_DEPRECATED = EPCGSettingsExecutionMode::Enabled;
 	}
 #endif
diff --git a/Engine/Plugins/Experimental/PCG/Source/PCG/Public/PCGSettings.h b/Engine/Plugins/Experimental/PCG/Source/PCG/Public/PCGSettings.h
--- a/Engine/Plugins/Experimental/PCG/Source/PCG/Public/PCGSettings.h
+++ b/Engine/Plugins/Experimental/PCG/Source/PCG/Public/PCGSettings.h
@@ -62,7 +62,10 @@ public:
 #endif
 
 	UPROPERTY(EditAnywhere, BlueprintReadWrite, Category = Debug)
-	EPCGSettingsExecutionMode ExecutionMode = EPCGSettingsExecutionMode::Enabled;
+	bool bEnabled = true;
+
+	UPROPERTY(EditAnywhere, BlueprintReadWrite, Category = Debug)
+	bool bDebug = false;
 
 #if WITH_EDITORONLY_DATA
 	UPROPERTY(EditAnywhere, BlueprintReadWrite, Category = Debug, meta = (ShowOnlyInnerProperties))
diff --git a/Engine/Plugins/Experimental/PCG/Source/PCGEditor/Private/PCGEditor.cpp b/Engine/Plugins/Experimental/PCG/Source/PCGEditor/Private/PCGEditor.cpp
--- a/Engine/Plugins/Experimental/PCG/Source/PCGEditor/Private/PCGEditor.cpp
+++ b/Engine/Plugins/Experimental/PCG/Source/PCGEditor/Private/PCGEditor.cpp
@@ -412,28 +412,16 @@ void FPCGEditor::BindCommands()
 		FCanExecuteAction::CreateSP(this, &FPCGEditor::CanRunDeterminismNodeTest));
 
 	GraphEditorCommands->MapAction(
-		PCGEditorCommands.ExecutionModeEnabled,
-		FExecuteAction::CreateSP(this, &FPCGEditor::OnSetExecutionMode, EPCGSettingsExecutionMode::Enabled),
+		PCGEditorCommands.ToggleEnabled,
+		FExecuteAction::CreateSP(this, &FPCGEditor::OnToggleEnabled),
 		FCanExecuteAction(),
-		FIsActionChecked::CreateSP(this, &FPCGEditor::IsExecutionModeActive, EPCGSettingsExecutionMode::Enabled));
+		FGetActionCheckState::CreateSP(this, &FPCGEditor::GetEnabledCheckState));
 
 	GraphEditorCommands->MapAction(
-		PCGEditorCommands.ExecutionModeDebug,
-		FExecuteAction::CreateSP(this, &FPCGEditor::OnSetExecutionMode, EPCGSettingsExecutionMode::Debug),
+		PCGEditorCommands.ToggleDebug,
+		FExecuteAction::CreateSP(this, &FPCGEditor::OnToggleDebug),
 		FCanExecuteAction(),
-		FIsActionChecked::CreateSP(this, &FPCGEditor::IsExecutionModeActive, EPCGSettingsExecutionMode::Debug));
-
-	GraphEditorCommands->MapAction(
-		PCGEditorCommands.ExecutionModeDisabled,
-		FExecuteAction::CreateSP(this, &FPCGEditor::OnSetExecutionMode, EPCGSettingsExecutionMode::Disabled),
-		FCanExecuteAction(),
-		FIsActionChecked::CreateSP(this, &FPCGEditor::IsExecutionModeActive, EPCGSettingsExecutionMode::Disabled));
-
-	GraphEditorCommands->MapAction(
-		PCGEditorCommands.ExecutionModeIsolated,
-		FExecuteAction::CreateSP(this, &FPCGEditor::OnSetExecutionMode, EPCGSettingsExecutionMode::Isolated),
-		FCanExecuteAction(),
-		FIsActionChecked::CreateSP(this, &FPCGEditor::IsExecutionModeActive, EPCGSettingsExecutionMode::Isolated));
+		FGetActionCheckState::CreateSP(this, &FPCGEditor::GetDebugCheckState));
 }
 
 void FPCGEditor::OnFind()
@@ -1209,46 +1197,97 @@ void FPCGEditor::OnStopInspectNode()
 	}
 }
 
-bool FPCGEditor::IsExecutionModeActive(EPCGSettingsExecutionMode InExecutionMode) const
+void FPCGEditor::OnToggleEnabled()
 {
+	const ECheckBoxState CheckState = GetEnabledCheckState();
+	const bool bNewCheckState = !(CheckState != ECheckBoxState::Unchecked);
 	if (GraphEditorWidget.IsValid())
 	{
-		for (const UObject* Object : GraphEditorWidget->GetSelectedNodes())
+		const FScopedTransaction Transaction(*FPCGEditorCommon::ContextIdentifier, LOCTEXT("PCGEditorToggleEnableTransactionMessage", "PCG Editor: Toggle Enable Nodes"), nullptr);
+		
+		for (UObject* Object : GraphEditorWidget->GetSelectedNodes())
 		{
-			const UPCGEditorGraphNodeBase* PCGEditorGraphNode = Cast<const UPCGEditorGraphNodeBase>(Object);
+			UPCGEditorGraphNodeBase* PCGEditorGraphNode = Cast<UPCGEditorGraphNodeBase>(Object);
 			if (!PCGEditorGraphNode)
 			{
 				continue;
 			}
 
-			const UPCGNode* PCGNode = PCGEditorGraphNode->GetPCGNode();
+			UPCGNode* PCGNode = PCGEditorGraphNode->GetPCGNode();
 			if (!PCGNode)
 			{
 				continue;
 			}
 
-			const UPCGSettingsInterface* PCGSettingsInterface = PCGNode->GetSettingsInterface();
+			UPCGSettingsInterface* PCGSettingsInterface = PCGNode->GetSettingsInterface();
 			if (!PCGSettingsInterface)
 			{
 				continue;
 			}
 
-			if (PCGSettingsInterface->ExecutionMode == InExecutionMode)
+			//const bool bNewCheckState = CheckState != ECheckBoxState::Checked;
+			if (PCGSettingsInterface->bEnabled != bNewCheckState)
 			{
-				return true;
+				PCGSettingsInterface->Modify();
+				PCGSettingsInterface->bEnabled = bNewCheckState;
+				PCGNode->OnNodeChangedDelegate.Broadcast(PCGNode, EPCGChangeType::Settings);
 			}
 		}
 	}
+}
 
-	return false;
+ECheckBoxState FPCGEditor::GetEnabledCheckState() const
+{	
+	if (GraphEditorWidget.IsValid())
+	{
+		bool bAllEnabled = true;
+		bool bAnyEnabled = false;
+		
+		for (UObject* Object : GraphEditorWidget->GetSelectedNodes())
+		{
+			const UPCGEditorGraphNodeBase* PCGEditorGraphNode = Cast<UPCGEditorGraphNodeBase>(Object);
+			if (!PCGEditorGraphNode)
+			{
+				continue;
+			}
+
+			const UPCGNode* PCGNode = PCGEditorGraphNode->GetPCGNode();
+			if (!PCGNode)
+			{
+				continue;
+			}
+
+			const UPCGSettingsInterface* PCGSettingsInterface = PCGNode->GetSettingsInterface();
+			if (!PCGSettingsInterface)
+			{
+				continue;
+			}
+
+			bAllEnabled &= PCGSettingsInterface->bEnabled;
+			bAnyEnabled |= PCGSettingsInterface->bEnabled;
+		}
+
+		if (bAllEnabled)
+		{
+			return ECheckBoxState::Checked;
+		}
+		else if (bAnyEnabled)
+		{
+			return ECheckBoxState::Undetermined; 
+		}
+	}
+	
+	return ECheckBoxState::Unchecked;
 }
 
-void FPCGEditor::OnSetExecutionMode(EPCGSettingsExecutionMode InExecutionMode)
+void FPCGEditor::OnToggleDebug()
 {
-	// todo should perhaps make a utility function to get all selected nodes.
+	const ECheckBoxState CheckState = GetDebugCheckState();
+	const bool bNewCheckState = !(CheckState != ECheckBoxState::Unchecked);
+	
 	if (GraphEditorWidget.IsValid())
 	{
-		const FScopedTransaction Transaction(*FPCGEditorCommon::ContextIdentifier, LOCTEXT("PCGEditorExecutionModeTransactionMessage", "PCG Editor: ExecutionMode"), nullptr);
+		const FScopedTransaction Transaction(*FPCGEditorCommon::ContextIdentifier, LOCTEXT("PCGEditorToggleDebugTransactionMessage", "PCG Editor: Toggle Debug Nodes"), nullptr);
 		
 		for (UObject* Object : GraphEditorWidget->GetSelectedNodes())
 		{
@@ -1270,16 +1309,60 @@ void FPCGEditor::OnSetExecutionMode(EPCGSettingsExecutionMode InExecutionMode)
 				continue;
 			}
 
-			if (PCGSettingsInterface->ExecutionMode != InExecutionMode)
+			if (PCGSettingsInterface->bDebug != bNewCheckState)
 			{
 				PCGSettingsInterface->Modify();
-				PCGSettingsInterface->ExecutionMode = InExecutionMode;
+				PCGSettingsInterface->bDebug = bNewCheckState;
 				PCGNode->OnNodeChangedDelegate.Broadcast(PCGNode, EPCGChangeType::Settings);
 			}
 		}
 	}
 }
 
+ECheckBoxState FPCGEditor::GetDebugCheckState() const
+{
+	if (GraphEditorWidget.IsValid())
+	{
+		bool bAllDebug = true;
+		bool bAnyDebug = false;
+		
+		for (UObject* Object : GraphEditorWidget->GetSelectedNodes())
+		{
+			const UPCGEditorGraphNodeBase* PCGEditorGraphNode = Cast<UPCGEditorGraphNodeBase>(Object);
+			if (!PCGEditorGraphNode)
+			{
+				continue;
+			}
+
+			const UPCGNode* PCGNode = PCGEditorGraphNode->GetPCGNode();
+			if (!PCGNode)
+			{
+				continue;
+			}
+
+			const UPCGSettingsInterface* PCGSettingsInterface = PCGNode->GetSettingsInterface();
+			if (!PCGSettingsInterface)
+			{
+				continue;
+			}
+
+			bAllDebug &= PCGSettingsInterface->bDebug;
+			bAnyDebug |= PCGSettingsInterface->bDebug;
+		}
+
+		if (bAllDebug)
+		{
+			return ECheckBoxState::Checked;
+		}
+		else if (bAnyDebug)
+		{
+			return ECheckBoxState::Undetermined;
+		}
+	}
+	
+	return ECheckBoxState::Unchecked;
+}
+
 void FPCGEditor::SelectAllNodes()
 {
 	if (GraphEditorWidget.IsValid())
@@ -1908,7 +1991,8 @@ bool FPCGEditor::IsVisibleProperty(const FPropertyAndParent& InPropertyAndParent
 	}
 
 	// Hide debug settings from the setting when showing the instance settings.
-	if (InPropertyAndParent.Property.GetFName() == GET_MEMBER_NAME_CHECKED(UPCGSettings, ExecutionMode) ||
+	if (InPropertyAndParent.Property.GetFName() == GET_MEMBER_NAME_CHECKED(UPCGSettings, bEnabled) ||
+		InPropertyAndParent.Property.GetFName() == GET_MEMBER_NAME_CHECKED(UPCGSettings, bDebug) ||
 		(InPropertyAndParent.ParentProperties.Num() >= 2 && 
 			InPropertyAndParent.ParentProperties[InPropertyAndParent.ParentProperties.Num() - 2]->GetFName() == GET_MEMBER_NAME_CHECKED(UPCGSettings, DebugSettings)))
 	{
diff --git a/Engine/Plugins/Experimental/PCG/Source/PCGEditor/Private/PCGEditor.h b/Engine/Plugins/Experimental/PCG/Source/PCGEditor/Private/PCGEditor.h
--- a/Engine/Plugins/Experimental/PCG/Source/PCGEditor/Private/PCGEditor.h
+++ b/Engine/Plugins/Experimental/PCG/Source/PCGEditor/Private/PCGEditor.h
@@ -114,6 +114,16 @@ private:
 	/** Stop inspecting the current inspected node */
 	void OnStopInspectNode();
 
+	/** Toggle node enabled state for selected nodes */
+	void OnToggleEnabled();
+	/** Whether selected nodes are enabled or not */
+	ECheckBoxState GetEnabledCheckState() const;
+	
+	/** Toggle node debug state for selected nodes */
+	void OnToggleDebug();
+	/** Whether selected nodes are being debugged or not */
+	ECheckBoxState GetDebugCheckState() const;
+
 	/** Cancels the current execution of the selected graph */
 	void OnCancelExecution_Clicked();
 
@@ -132,12 +142,6 @@ private:
 	/** Whether the PCG object being edited is opened in details view or not */
 	bool IsEditGraphSettingsToggled() const;
 
-	/** Whether or not an execution mode is active for the selected nodes */
-	bool IsExecutionModeActive(EPCGSettingsExecutionMode InExecutionMode) const;
-
-	/** Set execution mode for selected nodes */
-	void OnSetExecutionMode(EPCGSettingsExecutionMode InExecutionMode);
-
 	/** Select every node in the graph */
 	void SelectAllNodes();
 	/** Whether we can select every node */
diff --git a/Engine/Plugins/Experimental/PCG/Source/PCGEditor/Private/PCGEditorCommands.cpp b/Engine/Plugins/Experimental/PCG/Source/PCGEditor/Private/PCGEditorCommands.cpp
--- a/Engine/Plugins/Experimental/PCG/Source/PCGEditor/Private/PCGEditorCommands.cpp
+++ b/Engine/Plugins/Experimental/PCG/Source/PCGEditor/Private/PCGEditorCommands.cpp
@@ -29,11 +29,9 @@ void FPCGEditorCommands::RegisterCommands()
 	UI_COMMAND(EditGraphSettings, "Graph Settings", "Edit the graph settings.", EUserInterfaceActionType::ToggleButton, FInputChord());
 	UI_COMMAND(StartInspectNode, "Start Inspecting Node", "Start Inspecting Node", EUserInterfaceActionType::Button, FInputChord(EKeys::I));
 	UI_COMMAND(StopInspectNode, "Stop Inspecting Node", "Stop Inspecting Node", EUserInterfaceActionType::Button, FInputChord(EModifierKey::Shift, EKeys::I));
-	UI_COMMAND(ExecutionModeEnabled, "Execution Mode Enabled", "Set the execution mode for this node to Enabled.", EUserInterfaceActionType::RadioButton, FInputChord(EKeys::E));
-	UI_COMMAND(ExecutionModeDebug, "Execution Mode Debug", "Set the execution mode for this node to Debug.", EUserInterfaceActionType::RadioButton, FInputChord(EKeys::D));
-	UI_COMMAND(ExecutionModeIsolated, "Execution Mode Isolated", "Set the execution mode for this node to Isolated.", EUserInterfaceActionType::RadioButton, FInputChord());
-	UI_COMMAND(ExecutionModeDisabled, "Execution Mode Disabled", "Set the execution mode for this node to Disabled.", EUserInterfaceActionType::RadioButton, FInputChord(EModifierKey::Shift, EKeys::E));
 	UI_COMMAND(CancelExecution, "Cancel Execution", "Cancels the execution of the current graph", EUserInterfaceActionType::Button, FInputChord(EModifierKey::Control, EKeys::Escape));
+	UI_COMMAND(ToggleEnabled, "Toggle Enabled", "Toggle node enabled state for selected nodes.", EUserInterfaceActionType::ToggleButton, FInputChord(EKeys::E));
+	UI_COMMAND(ToggleDebug, "Toggle Debug", "Toggle node debug state for selected nodes", EUserInterfaceActionType::ToggleButton, FInputChord(EKeys::D));
 }
 
 #undef LOCTEXT_NAMESPACE
diff --git a/Engine/Plugins/Experimental/PCG/Source/PCGEditor/Private/PCGEditorCommands.h b/Engine/Plugins/Experimental/PCG/Source/PCGEditor/Private/PCGEditorCommands.h
--- a/Engine/Plugins/Experimental/PCG/Source/PCGEditor/Private/PCGEditorCommands.h
+++ b/Engine/Plugins/Experimental/PCG/Source/PCGEditor/Private/PCGEditorCommands.h
@@ -25,9 +25,7 @@ public:
 	TSharedPtr<FUICommandInfo> EditGraphSettings;
 	TSharedPtr<FUICommandInfo> StartInspectNode;
 	TSharedPtr<FUICommandInfo> StopInspectNode;
-	TSharedPtr<FUICommandInfo> ExecutionModeEnabled;
-	TSharedPtr<FUICommandInfo> ExecutionModeDebug;
-	TSharedPtr<FUICommandInfo> ExecutionModeIsolated;
-	TSharedPtr<FUICommandInfo> ExecutionModeDisabled;
 	TSharedPtr<FUICommandInfo> CancelExecution;
+	TSharedPtr<FUICommandInfo> ToggleEnabled;
+	TSharedPtr<FUICommandInfo> ToggleDebug;
 };
diff --git a/Engine/Plugins/Experimental/PCG/Source/PCGEditor/Private/PCGEditorGraphNodeBase.cpp b/Engine/Plugins/Experimental/PCG/Source/PCGEditor/Private/PCGEditorGraphNodeBase.cpp
--- a/Engine/Plugins/Experimental/PCG/Source/PCGEditor/Private/PCGEditorGraphNodeBase.cpp
+++ b/Engine/Plugins/Experimental/PCG/Source/PCGEditor/Private/PCGEditorGraphNodeBase.cpp
@@ -38,7 +38,7 @@ void UPCGEditorGraphNodeBase::Construct(UPCGNode* InPCGNode)
 
 	if (const UPCGSettingsInterface* PCGSettingsInterface = InPCGNode->GetSettingsInterface())
 	{
-		const ENodeEnabledState NewEnabledState = (PCGSettingsInterface->ExecutionMode == EPCGSettingsExecutionMode::Disabled) ? ENodeEnabledState::Disabled : ENodeEnabledState::Enabled;
+		const ENodeEnabledState NewEnabledState = !PCGSettingsInterface->bEnabled ? ENodeEnabledState::Disabled : ENodeEnabledState::Enabled;
 		SetEnabledState(NewEnabledState);
 	}
 }
@@ -79,19 +79,11 @@ void UPCGEditorGraphNodeBase::GetNodeContextMenuActions(UToolMenu* Menu, class U
 
 	{
 		FToolMenuSection& Section = Menu->AddSection("EdGraphSchemaNodeActions", LOCTEXT("NodeActionsHeader", "Node Actions"));
+		Section.AddMenuEntry(FPCGEditorCommands::Get().ToggleEnabled, LOCTEXT("ToggleEnabledLabel", "Enable"));
+		Section.AddMenuEntry(FPCGEditorCommands::Get().ToggleDebug, LOCTEXT("ToggleDebugLabel", "Debug"));
 		Section.AddMenuEntry(bIsInspected ? FPCGEditorCommands::Get().StopInspectNode : FPCGEditorCommands::Get().StartInspectNode);
 		Section.AddMenuEntry(FGraphEditorCommands::Get().BreakNodeLinks);
 		Section.AddMenuEntry(FPCGEditorCommands::Get().CollapseNodes);
-
-		Section.AddSubMenu("ExecutionMode", LOCTEXT("ExecutionModeHeader", "Execution Mode"), FText(), FNewToolMenuDelegate::CreateLambda([](UToolMenu* ExecutionModeMenu)
-		{
-			FToolMenuSection& SubSection = ExecutionModeMenu->AddSection("ExecutionMode");
-			SubSection.AddMenuEntry(FPCGEditorCommands::Get().ExecutionModeEnabled, LOCTEXT("ExecutionModeEnabledLabel", "Enabled"));
-			SubSection.AddMenuEntry(FPCGEditorCommands::Get().ExecutionModeDebug, LOCTEXT("ExecutionModeDebugLabel", "Debug"));
-			SubSection.AddMenuEntry(FPCGEditorCommands::Get().ExecutionModeDisabled, LOCTEXT("ExecutionModeDisabledLabel", "Disabled"));
-			SubSection.AddMenuEntry(FPCGEditorCommands::Get().ExecutionModeIsolated, LOCTEXT("ExecutionModeIsolatedLabel", "Isolated"));
-		}));
-
 		Section.AddMenuEntry(FPCGEditorCommands::Get().ExportNodes);
 		Section.AddMenuEntry(FPCGEditorCommands::Get().ConvertToStandaloneNodes);
 	}
@@ -253,7 +245,7 @@ void UPCGEditorGraphNodeBase::OnNodeChanged(UPCGNode* InNode, EPCGChangeType Cha
 		{
 			if (const UPCGSettingsInterface* PCGSettingsInterface = InNode->GetSettingsInterface())
 			{
-				const ENodeEnabledState NewEnabledState = (PCGSettingsInterface->ExecutionMode == EPCGSettingsExecutionMode::Disabled) ? ENodeEnabledState::Disabled : ENodeEnabledState::Enabled;
+				const ENodeEnabledState NewEnabledState = PCGSettingsInterface->bEnabled ? ENodeEnabledState::Enabled : ENodeEnabledState::Disabled;
 				if (NewEnabledState != GetDesiredEnabledState())
 				{
 					SetEnabledState(NewEnabledState);
@@ -261,7 +253,10 @@ void UPCGEditorGraphNodeBase::OnNodeChanged(UPCGNode* InNode, EPCGChangeType Cha
 			}
 		}
 		
-		ReconstructNode();
+		if (!!(ChangeType & (EPCGChangeType::Structural | EPCGChangeType::Cosmetic)))
+		{
+			ReconstructNode();
+		}
 	}
 }
 
@@ -326,11 +321,6 @@ FLinearColor UPCGEditorGraphNodeBase::GetNodeTitleColor() const
 		const UPCGSettingsInterface* PCGSettingsInterface = PCGNode->GetSettingsInterface();
 		const UPCGSettings* PCGSettings = PCGSettingsInterface ? PCGSettingsInterface->GetSettings() : nullptr;
 
-		if (PCGSettingsInterface && PCGSettingsInterface->ExecutionMode == EPCGSettingsExecutionMode::Isolated)
-		{
-			return GetDefault<UPCGEditorSettings>()->IsolatedNodeColor;
-		}
-
 		if (PCGNode->NodeTitleColor != FLinearColor::White)
 		{
 			return PCGNode->NodeTitleColor;
@@ -364,10 +354,6 @@ FLinearColor UPCGEditorGraphNodeBase::GetNodeBodyTintColor() const
 			{
 				return GetDefault<UPCGEditorSettings>()->InstancedNodeBodyTintColor;
 			}
-			else if (PCGSettingsInterface->ExecutionMode == EPCGSettingsExecutionMode::Isolated)
-			{
-				return GetDefault<UPCGEditorSettings>()->IsolatedNodeColor;
-			}
 		}
 	}
 
diff --git a/Engine/Plugins/Experimental/PCG/Source/PCGEditor/Private/PCGEditorSettings.cpp b/Engine/Plugins/Experimental/PCG/Source/PCGEditor/Private/PCGEditorSettings.cpp
--- a/Engine/Plugins/Experimental/PCG/Source/PCGEditor/Private/PCGEditorSettings.cpp
+++ b/Engine/Plugins/Experimental/PCG/Source/PCGEditor/Private/PCGEditorSettings.cpp
@@ -10,7 +10,6 @@ UPCGEditorSettings::UPCGEditorSettings(const FObjectInitializer& ObjectInitializ
 {
 	DefaultNodeColor = FLinearColor(0.4f, 0.62f, 1.0f);
 	InstancedNodeBodyTintColor = FLinearColor(0.5f, 0.5f, 0.5f);
-	IsolatedNodeColor = FLinearColor(0.0f, 0.0f, 1.0f);
 	InputOutputNodeColor = FLinearColor(1.0f, 0.0f, 0.0f);
 	SetOperationNodeColor = FLinearColor(0.8f, 0.2f, 0.8f);
 	DensityOperationNodeColor = FLinearColor(0.6f, 1.0f, 0.6f);
diff --git a/Engine/Plugins/Experimental/PCG/Source/PCGEditor/Private/SPCGEditorGraphNode.cpp b/Engine/Plugins/Experimental/PCG/Source/PCGEditor/Private/SPCGEditorGraphNode.cpp
--- a/Engine/Plugins/Experimental/PCG/Source/PCGEditor/Private/SPCGEditorGraphNode.cpp
+++ b/Engine/Plugins/Experimental/PCG/Source/PCGEditor/Private/SPCGEditorGraphNode.cpp
@@ -87,7 +87,7 @@ void SPCGEditorGraphNode::GetOverlayBrushes(bool bSelected, const FVector2D Widg
 
 	if (const UPCGNode* PCGNode = PCGEditorGraphNode->GetPCGNode())
 	{
-		if (PCGNode->GetSettingsInterface() && PCGNode->GetSettingsInterface()->ExecutionMode == EPCGSettingsExecutionMode::Debug)
+		if (PCGNode->GetSettingsInterface() && PCGNode->GetSettingsInterface()->bDebug)
 		{
 			FOverlayBrushInfo BrushInfo;
 			BrushInfo.Brush = DebugBrush;
diff --git a/Engine/Plugins/Experimental/PCG/Source/PCGEditor/Public/PCGEditorSettings.h b/Engine/Plugins/Experimental/PCG/Source/PCGEditor/Public/PCGEditorSettings.h
--- a/Engine/Plugins/Experimental/PCG/Source/PCGEditor/Public/PCGEditorSettings.h
+++ b/Engine/Plugins/Experimental/PCG/Source/PCGEditor/Public/PCGEditorSettings.h
@@ -29,10 +29,6 @@ public:
 	UPROPERTY(EditAnywhere, config, Category = Node, meta = (HideAlphaChannel))
 	FLinearColor InstancedNodeBodyTintColor;
 
-	/** Isolated node color */
-	UPROPERTY(EditAnywhere, config, Category = Node, meta = (HideAlphaChannel))
-	FLinearColor IsolatedNodeColor;
-
 	/** Color used for input & output nodes */
 	UPROPERTY(EditAnywhere, config, Category = Node, meta = (HideAlphaChannel))
 	FLinearColor InputOutputNodeColor;
