Normalized behavior of elements vs output connections in editor / fixed some caching issues.
Improved point filter in/out cases with non-existent attributes

diff --git a/Engine/Plugins/Experimental/PCG/Source/PCG/Private/Elements/PCGAttributeGetFromPointIndexElement.cpp b/Engine/Plugins/Experimental/PCG/Source/PCG/Private/Elements/PCGAttributeGetFromPointIndexElement.cpp
--- a/Engine/Plugins/Experimental/PCG/Source/PCG/Private/Elements/PCGAttributeGetFromPointIndexElement.cpp
+++ b/Engine/Plugins/Experimental/PCG/Source/PCG/Private/Elements/PCGAttributeGetFromPointIndexElement.cpp
@@ -80,8 +80,10 @@ bool FPCGAttributeGetFromPointIndexElement::ExecuteInternal(FPCGContext* Context
 	TArray<FPCGTaggedData>& Outputs = Context->OutputData.TaggedData;
 	const FPCGPoint& Point = PointData->GetPoints()[Index];
 
-	// Only create a point if we are connected, or in editor with the component being inspected.
-	if (Context->IsOutputConnectedOrInspecting(PCGAttributeGetFromPointIndexConstants::OutputPointLabel))
+#if !WITH_EDITOR
+	// Eschew output creation only in non-editor builds
+	if(Context->Node && Context->Node->IsOutputPinConnected(PCGAttributeGetFromPointIndexConstants::OutputPointLabel))
+#endif
 	{
 		UPCGPointData* OutputPointData = NewObject<UPCGPointData>();
 		OutputPointData->InitializeFromData(PointData);
diff --git a/Engine/Plugins/Experimental/PCG/Source/PCG/Private/Elements/PCGAttributeSelectElement.cpp b/Engine/Plugins/Experimental/PCG/Source/PCG/Private/Elements/PCGAttributeSelectElement.cpp
--- a/Engine/Plugins/Experimental/PCG/Source/PCG/Private/Elements/PCGAttributeSelectElement.cpp
+++ b/Engine/Plugins/Experimental/PCG/Source/PCG/Private/Elements/PCGAttributeSelectElement.cpp
@@ -404,7 +404,11 @@ bool FPCGAttributeSelectElement::ExecuteInternal(FPCGContext* Context) const
 	Output.Data = OutputParamData;
 	Output.Pin = PCGAttributeSelectConstants::OutputAttributeLabel;
 
-	if (PointData && Context->IsOutputConnectedOrInspecting(PCGAttributeSelectConstants::OutputPointLabel))
+#if WITH_EDITOR
+	if(PointData)
+#else
+	if(PointData && Context->Node && Context->Node->IsOutputPinConnected(PCGAttributeSelectConstants::OutputPointLabel))
+#endif
 	{
 		UPCGPointData* OutputPointData = NewObject<UPCGPointData>();
 		OutputPointData->InitializeFromData(PointData);
diff --git a/Engine/Plugins/Experimental/PCG/Source/PCG/Private/Elements/PCGPointFilter.cpp b/Engine/Plugins/Experimental/PCG/Source/PCG/Private/Elements/PCGPointFilter.cpp
--- a/Engine/Plugins/Experimental/PCG/Source/PCG/Private/Elements/PCGPointFilter.cpp
+++ b/Engine/Plugins/Experimental/PCG/Source/PCG/Private/Elements/PCGPointFilter.cpp
@@ -608,11 +608,13 @@ bool FPCGPointFilterElement::ExecuteInternal(FPCGContext* Context) const
 	const bool bHasInFilterOutputPin = Context->Node && Context->Node->IsOutputPinConnected(PCGPointFilterConstants::InFilterLabel);
 	const bool bHasOutsideFilterOutputPin = Context->Node && Context->Node->IsOutputPinConnected(PCGPointFilterConstants::OutFilterLabel);
 
-	// Early out
+#if !WITH_EDITOR
+	// Early out - only in non-editor builds, otherwise we will potentially poison the cache, since it is input-driven
 	if (!bHasInFilterOutputPin && !bHasOutsideFilterOutputPin)
 	{
 		return true;
 	}
+#endif
 
 	const UPCGPointFilterSettings* Settings = Context->GetInputSettings<UPCGPointFilterSettings>();
 	check(Settings);
@@ -712,14 +714,23 @@ bool FPCGPointFilterElement::ExecuteInternal(FPCGContext* Context) const
 		// Additional validation
 		if (TargetFilterType == EPCGPointTargetFilterType::Metadata)
 		{
+			bool bExitToOutFilter = false;
+
 			if (!TargetMetadata)
 			{
-				PCGE_LOG(Error, "Target data to filter has no metadata which is required to filter by metadata");
-				continue;
+				PCGE_LOG(Warning, "Target data to filter has no metadata which is required to filter by metadata");
+				bExitToOutFilter = true;
 			}
 			else if (!TargetMetadata->GetConstAttribute(LocalTargetAttributeName))
 			{
-				PCGE_LOG(Error, "Target metadata does not have the %s attribute", *LocalTargetAttributeName.ToString());
+				PCGE_LOG(Warning, "Target metadata does not have the %s attribute", *LocalTargetAttributeName.ToString());
+				bExitToOutFilter = true;
+			}
+
+			if (bExitToOutFilter)
+			{
+				FPCGTaggedData& OutFilterTaggedData = Outputs.Add_GetRef(Input);
+				OutFilterTaggedData.Pin = PCGPointFilterConstants::OutFilterLabel;
 				continue;
 			}
 		}
@@ -728,14 +739,23 @@ bool FPCGPointFilterElement::ExecuteInternal(FPCGContext* Context) const
 		
 		if (ThresholdFilterType == EPCGPointThresholdType::Metadata)
 		{
+			bool bExitToInFilter = false;
+
 			if (!ThresholdMetadata)
 			{
-				PCGE_LOG(Error, "Filter data has no metadata which is required to filter by metadata");
-				continue;
+				PCGE_LOG(Warning, "Filter data has no metadata which is required to filter by metadata");
+				bExitToInFilter = true;
 			}
 			else if (!ThresholdMetadata->GetConstAttribute(LocalThresholdAttributeName))
 			{
-				PCGE_LOG(Error, "Filter metadata does not have the %s attribute", *LocalThresholdAttributeName.ToString());
+				PCGE_LOG(Warning, "Filter metadata does not have the %s attribute", *LocalThresholdAttributeName.ToString());
+				bExitToInFilter = true;
+			}
+
+			if (bExitToInFilter)
+			{
+				FPCGTaggedData& InFilterTaggedData = Outputs.Add_GetRef(Input);
+				InFilterTaggedData.Pin = PCGPointFilterConstants::InFilterLabel;
 				continue;
 			}
 		}
diff --git a/Engine/Plugins/Experimental/PCG/Source/PCG/Private/Elements/PCGPropertyToParamData.cpp b/Engine/Plugins/Experimental/PCG/Source/PCG/Private/Elements/PCGPropertyToParamData.cpp
--- a/Engine/Plugins/Experimental/PCG/Source/PCG/Private/Elements/PCGPropertyToParamData.cpp
+++ b/Engine/Plugins/Experimental/PCG/Source/PCG/Private/Elements/PCGPropertyToParamData.cpp
@@ -71,12 +71,15 @@ bool FPCGPropertyToParamDataElement::ExecuteInternal(FPCGContext* Context) const
 		return true;
 	}
 
+#if !WITH_EDITOR
 	// If we have no output connected, nothing to do
+	// Optimization possibly only in non-editor builds, otherwise we could poison the input-driven cache
 	if (!Context->Node || !Context->Node->IsOutputPinConnected(PCGPinConstants::DefaultOutputLabel))
 	{
 		PCGE_LOG(Verbose, "Node is not connected, nothing to do");
 		return true;
 	}
+#endif
 
 	// First find the actor depending on the selection
 	UPCGComponent* OriginalComponent = UPCGBlueprintHelpers::GetOriginalComponent(*Context);
diff --git a/Engine/Plugins/Experimental/PCG/Source/PCG/Private/Elements/PCGStaticMeshSpawner.cpp b/Engine/Plugins/Experimental/PCG/Source/PCG/Private/Elements/PCGStaticMeshSpawner.cpp
--- a/Engine/Plugins/Experimental/PCG/Source/PCG/Private/Elements/PCGStaticMeshSpawner.cpp
+++ b/Engine/Plugins/Experimental/PCG/Source/PCG/Private/Elements/PCGStaticMeshSpawner.cpp
@@ -51,7 +51,12 @@ bool FPCGStaticMeshSpawnerElement::PrepareDataInternal(FPCGContext* InContext) c
 	// Forward any non-input data
 	Outputs.Append(Context->InputData.GetAllSettings());
 
-	const bool bOutputPinConnected = Context->Node && Context->Node->IsOutputPinConnected(PCGPinConstants::DefaultOutputLabel);
+#if WITH_EDITOR
+	// In editor, we always want to generate this data for inspection & to prevent caching issues
+	const bool bGenerateOutput = true;
+#else
+	const bool bGenerateOutput = Context->Node && Context->Node->IsOutputPinConnected(PCGPinConstants::DefaultOutputLabel);
+#endif
 
 	for (const FPCGTaggedData& Input : Inputs)
 	{
@@ -73,7 +78,7 @@ bool FPCGStaticMeshSpawnerElement::PrepareDataInternal(FPCGContext* InContext) c
 
 		UPCGPointData* OutputPointData = nullptr;
 
-		if (bOutputPinConnected || Settings->bForceConnectOutput)
+		if (bGenerateOutput)
 		{
 			FPCGTaggedData& Output = Outputs.Add_GetRef(Input);
 
diff --git a/Engine/Plugins/Experimental/PCG/Source/PCG/Private/PCGComponent.cpp b/Engine/Plugins/Experimental/PCG/Source/PCG/Private/PCGComponent.cpp
--- a/Engine/Plugins/Experimental/PCG/Source/PCG/Private/PCGComponent.cpp
+++ b/Engine/Plugins/Experimental/PCG/Source/PCG/Private/PCGComponent.cpp
@@ -2252,7 +2252,9 @@ void UPCGComponent::DirtyCacheForAllTrackedTags()
 
 bool UPCGComponent::GraphUsesLandscapePin() const
 {
-	return Graph && Graph->GetInputNode()->IsOutputPinConnected(PCGInputOutputConstants::DefaultLandscapeLabel);
+	return Graph &&
+		(Graph->GetInputNode()->IsOutputPinConnected(PCGInputOutputConstants::DefaultLandscapeLabel) ||
+		Graph->GetInputNode()->IsOutputPinConnected(PCGInputOutputConstants::DefaultLandscapeHeightLabel));
 }
 
 #endif // WITH_EDITOR
diff --git a/Engine/Plugins/Experimental/PCG/Source/PCG/Private/PCGContext.cpp b/Engine/Plugins/Experimental/PCG/Source/PCG/Private/PCGContext.cpp
--- a/Engine/Plugins/Experimental/PCG/Source/PCG/Private/PCGContext.cpp
+++ b/Engine/Plugins/Experimental/PCG/Source/PCG/Private/PCGContext.cpp
@@ -41,16 +41,6 @@ bool FPCGContext::ShouldStop() const
 	return FPlatformTime::Seconds() > EndTime;
 }
 
-bool FPCGContext::IsOutputConnectedOrInspecting(FName PinLabel) const
-{
-	const bool bOutputConnected = Node && Node->IsOutputPinConnected(PinLabel);
-#if WITH_EDITOR
-	return bOutputConnected || (!PCGHelpers::IsRuntimeOrPIE() && SourceComponent.IsValid() && SourceComponent->IsInspecting());
-#else
-	return bOutputConnected;
-#endif // WITH_EDITOR
-}
-
 const UPCGSettingsInterface* FPCGContext::GetInputSettingsInterface() const
 {
 	if (Node)
diff --git a/Engine/Plugins/Experimental/PCG/Source/PCG/Private/Tests/Elements/PCGStaticMeshSpawnerTest.cpp b/Engine/Plugins/Experimental/PCG/Source/PCG/Private/Tests/Elements/PCGStaticMeshSpawnerTest.cpp
--- a/Engine/Plugins/Experimental/PCG/Source/PCG/Private/Tests/Elements/PCGStaticMeshSpawnerTest.cpp
+++ b/Engine/Plugins/Experimental/PCG/Source/PCG/Private/Tests/Elements/PCGStaticMeshSpawnerTest.cpp
@@ -143,11 +143,6 @@ void TestMeshSelectorByAttribute(
 	UPCGMeshSelectorByAttribute* MeshSelector = CastChecked<UPCGMeshSelectorByAttribute>(Settings->MeshSelectorInstance);
 	MeshSelector->AttributeName = AttributeName;
 
-	if (bValidateOutput)
-	{
-		Settings->bForceConnectOutput = true;
-	}
-
 	// initialize and execute the StaticMeshSpawner
 	FPCGElementPtr StaticMeshSpawner = Settings->GetElement();
 	FPCGContext* Context = StaticMeshSpawner->Initialize(TestData.InputData, TestData.TestPCGComponent, nullptr);
@@ -272,11 +267,6 @@ void TestMeshSelectorWeighted(
 	UPCGMeshSelectorWeighted* MeshSelector = CastChecked<UPCGMeshSelectorWeighted>(Settings->MeshSelectorInstance);
 	MeshSelector->MeshEntries = Entries;
 
-	if (bValidateOutput)
-	{
-		Settings->bForceConnectOutput = true;
-	}
-
 	// initialize and execute the StaticMeshSpawner
 	FPCGElementPtr StaticMeshSpawner = Settings->GetElement();
 	FPCGContext* Context = StaticMeshSpawner->Initialize(TestData.InputData, TestData.TestPCGComponent, nullptr);
@@ -365,11 +355,6 @@ void TestMeshSelectorWeightedByCategory(
 	MeshSelector->CategoryAttribute = AttributeName;
 	MeshSelector->Entries = Entries;
 
-	if (bValidateOutput)
-	{
-		Settings->bForceConnectOutput = true;
-	}
-
 	// initialize and execute the StaticMeshSpawner
 	FPCGElementPtr StaticMeshSpawner = Settings->GetElement();
 	FPCGContext* Context = StaticMeshSpawner->Initialize(TestData.InputData, TestData.TestPCGComponent, nullptr);
diff --git a/Engine/Plugins/Experimental/PCG/Source/PCG/Public/Elements/PCGStaticMeshSpawner.h b/Engine/Plugins/Experimental/PCG/Source/PCG/Public/Elements/PCGStaticMeshSpawner.h
--- a/Engine/Plugins/Experimental/PCG/Source/PCG/Public/Elements/PCGStaticMeshSpawner.h
+++ b/Engine/Plugins/Experimental/PCG/Source/PCG/Public/Elements/PCGStaticMeshSpawner.h
@@ -88,8 +88,6 @@ public:
 	UPROPERTY(meta=(DeprecatedProperty, DeprecationMessage="Use MeshSelectorType and MeshSelectorInstance instead."))
 	TArray<FPCGStaticMeshSpawnerEntry> Meshes_DEPRECATED;
 
-	bool bForceConnectOutput = false;
-
 protected:
 	void RefreshMeshSelector();
 	void RefreshInstancePacker();
diff --git a/Engine/Plugins/Experimental/PCG/Source/PCG/Public/PCGContext.h b/Engine/Plugins/Experimental/PCG/Source/PCG/Public/PCGContext.h
--- a/Engine/Plugins/Experimental/PCG/Source/PCG/Public/PCGContext.h
+++ b/Engine/Plugins/Experimental/PCG/Source/PCG/Public/PCGContext.h
@@ -76,9 +76,4 @@ struct PCG_API FPCGContext
 	FString GetTaskName() const;
 	FString GetComponentName() const;
 	bool ShouldStop() const;
-
-	/** Helper function to return if the output is connected, or if we need to force the connection because the component is being inspected. 
-	*   Useful to avoid creating output when it is not necessary.
-	*/
-	bool IsOutputConnectedOrInspecting(FName PinLabel) const;
 };
