"optimization" to PCG surface sampler. This made no measurable difference as it's very minor, but just removes a potential performance gotcha as the optimizer is not always able to deal with member functions that take references. More details in the CR

diff --git a/Engine/Plugins/Experimental/PCG/Source/PCG/Private/Elements/PCGSurfaceSampler.cpp b/Engine/Plugins/Experimental/PCG/Source/PCG/Private/Elements/PCGSurfaceSampler.cpp
--- a/Engine/Plugins/Experimental/PCG/Source/PCG/Private/Elements/PCGSurfaceSampler.cpp
+++ b/Engine/Plugins/Experimental/PCG/Source/PCG/Private/Elements/PCGSurfaceSampler.cpp
@@ -87,13 +87,12 @@ namespace PCGSurfaceSampler
 		return true;
 	}
 
-	void FSurfaceSamplerSettings::ComputeCellIndices(int32 Index, int32& CellX, int32& CellY) const
+	FIntVector2 FSurfaceSamplerSettings::ComputeCellIndices(int32 Index) const
 	{
 		check(Index >= 0 && Index < CellCount);
 		const int32 CellCountX = 1 + CellMaxX - CellMinX;
 
-		CellX = CellMinX + (Index % CellCountX);
-		CellY = CellMinY + (Index / CellCountX);
+		return FIntVector2(CellMinX + (Index % CellCountX), CellMinY + (Index / CellCountX));
 	}
 
 	UPCGPointData* SampleSurface(FPCGContext* Context, const UPCGSpatialData* SpatialInput, const FSurfaceSamplerSettings& LoopData)
@@ -112,15 +111,13 @@ namespace PCGSurfaceSampler
 
 		FPCGAsync::AsyncPointProcessing(Context, LoopData.CellCount, SampledPoints, [&LoopData, SampledData, SpatialInput](int32 Index, FPCGPoint& OutPoint)
 		{
-			int32 CellX;
-			int32 CellY;
-			LoopData.ComputeCellIndices(Index, CellX, CellY);
+			const FIntVector2 Indices = LoopData.ComputeCellIndices(Index);
 
-			const FVector::FReal CurrentX = CellX * LoopData.CellSize.X;
-			const FVector::FReal CurrentY = CellY * LoopData.CellSize.Y;
+			const FVector::FReal CurrentX = Indices.X * LoopData.CellSize.X;
+			const FVector::FReal CurrentY = Indices.Y * LoopData.CellSize.Y;
 			const FVector InnerCellSize = LoopData.InnerCellSize;
 
-			FRandomStream RandomSource(PCGHelpers::ComputeSeed(LoopData.Seed, CellX, CellY));
+			FRandomStream RandomSource(PCGHelpers::ComputeSeed(LoopData.Seed, Indices.X, Indices.Y));
 			float Chance = RandomSource.FRand();
 
 			const float Ratio = LoopData.Ratio;
diff --git a/Engine/Plugins/Experimental/PCG/Source/PCG/Public/Elements/PCGSurfaceSampler.h b/Engine/Plugins/Experimental/PCG/Source/PCG/Public/Elements/PCGSurfaceSampler.h
--- a/Engine/Plugins/Experimental/PCG/Source/PCG/Public/Elements/PCGSurfaceSampler.h
+++ b/Engine/Plugins/Experimental/PCG/Source/PCG/Public/Elements/PCGSurfaceSampler.h
@@ -83,7 +83,7 @@ namespace PCGSurfaceSampler
 #endif
 
 		bool Initialize(const UPCGSurfaceSamplerSettings* InSettings, FPCGContext* Context, const FBox& InputBounds);
-		void ComputeCellIndices(int32 Index, int32& CellX, int32& CellY) const;
+		FIntVector2 ComputeCellIndices(int32 Index) const;
 
 		/** Computed values **/
 		FVector InterstitialDistance;
