Added point merge node, supports points only at the moment

diff --git a/Engine/Plugins/Experimental/PCG/Source/PCG/Private/Elements/PCGMergeElement.cpp b/Engine/Plugins/Experimental/PCG/Source/PCG/Private/Elements/PCGMergeElement.cpp
new file mode 100644
--- /dev/null
+++ b/Engine/Plugins/Experimental/PCG/Source/PCG/Private/Elements/PCGMergeElement.cpp
@@ -0,0 +1,134 @@
+// Copyright Epic Games, Inc. All Rights Reserved.
+
+#include "Elements/PCGMergeElement.h"
+
+#include "Data/PCGPointData.h"
+#include "Metadata/PCGMetadata.h"
+#include "Metadata/PCGMetadataAttribute.h"
+
+#include UE_INLINE_GENERATED_CPP_BY_NAME(PCGMergeElement)
+
+#define LOCTEXT_NAMESPACE "PCGMergeElement"
+
+#if WITH_EDITOR
+FText UPCGMergeSettings::GetNodeTooltipText() const
+{
+	return LOCTEXT("MergeNodeTooltip", "Merges multiple data sources into a single data output.");
+}
+#endif
+
+TArray<FPCGPinProperties> UPCGMergeSettings::InputPinProperties() const
+{
+	TArray<FPCGPinProperties> PinProperties;
+	PinProperties.Emplace(PCGPinConstants::DefaultInputLabel, EPCGDataType::Point, /*bAllowMultipleConnections=*/true);
+
+	return PinProperties;
+}
+
+TArray<FPCGPinProperties> UPCGMergeSettings::OutputPinProperties() const
+{
+	TArray<FPCGPinProperties> PinProperties;
+	PinProperties.Emplace(PCGPinConstants::DefaultOutputLabel, EPCGDataType::Point, /*bAllowMultipleConnections=*/false);
+
+	return PinProperties;
+}
+
+FPCGElementPtr UPCGMergeSettings::CreateElement() const
+{
+	return MakeShared<FPCGMergeElement>();
+}
+
+bool FPCGMergeElement::ExecuteInternal(FPCGContext* Context) const
+{
+	TRACE_CPUPROFILER_EVENT_SCOPE(FPCGMergeElement::Execute);
+	check(Context);
+
+	const UPCGMergeSettings* Settings = Context->GetInputSettings<UPCGMergeSettings>();
+	check(Settings);
+
+	const bool bMergeMetadata = Settings->bMergeMetadata;
+
+	TArray<FPCGTaggedData> Sources = Context->InputData.GetInputs();
+	TArray<FPCGTaggedData>& Outputs = Context->OutputData.TaggedData;
+
+	if (Sources.IsEmpty())
+	{
+		return true;
+	}
+
+	UPCGPointData* TargetPointData = nullptr;
+	FPCGTaggedData* TargetTaggedData = nullptr;
+
+	// Prepare data & metadata
+	// Done in two passes for futureproofing - expecting changes in the metadata attribute creation vs. usage in points
+	for (const FPCGTaggedData& Source : Sources)
+	{
+		const UPCGPointData* SourcePointData = Cast<const UPCGPointData>(Source.Data);
+
+		if (!SourcePointData)
+		{
+			PCGE_LOG(Error, "Unsupported data type in merge");
+			continue;
+		}
+
+		if (!TargetPointData)
+		{
+			TargetPointData = NewObject<UPCGPointData>();
+			TargetPointData->InitializeFromData(SourcePointData, nullptr, bMergeMetadata);
+
+			TargetTaggedData = &(Outputs.Emplace_GetRef(Source));
+			TargetTaggedData->Data = TargetPointData;
+		}
+		else
+		{
+			if (bMergeMetadata)
+			{
+				TargetPointData->Metadata->AddAttributes(SourcePointData->Metadata);
+			}
+			
+			check(TargetTaggedData);
+			TargetTaggedData->Tags.Append(Source.Tags); // TODO: only unique? if yes, fix union too
+		}
+	}
+
+	// No valid input types
+	if (!TargetPointData)
+	{
+		return true;
+	}
+
+	bool bIsFirst = true;
+	TArray<FPCGPoint>& TargetPoints = TargetPointData->GetMutablePoints();
+	
+	for (FPCGTaggedData& Source : Sources)
+	{
+		const UPCGPointData* SourcePointData = Cast<const UPCGPointData>(Source.Data);
+
+		if (!SourcePointData)
+		{
+			continue;
+		}
+
+		if (bIsFirst)
+		{
+			check(!bMergeMetadata || TargetPointData->Metadata->GetParent() == SourcePointData->Metadata);
+			TargetPoints.Append(SourcePointData->GetPoints());
+			bIsFirst = false;
+		}
+		else
+		{
+			int32 PointOffset = TargetPoints.Num();
+			TargetPoints.Append(SourcePointData->GetPoints());
+
+			if (bMergeMetadata && TargetPointData->Metadata && SourcePointData->Metadata && !SourcePointData->GetPoints().IsEmpty())
+			{
+				TargetPointData->Metadata->SetPointAttributes(MakeArrayView(SourcePointData->GetPoints()), SourcePointData->Metadata, MakeArrayView(&TargetPoints[PointOffset], SourcePointData->GetPoints().Num()));
+			}
+			// TBD: should we null out the metadata entry keys on the vertices if we didn't setup the metadata?
+		}
+	}
+	
+	return true;
+}
+
+#undef LOCTEXT_NAMESPACE
\ No newline at end of file
diff --git a/Engine/Plugins/Experimental/PCG/Source/PCG/Public/Elements/PCGMergeElement.h b/Engine/Plugins/Experimental/PCG/Source/PCG/Public/Elements/PCGMergeElement.h
new file mode 100644
--- /dev/null
+++ b/Engine/Plugins/Experimental/PCG/Source/PCG/Public/Elements/PCGMergeElement.h
@@ -0,0 +1,41 @@
+// Copyright Epic Games, Inc. All Rights Reserved.
+
+#pragma once
+
+#include "PCGSettings.h"
+#include "PCGElement.h"
+
+#include "PCGMergeElement.generated.h"
+
+/** Merges multiple data sources (currently only points supported) into a single output. */
+UCLASS(BlueprintType, Classgroup = (Procedural))
+class PCG_API UPCGMergeSettings : public UPCGSettings
+{
+	GENERATED_BODY()
+
+public:
+	//~Begin UPCGSettings interface
+#if WITH_EDITOR
+	virtual FName GetDefaultNodeName() const override { return FName(TEXT("Merge")); }
+	virtual FText GetNodeTooltipText() const override;
+	virtual EPCGSettingsType GetType() const override { return EPCGSettingsType::Spatial; }
+#endif
+
+	virtual TArray<FPCGPinProperties> InputPinProperties() const override;
+	virtual TArray<FPCGPinProperties> OutputPinProperties() const override;
+
+protected:
+	virtual FPCGElementPtr CreateElement() const override;
+	//~End UPCGSettings interface
+
+public:
+	/** Controls whether the resulting merge data will have any metadata */
+	UPROPERTY(BlueprintReadWrite, EditAnywhere, Category = Settings)
+	bool bMergeMetadata = true;
+};
+
+class FPCGMergeElement : public FSimplePCGElement
+{
+protected:
+	virtual bool ExecuteInternal(FPCGContext* Context) const override;
+};
\ No newline at end of file
