Fixed issue where operation on settings that would yield changes on the node (incl. pin changes) would not trigger reconstruct node as expected

#rnx

diff --git a/Engine/Plugins/Experimental/PCG/Source/PCGEditor/Private/PCGEditorGraphNodeBase.cpp b/Engine/Plugins/Experimental/PCG/Source/PCGEditor/Private/PCGEditorGraphNodeBase.cpp
--- a/Engine/Plugins/Experimental/PCG/Source/PCGEditor/Private/PCGEditorGraphNodeBase.cpp
+++ b/Engine/Plugins/Experimental/PCG/Source/PCGEditor/Private/PCGEditorGraphNodeBase.cpp
@@ -253,7 +253,7 @@ void UPCGEditorGraphNodeBase::OnNodeChanged(UPCGNode* InNode, EPCGChangeType Cha
 			}
 		}
 		
-		if (!!(ChangeType & (EPCGChangeType::Structural | EPCGChangeType::Cosmetic)))
+		if (!!(ChangeType & (EPCGChangeType::Structural | EPCGChangeType::Node | EPCGChangeType::Edge | EPCGChangeType::Cosmetic)))
 		{
 			ReconstructNode();
 		}
