Small improvements to attribute list view entry naming

#rnx

diff --git a/Engine/Plugins/Experimental/PCG/Source/PCGEditor/Private/SPCGEditorGraphAttributeListView.cpp b/Engine/Plugins/Experimental/PCG/Source/PCGEditor/Private/SPCGEditorGraphAttributeListView.cpp
--- a/Engine/Plugins/Experimental/PCG/Source/PCGEditor/Private/SPCGEditorGraphAttributeListView.cpp
+++ b/Engine/Plugins/Experimental/PCG/Source/PCGEditor/Private/SPCGEditorGraphAttributeListView.cpp
@@ -641,9 +641,17 @@ void SPCGEditorGraphAttributeListView::RefreshDataComboBox()
 		return;
 	}
 
-	for (const FPCGTaggedData& TaggedData : InspectionData->TaggedData)
+	for(int32 TaggedDataIndex = 0; TaggedDataIndex < InspectionData->TaggedData.Num(); ++TaggedDataIndex)
 	{
-		DataComboBoxItems.Add(MakeShared<FName>(TaggedData.Pin));
+		const FPCGTaggedData& TaggedData = InspectionData->TaggedData[TaggedDataIndex];
+		FString ItemName = FString::Format(TEXT("[{0}] {1} - {2}"), { FText::AsNumber(TaggedDataIndex).ToString(), TaggedData.Pin.ToString(), (TaggedData.Data ? TaggedData.Data->GetClass()->GetFName().ToString(): TEXT("No Data")) });
+
+		if (!TaggedData.Tags.IsEmpty())
+		{
+			ItemName.Append(FString::Format(TEXT(": ({0})"), { FString::Join(TaggedData.Tags, TEXT(", ")) }));
+		}
+
+		DataComboBoxItems.Add(MakeShared<FName>(ItemName));
 	}
 
 	if (DataComboBoxItems.Num() > 0)
