[PCG] Rework projection logic, fix projection of landscape splines, and support projection params on spline projection.

We were forcing collapse to points when sampling general (non-spline) projections, because projection is not a 1:1 function (many points in source can project to same point on target). However this is not the case if ProjectionParams.bProjectionPositions==false because nothing actually moves, in which case we can stick to analytic sampling. Therefore RequireCollapseToSample() uses this as the criteria.

The spline projection case with ProjectionParams.bProjectionPositions==false falls back to the base projection code, which will do a tidy sampling of the 3D spline (as we are not moving the spline positions - quirky but is consistent with processing of 'concrete' data (points)).

The landscape spline makes use of a similar path. It is probably possible to avoid this new UPCGLandscapeSplineProjectionData class completely and put logic for detecting this case in UPCGProjectionData::RequiresCollapseToSample().

Rename projection params to disambiguate.

diff --git a/Engine/Plugins/Experimental/PCG/Source/PCG/Private/Data/PCGLandscapeSplineData.cpp b/Engine/Plugins/Experimental/PCG/Source/PCG/Private/Data/PCGLandscapeSplineData.cpp
--- a/Engine/Plugins/Experimental/PCG/Source/PCG/Private/Data/PCGLandscapeSplineData.cpp
+++ b/Engine/Plugins/Experimental/PCG/Source/PCG/Private/Data/PCGLandscapeSplineData.cpp
@@ -229,3 +229,26 @@ UPCGSpatialData* UPCGLandscapeSplineData::CopyInternal() const
 
 	return NewLandscapeSplineData;
 }
+
+UPCGProjectionData* UPCGLandscapeSplineData::ProjectOn(const UPCGSpatialData* InOther, const FPCGProjectionParams& InParams) const
+{
+	if (InOther->GetDimension() == 2)
+	{
+		UPCGLandscapeSplineProjectionData* SplineProjectionData = NewObject<UPCGLandscapeSplineProjectionData>();
+		SplineProjectionData->Initialize(this, InOther, InParams);
+		return SplineProjectionData;
+	}
+	else
+	{
+		return Super::ProjectOn(InOther, InParams);
+	}
+}
+
+UPCGSpatialData* UPCGLandscapeSplineProjectionData::CopyInternal() const
+{
+	UPCGLandscapeSplineProjectionData* NewProjectionData = NewObject<UPCGLandscapeSplineProjectionData>();
+
+	CopyBaseProjectionClass(NewProjectionData);
+
+	return NewProjectionData;
+}
diff --git a/Engine/Plugins/Experimental/PCG/Source/PCG/Private/Data/PCGProjectionData.cpp b/Engine/Plugins/Experimental/PCG/Source/PCG/Private/Data/PCGProjectionData.cpp
--- a/Engine/Plugins/Experimental/PCG/Source/PCG/Private/Data/PCGProjectionData.cpp
+++ b/Engine/Plugins/Experimental/PCG/Source/PCG/Private/Data/PCGProjectionData.cpp
@@ -6,7 +6,7 @@
 
 #include UE_INLINE_GENERATED_CPP_BY_NAME(PCGProjectionData)
 
-void UPCGProjectionData::Initialize(const UPCGSpatialData* InSource, const UPCGSpatialData* InTarget, const FPCGProjectionParams& InParams)
+void UPCGProjectionData::Initialize(const UPCGSpatialData* InSource, const UPCGSpatialData* InTarget, const FPCGProjectionParams& InProjectionParams)
 {
 	check(InSource && InTarget);
 	// TODO: improve support for higher-dimension projection.
@@ -16,7 +16,7 @@ void UPCGProjectionData::Initialize(const UPCGSpatialData* InSource, const UPCGS
 	Target = InTarget;
 	TargetActor = InSource->TargetActor;
 
-	Params = InParams;
+	ProjectionParams = InProjectionParams;
 
 	CachedBounds = ProjectBounds(Source->GetBounds());
 	CachedStrictBounds = ProjectBounds(Source->GetStrictBounds());
@@ -77,13 +77,42 @@ FBox UPCGProjectionData::ProjectBounds(const FBox& InBounds) const
 
 bool UPCGProjectionData::SamplePoint(const FTransform& InTransform, const FBox& InBounds, FPCGPoint& OutPoint, UPCGMetadata* OutMetadata) const
 {
-	// Detecting if a point is in a projection is a non-trivial. Projection is not in general a bijection and we cannot simply unproject
-	// the point from the Target and check if it is in the Source. However we can approximate the image of the projection and check
-	// if the query point is in the image.
+	//TRACE_CPUPROFILER_EVENT_SCOPE(UPCGProjectionData::SamplePoint);
 	
-	// Passing nullptr for the context means the operation will execute single threaded which is not ideal. To mitigate this we
-	// prewarm the point cache when this projection data is constructed in the projection element.
-	return ToPointData(nullptr)->SamplePoint(InTransform, InBounds, OutPoint, OutMetadata);
+	// Detecting if a point is in a projection is often non-trivial. Projection is not in general a bijection and we cannot simply unproject
+	// the point from the Target and check if it is in the Source. In this case we approximate the image of the projection and check
+	// if the query point is in the image.
+	if (RequiresCollapseToSample())
+	{
+		// Passing nullptr for the context means the operation will execute single threaded which is not ideal. To mitigate this we
+		// prewarm the point cache when this projection data is constructed in the projection element.
+		return ToPointData(nullptr)->SamplePoint(InTransform, InBounds, OutPoint, OutMetadata);
+	}
+
+	FPCGPoint PointFromSource;
+	if (!Source->SamplePoint(InTransform, InBounds, PointFromSource, OutMetadata))
+	{
+		return false;
+	}
+
+	// This relies on fact that SamplePoint moves the point. This will be replaced with a ProjectPoint() call.
+	FPCGPoint PointFromTarget;
+	if (!Target->SamplePoint(PointFromSource.Transform, PointFromSource.GetLocalBounds(), PointFromTarget, OutMetadata))
+	{
+		return false;
+	}
+
+	// Merge points into a single point
+	OutPoint = PointFromSource;
+
+	ApplyProjectionResult(PointFromTarget, OutPoint);
+
+	if (OutMetadata && PointFromTarget.MetadataEntry != PCGInvalidEntryKey)
+	{
+		OutMetadata->MergePointAttributesSubset(PointFromSource, OutMetadata, Source->Metadata, PointFromTarget, OutMetadata, Target->Metadata, OutPoint, ProjectionParams.AttributeMergeOperation);
+	}
+
+	return true;
 }
 
 bool UPCGProjectionData::HasNonTrivialTransform() const
@@ -91,6 +120,34 @@ bool UPCGProjectionData::HasNonTrivialTransform() const
 	return Target->HasNonTrivialTransform();
 }
 
+bool UPCGProjectionData::RequiresCollapseToSample() const
+{
+	// Detecting if a point is in a projection is often non-trivial. Projection is not in general a bijection and we cannot simply unproject
+	// the point from the Target and check if it is in the Source. 
+	
+	// There are cases where projection is a bijection. Like projecting volumes onto volumes (which is sampling). A non-PCG example is projection in
+	// graphics using homogeneous coordinates - points can be unprojected back to original positions in 3D space.
+	
+	// There are cases where a projection is not technically a bijection, however we can still sample it. To illustrate, projection of a spline straight
+	// down onto a terrain is such an example and is already covered via UPCGSplineProjectionData which overrides methods from this class. On the other
+	// hand projecting a spline onto a terrain in a non-straight-down direction already complicates things because the spline projection will get
+	// shadowed by the terrain (akin to terrain shadows cast by sunlight). We could raycast/raymarch from each query point towards the spline
+	// to check for occlusion by the terrain, and also do a similar trick to what's in UPCGSplineProjectionData to get closest point. The spline
+	// could intersect the terrain multiple times, so this will likely be expensive and take time to implement robustly. The alternative
+	// of collapsing might seems favorable.
+
+	// If we are losing precision from a collapse and we think we can sample without collapse, such cases could be detected and added here. Cases
+	// involving projecting points should not be added here because a collapse calls ToPointData() which just returns the point data.
+	
+	// Keep in mind that detecting these cases robustly would ideally walk the upstream graph if it is a 'composite' network - i.e. if we want
+	// to allow projection onto landscapes without collapse, we'd ideally check if the composite network is equivalent to a landscape (e.g. a landscape
+	// intersected with a volume) rather than only checking if the immediate projection source is a particular type. A concrete example of this
+	// failing would be the In and Actor graph input pins which can be backed by composite networks.
+	
+	// It is however trivial if we are not actually moving anything around..
+	return ProjectionParams.bProjectPositions;
+}
+
 const UPCGPointData* UPCGProjectionData::CreatePointData(FPCGContext* Context) const
 {
 	TRACE_CPUPROFILER_EVENT_SCOPE(UPCGProjectionData::CreatePointData);
@@ -123,7 +180,7 @@ const UPCGPointData* UPCGProjectionData::CreatePointData(FPCGContext* Context) c
 		TSet<FName> AttributesFilter;
 		GetIncludeExcludeAttributeNames(AttributesFilter);
 
-		TempTargetMetadata->InitializeWithAttributeFilter(Target->Metadata, AttributesFilter, Params.AttributeMode);
+		TempTargetMetadata->InitializeWithAttributeFilter(Target->Metadata, AttributesFilter, ProjectionParams.AttributeMode);
 	}
 
 	// Add any attributes from filtered metadata to produce final list
@@ -155,7 +212,7 @@ const UPCGPointData* UPCGProjectionData::CreatePointData(FPCGContext* Context) c
 		if (PointData->Metadata && TempTargetMetadata && PointFromTarget.MetadataEntry != PCGInvalidEntryKey)
 		{
 			// Merge metadata to produce final attribute values
-			PointData->Metadata->MergePointAttributesSubset(SourcePoint, SourcePointData->Metadata, SourcePointData->Metadata, PointFromTarget, TempTargetMetadata, TempTargetMetadata, OutPoint, Params.AttributeMergeOperation);
+			PointData->Metadata->MergePointAttributesSubset(SourcePoint, SourcePointData->Metadata, SourcePointData->Metadata, PointFromTarget, TempTargetMetadata, TempTargetMetadata, OutPoint, ProjectionParams.AttributeMergeOperation);
 		}
 
 		return true;
@@ -168,22 +225,22 @@ const UPCGPointData* UPCGProjectionData::CreatePointData(FPCGContext* Context) c
 
 void UPCGProjectionData::ApplyProjectionResult(const FPCGPoint& InTargetPoint, FPCGPoint& InOutProjected) const
 {
-	if (Params.bProjectPositions)
+	if (ProjectionParams.bProjectPositions)
 	{
 		InOutProjected.Transform.SetLocation(InTargetPoint.Transform.GetLocation());
 	}
 
-	if (Params.bProjectRotations)
+	if (ProjectionParams.bProjectRotations)
 	{
 		InOutProjected.Transform.SetRotation(InTargetPoint.Transform.GetRotation());
 	}
 
-	if (Params.bProjectScales)
+	if (ProjectionParams.bProjectScales)
 	{
 		InOutProjected.Transform.SetScale3D(InTargetPoint.Transform.GetScale3D());
 	}
 
-	if (Params.bProjectColors)
+	if (ProjectionParams.bProjectColors)
 	{
 		InOutProjected.Color *= InTargetPoint.Color;
 	}
@@ -193,13 +250,13 @@ void UPCGProjectionData::ApplyProjectionResult(const FPCGPoint& InTargetPoint, F
 
 void UPCGProjectionData::GetIncludeExcludeAttributeNames(TSet<FName>& OutAttributeNames) const
 {
-	if (Params.AttributeList.IsEmpty())
+	if (ProjectionParams.AttributeList.IsEmpty())
 	{
 		return;
 	}
 
 	TArray<FString> AttributeNameStrings;
-	Params.AttributeList.ParseIntoArray(AttributeNameStrings, TEXT(","), true);
+	ProjectionParams.AttributeList.ParseIntoArray(AttributeNameStrings, TEXT(","), true);
 
 	for (const FString& Attribute : AttributeNameStrings)
 	{
@@ -213,7 +270,7 @@ void UPCGProjectionData::CopyBaseProjectionClass(UPCGProjectionData* NewProjecti
 	NewProjectionData->Target = Target;
 	NewProjectionData->CachedBounds = CachedBounds;
 	NewProjectionData->CachedStrictBounds = CachedStrictBounds;
-	NewProjectionData->Params = Params;
+	NewProjectionData->ProjectionParams = ProjectionParams;
 }
 
 UPCGSpatialData* UPCGProjectionData::CopyInternal() const
diff --git a/Engine/Plugins/Experimental/PCG/Source/PCG/Private/Data/PCGSplineData.cpp b/Engine/Plugins/Experimental/PCG/Source/PCG/Private/Data/PCGSplineData.cpp
--- a/Engine/Plugins/Experimental/PCG/Source/PCG/Private/Data/PCGSplineData.cpp
+++ b/Engine/Plugins/Experimental/PCG/Source/PCG/Private/Data/PCGSplineData.cpp
@@ -132,10 +132,19 @@ UPCGSpatialData* UPCGSplineData::CopyInternal() const
 
 bool UPCGSplineProjectionData::SamplePoint(const FTransform& InTransform, const FBox& InBounds, FPCGPoint& OutPoint, UPCGMetadata* OutMetadata) const
 {
-	// TODO: support metadata - we don't currently have a good preseentation of what metadata entries mean for non-point data
+	// TODO: support metadata - we don't currently have a good representation of what metadata entries mean for non-point data
 	// TODO: use InBounds when sampling spline (sample in area rather than at closest point)
 
-	// Find nearest point on projected spline
+	if (!ProjectionParams.bProjectPositions)
+	{
+		// If we're not moving anything around, then just defer to super which will sample 3D spline, to make SamplePoint() consistent with behaviour
+		// on 'concrete' data (points).
+		return Super::SamplePoint(InTransform, InBounds, OutPoint, OutMetadata);
+	}
+
+	// Find nearest point on projected spline by lifting point along projection direction to closest position on spline. This way
+	// when we sample the spline we get a similar result to if the spline had been projected onto the surface.
+
 	const FVector InPosition = InTransform.GetLocation();
 	const USplineComponent* Spline = GetSpline()->Spline.Get();
 	const FVector& SurfaceNormal = GetSurface()->GetNormal();
@@ -162,16 +171,14 @@ bool UPCGSplineProjectionData::SamplePoint(const FTransform& InTransform, const
 		if (GetSurface()->SamplePoint(SplinePoint.Transform, InBounds, SurfacePoint, OutMetadata))
 		{
 			OutPoint = SplinePoint;
-			OutPoint.Transform = SurfacePoint.Transform;
-			OutPoint.Density *= SurfacePoint.Density;
-			OutPoint.Color *= SurfacePoint.Color;
+
+			ApplyProjectionResult(SurfacePoint, OutPoint);
 
 			if (OutMetadata)
 			{
 				if (SplinePoint.MetadataEntry != PCGInvalidEntryKey && SurfacePoint.MetadataEntry)
 				{
-					// TODO use metadata op from parameters once we have metadata support
-					OutMetadata->MergePointAttributesSubset(SplinePoint, OutMetadata, GetSpline()->Metadata, SurfacePoint, OutMetadata, GetSurface()->Metadata, OutPoint, EPCGMetadataOp::Max);
+					OutMetadata->MergePointAttributesSubset(SplinePoint, OutMetadata, GetSpline()->Metadata, SurfacePoint, OutMetadata, GetSurface()->Metadata, OutPoint, ProjectionParams.AttributeMergeOperation);
 				}
 				else if (SurfacePoint.MetadataEntry != PCGInvalidEntryKey)
 				{
diff --git a/Engine/Plugins/Experimental/PCG/Source/PCG/Public/Data/PCGLandscapeSplineData.h b/Engine/Plugins/Experimental/PCG/Source/PCG/Public/Data/PCGLandscapeSplineData.h
--- a/Engine/Plugins/Experimental/PCG/Source/PCG/Public/Data/PCGLandscapeSplineData.h
+++ b/Engine/Plugins/Experimental/PCG/Source/PCG/Public/Data/PCGLandscapeSplineData.h
@@ -4,6 +4,7 @@
 
 #include "CoreMinimal.h"
 #include "PCGPolyLineData.h"
+#include "PCGProjectionData.h"
 
 #include "PCGLandscapeSplineData.generated.h"
 
@@ -33,6 +34,7 @@ public:
 	//~Begin UPCGSpatialData interface
 	virtual FBox GetBounds() const override;
 	virtual bool SamplePoint(const FTransform& Transform, const FBox& Bounds, FPCGPoint& OutPoint, UPCGMetadata* OutMetadata) const override;
+	virtual UPCGProjectionData* ProjectOn(const UPCGSpatialData* InOther, const FPCGProjectionParams& InParams = FPCGProjectionParams()) const override;
 protected:
 	virtual UPCGSpatialData* CopyInternal() const override;
 	//~End UPCGSpatialData interface
@@ -41,3 +43,22 @@ protected:
 	UPROPERTY(BlueprintReadWrite, EditAnywhere, Category = SourceData)
 	TWeakObjectPtr<ULandscapeSplinesComponent> Spline;
 };
+
+/* The projection of a landscape spline onto a surface. */
+UCLASS(BlueprintType, ClassGroup = (Procedural))
+class PCG_API UPCGLandscapeSplineProjectionData : public UPCGProjectionData
+{
+	GENERATED_BODY()
+public:
+	//~Begin UPCGSpatialData interface
+	// It is easy to determine if a point lies on a top-down projection onto a surface. No need to convert to points.
+	// NOTE: It will be less easy if the projection is not straight downwards, as the landscape will 'shadow' the projection. This could be detected
+	// here.
+	virtual bool RequiresCollapseToSample() const override { return false; }
+	//~End UPCGSpatialData interface
+
+protected:
+	//~Begin UPCGSpatialData interface
+	virtual UPCGSpatialData* CopyInternal() const override;
+	//~End UPCGSpatialData interface
+};
diff --git a/Engine/Plugins/Experimental/PCG/Source/PCG/Public/Data/PCGProjectionData.h b/Engine/Plugins/Experimental/PCG/Source/PCG/Public/Data/PCGProjectionData.h
--- a/Engine/Plugins/Experimental/PCG/Source/PCG/Public/Data/PCGProjectionData.h
+++ b/Engine/Plugins/Experimental/PCG/Source/PCG/Public/Data/PCGProjectionData.h
@@ -16,7 +16,7 @@ class PCG_API UPCGProjectionData : public UPCGSpatialDataWithPointCache
 {
 	GENERATED_BODY()
 public:
-	void Initialize(const UPCGSpatialData* InSource, const UPCGSpatialData* InTarget, const FPCGProjectionParams& InParams);
+	void Initialize(const UPCGSpatialData* InSource, const UPCGSpatialData* InTarget, const FPCGProjectionParams& InProjectionParams);
 
 	//~Begin UPCGSpatialData interface
 	virtual int GetDimension() const override;
@@ -25,7 +25,7 @@ public:
 	virtual FVector GetNormal() const override;
 	virtual bool SamplePoint(const FTransform& Transform, const FBox& Bounds, FPCGPoint& OutPoint, UPCGMetadata* OutMetadata) const override;
 	virtual bool HasNonTrivialTransform() const override;
-	virtual bool RequiresCollapseToSample() const override { return true; }
+	virtual bool RequiresCollapseToSample() const override;
 protected:
 	virtual UPCGSpatialData* CopyInternal() const override;
 	//~End UPCGSpatialData interface
@@ -58,5 +58,5 @@ protected:
 	FBox CachedStrictBounds = FBox(EForceInit::ForceInit);
 
 	UPROPERTY(BlueprintReadwrite, VisibleAnywhere, Category = SpatialData)
-	FPCGProjectionParams Params;
+	FPCGProjectionParams ProjectionParams;
 };
