Improved the update pins method to return more information on changes.
Fixed some issues with the use of the advanced pins toggle, and likely some issues with custom pins on blueprint nodes.

diff --git a/Engine/Plugins/Experimental/PCG/Source/PCG/Private/PCGNode.cpp b/Engine/Plugins/Experimental/PCG/Source/PCG/Private/PCGNode.cpp
--- a/Engine/Plugins/Experimental/PCG/Source/PCG/Private/PCGNode.cpp
+++ b/Engine/Plugins/Experimental/PCG/Source/PCG/Private/PCGNode.cpp
@@ -417,8 +417,8 @@ void UPCGNode::OnSettingsChanged(UPCGSettings* InSettings, EPCGChangeType Change
 {
 	if (InSettings == GetSettings())
 	{
-		const bool bUpdatedPins = UpdatePins();
-		OnNodeChangedDelegate.Broadcast(this, ((bUpdatedPins ? EPCGChangeType::Edge : EPCGChangeType::None) | ChangeType));
+		const EPCGChangeType PinChangeType = UpdatePins();
+		OnNodeChangedDelegate.Broadcast(this, ChangeType | PinChangeType);
 	}
 }
 
@@ -438,12 +438,12 @@ void UPCGNode::UpdateAfterSettingsChangeDuringCreation()
 	UpdatePins();
 }
 
-bool UPCGNode::UpdatePins()
+EPCGChangeType UPCGNode::UpdatePins()
 {
 	return UpdatePins([](UPCGNode* Node){ return NewObject<UPCGPin>(Node); });
 }
 
-bool UPCGNode::UpdatePins(TFunctionRef<UPCGPin*(UPCGNode*)> PinAllocator)
+EPCGChangeType UPCGNode::UpdatePins(TFunctionRef<UPCGPin*(UPCGNode*)> PinAllocator)
 {
 	if (!GetSettings())
 	{
@@ -473,7 +473,7 @@ bool UPCGNode::UpdatePins(TFunctionRef<UPCGPin*(UPCGNode*)> PinAllocator)
 
 		InputPins.Reset();
 		OutputPins.Reset();
-		return bChanged;
+		return EPCGChangeType::Edge | EPCGChangeType::Node;
 	}
 
 	UPCGSettings* Settings = GetSettings();
@@ -483,6 +483,7 @@ bool UPCGNode::UpdatePins(TFunctionRef<UPCGPin*(UPCGNode*)> PinAllocator)
 	auto UpdatePins = [this, &PinAllocator](TArray<UPCGPin*>& Pins, const TArray<FPCGPinProperties>& PinProperties)
 	{
 		bool bAppliedEdgeChanges = false;
+		bool bChangedPins = false;
 
 		// Find unmatched pins vs. properties on a name basis
 		TArray<UPCGPin*> UnmatchedPins;
@@ -495,6 +496,7 @@ bool UPCGNode::UpdatePins(TFunctionRef<UPCGPin*(UPCGNode*)> PinAllocator)
 					Pin->Modify();
 					Pin->Properties = *MatchingProperties;
 					bAppliedEdgeChanges |= Pin->BreakAllIncompatibleEdges();
+					bChangedPins = true;
 				}
 			}
 			else
@@ -519,12 +521,14 @@ bool UPCGNode::UpdatePins(TFunctionRef<UPCGPin*(UPCGNode*)> PinAllocator)
 			UnmatchedPins[0]->Modify();
 			UnmatchedPins[0]->Properties = UnmatchedProperties[0];
 			bAppliedEdgeChanges |= UnmatchedPins[0]->BreakAllIncompatibleEdges();
+			bChangedPins = true;
 		}
 		else
 		{
 			if(!UnmatchedPins.IsEmpty() || !UnmatchedProperties.IsEmpty())
 			{
 				Modify();
+				bChangedPins = true;
 			}
 
 			// Remove old pins
@@ -548,14 +552,14 @@ bool UPCGNode::UpdatePins(TFunctionRef<UPCGPin*(UPCGNode*)> PinAllocator)
 			}
 		}
 
-		return bAppliedEdgeChanges;
+		return (bAppliedEdgeChanges ? EPCGChangeType::Edge : EPCGChangeType::None) | (bChangedPins ? EPCGChangeType::Node : EPCGChangeType::None);
 	};
 
-	bool bChanged = false;
-	bChanged |= UpdatePins(InputPins, InboundPinProperties);
-	bChanged |= UpdatePins(OutputPins, OutboundPinProperties);
+	EPCGChangeType ChangeType = EPCGChangeType::None;
+	ChangeType |= UpdatePins(InputPins, InboundPinProperties);
+	ChangeType |= UpdatePins(OutputPins, OutboundPinProperties);
 
-	return bChanged;
+	return ChangeType;
 }
 
 #if WITH_EDITOR
diff --git a/Engine/Plugins/Experimental/PCG/Source/PCG/Public/PCGNode.h b/Engine/Plugins/Experimental/PCG/Source/PCG/Public/PCGNode.h
--- a/Engine/Plugins/Experimental/PCG/Source/PCG/Public/PCGNode.h
+++ b/Engine/Plugins/Experimental/PCG/Source/PCG/Public/PCGNode.h
@@ -148,8 +148,8 @@ public:
 #endif // WITH_EDITORONLY_DATA
 
 protected:
-	bool UpdatePins();
-	bool UpdatePins(TFunctionRef<UPCGPin* (UPCGNode*)> PinAllocator);
+	EPCGChangeType UpdatePins();
+	EPCGChangeType UpdatePins(TFunctionRef<UPCGPin* (UPCGNode*)> PinAllocator);
 
 #if WITH_EDITOR
 	virtual void PostEditChangeProperty(struct FPropertyChangedEvent& PropertyChangedEvent) override;
