Added basic tracking on data from actor, that also waits for generation if any.
First version, might require a bit more checks to prevent issues.
Limited dependency depth in tracked actors in the pcg component

diff --git a/Engine/Plugins/Experimental/PCG/Source/PCG/Private/Elements/PCGDataFromActor.cpp b/Engine/Plugins/Experimental/PCG/Source/PCG/Private/Elements/PCGDataFromActor.cpp
--- a/Engine/Plugins/Experimental/PCG/Source/PCG/Private/Elements/PCGDataFromActor.cpp
+++ b/Engine/Plugins/Experimental/PCG/Source/PCG/Private/Elements/PCGDataFromActor.cpp
@@ -3,6 +3,7 @@
 #include "Elements/PCGDataFromActor.h"
 
 #include "PCGComponent.h"
+#include "PCGSubsystem.h"
 #include "Data/PCGSpatialData.h"
 
 #include "GameFramework/Actor.h"
@@ -10,6 +11,17 @@
 
 #include UE_INLINE_GENERATED_CPP_BY_NAME(PCGDataFromActor)
 
+#if WITH_EDITOR
+void UPCGDataFromActorSettings::GetTrackedActorTags(FPCGTagToSettingsMap& OutTagToSettings, TArray<TObjectPtr<const UPCGGraph>>& OutVisitedGraphs) const
+{
+	if (ActorSelector.ActorSelection == EPCGActorSelection::ByTag &&
+		ActorSelector.ActorFilter == EPCGActorFilter::AllWorldActors)
+	{
+		OutTagToSettings.FindOrAdd(ActorSelector.ActorSelectionTag).Add(this);
+	}
+}
+#endif
+
 FPCGElementPtr UPCGDataFromActorSettings::CreateElement() const
 {
 	return MakeShared<FPCGDataFromActorElement>();
@@ -30,42 +42,114 @@ TArray<FPCGPinProperties> UPCGDataFromActorSettings::OutputPinProperties() const
 	return Pins;
 }
 
-bool FPCGDataFromActorElement::ExecuteInternal(FPCGContext* Context) const
+FPCGContext* FPCGDataFromActorElement::Initialize(const FPCGDataCollection& InputData, TWeakObjectPtr<UPCGComponent> SourceComponent, const UPCGNode* Node)
+{
+	FPCGDataFromActorContext* Context = new FPCGDataFromActorContext();
+	Context->InputData = InputData;
+	Context->SourceComponent = SourceComponent;
+	Context->Node = Node;
+
+	return Context;
+}
+
+bool FPCGDataFromActorElement::ExecuteInternal(FPCGContext* InContext) const
 {
 	TRACE_CPUPROFILER_EVENT_SCOPE(FPCGDataFromActorElement::Execute);
 
-	check(Context);
+	check(InContext);
+	FPCGDataFromActorContext* Context = static_cast<FPCGDataFromActorContext*>(InContext);
+
 	const UPCGDataFromActorSettings* Settings = Context->GetInputSettings<UPCGDataFromActorSettings>();
 	check(Settings);
 
 	UWorld* World = Context->SourceComponent.IsValid() ? Context->SourceComponent->GetWorld() : nullptr;
-	TArray<AActor*> FoundActors = PCGActorSelector::FindActors(Settings->ActorSelector, World, nullptr);
 
-	if (FoundActors.IsEmpty())
+	if (!Context->bPerformedQuery)
 	{
-		PCGE_LOG(Warning, "No matching actor was found.");
-		return true;
+		Context->FoundActors = PCGActorSelector::FindActors(Settings->ActorSelector, World, nullptr);
+		Context->bPerformedQuery = true;
+
+		if (Context->FoundActors.IsEmpty())
+		{
+			PCGE_LOG(Warning, "No matching actor was found.");
+			return true;
+		}
+
+		// If we're looking for PCG component data, we might have to wait for it.
+		if (Settings->Mode == EPCGGetDataFromActorMode::GetDataFromPCGComponent || Settings->Mode == EPCGGetDataFromActorMode::GetDataFromPCGComponentOrParseComponents)
+		{
+			TArray<FPCGTaskId> WaitOnTaskIds;
+			for (AActor* Actor : Context->FoundActors)
+			{
+				GatherWaitTasks(Actor, WaitOnTaskIds);
+			}
+
+			if (!WaitOnTaskIds.IsEmpty())
+			{
+				UPCGSubsystem* Subsystem = Context->SourceComponent.IsValid() ? Context->SourceComponent->GetSubsystem() : nullptr;
+				if (Subsystem)
+				{
+					// Add a trivial task after these generations that wakes up this task
+					Context->bIsPaused = true;
+
+					Subsystem->ScheduleGeneric([Context]()
+					{
+						// Wake up the current task
+						Context->bIsPaused = false;
+						return true;
+					}, Context->SourceComponent.Get(), WaitOnTaskIds);
+
+					return false;
+				}
+				else
+				{
+					PCGE_LOG(Error, "Was unable to wait for end of generation tasks");
+				}
+			}
+		}
 	}
 
-	for (AActor* Actor : FoundActors)
+	if (Context->bPerformedQuery)
 	{
-		ProcessActor(Context, Settings, Actor);
+		for (AActor* Actor : Context->FoundActors)
+		{
+			ProcessActor(Context, Settings, Actor);
+		}
 	}
 
 	return true;
 }
 
+void FPCGDataFromActorElement::GatherWaitTasks(AActor* FoundActor, TArray<FPCGTaskId>& OutWaitTasks) const
+{
+	if (!FoundActor)
+	{
+		return;
+	}
+
+	TInlineComponentArray<UPCGComponent*, 1> PCGComponents;
+	FoundActor->GetComponents(PCGComponents);
+
+	for (UPCGComponent* Component : PCGComponents)
+	{
+		if (Component->IsGenerating())
+		{
+			OutWaitTasks.Add(Component->GetGenerationTaskId());
+		}
+	}
+}
+
 void FPCGDataFromActorElement::ProcessActor(FPCGContext* Context, const UPCGDataFromActorSettings* Settings, AActor* FoundActor) const
 {
 	check(Context);
 	check(Settings);
 
-	if (!FoundActor)
+	if (!FoundActor || !IsValid(FoundActor))
 	{
 		return;
 	}
 
-	TArray<UPCGComponent*> PCGComponents;
+	TInlineComponentArray<UPCGComponent*, 1> PCGComponents;
 	bool bHasGeneratedPCGData = false;
 	FProperty* FoundProperty = nullptr;
 
diff --git a/Engine/Plugins/Experimental/PCG/Source/PCG/Private/PCGComponent.cpp b/Engine/Plugins/Experimental/PCG/Source/PCG/Private/PCGComponent.cpp
--- a/Engine/Plugins/Experimental/PCG/Source/PCG/Private/PCGComponent.cpp
+++ b/Engine/Plugins/Experimental/PCG/Source/PCG/Private/PCGComponent.cpp
@@ -2116,7 +2116,7 @@ bool UPCGComponent::AddTrackedActor(AActor* InActor, bool bForce)
 
 		bAppliedChange = true;
 		CachedTrackedActorToTags.FindOrAdd(InActor).Add(Tag);
-		PCGHelpers::GatherDependencies(InActor, CachedTrackedActorToDependencies.FindOrAdd(InActor));
+		PCGHelpers::GatherDependencies(InActor, CachedTrackedActorToDependencies.FindOrAdd(InActor), 1);
 
 		if (!bForce)
 		{
diff --git a/Engine/Plugins/Experimental/PCG/Source/PCG/Public/Elements/PCGDataFromActor.h b/Engine/Plugins/Experimental/PCG/Source/PCG/Public/Elements/PCGDataFromActor.h
--- a/Engine/Plugins/Experimental/PCG/Source/PCG/Public/Elements/PCGDataFromActor.h
+++ b/Engine/Plugins/Experimental/PCG/Source/PCG/Public/Elements/PCGDataFromActor.h
@@ -28,6 +28,7 @@ public:
 #if WITH_EDITOR
 	virtual FName GetDefaultNodeName() const override { return FName(TEXT("DataFromActor")); }
 	virtual EPCGSettingsType GetType() const override { return EPCGSettingsType::Spatial; }
+	virtual void GetTrackedActorTags(FPCGTagToSettingsMap& OutTagToSettings, TArray<TObjectPtr<const UPCGGraph>>& OutVisitedGraphs) const override;
 #endif
 
 	virtual TArray<FPCGPinProperties> InputPinProperties() const override { return TArray<FPCGPinProperties>(); }
@@ -51,12 +52,21 @@ public:
 	FName PropertyName = NAME_None;
 };
 
-class FPCGDataFromActorElement : public FSimplePCGElement
+class FPCGDataFromActorContext : public FPCGContext
 {
 public:
+	TArray<AActor*> FoundActors;
+	bool bPerformedQuery = false;
+};
+
+class FPCGDataFromActorElement : public IPCGElement
+{
+public:
+	virtual FPCGContext* Initialize(const FPCGDataCollection& InputData, TWeakObjectPtr<UPCGComponent> SourceComponent, const UPCGNode* Node) override;
 	virtual bool CanExecuteOnlyOnMainThread(FPCGContext* Context) const override { return true; }
 	virtual bool IsCacheable(const UPCGSettings* InSettings) const override { return false; }
 protected:
 	virtual bool ExecuteInternal(FPCGContext* Context) const;
+	void GatherWaitTasks(AActor* FoundActor, TArray<FPCGTaskId>& OutWaitTasks) const;
 	void ProcessActor(FPCGContext* Context, const UPCGDataFromActorSettings* Settings, AActor* FoundActor) const;
 };
\ No newline at end of file
diff --git a/Engine/Plugins/Experimental/PCG/Source/PCG/Public/PCGComponent.h b/Engine/Plugins/Experimental/PCG/Source/PCG/Public/PCGComponent.h
--- a/Engine/Plugins/Experimental/PCG/Source/PCG/Public/PCGComponent.h
+++ b/Engine/Plugins/Experimental/PCG/Source/PCG/Public/PCGComponent.h
@@ -195,6 +195,9 @@ public:
 	bool IsGenerating() const { return CurrentGenerationTask != InvalidPCGTaskId; }
 	bool IsCleaningUp() const { return CurrentCleanupTask != InvalidPCGTaskId; }
 
+	/** Returns task ids to do internal chaining */
+	FPCGTaskId GetGenerationTaskId() const { return CurrentGenerationTask; }
+
 #if WITH_EDITOR
 	void Refresh();
 	void OnRefresh();
