Bug fix to make the PCG volumes work correctly when spawned via editor scripting

diff --git a/Engine/Plugins/Experimental/PCG/Source/PCGEditor/Private/PCGVolumeFactory.cpp b/Engine/Plugins/Experimental/PCG/Source/PCGEditor/Private/PCGVolumeFactory.cpp
--- a/Engine/Plugins/Experimental/PCG/Source/PCGEditor/Private/PCGVolumeFactory.cpp
+++ b/Engine/Plugins/Experimental/PCG/Source/PCGEditor/Private/PCGVolumeFactory.cpp
@@ -19,7 +19,12 @@ UPCGVolumeFactory::UPCGVolumeFactory(const FObjectInitializer& ObjectInitializer
 
 bool UPCGVolumeFactory::CanCreateActorFrom(const FAssetData& AssetData, FText& OutErrorMsg)
 {
-	if ( !AssetData.IsValid() || !AssetData.IsInstanceOf(UPCGGraph::StaticClass()))
+	if (UActorFactory::CanCreateActorFrom(AssetData, OutErrorMsg))
+	{
+		return true;
+	}
+
+	if ( AssetData.IsValid() && !AssetData.IsInstanceOf(UPCGGraph::StaticClass()))
 	{
 		OutErrorMsg = LOCTEXT("NoPCGGraph", "A valid PCG graph asset must be specified.");
 		return false;
