Added a debug log capture view for PCG editor, to be able to track down where errors come from.

diff --git a/Engine/Plugins/Experimental/PCG/Source/PCG/Private/PCGElement.cpp b/Engine/Plugins/Experimental/PCG/Source/PCG/Private/PCGElement.cpp
--- a/Engine/Plugins/Experimental/PCG/Source/PCG/Private/PCGElement.cpp
+++ b/Engine/Plugins/Experimental/PCG/Source/PCG/Private/PCGElement.cpp
@@ -29,7 +29,7 @@ bool IPCGElement::Execute(FPCGContext* Context) const
 
 			case EPCGExecutionPhase::PrepareData:
 			{
-				FScopedCallTimer CallTimer(*this, Context);
+				FScopedCall ScopedCall(*this, Context);
 				if (PrepareDataInternal(Context))
 				{
 					Context->CurrentPhase = EPCGExecutionPhase::Execute;
@@ -43,7 +43,7 @@ bool IPCGElement::Execute(FPCGContext* Context) const
 
 			case EPCGExecutionPhase::Execute:
 			{
-				FScopedCallTimer CallTimer(*this, Context);
+				FScopedCall ScopedCall(*this, Context);
 				if (ExecuteInternal(Context))
 				{
 					Context->CurrentPhase = EPCGExecutionPhase::PostExecute;
@@ -154,7 +154,7 @@ void IPCGElement::PostExecute(FPCGContext* Context) const
 
 #if WITH_EDITOR
 	{
-		FScopeLock Lock(&TimersLock);
+		FScopeLock Lock(&CapturedDataLock);
 		PCGE_LOG(Verbose, "Executed in (%f)s and (%d) frames(s)", Timers[CurrentTimerIndex].ExecutionTime, Timers[CurrentTimerIndex].ExecutionFrameCount);
 	}
 #endif
@@ -212,12 +212,33 @@ void IPCGElement::DebugDisplay(FPCGContext* Context) const
 	}
 }
 
+
+TArray<IPCGElement::FCallTime> IPCGElement::GetTimers() const
+{
+	FScopeLock Lock(&CapturedDataLock);
+	return Timers;
+}
+
+TArray<IPCGElement::FCapturedMessage> IPCGElement::GetCapturedMessages() const
+{
+	FScopeLock Lock(&CapturedDataLock);
+	return CapturedMessages;
+
+}
+
 void IPCGElement::ResetTimers()
 {
-	FScopeLock Lock(&TimersLock);
+	FScopeLock Lock(&CapturedDataLock);
 	Timers.Empty();
 	CurrentTimerIndex = 0;
 }
+
+void IPCGElement::ResetMessages()
+{
+	FScopeLock Lock(&CapturedDataLock);
+	CapturedMessages.Reset();
+}
+
 #endif // WITH_EDITOR
 
 void IPCGElement::CleanupAndValidateOutput(FPCGContext* Context) const
@@ -281,17 +302,24 @@ bool IPCGElement::IsCacheableInstance(const UPCGSettingsInterface* InSettingsInt
 }
 
 #if WITH_EDITOR
-IPCGElement::FScopedCallTimer::FScopedCallTimer(const IPCGElement& InOwner, FPCGContext* InContext)
-	: Owner(InOwner), Context(InContext), Phase(InContext->CurrentPhase)
+IPCGElement::FScopedCall::FScopedCall(const IPCGElement& InOwner, FPCGContext* InContext)
+	: Owner(InOwner)
+	, Context(InContext)
+	, Phase(InContext->CurrentPhase)
+	, ThreadID(FPlatformTLS::GetCurrentThreadId())
 {
 	StartTime = FPlatformTime::Seconds();
+
+	GLog->AddOutputDevice(this);
 }
 
-IPCGElement::FScopedCallTimer::~FScopedCallTimer()
+IPCGElement::FScopedCall::~FScopedCall()
 {
+	GLog->RemoveOutputDevice(this);
+
 	const double ThisFrameTime = FPlatformTime::Seconds() - StartTime;
 
-	FScopeLock Lock(&Owner.TimersLock);
+	FScopeLock Lock(&Owner.CapturedDataLock);
 
 	constexpr int32 MaxNumberOfTrackedTimers = 100;
 	if (Phase == EPCGExecutionPhase::PrepareData)
@@ -327,7 +355,27 @@ IPCGElement::FScopedCallTimer::~FScopedCallTimer()
 
 		Time.PostExecuteTime = ThisFrameTime;
 	}
+
+	Owner.CapturedMessages += std::move(CapturedMessages);
+}
+
+void IPCGElement::FScopedCall::Serialize(const TCHAR* V, ELogVerbosity::Type Verbosity, const class FName& Category)
+{
+	// TODO: this thread id check will also filter out messages spawned from threads spawned inside of nodes. To improve that,
+	// perhaps set at TLS bit on things from here and inside of PCGAsync spawned jobs. If this was done, CapturedMessages below also will
+	// need protection
+	if (Verbosity > ELogVerbosity::Warning || FPlatformTLS::GetCurrentThreadId() != ThreadID)
+	{
+		// ignore
+		return;
+	}
+
+	// this is a dumb counter just so messages can be sorted in a similar order as when they were logged
+	static volatile int32 MessageCounter = 0;
+
+	CapturedMessages.Add(FCapturedMessage { MessageCounter++, Category, V, Verbosity});
 }
+
 #endif // WITH_EDITOR
 
 FPCGContext* FSimplePCGElement::Initialize(const FPCGDataCollection& InputData, TWeakObjectPtr<UPCGComponent> SourceComponent, const UPCGNode* Node)
diff --git a/Engine/Plugins/Experimental/PCG/Source/PCG/Public/PCGElement.h b/Engine/Plugins/Experimental/PCG/Source/PCG/Public/PCGElement.h
--- a/Engine/Plugins/Experimental/PCG/Source/PCG/Public/PCGElement.h
+++ b/Engine/Plugins/Experimental/PCG/Source/PCG/Public/PCGElement.h
@@ -6,6 +6,8 @@
 #include "PCGContext.h"
 #include "PCGData.h"
 #include "Math/NumericLimits.h"
+#include "Misc/OutputDevice.h"
+#include "Misc/OutputDeviceRedirector.h"
 
 class IPCGElement;
 class UPCGComponent;
@@ -67,9 +69,22 @@ public:
 		double PostExecuteTime = 0.0;
 	};
 
+	struct FCapturedMessage
+	{
+		int32 Index=0;
+		FName Namespace;
+		FString Message;
+		ELogVerbosity::Type Verbosity;
+	};
+
 	void DebugDisplay(FPCGContext* Context) const;
-	const TArray<FCallTime>& GetTimers() const { return Timers; }
+
+	// returns a copy because they can be modified on another thread otherwise
+	TArray<FCallTime> GetTimers() const;
+	TArray<FCapturedMessage> GetCapturedMessages() const;
+
 	void ResetTimers();
+	void ResetMessages();
 #endif
 
 protected:
@@ -95,15 +110,24 @@ private:
 	void CleanupAndValidateOutput(FPCGContext* Context) const;
 
 #if WITH_EDITOR
-	struct FScopedCallTimer
+
+	struct FScopedCall : public FOutputDevice
 	{
-		FScopedCallTimer(const IPCGElement& InOwner, FPCGContext* InContext);
-		~FScopedCallTimer();
+		FScopedCall(const IPCGElement& InOwner, FPCGContext* InContext);
+		virtual ~FScopedCall();
 
 		const IPCGElement& Owner;
 		FPCGContext* Context;
 		double StartTime;
 		EPCGExecutionPhase Phase;
+		const uint32 ThreadID;
+		TArray<FCapturedMessage> CapturedMessages;
+
+		// FOutputDevice
+		virtual bool IsMemoryOnly() const override { return true; }
+		virtual bool CanBeUsedOnMultipleThreads() const override { return true; }
+		virtual bool CanBeUsedOnAnyThread() const override { return true; }
+		virtual void Serialize(const TCHAR* V, ELogVerbosity::Type Verbosity, const class FName& Category) override;
 	};
 
 	// Set mutable because we need to modify them in the execute call, which is const
@@ -111,14 +135,15 @@ private:
 	// For now, it will track all calls to execute (excluding call where the result is already in cache).
 	mutable TArray<FCallTime> Timers;
 	mutable int32 CurrentTimerIndex = 0;
+	mutable TArray<FCapturedMessage> CapturedMessages;
 
 	// Perhaps overkill but there is a slight chance that we need to protect the timers array. If we call reset from the UI while an element is executing,
 	// it could crash while writing to the timers array.
-	mutable FCriticalSection TimersLock;
+	mutable FCriticalSection CapturedDataLock;
 #else // !WITH_EDITOR
-	struct FScopedCallTimer
+	struct FScopedCall
 	{
-		FScopedCallTimer(const IPCGElement& InOwner, FPCGContext* InContext) {}
+		FScopedCall(const IPCGElement& InOwner, FPCGContext* InContext) {}
 		void AdvanceTimerIndex() {}
 	};
 #endif // WITH_EDITOR
diff --git a/Engine/Plugins/Experimental/PCG/Source/PCGEditor/Private/PCGEditor.cpp b/Engine/Plugins/Experimental/PCG/Source/PCGEditor/Private/PCGEditor.cpp
--- a/Engine/Plugins/Experimental/PCG/Source/PCGEditor/Private/PCGEditor.cpp
+++ b/Engine/Plugins/Experimental/PCG/Source/PCGEditor/Private/PCGEditor.cpp
@@ -31,6 +31,7 @@
 #include "SPCGEditorGraphFind.h"
 #include "SPCGEditorGraphNodePalette.h"
 #include "SPCGEditorGraphProfilingView.h"
+#include "SPCGEditorGraphLogView.h"
 
 #include "AssetToolsModule.h"
 #include "EdGraphUtilities.h"
@@ -73,6 +74,7 @@ namespace FPCGEditor_private
 	const FName FindID = FName(TEXT("Find"));
 	const FName DeterminismID = FName(TEXT("Determinism"));
 	const FName ProfilingID = FName(TEXT("Profiling"));
+	const FName LogID = FName(TEXT("Log"));
 }
 
 void FPCGEditor::Initialize(const EToolkitMode::Type InMode, const TSharedPtr<class IToolkitHost>& InToolkitHost, UPCGGraph* InPCGGraph)
@@ -102,6 +104,7 @@ void FPCGEditor::Initialize(const EToolkitMode::Type InMode, const TSharedPtr<cl
 	AttributesWidget = CreateAttributesWidget();
 	DeterminismWidget = CreateDeterminismWidget();
 	ProfilingWidget = CreateProfilingWidget();
+	LogWidget = CreateLogWidget();
 
 	BindCommands();
 	RegisterToolbar();
@@ -225,6 +228,10 @@ void FPCGEditor::RegisterTabSpawners(const TSharedRef<FTabManager>& InTabManager
 	InTabManager->RegisterTabSpawner(FPCGEditor_private::ProfilingID, FOnSpawnTab::CreateSP(this, &FPCGEditor::SpawnTab_Profiling))
 		.SetDisplayName(LOCTEXT("ProfilingTab", "Profiling"))
 		.SetGroup(WorkspaceMenuCategoryRef);
+
+	InTabManager->RegisterTabSpawner(FPCGEditor_private::LogID, FOnSpawnTab::CreateSP(this, &FPCGEditor::SpawnTab_Log))
+		.SetDisplayName(LOCTEXT("LogCaptureTab", "Log Capture"))
+		.SetGroup(WorkspaceMenuCategoryRef);
 }
 
 void FPCGEditor::UnregisterTabSpawners(const TSharedRef<class FTabManager>& InTabManager)
@@ -1849,6 +1856,11 @@ TSharedRef<SPCGEditorGraphProfilingView> FPCGEditor::CreateProfilingWidget()
 	return SNew(SPCGEditorGraphProfilingView, SharedThis(this));
 }
 
+TSharedRef<SPCGEditorGraphLogView> FPCGEditor::CreateLogWidget()
+{
+	return SNew(SPCGEditorGraphLogView, SharedThis(this));
+}
+
 void FPCGEditor::OnSelectedNodesChanged(const TSet<UObject*>& NewSelection)
 {
 	TArray<TWeakObjectPtr<UObject>> SelectedObjects;
@@ -2072,4 +2084,14 @@ TSharedRef<SDockTab> FPCGEditor::SpawnTab_Profiling(const FSpawnTabArgs& Args)
 		];
 }
 
+TSharedRef<SDockTab> FPCGEditor::SpawnTab_Log(const FSpawnTabArgs& Args)
+{
+	return SNew(SDockTab)
+		.Label(LOCTEXT("PCGLogTitle", "Log Capture"))
+		.TabColorScale(GetTabColorScale())
+		[
+			LogWidget.ToSharedRef()
+		];
+}
+
 #undef LOCTEXT_NAMESPACE
diff --git a/Engine/Plugins/Experimental/PCG/Source/PCGEditor/Private/PCGEditor.h b/Engine/Plugins/Experimental/PCG/Source/PCGEditor/Private/PCGEditor.h
--- a/Engine/Plugins/Experimental/PCG/Source/PCGEditor/Private/PCGEditor.h
+++ b/Engine/Plugins/Experimental/PCG/Source/PCGEditor/Private/PCGEditor.h
@@ -17,6 +17,7 @@ class SPCGEditorGraphDeterminismListView;
 class SPCGEditorGraphFind;
 class SPCGEditorGraphNodePalette;
 class SPCGEditorGraphProfilingView;
+class SPCGEditorGraphLogView;
 class UEdGraphNode;
 class UPCGComponent;
 class UPCGEditorGraph;
@@ -220,6 +221,9 @@ private:
 	/** Create a new profiling tab widget */
 	TSharedRef<SPCGEditorGraphProfilingView> CreateProfilingWidget();
 
+	/** Create a new profiling tab widget */
+	TSharedRef<SPCGEditorGraphLogView> CreateLogWidget();
+
 	/** Called when the selection changes in the GraphEditor */
 	void OnSelectedNodesChanged(const TSet<UObject*>& NewSelection);
 
@@ -256,6 +260,7 @@ private:
 	TSharedRef<SDockTab> SpawnTab_Find(const FSpawnTabArgs& Args);
 	TSharedRef<SDockTab> SpawnTab_Determinism(const FSpawnTabArgs& Args);
 	TSharedRef<SDockTab> SpawnTab_Profiling(const FSpawnTabArgs& Args);
+	TSharedRef<SDockTab> SpawnTab_Log(const FSpawnTabArgs& Args);
 
 	TSharedPtr<SGraphEditor> GraphEditorWidget;
 	TSharedPtr<IDetailsView> PropertyDetailsWidget;
@@ -264,6 +269,7 @@ private:
 	TSharedPtr<SPCGEditorGraphAttributeListView> AttributesWidget;
 	TSharedPtr<SPCGEditorGraphDeterminismListView> DeterminismWidget;
 	TSharedPtr<SPCGEditorGraphProfilingView> ProfilingWidget;
+	TSharedPtr<SPCGEditorGraphLogView> LogWidget;
 
 	TSharedPtr<FUICommandList> GraphEditorCommands;
 
diff --git a/Engine/Plugins/Experimental/PCG/Source/PCGEditor/Private/SPCGEditorGraphLogView.cpp b/Engine/Plugins/Experimental/PCG/Source/PCGEditor/Private/SPCGEditorGraphLogView.cpp
new file mode 100644
--- /dev/null
+++ b/Engine/Plugins/Experimental/PCG/Source/PCGEditor/Private/SPCGEditorGraphLogView.cpp
@@ -0,0 +1,373 @@
+// Copyright Epic Games, Inc. All Rights Reserved.
+
+#include "SPCGEditorGraphLogView.h"
+
+#include "PCGComponent.h"
+#include "PCGEditor.h"
+#include "PCGEditorGraph.h"
+#include "PCGEditorGraphNode.h"
+#include "PCGGraph.h"
+
+#include "Styling/AppStyle.h"
+#include "Widgets/Input/SSearchBox.h"
+#include "Widgets/Layout/SScrollBox.h"
+#include "Widgets/SBoxPanel.h"
+#include "Widgets/SWidget.h"
+
+#include "Styling/StyleColors.h"
+
+#define LOCTEXT_NAMESPACE "SPCGEditorGraphLogView"
+
+namespace PCGEditorGraphLogView
+{	
+	const FName NAME_Order = FName(TEXT("Order"));
+	const FName NAME_Node = FName(TEXT("Node"));
+	const FName NAME_Namespace = FName(TEXT("Namespace"));
+	const FName NAME_Message = FName(TEXT("Message"));
+
+	/** Labels of the columns */
+	const FText TEXT_OrderLabel = LOCTEXT("Order", "Order");
+	const FText TEXT_NodeLabel = LOCTEXT("Node", "Node");
+	const FText TEXT_NamespaceLabel = LOCTEXT("Namespace", "Namespace");
+	const FText TEXT_MessageLabel = LOCTEXT("Message", "Message");
+}
+
+void SPCGLogListViewItemRow::Construct(const FArguments& InArgs, const TSharedRef<STableViewBase>& InOwnerTableView, const PCGLogListViewItemPtr& Item)
+{
+	InternalItem = Item;
+
+	SMultiColumnTableRow<PCGLogListViewItemPtr>::Construct(
+		SMultiColumnTableRow::FArguments()
+		.Style(FAppStyle::Get(), "DataTableEditor.CellListViewRow"),
+		InOwnerTableView);
+}
+
+TSharedRef<SWidget> SPCGLogListViewItemRow::GenerateWidgetForColumn(const FName& ColumnId)
+{
+	FText ColumnData = LOCTEXT("ColumnError", "Unrecognized Column");
+
+	if (InternalItem.IsValid())
+	{
+		if (ColumnId == PCGEditorGraphLogView::NAME_Order)
+		{
+			ColumnData = FText::AsNumber(InternalItem->Order);
+		}
+		else if (ColumnId == PCGEditorGraphLogView::NAME_Node)
+		{
+			return SNew(STextBlock)
+				.ColorAndOpacity(InternalItem->EditorNode->GetNodeTitleColor())
+				.Text(FText::FromName(InternalItem->NodeName));
+		}
+		else if (ColumnId == PCGEditorGraphLogView::NAME_Namespace)
+		{
+			ColumnData = FText::FromName(InternalItem->Namespace);
+		}
+		else if (ColumnId == PCGEditorGraphLogView::NAME_Message)
+		{
+			EStyleColor Color = EStyleColor::Foreground;
+
+			switch (InternalItem->Verbosity)
+			{
+			case ELogVerbosity::Error:
+				Color = EStyleColor::Error;
+				break;
+
+			case ELogVerbosity::Warning:
+				Color = EStyleColor::Warning;
+				break;
+			}
+
+			return SNew(STextBlock)
+				.ColorAndOpacity(Color)
+				.Text(FText::FromString(InternalItem->Message));
+		}
+	}
+
+	return SNew(STextBlock)
+		.Text(ColumnData);
+}
+
+void SPCGEditorGraphLogView::OnItemDoubleClicked(PCGLogListViewItemPtr Item)
+{
+	const TSharedPtr<FPCGEditor> PCGEditor = PCGEditorPtr.Pin();
+	if (!PCGEditor.IsValid())
+	{
+		return;
+	}
+
+	if (!Item.IsValid() || !Item->EditorNode)
+	{
+		return;
+	}
+
+	PCGEditor->JumpToNode(Item->EditorNode);
+}
+
+void SPCGEditorGraphLogView::Construct(const FArguments& InArgs, TSharedPtr<FPCGEditor> InPCGEditor)
+{
+	PCGEditorPtr = InPCGEditor;
+
+	const TSharedPtr<FPCGEditor> PCGEditor = PCGEditorPtr.Pin();
+	if (PCGEditor)
+	{
+		PCGEditorGraph = PCGEditor->GetPCGEditorGraph();
+	}
+
+	ListViewHeader = CreateHeaderRowWidget();
+
+	const TSharedRef<SScrollBar> HorizontalScrollBar = SNew(SScrollBar)
+		.Orientation(Orient_Horizontal)
+		.Thickness(FVector2D(12.0f, 12.0f));
+
+	const TSharedRef<SScrollBar> VerticalScrollBar = SNew(SScrollBar)
+		.Orientation(Orient_Vertical)
+		.Thickness(FVector2D(12.0f, 12.0f));
+	
+	SAssignNew(ListView, SListView<PCGLogListViewItemPtr>)
+		.ListItemsSource(&ListViewItems)
+		.HeaderRow(ListViewHeader)
+		.OnGenerateRow(this, &SPCGEditorGraphLogView::OnGenerateRow)
+		.OnMouseButtonDoubleClick(this, &SPCGEditorGraphLogView::OnItemDoubleClicked)
+		.AllowOverscroll(EAllowOverscroll::No)
+		.ExternalScrollbar(VerticalScrollBar)
+		.ConsumeMouseWheel(EConsumeMouseWheel::Always);
+	
+	this->ChildSlot
+	[
+		SNew(SVerticalBox)
+		+SVerticalBox::Slot()
+		.AutoHeight()
+		[
+			SNew(SHorizontalBox)
+			+SHorizontalBox::Slot()
+			.AutoWidth()
+			[
+				SNew(SButton)
+				.Text(LOCTEXT("RefreshButton", "Refresh"))
+				.OnClicked(this, &SPCGEditorGraphLogView::Refresh)
+			]
+			+SHorizontalBox::Slot()
+			.AutoWidth()
+			[
+				SNew(SButton)
+				.Text(LOCTEXT("ClearButton", "Clear"))
+				.OnClicked(this, &SPCGEditorGraphLogView::Clear)
+			]
+		]
+		+SVerticalBox::Slot()
+		.FillHeight(1.0f)
+		[
+			SNew(SHorizontalBox)
+			+SHorizontalBox::Slot()
+			.FillWidth(1.0f)
+			[
+				SNew(SScrollBox)
+				.Orientation(Orient_Horizontal)
+				.ExternalScrollbar(HorizontalScrollBar)
+				+SScrollBox::Slot()
+				[
+					ListView->AsShared()
+				]
+			]
+			+SHorizontalBox::Slot()
+			.AutoWidth()
+			[
+				VerticalScrollBar
+			]
+		]
+		+SVerticalBox::Slot()
+		.AutoHeight()
+		[
+			SNew(SHorizontalBox)
+			+SHorizontalBox::Slot()
+			[
+				HorizontalScrollBar
+			]
+		]
+	];
+
+	Refresh();
+}
+
+TSharedRef<SHeaderRow> SPCGEditorGraphLogView::CreateHeaderRowWidget()
+{
+	return SNew(SHeaderRow)
+		.ResizeMode(ESplitterResizeMode::FixedPosition)
+		.CanSelectGeneratedColumn(true)
+		+ SHeaderRow::Column(PCGEditorGraphLogView::NAME_Order)
+			.ManualWidth(64)
+			.DefaultLabel(PCGEditorGraphLogView::TEXT_OrderLabel)
+			.HAlignHeader(HAlign_Center)
+			.HAlignCell(HAlign_Center)
+			.SortMode(this, &SPCGEditorGraphLogView::GetColumnSortMode, PCGEditorGraphLogView::NAME_Order)
+			.OnSort(this, &SPCGEditorGraphLogView::OnSortColumnHeader)
+		+ SHeaderRow::Column(PCGEditorGraphLogView::NAME_Namespace)
+			.ManualWidth(180)
+			.DefaultLabel(PCGEditorGraphLogView::TEXT_NamespaceLabel)
+			.HAlignHeader(HAlign_Center)
+			.HAlignCell(HAlign_Left)
+			.SortMode(this, &SPCGEditorGraphLogView::GetColumnSortMode, PCGEditorGraphLogView::NAME_Namespace)
+			.OnSort(this, &SPCGEditorGraphLogView::OnSortColumnHeader)
+		+ SHeaderRow::Column(PCGEditorGraphLogView::NAME_Node)
+			.ManualWidth(180)
+			.DefaultLabel(PCGEditorGraphLogView::TEXT_NodeLabel)
+			.HAlignHeader(HAlign_Center)
+			.HAlignCell(HAlign_Left)
+			.SortMode(this, &SPCGEditorGraphLogView::GetColumnSortMode, PCGEditorGraphLogView::NAME_Node)
+			.OnSort(this, &SPCGEditorGraphLogView::OnSortColumnHeader)
+		+ SHeaderRow::Column(PCGEditorGraphLogView::NAME_Message)
+			.FillWidth(1.0)
+			.DefaultLabel(PCGEditorGraphLogView::TEXT_MessageLabel)
+			.HAlignHeader(HAlign_Center)
+			.HAlignCell(HAlign_Left)
+			.SortMode(this, &SPCGEditorGraphLogView::GetColumnSortMode, PCGEditorGraphLogView::NAME_Message)
+			.OnSort(this, &SPCGEditorGraphLogView::OnSortColumnHeader);
+}
+
+void SPCGEditorGraphLogView::OnSortColumnHeader(const EColumnSortPriority::Type SortPriority, const FName& ColumnId, const EColumnSortMode::Type NewSortMode)
+{
+	if (SortingColumn == ColumnId)
+	{
+		// Circling
+		SortMode = EColumnSortMode::Type((SortMode + 1) % 3);
+	}
+	else
+	{
+		SortingColumn = ColumnId;
+		SortMode = NewSortMode;
+	}
+
+	Refresh();
+}
+
+EColumnSortMode::Type SPCGEditorGraphLogView::GetColumnSortMode(const FName ColumnId) const
+{
+	if (SortingColumn != ColumnId)
+	{
+		return EColumnSortMode::None;
+	}
+
+	return SortMode;
+}
+
+FReply SPCGEditorGraphLogView::Clear()
+{
+	const TSharedPtr<FPCGEditor> PCGEditor = PCGEditorPtr.Pin();
+	if (!PCGEditor.IsValid())
+	{
+		return FReply::Handled();
+	}
+
+	if (!PCGEditorGraph)
+	{
+		return FReply::Handled();
+	}
+
+	TArray<UPCGEditorGraphNode*> EditorNodes;
+	PCGEditorGraph->GetNodesOfClass<UPCGEditorGraphNode>(EditorNodes);
+
+	for (UPCGEditorGraphNode* PCGEditorNode : EditorNodes)
+	{
+		if (PCGEditorNode)
+		{
+			if (UPCGNode* PCGNode = PCGEditorNode->GetPCGNode())
+			{
+				if (UPCGSettings* Settings = PCGNode->GetSettings())
+				{
+					if (IPCGElement* Element = Settings->GetElement().Get())
+					{
+						Element->ResetMessages();
+					}
+				}
+			}
+		}
+	}
+
+	Refresh();
+
+	return FReply::Handled();
+}
+
+FReply SPCGEditorGraphLogView::Refresh()
+{
+	ListViewItems.Empty();
+	ListView->RequestListRefresh();
+
+	const TSharedPtr<FPCGEditor> PCGEditor = PCGEditorPtr.Pin();
+	if (!PCGEditor.IsValid())
+	{
+		return FReply::Handled();
+	}
+
+	if (!PCGEditorGraph)
+	{
+		return FReply::Handled();
+	}
+
+	TArray<UPCGEditorGraphNode*> EditorNodes;
+	PCGEditorGraph->GetNodesOfClass<UPCGEditorGraphNode>(EditorNodes);
+
+	for (const UPCGEditorGraphNode* PCGEditorNode : EditorNodes)
+	{
+		const UPCGNode* PCGNode = PCGEditorNode ? PCGEditorNode->GetPCGNode() : nullptr;
+		if (PCGNode)
+		{
+			if (const UPCGSettings* Settings = PCGNode->GetSettings())
+			{
+				if (const IPCGElement* Element = Settings->GetElement().Get())
+				{
+					for (const IPCGElement::FCapturedMessage& Message : Element->GetCapturedMessages())
+					{
+						PCGLogListViewItemPtr ListViewItem = MakeShared<FPCGLogListViewItem>();
+						ListViewItem->EditorNode = PCGEditorNode;
+						ListViewItem->PCGNode = PCGNode;
+						ListViewItem->NodeName = PCGNode->GetNodeTitle();
+						ListViewItem->Order = Message.Index;
+						ListViewItem->Namespace = Message.Namespace;
+						ListViewItem->Message = Message.Message;
+						ListViewItem->Verbosity = Message.Verbosity;
+
+						this->ListViewItems.Add(ListViewItem);
+					};
+				}
+			}
+		}
+	}
+
+	if (SortingColumn != NAME_None && SortMode != EColumnSortMode::None)
+	{
+		Algo::Sort(ListViewItems, [this](const PCGLogListViewItemPtr& A, const PCGLogListViewItemPtr& B)
+			{
+				bool isLess = false;
+				if (SortingColumn == PCGEditorGraphLogView::NAME_Order)
+				{
+					isLess = A->Order < B->Order;
+				}
+				else if (SortingColumn == PCGEditorGraphLogView::NAME_Namespace)
+				{
+					isLess = A->Namespace.LexicalLess(B->Namespace);
+				}
+				else if (SortingColumn == PCGEditorGraphLogView::NAME_Node)
+				{
+					isLess = A->NodeName.LexicalLess(B->NodeName);
+				}
+				else if (SortingColumn == PCGEditorGraphLogView::NAME_Message)
+				{
+					isLess = A->Message < B->Message;
+				}
+
+				return SortMode == EColumnSortMode::Ascending ? isLess : !isLess;
+			});
+	}
+
+	ListView->SetListItemsSource(ListViewItems);
+
+	return FReply::Handled();
+}
+
+TSharedRef<ITableRow> SPCGEditorGraphLogView::OnGenerateRow(PCGLogListViewItemPtr Item, const TSharedRef<STableViewBase>& OwnerTable) const
+{
+	return SNew(SPCGLogListViewItemRow, OwnerTable, Item);
+}
+
+#undef LOCTEXT_NAMESPACE
diff --git a/Engine/Plugins/Experimental/PCG/Source/PCGEditor/Private/SPCGEditorGraphLogView.h b/Engine/Plugins/Experimental/PCG/Source/PCGEditor/Private/SPCGEditorGraphLogView.h
new file mode 100644
--- /dev/null
+++ b/Engine/Plugins/Experimental/PCG/Source/PCGEditor/Private/SPCGEditorGraphLogView.h
@@ -0,0 +1,86 @@
+// Copyright Epic Games, Inc. All Rights Reserved.
+
+#pragma once
+
+#include "CoreMinimal.h"
+#include "PCGElement.h"
+#include "Widgets/Input/SButton.h"
+#include "Widgets/SCompoundWidget.h"
+#include "Widgets/Views/SHeaderRow.h"
+#include "Widgets/Views/SListView.h"
+#include "Logging/LogVerbosity.h"
+
+class FPCGEditor;
+class UPCGComponent;
+class UPCGEditorGraphNode;
+class UPCGEditorGraph;
+class UPCGNode;
+
+namespace PCGEditorGraphLogView
+{	
+	/** Names of the columns in the attribute list */
+	extern const FName NAME_Order;
+}
+
+struct FPCGLogListViewItem
+{
+	const UPCGNode* PCGNode = nullptr;	
+	const UPCGEditorGraphNode* EditorNode = nullptr;
+
+	int32 Order = 0;
+	FName NodeName;
+	FName Namespace;
+	FString Message;
+
+	ELogVerbosity::Type Verbosity = ELogVerbosity::NoLogging;
+};
+
+typedef TSharedPtr<FPCGLogListViewItem> PCGLogListViewItemPtr;
+
+class SPCGLogListViewItemRow : public SMultiColumnTableRow<PCGLogListViewItemPtr>
+{
+public:
+	SLATE_BEGIN_ARGS(SPCGLogListViewItemRow) {}
+	SLATE_END_ARGS()
+
+	void Construct(const FArguments& InArgs, const TSharedRef<STableViewBase>& InOwnerTableView, const PCGLogListViewItemPtr& Item);
+
+	virtual TSharedRef<SWidget> GenerateWidgetForColumn(const FName& ColumnId) override;
+
+protected:
+	PCGLogListViewItemPtr InternalItem;
+};
+
+class SPCGEditorGraphLogView : public SCompoundWidget
+{
+public:
+	SLATE_BEGIN_ARGS(SPCGEditorGraphLogView) { }
+	SLATE_END_ARGS()
+
+	void Construct(const FArguments& InArgs, TSharedPtr<FPCGEditor> InPCGEditor);
+
+private:
+	TSharedRef<SHeaderRow> CreateHeaderRowWidget();
+	
+	// Callbacks
+	FReply Refresh();
+	FReply Clear();
+	TSharedRef<ITableRow> OnGenerateRow(PCGLogListViewItemPtr Item, const TSharedRef<STableViewBase>& OwnerTable) const;
+	void OnItemDoubleClicked(PCGLogListViewItemPtr Item);
+	void OnSortColumnHeader(const EColumnSortPriority::Type SortPriority, const FName& ColumnId, const EColumnSortMode::Type NewSortMode);
+	EColumnSortMode::Type GetColumnSortMode(const FName ColumnId) const;
+
+	/** Pointer back to the PCG editor that owns us */
+	TWeakPtr<FPCGEditor> PCGEditorPtr;
+
+	/** Cached PCGGraph being viewed */
+	UPCGEditorGraph* PCGEditorGraph = nullptr;
+
+	TSharedPtr<SHeaderRow> ListViewHeader;
+	TSharedPtr<SListView<PCGLogListViewItemPtr>> ListView;
+	TArray<PCGLogListViewItemPtr> ListViewItems;
+
+	// To allow sorting
+	FName SortingColumn = PCGEditorGraphLogView::NAME_Order;
+	EColumnSortMode::Type SortMode = EColumnSortMode::Type::Ascending;
+};
