[PCG] Revamp property/attribute wrapper

* Avoid multiple dispatch and use of virtual classes instead.
* Should be easier to extend the system to other data in the future
* Performances are on par or slightly better

diff --git a/Engine/Plugins/Experimental/PCG/Source/PCG/Private/Elements/Metadata/PCGMetadataOpElementBase.cpp b/Engine/Plugins/Experimental/PCG/Source/PCG/Private/Elements/Metadata/PCGMetadataOpElementBase.cpp
--- a/Engine/Plugins/Experimental/PCG/Source/PCG/Private/Elements/Metadata/PCGMetadataOpElementBase.cpp
+++ b/Engine/Plugins/Experimental/PCG/Source/PCG/Private/Elements/Metadata/PCGMetadataOpElementBase.cpp
@@ -8,6 +8,8 @@
 #include "Elements/Metadata/PCGMetadataElementCommon.h"
 #include "Helpers/PCGSettingsHelpers.h"
 #include "Metadata/PCGMetadata.h"
+#include "Metadata/Accessors/PCGAttributeAccessor.h"
+#include "Metadata/Accessors/PCGAttributeAccessorHelpers.h"
 
 #include UE_INLINE_GENERATED_CPP_BY_NAME(PCGMetadataOpElementBase)
 
@@ -80,8 +82,10 @@ bool FPCGMetadataElementBase::ExecuteInternal(FPCGContext* Context) const
 
 	// Gathering all the inputs metadata
 	TArray<const UPCGMetadata*> SourceMetadata;
+	TArray<const FPCGMetadataAttributeBase*> SourceAttribute;
 	TArray<FPCGTaggedData> InputTaggedData;
 	SourceMetadata.SetNum(NumberOfInputs);
+	SourceAttribute.SetNum(NumberOfInputs);
 	InputTaggedData.SetNum(NumberOfInputs);
 
 	for (uint32 i = 0; i < NumberOfInputs; ++i)
@@ -113,13 +117,20 @@ bool FPCGMetadataElementBase::ExecuteInternal(FPCGContext* Context) const
 	}
 
 	FOperationData OperationData;
-	OperationData.InputIterators.SetNum(NumberOfInputs);
+	OperationData.InputAccessors.SetNum(NumberOfInputs);
+	OperationData.InputKeys.SetNum(NumberOfInputs);
 	TArray<int32> NumberOfElements;
 	NumberOfElements.SetNum(NumberOfInputs);
 
 	OperationData.MostComplexInputType = (uint16)EPCGMetadataTypes::Unknown;
 	OperationData.NumberOfElementsToProcess = -1;
 
+	bool bNoOperationNeeded = false;
+
+	// Use this name to forward it to the output if needed.
+	// Only set it if the first input is an attribute.
+	FName InputName = NAME_None;
+
 	for (uint32 i = 0; i < NumberOfInputs; ++i)
 	{
 		// First we verify if the input data match the first one.
@@ -132,21 +143,30 @@ bool FPCGMetadataElementBase::ExecuteInternal(FPCGContext* Context) const
 		}
 
 		FPCGAttributePropertySelector InputSource = Settings->GetInputSource(i);
+		// Make sure to update it to transform it into a property if needed.
+		// TODO: Remove when it will be handled in the UI widget logic.
+		InputSource.Update();
 
 		if (InputSource.Selection == EPCGAttributePropertySelection::Attribute && InputSource.AttributeName == NAME_None)
 		{
 			InputSource.AttributeName = SourceMetadata[i]->GetLatestAttributeNameOrNone();
 		}
 
-		OperationData.InputIterators[i] = PCGMetadataAttributeWrapper::CreateIteratorWrapper(InputTaggedData[i].Data, InputSource);
+		if (i == 0 && InputSource.Selection == EPCGAttributePropertySelection::Attribute)
+		{
+			InputName = InputSource.AttributeName;
+		}
+
+		OperationData.InputAccessors[i] = PCGAttributeAccessorHelpers::CreateConstAccessor(InputTaggedData[i].Data, InputSource);
+		OperationData.InputKeys[i] = PCGAttributeAccessorHelpers::CreateConstKeys(InputTaggedData[i].Data, InputSource);
 
-		if (!OperationData.InputIterators[i].IsValid())
+		if (!OperationData.InputAccessors[i].IsValid() || !OperationData.InputKeys[i].IsValid())
 		{
 			PCGE_LOG(Error, "Attribute/Property %s does not exist for input %d", *InputSource.GetName().ToString(), i);
 			return true;
 		}
 
-		uint16 AttributeTypeId = OperationData.InputIterators[i].GetType();
+		uint16 AttributeTypeId = OperationData.InputAccessors[i]->GetUnderlyingType();
 
 		// Then verify that the type is OK
 		bool bHasSpecialRequirement = false;
@@ -170,29 +190,64 @@ bool FPCGMetadataElementBase::ExecuteInternal(FPCGContext* Context) const
 			}
 		}
 
-		NumberOfElements[i] = OperationData.InputIterators[i].Num();
+		NumberOfElements[i] = OperationData.InputKeys[i]->GetNum();
 
 		if (OperationData.NumberOfElementsToProcess == -1)
 		{
 			OperationData.NumberOfElementsToProcess = NumberOfElements[i];
 		}
 
+		// There is nothing to do if one input doesn't have any element to process.
+		// Therefore mark that we have nothing to do and early out.
+		if (NumberOfElements[i] == 0)
+		{
+			PCGE_LOG(Verbose, "No elements in input %d.", i);
+			bNoOperationNeeded = true;
+			break;
+		}
+
 		// Verify that the number of elements makes sense
 		if (OperationData.NumberOfElementsToProcess % NumberOfElements[i] != 0)
 		{
 			PCGE_LOG(Error, "Mismatch between the number of elements in input 0 (%d) and in input %d (%d).", OperationData.NumberOfElementsToProcess, i, NumberOfElements[i]);
 			return true;
 		}
+
+		if (InputSource.Selection == EPCGAttributePropertySelection::Attribute)
+		{
+			SourceAttribute[i] = SourceMetadata[i]->GetConstAttribute(InputSource.GetName());
+		}
+		else
+		{
+			SourceAttribute[i] = nullptr;
+		}
+	}
+
+	// If no operation is needed, just forward input 0
+	if (bNoOperationNeeded)
+	{
+		for (uint32 OutputIndex = 0; OutputIndex < Settings->GetOutputPinNum(); ++OutputIndex)
+		{
+			FPCGTaggedData& OutputData = Outputs.Add_GetRef(InputTaggedData[0]);
+			OutputData.Pin = Settings->GetOutputPinLabel(OutputIndex);
+		}
+
+		return true;
 	}
 
+
 	// At this point, we verified everything, so we can go forward with the computation, depending on the most complex type
 	// So first forward outputs and create the attribute
-	OperationData.OutputIterators.SetNum(Settings->GetOutputPinNum());
+	OperationData.OutputAccessors.SetNum(Settings->GetOutputPinNum());
+	OperationData.OutputKeys.SetNum(Settings->GetOutputPinNum());
 
 	FPCGAttributePropertySelector OutputTarget = Settings->OutputTarget;
+	// Make sure to update it to transform it into a property if needed.
+	OutputTarget.Update();
+
 	if (OutputTarget.Selection == EPCGAttributePropertySelection::Attribute && OutputTarget.AttributeName == NAME_None)
 	{
-		OutputTarget.AttributeName = OperationData.InputIterators[0].GetWrapper().GetName();
+		OutputTarget.AttributeName = InputName;
 	}
 
 	auto CreateAttribute = [&](uint32 OutputIndex, auto DummyOutValue) -> bool
@@ -206,57 +261,50 @@ bool FPCGMetadataElementBase::ExecuteInternal(FPCGContext* Context) const
 
 		FName OutputName = OutputTarget.GetName();
 
-		bool bShouldClearOrCreateAttribute = false;
-		bool bShouldCopyEntryMap = false;
+		if (OutputTarget.Selection == EPCGAttributePropertySelection::Attribute)
+		{
+			// In case of attribute, there is no point of failure before duplicating. So duplicate, create the attribute and then the accessor.
+			PCGMetadataElementCommon::DuplicateTaggedData(InputTaggedData[0], OutputData, OutMetadata);
+			FPCGMetadataAttributeBase* OutputAttribute = PCGMetadataElementCommon::ClearOrCreateAttribute(OutMetadata, OutputName, AttributeType{});
+			if (!OutputAttribute)
+			{
+				return false;
+			}
+
+			// And copy the mapping from the original attribute, if it is not points
+			if (!InputTaggedData[0].Data->IsA<UPCGPointData>() && SourceMetadata[0] && SourceAttribute[0])
+			{
+				PCGMetadataElementCommon::CopyEntryToValueKeyMap(SourceMetadata[0], SourceAttribute[0], OutputAttribute);
+			}
 
-		if (UPCGPointData* PointData = Cast<UPCGPointData>(InputTaggedData[0].Data))
+			OperationData.OutputAccessors[OutputIndex] = PCGAttributeAccessorHelpers::CreateAccessor(Cast<UPCGData>(OutputData.Data), OutputTarget);
+		}
+		else if (OutputTarget.Selection == EPCGAttributePropertySelection::PointProperty)
 		{
-			int16 PropertyType = -1;
-			if (PCGMetadataAttributeWrapper::IsPropertyWithType(InputTaggedData[0].Data, OutputName, &PropertyType))
+			// In case of property, we need to validate that the property can accept the output type. Verify this before duplicating.
+			OperationData.OutputAccessors[OutputIndex] = PCGAttributeAccessorHelpers::CreateAccessor(Cast<UPCGData>(OutputData.Data), OutputTarget);
+
+			if (OperationData.OutputAccessors[OutputIndex].IsValid())
 			{
 				// We matched a property, check if the output type is valid
-				if (!PCG::Private::IsBroadcastable(PCG::Private::MetadataTypes<AttributeType>::Id, PropertyType))
+				if (!PCG::Private::IsBroadcastable(PCG::Private::MetadataTypes<AttributeType>::Id, OperationData.OutputAccessors[OutputIndex]->GetUnderlyingType()))
 				{
 					PCGE_LOG(Error, "Property %s cannot be broadcasted to match types for input", *OutputName.ToString());
 					return false;
 				}
-			}
-			else
-			{
-				// Otherwise it is an attribute, we will create a clean one, so output type will match.
-				bShouldClearOrCreateAttribute = true;
+
+				PCGMetadataElementCommon::DuplicateTaggedData(InputTaggedData[0], OutputData, OutMetadata);
 			}
 		}
-		else
-		{
-			// If it is not a point data, we will only deal with attributes, and the entry map will be necessary.
-			bShouldClearOrCreateAttribute = true;
-			bShouldCopyEntryMap = true;
-		}
 
-		PCGMetadataElementCommon::DuplicateTaggedData(InputTaggedData[0], OutputData, OutMetadata);
-		if (bShouldClearOrCreateAttribute)
+		if (!OperationData.OutputAccessors[OutputIndex].IsValid())
 		{
-			AttributeType DefaultValue{};
-			OperationData.InputIterators[0].GetWrapper().Get<AttributeType>(PCGInvalidEntryKey, DefaultValue);
-
-			FPCGMetadataAttributeBase* OutputAttribute = PCGMetadataElementCommon::ClearOrCreateAttribute(OutMetadata, OutputName, DefaultValue);
-
-			if (!OutputAttribute)
-			{
-				return false;
-			}
-
-			if (bShouldCopyEntryMap)
-			{
-				PCGMetadataElementCommon::CopyEntryToValueKeyMap(SourceMetadata[0], OperationData.InputIterators[0].GetWrapper().GetAttribute(), OutputAttribute);
-			}
+			return false;
 		}
 
-		// Cast to use the non-const version of this method.
-		OperationData.OutputIterators[OutputIndex] = PCGMetadataAttributeWrapper::CreateIteratorWrapper(Cast<UPCGData>(OutputData.Data), OutputTarget);
+		OperationData.OutputKeys[OutputIndex] = PCGAttributeAccessorHelpers::CreateKeys(Cast<UPCGData>(OutputData.Data), OutputTarget);
 
-		return OperationData.OutputIterators[OutputIndex].IsValid();
+		return OperationData.OutputKeys[OutputIndex].IsValid();
 	};
 
 	auto CreateAllSameAttributes = [&](auto DummyOutValue) -> bool
diff --git a/Engine/Plugins/Experimental/PCG/Source/PCG/Private/Helpers/PCGPropertyHelpers.cpp b/Engine/Plugins/Experimental/PCG/Source/PCG/Private/Helpers/PCGPropertyHelpers.cpp
new file mode 100644
--- /dev/null
+++ b/Engine/Plugins/Experimental/PCG/Source/PCG/Private/Helpers/PCGPropertyHelpers.cpp
@@ -0,0 +1,21 @@
+// Copyright Epic Games, Inc. All Rights Reserved.
+
+#include "Helpers/PCGPropertyHelpers.h"
+
+EPCGMetadataTypes PCGPropertyHelpers::GetMetadataTypeFromProperty(const FProperty* InProperty)
+{
+	if (!InProperty)
+	{
+		return EPCGMetadataTypes::Unknown;
+	}
+
+	// Object are not yet supported as accessors
+	if (const FObjectPropertyBase* ObjectProperty = CastField<FObjectPropertyBase>(InProperty))
+	{
+		return EPCGMetadataTypes::String;
+	}
+
+	TUniquePtr<IPCGAttributeAccessor> PropertyAccessor = PCGAttributeAccessorHelpers::CreatePropertyAccessor(InProperty);
+
+	return PropertyAccessor.IsValid() ? EPCGMetadataTypes(PropertyAccessor->GetUnderlyingType()) : EPCGMetadataTypes::Unknown;
+}
diff --git a/Engine/Plugins/Experimental/PCG/Source/PCG/Private/Metadata/Accessors/IPCGAttributeAccessor.cpp b/Engine/Plugins/Experimental/PCG/Source/PCG/Private/Metadata/Accessors/IPCGAttributeAccessor.cpp
new file mode 100644
--- /dev/null
+++ b/Engine/Plugins/Experimental/PCG/Source/PCG/Private/Metadata/Accessors/IPCGAttributeAccessor.cpp
@@ -0,0 +1,10 @@
+// Copyright Epic Games, Inc. All Rights Reserved.
+
+#include "Metadata/Accessors/IPCGAttributeAccessor.h"
+
+// Specialization of IPCGAttributeAccessor::Get and IPCGAttributeAccessor::Set, for all supported types.
+#define IACCESSOR_DECL(T) \
+template <> bool IPCGAttributeAccessor::GetRange<T>(TArrayView<T> OutValues, int32 Index, const IPCGAttributeAccessorKeys& Keys, EPCGAttributeAccessorFlags Flags) const { return GetRange##T(OutValues, Index, Keys, Flags); } \
+template <> bool IPCGAttributeAccessor::SetRange<T>(TArrayView<const T> InValues, int32 Index, IPCGAttributeAccessorKeys& Keys, EPCGAttributeAccessorFlags Flags) { return SetRange##T(InValues, Index, Keys, Flags); }
+PCG_FOREACH_SUPPORTEDTYPES(IACCESSOR_DECL);
+#undef IACCESSOR_DECL
\ No newline at end of file
diff --git a/Engine/Plugins/Experimental/PCG/Source/PCG/Private/Metadata/Accessors/PCGAttributeAccessorHelpers.cpp b/Engine/Plugins/Experimental/PCG/Source/PCG/Private/Metadata/Accessors/PCGAttributeAccessorHelpers.cpp
new file mode 100644
--- /dev/null
+++ b/Engine/Plugins/Experimental/PCG/Source/PCG/Private/Metadata/Accessors/PCGAttributeAccessorHelpers.cpp
@@ -0,0 +1,259 @@
+// Copyright Epic Games, Inc. All Rights Reserved.
+
+#include "Metadata/Accessors/PCGAttributeAccessorHelpers.h"
+
+#include "PCGPoint.h"
+#include "PCGParamData.h"
+#include "Data/PCGPointData.h"
+#include "Data/PCGSpatialData.h"
+#include "Helpers/PCGPropertyHelpers.h"
+#include "Metadata/PCGMetadataAttributeTpl.h"
+#include "Metadata/PCGMetadataAttributeTraits.h"
+#include "Metadata/PCGAttributePropertySelector.h"
+#include "Metadata/Accessors/PCGPropertyAccessor.h"
+#include "Metadata/Accessors/PCGAttributeAccessor.h"
+#include "Metadata/Accessors/PCGAttributeAccessorKeys.h"
+
+#include "UObject/UnrealType.h"
+#include "UObject/EnumProperty.h"
+
+namespace PCGAttributeAccessorHelpers
+{
+	void ExtractMetadataAtribute(UPCGData* InData, FName Name, UPCGMetadata*& OutMetadata, FPCGMetadataAttributeBase*& OutAttribute)
+	{
+		OutMetadata = nullptr;
+		OutAttribute = nullptr;
+
+		if (UPCGSpatialData* SpatialData = Cast<UPCGSpatialData>(InData))
+		{
+			OutMetadata = SpatialData->Metadata;
+		}
+		else if (UPCGParamData* ParamData = Cast<UPCGParamData>(InData))
+		{
+			OutMetadata = ParamData->Metadata;
+		}
+
+		if (OutMetadata)
+		{
+			OutAttribute = OutMetadata->GetMutableAttribute(Name);
+		}
+	}
+
+	void ExtractMetadataAtribute(const UPCGData* InData, FName Name, UPCGMetadata const*& OutMetadata, FPCGMetadataAttributeBase const*& OutAttribute)
+	{
+		UPCGMetadata* Metadata = nullptr;
+		FPCGMetadataAttributeBase* Attribute = nullptr;
+		ExtractMetadataAtribute(const_cast<UPCGData*>(InData), Name, Metadata, Attribute);
+		OutMetadata = Metadata;
+		OutAttribute = Attribute;
+	}
+}
+
+TUniquePtr<IPCGAttributeAccessor> PCGAttributeAccessorHelpers::CreatePropertyAccessor(const FProperty* InProperty)
+{
+	if (const FNumericProperty* NumericProperty = CastField<FNumericProperty>(InProperty))
+	{
+		if (NumericProperty->IsFloatingPoint())
+		{
+			return MakeUnique<FPCGNumericPropertyAccessor<double>>(NumericProperty);
+		}
+		else if (NumericProperty->IsInteger())
+		{
+			return MakeUnique<FPCGNumericPropertyAccessor<int64>>(NumericProperty);
+		}
+	}
+	else if (const FBoolProperty* BoolProperty = CastField<FBoolProperty>(InProperty))
+	{
+		return MakeUnique<FPCGPropertyAccessor<bool>>(InProperty);
+	}
+	else if (const FStrProperty* StringProperty = CastField<FStrProperty>(InProperty))
+	{
+		return MakeUnique<FPCGPropertyAccessor<FString>>(InProperty);
+	}
+	else if (const FNameProperty* NameProperty = CastField<FNameProperty>(InProperty))
+	{
+		return MakeUnique<FPCGPropertyAccessor<FName>>(InProperty);
+	}
+	else if (const FEnumProperty* EnumProperty = CastField<FEnumProperty>(InProperty))
+	{
+		return MakeUnique<FPCGEnumPropertyAccessor>(EnumProperty);
+	}
+	else if (const FStructProperty* StructProperty = CastField<FStructProperty>(InProperty))
+	{
+		if (StructProperty->Struct == TBaseStructure<FVector>::Get())
+		{
+			return MakeUnique<FPCGPropertyAccessor<FVector>>(InProperty);
+		}
+		else if (StructProperty->Struct == TBaseStructure<FVector4>::Get())
+		{
+			return MakeUnique<FPCGPropertyAccessor<FVector4>>(InProperty);
+		}
+		else if (StructProperty->Struct == TBaseStructure<FQuat>::Get())
+		{
+			return MakeUnique<FPCGPropertyAccessor<FQuat>>(InProperty);
+		}
+		else if (StructProperty->Struct == TBaseStructure<FTransform>::Get())
+		{
+			return MakeUnique<FPCGPropertyAccessor<FTransform>>(InProperty);
+		}
+		else if (StructProperty->Struct == TBaseStructure<FRotator>::Get())
+		{
+			return MakeUnique<FPCGPropertyAccessor<FRotator>>(InProperty);
+		}
+		else if (StructProperty->Struct == TBaseStructure<FVector2D>::Get())
+		{
+			return MakeUnique<FPCGPropertyAccessor<FVector2D>>(InProperty);
+		}
+		else if (StructProperty->Struct == TBaseStructure<FSoftObjectPath>::Get())
+		{
+			return MakeUnique<FPCGPropertyPathAccessor<FSoftObjectPath>>(InProperty);
+		}
+		else if (StructProperty->Struct == TBaseStructure<FSoftClassPath>::Get())
+		{
+			return MakeUnique<FPCGPropertyPathAccessor<FSoftClassPath>>(InProperty);
+		}
+	}
+
+	return TUniquePtr<IPCGAttributeAccessor>();
+}
+
+TUniquePtr<const IPCGAttributeAccessor> PCGAttributeAccessorHelpers::CreateConstAccessor(const UPCGData* InData, const FPCGAttributePropertySelector& InSelector)
+{
+	FName Name = InSelector.GetName();
+
+	if (InSelector.Selection == EPCGAttributePropertySelection::PointProperty)
+	{
+		if (const UPCGPointData* PointData = Cast<const UPCGPointData>(InData))
+		{
+			if (const FProperty* Property = FPCGPoint::StaticStruct()->FindPropertyByName(Name))
+			{
+				return CreatePropertyAccessor(Property);
+			}
+			else if (FPCGPoint::HasCustomPropertyGetterSetter(Name))
+			{
+				return FPCGPoint::CreateCustomPropertyAccessor(Name);
+			}
+		}
+	}
+
+	// At this point, it is not a point data or we didn't find a property.
+	// We can't continue if it is a property wanted.
+	if (InSelector.Selection == EPCGAttributePropertySelection::PointProperty)
+	{
+		UE_LOG(LogPCG, Error, TEXT("[PCGAttributeAccessorHelpers::CreateConstAccessor] Expected to select a property but the data doesn't support this property."));
+		return TUniquePtr<const IPCGAttributeAccessor>();
+	}
+
+	const UPCGMetadata* Metadata = nullptr;
+	const FPCGMetadataAttributeBase* Attribute = nullptr;
+
+	ExtractMetadataAtribute(InData, Name, Metadata, Attribute);
+
+	auto CreateTypedAccessor = [Attribute, Metadata](auto Dummy) -> TUniquePtr<const IPCGAttributeAccessor>
+	{
+		using AttributeType = decltype(Dummy);
+		return MakeUnique<FPCGAttributeAccessor<AttributeType>>(static_cast<const FPCGMetadataAttribute<AttributeType>*>(Attribute), Metadata);
+	};
+
+	if (Attribute && Metadata)
+	{
+		return PCGMetadataAttribute::CallbackWithRightType(Attribute->GetTypeId(), CreateTypedAccessor);
+	}
+	else
+	{
+		return TUniquePtr<const IPCGAttributeAccessor>();
+	}
+}
+
+TUniquePtr<IPCGAttributeAccessor> PCGAttributeAccessorHelpers::CreateAccessor(UPCGData* InData, const FPCGAttributePropertySelector& InSelector)
+{
+	FName Name = InSelector.GetName();
+
+	if (InSelector.Selection == EPCGAttributePropertySelection::PointProperty)
+	{
+		if (const UPCGPointData* PointData = Cast<const UPCGPointData>(InData))
+		{
+			if (const FProperty* Property = FPCGPoint::StaticStruct()->FindPropertyByName(Name))
+			{
+				return CreatePropertyAccessor(Property);
+			}
+			else if (FPCGPoint::HasCustomPropertyGetterSetter(Name))
+			{
+				return FPCGPoint::CreateCustomPropertyAccessor(Name);
+			}
+		}
+	}
+
+	// At this point, it is not a point data or we didn't find a property.
+	// We can't continue if it is a property wanted.
+	if (InSelector.Selection == EPCGAttributePropertySelection::PointProperty)
+	{
+		UE_LOG(LogPCG, Error, TEXT("[PCGAttributeAccessorHelpers::CreateAccessor] Expected to select a property but the data doesn't support this property."));
+		return TUniquePtr<IPCGAttributeAccessor>();
+	}
+
+	UPCGMetadata* Metadata = nullptr;
+	FPCGMetadataAttributeBase* Attribute = nullptr;
+
+	ExtractMetadataAtribute(InData, Name, Metadata, Attribute);
+
+	auto CreateTypedAccessor = [Attribute, Metadata](auto Dummy) -> TUniquePtr<IPCGAttributeAccessor>
+	{
+		using AttributeType = decltype(Dummy);
+		return MakeUnique<FPCGAttributeAccessor<AttributeType>>(static_cast<FPCGMetadataAttribute<AttributeType>*>(Attribute), Metadata);
+	};
+
+	if (Attribute && Metadata)
+	{
+		return PCGMetadataAttribute::CallbackWithRightType(Attribute->GetTypeId(), CreateTypedAccessor);
+	}
+	else
+	{
+		return TUniquePtr<IPCGAttributeAccessor>();
+	}
+}
+
+TUniquePtr<const IPCGAttributeAccessorKeys> PCGAttributeAccessorHelpers::CreateConstKeys(const UPCGData* InData, const FPCGAttributePropertySelector& InSelector)
+{
+	if (const UPCGPointData* PointData = Cast<const UPCGPointData>(InData))
+	{
+		return MakeUnique<FPCGAttributeAccessorKeysPoints>(PointData->GetPoints());
+	}
+
+	const UPCGMetadata* Metadata = nullptr;
+	const FPCGMetadataAttributeBase* Attribute = nullptr;
+
+	ExtractMetadataAtribute(InData, InSelector.GetName(), Metadata, Attribute);
+
+	if (Attribute)
+	{
+		return MakeUnique<FPCGAttributeAccessorKeysEntries>(Attribute);
+	}
+	else
+	{
+		return TUniquePtr<IPCGAttributeAccessorKeys>();
+	}
+}
+
+TUniquePtr<IPCGAttributeAccessorKeys> PCGAttributeAccessorHelpers::CreateKeys(UPCGData* InData, const FPCGAttributePropertySelector& InSelector)
+{
+	if (UPCGPointData* PointData = Cast<UPCGPointData>(InData))
+	{
+		TArrayView<FPCGPoint> View(PointData->GetMutablePoints());
+		return MakeUnique<FPCGAttributeAccessorKeysPoints>(View);
+	}
+
+	UPCGMetadata* Metadata = nullptr;
+	FPCGMetadataAttributeBase* Attribute = nullptr;
+
+	ExtractMetadataAtribute(InData, InSelector.GetName(), Metadata, Attribute);
+
+	if (Attribute)
+	{
+		return MakeUnique<FPCGAttributeAccessorKeysEntries>(Attribute);
+	}
+	else
+	{
+		return TUniquePtr<IPCGAttributeAccessorKeys>();
+	}
+}
diff --git a/Engine/Plugins/Experimental/PCG/Source/PCG/Private/Metadata/Accessors/PCGAttributeAccessorKeys.cpp b/Engine/Plugins/Experimental/PCG/Source/PCG/Private/Metadata/Accessors/PCGAttributeAccessorKeys.cpp
new file mode 100644
--- /dev/null
+++ b/Engine/Plugins/Experimental/PCG/Source/PCG/Private/Metadata/Accessors/PCGAttributeAccessorKeys.cpp
@@ -0,0 +1,93 @@
+// Copyright Epic Games, Inc. All Rights Reserved.
+
+#include "Metadata/Accessors/PCGAttributeAccessorKeys.h"
+#include "Metadata/PCGMetadataAttribute.h"
+
+#include "Templates/Function.h"
+
+////////////////////////////////////////////////////////////////////
+
+FPCGAttributeAccessorKeysEntries::FPCGAttributeAccessorKeysEntries(const FPCGMetadataAttributeBase* Attribute)
+	: IPCGAttributeAccessorKeys(/*bInReadOnly=*/ false)
+{
+	const FPCGMetadataAttributeBase* Current = Attribute;
+	while (Current)
+	{
+		TArray<PCGMetadataEntryKey> Temp;
+		Current->GetEntryToValueKeyMap_NotThreadSafe().GenerateKeyArray(Temp);
+		Entries.Append(Temp);
+		Current = Current->GetParent();
+	}
+
+	// If the attribute doesn't have any entry, we will always take the default value.
+	if (Entries.IsEmpty())
+	{
+		Entries.Add(PCGInvalidEntryKey);
+	}
+}
+
+FPCGAttributeAccessorKeysEntries::FPCGAttributeAccessorKeysEntries(PCGMetadataEntryKey EntryKey)
+	: IPCGAttributeAccessorKeys(/*bInReadOnly=*/ false)
+{
+	Entries.Add(EntryKey);
+}
+
+bool FPCGAttributeAccessorKeysEntries::GetMetadataEntryKeys(int32 InStart, TArrayView<PCGMetadataEntryKey*>& OutEntryKeys)
+{
+	return PCGAttributeAccessorKeys::GetKeys(Entries, InStart, OutEntryKeys, [](PCGMetadataEntryKey& Key) -> PCGMetadataEntryKey* { return &Key; });
+}
+
+bool FPCGAttributeAccessorKeysEntries::GetMetadataEntryKeys(int32 InStart, TArrayView<const PCGMetadataEntryKey*>& OutEntryKeys) const
+{
+	return PCGAttributeAccessorKeys::GetKeys(Entries, InStart, OutEntryKeys, [](const PCGMetadataEntryKey& Key) -> const PCGMetadataEntryKey* { return &Key; });
+}
+
+////////////////////////////////////////////////////////////////////
+
+FPCGAttributeAccessorKeysPoints::FPCGAttributeAccessorKeysPoints(const TArrayView<FPCGPoint>& InPoints)
+	: IPCGAttributeAccessorKeys(/*bInReadOnly=*/ false)
+	, Points(InPoints)
+{}
+
+FPCGAttributeAccessorKeysPoints::FPCGAttributeAccessorKeysPoints(const TArrayView<const FPCGPoint>& InPoints)
+	: IPCGAttributeAccessorKeys(/*bInReadOnly=*/ true)
+	, Points(const_cast<FPCGPoint*>(InPoints.GetData()), InPoints.Num())
+{}
+
+FPCGAttributeAccessorKeysPoints::FPCGAttributeAccessorKeysPoints(FPCGPoint& InPoint)
+	: FPCGAttributeAccessorKeysPoints(TArrayView<FPCGPoint>(&InPoint, 1))
+{}
+
+FPCGAttributeAccessorKeysPoints::FPCGAttributeAccessorKeysPoints(const FPCGPoint& InPoint)
+	: FPCGAttributeAccessorKeysPoints(TArrayView<const FPCGPoint>(&InPoint, 1))
+{}
+
+bool FPCGAttributeAccessorKeysPoints::GetPointKeys(int32 InStart, TArrayView<FPCGPoint*>& OutPoints)
+{
+	return PCGAttributeAccessorKeys::GetKeys(Points, InStart, OutPoints, [](FPCGPoint& Point) -> FPCGPoint* { return &Point; });
+}
+
+bool FPCGAttributeAccessorKeysPoints::GetPointKeys(int32 InStart, TArrayView<const FPCGPoint*>& OutPoints) const
+{
+	return PCGAttributeAccessorKeys::GetKeys(Points, InStart, OutPoints, [](const FPCGPoint& Point) -> const FPCGPoint* { return &Point; });
+}
+
+bool FPCGAttributeAccessorKeysPoints::GetGenericObjectKeys(int32 InStart, TArrayView<void*>& OutObjects)
+{
+	return PCGAttributeAccessorKeys::GetKeys(Points, InStart, OutObjects, [](FPCGPoint& Point) -> void* { return &Point; });
+}
+
+bool FPCGAttributeAccessorKeysPoints::GetGenericObjectKeys(int32 InStart, TArrayView<const void*>& OutObjects) const
+{
+	return PCGAttributeAccessorKeys::GetKeys(Points, InStart, OutObjects, [](const FPCGPoint& Point) -> const void* { return &Point; });
+}
+
+bool FPCGAttributeAccessorKeysPoints::GetMetadataEntryKeys(int32 InStart, TArrayView<PCGMetadataEntryKey*>& OutEntryKeys)
+{
+	return PCGAttributeAccessorKeys::GetKeys(Points, InStart, OutEntryKeys, [](FPCGPoint& Point) -> PCGMetadataEntryKey* { return &(Point.MetadataEntry); });
+}
+
+bool FPCGAttributeAccessorKeysPoints::GetMetadataEntryKeys(int32 InStart, TArrayView<const PCGMetadataEntryKey*>& OutEntryKeys) const
+{
+	return PCGAttributeAccessorKeys::GetKeys(Points, InStart, OutEntryKeys, [](const FPCGPoint& Point) -> const PCGMetadataEntryKey* { return &(Point.MetadataEntry); });
+}
\ No newline at end of file
diff --git a/Engine/Plugins/Experimental/PCG/Source/PCG/Private/Metadata/Accessors/PCGPropertyAccessor.cpp b/Engine/Plugins/Experimental/PCG/Source/PCG/Private/Metadata/Accessors/PCGPropertyAccessor.cpp
new file mode 100644
--- /dev/null
+++ b/Engine/Plugins/Experimental/PCG/Source/PCG/Private/Metadata/Accessors/PCGPropertyAccessor.cpp
@@ -0,0 +1,47 @@
+// Copyright Epic Games, Inc. All Rights Reserved.
+
+#include "Metadata/Accessors/PCGPropertyAccessor.h"
+
+#include "UObject/EnumProperty.h"
+
+inline TArray<const void*> PCGPropertyAccessor::GetContainerKeys(int32 Index, int32 Range, const IPCGAttributeAccessorKeys& Keys)
+{
+	TArray<const void*> ContainerKeys;
+	ContainerKeys.SetNum(Range);
+	TArrayView<const void*> ContainerKeysView(ContainerKeys);
+	if (!Keys.GetKeys(Index, ContainerKeysView))
+	{
+		return TArray<const void*>{};
+	}
+
+	return ContainerKeys;
+}
+
+inline TArray<void*> PCGPropertyAccessor::GetContainerKeys(int32 Index, int32 Range, IPCGAttributeAccessorKeys& Keys)
+{
+	TArray<void*> ContainerKeys;
+	ContainerKeys.SetNum(Range);
+	TArrayView<void*> ContainerKeysView(ContainerKeys);
+	if (!Keys.GetKeys(Index, ContainerKeysView))
+	{
+		return TArray<void*>{};
+	}
+
+	return ContainerKeys;
+}
+
+bool FPCGEnumPropertyAccessor::GetRangeImpl(TArrayView<int64> OutValues, int32 Index, const IPCGAttributeAccessorKeys& Keys) const
+{
+	return PCGPropertyAccessor::IterateGet(Property, OutValues, Index, Keys, [this](const void* PropertyAddressData) -> Type
+		{
+			return Property->GetUnderlyingProperty()->GetSignedIntPropertyValue(PropertyAddressData);
+		});
+}
+
+bool FPCGEnumPropertyAccessor::SetRangeImpl(TArrayView<const int64> InValues, int32 Index, IPCGAttributeAccessorKeys& Keys, EPCGAttributeAccessorFlags)
+{
+	return PCGPropertyAccessor::IterateSet(Property, InValues, Index, Keys, [this](void* PropertyAddressData, const int64& Value) -> void
+		{
+			Property->GetUnderlyingProperty()->SetIntPropertyValue(PropertyAddressData, Value);
+		});
+}
\ No newline at end of file
diff --git a/Engine/Plugins/Experimental/PCG/Source/PCG/Private/Metadata/PCGAttributePropertySelector.cpp b/Engine/Plugins/Experimental/PCG/Source/PCG/Private/Metadata/PCGAttributePropertySelector.cpp
new file mode 100644
--- /dev/null
+++ b/Engine/Plugins/Experimental/PCG/Source/PCG/Private/Metadata/PCGAttributePropertySelector.cpp
@@ -0,0 +1,131 @@
+// Copyright Epic Games, Inc. All Rights Reserved.
+
+#include "Metadata/PCGAttributePropertySelector.h"
+#include "Metadata/PCGMetadataAttribute.h"
+
+FName FPCGAttributePropertySelector::GetName() const
+{
+	switch (Selection)
+	{
+	case EPCGAttributePropertySelection::PointProperty:
+	{
+		if (const UEnum* EnumPtr = StaticEnum<EPCGPointProperties>())
+		{
+			// Need to use the string version and not the name version, because the name verison has "EPCGPointProperties::" as a prefix.
+			return FName(EnumPtr->GetNameStringByValue((int64)PointProperty));
+		}
+		else
+		{
+			return NAME_None;
+		}
+	}
+	default:
+		return AttributeName;
+	}
+}
+
+bool FPCGAttributePropertySelector::SetPointProperty(EPCGPointProperties InPointProperty)
+{
+	if (Selection == EPCGAttributePropertySelection::PointProperty && InPointProperty == PointProperty)
+	{
+		// Nothing changed
+		return false;
+	}
+	else
+	{
+		Selection = EPCGAttributePropertySelection::PointProperty;
+		PointProperty = InPointProperty;
+		return true;
+	}
+}
+
+bool FPCGAttributePropertySelector::SetAttributeName(FName InAttributeName)
+{
+	if (Selection == EPCGAttributePropertySelection::Attribute && AttributeName == InAttributeName)
+	{
+		// Nothing changed
+		return false;
+	}
+	else
+	{
+		Selection = EPCGAttributePropertySelection::Attribute;
+		AttributeName = InAttributeName;
+		return true;
+	}
+}
+
+// TODO: To be removed when UI widget is done.
+void FPCGAttributePropertySelector::Update()
+{
+	const FString AttributeNameStr = AttributeName.ToString();
+
+	// If we attribute we select starts with a '$' and matches a property, we convert it to a property
+	if (!AttributeNameStr.IsEmpty() && AttributeNameStr[0] == TEXT('$'))
+	{
+		if (const UEnum* EnumPtr = StaticEnum<EPCGPointProperties>())
+		{
+			int32 Index = EnumPtr->GetIndexByNameString(AttributeNameStr.Mid(/*Start=*/1));
+			if (Index != INDEX_NONE)
+			{
+				SetPointProperty(static_cast<EPCGPointProperties>(EnumPtr->GetValueByIndex(Index)));
+				return;
+			}
+		}
+	}
+}
+
+#if WITH_EDITOR
+bool FPCGAttributePropertySelector::IsValid() const
+{
+	return (Selection != EPCGAttributePropertySelection::Attribute) || FPCGMetadataAttributeBase::IsValidName(AttributeName);
+}
+
+FText FPCGAttributePropertySelector::GetDisplayText() const
+{
+	const FName Name = GetName();
+
+	// Add a '$' if it is a property
+	if (Selection == EPCGAttributePropertySelection::PointProperty && (Name != NAME_None))
+	{
+		return FText::FromString(FString(TEXT("$")) + Name.ToString());
+	}
+	else
+	{
+		return FText::FromName(Name);
+	}
+}
+
+bool FPCGAttributePropertySelector::Update(FString NewValue)
+{
+	if (!NewValue.IsEmpty() && NewValue[0] == TEXT('$'))
+	{
+		if (const UEnum* EnumPtr = StaticEnum<EPCGPointProperties>())
+		{
+			int32 Index = EnumPtr->GetIndexByNameString(NewValue.Mid(/*Start=*/1));
+			if (Index != INDEX_NONE)
+			{
+				return SetPointProperty(static_cast<EPCGPointProperties>(EnumPtr->GetValueByIndex(Index)));
+			}
+		}
+	}
+
+	return SetAttributeName(NewValue.IsEmpty() ? NAME_None : FName(NewValue));
+}
+#endif // WITH_EDITOR
+
+///////////////////////////////////////////////////////////////////////
+
+bool UPCGAttributePropertySelectorBlueprintHelpers::SetPointProperty(FPCGAttributePropertySelector& Selector, EPCGPointProperties InPointProperty)
+{
+	return Selector.SetPointProperty(InPointProperty);
+}
+
+bool UPCGAttributePropertySelectorBlueprintHelpers::SetAttributeName(FPCGAttributePropertySelector& Selector, FName InAttributeName)
+{
+	return Selector.SetAttributeName(InAttributeName);
+}
+
+FName UPCGAttributePropertySelectorBlueprintHelpers::GetName(const FPCGAttributePropertySelector& Selector)
+{
+	return Selector.GetName();
+}
\ No newline at end of file
diff --git a/Engine/Plugins/Experimental/PCG/Source/PCG/Private/Metadata/PCGMetadata.cpp b/Engine/Plugins/Experimental/PCG/Source/PCG/Private/Metadata/PCGMetadata.cpp
--- a/Engine/Plugins/Experimental/PCG/Source/PCG/Private/Metadata/PCGMetadata.cpp
+++ b/Engine/Plugins/Experimental/PCG/Source/PCG/Private/Metadata/PCGMetadata.cpp
@@ -3,7 +3,7 @@
 #include "Metadata/PCGMetadata.h"
 
 #include "Misc/ScopeRWLock.h"
-#include "Metadata/PCGMetadataAttributeWrapper.h"
+#include "Helpers/PCGPropertyHelpers.h"
 
 #include UE_INLINE_GENERATED_CPP_BY_NAME(PCGMetadata)
 
@@ -490,7 +490,7 @@ namespace PCGMetadata
 				});
 		};
 
-		return PCGMetadataAttributeWrapper::GetPropertyValueWithCallback(DataPtr, InProperty, CreateAttributeAndSet);
+		return PCGPropertyHelpers::GetPropertyValueWithCallback(DataPtr, InProperty, CreateAttributeAndSet);
 	}
 }
 
@@ -504,9 +504,9 @@ bool UPCGMetadata::SetAttributeFromDataProperty(FName AttributeName, PCGMetadata
 	return PCGMetadata::SetAttributeFromPropertyHelper<void>(*this, AttributeName, EntryKey, Data, InProperty, bCreate);
 }
 
-void UPCGMetadata::CopyExistingAttribute(FName AttributeToCopy, FName NewAttributeName, bool bKeepParent)
+bool UPCGMetadata::CopyExistingAttribute(FName AttributeToCopy, FName NewAttributeName, bool bKeepParent)
 {
-	CopyAttribute(AttributeToCopy, NewAttributeName, bKeepParent, /*bCopyEntries=*/true, /*bCopyValues=*/true);
+	return CopyAttribute(AttributeToCopy, NewAttributeName, bKeepParent, /*bCopyEntries=*/true, /*bCopyValues=*/true) != nullptr;
 }
 
 FPCGMetadataAttributeBase* UPCGMetadata::CopyAttribute(FName AttributeToCopy, FName NewAttributeName, bool bKeepParent, bool bCopyEntries, bool bCopyValues)
@@ -551,12 +551,12 @@ FPCGMetadataAttributeBase* UPCGMetadata::CopyAttribute(const FPCGMetadataAttribu
 	return NewAttribute;
 }
 
-void UPCGMetadata::RenameAttribute(FName AttributeToRename, FName NewAttributeName)
+bool UPCGMetadata::RenameAttribute(FName AttributeToRename, FName NewAttributeName)
 {
 	if (!FPCGMetadataAttributeBase::IsValidName(NewAttributeName))
 	{
 		UE_LOG(LogPCG, Error, TEXT("New attribute name %s is not valid"), *NewAttributeName.ToString());
-		return;
+		return false;
 	}
 
 	bool bRenamed = false;
@@ -576,6 +576,8 @@ void UPCGMetadata::RenameAttribute(FName AttributeToRename, FName NewAttributeNa
 	{
 		UE_LOG(LogPCG, Warning, TEXT("Attribute %s does not exist and therefore cannot be renamed"), *AttributeToRename.ToString());
 	}
+
+	return bRenamed;
 }
 
 void UPCGMetadata::ClearAttribute(FName AttributeToClear)
diff --git a/Engine/Plugins/Experimental/PCG/Source/PCG/Private/Metadata/PCGMetadataAttributeWrapper.cpp b/Engine/Plugins/Experimental/PCG/Source/PCG/Private/Metadata/PCGMetadataAttributeWrapper.cpp
deleted file mode 100644
diff --git a/Engine/Plugins/Experimental/PCG/Source/PCG/Private/PCGPoint.cpp b/Engine/Plugins/Experimental/PCG/Source/PCG/Private/PCGPoint.cpp
--- a/Engine/Plugins/Experimental/PCG/Source/PCG/Private/PCGPoint.cpp
+++ b/Engine/Plugins/Experimental/PCG/Source/PCG/Private/PCGPoint.cpp
@@ -2,6 +2,8 @@
 
 #include "PCGPoint.h"
 
+#include "Metadata/Accessors/PCGCustomAccessor.h"
+
 #include UE_INLINE_GENERATED_CPP_BY_NAME(PCGPoint)
 
 namespace PCGPointCustomPropertyNames
@@ -36,13 +38,6 @@ FBoxSphereBounds FPCGPoint::GetDensityBounds() const
 	return FBoxSphereBounds(FBox((2 - Steepness) * BoundsMin, (2 - Steepness) * BoundsMax).TransformBy(Transform));
 }
 
-FPCGPoint::PointCustomPropertyGetterSetter::PointCustomPropertyGetterSetter(const FPCGPoint::PointCustomPropertyGetter& InGetter, const FPCGPoint::PointCustomPropertySetter& InSetter, int16 InType, FName InName)
-	: Getter(InGetter)
-	, Setter(InSetter)
-	, Type(InType)
-	, Name(InName)
-{}
-
 bool FPCGPoint::HasCustomPropertyGetterSetter(FName Name)
 {
 	return Name == PCGPointCustomPropertyNames::ExtentsName ||
@@ -52,53 +47,43 @@ bool FPCGPoint::HasCustomPropertyGetterSetter(FName Name)
 		Name == PCGPointCustomPropertyNames::ScaleName;
 }
 
-FPCGPoint::PointCustomPropertyGetterSetter FPCGPoint::CreateCustomPropertyGetterSetter(FName Name)
+TUniquePtr<IPCGAttributeAccessor> FPCGPoint::CreateCustomPropertyAccessor(FName Name)
 {
 	if (Name == PCGPointCustomPropertyNames::ExtentsName)
 	{
-		return FPCGPoint::PointCustomPropertyGetterSetter(
+		return MakeUnique<FPCGCustomPointAccessor<FVector>>(
 			[](const FPCGPoint& Point, void* OutValue) { *reinterpret_cast<FVector*>(OutValue) = Point.GetExtents(); return true; },
-			[](FPCGPoint& Point, const void* InValue) { Point.SetExtents(*reinterpret_cast<const FVector*>(InValue)); return true; },
-			(int16)PCG::Private::MetadataTypes<FVector>::Id,
-			PCGPointCustomPropertyNames::ExtentsName
+			[](FPCGPoint& Point, const void* InValue) { Point.SetExtents(*reinterpret_cast<const FVector*>(InValue)); return true; }
 		);
 	}
 	else if (Name == PCGPointCustomPropertyNames::LocalCenterName)
 	{
-		return FPCGPoint::PointCustomPropertyGetterSetter(
+		return MakeUnique<FPCGCustomPointAccessor<FVector>>(
 			[](const FPCGPoint& Point, void* OutValue) { *reinterpret_cast<FVector*>(OutValue) = Point.GetLocalCenter(); return true; },
-			[](FPCGPoint& Point, const void* InValue) { Point.SetLocalCenter(*reinterpret_cast<const FVector*>(InValue)); return true; },
-			(int16)PCG::Private::MetadataTypes<FVector>::Id,
-			PCGPointCustomPropertyNames::LocalCenterName
+			[](FPCGPoint& Point, const void* InValue) { Point.SetLocalCenter(*reinterpret_cast<const FVector*>(InValue)); return true; }
 		);
 	}
 	else if (Name == PCGPointCustomPropertyNames::PositionName)
 	{
-		return FPCGPoint::PointCustomPropertyGetterSetter(
+		return MakeUnique<FPCGCustomPointAccessor<FVector>>(
 			[](const FPCGPoint& Point, void* OutValue) { *reinterpret_cast<FVector*>(OutValue) = Point.Transform.GetLocation(); return true; },
-			[](FPCGPoint& Point, const void* InValue) { Point.Transform.SetLocation(*reinterpret_cast<const FVector*>(InValue)); return true; },
-			(int16)PCG::Private::MetadataTypes<FVector>::Id,
-			PCGPointCustomPropertyNames::PositionName
+			[](FPCGPoint& Point, const void* InValue) { Point.Transform.SetLocation(*reinterpret_cast<const FVector*>(InValue)); return true; }
 		);
 	}
 	else if (Name == PCGPointCustomPropertyNames::RotationName)
 	{
-		return FPCGPoint::PointCustomPropertyGetterSetter(
+		return MakeUnique<FPCGCustomPointAccessor<FQuat>>(
 			[](const FPCGPoint& Point, void* OutValue) { *reinterpret_cast<FQuat*>(OutValue) = Point.Transform.GetRotation(); return true; },
-			[](FPCGPoint& Point, const void* InValue) { Point.Transform.SetRotation(*reinterpret_cast<const FQuat*>(InValue)); return true; },
-			(int16)PCG::Private::MetadataTypes<FQuat>::Id,
-			PCGPointCustomPropertyNames::RotationName
+			[](FPCGPoint& Point, const void* InValue) { Point.Transform.SetRotation(*reinterpret_cast<const FQuat*>(InValue)); return true; }
 		);
 	}
 	else if (Name == PCGPointCustomPropertyNames::ScaleName)
 	{
-		return FPCGPoint::PointCustomPropertyGetterSetter(
+		return MakeUnique<FPCGCustomPointAccessor<FVector>>(
 			[](const FPCGPoint& Point, void* OutValue) { *reinterpret_cast<FVector*>(OutValue) = Point.Transform.GetScale3D(); return true; },
-			[](FPCGPoint& Point, const void* InValue) { Point.Transform.SetScale3D(*reinterpret_cast<const FVector*>(InValue)); return true; },
-			(int16)PCG::Private::MetadataTypes<FVector>::Id,
-			PCGPointCustomPropertyNames::ScaleName
+			[](FPCGPoint& Point, const void* InValue) { Point.Transform.SetScale3D(*reinterpret_cast<const FVector*>(InValue)); return true; }
 		);
 	}
 
-	return FPCGPoint::PointCustomPropertyGetterSetter();
+	return TUniquePtr<IPCGAttributeAccessor>();
 }
diff --git a/Engine/Plugins/Experimental/PCG/Source/PCG/Public/Elements/Metadata/PCGMetadataOpElementBase.h b/Engine/Plugins/Experimental/PCG/Source/PCG/Public/Elements/Metadata/PCGMetadataOpElementBase.h
--- a/Engine/Plugins/Experimental/PCG/Source/PCG/Public/Elements/Metadata/PCGMetadataOpElementBase.h
+++ b/Engine/Plugins/Experimental/PCG/Source/PCG/Public/Elements/Metadata/PCGMetadataOpElementBase.h
@@ -7,7 +7,9 @@
 
 #include "Metadata/PCGMetadataAttribute.h"
 #include "Metadata/PCGMetadataAttributeTpl.h"
-#include "Metadata/PCGMetadataAttributeWrapper.h"
+#include "Metadata/PCGAttributePropertySelector.h"
+#include "Metadata/Accessors/IPCGAttributeAccessor.h"
+#include "Metadata/Accessors/PCGAttributeAccessorKeys.h"
 
 #include "Containers/StaticArray.h"
 
@@ -113,6 +115,7 @@ public:
 	EPCGMetadataSettingsBaseMode Mode_DEPRECATED = EPCGMetadataSettingsBaseMode::Inferred;
 #endif // WITH_EDITORONLY_DATA
 
+	static constexpr uint32 MaxNumberOfInputs = 4;
 	static constexpr uint32 MaxNumberOfOutputs = 4;
 
 #if WITH_EDITORONLY_DATA
@@ -132,243 +135,271 @@ public:
 		uint16 OutputType;
 		const UPCGMetadataSettingsBase* Settings = nullptr;
 
-		TArray<FPCGPropertyAttributeIterator> InputIterators;
-		TArray<FPCGPropertyAttributeIterator> OutputIterators;
+		TArray<TUniquePtr<const IPCGAttributeAccessorKeys>> InputKeys;
+		TArray<TUniquePtr<IPCGAttributeAccessorKeys>> OutputKeys;
+
+		TArray<TUniquePtr<const IPCGAttributeAccessor>> InputAccessors;
+		TArray<TUniquePtr<IPCGAttributeAccessor>> OutputAccessors;
+
+		template <int32 NbInputs, int32 NbOutputs>
+		void Validate();
 	};
 
 protected:
 	virtual bool ExecuteInternal(FPCGContext* Context) const override;
 
-	virtual bool DoOperation(FOperationData& OperationData) const = 0;
+	virtual bool DoOperation(FOperationData& InOperationData) const = 0;
+
+	/**
+	* Generic method to factorise all the boilerplate code for a variable number of inputs/outputs
+	*/
+	template <typename... InputTypes, typename... Callbacks>
+	inline bool DoNAryOp(FOperationData& InOperationData, TTuple<Callbacks...>&& InCallbacks) const;
 
 	/* All operations can have a fixed number of inputs and a variable number of outputs.
 	* Each output need to have its own callback, all taking the exact number of "const InType&" as input
 	* and each can return a different output type.
 	*/
 	template <typename InType, typename... Callbacks>
-	bool DoUnaryOp(FOperationData& OperationData, Callbacks&& ...InCallbacks) const;
+	bool DoUnaryOp(FOperationData& InOperationData, Callbacks&& ...InCallbacks) const;
 
 	template <typename InType1, typename InType2, typename... Callbacks>
-	bool DoBinaryOp(FOperationData& OperationData, Callbacks&& ...InCallbacks) const;
+	bool DoBinaryOp(FOperationData& InOperationData, Callbacks&& ...InCallbacks) const;
 
 	template <typename InType1, typename InType2, typename InType3, typename... Callbacks>
-	bool DoTernaryOp(FOperationData& OperationData, Callbacks&& ...InCallbacks) const;
+	bool DoTernaryOp(FOperationData& InOperationData, Callbacks&& ...InCallbacks) const;
 
 	template <typename InType1, typename InType2, typename InType3, typename InType4, typename... Callbacks>
-	bool DoQuaternaryOp(FOperationData& OperationData, Callbacks&& ...InCallbacks) const;
+	bool DoQuaternaryOp(FOperationData& InOperationData, Callbacks&& ...InCallbacks) const;
 };
 
-template <typename InType, typename... Callbacks>
-inline bool FPCGMetadataElementBase::DoUnaryOp(FOperationData& OperationData, Callbacks&& ...InCallbacks) const
+template <int32 NbInputs, int32 NbOutputs>
+inline void FPCGMetadataElementBase::FOperationData::Validate()
 {
-	check(OperationData.InputIterators[0].IsValid());
-
-	const uint32 NumOutputs = OperationData.OutputIterators.Num();
-	constexpr uint32 NumCallbacks = (uint32)sizeof...(InCallbacks);
-	check(NumOutputs == NumCallbacks);
+	check(OutputKeys.Num() == NbOutputs);
 
-	InType DefaultValue{};
-	if (OperationData.InputIterators[0].GetWrapper().Get<InType>(PCGInvalidEntryKey, DefaultValue))
+	for (int32 i = 0; i < NbInputs; ++i)
 	{
-		// Fold expressions. It will un-roll the packed argument "InCallbacks", and we can use "InCallbacks" as the current unrolled argument.
-		// In this case, we have multiple callbacks that all take "const InType&" as input, but can return different outputs.
-		// We use this return value to know in which type we need to cast our attribute.
-		uint32 j = 0;
-		([&]
-		{
-			check(j < NumOutputs);
-			auto OutValue = InCallbacks(DefaultValue);
-			typedef decltype(OutValue) OutType;
-
-			OperationData.OutputIterators[j++].GetWrapper().Set<OutType>(PCGInvalidEntryKey, OutValue);
-		} (), ...);
+		check(InputAccessors[i].IsValid());
+		check(InputKeys[i].IsValid());
 	}
 
-	InType Value{};
-
-	for (int32 i = 0; i < OperationData.NumberOfElementsToProcess; ++i)
+	for (int32 j = 0; j < NbOutputs; ++j)
 	{
-		if (OperationData.InputIterators[0].GetAndAdvance<InType>(Value))
-		{
-			uint32 j = 0;
-			([&]
-			{
-				check(j < NumOutputs);
-				auto OutValue = InCallbacks(Value);
-				typedef decltype(OutValue) OutType;
-
-				OperationData.OutputIterators[j++].SetAndAdvance<OutType>(OutValue);
-			} (), ...);
-		}
+		check(OutputAccessors[j].IsValid());
+		check(OutputKeys[j].IsValid());
 	}
-
-	return true;
 }
 
-template <typename InType1, typename InType2, typename... Callbacks>
-inline bool FPCGMetadataElementBase::DoBinaryOp(FOperationData& OperationData, Callbacks&& ...InCallbacks) const
+/**
+* Heavy templated code to factorize the gathering of inputs, the application of callbacks and the set of outputs.
+* Note that the order of calls is from the bottom to the top: Operation -> Gather -> Apply.
+* So start at the bottom!
+*/
+namespace PCG
+{
+namespace Private
 {
-	check(OperationData.InputIterators[0].IsValid());
-	check(OperationData.InputIterators[1].IsValid());
+namespace NAryOperation
+{
+	/**
+	* Empty struct to pack/unpack types in templates.
+	*/
+	template <typename... T>
+	struct Signature {};
 
-	const uint32 NumOutputs = OperationData.OutputIterators.Num();
-	constexpr uint32 NumCallbacks = (uint32)sizeof...(InCallbacks);
-	check(NumOutputs == NumCallbacks);
+	constexpr int32 DefaultChunkSize = 256;
 
-	InType1 DefaultValue1{};
-	InType2 DefaultValue2{};
+	/**
+	* Set of options to know if we need to use the default key + flags for get and set.
+	*/
+	struct Options
+	{
+		EPCGAttributeAccessorFlags GetFlags;
+		EPCGAttributeAccessorFlags SetFlags;
+		bool bUseDefaultKey = false;
+	};
 
-	if (OperationData.InputIterators[0].GetWrapper().Get<InType1>(PCGInvalidEntryKey, DefaultValue1) &&
-		OperationData.InputIterators[1].GetWrapper().Get<InType2>(PCGInvalidEntryKey, DefaultValue2))
+	/**
+	* Finally we call our callback with the packed input values from InArgs, and set the output value in its accessor.
+	* We use OutputIndex to know which callback to get (and therefore which output to set).
+	* Note that we go in reverse order, and stop when OutputIndex is negative.
+	*/
+	template<int OutputIndex, typename ...Callbacks, typename ...Args>
+	bool Apply(FPCGMetadataElementBase::FOperationData& InOperationData, int32 StartIndex, int32 Range, const Options& InOptions, const TTuple<Callbacks...>& InCallbacks, Args&& ...InArgs)
 	{
-		// Fold expression, cf. Unary op
-		uint32 j = 0;
-		([&]
+		if constexpr (OutputIndex < 0)
+		{
+			return true;
+		}
+		else
 		{
-			check(j < NumOutputs);
-			auto OutValue = InCallbacks(DefaultValue1, DefaultValue2);
-			typedef decltype(OutValue) OutType;
+			if (Range <= 0)
+			{
+				return true;
+			}
 
-			OperationData.OutputIterators[j++].GetWrapper().Set<OutType>(PCGInvalidEntryKey, OutValue);
-		} (), ...);
-	}
+			// First compute the first element to get its type.
+			// We use argument unpacking, since our InArgs contains all the arrays with our values.
+			// If in InArgs we have TArray<int>& IntValues, TArray<float>& FloatValues, the unpack will do:
+			// (IntValues[0], FloatValues[0])
+			auto OutValue = InCallbacks.template Get<OutputIndex>()(InArgs[0]...);
+			using OutType = decltype(OutValue);
 
-	InType1 Value1{};
-	InType2 Value2{};
+			TArray<OutType, TInlineAllocator<DefaultChunkSize>> OutValues;
+			OutValues.SetNum(Range);
+			OutValues[0] = OutValue;
 
-	for (int32 i = 0; i < OperationData.NumberOfElementsToProcess; ++i)
-	{
-		if (OperationData.InputIterators[0].GetAndAdvance<InType1>(Value1) &&
-			OperationData.InputIterators[1].GetAndAdvance<InType2>(Value2))
-		{
-			uint32 j = 0;
-			([&]
+			for (int32 i = 1; i < Range; ++i)
 			{
-				check(j < NumOutputs);
-				auto OutValue = InCallbacks(Value1, Value2);
-				typedef decltype(OutValue) OutType;
+				OutValues[i] = InCallbacks.template Get<OutputIndex>()(InArgs[i]...);
+			}
 
-				OperationData.OutputIterators[j++].SetAndAdvance<OutType>(OutValue);
-			} (), ...);
-		}
-	}
-
-	return true;
-}
+			bool bSuccess = true;
+			TArrayView<const OutType> OutValuesView(OutValues);
 
-template <typename InType1, typename InType2, typename InType3, typename... Callbacks>
-inline bool FPCGMetadataElementBase::DoTernaryOp(FOperationData& OperationData, Callbacks&& ...InCallbacks) const
-{
-	check(OperationData.InputIterators[0].IsValid());
-	check(OperationData.InputIterators[1].IsValid());
-	check(OperationData.InputIterators[2].IsValid());
+			if (InOptions.bUseDefaultKey)
+			{
+				FPCGAttributeAccessorKeysEntries DefaultAccessorKey(PCGInvalidEntryKey);
+				bSuccess = InOperationData.OutputAccessors[OutputIndex]->SetRange(OutValuesView, /*Index=*/0, DefaultAccessorKey, InOptions.SetFlags);
+			}
+			else
+			{
+				bSuccess = InOperationData.OutputAccessors[OutputIndex]->SetRange(OutValuesView, StartIndex, *InOperationData.OutputKeys[OutputIndex], InOptions.SetFlags);
+			}
 
-	const uint32 NumOutputs = OperationData.OutputIterators.Num();
-	constexpr uint32 NumCallbacks = (uint32)sizeof...(InCallbacks);
-	check(NumOutputs == NumCallbacks);
+			if (!bSuccess)
+			{
+				return false;
+			}
 
-	InType1 DefaultValue1{};
-	InType2 DefaultValue2{};
-	InType3 DefaultValue3{};
+			return Apply<OutputIndex - 1>(InOperationData, StartIndex, Range, InOptions, InCallbacks, std::forward<Args>(InArgs)...);
+		}
+	}
 
-	if (OperationData.InputIterators[0].GetWrapper().Get<InType1>(PCGInvalidEntryKey, DefaultValue1) &&
-		OperationData.InputIterators[1].GetWrapper().Get<InType2>(PCGInvalidEntryKey, DefaultValue2) &&
-		OperationData.InputIterators[2].GetWrapper().Get<InType3>(PCGInvalidEntryKey, DefaultValue3))
+	/**
+	* When Signature doesn't have any templated types anymore (Signature<>), we got all our input values packed in "InArgs", so it's time to compute our operation and set the outputs.
+	* To do so, we use Apply templated with the LastOutputIndex, and go backwards (last output to first output).
+	*/
+	template <typename... Callbacks, typename... Args>
+	bool Gather(FPCGMetadataElementBase::FOperationData& InOperationData, int32 StartIndex, int32 Range, const Options& InOptions, const TTuple<Callbacks...>& InCallbacks, int InputIndex, Signature<> S, Args&& ...InArgs)
 	{
-		// Fold expression, cf. Unary op
-		uint32 j = 0;
-		([&]
-		{
-			check(j < NumOutputs);
-			auto OutValue = InCallbacks(DefaultValue1, DefaultValue2, DefaultValue3);
-			typedef decltype(OutValue) OutType;
+		TRACE_CPUPROFILER_EVENT_SCOPE(PCG::Private::NAryOperation::Apply);
 
-			OperationData.OutputIterators[j++].GetWrapper().Set<OutType>(PCGInvalidEntryKey, OutValue);
-		} (), ...);
+		constexpr int LastOutputIndex = (int)sizeof...(Callbacks) - 1;
+		return Apply<LastOutputIndex>(InOperationData, StartIndex, Range, InOptions, InCallbacks, std::forward<Args>(InArgs)...);
 	}
 
-	InType1 Value1{};
-	InType2 Value2{};
-	InType3 Value3{};
-
-	for (int32 i = 0; i < OperationData.NumberOfElementsToProcess; ++i)
+	/**
+	* The idea of gather is to pack all the input values at the end of the function call.
+	* We use InputType, from our Signature struct to extract the current InputType, get a range of values from our accessor (using the InputIndex)
+	* then call Gather recursively, with the rest of our input types in Signature, InputIndex incremented by one and our newly InputValues moved at the end and packed into InArgs.
+	* We specialize the case where Signature is empty, meaning we got all our input values.
+	* 
+	* For example, if InputTypes = <int, float>, we have 3 calls:
+	* Gather(InOperationData, StartIndex, Range, InOptions, InCallbacks, 0, Signature<int, float>);
+	* Gather(InOperationData, StartIndex, Range, InOptions, InCallbacks, 1, Signature<float>, IntValues);
+	* Gather(InOperationData, StartIndex, Range, InOptions, InCallbacks, 2, Signature<>, IntValues, FloatValues);
+	*/
+	template <typename... Callbacks, typename InputType, typename... InputTypes, typename... Args>
+	bool Gather(FPCGMetadataElementBase::FOperationData& InOperationData, int32 StartIndex, int32 Range, const Options& InOptions, const TTuple<Callbacks...>& InCallbacks, int InputIndex, Signature<InputType, InputTypes...> S, Args&& ...InArgs)
 	{
-		if (OperationData.InputIterators[0].GetAndAdvance<InType1>(Value1) &&
-			OperationData.InputIterators[1].GetAndAdvance<InType2>(Value2) &&
-			OperationData.InputIterators[2].GetAndAdvance<InType3>(Value3))
+		bool bSuccess = true;
+
+		TArray<InputType, TInlineAllocator<DefaultChunkSize>> InputValues;
+		InputValues.SetNum(Range);
+		if (InOptions.bUseDefaultKey)
 		{
-			uint32 j = 0;
-			([&]
-			{
-				check(j < NumOutputs);
-				auto OutValue = InCallbacks(Value1, Value2, Value3);
-				typedef decltype(OutValue) OutType;
+			FPCGAttributeAccessorKeysEntries DefaultAccessorKey(PCGInvalidEntryKey);
+			bSuccess = InOperationData.InputAccessors[InputIndex]->GetRange<InputType>(InputValues, /*Index=*/0, DefaultAccessorKey, InOptions.GetFlags);
+		}
+		else
+		{
+			bSuccess = InOperationData.InputAccessors[InputIndex]->GetRange<InputType>(InputValues, StartIndex, *InOperationData.InputKeys[InputIndex], InOptions.GetFlags);
+		}
 
-				OperationData.OutputIterators[j++].SetAndAdvance<OutType>(OutValue);
-			} (), ...);
+		if (!bSuccess)
+		{
+			return false;
 		}
+
+		return Gather(InOperationData, StartIndex, Range, InOptions, InCallbacks, InputIndex + 1, Signature<InputTypes...>(), std::forward<Args>(InArgs)..., std::move(InputValues));
 	}
 
-	return true;
+	/**
+	* First we need to gather all our input values.
+	* We use an index, starting at 0, and also use our empty struct to pack all our input types.
+	*/
+	template <typename... InputTypes, typename... Callbacks>
+	bool Operation(FPCGMetadataElementBase::FOperationData& InOperationData, int32 StartIndex, int32 Range, const Options& InOptions, const TTuple<Callbacks...>& InCallbacks)
+	{
+		TRACE_CPUPROFILER_EVENT_SCOPE(PCG::Private::NAryOperation::Operation);
+		return Gather(InOperationData, StartIndex, Range, InOptions, InCallbacks, 0, Signature<InputTypes...>());
+	}
+}
+}
 }
 
-template <typename InType1, typename InType2, typename InType3, typename InType4, typename... Callbacks>
-inline bool FPCGMetadataElementBase::DoQuaternaryOp(FOperationData& OperationData, Callbacks&& ...InCallbacks) const
+template <typename... InputTypes, typename... Callbacks>
+inline bool FPCGMetadataElementBase::DoNAryOp(FOperationData& InOperationData, TTuple<Callbacks...>&& InCallbacks) const
 {
-	check(OperationData.InputIterators[0].IsValid());
-	check(OperationData.InputIterators[1].IsValid());
-	check(OperationData.InputIterators[2].IsValid());
-	check(OperationData.InputIterators[3].IsValid());
-
-	const uint32 NumOutputs = OperationData.OutputIterators.Num();
-	constexpr uint32 NumCallbacks = (uint32)sizeof...(InCallbacks);
-	check(NumOutputs == NumCallbacks);
-
-	InType1 DefaultValue1{};
-	InType2 DefaultValue2{};
-	InType3 DefaultValue3{};
-	InType4 DefaultValue4{};
-	
-	if (OperationData.InputIterators[0].GetWrapper().Get<InType1>(PCGInvalidEntryKey, DefaultValue1) &&
-		OperationData.InputIterators[1].GetWrapper().Get<InType2>(PCGInvalidEntryKey, DefaultValue2) &&
-		OperationData.InputIterators[2].GetWrapper().Get<InType3>(PCGInvalidEntryKey, DefaultValue3) &&
-		OperationData.InputIterators[3].GetWrapper().Get<InType4>(PCGInvalidEntryKey, DefaultValue4))
+	// If nothing to do, exit immediatly
+	if (InOperationData.NumberOfElementsToProcess == 0)
 	{
-		// Fold expression, cf. Unary op
-		uint32 j = 0;
-		([&]
-		{
-			check(j < NumOutputs);
-			auto OutValue = InCallbacks(DefaultValue1, DefaultValue2, DefaultValue3, DefaultValue4);
-			typedef decltype(OutValue) OutType;
-
-			OperationData.OutputIterators[j++].GetWrapper().Set<OutType>(PCGInvalidEntryKey, OutValue);
-		} (), ...);
+		return true;
 	}
 
-	InType1 Value1{};
-	InType2 Value2{};
-	InType3 Value3{};
-	InType4 Value4{};
+	// Validate that all is good
+	constexpr uint32 NbInputs = (uint32)sizeof...(InputTypes);
+	constexpr uint32 NbOutputs = (uint32)sizeof...(Callbacks);
 
-	for (int32 i = 0; i < OperationData.NumberOfElementsToProcess; ++i)
-	{
-		if (OperationData.InputIterators[0].GetAndAdvance<InType1>(Value1) &&
-			OperationData.InputIterators[1].GetAndAdvance<InType2>(Value2) &&
-			OperationData.InputIterators[2].GetAndAdvance<InType3>(Value3) &&
-			OperationData.InputIterators[3].GetAndAdvance<InType4>(Value4))
-		{
-			uint32 j = 0;
-			([&]
-			{
-				check(j < NumOutputs);
-				auto OutValue = InCallbacks(Value1, Value2, Value3, Value4);
-				typedef decltype(OutValue) OutType;
+	static_assert(NbInputs <= UPCGMetadataSettingsBase::MaxNumberOfInputs);
+	static_assert(NbOutputs <= UPCGMetadataSettingsBase::MaxNumberOfOutputs);
 
-				OperationData.OutputIterators[j++].SetAndAdvance<OutType>(OutValue);
-			} (), ...);
-		}
+	InOperationData.Validate<NbInputs, NbOutputs>();
+
+	EPCGAttributeAccessorFlags Flags = EPCGAttributeAccessorFlags::AllowBroadcast;
+
+	// First set the default value
+	PCG::Private::NAryOperation::Options Options{ Flags, Flags | EPCGAttributeAccessorFlags::AllowSetDefaultValue, true };
+	PCG::Private::NAryOperation::Operation<InputTypes...>(InOperationData, /*StartIndex=*/0, /*Range=*/1, Options, InCallbacks);
+
+	// Then iterate over all the values
+	Options.SetFlags = Flags;
+	Options.bUseDefaultKey = false;
+
+	const int32 NumberOfIterations = (InOperationData.NumberOfElementsToProcess + PCG::Private::NAryOperation::DefaultChunkSize - 1) / PCG::Private::NAryOperation::DefaultChunkSize;
+
+	for (int32 i = 0; i < NumberOfIterations; ++i)
+	{
+		int32 StartIndex = i * PCG::Private::NAryOperation::DefaultChunkSize;
+		int32 Range = FMath::Min(InOperationData.NumberOfElementsToProcess - StartIndex, PCG::Private::NAryOperation::DefaultChunkSize);
+		PCG::Private::NAryOperation::Operation<InputTypes...>(InOperationData, StartIndex, Range, Options, InCallbacks);
 	}
 
 	return true;
-}
\ No newline at end of file
+}
+
+template <typename InType, typename... Callbacks>
+inline bool FPCGMetadataElementBase::DoUnaryOp(FOperationData& InOperationData, Callbacks&& ...InCallbacks) const
+{
+	return DoNAryOp<InType>(InOperationData, ForwardAsTuple(std::forward<Callbacks>(InCallbacks)...));
+}
+
+template <typename InType1, typename InType2, typename... Callbacks>
+inline bool FPCGMetadataElementBase::DoBinaryOp(FOperationData& InOperationData, Callbacks&& ...InCallbacks) const
+{
+	return DoNAryOp<InType1, InType2>(InOperationData, ForwardAsTuple(std::forward<Callbacks>(InCallbacks)...));
+}
+
+template <typename InType1, typename InType2, typename InType3, typename... Callbacks>
+inline bool FPCGMetadataElementBase::DoTernaryOp(FOperationData& InOperationData, Callbacks&& ...InCallbacks) const
+{
+	return DoNAryOp<InType1, InType2, InType3>(InOperationData, ForwardAsTuple(std::forward<Callbacks>(InCallbacks)...));
+}
+
+template <typename InType1, typename InType2, typename InType3, typename InType4, typename... Callbacks>
+inline bool FPCGMetadataElementBase::DoQuaternaryOp(FOperationData& InOperationData, Callbacks&& ...InCallbacks) const
+{
+	return DoNAryOp<InType1, InType2, InType3, InType4>(InOperationData, ForwardAsTuple(std::forward<Callbacks>(InCallbacks)...));
+}
diff --git a/Engine/Plugins/Experimental/PCG/Source/PCG/Public/Helpers/PCGPropertyHelpers.h b/Engine/Plugins/Experimental/PCG/Source/PCG/Public/Helpers/PCGPropertyHelpers.h
new file mode 100644
--- /dev/null
+++ b/Engine/Plugins/Experimental/PCG/Source/PCG/Public/Helpers/PCGPropertyHelpers.h
@@ -0,0 +1,137 @@
+// Copyright Epic Games, Inc. All Rights Reserved.
+
+#pragma once
+
+#include "Metadata/PCGMetadataAttributeTpl.h"
+#include "Metadata/Accessors/IPCGAttributeAccessor.h"
+#include "Metadata/Accessors/PCGAttributeAccessorHelpers.h"
+#include "Metadata/Accessors/PCGAttributeAccessorKeys.h"
+
+class UPCGData;
+
+namespace PCGPropertyHelpers
+{
+	/**
+	* Get a property value and pass it as a parameter to a callback function.
+	* @param InObject - The object to read from
+	* @param InProperty - The property to look for
+	* @param InFunc - Callback function that can return anything, and should have a single templated argument, where the property will be.
+	* @returns Forward the result of the callback.
+	*/
+	template <typename ObjectType, typename Func>
+	decltype(auto) GetPropertyValueWithCallback(const ObjectType* InObject, const FProperty* InProperty, Func InFunc);
+
+	/**
+	* Set a property value given by a callback function.
+	* @param InObject - The object to write to
+	* @param InProperty - The property to look for
+	* @param InFunc - Callback function that take a reference to a templated type. It will set the property with this value. returns true if we should set, false otherwise.
+	* @returns Forward the result of the callback
+	*/
+	template <typename ObjectType, typename Func>
+	bool SetPropertyValueFromCallback(ObjectType* InObject, const FProperty* InProperty, Func InFunc);
+
+	/**
+	* Conversion between property type and PCG type.
+	* @param InProperty - The property to look for
+	* @returns PCG type if the property is supported, Unknown otherwise.
+	*/
+	PCG_API EPCGMetadataTypes GetMetadataTypeFromProperty(const FProperty* InProperty);
+}
+
+//////
+/// PCGPropertiesHelpers Implementation
+//////
+
+// Func signature : auto(auto&&)
+template <typename ObjectType, typename Func>
+inline decltype(auto) PCGPropertyHelpers::GetPropertyValueWithCallback(const ObjectType* InObject, const FProperty* InProperty, Func InFunc)
+{
+	if (!InObject || !InProperty)
+	{
+		return false;
+	}
+
+	// Object are not yet supported as accessors
+	if (const FObjectPropertyBase* ObjectProperty = CastField<FObjectPropertyBase>(InProperty))
+	{
+		if (const UObject* Object = ObjectProperty->GetObjectPropertyValue(InProperty->ContainerPtrToValuePtr<void>(InObject)))
+		{
+			// Object are transformed into their soft path name (as a string attribute)
+			return InFunc(Object->GetPathName());
+		}
+		else
+		{
+			return false;
+		}
+	}
+
+	TUniquePtr<IPCGAttributeAccessor> PropertyAccessor = PCGAttributeAccessorHelpers::CreatePropertyAccessor(InProperty);
+
+	if (!PropertyAccessor.IsValid())
+	{
+		return false;
+	}
+
+	auto Getter = [&PropertyAccessor, &InFunc, InObject](auto Dummy)
+	{
+		using Type = decltype(Dummy);
+		Type Value{};
+		FPCGAttributeAccessorKeysGeneric Key(const_cast<ObjectType*>(InObject));
+
+		if (PropertyAccessor->Get<Type>(Value, Key))
+		{
+			return InFunc(Value);
+		}
+		else
+		{
+			using ReturnType = decltype(InFunc(0.0));
+			if constexpr (std::is_same_v<ReturnType, void>)
+			{
+				return;
+			}
+			else
+			{
+				return ReturnType{};
+			}
+		}
+	};
+
+	return PCGMetadataAttribute::CallbackWithRightType(PropertyAccessor->GetUnderlyingType(), Getter);
+}
+
+// Func signature : bool(auto&)
+// Will have property value in first arg, and boolean return if we should set the property after.
+// Returns true if the set succeeded
+template <typename ObjectType, typename Func>
+inline bool PCGPropertyHelpers::SetPropertyValueFromCallback(ObjectType* InObject, const FProperty* InProperty, Func InFunc)
+{
+	if (!InObject || !InProperty)
+	{
+		return false;
+	}
+
+	TUniquePtr<IPCGAttributeAccessor> PropertyAccessor = PCGAttributeAccessorHelpers::CreatePropertyAccessor(InProperty);
+
+	if (!PropertyAccessor.IsValid())
+	{
+		return false;
+	}
+
+	auto Setter = [&PropertyAccessor, &InFunc, InObject](auto Dummy) -> bool
+	{
+		using Type = decltype(Dummy);
+		Type Value{};
+		FPCGAttributeAccessorKeysGeneric Key(InObject);
+		if (InFunc(Value))
+		{
+			return PropertyAccessor->Set<Type>(Value, Key);
+		}
+		else
+		{
+			return false;
+		}
+	};
+
+	return PCGMetadataAttribute::CallbackWithRightType(PropertyAccessor->GetUnderlyingType(), Setter);
+}
\ No newline at end of file
diff --git a/Engine/Plugins/Experimental/PCG/Source/PCG/Public/Metadata/Accessors/IPCGAttributeAccessor.h b/Engine/Plugins/Experimental/PCG/Source/PCG/Public/Metadata/Accessors/IPCGAttributeAccessor.h
new file mode 100644
--- /dev/null
+++ b/Engine/Plugins/Experimental/PCG/Source/PCG/Public/Metadata/Accessors/IPCGAttributeAccessor.h
@@ -0,0 +1,160 @@
+// Copyright Epic Games, Inc. All Rights Reserved.
+
+#pragma once
+
+#include "Metadata/PCGMetadataCommon.h"
+#include "Metadata/PCGMetadataAttributeTraits.h"
+
+#include "UObject/UnrealType.h"
+
+#include "IPCGAttributeAccessor.generated.h"
+
+class IPCGAttributeAccessorKeys;
+
+//////////////////////////////////////////////////////////////////////////////
+
+UENUM(Meta = (Bitflags))
+enum class EPCGAttributeAccessorFlags
+{	
+	// Always require that the underlying type of the accessor match the expected type, 1 for 1.
+	StrictType = 1 << 0,
+
+	// Allow to broadcast the expected type to the underlying type (or vice versa, depending on the operation)
+	AllowBroadcast = 1 << 1,
+
+	// Allow to construct the expected type from the underlying type (or vice versa, depending on the operation)
+	AllowConstructible = 1 << 2,
+
+	// By default, if the key is a PCGInvalidEntryKey, it will add a new entry. With this set, it will override the default value.
+	// USE WITH CAUTION
+	AllowSetDefaultValue = 1 << 3
+};
+ENUM_CLASS_FLAGS(EPCGAttributeAccessorFlags);
+
+//////////////////////////////////////////////////////////////////////////////
+
+/**
+* Base class for accessor. GetRange and SetRange will be sepcialized for all supported types, and there are a GetRange and SetRange virtual
+* for each supported type.
+* For Get/Set, you need a key, that will represent how you can access the value wanted. cf PCGAttributeAccessorKeys.
+* NOTE: This is not threadsafe and not intended to be used unprotected.
+*/
+class IPCGAttributeAccessor
+{
+public:
+	virtual ~IPCGAttributeAccessor() = default;
+	
+	// Both GetRange and SetRange will be specialized for all supported types. (declared below and defined in the cpp). 
+	// By default, all unsupported types will return false.
+
+	/**
+	* Get a value from the accessor for a given type. Not threadsafe for all accessors.
+	* @param OutValue - Where the value will be written to
+	* @param Index - The index to look for in the keys
+	* @param Keys - Identification to know how to retrieve the value
+	* @param Flags - Optional flag to allow for specific operations. Cf EPCGAttributeAccessorFlags.
+	* @return true if the get succeeded, false otherwise
+	*/
+	template <typename T>
+	bool Get(T& OutValue, int32 Index, const IPCGAttributeAccessorKeys& Keys, EPCGAttributeAccessorFlags Flags = EPCGAttributeAccessorFlags::StrictType) const
+	{
+		TArrayView<T> Temp(&OutValue, 1);
+		return GetRange(TArrayView<T>(&OutValue, 1), Index, Keys, Flags);
+	}
+
+	/**
+	* Get a value from the accessor for a given type at index 0. Not threadsafe for all accessors.
+	* @param OutValue - Where the value will be written to
+	* @param Keys - Identification to know how to retrieve the value
+	* @param Flags - Optional flag to allow for specific operations. Cf EPCGAttributeAccessorFlags.
+	* @return true if the get succeeded, false otherwise
+	*/
+	template <typename T>
+	bool Get(T& OutValue, const IPCGAttributeAccessorKeys& Keys, EPCGAttributeAccessorFlags Flags = EPCGAttributeAccessorFlags::StrictType) const
+	{
+		return Get(OutValue, 0, Keys, Flags);
+	}
+
+	/**
+	* Get a range of values from the accessor for a given type. Not threadsafe for all accessors.
+	* If the number of elements asked is greater that the number of keys, it will wrap around.
+	* @param OutValues - View on memory where to write values to. Its "Num" will determine how many elements we will read.
+	* @param Index - The index to start looking for in the keys
+	* @param Keys - Identification to know how to retrieve the value
+	* @param Flags - Optional flag to allow for specific operations. Cf EPCGAttributeAccessorFlags.
+	* @return true if the get succeeded, false otherwise
+	*/
+	template <typename T>
+	bool GetRange(TArrayView<T> OutValues, int32 Index, const IPCGAttributeAccessorKeys& Keys, EPCGAttributeAccessorFlags Flags = EPCGAttributeAccessorFlags::StrictType) const
+	{
+		return false;
+	}
+
+	/**
+	* Set a value to the accessor for a given type. Not threadsafe for all accessors. 
+	* If Index is greater than the number of keys, it will fail (return false).
+	* @param InValue - the value to write
+	* @param Index - The index to look for in the keys
+	* @param Keys - Identification to know how to identify the value.
+	* @param Flags - Optional flag to allow for specific operations. Cf EPCGAttributeAccessorFlags.
+	* @return true if the set succeeded, false otherwise
+	*/
+	template <typename T>
+	bool Set(const T& InValue, int32 Index, IPCGAttributeAccessorKeys& Keys, EPCGAttributeAccessorFlags Flags = EPCGAttributeAccessorFlags::StrictType)
+	{
+		return SetRange(TArrayView<const T>(&InValue, 1), Index, Keys, Flags);
+	}
+
+	/**
+	* Set a value to the accessor for a given type at index 0. Not threadsafe for all accessors.
+	* @param InValue - the value to write
+	* @param Keys - Identification to know how to identify the value.
+	* @param Flags - Optional flag to allow for specific operations. Cf EPCGAttributeAccessorFlags.
+	* @return true if the set succeeded, false otherwise
+	*/
+	template <typename T>
+	bool Set(const T& InValue, IPCGAttributeAccessorKeys& Keys, EPCGAttributeAccessorFlags Flags = EPCGAttributeAccessorFlags::StrictType)
+	{
+		return Set(InValue, 0, Keys, Flags);
+	}
+
+	/**
+	* Set a range of values to the accessor for a given type. Not threadsafe for all accessors.
+	* If the number of elements asked is greater that the number of keys, it will fail. It is to avoid writing at the same memory place multiple times.
+	* @param InValues - View on memory where to read values from. Its "Num" will determine how many elements we will write.
+	* @param Index - The index to start looking for in the keys
+	* @param Keys - Identification to know how to retrieve the value
+	* @param Flags - Optional flag to allow for specific operations. Cf EPCGAttributeAccessorFlags.
+	* @return true if the get succeeded, false otherwise
+	*/
+	template <typename T>
+	bool SetRange(TArrayView<const T> InValues, int32 Index, IPCGAttributeAccessorKeys& Keys, EPCGAttributeAccessorFlags Flags = EPCGAttributeAccessorFlags::StrictType)
+	{
+		return false;
+	}
+
+	int16 GetUnderlyingType() const { return UnderlyingType; }
+
+protected:
+	IPCGAttributeAccessor(bool bInReadOnly, int16 InUnderlyingType)
+		: bReadOnly(bInReadOnly)
+		, UnderlyingType(InUnderlyingType)
+	{}
+
+	// Define all virtual calls for all supported types.
+#define IACCESSOR_DECL(T) \
+	virtual bool GetRange##T(TArrayView<T>& OutValues, int32 Index, const IPCGAttributeAccessorKeys& Keys, EPCGAttributeAccessorFlags Flags) const = 0; \
+	virtual bool SetRange##T(TArrayView<const T>& InValues, int32 Index, IPCGAttributeAccessorKeys& Keys, EPCGAttributeAccessorFlags Flags) = 0;
+	PCG_FOREACH_SUPPORTEDTYPES(IACCESSOR_DECL);
+#undef IACCESSOR_DECL
+
+	bool bReadOnly = true;
+	int16 UnderlyingType = (int16)EPCGMetadataTypes::Unknown;
+};
+
+// Specialization of IPCGAttributeAccessor::Get and IPCGAttributeAccessor::Set, for all supported types.
+#define IACCESSOR_DECL(T) \
+template <> bool IPCGAttributeAccessor::GetRange<T>(TArrayView<T> OutValues, int32 Index, const IPCGAttributeAccessorKeys& Keys, EPCGAttributeAccessorFlags Flags) const; \
+template <> bool IPCGAttributeAccessor::SetRange<T>(TArrayView<const T> InValues, int32 Index, IPCGAttributeAccessorKeys& Keys, EPCGAttributeAccessorFlags Flags);
+PCG_FOREACH_SUPPORTEDTYPES(IACCESSOR_DECL);
+#undef IACCESSOR_DECL
\ No newline at end of file
diff --git a/Engine/Plugins/Experimental/PCG/Source/PCG/Public/Metadata/Accessors/IPCGAttributeAccessorTpl.h b/Engine/Plugins/Experimental/PCG/Source/PCG/Public/Metadata/Accessors/IPCGAttributeAccessorTpl.h
new file mode 100644
--- /dev/null
+++ b/Engine/Plugins/Experimental/PCG/Source/PCG/Public/Metadata/Accessors/IPCGAttributeAccessorTpl.h
@@ -0,0 +1,143 @@
+// Copyright Epic Games, Inc. All Rights Reserved.
+
+#pragma once
+
+#include "CoreMinimal.h"
+#include "Metadata/PCGMetadataCommon.h"
+#include "Metadata/PCGMetadataAttributeTraits.h"
+#include "Metadata/Accessors/IPCGAttributeAccessor.h"
+#include "Metadata/Accessors/PCGAttributeAccessorKeys.h"
+
+/**
+* Use of curiously recursive template pattern (CRTP) to dispatch GetRangeImpl and SetRangeImpl at compile time.
+* Override all virtual functions for the supported types and will handle the conversion between
+* "U" the incoming type and "T" the underlying type.
+* 
+* Class that inherit this one needs to define:
+* -> bool GetRangeImpl(TArrayView<T>& OutValues, int32 Index, const IPCGAttributeAccessorKeys& Keys) const
+* -> bool SetRangeImpl(TArrayView<const T>& InValues, int32 Index, IPCGAttributeAccessorKeys& Keys, EPCGAttributeAccessorFlags Flags)
+* -> The underlying type: (using Type = ...);
+*/
+template <typename Derived>
+class IPCGAttributeAccessorT : public IPCGAttributeAccessor
+{
+protected:
+	IPCGAttributeAccessorT(bool bInReadOnly)
+		: IPCGAttributeAccessor(bInReadOnly, /*UnderlyingType=*/ PCG::Private::MetadataTypes<typename Derived::Type>::Id)
+	{}
+
+	// Override all virtual calls to call our templated method.
+#define IACCESSOR_DECL(T) \
+	virtual bool GetRange##T(TArrayView<T>& OutValues, int32 Index, const IPCGAttributeAccessorKeys& Keys, EPCGAttributeAccessorFlags Flags) const override { return InternalGetRange(OutValues, Index, Keys, Flags); } \
+	virtual bool SetRange##T(TArrayView<const T>& InValues, int32 Index, IPCGAttributeAccessorKeys& Keys, EPCGAttributeAccessorFlags Flags) override { return InternalSetRange(InValues, Index, Keys, Flags); }
+	PCG_FOREACH_SUPPORTEDTYPES(IACCESSOR_DECL);
+#undef IACCESSOR_DECL
+
+private:
+	template<typename U>
+	bool InternalGetRange(TArrayView<U>& OutValues, int32 Index, const IPCGAttributeAccessorKeys& Keys, EPCGAttributeAccessorFlags Flags) const
+	{
+		TRACE_CPUPROFILER_EVENT_SCOPE(IPCGAttributeAccessorT::InternalGetRange);
+
+		using T = typename Derived::Type;
+		const Derived* This = static_cast<const Derived*>(this);
+
+		if constexpr (std::is_same_v<T, U>)
+		{
+			return This->GetRangeImpl(OutValues, Index, Keys);
+		}
+		else
+		{
+			if constexpr (PCG::Private::IsBroadcastable<T, U>())
+			{
+				if (!!(Flags & EPCGAttributeAccessorFlags::AllowBroadcast))
+				{
+					TArray<T, TInlineAllocator<4>> InValues;
+					InValues.SetNum(OutValues.Num());
+					if (This->GetRangeImpl(InValues, Index, Keys))
+					{
+						for (int32 i = 0; i < OutValues.Num(); ++i)
+						{
+							PCG::Private::GetValueWithBroadcast<T, U>(InValues[i], OutValues[i]);
+						}
+
+						return true;
+					}
+				}
+			}
+			else if constexpr (std::is_constructible_v<U, T>)
+			{
+				if (!!(Flags & EPCGAttributeAccessorFlags::AllowConstructible))
+				{
+					TArray<T, TInlineAllocator<4>> InValues;
+					InValues.SetNum(OutValues.Num());
+					if (This->GetRangeImpl(InValues, Index, Keys))
+					{
+						for (int32 i = 0; i < OutValues.Num(); ++i)
+						{
+							OutValues[i] = U(InValues[i]);
+						}
+						return true;
+					}
+				}
+			}
+		}
+
+		// If U is different of T, and we either can't or allow broadcast or construct from, we will fail, and return false.
+		return false;
+	}
+
+	template<typename U>
+	bool InternalSetRange(TArrayView<const U>& InValues, int32 Index, IPCGAttributeAccessorKeys& Keys, EPCGAttributeAccessorFlags Flags)
+	{
+		TRACE_CPUPROFILER_EVENT_SCOPE(IPCGAttributeAccessorT::InternalSetRange);
+
+		// Can't set if read only, and discard any set outside the keys range.
+		if (bReadOnly || Index >= Keys.GetNum())
+		{
+			return false;
+		}
+
+		using T = typename Derived::Type;
+		Derived* This = static_cast<Derived*>(this);
+
+		if constexpr (std::is_same_v<T, U>)
+		{
+			return This->SetRangeImpl(InValues, Index, Keys, Flags);
+		}
+		else
+		{
+			if constexpr (PCG::Private::IsBroadcastable<U, T>())
+			{
+				if (!!(Flags & EPCGAttributeAccessorFlags::AllowBroadcast))
+				{
+					TArray<T> OutValues;
+					OutValues.SetNum(InValues.Num());
+					for (int32 i = 0; i < OutValues.Num(); ++i)
+					{
+						PCG::Private::GetValueWithBroadcast<U, T>(InValues[i], OutValues[i]);
+					}
+
+					return This->SetRangeImpl(OutValues, Index, Keys, Flags);
+				}
+			}
+			else if constexpr (std::is_constructible_v<T, U>)
+			{
+				if (!!(Flags & EPCGAttributeAccessorFlags::AllowConstructible))
+				{
+					TArray<T> OutValues;
+					OutValues.SetNum(InValues.Num());
+					for (int32 i = 0; i < OutValues.Num(); ++i)
+					{
+						OutValues[i] = T(InValues[i]);
+					}
+
+					return This->SetRangeImpl(OutValues, Index, Keys, Flags);
+				}
+			}
+		}
+
+		// If U is different of T, and we either can't or allow broadcast or construct from, we will fail, and return false.
+		return false;
+	}
+};
diff --git a/Engine/Plugins/Experimental/PCG/Source/PCG/Public/Metadata/Accessors/PCGAttributeAccessor.h b/Engine/Plugins/Experimental/PCG/Source/PCG/Public/Metadata/Accessors/PCGAttributeAccessor.h
new file mode 100644
--- /dev/null
+++ b/Engine/Plugins/Experimental/PCG/Source/PCG/Public/Metadata/Accessors/PCGAttributeAccessor.h
@@ -0,0 +1,96 @@
+// Copyright Epic Games, Inc. All Rights Reserved.
+
+#pragma once
+
+#include "IPCGAttributeAccessorTpl.h"
+
+#include "Metadata/PCGMetadata.h"
+#include "Metadata/PCGMetadataAttributeTpl.h"
+#include "Metadata/PCGMetadataAttributeTraits.h"
+#include "Metadata/Accessors/PCGAttributeAccessorKeys.h"
+
+
+/**
+* Templated accessor class for attributes. Will wrap around a typed attribute.
+* Key supported: MetadataEntryKey and Points
+*/
+template <typename T>
+class FPCGAttributeAccessor : public IPCGAttributeAccessorT<FPCGAttributeAccessor<T>>
+{
+public:
+	using Type = T;
+	using Super = IPCGAttributeAccessorT<FPCGAttributeAccessor<T>>;
+
+	FPCGAttributeAccessor(FPCGMetadataAttribute<T>* InAttribute, UPCGMetadata* InMetadata)
+		: Super(/*bInReadOnly=*/ false)
+		, Attribute(InAttribute)
+		, Metadata(InMetadata)
+	{
+		check(InAttribute && InMetadata);
+	}
+
+	FPCGAttributeAccessor(const FPCGMetadataAttribute<T>* InAttribute, const UPCGMetadata* InMetadata)
+		: Super(/*bInReadOnly=*/ true)
+		, Attribute(const_cast<FPCGMetadataAttribute<T>*>(InAttribute))
+		, Metadata(const_cast<UPCGMetadata*>(InMetadata))
+	{
+		check(InAttribute && InMetadata);
+	}
+
+	bool GetRangeImpl(TArrayView<T> OutValues, int32 Index, const IPCGAttributeAccessorKeys& Keys) const
+	{
+		TArray<const PCGMetadataEntryKey*> EntryKeys;
+		EntryKeys.SetNum(OutValues.Num());
+		TArrayView<const PCGMetadataEntryKey*> EntryKeysView(EntryKeys);
+		if (!Keys.GetKeys<PCGMetadataEntryKey>(Index, EntryKeysView))
+		{
+			return false;
+		}
+
+		// TODO: Might be good to hase a "GetValuesFromItemKeys" to try locking less often.
+		for (int32 i = 0; i < OutValues.Num(); ++i)
+		{
+			OutValues[i] = Attribute->GetValueFromItemKey(*EntryKeys[i]);
+		}
+
+		return true;
+	}
+
+	bool SetRangeImpl(TArrayView<const T> InValues, int32 Index, IPCGAttributeAccessorKeys& Keys, EPCGAttributeAccessorFlags Flags)
+	{
+		TArray<PCGMetadataEntryKey*> EntryKeys;
+		EntryKeys.SetNum(InValues.Num());
+		TArrayView<PCGMetadataEntryKey*> EntryKeysView(EntryKeys);
+		if (!Keys.GetKeys<PCGMetadataEntryKey>(Index, EntryKeysView))
+		{
+			return false;
+		}
+
+		// TODO: Same than above (avoid locking too many times), but perhaps will be a bit more complex, because of
+		// the added logic with PCGInvalidEntryKey, AddEntry and SetDefaultValue.
+		for (int32 i = 0; i < InValues.Num(); ++i)
+		{
+			PCGMetadataEntryKey& EntryKey = *EntryKeys[i];
+			if (EntryKey != PCGInvalidEntryKey || !(Flags & EPCGAttributeAccessorFlags::AllowSetDefaultValue))
+			{
+				// TODO: This part seems costly for a lot of entries. Maybe there are some optimizations to do.
+				if (EntryKey == PCGInvalidEntryKey)
+				{
+					EntryKey = Metadata->AddEntry();
+				}
+
+				Attribute->SetValue(EntryKey, InValues[i]);
+			}
+			else
+			{
+				Attribute->SetDefaultValue(InValues[i]);
+			}
+		}
+
+		return true;
+	}
+
+private:
+	FPCGMetadataAttribute<T>* Attribute = nullptr;
+	UPCGMetadata* Metadata = nullptr;
+};
\ No newline at end of file
diff --git a/Engine/Plugins/Experimental/PCG/Source/PCG/Public/Metadata/Accessors/PCGAttributeAccessorHelpers.h b/Engine/Plugins/Experimental/PCG/Source/PCG/Public/Metadata/Accessors/PCGAttributeAccessorHelpers.h
new file mode 100644
--- /dev/null
+++ b/Engine/Plugins/Experimental/PCG/Source/PCG/Public/Metadata/Accessors/PCGAttributeAccessorHelpers.h
@@ -0,0 +1,24 @@
+// Copyright Epic Games, Inc. All Rights Reserved.
+
+#pragma once
+
+#include "Templates/UniquePtr.h"
+
+class IPCGAttributeAccessor;
+class IPCGAttributeAccessorKeyIterator;
+class IPCGAttributeAccessorKeys;
+class FProperty;
+class UPCGData;
+struct FPCGAttributePropertySelector;
+
+namespace PCGAttributeAccessorHelpers
+{
+	TUniquePtr<IPCGAttributeAccessor> CreatePropertyAccessor(const FProperty* InProperty);
+
+	TUniquePtr<const IPCGAttributeAccessor> CreateConstAccessor(const UPCGData* InData, const FPCGAttributePropertySelector& InSelector);
+	TUniquePtr<IPCGAttributeAccessor> CreateAccessor(UPCGData* InData, const FPCGAttributePropertySelector& InSelector);
+
+	TUniquePtr<const IPCGAttributeAccessorKeys> CreateConstKeys(const UPCGData* InData, const FPCGAttributePropertySelector& InSelector);
+	TUniquePtr<IPCGAttributeAccessorKeys> CreateKeys(UPCGData* InData, const FPCGAttributePropertySelector& InSelector);
+
+}
\ No newline at end of file
diff --git a/Engine/Plugins/Experimental/PCG/Source/PCG/Public/Metadata/Accessors/PCGAttributeAccessorKeys.h b/Engine/Plugins/Experimental/PCG/Source/PCG/Public/Metadata/Accessors/PCGAttributeAccessorKeys.h
new file mode 100644
--- /dev/null
+++ b/Engine/Plugins/Experimental/PCG/Source/PCG/Public/Metadata/Accessors/PCGAttributeAccessorKeys.h
@@ -0,0 +1,279 @@
+// Copyright Epic Games, Inc. All Rights Reserved.
+
+#pragma once
+
+#include "Metadata/PCGMetadataCommon.h"
+#include "PCGPoint.h"
+
+#include "Containers/ArrayView.h"
+#include "UObject/UnrealType.h"
+
+class FPCGMetadataAttributeBase;
+
+///////////////////////////////////////////////////////////////////////
+
+/**
+* Base class to identify keys to use with an accessor.
+*/
+class IPCGAttributeAccessorKeys
+{
+public:
+	IPCGAttributeAccessorKeys(bool bInReadOnly)
+		: bIsReadOnly(bInReadOnly)
+	{}
+
+	virtual ~IPCGAttributeAccessorKeys() = default;
+
+	/**
+	* Retrieve in the given view pointers of the wanted type
+	* Need to be a supported type, such as FPCGPoint, PCGMetadataEntryKey or void.
+	* It will wrap around if the index/range goes outside the number of keys.
+	* @param InStart - Index to start looking in the keys.
+	* @param OutKeys - View on the out keys. Its size will indicate the number of elements to get.
+	* @return true if it succeeded, false otherwise. (like num == 0,unsupported type or read only)
+	*/
+	template <typename ObjectType>
+	bool GetKeys(int32 InStart, TArrayView<ObjectType*>& OutKeys);
+
+	// Same function but const.
+	template <typename ObjectType>
+	bool GetKeys(int32 InStart, TArrayView<const ObjectType*>& OutKeys) const;
+
+	/**
+	* Retrieve in the given argument pointer of the wanted type.
+	* Need to be a supported type, such as FPCGPoint, PCGMetadataEntryKey or void.
+	* It will wrap around if the index goes outside the number of keys.
+	* @param InStart - Index to start looking in the keys.
+	* @param OutObject - Reference on a pointer, where the key will be written.
+	* @return true if it succeeded, false otherwise. (like num == 0, unsupported type or read only)
+	*/
+	template <typename ObjectType>
+	bool GetKey(int32 InStart, ObjectType*& OutObject)
+	{
+		return GetKeys(InStart, TArrayView<ObjectType*>(&OutObject, 1));
+	}
+
+	// Same function but const
+	template <typename ObjectType>
+	bool GetKey(int32 InStart, ObjectType const*& OutObject)
+	{
+		return GetKeys(InStart, TArrayView<ObjectType*>(&OutObject, 1));
+	}
+
+	/**
+	* Retrieve in the given argument pointer of the wanted type at the index 0.
+	* Need to be a supported type, such as FPCGPoint, PCGMetadataEntryKey or void.
+	* @param OutObject - Reference on a pointer, where the key will be written.
+	* @return true if it succeeded, false otherwise. (like num == 0, unsupported type or read only)
+	*/
+	template <typename ObjectType>
+	bool GetKey(ObjectType*& OutObject)
+	{
+		return GetKeys(/*InStart=*/ 0, TArrayView<ObjectType*>(&OutObject, 1));
+	}
+
+	// Same function but const
+	template <typename ObjectType>
+	bool GetKey(ObjectType const*& OutObject) const
+	{
+		return GetKeys(/*InStart=*/ 0, TArrayView<const ObjectType*>(&OutObject, 1));
+	}
+
+	/*
+	* Returns the number of keys.
+	*/
+	virtual int32 GetNum() const = 0;
+
+protected:
+	virtual bool GetPointKeys(int32 InStart, TArrayView<FPCGPoint*>& OutPoints) { return false; }
+	virtual bool GetPointKeys(int32 InStart, TArrayView<const FPCGPoint*>& OutPoints) const { return false; }
+
+	virtual bool GetGenericObjectKeys(int32 InStart, TArrayView<void*>& OutObjects) { return false; }
+	virtual bool GetGenericObjectKeys(int32 InStart, TArrayView<const void*>& OutObjects) const { return false; }
+
+	virtual bool GetMetadataEntryKeys(int32 InStart, TArrayView<PCGMetadataEntryKey*>& OutEntryKeys) { return false; }
+	virtual bool GetMetadataEntryKeys(int32 InStart, TArrayView<const PCGMetadataEntryKey*>& OutEntryKeys) const { return false; }
+
+	bool bIsReadOnly = false;
+};
+
+///////////////////////////////////////////////////////////////////////
+
+/**
+* Key around a metadata entry key
+*/
+class FPCGAttributeAccessorKeysEntries : public IPCGAttributeAccessorKeys
+{
+public:
+	FPCGAttributeAccessorKeysEntries(const FPCGMetadataAttributeBase* Attribute);
+	FPCGAttributeAccessorKeysEntries(PCGMetadataEntryKey EntryKey);
+
+	virtual int32 GetNum() const override { return Entries.Num(); }
+
+protected:
+	virtual bool GetMetadataEntryKeys(int32 InStart, TArrayView<PCGMetadataEntryKey*>& OutEntryKeys) override;
+	virtual bool GetMetadataEntryKeys(int32 InStart, TArrayView<const PCGMetadataEntryKey*>& OutEntryKeys) const override;
+
+	TArray<PCGMetadataEntryKey> Entries;
+};
+
+///////////////////////////////////////////////////////////////////////
+
+/**
+* Key around points
+*/
+class FPCGAttributeAccessorKeysPoints : public IPCGAttributeAccessorKeys
+{
+public:
+	FPCGAttributeAccessorKeysPoints(const TArrayView<FPCGPoint>& InPoints);
+	FPCGAttributeAccessorKeysPoints(const TArrayView<const FPCGPoint>& InPoints);
+
+	FPCGAttributeAccessorKeysPoints(FPCGPoint& InPoint);
+	FPCGAttributeAccessorKeysPoints(const FPCGPoint& InPoint);
+
+	virtual int32 GetNum() const override { return Points.Num(); }
+
+protected:
+	virtual bool GetPointKeys(int32 InStart, TArrayView<FPCGPoint*>& OutPoints) override;
+	virtual bool GetPointKeys(int32 InStart, TArrayView<const FPCGPoint*>& OutPoints) const override;
+
+	virtual bool GetGenericObjectKeys(int32 InStart, TArrayView<void*>& OutObjects) override;
+	virtual bool GetGenericObjectKeys(int32 InStart, TArrayView<const void*>& OutObjects) const override;
+
+	virtual bool GetMetadataEntryKeys(int32 InStart, TArrayView<PCGMetadataEntryKey*>& OutEntryKeys) override;
+	virtual bool GetMetadataEntryKeys(int32 InStart, TArrayView<const PCGMetadataEntryKey*>& OutEntryKeys) const override;
+
+	TArrayView<FPCGPoint> Points;
+};
+
+/////////////////////////////////////////////////////////////////
+
+/**
+* Key around generic objects
+*/
+template <typename ObjectType>
+class FPCGAttributeAccessorKeysGeneric : public IPCGAttributeAccessorKeys
+{
+public:
+	FPCGAttributeAccessorKeysGeneric(const TArrayView<ObjectType>& InObjects)
+		: IPCGAttributeAccessorKeys(/*bInReadOnly=*/ false)
+		, Objects(InObjects)
+	{}
+
+	FPCGAttributeAccessorKeysGeneric(const TArrayView<const ObjectType>& InObjects)
+		: IPCGAttributeAccessorKeys(/*bInReadOnly=*/ true)
+		, Objects(const_cast<ObjectType*>(InObjects.GetData()), InObjects.Num())
+	{}
+
+	FPCGAttributeAccessorKeysGeneric(ObjectType& InObject)
+		: FPCGAttributeAccessorKeysGeneric(TArrayView<ObjectType>(&InObject, 1))
+	{}
+
+	FPCGAttributeAccessorKeysGeneric(const ObjectType& InObject)
+		: FPCGAttributeAccessorKeysGeneric(TArrayView<const ObjectType>(&InObject, 1))
+	{}
+
+	virtual int32 GetNum() const override { return Objects.Num(); }
+
+protected:
+	virtual bool GetGenericObjectKeys(int32 InStart, TArrayView<void*>& OutObjects) override;
+	virtual bool GetGenericObjectKeys(int32 InStart, TArrayView<const void*>& OutObjects) const override;
+
+	TArrayView<ObjectType> Objects;
+};
+
+/////////////////////////////////////////////////////////////////
+
+template <typename ObjectType>
+inline bool IPCGAttributeAccessorKeys::GetKeys(int32 InStart, TArrayView<ObjectType*>& OutKeys)
+{
+	TRACE_CPUPROFILER_EVENT_SCOPE(IPCGAttributeAccessorKeys::GetKeys);
+
+	if (bIsReadOnly)
+	{
+		return false;
+	}
+
+	if constexpr (std::is_same_v<ObjectType, FPCGPoint>)
+	{
+		return GetPointKeys(InStart, OutKeys);
+	}
+	else if constexpr (std::is_same_v<ObjectType, PCGMetadataEntryKey>)
+	{
+		return GetMetadataEntryKeys(InStart, OutKeys);
+	}
+	else if constexpr (std::is_same_v<ObjectType, void>)
+	{
+		return GetGenericObjectKeys(InStart, OutKeys);
+	}
+	else
+	{
+		return false;
+	}
+}
+
+template <typename ObjectType>
+inline bool IPCGAttributeAccessorKeys::GetKeys(int32 InStart, TArrayView<const ObjectType*>& OutKeys) const
+{
+	TRACE_CPUPROFILER_EVENT_SCOPE(IPCGAttributeAccessorKeys::GetKeys);
+
+	if constexpr (std::is_same_v<ObjectType, FPCGPoint>)
+	{
+		return GetPointKeys(InStart, OutKeys);
+	}
+	else if constexpr (std::is_same_v<ObjectType, PCGMetadataEntryKey>)
+	{
+		return GetMetadataEntryKeys(InStart, OutKeys);
+	}
+	else if constexpr (std::is_same_v<ObjectType, void>)
+	{
+		return GetGenericObjectKeys(InStart, OutKeys);
+	}
+	else
+	{
+		return false;
+	}
+}
+
+/////////////////////////////////////////////////////////////////
+
+namespace PCGAttributeAccessorKeys
+{
+	template <typename T, typename Container, typename Func>
+	bool GetKeys(Container& InContainer, int32 InStart, TArrayView<T*>& OutItems, Func&& Transform)
+	{
+		if (InContainer.Num() == 0)
+		{
+			return false;
+		}
+
+		int32 Current = InStart;
+		if (Current >= InContainer.Num())
+		{
+			Current %= InContainer.Num();
+		}
+
+		for (int32 i = 0; i < OutItems.Num(); ++i)
+		{
+			OutItems[i] = Transform(InContainer[Current++]);
+			if (Current >= InContainer.Num())
+			{
+				Current = 0;
+			}
+		}
+
+		return true;
+	}
+}
+
+template <typename ObjectType>
+bool FPCGAttributeAccessorKeysGeneric<ObjectType>::GetGenericObjectKeys(int32 InStart, TArrayView<void*>& OutObjects)
+{
+	return PCGAttributeAccessorKeys::GetKeys(Objects, InStart, OutObjects, [](ObjectType& Obj) -> ObjectType* { return &Obj; });
+}
+
+template <typename ObjectType>
+bool FPCGAttributeAccessorKeysGeneric<ObjectType>::GetGenericObjectKeys(int32 InStart, TArrayView<const void*>& OutObjects) const
+{
+	return PCGAttributeAccessorKeys::GetKeys(Objects, InStart, OutObjects, [](const ObjectType& Obj) -> const ObjectType* { return &Obj; });
+}
\ No newline at end of file
diff --git a/Engine/Plugins/Experimental/PCG/Source/PCG/Public/Metadata/Accessors/PCGCustomAccessor.h b/Engine/Plugins/Experimental/PCG/Source/PCG/Public/Metadata/Accessors/PCGCustomAccessor.h
new file mode 100644
--- /dev/null
+++ b/Engine/Plugins/Experimental/PCG/Source/PCG/Public/Metadata/Accessors/PCGCustomAccessor.h
@@ -0,0 +1,66 @@
+// Copyright Epic Games, Inc. All Rights Reserved.
+
+#pragma once
+
+#include "IPCGAttributeAccessorTpl.h"
+
+#include "PCGPoint.h"
+
+/**
+* Templated accessor class for custom point properties. Need a getter and a setter, defined in the
+* FPCGPoint class.
+* Key supported: Points
+*/
+template <typename T>
+class FPCGCustomPointAccessor : public IPCGAttributeAccessorT<FPCGCustomPointAccessor<T>>
+{
+public:
+	using Type = T;
+	using Super = IPCGAttributeAccessorT<FPCGCustomPointAccessor<T>>;
+
+	FPCGCustomPointAccessor(const FPCGPoint::PointCustomPropertyGetter& InGetter, const FPCGPoint::PointCustomPropertySetter& InSetter)
+		: Super(/*bInReadOnly=*/ false)
+		, Getter(InGetter)
+		, Setter(InSetter)
+	{}
+
+	bool GetRangeImpl(TArrayView<T> OutValues, int32 Index, const IPCGAttributeAccessorKeys& Keys) const
+	{
+		TArray<const FPCGPoint*> PointKeys;
+		PointKeys.SetNum(OutValues.Num());
+		TArrayView<const FPCGPoint*> PointKeysView(PointKeys);
+		if (!Keys.GetKeys(Index, PointKeysView))
+		{
+			return false;
+		}
+
+		for (int32 i = 0; i < OutValues.Num(); ++i)
+		{
+			Getter(*PointKeys[i], &OutValues[i]);
+		}
+
+		return true;
+	}
+
+	bool SetRangeImpl(TArrayView<const T> InValues, int32 Index, IPCGAttributeAccessorKeys& Keys, EPCGAttributeAccessorFlags Flags)
+	{
+		TArray<FPCGPoint*> PointKeys;
+		PointKeys.SetNum(InValues.Num());
+		TArrayView<FPCGPoint*> PointKeysView(PointKeys);
+		if (!Keys.GetKeys(Index, PointKeysView))
+		{
+			return false;
+		}
+
+		for (int32 i = 0; i < InValues.Num(); ++i)
+		{
+			Setter(*PointKeys[i], &InValues[i]);
+		}
+
+		return true;
+	}
+
+private:
+	FPCGPoint::PointCustomPropertyGetter Getter;
+	FPCGPoint::PointCustomPropertySetter Setter;
+};
\ No newline at end of file
diff --git a/Engine/Plugins/Experimental/PCG/Source/PCG/Public/Metadata/Accessors/PCGPropertyAccessor.h b/Engine/Plugins/Experimental/PCG/Source/PCG/Public/Metadata/Accessors/PCGPropertyAccessor.h
new file mode 100644
--- /dev/null
+++ b/Engine/Plugins/Experimental/PCG/Source/PCG/Public/Metadata/Accessors/PCGPropertyAccessor.h
@@ -0,0 +1,215 @@
+// Copyright Epic Games, Inc. All Rights Reserved.
+
+#pragma once
+
+#include "IPCGAttributeAccessorTpl.h"
+
+#include "UObject/UnrealType.h"
+
+class FEnumProperty;
+
+namespace PCGPropertyAccessor
+{
+	TArray<const void*> GetContainerKeys(int32 Index, int32 Range, const IPCGAttributeAccessorKeys& Keys);
+	TArray<void*> GetContainerKeys(int32 Index, int32 Range, IPCGAttributeAccessorKeys& Keys);
+
+	template <typename T, typename Func>
+	bool IterateGet(const FProperty* Property, TArrayView<T>& OutValues, int32 Index, const IPCGAttributeAccessorKeys& Keys, Func&& Getter)
+	{
+		TArray<const void*> ContainerKeys = GetContainerKeys(Index, OutValues.Num(), Keys);
+		if (ContainerKeys.IsEmpty())
+		{
+			return false;
+		}
+
+		for (int32 i = 0; i < OutValues.Num(); ++i)
+		{
+			const void* PropertyAddressData = Property->ContainerPtrToValuePtr<void>(ContainerKeys[i]);
+			OutValues[i] = Getter(PropertyAddressData);
+		}
+
+		return true;
+	}
+
+	template <typename T, typename Func>
+	bool IterateSet(const FProperty* Property, TArrayView<const T>& InValues, int32 Index, IPCGAttributeAccessorKeys& Keys, Func&& Setter)
+	{
+		TArray<void*> ContainerKeys = GetContainerKeys(Index, InValues.Num(), Keys);
+		if (ContainerKeys.IsEmpty())
+		{
+			return false;
+		}
+
+		for (int32 i = 0; i < InValues.Num(); ++i)
+		{
+			void* PropertyAddressData = Property->ContainerPtrToValuePtr<void>(ContainerKeys[i]);
+			Setter(PropertyAddressData, InValues[i]);
+		}
+
+		return true;
+	}
+}
+
+/**
+* Templated accessor class for numeric properties. Will wrap around a numeric property.
+* Do not instanciate it manually, use PCGAttributeAccessorHelpers::CreatePropertyAccessor.
+* Only support integral and floating point types.
+* Key supported: Generic object
+*/
+template <typename T>
+class FPCGNumericPropertyAccessor : public IPCGAttributeAccessorT<FPCGNumericPropertyAccessor<T>>
+{
+public:
+	using Type = T;
+	using Super = IPCGAttributeAccessorT<FPCGNumericPropertyAccessor<T>>;
+
+	FPCGNumericPropertyAccessor(const FNumericProperty* InProperty)
+		: Super(/*bInReadOnly=*/ false)
+		, Property(InProperty)
+	{
+		static_assert(std::is_integral_v<T> || std::is_floating_point_v<T>);
+		check(Property);
+	}
+
+	bool GetRangeImpl(TArrayView<T> OutValues, int32 Index, const IPCGAttributeAccessorKeys& Keys) const
+	{
+		return PCGPropertyAccessor::IterateGet(Property, OutValues, Index, Keys, [this](const void* PropertyAddressData) -> T
+			{
+				if constexpr (std::is_integral_v<T>)
+				{
+					return T(Property->GetSignedIntPropertyValue(PropertyAddressData));
+				}
+				else
+				{
+					return T(Property->GetFloatingPointPropertyValue(PropertyAddressData));
+				}
+			});
+	}
+
+	bool SetRangeImpl(TArrayView<const T> InValues, int32 Index, IPCGAttributeAccessorKeys& Keys, EPCGAttributeAccessorFlags)
+	{
+		return PCGPropertyAccessor::IterateSet(Property, InValues, Index, Keys, [this](void* PropertyAddressData, const T& Value) -> void
+			{
+				if constexpr (std::is_integral_v<T>)
+				{
+					Property->SetIntPropertyValue(PropertyAddressData, Value);
+				}
+				else
+				{
+					Property->SetFloatingPointPropertyValue(PropertyAddressData, Value);
+				}
+			});
+	}
+
+private:
+	const FNumericProperty* Property = nullptr;
+};
+
+/**
+* Templated accessor class for enum properties. Will wrap around an enum property.
+* Will always convert to int64 for PCG
+* Do not instanciate it manually, use PCGAttributeAccessorHelpers::CreatePropertyAccessor.
+* Key supported: Generic object
+*/
+class FPCGEnumPropertyAccessor : public IPCGAttributeAccessorT<FPCGEnumPropertyAccessor>
+{
+public:
+	using Type = int64;
+	using Super = IPCGAttributeAccessorT<FPCGEnumPropertyAccessor>;
+
+	FPCGEnumPropertyAccessor(const FEnumProperty* InProperty)
+		: Super(/*bInReadOnly=*/ false)
+		, Property(InProperty)
+	{
+		check(Property);
+	}
+
+	bool GetRangeImpl(TArrayView<int64> OutValues, int32 Index, const IPCGAttributeAccessorKeys& Keys) const;
+	bool SetRangeImpl(TArrayView<const int64> InValues, int32 Index, IPCGAttributeAccessorKeys& Keys, EPCGAttributeAccessorFlags);
+
+private:
+	const FEnumProperty* Property = nullptr;
+};
+
+/**
+* Templated accessor class for generic properties. Will wrap around a property.
+* Do not instanciate it manually, use PCGAttributeAccessorHelpers::CreatePropertyAccessor.
+* WARNING: Getting the address of the underlying data from a container using this property
+* should point to a "T" object in memory, otherwise it is UD.
+* Key supported: Generic object
+*/
+template <typename T>
+class FPCGPropertyAccessor : public IPCGAttributeAccessorT<FPCGPropertyAccessor<T>>
+{
+public:
+	using Type = T;
+	using Super = IPCGAttributeAccessorT<FPCGPropertyAccessor<T>>;
+
+	FPCGPropertyAccessor(const FProperty* InProperty)
+		: Super(/*bInReadOnly=*/ false)
+		, Property(InProperty)
+	{
+		static_assert(PCG::Private::IsPCGType<T>());
+		check(Property);
+	}
+
+	bool GetRangeImpl(TArrayView<T> OutValues, int32 Index, const IPCGAttributeAccessorKeys& Keys) const
+	{
+		return PCGPropertyAccessor::IterateGet(Property, OutValues, Index, Keys, [this](const void* PropertyAddressData) -> T
+			{
+				return *reinterpret_cast<const T*>(PropertyAddressData);
+			});
+	}
+
+	bool SetRangeImpl(TArrayView<const T> InValues, int32 Index, IPCGAttributeAccessorKeys& Keys, EPCGAttributeAccessorFlags)
+	{
+		return PCGPropertyAccessor::IterateSet(Property, InValues, Index, Keys, [this](void* PropertyAddressData, const T& Value) -> void
+			{
+				*reinterpret_cast<T*>(PropertyAddressData) = Value;
+			});
+	}
+
+private:
+	const FProperty* Property = nullptr;
+};
+
+/**
+* Templated accessor class for path properties. Will wrap around a soft path property.
+* Do not instanciate it manually, use PCGAttributeAccessorHelpers::CreatePropertyAccessor.
+* Will always convert to FString for PCG
+* Key supported: Generic object
+*/
+template <typename T>
+class FPCGPropertyPathAccessor : public IPCGAttributeAccessorT<FPCGPropertyPathAccessor<T>>
+{
+public:
+	using Type = FString;
+	using Super = IPCGAttributeAccessorT<FPCGPropertyPathAccessor<T>>;
+
+	FPCGPropertyPathAccessor(const FProperty* InProperty)
+		: Super(/*bInReadOnly=*/ false)
+		, Property(InProperty)
+	{
+		static_assert(std::is_same_v<FSoftObjectPath, T> || std::is_same_v<FSoftClassPath, T>);
+		check(Property);
+	}
+
+	bool GetRangeImpl(TArrayView<Type> OutValues, int32 Index, const IPCGAttributeAccessorKeys& Keys) const
+	{
+		return PCGPropertyAccessor::IterateGet(Property, OutValues, Index, Keys, [this](const void* PropertyAddressData) -> Type
+			{
+				return reinterpret_cast<const T*>(PropertyAddressData)->ToString();
+			});
+	}
+
+	bool SetRangeImpl(TArrayView<const Type> InValues, int32 Index, IPCGAttributeAccessorKeys& Keys, EPCGAttributeAccessorFlags)
+	{
+		return PCGPropertyAccessor::IterateSet(Property, InValues, Index, Keys, [this](void* PropertyAddressData, const Type& Value) -> void
+			{
+				reinterpret_cast<T*>(PropertyAddressData)->SetPath(Value);
+			});
+	}
+
+private:
+	const FProperty* Property = nullptr;
+};
\ No newline at end of file
diff --git a/Engine/Plugins/Experimental/PCG/Source/PCG/Public/Metadata/PCGAttributePropertySelector.h b/Engine/Plugins/Experimental/PCG/Source/PCG/Public/Metadata/PCGAttributePropertySelector.h
new file mode 100644
--- /dev/null
+++ b/Engine/Plugins/Experimental/PCG/Source/PCG/Public/Metadata/PCGAttributePropertySelector.h
@@ -0,0 +1,71 @@
+// Copyright Epic Games, Inc. All Rights Reserved.
+
+#pragma once
+
+#include "PCGPoint.h"
+
+#include "Kismet/BlueprintFunctionLibrary.h"
+
+#include "PCGAttributePropertySelector.generated.h"
+
+UENUM()
+enum class EPCGAttributePropertySelection
+{
+	Attribute,
+	PointProperty,
+};
+
+USTRUCT(BlueprintType)
+struct PCG_API FPCGAttributePropertySelector
+{
+	GENERATED_BODY()
+
+public:
+	UPROPERTY(BlueprintReadOnly, VisibleAnywhere, Category = "Settings")
+	EPCGAttributePropertySelection Selection = EPCGAttributePropertySelection::Attribute;
+
+	UPROPERTY(BlueprintReadOnly, VisibleAnywhere, Category = "Settings", meta = (EditCondition = "Selection == EPCGAttributePropertySelection::Attribute", EditConditionHides))
+	FName AttributeName;
+
+	UPROPERTY(BlueprintReadOnly, VisibleAnywhere, Category = "Settings", meta = (EditCondition = "Selection == EPCGAttributePropertySelection::PointProperty", EditConditionHides))
+	EPCGPointProperties PointProperty = EPCGPointProperties::Position;
+
+	// Setters, retrurn true if something changed.
+	bool SetPointProperty(EPCGPointProperties InPointProperty);
+	bool SetAttributeName(FName InAttributeName);
+
+	FName GetName() const;
+
+	// TODO: To be removed when UI widget is done.
+	void Update();
+
+#if WITH_EDITOR
+	// Return true if the underlying name is valid.
+	bool IsValid() const;
+
+	// Returns the text to display in the widget.
+	FText GetDisplayText() const;
+
+	// Update the selector with an incoming string.
+	bool Update(FString NewValue);
+#endif // WITH_EDITOR
+};
+
+/**
+* Helper class to allow the BP to call the custom setters and getters on FPCGAttributePropertySelector.
+*/
+UCLASS()
+class PCG_API UPCGAttributePropertySelectorBlueprintHelpers : public UBlueprintFunctionLibrary
+{
+	GENERATED_BODY()
+
+public:
+	UFUNCTION(BlueprintCallable, Category = "PCG|PCGAttributePropertySelector", meta = (ScriptMethod))
+	static bool SetPointProperty(UPARAM(ref) FPCGAttributePropertySelector& Selector, EPCGPointProperties InPointProperty);
+
+	UFUNCTION(BlueprintCallable, Category = "PCG|PCGAttributePropertySelector", meta = (ScriptMethod))
+	static bool SetAttributeName(UPARAM(ref) FPCGAttributePropertySelector& Selector, FName InAttributeName);
+
+	UFUNCTION(BlueprintCallable, Category = "PCG|PCGAttributePropertySelector", meta = (ScriptMethod))
+	static FName GetName(UPARAM(const, ref) const FPCGAttributePropertySelector& Selector);
+};
\ No newline at end of file
diff --git a/Engine/Plugins/Experimental/PCG/Source/PCG/Public/Metadata/PCGMetadata.h b/Engine/Plugins/Experimental/PCG/Source/PCG/Public/Metadata/PCGMetadata.h
--- a/Engine/Plugins/Experimental/PCG/Source/PCG/Public/Metadata/PCGMetadata.h
+++ b/Engine/Plugins/Experimental/PCG/Source/PCG/Public/Metadata/PCGMetadata.h
@@ -154,11 +154,11 @@ public:
 
 	/** Copy attribute */
 	UFUNCTION(BlueprintCallable, Category = "PCG|Metadata")
-	void CopyExistingAttribute(FName AttributeToCopy, FName NewAttributeName, bool bKeepParent = true);
+	bool CopyExistingAttribute(FName AttributeToCopy, FName NewAttributeName, bool bKeepParent = true);
 
 	/** Rename attribute */
 	UFUNCTION(BlueprintCallable, Category = "PCG|Metadata")
-	void RenameAttribute(FName AttributeToRename, FName NewAttributeName);
+	bool RenameAttribute(FName AttributeToRename, FName NewAttributeName);
 
 	/** Clear/Reinit attribute */
 	UFUNCTION(BlueprintCallable, Category = "PCG|Metadata")
diff --git a/Engine/Plugins/Experimental/PCG/Source/PCG/Public/Metadata/PCGMetadataAttributeTpl.h b/Engine/Plugins/Experimental/PCG/Source/PCG/Public/Metadata/PCGMetadataAttributeTpl.h
--- a/Engine/Plugins/Experimental/PCG/Source/PCG/Public/Metadata/PCGMetadataAttributeTpl.h
+++ b/Engine/Plugins/Experimental/PCG/Source/PCG/Public/Metadata/PCGMetadataAttributeTpl.h
@@ -3,10 +3,12 @@
 #pragma once
 
 #include "CoreMinimal.h"
-#include "Metadata/PCGMetadataAttribute.h"
-#include "Metadata/PCGMetadataAttributeTraits.h"
+
+#include "PCGModule.h"
 
 #include "Helpers/PCGMetadataHelpers.h"
+#include "Metadata/PCGMetadataAttribute.h"
+#include "Metadata/PCGMetadataAttributeTraits.h"
 
 class UPCGMetadata;
 
@@ -551,23 +553,13 @@ namespace PCGMetadataAttribute
 {
 	inline FPCGMetadataAttributeBase* AllocateEmptyAttributeFromType(int16 TypeId)
 	{
-#define AllocatePCGMetadataAttributeOnType(Type) case PCG::Private::MetadataTypes<Type>::Id : { return new FPCGMetadataAttribute<Type>(); } break;
 
 		switch (TypeId)
 		{
-			AllocatePCGMetadataAttributeOnType(float);
-			AllocatePCGMetadataAttributeOnType(double);
-			AllocatePCGMetadataAttributeOnType(int32);
-			AllocatePCGMetadataAttributeOnType(int64);
-			AllocatePCGMetadataAttributeOnType(FVector);
-			AllocatePCGMetadataAttributeOnType(FVector2D);
-			AllocatePCGMetadataAttributeOnType(FVector4);
-			AllocatePCGMetadataAttributeOnType(FQuat);
-			AllocatePCGMetadataAttributeOnType(FTransform);
-			AllocatePCGMetadataAttributeOnType(FString);
-			AllocatePCGMetadataAttributeOnType(bool);
-			AllocatePCGMetadataAttributeOnType(FRotator);
-			AllocatePCGMetadataAttributeOnType(FName);
+
+#define PCG_ALLOCATEEMPTY_DECL(T) case PCG::Private::MetadataTypes<T>::Id: return new FPCGMetadataAttribute<T>();
+		PCG_FOREACH_SUPPORTEDTYPES(PCG_ALLOCATEEMPTY_DECL)
+#undef PCG_ALLOCATEEMPTY_DECL
 
 		default:
 			return nullptr;
@@ -583,32 +575,11 @@ namespace PCGMetadataAttribute
 
 		switch (TypeId)
 		{
-		case (uint16)EPCGMetadataTypes::Integer32:
-			return Callback(int32{});
-		case (uint16)EPCGMetadataTypes::Integer64:
-			return Callback(int64{});
-		case (uint16)EPCGMetadataTypes::Float:
-			return Callback(float{});
-		case (uint16)EPCGMetadataTypes::Double:
-			return Callback(double{});
-		case (uint16)EPCGMetadataTypes::Vector2:
-			return Callback(FVector2D{});
-		case (uint16)EPCGMetadataTypes::Vector:
-			return Callback(FVector{});
-		case (uint16)EPCGMetadataTypes::Vector4:
-			return Callback(FVector4{});
-		case (uint16)EPCGMetadataTypes::Quaternion:
-			return Callback(FQuat{});
-		case (uint16)EPCGMetadataTypes::Transform:
-			return Callback(FTransform{});
-		case (uint16)EPCGMetadataTypes::String:
-			return Callback(FString{});
-		case (uint16)EPCGMetadataTypes::Boolean:
-			return Callback(bool{});
-		case (uint16)EPCGMetadataTypes::Rotator:
-			return Callback(FRotator{});
-		case (uint16)EPCGMetadataTypes::Name:
-			return Callback(FName{});
+
+#define PCG_CALLBACKWITHRIGHTTYPE_DECL(T) case (uint16)(PCG::Private::MetadataTypes<T>::Id): return Callback(T{});
+		PCG_FOREACH_SUPPORTEDTYPES(PCG_CALLBACKWITHRIGHTTYPE_DECL)
+#undef PCG_CALLBACKWITHRIGHTTYPE_DECL
+
 		default:
 		{
 			// ReturnType{} is invalid if ReturnType is void
diff --git a/Engine/Plugins/Experimental/PCG/Source/PCG/Public/Metadata/PCGMetadataAttributeTraits.h b/Engine/Plugins/Experimental/PCG/Source/PCG/Public/Metadata/PCGMetadataAttributeTraits.h
--- a/Engine/Plugins/Experimental/PCG/Source/PCG/Public/Metadata/PCGMetadataAttributeTraits.h
+++ b/Engine/Plugins/Experimental/PCG/Source/PCG/Public/Metadata/PCGMetadataAttributeTraits.h
@@ -28,6 +28,21 @@ enum class EPCGMetadataTypes : uint8
 	Unknown = 255
 };
 
+// Convenient macro to avoid duplicating a lot of code with all our supported types.
+#define PCG_FOREACH_SUPPORTEDTYPES(MACRO) \
+	MACRO(int32)      \
+	MACRO(int64)      \
+	MACRO(float)      \
+	MACRO(double)     \
+	MACRO(FVector2D)  \
+	MACRO(FVector)    \
+	MACRO(FVector4)   \
+	MACRO(FQuat)      \
+	MACRO(FRotator)   \
+	MACRO(FName)      \
+	MACRO(FString)    \
+	MACRO(FTransform)
+
 namespace PCG
 {
 	namespace Private
@@ -56,6 +71,17 @@ namespace PCG
 
 #undef PCGMetadataGenerateDataTypes
 
+		constexpr inline bool IsPCGType(uint16 TypeId)
+		{
+			return TypeId < (uint16)EPCGMetadataTypes::Count;
+		}
+
+		template <typename T>
+		constexpr inline bool IsPCGType()
+		{
+			return IsPCGType(MetadataTypes<T>::Id);
+		}
+
 		// Verify if the TypeId is a type known by PCG and matches any types provided in the template.
 		// Example: IsOfTypes<FVector2D, FVector, FVector4>(TypeId);
 		template <typename ...Types>
diff --git a/Engine/Plugins/Experimental/PCG/Source/PCG/Public/Metadata/PCGMetadataAttributeWrapper.h b/Engine/Plugins/Experimental/PCG/Source/PCG/Public/Metadata/PCGMetadataAttributeWrapper.h
deleted file mode 100644
diff --git a/Engine/Plugins/Experimental/PCG/Source/PCG/Public/PCGPoint.h b/Engine/Plugins/Experimental/PCG/Source/PCG/Public/PCGPoint.h
--- a/Engine/Plugins/Experimental/PCG/Source/PCG/Public/PCGPoint.h
+++ b/Engine/Plugins/Experimental/PCG/Source/PCG/Public/PCGPoint.h
@@ -7,6 +7,8 @@
 
 #include "PCGPoint.generated.h"
 
+class IPCGAttributeAccessor;
+
 UENUM()
 enum class EPCGPointProperties : uint8
 {
@@ -78,50 +80,6 @@ public:
 	using PointCustomPropertyGetter = TFunction<bool(const FPCGPoint&, void*)>;
 	using PointCustomPropertySetter = TFunction<bool(FPCGPoint&, const void*)>;
 
-	struct PointCustomPropertyGetterSetter
-	{
-	public:
-		PointCustomPropertyGetterSetter() = default;
-
-		PointCustomPropertyGetterSetter(const PointCustomPropertyGetter& InGetter, const PointCustomPropertySetter& InSetter, int16 InType, FName InName);
-
-		template<typename T>
-		bool Get(const FPCGPoint& Point, T& OutValue) const
-		{
-			if (PCG::Private::IsOfTypes<T>(Type))
-			{
-				return Getter(Point, &OutValue);
-			}
-			else
-			{
-				return false;
-			}
-		}
-
-		template<typename T>
-		bool Set(FPCGPoint& Point, const T& InValue)
-		{
-			if (PCG::Private::IsOfTypes<T>(Type))
-			{
-				return Setter(Point, &InValue);
-			}
-			else
-			{
-				return false;
-			}
-		}
-
-		bool IsValid() const { return Type >= 0 && Type < (int16)EPCGMetadataTypes::Count; }
-		int16 GetType() const { return Type; }
-		FName GetName() const { return Name; }
-
-	private:
-		PointCustomPropertyGetter Getter = [](const FPCGPoint&, void*) { return false; };
-		PointCustomPropertySetter Setter = [](FPCGPoint&, const void*) { return false; };
-		int16 Type = -1;
-		FName Name = NAME_None;
-	};
-
 	static bool HasCustomPropertyGetterSetter(FName Name);
-	static PointCustomPropertyGetterSetter CreateCustomPropertyGetterSetter(FName Name);
+	static TUniquePtr<IPCGAttributeAccessor> CreateCustomPropertyAccessor(FName Name);
 };
