Adjusted the way pcg world actors register & unregister to the subsystem. The previous behaivor could cause some issues when the pcg world actor gets deleted but the subsystem does not.
Note that this is a first change towards multi pcg world actor support.

#rnx

diff --git a/Engine/Plugins/Experimental/PCG/Source/PCG/Private/PCGSubsystem.cpp b/Engine/Plugins/Experimental/PCG/Source/PCG/Private/PCGSubsystem.cpp
--- a/Engine/Plugins/Experimental/PCG/Source/PCG/Private/PCGSubsystem.cpp
+++ b/Engine/Plugins/Experimental/PCG/Source/PCG/Private/PCGSubsystem.cpp
@@ -82,6 +82,8 @@ void UPCGSubsystem::Deinitialize()
 	delete GraphExecutor;
 	GraphExecutor = nullptr;
 
+	PCGWorldActor = nullptr;
+
 	Super::Deinitialize();
 }
 
diff --git a/Engine/Plugins/Experimental/PCG/Source/PCG/Private/PCGWorldActor.cpp b/Engine/Plugins/Experimental/PCG/Source/PCG/Private/PCGWorldActor.cpp
--- a/Engine/Plugins/Experimental/PCG/Source/PCG/Private/PCGWorldActor.cpp
+++ b/Engine/Plugins/Experimental/PCG/Source/PCG/Private/PCGWorldActor.cpp
@@ -31,16 +31,16 @@ void APCGWorldActor::BeginCacheForCookedPlatformData(const ITargetPlatform* Targ
 }
 #endif
 
-void APCGWorldActor::PostLoad()
+void APCGWorldActor::BeginPlay()
 {
-	Super::PostLoad();
+	Super::BeginPlay();
 	RegisterToSubsystem();
 }
 
-void APCGWorldActor::BeginDestroy()
+void APCGWorldActor::EndPlay(const EEndPlayReason::Type EndPlayReason)
 {
 	UnregisterFromSubsystem();
-	Super::BeginDestroy();
+	Super::EndPlay(EndPlayReason);
 }
 
 #if WITH_EDITOR
diff --git a/Engine/Plugins/Experimental/PCG/Source/PCG/Public/PCGWorldActor.h b/Engine/Plugins/Experimental/PCG/Source/PCG/Public/PCGWorldActor.h
--- a/Engine/Plugins/Experimental/PCG/Source/PCG/Public/PCGWorldActor.h
+++ b/Engine/Plugins/Experimental/PCG/Source/PCG/Public/PCGWorldActor.h
@@ -18,8 +18,8 @@ public:
 	APCGWorldActor(const FObjectInitializer& ObjectInitializer);
 
 	//~Begin AActor Interface
-	virtual void PostLoad() override;
-	virtual void BeginDestroy() override;
+	virtual void BeginPlay() override;
+	virtual void EndPlay(const EEndPlayReason::Type EndPlayReason) override;
 
 #if WITH_EDITOR	
 	virtual bool CanChangeIsSpatiallyLoadedFlag() const { return false; }
