[PCG] Add PropertyAttributeSelector detail view + move all details in their own folder

diff --git a/Engine/Plugins/Experimental/PCG/Source/PCG/Public/Metadata/PCGAttributePropertySelector.h b/Engine/Plugins/Experimental/PCG/Source/PCG/Public/Metadata/PCGAttributePropertySelector.h
--- a/Engine/Plugins/Experimental/PCG/Source/PCG/Public/Metadata/PCGAttributePropertySelector.h
+++ b/Engine/Plugins/Experimental/PCG/Source/PCG/Public/Metadata/PCGAttributePropertySelector.h
@@ -15,6 +15,11 @@ enum class EPCGAttributePropertySelection
 	PointProperty,
 };
 
+/**
+* Blueprint class to allow to select an attribute or a property.
+* It will handle the logic and can only be modified using the blueprint library defined below.
+* Also has a custom detail view in the PCGEditor plugin.
+*/
 USTRUCT(BlueprintType)
 struct PCG_API FPCGAttributePropertySelector
 {
diff --git a/Engine/Plugins/Experimental/PCG/Source/PCGEditor/Private/Details/PCGAttributePropertySelectorDetails.cpp b/Engine/Plugins/Experimental/PCG/Source/PCGEditor/Private/Details/PCGAttributePropertySelectorDetails.cpp
new file mode 100644
--- /dev/null
+++ b/Engine/Plugins/Experimental/PCG/Source/PCGEditor/Private/Details/PCGAttributePropertySelectorDetails.cpp
@@ -0,0 +1,192 @@
+// Copyright Epic Games, Inc. All Rights Reserved.
+
+#include "Details/PCGAttributePropertySelectorDetails.h"
+
+#include "Metadata/PCGAttributePropertySelector.h"
+
+#include "DetailLayoutBuilder.h"
+#include "DetailWidgetRow.h"
+#include "ScopedTransaction.h"
+#include "Framework/Application/SlateApplication.h"
+#include "Framework/MultiBox/MultiBoxBuilder.h"
+#include "Styling/StyleColors.h"
+#include "Widgets/SBoxPanel.h"
+#include "Widgets/Images/SImage.h"
+#include "Widgets/Input/SComboButton.h"
+#include "Widgets/Input/SEditableTextBox.h"
+
+#define LOCTEXT_NAMESPACE "PCGAttributePropertySelectorDetails"
+
+void FPCGAttributePropertySelectorDetails::CustomizeHeader(TSharedRef<IPropertyHandle> InPropertyHandle, FDetailWidgetRow& HeaderRow, IPropertyTypeCustomizationUtils& CustomizationUtils)
+{
+	PropertyHandle = InPropertyHandle;
+	check(PropertyHandle.IsValid());
+
+	FText TooltipText = LOCTEXT("TooltipText", "Enter the name of the attribute. You can prefix it by '$' to get a property. Red text means that your attribute or property is invalid.");
+
+	auto Validation = [this]() -> FSlateColor
+	{
+		FPCGAttributePropertySelector* Selector = GetStruct();
+		return (Selector && !Selector->IsValid()) ? FStyleColors::AccentRed : FSlateColor::UseForeground();
+	};
+
+	HeaderRow
+		.NameContent()
+		[
+			PropertyHandle->CreatePropertyNameWidget()
+		]
+		.ValueContent()
+		.MinDesiredWidth(250.0f)
+		[
+			SNew(SHorizontalBox)
+			+ SHorizontalBox::Slot()
+			[
+				SNew(SEditableTextBox)
+				.Text(this, &FPCGAttributePropertySelectorDetails::GetText)
+				.ToolTipText(TooltipText)
+				.OnTextCommitted(this, &FPCGAttributePropertySelectorDetails::SetText)
+				.Font(IDetailLayoutBuilder::GetDetailFont())
+				.ForegroundColor_Lambda(Validation)
+			]
+			+ SHorizontalBox::Slot()
+			.MaxWidth(20.0f)
+			[
+				SNew(SComboButton)
+				.ButtonStyle(FAppStyle::Get(), "SimpleButton")
+				.OnGetMenuContent(this, &FPCGAttributePropertySelectorDetails::GenerateExtraMenu)
+				.HasDownArrow(false)
+				.ButtonContent()
+				[
+					SNew(SImage)
+					.ColorAndOpacity(FSlateColor::UseForeground())
+					.Image(FAppStyle::Get().GetBrush("Icons.PlusCircle"))
+				]
+			]
+		];
+}
+
+TSharedRef<SWidget> FPCGAttributePropertySelectorDetails::GenerateExtraMenu()
+{
+	// Clear the focus on the text box. It is preventing from updating the text.
+	FSlateApplication::Get().ClearKeyboardFocus(EFocusCause::Cleared);
+
+	FMenuBuilder MenuBuilder(true, nullptr);
+
+	MenuBuilder.BeginSection("Attributes", LOCTEXT("AttributesHeader", "Attributes"));
+	{
+		MenuBuilder.AddMenuEntry(
+			LOCTEXT("LastAttributeHeader", "Last Attribute(None)"),
+			FText(),
+			FSlateIcon(),
+			FExecuteAction::CreateSP(this, &FPCGAttributePropertySelectorDetails::SetAttributeName, FName()),
+			NAME_None,
+			EUserInterfaceActionType::Button
+		);
+	}
+	MenuBuilder.EndSection();
+
+	MenuBuilder.BeginSection("PointProperties", LOCTEXT("PointPropertiesHeader", "Point Properties"));
+	{
+		if (const UEnum* EnumPtr = StaticEnum<EPCGPointProperties>())
+		{
+			for (int32 i = 0; i < EnumPtr->NumEnums() - 1; ++i)
+			{
+				FString EnumName = EnumPtr->GetNameStringByIndex(i);
+				EPCGPointProperties EnumValue = (EPCGPointProperties)EnumPtr->GetValueByIndex(i);
+
+				MenuBuilder.AddMenuEntry(
+					FText::FromString(EnumName),
+					FText(),
+					FSlateIcon(),
+					FExecuteAction::CreateSP(this, &FPCGAttributePropertySelectorDetails::SetPointProperty, EnumValue),
+					NAME_None,
+					EUserInterfaceActionType::Button
+				);
+			}
+		}
+	}
+	MenuBuilder.EndSection();
+
+	return MenuBuilder.MakeWidget();
+}
+
+FPCGAttributePropertySelector* FPCGAttributePropertySelectorDetails::GetStruct()
+{
+	void* Data = nullptr;
+	FPropertyAccess::Result Result = PropertyHandle->GetValueData(Data);
+
+	return (Result == FPropertyAccess::Success) ? reinterpret_cast<FPCGAttributePropertySelector*>(Data) : nullptr;
+}
+
+const FPCGAttributePropertySelector* FPCGAttributePropertySelectorDetails::GetStruct() const
+{
+	void* Data = nullptr;
+	FPropertyAccess::Result Result = PropertyHandle->GetValueData(Data);
+
+	return (Result == FPropertyAccess::Success) ? reinterpret_cast<const FPCGAttributePropertySelector*>(Data) : nullptr;
+}
+
+FText FPCGAttributePropertySelectorDetails::GetText() const
+{
+	const FPCGAttributePropertySelector* Selector = GetStruct();
+
+	return Selector ? Selector->GetDisplayText() : FText();
+}
+
+void FPCGAttributePropertySelectorDetails::SetText(const FText& NewText, ETextCommit::Type CommitInfo)
+{
+	if (CommitInfo == ETextCommit::Type::OnCleared)
+	{
+		return;
+	}
+
+	if (FPCGAttributePropertySelector* Selector = GetStruct())
+	{
+		FScopedTransaction Transaction(LOCTEXT("SetAttributePropertyName", "[PCG] Set Attribute/Property Name"));
+
+		PropertyHandle->NotifyPreChange();
+		if (!Selector->Update(NewText.ToString()))
+		{
+			Transaction.Cancel();
+			return;
+		}
+
+		PropertyHandle->NotifyPostChange(EPropertyChangeType::ValueSet);
+	}
+}
+
+void FPCGAttributePropertySelectorDetails::SetPointProperty(EPCGPointProperties EnumValue)
+{
+	if (FPCGAttributePropertySelector* Selector = GetStruct())
+	{
+		FScopedTransaction Transaction(LOCTEXT("SetAttributePropertyName", "[PCG] Set Attribute/Property Name"));
+
+		PropertyHandle->NotifyPreChange();
+		if (!Selector->SetPointProperty(EnumValue))
+		{
+			Transaction.Cancel();
+			return;
+		}
+
+		PropertyHandle->NotifyPostChange(EPropertyChangeType::ValueSet);
+	}
+}
+
+void FPCGAttributePropertySelectorDetails::SetAttributeName(FName NewName)
+{
+	if (FPCGAttributePropertySelector* Selector = GetStruct())
+	{
+		FScopedTransaction Transaction(LOCTEXT("SetAttributePropertyName", "[PCG] Set Attribute/Property Name"));
+
+		PropertyHandle->NotifyPreChange();
+		if (!Selector->SetAttributeName(NewName))
+		{
+			Transaction.Cancel();
+			return;
+		}
+
+		PropertyHandle->NotifyPostChange(EPropertyChangeType::ValueSet);
+	}
+}
+
+#undef LOCTEXT_NAMESPACE
\ No newline at end of file
diff --git a/Engine/Plugins/Experimental/PCG/Source/PCGEditor/Private/Details/PCGBlueprintSettingsDetails.cpp b/Engine/Plugins/Experimental/PCG/Source/PCGEditor/Private/Details/PCGBlueprintSettingsDetails.cpp
new file mode 100644
--- /dev/null
+++ b/Engine/Plugins/Experimental/PCG/Source/PCGEditor/Private/Details/PCGBlueprintSettingsDetails.cpp
@@ -0,0 +1,224 @@
+// Copyright Epic Games, Inc. All Rights Reserved.
+
+#include "Details/PCGBlueprintSettingsDetails.h"
+
+#include "DetailLayoutBuilder.h"
+//#include "ObjectDetailsTools.h"
+
+#include "Containers/ContainerAllocationPolicies.h"
+#include "Containers/UnrealString.h"
+#include "Delegates/Delegate.h"
+#include "DetailCategoryBuilder.h"
+#include "DetailWidgetRow.h"
+#include "EdGraphSchema_K2.h"
+#include "Engine/Blueprint.h"
+#include "Framework/SlateDelegates.h"
+#include "HAL/Platform.h"
+#include "HAL/PlatformCrt.h"
+#include "Internationalization/Internationalization.h"
+#include "Internationalization/Text.h"
+#include "Kismet2/BlueprintEditorUtils.h"
+#include "Layout/Margin.h"
+#include "Math/NumericLimits.h"
+#include "Misc/Attribute.h"
+#include "Misc/CString.h"
+#include "ObjectEditorUtils.h"
+#include "SWarningOrErrorBox.h"
+#include "ScopedTransaction.h"
+#include "SlotBase.h"
+#include "Templates/Casts.h"
+#include "UObject/Class.h"
+#include "UObject/NameTypes.h"
+#include "UObject/Object.h"
+#include "UObject/ObjectMacros.h"
+#include "UObject/Script.h"
+#include "UObject/UnrealNames.h"
+#include "UObject/UnrealType.h"
+#include "Widgets/DeclarativeSyntaxSupport.h"
+#include "Widgets/Input/SButton.h"
+#include "Widgets/Layout/SBox.h"
+#include "Widgets/Layout/SWrapBox.h"
+
+#define LOCTEXT_NAMESPACE "PCGSettingsDetails"
+
+TSharedRef<IDetailCustomization> FPCGBlueprintSettingsDetails::MakeInstance()
+{
+	return MakeShareable(new FPCGBlueprintSettingsDetails());
+}
+
+void FPCGBlueprintSettingsDetails::CustomizeDetails(IDetailLayoutBuilder& DetailBuilder)
+{
+	TArray<TWeakObjectPtr<UObject>> ObjectsBeingCustomized;
+	DetailBuilder.GetObjectsBeingCustomized(ObjectsBeingCustomized);
+
+	if (ObjectsBeingCustomized.Num() != 1)
+	{
+		return;
+	}
+
+	UPCGSettingsInterface* SettingsInterface = Cast<UPCGSettingsInterface>(ObjectsBeingCustomized[0]);
+
+	if (!SettingsInterface)
+	{
+		return;
+	}
+
+	UPCGBlueprintSettings* BlueprintSettings = Cast<UPCGBlueprintSettings>(SettingsInterface->GetSettings());
+
+	if (!BlueprintSettings)
+	{
+		return;
+	}
+
+	// Go through all the properties, if we find properties that have CallInEditor,
+	// create buttons for them to be executed
+	UPCGBlueprintElement* BlueprintElement = BlueprintSettings->GetElementInstance();
+	if (BlueprintElement)
+	{
+		SelectedObject = TWeakObjectPtr<UPCGBlueprintElement>(BlueprintElement);
+		AddCallInEditorMethods(DetailBuilder);
+		//FObjectDetailsTools::AddCallInEditorMethods(
+		//	DetailBuilder,
+		//	BlueprintElement->GetClass(),
+		//	[this](TWeakObjectPtr<UFunction> WeakFunctionPtr) { return FOnClicked::CreateSP(this, &FPCGBlueprintSettingsDetails::OnExecuteCallInEditorFunction, WeakFunctionPtr); });
+	}
+}
+
+void FPCGBlueprintSettingsDetails::AddCallInEditorMethods(IDetailLayoutBuilder& DetailBuilder)
+{
+	check(SelectedObject.IsValid());
+	// NOTE: this is an adaptation of FObjectDetails::AddCallInEditorMethods
+	// metadata tag for defining sort order of function buttons within a Category
+	static const FName NAME_DisplayPriority("DisplayPriority");
+
+	TArray<UFunction*, TInlineAllocator<8>> CallInEditorFunctions;
+
+	for (TFieldIterator<UFunction> FunctionIter(SelectedObject->GetClass(), EFieldIteratorFlags::IncludeSuper); FunctionIter; ++FunctionIter)
+	{
+		UFunction* Function = *FunctionIter;
+		if (Function->GetBoolMetaData(FBlueprintMetadata::MD_CallInEditor) && (Function->ParmsSize == 0))
+		{
+			const FName FunctionName = Function->GetFName();
+			if (!CallInEditorFunctions.FindByPredicate([&FunctionName](const UFunction* Func) { return Func->GetFName() == FunctionName; }))
+			{
+				CallInEditorFunctions.Add(*FunctionIter);
+			}
+		}
+	}
+
+	if (CallInEditorFunctions.Num() > 0)
+	{
+		// Sort the functions by category and then by DisplayPriority meta tag, and then by name
+		CallInEditorFunctions.Sort([](UFunction& A, UFunction& B)
+		{
+			const int32 CategorySort = A.GetMetaData(FBlueprintMetadata::MD_FunctionCategory).Compare(B.GetMetaData(FBlueprintMetadata::MD_FunctionCategory));
+			if (CategorySort != 0)
+			{
+				return (CategorySort <= 0);
+			}
+			else
+			{
+				FString DisplayPriorityAStr = A.GetMetaData(NAME_DisplayPriority);
+				int32 DisplayPriorityA = (DisplayPriorityAStr.IsEmpty() ? MAX_int32 : FCString::Atoi(*DisplayPriorityAStr));
+				if (DisplayPriorityA == 0 && !FCString::IsNumeric(*DisplayPriorityAStr))
+				{
+					DisplayPriorityA = MAX_int32;
+				}
+
+				FString DisplayPriorityBStr = B.GetMetaData(NAME_DisplayPriority);
+				int32 DisplayPriorityB = (DisplayPriorityBStr.IsEmpty() ? MAX_int32 : FCString::Atoi(*DisplayPriorityBStr));
+				if (DisplayPriorityB == 0 && !FCString::IsNumeric(*DisplayPriorityBStr))
+				{
+					DisplayPriorityB = MAX_int32;
+				}
+
+				return (DisplayPriorityA == DisplayPriorityB) ? (A.GetName() <= B.GetName()) : (DisplayPriorityA <= DisplayPriorityB);
+			}
+		});
+
+		struct FCategoryEntry
+		{
+			FName CategoryName;
+			FName RowTag;
+			TSharedPtr<SWrapBox> WrapBox;
+			FTextBuilder FunctionSearchText;
+
+			FCategoryEntry(FName InCategoryName)
+				: CategoryName(InCategoryName)
+			{
+				WrapBox = SNew(SWrapBox).UseAllottedSize(true);
+			}
+		};
+
+		// Build up a set of functions for each category, accumulating search text and buttons in a wrap box
+		FName ActiveCategory;
+		TArray<FCategoryEntry, TInlineAllocator<8>> CategoryList;
+		for (UFunction* Function : CallInEditorFunctions)
+		{
+			FName FunctionCategoryName(NAME_Default);
+			if (Function->HasMetaData(FBlueprintMetadata::FBlueprintMetadata::MD_FunctionCategory))
+			{
+				FunctionCategoryName = FName(*Function->GetMetaData(FBlueprintMetadata::MD_FunctionCategory));
+			}
+
+			if (FunctionCategoryName != ActiveCategory)
+			{
+				ActiveCategory = FunctionCategoryName;
+				CategoryList.Emplace(FunctionCategoryName);
+			}
+			FCategoryEntry& CategoryEntry = CategoryList.Last();
+
+			//@TODO: Expose the code in UK2Node_CallFunction::GetUserFacingFunctionName / etc...
+			const FText ButtonCaption = FText::FromString(FName::NameToDisplayString(*Function->GetName(), false));
+			FText FunctionTooltip = Function->GetToolTipText();
+			if (FunctionTooltip.IsEmpty())
+			{
+				FunctionTooltip = FText::FromString(Function->GetName());
+			}
+
+			TWeakObjectPtr<UFunction> WeakFunctionPtr(Function);
+			CategoryEntry.WrapBox->AddSlot()
+				.Padding(0.0f, 0.0f, 5.0f, 3.0f)
+				[
+					SNew(SButton)
+					.Text(ButtonCaption)
+					.OnClicked(FOnClicked::CreateSP(this, &FPCGBlueprintSettingsDetails::OnExecuteCallInEditorFunction, WeakFunctionPtr))
+					.ToolTipText(FunctionTooltip.IsEmptyOrWhitespace() ? LOCTEXT("CallInEditorTooltip", "Call an event on the selected object(s)") : FunctionTooltip)
+				];
+
+			CategoryEntry.RowTag = Function->GetFName();
+			CategoryEntry.FunctionSearchText.AppendLine(ButtonCaption);
+			CategoryEntry.FunctionSearchText.AppendLine(FunctionTooltip);
+		}
+
+		// Now edit the categories, adding the button strips to the details panel
+		for (FCategoryEntry& CategoryEntry : CategoryList)
+		{
+			IDetailCategoryBuilder& CategoryBuilder = DetailBuilder.EditCategory(CategoryEntry.CategoryName);
+			CategoryBuilder.AddCustomRow(CategoryEntry.FunctionSearchText.ToText())
+				.RowTag(CategoryEntry.RowTag)
+				[
+					CategoryEntry.WrapBox.ToSharedRef()
+				];
+		}
+	}
+}
+
+FReply FPCGBlueprintSettingsDetails::OnExecuteCallInEditorFunction(TWeakObjectPtr<UFunction> WeakFunctionPtr)
+{
+	if (UFunction* Function = WeakFunctionPtr.Get())
+	{
+		//@TODO: Consider naming the transaction scope after the fully qualified function name for better UX
+		FScopedTransaction Transaction(LOCTEXT("ExecuteCallInEditorMethod", "Call In Editor Action"));
+
+		FEditorScriptExecutionGuard ScriptGuard;
+		if (UObject* Object = SelectedObject.Get())
+		{
+			Object->ProcessEvent(Function, nullptr);
+		}
+	}
+
+	return FReply::Handled();
+}
+
+#undef LOCTEXT_NAMESPACE
\ No newline at end of file
diff --git a/Engine/Plugins/Experimental/PCG/Source/PCGEditor/Private/Details/PCGComponentDetails.cpp b/Engine/Plugins/Experimental/PCG/Source/PCGEditor/Private/Details/PCGComponentDetails.cpp
new file mode 100644
--- /dev/null
+++ b/Engine/Plugins/Experimental/PCG/Source/PCGEditor/Private/Details/PCGComponentDetails.cpp
@@ -0,0 +1,231 @@
+// Copyright Epic Games, Inc. All Rights Reserved.
+
+#include "Details/PCGComponentDetails.h"
+#include "PCGComponent.h"
+
+#include "Widgets/DeclarativeSyntaxSupport.h"
+#include "Widgets/Text/STextBlock.h"
+#include "Widgets/Input/SButton.h"
+#include "PropertyHandle.h"
+#include "DetailLayoutBuilder.h"
+#include "DetailWidgetRow.h"
+#include "DetailCategoryBuilder.h"
+#include "Modules/ModuleManager.h"
+#include "LevelEditor.h"
+
+#define LOCTEXT_NAMESPACE "PCGComponentDetails"
+
+TSharedRef<IDetailCustomization> FPCGComponentDetails::MakeInstance()
+{
+	return MakeShareable(new FPCGComponentDetails());
+}
+
+void FPCGComponentDetails::GatherPCGComponentsFromSelection(const TArray<TWeakObjectPtr<UObject>>& InObjectSelected)
+{
+	for (const TWeakObjectPtr<UObject>& Object : InObjectSelected)
+	{
+		UPCGComponent* Component = Cast<UPCGComponent>(Object.Get());
+		if (ensure(Component))
+		{
+			SelectedComponents.Add(Component);
+		}
+	}
+}
+
+void FPCGComponentDetails::CustomizeDetails(IDetailLayoutBuilder& DetailBuilder)
+{
+	const FName PCGCategoryName("PCG");
+	IDetailCategoryBuilder& PCGCategory = DetailBuilder.EditCategory(PCGCategoryName);
+
+	TArray<TWeakObjectPtr<UObject>> ObjectsBeingCustomized;
+	DetailBuilder.GetObjectsBeingCustomized(ObjectsBeingCustomized);
+
+	GatherPCGComponentsFromSelection(ObjectsBeingCustomized);
+
+	if (AddDefaultProperties())
+	{
+		TArray<TSharedRef<IPropertyHandle>> AllProperties;
+		bool bSimpleProperties = true;
+		bool bAdvancedProperties = false;
+		// Add all properties in the category in order
+		PCGCategory.GetDefaultProperties(AllProperties, bSimpleProperties, bAdvancedProperties);
+
+		for (auto& Property : AllProperties)
+		{
+			PCGCategory.AddProperty(Property);
+		}
+	}
+
+	FDetailWidgetRow& NewRow = PCGCategory.AddCustomRow(FText::GetEmpty());
+
+	NewRow.ValueContent()
+		.MaxDesiredWidth(120.f)
+		[
+			SNew(SHorizontalBox)
+			+ SHorizontalBox::Slot()
+			.AutoWidth()
+			.Padding(4.0f)
+			[
+				SNew(SButton)
+				.OnClicked(this, &FPCGComponentDetails::OnGenerateClicked)
+				.Visibility(this, &FPCGComponentDetails::GenerateButtonVisible)
+				[
+					SNew(STextBlock)
+					.Font(IDetailLayoutBuilder::GetDetailFont())
+					.Text(LOCTEXT("GenerateButton", "Generate"))
+				]
+			]
+			+ SHorizontalBox::Slot()
+			.AutoWidth()
+			.Padding(4.0f)
+			[
+				SNew(SButton)
+				.OnClicked(this, &FPCGComponentDetails::OnCancelClicked)
+				.Visibility(this, &FPCGComponentDetails::CancelButtonVisible)
+				[
+					SNew(STextBlock)
+					.Font(IDetailLayoutBuilder::GetDetailFont())
+					.Text(LOCTEXT("CancelButton", "Cancel"))
+				]
+			]
+			+ SHorizontalBox::Slot()
+			.AutoWidth()
+			.Padding(4.0f)
+			[
+				SNew(SButton)
+				.OnClicked(this, &FPCGComponentDetails::OnCleanupClicked)				
+				[
+					SNew(STextBlock)					
+					.Font(IDetailLayoutBuilder::GetDetailFont())
+					.Text(LOCTEXT("CleanupButton", "Cleanup"))
+				]
+			]
+			+ SHorizontalBox::Slot()
+			.AutoWidth()
+			.Padding(4.0f)
+			[
+				SNew(SButton)
+				.OnClicked(this, &FPCGComponentDetails::OnClearPCGLinkClicked)
+				[
+					SNew(STextBlock)
+					.Font(IDetailLayoutBuilder::GetDetailFont())
+					.Text(LOCTEXT("ClearPCGLinkButton", "Clear PCG Link"))
+				]
+			]
+		];
+
+	// Attach to generated delegate on the selected components
+	for (TWeakObjectPtr<UPCGComponent>& Component : SelectedComponents)
+	{
+		if (Component.IsValid())
+		{
+			Component.Get()->OnPCGGraphGeneratedDelegate.AddRaw(this, &FPCGComponentDetails::OnGraphChanged);
+			Component.Get()->OnPCGGraphCleanedDelegate.AddRaw(this, &FPCGComponentDetails::OnGraphChanged);
+		}
+	}
+}
+
+void FPCGComponentDetails::PendingDelete()
+{
+	// Detach from the generated delegate
+	for (TWeakObjectPtr<UPCGComponent>& Component : SelectedComponents)
+	{
+		if (Component.IsValid())
+		{
+			Component.Get()->OnPCGGraphGeneratedDelegate.RemoveAll(this);
+			Component.Get()->OnPCGGraphCleanedDelegate.RemoveAll(this);
+		}
+	}
+}
+
+EVisibility FPCGComponentDetails::GenerateButtonVisible() const
+{
+	for (const TWeakObjectPtr<UPCGComponent>& Component : SelectedComponents)
+	{
+		if (Component.IsValid() && !Component->IsGenerating())
+		{
+			return EVisibility::Visible;
+		}
+	}
+
+	return EVisibility::Collapsed;
+}
+
+EVisibility FPCGComponentDetails::CancelButtonVisible() const
+{
+	for (const TWeakObjectPtr<UPCGComponent>& Component : SelectedComponents)
+	{
+		if (Component.IsValid() && Component->IsGenerating())
+		{
+			return EVisibility::Visible;
+		}
+	}
+
+	return EVisibility::Collapsed;
+}
+
+FReply FPCGComponentDetails::OnGenerateClicked()
+{
+	for (TWeakObjectPtr<UPCGComponent>& Component : SelectedComponents)
+	{
+		if (Component.IsValid())
+		{
+			Component.Get()->Generate();
+		}
+	}
+
+	return FReply::Handled();
+}
+
+FReply FPCGComponentDetails::OnCancelClicked()
+{
+	for (TWeakObjectPtr<UPCGComponent>& Component : SelectedComponents)
+	{
+		if (Component.IsValid())
+		{
+			Component.Get()->CancelGeneration();
+		}
+	}
+
+	return FReply::Handled();
+}
+
+FReply FPCGComponentDetails::OnClearPCGLinkClicked()
+{
+	for (TWeakObjectPtr<UPCGComponent>& Component : SelectedComponents)
+	{
+		if (Component.IsValid())
+		{
+			Component.Get()->ClearPCGLink();
+		}
+	}
+
+	return FReply::Handled();
+}
+
+void FPCGComponentDetails::OnGraphChanged(UPCGComponent* InComponent)
+{
+	if (!InComponent || !InComponent->GetOwner())
+	{
+		return;
+	}
+
+	// Notify editor that some components might have changed
+	FLevelEditorModule& LevelEditor = FModuleManager::LoadModuleChecked<FLevelEditorModule>("LevelEditor");
+	LevelEditor.BroadcastComponentsEdited();
+}
+
+FReply FPCGComponentDetails::OnCleanupClicked()
+{
+	for (TWeakObjectPtr<UPCGComponent>& Component : SelectedComponents)
+	{
+		if (Component.IsValid())
+		{
+			Component.Get()->Cleanup();
+		}
+	}
+
+	return FReply::Handled();
+}
+
+#undef LOCTEXT_NAMESPACE
diff --git a/Engine/Plugins/Experimental/PCG/Source/PCGEditor/Private/Details/PCGGraphDetails.cpp b/Engine/Plugins/Experimental/PCG/Source/PCGEditor/Private/Details/PCGGraphDetails.cpp
new file mode 100644
--- /dev/null
+++ b/Engine/Plugins/Experimental/PCG/Source/PCGEditor/Private/Details/PCGGraphDetails.cpp
@@ -0,0 +1,50 @@
+// Copyright Epic Games, Inc. All Rights Reserved.
+
+#include "Details/PCGGraphDetails.h"
+#include "PCGGraph.h"
+
+#include "Widgets/DeclarativeSyntaxSupport.h"
+#include "Widgets/Text/STextBlock.h"
+#include "Widgets/Input/SButton.h"
+#include "PropertyHandle.h"
+#include "DetailLayoutBuilder.h"
+#include "DetailWidgetRow.h"
+#include "DetailCategoryBuilder.h"
+
+#define LOCTEXT_NAMESPACE "PCGGraphDetails"
+
+TSharedRef<IDetailCustomization> FPCGGraphDetails::MakeInstance()
+{
+	return MakeShareable(new FPCGGraphDetails());
+}
+
+void FPCGGraphDetails::CustomizeDetails(IDetailLayoutBuilder& DetailBuilder)
+{
+	const FName PCGCategoryName("PCG");
+	IDetailCategoryBuilder& PCGCategory = DetailBuilder.EditCategory(PCGCategoryName);
+
+	TArray<TWeakObjectPtr<UObject>> ObjectsBeingCustomized;
+	DetailBuilder.GetObjectsBeingCustomized(ObjectsBeingCustomized);
+
+	for (TWeakObjectPtr<UObject>& Object : ObjectsBeingCustomized)
+	{
+		UPCGGraph* Graph = Cast<UPCGGraph>(Object.Get());
+		if (ensure(Graph))
+		{
+			SelectedGraphs.Add(Graph);
+		}
+	}
+
+	TArray<TSharedRef<IPropertyHandle>> AllProperties;
+	bool bSimpleProperties = true;
+	bool bAdvancedProperties = false;
+	// Add all properties in the category in order
+	PCGCategory.GetDefaultProperties(AllProperties, bSimpleProperties, bAdvancedProperties);
+
+	for (auto& Property : AllProperties)
+	{
+		PCGCategory.AddProperty(Property);
+	}
+}
+
+#undef LOCTEXT_NAMESPACE
\ No newline at end of file
diff --git a/Engine/Plugins/Experimental/PCG/Source/PCGEditor/Private/Details/PCGVolumeDetails.cpp b/Engine/Plugins/Experimental/PCG/Source/PCGEditor/Private/Details/PCGVolumeDetails.cpp
new file mode 100644
--- /dev/null
+++ b/Engine/Plugins/Experimental/PCG/Source/PCGEditor/Private/Details/PCGVolumeDetails.cpp
@@ -0,0 +1,30 @@
+// Copyright Epic Games, Inc. All Rights Reserved.
+
+#include "Details/PCGVolumeDetails.h"
+
+#include "PCGComponent.h"
+#include "PCGVolume.h"
+
+TSharedRef<IDetailCustomization> FPCGVolumeDetails::MakeInstance()
+{
+	return MakeShareable(new FPCGVolumeDetails());
+}
+
+void FPCGVolumeDetails::GatherPCGComponentsFromSelection(const TArray<TWeakObjectPtr<UObject>>& InObjectSelected)
+{
+	for (const TWeakObjectPtr<UObject>& Object : InObjectSelected)
+	{
+		if (APCGVolume* Volume = Cast<APCGVolume>(Object))
+		{
+			TInlineComponentArray<UPCGComponent*, 1> Components;
+			Volume->GetComponents<UPCGComponent>(Components);
+			for (UPCGComponent* Component : Components)
+			{
+				if (Component)
+				{
+					SelectedComponents.Add(Component);
+				}
+			}
+		}
+	}
+}
\ No newline at end of file
diff --git a/Engine/Plugins/Experimental/PCG/Source/PCGEditor/Private/PCGBlueprintSettingsDetails.cpp b/Engine/Plugins/Experimental/PCG/Source/PCGEditor/Private/PCGBlueprintSettingsDetails.cpp
deleted file mode 100644
diff --git a/Engine/Plugins/Experimental/PCG/Source/PCGEditor/Private/PCGComponentDetails.cpp b/Engine/Plugins/Experimental/PCG/Source/PCGEditor/Private/PCGComponentDetails.cpp
deleted file mode 100644
diff --git a/Engine/Plugins/Experimental/PCG/Source/PCGEditor/Private/PCGEditorModule.cpp b/Engine/Plugins/Experimental/PCG/Source/PCGEditor/Private/PCGEditorModule.cpp
--- a/Engine/Plugins/Experimental/PCG/Source/PCGEditor/Private/PCGEditorModule.cpp
+++ b/Engine/Plugins/Experimental/PCG/Source/PCGEditor/Private/PCGEditorModule.cpp
@@ -4,19 +4,26 @@
 
 #include "AssetTypeActions/PCGGraphAssetTypeActions.h"
 #include "AssetTypeActions/PCGSettingsAssetTypeActions.h"
-#include "PCGComponentDetails.h"
+#include "PCGComponent.h"
 #include "PCGEditorCommands.h"
 #include "PCGEditorGraphNodeFactory.h"
 #include "PCGEditorSettings.h"
 #include "PCGEditorStyle.h"
 #include "PCGEditorUtils.h"
-#include "PCGGraphDetails.h"
-#include "PCGBlueprintSettingsDetails.h"
+#include "PCGGraph.h"
 #include "PCGSubsystem.h"
-#include "PCGVolumeDetails.h"
+#include "PCGVolume.h"
 #include "PCGVolumeFactory.h"
 #include "PCGWorldActor.h"
 
+#include "Details/PCGBlueprintSettingsDetails.h"
+#include "Details/PCGComponentDetails.h"
+#include "Details/PCGGraphDetails.h"
+#include "Details/PCGVolumeDetails.h"
+#include "Details/PCGAttributePropertySelectorDetails.h"
+#include "Elements/PCGExecuteBlueprint.h"
+#include "Metadata/PCGAttributePropertySelector.h"
+
 #include "EdGraphUtilities.h"
 #include "Framework/MultiBox/MultiBoxBuilder.h"
 #include "IAssetTools.h"
@@ -75,10 +82,12 @@ bool FPCGEditorModule::SupportsDynamicReloading()
 void FPCGEditorModule::RegisterDetailsCustomizations()
 {
 	FPropertyEditorModule& PropertyEditor = FModuleManager::LoadModuleChecked<FPropertyEditorModule>("PropertyEditor");
-	PropertyEditor.RegisterCustomClassLayout("PCGComponent", FOnGetDetailCustomizationInstance::CreateStatic(&FPCGComponentDetails::MakeInstance));
-	PropertyEditor.RegisterCustomClassLayout("PCGGraph", FOnGetDetailCustomizationInstance::CreateStatic(&FPCGGraphDetails::MakeInstance));
-	PropertyEditor.RegisterCustomClassLayout("PCGBlueprintSettings", FOnGetDetailCustomizationInstance::CreateStatic(&FPCGBlueprintSettingsDetails::MakeInstance));
-	PropertyEditor.RegisterCustomClassLayout("PCGVolume", FOnGetDetailCustomizationInstance::CreateStatic(&FPCGVolumeDetails::MakeInstance));
+	PropertyEditor.RegisterCustomClassLayout(UPCGBlueprintSettings::StaticClass()->GetFName(), FOnGetDetailCustomizationInstance::CreateStatic(&FPCGBlueprintSettingsDetails::MakeInstance));
+	PropertyEditor.RegisterCustomClassLayout(UPCGComponent::StaticClass()->GetFName(), FOnGetDetailCustomizationInstance::CreateStatic(&FPCGComponentDetails::MakeInstance));
+	PropertyEditor.RegisterCustomClassLayout(UPCGGraph::StaticClass()->GetFName(), FOnGetDetailCustomizationInstance::CreateStatic(&FPCGGraphDetails::MakeInstance));
+	PropertyEditor.RegisterCustomClassLayout(APCGVolume::StaticClass()->GetFName(), FOnGetDetailCustomizationInstance::CreateStatic(&FPCGVolumeDetails::MakeInstance));
+
+	PropertyEditor.RegisterCustomPropertyTypeLayout(FPCGAttributePropertySelector::StaticStruct()->GetFName(), FOnGetPropertyTypeCustomizationInstance::CreateStatic(&FPCGAttributePropertySelectorDetails::MakeInstance));
 }
 
 void FPCGEditorModule::UnregisterDetailsCustomizations()
@@ -86,10 +95,12 @@ void FPCGEditorModule::UnregisterDetailsCustomizations()
 	if (FModuleManager::Get().IsModuleLoaded("PropertyEditor"))
 	{
 		FPropertyEditorModule& PropertyModule = FModuleManager::LoadModuleChecked<FPropertyEditorModule>("PropertyEditor");
-		PropertyModule.UnregisterCustomClassLayout("PCGComponent");
-		PropertyModule.UnregisterCustomClassLayout("PCGGraph");
-		PropertyModule.UnregisterCustomClassLayout("PCGBlueprintSettings");
-		PropertyModule.UnregisterCustomClassLayout("PCGVolume");
+		PropertyModule.UnregisterCustomClassLayout(UPCGBlueprintSettings::StaticClass()->GetFName());
+		PropertyModule.UnregisterCustomClassLayout(UPCGComponent::StaticClass()->GetFName());
+		PropertyModule.UnregisterCustomClassLayout(UPCGGraph::StaticClass()->GetFName());
+		PropertyModule.UnregisterCustomClassLayout(APCGVolume::StaticClass()->GetFName());
+
+		PropertyModule.UnregisterCustomPropertyTypeLayout(FPCGAttributePropertySelector::StaticStruct()->GetFName());
 		PropertyModule.NotifyCustomizationModuleChanged();
 	}
 }
diff --git a/Engine/Plugins/Experimental/PCG/Source/PCGEditor/Private/PCGGraphDetails.cpp b/Engine/Plugins/Experimental/PCG/Source/PCGEditor/Private/PCGGraphDetails.cpp
deleted file mode 100644
diff --git a/Engine/Plugins/Experimental/PCG/Source/PCGEditor/Private/PCGVolumeDetails.cpp b/Engine/Plugins/Experimental/PCG/Source/PCGEditor/Private/PCGVolumeDetails.cpp
deleted file mode 100644
diff --git a/Engine/Plugins/Experimental/PCG/Source/PCGEditor/Public/Details/PCGAttributePropertySelectorDetails.h b/Engine/Plugins/Experimental/PCG/Source/PCGEditor/Public/Details/PCGAttributePropertySelectorDetails.h
new file mode 100644
--- /dev/null
+++ b/Engine/Plugins/Experimental/PCG/Source/PCGEditor/Public/Details/PCGAttributePropertySelectorDetails.h
@@ -0,0 +1,37 @@
+// Copyright Epic Games, Inc. All Rights Reserved.
+
+#pragma once
+
+#include "IPropertyTypeCustomization.h"
+
+#include "PCGPoint.h"
+
+struct FPCGAttributePropertySelector;
+class SWidget;
+
+class FPCGAttributePropertySelectorDetails : public IPropertyTypeCustomization
+{
+public:
+	static TSharedRef<IPropertyTypeCustomization> MakeInstance()
+	{
+		return MakeShareable(new FPCGAttributePropertySelectorDetails);
+	}
+
+	/** ~Begin IPropertyTypeCustomization interface */
+	virtual void CustomizeHeader(TSharedRef<IPropertyHandle> InPropertyHandle, FDetailWidgetRow& HeaderRow, IPropertyTypeCustomizationUtils& CustomizationUtils) override;
+	virtual void CustomizeChildren(TSharedRef<IPropertyHandle> InPropertyHandle, IDetailChildrenBuilder& ChildBuilder, IPropertyTypeCustomizationUtils& CustomizationUtils) override {};
+	/** ~End IPropertyTypeCustomization interface */
+
+protected:
+	FPCGAttributePropertySelector* GetStruct();
+	const FPCGAttributePropertySelector* GetStruct() const;
+
+	TSharedRef<SWidget> GenerateExtraMenu();
+
+	FText GetText() const;
+	void SetText(const FText& NewText, ETextCommit::Type CommitInfo);
+	void SetPointProperty(EPCGPointProperties EnumValue);
+	void SetAttributeName(FName NewName);
+
+	TSharedPtr<IPropertyHandle> PropertyHandle;
+};
\ No newline at end of file
diff --git a/Engine/Plugins/Experimental/PCG/Source/PCGEditor/Public/Details/PCGBlueprintSettingsDetails.h b/Engine/Plugins/Experimental/PCG/Source/PCGEditor/Public/Details/PCGBlueprintSettingsDetails.h
new file mode 100644
--- /dev/null
+++ b/Engine/Plugins/Experimental/PCG/Source/PCGEditor/Public/Details/PCGBlueprintSettingsDetails.h
@@ -0,0 +1,37 @@
+// Copyright Epic Games, Inc. All Rights Reserved.
+
+#pragma once
+
+#include "IDetailCustomization.h"
+#include "IPropertyTypeCustomization.h"
+#include "Input/Reply.h"
+#include "Templates/SharedPointer.h"
+#include "UObject/WeakObjectPtr.h"
+#include "UObject/WeakObjectPtrTemplates.h"
+
+#include "PCGSettings.h"
+#include "Elements/PCGExecuteBlueprint.h"
+
+class IDetailLayoutBuilder;
+class UFunction;
+class UPCGBlueprintElement;
+
+class FPCGBlueprintSettingsDetails : public IDetailCustomization
+{
+public:
+	static TSharedRef<IDetailCustomization> MakeInstance();
+
+	/** ~Begin IDetailCustomization interface */
+	virtual void CustomizeDetails(IDetailLayoutBuilder& DetailBuilder) override;
+	/** ~End IDetailCustomization interface */
+
+private:
+	// Creates a button strip in each category that contains reflected functions marked as CallInEditor
+	void AddCallInEditorMethods(IDetailLayoutBuilder& DetailBuilder);
+
+	// Executes the specified method on the selected objects
+	FReply OnExecuteCallInEditorFunction(TWeakObjectPtr<UFunction> WeakFunctionPtr);
+
+private:
+	TWeakObjectPtr<UPCGBlueprintElement> SelectedObject;
+};
\ No newline at end of file
diff --git a/Engine/Plugins/Experimental/PCG/Source/PCGEditor/Public/Details/PCGComponentDetails.h b/Engine/Plugins/Experimental/PCG/Source/PCGEditor/Public/Details/PCGComponentDetails.h
new file mode 100644
--- /dev/null
+++ b/Engine/Plugins/Experimental/PCG/Source/PCGEditor/Public/Details/PCGComponentDetails.h
@@ -0,0 +1,37 @@
+// Copyright Epic Games, Inc. All Rights Reserved.
+
+#pragma once
+
+#include "UObject/WeakObjectPtr.h"
+#include "Input/Reply.h"
+#include "Layout/Visibility.h"
+#include "IDetailCustomization.h"
+
+class UPCGComponent;
+
+class FPCGComponentDetails : public IDetailCustomization
+{
+public:
+	static TSharedRef<IDetailCustomization> MakeInstance();
+
+	/** ~Begin IDetailCustomization interface */
+	virtual void CustomizeDetails(IDetailLayoutBuilder& DetailBuilder) override;
+	virtual void PendingDelete() override;
+	/** ~End IDetailCustomization interface */
+
+protected:
+	virtual void GatherPCGComponentsFromSelection(const TArray<TWeakObjectPtr<UObject>>& InObjectSelected);
+	virtual bool AddDefaultProperties() const { return true; }
+
+private:
+	EVisibility GenerateButtonVisible() const;
+	EVisibility CancelButtonVisible() const;
+	FReply OnGenerateClicked();
+	FReply OnCancelClicked();
+	FReply OnCleanupClicked();
+	FReply OnClearPCGLinkClicked();
+	void OnGraphChanged(UPCGComponent* InComponent);
+
+protected:
+	TArray<TWeakObjectPtr<UPCGComponent>> SelectedComponents;
+};
\ No newline at end of file
diff --git a/Engine/Plugins/Experimental/PCG/Source/PCGEditor/Public/Details/PCGGraphDetails.h b/Engine/Plugins/Experimental/PCG/Source/PCGEditor/Public/Details/PCGGraphDetails.h
new file mode 100644
--- /dev/null
+++ b/Engine/Plugins/Experimental/PCG/Source/PCGEditor/Public/Details/PCGGraphDetails.h
@@ -0,0 +1,20 @@
+// Copyright Epic Games, Inc. All Rights Reserved.
+
+#pragma once
+
+#include "UObject/WeakObjectPtr.h"
+#include "Input/Reply.h"
+#include "IDetailCustomization.h"
+
+class UPCGGraph;
+
+class FPCGGraphDetails : public IDetailCustomization
+{
+public:
+	static TSharedRef<IDetailCustomization> MakeInstance();
+
+	virtual void CustomizeDetails(IDetailLayoutBuilder& DetailBuilder) override;
+
+private:
+	TArray<TWeakObjectPtr<UPCGGraph>> SelectedGraphs;
+};
\ No newline at end of file
diff --git a/Engine/Plugins/Experimental/PCG/Source/PCGEditor/Public/Details/PCGVolumeDetails.h b/Engine/Plugins/Experimental/PCG/Source/PCGEditor/Public/Details/PCGVolumeDetails.h
new file mode 100644
--- /dev/null
+++ b/Engine/Plugins/Experimental/PCG/Source/PCGEditor/Public/Details/PCGVolumeDetails.h
@@ -0,0 +1,17 @@
+// Copyright Epic Games, Inc. All Rights Reserved.
+
+#pragma once
+
+#include "Details/PCGComponentDetails.h"
+
+class FPCGVolumeDetails : public FPCGComponentDetails
+{
+public:
+	static TSharedRef<IDetailCustomization> MakeInstance();
+
+protected:
+	/** ~Begin FPCGComponentDetails interface */
+	virtual void GatherPCGComponentsFromSelection(const TArray<TWeakObjectPtr<UObject>>& InObjectSelected) override;
+	virtual bool AddDefaultProperties() const override { return false; }
+	/** ~End FPCGComponentDetails interface */
+};
\ No newline at end of file
diff --git a/Engine/Plugins/Experimental/PCG/Source/PCGEditor/Public/PCGBlueprintSettingsDetails.h b/Engine/Plugins/Experimental/PCG/Source/PCGEditor/Public/PCGBlueprintSettingsDetails.h
deleted file mode 100644
diff --git a/Engine/Plugins/Experimental/PCG/Source/PCGEditor/Public/PCGComponentDetails.h b/Engine/Plugins/Experimental/PCG/Source/PCGEditor/Public/PCGComponentDetails.h
deleted file mode 100644
diff --git a/Engine/Plugins/Experimental/PCG/Source/PCGEditor/Public/PCGGraphDetails.h b/Engine/Plugins/Experimental/PCG/Source/PCGEditor/Public/PCGGraphDetails.h
deleted file mode 100644
diff --git a/Engine/Plugins/Experimental/PCG/Source/PCGEditor/Public/PCGVolumeDetails.h b/Engine/Plugins/Experimental/PCG/Source/PCGEditor/Public/PCGVolumeDetails.h
deleted file mode 100644
