[PCG] Refresh metadata nodes

* Make sure that all nodes supports "GetLatestAttributeOrNone"
* Replace all fields where it makes sense to select attribute and properties.
* Cleanup nodes that implemented specific code for attribute<->property translation
* Try to normalize how metadata nodes are implemented.

diff --git a/Engine/Plugins/Experimental/PCG/Source/PCG/Private/Elements/Metadata/PCGMetadataElement.cpp b/Engine/Plugins/Experimental/PCG/Source/PCG/Private/Elements/Metadata/PCGMetadataElement.cpp
--- a/Engine/Plugins/Experimental/PCG/Source/PCG/Private/Elements/Metadata/PCGMetadataElement.cpp
+++ b/Engine/Plugins/Experimental/PCG/Source/PCG/Private/Elements/Metadata/PCGMetadataElement.cpp
@@ -1,97 +1,70 @@
 // Copyright Epic Games, Inc. All Rights Reserved.
 
 #include "Elements/Metadata/PCGMetadataElement.h"
+
+#include "Elements/Metadata/PCGMetadataElementCommon.h"
 #include "Data/PCGPointData.h"
 #include "Helpers/PCGSettingsHelpers.h"
 #include "Metadata/PCGMetadata.h"
 #include "Metadata/PCGMetadataAttribute.h"
 #include "Metadata/PCGMetadataAttributeTpl.h"
+#include "Metadata/Accessors/IPCGAttributeAccessor.h"
+#include "Metadata/Accessors/PCGAttributeAccessorHelpers.h"
+#include "Metadata/Accessors/PCGAttributeAccessorKeys.h"
 
 #include UE_INLINE_GENERATED_CPP_BY_NAME(PCGMetadataElement)
 
-namespace PCGMetadataOperations
+UPCGMetadataOperationSettings::UPCGMetadataOperationSettings()
 {
-	// Getter and setter for properties.
-	// Need to specify the property type and the getter/setter.
-	// Will work if AttributeType and PropertyType are the same, or if we can construct one from the other.
-	template <typename T>
-	bool SetValue(TArray<FPCGPoint>& InPoints, FPCGMetadataAttributeBase* AttributeBase, UPCGMetadata* Metadata, TFunctionRef<T(const FPCGPoint& InPoint)> PropGetter)
-	{
-		if (!AttributeBase)
-		{
-			return false;
-		}
-
-		auto Func = [&InPoints, AttributeBase, Metadata, &PropGetter](auto DummyValue) -> bool
-		{
-			using AttributeType = decltype(DummyValue);
-
-			FPCGMetadataAttribute<AttributeType>* Attribute = static_cast<FPCGMetadataAttribute<AttributeType>*>(AttributeBase);
-
-			for (FPCGPoint& Point : InPoints)
-			{
-				if constexpr (std::is_same_v<AttributeType, T>)
-				{
-					Metadata->InitializeOnSet(Point.MetadataEntry);
-					Attribute->SetValue(Point.MetadataEntry, PropGetter(Point));
-				}
-				else if constexpr (std::is_constructible_v<AttributeType, T>)
-				{
-					Metadata->InitializeOnSet(Point.MetadataEntry);
-					Attribute->SetValue(Point.MetadataEntry, AttributeType(PropGetter(Point)));
-				}
-				else
-				{
-					return false;
-				}
-			}
-
-			return true;
-		};
+	// Previous default object was: None for source and target attribute, density for point property
+	// and Property to attribute
+	// Recreate the same default
+	InputSource.SetPointProperty(EPCGPointProperties::Density);
+	OutputTarget.SetAttributeName(NAME_None);
+}
 
-		return PCGMetadataAttribute::CallbackWithRightType(AttributeBase->GetTypeId(), Func);
-	}
+FPCGElementPtr UPCGMetadataOperationSettings::CreateElement() const
+{
+	return MakeShared<FPCGMetadataOperationElement>();
+}
 
-	template <typename T>
-	bool SetValue(const FPCGMetadataAttributeBase* AttributeBase, TArray<FPCGPoint>& InPoints, TFunctionRef<void(FPCGPoint& OutPoint, const T& InValue)> PropSetter)
+void UPCGMetadataOperationSettings::PostLoad()
+{
+	Super::PostLoad();
+
+#if WITH_EDITOR
+PRAGMA_DISABLE_DEPRECATION_WARNINGS
+	if (Target_DEPRECATED != EPCGMetadataOperationTarget::PropertyToAttribute
+		|| (SourceAttribute_DEPRECATED != NAME_None)
+		|| (DestinationAttribute_DEPRECATED != NAME_None)
+		|| (PointProperty_DEPRECATED != EPCGPointProperties::Density))
 	{
-		if (!AttributeBase)
+		switch (Target_DEPRECATED)
 		{
-			return false;
+		case EPCGMetadataOperationTarget::PropertyToAttribute:
+			InputSource.SetPointProperty(PointProperty_DEPRECATED);
+			OutputTarget.SetAttributeName(DestinationAttribute_DEPRECATED);
+			break;
+		case EPCGMetadataOperationTarget::AttributeToProperty:
+			InputSource.SetAttributeName(SourceAttribute_DEPRECATED);
+			OutputTarget.SetPointProperty(PointProperty_DEPRECATED);
+			break;
+		case EPCGMetadataOperationTarget::AttributeToAttribute:
+			InputSource.SetAttributeName(SourceAttribute_DEPRECATED);
+			OutputTarget.SetAttributeName(DestinationAttribute_DEPRECATED);
+			break;
+		default:
+			break;
 		}
 
-		auto Func = [&InPoints, AttributeBase, &PropSetter](auto DummyValue) -> bool
-		{
-			using AttributeType = decltype(DummyValue);
-
-			const FPCGMetadataAttribute<AttributeType>* Attribute = static_cast<const FPCGMetadataAttribute<AttributeType>*>(AttributeBase);
-
-			for (FPCGPoint& Point : InPoints)
-			{
-				if constexpr (std::is_same_v<T, AttributeType>)
-				{
-					PropSetter(Point, Attribute->GetValueFromItemKey(Point.MetadataEntry));
-				}
-				else if constexpr (std::is_constructible_v<T, AttributeType>)
-				{
-					PropSetter(Point, T(Attribute->GetValueFromItemKey(Point.MetadataEntry)));
-				}
-				else
-				{
-					return false;
-				}
-			}
-
-			return true;
-		};
-
-		return PCGMetadataAttribute::CallbackWithRightType(AttributeBase->GetTypeId(), Func);
+		// Default values.
+		SourceAttribute_DEPRECATED = NAME_None;
+		DestinationAttribute_DEPRECATED = NAME_None;
+		Target_DEPRECATED = EPCGMetadataOperationTarget::PropertyToAttribute;
+		PointProperty_DEPRECATED = EPCGPointProperties::Density;
 	}
-}
-
-FPCGElementPtr UPCGMetadataOperationSettings::CreateElement() const
-{
-	return MakeShared<FPCGMetadataOperationElement>();
+PRAGMA_ENABLE_DEPRECATION_WARNINGS
+#endif // WITH_EDITOR
 }
 
 bool FPCGMetadataOperationElement::ExecuteInternal(FPCGContext* Context) const
@@ -106,11 +79,6 @@ bool FPCGMetadataOperationElement::ExecuteInternal(FPCGContext* Context) const
 
 	TArray<FPCGTaggedData>& Outputs = Context->OutputData.TaggedData;
 
-	const FName SourceAttribute = PCGSettingsHelpers::GetValue(GET_MEMBER_NAME_CHECKED(UPCGMetadataOperationSettings, SourceAttribute), Settings->SourceAttribute, Params);
-	const EPCGPointProperties PointProperty = PCGSettingsHelpers::GetValue(GET_MEMBER_NAME_CHECKED(UPCGMetadataOperationSettings, PointProperty), Settings->PointProperty, Params);
-	const FName DestinationAttribute = PCGSettingsHelpers::GetValue(GET_MEMBER_NAME_CHECKED(UPCGMetadataOperationSettings, DestinationAttribute), Settings->DestinationAttribute, Params);
-	const EPCGMetadataOperationTarget Target = PCGSettingsHelpers::GetValue(GET_MEMBER_NAME_CHECKED(UPCGMetadataOperationSettings, Target), Settings->Target, Params);
-
 	// Forward any non-input data
 	Outputs.Append(Context->InputData.GetAllSettings());
 
@@ -140,10 +108,16 @@ bool FPCGMetadataOperationElement::ExecuteInternal(FPCGContext* Context) const
 			continue;
 		}
 
+		FPCGAttributePropertySelector InputSource = Settings->InputSource;
+
+		const FName SourceAttribute = InputSource.GetName();
 		const FName LocalSourceAttribute = ((SourceAttribute != NAME_None) ? SourceAttribute : OriginalData->Metadata->GetLatestAttributeNameOrNone());
+		const FName DestinationAttribute = Settings->OutputTarget.GetName();
+
+		// Make sure we use the right attribute name
+		InputSource.AttributeName = LocalSourceAttribute;
 
-		// Check if the attribute exists
-		if ((Target == EPCGMetadataOperationTarget::AttributeToProperty || Target == EPCGMetadataOperationTarget::AttributeToAttribute) && !OriginalData->Metadata->HasAttribute(LocalSourceAttribute))
+		if (Settings->InputSource.Selection == EPCGAttributePropertySelection::Attribute && !OriginalData->Metadata->HasAttribute(LocalSourceAttribute))
 		{
 			PCGE_LOG(Warning, "Input does not have the %s attribute", *LocalSourceAttribute.ToString());
 			continue;
@@ -161,265 +135,99 @@ bool FPCGMetadataOperationElement::ExecuteInternal(FPCGContext* Context) const
 		// Copy points and then apply the operation
 		SampledPoints = Points;
 
-		if (Target == EPCGMetadataOperationTarget::PropertyToAttribute)
+		// If it is attribute to attribute, just copy the attributes, if they exist and are valid
+		if (Settings->InputSource.Selection == EPCGAttributePropertySelection::Attribute && Settings->OutputTarget.Selection == EPCGAttributePropertySelection::Attribute)
 		{
-			if (PointProperty == EPCGPointProperties::Density)
-			{
-				auto DensityGetter = [](const FPCGPoint& InPoint) { return InPoint.Density; };
-
-				if (!SampledData->Metadata->HasAttribute(DestinationAttribute))
-				{
-					SampledData->Metadata->CreateFloatAttribute(DestinationAttribute, 0.0f, /*bAllowsInterpolation=*/true, /*bOverrideParent=*/true);
-				}
-
-				FPCGMetadataAttributeBase* AttributeBase = SampledData->Metadata->GetMutableAttribute(DestinationAttribute);
-				if(!PCGMetadataOperations::SetValue<float>(SampledPoints, AttributeBase, SampledData->Metadata, DensityGetter))
-				{
-					PCGE_LOG(Error, "Attribute %s already exists but its type is not compatible", *DestinationAttribute.ToString());
-				}
-			}
-			else if (PointProperty == EPCGPointProperties::BoundsMin)
-			{
-				auto ExtentsGetter = [](const FPCGPoint& InPoint) { return InPoint.BoundsMin; };
-
-				if (!SampledData->Metadata->HasAttribute(DestinationAttribute))
-				{
-					SampledData->Metadata->CreateVectorAttribute(DestinationAttribute, FVector::Zero(), /*bAllowsInterpolation=*/true, /*bOverrideParent=*/true);
-				}
-
-				FPCGMetadataAttributeBase* AttributeBase = SampledData->Metadata->GetMutableAttribute(DestinationAttribute);
-				if (!PCGMetadataOperations::SetValue<FVector>(SampledPoints, AttributeBase, SampledData->Metadata, ExtentsGetter))
-				{
-					PCGE_LOG(Error, "Attribute %s already exists but its type is not compatible", *DestinationAttribute.ToString());
-				}
-			}
-			else if (PointProperty == EPCGPointProperties::BoundsMax)
-			{
-				auto ExtentsGetter = [](const FPCGPoint& InPoint) { return InPoint.BoundsMax; };
-
-				if (!SampledData->Metadata->HasAttribute(DestinationAttribute))
-				{
-					SampledData->Metadata->CreateVectorAttribute(DestinationAttribute, FVector::Zero(), /*bAllowsInterpolation=*/true, /*bOverrideParent=*/true);
-				}
-
-				FPCGMetadataAttributeBase* AttributeBase = SampledData->Metadata->GetMutableAttribute(DestinationAttribute);
-				if (!PCGMetadataOperations::SetValue<FVector>(SampledPoints, AttributeBase, SampledData->Metadata, ExtentsGetter))
-				{
-					PCGE_LOG(Error, "Attribute %s already exists but its type is not compatible", *DestinationAttribute.ToString());
-				}
-			}
-			else if (PointProperty == EPCGPointProperties::Extents)
-			{
-				auto ExtentsGetter = [](const FPCGPoint& InPoint) { return InPoint.GetExtents(); };
-
-				if (!SampledData->Metadata->HasAttribute(DestinationAttribute))
-				{
-					SampledData->Metadata->CreateVectorAttribute(DestinationAttribute, FVector::Zero() , /*bAllowsInterpolation=*/true, /*bOverrideParent=*/true);
-				}
-
-				FPCGMetadataAttributeBase* AttributeBase = SampledData->Metadata->GetMutableAttribute(DestinationAttribute);
-				if(!PCGMetadataOperations::SetValue<FVector>(SampledPoints, AttributeBase, SampledData->Metadata, ExtentsGetter))
-				{
-					PCGE_LOG(Error, "Attribute %s already exists but its type is not compatible", *DestinationAttribute.ToString());
-				}
-			}
-			else if(PointProperty == EPCGPointProperties::Color)
-			{
-				auto ColorGetter = [](const FPCGPoint& InPoint) { return InPoint.Color; };
-
-				if (!SampledData->Metadata->HasAttribute(DestinationAttribute))
-				{
-					SampledData->Metadata->CreateVector4Attribute(DestinationAttribute, FVector4::Zero(), /*bAllowsInterpolation=*/true, /*bOverrideParent=*/true);
-				}
-
-				FPCGMetadataAttributeBase* AttributeBase = SampledData->Metadata->GetMutableAttribute(DestinationAttribute);
-				if(!PCGMetadataOperations::SetValue<FVector4>(SampledPoints, AttributeBase, SampledData->Metadata, ColorGetter))
-				{
-					PCGE_LOG(Error, "Attribute %s already exists but its type is not compatible", *DestinationAttribute.ToString());
-				}
-			}
-			else if (PointProperty == EPCGPointProperties::Position)
-			{
-				auto PositionGetter = [](const FPCGPoint& InPoint) { return InPoint.Transform.GetLocation(); };
-
-				if (!SampledData->Metadata->HasAttribute(DestinationAttribute))
-				{
-					SampledData->Metadata->CreateVectorAttribute(DestinationAttribute, FVector::Zero(), /*bAllowsInterpolation=*/true, /*bOverrideParent=*/true);
-				}
-
-				FPCGMetadataAttributeBase* AttributeBase = SampledData->Metadata->GetMutableAttribute(DestinationAttribute);
-				if (!PCGMetadataOperations::SetValue<FVector>(SampledPoints, AttributeBase, SampledData->Metadata, PositionGetter))
-				{
-					PCGE_LOG(Error, "Attribute %s already exists but its type is not compatible", *DestinationAttribute.ToString());
-				}
-			}
-			else if (PointProperty == EPCGPointProperties::Rotation)
+			if (LocalSourceAttribute == DestinationAttribute)
 			{
-				auto RotationGetter = [](const FPCGPoint& InPoint) { return InPoint.Transform.GetRotation(); };
-
-				if (!SampledData->Metadata->HasAttribute(DestinationAttribute))
-				{
-					SampledData->Metadata->CreateQuatAttribute(DestinationAttribute, FQuat::Identity, /*bAllowsInterpolation=*/true, /*bOverrideParent=*/true);
-				}
-
-				FPCGMetadataAttributeBase* AttributeBase = SampledData->Metadata->GetMutableAttribute(DestinationAttribute);
-				if (!PCGMetadataOperations::SetValue<FQuat>(SampledPoints, AttributeBase, SampledData->Metadata, RotationGetter))
-				{
-					PCGE_LOG(Error, "Attribute %s already exists but its type is not compatible", *DestinationAttribute.ToString());
-				}
+				// Nothing to do...
+				continue;
 			}
-			else if (PointProperty == EPCGPointProperties::Scale)
+
+			if (!SampledData->Metadata->CopyExistingAttribute(LocalSourceAttribute, DestinationAttribute))
 			{
-				auto ScaleGetter = [](const FPCGPoint& InPoint) { return InPoint.Transform.GetScale3D(); };
-
-				if (!SampledData->Metadata->HasAttribute(DestinationAttribute))
-				{
-					SampledData->Metadata->CreateVectorAttribute(DestinationAttribute, FVector::One(), /*bAllowsInterpolation=*/true, /*bOverrideParent=*/true);
-				}
-
-				FPCGMetadataAttributeBase* AttributeBase = SampledData->Metadata->GetMutableAttribute(DestinationAttribute);
-				if (!PCGMetadataOperations::SetValue<FVector>(SampledPoints, AttributeBase, SampledData->Metadata, ScaleGetter))
-				{
-					PCGE_LOG(Error, "Attribute %s already exists but its type is not compatible", *DestinationAttribute.ToString());
-				}
+				PCGE_LOG(Warning, "Failed to copy to new attribute %s", *DestinationAttribute.ToString());
 			}
-			else if (PointProperty == EPCGPointProperties::Transform)
+
+			continue;
+		}
+
+		TUniquePtr<const IPCGAttributeAccessor> InputAccessor = PCGAttributeAccessorHelpers::CreateConstAccessor(OriginalData, InputSource);
+		TUniquePtr<const IPCGAttributeAccessorKeys> InputKeys = PCGAttributeAccessorHelpers::CreateConstKeys(OriginalData, InputSource);
+
+		if (!InputAccessor.IsValid() || !InputKeys.IsValid())
+		{
+			PCGE_LOG(Warning, "Failed to create input accessor or iterator");
+			continue;
+		}
+
+		if (Settings->OutputTarget.Selection == EPCGAttributePropertySelection::Attribute)
+		{
+			auto CreateAttribute = [&](auto Dummy)
 			{
-				auto TransformGetter = [](const FPCGPoint& InPoint) { return InPoint.Transform; };
-
-				if (!SampledData->Metadata->HasAttribute(DestinationAttribute))
-				{
-					SampledData->Metadata->CreateTransformAttribute(DestinationAttribute, FTransform::Identity, /*bAllowsInterpolation=*/true, /*bOverrideParent=*/true);
-				}
-
-				FPCGMetadataAttributeBase* AttributeBase = SampledData->Metadata->GetMutableAttribute(DestinationAttribute);
-				if (!PCGMetadataOperations::SetValue<FTransform>(SampledPoints, AttributeBase, SampledData->Metadata, TransformGetter))
-				{
-					PCGE_LOG(Error, "Attribute %s already exists but its type is not compatible", *DestinationAttribute.ToString());
-				}
-			}
-			else if (PointProperty == EPCGPointProperties::Steepness)
+				using AttributeType = decltype(Dummy);
+				return PCGMetadataElementCommon::ClearOrCreateAttribute(SampledData->Metadata, DestinationAttribute, AttributeType{}) != nullptr;
+			};
+			
+			if (!PCGMetadataAttribute::CallbackWithRightType(InputAccessor->GetUnderlyingType(), CreateAttribute))
 			{
-				auto SteepnessGetter = [](const FPCGPoint& InPoint) { return InPoint.Steepness; };
-
-				if (!SampledData->Metadata->HasAttribute(DestinationAttribute))
-				{
-					SampledData->Metadata->CreateFloatAttribute(DestinationAttribute, 0.5f, /*bAllowsInterpolation=*/true, /*bOverrideParent=*/true);
-				}
-
-				FPCGMetadataAttributeBase* AttributeBase = SampledData->Metadata->GetMutableAttribute(DestinationAttribute);
-				if (!PCGMetadataOperations::SetValue<float>(SampledPoints, AttributeBase, SampledData->Metadata, SteepnessGetter))
-				{
-					PCGE_LOG(Error, "Attribute %s already exists but its type is not compatible", *DestinationAttribute.ToString());
-				}
+				PCGE_LOG(Warning, "Failed to create new attribute %s", *DestinationAttribute.ToString());
+				continue;
 			}
 		}
-		else if(Target == EPCGMetadataOperationTarget::AttributeToProperty) // Attribute to property
+
+		TUniquePtr<IPCGAttributeAccessor> OutputAccessor = PCGAttributeAccessorHelpers::CreateAccessor(SampledData, Settings->OutputTarget);
+		TUniquePtr<IPCGAttributeAccessorKeys> OutputKeys = PCGAttributeAccessorHelpers::CreateKeys(SampledData, Settings->OutputTarget);
+
+		if (!OutputAccessor.IsValid() || !OutputKeys.IsValid())
 		{
-			if (PointProperty == EPCGPointProperties::Density)
-			{
-				auto DensitySetter = [](FPCGPoint& InPoint, const float& InValue) { InPoint.Density = InValue; };
+			PCGE_LOG(Warning, "Failed to create output accessor or iterator");
+			continue;
+		}
 
-				const FPCGMetadataAttributeBase* AttributeBase = SampledData->Metadata->GetConstAttribute(LocalSourceAttribute);
-				if(!PCGMetadataOperations::SetValue<float>(AttributeBase, SampledPoints, DensitySetter))
-				{
-					PCGE_LOG(Error, "Attribute %s already exists but its type is not compatible", *LocalSourceAttribute.ToString());
-				}
-			}
-			else if (PointProperty == EPCGPointProperties::BoundsMin)
-			{
-				auto ExtentsSetter = [](FPCGPoint& InPoint, const FVector& InValue) { InPoint.BoundsMin = InValue; };
+		// By construction, they should be the same
+		check(InputKeys->GetNum() == OutputKeys->GetNum());
 
-				const FPCGMetadataAttributeBase* AttributeBase = SampledData->Metadata->GetConstAttribute(LocalSourceAttribute);
-				if (!PCGMetadataOperations::SetValue<FVector>(AttributeBase, SampledPoints, ExtentsSetter))
-				{
-					PCGE_LOG(Error, "Attribute %s already exists but its type is not compatible", *LocalSourceAttribute.ToString());
-				}
-			}
-			else if (PointProperty == EPCGPointProperties::BoundsMax)
-			{
-				auto ExtentsSetter = [](FPCGPoint& InPoint, const FVector& InValue) { InPoint.BoundsMax = InValue; };
+		// Final verification, if we can put the value of input into output
+		if (!PCG::Private::IsBroadcastable(InputAccessor->GetUnderlyingType(), OutputAccessor->GetUnderlyingType()))
+		{
+			PCGE_LOG(Error, "Can't broadcast input type into output type");
+			continue;
+		}
 
-				const FPCGMetadataAttributeBase* AttributeBase = SampledData->Metadata->GetConstAttribute(LocalSourceAttribute);
-				if (!PCGMetadataOperations::SetValue<FVector>(AttributeBase, SampledPoints, ExtentsSetter))
-				{
-					PCGE_LOG(Error, "Attribute %s already exists but its type is not compatible", *LocalSourceAttribute.ToString());
-				}
-			}
-			else if (PointProperty == EPCGPointProperties::Extents)
-			{
-				auto ExtentsSetter = [](FPCGPoint& InPoint, const FVector& InValue) { InPoint.SetExtents(InValue); };
+		// At this point, we are ready.
+		auto Operation = [&](auto Dummy)
+		{
+			using OutputType = decltype(Dummy);
+			OutputType Value{};
 
-				const FPCGMetadataAttributeBase* AttributeBase = SampledData->Metadata->GetConstAttribute(LocalSourceAttribute);
-				if(!PCGMetadataOperations::SetValue<FVector>(AttributeBase, SampledPoints, ExtentsSetter))
-				{
-					PCGE_LOG(Error, "Attribute %s already exists but its type is not compatible", *LocalSourceAttribute.ToString());
-				}
-			}
-			else if (PointProperty == EPCGPointProperties::Color)
-			{
-				auto ColorSetter = [](FPCGPoint& InPoint, const FVector4& InValue) { InPoint.Color = InValue; };
+			EPCGAttributeAccessorFlags Flags = EPCGAttributeAccessorFlags::AllowBroadcast;
 
-				const FPCGMetadataAttributeBase* AttributeBase = SampledData->Metadata->GetConstAttribute(LocalSourceAttribute);
-				if(!PCGMetadataOperations::SetValue<FVector4>(AttributeBase, SampledPoints, ColorSetter))
-				{
-					PCGE_LOG(Error, "Attribute %s already exists but its type is not compatible", *LocalSourceAttribute.ToString());
-				}
-			}
-			else if (PointProperty == EPCGPointProperties::Position)
-			{
-				auto PositionSetter = [](FPCGPoint& InPoint, const FVector& InValue) { InPoint.Transform.SetLocation(InValue); };
+			const int32 NumberOfElements = InputKeys->GetNum();
+			constexpr int32 ChunkSize = 256;
 
-				const FPCGMetadataAttributeBase* AttributeBase = SampledData->Metadata->GetConstAttribute(LocalSourceAttribute);
-				if (!PCGMetadataOperations::SetValue<FVector>(AttributeBase, SampledPoints, PositionSetter))
-				{
-					PCGE_LOG(Error, "Attribute %s already exists but its type is not compatible", *LocalSourceAttribute.ToString());
-				}
-			}
-			else if (PointProperty == EPCGPointProperties::Rotation)
-			{
-				auto RotationSetter = [](FPCGPoint& InPoint, const FQuat& InValue) { InPoint.Transform.SetRotation(InValue.GetNormalized()); };
+			TArray<OutputType, TInlineAllocator<ChunkSize>> TempValues;
+			TempValues.SetNum(ChunkSize);
 
-				const FPCGMetadataAttributeBase* AttributeBase = SampledData->Metadata->GetConstAttribute(LocalSourceAttribute);
-				if (!PCGMetadataOperations::SetValue<FQuat>(AttributeBase, SampledPoints, RotationSetter))
-				{
-					PCGE_LOG(Error, "Attribute %s already exists but its type is not compatible", *LocalSourceAttribute.ToString());
-				}
-			}
-			else if (PointProperty == EPCGPointProperties::Scale)
-			{
-				auto ScaleSetter = [](FPCGPoint& InPoint, const FVector& InValue) { InPoint.Transform.SetScale3D(InValue); };
+			const int32 NumberOfIterations = (NumberOfElements + ChunkSize - 1) / ChunkSize;
 
-				const FPCGMetadataAttributeBase* AttributeBase = SampledData->Metadata->GetConstAttribute(LocalSourceAttribute);
-				if (!PCGMetadataOperations::SetValue<FVector>(AttributeBase, SampledPoints, ScaleSetter))
-				{
-					PCGE_LOG(Error, "Attribute %s already exists but its type is not compatible", *LocalSourceAttribute.ToString());
-				}
-			}
-			else if (PointProperty == EPCGPointProperties::Transform)
+			for (int32 i = 0; i < NumberOfIterations; ++i)
 			{
-				auto TransformSetter = [](FPCGPoint& InPoint, const FTransform& InValue) { InPoint.Transform = InValue; };
+				const int32 StartIndex = i * ChunkSize;
+				const int32 Range = FMath::Min(NumberOfElements - StartIndex, ChunkSize);
+				TArrayView<OutputType> View(TempValues.GetData(), Range);
 
-				const FPCGMetadataAttributeBase* AttributeBase = SampledData->Metadata->GetConstAttribute(LocalSourceAttribute);
-				if (!PCGMetadataOperations::SetValue<FTransform>(AttributeBase, SampledPoints, TransformSetter))
-				{
-					PCGE_LOG(Error, "Attribute %s already exists but its type is not compatible", *LocalSourceAttribute.ToString());
-				}
+				InputAccessor->GetRange<OutputType>(View, StartIndex, *InputKeys, Flags);
+				OutputAccessor->SetRange<OutputType>(View, StartIndex, *OutputKeys, Flags);
 			}
-			else if (PointProperty == EPCGPointProperties::Steepness)
-			{
-				auto SteepnessSetter = [](FPCGPoint& InPoint, const float& InValue) { InPoint.Steepness = InValue; };
 
-				const FPCGMetadataAttributeBase* AttributeBase = SampledData->Metadata->GetConstAttribute(LocalSourceAttribute);
-				if (!PCGMetadataOperations::SetValue<float>(AttributeBase, SampledPoints, SteepnessSetter))
-				{
-					PCGE_LOG(Error, "Attribute %s already exists but its type is not compatible", *LocalSourceAttribute.ToString());
-				}
-			}
-		}
-		else // Attribute to attribute
+			return true;
+		};
+
+		if (!PCGMetadataAttribute::CallbackWithRightType(OutputAccessor->GetUnderlyingType(), Operation))
 		{
-			SampledData->Metadata->CopyExistingAttribute(LocalSourceAttribute, DestinationAttribute);
+			PCGE_LOG(Warning, "Error while getting/setting values");
+			continue;
 		}
 	}
 
diff --git a/Engine/Plugins/Experimental/PCG/Source/PCG/Private/Elements/Metadata/PCGMetadataElementCommon.cpp b/Engine/Plugins/Experimental/PCG/Source/PCG/Private/Elements/Metadata/PCGMetadataElementCommon.cpp
--- a/Engine/Plugins/Experimental/PCG/Source/PCG/Private/Elements/Metadata/PCGMetadataElementCommon.cpp
+++ b/Engine/Plugins/Experimental/PCG/Source/PCG/Private/Elements/Metadata/PCGMetadataElementCommon.cpp
@@ -22,7 +22,7 @@ namespace PCGMetadataElementCommon
 		else if (const UPCGParamData* ParamsInput = Cast<const UPCGParamData>(InTaggedData.Data))
 		{
 			UPCGParamData* NewParamData = NewObject<UPCGParamData>();
-			NewParamData->Metadata->Initialize(ParamsInput->Metadata);
+			NewParamData->Metadata->InitializeAsCopy(ParamsInput->Metadata);
 
 			OutTaggedData.Data = NewParamData;
 
diff --git a/Engine/Plugins/Experimental/PCG/Source/PCG/Private/Elements/Metadata/PCGMetadataRenameElement.cpp b/Engine/Plugins/Experimental/PCG/Source/PCG/Private/Elements/Metadata/PCGMetadataRenameElement.cpp
--- a/Engine/Plugins/Experimental/PCG/Source/PCG/Private/Elements/Metadata/PCGMetadataRenameElement.cpp
+++ b/Engine/Plugins/Experimental/PCG/Source/PCG/Private/Elements/Metadata/PCGMetadataRenameElement.cpp
@@ -1,13 +1,24 @@
 // Copyright Epic Games, Inc. All Rights Reserved.
 
 #include "Elements/Metadata/PCGMetadataRenameElement.h"
+
 #include "Data/PCGSpatialData.h"
+#include "Elements/Metadata/PCGMetadataElementCommon.h"
 #include "Helpers/PCGSettingsHelpers.h"
 #include "Metadata/PCGMetadata.h"
 #include "Metadata/PCGMetadataAttribute.h"
 
 #include UE_INLINE_GENERATED_CPP_BY_NAME(PCGMetadataRenameElement)
 
+TArray<FPCGPinProperties> UPCGMetadataRenameSettings::InputPinProperties() const
+{
+	TArray<FPCGPinProperties> PinProperties;
+	PinProperties.Emplace(PCGPinConstants::DefaultInputLabel, EPCGDataType::Any, /*bInAllowMultipleConnections=*/ true);
+	PinProperties.Emplace(PCGPinConstants::DefaultParamsLabel, EPCGDataType::Param, /*bInAllowMultipleConnections=*/ false);
+
+	return PinProperties;
+}
+
 TArray<FPCGPinProperties> UPCGMetadataRenameSettings::OutputPinProperties() const
 {
 	TArray<FPCGPinProperties> PinProperties;
@@ -28,9 +39,9 @@ bool FPCGMetadataRenameElement::ExecuteInternal(FPCGContext* Context) const
 	const UPCGMetadataRenameSettings* Settings = Context->GetInputSettings<UPCGMetadataRenameSettings>();
 	check(Settings);
 
-	TArray<FPCGTaggedData> Inputs = Context->InputData.GetInputs();
-	TArray<FPCGTaggedData> AllParams = Context->InputData.GetAllParams();
-	UPCGParamData* Params = Context->InputData.GetParams();
+	TArray<FPCGTaggedData> Inputs = Context->InputData.GetInputsByPin(PCGPinConstants::DefaultInputLabel);
+	TArray<FPCGTaggedData> InputParams = Context->InputData.GetInputsByPin(PCGPinConstants::DefaultParamsLabel);
+	const UPCGParamData* Params = InputParams.IsEmpty() ? nullptr : Cast<const UPCGParamData>(InputParams[0].Data);
 
 	TArray<FPCGTaggedData>& Outputs = Context->OutputData.TaggedData;
 	
@@ -49,50 +60,23 @@ bool FPCGMetadataRenameElement::ExecuteInternal(FPCGContext* Context) const
 	{
 		FPCGTaggedData& Output = Outputs.Add_GetRef(Input);
 
-		const UPCGSpatialData* SpatialInput = Cast<const UPCGSpatialData>(Input.Data);
-		if (!SpatialInput)
-		{
-			continue;
-		}
+		const UPCGMetadata* Metadata = nullptr;
 
-		// If the data has a metadata & the attribute to rename, then duplicate the data
-		// otherwise, keep it as is.
-		const UPCGMetadata* Metadata = SpatialInput->Metadata;
-
-		if (!Metadata)
+		if (const UPCGSpatialData* SpatialInput = Cast<const UPCGSpatialData>(Input.Data))
 		{
-			continue;
+			Metadata = SpatialInput->Metadata;
+			Output.Data = SpatialInput;
 		}
-
-		const FName LocalAttributeToRename = ((AttributeToRename != NAME_None) ? AttributeToRename : Metadata->GetLatestAttributeNameOrNone());
-
-		if (!Metadata->HasAttribute(LocalAttributeToRename))
-		{
-			continue;
-		}
-
-		//TODO: this might require to execute on the main thread
-		UPCGSpatialData* NewSpatialData = Cast<UPCGSpatialData>(StaticDuplicateObject(SpatialInput, const_cast<UPCGSpatialData*>(SpatialInput), FName()));
-		NewSpatialData->InitializeFromData(SpatialInput);
-		NewSpatialData->Metadata->RenameAttribute(LocalAttributeToRename, NewAttributeName);
-
-		Output.Data = NewSpatialData;
-	}
-
-	for (const FPCGTaggedData& ParamTaggedData : AllParams)
-	{
-		FPCGTaggedData& Output = Outputs.Add_GetRef(ParamTaggedData);
-
-		const UPCGParamData* InputParams = Cast<const UPCGParamData>(ParamTaggedData.Data);
-		if (!InputParams)
+		else if (const UPCGParamData* ParamInput = Cast<UPCGParamData>(Input.Data))
 		{
-			continue;
+			Metadata = ParamInput->Metadata;
+			Output.Data = ParamInput;
 		}
-		
-		const UPCGMetadata* Metadata = InputParams->Metadata;
 
 		if (!Metadata)
 		{
+			UE_LOG(LogPCG, Warning, TEXT("Input is not supported, only supports spatial and params"));
+			Output.Data = nullptr;
 			continue;
 		}
 
@@ -103,14 +87,13 @@ bool FPCGMetadataRenameElement::ExecuteInternal(FPCGContext* Context) const
 			continue;
 		}
 
-		//TODO: this might require to execute on the main thread
-		UPCGParamData* NewParams = Cast<UPCGParamData>(StaticDuplicateObject(InputParams, const_cast<UPCGParamData*>(InputParams), FName()));
-		// Note: we will not parent the metadata here to ensure that the 0th entry is present on this metadata
-		// We will instead do a copy (and keep its parent, etc.)
-		NewParams->Metadata->InitializeAsCopy(Metadata);
-		NewParams->Metadata->RenameAttribute(LocalAttributeToRename, NewAttributeName);
+		UPCGMetadata* NewMetadata = nullptr;
+		PCGMetadataElementCommon::DuplicateTaggedData(Input, Output, NewMetadata);
 
-		Output.Data = NewParams;
+		if (!NewMetadata || !NewMetadata->RenameAttribute(LocalAttributeToRename, NewAttributeName))
+		{
+			UE_LOG(LogPCG, Warning, TEXT("Failed to rename attribute from %s to %s"), *LocalAttributeToRename.ToString(), *NewAttributeName.ToString());
+		}
 	}
 
 	// Pass-through settings
diff --git a/Engine/Plugins/Experimental/PCG/Source/PCG/Private/Elements/PCGAttributeFilterElement.cpp b/Engine/Plugins/Experimental/PCG/Source/PCG/Private/Elements/PCGAttributeFilterElement.cpp
--- a/Engine/Plugins/Experimental/PCG/Source/PCG/Private/Elements/PCGAttributeFilterElement.cpp
+++ b/Engine/Plugins/Experimental/PCG/Source/PCG/Private/Elements/PCGAttributeFilterElement.cpp
@@ -15,6 +15,41 @@ namespace PCGAttributeFilterConstants
 	const FName NodeName = TEXT("FilterAttribute");
 }
 
+namespace PCGAttributeFilterSettings
+{
+	TArray<FString> GenerateNameArray(const FString& InString)
+	{
+		TArray<FString> Result;
+		InString.ParseIntoArrayWS(Result, TEXT(","));
+		return Result;
+	}
+}
+
+void UPCGAttributeFilterSettings::PostLoad()
+{
+	Super::PostLoad();
+
+#if WITH_EDITOR
+	if (!AttributesToKeep_DEPRECATED.IsEmpty())
+	{
+		SelectedAttributes.Empty();
+		Operation = EPCGAttributeFilterOperation::KeepSelectedAttributes;
+		// Can't use FString::Join since it is an array of FName
+		for (int i = 0; i < AttributesToKeep_DEPRECATED.Num(); ++i)
+		{
+			if (i != 0)
+			{
+				SelectedAttributes += TEXT(",");
+			}
+
+			SelectedAttributes += AttributesToKeep_DEPRECATED[i].ToString();
+		}
+
+		AttributesToKeep_DEPRECATED.Empty();
+	}
+#endif // WITH_EDITOR
+}
+
 #if WITH_EDITOR
 FName UPCGAttributeFilterSettings::GetDefaultNodeName() const
 {
@@ -24,14 +59,28 @@ FName UPCGAttributeFilterSettings::GetDefaultNodeName() const
 
 FName UPCGAttributeFilterSettings::AdditionalTaskName() const
 {
+	TArray<FString> AttributesToKeep = PCGAttributeFilterSettings::GenerateNameArray(SelectedAttributes);
+
+	FString NodeName = PCGAttributeFilterConstants::NodeName.ToString();
+
+	switch (Operation)
+	{
+	case EPCGAttributeFilterOperation::KeepSelectedAttributes:
+		NodeName += TEXT(" (Keep)");
+		break;
+	case EPCGAttributeFilterOperation::DeleteSelectedAttributes:
+		NodeName += TEXT(" (Delete)");
+		break;
+	}
+
 	// If we filter only one attribute, show its name
 	if (AttributesToKeep.Num() == 1)
 	{
-		return FName(FString::Printf(TEXT("%s: %s"), *PCGAttributeFilterConstants::NodeName.ToString(), *AttributesToKeep[0].ToString()));
+		return FName(FString::Printf(TEXT("%s: %s"), *NodeName, *AttributesToKeep[0]));
 	}
 	else
 	{
-		return NAME_None;
+		return FName(NodeName);
 	}
 }
 
@@ -56,6 +105,8 @@ bool FPCGAttributeFilterElement::ExecuteInternal(FPCGContext* Context) const
 
 	const UPCGAttributeFilterSettings* Settings = Context->GetInputSettings<UPCGAttributeFilterSettings>();
 
+	const bool bAddAttributesFromParent = (Settings->Operation == EPCGAttributeFilterOperation::DeleteSelectedAttributes);
+
 	TArray<FPCGTaggedData> Inputs = Context->InputData.GetInputsByPin(PCGPinConstants::DefaultInputLabel);
 
 	for (const FPCGTaggedData& InputTaggedData : Inputs)
@@ -72,7 +123,7 @@ bool FPCGAttributeFilterElement::ExecuteInternal(FPCGContext* Context) const
 
 			UPCGSpatialData* NewSpatialData = InputSpatialData->DuplicateData(/*bInitializeFromThisData=*/false);
 			Metadata = NewSpatialData->Metadata;
-			NewSpatialData->Metadata->Initialize(ParentMetadata, /*bAddAttributesFromParent=*/false);
+			NewSpatialData->Metadata->Initialize(ParentMetadata, bAddAttributesFromParent);
 
 			// No need to inherit metadata since we already initialized it.
 			NewSpatialData->InitializeFromData(InputSpatialData, /*InMetadataParentOverride=*/ nullptr, /*bInheritMetadata=*/ false);
@@ -86,7 +137,7 @@ bool FPCGAttributeFilterElement::ExecuteInternal(FPCGContext* Context) const
 			UPCGParamData* NewParamData = NewObject<UPCGParamData>();
 			Metadata = NewParamData->Metadata;
 
-			Metadata->Initialize(InputParamData->Metadata, /*bAddAttributesFromParent=*/false);
+			Metadata->Initialize(InputParamData->Metadata, bAddAttributesFromParent);
 			OutputData = NewParamData;
 		}
 		else
@@ -95,10 +146,20 @@ bool FPCGAttributeFilterElement::ExecuteInternal(FPCGContext* Context) const
 			continue;
 		}
 
-		// Then add each attribute in the list explicitly
-		for (const FName& AttributeName : Settings->AttributesToKeep)
+		TArray<FString> AttributesToKeep = PCGAttributeFilterSettings::GenerateNameArray(Settings->SelectedAttributes);
+
+		// Then add/remove each attribute in the list explicitly
+		for (const FString& AttributeName : AttributesToKeep)
 		{
-			Metadata->AddAttribute(ParentMetadata, AttributeName);
+			switch (Settings->Operation)
+			{
+			case EPCGAttributeFilterOperation::KeepSelectedAttributes:
+				Metadata->AddAttribute(ParentMetadata, FName(AttributeName));
+				break;
+			case EPCGAttributeFilterOperation::DeleteSelectedAttributes:
+				Metadata->DeleteAttribute(FName(AttributeName));
+				break;
+			}
 		}
 
 		TArray<FPCGTaggedData>& Outputs = Context->OutputData.TaggedData;
diff --git a/Engine/Plugins/Experimental/PCG/Source/PCG/Private/Elements/PCGAttributeGetFromPointIndexElement.cpp b/Engine/Plugins/Experimental/PCG/Source/PCG/Private/Elements/PCGAttributeGetFromPointIndexElement.cpp
--- a/Engine/Plugins/Experimental/PCG/Source/PCG/Private/Elements/PCGAttributeGetFromPointIndexElement.cpp
+++ b/Engine/Plugins/Experimental/PCG/Source/PCG/Private/Elements/PCGAttributeGetFromPointIndexElement.cpp
@@ -6,9 +6,12 @@
 #include "PCGHelpers.h"
 #include "PCGParamData.h"
 #include "Data/PCGPointData.h"
+#include "Helpers/PCGSettingsHelpers.h"
 #include "Metadata/PCGMetadata.h"
 #include "Metadata/PCGMetadataAttribute.h"
-#include "Helpers/PCGSettingsHelpers.h"
+#include "Metadata/Accessors/IPCGAttributeAccessor.h"
+#include "Metadata/Accessors/PCGAttributeAccessorKeys.h"
+#include "Metadata/Accessors/PCGAttributeAccessorHelpers.h"
 
 #include UE_INLINE_GENERATED_CPP_BY_NAME(PCGAttributeGetFromPointIndexElement)
 
@@ -19,6 +22,19 @@ FName UPCGAttributeGetFromPointIndexSettings::GetDefaultNodeName() const
 }
 #endif
 
+void UPCGAttributeGetFromPointIndexSettings::PostLoad()
+{
+	Super::PostLoad();
+
+#if WITH_EDITOR
+	if (InputAttributeName_DEPRECATED != NAME_None)
+	{
+		InputSource.SetAttributeName(InputAttributeName_DEPRECATED);
+		InputAttributeName_DEPRECATED = NAME_None;
+	}
+#endif
+}
+
 TArray<FPCGPinProperties> UPCGAttributeGetFromPointIndexSettings::InputPinProperties() const
 {
 	TArray<FPCGPinProperties> PinProperties;
@@ -94,39 +110,58 @@ bool FPCGAttributeGetFromPointIndexElement::ExecuteInternal(FPCGContext* Context
 		Output.Pin = PCGAttributeGetFromPointIndexConstants::OutputPointLabel;
 	}
 
-	if (PointData->Metadata && PointData->Metadata->HasAttribute(Settings->InputAttributeName))
+	FPCGAttributePropertySelector InputSource = Settings->InputSource;
+	if (InputSource.Selection == EPCGAttributePropertySelection::Attribute && InputSource.AttributeName == NAME_None)
+	{
+		InputSource.SetAttributeName(PointData->Metadata->GetLatestAttributeNameOrNone());
+	}
+
+	FName OutputAttributeName = (Settings->OutputAttributeName == NAME_None) ? InputSource.GetName() : Settings->OutputAttributeName;
+
+	TUniquePtr<const IPCGAttributeAccessor> Accessor = PCGAttributeAccessorHelpers::CreateConstAccessor(PointData, InputSource);
+	FPCGAttributeAccessorKeysPoints PointKey(Point);
+
+	if (Accessor.IsValid())
 	{
-		const FPCGMetadataAttributeBase* Attribute = PointData->Metadata->GetConstAttribute(Settings->InputAttributeName);
 		UPCGParamData* OutputParamData = NewObject<UPCGParamData>();
 
-		auto ExtractAttribute = [this, Context, AttributeName = Settings->InputAttributeName, &Point, &OutputParamData, Attribute](auto DummyValue) -> bool
+		auto ExtractAttribute = [this, Context, OutputAttributeName, &OutputParamData, &Accessor, &PointKey](auto DummyValue) -> bool
 		{
 			using AttributeType = decltype(DummyValue);
 
-			const FPCGMetadataAttribute<AttributeType>* TypedAttribute = static_cast<const FPCGMetadataAttribute<AttributeType>*>(Attribute);
+			AttributeType Value{};
+
+			// Should never fail, as OutputType == Accessor->UnderlyingType
+			if (!ensure(Accessor->Get<AttributeType>(Value, PointKey)))
+			{
+				return false;
+			}
 
-			FPCGMetadataAttributeBase* NewAttributeBase =
-				OutputParamData->Metadata->CreateAttribute<AttributeType>(AttributeName, TypedAttribute->GetValueFromItemKey(Point.MetadataEntry), TypedAttribute->AllowsInterpolation(), /*bOverrideParent=*/false);
+			FPCGMetadataAttribute<AttributeType>* NewAttribute = static_cast<FPCGMetadataAttribute<AttributeType>*>(
+				OutputParamData->Metadata->CreateAttribute<AttributeType>(OutputAttributeName, Value, /*bAllowInterpolation=*/true, /*bOverrideParent=*/false));
 
-			if (!NewAttributeBase)
+			if (!NewAttribute)
 			{
-				PCGE_LOG(Error, "Error while creating target attribute %s", *AttributeName.ToString());
+				PCGE_LOG(Error, "Error while creating target attribute %s", *OutputAttributeName.ToString());
 				return false;
 			}
 
-			PCGMetadataEntryKey EntryKey = OutputParamData->Metadata->AddEntry();
-			NewAttributeBase->SetValueFromValueKey(EntryKey, PCGDefaultValueKey);
+			NewAttribute->SetValue(OutputParamData->Metadata->AddEntry(), Value);
 
 			return true;
 		};
 
-		if (PCGMetadataAttribute::CallbackWithRightType(Attribute->GetTypeId(), ExtractAttribute))
+		if (PCGMetadataAttribute::CallbackWithRightType(Accessor->GetUnderlyingType(), ExtractAttribute))
 		{
 			FPCGTaggedData& Output = Outputs.Emplace_GetRef();
 			Output.Data = OutputParamData;
 			Output.Pin = PCGAttributeGetFromPointIndexConstants::OutputAttributeLabel;
 		}
 	}
+	else
+	{
+		PCGE_LOG(Warning, "Can't find attribute/property %s in input", *InputSource.GetName().ToString());
+	}
 
 	return true;
 }
diff --git a/Engine/Plugins/Experimental/PCG/Source/PCG/Private/Elements/PCGAttributeReduceElement.cpp b/Engine/Plugins/Experimental/PCG/Source/PCG/Private/Elements/PCGAttributeReduceElement.cpp
--- a/Engine/Plugins/Experimental/PCG/Source/PCG/Private/Elements/PCGAttributeReduceElement.cpp
+++ b/Engine/Plugins/Experimental/PCG/Source/PCG/Private/Elements/PCGAttributeReduceElement.cpp
@@ -5,17 +5,20 @@
 #include "PCGData.h"
 #include "PCGParamData.h"
 #include "Data/PCGPointData.h"
+#include "Elements/Metadata/PCGMetadataElementCommon.h"
+#include "Helpers/PCGSettingsHelpers.h"
 #include "Metadata/PCGMetadata.h"
 #include "Metadata/PCGMetadataAttribute.h"
-#include "Metadata/PCGMetadataEntryKeyIterator.h"
-#include "Helpers/PCGSettingsHelpers.h"
+#include "Metadata/Accessors/IPCGAttributeAccessor.h"
+#include "Metadata/Accessors/PCGAttributeAccessorHelpers.h"
+#include "Metadata/Accessors/PCGAttributeAccessorKeys.h"
 
 #include UE_INLINE_GENERATED_CPP_BY_NAME(PCGAttributeReduceElement)
 
 namespace PCGAttributeReduceElement
 {
 	template <typename T>
-	bool Average(IPCGMetadataEntryIterator& Iterator, const FPCGMetadataAttribute<T>* InAttribute, T& OutValue)
+	bool Average(const IPCGAttributeAccessorKeys& Keys, const IPCGAttributeAccessor& Accessor, T& OutValue)
 	{
 		if constexpr (!PCG::Private::MetadataTraits<T>::CanSubAdd || !PCG::Private::MetadataTraits<T>::CanInterpolate)
 		{
@@ -23,28 +26,24 @@ namespace PCGAttributeReduceElement
 		}
 		else
 		{
-			// Can't have a repeat iterator
-			check(!Iterator.IsRepeat());
-
-			int32 NumberOfEntries = 0;
 			OutValue = PCG::Private::MetadataTraits<T>::ZeroValue();
-			for (; !Iterator.IsEnd(); ++Iterator)
-			{
-				OutValue = PCG::Private::MetadataTraits<T>::Add(OutValue, InAttribute->GetValueFromItemKey(*Iterator));
-				++NumberOfEntries;
-			}
 
-			if (NumberOfEntries != 0)
+			bool bSuccess = PCGMetadataElementCommon::ApplyOnAccessor<T>(Keys, Accessor, [&OutValue](const T& InValue, int32)
+				{
+					OutValue = PCG::Private::MetadataTraits<T>::Add(OutValue, InValue);
+				});
+
+			if (bSuccess)
 			{
-				OutValue = PCG::Private::MetadataTraits<T>::WeightedSum(PCG::Private::MetadataTraits<T>::ZeroValue(), OutValue, 1.0f / NumberOfEntries);
+				OutValue = PCG::Private::MetadataTraits<T>::WeightedSum(PCG::Private::MetadataTraits<T>::ZeroValue(), OutValue, 1.0f / Keys.GetNum());
 			}
 
-			return true;
+			return bSuccess;
 		}
 	}
 
 	template <typename T, bool bIsMin>
-	bool MinMax(IPCGMetadataEntryIterator& Iterator, const FPCGMetadataAttribute<T>* InAttribute, T& OutValue)
+	bool MinMax(const IPCGAttributeAccessorKeys& Keys, const IPCGAttributeAccessor& Accessor, T& OutValue)
 	{
 		if constexpr (!PCG::Private::MetadataTraits<T>::CanMinMax)
 		{
@@ -52,33 +51,28 @@ namespace PCGAttributeReduceElement
 		}
 		else
 		{
-			// Can't have a repeat iterator
-			check(!Iterator.IsRepeat());
-
-			int32 NumberOfEntries = 0;
-			for (; !Iterator.IsEnd(); ++Iterator)
-			{
-				T AttributeValue = InAttribute->GetValueFromItemKey(*Iterator);
+			OutValue = PCG::Private::MetadataTraits<T>::ZeroValue();
+			bool bFirstValue = true;
 
-				if (NumberOfEntries == 0)
+			return PCGMetadataElementCommon::ApplyOnAccessor<T>(Keys, Accessor, [&OutValue, &bFirstValue](const T& InValue, int32)
 				{
-					OutValue = AttributeValue;
-				}
-				else
-				{
-					if constexpr (bIsMin)
+					if (bFirstValue)
 					{
-						OutValue = PCG::Private::MetadataTraits<T>::Min(OutValue, AttributeValue);
+						OutValue = InValue;
+						bFirstValue = false;
 					}
 					else
 					{
-						OutValue = PCG::Private::MetadataTraits<T>::Max(OutValue, AttributeValue);
+						if constexpr (bIsMin)
+						{
+							OutValue = PCG::Private::MetadataTraits<T>::Min(OutValue, InValue);
+						}
+						else
+						{
+							OutValue = PCG::Private::MetadataTraits<T>::Max(OutValue, InValue);
+						}
 					}
-				}
-				++NumberOfEntries;
-			}
-
-			return true;
+				});
 		}
 	}
 }
@@ -90,12 +84,31 @@ FName UPCGAttributeReduceSettings::GetDefaultNodeName() const
 }
 #endif
 
+void UPCGAttributeReduceSettings::PostLoad()
+{
+	Super::PostLoad();
+
+#if WITH_EDITOR
+	if (InputAttributeName_DEPRECATED != NAME_None)
+	{
+		InputSource.SetAttributeName(InputAttributeName_DEPRECATED);
+		InputAttributeName_DEPRECATED = NAME_None;
+	}
+#endif
+}
+
 FName UPCGAttributeReduceSettings::AdditionalTaskName() const
 {
-	if (const UEnum* EnumPtr = FindObject<UEnum>(nullptr, TEXT("/Script/PCG.EPCGAttributeReduceOperation"), true))
+	if (const UEnum* EnumPtr = StaticEnum<EPCGAttributeReduceOperation>())
 	{
 		const FString OperationName = EnumPtr->GetNameStringByValue(static_cast<int>(Operation));
 
+		FName InputAttributeName = InputSource.GetName();
+		if (InputAttributeName == NAME_None)
+		{
+			InputAttributeName = FName(TEXT("LastAttribute"));
+		}
+
 		if (InputAttributeName != OutputAttributeName && OutputAttributeName != NAME_None)
 		{
 			return FName(FString::Printf(TEXT("Reduce %s to %s: %s"), *InputAttributeName.ToString(), *OutputAttributeName.ToString(), *OperationName));
@@ -157,64 +170,52 @@ bool FPCGAttributeReduceElement::ExecuteInternal(FPCGContext* Context) const
 		return true;
 	}
 
-	if (!SpatialData->Metadata)
-	{
-		PCGE_LOG(Error, "Input data doesn't have metadata");
-		return true;
-	}
+	const UPCGPointData* PointData = Cast<UPCGPointData>(SpatialData);
 
-	if (!SpatialData->Metadata->HasAttribute(Settings->InputAttributeName))
+	FPCGAttributePropertySelector InputSource = Settings->InputSource;
+	if (InputSource.Selection == EPCGAttributePropertySelection::Attribute && InputSource.AttributeName == NAME_None && SpatialData->Metadata)
 	{
-		PCGE_LOG(Error, "Input metadata doesn't have attribute \"%s\"", *Settings->InputAttributeName.ToString());
-		return true;
+		InputSource.SetAttributeName(SpatialData->Metadata->GetLatestAttributeNameOrNone());
 	}
 
-	const UPCGPointData* PointData = Cast<UPCGPointData>(SpatialData);
-
-	const FPCGMetadataAttributeBase* InputAttribute = SpatialData->Metadata->GetConstAttribute(Settings->InputAttributeName);
-	const FName OutputAttributeName = (Settings->OutputAttributeName == NAME_None) ? Settings->InputAttributeName : Settings->OutputAttributeName;
+	const FName OutputAttributeName = (Settings->OutputAttributeName == NAME_None) ? InputSource.GetName() : Settings->OutputAttributeName;
 	UPCGParamData* OutputParamData = NewObject<UPCGParamData>();
-	
-	FPCGMetadataAttributeBase* OutputAttribute = OutputParamData->Metadata->CopyAttribute(InputAttribute, OutputAttributeName, /*bKeepParent=*/ false, /*bCopyEntries=*/ false, /*bCopyValues=*/ false);
 
-	if (!OutputAttribute)
+	TUniquePtr<const IPCGAttributeAccessor> Accessor = PCGAttributeAccessorHelpers::CreateConstAccessor(PointData, InputSource);
+	TUniquePtr<const IPCGAttributeAccessorKeys> Keys = PCGAttributeAccessorHelpers::CreateConstKeys(PointData, InputSource);
+
+	if (!Accessor.IsValid() || !Keys.IsValid())
 	{
-		PCGE_LOG(Error, "Error while creating target attribute %s", *OutputAttributeName.ToString());
+		PCGE_LOG(Error, "Input attribute/property doesn't exists");
 		return true;
 	}
 
-	auto DoOperation = [InputAttribute, OutputAttribute, Operation = Settings->Operation, SpatialData, PointData, OutputParamData](auto DummyValue) -> bool
+	auto DoOperation = [&Accessor, &Keys, Operation = Settings->Operation, OutputParamData, OutputAttributeName](auto DummyValue) -> bool
 	{
 		using AttributeType = decltype(DummyValue);
 
 		bool bSuccess = false;
-		const bool bIsRepeat = false;
-
-		const FPCGMetadataAttribute<AttributeType>* TypedInputAttribute = static_cast<const FPCGMetadataAttribute<AttributeType>*>(InputAttribute);
-		FPCGMetadataAttribute<AttributeType>* TypedOutputAttribute = static_cast<FPCGMetadataAttribute<AttributeType>*>(OutputAttribute);
 
 		AttributeType OutputValue{};
 
-		TUniquePtr<IPCGMetadataEntryIterator> Iterator;
-		if (PointData)
-		{
-			Iterator = MakeUnique<FPCGMetadataEntryPointIterator>(PointData, bIsRepeat);
-		}
-		else
+		FPCGMetadataAttribute<AttributeType>* NewAttribute = static_cast<FPCGMetadataAttribute<AttributeType>*>(
+			OutputParamData->Metadata->CreateAttribute<AttributeType>(OutputAttributeName, OutputValue, /*bAllowInterpolation=*/ true, /*bOverrideParent=*/false));
+
+		if (!NewAttribute)
 		{
-			Iterator = MakeUnique<FPCGMetadataEntryAttributeIterator>(*InputAttribute, bIsRepeat);
+			return false;
 		}
 
 		switch (Operation)
 		{
 		case EPCGAttributeReduceOperation::Average:
-			bSuccess = PCGAttributeReduceElement::Average<AttributeType>(*Iterator, TypedInputAttribute, OutputValue);
+			bSuccess = PCGAttributeReduceElement::Average<AttributeType>(*Keys, *Accessor, OutputValue);
 			break;
 		case EPCGAttributeReduceOperation::Max:
-			bSuccess = PCGAttributeReduceElement::MinMax<AttributeType, /*bIsMin*/false>(*Iterator, TypedInputAttribute, OutputValue);
+			bSuccess = PCGAttributeReduceElement::MinMax<AttributeType, /*bIsMin*/false>(*Keys, *Accessor, OutputValue);
 			break;
 		case EPCGAttributeReduceOperation::Min:
-			bSuccess = PCGAttributeReduceElement::MinMax<AttributeType, /*bIsMin*/true>(*Iterator, TypedInputAttribute, OutputValue);
+			bSuccess = PCGAttributeReduceElement::MinMax<AttributeType, /*bIsMin*/true>(*Keys, *Accessor, OutputValue);
 			break;
 		default:
 			break;
@@ -222,16 +223,16 @@ bool FPCGAttributeReduceElement::ExecuteInternal(FPCGContext* Context) const
 
 		if (bSuccess)
 		{
-			TypedOutputAttribute->SetDefaultValue(OutputValue);
-			TypedOutputAttribute->SetValueFromValueKey(OutputParamData->Metadata->AddEntry(), PCGDefaultValueKey);
+			NewAttribute->SetDefaultValue(OutputValue);
+			NewAttribute->SetValue(OutputParamData->Metadata->AddEntry(), OutputValue);
 		}
 
 		return bSuccess;
 	};
 
-	if (!PCGMetadataAttribute::CallbackWithRightType(InputAttribute->GetTypeId(), DoOperation))
+	if (!PCGMetadataAttribute::CallbackWithRightType(Accessor->GetUnderlyingType(), DoOperation))
 	{
-		PCGE_LOG(Error, "Operation was not compatible with the attribute type.");
+		PCGE_LOG(Error, "Operation was not compatible with the attribute type or couldn't create attribute %s.", *OutputAttributeName.ToString());
 		return true;
 	}
 
diff --git a/Engine/Plugins/Experimental/PCG/Source/PCG/Private/Elements/PCGAttributeSelectElement.cpp b/Engine/Plugins/Experimental/PCG/Source/PCG/Private/Elements/PCGAttributeSelectElement.cpp
--- a/Engine/Plugins/Experimental/PCG/Source/PCG/Private/Elements/PCGAttributeSelectElement.cpp
+++ b/Engine/Plugins/Experimental/PCG/Source/PCG/Private/Elements/PCGAttributeSelectElement.cpp
@@ -5,10 +5,13 @@
 #include "PCGData.h"
 #include "PCGParamData.h"
 #include "Data/PCGPointData.h"
+#include "Elements/Metadata/PCGMetadataElementCommon.h"
+#include "Helpers/PCGSettingsHelpers.h"
 #include "Metadata/PCGMetadata.h"
 #include "Metadata/PCGMetadataAttribute.h"
-#include "Metadata/PCGMetadataEntryKeyIterator.h"
-#include "Helpers/PCGSettingsHelpers.h"
+#include "Metadata/Accessors/IPCGAttributeAccessor.h"
+#include "Metadata/Accessors/PCGAttributeAccessorHelpers.h"
+#include "Metadata/Accessors/PCGAttributeAccessorKeys.h"
 
 #include UE_INLINE_GENERATED_CPP_BY_NAME(PCGAttributeSelectElement)
 
@@ -68,62 +71,53 @@ namespace PCGAttributeSelectElement
 	}
 
 	template <typename T, bool bIsMin, int Axis>
-	bool MinMaxSelect(IPCGMetadataEntryIterator& Iterator, const FPCGMetadataAttribute<T>* InAttribute, const T& InAxis, T& OutValue, int32& OutIndex)
+	bool MinMaxSelect(const IPCGAttributeAccessorKeys& Keys, const IPCGAttributeAccessor& Accessor, const T& InAxis, T& OutValue, int32& OutIndex)
 	{
-		// Can't have a repeat iterator
-		check(!Iterator.IsRepeat());
-
 		// Will be int or double
 		using CompareType = decltype(Projection<T, Axis>(T{}, InAxis));
 
 		CompareType MinMaxValue{};
+		bool bFirstValue = true;
 
-		int32 NumberOfEntries = 0;
-		for (; !Iterator.IsEnd(); ++Iterator)
-		{
-			T AttributeValue = InAttribute->GetValueFromItemKey(*Iterator);
-			CompareType CurrentValue = Projection<T, Axis>(AttributeValue, InAxis);
-
-			if constexpr (bIsMin)
+		return PCGMetadataElementCommon::ApplyOnAccessor<T>(Keys, Accessor, [&OutValue, &OutIndex, &bFirstValue, &MinMaxValue, &InAxis](const T& InValue, int32 InIndex)
 			{
-				if (NumberOfEntries == 0 || CurrentValue < MinMaxValue)
+				CompareType CurrentValue = Projection<T, Axis>(InValue, InAxis);
+
+				if constexpr (bIsMin)
 				{
-					MinMaxValue = CurrentValue;
-					OutValue = AttributeValue;
-					OutIndex = NumberOfEntries;
+					if (bFirstValue || CurrentValue < MinMaxValue)
+					{
+						MinMaxValue = CurrentValue;
+						OutValue = InValue;
+						OutIndex = InIndex;
+					}
 				}
-			}
-			else
-			{
-				if (NumberOfEntries == 0 || CurrentValue > MinMaxValue)
+				else
 				{
-					MinMaxValue = CurrentValue;
-					OutValue = AttributeValue;
-					OutIndex = NumberOfEntries;
+					if (bFirstValue || CurrentValue > MinMaxValue)
+					{
+						MinMaxValue = CurrentValue;
+						OutValue = InValue;
+						OutIndex = InIndex;
+					}
 				}
-			}
 
-			++NumberOfEntries;
-		}
-
-		return NumberOfEntries != 0;
+				bFirstValue = false;
+			});
 	}
 
 	template <typename T, int Axis>
-	bool MedianSelect(IPCGMetadataEntryIterator& Iterator, const FPCGMetadataAttribute<T>* InAttribute, const T& InAxis, T& OutValue, int32& OutIndex)
+	bool MedianSelect(const IPCGAttributeAccessorKeys& Keys, const IPCGAttributeAccessor& Accessor, const T& InAxis, T& OutValue, int32& OutIndex)
 	{
-		// Can't have a repeat iterator
-		check(!Iterator.IsRepeat());
-
 		// Will be int or double
 		using CompareType = decltype(Projection<T, Axis>(T{}, InAxis));
 
 		struct Item
 		{
-			Item(int32 InIndex, CompareType InCompareValue, T&& InAttributeValue)
+			Item(int32 InIndex, CompareType InCompareValue, const T& InAttributeValue)
 				: Index(InIndex)
 				, CompareValue(InCompareValue)
-				, AttributeValue(Forward<T>(InAttributeValue)) 
+				, AttributeValue(InAttributeValue)
 			{}
 
 			int32 Index;
@@ -131,49 +125,63 @@ namespace PCGAttributeSelectElement
 			T AttributeValue;
 		};
 
-		int32 NumberOfEntries = 0;
-		TArray<Item> Items;
-		for (; !Iterator.IsEnd(); ++Iterator)
-		{
-			T AttributeValue = InAttribute->GetValueFromItemKey(*Iterator);
-			CompareType CurrentValue = Projection<T, Axis>(AttributeValue, InAxis);
-
-			Items.Emplace(NumberOfEntries++, MoveTemp(CurrentValue), MoveTemp(AttributeValue));
-		}
+		const int32 NumberOfEntries = Keys.GetNum();
 
 		if (NumberOfEntries == 0)
 		{
 			return false;
 		}
 
-		Algo::Sort(Items, [](const Item& A, const Item& B) -> bool { return A.CompareValue < B.CompareValue; });
+		TArray<Item> CompareItems;
+		CompareItems.Reserve(NumberOfEntries);
+
+		PCGMetadataElementCommon::ApplyOnAccessor<T>(Keys, Accessor, [&CompareItems, &InAxis](const T& InValue, int32 InIndex)
+			{
+				CompareItems.Emplace(InIndex, Projection<T, Axis>(InValue, InAxis), InValue);
+			});
+
+		// TODO: Use a better algo than just sorting.
+		Algo::Sort(CompareItems, [](const Item& A, const Item& B) -> bool { return A.CompareValue < B.CompareValue; });
 
 		// Since we need to return an index, we can't do the mean on 2 values if the number of entries is even, since it will yield a value that might not exist in the original dataset.
 		// In this case we arbitrarily chose one entry.
 		const int32 Index = NumberOfEntries / 2;
-		OutValue = Items[Index].AttributeValue;
-		OutIndex = Items[Index].Index;
+		OutIndex = CompareItems[Index].Index;
+		OutValue = CompareItems[Index].AttributeValue;
 
 		return true;
 	}
 
 	template <typename T, int Axis>
-	inline bool DispatchOperation(IPCGMetadataEntryIterator& Iterator, const FPCGMetadataAttribute<T>* InAttribute, const T& InAxis, EPCGAttributeSelectOperation Operation, T& OutValue, int32& OutIndex)
+	inline bool DispatchOperation(const IPCGAttributeAccessorKeys& Keys, const IPCGAttributeAccessor& Accessor, const T& InAxis, EPCGAttributeSelectOperation Operation, T& OutValue, int32& OutIndex)
 	{
 		switch (Operation)
 		{
 		case EPCGAttributeSelectOperation::Min:
-			return PCGAttributeSelectElement::MinMaxSelect<T, /*bIsMin=*/true, Axis>(Iterator, InAttribute, InAxis, OutValue, OutIndex);
+			return PCGAttributeSelectElement::MinMaxSelect<T, /*bIsMin=*/true, Axis>(Keys, Accessor, InAxis, OutValue, OutIndex);
 		case EPCGAttributeSelectOperation::Max:
-			return PCGAttributeSelectElement::MinMaxSelect<T, /*bIsMin=*/false, Axis>(Iterator, InAttribute, InAxis, OutValue, OutIndex);
+			return PCGAttributeSelectElement::MinMaxSelect<T, /*bIsMin=*/false, Axis>(Keys, Accessor, InAxis, OutValue, OutIndex);
 		case EPCGAttributeSelectOperation::Median:
-			return PCGAttributeSelectElement::MedianSelect<T, Axis>(Iterator, InAttribute, InAxis, OutValue, OutIndex);
+			return PCGAttributeSelectElement::MedianSelect<T, Axis>(Keys, Accessor, InAxis, OutValue, OutIndex);
 		default:
 			return false;
 		}
 	}
 }
 
+void UPCGAttributeSelectSettings::PostLoad()
+{
+	Super::PostLoad();
+
+#if WITH_EDITOR
+	if (InputAttributeName_DEPRECATED != NAME_None)
+	{
+		InputSource.SetAttributeName(InputAttributeName_DEPRECATED);
+		InputAttributeName_DEPRECATED = NAME_None;
+	}
+#endif
+}
+
 #if WITH_EDITOR
 FName UPCGAttributeSelectSettings::GetDefaultNodeName() const
 {
@@ -183,9 +191,9 @@ FName UPCGAttributeSelectSettings::GetDefaultNodeName() const
 
 FName UPCGAttributeSelectSettings::AdditionalTaskName() const
 {
-	if (const UEnum* EnumOpPtr = FindObject<UEnum>(nullptr, TEXT("/Script/PCG.EPCGAttributeSelectOperation"), true))
+	if (const UEnum* EnumOpPtr = StaticEnum<EPCGAttributeSelectOperation>())
 	{
-		if (const UEnum* EnumAxisPtr = FindObject<UEnum>(nullptr, TEXT("/Script/PCG.EPCGAttributeSelectAxis"), true))
+		if (const UEnum* EnumAxisPtr = StaticEnum<EPCGAttributeSelectAxis>())
 		{
 			const FString OperationName = EnumOpPtr->GetNameStringByValue(static_cast<int>(Operation));
 			FString AxisName;
@@ -198,6 +206,12 @@ FName UPCGAttributeSelectSettings::AdditionalTaskName() const
 				AxisName = EnumAxisPtr->GetNameStringByValue(static_cast<int>(Axis));
 			}
 
+			FName InputAttributeName = InputSource.GetName();
+			if (InputAttributeName == NAME_None)
+			{
+				InputAttributeName = FName(TEXT("LastAttribute"));
+			}
+
 			if (InputAttributeName != OutputAttributeName && OutputAttributeName != NAME_None)
 			{
 				return FName(FString::Printf(TEXT("Select %s to %s: %s on %s"), *InputAttributeName.ToString(), *OutputAttributeName.ToString(), *OperationName, *AxisName));
@@ -272,25 +286,25 @@ bool FPCGAttributeSelectElement::ExecuteInternal(FPCGContext* Context) const
 		return true;
 	}
 
-	if (!SpatialData->Metadata->HasAttribute(Settings->InputAttributeName))
+	FPCGAttributePropertySelector InputSource = Settings->InputSource;
+	if (InputSource.Selection == EPCGAttributePropertySelection::Attribute && InputSource.AttributeName == NAME_None && SpatialData->Metadata)
 	{
-		PCGE_LOG(Error, "Input metadata doesn't have attribute \"%s\"", *Settings->InputAttributeName.ToString());
-		return true;
+		InputSource.SetAttributeName(SpatialData->Metadata->GetLatestAttributeNameOrNone());
 	}
 
-	const FPCGMetadataAttributeBase* InputAttribute = SpatialData->Metadata->GetConstAttribute(Settings->InputAttributeName);
-	const FName OutputAttributeName = (Settings->OutputAttributeName == NAME_None) ? Settings->InputAttributeName : Settings->OutputAttributeName;
+	const FName OutputAttributeName = (Settings->OutputAttributeName == NAME_None) ? InputSource.GetName() : Settings->OutputAttributeName;
 	UPCGParamData* OutputParamData = NewObject<UPCGParamData>();
 
-	FPCGMetadataAttributeBase* OutputAttribute = OutputParamData->Metadata->CopyAttribute(InputAttribute, OutputAttributeName, /*bKeepParent=*/ false, /*bCopyEntries=*/ false, /*bCopyValues=*/ false);
+	TUniquePtr<const IPCGAttributeAccessor> Accessor = PCGAttributeAccessorHelpers::CreateConstAccessor(PointData, InputSource);
+	TUniquePtr<const IPCGAttributeAccessorKeys> Keys = PCGAttributeAccessorHelpers::CreateConstKeys(PointData, InputSource);
 
-	if (!OutputAttribute)
+	if (!Accessor.IsValid() || !Keys.IsValid())
 	{
-		PCGE_LOG(Error, "Error while creating target attribute %s", *OutputAttributeName.ToString());
+		PCGE_LOG(Error, "Input attribute/property doesn't exists");
 		return true;
 	}
 
-	auto DoOperation = [&](auto DummyValue) -> int32
+	auto DoOperation = [this, &Accessor, &Keys, &Settings, OutputAttributeName, OutputParamData, Context](auto DummyValue) -> int32
 	{
 		using AttributeType = decltype(DummyValue);
 
@@ -301,15 +315,19 @@ bool FPCGAttributeSelectElement::ExecuteInternal(FPCGContext* Context) const
 		}
 		else
 		{
-			const FPCGMetadataAttribute<AttributeType>* TypedInputAttribute = static_cast<const FPCGMetadataAttribute<AttributeType>*>(InputAttribute);
-			FPCGMetadataAttribute<AttributeType>* TypedOutputAttribute = static_cast<FPCGMetadataAttribute<AttributeType>*>(OutputAttribute);
-
 			bool bSuccess = false;
-			const bool bIsRepeat = false;
 			AttributeType OutputValue{};
 			AttributeType Axis{};
 			int32 OutputIndex = -1;
 
+			FPCGMetadataAttribute<AttributeType>* NewAttribute = static_cast<FPCGMetadataAttribute<AttributeType>*>(
+				OutputParamData->Metadata->CreateAttribute<AttributeType>(OutputAttributeName, OutputValue, /*bAllowInterpolation=*/ true, /*bOverrideParent=*/false));
+
+			if (!NewAttribute)
+			{
+				return -1;
+			}
+
 			// First we need to verify if the axis we want to project on is valid for dimension of our vector type.
 			// If it is a scalar, we won't project anything.
 			if constexpr (PCG::Private::IsOfTypes<AttributeType, FVector2D, FVector, FVector4>())
@@ -343,36 +361,25 @@ bool FPCGAttributeSelectElement::ExecuteInternal(FPCGContext* Context) const
 				}
 			}
 
-			// Then create our iterator depending if it is a point data or a spatial data
-			TUniquePtr<IPCGMetadataEntryIterator> Iterator;
-			if (PointData)
-			{
-				Iterator = MakeUnique<FPCGMetadataEntryPointIterator>(PointData, bIsRepeat);
-			}
-			else
-			{
-				Iterator = MakeUnique<FPCGMetadataEntryAttributeIterator>(*InputAttribute, bIsRepeat);
-			}
-
 			// Finally dispatch the operation depending on the axis.
 			// If the axis is X, Y, Z or W, the projection is overkill (Dot product), so use templates to 
 			// indicate which coordinate we should take. If the axis value is -1, it will do the projection with the custom axis, passed as parameter.
 			switch (Settings->Axis)
 			{
 			case EPCGAttributeSelectAxis::X:
-				bSuccess = PCGAttributeSelectElement::DispatchOperation<AttributeType, 0>(*Iterator, TypedInputAttribute, Axis, Settings->Operation, OutputValue, OutputIndex);
+				bSuccess = PCGAttributeSelectElement::DispatchOperation<AttributeType, 0>(*Keys, *Accessor, Axis, Settings->Operation, OutputValue, OutputIndex);
 				break;
 			case EPCGAttributeSelectAxis::Y:
-				bSuccess = PCGAttributeSelectElement::DispatchOperation<AttributeType, 1>(*Iterator, TypedInputAttribute, Axis, Settings->Operation, OutputValue, OutputIndex);
+				bSuccess = PCGAttributeSelectElement::DispatchOperation<AttributeType, 1>(*Keys, *Accessor, Axis, Settings->Operation, OutputValue, OutputIndex);
 				break;
 			case EPCGAttributeSelectAxis::Z:
-				bSuccess = PCGAttributeSelectElement::DispatchOperation<AttributeType, 2>(*Iterator, TypedInputAttribute, Axis, Settings->Operation, OutputValue, OutputIndex);
+				bSuccess = PCGAttributeSelectElement::DispatchOperation<AttributeType, 2>(*Keys, *Accessor, Axis, Settings->Operation, OutputValue, OutputIndex);
 				break;
 			case EPCGAttributeSelectAxis::W:
-				bSuccess = PCGAttributeSelectElement::DispatchOperation<AttributeType, 3>(*Iterator, TypedInputAttribute, Axis, Settings->Operation, OutputValue, OutputIndex);
+				bSuccess = PCGAttributeSelectElement::DispatchOperation<AttributeType, 3>(*Keys, *Accessor, Axis, Settings->Operation, OutputValue, OutputIndex);
 				break;
 			case EPCGAttributeSelectAxis::CustomAxis:
-				bSuccess = PCGAttributeSelectElement::DispatchOperation<AttributeType, -1>(*Iterator, TypedInputAttribute, Axis, Settings->Operation, OutputValue, OutputIndex);
+				bSuccess = PCGAttributeSelectElement::DispatchOperation<AttributeType, -1>(*Keys, *Accessor, Axis, Settings->Operation, OutputValue, OutputIndex);
 				break;
 			default:
 				break;
@@ -380,8 +387,8 @@ bool FPCGAttributeSelectElement::ExecuteInternal(FPCGContext* Context) const
 
 			if (bSuccess)
 			{
-				TypedOutputAttribute->SetDefaultValue(OutputValue);
-				TypedOutputAttribute->SetValueFromValueKey(OutputParamData->Metadata->AddEntry(), PCGDefaultValueKey);
+				NewAttribute->SetDefaultValue(OutputValue);
+				NewAttribute->SetValue(OutputParamData->Metadata->AddEntry(), OutputValue);
 			}
 			else
 			{
@@ -393,7 +400,7 @@ bool FPCGAttributeSelectElement::ExecuteInternal(FPCGContext* Context) const
 		}
 	};
 
-	int32 OutputIndex = PCGMetadataAttribute::CallbackWithRightType(InputAttribute->GetTypeId(), DoOperation);
+	int32 OutputIndex = PCGMetadataAttribute::CallbackWithRightType(Accessor->GetUnderlyingType(), DoOperation);
 	if (OutputIndex < 0)
 	{
 		return true;
diff --git a/Engine/Plugins/Experimental/PCG/Source/PCG/Public/Elements/Metadata/PCGMetadataElement.h b/Engine/Plugins/Experimental/PCG/Source/PCG/Public/Elements/Metadata/PCGMetadataElement.h
--- a/Engine/Plugins/Experimental/PCG/Source/PCG/Public/Elements/Metadata/PCGMetadataElement.h
+++ b/Engine/Plugins/Experimental/PCG/Source/PCG/Public/Elements/Metadata/PCGMetadataElement.h
@@ -5,15 +5,16 @@
 #include "PCGSettings.h"
 #include "PCGElement.h"
 #include "PCGPoint.h"
+#include "Metadata/PCGAttributePropertySelector.h"
 
 #include "PCGMetadataElement.generated.h"
 
 UENUM()
-enum class EPCGMetadataOperationTarget : uint8
+enum class UE_DEPRECATED(5.2, "Not used anymore") EPCGMetadataOperationTarget : uint8
 {
 	PropertyToAttribute,
 	AttributeToProperty,
-	AttributeToAttribute
+	AttributeToAttribute,
 };
 
 UCLASS(BlueprintType, ClassGroup = (Procedural))
@@ -22,9 +23,15 @@ class PCG_API UPCGMetadataOperationSettings : public UPCGSettings
 	GENERATED_BODY()
 
 public:
+	UPCGMetadataOperationSettings();
+
+	//~Begin UObject interface
+	virtual void PostLoad() override;
+	//~End UObject interface
+
 	//~Begin UPCGSettings interface
 #if WITH_EDITOR
-	virtual FName GetDefaultNodeName() const override { return FName(TEXT("MetadataOperation")); }
+	virtual FName GetDefaultNodeName() const override { return FName(TEXT("AttributeOperation")); }
 	virtual EPCGSettingsType GetType() const override { return EPCGSettingsType::Metadata; }
 #endif
 
@@ -35,17 +42,27 @@ protected:
 	//~End UPCGSettings interface
 
 public:
-	UPROPERTY(BlueprintReadWrite, EditAnywhere, Category = Settings, meta = (EditCondition = "Target==EPCGMetadataOperationTarget::AttributeToProperty||Target==EPCGMetadataOperationTarget::AttributeToAttribute"))
-	FName SourceAttribute = NAME_None;
+	UPROPERTY(BlueprintReadWrite, EditAnywhere, Category = Settings)
+	FPCGAttributePropertySelector InputSource;
 
 	UPROPERTY(BlueprintReadWrite, EditAnywhere, Category = Settings)
-	EPCGPointProperties PointProperty = EPCGPointProperties::Density;
+	FPCGAttributePropertySelector OutputTarget;
 
-	UPROPERTY(BlueprintReadWrite, EditAnywhere, Category = Settings, meta = (EditCondition="Target==EPCGMetadataOperationTarget::PropertyToAttribute||Target==EPCGMetadataOperationTarget::AttributeToAttribute"))
-	FName DestinationAttribute = NAME_None;
+#if WITH_EDITORONLY_DATA
+	UPROPERTY()
+	FName SourceAttribute_DEPRECATED = NAME_None;
 
-	UPROPERTY(BlueprintReadWrite, EditAnywhere, Category = Settings)
-	EPCGMetadataOperationTarget Target = EPCGMetadataOperationTarget::PropertyToAttribute;
+	UPROPERTY()
+	EPCGPointProperties PointProperty_DEPRECATED = EPCGPointProperties::Density;
+
+	UPROPERTY()
+	FName DestinationAttribute_DEPRECATED = NAME_None;
+
+PRAGMA_DISABLE_DEPRECATION_WARNINGS
+	UPROPERTY()
+	EPCGMetadataOperationTarget Target_DEPRECATED = EPCGMetadataOperationTarget::PropertyToAttribute;
+PRAGMA_ENABLE_DEPRECATION_WARNINGS
+#endif // WITH_EDITORONLY_DATA
 };
 
 class FPCGMetadataOperationElement : public FSimplePCGElement
diff --git a/Engine/Plugins/Experimental/PCG/Source/PCG/Public/Elements/Metadata/PCGMetadataElementCommon.h b/Engine/Plugins/Experimental/PCG/Source/PCG/Public/Elements/Metadata/PCGMetadataElementCommon.h
--- a/Engine/Plugins/Experimental/PCG/Source/PCG/Public/Elements/Metadata/PCGMetadataElementCommon.h
+++ b/Engine/Plugins/Experimental/PCG/Source/PCG/Public/Elements/Metadata/PCGMetadataElementCommon.h
@@ -6,6 +6,9 @@
 #include "Metadata/PCGMetadata.h"
 #include "Metadata/PCGMetadataAttribute.h"
 #include "Metadata/PCGMetadataAttributeTpl.h"
+#include "Metadata/PCGMetadataAttributeTpl.h"
+#include "Metadata/Accessors/IPCGAttributeAccessor.h"
+#include "Metadata/Accessors/PCGAttributeAccessorKeys.h"
 
 struct FPCGTaggedData;
 
@@ -38,4 +41,47 @@ namespace PCGMetadataElementCommon
 
 		return static_cast<FPCGMetadataAttribute<T>*>(Metadata->GetMutableAttribute(DestinationAttribute));
 	}
+
+	constexpr int32 DefaultChunkSize = 256;
+
+	/**
+	* Iterate over the full range of the keys, calling the callback with values get from the accessor.
+	* ChunkSize influence the max number of values to get in one go with GetRange.
+	* Callback should have a signature: void(const T&, int32)
+	* Return false if it process nothing.
+	*/
+	template <typename T, typename Func>
+	bool ApplyOnAccessor(const IPCGAttributeAccessorKeys& Keys, const IPCGAttributeAccessor& Accessor, Func&& Callback, const int32 ChunkSize = DefaultChunkSize)
+	{
+		const int32 NumberOfEntries = Keys.GetNum();
+
+		if (NumberOfEntries == 0)
+		{
+			return false;
+		}
+
+		TArray<T, TInlineAllocator<DefaultChunkSize>> TempValues;
+		TempValues.SetNum(ChunkSize);
+
+		const int32 NumberOfIterations = (NumberOfEntries + ChunkSize - 1) / ChunkSize;
+
+		for (int32 i = 0; i < NumberOfIterations; ++i)
+		{
+			const int32 StartIndex = i * ChunkSize;
+			const int32 Range = FMath::Min(NumberOfEntries - StartIndex, ChunkSize);
+			TArrayView<T> View(TempValues.GetData(), Range);
+
+			if (!Accessor.GetRange(View, StartIndex, Keys))
+			{
+				return false;
+			}
+
+			for (int32 j = 0; j < Range; ++j)
+			{
+				Callback(TempValues[j], StartIndex + j);
+			}
+		}
+
+		return true;
+	}
 }
\ No newline at end of file
diff --git a/Engine/Plugins/Experimental/PCG/Source/PCG/Public/Elements/Metadata/PCGMetadataPartition.h b/Engine/Plugins/Experimental/PCG/Source/PCG/Public/Elements/Metadata/PCGMetadataPartition.h
--- a/Engine/Plugins/Experimental/PCG/Source/PCG/Public/Elements/Metadata/PCGMetadataPartition.h
+++ b/Engine/Plugins/Experimental/PCG/Source/PCG/Public/Elements/Metadata/PCGMetadataPartition.h
@@ -15,7 +15,7 @@ class PCG_API UPCGMetadataPartitionSettings : public UPCGSettings
 public:
 	//~Begin UPCGSettings interface
 #if WITH_EDITOR
-	virtual FName GetDefaultNodeName() const override { return FName(TEXT("MetadataPartition")); }
+	virtual FName GetDefaultNodeName() const override { return FName(TEXT("AttributePartition")); }
 	virtual EPCGSettingsType GetType() const override { return EPCGSettingsType::Metadata; }
 #endif
 
diff --git a/Engine/Plugins/Experimental/PCG/Source/PCG/Public/Elements/Metadata/PCGMetadataRenameElement.h b/Engine/Plugins/Experimental/PCG/Source/PCG/Public/Elements/Metadata/PCGMetadataRenameElement.h
--- a/Engine/Plugins/Experimental/PCG/Source/PCG/Public/Elements/Metadata/PCGMetadataRenameElement.h
+++ b/Engine/Plugins/Experimental/PCG/Source/PCG/Public/Elements/Metadata/PCGMetadataRenameElement.h
@@ -19,6 +19,7 @@ public:
 	virtual EPCGSettingsType GetType() const override { return EPCGSettingsType::Metadata; }
 #endif
 
+	virtual TArray<FPCGPinProperties> InputPinProperties() const override;
 	virtual TArray<FPCGPinProperties> OutputPinProperties() const override;
 
 protected:
@@ -35,9 +36,6 @@ public:
 
 class FPCGMetadataRenameElement : public FSimplePCGElement
 {
-public:
-	virtual bool CanExecuteOnlyOnMainThread(FPCGContext* Context) const override { return true; }
-
 protected:
 	virtual bool ExecuteInternal(FPCGContext* Context) const override;
 };
\ No newline at end of file
diff --git a/Engine/Plugins/Experimental/PCG/Source/PCG/Public/Elements/PCGAttributeFilterElement.h b/Engine/Plugins/Experimental/PCG/Source/PCG/Public/Elements/PCGAttributeFilterElement.h
--- a/Engine/Plugins/Experimental/PCG/Source/PCG/Public/Elements/PCGAttributeFilterElement.h
+++ b/Engine/Plugins/Experimental/PCG/Source/PCG/Public/Elements/PCGAttributeFilterElement.h
@@ -7,6 +7,13 @@
 
 #include "PCGAttributeFilterElement.generated.h"
 
+UENUM()
+enum class EPCGAttributeFilterOperation
+{
+	KeepSelectedAttributes,
+	DeleteSelectedAttributes
+};
+
 /**
 * Filter the attributes from a given input metadata.
 * Will remove all attributes that are not listed in AttributesToKeep.
@@ -20,6 +27,10 @@ class PCG_API UPCGAttributeFilterSettings : public UPCGSettings
 	GENERATED_BODY()
 
 public:
+	//~Begin UObject interface
+	virtual void PostLoad() override;
+	//~End UObject interface
+
 	//~Begin UPCGSettings interface
 #if WITH_EDITOR
 	virtual FName GetDefaultNodeName() const override;
@@ -31,7 +42,15 @@ public:
 	//~End UPCGSettings interface
 
 	UPROPERTY(BlueprintReadWrite, EditAnywhere, Category = Settings)
-	TArray<FName> AttributesToKeep;
+	EPCGAttributeFilterOperation Operation = EPCGAttributeFilterOperation::KeepSelectedAttributes;
+
+	UPROPERTY(BlueprintReadWrite, EditAnywhere, Category = Settings)
+	FString SelectedAttributes;
+
+#if WITH_EDITORONLY_DATA
+	UPROPERTY()
+	TArray<FName> AttributesToKeep_DEPRECATED;
+#endif // WITH_EDITORONLY_DATA
 
 protected:
 	virtual FPCGElementPtr CreateElement() const override;
diff --git a/Engine/Plugins/Experimental/PCG/Source/PCG/Public/Elements/PCGAttributeGetFromPointIndexElement.h b/Engine/Plugins/Experimental/PCG/Source/PCG/Public/Elements/PCGAttributeGetFromPointIndexElement.h
--- a/Engine/Plugins/Experimental/PCG/Source/PCG/Public/Elements/PCGAttributeGetFromPointIndexElement.h
+++ b/Engine/Plugins/Experimental/PCG/Source/PCG/Public/Elements/PCGAttributeGetFromPointIndexElement.h
@@ -4,6 +4,7 @@
 
 #include "PCGElement.h"
 #include "PCGSettings.h"
+#include "Metadata/PCGAttributePropertySelector.h"
 
 #include "PCGAttributeGetFromPointIndexElement.generated.h"
 
@@ -14,9 +15,9 @@ namespace PCGAttributeGetFromPointIndexConstants
 }
 
 /**
-* Get the attribute of a point given its index. The result will be in a ParamData.
+* Get the attribute/property of a point given its index. The result will be in a ParamData.
 * There is also a second output that will output the selected point. This point will be output
-* even if the attribute doesn't exist.
+* even if the property/attribute doesn't exist.
 * 
 * The Index can be overridden by a second Params input.
 */
@@ -26,6 +27,10 @@ class PCG_API UPCGAttributeGetFromPointIndexSettings : public UPCGSettings
 	GENERATED_BODY()
 
 public:
+	//~Begin UObject interface
+	virtual void PostLoad() override;
+	//~End UObject interface
+
 	//~Begin UPCGSettings interface
 #if WITH_EDITOR
 	virtual FName GetDefaultNodeName() const override;
@@ -37,11 +42,19 @@ public:
 	//~End UPCGSettings interface
 
 	UPROPERTY(BlueprintReadWrite, EditAnywhere, Category = Settings)
-	FName InputAttributeName = NAME_None;
+	FPCGAttributePropertySelector InputSource;
 
 	UPROPERTY(BlueprintReadWrite, EditAnywhere, Category = Settings)
 	int32 Index = 0;
 
+	UPROPERTY(BlueprintReadWrite, EditAnywhere, Category = Settings)
+	FName OutputAttributeName = NAME_None;
+
+#if WITH_EDITORONLY_DATA
+	UPROPERTY()
+	FName InputAttributeName_DEPRECATED = NAME_None;
+#endif // WITH_EDITORONLY_DATA
+
 protected:
 	virtual FPCGElementPtr CreateElement() const override;
 };
diff --git a/Engine/Plugins/Experimental/PCG/Source/PCG/Public/Elements/PCGAttributeReduceElement.h b/Engine/Plugins/Experimental/PCG/Source/PCG/Public/Elements/PCGAttributeReduceElement.h
--- a/Engine/Plugins/Experimental/PCG/Source/PCG/Public/Elements/PCGAttributeReduceElement.h
+++ b/Engine/Plugins/Experimental/PCG/Source/PCG/Public/Elements/PCGAttributeReduceElement.h
@@ -4,6 +4,7 @@
 
 #include "PCGElement.h"
 #include "PCGSettings.h"
+#include "Metadata/PCGAttributePropertySelector.h"
 
 #include "PCGAttributeReduceElement.generated.h"
 
@@ -16,10 +17,10 @@ enum class EPCGAttributeReduceOperation
 };
 
 /**
-* Take all the entries/points from the input and perform a reduce operation on the given attribute
+* Take all the entries/points from the input and perform a reduce operation on the given attribute/property
 * and output the result into a ParamData.
 * 
-* If the OutputAttributeName is None, we will use InputAttributeName
+* If the OutputAttributeName is None, we will use the InputSource name
 */
 UCLASS(BlueprintType, ClassGroup = (Procedural))
 class PCG_API UPCGAttributeReduceSettings : public UPCGSettings
@@ -27,6 +28,10 @@ class PCG_API UPCGAttributeReduceSettings : public UPCGSettings
 	GENERATED_BODY()
 
 public:
+	//~Begin UObject interface
+	virtual void PostLoad() override;
+	//~End UObject interface
+
 	//~Begin UPCGSettings interface
 #if WITH_EDITOR
 	virtual FName GetDefaultNodeName() const override;
@@ -40,7 +45,7 @@ public:
 	//~End UPCGSettings interface
 
 	UPROPERTY(BlueprintReadWrite, EditAnywhere, Category = Settings)
-	FName InputAttributeName = NAME_None;
+	FPCGAttributePropertySelector InputSource;
 
 	UPROPERTY(BlueprintReadWrite, EditAnywhere, Category = Settings)
 	FName OutputAttributeName = NAME_None;
@@ -48,6 +53,11 @@ public:
 	UPROPERTY(BlueprintReadWrite, EditAnywhere, Category = Settings)
 	EPCGAttributeReduceOperation Operation = EPCGAttributeReduceOperation::Average;
 
+#if WITH_EDITORONLY_DATA
+	UPROPERTY()
+	FName InputAttributeName_DEPRECATED = NAME_None;
+#endif // WITH_EDITORONLY_DATA
+
 protected:
 	virtual FPCGElementPtr CreateElement() const override;
 };
diff --git a/Engine/Plugins/Experimental/PCG/Source/PCG/Public/Elements/PCGAttributeSelectElement.h b/Engine/Plugins/Experimental/PCG/Source/PCG/Public/Elements/PCGAttributeSelectElement.h
--- a/Engine/Plugins/Experimental/PCG/Source/PCG/Public/Elements/PCGAttributeSelectElement.h
+++ b/Engine/Plugins/Experimental/PCG/Source/PCG/Public/Elements/PCGAttributeSelectElement.h
@@ -4,6 +4,7 @@
 
 #include "PCGElement.h"
 #include "PCGSettings.h"
+#include "Metadata/PCGAttributePropertySelector.h"
 
 #include "PCGAttributeSelectElement.generated.h"
 
@@ -32,8 +33,8 @@ enum class EPCGAttributeSelectAxis
 };
 
 /**
-* Take all the entries/points from the input and perform a select operation on the given attribute on the given axis 
-* (if the attribute is a vector) and output the result into a ParamData.
+* Take all the entries/points from the input and perform a select operation on the given attribute/property on the given axis 
+* (if the attribute/property is a vector) and output the result into a ParamData.
 * It will also output the selected point if the input is a PointData.
 * 
 * Only support vector attributes and scalar attributes.
@@ -42,7 +43,7 @@ enum class EPCGAttributeSelectAxis
 * 
 * In case of the median operation, and the number of elements is even, we arbitrarily chose a point (Index = Num / 2)
 *
-* If the OutputAttributeName is None, we will use InputAttributeName.
+* If the OutputAttributeName is None, we will use InputSource.GetName().
 */
 UCLASS(BlueprintType, ClassGroup = (Procedural))
 class PCG_API UPCGAttributeSelectSettings : public UPCGSettings
@@ -50,6 +51,10 @@ class PCG_API UPCGAttributeSelectSettings : public UPCGSettings
 	GENERATED_BODY()
 
 public:
+	//~Begin UObject interface
+	virtual void PostLoad() override;
+	//~End UObject interface
+
 	//~Begin UPCGSettings interface
 #if WITH_EDITOR
 	virtual FName GetDefaultNodeName() const override;
@@ -63,7 +68,7 @@ public:
 	//~End UPCGSettings interface
 
 	UPROPERTY(BlueprintReadWrite, EditAnywhere, Category = Settings)
-	FName InputAttributeName = NAME_None;
+	FPCGAttributePropertySelector InputSource;
 
 	UPROPERTY(BlueprintReadWrite, EditAnywhere, Category = Settings)
 	FName OutputAttributeName = NAME_None;
@@ -77,6 +82,11 @@ public:
 	UPROPERTY(BlueprintReadWrite, EditAnywhere, Category = Settings, meta = (EditCondition = "Axis == EPCGAttributeSelectAxis::CustomAxis", EditConditionHides))
 	FVector4 CustomAxis = FVector4::Zero();
 
+#if WITH_EDITORONLY_DATA
+	UPROPERTY()
+	FName InputAttributeName_DEPRECATED = NAME_None;
+#endif // WITH_EDITORONLY_DATA
+
 protected:
 	virtual FPCGElementPtr CreateElement() const override;
 };
