[PCG] Remove PCGMetadataEntryKeyIterator

#rnx

diff --git a/Engine/Plugins/Experimental/PCG/Source/PCG/Private/Elements/Metadata/PCGMetadataBitwiseOpElement.cpp b/Engine/Plugins/Experimental/PCG/Source/PCG/Private/Elements/Metadata/PCGMetadataBitwiseOpElement.cpp
--- a/Engine/Plugins/Experimental/PCG/Source/PCG/Private/Elements/Metadata/PCGMetadataBitwiseOpElement.cpp
+++ b/Engine/Plugins/Experimental/PCG/Source/PCG/Private/Elements/Metadata/PCGMetadataBitwiseOpElement.cpp
@@ -8,7 +8,6 @@
 #include "Metadata/PCGMetadata.h"
 #include "Metadata/PCGMetadataAttribute.h"
 #include "Metadata/PCGMetadataAttributeTpl.h"
-#include "Metadata/PCGMetadataEntryKeyIterator.h"
 #include "Elements/Metadata/PCGMetadataElementCommon.h"
 
 #include UE_INLINE_GENERATED_CPP_BY_NAME(PCGMetadataBitwiseOpElement)
@@ -44,15 +43,13 @@ void UPCGMetadataBitwiseSettings::PostLoad()
 #if WITH_EDITOR
 	if (Input1AttributeName_DEPRECATED != NAME_None)
 	{
-		InputSource1.Selection = EPCGAttributePropertySelection::Attribute;
-		InputSource1.AttributeName = Input1AttributeName_DEPRECATED;
+		InputSource1.SetAttributeName(Input1AttributeName_DEPRECATED);
 		Input1AttributeName_DEPRECATED = NAME_None;
 	}
 
 	if (Input2AttributeName_DEPRECATED != NAME_None)
 	{
-		InputSource2.Selection = EPCGAttributePropertySelection::Attribute;
-		InputSource2.AttributeName = Input2AttributeName_DEPRECATED;
+		InputSource2.SetAttributeName(Input2AttributeName_DEPRECATED);
 		Input2AttributeName_DEPRECATED = NAME_None;
 	}
 #endif // WITH_EDITOR
diff --git a/Engine/Plugins/Experimental/PCG/Source/PCG/Private/Elements/Metadata/PCGMetadataBooleanOpElement.cpp b/Engine/Plugins/Experimental/PCG/Source/PCG/Private/Elements/Metadata/PCGMetadataBooleanOpElement.cpp
--- a/Engine/Plugins/Experimental/PCG/Source/PCG/Private/Elements/Metadata/PCGMetadataBooleanOpElement.cpp
+++ b/Engine/Plugins/Experimental/PCG/Source/PCG/Private/Elements/Metadata/PCGMetadataBooleanOpElement.cpp
@@ -8,7 +8,6 @@
 #include "Metadata/PCGMetadata.h"
 #include "Metadata/PCGMetadataAttribute.h"
 #include "Metadata/PCGMetadataAttributeTpl.h"
-#include "Metadata/PCGMetadataEntryKeyIterator.h"
 #include "Elements/Metadata/PCGMetadataElementCommon.h"
 
 #include UE_INLINE_GENERATED_CPP_BY_NAME(PCGMetadataBooleanOpElement)
@@ -44,15 +43,13 @@ void UPCGMetadataBooleanSettings::PostLoad()
 #if WITH_EDITOR
 	if (Input1AttributeName_DEPRECATED != NAME_None)
 	{
-		InputSource1.Selection = EPCGAttributePropertySelection::Attribute;
-		InputSource1.AttributeName = Input1AttributeName_DEPRECATED;
+		InputSource1.SetAttributeName(Input1AttributeName_DEPRECATED);
 		Input1AttributeName_DEPRECATED = NAME_None;
 	}
 
 	if (Input2AttributeName_DEPRECATED != NAME_None)
 	{
-		InputSource2.Selection = EPCGAttributePropertySelection::Attribute;
-		InputSource2.AttributeName = Input2AttributeName_DEPRECATED;
+		InputSource2.SetAttributeName(Input2AttributeName_DEPRECATED);
 		Input2AttributeName_DEPRECATED = NAME_None;
 	}
 #endif // WITH_EDITOR
diff --git a/Engine/Plugins/Experimental/PCG/Source/PCG/Private/Elements/Metadata/PCGMetadataBreakTransform.cpp b/Engine/Plugins/Experimental/PCG/Source/PCG/Private/Elements/Metadata/PCGMetadataBreakTransform.cpp
--- a/Engine/Plugins/Experimental/PCG/Source/PCG/Private/Elements/Metadata/PCGMetadataBreakTransform.cpp
+++ b/Engine/Plugins/Experimental/PCG/Source/PCG/Private/Elements/Metadata/PCGMetadataBreakTransform.cpp
@@ -37,8 +37,7 @@ void UPCGMetadataBreakTransformSettings::PostLoad()
 #if WITH_EDITOR
 	if (InputAttributeName_DEPRECATED != NAME_None)
 	{
-		InputSource.Selection = EPCGAttributePropertySelection::Attribute;
-		InputSource.AttributeName = InputAttributeName_DEPRECATED;
+		InputSource.SetAttributeName(InputAttributeName_DEPRECATED);
 		InputAttributeName_DEPRECATED = NAME_None;
 	}
 #endif // WITH_EDITOR
diff --git a/Engine/Plugins/Experimental/PCG/Source/PCG/Private/Elements/Metadata/PCGMetadataBreakVector.cpp b/Engine/Plugins/Experimental/PCG/Source/PCG/Private/Elements/Metadata/PCGMetadataBreakVector.cpp
--- a/Engine/Plugins/Experimental/PCG/Source/PCG/Private/Elements/Metadata/PCGMetadataBreakVector.cpp
+++ b/Engine/Plugins/Experimental/PCG/Source/PCG/Private/Elements/Metadata/PCGMetadataBreakVector.cpp
@@ -87,8 +87,7 @@ void UPCGMetadataBreakVectorSettings::PostLoad()
 #if WITH_EDITOR
 	if (InputAttributeName_DEPRECATED != NAME_None)
 	{
-		InputSource.Selection = EPCGAttributePropertySelection::Attribute;
-		InputSource.AttributeName = InputAttributeName_DEPRECATED;
+		InputSource.SetAttributeName(InputAttributeName_DEPRECATED);
 		InputAttributeName_DEPRECATED = NAME_None;
 	}
 #endif // WITH_EDITOR
diff --git a/Engine/Plugins/Experimental/PCG/Source/PCG/Private/Elements/Metadata/PCGMetadataCompareOpElement.cpp b/Engine/Plugins/Experimental/PCG/Source/PCG/Private/Elements/Metadata/PCGMetadataCompareOpElement.cpp
--- a/Engine/Plugins/Experimental/PCG/Source/PCG/Private/Elements/Metadata/PCGMetadataCompareOpElement.cpp
+++ b/Engine/Plugins/Experimental/PCG/Source/PCG/Private/Elements/Metadata/PCGMetadataCompareOpElement.cpp
@@ -6,7 +6,6 @@
 #include "Metadata/PCGMetadata.h"
 #include "Metadata/PCGMetadataAttribute.h"
 #include "Metadata/PCGMetadataAttributeTpl.h"
-#include "Metadata/PCGMetadataEntryKeyIterator.h"
 #include "Elements/Metadata/PCGMetadataElementCommon.h"
 
 #include UE_INLINE_GENERATED_CPP_BY_NAME(PCGMetadataCompareOpElement)
@@ -61,15 +60,13 @@ void UPCGMetadataCompareSettings::PostLoad()
 #if WITH_EDITOR
 	if (Input1AttributeName_DEPRECATED != NAME_None)
 	{
-		InputSource1.Selection = EPCGAttributePropertySelection::Attribute;
-		InputSource1.AttributeName = Input1AttributeName_DEPRECATED;
+		InputSource1.SetAttributeName(Input1AttributeName_DEPRECATED);
 		Input1AttributeName_DEPRECATED = NAME_None;
 	}
 
 	if (Input2AttributeName_DEPRECATED != NAME_None)
 	{
-		InputSource2.Selection = EPCGAttributePropertySelection::Attribute;
-		InputSource2.AttributeName = Input2AttributeName_DEPRECATED;
+		InputSource2.SetAttributeName(Input2AttributeName_DEPRECATED);
 		Input2AttributeName_DEPRECATED = NAME_None;
 	}
 #endif // WITH_EDITOR
diff --git a/Engine/Plugins/Experimental/PCG/Source/PCG/Private/Elements/Metadata/PCGMetadataMakeTransform.cpp b/Engine/Plugins/Experimental/PCG/Source/PCG/Private/Elements/Metadata/PCGMetadataMakeTransform.cpp
--- a/Engine/Plugins/Experimental/PCG/Source/PCG/Private/Elements/Metadata/PCGMetadataMakeTransform.cpp
+++ b/Engine/Plugins/Experimental/PCG/Source/PCG/Private/Elements/Metadata/PCGMetadataMakeTransform.cpp
@@ -36,22 +36,19 @@ void UPCGMetadataMakeTransformSettings::PostLoad()
 #if WITH_EDITOR
 	if (Input1AttributeName_DEPRECATED != NAME_None)
 	{
-		InputSource1.Selection = EPCGAttributePropertySelection::Attribute;
-		InputSource1.AttributeName = Input1AttributeName_DEPRECATED;
+		InputSource1.SetAttributeName(Input1AttributeName_DEPRECATED);
 		Input1AttributeName_DEPRECATED = NAME_None;
 	}
 
 	if (Input2AttributeName_DEPRECATED != NAME_None)
 	{
-		InputSource2.Selection = EPCGAttributePropertySelection::Attribute;
-		InputSource2.AttributeName = Input2AttributeName_DEPRECATED;
+		InputSource2.SetAttributeName(Input2AttributeName_DEPRECATED);
 		Input2AttributeName_DEPRECATED = NAME_None;
 	}
 
 	if (Input3AttributeName_DEPRECATED != NAME_None)
 	{
-		InputSource3.Selection = EPCGAttributePropertySelection::Attribute;
-		InputSource3.AttributeName = Input3AttributeName_DEPRECATED;
+		InputSource3.SetAttributeName(Input3AttributeName_DEPRECATED);
 		Input3AttributeName_DEPRECATED = NAME_None;
 	}
 #endif // WITH_EDITOR
diff --git a/Engine/Plugins/Experimental/PCG/Source/PCG/Private/Elements/Metadata/PCGMetadataMakeVector.cpp b/Engine/Plugins/Experimental/PCG/Source/PCG/Private/Elements/Metadata/PCGMetadataMakeVector.cpp
--- a/Engine/Plugins/Experimental/PCG/Source/PCG/Private/Elements/Metadata/PCGMetadataMakeVector.cpp
+++ b/Engine/Plugins/Experimental/PCG/Source/PCG/Private/Elements/Metadata/PCGMetadataMakeVector.cpp
@@ -53,29 +53,25 @@ void UPCGMetadataMakeVectorSettings::PostLoad()
 #if WITH_EDITOR
 	if (Input1AttributeName_DEPRECATED != NAME_None)
 	{
-		InputSource1.Selection = EPCGAttributePropertySelection::Attribute;
-		InputSource1.AttributeName = Input1AttributeName_DEPRECATED;
+		InputSource1.SetAttributeName(Input1AttributeName_DEPRECATED);
 		Input1AttributeName_DEPRECATED = NAME_None;
 	}
 
 	if (Input2AttributeName_DEPRECATED != NAME_None)
 	{
-		InputSource2.Selection = EPCGAttributePropertySelection::Attribute;
-		InputSource2.AttributeName = Input2AttributeName_DEPRECATED;
+		InputSource2.SetAttributeName(Input2AttributeName_DEPRECATED);
 		Input2AttributeName_DEPRECATED = NAME_None;
 	}
 
 	if (Input3AttributeName_DEPRECATED != NAME_None)
 	{
-		InputSource3.Selection = EPCGAttributePropertySelection::Attribute;
-		InputSource3.AttributeName = Input3AttributeName_DEPRECATED;
+		InputSource3.SetAttributeName(Input3AttributeName_DEPRECATED);
 		Input3AttributeName_DEPRECATED = NAME_None;
 	}
 
 	if (Input4AttributeName_DEPRECATED != NAME_None)
 	{
-		InputSource4.Selection = EPCGAttributePropertySelection::Attribute;
-		InputSource4.AttributeName = Input4AttributeName_DEPRECATED;
+		InputSource4.SetAttributeName(Input4AttributeName_DEPRECATED);
 		Input4AttributeName_DEPRECATED = NAME_None;
 	}
 #endif // WITH_EDITOR
diff --git a/Engine/Plugins/Experimental/PCG/Source/PCG/Private/Elements/Metadata/PCGMetadataMathsOpElement.cpp b/Engine/Plugins/Experimental/PCG/Source/PCG/Private/Elements/Metadata/PCGMetadataMathsOpElement.cpp
--- a/Engine/Plugins/Experimental/PCG/Source/PCG/Private/Elements/Metadata/PCGMetadataMathsOpElement.cpp
+++ b/Engine/Plugins/Experimental/PCG/Source/PCG/Private/Elements/Metadata/PCGMetadataMathsOpElement.cpp
@@ -8,7 +8,6 @@
 #include "Metadata/PCGMetadata.h"
 #include "Metadata/PCGMetadataAttribute.h"
 #include "Metadata/PCGMetadataAttributeTpl.h"
-#include "Metadata/PCGMetadataEntryKeyIterator.h"
 #include "Elements/Metadata/PCGMetadataElementCommon.h"
 
 #include "Elements/Metadata/PCGMetadataMaths.inl"
@@ -162,22 +161,19 @@ void UPCGMetadataMathsSettings::PostLoad()
 #if WITH_EDITOR
 	if (Input1AttributeName_DEPRECATED != NAME_None)
 	{
-		InputSource1.Selection = EPCGAttributePropertySelection::Attribute;
-		InputSource1.AttributeName = Input1AttributeName_DEPRECATED;
+		InputSource1.SetAttributeName(Input1AttributeName_DEPRECATED);
 		Input1AttributeName_DEPRECATED = NAME_None;
 	}
 
 	if (Input2AttributeName_DEPRECATED != NAME_None)
 	{
-		InputSource2.Selection = EPCGAttributePropertySelection::Attribute;
-		InputSource2.AttributeName = Input2AttributeName_DEPRECATED;
+		InputSource2.SetAttributeName(Input2AttributeName_DEPRECATED);
 		Input2AttributeName_DEPRECATED = NAME_None;
 	}
 
 	if (Input3AttributeName_DEPRECATED != NAME_None)
 	{
-		InputSource3.Selection = EPCGAttributePropertySelection::Attribute;
-		InputSource3.AttributeName = Input3AttributeName_DEPRECATED;
+		InputSource3.SetAttributeName(Input3AttributeName_DEPRECATED);
 		Input3AttributeName_DEPRECATED = NAME_None;
 	}
 #endif // WITH_EDITOR
diff --git a/Engine/Plugins/Experimental/PCG/Source/PCG/Private/Elements/Metadata/PCGMetadataRotatorOpElement.cpp b/Engine/Plugins/Experimental/PCG/Source/PCG/Private/Elements/Metadata/PCGMetadataRotatorOpElement.cpp
--- a/Engine/Plugins/Experimental/PCG/Source/PCG/Private/Elements/Metadata/PCGMetadataRotatorOpElement.cpp
+++ b/Engine/Plugins/Experimental/PCG/Source/PCG/Private/Elements/Metadata/PCGMetadataRotatorOpElement.cpp
@@ -6,7 +6,6 @@
 #include "Metadata/PCGMetadata.h"
 #include "Metadata/PCGMetadataAttribute.h"
 #include "Metadata/PCGMetadataAttributeTpl.h"
-#include "Metadata/PCGMetadataEntryKeyIterator.h"
 #include "Elements/Metadata/PCGMetadataElementCommon.h"
 
 #include UE_INLINE_GENERATED_CPP_BY_NAME(PCGMetadataRotatorOpElement)
@@ -103,22 +102,19 @@ void UPCGMetadataRotatorSettings::PostLoad()
 #if WITH_EDITOR
 	if (Input1AttributeName_DEPRECATED != NAME_None)
 	{
-		InputSource1.Selection = EPCGAttributePropertySelection::Attribute;
-		InputSource1.AttributeName = Input1AttributeName_DEPRECATED;
+		InputSource1.SetAttributeName(Input1AttributeName_DEPRECATED);
 		Input1AttributeName_DEPRECATED = NAME_None;
 	}
 
 	if (Input2AttributeName_DEPRECATED != NAME_None)
 	{
-		InputSource2.Selection = EPCGAttributePropertySelection::Attribute;
-		InputSource2.AttributeName = Input2AttributeName_DEPRECATED;
+		InputSource2.SetAttributeName(Input2AttributeName_DEPRECATED);
 		Input2AttributeName_DEPRECATED = NAME_None;
 	}
 
 	if (Input3AttributeName_DEPRECATED != NAME_None)
 	{
-		InputSource3.Selection = EPCGAttributePropertySelection::Attribute;
-		InputSource3.AttributeName = Input3AttributeName_DEPRECATED;
+		InputSource3.SetAttributeName(Input3AttributeName_DEPRECATED);
 		Input3AttributeName_DEPRECATED = NAME_None;
 	}
 #endif // WITH_EDITOR
diff --git a/Engine/Plugins/Experimental/PCG/Source/PCG/Private/Elements/Metadata/PCGMetadataTransformOpElement.cpp b/Engine/Plugins/Experimental/PCG/Source/PCG/Private/Elements/Metadata/PCGMetadataTransformOpElement.cpp
--- a/Engine/Plugins/Experimental/PCG/Source/PCG/Private/Elements/Metadata/PCGMetadataTransformOpElement.cpp
+++ b/Engine/Plugins/Experimental/PCG/Source/PCG/Private/Elements/Metadata/PCGMetadataTransformOpElement.cpp
@@ -8,7 +8,6 @@
 #include "Metadata/PCGMetadata.h"
 #include "Metadata/PCGMetadataAttribute.h"
 #include "Metadata/PCGMetadataAttributeTpl.h"
-#include "Metadata/PCGMetadataEntryKeyIterator.h"
 
 #include "Math/DualQuat.h"
 
@@ -84,22 +83,19 @@ void UPCGMetadataTransformSettings::PostLoad()
 #if WITH_EDITOR
 	if (Input1AttributeName_DEPRECATED != NAME_None)
 	{
-		InputSource1.Selection = EPCGAttributePropertySelection::Attribute;
-		InputSource1.AttributeName = Input1AttributeName_DEPRECATED;
+		InputSource1.SetAttributeName(Input1AttributeName_DEPRECATED);
 		Input1AttributeName_DEPRECATED = NAME_None;
 	}
 
 	if (Input2AttributeName_DEPRECATED != NAME_None)
 	{
-		InputSource2.Selection = EPCGAttributePropertySelection::Attribute;
-		InputSource2.AttributeName = Input2AttributeName_DEPRECATED;
+		InputSource2.SetAttributeName(Input2AttributeName_DEPRECATED);
 		Input2AttributeName_DEPRECATED = NAME_None;
 	}
 
 	if (Input3AttributeName_DEPRECATED != NAME_None)
 	{
-		InputSource3.Selection = EPCGAttributePropertySelection::Attribute;
-		InputSource3.AttributeName = Input3AttributeName_DEPRECATED;
+		InputSource3.SetAttributeName(Input3AttributeName_DEPRECATED);
 		Input3AttributeName_DEPRECATED = NAME_None;
 	}
 #endif // WITH_EDITOR
diff --git a/Engine/Plugins/Experimental/PCG/Source/PCG/Private/Elements/Metadata/PCGMetadataTrigOpElement.cpp b/Engine/Plugins/Experimental/PCG/Source/PCG/Private/Elements/Metadata/PCGMetadataTrigOpElement.cpp
--- a/Engine/Plugins/Experimental/PCG/Source/PCG/Private/Elements/Metadata/PCGMetadataTrigOpElement.cpp
+++ b/Engine/Plugins/Experimental/PCG/Source/PCG/Private/Elements/Metadata/PCGMetadataTrigOpElement.cpp
@@ -6,7 +6,6 @@
 #include "Metadata/PCGMetadata.h"
 #include "Metadata/PCGMetadataAttribute.h"
 #include "Metadata/PCGMetadataAttributeTpl.h"
-#include "Metadata/PCGMetadataEntryKeyIterator.h"
 #include "Elements/Metadata/PCGMetadataElementCommon.h"
 
 #include UE_INLINE_GENERATED_CPP_BY_NAME(PCGMetadataTrigOpElement)
@@ -54,15 +53,13 @@ void UPCGMetadataTrigSettings::PostLoad()
 #if WITH_EDITOR
 	if (Input1AttributeName_DEPRECATED != NAME_None)
 	{
-		InputSource1.Selection = EPCGAttributePropertySelection::Attribute;
-		InputSource1.AttributeName = Input1AttributeName_DEPRECATED;
+		InputSource1.SetAttributeName(Input1AttributeName_DEPRECATED);
 		Input1AttributeName_DEPRECATED = NAME_None;
 	}
 
 	if (Input2AttributeName_DEPRECATED != NAME_None)
 	{
-		InputSource2.Selection = EPCGAttributePropertySelection::Attribute;
-		InputSource2.AttributeName = Input2AttributeName_DEPRECATED;
+		InputSource2.SetAttributeName(Input2AttributeName_DEPRECATED);
 		Input2AttributeName_DEPRECATED = NAME_None;
 	}
 #endif // WITH_EDITOR
diff --git a/Engine/Plugins/Experimental/PCG/Source/PCG/Private/Elements/Metadata/PCGMetadataVectorOpElement.cpp b/Engine/Plugins/Experimental/PCG/Source/PCG/Private/Elements/Metadata/PCGMetadataVectorOpElement.cpp
--- a/Engine/Plugins/Experimental/PCG/Source/PCG/Private/Elements/Metadata/PCGMetadataVectorOpElement.cpp
+++ b/Engine/Plugins/Experimental/PCG/Source/PCG/Private/Elements/Metadata/PCGMetadataVectorOpElement.cpp
@@ -8,7 +8,6 @@
 #include "Metadata/PCGMetadata.h"
 #include "Metadata/PCGMetadataAttribute.h"
 #include "Metadata/PCGMetadataAttributeTpl.h"
-#include "Metadata/PCGMetadataEntryKeyIterator.h"
 #include "Elements/Metadata/PCGMetadataElementCommon.h"
 
 #include UE_INLINE_GENERATED_CPP_BY_NAME(PCGMetadataVectorOpElement)
@@ -217,22 +216,19 @@ void UPCGMetadataVectorSettings::PostLoad()
 #if WITH_EDITOR
 	if (Input1AttributeName_DEPRECATED != NAME_None)
 	{
-		InputSource1.Selection = EPCGAttributePropertySelection::Attribute;
-		InputSource1.AttributeName = Input1AttributeName_DEPRECATED;
+		InputSource1.SetAttributeName(Input1AttributeName_DEPRECATED);
 		Input1AttributeName_DEPRECATED = NAME_None;
 	}
 
 	if (Input2AttributeName_DEPRECATED != NAME_None)
 	{
-		InputSource2.Selection = EPCGAttributePropertySelection::Attribute;
-		InputSource2.AttributeName = Input2AttributeName_DEPRECATED;
+		InputSource2.SetAttributeName(Input2AttributeName_DEPRECATED);
 		Input2AttributeName_DEPRECATED = NAME_None;
 	}
 
 	if (Input3AttributeName_DEPRECATED != NAME_None)
 	{
-		InputSource3.Selection = EPCGAttributePropertySelection::Attribute;
-		InputSource3.AttributeName = Input3AttributeName_DEPRECATED;
+		InputSource3.SetAttributeName(Input3AttributeName_DEPRECATED);
 		Input3AttributeName_DEPRECATED = NAME_None;
 	}
 #endif // WITH_EDITOR
diff --git a/Engine/Plugins/Experimental/PCG/Source/PCG/Private/Metadata/PCGMetadataEntryKeyIterator.cpp b/Engine/Plugins/Experimental/PCG/Source/PCG/Private/Metadata/PCGMetadataEntryKeyIterator.cpp
deleted file mode 100644
diff --git a/Engine/Plugins/Experimental/PCG/Source/PCG/Public/Metadata/PCGMetadataEntryKeyIterator.h b/Engine/Plugins/Experimental/PCG/Source/PCG/Public/Metadata/PCGMetadataEntryKeyIterator.h
deleted file mode 100644
