[PCG] Remove "KeepExistingAttributes" to create attribute node + make sure we don't add a lot of entries for nothing.

diff --git a/Engine/Plugins/Experimental/PCG/Source/PCG/Private/Elements/PCGCreateAttribute.cpp b/Engine/Plugins/Experimental/PCG/Source/PCG/Private/Elements/PCGCreateAttribute.cpp
--- a/Engine/Plugins/Experimental/PCG/Source/PCG/Private/Elements/PCGCreateAttribute.cpp
+++ b/Engine/Plugins/Experimental/PCG/Source/PCG/Private/Elements/PCGCreateAttribute.cpp
@@ -217,7 +217,7 @@ bool FPCGCreateAttributeElement::ExecuteInternal(FPCGContext* Context) const
 		if (const UPCGSpatialData* InputSpatialData = Cast<UPCGSpatialData>(InputData))
 		{
 			UPCGSpatialData* NewSpatialData = InputSpatialData->DuplicateData(/*bInitializeFromData=*/false);
-			NewSpatialData->InitializeFromData(InputSpatialData, /*InMetadataParentOverride=*/ nullptr, /*bInheritMetadata=*/true, /*bInheritMetadata=*/ Settings->bKeepExistingAttributes);
+			NewSpatialData->InitializeFromData(InputSpatialData, /*InMetadataParentOverride=*/ nullptr, /*bInheritMetadata=*/true);
 
 			OutputData = NewSpatialData;
 			Metadata = NewSpatialData->Metadata;
@@ -226,12 +226,12 @@ bool FPCGCreateAttributeElement::ExecuteInternal(FPCGContext* Context) const
 		{
 			// If we can reuse input data, it is safe to const_cast, as it was created by ourselves above.
 			UPCGParamData* NewParamData = bCanReuseInputData ? const_cast<UPCGParamData*>(InputParamData) : NewObject<UPCGParamData>();
-			NewParamData->Metadata->Initialize(bCanReuseInputData ? nullptr : InputParamData->Metadata, Settings->bKeepExistingAttributes);
+			NewParamData->Metadata->InitializeAsCopy(bCanReuseInputData ? nullptr : InputParamData->Metadata);
 
 			OutputData = NewParamData;
 			Metadata = NewParamData->Metadata;
 
-			// In case of param data, we want to add a new entry too
+			// In case of param data, we want to add a new entry too, if needed
 			bShouldAddNewEntry = true;
 		}
 		else
@@ -267,7 +267,9 @@ bool FPCGCreateAttributeElement::ExecuteInternal(FPCGContext* Context) const
 		// Add a new entry if it is a param data and not from source (because entries are already copied)
 		if (bShouldAddNewEntry && !Settings->bFromSourceParam)
 		{
-			PCGMetadataEntryKey EntryKey = Metadata->AddEntry();
+			// If the metadata is empty, we need to add a new entry, so set it to PCGInvalidEntryKey.
+			// Otherwise, use the entry key 0.
+			PCGMetadataEntryKey EntryKey = Metadata->GetItemCountForChild() == 0 ? PCGInvalidEntryKey : 0;
 			SetAttribute(Settings, Attribute, Metadata, EntryKey, nullptr);
 		}
 	}
diff --git a/Engine/Plugins/Experimental/PCG/Source/PCG/Public/Elements/PCGCreateAttribute.h b/Engine/Plugins/Experimental/PCG/Source/PCG/Public/Elements/PCGCreateAttribute.h
--- a/Engine/Plugins/Experimental/PCG/Source/PCG/Public/Elements/PCGCreateAttribute.h
+++ b/Engine/Plugins/Experimental/PCG/Source/PCG/Public/Elements/PCGCreateAttribute.h
@@ -37,9 +37,6 @@ public:
 	virtual FName AdditionalTaskName() const override;
 	//~End UPCGSettings interface
 
-	UPROPERTY(BlueprintReadWrite, EditAnywhere, Category = Settings)
-	bool bKeepExistingAttributes = false;
-
 	UPROPERTY(BlueprintReadWrite, EditAnywhere, Category = Settings)
 	FName OutputAttributeName = NAME_None;
 
@@ -92,6 +89,11 @@ public:
 	UPROPERTY(BlueprintReadWrite, EditAnywhere, Category = Settings, meta = (EditCondition = "Type == EPCGMetadataTypes::Name && !bFromSourceParam", EditConditionHides))
 	FName NameValue = NAME_None;
 
+#if WITH_EDITORONLY_DATA
+	UPROPERTY()
+	bool bKeepExistingAttributes_DEPRECATED = false;
+#endif // WITH_EDITORONLY_DATA
+
 protected:
 	virtual FPCGElementPtr CreateElement() const override;
 };
