[PCG] Add ProjectPoint function to spatial data which performs a projection.

Often SamplePoint can be expressed as a projection one way or another - a majority of the functionality is shared. Therefore data's that have received this new ProjectPoint function have had their SamplePoint refactored to use it.

This has the benefit of surfacing the fact that the landscape data does actually perform a projection in its SamplePoint (for now - will be addressed next).

UPCGProjectionData now calls ProjectPoint. To stop this from breaking loads of things, UPCGSpatialData has a crutch fallback that forwards the call to SamplePoint to retain behaviour for most data's. The exception is UPCGPointData where I've proceeded to actually purify SamplePoint, which will make cases which wanted projection via SamplePoint() instead of using UPCGProjectionData fail. I will message this.

UPCGProjectionData is updated to use the new project function, and can pass its projection params into it which is nice. Slightly oddly, the project function takes Transform and Bounds from the source point but not the rest of the info like Color or metadata, and therefore it only makes use of some of the projection params (bProjectPositions, bProjectScales, bProjectRotations) and not the rest. I feel this is odd but acceptable - I think project should only need Transform and Bounds as input rather than an entire point. If a caller has access to more data then it can decide which data to take using the other projection params which is exactly what UPCGProjectionData does.

Some decisions will need to be made about how projection is implemented in some of the data's before we can move the crutch fallback. TODO's are added to mark some non-trivial cases.

diff --git a/Engine/Plugins/Experimental/PCG/Source/PCG/Private/Data/PCGLandscapeData.cpp b/Engine/Plugins/Experimental/PCG/Source/PCG/Private/Data/PCGLandscapeData.cpp
--- a/Engine/Plugins/Experimental/PCG/Source/PCG/Private/Data/PCGLandscapeData.cpp
+++ b/Engine/Plugins/Experimental/PCG/Source/PCG/Private/Data/PCGLandscapeData.cpp
@@ -106,6 +106,14 @@ FBox UPCGLandscapeData::GetStrictBounds() const
 }
 
 bool UPCGLandscapeData::SamplePoint(const FTransform& InTransform, const FBox& InBounds, FPCGPoint& OutPoint, UPCGMetadata* OutMetadata) const
+{
+	// SamplePoint actually does a projection for this data, so lets make it explicit and delegate to a project function.
+	// TODO fixup samplers and other things that rely on this behavior and then fix this code.
+	FPCGProjectionParams Params{};
+	return ProjectPoint(InTransform, InBounds, Params, OutPoint, OutMetadata);
+}
+
+bool UPCGLandscapeData::ProjectPoint(const FTransform& InTransform, const FBox& InBounds, const FPCGProjectionParams& InParams, FPCGPoint& OutPoint, UPCGMetadata* OutMetadata) const
 {
 	//TRACE_CPUPROFILER_EVENT_SCOPE(UPCGLandscapeData::SamplePoint);
 	if (!LandscapeCache)
@@ -145,6 +153,22 @@ bool UPCGLandscapeData::SamplePoint(const FTransform& InTransform, const FBox& I
 			LandscapeCacheEntry->GetInterpolatedPoint(ComponentLocalPoint, OutPoint, bUseMetadata ? OutMetadata : nullptr);
 		}
 
+		// Respect projection settings
+		if (!InParams.bProjectPositions)
+		{
+			OutPoint.Transform.SetLocation(InTransform.GetLocation());
+		}
+		
+		if (!InParams.bProjectRotations)
+		{
+			OutPoint.Transform.SetRotation(InTransform.GetRotation());
+		}
+
+		if (!InParams.bProjectScales)
+		{
+			OutPoint.Transform.SetScale3D(InTransform.GetScale3D());
+		}
+
 		return true;
 	}
 	else
diff --git a/Engine/Plugins/Experimental/PCG/Source/PCG/Private/Data/PCGLandscapeSplineData.cpp b/Engine/Plugins/Experimental/PCG/Source/PCG/Private/Data/PCGLandscapeSplineData.cpp
--- a/Engine/Plugins/Experimental/PCG/Source/PCG/Private/Data/PCGLandscapeSplineData.cpp
+++ b/Engine/Plugins/Experimental/PCG/Source/PCG/Private/Data/PCGLandscapeSplineData.cpp
@@ -165,6 +165,10 @@ bool UPCGLandscapeSplineData::SamplePoint(const FTransform& InTransform, const F
 {
 	// TODO : add metadata support on poly lines
 	// TODO : add support for bounds
+
+	// This does not move the query point, but does not take into account the Z axis at all - so this is inherently a projection. There's some
+	// things that need double checking but we should look at incorporating the Z axis into the calculation.
+
 	const ULandscapeSplinesComponent* SplinePtr = Spline.Get();
 	check(SplinePtr);
 
@@ -229,26 +233,3 @@ UPCGSpatialData* UPCGLandscapeSplineData::CopyInternal() const
 
 	return NewLandscapeSplineData;
 }
-
-UPCGProjectionData* UPCGLandscapeSplineData::ProjectOn(const UPCGSpatialData* InOther, const FPCGProjectionParams& InParams) const
-{
-	if (InOther->GetDimension() == 2)
-	{
-		UPCGLandscapeSplineProjectionData* SplineProjectionData = NewObject<UPCGLandscapeSplineProjectionData>();
-		SplineProjectionData->Initialize(this, InOther, InParams);
-		return SplineProjectionData;
-	}
-	else
-	{
-		return Super::ProjectOn(InOther, InParams);
-	}
-}
-
-UPCGSpatialData* UPCGLandscapeSplineProjectionData::CopyInternal() const
-{
-	UPCGLandscapeSplineProjectionData* NewProjectionData = NewObject<UPCGLandscapeSplineProjectionData>();
-
-	CopyBaseProjectionClass(NewProjectionData);
-
-	return NewProjectionData;
-}
diff --git a/Engine/Plugins/Experimental/PCG/Source/PCG/Private/Data/PCGPointData.cpp b/Engine/Plugins/Experimental/PCG/Source/PCG/Private/Data/PCGPointData.cpp
--- a/Engine/Plugins/Experimental/PCG/Source/PCG/Private/Data/PCGPointData.cpp
+++ b/Engine/Plugins/Experimental/PCG/Source/PCG/Private/Data/PCGPointData.cpp
@@ -326,6 +326,16 @@ FPCGPoint UPCGPointData::GetPoint(int32 Index) const
 }
 
 bool UPCGPointData::SamplePoint(const FTransform& InTransform, const FBox& InBounds, FPCGPoint& OutPoint, UPCGMetadata* OutMetadata) const
+{
+	// Run a projection but don't change the point transform. There is a large overlap in code/functionality so this shares one code path.
+	FPCGProjectionParams Params{};
+	Params.bProjectPositions = Params.bProjectRotations = Params.bProjectScales = Params.bProjectColors = false;
+
+	// The ProjectPoint implementation in this class returns true if the query point is overlapping the point data, which is what SamplePoint should return, so forward the return value.
+	return ProjectPoint(InTransform, InBounds, Params, OutPoint, OutMetadata);
+}
+
+bool UPCGPointData::ProjectPoint(const FTransform& InTransform, const FBox& InBounds, const FPCGProjectionParams& InParams, FPCGPoint& OutPoint, UPCGMetadata* OutMetadata) const
 {
 	//TRACE_CPUPROFILER_EVENT_SCOPE(UPCGPointData::SamplePoint);
 	if (bOctreeIsDirty)
@@ -414,11 +424,35 @@ bool UPCGPointData::SamplePoint(const FTransform& InTransform, const FBox& InBou
 	}
 
 	// Finally, apply changes to point
-	WeightedQuat.Normalize();
+	
+	if (InParams.bProjectPositions)
+	{
+		OutPoint.Transform.SetLocation(bSampleInVolume ? WeightedPosition : InTransform.GetLocation());
+	}
+	else
+	{
+		OutPoint.Transform.SetLocation(InTransform.GetLocation());
+	}
+
+	if (InParams.bProjectRotations)
+	{
+		WeightedQuat.Normalize();
+		OutPoint.Transform.SetRotation(WeightedQuat);
+	}
+	else
+	{
+		OutPoint.Transform.SetRotation(InTransform.GetRotation());
+	}
+
+	if (InParams.bProjectScales)
+	{
+		OutPoint.Transform.SetScale3D(WeightedScale);
+	}
+	else
+	{
+		OutPoint.Transform.SetScale3D(InTransform.GetScale3D());
+	}
 
-	OutPoint.Transform.SetRotation(WeightedQuat);
-	OutPoint.Transform.SetScale3D(WeightedScale);
-	OutPoint.Transform.SetLocation(bSampleInVolume ? WeightedPosition : InTransform.GetLocation());
 	OutPoint.Density = WeightedDensity;
 	OutPoint.BoundsMin = WeightedBoundsMin;
 	OutPoint.BoundsMax = WeightedBoundsMax;
diff --git a/Engine/Plugins/Experimental/PCG/Source/PCG/Private/Data/PCGPrimitiveData.cpp b/Engine/Plugins/Experimental/PCG/Source/PCG/Private/Data/PCGPrimitiveData.cpp
--- a/Engine/Plugins/Experimental/PCG/Source/PCG/Private/Data/PCGPrimitiveData.cpp
+++ b/Engine/Plugins/Experimental/PCG/Source/PCG/Private/Data/PCGPrimitiveData.cpp
@@ -19,6 +19,8 @@ void UPCGPrimitiveData::Initialize(UPrimitiveComponent* InPrimitive)
 
 bool UPCGPrimitiveData::SamplePoint(const FTransform& InTransform, const FBox& InBounds, FPCGPoint& OutPoint, UPCGMetadata* OutMetadata) const
 {
+	// Pure implementation
+
 	FCollisionShape CollisionShape;
 	CollisionShape.SetBox(FVector3f(InBounds.GetExtent() * InTransform.GetScale3D()));
 
diff --git a/Engine/Plugins/Experimental/PCG/Source/PCG/Private/Data/PCGProjectionData.cpp b/Engine/Plugins/Experimental/PCG/Source/PCG/Private/Data/PCGProjectionData.cpp
--- a/Engine/Plugins/Experimental/PCG/Source/PCG/Private/Data/PCGProjectionData.cpp
+++ b/Engine/Plugins/Experimental/PCG/Source/PCG/Private/Data/PCGProjectionData.cpp
@@ -1,6 +1,9 @@
 // Copyright Epic Games, Inc. All Rights Reserved.
 
 #include "Data/PCGProjectionData.h"
+
+#include "Data/PCGLandscapeSplineData.h"
+#include "Data/PCGSplineData.h"
 #include "Helpers/PCGAsync.h"
 #include "Metadata/PCGMetadataAccessor.h"
 
@@ -144,8 +147,22 @@ bool UPCGProjectionData::RequiresCollapseToSample() const
 	// intersected with a volume) rather than only checking if the immediate projection source is a particular type. A concrete example of this
 	// failing would be the In and Actor graph input pins which can be backed by composite networks.
 	
-	// It is however trivial if we are not actually moving anything around..
-	return ProjectionParams.bProjectPositions;
+	// Sampling is trivial if we are not actually moving anything around..
+	bool bRequiresCollapse = ProjectionParams.bProjectPositions;
+
+	// Projection of a spline onto a surface is currently easy to sample - don't need a collapse. If projection direction exposed, will need to check that.
+	if (Cast<UPCGSplineData>(Source) != nullptr && Target->GetDimension() == 2)
+	{
+		bRequiresCollapse = false;
+	}
+
+	// Projection of a landscape spline onto a surface is currently easy to sample - don't need a collapse. If projection direction exposed, will need to check that.
+	if (Cast<UPCGLandscapeSplineData>(Source) != nullptr && Target->GetDimension() == 2)
+	{
+		bRequiresCollapse = false;
+	}
+	
+	return bRequiresCollapse;
 }
 
 const UPCGPointData* UPCGProjectionData::CreatePointData(FPCGContext* Context) const
@@ -194,9 +211,9 @@ const UPCGPointData* UPCGProjectionData::CreatePointData(FPCGContext* Context) c
 
 		FPCGPoint PointFromTarget;
 #if WITH_EDITORONLY_DATA
-		if (!Target->SamplePoint(SourcePoint.Transform, SourcePoint.GetLocalBounds(), PointFromTarget, TempTargetMetadata) && !bKeepZeroDensityPoints)
+		if (!Target->ProjectPoint(SourcePoint.Transform, SourcePoint.GetLocalBounds(), ProjectionParams, PointFromTarget, TempTargetMetadata) && !bKeepZeroDensityPoints)
 #else
-		if (!Target->SamplePoint(SourcePoint.Transform, SourcePoint.GetLocalBounds(), PointFromTarget, TempTargetMetadata))
+		if (!Target->ProjectPoint(SourcePoint.Transform, SourcePoint.GetLocalBounds(), ProjectionParams, PointFromTarget, TempTargetMetadata))
 #endif
 		{
 			return false;
@@ -207,8 +224,12 @@ const UPCGPointData* UPCGProjectionData::CreatePointData(FPCGContext* Context) c
 		// Assign metadata key and copy values
 		UPCGMetadataAccessorHelpers::InitializeMetadata(OutPoint, PointData->Metadata, SourcePoint);
 
-		ApplyProjectionResult(PointFromTarget, OutPoint);
-
+		// Apply projection result. Some of the params are already used inside ProjectPoint, so we are just applying the remaining bits that ProjectPoint() did not have access to.
+		// TODO this would be cleaner if there was a ProjectPoint that took an FPCGPoint
+		OutPoint.Transform = PointFromTarget.Transform;
+		OutPoint.Color = ProjectionParams.bProjectColors ? (PointFromTarget.Color * SourcePoint.Color) : SourcePoint.Color;
+		OutPoint.Density *= PointFromTarget.Density;
+		
 		if (PointData->Metadata && TempTargetMetadata && PointFromTarget.MetadataEntry != PCGInvalidEntryKey)
 		{
 			// Merge metadata to produce final attribute values
diff --git a/Engine/Plugins/Experimental/PCG/Source/PCG/Private/Data/PCGSpatialData.cpp b/Engine/Plugins/Experimental/PCG/Source/PCG/Private/Data/PCGSpatialData.cpp
--- a/Engine/Plugins/Experimental/PCG/Source/PCG/Private/Data/PCGSpatialData.cpp
+++ b/Engine/Plugins/Experimental/PCG/Source/PCG/Private/Data/PCGSpatialData.cpp
@@ -84,6 +84,34 @@ FVector UPCGSpatialData::TransformPosition(const FVector& InPosition) const
 	}
 }
 
+bool UPCGSpatialData::ProjectPoint(const FTransform& InTransform, const FBox& InBounds, const FPCGProjectionParams& InParams, FPCGPoint& OutPoint, UPCGMetadata* OutMetadata) const
+{
+	// Fallback implementation - calls SamplePoint because SamplePoint was being used for projection previously.
+	
+	// TODO This is a crutch until we implement ProjectPoint everywhere
+
+	const bool bResult = SamplePoint(InTransform, InBounds, OutPoint, OutMetadata);
+
+	// Respect the projection params that we can at this point given our available data (InTransform)
+
+	if (!InParams.bProjectPositions)
+	{
+		OutPoint.Transform.SetLocation(InTransform.GetLocation());
+	}
+
+	if (!InParams.bProjectRotations)
+	{
+		OutPoint.Transform.SetRotation(InTransform.GetRotation());
+	}
+
+	if (!InParams.bProjectScales)
+	{
+		OutPoint.Transform.SetScale3D(InTransform.GetScale3D());
+	}
+
+	return bResult;
+}
+
 UPCGIntersectionData* UPCGSpatialData::IntersectWith(const UPCGSpatialData* InOther) const
 {
 	UPCGIntersectionData* IntersectionData = NewObject<UPCGIntersectionData>();
diff --git a/Engine/Plugins/Experimental/PCG/Source/PCG/Private/Data/PCGSplineData.cpp b/Engine/Plugins/Experimental/PCG/Source/PCG/Private/Data/PCGSplineData.cpp
--- a/Engine/Plugins/Experimental/PCG/Source/PCG/Private/Data/PCGSplineData.cpp
+++ b/Engine/Plugins/Experimental/PCG/Source/PCG/Private/Data/PCGSplineData.cpp
@@ -82,6 +82,8 @@ bool UPCGSplineData::SamplePoint(const FTransform& InTransform, const FBox& InBo
 	// TODO: support metadata
 	// TODO: support proper bounds
 
+	// This is a pure SamplePoint implementation.
+	
 	// Find nearest point on spline
 	const FVector InPosition = InTransform.GetLocation();
 	float NearestPointKey = Spline->FindInputKeyClosestToWorldLocation(InPosition);
diff --git a/Engine/Plugins/Experimental/PCG/Source/PCG/Private/Data/PCGTextureData.cpp b/Engine/Plugins/Experimental/PCG/Source/PCG/Private/Data/PCGTextureData.cpp
--- a/Engine/Plugins/Experimental/PCG/Source/PCG/Private/Data/PCGTextureData.cpp
+++ b/Engine/Plugins/Experimental/PCG/Source/PCG/Private/Data/PCGTextureData.cpp
@@ -132,6 +132,12 @@ bool UPCGBaseTextureData::SamplePoint(const FTransform& InTransform, const FBox&
 {
 	// TODO: add metadata support
 	// TODO: add sampling along the bounds
+
+	// TODO: needs unpicking of sample vs projection. I believe the below is a projection.. But semantics are slightly different.
+	// 1 - We have some information telling us the 'z' size of the surface allowing us to reject points that would be too far from the surface, maybe including some density falloff by distance
+	// 2 - We suppose that the surface has an infinite 'z' size, in which case the sampling is basically the same as the sampling, except that it does not change the position
+	// 3 - The surface is infinitesimal - we'll return something if and only if the point overlaps with the projected position
+
 	if (!IsValid())
 	{
 		return false;
diff --git a/Engine/Plugins/Experimental/PCG/Source/PCG/Private/Data/PCGVolumeData.cpp b/Engine/Plugins/Experimental/PCG/Source/PCG/Private/Data/PCGVolumeData.cpp
--- a/Engine/Plugins/Experimental/PCG/Source/PCG/Private/Data/PCGVolumeData.cpp
+++ b/Engine/Plugins/Experimental/PCG/Source/PCG/Private/Data/PCGVolumeData.cpp
@@ -62,6 +62,9 @@ bool UPCGVolumeData::SamplePoint(const FTransform& InTransform, const FBox& InBo
 	//TRACE_CPUPROFILER_EVENT_SCOPE(UPCGVolumeData::SamplePoint);
 	// TODO: add metadata
 	// TODO: consider bounds
+
+	// This is a pure implementation
+
 	const FVector InPosition = InTransform.GetLocation();
 	if (PCGHelpers::IsInsideBounds(GetBounds(), InPosition))
 	{
diff --git a/Engine/Plugins/Experimental/PCG/Source/PCG/Private/Data/PCGWorldData.cpp b/Engine/Plugins/Experimental/PCG/Source/PCG/Private/Data/PCGWorldData.cpp
--- a/Engine/Plugins/Experimental/PCG/Source/PCG/Private/Data/PCGWorldData.cpp
+++ b/Engine/Plugins/Experimental/PCG/Source/PCG/Private/Data/PCGWorldData.cpp
@@ -24,6 +24,8 @@ void UPCGWorldVolumetricData::Initialize(UWorld* InWorld, const FBox& InBounds)
 
 bool UPCGWorldVolumetricData::SamplePoint(const FTransform& InTransform, const FBox& InBounds, FPCGPoint& OutPoint, UPCGMetadata* OutMetadata) const
 {
+	// This is a pure implementation.
+
 	check(World.IsValid());
 
 	FCollisionObjectQueryParams ObjectQueryParams(QueryParams.CollisionChannel);
@@ -152,6 +154,8 @@ void UPCGWorldRayHitData::Initialize(UWorld* InWorld, const FBox& InBounds)
 
 bool UPCGWorldRayHitData::SamplePoint(const FTransform& InTransform, const FBox& InBounds, FPCGPoint& OutPoint, UPCGMetadata* OutMetadata) const
 {
+	// TODO: This seems to be a projection - along a direction. I suspect that UPCGWorldVolumetricData is the SamplePoint(), and this is the ProjectPoint() (in a direction)?
+
 	check(World.IsValid());
 
 	// Todo: consider prebuilding this
diff --git a/Engine/Plugins/Experimental/PCG/Source/PCG/Public/Data/PCGLandscapeData.h b/Engine/Plugins/Experimental/PCG/Source/PCG/Public/Data/PCGLandscapeData.h
--- a/Engine/Plugins/Experimental/PCG/Source/PCG/Public/Data/PCGLandscapeData.h
+++ b/Engine/Plugins/Experimental/PCG/Source/PCG/Public/Data/PCGLandscapeData.h
@@ -34,6 +34,7 @@ public:
 	virtual FBox GetBounds() const override;
 	virtual FBox GetStrictBounds() const override;
 	virtual bool SamplePoint(const FTransform& Transform, const FBox& Bounds, FPCGPoint& OutPoint, UPCGMetadata* OutMetadata) const override;
+	virtual bool ProjectPoint(const FTransform& InTransform, const FBox& InBounds, const FPCGProjectionParams& InParams, FPCGPoint& OutPoint, UPCGMetadata* OutMetadata) const override;
 	virtual bool HasNonTrivialTransform() const override { return true; }
 protected:
 	virtual UPCGSpatialData* CopyInternal() const override;
diff --git a/Engine/Plugins/Experimental/PCG/Source/PCG/Public/Data/PCGLandscapeSplineData.h b/Engine/Plugins/Experimental/PCG/Source/PCG/Public/Data/PCGLandscapeSplineData.h
--- a/Engine/Plugins/Experimental/PCG/Source/PCG/Public/Data/PCGLandscapeSplineData.h
+++ b/Engine/Plugins/Experimental/PCG/Source/PCG/Public/Data/PCGLandscapeSplineData.h
@@ -34,7 +34,6 @@ public:
 	//~Begin UPCGSpatialData interface
 	virtual FBox GetBounds() const override;
 	virtual bool SamplePoint(const FTransform& Transform, const FBox& Bounds, FPCGPoint& OutPoint, UPCGMetadata* OutMetadata) const override;
-	virtual UPCGProjectionData* ProjectOn(const UPCGSpatialData* InOther, const FPCGProjectionParams& InParams = FPCGProjectionParams()) const override;
 protected:
 	virtual UPCGSpatialData* CopyInternal() const override;
 	//~End UPCGSpatialData interface
@@ -43,22 +42,3 @@ protected:
 	UPROPERTY(BlueprintReadWrite, EditAnywhere, Category = SourceData)
 	TWeakObjectPtr<ULandscapeSplinesComponent> Spline;
 };
-
-/* The projection of a landscape spline onto a surface. */
-UCLASS(BlueprintType, ClassGroup = (Procedural))
-class PCG_API UPCGLandscapeSplineProjectionData : public UPCGProjectionData
-{
-	GENERATED_BODY()
-public:
-	//~Begin UPCGSpatialData interface
-	// It is easy to determine if a point lies on a top-down projection onto a surface. No need to convert to points.
-	// NOTE: It will be less easy if the projection is not straight downwards, as the landscape will 'shadow' the projection. This could be detected
-	// here.
-	virtual bool RequiresCollapseToSample() const override { return false; }
-	//~End UPCGSpatialData interface
-
-protected:
-	//~Begin UPCGSpatialData interface
-	virtual UPCGSpatialData* CopyInternal() const override;
-	//~End UPCGSpatialData interface
-};
diff --git a/Engine/Plugins/Experimental/PCG/Source/PCG/Public/Data/PCGPointData.h b/Engine/Plugins/Experimental/PCG/Source/PCG/Public/Data/PCGPointData.h
--- a/Engine/Plugins/Experimental/PCG/Source/PCG/Public/Data/PCGPointData.h
+++ b/Engine/Plugins/Experimental/PCG/Source/PCG/Public/Data/PCGPointData.h
@@ -75,6 +75,7 @@ public:
 	virtual FBox GetBounds() const override;
 	virtual const UPCGPointData* ToPointData(FPCGContext* Context, const FBox& InBounds = FBox(EForceInit::ForceInit)) const { return this; }
 	virtual bool SamplePoint(const FTransform& Transform, const FBox& Bounds, FPCGPoint& OutPoint, UPCGMetadata* OutMetadata) const override;
+	virtual bool ProjectPoint(const FTransform& InTransform, const FBox& InBounds, const FPCGProjectionParams& InParams, FPCGPoint& OutPoint, UPCGMetadata* OutMetadata) const override;
 protected:
 	virtual UPCGSpatialData* CopyInternal() const override;
 	// ~End UPCGSpatialData interface
diff --git a/Engine/Plugins/Experimental/PCG/Source/PCG/Public/Data/PCGPrimitiveData.h b/Engine/Plugins/Experimental/PCG/Source/PCG/Public/Data/PCGPrimitiveData.h
--- a/Engine/Plugins/Experimental/PCG/Source/PCG/Public/Data/PCGPrimitiveData.h
+++ b/Engine/Plugins/Experimental/PCG/Source/PCG/Public/Data/PCGPrimitiveData.h
@@ -23,6 +23,8 @@ public:
 	virtual FBox GetBounds() const override { return CachedBounds; }
 	virtual FBox GetStrictBounds() const override { return CachedStrictBounds; }
 	virtual bool SamplePoint(const FTransform& Transform, const FBox& Bounds, FPCGPoint& OutPoint, UPCGMetadata* OutMetadata) const override;
+	// TODO needs an implementation to support projection
+	//virtual bool ProjectPoint(const FTransform& InTransform, const FBox& InBounds, const FPCGProjectionParams& InParams, FPCGPoint& OutPoint, UPCGMetadata* OutMetadata) const;
 protected:
 	virtual UPCGSpatialData* CopyInternal() const override;
 	//~End UPCGSpatialData interface
diff --git a/Engine/Plugins/Experimental/PCG/Source/PCG/Public/Data/PCGSpatialData.h b/Engine/Plugins/Experimental/PCG/Source/PCG/Public/Data/PCGSpatialData.h
--- a/Engine/Plugins/Experimental/PCG/Source/PCG/Public/Data/PCGSpatialData.h
+++ b/Engine/Plugins/Experimental/PCG/Source/PCG/Public/Data/PCGSpatialData.h
@@ -82,9 +82,14 @@ public:
 	/** Transform a world-space position to a world-space position in relation to the current data. (ex: projection on surface) */
 	FVector TransformPosition(const FVector& InPosition) const;
 
+	/** Sample information from this data at the query. Generally OutPoint should be located at the query position, however some data such as Landscape currently also projects the point. Returns true if Transform location and Bounds overlaps this data. */
 	UFUNCTION(BlueprintCallable, Category = SpatialData)
 	virtual bool SamplePoint(const FTransform& Transform, const FBox& Bounds, FPCGPoint& OutPoint, UPCGMetadata* OutMetadata) const PURE_VIRTUAL(UPCGSpatialData::SamplePoint, return false;);
 
+	/** Project the query point onto this data, and sample point and metadata information at the projected position. Returns true if successful. */
+	UFUNCTION(BlueprintCallable, Category = SpatialData)
+	virtual bool ProjectPoint(const FTransform& InTransform, const FBox& InBounds, const FPCGProjectionParams& InParams, FPCGPoint& OutPoint, UPCGMetadata* OutMetadata) const;
+
 	/** Returns true if the data has a non-trivial transform */
 	UFUNCTION(BlueprintCallable, Category = SpatialData)
 	virtual bool HasNonTrivialTransform() const { return false; }
diff --git a/Engine/Plugins/Experimental/PCG/Source/PCG/Public/Data/PCGSplineData.h b/Engine/Plugins/Experimental/PCG/Source/PCG/Public/Data/PCGSplineData.h
--- a/Engine/Plugins/Experimental/PCG/Source/PCG/Public/Data/PCGSplineData.h
+++ b/Engine/Plugins/Experimental/PCG/Source/PCG/Public/Data/PCGSplineData.h
@@ -66,11 +66,6 @@ public:
 
 	//~Begin UPCGSpatialData interface
 	virtual bool SamplePoint(const FTransform& Transform, const FBox& Bounds, FPCGPoint& OutPoint, UPCGMetadata* OutMetadata) const override;
-
-	// It is easy to determine if a point lies on a top-down projection onto a surface. No need to convert to points.
-	// NOTE: It will be less easy if the projection is not straight downwards, as the landscape will 'shadow' the projection. This could be detected
-	// here.
-	virtual bool RequiresCollapseToSample() const override { return false; }
 	//~End UPCGSpatialData interface
 
 	UPROPERTY(BlueprintReadWrite, EditAnywhere, Category = SpatialData)
diff --git a/Engine/Plugins/Experimental/PCG/Source/PCG/Public/Data/PCGVolumeData.h b/Engine/Plugins/Experimental/PCG/Source/PCG/Public/Data/PCGVolumeData.h
--- a/Engine/Plugins/Experimental/PCG/Source/PCG/Public/Data/PCGVolumeData.h
+++ b/Engine/Plugins/Experimental/PCG/Source/PCG/Public/Data/PCGVolumeData.h
@@ -27,6 +27,8 @@ public:
 	virtual FBox GetBounds() const override;
 	virtual FBox GetStrictBounds() const override;
 	virtual bool SamplePoint(const FTransform& Transform, const FBox& Bounds, FPCGPoint& OutPoint, UPCGMetadata* OutMetadata) const override;
+	// TODO what should this do - closest point on volume?
+	//virtual bool ProjectPoint(const FTransform& InTransform, const FBox& InBounds, const FPCGProjectionParams& InParams, FPCGPoint& OutPoint, UPCGMetadata* OutMetadata) const;
 protected:
 	virtual UPCGSpatialData* CopyInternal() const override;
 	//~End UPCGSpatialData interface
diff --git a/Engine/Plugins/Experimental/PCG/Source/PCG/Public/Data/PCGWorldData.h b/Engine/Plugins/Experimental/PCG/Source/PCG/Public/Data/PCGWorldData.h
--- a/Engine/Plugins/Experimental/PCG/Source/PCG/Public/Data/PCGWorldData.h
+++ b/Engine/Plugins/Experimental/PCG/Source/PCG/Public/Data/PCGWorldData.h
@@ -32,6 +32,7 @@ struct FPCGWorldVolumetricQueryParams
 	TEnumAsByte<ECollisionChannel> CollisionChannel = ECC_WorldStatic;
 };
 
+/** Queries volume for presence of world collision or not. Can be used to voxelize environment. */
 UCLASS(BlueprintType, ClassGroup = (Procedural))
 class PCG_API UPCGWorldVolumetricData : public UPCGVolumeData
 {
@@ -42,6 +43,8 @@ public:
 
 	//~Begin UPCGSpatialData interface
 	virtual bool SamplePoint(const FTransform& Transform, const FBox& Bounds, FPCGPoint& OutPoint, UPCGMetadata* OutMetadata) const override;
+	// TODO not sure what this would mean. Without a direction, this means perhaps finding closest point on any collision surface? Should we implement this disabled?
+	//virtual bool ProjectPoint(const FTransform& InTransform, const FBox& InBounds, const FPCGProjectionParams& InParams, FPCGPoint& OutPoint, UPCGMetadata* OutMetadata) const;
 protected:
 	virtual UPCGSpatialData* CopyInternal() const override;
 	//~End UPCGSpatialData interface
@@ -96,6 +99,7 @@ struct FPCGWorldRayHitQueryParams
 	//bool bUseMetadataFromLandscape = true;
 };
 
+/** Executes collision queries against world collision. */
 UCLASS(BlueprintType, ClassGroup = (Procedural))
 class PCG_API UPCGWorldRayHitData : public UPCGSurfaceData
 {
