[PCG] Add tooltips to most explanation-worthy InputOutput pins, and to projection node.

diff --git a/Engine/Plugins/Experimental/PCG/Source/PCG/Private/Elements/PCGProjectionElement.cpp b/Engine/Plugins/Experimental/PCG/Source/PCG/Private/Elements/PCGProjectionElement.cpp
--- a/Engine/Plugins/Experimental/PCG/Source/PCG/Private/Elements/PCGProjectionElement.cpp
+++ b/Engine/Plugins/Experimental/PCG/Source/PCG/Private/Elements/PCGProjectionElement.cpp
@@ -13,8 +13,8 @@
 TArray<FPCGPinProperties> UPCGProjectionSettings::InputPinProperties() const
 {
 	TArray<FPCGPinProperties> PinProperties;
-	PinProperties.Emplace(PCGPinConstants::DefaultInputLabel, EPCGDataType::Spatial);
-	PinProperties.Emplace(PCGProjectionConstants::ProjectionTargetLabel, EPCGDataType::Spatial, /*bAllowMultipleConnections=*/false);
+	PinProperties.Emplace(PCGPinConstants::DefaultInputLabel, EPCGDataType::Spatial, /*bInAllowMultipleConnections=*/true, LOCTEXT("ProjectionSourcePinTooltip", "The data to project."));
+	PinProperties.Emplace(PCGProjectionConstants::ProjectionTargetLabel, EPCGDataType::Spatial, /*bAllowMultipleConnections=*/false, LOCTEXT("ProjectionTargetPinTooltip", "The projection target."));
 	PinProperties.Emplace(PCGPinConstants::DefaultParamsLabel, EPCGDataType::Param, /*bInAllowMultipleConnections=*/false);
 
 	return PinProperties;
diff --git a/Engine/Plugins/Experimental/PCG/Source/PCG/Private/PCGInputOutputSettings.cpp b/Engine/Plugins/Experimental/PCG/Source/PCG/Private/PCGInputOutputSettings.cpp
--- a/Engine/Plugins/Experimental/PCG/Source/PCG/Private/PCGInputOutputSettings.cpp
+++ b/Engine/Plugins/Experimental/PCG/Source/PCG/Private/PCGInputOutputSettings.cpp
@@ -7,6 +7,8 @@
 
 #include UE_INLINE_GENERATED_CPP_BY_NAME(PCGInputOutputSettings)
 
+#define LOCTEXT_NAMESPACE "PCGInputOutputElement"
+
 bool FPCGInputOutputElement::ExecuteInternal(FPCGContext* Context) const
 {
 	// Essentially a pass-through element
@@ -17,15 +19,34 @@ bool FPCGInputOutputElement::ExecuteInternal(FPCGContext* Context) const
 UPCGGraphInputOutputSettings::UPCGGraphInputOutputSettings(const FObjectInitializer& ObjectInitializer)
 	: Super(ObjectInitializer)
 {
-	StaticInLabels.Add(PCGPinConstants::DefaultInputLabel);
-	StaticAdvancedInLabels.Add(PCGInputOutputConstants::DefaultInputLabel);
-	StaticAdvancedInLabels.Add(PCGInputOutputConstants::DefaultActorLabel);
-	StaticAdvancedInLabels.Add(PCGInputOutputConstants::DefaultOriginalActorLabel);
-	StaticAdvancedInLabels.Add(PCGInputOutputConstants::DefaultLandscapeLabel);
-	StaticAdvancedInLabels.Add(PCGInputOutputConstants::DefaultLandscapeHeightLabel);
-	StaticAdvancedInLabels.Add(PCGInputOutputConstants::DefaultExcludedActorsLabel);
+	StaticInLabels.Emplace(PCGPinConstants::DefaultInputLabel, LOCTEXT("InputOutputInPinTooltip",
+		"Provides the result of the Input pin, but with 'Excluded Tags' actors removed (configured on the PCG Component)."
+	));
+	StaticAdvancedInLabels.Emplace(PCGInputOutputConstants::DefaultInputLabel, LOCTEXT("InputOutputInputPinTooltip",
+		"Takes the output of the Actor pin and if the 'Input Type' setting on the PCG Component is set to Landscape, combines it with the result of the Landscape pin. "
+		"If the Actor data is two dimensional it will be projected onto the landscape, otherwise it will be intersected."
+	));
+	StaticAdvancedInLabels.Emplace(PCGInputOutputConstants::DefaultActorLabel, LOCTEXT("InputOutputActorPinTooltip",
+		"If this is a partitioned component, then this will be the intersection of the current partition actor bounds with the following. "
+		"If the actor is a Landscape Proxy, then this provide a landscape data. "
+		"Otherwise if the actor is a volume, this will provide a volume shape matching the actor bounds. "
+		"Otherwise if the 'Parse Actor Components' setting is enabled on the PCG Component, this will be all compatible components on the actor (Landscape Splines, Splines, Shapes, Primitives) unioned together. "
+		"Otherwise a single point will be provided at the actor position."
+	));
+	StaticAdvancedInLabels.Emplace(PCGInputOutputConstants::DefaultOriginalActorLabel, LOCTEXT("InputOutputOriginalActorPinTooltip",
+		"If the actor is a partition actor, this will pull data from the generating PCG actor. Otherwise it will provide the same data as the Actor pin."
+	));
+	StaticAdvancedInLabels.Emplace(PCGInputOutputConstants::DefaultLandscapeLabel,
+		LOCTEXT("InputOutputLandscapePinTooltip", "Provides the landscape represented by this actor if it is a Landscape Proxy, otherwise it returns any landscapes overlapping this actor in the level."
+	));
+	StaticAdvancedInLabels.Emplace(PCGInputOutputConstants::DefaultLandscapeHeightLabel, LOCTEXT("InputOutputLandscapeHeightPinTooltip",
+		"Similar to Landscape pin, but only provides height data and not other layers."
+	));
+	StaticAdvancedInLabels.Emplace(PCGInputOutputConstants::DefaultExcludedActorsLabel, LOCTEXT("InputOutputExcludedPinTooltip",
+		"Returns union of data read from actors have any tag specified in the 'Excluded Tags' setting on the PCG Component."
+	));
 	
-	StaticOutLabels.Add(PCGPinConstants::DefaultOutputLabel);
+	StaticOutLabels.Emplace(PCGPinConstants::DefaultOutputLabel);
 }
 
 void UPCGGraphInputOutputSettings::PostLoad()
@@ -47,13 +68,14 @@ TArray<FPCGPinProperties> UPCGGraphInputOutputSettings::GetPinProperties() const
 {
 	TArray<FPCGPinProperties> PinProperties;
 	const EPCGDataType DefaultPinDataType = bIsInput ? EPCGDataType::Spatial : EPCGDataType::Any;
-	Algo::Transform(StaticLabels(), PinProperties, [DefaultPinDataType](const FName& InLabel) { return FPCGPinProperties(InLabel, DefaultPinDataType); });
+	Algo::Transform(StaticLabels(), PinProperties, [DefaultPinDataType](const FLabelAndTooltip& InLabelAndTooltip) {
+		return FPCGPinProperties(InLabelAndTooltip.Label, DefaultPinDataType, true, InLabelAndTooltip.Tooltip);
+	});
 	
 	if (bShowAdvancedPins)
 	{
-		Algo::Transform(StaticAdvancedLabels(), PinProperties, [DefaultPinDataType](const FName& InLabel)
-		{
-			return FPCGPinProperties(InLabel, (InLabel == PCGInputOutputConstants::DefaultLandscapeLabel) ? EPCGDataType::Surface: DefaultPinDataType); 
+		Algo::Transform(StaticAdvancedLabels(), PinProperties, [DefaultPinDataType](const FLabelAndTooltip& InLabelAndTooltip) {
+			return FPCGPinProperties(InLabelAndTooltip.Label, (InLabelAndTooltip.Label == PCGInputOutputConstants::DefaultLandscapeLabel) ? EPCGDataType::Surface : DefaultPinDataType, true, InLabelAndTooltip.Tooltip);
 		});
 	}
 
@@ -73,7 +95,7 @@ TArray<FPCGPinProperties> UPCGGraphInputOutputSettings::OutputPinProperties() co
 
 bool UPCGGraphInputOutputSettings::IsPinAdvanced(const UPCGPin* Pin) const
 {
-	return Pin != nullptr ? StaticAdvancedLabels().Contains(Pin->Properties.Label) : false;
+	return Pin && StaticAdvancedLabels().FindByPredicate([Pin](const FLabelAndTooltip& InLabelAndTooltip) -> bool { return InLabelAndTooltip.Label == Pin->Properties.Label; });
 }
 
 void UPCGGraphInputOutputSettings::SetShowAdvancedPins(bool bValue)
@@ -90,3 +112,5 @@ void UPCGGraphInputOutputSettings::AddCustomPin(const FPCGPinProperties& NewCust
 	Modify();
 	CustomPins.Add(NewCustomPinProperties);
 }
+
+#undef LOCTEXT_NAMESPACE
diff --git a/Engine/Plugins/Experimental/PCG/Source/PCG/Private/PCGPin.cpp b/Engine/Plugins/Experimental/PCG/Source/PCG/Private/PCGPin.cpp
--- a/Engine/Plugins/Experimental/PCG/Source/PCG/Private/PCGPin.cpp
+++ b/Engine/Plugins/Experimental/PCG/Source/PCG/Private/PCGPin.cpp
@@ -6,8 +6,8 @@
 
 #include UE_INLINE_GENERATED_CPP_BY_NAME(PCGPin)
 
-FPCGPinProperties::FPCGPinProperties(const FName& InLabel, EPCGDataType InAllowedTypes, bool bInAllowMultipleConnections)
-	: Label(InLabel), AllowedTypes(InAllowedTypes), bAllowMultipleConnections(bInAllowMultipleConnections)
+FPCGPinProperties::FPCGPinProperties(const FName& InLabel, EPCGDataType InAllowedTypes, bool bInAllowMultipleConnections, const FText& InTooltip)
+	: Label(InLabel), AllowedTypes(InAllowedTypes), bAllowMultipleConnections(bInAllowMultipleConnections), Tooltip(InTooltip)
 {}
 
 bool FPCGPinProperties::operator==(const FPCGPinProperties& Other) const
diff --git a/Engine/Plugins/Experimental/PCG/Source/PCG/Public/PCGInputOutputSettings.h b/Engine/Plugins/Experimental/PCG/Source/PCG/Public/PCGInputOutputSettings.h
--- a/Engine/Plugins/Experimental/PCG/Source/PCG/Public/PCGInputOutputSettings.h
+++ b/Engine/Plugins/Experimental/PCG/Source/PCG/Public/PCGInputOutputSettings.h
@@ -56,8 +56,19 @@ protected:
 	virtual FPCGElementPtr CreateElement() const override { return MakeShared<FPCGInputOutputElement>(); }
 	// ~End UPCGSettings interface
 
-	const TArray<FName>& StaticLabels() const { return bIsInput ? StaticInLabels : StaticOutLabels; }
-	const TArray<FName>& StaticAdvancedLabels() const { return bIsInput ? StaticAdvancedInLabels : StaticAdvancedOutLabels; }
+	struct FLabelAndTooltip
+	{
+		explicit FLabelAndTooltip(FName InLabel, FText InTooltip = FText::GetEmpty())
+			: Label(InLabel), Tooltip(InTooltip)
+		{
+		}
+
+		FName Label;
+		FText Tooltip;
+	};
+
+	const TArray<FLabelAndTooltip>& StaticLabels() const { return bIsInput ? StaticInLabels : StaticOutLabels; }
+	const TArray<FLabelAndTooltip>& StaticAdvancedLabels() const { return bIsInput ? StaticAdvancedInLabels : StaticAdvancedOutLabels; }
 
 	TArray<FPCGPinProperties> GetPinProperties() const;
 
@@ -71,10 +82,10 @@ protected:
 	UPROPERTY(EditAnywhere, BlueprintReadWrite, Category = "Settings|Input")
 	bool bShowAdvancedPins = false;
 
-	TArray<FName> StaticInLabels;
-	TArray<FName> StaticAdvancedInLabels;
-	TArray<FName> StaticOutLabels;
-	TArray<FName> StaticAdvancedOutLabels;
+	TArray<FLabelAndTooltip> StaticInLabels;
+	TArray<FLabelAndTooltip> StaticAdvancedInLabels;
+	TArray<FLabelAndTooltip> StaticOutLabels;
+	TArray<FLabelAndTooltip> StaticAdvancedOutLabels;
 
 	bool bIsInput = false;
 };
\ No newline at end of file
diff --git a/Engine/Plugins/Experimental/PCG/Source/PCG/Public/PCGPin.h b/Engine/Plugins/Experimental/PCG/Source/PCG/Public/PCGPin.h
--- a/Engine/Plugins/Experimental/PCG/Source/PCG/Public/PCGPin.h
+++ b/Engine/Plugins/Experimental/PCG/Source/PCG/Public/PCGPin.h
@@ -16,7 +16,7 @@ struct PCG_API FPCGPinProperties
 	GENERATED_BODY()
 
 	FPCGPinProperties() = default;
-	explicit FPCGPinProperties(const FName& InLabel, EPCGDataType InAllowedTypes = EPCGDataType::Any, bool bInAllowMultipleConnections = true);
+	explicit FPCGPinProperties(const FName& InLabel, EPCGDataType InAllowedTypes = EPCGDataType::Any, bool bInAllowMultipleConnections = true, const FText& InTooltip = FText::GetEmpty());
 
 	bool operator==(const FPCGPinProperties& Other) const;
 
@@ -28,6 +28,9 @@ struct PCG_API FPCGPinProperties
 
 	UPROPERTY(BlueprintReadWrite, EditAnywhere, Category = Settings)
 	bool bAllowMultipleConnections = true;
+
+	UPROPERTY(BlueprintReadWrite, EditAnywhere, Category = Settings)
+	FText Tooltip;
 };
 
 UCLASS(ClassGroup = (Procedural))
@@ -62,4 +65,4 @@ public:
 	bool BreakAllIncompatibleEdges();
 	bool IsConnected() const;
 	int32 EdgeCount() const;
-};
\ No newline at end of file
+};
diff --git a/Engine/Plugins/Experimental/PCG/Source/PCGEditor/Private/PCGEditorGraphNodeBase.cpp b/Engine/Plugins/Experimental/PCG/Source/PCGEditor/Private/PCGEditorGraphNodeBase.cpp
--- a/Engine/Plugins/Experimental/PCG/Source/PCGEditor/Private/PCGEditorGraphNodeBase.cpp
+++ b/Engine/Plugins/Experimental/PCG/Source/PCGEditor/Private/PCGEditorGraphNodeBase.cpp
@@ -474,19 +474,22 @@ void UPCGEditorGraphNodeBase::GetPinHoverText(const UEdGraphPin& Pin, FString& H
 
 	const FText DataTypeText = PinTypeToText(Pin.PinType.PinCategory, MatchingPin);
 	const FText DataSubtypeText = PinTypeToText(Pin.PinType.PinSubCategory, MatchingPin);
+	const FText Description = (MatchingPin && !MatchingPin->Properties.Tooltip.IsEmpty()) ? MatchingPin->Properties.Tooltip : FText::FromName(MatchingPin->Properties.Label);
 
 	if (MatchingPin != nullptr && bIsInputPin)
 	{
 		const FText AdditionalInfo = MatchingPin->Properties.bAllowMultipleConnections ? FText(LOCTEXT("SupportsMultiInput", "Supports multiple inputs")) : FText(LOCTEXT("SingleInputOnly", "Supports only one input"));
 
-		HoverTextOut = FText::Format(LOCTEXT("PinHoverToolTipFull", "Type: {0}\nSubtype: {1}\nAdditional information: {2}"),
+		HoverTextOut = FText::Format(LOCTEXT("PinHoverToolTipFull", "{0}\n\nType: {1}\nSubtype: {2}\nAdditional information: {3}"),
+			Description,
 			DataTypeText,
 			DataSubtypeText,
 			AdditionalInfo).ToString();
 	}
 	else
 	{
-		HoverTextOut = FText::Format(LOCTEXT("PinHoverToolTipPartial", "Type: {0}\nSubtype: {1}"),
+		HoverTextOut = FText::Format(LOCTEXT("PinHoverToolTipPartial", "{0}\n\nType: {1}\nSubtype: {2}"),
+			Description,
 			DataTypeText,
 			DataSubtypeText).ToString();
 	}
