Fixed an issue where disabling/enabling nodes would not update graph visuals properly.

diff --git a/Engine/Plugins/Experimental/PCG/Source/PCGEditor/Private/PCGEditor.cpp b/Engine/Plugins/Experimental/PCG/Source/PCGEditor/Private/PCGEditor.cpp
--- a/Engine/Plugins/Experimental/PCG/Source/PCGEditor/Private/PCGEditor.cpp
+++ b/Engine/Plugins/Experimental/PCG/Source/PCGEditor/Private/PCGEditor.cpp
@@ -1211,7 +1211,8 @@ void FPCGEditor::OnToggleEnabled()
 	if (GraphEditorWidget.IsValid())
 	{
 		const FScopedTransaction Transaction(*FPCGEditorCommon::ContextIdentifier, LOCTEXT("PCGEditorToggleEnableTransactionMessage", "PCG Editor: Toggle Enable Nodes"), nullptr);
-		
+
+		bool bChanged = false;
 		for (UObject* Object : GraphEditorWidget->GetSelectedNodes())
 		{
 			UPCGEditorGraphNodeBase* PCGEditorGraphNode = Cast<UPCGEditorGraphNodeBase>(Object);
@@ -1232,14 +1233,19 @@ void FPCGEditor::OnToggleEnabled()
 				continue;
 			}
 
-			//const bool bNewCheckState = CheckState != ECheckBoxState::Checked;
 			if (PCGSettingsInterface->bEnabled != bNewCheckState)
 			{
 				PCGSettingsInterface->Modify();
 				PCGSettingsInterface->bEnabled = bNewCheckState;
 				PCGNode->OnNodeChangedDelegate.Broadcast(PCGNode, EPCGChangeType::Settings);
+				bChanged = true;
 			}
 		}
+		
+		if (bChanged)
+		{
+			GraphEditorWidget->NotifyGraphChanged();	
+		}
 	}
 }
 
