[PCG] Fix issue with GetProperty/SetProperty for metadata.

Should fix issue with DataRow node.

#rnx

diff --git a/Engine/Plugins/Experimental/PCG/Source/PCG/Private/Metadata/Accessors/PCGAttributeAccessorHelpers.cpp b/Engine/Plugins/Experimental/PCG/Source/PCG/Private/Metadata/Accessors/PCGAttributeAccessorHelpers.cpp
--- a/Engine/Plugins/Experimental/PCG/Source/PCG/Private/Metadata/Accessors/PCGAttributeAccessorHelpers.cpp
+++ b/Engine/Plugins/Experimental/PCG/Source/PCG/Private/Metadata/Accessors/PCGAttributeAccessorHelpers.cpp
@@ -64,15 +64,15 @@ TUniquePtr<IPCGAttributeAccessor> PCGAttributeAccessorHelpers::CreatePropertyAcc
 	}
 	else if (const FBoolProperty* BoolProperty = CastField<FBoolProperty>(InProperty))
 	{
-		return MakeUnique<FPCGPropertyAccessor<bool>>(InProperty);
+		return MakeUnique<FPCGPropertyAccessor<bool, FBoolProperty>>(BoolProperty);
 	}
 	else if (const FStrProperty* StringProperty = CastField<FStrProperty>(InProperty))
 	{
-		return MakeUnique<FPCGPropertyAccessor<FString>>(InProperty);
+		return MakeUnique<FPCGPropertyAccessor<FString, FStrProperty>>(StringProperty);
 	}
 	else if (const FNameProperty* NameProperty = CastField<FNameProperty>(InProperty))
 	{
-		return MakeUnique<FPCGPropertyAccessor<FName>>(InProperty);
+		return MakeUnique<FPCGPropertyAccessor<FName, FNameProperty>>(NameProperty);
 	}
 	else if (const FEnumProperty* EnumProperty = CastField<FEnumProperty>(InProperty))
 	{
@@ -82,27 +82,27 @@ TUniquePtr<IPCGAttributeAccessor> PCGAttributeAccessorHelpers::CreatePropertyAcc
 	{
 		if (StructProperty->Struct == TBaseStructure<FVector>::Get())
 		{
-			return MakeUnique<FPCGPropertyAccessor<FVector>>(InProperty);
+			return MakeUnique<FPCGPropertyStructAccessor<FVector>>(StructProperty);
 		}
 		else if (StructProperty->Struct == TBaseStructure<FVector4>::Get())
 		{
-			return MakeUnique<FPCGPropertyAccessor<FVector4>>(InProperty);
+			return MakeUnique<FPCGPropertyStructAccessor<FVector4>>(StructProperty);
 		}
 		else if (StructProperty->Struct == TBaseStructure<FQuat>::Get())
 		{
-			return MakeUnique<FPCGPropertyAccessor<FQuat>>(InProperty);
+			return MakeUnique<FPCGPropertyStructAccessor<FQuat>>(StructProperty);
 		}
 		else if (StructProperty->Struct == TBaseStructure<FTransform>::Get())
 		{
-			return MakeUnique<FPCGPropertyAccessor<FTransform>>(InProperty);
+			return MakeUnique<FPCGPropertyStructAccessor<FTransform>>(StructProperty);
 		}
 		else if (StructProperty->Struct == TBaseStructure<FRotator>::Get())
 		{
-			return MakeUnique<FPCGPropertyAccessor<FRotator>>(InProperty);
+			return MakeUnique<FPCGPropertyStructAccessor<FRotator>>(StructProperty);
 		}
 		else if (StructProperty->Struct == TBaseStructure<FVector2D>::Get())
 		{
-			return MakeUnique<FPCGPropertyAccessor<FVector2D>>(InProperty);
+			return MakeUnique<FPCGPropertyStructAccessor<FVector2D>>(StructProperty);
 		}
 		else if (StructProperty->Struct == TBaseStructure<FSoftObjectPath>::Get())
 		{
diff --git a/Engine/Plugins/Experimental/PCG/Source/PCG/Public/Helpers/PCGPropertyHelpers.h b/Engine/Plugins/Experimental/PCG/Source/PCG/Public/Helpers/PCGPropertyHelpers.h
--- a/Engine/Plugins/Experimental/PCG/Source/PCG/Public/Helpers/PCGPropertyHelpers.h
+++ b/Engine/Plugins/Experimental/PCG/Source/PCG/Public/Helpers/PCGPropertyHelpers.h
@@ -77,7 +77,7 @@ inline decltype(auto) PCGPropertyHelpers::GetPropertyValueWithCallback(const Obj
 	{
 		using Type = decltype(Dummy);
 		Type Value{};
-		FPCGAttributeAccessorKeysGeneric Key(const_cast<ObjectType*>(InObject));
+		FPCGAttributeAccessorKeysSingleObjectPtr<ObjectType> Key(InObject);
 
 		if (PropertyAccessor->Get<Type>(Value, Key))
 		{
@@ -122,7 +122,8 @@ inline bool PCGPropertyHelpers::SetPropertyValueFromCallback(ObjectType* InObjec
 	{
 		using Type = decltype(Dummy);
 		Type Value{};
-		FPCGAttributeAccessorKeysGeneric Key(InObject);
+		FPCGAttributeAccessorKeysSingleObjectPtr<ObjectType> Key(InObject);
+
 		if (InFunc(Value))
 		{
 			return PropertyAccessor->Set<Type>(Value, Key);
diff --git a/Engine/Plugins/Experimental/PCG/Source/PCG/Public/Metadata/Accessors/PCGAttributeAccessorKeys.h b/Engine/Plugins/Experimental/PCG/Source/PCG/Public/Metadata/Accessors/PCGAttributeAccessorKeys.h
--- a/Engine/Plugins/Experimental/PCG/Source/PCG/Public/Metadata/Accessors/PCGAttributeAccessorKeys.h
+++ b/Engine/Plugins/Experimental/PCG/Source/PCG/Public/Metadata/Accessors/PCGAttributeAccessorKeys.h
@@ -184,6 +184,36 @@ protected:
 
 /////////////////////////////////////////////////////////////////
 
+/**
+* Unique Key around a single object.
+* Necessary if ObjectType is void, but keep a template version for completeness.
+* Useful when you want to use the accessors Get/Set methods on a single object.
+*/
+template <typename ObjectType>
+class FPCGAttributeAccessorKeysSingleObjectPtr : public IPCGAttributeAccessorKeys
+{
+public:
+	FPCGAttributeAccessorKeysSingleObjectPtr(ObjectType* InPtr)
+		: IPCGAttributeAccessorKeys(/*bInReadOnly=*/ false)
+		, Ptr(InPtr)
+	{}
+
+	FPCGAttributeAccessorKeysSingleObjectPtr(const ObjectType* InPtr)
+		: IPCGAttributeAccessorKeys(/*bInReadOnly=*/ true)
+		, Ptr(const_cast<ObjectType*>(InPtr))
+	{}
+
+	virtual int32 GetNum() const override { return 1; }
+
+protected:
+	virtual bool GetGenericObjectKeys(int32 InStart, TArrayView<void*>& OutObjects) override;
+	virtual bool GetGenericObjectKeys(int32 InStart, TArrayView<const void*>& OutObjects) const override;
+
+	ObjectType* Ptr = nullptr;
+};
+
+/////////////////////////////////////////////////////////////////
+
 template <typename ObjectType>
 inline bool IPCGAttributeAccessorKeys::GetKeys(int32 InStart, TArrayView<ObjectType*>& OutKeys)
 {
@@ -276,4 +306,36 @@ template <typename ObjectType>
 bool FPCGAttributeAccessorKeysGeneric<ObjectType>::GetGenericObjectKeys(int32 InStart, TArrayView<const void*>& OutObjects) const
 {
 	return PCGAttributeAccessorKeys::GetKeys(Objects, InStart, OutObjects, [](const ObjectType& Obj) -> const ObjectType* { return &Obj; });
+}
+
+template <typename ObjectType>
+bool FPCGAttributeAccessorKeysSingleObjectPtr<ObjectType>::GetGenericObjectKeys(int32 InStart, TArrayView<void*>& OutObjects)
+{
+	if (Ptr == nullptr)
+	{
+		return false;
+	}
+
+	for (int32 i = 0; i < OutObjects.Num(); ++i)
+	{
+		OutObjects[i] = Ptr;
+	}
+
+	return true;
+}
+
+template <typename ObjectType>
+bool FPCGAttributeAccessorKeysSingleObjectPtr<ObjectType>::GetGenericObjectKeys(int32 InStart, TArrayView<const void*>& OutObjects) const
+{
+	if (Ptr == nullptr)
+	{
+		return false;
+	}
+
+	for (int32 i = 0; i < OutObjects.Num(); ++i)
+	{
+		OutObjects[i] = Ptr;
+	}
+
+	return true;
 }
\ No newline at end of file
diff --git a/Engine/Plugins/Experimental/PCG/Source/PCG/Public/Metadata/Accessors/PCGPropertyAccessor.h b/Engine/Plugins/Experimental/PCG/Source/PCG/Public/Metadata/Accessors/PCGPropertyAccessor.h
--- a/Engine/Plugins/Experimental/PCG/Source/PCG/Public/Metadata/Accessors/PCGPropertyAccessor.h
+++ b/Engine/Plugins/Experimental/PCG/Source/PCG/Public/Metadata/Accessors/PCGPropertyAccessor.h
@@ -132,20 +132,20 @@ private:
 };
 
 /**
-* Templated accessor class for generic properties. Will wrap around a property.
+* Templated accessor class for struct properties. Will wrap around a struct property.
 * Do not instanciate it manually, use PCGAttributeAccessorHelpers::CreatePropertyAccessor.
 * WARNING: Getting the address of the underlying data from a container using this property
 * should point to a "T" object in memory, otherwise it is UD.
 * Key supported: Generic object
 */
 template <typename T>
-class FPCGPropertyAccessor : public IPCGAttributeAccessorT<FPCGPropertyAccessor<T>>
+class FPCGPropertyStructAccessor : public IPCGAttributeAccessorT<FPCGPropertyStructAccessor<T>>
 {
 public:
 	using Type = T;
-	using Super = IPCGAttributeAccessorT<FPCGPropertyAccessor<T>>;
+	using Super = IPCGAttributeAccessorT<FPCGPropertyStructAccessor<T>>;
 
-	FPCGPropertyAccessor(const FProperty* InProperty)
+	FPCGPropertyStructAccessor(const FStructProperty* InProperty)
 		: Super(/*bInReadOnly=*/ false)
 		, Property(InProperty)
 	{
@@ -170,7 +170,48 @@ public:
 	}
 
 private:
-	const FProperty* Property = nullptr;
+	const FStructProperty* Property = nullptr;
+};
+
+/**
+* Templated accessor class for properties that has a (Get/Set)PropertyValue. Will wrap around a property.
+* For example String/Name/Bool
+* Do not instanciate it manually, use PCGAttributeAccessorHelpers::CreatePropertyAccessor.
+* Key supported: Generic object
+*/
+template <typename T, typename PropertyType>
+class FPCGPropertyAccessor : public IPCGAttributeAccessorT<FPCGPropertyAccessor<T, PropertyType>>
+{
+public:
+	using Type = T;
+	using Super = IPCGAttributeAccessorT<FPCGPropertyAccessor<T, PropertyType>>;
+
+	FPCGPropertyAccessor(const PropertyType* InProperty)
+		: Super(/*bInReadOnly=*/ false)
+		, Property(InProperty)
+	{
+		static_assert(PCG::Private::IsPCGType<T>());
+		check(Property);
+	}
+
+	bool GetRangeImpl(TArrayView<T> OutValues, int32 Index, const IPCGAttributeAccessorKeys& Keys) const
+	{
+		return PCGPropertyAccessor::IterateGet(Property, OutValues, Index, Keys, [this](const void* PropertyAddressData) -> T
+			{
+				return Property->GetPropertyValue(PropertyAddressData);
+			});
+	}
+
+	bool SetRangeImpl(TArrayView<const T> InValues, int32 Index, IPCGAttributeAccessorKeys& Keys, EPCGAttributeAccessorFlags)
+	{
+		return PCGPropertyAccessor::IterateSet(Property, InValues, Index, Keys, [this](void* PropertyAddressData, const T& Value) -> void
+			{
+				Property->SetPropertyValue(PropertyAddressData, Value);
+			});
+	}
+
+private:
+	const PropertyType* Property = nullptr;
 };
 
 /**
