Fixed missing main thread execution restriction on PCG debug node

#rnx

diff --git a/Engine/Plugins/Experimental/PCG/Source/PCG/Public/Elements/PCGDebugElement.h b/Engine/Plugins/Experimental/PCG/Source/PCG/Public/Elements/PCGDebugElement.h
--- a/Engine/Plugins/Experimental/PCG/Source/PCG/Public/Elements/PCGDebugElement.h
+++ b/Engine/Plugins/Experimental/PCG/Source/PCG/Public/Elements/PCGDebugElement.h
@@ -35,6 +35,7 @@ class FPCGDebugElement : public FSimplePCGElement
 {
 public:
 	virtual bool IsCacheable(const UPCGSettings* InSettings) const override { return false; }
+	virtual bool CanExecuteOnlyOnMainThread(FPCGContext* Context) const { return true; }
 
 protected:
 	virtual bool ExecuteInternal(FPCGContext* Context) const override;	
