Forced spawned PCG actors to be transient when spawned at runtime

#rnx

diff --git a/Engine/Plugins/Experimental/PCG/Source/PCG/Private/Elements/PCGSpawnActor.cpp b/Engine/Plugins/Experimental/PCG/Source/PCG/Private/Elements/PCGSpawnActor.cpp
--- a/Engine/Plugins/Experimental/PCG/Source/PCG/Private/Elements/PCGSpawnActor.cpp
+++ b/Engine/Plugins/Experimental/PCG/Source/PCG/Private/Elements/PCGSpawnActor.cpp
@@ -236,6 +236,11 @@ bool FPCGSpawnActorElement::ExecuteInternal(FPCGContext* Context) const
 					SpawnParams.Owner = TargetActor;
 					SpawnParams.SpawnCollisionHandlingOverride = ESpawnActorCollisionHandlingMethod::AlwaysSpawn;
 
+					if (PCGHelpers::IsRuntimeOrPIE())
+					{
+						SpawnParams.ObjectFlags |= RF_Transient;
+					}
+
 					TArray<FName> NewActorTags = GetNewActorTags(Context, TargetActor, Settings->bInheritActorTags, Settings->TagsToAddOnActors);
 
 					// Create managed resource for actor tracking
