[PCG] Return first non-params data from first non-params input pin when element is disabled.

Seems to work well across a range of different elements tested.

Behaviour can be overridden so that disabled-behaviour can be changed for specific nodes if needed.

diff --git a/Engine/Plugins/Experimental/PCG/Source/PCG/Private/PCGElement.cpp b/Engine/Plugins/Experimental/PCG/Source/PCG/Private/PCGElement.cpp
--- a/Engine/Plugins/Experimental/PCG/Source/PCG/Private/PCGElement.cpp
+++ b/Engine/Plugins/Experimental/PCG/Source/PCG/Private/PCGElement.cpp
@@ -4,6 +4,7 @@
 #include "PCGComponent.h"
 #include "PCGContext.h"
 #include "PCGModule.h"
+#include "PCGParamData.h"
 #include "PCGSettings.h"
 #include "Elements/PCGDebugElement.h"
 #include "Elements/PCGSelfPruning.h"
@@ -109,7 +110,8 @@ void IPCGElement::PreExecute(FPCGContext* Context) const
 	if (!SettingsInterface->bEnabled)
 	{
 		//Pass-through - no execution
-		Context->OutputData = Context->InputData;
+		DisabledPassThroughData(Context);
+
 		Context->CurrentPhase = EPCGExecutionPhase::PostExecute;
 	}
 	else
@@ -193,6 +195,63 @@ void IPCGElement::PostExecute(FPCGContext* Context) const
 	Context->CurrentPhase = EPCGExecutionPhase::Done;
 }
 
+void IPCGElement::DisabledPassThroughData(FPCGContext* Context) const
+{
+	// Copy as baseline
+	Context->OutputData = Context->InputData;
+
+	if (!Context->Node)
+	{
+		// Full pass-through if we don't have a node
+		return;
+	}
+
+	if (Context->Node->GetInputPins().Num() == 0)
+	{
+		// No input pins, return nothing
+		Context->OutputData.TaggedData.Empty();
+
+		return;
+	}
+
+	const UPCGPin* FirstNonParamsPin = nullptr;
+
+	// Find first non-params pin. Choosing to pass through params does not make sense
+	auto NonParamsPredicate = [](const TObjectPtr<UPCGPin>& InPin) { return InPin->Properties.AllowedTypes != EPCGDataType::Param; };
+	if (const TObjectPtr<UPCGPin>* FirstNonParamsPinPtr = Algo::FindByPredicate(Context->Node->GetInputPins(), NonParamsPredicate))
+	{
+		FirstNonParamsPin = *FirstNonParamsPinPtr;
+	}
+
+	if (FirstNonParamsPin == nullptr)
+	{
+		// No pin to grab pass through data from
+		Context->OutputData.TaggedData.Empty();
+
+		return;
+	}
+
+	// Find first incoming non-params data that is coming through the identified pin
+	TArray<FPCGTaggedData> InputsOnFirstPin = Context->InputData.GetInputsByPin(FirstNonParamsPin->Properties.Label);
+	const int FirstNonParamsDataIndex = InputsOnFirstPin.IndexOfByPredicate([Context](const FPCGTaggedData& InData) { return Cast<UPCGParamData>(InData.Data) == nullptr; });
+	if (FirstNonParamsDataIndex != -1)
+	{
+		// Remove everything except the data we found above
+		for (int Index = Context->OutputData.TaggedData.Num() - 1; Index >= 0; --Index)
+		{
+			if (Index != FirstNonParamsDataIndex)
+			{
+				Context->OutputData.TaggedData.RemoveAt(Index);
+			}
+		}
+	}
+	else
+	{
+		// No data found to return
+		Context->OutputData.TaggedData.Empty();
+	}
+}
+
 #if WITH_EDITOR
 void IPCGElement::DebugDisplay(FPCGContext* Context) const
 {
diff --git a/Engine/Plugins/Experimental/PCG/Source/PCG/Public/PCGElement.h b/Engine/Plugins/Experimental/PCG/Source/PCG/Public/PCGElement.h
--- a/Engine/Plugins/Experimental/PCG/Source/PCG/Public/PCGElement.h
+++ b/Engine/Plugins/Experimental/PCG/Source/PCG/Public/PCGElement.h
@@ -102,6 +102,9 @@ protected:
 	/** Used to specify that the element passes through the data without any manipulation - used to correct target pins, etc. */
 	virtual bool IsPassthrough() const { return false; }
 
+	/** Passes through data when the element is Disabled. Can be implemented to override what gets passed through. */
+	virtual void DisabledPassThroughData(FPCGContext* Context) const;
+
 #if WITH_EDITOR
 	virtual bool ShouldLog() const { return true; }
 #endif
