[PCG] Fix non-unity build (Issue 243070)

diff --git a/Engine/Plugins/Experimental/PCG/Source/PCG/Private/PCGElement.cpp b/Engine/Plugins/Experimental/PCG/Source/PCG/Private/PCGElement.cpp
--- a/Engine/Plugins/Experimental/PCG/Source/PCG/Private/PCGElement.cpp
+++ b/Engine/Plugins/Experimental/PCG/Source/PCG/Private/PCGElement.cpp
@@ -1,6 +1,7 @@
 // Copyright Epic Games, Inc. All Rights Reserved.
 
 #include "PCGElement.h"
+
 #include "PCGComponent.h"
 #include "PCGContext.h"
 #include "PCGModule.h"
@@ -10,6 +11,8 @@
 #include "Elements/PCGSelfPruning.h"
 #include "Graph/PCGGraphCache.h"
 
+#include "Algo/Find.h"
+
 bool IPCGElement::Execute(FPCGContext* Context) const
 {
 	TRACE_CPUPROFILER_EVENT_SCOPE(IPCGElement::Execute);
