Added node name and point/metadata count to the attribute list view.
Removed raw ptr to the pcg node being inspected, we can access it through the pcg editor graph node being inspected.

diff --git a/Engine/Plugins/Experimental/PCG/Source/PCGEditor/Private/PCGEditor.cpp b/Engine/Plugins/Experimental/PCG/Source/PCGEditor/Private/PCGEditor.cpp
--- a/Engine/Plugins/Experimental/PCG/Source/PCGEditor/Private/PCGEditor.cpp
+++ b/Engine/Plugins/Experimental/PCG/Source/PCGEditor/Private/PCGEditor.cpp
@@ -176,15 +176,6 @@ void FPCGEditor::SetPCGComponentBeingDebugged(UPCGComponent* InPCGComponent)
 	}
 }
 
-void FPCGEditor::SetPCGNodeBeingInspected(UPCGNode* InPCGNode)
-{
-	if (GetPCGNodeBeingInspected() != InPCGNode)
-	{
-		PCGNodeBeingInspected = InPCGNode;
-		OnInspectedNodeChangedDelegate.Broadcast(PCGNodeBeingInspected);
-	}
-}
-
 void FPCGEditor::JumpToNode(const UEdGraphNode* InNode)
 {
 	if (GraphEditorWidget.IsValid())
@@ -1187,11 +1178,11 @@ void FPCGEditor::OnStartInspectNode()
 	if (PCGGraphNodeBeingInspected)
 	{
 		PCGGraphNodeBeingInspected->SetInspected(false);
-	}	
+	}
 
 	PCGGraphNodeBeingInspected = PCGGraphNodeBase;
 	PCGGraphNodeBeingInspected->SetInspected(true);
-	SetPCGNodeBeingInspected(PCGGraphNodeBase->GetPCGNode());
+	OnInspectedNodeChangedDelegate.Broadcast(PCGGraphNodeBeingInspected);
 }
 
 void FPCGEditor::OnStopInspectNode()
@@ -1200,7 +1191,7 @@ void FPCGEditor::OnStopInspectNode()
 	{
 		PCGGraphNodeBeingInspected->SetInspected(false);
 		PCGGraphNodeBeingInspected = nullptr;
-		SetPCGNodeBeingInspected(nullptr);
+		OnInspectedNodeChangedDelegate.Broadcast(nullptr);
 	}
 }
 
diff --git a/Engine/Plugins/Experimental/PCG/Source/PCGEditor/Private/PCGEditor.h b/Engine/Plugins/Experimental/PCG/Source/PCGEditor/Private/PCGEditor.h
--- a/Engine/Plugins/Experimental/PCG/Source/PCGEditor/Private/PCGEditor.h
+++ b/Engine/Plugins/Experimental/PCG/Source/PCGEditor/Private/PCGEditor.h
@@ -26,7 +26,7 @@ class UPCGGraph;
 class UPCGNode;
 
 DECLARE_MULTICAST_DELEGATE_OneParam(FOnDebugObjectChanged, UPCGComponent*);
-DECLARE_MULTICAST_DELEGATE_OneParam(FOnInspectedNodeChanged, UPCGNode*);
+DECLARE_MULTICAST_DELEGATE_OneParam(FOnInspectedNodeChanged, UPCGEditorGraphNodeBase*);
 
 class FPCGEditor : public FAssetEditorToolkit, public FGCObject, public FSelfRegisteringEditorUndoClient
 {
@@ -43,12 +43,6 @@ public:
 	/** Gets the PCG component we are debugging */
 	UPCGComponent* GetPCGComponentBeingDebugged() const { return PCGComponentBeingDebugged.Get(); }
 
-	/** Sets the PCG node we want to inspect */
-	void SetPCGNodeBeingInspected(UPCGNode* InPCGNode);
-
-	/** Gets the PCG node we are inspecting */
-	UPCGNode* GetPCGNodeBeingInspected() const { return PCGNodeBeingInspected; }
-
 	/** Focus the graph view on a specific node */
 	void JumpToNode(const UEdGraphNode* InNode);
 
@@ -277,6 +271,5 @@ private:
 	UPCGEditorGraph* PCGEditorGraph = nullptr;
 
 	TWeakObjectPtr<UPCGComponent> PCGComponentBeingDebugged;
-	UPCGNode* PCGNodeBeingInspected = nullptr;
 	UPCGEditorGraphNodeBase* PCGGraphNodeBeingInspected = nullptr;
 };
diff --git a/Engine/Plugins/Experimental/PCG/Source/PCGEditor/Private/SPCGEditorGraphAttributeListView.cpp b/Engine/Plugins/Experimental/PCG/Source/PCGEditor/Private/SPCGEditorGraphAttributeListView.cpp
--- a/Engine/Plugins/Experimental/PCG/Source/PCGEditor/Private/SPCGEditorGraphAttributeListView.cpp
+++ b/Engine/Plugins/Experimental/PCG/Source/PCGEditor/Private/SPCGEditorGraphAttributeListView.cpp
@@ -11,6 +11,7 @@
 #include "PCGEditor.h"
 #include "PCGNode.h"
 #include "PCGParamData.h"
+#include "PCGEditorGraphNodeBase.h"
 
 #include "Styling/AppStyle.h"
 #include "Widgets/Input/SSearchBox.h"
@@ -28,7 +29,9 @@ static TAutoConsoleVariable<bool> CVarShowAdvancedAttributesFields(
 namespace PCGEditorGraphAttributeListView
 {
 	const FText NoDataAvailableText = LOCTEXT("NoDataAvailableText", "No data available");
-
+	const FText NoNodeInspectedText = LOCTEXT("NoNodeInspectedText", "No node being inspected");
+	const FText NoNodeInspectedToolTip = LOCTEXT("NoNodeInspectedToolTip", "Inspect a node using the right click menu");
+	
 	/** Names of the columns in the attribute list */
 	const FName NAME_IndexColumn = FName(TEXT("IndexColumn"));
 	const FName NAME_PointPositionX = FName(TEXT("PointPositionX"));
@@ -442,7 +445,31 @@ void SPCGEditorGraphAttributeListView::Construct(const FArguments& InArgs, TShar
 		+SVerticalBox::Slot()
 		.AutoHeight()
 		[
-			DataComboBox->AsShared()
+			SNew(SHorizontalBox)
+			+SHorizontalBox::Slot()
+			.AutoWidth()
+			[
+				DataComboBox->AsShared()
+			]
+			+SHorizontalBox::Slot()
+			.AutoWidth()
+			.VAlign(EVerticalAlignment::VAlign_Center)
+			.Padding(4.0f, 0.0f)
+			[
+				SAssignNew(NodeNameTextBlock, STextBlock)
+				.Text(PCGEditorGraphAttributeListView::NoNodeInspectedText)
+				.ToolTipText(PCGEditorGraphAttributeListView::NoNodeInspectedToolTip)
+			]
+			+ SHorizontalBox::Slot()
+			.FillWidth(1.0f)
+			+SHorizontalBox::Slot()
+			.AutoWidth()
+			.HAlign(EHorizontalAlignment::HAlign_Right)
+			.VAlign(EVerticalAlignment::VAlign_Center)
+			.Padding(4.0f, 0.0f)
+			[
+				SAssignNew(InfoTextBlock, STextBlock)
+			]
 		]
 		+SVerticalBox::Slot()
 		.FillHeight(1.0f)
@@ -508,8 +535,26 @@ void SPCGEditorGraphAttributeListView::OnDebugObjectChanged(UPCGComponent* InPCG
 	}
 }
 
-void SPCGEditorGraphAttributeListView::OnInspectedNodeChanged(UPCGNode* /*InPCGNode*/)
+void SPCGEditorGraphAttributeListView::OnInspectedNodeChanged(UPCGEditorGraphNodeBase* InPCGEditorGraphNode)
 {
+	if (PCGEditorGraphNode == InPCGEditorGraphNode)
+	{
+		return;
+	}
+	
+	PCGEditorGraphNode = InPCGEditorGraphNode;
+	
+	if (PCGEditorGraphNode.IsValid())
+	{
+		NodeNameTextBlock->SetText(PCGEditorGraphNode->GetNodeTitle(ENodeTitleType::FullTitle));
+		NodeNameTextBlock->SetToolTipText(PCGEditorGraphNode->GetTooltipText());
+	}
+	else
+	{
+		NodeNameTextBlock->SetText(PCGEditorGraphAttributeListView::NoNodeInspectedText);
+		NodeNameTextBlock->SetToolTipText(PCGEditorGraphAttributeListView::NoNodeInspectedToolTip);
+	}
+	
 	RefreshDataComboBox();
 	RefreshAttributeList();
 }
@@ -529,19 +574,19 @@ void SPCGEditorGraphAttributeListView::RefreshAttributeList()
 	ListViewItems.Empty();
 	ListViewHeader->ClearColumns();
 	MetadataInfos.Empty();
-
-	const TSharedPtr<FPCGEditor> PCGEditor = PCGEditorPtr.Pin();
-	if (!PCGEditor.IsValid())
+	InfoTextBlock->SetText(FText::GetEmpty());
+	
+	if (!PCGComponent.IsValid())
 	{
 		return;
 	}
 
-	if (!PCGComponent.IsValid())
+	if (!PCGEditorGraphNode.IsValid())
 	{
 		return;
 	}
 
-	const UPCGNode* PCGNode = PCGEditor->GetPCGNodeBeingInspected();
+	const UPCGNode* PCGNode = PCGEditorGraphNode->GetPCGNode();
 	if (!PCGNode)
 	{
 		return;
@@ -581,6 +626,8 @@ void SPCGEditorGraphAttributeListView::RefreshAttributeList()
 				ListViewItem->MetadataInfos = &MetadataInfos;
 				ListViewItems.Add(ListViewItem);
 			}
+
+			InfoTextBlock->SetText(FText::Format(LOCTEXT("MetadataInfoTextBlockFmt", "Number of metadata: {0}"), ItemKeyUpperBound - ItemKeyLowerBound));
 		}
 	}
 	else if (const UPCGSpatialData* PCGSpatialData = Cast<const UPCGSpatialData>(PCGData))
@@ -593,8 +640,9 @@ void SPCGEditorGraphAttributeListView::RefreshAttributeList()
 			GenerateColumnsFromMetadata(PCGMetadata);
 
 			const TArray<FPCGPoint>& PCGPoints = PCGPointData->GetPoints();
-			ListViewItems.Reserve(PCGPoints.Num());
-			for (int32 PointIndex = 0; PointIndex < PCGPoints.Num(); PointIndex++)
+			const int32 NumPoints = PCGPoints.Num();
+			ListViewItems.Reserve(NumPoints);
+			for (int32 PointIndex = 0; PointIndex < NumPoints; PointIndex++)
 			{
 				const FPCGPoint& PCGPoint = PCGPoints[PointIndex];
 				// TODO: Investigate swapping out the shared ptr's for better performance on huge data sets
@@ -606,6 +654,8 @@ void SPCGEditorGraphAttributeListView::RefreshAttributeList()
 				ListViewItem->MetadataInfos = &MetadataInfos;
 				ListViewItems.Add(ListViewItem);
 			}
+
+			InfoTextBlock->SetText(FText::Format(LOCTEXT("PointInfoTextBlockFmt", "Number of points: {0}"), NumPoints));
 		}
 	}
 
@@ -623,13 +673,12 @@ void SPCGEditorGraphAttributeListView::RefreshDataComboBox()
 		return;
 	}
 
-	const TSharedPtr<FPCGEditor> PCGEditor = PCGEditorPtr.Pin();
-	if (!PCGEditor.IsValid())
+	if (!PCGEditorGraphNode.IsValid())
 	{
 		return;
 	}
 
-	const UPCGNode* PCGNode = PCGEditor->GetPCGNodeBeingInspected();
+	const UPCGNode* PCGNode = PCGEditorGraphNode->GetPCGNode();
 	if (!PCGNode)
 	{
 		return;
diff --git a/Engine/Plugins/Experimental/PCG/Source/PCGEditor/Private/SPCGEditorGraphAttributeListView.h b/Engine/Plugins/Experimental/PCG/Source/PCGEditor/Private/SPCGEditorGraphAttributeListView.h
--- a/Engine/Plugins/Experimental/PCG/Source/PCGEditor/Private/SPCGEditorGraphAttributeListView.h
+++ b/Engine/Plugins/Experimental/PCG/Source/PCGEditor/Private/SPCGEditorGraphAttributeListView.h
@@ -11,7 +11,7 @@
 
 class FPCGEditor;
 class UPCGComponent;
-class UPCGNode;
+class UPCGEditorGraphNodeBase;
 class UPCGParamData;
 struct FPCGPoint;
 
@@ -64,7 +64,7 @@ private:
 	TSharedRef<SHeaderRow> CreateHeaderRowWidget() const;
 
 	void OnDebugObjectChanged(UPCGComponent* InPCGComponent);
-	void OnInspectedNodeChanged(UPCGNode* InPCGNode);
+	void OnInspectedNodeChanged(UPCGEditorGraphNodeBase* InPCGEditorGraphNode);
 
 	void OnGenerateUpdated(UPCGComponent* InPCGComponent);
 
@@ -98,12 +98,18 @@ private:
 	/** Cached PCGComponent being viewed */
 	TWeakObjectPtr<UPCGComponent> PCGComponent;
 
+	/** Cached PCGGraphNode being viewed */
+	TWeakObjectPtr<UPCGEditorGraphNodeBase> PCGEditorGraphNode;
+
 	TSharedPtr<SHeaderRow> ListViewHeader;
 	TSharedPtr<SListView<PCGListviewItemPtr>> ListView;
 	TArray<PCGListviewItemPtr> ListViewItems;
 
 	TSharedPtr<SComboBox<TSharedPtr<FName>>> DataComboBox;
 	TArray<TSharedPtr<FName>> DataComboBoxItems;
+	
+	TSharedPtr<STextBlock> NodeNameTextBlock;
+	TSharedPtr<STextBlock> InfoTextBlock;
 
 	TMap<FName, FPCGMetadataInfo> MetadataInfos;
 };
