Exposing PCG unit test stuff so we can write some unit tests in projects

diff --git a/Engine/Plugins/Experimental/PCG/Source/PCG/Public/Tests/PCGTestsCommon.h b/Engine/Plugins/Experimental/PCG/Source/PCG/Public/Tests/PCGTestsCommon.h
--- a/Engine/Plugins/Experimental/PCG/Source/PCG/Public/Tests/PCGTestsCommon.h
+++ b/Engine/Plugins/Experimental/PCG/Source/PCG/Public/Tests/PCGTestsCommon.h
@@ -23,7 +23,7 @@ namespace PCGTestsCommon
 {
 	constexpr int TestFlags = EAutomationTestFlags::EditorContext | EAutomationTestFlags::EngineFilter;
 
-	struct FTestData
+	struct PCG_API FTestData
 	{
 		FTestData(int32 Seed, UPCGSettings* DefaultSettings = nullptr, TSubclassOf<AActor> ActorClass = AActor::StaticClass());
 		~FTestData();
@@ -62,7 +62,7 @@ namespace PCGTestsCommon
 	bool PointsAreIdentical(const FPCGPoint& FirstPoint, const FPCGPoint& SecondPoint);
 }
 
-class FPCGTestBaseClass : public FAutomationTestBase
+class PCG_API FPCGTestBaseClass : public FAutomationTestBase
 {
 public:
 	using FAutomationTestBase::FAutomationTestBase;
