Added unit test for PCGDataTableRowToParamData node

diff --git a/Engine/Plugins/Experimental/PCG/Source/PCG/Private/Tests/Elements/PCGDataTableRowToParamDataTest.cpp b/Engine/Plugins/Experimental/PCG/Source/PCG/Private/Tests/Elements/PCGDataTableRowToParamDataTest.cpp
new file mode 100644
--- /dev/null
+++ b/Engine/Plugins/Experimental/PCG/Source/PCG/Private/Tests/Elements/PCGDataTableRowToParamDataTest.cpp
@@ -0,0 +1,110 @@
+// Copyright Epic Games, Inc. All Rights Reserved.
+
+#include "Tests/Determinism/PCGDeterminismTestsCommon.h"
+#include "Tests/PCGTestsCommon.h"
+
+#include "PCGComponent.h"
+#include "PCGData.h"
+#include "PCGHelpers.h"
+#include "PCGParamData.h"
+
+#include "Elements/PCGDataTableRowToParamData.h"
+#include "Tests/Elements/PCGDataTableRowToParamDataTestStruct.h"
+#include "Engine/DataTable.h"
+
+IMPLEMENT_CUSTOM_SIMPLE_AUTOMATION_TEST(FPCGDataTableRowToParamDataTest, FPCGTestBaseClass, "pcg.tests.DataTableRowToParamData.Basic", PCGTestsCommon::TestFlags)
+
+namespace PCGDataTableRowToParamDataTest
+{
+	template<typename StructType, typename PCGType>
+	bool TestAttribute(FPCGTestBaseClass& Self, void *Data, const UPCGParamData& Params, FName FieldName, PCGType ExpectedValue)
+	{
+		const FPCGMetadataAttribute<PCGType>* Attribute = static_cast<const FPCGMetadataAttribute<PCGType>*>(Params.Metadata->GetConstAttribute(FieldName));
+
+		if (!Self.TestNotNull(*FString::Printf(TEXT("'%s' Metadata attribute"), *FieldName.ToString()), Attribute))
+		{
+			return false;
+		}
+
+		if (!Self.TestEqual(*FString::Printf(TEXT("'%s' Metadata type"), *FieldName.ToString()), Attribute->GetTypeId(), PCG::Private::MetadataTypes<PCGType>::Id))
+		{
+			return false;
+		}
+
+		const FProperty* Property = FindFProperty<FProperty>(FPCGDataTableRowToParamDataTestStruct::StaticStruct(), FieldName);
+		if (!Self.TestNotNull(*FString::Printf(TEXT("'%s' Property not found"), *FieldName.ToString()), Property))
+		{
+			return false;
+		}
+
+		const StructType* RowValuePtr = Property->ContainerPtrToValuePtr<StructType>(Data);
+		if (!Self.TestNotNull(*FString::Printf(TEXT("'%s' Row Value not found"), *FieldName.ToString()), RowValuePtr))
+		{
+			return false;
+		}
+
+		if (!Self.TestEqual(*FString::Printf(TEXT("'%s' Row Value"), *FieldName.ToString()), *RowValuePtr, StructType(ExpectedValue)))
+		{
+			return false;
+		}
+		
+		PCGType Value = Attribute->GetValueFromItemKey(0);
+		if (!Self.TestEqual(*FString::Printf(TEXT("'%s' Attribute Value"), *FieldName.ToString()), Value, ExpectedValue))
+		{
+			return false;
+		}
+
+		return true;
+	}
+}
+
+bool FPCGDataTableRowToParamDataTest::RunTest(const FString& Parameters)
+{
+	PCGTestsCommon::FTestData TestData(PCGDeterminismTests::Defaults::Seed);
+	PCGDeterminismTests::GenerateSettings<UPCGDataTableRowToParamDataSettings>(TestData);
+	UPCGDataTableRowToParamDataSettings* Settings = CastChecked<UPCGDataTableRowToParamDataSettings>(TestData.Settings);
+
+	UDataTable *TestDataTable = NewObject<UDataTable>();
+	TestDataTable->RowStruct = FPCGDataTableRowToParamDataTestStruct::StaticStruct();
+	TArray<FString> Errors = TestDataTable->CreateTableFromCSVString(R"CSV(---,Name,String,I32,I64,F32,F64,V2,V3,V4,SoftPath
+AAA,"A Name","A String","111","1111","1.1","1.11","(X=1.0,Y=1.0)","(X=1.0,Y=1.0,Z=1.0)","(X=1.0,Y=1.0,Z=1.0,W=1.0)","/Script/PCG"
+BBB,"B Name","B String","222","2222","2.2","2.22","(X=2.0,Y=2.0)","(X=2.0,Y=2.0,Z=2.0)","(X=2.0,Y=2.0,Z=2.0,W=2.0)","/Script/PCG"
+CCC,"C Name","C String","333","3333","3.3","3.33","(X=3.0,Y=3.0)","(X=3.0,Y=3.0,Z=3.0)","(X=3.0,Y=3.0,Z=3.0,W=3.0)","/Script/PCG")CSV");
+
+	UTEST_EQUAL(*FString::Printf(TEXT("CSV Errors:\n%s"), *FString::Join(Errors, TEXT("\n"))), Errors.Num(), 0);
+
+	Settings->RowName = TEXT("BBB");
+	Settings->DataTable = TestDataTable;
+
+	FPCGElementPtr TestElement = TestData.Settings->GetElement();
+
+	TUniquePtr<FPCGContext> Context(TestElement->Initialize(TestData.InputData, TestData.TestPCGComponent, nullptr));
+	Context->NumAvailableTasks = 1;
+
+	while (!TestElement->Execute(Context.Get())) {}
+
+	const UPCGParamData* Params = Context->OutputData.GetParams();
+	UTEST_NOT_NULL("Output params", Params);
+
+	UTEST_NOT_NULL("Params has metadata", Params->Metadata.Get());
+
+	void* RowData = TestDataTable->FindRow<FPCGDataTableRowToParamDataTestStruct>(Settings->RowName, TEXT("FPCGDataTableRowToParamDataTest"));
+	UTEST_NOT_NULL(*FString::Printf(TEXT("'%s' Row not found in data table"), *Settings->RowName.ToString()), RowData);
+
+	bool bSuccess = true;
+
+	// first template type is the value of the row struct member, second is the pcg attribute type
+	// this is because sometimes conversion occurs as PCG can't represent all types
+	bSuccess |= PCGDataTableRowToParamDataTest::TestAttribute<FName,FName>(*this, RowData, *Params, TEXT("Name"), TEXT("B Name"));
+	bSuccess |= PCGDataTableRowToParamDataTest::TestAttribute<FString,FString>(*this, RowData, *Params, TEXT("String"), TEXT("B String"));
+	bSuccess |= PCGDataTableRowToParamDataTest::TestAttribute<int32,int64>(*this, RowData, *Params, TEXT("I32"), 222);
+	bSuccess |= PCGDataTableRowToParamDataTest::TestAttribute<int64,int64>(*this, RowData, *Params, TEXT("I64"), 2222);
+	bSuccess |= PCGDataTableRowToParamDataTest::TestAttribute<float,double>(*this, RowData, *Params, TEXT("F32"), 2.2);
+	bSuccess |= PCGDataTableRowToParamDataTest::TestAttribute<double,double>(*this, RowData, *Params, TEXT("F64"), 2.22);
+	bSuccess |= PCGDataTableRowToParamDataTest::TestAttribute<FVector2D,FVector2D>(*this, RowData, *Params, TEXT("V2"), FVector2D(2.0, 2.0));
+	bSuccess |= PCGDataTableRowToParamDataTest::TestAttribute<FVector,FVector>(*this, RowData, *Params, TEXT("V3"), FVector(2.0, 2.0, 2.0));
+	bSuccess |= PCGDataTableRowToParamDataTest::TestAttribute<FVector4,FVector4>(*this, RowData, *Params, TEXT("V4"), FVector4(2.0, 2.0, 2.0, 2.0));
+	bSuccess |= PCGDataTableRowToParamDataTest::TestAttribute<FSoftObjectPath,FString>(*this, RowData, *Params, TEXT("SoftPath"), TEXT("/Script/PCG"));
+
+	return bSuccess;
+}
diff --git a/Engine/Plugins/Experimental/PCG/Source/PCG/Private/Tests/Elements/PCGDataTableRowToParamDataTestStruct.h b/Engine/Plugins/Experimental/PCG/Source/PCG/Private/Tests/Elements/PCGDataTableRowToParamDataTestStruct.h
new file mode 100644
--- /dev/null
+++ b/Engine/Plugins/Experimental/PCG/Source/PCG/Private/Tests/Elements/PCGDataTableRowToParamDataTestStruct.h
@@ -0,0 +1,42 @@
+// Copyright Epic Games, Inc. All Rights Reserved.
+
+#pragma once
+
+#include "CoreMinimal.h"
+#include "PCGDataTableRowToParamDataTestStruct.generated.h"
+
+USTRUCT()
+struct FPCGDataTableRowToParamDataTestStruct
+{
+	GENERATED_BODY()
+
+	UPROPERTY()
+	FName Name = NAME_None;
+
+	UPROPERTY()
+	FString String = "";
+
+	UPROPERTY()
+	int32 I32 = 0;
+
+	UPROPERTY()
+	int64 I64 = 0;
+
+	UPROPERTY()
+	float F32 = 0.0f;
+
+	UPROPERTY()
+	double F64 = 0.0;
+
+	UPROPERTY()
+	FVector2D V2 = FVector2D::Zero();
+
+	UPROPERTY()
+	FVector V3 = FVector::Zero();
+
+	UPROPERTY()
+	FVector4 V4 = FVector4::Zero();
+
+	UPROPERTY()
+	FSoftObjectPath SoftPath;
+};
\ No newline at end of file
