Fixed visibility issue in graph settings

#rnx

diff --git a/Engine/Plugins/Experimental/PCG/Source/PCGEditor/Private/PCGEditor.cpp b/Engine/Plugins/Experimental/PCG/Source/PCGEditor/Private/PCGEditor.cpp
--- a/Engine/Plugins/Experimental/PCG/Source/PCGEditor/Private/PCGEditor.cpp
+++ b/Engine/Plugins/Experimental/PCG/Source/PCGEditor/Private/PCGEditor.cpp
@@ -1985,6 +1985,12 @@ bool FPCGEditor::IsVisibleProperty(const FPropertyAndParent& InPropertyAndParent
 		return true;
 	}
 
+	// Currently never hide anything from the graph settings
+	if (InPropertyAndParent.Objects.Num() == 1 && Cast<UPCGGraph>(InPropertyAndParent.Objects[0]))
+	{
+		return true;
+	}
+
 	// Always hide asset info information
 	if (InPropertyAndParent.Property.HasMetaData(TEXT("Category")) &&
 		InPropertyAndParent.Property.GetMetaData(TEXT("Category")) == TEXT("AssetInfo"))
