[PCG] When a node is disabled, desaturate any edges that will not contribute / be passed through the node.

diff --git a/Engine/Plugins/Experimental/PCG/Source/PCG/Private/PCGElement.cpp b/Engine/Plugins/Experimental/PCG/Source/PCG/Private/PCGElement.cpp
--- a/Engine/Plugins/Experimental/PCG/Source/PCG/Private/PCGElement.cpp
+++ b/Engine/Plugins/Experimental/PCG/Source/PCG/Private/PCGElement.cpp
@@ -217,16 +217,9 @@ void IPCGElement::DisabledPassThroughData(FPCGContext* Context) const
 		return;
 	}
 
-	const UPCGPin* FirstNonParamsPin = nullptr;
+	const UPCGPin* PassThroughPin = Context->Node->GetPassThroughInputPin();
 
-	// Find first non-params pin. Choosing to pass through params does not make sense
-	auto NonParamsPredicate = [](const TObjectPtr<UPCGPin>& InPin) { return InPin->Properties.AllowedTypes != EPCGDataType::Param; };
-	if (const TObjectPtr<UPCGPin>* FirstNonParamsPinPtr = Algo::FindByPredicate(Context->Node->GetInputPins(), NonParamsPredicate))
-	{
-		FirstNonParamsPin = *FirstNonParamsPinPtr;
-	}
-
-	if (FirstNonParamsPin == nullptr)
+	if (PassThroughPin == nullptr)
 	{
 		// No pin to grab pass through data from
 		Context->OutputData.TaggedData.Empty();
@@ -235,9 +228,9 @@ void IPCGElement::DisabledPassThroughData(FPCGContext* Context) const
 	}
 
 	// Find first incoming non-params data that is coming through the identified pin
-	TArray<FPCGTaggedData> InputsOnFirstPin = Context->InputData.GetInputsByPin(FirstNonParamsPin->Properties.Label);
+	TArray<FPCGTaggedData> InputsOnFirstPin = Context->InputData.GetInputsByPin(PassThroughPin->Properties.Label);
 	const int FirstNonParamsDataIndex = InputsOnFirstPin.IndexOfByPredicate([Context](const FPCGTaggedData& InData) { return Cast<UPCGParamData>(InData.Data) == nullptr; });
-	if (FirstNonParamsDataIndex != -1)
+	if (FirstNonParamsDataIndex != INDEX_NONE)
 	{
 		// Remove everything except the data we found above
 		for (int Index = Context->OutputData.TaggedData.Num() - 1; Index >= 0; --Index)
diff --git a/Engine/Plugins/Experimental/PCG/Source/PCG/Private/PCGNode.cpp b/Engine/Plugins/Experimental/PCG/Source/PCG/Private/PCGNode.cpp
--- a/Engine/Plugins/Experimental/PCG/Source/PCG/Private/PCGNode.cpp
+++ b/Engine/Plugins/Experimental/PCG/Source/PCG/Private/PCGNode.cpp
@@ -364,6 +364,34 @@ bool UPCGNode::HasInboundEdges() const
 	return false;
 }
 
+const UPCGPin* UPCGNode::GetPassThroughInputPin() const
+{
+	// No outputs means nothing will be passed through
+	if (GetOutputPins().Num() == 0)
+	{
+		return nullptr;
+	}
+
+	// We assume a node whose primary output is params is a param processing node.
+	const bool bNodeOutputsParams = (GetOutputPins()[0]->Properties.AllowedTypes == EPCGDataType::Param);
+
+	if (bNodeOutputsParams)
+	{
+		// If this node primarily processes params, then look for a params pin, if we find it then we will pass it through
+		auto FindParams = [](const TObjectPtr<UPCGPin>& InPin) { return InPin->Properties.AllowedTypes == EPCGDataType::Param; };
+		const TObjectPtr<UPCGPin>* FirstParamsPinPtr = Algo::FindByPredicate(GetInputPins(), FindParams);
+		return FirstParamsPinPtr ? *FirstParamsPinPtr : nullptr;
+	}
+	else
+	{
+		// 'Normal' node. Look for the first pin that is not of type Params to pass through. If the node is disabled, the params are unused.
+		// Params-only pins will be rejected/ignored.
+		auto FindNonParams = [](const TObjectPtr<UPCGPin>& InPin) { return InPin->Properties.AllowedTypes != EPCGDataType::Param; };
+		const TObjectPtr<UPCGPin>* FirstNonParamsPinPtr = Algo::FindByPredicate(GetInputPins(), FindNonParams);
+		return FirstNonParamsPinPtr ? *FirstNonParamsPinPtr : nullptr;
+	}
+}
+
 void UPCGNode::SetSettingsInterface(UPCGSettingsInterface* InSettingsInterface, bool bUpdatePins)
 {
 #if WITH_EDITOR
diff --git a/Engine/Plugins/Experimental/PCG/Source/PCG/Public/PCGNode.h b/Engine/Plugins/Experimental/PCG/Source/PCG/Public/PCGNode.h
--- a/Engine/Plugins/Experimental/PCG/Source/PCG/Public/PCGNode.h
+++ b/Engine/Plugins/Experimental/PCG/Source/PCG/Public/PCGNode.h
@@ -97,6 +97,9 @@ public:
 	const UPCGPin* GetOutputPin(const FName& Label) const;
 	bool HasInboundEdges() const;
 
+	/** Pin from which data is passed through when this node is disabled. */
+	virtual const UPCGPin* GetPassThroughInputPin() const;
+
 	const TArray<TObjectPtr<UPCGPin>>& GetInputPins() const { return InputPins; }
 	const TArray<TObjectPtr<UPCGPin>>& GetOutputPins() const { return OutputPins; }
 
diff --git a/Engine/Plugins/Experimental/PCG/Source/PCGEditor/Private/PCGEditorGraphSchema.cpp b/Engine/Plugins/Experimental/PCG/Source/PCGEditor/Private/PCGEditorGraphSchema.cpp
--- a/Engine/Plugins/Experimental/PCG/Source/PCGEditor/Private/PCGEditorGraphSchema.cpp
+++ b/Engine/Plugins/Experimental/PCG/Source/PCGEditor/Private/PCGEditorGraphSchema.cpp
@@ -406,6 +406,7 @@ FPCGEditorConnectionDrawingPolicy::FPCGEditorConnectionDrawingPolicy(int32 InBac
 void FPCGEditorConnectionDrawingPolicy::DetermineWiringStyle(UEdGraphPin* OutputPin, UEdGraphPin* InputPin, /*inout*/ FConnectionParams& Params)
 {
 	FConnectionDrawingPolicy::DetermineWiringStyle(OutputPin, InputPin, Params);
+
 	// Emphasize wire thickness on hovered pins
 	if (HoveredPins.Contains(InputPin) && HoveredPins.Contains(OutputPin))
 	{
@@ -417,6 +418,26 @@ void FPCGEditorConnectionDrawingPolicy::DetermineWiringStyle(UEdGraphPin* Output
 	{
 		Params.WireColor = GetDefault<UPCGEditorSettings>()->GetPinColor(OutputPin->PinType);
 	}
+
+	// Desaturate and connection if the node is disabled and the data on this wire won't be used
+	if (InputPin)
+	{
+		if (const UPCGEditorGraphNodeBase* EditorNode = CastChecked<const UPCGEditorGraphNodeBase>(InputPin->GetOwningNode()))
+		{
+			const UPCGNode* PCGNode = EditorNode->GetPCGNode();
+			if (PCGNode && PCGNode->GetSettings() && !PCGNode->GetSettings()->bEnabled)
+			{
+				const UPCGPin* PassThroughInputPin = PCGNode->GetPassThroughInputPin();
+
+				// If this wire does not connect to the pass-through pin, or if the wire is not the first link on that pin, then it will
+				// not be used by the node. In this case desaturate the wire.
+				if (PassThroughInputPin && PassThroughInputPin->Properties.Label != InputPin->GetFName() || InputPin->LinkedTo[0] != OutputPin)
+				{
+					Params.WireColor = Params.WireColor.Desaturate(0.7f);
+				}
+			}
+		}
+	}
 }
 
 #undef LOCTEXT_NAMESPACE
