[PCG] Add Composite/Concrete types. Composite is output by boolean ops like Union/Difference.

No new type constraints on connections yet, and no colouring either. Lenient connections+same colours should mean this is almost completely transparent to users (only evidence is type info in tooltips).

This paves the way to varying the colour and adding type constraints to connections, and to placing collapse nodes to make conversion to points explicit. Also future sampler improvements will make use of the new type.

Misc fix - landscape height pin on Input node was missing its type.

I tested the PCG sample content and everything looks fine.

diff --git a/Engine/Plugins/Experimental/PCG/Source/PCG/Private/Elements/PCGDifferenceElement.cpp b/Engine/Plugins/Experimental/PCG/Source/PCG/Private/Elements/PCGDifferenceElement.cpp
--- a/Engine/Plugins/Experimental/PCG/Source/PCG/Private/Elements/PCGDifferenceElement.cpp
+++ b/Engine/Plugins/Experimental/PCG/Source/PCG/Private/Elements/PCGDifferenceElement.cpp
@@ -18,6 +18,14 @@ TArray<FPCGPinProperties> UPCGDifferenceSettings::InputPinProperties() const
 	return PinProperties;
 }
 
+TArray<FPCGPinProperties> UPCGDifferenceSettings::OutputPinProperties() const
+{
+	TArray<FPCGPinProperties> PinProperties;
+	PinProperties.Emplace(PCGPinConstants::DefaultOutputLabel, EPCGDataType::Composite);
+
+	return PinProperties;
+}
+
 FPCGElementPtr UPCGDifferenceSettings::CreateElement() const
 {
 	return MakeShared<FPCGDifferenceElement>();
diff --git a/Engine/Plugins/Experimental/PCG/Source/PCG/Private/Elements/PCGIntersectionElement.cpp b/Engine/Plugins/Experimental/PCG/Source/PCG/Private/Elements/PCGIntersectionElement.cpp
--- a/Engine/Plugins/Experimental/PCG/Source/PCG/Private/Elements/PCGIntersectionElement.cpp
+++ b/Engine/Plugins/Experimental/PCG/Source/PCG/Private/Elements/PCGIntersectionElement.cpp
@@ -5,6 +5,14 @@
 
 #include UE_INLINE_GENERATED_CPP_BY_NAME(PCGIntersectionElement)
 
+TArray<FPCGPinProperties> UPCGIntersectionSettings::OutputPinProperties() const
+{
+	TArray<FPCGPinProperties> PinProperties;
+	PinProperties.Emplace(PCGPinConstants::DefaultOutputLabel, EPCGDataType::Composite);
+
+	return PinProperties;
+}
+
 FPCGElementPtr UPCGIntersectionSettings::CreateElement() const
 {
 	return MakeShared<FPCGIntersectionElement>();
diff --git a/Engine/Plugins/Experimental/PCG/Source/PCG/Private/Elements/PCGUnionElement.cpp b/Engine/Plugins/Experimental/PCG/Source/PCG/Private/Elements/PCGUnionElement.cpp
--- a/Engine/Plugins/Experimental/PCG/Source/PCG/Private/Elements/PCGUnionElement.cpp
+++ b/Engine/Plugins/Experimental/PCG/Source/PCG/Private/Elements/PCGUnionElement.cpp
@@ -5,6 +5,14 @@
 
 #include UE_INLINE_GENERATED_CPP_BY_NAME(PCGUnionElement)
 
+TArray<FPCGPinProperties> UPCGUnionSettings::OutputPinProperties() const
+{
+	TArray<FPCGPinProperties> PinProperties;
+	PinProperties.Emplace(PCGPinConstants::DefaultOutputLabel, EPCGDataType::Composite);
+
+	return PinProperties;
+}
+
 FPCGElementPtr UPCGUnionSettings::CreateElement() const
 {
 	return MakeShared<FPCGUnionElement>();
diff --git a/Engine/Plugins/Experimental/PCG/Source/PCG/Private/PCGInputOutputSettings.cpp b/Engine/Plugins/Experimental/PCG/Source/PCG/Private/PCGInputOutputSettings.cpp
--- a/Engine/Plugins/Experimental/PCG/Source/PCG/Private/PCGInputOutputSettings.cpp
+++ b/Engine/Plugins/Experimental/PCG/Source/PCG/Private/PCGInputOutputSettings.cpp
@@ -67,7 +67,7 @@ void UPCGGraphInputOutputSettings::PostLoad()
 TArray<FPCGPinProperties> UPCGGraphInputOutputSettings::GetPinProperties() const
 {
 	TArray<FPCGPinProperties> PinProperties;
-	const EPCGDataType DefaultPinDataType = bIsInput ? EPCGDataType::Spatial : EPCGDataType::Any;
+	const EPCGDataType DefaultPinDataType = bIsInput ? EPCGDataType::Composite : EPCGDataType::Any;
 	Algo::Transform(StaticLabels(), PinProperties, [DefaultPinDataType](const FLabelAndTooltip& InLabelAndTooltip) {
 		return FPCGPinProperties(InLabelAndTooltip.Label, DefaultPinDataType, true, InLabelAndTooltip.Tooltip);
 	});
@@ -75,7 +75,8 @@ TArray<FPCGPinProperties> UPCGGraphInputOutputSettings::GetPinProperties() const
 	if (bShowAdvancedPins)
 	{
 		Algo::Transform(StaticAdvancedLabels(), PinProperties, [DefaultPinDataType](const FLabelAndTooltip& InLabelAndTooltip) {
-			return FPCGPinProperties(InLabelAndTooltip.Label, (InLabelAndTooltip.Label == PCGInputOutputConstants::DefaultLandscapeLabel) ? EPCGDataType::Surface : DefaultPinDataType, true, InLabelAndTooltip.Tooltip);
+			const EPCGDataType PinType = (InLabelAndTooltip.Label == PCGInputOutputConstants::DefaultLandscapeLabel || InLabelAndTooltip.Label == PCGInputOutputConstants::DefaultLandscapeHeightLabel) ? EPCGDataType::Surface : DefaultPinDataType;
+			return FPCGPinProperties(InLabelAndTooltip.Label, PinType, true, InLabelAndTooltip.Tooltip);
 		});
 	}
 
diff --git a/Engine/Plugins/Experimental/PCG/Source/PCG/Private/PCGPin.cpp b/Engine/Plugins/Experimental/PCG/Source/PCG/Private/PCGPin.cpp
--- a/Engine/Plugins/Experimental/PCG/Source/PCG/Private/PCGPin.cpp
+++ b/Engine/Plugins/Experimental/PCG/Source/PCG/Private/PCGPin.cpp
@@ -48,10 +48,10 @@ bool UPCGPin::AddEdgeTo(UPCGPin* OtherPin)
 	}
 
 	// This pin is upstream if the pin is an output pin
-	const bool bThisPinIsUpstream = Node->GetOutputPin(Properties.Label) == this;
+	const bool bThisPinIsUpstream = IsOutputPin();
+	const bool bOtherPinIsUpstream = OtherPin->IsOutputPin();
 
 	// Pins should not both be upstream or both be downstream..
-	const bool bOtherPinIsUpstream = OtherPin->Node->GetOutputPin(OtherPin->Properties.Label) == OtherPin;
 	if (!ensure(bThisPinIsUpstream != bOtherPinIsUpstream))
 	{
 		return false;
@@ -162,6 +162,12 @@ bool UPCGPin::IsConnected() const
 	return false;
 }
 
+bool UPCGPin::IsOutputPin() const
+{
+	check(Node);
+	return Node->GetOutputPin(Properties.Label) == this;
+}
+
 int32 UPCGPin::EdgeCount() const
 {
 	int32 EdgeNum = 0;
@@ -178,7 +184,41 @@ int32 UPCGPin::EdgeCount() const
 
 bool UPCGPin::IsCompatible(const UPCGPin* OtherPin) const
 {
-	return OtherPin && !!(Properties.AllowedTypes & OtherPin->Properties.AllowedTypes);
+	if (!OtherPin)
+	{
+		return false;
+	}
+
+	const bool bThisPinOutput = IsOutputPin();
+	if (!ensure(bThisPinOutput != OtherPin->IsOutputPin()))
+	{
+		// Cannot connect two pins of same polarity
+		return false;
+	}
+
+	// Sort pins
+	const UPCGPin* UpstreamPin = bThisPinOutput ? this : OtherPin;
+	const UPCGPin* DownstreamPin = bThisPinOutput ? OtherPin : this;
+
+	// Concrete can always be used as a composite - allow connections from concrete to composite
+	const bool bUpstreamConcrete = !!(UpstreamPin->Properties.AllowedTypes & EPCGDataType::Concrete);
+	const bool bDownstreamComposite = !!(DownstreamPin->Properties.AllowedTypes & EPCGDataType::Composite);
+	if (bUpstreamConcrete && bDownstreamComposite)
+	{
+		return true;
+	}
+
+	// Catch case when a composite type is being connected to a concrete type - that can require collapse
+	const bool bUpstreamComposite = !!(UpstreamPin->Properties.AllowedTypes & EPCGDataType::Composite);
+	const bool bDownstreamConcrete = !!(DownstreamPin->Properties.AllowedTypes & EPCGDataType::Concrete);
+	if (bUpstreamComposite && bDownstreamConcrete)
+	{
+		// This will trigger a collapse, but let it slide for now.
+		// TODO in the future we should inject a conversion node.
+		return true;
+	}
+
+	return !!(Properties.AllowedTypes & OtherPin->Properties.AllowedTypes);
 }
 
 bool UPCGPin::CanConnect(const UPCGPin* OtherPin) const
diff --git a/Engine/Plugins/Experimental/PCG/Source/PCG/Public/Elements/PCGDifferenceElement.h b/Engine/Plugins/Experimental/PCG/Source/PCG/Public/Elements/PCGDifferenceElement.h
--- a/Engine/Plugins/Experimental/PCG/Source/PCG/Public/Elements/PCGDifferenceElement.h
+++ b/Engine/Plugins/Experimental/PCG/Source/PCG/Public/Elements/PCGDifferenceElement.h
@@ -28,6 +28,7 @@ public:
 #endif
 
 	virtual TArray<FPCGPinProperties> InputPinProperties() const override;
+	virtual TArray<FPCGPinProperties> OutputPinProperties() const override;
 
 protected:
 	virtual FPCGElementPtr CreateElement() const override;
diff --git a/Engine/Plugins/Experimental/PCG/Source/PCG/Public/Elements/PCGIntersectionElement.h b/Engine/Plugins/Experimental/PCG/Source/PCG/Public/Elements/PCGIntersectionElement.h
--- a/Engine/Plugins/Experimental/PCG/Source/PCG/Public/Elements/PCGIntersectionElement.h
+++ b/Engine/Plugins/Experimental/PCG/Source/PCG/Public/Elements/PCGIntersectionElement.h
@@ -21,6 +21,8 @@ public:
 	virtual EPCGSettingsType GetType() const override { return EPCGSettingsType::Spatial; }
 #endif
 
+	virtual TArray<FPCGPinProperties> OutputPinProperties() const override;
+
 protected:
 	virtual FPCGElementPtr CreateElement() const override;
 	//~End UPCGSettings interface
diff --git a/Engine/Plugins/Experimental/PCG/Source/PCG/Public/Elements/PCGUnionElement.h b/Engine/Plugins/Experimental/PCG/Source/PCG/Public/Elements/PCGUnionElement.h
--- a/Engine/Plugins/Experimental/PCG/Source/PCG/Public/Elements/PCGUnionElement.h
+++ b/Engine/Plugins/Experimental/PCG/Source/PCG/Public/Elements/PCGUnionElement.h
@@ -21,6 +21,8 @@ public:
 	virtual EPCGSettingsType GetType() const override { return EPCGSettingsType::Spatial; }
 #endif
 
+	virtual TArray<FPCGPinProperties> OutputPinProperties() const override;
+
 protected:
 	virtual FPCGElementPtr CreateElement() const override;
 	//~End UPCGSettings interface
diff --git a/Engine/Plugins/Experimental/PCG/Source/PCG/Public/PCGCommon.h b/Engine/Plugins/Experimental/PCG/Source/PCG/Public/PCGCommon.h
--- a/Engine/Plugins/Experimental/PCG/Source/PCG/Public/PCGCommon.h
+++ b/Engine/Plugins/Experimental/PCG/Source/PCG/Public/PCGCommon.h
@@ -39,7 +39,13 @@ enum class EPCGDataType : uint32
 	Volume = 1 << 7,
 	Primitive = 1 << 8,
 
-	Spatial = Point | PolyLine | Surface | Volume | Primitive,
+	/** Simple concrete data. */
+	Concrete = Point | PolyLine | Surface | Volume | Primitive,
+
+	/** Composite spatial data which can include boolean operations. Points can be sampled from the set without collapsing to concrete data. */
+	Composite = 1 << 9,
+
+	Spatial = Composite | Concrete,
 
 	Param = 1 << 27,
 	Settings = 1 << 28 UMETA(Hidden),
diff --git a/Engine/Plugins/Experimental/PCG/Source/PCG/Public/PCGPin.h b/Engine/Plugins/Experimental/PCG/Source/PCG/Public/PCGPin.h
--- a/Engine/Plugins/Experimental/PCG/Source/PCG/Public/PCGPin.h
+++ b/Engine/Plugins/Experimental/PCG/Source/PCG/Public/PCGPin.h
@@ -64,5 +64,6 @@ public:
 	bool BreakAllEdges();
 	bool BreakAllIncompatibleEdges();
 	bool IsConnected() const;
+	bool IsOutputPin() const;
 	int32 EdgeCount() const;
 };
diff --git a/Engine/Plugins/Experimental/PCG/Source/PCGEditor/Private/PCGEditorCommon.h b/Engine/Plugins/Experimental/PCG/Source/PCGEditor/Private/PCGEditorCommon.h
--- a/Engine/Plugins/Experimental/PCG/Source/PCGEditor/Private/PCGEditorCommon.h
+++ b/Engine/Plugins/Experimental/PCG/Source/PCGEditor/Private/PCGEditorCommon.h
@@ -32,6 +32,9 @@ namespace FPCGEditorCommon
 	const FName SettingsDataType = FName(TEXT("Settings Data"));
 	const FName OtherDataType = FName(TEXT("Other Data"));
 
+	const FName CompositeDataType = FName(TEXT("Composite Data"));
+	
+	const FName ConcreteDataType = FName(TEXT("Concrete Data"));
 	const FName PointDataType = FName(TEXT("Point Data"));
 	const FName PolyLineDataType = FName(TEXT("Poly Line Data"));
 	const FName SurfaceDataType = FName(TEXT("Surface Data"));
diff --git a/Engine/Plugins/Experimental/PCG/Source/PCGEditor/Private/PCGEditorGraphNodeBase.cpp b/Engine/Plugins/Experimental/PCG/Source/PCGEditor/Private/PCGEditorGraphNodeBase.cpp
--- a/Engine/Plugins/Experimental/PCG/Source/PCGEditor/Private/PCGEditorGraphNodeBase.cpp
+++ b/Engine/Plugins/Experimental/PCG/Source/PCGEditor/Private/PCGEditorGraphNodeBase.cpp
@@ -379,7 +379,15 @@ FEdGraphPinType UPCGEditorGraphNodeBase::GetPinType(const UPCGPin* InPin)
 		EdPinType.PinCategory = FPCGEditorCommon::SpatialDataType;
 
 		// Assign subcategory if we have precise information
-		if (CheckType(EPCGDataType::Point))
+		if (CheckType(EPCGDataType::Composite))
+		{
+			EdPinType.PinCategory = FPCGEditorCommon::CompositeDataType;
+		}
+		else if (CheckType(EPCGDataType::Concrete))
+		{
+			EdPinType.PinCategory = FPCGEditorCommon::ConcreteDataType;
+		}
+		else if (CheckType(EPCGDataType::Point))
 		{
 			EdPinType.PinSubCategory = FPCGEditorCommon::PointDataType;
 		}
@@ -415,7 +423,7 @@ FEdGraphPinType UPCGEditorGraphNodeBase::GetPinType(const UPCGPin* InPin)
 	else if (CheckType(EPCGDataType::Other))
 	{
 		EdPinType.PinCategory = FPCGEditorCommon::OtherDataType;
-	}	
+	}
 
 	return EdPinType;
 }
diff --git a/Engine/Plugins/Experimental/PCG/Source/PCGEditor/Private/PCGEditorSettings.cpp b/Engine/Plugins/Experimental/PCG/Source/PCGEditor/Private/PCGEditorSettings.cpp
--- a/Engine/Plugins/Experimental/PCG/Source/PCGEditor/Private/PCGEditorSettings.cpp
+++ b/Engine/Plugins/Experimental/PCG/Source/PCGEditor/Private/PCGEditorSettings.cpp
@@ -24,6 +24,8 @@ UPCGEditorSettings::UPCGEditorSettings(const FObjectInitializer& ObjectInitializ
 
 	DefaultPinColor = FLinearColor(1.0f, 1.0f, 1.0f);
 	SpatialDataPinColor = FLinearColor(0.2f, 0.2f, 1.0f);
+	CompositeDataPinColor = FLinearColor(0.2f, 0.2f, 1.0f);
+	ConcreteDataPinColor = FLinearColor(0.2f, 0.2f, 1.0f);
 	RenderTargetDataPinColor = FLinearColor(1.0f, 0.3f, 0.f);
 	ParamDataPinColor = FLinearColor(1.0f, 0.6f, 0.0f);
 	UnknownDataPinColor = FLinearColor(0.3f, 0.3f, 0.3f);
@@ -96,17 +98,26 @@ FLinearColor UPCGEditorSettings::GetColor(UPCGSettings* Settings) const
 
 FLinearColor UPCGEditorSettings::GetPinColor(const FEdGraphPinType& PinType) const
 {
-	if (PinType.PinCategory == FPCGEditorCommon::SpatialDataType)
+	if (PinType.PinCategory == FPCGEditorCommon::SpatialDataType || PinType.PinCategory == FPCGEditorCommon::ConcreteDataType)
 	{
+		// Clauses below try to pick the narrowest type possible, falling back to Spatial
 		if (PinType.PinSubCategory == FPCGEditorCommon::RenderTargetDataType)
 		{
 			return RenderTargetDataPinColor;
 		}
+		else if (PinType.PinCategory == FPCGEditorCommon::ConcreteDataType)
+		{
+			return ConcreteDataPinColor;
+		}
 		else
 		{
 			return SpatialDataPinColor;
 		}
 	}
+	else if (PinType.PinCategory == FPCGEditorCommon::CompositeDataType)
+	{
+		return CompositeDataPinColor;
+	}
 	else if (PinType.PinCategory == FPCGEditorCommon::ParamDataType)
 	{
 		return ParamDataPinColor;
diff --git a/Engine/Plugins/Experimental/PCG/Source/PCGEditor/Public/PCGEditorSettings.h b/Engine/Plugins/Experimental/PCG/Source/PCGEditor/Public/PCGEditorSettings.h
--- a/Engine/Plugins/Experimental/PCG/Source/PCGEditor/Public/PCGEditorSettings.h
+++ b/Engine/Plugins/Experimental/PCG/Source/PCGEditor/Public/PCGEditorSettings.h
@@ -81,6 +81,14 @@ public:
 	UPROPERTY(EditAnywhere, config, Category = Node, meta = (HideAlphaChannel))
 	FLinearColor SpatialDataPinColor;
 
+	/** Color used for concrete/simple spatial data pins */
+	UPROPERTY(EditAnywhere, config, Category = Node, meta = (HideAlphaChannel))
+	FLinearColor ConcreteDataPinColor;
+
+	/** Color used for composite shape spatial data pins */
+	UPROPERTY(EditAnywhere, config, Category = Node, meta = (HideAlphaChannel))
+	FLinearColor CompositeDataPinColor;
+
 	/** Color used for render targets / GPU pins */
 	UPROPERTY(EditAnywhere, config, Category = Node, meta = (HideAlphaChannel))
 	FLinearColor RenderTargetDataPinColor;
