[PCG] Fix non-unity build

diff --git a/Engine/Plugins/Experimental/PCG/Source/PCG/Private/PCGNode.cpp b/Engine/Plugins/Experimental/PCG/Source/PCG/Private/PCGNode.cpp
--- a/Engine/Plugins/Experimental/PCG/Source/PCG/Private/PCGNode.cpp
+++ b/Engine/Plugins/Experimental/PCG/Source/PCG/Private/PCGNode.cpp
@@ -1,11 +1,13 @@
 // Copyright Epic Games, Inc. All Rights Reserved.
 
 #include "PCGNode.h"
+
 #include "PCGEdge.h"
 #include "PCGGraph.h"
 #include "PCGSettings.h"
 #include "PCGSubgraph.h"
 
+#include "Algo/Find.h"
 #include "Algo/Transform.h"
 
 #include UE_INLINE_GENERATED_CPP_BY_NAME(PCGNode)
