[PCG] Fix warning around logic, and simplify at same time

diff --git a/Engine/Plugins/Experimental/PCG/Source/PCGEditor/Private/PCGEditorGraphSchema.cpp b/Engine/Plugins/Experimental/PCG/Source/PCGEditor/Private/PCGEditorGraphSchema.cpp
--- a/Engine/Plugins/Experimental/PCG/Source/PCGEditor/Private/PCGEditorGraphSchema.cpp
+++ b/Engine/Plugins/Experimental/PCG/Source/PCGEditor/Private/PCGEditorGraphSchema.cpp
@@ -431,7 +431,9 @@ void FPCGEditorConnectionDrawingPolicy::DetermineWiringStyle(UEdGraphPin* Output
 
 				// If this wire does not connect to the pass-through pin, or if the wire is not the first link on that pin, then it will
 				// not be used by the node. In this case desaturate the wire.
-				if (PassThroughInputPin && PassThroughInputPin->Properties.Label != InputPin->GetFName() || InputPin->LinkedTo[0] != OutputPin)
+				const bool bConnectedToPassThrough = PassThroughInputPin && PassThroughInputPin->Properties.Label == InputPin->GetFName();
+				const bool bFirstWire = InputPin->LinkedTo[0] == OutputPin;
+				if (!bConnectedToPassThrough || !bFirstWire)
 				{
 					Params.WireColor = Params.WireColor.Desaturate(0.7f);
 				}
