Adding missing pin on debug node
Fixed mesh offset issue that is misleading for debug display
Clarified error message in copy points node

#rnx

diff --git a/Engine/Plugins/Experimental/PCG/Source/PCG/Private/Elements/PCGCopyPoints.cpp b/Engine/Plugins/Experimental/PCG/Source/PCG/Private/Elements/PCGCopyPoints.cpp
--- a/Engine/Plugins/Experimental/PCG/Source/PCG/Private/Elements/PCGCopyPoints.cpp
+++ b/Engine/Plugins/Experimental/PCG/Source/PCG/Private/Elements/PCGCopyPoints.cpp
@@ -46,7 +46,7 @@ bool FPCGCopyPointsElement::ExecuteInternal(FPCGContext* Context) const
 
 	if (Sources.Num() != 1 || Targets.Num() != 1)
 	{
-		PCGE_LOG(Error, "Invalid number of inputs (should be exactly one source and one target)");
+		PCGE_LOG(Error, "Invalid number of inputs - Expected one source, got %d; Expected one target, got %d.", Sources.Num(), Targets.Num());
 		return true;
 	}
 	
diff --git a/Engine/Plugins/Experimental/PCG/Source/PCG/Private/Elements/PCGDebugElement.cpp b/Engine/Plugins/Experimental/PCG/Source/PCG/Private/Elements/PCGDebugElement.cpp
--- a/Engine/Plugins/Experimental/PCG/Source/PCG/Private/Elements/PCGDebugElement.cpp
+++ b/Engine/Plugins/Experimental/PCG/Source/PCG/Private/Elements/PCGDebugElement.cpp
@@ -133,7 +133,7 @@ namespace PCGDebugElement
 				else if (bScaleWithExtents)
 				{
 					const FVector ScaleWithExtents = Point.GetExtents() / MeshExtents;
-					const FVector TransformedBoxCenterWithOffset = InstanceTransform.TransformPosition(Point.GetLocalCenter()) - InstanceTransform.GetLocation() - MeshCenter * ScaleWithExtents;
+					const FVector TransformedBoxCenterWithOffset = InstanceTransform.TransformPosition(Point.GetLocalCenter()) - InstanceTransform.GetLocation();
 					InstanceTransform.SetTranslation(InstanceTransform.GetTranslation() + TransformedBoxCenterWithOffset);
 					InstanceTransform.SetScale3D(InstanceTransform.GetScale3D() * ScaleWithExtents);
 				}
@@ -188,6 +188,14 @@ FPCGElementPtr UPCGDebugSettings::CreateElement() const
 	return MakeShared<FPCGDebugElement>();
 }
 
+TArray<FPCGPinProperties> UPCGDebugSettings::OutputPinProperties() const
+{
+	TArray<FPCGPinProperties> PinProperties;
+	PinProperties.Emplace(PCGPinConstants::DefaultOutputLabel, EPCGDataType::Any);
+
+	return PinProperties;
+}
+
 bool FPCGDebugElement::ExecuteInternal(FPCGContext* Context) const
 {
 	TRACE_CPUPROFILER_EVENT_SCOPE(FPCGDebugElement::Execute);
diff --git a/Engine/Plugins/Experimental/PCG/Source/PCG/Public/Elements/PCGDebugElement.h b/Engine/Plugins/Experimental/PCG/Source/PCG/Public/Elements/PCGDebugElement.h
--- a/Engine/Plugins/Experimental/PCG/Source/PCG/Public/Elements/PCGDebugElement.h
+++ b/Engine/Plugins/Experimental/PCG/Source/PCG/Public/Elements/PCGDebugElement.h
@@ -24,7 +24,7 @@ public:
 	virtual EPCGSettingsType GetType() const override { return EPCGSettingsType::Debug; }
 #endif
 
-	virtual TArray<FPCGPinProperties> OutputPinProperties() const override { return TArray<FPCGPinProperties>(); }
+	virtual TArray<FPCGPinProperties> OutputPinProperties() const override;
 
 protected:
 	virtual FPCGElementPtr CreateElement() const override;
