Fixed reference issue in the partition actors that would hold on potentially larger spatially loaded objects, and this would interfere with streaming.
Also removed cleanup on end play, since this could interfere with "not visible but loaded" states
Note: might require a cleanup+generate to take into account

#rnx

diff --git a/Engine/Plugins/Experimental/PCG/Source/PCG/Private/Grid/PCGPartitionActor.cpp b/Engine/Plugins/Experimental/PCG/Source/PCG/Private/Grid/PCGPartitionActor.cpp
--- a/Engine/Plugins/Experimental/PCG/Source/PCG/Private/Grid/PCGPartitionActor.cpp
+++ b/Engine/Plugins/Experimental/PCG/Source/PCG/Private/Grid/PCGPartitionActor.cpp
@@ -47,6 +47,18 @@ void APCGPartitionActor::PostLoad()
 		PCGGridSize = APCGWorldActor::DefaultPartitionGridSize;
 	}
 
+#if WITH_EDITOR
+	if (LocalToOriginalMap_DEPRECATED.Num() > 0)
+	{
+		for (const TPair<TObjectPtr<UPCGComponent>, TWeakObjectPtr<UPCGComponent>>& LocalToOriginalPair : LocalToOriginalMap_DEPRECATED)
+		{
+			LocalToOriginal.Emplace(LocalToOriginalPair.Key, LocalToOriginalPair.Value.Get());
+		}
+
+		LocalToOriginalMap_DEPRECATED.Reset();
+	}
+#endif
+
 	// Make sure that we don't track objects that do not exist anymore
 	CleanupDeadGraphInstances();
 
@@ -96,6 +108,14 @@ void APCGPartitionActor::PostRegisterAllComponents()
 	}
 #endif // WITH_EDITOR
 
+	// Reset the OriginalToLocal and build it from the local map
+	OriginalToLocal.Reset();
+	OriginalToLocal.Reserve(LocalToOriginal.Num());
+	for (const auto& It : LocalToOriginal)
+	{
+		OriginalToLocal.Add(It.Value, It.Key);
+	}
+
 	// Make the Partition actor register itself to the PCG Subsystem
 	// Always do it at runtime, wait for the post load/creation in editor
 	// Only do the mapping if we are at runtime.
@@ -119,7 +139,7 @@ void APCGPartitionActor::PostRegisterAllComponents()
 void APCGPartitionActor::BeginPlay()
 {
 	// Pass through all the pcg components, to verify if we need to generate them
-	for (auto& It : OriginalToLocalMap)
+	for (auto& It : OriginalToLocal)
 	{
 		if (It.Key && It.Value)
 		{
@@ -196,21 +216,21 @@ void APCGPartitionActor::GetActorBounds(bool bOnlyCollidingComponents, FVector&
 
 UPCGComponent* APCGPartitionActor::GetLocalComponent(const UPCGComponent* OriginalComponent) const
 {
-	const TObjectPtr<UPCGComponent>* LocalComponent = OriginalToLocalMap.Find(OriginalComponent);
+	const TObjectPtr<UPCGComponent>* LocalComponent = OriginalToLocal.Find(OriginalComponent);
 	return LocalComponent ? *LocalComponent : nullptr;
 }
 
 UPCGComponent* APCGPartitionActor::GetOriginalComponent(const UPCGComponent* LocalComponent) const
 {
-	const TObjectPtr<UPCGComponent>* OriginalComponent = LocalToOriginalMap.Find(LocalComponent);
+	const TSoftObjectPtr<UPCGComponent>* OriginalComponent = LocalToOriginal.Find(LocalComponent);
 	return OriginalComponent ? (*OriginalComponent).Get() : nullptr;
 }
 
 void APCGPartitionActor::CleanupDeadGraphInstances()
 {
 	// First find if we have any local dead instance (= nullptr) hooked to an original component.
-	TSet<TObjectPtr<UPCGComponent>> DeadOriginalInstances;
-	for (const auto& OriginalToLocalItem : OriginalToLocalMap)
+	TSet<TSoftObjectPtr<UPCGComponent>> DeadOriginalInstances;
+	for (const auto& OriginalToLocalItem : OriginalToLocal)
 	{
 		if (!OriginalToLocalItem.Value)
 		{
@@ -222,32 +242,31 @@ void APCGPartitionActor::CleanupDeadGraphInstances()
 	{
 		Modify();
 
-		for (const TObjectPtr<UPCGComponent>& DeadInstance : DeadOriginalInstances)
+		for (const TSoftObjectPtr<UPCGComponent>& DeadInstance : DeadOriginalInstances)
 		{
-			OriginalToLocalMap.Remove(DeadInstance);
+			OriginalToLocal.Remove(DeadInstance);
 		}
 
-		LocalToOriginalMap.Remove(nullptr);
+		LocalToOriginal.Remove(nullptr);
 	}
 
 	// And do the same with dead original ones.
 	TSet<TObjectPtr<UPCGComponent>> DeadLocalInstances;
-	for (const auto& LocalToOriginalItem : LocalToOriginalMap)
+	for (const auto& LocalToOriginalItem : LocalToOriginal)
 	{
-		if (!LocalToOriginalItem.Value)
+		if (!LocalToOriginalItem.Value.IsValid())
 		{
 			DeadLocalInstances.Add(LocalToOriginalItem.Key);
 		}
 	}
 
-
 	if (!DeadLocalInstances.IsEmpty())
 	{
 		Modify();
 
 		for (const TObjectPtr<UPCGComponent>& DeadInstance : DeadLocalInstances)
 		{
-			LocalToOriginalMap.Remove(DeadInstance);
+			LocalToOriginal.Remove(DeadInstance);
 
 			if (DeadInstance)
 			{
@@ -257,7 +276,7 @@ void APCGPartitionActor::CleanupDeadGraphInstances()
 		}
 
 		// Remove all dead entries
-		OriginalToLocalMap.Remove(nullptr);
+		OriginalToLocal.Remove(nullptr);
 	}
 }
 
@@ -291,8 +310,8 @@ void APCGPartitionActor::AddGraphInstance(UPCGComponent* OriginalComponent)
 	// TODO: check if we should use a non-instanced component?
 	AddInstanceComponent(LocalComponent);
 
-	OriginalToLocalMap.Add(OriginalComponent, LocalComponent);
-	LocalToOriginalMap.Add(LocalComponent, OriginalComponent);
+	OriginalToLocal.Emplace(OriginalComponent, LocalComponent);
+	LocalToOriginal.Emplace(LocalComponent, OriginalComponent);
 }
 
 void APCGPartitionActor::RemapGraphInstance(const UPCGComponent* OldOriginalComponent, UPCGComponent* NewOriginalComponent)
@@ -306,12 +325,12 @@ void APCGPartitionActor::RemapGraphInstance(const UPCGComponent* OldOriginalComp
 
 	Modify();
 
-	OriginalToLocalMap.Remove(OldOriginalComponent);
-	LocalToOriginalMap.Remove(LocalComponent);
+	OriginalToLocal.Remove(OldOriginalComponent);
+	LocalToOriginal.Remove(LocalComponent);
 
 	LocalComponent->SetPropertiesFromOriginal(NewOriginalComponent);
-	OriginalToLocalMap.Add(NewOriginalComponent, LocalComponent);
-	LocalToOriginalMap.Add(LocalComponent, NewOriginalComponent);
+	OriginalToLocal.Emplace(NewOriginalComponent, LocalComponent);
+	LocalToOriginal.Emplace(LocalComponent, NewOriginalComponent);
 }
 
 bool APCGPartitionActor::RemoveGraphInstance(UPCGComponent* OriginalComponent)
@@ -328,14 +347,14 @@ bool APCGPartitionActor::RemoveGraphInstance(UPCGComponent* OriginalComponent)
 
 	Modify();
 
-	OriginalToLocalMap.Remove(OriginalComponent);
-	LocalToOriginalMap.Remove(LocalComponent);
+	OriginalToLocal.Remove(OriginalComponent);
+	LocalToOriginal.Remove(LocalComponent);
 
 	// TODO Add option to not cleanup?
 	LocalComponent->CleanupLocalImmediate(/*bRemoveComponents=*/true);
 	LocalComponent->DestroyComponent();
 
-	return OriginalToLocalMap.IsEmpty();
+	return OriginalToLocal.IsEmpty();
 }
 
 #if WITH_EDITOR
@@ -404,7 +423,7 @@ bool APCGPartitionActor::IsSafeForDeletion() const
 TSet<TObjectPtr<UPCGComponent>> APCGPartitionActor::GetAllLocalPCGComponents() const
 {
 	TSet<TObjectPtr<UPCGComponent>> ResultComponents;
-	LocalToOriginalMap.GetKeys(ResultComponents);
+	LocalToOriginal.GetKeys(ResultComponents);
 
 	return ResultComponents;
 }
@@ -412,7 +431,10 @@ TSet<TObjectPtr<UPCGComponent>> APCGPartitionActor::GetAllLocalPCGComponents() c
 TSet<TObjectPtr<UPCGComponent>> APCGPartitionActor::GetAllOriginalPCGComponents() const
 {
 	TSet<TObjectPtr<UPCGComponent>> ResultComponents;
-	OriginalToLocalMap.GetKeys(ResultComponents);
+	for(const TPair<TSoftObjectPtr<UPCGComponent>, TObjectPtr<UPCGComponent>>& OriginalPair : OriginalToLocal)
+	{
+		ResultComponents.Add(OriginalPair.Key.Get());
+	}
 
 	return ResultComponents;
 }
diff --git a/Engine/Plugins/Experimental/PCG/Source/PCG/Private/Grid/PCGPartitionActor.h b/Engine/Plugins/Experimental/PCG/Source/PCG/Private/Grid/PCGPartitionActor.h
--- a/Engine/Plugins/Experimental/PCG/Source/PCG/Private/Grid/PCGPartitionActor.h
+++ b/Engine/Plugins/Experimental/PCG/Source/PCG/Private/Grid/PCGPartitionActor.h
@@ -86,12 +86,16 @@ public:
 private:
 	UPCGSubsystem* GetSubsystem() const;
 
-	// TODO: Make these properties editor only (see comment before).
+	// Note: this map is not a property and not serialized since we will rebuild it from the LocalToOriginal
+	TMap<TSoftObjectPtr<UPCGComponent>, TObjectPtr<UPCGComponent>> OriginalToLocal;
+
 	UPROPERTY()
-	TMap<TObjectPtr<UPCGComponent>, TObjectPtr<UPCGComponent>> OriginalToLocalMap;
+	TMap<TObjectPtr<UPCGComponent>, TSoftObjectPtr<UPCGComponent>> LocalToOriginal;
 
+#if WITH_EDITORONLY_DATA
 	UPROPERTY()
-	TMap<TObjectPtr<UPCGComponent>, TObjectPtr<UPCGComponent>> LocalToOriginalMap;
+	TMap<TObjectPtr<UPCGComponent>, TWeakObjectPtr<UPCGComponent>> LocalToOriginalMap_DEPRECATED;
+#endif
 
 	UPROPERTY(VisibleAnywhere, Category = WorldPartition)
 	uint32 PCGGridSize;
diff --git a/Engine/Plugins/Experimental/PCG/Source/PCG/Private/PCGSubsystem.cpp b/Engine/Plugins/Experimental/PCG/Source/PCG/Private/PCGSubsystem.cpp
--- a/Engine/Plugins/Experimental/PCG/Source/PCG/Private/PCGSubsystem.cpp
+++ b/Engine/Plugins/Experimental/PCG/Source/PCG/Private/PCGSubsystem.cpp
@@ -826,7 +826,6 @@ void UPCGSubsystem::UnregisterPartitionActor(APCGPartitionActor* Actor)
 				TSet<TObjectPtr<APCGPartitionActor>>* PartitionActorsPtr = ComponentToPartitionActorsMap.Find(Component);
 				if (PartitionActorsPtr)
 				{
-					Actor->RemoveGraphInstance(Component);
 					PartitionActorsPtr->Remove(Actor);
 				}
 			});
