[PCG] Refresh PointFilter to use the latest accessors

* Try to re-use compare operators defined in the traits
* Also use those operators in the Compare Op Element
* Also factorize constant definition between point filter and create attribute

diff --git a/Engine/Plugins/Experimental/PCG/Source/PCG/Private/Elements/Metadata/PCGMetadataCompareOpElement.cpp b/Engine/Plugins/Experimental/PCG/Source/PCG/Private/Elements/Metadata/PCGMetadataCompareOpElement.cpp
--- a/Engine/Plugins/Experimental/PCG/Source/PCG/Private/Elements/Metadata/PCGMetadataCompareOpElement.cpp
+++ b/Engine/Plugins/Experimental/PCG/Source/PCG/Private/Elements/Metadata/PCGMetadataCompareOpElement.cpp
@@ -6,6 +6,7 @@
 #include "Metadata/PCGMetadata.h"
 #include "Metadata/PCGMetadataAttribute.h"
 #include "Metadata/PCGMetadataAttributeTpl.h"
+#include "Metadata/PCGMetadataAttributeTraits.h"
 #include "Elements/Metadata/PCGMetadataElementCommon.h"
 
 #include UE_INLINE_GENERATED_CPP_BY_NAME(PCGMetadataCompareOpElement)
@@ -15,39 +16,33 @@ namespace PCGMetadataCompareSettings
 	template <typename T>
 	bool ApplyCompare(const T& Input1, const T& Input2, EPCGMedadataCompareOperation Operation, double Tolerance)
 	{
-		switch (Operation)
+		if (Operation == EPCGMedadataCompareOperation::Equal)
 		{
-		case EPCGMedadataCompareOperation::Equal:
+			return PCG::Private::MetadataTraits<T>::Equal(Input1, Input2);
+		}
+		else if (Operation == EPCGMedadataCompareOperation::NotEqual)
 		{
-			if constexpr (std::is_same_v<T, int32> || std::is_same_v<T, int64>)
-			{
-				return Input1 == Input2;
-			}
-			else
-			{
-				return FMath::IsNearlyEqual(Input1, Input2, static_cast<T>(Tolerance));
-			}
+			return !PCG::Private::MetadataTraits<T>::Equal(Input1, Input2);
 		}
-		case EPCGMedadataCompareOperation::NotEqual:
+
+		if constexpr (PCG::Private::MetadataTraits<T>::CanCompare)
 		{
-			if constexpr (std::is_same_v<T, int32> || std::is_same_v<T, int64>)
+			switch (Operation)
 			{
-				return Input1 != Input2;
-			}
-			else
-			{
-				return !FMath::IsNearlyEqual(Input1, Input2, static_cast<T>(Tolerance));
+			case EPCGMedadataCompareOperation::Greater:
+				return PCG::Private::MetadataTraits<T>::Greater(Input1, Input2);
+			case EPCGMedadataCompareOperation::GreaterOrEqual:
+				return PCG::Private::MetadataTraits<T>::GreaterOrEqual(Input1, Input2);
+			case EPCGMedadataCompareOperation::Less:
+				return PCG::Private::MetadataTraits<T>::Less(Input1, Input2);
+			case EPCGMedadataCompareOperation::LessOrEqual:
+				return PCG::Private::MetadataTraits<T>::LessOrEqual(Input1, Input2);
+			default:
+				return false;
 			}
 		}
-		case EPCGMedadataCompareOperation::Greater:
-			return Input1 > Input2;
-		case EPCGMedadataCompareOperation::GreaterOrEqual:
-			return Input1 >= Input2;
-		case EPCGMedadataCompareOperation::Less:
-			return Input1 < Input2;
-		case EPCGMedadataCompareOperation::LessOrEqual:
-			return Input1 <= Input2;
-		default:
+		else
+		{
 			return false;
 		}
 	}
@@ -93,7 +88,19 @@ uint32 UPCGMetadataCompareSettings::GetInputPinNum() const
 bool UPCGMetadataCompareSettings::IsSupportedInputType(uint16 TypeId, uint32 InputIndex, bool& bHasSpecialRequirement) const
 {
 	bHasSpecialRequirement = false;
-	return PCG::Private::IsOfTypes<int32, int64, float, double>(TypeId);
+	if (Operation == EPCGMedadataCompareOperation::Equal || Operation == EPCGMedadataCompareOperation::NotEqual)
+	{
+		return PCG::Private::IsPCGType(TypeId);
+	}
+	else
+	{
+		auto IsValid = [](auto Dummy) -> bool
+		{
+			return PCG::Private::MetadataTraits<decltype(Dummy)>::CanCompare;
+		};
+
+		return PCGMetadataAttribute::CallbackWithRightType(TypeId, IsValid);
+	}
 }
 
 FPCGAttributePropertySelector UPCGMetadataCompareSettings::GetInputSource(uint32 Index) const
@@ -116,7 +123,7 @@ uint16 UPCGMetadataCompareSettings::GetOutputType(uint16 InputTypeId) const
 
 FName UPCGMetadataCompareSettings::AdditionalTaskName() const
 {
-	if (const UEnum* EnumPtr = FindObject<UEnum>(nullptr, TEXT("/Script/PCG.EPCGMedadataCompareOperation"), true))
+	if (const UEnum* EnumPtr = StaticEnum<EPCGMedadataCompareOperation>())
 	{
 		return FName(FString("Compare: ") + EnumPtr->GetNameStringByValue(static_cast<int>(Operation)));
 	}
@@ -144,25 +151,15 @@ bool FPCGMetadataCompareElement::DoOperation(FOperationData& OperationData) cons
 
 	const UPCGMetadataCompareSettings* Settings = CastChecked<UPCGMetadataCompareSettings>(OperationData.Settings);
 
-	auto CompareFunc = [this, Operation = Settings->Operation, Tolerance = Settings->Tolerance, &OperationData](auto DummyValue)
+	auto CompareFunc = [this, Operation = Settings->Operation, Tolerance = Settings->Tolerance, &OperationData](auto DummyValue) -> bool
 	{
 		using AttributeType = decltype(DummyValue);
 
-		// Need to remove types that would not compile
-		if constexpr (!PCG::Private::IsOfTypes<AttributeType, int32, int64, float, double>())
-		{
-			return;
-		}
-		else
-		{
-			DoBinaryOp<AttributeType, AttributeType>(OperationData, 
-				[Operation, Tolerance](const AttributeType& Value1, const AttributeType& Value2) -> bool { 
-					return PCGMetadataCompareSettings::ApplyCompare(Value1, Value2, Operation, Tolerance); 
-				});
-		}
+		return DoBinaryOp<AttributeType, AttributeType>(OperationData, 
+			[Operation, Tolerance](const AttributeType& Value1, const AttributeType& Value2) -> bool { 
+				return PCGMetadataCompareSettings::ApplyCompare(Value1, Value2, Operation, Tolerance); 
+			});
 	};
 
-	PCGMetadataAttribute::CallbackWithRightType(OperationData.MostComplexInputType, CompareFunc);
-
-	return true;
+	return PCGMetadataAttribute::CallbackWithRightType(OperationData.MostComplexInputType, CompareFunc);
 }
diff --git a/Engine/Plugins/Experimental/PCG/Source/PCG/Private/Elements/Metadata/PCGMetadataOpElementBase.cpp b/Engine/Plugins/Experimental/PCG/Source/PCG/Private/Elements/Metadata/PCGMetadataOpElementBase.cpp
--- a/Engine/Plugins/Experimental/PCG/Source/PCG/Private/Elements/Metadata/PCGMetadataOpElementBase.cpp
+++ b/Engine/Plugins/Experimental/PCG/Source/PCG/Private/Elements/Metadata/PCGMetadataOpElementBase.cpp
@@ -59,11 +59,6 @@ TArray<FPCGPinProperties> UPCGMetadataSettingsBase::OutputPinProperties() const
 	return PinProperties;
 }
 
-bool UPCGMetadataSettingsBase::IsMoreComplexType(uint16 FirstType, uint16 SecondType) const
-{
-	return FirstType != SecondType && FirstType <= (uint16)(EPCGMetadataTypes::Count) && SecondType <= (uint16)(EPCGMetadataTypes::Count) && PCG::Private::BroadcastableTypes[SecondType][FirstType];
-}
-
 bool FPCGMetadataElementBase::ExecuteInternal(FPCGContext* Context) const
 {
 	TRACE_CPUPROFILER_EVENT_SCOPE(FPCGMetadataElementBase::Execute);
@@ -179,7 +174,7 @@ bool FPCGMetadataElementBase::ExecuteInternal(FPCGContext* Context) const
 		if (!bHasSpecialRequirement)
 		{
 			// In this case, check if we have a more complex type, or if we can broadcast to the most complex type.
-			if (OperationData.MostComplexInputType == (uint16)EPCGMetadataTypes::Unknown || Settings->IsMoreComplexType(AttributeTypeId, OperationData.MostComplexInputType))
+			if (OperationData.MostComplexInputType == (uint16)EPCGMetadataTypes::Unknown || PCG::Private::IsMoreComplexType(AttributeTypeId, OperationData.MostComplexInputType))
 			{
 				OperationData.MostComplexInputType = AttributeTypeId;
 			}
diff --git a/Engine/Plugins/Experimental/PCG/Source/PCG/Private/Elements/PCGCreateAttribute.cpp b/Engine/Plugins/Experimental/PCG/Source/PCG/Private/Elements/PCGCreateAttribute.cpp
--- a/Engine/Plugins/Experimental/PCG/Source/PCG/Private/Elements/PCGCreateAttribute.cpp
+++ b/Engine/Plugins/Experimental/PCG/Source/PCG/Private/Elements/PCGCreateAttribute.cpp
@@ -18,53 +18,32 @@ namespace PCGCreateAttributeConstants
 	const FName SourceLabel = TEXT("Source");
 }
 
-namespace PCGCreateAttributeElement
+void UPCGCreateAttributeSettings::PostLoad()
 {
-	template <typename Func>
-	decltype(auto) Dispatcher(const UPCGCreateAttributeSettings* Settings, const UPCGParamData* Params, Func Callback)
-	{
-		using ReturnType = decltype(Callback(double{}));
+	Super::PostLoad();
 
-		switch (Settings->Type)
-		{
-		case EPCGMetadataTypes::Integer64:
-			return Callback(PCG_GET_OVERRIDEN_VALUE(Settings, IntValue, Params));
-		case EPCGMetadataTypes::Integer32:
-			return Callback(PCG_GET_OVERRIDEN_VALUE(Settings, Int32Value, Params));
-		case EPCGMetadataTypes::Float:
-			return Callback(PCG_GET_OVERRIDEN_VALUE(Settings, FloatValue, Params));
-		case EPCGMetadataTypes::Double:
-			return Callback(PCG_GET_OVERRIDEN_VALUE(Settings, DoubleValue, Params));
-		case EPCGMetadataTypes::Vector2:
-			return Callback(PCG_GET_OVERRIDEN_VALUE(Settings, Vector2Value, Params));
-		case EPCGMetadataTypes::Vector:
-			return Callback(PCG_GET_OVERRIDEN_VALUE(Settings, VectorValue, Params));
-		case EPCGMetadataTypes::Vector4:
-			return Callback(PCG_GET_OVERRIDEN_VALUE(Settings, Vector4Value, Params));
-		case EPCGMetadataTypes::Quaternion:
-			return Callback(PCG_GET_OVERRIDEN_VALUE(Settings, QuatValue, Params));
-		case EPCGMetadataTypes::Transform:
-			return Callback(PCG_GET_OVERRIDEN_VALUE(Settings, TransformValue, Params));
-		case EPCGMetadataTypes::String:
-			return Callback(PCG_GET_OVERRIDEN_VALUE(Settings, StringValue, Params));
-		case EPCGMetadataTypes::Boolean:
-			return Callback(PCG_GET_OVERRIDEN_VALUE(Settings, BoolValue, Params));
-		case EPCGMetadataTypes::Rotator:
-			return Callback(PCG_GET_OVERRIDEN_VALUE(Settings, RotatorValue, Params));
-		case EPCGMetadataTypes::Name:
-			return Callback(PCG_GET_OVERRIDEN_VALUE(Settings, NameValue, Params));
-		default:
-			// ReturnType{} is invalid if ReturnType is void
-			if constexpr (std::is_same_v<ReturnType, void>)
-			{
-				return;
-			}
-			else
-			{
-				return ReturnType{};
-			}
-		}
+#if WITH_EDITOR
+	if ((Type_DEPRECATED != EPCGMetadataTypes::Double) || (DoubleValue_DEPRECATED != 0.0))
+	{
+		AttributeTypes.Type = Type_DEPRECATED;
+		AttributeTypes.DoubleValue = DoubleValue_DEPRECATED;
+		AttributeTypes.FloatValue = FloatValue_DEPRECATED;
+		AttributeTypes.IntValue = IntValue_DEPRECATED;
+		AttributeTypes.Int32Value = Int32Value_DEPRECATED;
+		AttributeTypes.Vector2Value = Vector2Value_DEPRECATED;
+		AttributeTypes.VectorValue = VectorValue_DEPRECATED;
+		AttributeTypes.Vector4Value = Vector4Value_DEPRECATED;
+		AttributeTypes.RotatorValue = RotatorValue_DEPRECATED;
+		AttributeTypes.QuatValue = QuatValue_DEPRECATED;
+		AttributeTypes.TransformValue = TransformValue_DEPRECATED;
+		AttributeTypes.BoolValue = BoolValue_DEPRECATED;
+		AttributeTypes.StringValue = StringValue_DEPRECATED;
+		AttributeTypes.NameValue = NameValue_DEPRECATED;
+
+		Type_DEPRECATED = EPCGMetadataTypes::Double;
+		DoubleValue_DEPRECATED = 0.0;
 	}
+#endif // WITH_EDITOR
 }
 
 FName UPCGCreateAttributeSettings::AdditionalTaskName() const
@@ -73,7 +52,7 @@ FName UPCGCreateAttributeSettings::AdditionalTaskName() const
 	{
 		const FName NodeName = PCGCreateAttributeConstants::NodeName;
 
-		if (OutputAttributeName == NAME_None && SourceParamAttributeName == NAME_None)
+		if ((OutputAttributeName == NAME_None) && (SourceParamAttributeName == NAME_None))
 		{
 			return NodeName;
 		}
@@ -85,39 +64,7 @@ FName UPCGCreateAttributeSettings::AdditionalTaskName() const
 	}
 	else
 	{
-		const FString Name = OutputAttributeName.ToString();
-
-		switch (Type)
-		{
-		case EPCGMetadataTypes::Integer32:
-			return FName(FString::Printf(TEXT("%s: %d"), *Name, Int32Value));
-		case EPCGMetadataTypes::Float:
-			return FName(FString::Printf(TEXT("%s: %.2f"), *Name, FloatValue));
-		case EPCGMetadataTypes::Integer64:
-			return FName(FString::Printf(TEXT("%s: %lld"), *Name, IntValue));
-		case EPCGMetadataTypes::Double:
-			return FName(FString::Printf(TEXT("%s: %.2f"), *Name, DoubleValue));
-		case EPCGMetadataTypes::String:
-			return FName(FString::Printf(TEXT("%s: \"%s\""), *Name, *StringValue));
-		case EPCGMetadataTypes::Name:
-			return FName(FString::Printf(TEXT("%s: N(\"%s\")"), *Name, *NameValue.ToString()));
-		case EPCGMetadataTypes::Vector2:
-			return FName(FString::Printf(TEXT("%s: V(%.2f, %.2f)"), *Name, Vector2Value.X, Vector2Value.Y));
-		case EPCGMetadataTypes::Vector:
-			return FName(FString::Printf(TEXT("%s: V(%.2f, %.2f, %.2f)"), *Name, VectorValue.X, VectorValue.Y, VectorValue.Z));
-		case EPCGMetadataTypes::Vector4:
-			return FName(FString::Printf(TEXT("%s: V(%.2f, %.2f, %.2f, %.2f)"), *Name, Vector4Value.X, Vector4Value.Y, Vector4Value.Z, Vector4Value.W));
-		case EPCGMetadataTypes::Rotator:
-			return FName(FString::Printf(TEXT("%s: R(%.2f, %.2f, %.2f)"), *Name, RotatorValue.Roll, RotatorValue.Pitch, RotatorValue.Yaw));
-		case EPCGMetadataTypes::Quaternion:
-			return FName(FString::Printf(TEXT("%s: Q(%.2f, %.2f, %.2f, %.2f)"), *Name, QuatValue.X, QuatValue.Y, QuatValue.Z, QuatValue.W));
-		case EPCGMetadataTypes::Transform:
-			return FName(FString::Printf(TEXT("%s: Transform"), *Name));
-		case EPCGMetadataTypes::Boolean:
-			return FName(FString::Printf(TEXT("%s: %s"), *Name, (BoolValue ? TEXT("True") : TEXT("False"))));
-		default:
-			return NAME_None;
-		}
+		return FName(FString::Printf(TEXT("%s: %s"), *OutputAttributeName.ToString(), *AttributeTypes.ToString()));
 	}
 }
 
@@ -286,7 +233,7 @@ FPCGMetadataAttributeBase* FPCGCreateAttributeElement::ClearOrCreateAttribute(co
 		return PCGMetadataElementCommon::ClearOrCreateAttribute(Metadata, OutputAttributeNameOverride ? *OutputAttributeNameOverride : Settings->OutputAttributeName, Value);
 	};
 
-	return PCGCreateAttributeElement::Dispatcher(Settings, Params, CreateAttribute);
+	return Settings->AttributeTypes.DispatcherWithOverride(Params, CreateAttribute);
 }
 
 PCGMetadataEntryKey FPCGCreateAttributeElement::SetAttribute(const UPCGCreateAttributeSettings* Settings, FPCGMetadataAttributeBase* Attribute, UPCGMetadata* Metadata, PCGMetadataEntryKey EntryKey, const UPCGParamData* Params) const
@@ -306,6 +253,6 @@ PCGMetadataEntryKey FPCGCreateAttributeElement::SetAttribute(const UPCGCreateAtt
 		return FinalKey;
 	};
 
-	return PCGCreateAttributeElement::Dispatcher(Settings, Params, SetAttribute);
+	return Settings->AttributeTypes.DispatcherWithOverride(Params, SetAttribute);
 }
 
diff --git a/Engine/Plugins/Experimental/PCG/Source/PCG/Private/Elements/PCGPointFilter.cpp b/Engine/Plugins/Experimental/PCG/Source/PCG/Private/Elements/PCGPointFilter.cpp
--- a/Engine/Plugins/Experimental/PCG/Source/PCG/Private/Elements/PCGPointFilter.cpp
+++ b/Engine/Plugins/Experimental/PCG/Source/PCG/Private/Elements/PCGPointFilter.cpp
@@ -8,6 +8,9 @@
 #include "Metadata/PCGMetadataAttribute.h"
 #include "Metadata/PCGMetadataAttributeTpl.h"
 #include "Helpers/PCGAsync.h"
+#include "Metadata/Accessors/PCGCustomAccessor.h"
+#include "Metadata/Accessors/PCGAttributeAccessorKeys.h"
+#include "Metadata/Accessors/PCGAttributeAccessorHelpers.h"
 
 #include UE_INLINE_GENERATED_CPP_BY_NAME(PCGPointFilter)
 
@@ -17,587 +20,159 @@ namespace PCGPointFilterConstants
 	const FName FilterLabel = TEXT("Filter");
 	const FName InFilterLabel = TEXT("InsideFilter");
 	const FName OutFilterLabel = TEXT("OutsideFilter");
+
+	constexpr int32 ChunkSize = 256;
 }
 
 namespace PCGPointFilterHelpers
 {
-	TFunction<void(const FPCGPoint&, const UPCGMetadata*, void*)> ConstructPropertyGetter(EPCGPointProperties TargetPointProperty, EPCGPointFilterConstantType& OutType)
+	template <typename T>
+	bool ApplyCompare(const T& Input1, const T& Input2, EPCGPointFilterOperator Operation)
 	{
-		if (TargetPointProperty == EPCGPointProperties::Density)
-		{
-			OutType = EPCGPointFilterConstantType::Float;
-			return [](const FPCGPoint& InPoint, const UPCGMetadata*, void* OutData)
-			{
-				*(static_cast<float*>(OutData)) = InPoint.Density;
-			};
-		}
-		else if (TargetPointProperty == EPCGPointProperties::BoundsMin)
+		if (Operation == EPCGPointFilterOperator::Equal)
 		{
-			OutType = EPCGPointFilterConstantType::Vector;
-			return [](const FPCGPoint& InPoint, const UPCGMetadata*, void* OutData)
-			{
-				*(static_cast<FVector*>(OutData)) = InPoint.BoundsMin;
-			};
-		}
-		else if (TargetPointProperty == EPCGPointProperties::BoundsMax)
-		{
-			OutType = EPCGPointFilterConstantType::Vector;
-			return [](const FPCGPoint& InPoint, const UPCGMetadata*, void* OutData)
-			{
-				*(static_cast<FVector*>(OutData)) = InPoint.BoundsMax;
-			};
-		}
-		else if (TargetPointProperty == EPCGPointProperties::Extents)
-		{
-			OutType = EPCGPointFilterConstantType::Vector;
-			return [](const FPCGPoint& InPoint, const UPCGMetadata*, void* OutData)
-			{
-				*(static_cast<FVector*>(OutData)) = InPoint.GetExtents();
-			};
+			return PCG::Private::MetadataTraits<T>::Equal(Input1, Input2);
 		}
-		else if (TargetPointProperty == EPCGPointProperties::Color)
+		else if (Operation == EPCGPointFilterOperator::NotEqual)
 		{
-			OutType = EPCGPointFilterConstantType::Vector4;
-			return [](const FPCGPoint& InPoint, const UPCGMetadata*, void* OutData)
-			{
-				*(static_cast<FVector4*>(OutData)) = InPoint.Color;
-			};
-		}
-		else if (TargetPointProperty == EPCGPointProperties::Position)
-		{
-			OutType = EPCGPointFilterConstantType::Vector;
-			return [](const FPCGPoint& InPoint, const UPCGMetadata*, void* OutData)
-			{
-				*(static_cast<FVector*>(OutData)) = InPoint.Transform.GetLocation();
-			};
-		}
-		//else if (TargetPointProperty == EPCGPointProperties::Rotation)
-		//{
-		//	OutType = EPCGPointFilterConstantType::Rotation;
-		//	return [](const FPCGPoint& InPoint, const UPCGMetadata*, void* OutData)
-		//	{
-		//		*(static_cast<FQuat*>(OutData)) = InPoint.Transform.GetRotation();
-		//	};
-		//}
-		else if (TargetPointProperty == EPCGPointProperties::Scale)
-		{
-			OutType = EPCGPointFilterConstantType::Vector;
-			return [](const FPCGPoint& InPoint, const UPCGMetadata*, void* OutData)
-			{
-				*(static_cast<FVector*>(OutData)) = InPoint.Transform.GetScale3D();
-			};
-		}
-		//else if (TargetPointProperty == EPCGPointProperties::Transform)
-		//{
-		//	OutType = EPCGPointFilterConstantType::Transform;
-		//	return [](const FPCGPoint& InPoint, const UPCGMetadata*, void* OutData)
-		//	{
-		//		*(static_cast<FTransform*>(OutData)) = InPoint.Transform;
-		//	};
-		//}
-		else if (TargetPointProperty == EPCGPointProperties::Steepness)
-		{
-			OutType = EPCGPointFilterConstantType::Float;
-			return [](const FPCGPoint& InPoint, const UPCGMetadata*, void* OutData)
-			{
-				*(static_cast<float*>(OutData)) = InPoint.Steepness;
-			};
+			return !PCG::Private::MetadataTraits<T>::Equal(Input1, Input2);
 		}
-		else
-		{
-			// Log error
-			OutType = EPCGPointFilterConstantType::Unknown;
-			return [](const FPCGPoint&, const UPCGMetadata*, void*) {};
-		}
-	}
-
-	TFunction<void(const FPCGPoint&, const UPCGMetadata*, void*)> ConstructAttributeGetter(
-		const UPCGMetadata* Metadata, 
-		const FName& AttributeName, 
-		EPCGPointFilterConstantType& OutType)
-	{
-		OutType = EPCGPointFilterConstantType::Unknown;
 
-		const FPCGMetadataAttributeBase* AttributeBase = Metadata->GetConstAttribute(AttributeName);
-		if (AttributeBase)
+		if constexpr (PCG::Private::MetadataTraits<T>::CanCompare)
 		{
-			if (AttributeBase->GetTypeId() == PCG::Private::MetadataTypes<bool>::Id)
-			{
-				OutType = EPCGPointFilterConstantType::Integer64;
-				return [AttributeBase](const FPCGPoint& InPoint, const UPCGMetadata* InMetadata, void* OutData)
-				{
-					*(static_cast<int64*>(OutData)) = static_cast<const FPCGMetadataAttribute<bool>*>(AttributeBase)->GetValueFromItemKey(InPoint.MetadataEntry);
-				};
-			}
-			else if (AttributeBase->GetTypeId() == PCG::Private::MetadataTypes<int32>::Id)
-			{
-				OutType = EPCGPointFilterConstantType::Integer64;
-				return [AttributeBase](const FPCGPoint& InPoint, const UPCGMetadata* InMetadata, void* OutData)
-				{
-					*(static_cast<int64*>(OutData)) = static_cast<const FPCGMetadataAttribute<int32>*>(AttributeBase)->GetValueFromItemKey(InPoint.MetadataEntry);
-				};
-			}
-			else if (AttributeBase->GetTypeId() == PCG::Private::MetadataTypes<int64>::Id)
-			{
-				OutType = EPCGPointFilterConstantType::Integer64;
-				return [AttributeBase](const FPCGPoint& InPoint, const UPCGMetadata* InMetadata, void* OutData)
-				{
-					*(static_cast<int64*>(OutData)) = static_cast<const FPCGMetadataAttribute<int64>*>(AttributeBase)->GetValueFromItemKey(InPoint.MetadataEntry);
-				};
-			}
-			else if (AttributeBase->GetTypeId() == PCG::Private::MetadataTypes<float>::Id)
+			switch (Operation)
 			{
-				OutType = EPCGPointFilterConstantType::Float;
-				return [AttributeBase](const FPCGPoint& InPoint, const UPCGMetadata* InMetadata, void* OutData)
-				{
-					*(static_cast<float*>(OutData)) = static_cast<const FPCGMetadataAttribute<float>*>(AttributeBase)->GetValueFromItemKey(InPoint.MetadataEntry);
-				};
-			}
-			else if (AttributeBase->GetTypeId() == PCG::Private::MetadataTypes<double>::Id)
-			{
-				OutType = EPCGPointFilterConstantType::Float;
-				return [AttributeBase](const FPCGPoint& InPoint, const UPCGMetadata* InMetadata, void* OutData)
-				{
-					*(static_cast<float*>(OutData)) = static_cast<const FPCGMetadataAttribute<double>*>(AttributeBase)->GetValueFromItemKey(InPoint.MetadataEntry);
-				};
-			}
-			else if (AttributeBase->GetTypeId() == PCG::Private::MetadataTypes<FVector>::Id)
-			{
-				OutType = EPCGPointFilterConstantType::Vector;
-				return [AttributeBase](const FPCGPoint& InPoint, const UPCGMetadata* InMetadata, void* OutData)
-				{
-					*(static_cast<FVector*>(OutData)) = static_cast<const FPCGMetadataAttribute<FVector>*>(AttributeBase)->GetValueFromItemKey(InPoint.MetadataEntry);
-				};
-			}
-			else if (AttributeBase->GetTypeId() == PCG::Private::MetadataTypes<FVector4>::Id)
-			{
-				OutType = EPCGPointFilterConstantType::Vector4;
-				return [AttributeBase](const FPCGPoint& InPoint, const UPCGMetadata* InMetadata, void* OutData)
-				{
-					*(static_cast<FVector4*>(OutData)) = static_cast<const FPCGMetadataAttribute<FVector4>*>(AttributeBase)->GetValueFromItemKey(InPoint.MetadataEntry);
-				};
-			}
-			//else if (AttributeBase->GetTypeId() == PCG::Private::MetadataTypes<FQuat>::Id)
-			//{
-			//	OutType = EPCGPointFilterConstantType::Rotation;
-			//	return [AttributeBase](const FPCGPoint& InPoint, const UPCGMetadata* InMetadata, void* OutData)
-			//	{
-			//		*(static_cast<FQuat*>(OutData)) = static_cast<const FPCGMetadataAttribute<FQuat>*>(AttributeBase)->GetValueFromItemKey(InPoint.MetadataEntry);
-			//	};
-			//}
-			//else if (AttributeBase->GetTypeId() == PCG::Private::MetadataTypes<FTransform>::Id)
-			//{
-			//	*(static_cast<FTransform*>(OutData)) = static_cast<const FPCGMetadataAttribute<FTransform>*>(AttributeBase)->GetValueFromItemKey(InPoint.MetadataEntry);
-			//}
-			else if (AttributeBase->GetTypeId() == PCG::Private::MetadataTypes<FString>::Id)
-			{
-				OutType = EPCGPointFilterConstantType::String;
-				return [AttributeBase](const FPCGPoint& InPoint, const UPCGMetadata* InMetadata, void* OutData)
-				{
-					*(static_cast<FString*>(OutData)) = static_cast<const FPCGMetadataAttribute<FString>*>(AttributeBase)->GetValueFromItemKey(InPoint.MetadataEntry);
-				};
-			}
-			else
-			{
-				// Log error
-				OutType = EPCGPointFilterConstantType::Unknown;
-				return [](const FPCGPoint&, const UPCGMetadata*, void*) {};
+			case EPCGPointFilterOperator::Greater:
+				return PCG::Private::MetadataTraits<T>::Greater(Input1, Input2);
+			case EPCGPointFilterOperator::GreaterOrEqual:
+				return PCG::Private::MetadataTraits<T>::GreaterOrEqual(Input1, Input2);
+			case EPCGPointFilterOperator::Lesser:
+				return PCG::Private::MetadataTraits<T>::Less(Input1, Input2);
+			case EPCGPointFilterOperator::LesserOrEqual:
+				return PCG::Private::MetadataTraits<T>::LessOrEqual(Input1, Input2);
+			default:
+				return false;
 			}
 		}
 		else
 		{
-			// Log error
-			OutType = EPCGPointFilterConstantType::Unknown;
-			return [](const FPCGPoint&, const UPCGMetadata*, void*) {};
+			return false;
 		}
 	}
+}
 
-	TFunction<void(const FPCGPoint&, const UPCGMetadata*, void*)> ConstructSourceGetter(
-		EPCGPointTargetFilterType TargetFilterType, 
-		EPCGPointProperties TargetPointProperty, 
-		const UPCGMetadata* TargetMetadata,
-		const FName& TargetAttributeName,
-		EPCGPointFilterConstantType& OutType)
-	{
-		if (TargetFilterType == EPCGPointTargetFilterType::Property)
-		{
-			return ConstructPropertyGetter(TargetPointProperty, OutType);
-		}
-		else if (TargetFilterType == EPCGPointTargetFilterType::Metadata)
-		{
-			return ConstructAttributeGetter(TargetMetadata, TargetAttributeName, OutType);
-		}
-		else
-		{
-			// log error
-			OutType = EPCGPointFilterConstantType::Unknown;
-			return [](const FPCGPoint&, const UPCGMetadata*, void*) {};
-		}
-	}
+TArray<FPCGPinProperties> UPCGPointFilterSettings::InputPinProperties() const
+{
+	TArray<FPCGPinProperties> PinProperties;
+	PinProperties.Emplace(PCGPointFilterConstants::DataToFilterLabel, EPCGDataType::Any);
 
-	TFunction<void(const FPCGPoint&, const UPCGMetadata*, void*)> ConstructThresholdGetter(
-		EPCGPointThresholdType ThresholdFilterType, 
-		EPCGPointProperties ThresholdPointProperty, 
-		const UPCGMetadata* ThresholdMetadata,
-		const FName& ThresholdAttributeName, 
-		EPCGPointFilterConstantType ThresholdConstantType,
-		const int64& Integer64Constant,
-		const float& FloatConstant,
-		const FVector& VectorConstant,
-		const FVector4& Vector4Constant,
-		//const FQuat& RotationConstant,
-		const FString& StringConstant,
-		EPCGPointFilterConstantType& OutType)
+	if (!bUseConstantThreshold)
 	{
-		if (ThresholdFilterType == EPCGPointThresholdType::Property)
-		{
-			return ConstructPropertyGetter(ThresholdPointProperty, OutType);
-		}
-		else if (ThresholdFilterType == EPCGPointThresholdType::Metadata)
-		{
-			return ConstructAttributeGetter(ThresholdMetadata, ThresholdAttributeName, OutType);
-		}
-		else if (ThresholdFilterType == EPCGPointThresholdType::Constant)
-		{
-			OutType = ThresholdConstantType;
-
-			if (ThresholdConstantType == EPCGPointFilterConstantType::Integer64)
-			{
-				return [&Integer64Constant](const FPCGPoint&, const UPCGMetadata*, void* OutData)
-				{
-					*(static_cast<int64*>(OutData)) = Integer64Constant;
-				};
-			}
-			else if (ThresholdConstantType == EPCGPointFilterConstantType::Float)
-			{
-				return [&FloatConstant](const FPCGPoint&, const UPCGMetadata*, void* OutData)
-				{
-					*(static_cast<float*>(OutData)) = FloatConstant;
-				};
-			}
-			else if (ThresholdConstantType == EPCGPointFilterConstantType::Vector)
-			{
-				return [&VectorConstant](const FPCGPoint&, const UPCGMetadata*, void* OutData)
-				{
-					*(static_cast<FVector*>(OutData)) = VectorConstant;
-				};
-			}
-			else if (ThresholdConstantType == EPCGPointFilterConstantType::Vector4)
-			{
-				return [&Vector4Constant](const FPCGPoint&, const UPCGMetadata*, void* OutData)
-				{
-					*(static_cast<FVector4*>(OutData)) = Vector4Constant;
-				};
-			}
-			//else if (ThresholdConstantType == EPCGPointFilterConstantType::Rotation)
-			//{
-			//	return [&RotationConstant](const FPCGPoint&, const UPCGMetadata*, void* OutData)
-			//	{
-			//		*(static_cast<FQuat*>(OutData)) = RotationConstant;
-			//	};
-			//}
-			else if (ThresholdConstantType == EPCGPointFilterConstantType::String)
-			{
-				return [&StringConstant](const FPCGPoint&, const UPCGMetadata*, void* OutData)
-				{
-					*(static_cast<FString*>(OutData)) = StringConstant;
-				};
-			}
-			else // unknown constant type
-			{
-				// Log error
-				OutType = EPCGPointFilterConstantType::Unknown;
-				return [](const FPCGPoint&, const UPCGMetadata*, void*) {};
-			}
-		}
-		else // unknown filter type
-		{
-			// Log error
-			OutType = EPCGPointFilterConstantType::Unknown;
-			return [](const FPCGPoint&, const UPCGMetadata*, void*) {};
-		}
+		PinProperties.Emplace(PCGPointFilterConstants::FilterLabel, EPCGDataType::Any, /*bInAllowMultipleConnections=*/ false);
 	}
 
-	template<typename T, typename U>
-	struct TemplatedOperator
-	{
-		TemplatedOperator(EPCGPointFilterOperator Operator)
-		{
-			if (Operator == EPCGPointFilterOperator::Greater)
-			{
-				OperatorFn = [](const T& A, const U& B) { return A > B; };
-			}
-			else if (Operator == EPCGPointFilterOperator::GreaterOrEqual)
-			{
-				OperatorFn = [](const T& A, const U& B) { return A >= B; };
-			}
-			else if (Operator == EPCGPointFilterOperator::Lesser)
-			{
-				OperatorFn = [](const T& A, const U& B) { return A < B; };
-			}
-			else if (Operator == EPCGPointFilterOperator::LesserOrEqual)
-			{
-				OperatorFn = [](const T& A, const U& B) { return A <= B; };
-			}
-			else if (Operator == EPCGPointFilterOperator::Equal)
-			{
-				OperatorFn = [](const T& A, const U& B) { return A == B; };
-			}
-			else if (Operator == EPCGPointFilterOperator::NotEqual)
-			{
-				OperatorFn = [](const T& A, const U& B) { return A != B; };
-			}
-			else
-			{
-				OperatorFn = [](const T&, const U&) { return false; };
-			}
-		}
-
-		bool operator()(const T& A, const U& B) const
-		{
-			return OperatorFn(A, B);
-		}
+	return PinProperties;
+}
 
-		TFunction<bool(const T&, const U&)> OperatorFn;
-	};
+TArray<FPCGPinProperties> UPCGPointFilterSettings::OutputPinProperties() const
+{
+	TArray<FPCGPinProperties> PinProperties;
+	PinProperties.Emplace(PCGPointFilterConstants::InFilterLabel, EPCGDataType::Point);
+	PinProperties.Emplace(PCGPointFilterConstants::OutFilterLabel, EPCGDataType::Point);
 
-	template<typename U>
-	struct TemplatedOperator<FVector, U>
-	{
-		TemplatedOperator(EPCGPointFilterOperator Operator)
-		{
-			if (Operator == EPCGPointFilterOperator::Greater)
-			{
-				OperatorFn = [](const FVector& A, const U& B) { return A.Length() > B; };
-			}
-			else if (Operator == EPCGPointFilterOperator::GreaterOrEqual)
-			{
-				OperatorFn = [](const FVector& A, const U& B) { return A.Length() >= B; };
-			}
-			else if (Operator == EPCGPointFilterOperator::Lesser)
-			{
-				OperatorFn = [](const FVector& A, const U& B) { return A.Length() < B; };
-			}
-			else if (Operator == EPCGPointFilterOperator::LesserOrEqual)
-			{
-				OperatorFn = [](const FVector& A, const U& B) { return A.Length() <= B; };
-			}
-			else if (Operator == EPCGPointFilterOperator::Equal)
-			{
-				OperatorFn = [](const FVector& A, const U& B) { return A.Length() == B; };
-			}
-			else if (Operator == EPCGPointFilterOperator::NotEqual)
-			{
-				OperatorFn = [](const FVector& A, const U& B) { return A.Length() != B; };
-			}
-			else
-			{
-				OperatorFn = [](const FVector&, const U&) { return false; };
-			}
-		}
+	return PinProperties;
+}
 
-		bool operator()(const FVector& A, const U& B) const
-		{
-			return OperatorFn(A, B);
-		}
+FPCGElementPtr UPCGPointFilterSettings::CreateElement() const
+{
+	return MakeShared<FPCGPointFilterElement>();
+}
 
-		TFunction<bool(const FVector&, const U&)> OperatorFn;
-	};
+UPCGPointFilterSettings::UPCGPointFilterSettings()
+	: UPCGSettings()
+{
+	// Previous default object was: density for both selectors
+	// Recreate the same default
+	TargetAttribute.SetPointProperty(EPCGPointProperties::Density);
+	ThresholdAttribute.SetPointProperty(EPCGPointProperties::Density);
+}
 
-	template<typename U>
-	struct TemplatedOperator<U, FVector>
+void UPCGPointFilterSettings::PostLoad()
+{
+	Super::PostLoad();
+
+#if WITH_EDITOR
+PRAGMA_DISABLE_DEPRECATION_WARNINGS
+	if ((TargetFilterType_DEPRECATED != EPCGPointTargetFilterType::Property)
+		|| (TargetPointProperty_DEPRECATED != EPCGPointProperties::Density)
+		|| (TargetAttributeName_DEPRECATED != NAME_None)
+		|| (ThresholdFilterType_DEPRECATED != EPCGPointThresholdType::Property)
+		|| (ThresholdPointProperty_DEPRECATED != EPCGPointProperties::Density)
+		|| (ThresholdAttributeName_DEPRECATED != NAME_None))
 	{
-		TemplatedOperator(EPCGPointFilterOperator Operator)
-		{
-			if (Operator == EPCGPointFilterOperator::Greater)
-			{
-				OperatorFn = [](const U& A, const FVector& B) { return A > B.Length(); };
-			}
-			else if (Operator == EPCGPointFilterOperator::GreaterOrEqual)
-			{
-				OperatorFn = [](const U& A, const FVector& B) { return A >= B.Length(); };
-			}
-			else if (Operator == EPCGPointFilterOperator::Lesser)
-			{
-				OperatorFn = [](const U& A, const FVector& B) { return A < B.Length(); };
-			}
-			else if (Operator == EPCGPointFilterOperator::LesserOrEqual)
-			{
-				OperatorFn = [](const U& A, const FVector& B) { return A <= B.Length(); };
-			}
-			else if (Operator == EPCGPointFilterOperator::Equal)
-			{
-				OperatorFn = [](const U& A, const FVector& B) { return A == B.Length(); };
-			}
-			else if (Operator == EPCGPointFilterOperator::NotEqual)
-			{
-				OperatorFn = [](const U& A, const FVector& B) { return A != B.Length(); };
-			}
-			else
-			{
-				OperatorFn = [](const U&, const FVector&) { return false; };
-			}
-		}
-
-		bool operator()(const U& A, const FVector& B) const
+		if (TargetFilterType_DEPRECATED == EPCGPointTargetFilterType::Property)
 		{
-			return OperatorFn(A, B);
+			TargetAttribute.SetPointProperty(TargetPointProperty_DEPRECATED);
 		}
-
-		TFunction<bool(const U&, const FVector&)> OperatorFn;
-	};
-
-	template<>
-	struct TemplatedOperator<FVector, FVector>
-	{
-		TemplatedOperator(EPCGPointFilterOperator Operator)
+		else
 		{
-			if (Operator == EPCGPointFilterOperator::Greater)
-			{
-				OperatorFn = [](const FVector& A, const FVector& B) { return A.X > B.X && A.Y > B.Y && A.Z > B.Z; };
-			}
-			else if (Operator == EPCGPointFilterOperator::GreaterOrEqual)
-			{
-				OperatorFn = [](const FVector& A, const FVector& B) { return A.X >= B.X && A.Y >= B.Y && A.Z >= B.Z; };
-			}
-			else if (Operator == EPCGPointFilterOperator::Lesser)
-			{
-				OperatorFn = [](const FVector& A, const FVector& B) { return A.X < B.X && A.Y < B.Y && A.Z < B.Z; };
-			}
-			else if (Operator == EPCGPointFilterOperator::LesserOrEqual)
-			{
-				OperatorFn = [](const FVector& A, const FVector& B) { return A.X <= B.X && A.Y <= B.Y && A.Z <= B.Z; };
-			}
-			else if (Operator == EPCGPointFilterOperator::Equal)
-			{
-				OperatorFn = [](const FVector& A, const FVector& B) { return A == B; };
-			}
-			else if (Operator == EPCGPointFilterOperator::NotEqual)
-			{
-				OperatorFn = [](const FVector& A, const FVector& B) { return A != B; };
-			}
-			else
-			{
-				OperatorFn = [](const FVector&, const FVector&) { return false; };
-			}
+			TargetAttribute.SetAttributeName(TargetAttributeName_DEPRECATED);
 		}
 
-		bool operator()(const FVector& A, const FVector& B) const
+		if (ThresholdFilterType_DEPRECATED == EPCGPointThresholdType::Property)
 		{
-			return OperatorFn(A, B);
+			ThresholdAttribute.SetPointProperty(ThresholdPointProperty_DEPRECATED);
 		}
-
-		TFunction<bool(const FVector&, const FVector&)> OperatorFn;
-	};
-
-	template<>
-	struct TemplatedOperator<FVector4, FVector4>
-	{
-		TemplatedOperator(EPCGPointFilterOperator Operator)
+		else if (ThresholdFilterType_DEPRECATED == EPCGPointThresholdType::Metadata)
 		{
-			if (Operator == EPCGPointFilterOperator::Greater)
-			{
-				OperatorFn = [](const FVector4& A, const FVector4& B) { return A.X > B.X && A.Y > B.Y && A.Z > B.Z && A.W > B.W; };
-			}
-			else if (Operator == EPCGPointFilterOperator::GreaterOrEqual)
-			{
-				OperatorFn = [](const FVector4& A, const FVector4& B) { return A.X >= B.X && A.Y >= B.Y && A.Z >= B.Z && A.W >= B.W; };
-			}
-			else if (Operator == EPCGPointFilterOperator::Lesser)
-			{
-				OperatorFn = [](const FVector4& A, const FVector4& B) { return A.X < B.X&& A.Y < B.Y&& A.Z < B.Z && A.W < B.W; };
-			}
-			else if (Operator == EPCGPointFilterOperator::LesserOrEqual)
-			{
-				OperatorFn = [](const FVector4& A, const FVector4& B) { return A.X <= B.X && A.Y <= B.Y && A.Z <= B.Z && A.W <= B.W; };
-			}
-			else if (Operator == EPCGPointFilterOperator::Equal)
-			{
-				OperatorFn = [](const FVector4& A, const FVector4& B) { return A == B; };
-			}
-			else if (Operator == EPCGPointFilterOperator::NotEqual)
-			{
-				OperatorFn = [](const FVector4& A, const FVector4& B) { return A != B; };
-			}
-			else
-			{
-				OperatorFn = [](const FVector4&, const FVector4&) { return false; };
-			}
+			ThresholdAttribute.SetAttributeName(ThresholdAttributeName_DEPRECATED);
 		}
-
-		bool operator()(const FVector& A, const FVector& B) const
+		else
 		{
-			return OperatorFn(A, B);
+			bUseConstantThreshold = true;
 		}
+	}
 
-		TFunction<bool(const FVector&, const FVector&)> OperatorFn;
-	};
-
-	TFunction<bool(const FPCGPoint&, const UPCGMetadata*, const FPCGPoint&, const UPCGMetadata*)> ConstructConcreteOperator(
-		const TFunction<void(const FPCGPoint&, const UPCGMetadata*, void*)>& SourceGetter,
-		EPCGPointFilterConstantType SourceType,
-		const TFunction<void(const FPCGPoint&, const UPCGMetadata*, void*)>& ThresholdGetter,
-		EPCGPointFilterConstantType ThresholdType,
-		EPCGPointFilterOperator Operator,
-		bool& bOutOperatorCreatedCorrectly)
+	if (bUseConstantThreshold && 
+		((ThresholdConstantType_DEPRECATED != EPCGPointFilterConstantType::Float)
+		|| (FloatConstant_DEPRECATED != 0.0f)))
 	{
-		bOutOperatorCreatedCorrectly = true;
-
-#define CONCRETE_OPERATOR_SHORTHAND(TypeEnumA, TypeA, TypeEnumB, TypeB) \
-		if(SourceType == EPCGPointFilterConstantType::TypeEnumA && ThresholdType == EPCGPointFilterConstantType::TypeEnumB) \
-		{ \
-			TemplatedOperator<TypeA, TypeB> OperatorFn(Operator); \
-			\
-			return[&SourceGetter, &ThresholdGetter, OperatorFn](const FPCGPoint& InTargetPoint, const UPCGMetadata* InTargetMetadata, const FPCGPoint& InThresholdPoint, const UPCGMetadata* InThresholdMetadata) \
-			{ \
-				TypeA A; \
-				TypeB B; \
-				\
-				SourceGetter(InTargetPoint, InTargetMetadata, &A); \
-				ThresholdGetter(InThresholdPoint, InThresholdMetadata, &B); \
-				return OperatorFn(A, B); \
-			}; \
+		switch (ThresholdConstantType_DEPRECATED)
+		{
+		case EPCGPointFilterConstantType::Float:
+			AttributeTypes.Type = EPCGMetadataTypes::Float;
+			AttributeTypes.FloatValue = FloatConstant_DEPRECATED;
+			break;
+		case EPCGPointFilterConstantType::Integer64:
+			AttributeTypes.Type = EPCGMetadataTypes::Integer64;
+			AttributeTypes.IntValue = Integer64Constant_DEPRECATED;
+			break;
+		case EPCGPointFilterConstantType::Vector:
+			AttributeTypes.Type = EPCGMetadataTypes::Vector;
+			AttributeTypes.VectorValue = VectorConstant_DEPRECATED;
+			break;
+		case EPCGPointFilterConstantType::Vector4:
+			AttributeTypes.Type = EPCGMetadataTypes::Vector4;
+			AttributeTypes.Vector4Value = Vector4Constant_DEPRECATED;
+			break;
+		case EPCGPointFilterConstantType::String:
+			AttributeTypes.Type = EPCGMetadataTypes::String;
+			AttributeTypes.StringValue = StringConstant_DEPRECATED;
+			break;
 		}
-
-		CONCRETE_OPERATOR_SHORTHAND(Integer64, int64, Integer64, int64)
-		CONCRETE_OPERATOR_SHORTHAND(Integer64, int64, Vector, FVector)
-		CONCRETE_OPERATOR_SHORTHAND(Float, float, Float, float)
-		CONCRETE_OPERATOR_SHORTHAND(Float, float, Vector, FVector)
-		CONCRETE_OPERATOR_SHORTHAND(Vector, FVector, Vector, FVector)
-		CONCRETE_OPERATOR_SHORTHAND(Vector, FVector, Integer64, int64)
-		CONCRETE_OPERATOR_SHORTHAND(Vector, FVector, Float, float)
-		CONCRETE_OPERATOR_SHORTHAND(Vector4, FVector4, Vector4, FVector4)
-		//CONCRETE_OPERATOR_SHORTHAND(Rotation, FQuat, Rotation, FQuat)
-		CONCRETE_OPERATOR_SHORTHAND(String, FString, String, FString)
-
-#undef CONCRETE_OPERATOR_SHORTHAND
-
-		bOutOperatorCreatedCorrectly = false;
-		return [](const FPCGPoint&, const UPCGMetadata*, const FPCGPoint&, const UPCGMetadata*) { return false; };
 	}
-}
-
-TArray<FPCGPinProperties> UPCGPointFilterSettings::InputPinProperties() const
-{
-	TArray<FPCGPinProperties> PinProperties;
-	PinProperties.Emplace(PCGPointFilterConstants::DataToFilterLabel, EPCGDataType::Any);
-	PinProperties.Emplace(PCGPointFilterConstants::FilterLabel, EPCGDataType::Spatial);
-
-	return PinProperties;
-}
-
-TArray<FPCGPinProperties> UPCGPointFilterSettings::OutputPinProperties() const
-{
-	TArray<FPCGPinProperties> PinProperties;
-	PinProperties.Emplace(PCGPointFilterConstants::InFilterLabel, EPCGDataType::Point);
-	PinProperties.Emplace(PCGPointFilterConstants::OutFilterLabel, EPCGDataType::Point);
-
-	return PinProperties;
-}
 
-FPCGElementPtr UPCGPointFilterSettings::CreateElement() const
-{
-	return MakeShared<FPCGPointFilterElement>();
+	// Default values
+	TargetFilterType_DEPRECATED = EPCGPointTargetFilterType::Property;
+	TargetPointProperty_DEPRECATED = EPCGPointProperties::Density;
+	TargetAttributeName_DEPRECATED = NAME_None;
+	ThresholdFilterType_DEPRECATED = EPCGPointThresholdType::Property;
+	ThresholdPointProperty_DEPRECATED = EPCGPointProperties::Density;
+	ThresholdAttributeName_DEPRECATED = NAME_None;
+	ThresholdConstantType_DEPRECATED = EPCGPointFilterConstantType::Float;
+	FloatConstant_DEPRECATED = 0.0f;
+PRAGMA_ENABLE_DEPRECATION_WARNINGS
+#endif // WITH_EDITOR
 }
 
 bool FPCGPointFilterElement::ExecuteInternal(FPCGContext* Context) const
@@ -605,10 +180,10 @@ bool FPCGPointFilterElement::ExecuteInternal(FPCGContext* Context) const
 	TRACE_CPUPROFILER_EVENT_SCOPE(FPCGPointFilterElement::Execute);
 	check(Context);
 
+#if !WITH_EDITOR
 	const bool bHasInFilterOutputPin = Context->Node && Context->Node->IsOutputPinConnected(PCGPointFilterConstants::InFilterLabel);
 	const bool bHasOutsideFilterOutputPin = Context->Node && Context->Node->IsOutputPinConnected(PCGPointFilterConstants::OutFilterLabel);
 
-#if !WITH_EDITOR
 	// Early out - only in non-editor builds, otherwise we will potentially poison the cache, since it is input-driven
 	if (!bHasInFilterOutputPin && !bHasOutsideFilterOutputPin)
 	{
@@ -623,262 +198,268 @@ bool FPCGPointFilterElement::ExecuteInternal(FPCGContext* Context) const
 	UPCGParamData* Params = Context->InputData.GetParams(); //TODO: by pin?
 	TArray<FPCGTaggedData>& Outputs = Context->OutputData.TaggedData;
 	TArray<FPCGTaggedData> FilterData = Context->InputData.GetInputsByPin(PCGPointFilterConstants::FilterLabel);
-	FilterData.Append(Context->InputData.GetParamsByPin(PCGPointFilterConstants::FilterLabel));
 
 	// Forward any non-input data
 	Outputs.Append(Context->InputData.GetAllSettings());
 
 	const EPCGPointFilterOperator Operator = PCGSettingsHelpers::GetValue(GET_MEMBER_NAME_CHECKED(UPCGPointFilterSettings, Operator), Settings->Operator, Params);
-	const EPCGPointTargetFilterType TargetFilterType = PCGSettingsHelpers::GetValue(GET_MEMBER_NAME_CHECKED(UPCGPointFilterSettings, TargetFilterType), Settings->TargetFilterType, Params);
-	const EPCGPointProperties TargetPointProperty = PCGSettingsHelpers::GetValue(GET_MEMBER_NAME_CHECKED(UPCGPointFilterSettings, TargetPointProperty), Settings->TargetPointProperty, Params);
-	const FName TargetAttributeName = PCGSettingsHelpers::GetValue(GET_MEMBER_NAME_CHECKED(UPCGPointFilterSettings, TargetAttributeName), Settings->TargetAttributeName, Params);
-
-	const EPCGPointThresholdType ThresholdFilterType = PCGSettingsHelpers::GetValue(GET_MEMBER_NAME_CHECKED(UPCGPointFilterSettings, ThresholdFilterType), Settings->ThresholdFilterType, Params);
-	const EPCGPointProperties ThresholdPointProperty = PCGSettingsHelpers::GetValue(GET_MEMBER_NAME_CHECKED(UPCGPointFilterSettings, ThresholdPointProperty), Settings->ThresholdPointProperty, Params);
-	const FName ThresholdAttributeName = PCGSettingsHelpers::GetValue(GET_MEMBER_NAME_CHECKED(UPCGPointFilterSettings, ThresholdAttributeName), Settings->ThresholdAttributeName, Params);
-	const EPCGPointFilterConstantType ThresholdConstantType = PCGSettingsHelpers::GetValue(GET_MEMBER_NAME_CHECKED(UPCGPointFilterSettings, ThresholdConstantType), Settings->ThresholdConstantType, Params);
-
-	const int64 Integer64Constant = PCGSettingsHelpers::GetValue(GET_MEMBER_NAME_CHECKED(UPCGPointFilterSettings, Integer64Constant), Settings->Integer64Constant, Params);
-	const float FloatConstant = PCGSettingsHelpers::GetValue(GET_MEMBER_NAME_CHECKED(UPCGPointFilterSettings, FloatConstant), Settings->FloatConstant, Params);
-	const FVector VectorConstant = PCGSettingsHelpers::GetValue(GET_MEMBER_NAME_CHECKED(UPCGPointFilterSettings, VectorConstant), Settings->VectorConstant, Params);
-	const FVector4 Vector4Constant = PCGSettingsHelpers::GetValue(GET_MEMBER_NAME_CHECKED(UPCGPointFilterSettings, Vector4Constant), Settings->Vector4Constant, Params);
-	//const FQuat RotationConstant = PCGSettingsHelpers::GetValue(GET_MEMBER_NAME_CHECKED(UPCGPointFilterSettings, RotationConstant), Settings->RotationConstant, Params);
-	const FString StringConstant = PCGSettingsHelpers::GetValue(GET_MEMBER_NAME_CHECKED(UPCGPointFilterSettings, StringConstant), Settings->StringConstant, Params);
 
 	bool bUseSpatialQuery = PCGSettingsHelpers::GetValue(GET_MEMBER_NAME_CHECKED(UPCGPointFilterSettings, bUseSpatialQuery), Settings->bUseSpatialQuery, Params);
 
-	// Validate basic input data
-	if (ThresholdFilterType == EPCGPointThresholdType::Property && !FilterData.IsEmpty() && Cast<UPCGSpatialData>(FilterData[0].Data) == nullptr)
+	// If there is no input, do nothing
+	if (DataToFilter.IsEmpty())
 	{
-		PCGE_LOG(Error, "Cannot filter by property on a non-spatial filter data");
 		return true;
 	}
 
-	if (ThresholdFilterType == EPCGPointThresholdType::Constant && !FilterData.IsEmpty())
-	{
-		PCGE_LOG(Warning, "Point filter uses a constant but has filter data that will be ignored");
-	}
+	TUniquePtr<const IPCGAttributeAccessor> ThresholdAccessor;
+	TUniquePtr<const IPCGAttributeAccessorKeys> ThresholdKeys;
+	bool bUseInputDataForThreshold = false;
 
-	for (const FPCGTaggedData& Input : DataToFilter)
-	{
-		const UPCGSpatialData* SpatialInput = Cast<const UPCGSpatialData>(Input.Data);
-		check(SpatialInput);
+	UPCGPointData* ThresholdPointData = nullptr;
+	const UPCGSpatialData* ThresholdSpatialData = nullptr;
 
-		const UPCGPointData* OriginalData = SpatialInput->ToPointData(Context);
-		if (!OriginalData)
+	if (Settings->bUseConstantThreshold)
+	{
+		auto ConstantThreshold = [&ThresholdAccessor, &ThresholdKeys](auto&& Value)
 		{
-			PCGE_LOG(Error, "Unable to get point data from input");
-			continue;
-		}
-
-		const TArray<FPCGPoint>& TargetPoints = OriginalData->GetPoints();
-		const UPCGMetadata* TargetMetadata = OriginalData->ConstMetadata();
+			using ConstantType = std::decay_t<decltype(Value)>;
 
-		const UPCGSpatialData* ThresholdSpatialData = nullptr;
-		const TArray<FPCGPoint>* ThresholdPoints = nullptr;
-		const UPCGMetadata* ThresholdMetadata = nullptr;
-		UPCGMetadata* TemporaryMetadata = nullptr;
+			ThresholdAccessor = MakeUnique<FPCGConstantValueAccessor<ConstantType>>(std::forward<ConstantType>(Value));
+			// Dummy keys
+			ThresholdKeys = MakeUnique<FPCGAttributeAccessorKeysSingleObjectPtr<void>>();
+		};
 
-		if (FilterData.Num() > 0)
+		Settings->AttributeTypes.DispatcherWithOverride(Params, ConstantThreshold);
+	}
+	else
+	{
+		if (!FilterData.IsEmpty())
 		{
-			if (const UPCGSpatialData* SpatialDataFilter = Cast<const UPCGSpatialData>(FilterData[0].Data))
-			{
-				ThresholdSpatialData = SpatialDataFilter;
-				ThresholdMetadata = ThresholdSpatialData->ConstMetadata();
+			const UPCGData* ThresholdData = FilterData[0].Data;
 
-				if (bUseSpatialQuery && ThresholdMetadata && ThresholdFilterType == EPCGPointThresholdType::Metadata)
+			if (const UPCGSpatialData* TentativeThresholdSpatialData = Cast<const UPCGSpatialData>(ThresholdData))
+			{
+				// If the threshold is spatial, and we use spatial query, it means we'll have to sample points.
+				// Don't create an accessor yet (ThresholdData = nullptr), it will be created further down.
+				// Otherwise, we convert it to point data and create the accessor.
+				if (Settings->bUseSpatialQuery)
 				{
-					TemporaryMetadata = NewObject<UPCGMetadata>();
-					TemporaryMetadata->Initialize(ThresholdMetadata);
+					ThresholdSpatialData = TentativeThresholdSpatialData;
+					ThresholdData = nullptr;
 				}
-				else if(!bUseSpatialQuery)
+				else
 				{
-					const UPCGPointData* ThresholdData = SpatialDataFilter->ToPointData();
-					ThresholdPoints = &ThresholdData->GetPoints();
+					ThresholdSpatialData = TentativeThresholdSpatialData->ToPointData();
+					if (!ThresholdSpatialData)
+					{
+						PCGE_LOG(Error, "Unable to get point data from filter input");
+						return true;
+					}
+
+					ThresholdData = ThresholdSpatialData;
 				}
 			}
-			else if (const UPCGParamData* ParamDataFilter = Cast<const UPCGParamData>(FilterData[0].Data))
+
+			if (ThresholdData)
 			{
-				ThresholdMetadata = ParamDataFilter->ConstMetadata();
+				ThresholdAccessor = PCGAttributeAccessorHelpers::CreateConstAccessor(ThresholdData, Settings->ThresholdAttribute);
+				ThresholdKeys = PCGAttributeAccessorHelpers::CreateConstKeys(ThresholdData, Settings->ThresholdAttribute);
 			}
 		}
 		else
 		{
-			// Use self as filter
-			ThresholdMetadata = TargetMetadata;
-			ThresholdPoints = &TargetPoints;
+			bUseInputDataForThreshold = true;
 		}
+	}
 
-		const FName LocalTargetAttributeName = ((TargetAttributeName != NAME_None || !TargetMetadata) ? TargetAttributeName : TargetMetadata->GetLatestAttributeNameOrNone());
+	for (const FPCGTaggedData& Input : DataToFilter)
+	{
+		const UPCGSpatialData* SpatialInput = Cast<const UPCGSpatialData>(Input.Data);
+		check(SpatialInput);
 
-		// Additional validation
-		if (TargetFilterType == EPCGPointTargetFilterType::Metadata)
+		const UPCGPointData* OriginalData = SpatialInput->ToPointData(Context);
+		if (!OriginalData)
 		{
-			bool bExitToOutFilter = false;
-
-			if (!TargetMetadata)
-			{
-				PCGE_LOG(Warning, "Target data to filter has no metadata which is required to filter by metadata");
-				bExitToOutFilter = true;
-			}
-			else if (!TargetMetadata->GetConstAttribute(LocalTargetAttributeName))
-			{
-				PCGE_LOG(Warning, "Target metadata does not have the %s attribute", *LocalTargetAttributeName.ToString());
-				bExitToOutFilter = true;
-			}
-
-			if (bExitToOutFilter)
-			{
-				FPCGTaggedData& OutFilterTaggedData = Outputs.Add_GetRef(Input);
-				OutFilterTaggedData.Pin = PCGPointFilterConstants::OutFilterLabel;
-				continue;
-			}
+			PCGE_LOG(Error, "Unable to get point data from input");
+			continue;
 		}
 
-		const FName LocalThresholdAttributeName = ((ThresholdAttributeName != NAME_None || !ThresholdMetadata) ? ThresholdAttributeName : ThresholdMetadata->GetLatestAttributeNameOrNone());
-		
-		if (ThresholdFilterType == EPCGPointThresholdType::Metadata)
+		// Helper lambdas to fail nicely and forward input to in/out filter pin
+		// If there is a problem with threshold -> forward to InFilter
+		auto ForwardInputToInFilterPin = [&Outputs, Input]()
 		{
-			bool bExitToInFilter = false;
-
-			if (!ThresholdMetadata)
-			{
-				PCGE_LOG(Warning, "Filter data has no metadata which is required to filter by metadata");
-				bExitToInFilter = true;
-			}
-			else if (!ThresholdMetadata->GetConstAttribute(LocalThresholdAttributeName))
-			{
-				PCGE_LOG(Warning, "Filter metadata does not have the %s attribute", *LocalThresholdAttributeName.ToString());
-				bExitToInFilter = true;
-			}
+			FPCGTaggedData& InFilterOutput = Outputs.Add_GetRef(Input);
+			InFilterOutput.Pin = PCGPointFilterConstants::InFilterLabel;
+		};
 
-			if (bExitToInFilter)
-			{
-				FPCGTaggedData& InFilterTaggedData = Outputs.Add_GetRef(Input);
-				InFilterTaggedData.Pin = PCGPointFilterConstants::InFilterLabel;
-				continue;
-			}
-		}
-		else if (ThresholdFilterType == EPCGPointThresholdType::Property)
+		// If there is a problem with target -> forward to OutFilter
+		auto ForwardInputToOutFilterPin = [&Outputs, Input]()
 		{
-			if (!bUseSpatialQuery && (!ThresholdPoints || (ThresholdPoints->Num() > 1 && TargetPoints.Num() > 0 && ThresholdPoints->Num() != TargetPoints.Num())))
-			{
-				PCGE_LOG(Error, "Filter data points count (%d) mismatch vs the data to filter (%d)", (ThresholdPoints ? ThresholdPoints->Num() : 0), TargetPoints.Num());
-				continue;
-			}
+			FPCGTaggedData& OutFilterOutput = Outputs.Add_GetRef(Input);
+			OutFilterOutput.Pin = PCGPointFilterConstants::OutFilterLabel;
+		};
 
-			if (!bUseSpatialQuery && !ThresholdSpatialData)
-			{
-				PCGE_LOG(Error, "Filter data cannot be used as spatial data");
-				continue;
-			}
+		const TArray<FPCGPoint>& OriginalPoints = OriginalData->GetPoints();
+
+		TUniquePtr<const IPCGAttributeAccessor> TargetAccessor = PCGAttributeAccessorHelpers::CreateConstAccessor(OriginalData, Settings->TargetAttribute);
+		TUniquePtr<const IPCGAttributeAccessorKeys> TargetKeys = PCGAttributeAccessorHelpers::CreateConstKeys(OriginalData, Settings->TargetAttribute);
+
+		if (!TargetAccessor.IsValid() || !TargetKeys.IsValid())
+		{
+			PCGE_LOG(Warning, "TargetData doesn't have %s target attribute/property", *Settings->TargetAttribute.GetName().ToString());
+			ForwardInputToOutFilterPin();
+			continue;
 		}
 
-		// Build our functions
-		EPCGPointFilterConstantType TargetType = EPCGPointFilterConstantType::Unknown;
-		EPCGPointFilterConstantType ThresholdType = EPCGPointFilterConstantType::Unknown;
+		if (bUseInputDataForThreshold)
+		{
+			// If we have no threshold accessor, we use the same data as input
+			ThresholdAccessor = PCGAttributeAccessorHelpers::CreateConstAccessor(OriginalData, Settings->ThresholdAttribute);
+			ThresholdKeys = PCGAttributeAccessorHelpers::CreateConstKeys(OriginalData, Settings->ThresholdAttribute);
+		}
+		else if (ThresholdSpatialData != nullptr && Settings->bUseSpatialQuery)
+		{
+			// Reset the point data and reserving some points
+			// No need to reserve the full number of points, since we'll go by chunk
+			// Only allocate the chunk size
+			ThresholdPointData = NewObject<UPCGPointData>();
+			ThresholdPointData->InitializeFromData(ThresholdSpatialData);
+			ThresholdPointData->GetMutablePoints().SetNum(PCGPointFilterConstants::ChunkSize);
 
-		TFunction<void(const FPCGPoint&, const UPCGMetadata*, void*)> SourceGetter = PCGPointFilterHelpers::ConstructSourceGetter(
-			TargetFilterType, 
-			TargetPointProperty, 
-			TargetMetadata,
-			LocalTargetAttributeName,
-			TargetType);
+			// Accessor will be valid, but keys will point to default points. But since it is a view, it will be updated when we sample the points.
+			ThresholdAccessor = PCGAttributeAccessorHelpers::CreateConstAccessor(ThresholdPointData, Settings->ThresholdAttribute);
+			ThresholdKeys = PCGAttributeAccessorHelpers::CreateConstKeys(ThresholdPointData, Settings->ThresholdAttribute);
+		}
 
-		if (TargetType == EPCGPointFilterConstantType::Unknown)
+		if (!ThresholdAccessor.IsValid() || !ThresholdKeys.IsValid())
 		{
-			PCGE_LOG(Error, "Unable to generate the data value getter");
+			PCGE_LOG(Warning, "DataToFilter doesn't have %s threshold attribute/property", *Settings->ThresholdAttribute.GetName().ToString());
+			ForwardInputToInFilterPin();
 			continue;
 		}
 
-		// B - a getter that returns the value for the threshold data
-		TFunction<void(const FPCGPoint&, const UPCGMetadata*, void*)> ThresholdGetter = PCGPointFilterHelpers::ConstructThresholdGetter(
-			ThresholdFilterType,
-			ThresholdPointProperty,
-			ThresholdMetadata,
-			LocalThresholdAttributeName,
-			ThresholdConstantType,
-			Integer64Constant,
-			FloatConstant,
-			VectorConstant,
-			Vector4Constant,
-			//RotationConstant,
-			StringConstant,
-			ThresholdType);
-
-		if (ThresholdType == EPCGPointFilterConstantType::Unknown)
+		// Comparison between threshold and target data needs to be of the same type. So we have to make sure that we can
+		// request target type from threshold type. ie. We need to make sure we can broadcast threshold type to target type.
+		// For example: if target is double but threshold is int32, we can broadcast int32 to double, to compare a double with a double.
+		if (!PCG::Private::IsBroadcastable(ThresholdAccessor->GetUnderlyingType(), TargetAccessor->GetUnderlyingType()))
 		{
-			PCGE_LOG(Error, "Unable to generate the filter value getter");
+			PCGE_LOG(Warning, "Cannot broadcast threshold type to target type");
+			ForwardInputToInFilterPin();
 			continue;
 		}
 
-		// C - the operator
-		bool bConcreteOperatorConstructedCorrectly = false;
-		TFunction<bool(const FPCGPoint&, const UPCGMetadata*, const FPCGPoint&, const UPCGMetadata*)> ConcreteOperator = PCGPointFilterHelpers::ConstructConcreteOperator(
-			SourceGetter,
-			TargetType,
-			ThresholdGetter,
-			ThresholdType,
-			Operator,
-			bConcreteOperatorConstructedCorrectly);
-
-		if (!bConcreteOperatorConstructedCorrectly)
+		// And also validate that types are comparable
+		if (Settings->Operator != EPCGPointFilterOperator::Equal && Settings->Operator != EPCGPointFilterOperator::NotEqual)
 		{
-			PCGE_LOG(Error, "Unable to generate the operator (most likely data types mismatch");
-			continue;
+			bool bCanCompare = PCGMetadataAttribute::CallbackWithRightType(TargetAccessor->GetUnderlyingType(), [](auto Dummy) -> bool
+			{
+				return PCG::Private::MetadataTraits<decltype(Dummy)>::CanCompare;
+			});
+
+			if (!bCanCompare)
+			{
+				PCGE_LOG(Warning, "Cannot compare target type.");
+				ForwardInputToOutFilterPin();
+				continue;
+			}
 		}
 
 		UPCGPointData* InFilterData = NewObject<UPCGPointData>();
-		InFilterData->InitializeFromData(OriginalData);
-		TArray<FPCGPoint>& InPoints = InFilterData->GetMutablePoints();
-
 		UPCGPointData* OutFilterData = NewObject<UPCGPointData>();
+
+		InFilterData->InitializeFromData(OriginalData);
 		OutFilterData->InitializeFromData(OriginalData);
-		TArray<FPCGPoint>& OutPoints = OutFilterData->GetMutablePoints();
+		TArray<FPCGPoint>& InFilterPoints = InFilterData->GetMutablePoints();
+		TArray<FPCGPoint>& OutFilterPoints = OutFilterData->GetMutablePoints();
+
+		InFilterPoints.Reserve(OriginalPoints.Num());
+		OutFilterPoints.Reserve(OriginalPoints.Num());
 
-		FPCGAsync::AsyncPointFilterProcessing(Context, TargetPoints.Num(), InPoints, OutPoints, [&TargetPoints, TargetMetadata, ThresholdSpatialData, ThresholdPoints, ThresholdMetadata, TemporaryMetadata, &ConcreteOperator](int32 Index, FPCGPoint& InFilterPoint, FPCGPoint& OutFilterPoint)
+		auto Operation = [&InFilterPoints, &OutFilterPoints, Settings, &ThresholdAccessor, &ThresholdKeys, &TargetAccessor, &TargetKeys, &OriginalPoints, ThresholdPointData, ThresholdSpatialData](auto Dummy) -> bool
 		{
-			const FPCGPoint& SourcePoint = TargetPoints[Index];
+			using Type = decltype(Dummy);
 
-			bool bSkipTest = false;
-			FPCGPoint ThresholdPoint;
-			if (ThresholdPoints) // not null only in the direct match case
-			{
-				ThresholdPoint = (*ThresholdPoints)[(ThresholdPoints->Num() > 1) ? Index : 0];
-			}
-			else if(ThresholdSpatialData)
-			{
-				bSkipTest = !ThresholdSpatialData->SamplePoint(SourcePoint.Transform, SourcePoint.GetLocalBounds(), ThresholdPoint, TemporaryMetadata);
-			}
-			else // Only metadata entry (param)
-			{
-				ThresholdPoint.MetadataEntry = 0;
-			}
+			const int32 NumberOfEntries = TargetKeys->GetNum();
 
-			// Apply operator
-			bool bPointInFilter = bSkipTest || ConcreteOperator(SourcePoint, TargetMetadata, ThresholdPoint, (TemporaryMetadata ? TemporaryMetadata : ThresholdMetadata));
-			if (bPointInFilter)
+			if (NumberOfEntries <= 0)
 			{
-				InFilterPoint = SourcePoint;
-				return true;
+				return false;
 			}
-			else
+
+			TArray<Type, TInlineAllocator<PCGPointFilterConstants::ChunkSize>> TargetValues;
+			TArray<Type, TInlineAllocator<PCGPointFilterConstants::ChunkSize>> ThresholdValues;
+			TArray<bool, TInlineAllocator<PCGPointFilterConstants::ChunkSize>> SkipTests;
+			TargetValues.SetNum(PCGPointFilterConstants::ChunkSize);
+			ThresholdValues.SetNum(PCGPointFilterConstants::ChunkSize);
+			SkipTests.SetNum(PCGPointFilterConstants::ChunkSize);
+
+			const int32 NumberOfIterations = (NumberOfEntries + PCGPointFilterConstants::ChunkSize - 1) / PCGPointFilterConstants::ChunkSize;
+
+			for (int32 i = 0; i < NumberOfIterations; ++i)
 			{
-				OutFilterPoint = SourcePoint;
-				return false;
+				const int32 StartIndex = i * PCGPointFilterConstants::ChunkSize;
+				const int32 Range = FMath::Min(NumberOfEntries - StartIndex, PCGPointFilterConstants::ChunkSize);
+				TArrayView<Type> TargetView(TargetValues.GetData(), Range);
+				TArrayView<Type> ThresholdView(ThresholdValues.GetData(), Range);
+
+				// Sampling the points if needed
+				if (ThresholdPointData != nullptr)
+				{
+					// Threshold points only have "ChunkSize" points.
+					TArray<FPCGPoint>& ThresholdPoints = ThresholdPointData->GetMutablePoints();
+					for (int32 j = 0; j < Range; ++j)
+					{
+						FPCGPoint ThresholdPoint;
+						const FPCGPoint& SourcePoint = OriginalPoints[StartIndex + j];
+						if (ThresholdSpatialData->SamplePoint(SourcePoint.Transform, SourcePoint.GetLocalBounds(), ThresholdPoint, ThresholdPointData->Metadata))
+						{
+							ThresholdPoints[j] = ThresholdPoint;
+							SkipTests[j] = false;
+						}
+						else
+						{
+							SkipTests[j] = true;
+						}
+					}
+				}
+
+				// If ThresholdView point on ThresholdPointData points, there are only "ChunkSize" points in it.
+				// But it wraps around, and since StartIndex is a multiple of "ChunkSize", we'll always start at point 0, as wanted. 
+				if (!TargetAccessor->GetRange(TargetView, StartIndex, *TargetKeys) ||
+					!ThresholdAccessor->GetRange(ThresholdView, StartIndex, *ThresholdKeys, EPCGAttributeAccessorFlags::AllowBroadcast))
+				{
+					return false;
+				}
+
+				for (int32 j = 0; j < Range; ++j)
+				{
+					if (SkipTests[j] || PCGPointFilterHelpers::ApplyCompare(TargetValues[j], ThresholdValues[j], Settings->Operator))
+					{
+						InFilterPoints.Add(OriginalPoints[StartIndex + j]);
+					}
+					else
+					{
+						OutFilterPoints.Add(OriginalPoints[StartIndex + j]);
+					}
+				}
 			}
-		});
-
-		// Push the filtered data to the output data.
-		FPCGTaggedData& InFilterTaggedData = Outputs.Emplace_GetRef();
-		InFilterTaggedData.Data = InFilterData;
-		InFilterTaggedData.Tags = Input.Tags;
-		InFilterTaggedData.Pin = PCGPointFilterConstants::InFilterLabel;
-
-		FPCGTaggedData& OutFilterTaggedData = Outputs.Emplace_GetRef();
-		OutFilterTaggedData.Data = OutFilterData;
-		OutFilterTaggedData.Tags = Input.Tags;
-		OutFilterTaggedData.Pin = PCGPointFilterConstants::OutFilterLabel;
+
+			return true;
+		};
+
+		if (PCGMetadataAttribute::CallbackWithRightType(TargetAccessor->GetUnderlyingType(), Operation))
+		{
+			FPCGTaggedData& InFilterOutput = Outputs.Add_GetRef(Input);
+			InFilterOutput.Pin = PCGPointFilterConstants::InFilterLabel;
+			InFilterOutput.Data = InFilterData;
+			InFilterOutput.Tags = Input.Tags;
+
+			FPCGTaggedData& OutFilterOutput = Outputs.Add_GetRef(Input);
+			OutFilterOutput.Pin = PCGPointFilterConstants::OutFilterLabel;
+			OutFilterOutput.Data = OutFilterData;
+			OutFilterOutput.Tags = Input.Tags;
+		}
 	}
 
 	return true;
diff --git a/Engine/Plugins/Experimental/PCG/Source/PCG/Private/Metadata/Accessors/PCGAttributeAccessorHelpers.cpp b/Engine/Plugins/Experimental/PCG/Source/PCG/Private/Metadata/Accessors/PCGAttributeAccessorHelpers.cpp
--- a/Engine/Plugins/Experimental/PCG/Source/PCG/Private/Metadata/Accessors/PCGAttributeAccessorHelpers.cpp
+++ b/Engine/Plugins/Experimental/PCG/Source/PCG/Private/Metadata/Accessors/PCGAttributeAccessorHelpers.cpp
@@ -35,6 +35,12 @@ namespace PCGAttributeAccessorHelpers
 
 		if (OutMetadata)
 		{
+			// If Name is None, try ot get the latest attribute
+			if (Name == NAME_None)
+			{
+				Name = OutMetadata->GetLatestAttributeNameOrNone();
+			}
+
 			OutAttribute = OutMetadata->GetMutableAttribute(Name);
 		}
 	}
diff --git a/Engine/Plugins/Experimental/PCG/Source/PCG/Private/Metadata/PCGMetadataTypesConstantStruct.cpp b/Engine/Plugins/Experimental/PCG/Source/PCG/Private/Metadata/PCGMetadataTypesConstantStruct.cpp
new file mode 100644
--- /dev/null
+++ b/Engine/Plugins/Experimental/PCG/Source/PCG/Private/Metadata/PCGMetadataTypesConstantStruct.cpp
@@ -0,0 +1,38 @@
+// Copyright Epic Games, Inc. All Rights Reserved.
+
+#include "Metadata/PCGMetadataTypesConstantStruct.h"
+
+FString FPCGMetadataTypesConstantStruct::ToString() const
+{
+	switch (Type)
+	{
+	case EPCGMetadataTypes::Integer32:
+		return FString::Printf(TEXT("%d"), Int32Value);
+	case EPCGMetadataTypes::Float:
+		return FString::Printf(TEXT("%.2f"), FloatValue);
+	case EPCGMetadataTypes::Integer64:
+		return FString::Printf(TEXT("%lld"), IntValue);
+	case EPCGMetadataTypes::Double:
+		return FString::Printf(TEXT("%.2f"), DoubleValue);
+	case EPCGMetadataTypes::String:
+		return FString::Printf(TEXT("\"%s\""), *StringValue);
+	case EPCGMetadataTypes::Name:
+		return FString::Printf(TEXT("N(\"%s\")"), *NameValue.ToString());
+	case EPCGMetadataTypes::Vector2:
+		return FString::Printf(TEXT("V(%.2f, %.2f)"), Vector2Value.X, Vector2Value.Y);
+	case EPCGMetadataTypes::Vector:
+		return FString::Printf(TEXT("V(%.2f, %.2f, %.2f)"), VectorValue.X, VectorValue.Y, VectorValue.Z);
+	case EPCGMetadataTypes::Vector4:
+		return FString::Printf(TEXT("V(%.2f, %.2f, %.2f, %.2f)"), Vector4Value.X, Vector4Value.Y, Vector4Value.Z, Vector4Value.W);
+	case EPCGMetadataTypes::Rotator:
+		return FString::Printf(TEXT("R(%.2f, %.2f, %.2f)"), RotatorValue.Roll, RotatorValue.Pitch, RotatorValue.Yaw);
+	case EPCGMetadataTypes::Quaternion:
+		return FString::Printf(TEXT("Q(%.2f, %.2f, %.2f, %.2f)"), QuatValue.X, QuatValue.Y, QuatValue.Z, QuatValue.W);
+	case EPCGMetadataTypes::Transform:
+		return FString::Printf(TEXT("Transform"));
+	case EPCGMetadataTypes::Boolean:
+		return FString::Printf(TEXT("%s"), (BoolValue ? TEXT("True") : TEXT("False")));
+	default:
+		return FString();
+	}
+}
\ No newline at end of file
diff --git a/Engine/Plugins/Experimental/PCG/Source/PCG/Public/Elements/Metadata/PCGMetadataOpElementBase.h b/Engine/Plugins/Experimental/PCG/Source/PCG/Public/Elements/Metadata/PCGMetadataOpElementBase.h
--- a/Engine/Plugins/Experimental/PCG/Source/PCG/Public/Elements/Metadata/PCGMetadataOpElementBase.h
+++ b/Engine/Plugins/Experimental/PCG/Source/PCG/Public/Elements/Metadata/PCGMetadataOpElementBase.h
@@ -102,8 +102,6 @@ public:
 	virtual bool HasDifferentOutputTypes() const { return false; }
 	virtual TArray<uint16> GetAllOutputTypes() const { return TArray<uint16>(); };
 
-	bool IsMoreComplexType(uint16 FirstType, uint16 SecondType) const;
-
 	UPROPERTY(BlueprintReadWrite, EditAnywhere, Category = "Output")
 	FPCGAttributePropertySelector OutputTarget;
 
diff --git a/Engine/Plugins/Experimental/PCG/Source/PCG/Public/Elements/PCGCreateAttribute.h b/Engine/Plugins/Experimental/PCG/Source/PCG/Public/Elements/PCGCreateAttribute.h
--- a/Engine/Plugins/Experimental/PCG/Source/PCG/Public/Elements/PCGCreateAttribute.h
+++ b/Engine/Plugins/Experimental/PCG/Source/PCG/Public/Elements/PCGCreateAttribute.h
@@ -6,6 +6,7 @@
 #include "PCGSettings.h"
 #include "Metadata/PCGMetadataCommon.h"
 #include "Metadata/PCGMetadataAttributeTpl.h"
+#include "Metadata/PCGMetadataTypesConstantStruct.h"
 
 #include "PCGCreateAttribute.generated.h"
 
@@ -25,6 +26,10 @@ class PCG_API UPCGCreateAttributeSettings : public UPCGSettings
 {
 	GENERATED_BODY()
 public:
+	//~Begin UObject interface
+	virtual void PostLoad() override;
+	//~End UObject interface
+
 	//~Begin UPCGSettings interface
 #if WITH_EDITOR
 	virtual FName GetDefaultNodeName() const override;
@@ -46,50 +51,52 @@ public:
 	UPROPERTY(BlueprintReadWrite, EditAnywhere, Category = Settings, meta = (EditCondition = "bFromSourceParam", EditConditionHides))
 	FName SourceParamAttributeName = NAME_None;
 
-	UPROPERTY(BlueprintReadWrite, EditAnywhere, Category = Settings, meta = (ValidEnumValues = "Float, Double, Integer32, Integer64, Vector2, Vector, Vector4, Quaternion, Transform, String, Boolean, Rotator, Name", EditCondition = "!bFromSourceParam", EditConditionHides))
-	EPCGMetadataTypes Type = EPCGMetadataTypes::Double;
+	UPROPERTY(BlueprintReadWrite, EditAnywhere, Category = Settings, meta = (EditCondition = "!bFromSourceParam", EditConditionHides, ShowOnlyInnerProperties, DisplayAfter = "bFromSourceParam"))
+	FPCGMetadataTypesConstantStruct AttributeTypes;
 
-	// All different types
-	UPROPERTY(BlueprintReadWrite, EditAnywhere, Category = Settings, meta = (EditCondition = "Type == EPCGMetadataTypes::Float && !bFromSourceParam", EditConditionHides))
-	float FloatValue = 0.0;
+#if WITH_EDITORONLY_DATA
+	UPROPERTY()
+	EPCGMetadataTypes Type_DEPRECATED = EPCGMetadataTypes::Double;
 
-	UPROPERTY(BlueprintReadWrite, EditAnywhere, Category = Settings, meta = (EditCondition = "Type == EPCGMetadataTypes::Integer32 && !bFromSourceParam", EditConditionHides))
-	int32 Int32Value = 0;
+	UPROPERTY()
+	float FloatValue_DEPRECATED = 0.0;
 
-	UPROPERTY(BlueprintReadWrite, EditAnywhere, Category = Settings, meta = (EditCondition = "Type == EPCGMetadataTypes::Double && !bFromSourceParam", EditConditionHides))
-	double DoubleValue = 0.0;
+	UPROPERTY()
+	int32 Int32Value_DEPRECATED = 0;
 
-	UPROPERTY(BlueprintReadWrite, EditAnywhere, Category = Settings, meta = (EditCondition = "Type == EPCGMetadataTypes::Integer64 && !bFromSourceParam", EditConditionHides))
-	int64 IntValue = 0;
+	UPROPERTY()
+	double DoubleValue_DEPRECATED = 0.0;
+
+	UPROPERTY()
+	int64 IntValue_DEPRECATED = 0;
 
-	UPROPERTY(BlueprintReadWrite, EditAnywhere, Category = Settings, meta = (EditCondition = "Type == EPCGMetadataTypes::Vector2 && !bFromSourceParam", EditConditionHides))
-	FVector2D Vector2Value = FVector2D::ZeroVector;
+	UPROPERTY()
+	FVector2D Vector2Value_DEPRECATED = FVector2D::ZeroVector;
 
-	UPROPERTY(BlueprintReadWrite, EditAnywhere, Category = Settings, meta = (EditCondition = "Type == EPCGMetadataTypes::Vector && !bFromSourceParam", EditConditionHides))
-	FVector VectorValue = FVector::ZeroVector;
+	UPROPERTY()
+	FVector VectorValue_DEPRECATED = FVector::ZeroVector;
 
-	UPROPERTY(BlueprintReadWrite, EditAnywhere, Category = Settings, meta = (EditCondition = "Type == EPCGMetadataTypes::Vector4 && !bFromSourceParam", EditConditionHides))
-	FVector4 Vector4Value = FVector4::Zero();
+	UPROPERTY()
+	FVector4 Vector4Value_DEPRECATED = FVector4::Zero();
 
-	UPROPERTY(BlueprintReadWrite, EditAnywhere, Category = Settings, meta = (EditCondition = "Type == EPCGMetadataTypes::Quaternion && !bFromSourceParam", EditConditionHides))
-	FQuat QuatValue = FQuat::Identity;
+	UPROPERTY()
+	FQuat QuatValue_DEPRECATED = FQuat::Identity;
 
-	UPROPERTY(BlueprintReadWrite, EditAnywhere, Category = Settings, meta = (EditCondition = "Type == EPCGMetadataTypes::Transform && !bFromSourceParam", EditConditionHides))
-	FTransform TransformValue = FTransform::Identity;
+	UPROPERTY()
+	FTransform TransformValue_DEPRECATED = FTransform::Identity;
 
-	UPROPERTY(BlueprintReadWrite, EditAnywhere, Category = Settings, meta = (EditCondition = "Type == EPCGMetadataTypes::String && !bFromSourceParam", EditConditionHides))
-	FString StringValue = "";
+	UPROPERTY()
+	FString StringValue_DEPRECATED = "";
 
-	UPROPERTY(BlueprintReadWrite, EditAnywhere, Category = Settings, meta = (EditCondition = "Type == EPCGMetadataTypes::Boolean && !bFromSourceParam", EditConditionHides))
-	bool BoolValue = false;
+	UPROPERTY()
+	bool BoolValue_DEPRECATED = false;
 
-	UPROPERTY(BlueprintReadWrite, EditAnywhere, Category = Settings, meta = (EditCondition = "Type == EPCGMetadataTypes::Rotator && !bFromSourceParam", EditConditionHides))
-	FRotator RotatorValue = FRotator::ZeroRotator;
+	UPROPERTY()
+	FRotator RotatorValue_DEPRECATED = FRotator::ZeroRotator;
 
-	UPROPERTY(BlueprintReadWrite, EditAnywhere, Category = Settings, meta = (EditCondition = "Type == EPCGMetadataTypes::Name && !bFromSourceParam", EditConditionHides))
-	FName NameValue = NAME_None;
+	UPROPERTY()
+	FName NameValue_DEPRECATED = NAME_None;
 
-#if WITH_EDITORONLY_DATA
 	UPROPERTY()
 	bool bKeepExistingAttributes_DEPRECATED = false;
 #endif // WITH_EDITORONLY_DATA
diff --git a/Engine/Plugins/Experimental/PCG/Source/PCG/Public/Elements/PCGPointFilter.h b/Engine/Plugins/Experimental/PCG/Source/PCG/Public/Elements/PCGPointFilter.h
--- a/Engine/Plugins/Experimental/PCG/Source/PCG/Public/Elements/PCGPointFilter.h
+++ b/Engine/Plugins/Experimental/PCG/Source/PCG/Public/Elements/PCGPointFilter.h
@@ -5,18 +5,20 @@
 #include "PCGSettings.h"
 #include "PCGElement.h"
 #include "PCGPoint.h"
+#include "Metadata/PCGAttributePropertySelector.h"
+#include "Metadata/PCGMetadataTypesConstantStruct.h"
 
 #include "PCGPointFilter.generated.h"
 
 UENUM()
-enum class EPCGPointTargetFilterType : uint8
+enum class UE_DEPRECATED(5.2, "Not used anymore") EPCGPointTargetFilterType : uint8
 {
 	Property,
 	Metadata
 };
 
 UENUM()
-enum class EPCGPointThresholdType : uint8
+enum class UE_DEPRECATED(5.2, "Not used anymore") EPCGPointThresholdType : uint8
 {
 	Property,
 	Metadata,
@@ -24,7 +26,7 @@ enum class EPCGPointThresholdType : uint8
 };
 
 UENUM()
-enum class EPCGPointFilterConstantType : uint8
+enum class UE_DEPRECATED(5.2, "Not used anymore") EPCGPointFilterConstantType : uint8
 {
 	Integer64,
 	Float,
@@ -64,6 +66,12 @@ class PCG_API UPCGPointFilterSettings : public UPCGSettings
 	GENERATED_BODY()
 
 public:
+	UPCGPointFilterSettings();
+
+	//~Begin UObject interface
+	virtual void PostLoad() override;
+	//~End UObject interface
+
 	//~Begin UPCGSettings interface
 #if WITH_EDITOR
 	virtual FName GetDefaultNodeName() const override { return FName(TEXT("PointFilter")); }
@@ -83,48 +91,60 @@ public:
 
 	/** Target property/attribute related properties */
 	UPROPERTY(BlueprintReadWrite, EditAnywhere, Category = Settings)
-	EPCGPointTargetFilterType TargetFilterType = EPCGPointTargetFilterType::Property;
-
-	UPROPERTY(BlueprintReadWrite, EditAnywhere, Category = Settings, meta = (EditCondition = "TargetFilterType==EPCGPointTargetFilterType::Property"))
-	EPCGPointProperties TargetPointProperty = EPCGPointProperties::Density;
-
-	UPROPERTY(BlueprintReadWrite, EditAnywhere, Category = Settings, meta = (EditCondition = "TargetFilterType==EPCGPointTargetFilterType::Metadata"))
-	FName TargetAttributeName = NAME_None;
+	FPCGAttributePropertySelector TargetAttribute;
 
 	/** Threshold property/attribute/constant related properties */
 	UPROPERTY(BlueprintReadWrite, EditAnywhere, Category = Settings)
-	EPCGPointThresholdType ThresholdFilterType = EPCGPointThresholdType::Property;
+	bool bUseConstantThreshold = false;
 
-	UPROPERTY(BlueprintReadWrite, EditAnywhere, Category = Settings, meta = (EditCondition = "ThresholdFilterType==EPCGPointThresholdType::Property"))
-	EPCGPointProperties ThresholdPointProperty = EPCGPointProperties::Density;
+	UPROPERTY(BlueprintReadWrite, EditAnywhere, Category = Settings, meta = (EditCondition = "!bUseConstantThreshold", EditConditionHides))
+	FPCGAttributePropertySelector ThresholdAttribute;
 
-	UPROPERTY(BlueprintReadWrite, EditAnywhere, Category = Settings, meta = (EditCondition = "ThresholdFilterType==EPCGPointThresholdType::Metadata"))
-	FName ThresholdAttributeName = NAME_None;
+	UPROPERTY(BlueprintReadWrite, EditAnywhere, Category = Settings, meta = (EditCondition = "!bUseConstantThreshold", EditConditionHides))
+	bool bUseSpatialQuery = true;
 
-	UPROPERTY(BlueprintReadWrite, EditAnywhere, Category = Settings, meta = (EditCondition = "ThresholdFilterType==EPCGPointThresholdType::Constant"))
-	EPCGPointFilterConstantType ThresholdConstantType = EPCGPointFilterConstantType::Float;
+	UPROPERTY(BlueprintReadWrite, EditAnywhere, Category = Constant, meta = (EditCondition = "bUseConstantThreshold", EditConditionHides, ShowOnlyInnerProperties, DisplayAfter = "bUseConstantThreshold"))
+	FPCGMetadataTypesConstantStruct AttributeTypes;
 
-	/** Constants used as threshold comparator */
-	UPROPERTY(BlueprintReadWrite, EditAnywhere, Category = Settings, meta = (EditCondition = "ThresholdFilterType==EPCGPointThresholdType::Constant && ThresholdConstantType==EPCGPointFilterConstantType::Integer64"))
-	int64 Integer64Constant = 0;
+#if WITH_EDITORONLY_DATA
+PRAGMA_DISABLE_DEPRECATION_WARNINGS
+	UPROPERTY()
+	EPCGPointTargetFilterType TargetFilterType_DEPRECATED = EPCGPointTargetFilterType::Property;
 
-	UPROPERTY(BlueprintReadWrite, EditAnywhere, Category = Settings, meta=(EditCondition="ThresholdFilterType==EPCGPointThresholdType::Constant && ThresholdConstantType==EPCGPointFilterConstantType::Float"))
-	float FloatConstant = 0;
+	UPROPERTY()
+	EPCGPointProperties TargetPointProperty_DEPRECATED = EPCGPointProperties::Density;
 
-	UPROPERTY(BlueprintReadWrite, EditAnywhere, Category = Settings, meta = (EditCondition = "ThresholdFilterType==EPCGPointThresholdType::Constant && ThresholdConstantType==EPCGPointFilterConstantType::Vector"))
-	FVector VectorConstant = FVector::Zero();
+	UPROPERTY()
+	FName TargetAttributeName_DEPRECATED = NAME_None;
 
-	UPROPERTY(BlueprintReadWrite, EditAnywhere, Category = Settings, meta = (EditCondition = "ThresholdFilterType==EPCGPointThresholdType::Constant && ThresholdConstantType==EPCGPointFilterConstantType::Vector4"))
-	FVector4 Vector4Constant = FVector4::Zero();
+	UPROPERTY()
+	EPCGPointThresholdType ThresholdFilterType_DEPRECATED = EPCGPointThresholdType::Property;
 
-	//UPROPERTY(BlueprintReadWrite, EditAnywhere, Category = Settings, meta = (EditCondition = "ThresholdFilterType==EPCGPointThresholdType::Constant && ThresholdConstantType==EPCGPointFilterConstantType::Rotation"))
-	//FQuat RotationConstant = FQuat::Identity;
+	UPROPERTY()
+	EPCGPointProperties ThresholdPointProperty_DEPRECATED = EPCGPointProperties::Density;
 
-	UPROPERTY(BlueprintReadWrite, EditAnywhere, Category = Settings, meta = (EditCondition = "ThresholdFilterType==EPCGPointThresholdType::Constant && ThresholdConstantType==EPCGPointFilterConstantType::String"))
-	FString StringConstant;
+	UPROPERTY()
+	FName ThresholdAttributeName_DEPRECATED = NAME_None;
 
-	UPROPERTY(BlueprintReadWrite, EditAnywhere, Category = Settings, meta = (EditCondition = "ThresholdFilterType!=EPCGPointThresholdType::Constant"))
-	bool bUseSpatialQuery = true;
+	UPROPERTY()
+	EPCGPointFilterConstantType ThresholdConstantType_DEPRECATED = EPCGPointFilterConstantType::Float;
+
+	UPROPERTY()
+	int64 Integer64Constant_DEPRECATED = 0;
+
+	UPROPERTY()
+	float FloatConstant_DEPRECATED = 0;
+
+	UPROPERTY()
+	FVector VectorConstant_DEPRECATED = FVector::Zero();
+
+	UPROPERTY()
+	FVector4 Vector4Constant_DEPRECATED = FVector4::Zero();
+
+	UPROPERTY()
+	FString StringConstant_DEPRECATED;
+PRAGMA_ENABLE_DEPRECATION_WARNINGS
+#endif 
 };
 
 class FPCGPointFilterElement : public FSimplePCGElement
diff --git a/Engine/Plugins/Experimental/PCG/Source/PCG/Public/Metadata/Accessors/PCGAttributeAccessor.h b/Engine/Plugins/Experimental/PCG/Source/PCG/Public/Metadata/Accessors/PCGAttributeAccessor.h
--- a/Engine/Plugins/Experimental/PCG/Source/PCG/Public/Metadata/Accessors/PCGAttributeAccessor.h
+++ b/Engine/Plugins/Experimental/PCG/Source/PCG/Public/Metadata/Accessors/PCGAttributeAccessor.h
@@ -21,12 +21,13 @@ public:
 	using Type = T;
 	using Super = IPCGAttributeAccessorT<FPCGAttributeAccessor<T>>;
 
+	// Can't write if metadata is null
 	FPCGAttributeAccessor(FPCGMetadataAttribute<T>* InAttribute, UPCGMetadata* InMetadata)
-		: Super(/*bInReadOnly=*/ false)
+		: Super(/*bInReadOnly=*/ InMetadata == nullptr)
 		, Attribute(InAttribute)
 		, Metadata(InMetadata)
 	{
-		check(InAttribute && InMetadata);
+		check(InAttribute);
 	}
 
 	FPCGAttributeAccessor(const FPCGMetadataAttribute<T>* InAttribute, const UPCGMetadata* InMetadata)
@@ -34,7 +35,7 @@ public:
 		, Attribute(const_cast<FPCGMetadataAttribute<T>*>(InAttribute))
 		, Metadata(const_cast<UPCGMetadata*>(InMetadata))
 	{
-		check(InAttribute && InMetadata);
+		check(InAttribute);
 	}
 
 	bool GetRangeImpl(TArrayView<T> OutValues, int32 Index, const IPCGAttributeAccessorKeys& Keys) const
diff --git a/Engine/Plugins/Experimental/PCG/Source/PCG/Public/Metadata/Accessors/PCGAttributeAccessorKeys.h b/Engine/Plugins/Experimental/PCG/Source/PCG/Public/Metadata/Accessors/PCGAttributeAccessorKeys.h
--- a/Engine/Plugins/Experimental/PCG/Source/PCG/Public/Metadata/Accessors/PCGAttributeAccessorKeys.h
+++ b/Engine/Plugins/Experimental/PCG/Source/PCG/Public/Metadata/Accessors/PCGAttributeAccessorKeys.h
@@ -12,13 +12,44 @@ class FPCGMetadataAttributeBase;
 
 ///////////////////////////////////////////////////////////////////////
 
+namespace PCGAttributeAccessorKeys
+{
+	template <typename T, typename Container, typename Func>
+	bool GetKeys(Container& InContainer, int32 InStart, TArrayView<T*>& OutItems, Func&& Transform)
+	{
+		if (InContainer.Num() == 0)
+		{
+			return false;
+		}
+
+		int32 Current = InStart;
+		if (Current >= InContainer.Num())
+		{
+			Current %= InContainer.Num();
+		}
+
+		for (int32 i = 0; i < OutItems.Num(); ++i)
+		{
+			OutItems[i] = Transform(InContainer[Current++]);
+			if (Current >= InContainer.Num())
+			{
+				Current = 0;
+			}
+		}
+
+		return true;
+	}
+}
+
+///////////////////////////////////////////////////////////////////////
+
 /**
 * Base class to identify keys to use with an accessor.
 */
 class IPCGAttributeAccessorKeys
 {
 public:
-	IPCGAttributeAccessorKeys(bool bInReadOnly)
+	explicit IPCGAttributeAccessorKeys(bool bInReadOnly)
 		: bIsReadOnly(bInReadOnly)
 	{}
 
@@ -105,8 +136,8 @@ protected:
 class FPCGAttributeAccessorKeysEntries : public IPCGAttributeAccessorKeys
 {
 public:
-	FPCGAttributeAccessorKeysEntries(const FPCGMetadataAttributeBase* Attribute);
-	FPCGAttributeAccessorKeysEntries(PCGMetadataEntryKey EntryKey);
+	explicit FPCGAttributeAccessorKeysEntries(const FPCGMetadataAttributeBase* Attribute);
+	explicit FPCGAttributeAccessorKeysEntries(PCGMetadataEntryKey EntryKey);
 
 	virtual int32 GetNum() const override { return Entries.Num(); }
 
@@ -128,8 +159,8 @@ public:
 	FPCGAttributeAccessorKeysPoints(const TArrayView<FPCGPoint>& InPoints);
 	FPCGAttributeAccessorKeysPoints(const TArrayView<const FPCGPoint>& InPoints);
 
-	FPCGAttributeAccessorKeysPoints(FPCGPoint& InPoint);
-	FPCGAttributeAccessorKeysPoints(const FPCGPoint& InPoint);
+	explicit FPCGAttributeAccessorKeysPoints(FPCGPoint& InPoint);
+	explicit FPCGAttributeAccessorKeysPoints(const FPCGPoint& InPoint);
 
 	virtual int32 GetNum() const override { return Points.Num(); }
 
@@ -149,9 +180,11 @@ protected:
 /////////////////////////////////////////////////////////////////
 
 /**
-* Key around generic objects
+* Key around generic objects. 
+* Make sure ObjectType is not a pointer nor a reference, since we convert those to void*, it could lead to
+* very bad situations if we try to convert a T** to a void*.
 */
-template <typename ObjectType>
+template <typename ObjectType, typename = typename std::enable_if_t<!std::is_pointer_v<ObjectType> && !std::is_reference_v<ObjectType>>>
 class FPCGAttributeAccessorKeysGeneric : public IPCGAttributeAccessorKeys
 {
 public:
@@ -165,19 +198,26 @@ public:
 		, Objects(const_cast<ObjectType*>(InObjects.GetData()), InObjects.Num())
 	{}
 
-	FPCGAttributeAccessorKeysGeneric(ObjectType& InObject)
+	explicit FPCGAttributeAccessorKeysGeneric(ObjectType& InObject)
 		: FPCGAttributeAccessorKeysGeneric(TArrayView<ObjectType>(&InObject, 1))
 	{}
 
-	FPCGAttributeAccessorKeysGeneric(const ObjectType& InObject)
+	explicit FPCGAttributeAccessorKeysGeneric(const ObjectType& InObject)
 		: FPCGAttributeAccessorKeysGeneric(TArrayView<const ObjectType>(&InObject, 1))
 	{}
 
 	virtual int32 GetNum() const override { return Objects.Num(); }
 
 protected:
-	virtual bool GetGenericObjectKeys(int32 InStart, TArrayView<void*>& OutObjects) override;
-	virtual bool GetGenericObjectKeys(int32 InStart, TArrayView<const void*>& OutObjects) const override;
+	virtual bool GetGenericObjectKeys(int32 InStart, TArrayView<void*>& OutObjects) override
+	{
+		return PCGAttributeAccessorKeys::GetKeys(Objects, InStart, OutObjects, [](ObjectType& Obj) -> ObjectType* { return &Obj; });
+	}
+
+	virtual bool GetGenericObjectKeys(int32 InStart, TArrayView<const void*>& OutObjects) const override
+	{
+		return PCGAttributeAccessorKeys::GetKeys(Objects, InStart, OutObjects, [](const ObjectType& Obj) -> const ObjectType* { return &Obj; });
+	}
 
 	TArrayView<ObjectType> Objects;
 };
@@ -188,17 +228,25 @@ protected:
 * Unique Key around a single object.
 * Necessary if ObjectType is void, but keep a template version for completeness.
 * Useful when you want to use the accessors Get/Set methods on a single object.
+* Make sure ObjectType is not a pointer nor a reference, since we convert those to void*, it could lead to
+* very bad situations if we try to convert a T** to a void*.
 */
-template <typename ObjectType>
+template <typename ObjectType, typename = typename std::enable_if_t<!std::is_pointer_v<ObjectType> && !std::is_reference_v<ObjectType>>>
 class FPCGAttributeAccessorKeysSingleObjectPtr : public IPCGAttributeAccessorKeys
 {
 public:
-	FPCGAttributeAccessorKeysSingleObjectPtr(ObjectType* InPtr)
+
+	FPCGAttributeAccessorKeysSingleObjectPtr()
+		: IPCGAttributeAccessorKeys(/*bInReadOnly=*/ true)
+		, Ptr(nullptr)
+	{}
+
+	explicit FPCGAttributeAccessorKeysSingleObjectPtr(ObjectType* InPtr)
 		: IPCGAttributeAccessorKeys(/*bInReadOnly=*/ false)
 		, Ptr(InPtr)
 	{}
 
-	FPCGAttributeAccessorKeysSingleObjectPtr(const ObjectType* InPtr)
+	explicit FPCGAttributeAccessorKeysSingleObjectPtr(const ObjectType* InPtr)
 		: IPCGAttributeAccessorKeys(/*bInReadOnly=*/ true)
 		, Ptr(const_cast<ObjectType*>(InPtr))
 	{}
@@ -206,8 +254,35 @@ public:
 	virtual int32 GetNum() const override { return 1; }
 
 protected:
-	virtual bool GetGenericObjectKeys(int32 InStart, TArrayView<void*>& OutObjects) override;
-	virtual bool GetGenericObjectKeys(int32 InStart, TArrayView<const void*>& OutObjects) const override;
+	virtual bool GetGenericObjectKeys(int32 InStart, TArrayView<void*>& OutObjects) override
+	{
+		if (Ptr == nullptr)
+		{
+			return false;
+		}
+
+		for (int32 i = 0; i < OutObjects.Num(); ++i)
+		{
+			OutObjects[i] = Ptr;
+		}
+
+		return true;
+	}
+
+	virtual bool GetGenericObjectKeys(int32 InStart, TArrayView<const void*>& OutObjects) const override
+	{
+		if (Ptr == nullptr)
+		{
+			return false;
+		}
+
+		for (int32 i = 0; i < OutObjects.Num(); ++i)
+		{
+			OutObjects[i] = Ptr;
+		}
+
+		return true;
+	}
 
 	ObjectType* Ptr = nullptr;
 };
@@ -263,79 +338,4 @@ inline bool IPCGAttributeAccessorKeys::GetKeys(int32 InStart, TArrayView<const O
 	{
 		return false;
 	}
-}
-
-/////////////////////////////////////////////////////////////////
-
-namespace PCGAttributeAccessorKeys
-{
-	template <typename T, typename Container, typename Func>
-	bool GetKeys(Container& InContainer, int32 InStart, TArrayView<T*>& OutItems, Func&& Transform)
-	{
-		if (InContainer.Num() == 0)
-		{
-			return false;
-		}
-
-		int32 Current = InStart;
-		if (Current >= InContainer.Num())
-		{
-			Current %= InContainer.Num();
-		}
-
-		for (int32 i = 0; i < OutItems.Num(); ++i)
-		{
-			OutItems[i] = Transform(InContainer[Current++]);
-			if (Current >= InContainer.Num())
-			{
-				Current = 0;
-			}
-		}
-
-		return true;
-	}
-}
-
-template <typename ObjectType>
-bool FPCGAttributeAccessorKeysGeneric<ObjectType>::GetGenericObjectKeys(int32 InStart, TArrayView<void*>& OutObjects)
-{
-	return PCGAttributeAccessorKeys::GetKeys(Objects, InStart, OutObjects, [](ObjectType& Obj) -> ObjectType* { return &Obj; });
-}
-
-template <typename ObjectType>
-bool FPCGAttributeAccessorKeysGeneric<ObjectType>::GetGenericObjectKeys(int32 InStart, TArrayView<const void*>& OutObjects) const
-{
-	return PCGAttributeAccessorKeys::GetKeys(Objects, InStart, OutObjects, [](const ObjectType& Obj) -> const ObjectType* { return &Obj; });
-}
-
-template <typename ObjectType>
-bool FPCGAttributeAccessorKeysSingleObjectPtr<ObjectType>::GetGenericObjectKeys(int32 InStart, TArrayView<void*>& OutObjects)
-{
-	if (Ptr == nullptr)
-	{
-		return false;
-	}
-
-	for (int32 i = 0; i < OutObjects.Num(); ++i)
-	{
-		OutObjects[i] = Ptr;
-	}
-
-	return true;
-}
-
-template <typename ObjectType>
-bool FPCGAttributeAccessorKeysSingleObjectPtr<ObjectType>::GetGenericObjectKeys(int32 InStart, TArrayView<const void*>& OutObjects) const
-{
-	if (Ptr == nullptr)
-	{
-		return false;
-	}
-
-	for (int32 i = 0; i < OutObjects.Num(); ++i)
-	{
-		OutObjects[i] = Ptr;
-	}
-
-	return true;
 }
\ No newline at end of file
diff --git a/Engine/Plugins/Experimental/PCG/Source/PCG/Public/Metadata/Accessors/PCGCustomAccessor.h b/Engine/Plugins/Experimental/PCG/Source/PCG/Public/Metadata/Accessors/PCGCustomAccessor.h
--- a/Engine/Plugins/Experimental/PCG/Source/PCG/Public/Metadata/Accessors/PCGCustomAccessor.h
+++ b/Engine/Plugins/Experimental/PCG/Source/PCG/Public/Metadata/Accessors/PCGCustomAccessor.h
@@ -63,4 +63,39 @@ public:
 private:
 	FPCGPoint::PointCustomPropertyGetter Getter;
 	FPCGPoint::PointCustomPropertySetter Setter;
+};
+
+/**
+* Very simple accessor that returns a constant value. Read only
+* Key supported: All
+*/
+template <typename T>
+class FPCGConstantValueAccessor : public IPCGAttributeAccessorT<FPCGConstantValueAccessor<T>>
+{
+public:
+	using Type = T;
+	using Super = IPCGAttributeAccessorT<FPCGConstantValueAccessor<T>>;
+
+	FPCGConstantValueAccessor(const T& InValue)
+		: Super(/*bInReadOnly=*/ true)
+		, Value(InValue)
+	{}
+
+	bool GetRangeImpl(TArrayView<T> OutValues, int32, const IPCGAttributeAccessorKeys&) const
+	{
+		for (int32 i = 0; i < OutValues.Num(); ++i)
+		{
+			OutValues[i] = Value;
+		}
+
+		return true;
+	}
+
+	bool SetRangeImpl(TArrayView<const T>, int32, IPCGAttributeAccessorKeys&, EPCGAttributeAccessorFlags)
+	{
+		return false;
+	}
+
+private:
+	T Value;
 };
\ No newline at end of file
diff --git a/Engine/Plugins/Experimental/PCG/Source/PCG/Public/Metadata/PCGMetadataAttributeTraits.h b/Engine/Plugins/Experimental/PCG/Source/PCG/Public/Metadata/PCGMetadataAttributeTraits.h
--- a/Engine/Plugins/Experimental/PCG/Source/PCG/Public/Metadata/PCGMetadataAttributeTraits.h
+++ b/Engine/Plugins/Experimental/PCG/Source/PCG/Public/Metadata/PCGMetadataAttributeTraits.h
@@ -38,10 +38,11 @@ enum class EPCGMetadataTypes : uint8
 	MACRO(FVector)    \
 	MACRO(FVector4)   \
 	MACRO(FQuat)      \
-	MACRO(FRotator)   \
-	MACRO(FName)      \
+	MACRO(FTransform) \
 	MACRO(FString)    \
-	MACRO(FTransform)
+	MACRO(bool)       \
+	MACRO(FRotator)   \
+	MACRO(FName)
 
 namespace PCG
 {
@@ -170,6 +171,25 @@ namespace PCG
 			return IsBroadcastable(MetadataTypes<FirstType>::Id, MetadataTypes<SecondType>::Id);
 		}
 
+		/**
+		* FirstType is more complex than SecondType, in the sense of PCG types, if both types are different and valid
+		* and we can broadcast SecondType to FirstType.
+		* It is useful to find the most complex type between operands of a given operation.
+		*/
+		constexpr inline bool IsMoreComplexType(uint16 FirstType, uint16 SecondType)
+		{
+			return FirstType != SecondType && FirstType <= (uint16)(EPCGMetadataTypes::Count) && SecondType <= (uint16)(EPCGMetadataTypes::Count) && PCG::Private::BroadcastableTypes[SecondType][FirstType];
+		}
+
+		/**
+		* cf. non templated version of IsMoreComplexType
+		*/
+		template <typename FirstType, typename SecondType>
+		constexpr inline bool IsMoreComplexType()
+		{
+			return IsMoreComplexType(MetadataTypes<FirstType>::Id, MetadataTypes<SecondType>::Id);
+		}
+
 		/**
 		* Generic function to broadcast an InType to an OutType.
 		* Supports only PCG types
@@ -267,6 +287,32 @@ namespace PCG
 			}
 		};
 
+		template<typename T>
+		struct DefaultCompareTraits
+		{
+			enum { CanCompare = true };
+
+			static bool Less(const T& A, const T& B)
+			{
+				return A < B;
+			}
+
+			static bool Greater(const T& A, const T& B)
+			{
+				return A > B;
+			}
+
+			static bool LessOrEqual(const T& A, const T& B)
+			{
+				return A <= B;
+			}
+
+			static bool GreaterOrEqual(const T& A, const T& B)
+			{
+				return A >= B;
+			}
+		};
+
 		template<typename T>
 		struct DefaultWeightedSumTraits
 		{
@@ -283,11 +329,9 @@ namespace PCG
 			}
 		};
 
-		// Common traits for int32, int64, float, double
-		template<typename T>
-		struct MetadataTraits : DefaultOperationTraits<T>, DefaultWeightedSumTraits<T>
+		template <typename T>
+		struct DefaultMinMaxTraits
 		{
-			enum { CompressData = false };
 			enum { CanMinMax = true };
 
 			static T Min(const T& A, const T& B)
@@ -301,6 +345,35 @@ namespace PCG
 			}
 		};
 
+		// Common traits for int32, int64, float, double
+		template<typename T>
+		struct MetadataTraits : DefaultOperationTraits<T>, DefaultWeightedSumTraits<T>, DefaultMinMaxTraits<T>, DefaultCompareTraits<T>
+		{
+			enum { CompressData = false };
+		};
+
+		template<>
+		struct MetadataTraits<float> : DefaultOperationTraits<float>, DefaultWeightedSumTraits<float>, DefaultMinMaxTraits<float>, DefaultCompareTraits<float>
+		{
+			enum { CompressData = false };
+
+			static bool Equal(const float& A, const float& B)
+			{
+				return FMath::IsNearlyEqual(A, B);
+			}
+		};
+
+		template<>
+		struct MetadataTraits<double> : DefaultOperationTraits<double>, DefaultWeightedSumTraits<double>, DefaultMinMaxTraits<double>, DefaultCompareTraits<double>
+		{
+			enum { CompressData = false };
+
+			static bool Equal(const double& A, const double& B)
+			{
+				return FMath::IsNearlyEqual(A, B);
+			}
+		};
+
 		template<>
 		struct MetadataTraits<bool>
 		{
@@ -309,6 +382,12 @@ namespace PCG
 			enum { CanSubAdd = true };
 			enum { CanMulDiv = false };
 			enum { CanInterpolate = false };
+			enum { CanCompare = true };
+
+			static bool ZeroValue()
+			{
+				return false;
+			}
 
 			static bool Equal(const bool& A, const bool& B)
 			{
@@ -334,6 +413,26 @@ namespace PCG
 			{
 				return A && !B;
 			}
+
+			static bool Less(const bool& A, const bool& B)
+			{
+				return !A && B;
+			}
+
+			static bool Greater(const bool& A, const bool& B)
+			{
+				return A && !B;
+			}
+
+			static bool LessOrEqual(const bool& A, const bool& B)
+			{
+				return !Greater(A, B);
+			}
+
+			static bool GreaterOrEqual(const bool& A, const bool& B)
+			{
+				return !Less(A, B);
+			}
 		};
 
 		// Vector types
@@ -343,6 +442,7 @@ namespace PCG
 			enum { CompressData = false };
 			enum { CanMinMax = true };
 			enum { CanInterpolate = true };
+			enum { CanCompare = true };
 
 			static T ZeroValue()
 			{
@@ -353,6 +453,11 @@ namespace PCG
 		template<>
 		struct MetadataTraits<FVector2D> : VectorTraits<FVector2D>
 		{
+			static bool Equal(const FVector2D& A, const FVector2D& B)
+			{
+				return A.Equals(B);
+			}
+
 			static FVector2D Min(const FVector2D& A, const FVector2D& B)
 			{
 				return FVector2D::Min(A, B);
@@ -362,11 +467,36 @@ namespace PCG
 			{
 				return FVector2D::Max(A, B);
 			}
+
+			static bool Less(const FVector2D& A, const FVector2D& B)
+			{
+				return (A.X < B.X) && (A.Y < B.Y);
+			}
+
+			static bool Greater(const FVector2D& A, const FVector2D& B)
+			{
+				return (A.X > B.X) && (A.Y > B.Y);
+			}
+
+			static bool LessOrEqual(const FVector2D& A, const FVector2D& B)
+			{
+				return (A.X <= B.X) && (A.Y <= B.Y);
+			}
+
+			static bool GreaterOrEqual(const FVector2D& A, const FVector2D& B)
+			{
+				return (A.X >= B.X) && (A.Y >= B.Y);
+			}
 		};
 
 		template<>
 		struct MetadataTraits<FVector> : VectorTraits<FVector>
 		{
+			static bool Equal(const FVector& A, const FVector& B)
+			{
+				return A.Equals(B);
+			}
+
 			static FVector Min(const FVector& A, const FVector& B)
 			{
 				return FVector::Min(A, B);
@@ -376,11 +506,36 @@ namespace PCG
 			{
 				return FVector::Max(A, B);
 			}
+
+			static bool Less(const FVector& A, const FVector& B)
+			{
+				return (A.X < B.X) && (A.Y < B.Y) && (A.Z < B.Z);
+			}
+
+			static bool Greater(const FVector& A, const FVector& B)
+			{
+				return (A.X > B.X) && (A.Y > B.Y) && (A.Z > B.Z);
+			}
+
+			static bool LessOrEqual(const FVector& A, const FVector& B)
+			{
+				return (A.X <= B.X) && (A.Y <= B.Y) && (A.Z <= B.Z);
+			}
+
+			static bool GreaterOrEqual(const FVector& A, const FVector& B)
+			{
+				return (A.X >= B.X) && (A.Y >= B.Y) && (A.Z >= B.Z);
+			}
 		};
 		
 		template<>
 		struct MetadataTraits<FVector4> : VectorTraits<FVector4>
 		{
+			static bool Equal(const FVector4& A, const FVector4& B)
+			{
+				return A.Equals(B);
+			}
+
 			static FVector4 Min(const FVector4& A, const FVector4& B)
 			{
 				return FVector4(FMath::Min(A.X, B.X), FMath::Min(A.Y, B.Y), FMath::Min(A.Z, B.Z), FMath::Min(A.W, B.W));
@@ -390,6 +545,26 @@ namespace PCG
 			{
 				return FVector4(FMath::Max(A.X, B.X), FMath::Max(A.Y, B.Y), FMath::Max(A.Z, B.Z), FMath::Max(A.W, B.W));
 			}
+
+			static bool Less(const FVector4& A, const FVector4& B)
+			{
+				return (A.X < B.X) && (A.Y < B.Y) && (A.Z < B.Z) && (A.W < B.W);
+			}
+
+			static bool Greater(const FVector4& A, const FVector4& B)
+			{
+				return (A.X > B.X) && (A.Y > B.Y) && (A.Z > B.Z) && (A.W > B.W);
+			}
+
+			static bool LessOrEqual(const FVector4& A, const FVector4& B)
+			{
+				return (A.X <= B.X) && (A.Y <= B.Y) && (A.Z <= B.Z) && (A.W <= B.W);
+			}
+
+			static bool GreaterOrEqual(const FVector4& A, const FVector4& B)
+			{
+				return (A.X >= B.X) && (A.Y >= B.Y) && (A.Z >= B.Z) && (A.W >= B.W);
+			}
 		};
 
 		// Quaternion
@@ -401,10 +576,11 @@ namespace PCG
 			enum { CanSubAdd = true };
 			enum { CanMulDiv = true };
 			enum { CanInterpolate = true };
+			enum { CanCompare = false };
 
 			static bool Equal(const FQuat& A, const FQuat& B)
 			{
-				return A == B;
+				return A.Equals(B);
 			}
 
 			static FQuat Add(const FQuat& A, const FQuat& B)
@@ -453,10 +629,11 @@ namespace PCG
 			enum { CanSubAdd = true };
 			enum { CanMulDiv = true };
 			enum { CanInterpolate = true };
+			enum { CanCompare = false };
 
 			static bool Equal(const FRotator& A, const FRotator& B)
 			{
-				return A == B;
+				return A.Equals(B);
 			}
 
 			static FRotator Add(const FRotator& A, const FRotator& B)
@@ -500,12 +677,11 @@ namespace PCG
 			enum { CanSubAdd = false };
 			enum { CanMulDiv = true };
 			enum { CanInterpolate = true };
+			enum { CanCompare = false };
 
 			static bool Equal(const FTransform& A, const FTransform& B)
 			{
-				return A.GetLocation() == B.GetLocation() &&
-					A.GetRotation() == B.GetRotation() &&
-					A.GetScale3D() == B.GetScale3D();
+				return A.Equals(B);
 			}
 
 			static FTransform Mul(const FTransform& A, const FTransform& B)
@@ -545,7 +721,7 @@ namespace PCG
 
 		// Strings
 		template<>
-		struct MetadataTraits<FString>
+		struct MetadataTraits<FString> : DefaultCompareTraits<FString>
 		{
 			enum { CompressData = true };
 			enum { CanMinMax = false };
@@ -557,6 +733,11 @@ namespace PCG
 			{
 				return A == B;
 			}
+
+			static FString ZeroValue()
+			{
+				return FString();
+			}
 		};
 
 		template<>
@@ -567,11 +748,37 @@ namespace PCG
 			enum { CanSubAdd = false };
 			enum { CanMulDiv = false };
 			enum { CanInterpolate = false };
+			enum { CanCompare = true };
 
 			static bool Equal(const FName& A, const FName& B)
 			{
 				return A == B;
 			}
+
+			static bool Less(const FName& A, const FName& B)
+			{
+				return A.Compare(B) < 0;
+			}
+
+			static bool Greater(const FName& A, const FName& B)
+			{
+				return A.Compare(B) > 0;
+			}
+
+			static bool LessOrEqual(const FName& A, const FName& B)
+			{
+				return A.Compare(B) <= 0;
+			}
+
+			static bool GreaterOrEqual(const FName& A, const FName& B)
+			{
+				return A.Compare(B) >= 0;
+			}
+
+			static FName ZeroValue()
+			{
+				return NAME_None;
+			}
 		};
 	}
 }
\ No newline at end of file
diff --git a/Engine/Plugins/Experimental/PCG/Source/PCG/Public/Metadata/PCGMetadataTypesConstantStruct.h b/Engine/Plugins/Experimental/PCG/Source/PCG/Public/Metadata/PCGMetadataTypesConstantStruct.h
new file mode 100644
--- /dev/null
+++ b/Engine/Plugins/Experimental/PCG/Source/PCG/Public/Metadata/PCGMetadataTypesConstantStruct.h
@@ -0,0 +1,114 @@
+// Copyright Epic Games, Inc. All Rights Reserved.
+
+#pragma once
+
+#include "Metadata/PCGMetadataAttributeTraits.h"
+#include "Helpers/PCGSettingsHelpers.h"
+
+#include "PCGMetadataTypesConstantStruct.generated.h"
+
+class UPCGParamData;
+
+/**
+* Struct to be re-used when you need to show constants types for a metadata type
+* It will store all our values, and will display nicely depending on the type chosen
+*/
+USTRUCT(BlueprintType)
+struct PCG_API FPCGMetadataTypesConstantStruct
+{
+	GENERATED_BODY()
+public:
+	template <typename Func>
+	decltype(auto) DispatcherWithOverride(const UPCGParamData* Params, Func Callback) const;
+
+	FString ToString() const;
+
+	UPROPERTY(BlueprintReadWrite, EditAnywhere, Category = Settings, meta = (ValidEnumValues = "Float, Double, Integer32, Integer64, Vector2, Vector, Vector4, Quaternion, Transform, String, Boolean, Rotator, Name"))
+	EPCGMetadataTypes Type = EPCGMetadataTypes::Double;
+
+	// All different types
+	UPROPERTY(BlueprintReadWrite, EditAnywhere, Category = Settings, meta = (EditCondition = "Type == EPCGMetadataTypes::Float", EditConditionHides))
+	float FloatValue = 0.0f;
+
+	UPROPERTY(BlueprintReadWrite, EditAnywhere, Category = Settings, meta = (EditCondition = "Type == EPCGMetadataTypes::Integer32", EditConditionHides))
+	int32 Int32Value = 0;
+
+	UPROPERTY(BlueprintReadWrite, EditAnywhere, Category = Settings, meta = (EditCondition = "Type == EPCGMetadataTypes::Double", EditConditionHides))
+	double DoubleValue = 0.0;
+
+	UPROPERTY(BlueprintReadWrite, EditAnywhere, Category = Settings, meta = (EditCondition = "Type == EPCGMetadataTypes::Integer64", EditConditionHides))
+	int64 IntValue = 0;
+
+	UPROPERTY(BlueprintReadWrite, EditAnywhere, Category = Settings, meta = (EditCondition = "Type == EPCGMetadataTypes::Vector2", EditConditionHides))
+	FVector2D Vector2Value = FVector2D::ZeroVector;
+
+	UPROPERTY(BlueprintReadWrite, EditAnywhere, Category = Settings, meta = (EditCondition = "Type == EPCGMetadataTypes::Vector", EditConditionHides))
+	FVector VectorValue = FVector::ZeroVector;
+
+	UPROPERTY(BlueprintReadWrite, EditAnywhere, Category = Settings, meta = (EditCondition = "Type == EPCGMetadataTypes::Vector4", EditConditionHides))
+	FVector4 Vector4Value = FVector4::Zero();
+
+	UPROPERTY(BlueprintReadWrite, EditAnywhere, Category = Settings, meta = (EditCondition = "Type == EPCGMetadataTypes::Quaternion", EditConditionHides))
+	FQuat QuatValue = FQuat::Identity;
+
+	UPROPERTY(BlueprintReadWrite, EditAnywhere, Category = Settings, meta = (EditCondition = "Type == EPCGMetadataTypes::Transform", EditConditionHides))
+	FTransform TransformValue = FTransform::Identity;
+
+	UPROPERTY(BlueprintReadWrite, EditAnywhere, Category = Settings, meta = (EditCondition = "Type == EPCGMetadataTypes::String", EditConditionHides))
+	FString StringValue = "";
+
+	UPROPERTY(BlueprintReadWrite, EditAnywhere, Category = Settings, meta = (EditCondition = "Type == EPCGMetadataTypes::Boolean", EditConditionHides))
+	bool BoolValue = false;
+
+	UPROPERTY(BlueprintReadWrite, EditAnywhere, Category = Settings, meta = (EditCondition = "Type == EPCGMetadataTypes::Rotator", EditConditionHides))
+	FRotator RotatorValue = FRotator::ZeroRotator;
+
+	UPROPERTY(BlueprintReadWrite, EditAnywhere, Category = Settings, meta = (EditCondition = "Type == EPCGMetadataTypes::Name", EditConditionHides))
+	FName NameValue = NAME_None;
+};
+
+template <typename Func>
+decltype(auto) FPCGMetadataTypesConstantStruct::DispatcherWithOverride(const UPCGParamData* Params, Func Callback) const
+{
+	using ReturnType = decltype(Callback(double{}));
+
+	switch (Type)
+	{
+	case EPCGMetadataTypes::Integer64:
+		return Callback(PCG_GET_OVERRIDEN_VALUE(this, IntValue, Params));
+	case EPCGMetadataTypes::Integer32:
+		return Callback(PCG_GET_OVERRIDEN_VALUE(this, Int32Value, Params));
+	case EPCGMetadataTypes::Float:
+		return Callback(PCG_GET_OVERRIDEN_VALUE(this, FloatValue, Params));
+	case EPCGMetadataTypes::Double:
+		return Callback(PCG_GET_OVERRIDEN_VALUE(this, DoubleValue, Params));
+	case EPCGMetadataTypes::Vector2:
+		return Callback(PCG_GET_OVERRIDEN_VALUE(this, Vector2Value, Params));
+	case EPCGMetadataTypes::Vector:
+		return Callback(PCG_GET_OVERRIDEN_VALUE(this, VectorValue, Params));
+	case EPCGMetadataTypes::Vector4:
+		return Callback(PCG_GET_OVERRIDEN_VALUE(this, Vector4Value, Params));
+	case EPCGMetadataTypes::Quaternion:
+		return Callback(PCG_GET_OVERRIDEN_VALUE(this, QuatValue, Params));
+	case EPCGMetadataTypes::Transform:
+		return Callback(PCG_GET_OVERRIDEN_VALUE(this, TransformValue, Params));
+	case EPCGMetadataTypes::String:
+		return Callback(PCG_GET_OVERRIDEN_VALUE(this, StringValue, Params));
+	case EPCGMetadataTypes::Boolean:
+		return Callback(PCG_GET_OVERRIDEN_VALUE(this, BoolValue, Params));
+	case EPCGMetadataTypes::Rotator:
+		return Callback(PCG_GET_OVERRIDEN_VALUE(this, RotatorValue, Params));
+	case EPCGMetadataTypes::Name:
+		return Callback(PCG_GET_OVERRIDEN_VALUE(this, NameValue, Params));
+	default:
+		// ReturnType{} is invalid if ReturnType is void
+		if constexpr (std::is_same_v<ReturnType, void>)
+		{
+			return;
+		}
+		else
+		{
+			return ReturnType{};
+		}
+	}
+}
\ No newline at end of file
