[PCG] Fixing unit tests

#rnx

diff --git a/Engine/Plugins/Experimental/PCG/Source/PCG/Private/Tests/Elements/Metadata/PCGMetadataBooleanOpTest.cpp b/Engine/Plugins/Experimental/PCG/Source/PCG/Private/Tests/Elements/Metadata/PCGMetadataBooleanOpTest.cpp
--- a/Engine/Plugins/Experimental/PCG/Source/PCG/Private/Tests/Elements/Metadata/PCGMetadataBooleanOpTest.cpp
+++ b/Engine/Plugins/Experimental/PCG/Source/PCG/Private/Tests/Elements/Metadata/PCGMetadataBooleanOpTest.cpp
@@ -88,7 +88,7 @@ bool FPCGMetadataBooleanOpTest::RunTest(const FString& Parameters)
 			return false;
 		}
 
-		const FPCGMetadataAttributeBase* OutAttributeBase = OutMetadata->GetConstAttribute(Settings->OutputTarget.AttributeName);
+		const FPCGMetadataAttributeBase* OutAttributeBase = OutMetadata->GetConstAttribute(Settings->OutputTarget.GetName());
 		if (!TestNotNull(TEXT("Output attribute exists"), OutAttributeBase))
 		{
 			return false;
@@ -99,6 +99,11 @@ bool FPCGMetadataBooleanOpTest::RunTest(const FString& Parameters)
 			return false;
 		}
 
+		if (!TestEqual(TEXT("Output attribute has an entry"), OutAttributeBase->GetNumberOfEntries(), 1))
+		{
+			return false;
+		}
+
 		const FPCGMetadataAttribute<bool>* OutAttribute = static_cast<const FPCGMetadataAttribute<bool>*>(OutAttributeBase);
 
 		bool Out = OutAttribute->GetValueFromItemKey(PCGInvalidEntryKey);
@@ -144,45 +149,45 @@ bool FPCGMetadataBooleanOpTest::RunTest(const FString& Parameters)
 	AddExpectedError(TEXT("Attribute/Property Invalid is not a supported type for input"), EAutomationExpectedErrorFlags::Contains, 3);
 
 	{
-		Settings->InputSource1.AttributeName = InvalidAttribute;
-		Settings->InputSource2.AttributeName = FalseAttribute;
+		Settings->InputSource1.SetAttributeName(InvalidAttribute);
+		Settings->InputSource2.SetAttributeName(FalseAttribute);
 		bTestPassed &= ValidateOp(/*bExpectedResult=*/false, /*bIsValid=*/ false);
 	}
 
 	{
-		Settings->InputSource1.AttributeName = FalseAttribute;
-		Settings->InputSource2.AttributeName = InvalidAttribute;
+		Settings->InputSource1.SetAttributeName(FalseAttribute);
+		Settings->InputSource2.SetAttributeName(InvalidAttribute);
 		bTestPassed &= ValidateOp(/*bExpectedResult=*/false, /*bIsValid=*/ false);
 	}
 
 	{
-		Settings->InputSource1.AttributeName = InvalidAttribute;
-		Settings->InputSource2.AttributeName = InvalidAttribute;
+		Settings->InputSource1.SetAttributeName(InvalidAttribute);
+		Settings->InputSource2.SetAttributeName(InvalidAttribute);
 		bTestPassed &= ValidateOp(/*bExpectedResult=*/false, /*bIsValid=*/ false);
 	}
 
 	// Valid tests
 	{
-		Settings->InputSource1.AttributeName = TrueAttribute;
-		Settings->InputSource2.AttributeName = TrueAttribute;
+		Settings->InputSource1.SetAttributeName(TrueAttribute);
+		Settings->InputSource2.SetAttributeName(TrueAttribute);
 		bTestPassed &= ValidateOp(/*bExpectedResult=*/true);
 	}
 
 	{
-		Settings->InputSource1.AttributeName = FalseAttribute;
-		Settings->InputSource2.AttributeName = TrueAttribute;
+		Settings->InputSource1.SetAttributeName(FalseAttribute);
+		Settings->InputSource2.SetAttributeName(TrueAttribute);
 		bTestPassed &= ValidateOp(/*bExpectedResult=*/false);
 	}
 
 	{
-		Settings->InputSource1.AttributeName = TrueAttribute;
-		Settings->InputSource2.AttributeName = FalseAttribute;
+		Settings->InputSource1.SetAttributeName(TrueAttribute);
+		Settings->InputSource2.SetAttributeName(FalseAttribute);
 		bTestPassed &= ValidateOp(/*bExpectedResult=*/false);
 	}
 
 	{
-		Settings->InputSource1.AttributeName = FalseAttribute;
-		Settings->InputSource2.AttributeName = FalseAttribute;
+		Settings->InputSource1.SetAttributeName(FalseAttribute);
+		Settings->InputSource2.SetAttributeName(FalseAttribute);
 		bTestPassed &= ValidateOp(/*bExpectedResult=*/false);
 	}
 
@@ -192,26 +197,26 @@ bool FPCGMetadataBooleanOpTest::RunTest(const FString& Parameters)
 	/////////////////////////////////////////
 	Settings->Operation = EPCGMedadataBooleanOperation::Or;
 	{
-		Settings->InputSource1.AttributeName = TrueAttribute;
-		Settings->InputSource2.AttributeName = TrueAttribute;
+		Settings->InputSource1.SetAttributeName(TrueAttribute);
+		Settings->InputSource2.SetAttributeName(TrueAttribute);
 		bTestPassed &= ValidateOp(/*bExpectedResult=*/true);
 	}
 
 	{
-		Settings->InputSource1.AttributeName = FalseAttribute;
-		Settings->InputSource2.AttributeName = TrueAttribute;
+		Settings->InputSource1.SetAttributeName(FalseAttribute);
+		Settings->InputSource2.SetAttributeName(TrueAttribute);
 		bTestPassed &= ValidateOp(/*bExpectedResult=*/true);
 	}
 
 	{
-		Settings->InputSource1.AttributeName = TrueAttribute;
-		Settings->InputSource2.AttributeName = FalseAttribute;
+		Settings->InputSource1.SetAttributeName(TrueAttribute);
+		Settings->InputSource2.SetAttributeName(FalseAttribute);
 		bTestPassed &= ValidateOp(/*bExpectedResult=*/true);
 	}
 
 	{
-		Settings->InputSource1.AttributeName = FalseAttribute;
-		Settings->InputSource2.AttributeName = FalseAttribute;
+		Settings->InputSource1.SetAttributeName(FalseAttribute);
+		Settings->InputSource2.SetAttributeName(FalseAttribute);
 		bTestPassed &= ValidateOp(/*bExpectedResult=*/false);
 	}
 
@@ -220,26 +225,26 @@ bool FPCGMetadataBooleanOpTest::RunTest(const FString& Parameters)
 	/////////////////////////////////////////
 	Settings->Operation = EPCGMedadataBooleanOperation::Xor;
 	{
-		Settings->InputSource1.AttributeName = TrueAttribute;
-		Settings->InputSource2.AttributeName = TrueAttribute;
+		Settings->InputSource1.SetAttributeName(TrueAttribute);
+		Settings->InputSource2.SetAttributeName(TrueAttribute);
 		bTestPassed &= ValidateOp(/*bExpectedResult=*/false);
 	}
 
 	{
-		Settings->InputSource1.AttributeName = FalseAttribute;
-		Settings->InputSource2.AttributeName = TrueAttribute;
+		Settings->InputSource1.SetAttributeName(FalseAttribute);
+		Settings->InputSource2.SetAttributeName(TrueAttribute);
 		bTestPassed &= ValidateOp(/*bExpectedResult=*/true);
 	}
 
 	{
-		Settings->InputSource1.AttributeName = TrueAttribute;
-		Settings->InputSource2.AttributeName = FalseAttribute;
+		Settings->InputSource1.SetAttributeName(TrueAttribute);
+		Settings->InputSource2.SetAttributeName(FalseAttribute);
 		bTestPassed &= ValidateOp(/*bExpectedResult=*/true);
 	}
 
 	{
-		Settings->InputSource1.AttributeName = FalseAttribute;
-		Settings->InputSource2.AttributeName = FalseAttribute;
+		Settings->InputSource1.SetAttributeName(FalseAttribute);
+		Settings->InputSource2.SetAttributeName(FalseAttribute);
 		bTestPassed &= ValidateOp(/*bExpectedResult=*/false);
 	}
 
diff --git a/Engine/Plugins/Experimental/PCG/Source/PCG/Private/Tests/Elements/Metadata/PCGMetadataBreakTransformTest.cpp b/Engine/Plugins/Experimental/PCG/Source/PCG/Private/Tests/Elements/Metadata/PCGMetadataBreakTransformTest.cpp
--- a/Engine/Plugins/Experimental/PCG/Source/PCG/Private/Tests/Elements/Metadata/PCGMetadataBreakTransformTest.cpp
+++ b/Engine/Plugins/Experimental/PCG/Source/PCG/Private/Tests/Elements/Metadata/PCGMetadataBreakTransformTest.cpp
@@ -12,6 +12,9 @@
 #include "Metadata/PCGMetadata.h"
 #include "Metadata/PCGMetadataAttribute.h"
 #include "Metadata/PCGMetadataAttributeTpl.h"
+#include "Metadata/Accessors/PCGAttributeAccessorKeys.h"
+#include "Metadata/Accessors/IPCGAttributeAccessor.h"
+#include "Metadata/Accessors/PCGAttributeAccessorHelpers.h"
 
 #if WITH_EDITOR
 
@@ -88,31 +91,6 @@ namespace PCGBreakTransformTest
 			Attribute->SetValue(EntryKey, RandomSource.FRand());
 		}
 	}
-
-	TArray<PCGMetadataEntryKey> GatherEntryKeys(const UPCGData* InData, const FPCGMetadataAttributeBase* InAttribute)
-	{
-		TArray<PCGMetadataEntryKey> Keys;
-		if (const UPCGPointData* PointOutput = Cast<const UPCGPointData>(InData))
-		{
-			for (const FPCGPoint& Point : PointOutput->GetPoints())
-			{
-				Keys.Add(Point.MetadataEntry);
-			}
-		}
-		else
-		{
-			const FPCGMetadataAttributeBase* Current = InAttribute;
-			while (Current)
-			{
-				TArray<PCGMetadataEntryKey> Temp;
-				Current->GetEntryToValueKeyMap_NotThreadSafe().GenerateKeyArray(Temp);
-				Keys.Append(Temp);
-				Current = Current->GetParent();
-			}
-		}
-
-		return Keys;
-	}
 }
 
 
@@ -150,28 +128,15 @@ bool FPCGMetadataBreakTransformTest::RunTest(const FString& Parameters)
 
 		const FPCGTaggedData& Input = Inputs[0];
 
-		UPCGMetadata* SourceMetadata = nullptr; 
-		if (const UPCGSpatialData* SpatialInput = Cast<const UPCGSpatialData>(Input.Data))
-		{
-			SourceMetadata = SpatialInput->Metadata;
-		}
-		else if (const UPCGParamData* ParamsInput = Cast<const UPCGParamData>(Input.Data))
-		{
-			SourceMetadata = ParamsInput->Metadata;
-		}
-		else
+		TUniquePtr<const IPCGAttributeAccessor> InputAccessor = PCGAttributeAccessorHelpers::CreateConstAccessor(Input.Data, Settings->InputSource);
+		TUniquePtr<const IPCGAttributeAccessorKeys> InputKeys = PCGAttributeAccessorHelpers::CreateConstKeys(Input.Data, Settings->InputSource);
+
+		if (!TestTrue("InputSource was found", InputAccessor.IsValid() && InputKeys.IsValid()))
 		{
 			return false;
 		}
 
-		check(SourceMetadata);
-
-		const FPCGMetadataAttributeBase* SourceAttributeBase = SourceMetadata->GetConstAttribute(Settings->InputSource.GetName());
-		check(SourceAttributeBase);
-
-		TArray<PCGMetadataEntryKey> InputKeys = PCGBreakTransformTest::GatherEntryKeys(Input.Data, SourceAttributeBase);
-
-		auto ValidateComponentOutput = [&](const FPCGTaggedData& Output, const FName& OutAttributeName, PCGBreakTransformTest::EPCGComponentToCheck ComponentToCheck, auto DummyValue)
+		auto ValidateComponentOutput = [this, &bTestPassed, &InputKeys, &InputAccessor](const FPCGTaggedData& Output, const FPCGAttributePropertySelector& OutSelector, PCGBreakTransformTest::EPCGComponentToCheck ComponentToCheck, auto DummyValue)
 		{
 			using OutType = decltype(DummyValue);
 
@@ -181,59 +146,42 @@ bool FPCGMetadataBreakTransformTest::RunTest(const FString& Parameters)
 				return;
 			}
 
-			UPCGMetadata* OutMetadata = nullptr;
-			if (const UPCGSpatialData* SpatialOutput = Cast<const UPCGSpatialData>(Output.Data))
-			{
-				OutMetadata = SpatialOutput->Metadata;
-			}
-			else if (const UPCGParamData* ParamsOutput = Cast<const UPCGParamData>(Output.Data))
-			{
-				OutMetadata = ParamsOutput->Metadata;
-			}
-			else
-			{
-				return;
-			}
-
-			if (!TestNotNull("Valid output metadata", OutMetadata))
-			{
-				bTestPassed = false;
-				return;
-			}
-
-			const FPCGMetadataAttributeBase* OutAttributeBase = OutMetadata->GetConstAttribute(OutAttributeName);
+			TUniquePtr<const IPCGAttributeAccessor> OutputAccessor = PCGAttributeAccessorHelpers::CreateConstAccessor(Output.Data, OutSelector);
+			TUniquePtr<const IPCGAttributeAccessorKeys> OutputKeys = PCGAttributeAccessorHelpers::CreateConstKeys(Output.Data, OutSelector);
 
-			if (!TestNotNull("Valid output attribute", OutAttributeBase))
+			if (!TestTrue("OutputTarget was found", OutputAccessor.IsValid() && OutputKeys.IsValid()))
 			{
 				bTestPassed = false;
 				return;
 			}
 
-			if (!TestEqual("Output attribute is a valid type", OutAttributeBase->GetTypeId(), PCG::Private::MetadataTypes<OutType>::Id))
+			if (!TestEqual("Output has a valid type", OutputAccessor->GetUnderlyingType(), PCG::Private::MetadataTypes<OutType>::Id))
 			{
 				bTestPassed = false;
 				return;
 			}
 
-			const FPCGMetadataAttribute<OutType>* OutAttribute = static_cast<const FPCGMetadataAttribute<OutType>*>(OutAttributeBase);
-			check(OutAttribute);
-
-			TArray<PCGMetadataEntryKey> OutputKeys = PCGBreakTransformTest::GatherEntryKeys(Output.Data, OutAttribute);
-
-			if (!TestEqual("Identical EntryKeys count", InputKeys.Num(), OutputKeys.Num()))
+			if (!TestEqual("Identical EntryKeys count", InputKeys->GetNum(), OutputKeys->GetNum()))
 			{
 				bTestPassed = false;
 				return;
 			}
 
-			for (int32 i = 0; i < OutputKeys.Num(); ++i)
+			for (int32 i = 0; i < OutputKeys->GetNum(); ++i)
 			{
-				const PCGMetadataValueKey InValueKey = SourceAttributeBase->GetValueKey(InputKeys[i]);
-				const PCGMetadataValueKey OutValueKey = OutAttribute->GetValueKey(OutputKeys[i]);
-
-				FTransform SourceValue = static_cast<const FPCGMetadataAttribute<FTransform>*>(SourceAttributeBase)->GetValue(InValueKey);
+				FTransform SourceValue{};
+				if (!InputAccessor->Get<FTransform>(SourceValue, i, *InputKeys))
+				{
+					bTestPassed = false;
+					return;
+				}
 
-				const OutType OutValue = OutAttribute->GetValue(OutValueKey);
+				OutType OutValue{};
+				if (!OutputAccessor->Get<OutType>(OutValue, i, *OutputKeys))
+				{
+					bTestPassed = false;
+					return;
+				}
 
 				if constexpr (std::is_same_v<FVector, OutType>)
 				{
@@ -268,13 +216,9 @@ bool FPCGMetadataBreakTransformTest::RunTest(const FString& Parameters)
 		const TArray<FPCGTaggedData> OutputsRotation = Context->OutputData.GetInputsByPin(PCGMetadataTransformConstants::Rotation);
 		const TArray<FPCGTaggedData> OutputsScale = Context->OutputData.GetInputsByPin(PCGMetadataTransformConstants::Scale);
 
-		const FName DestinationAttributeForTranslation = Settings->OutputTarget.AttributeName; //Settings->GetOutputAttributeName(Settings->OutputName.AttributeName, 0);
-		const FName DestinationAttributeForRotation = Settings->OutputTarget.AttributeName; //Settings->GetOutputAttributeName(Settings->OutputName.AttributeName, 1);
-		const FName DestinationAttributeForScale = Settings->OutputTarget.AttributeName; //Settings->GetOutputAttributeName(Settings->OutputName.AttributeName, 2);
-
 		if (TestEqual("Appropriate number of outputs generated for Translation", OutputsTranslation.Num(), 1))
 		{
-			ValidateComponentOutput(OutputsTranslation[0], DestinationAttributeForTranslation, PCGBreakTransformTest::EPCGComponentToCheck::Translation, FVector{});
+			ValidateComponentOutput(OutputsTranslation[0], Settings->OutputTarget, PCGBreakTransformTest::EPCGComponentToCheck::Translation, FVector{});
 		}
 		else
 		{
@@ -283,7 +227,7 @@ bool FPCGMetadataBreakTransformTest::RunTest(const FString& Parameters)
 
 		if (TestEqual("Appropriate number of outputs generated for Rotation", OutputsRotation.Num(), 1))
 		{
-			ValidateComponentOutput(OutputsRotation[0], DestinationAttributeForRotation, PCGBreakTransformTest::EPCGComponentToCheck::Rotation, FQuat{});
+			ValidateComponentOutput(OutputsRotation[0], Settings->OutputTarget, PCGBreakTransformTest::EPCGComponentToCheck::Rotation, FQuat{});
 		}
 		else
 		{
@@ -292,7 +236,7 @@ bool FPCGMetadataBreakTransformTest::RunTest(const FString& Parameters)
 
 		if (TestEqual("Appropriate number of outputs generated for Scale", OutputsScale.Num(), 1))
 		{
-			ValidateComponentOutput(OutputsScale[0], DestinationAttributeForScale, PCGBreakTransformTest::EPCGComponentToCheck::Scale, FVector{});
+			ValidateComponentOutput(OutputsScale[0], Settings->OutputTarget, PCGBreakTransformTest::EPCGComponentToCheck::Scale, FVector{});
 		}
 		else
 		{
@@ -332,17 +276,15 @@ bool FPCGMetadataBreakTransformTest::RunTest(const FString& Parameters)
 		Settings->ForceOutputConnections[1] = true;
 		Settings->ForceOutputConnections[2] = true;
 
-		Settings->InputSource.Selection = EPCGAttributePropertySelection::Attribute;
-
 		{
 			AddInfo(TEXT("Test with Transform as input attribute"));
-			Settings->InputSource.AttributeName = PCGBreakTransformTest::TransformAttribute;
+			Settings->InputSource.SetAttributeName(PCGBreakTransformTest::TransformAttribute);
 			bTestPassed &= ValidateMetadataBreakVector(*TestData);
 		}
 
 		{
 			AddInfo(TEXT("Test with Float as input attribute (invalid type)"));
-			Settings->InputSource.AttributeName = PCGBreakTransformTest::InvalidAttribute;
+			Settings->InputSource.SetAttributeName(PCGBreakTransformTest::InvalidAttribute);
 
 			bTestPassed &= ValidateMetadataBreakVector(*TestData, /*bIsValid=*/false);
 		}
diff --git a/Engine/Plugins/Experimental/PCG/Source/PCG/Private/Tests/Elements/Metadata/PCGMetadataBreakVectorTest.cpp b/Engine/Plugins/Experimental/PCG/Source/PCG/Private/Tests/Elements/Metadata/PCGMetadataBreakVectorTest.cpp
--- a/Engine/Plugins/Experimental/PCG/Source/PCG/Private/Tests/Elements/Metadata/PCGMetadataBreakVectorTest.cpp
+++ b/Engine/Plugins/Experimental/PCG/Source/PCG/Private/Tests/Elements/Metadata/PCGMetadataBreakVectorTest.cpp
@@ -12,6 +12,9 @@
 #include "Metadata/PCGMetadata.h"
 #include "Metadata/PCGMetadataAttribute.h"
 #include "Metadata/PCGMetadataAttributeTpl.h"
+#include "Metadata/Accessors/PCGAttributeAccessorKeys.h"
+#include "Metadata/Accessors/IPCGAttributeAccessor.h"
+#include "Metadata/Accessors/PCGAttributeAccessorHelpers.h"
 
 #if WITH_EDITOR
 
@@ -144,26 +147,15 @@ bool FPCGMetadataBreakVectorTest::RunTest(const FString& Parameters)
 
 		const FPCGTaggedData& Input = Inputs[0];
 
-		UPCGMetadata* SourceMetadata = nullptr; 
-		if (const UPCGSpatialData* SpatialInput = Cast<const UPCGSpatialData>(Input.Data))
-		{
-			SourceMetadata = SpatialInput->Metadata;
-		}
-		else if (const UPCGParamData* ParamsInput = Cast<const UPCGParamData>(Input.Data))
-		{
-			SourceMetadata = ParamsInput->Metadata;
-		}
-		else
+		TUniquePtr<const IPCGAttributeAccessor> InputAccessor = PCGAttributeAccessorHelpers::CreateConstAccessor(Input.Data, Settings->InputSource);
+		TUniquePtr<const IPCGAttributeAccessorKeys> InputKeys = PCGAttributeAccessorHelpers::CreateConstKeys(Input.Data, Settings->InputSource);
+
+		if (!TestTrue("InputSource was found", InputAccessor.IsValid() && InputKeys.IsValid()))
 		{
 			return false;
 		}
 
-		check(SourceMetadata);
-
-		const FPCGMetadataAttributeBase* SourceAttributeBase = SourceMetadata->GetConstAttribute(Settings->InputSource.GetName());
-		check(SourceAttributeBase);
-
-		auto ValidateComponentOutput = [&](const FPCGTaggedData& Output, const FName& OutAttributeName, PCGBreakVectorTest::EPCGComponentToCheck ComponentToCheck)
+		auto ValidateComponentOutput = [this, &bTestPassed, &InputKeys, &InputAccessor](const FPCGTaggedData& Output, const FPCGAttributePropertySelector& OutSelector, PCGBreakVectorTest::EPCGComponentToCheck ComponentToCheck)
 		{
 			if (!TestTrue("Valid output data", Output.Data != nullptr))
 			{
@@ -171,91 +163,80 @@ bool FPCGMetadataBreakVectorTest::RunTest(const FString& Parameters)
 				return;
 			}
 
-			UPCGMetadata* OutMetadata = nullptr;
-			if (const UPCGSpatialData* SpatialInput = Cast<const UPCGSpatialData>(Output.Data))
-			{
-				OutMetadata = SpatialInput->Metadata;
-			}
-			else if (const UPCGParamData* ParamsInput = Cast<const UPCGParamData>(Output.Data))
-			{
-				OutMetadata = ParamsInput->Metadata;
-			}
-			else
-			{
-				return;
-			}
+			TUniquePtr<const IPCGAttributeAccessor> OutputAccessor = PCGAttributeAccessorHelpers::CreateConstAccessor(Output.Data, OutSelector);
+			TUniquePtr<const IPCGAttributeAccessorKeys> OutputKeys = PCGAttributeAccessorHelpers::CreateConstKeys(Output.Data, OutSelector);
 
-			if (!TestNotNull("Valid output metadata", OutMetadata))
+			if (!TestTrue("OutputTarget was found", OutputAccessor.IsValid() && OutputKeys.IsValid()))
 			{
 				bTestPassed = false;
 				return;
 			}
 
-			const FPCGMetadataAttributeBase* OutAttributeBase = OutMetadata->GetConstAttribute(OutAttributeName);
-
-			if (!TestNotNull("Valid output attribute", OutAttributeBase))
+			if (!TestEqual("Output has a valid type (double)", OutputAccessor->GetUnderlyingType(), PCG::Private::MetadataTypes<double>::Id))
 			{
 				bTestPassed = false;
 				return;
 			}
 
-			if (!TestEqual("Output attribute is a valid type (double)", OutAttributeBase->GetTypeId(), PCG::Private::MetadataTypes<double>::Id))
+			if (!TestEqual("Identical EntryKeys count", InputKeys->GetNum(), OutputKeys->GetNum()))
 			{
 				bTestPassed = false;
 				return;
 			}
 
-			const FPCGMetadataAttribute<double>* OutAttribute = static_cast<const FPCGMetadataAttribute<double>*>(OutAttributeBase);
-			check(OutAttribute);
-
-			const PCGMetadataEntryKey InEntryKeyCount = SourceMetadata->GetItemCountForChild();
-			const PCGMetadataEntryKey OutEntryKeyCount = OutMetadata->GetItemCountForChild();
-
-			if (!TestEqual("Identical EntryKey counts", InEntryKeyCount, OutEntryKeyCount))
+			for (int32 i = 0; i < OutputKeys->GetNum(); ++i)
 			{
-				bTestPassed = false;
-				return;
-			}
-
-			const PCGMetadataValueKey InValueKeyCount = SourceAttributeBase->GetValueKeyOffsetForChild();
-			const PCGMetadataValueKey OutValueKeyCount = OutAttribute->GetValueKeyOffsetForChild();
-
-			if (!TestEqual("Identical ValueKey counts", InValueKeyCount, OutValueKeyCount))
-			{
-				bTestPassed = false;
-				return;
-			}
-
-			for (PCGMetadataEntryKey EntryKey = 0; EntryKey < InEntryKeyCount; ++EntryKey)
-			{
-				const PCGMetadataValueKey InValueKey = SourceAttributeBase->GetValueKey(EntryKey);
-				const PCGMetadataValueKey OutValueKey = OutAttribute->GetValueKey(EntryKey);
-
-				bTestPassed &= TestEqual("Identical value keys", InValueKey, OutValueKey);
-
-				FVector4 SourceValue;
-				if (SourceAttributeBase->GetTypeId() == PCG::Private::MetadataTypes<FVector>::Id)
+				FVector4 SourceValue{};
+				if (InputAccessor->GetUnderlyingType() == PCG::Private::MetadataTypes<FVector>::Id)
 				{
-					SourceValue = static_cast<const FPCGMetadataAttribute<FVector>*>(SourceAttributeBase)->GetValue(InValueKey);
+					FVector TempValue{};
+					if (!InputAccessor->Get<FVector>(TempValue, i, *InputKeys))
+					{
+						bTestPassed = false;
+						return;
+					}
+
+					SourceValue = FVector4(TempValue);
 				}
-				else if (SourceAttributeBase->GetTypeId() == PCG::Private::MetadataTypes<FVector2D>::Id)
+				else if (InputAccessor->GetUnderlyingType() == PCG::Private::MetadataTypes<FVector2D>::Id)
 				{
-					const FVector2D Vec2 = static_cast<const FPCGMetadataAttribute<FVector2D>*>(SourceAttributeBase)->GetValue(InValueKey);
-					SourceValue = FVector4(Vec2.X, Vec2.Y);
+					FVector2D TempValue{};
+					if (!InputAccessor->Get<FVector2D>(TempValue, i, *InputKeys))
+					{
+						bTestPassed = false;
+						return;
+					}
+
+					SourceValue = FVector4(TempValue.X, TempValue.Y);
 				}
-				else if (SourceAttributeBase->GetTypeId() == PCG::Private::MetadataTypes<FVector4>::Id)
+				else if (InputAccessor->GetUnderlyingType() == PCG::Private::MetadataTypes<FVector4>::Id)
 				{
-					SourceValue = static_cast<const FPCGMetadataAttribute<FVector4>*>(SourceAttributeBase)->GetValue(InValueKey);
+					if (!InputAccessor->Get<FVector4>(SourceValue, i, *InputKeys))
+					{
+						bTestPassed = false;
+						return;
+					}
 				}
-				else //if (SourceAttributeBase->GetTypeId() == PCG::Private::MetadataTypes<FRotator>::Id)
+				else //if (InputAccessor->GetUnderlyingType() == PCG::Private::MetadataTypes<FRotator>::Id)
 				{
-					const FRotator Rotator = static_cast<const FPCGMetadataAttribute<FRotator>*>(SourceAttributeBase)->GetValue(InValueKey);
-					SourceValue.X = Rotator.Roll;
-					SourceValue.Y = Rotator.Pitch;
-					SourceValue.Z = Rotator.Yaw;
+					FRotator TempValue{};
+					if (!InputAccessor->Get<FRotator>(TempValue, i, *InputKeys))
+					{
+						bTestPassed = false;
+						return;
+					}
+
+					SourceValue.X = TempValue.Roll;
+					SourceValue.Y = TempValue.Pitch;
+					SourceValue.Z = TempValue.Yaw;
 				}
 
-				const double OutValue = OutAttribute->GetValue(OutValueKey);
+				double OutValue{};
+				if (!OutputAccessor->Get<double>(OutValue, i, *OutputKeys))
+				{
+					bTestPassed = false;
+					return;
+				}
 
 				if (ComponentToCheck == PCGBreakVectorTest::EPCGComponentToCheck::X)
 				{
@@ -281,16 +262,11 @@ bool FPCGMetadataBreakVectorTest::RunTest(const FString& Parameters)
 		const TArray<FPCGTaggedData> OutputsZ = Context->OutputData.GetInputsByPin(PCGMetadataBreakVectorConstants::ZLabel);
 		const TArray<FPCGTaggedData> OutputsW = Context->OutputData.GetInputsByPin(PCGMetadataBreakVectorConstants::WLabel);
 
-		const FName DestinationAttributeForX = Settings->OutputTarget.AttributeName; //Settings->GetOutputAttributeName(Settings->OutputName.AttributeName, 0);
-		const FName DestinationAttributeForY = Settings->OutputTarget.AttributeName; //Settings->GetOutputAttributeName(Settings->OutputName.AttributeName, 1);
-		const FName DestinationAttributeForZ = Settings->OutputTarget.AttributeName; //Settings->GetOutputAttributeName(Settings->OutputName.AttributeName, 2);
-		const FName DestinationAttributeForW = Settings->OutputTarget.AttributeName; //Settings->GetOutputAttributeName(Settings->OutputName.AttributeName, 3);
-
 		if (Settings->ForceOutputConnections[0])
 		{
 			if (TestEqual("Appropriate number of outputs generated fox X", OutputsX.Num(), 1))
 			{
-				ValidateComponentOutput(OutputsX[0], DestinationAttributeForX, PCGBreakVectorTest::EPCGComponentToCheck::X);
+				ValidateComponentOutput(OutputsX[0], Settings->OutputTarget, PCGBreakVectorTest::EPCGComponentToCheck::X);
 			}
 			else
 			{
@@ -302,7 +278,7 @@ bool FPCGMetadataBreakVectorTest::RunTest(const FString& Parameters)
 		{
 			if (TestEqual("Appropriate number of outputs generated for Y", OutputsY.Num(), 1))
 			{
-				ValidateComponentOutput(OutputsY[0], DestinationAttributeForY, PCGBreakVectorTest::EPCGComponentToCheck::Y);
+				ValidateComponentOutput(OutputsY[0], Settings->OutputTarget, PCGBreakVectorTest::EPCGComponentToCheck::Y);
 			}
 			else
 			{
@@ -314,7 +290,7 @@ bool FPCGMetadataBreakVectorTest::RunTest(const FString& Parameters)
 		{
 			if (TestEqual("Appropriate number of outputs generated for Z", OutputsZ.Num(), 1))
 			{
-				ValidateComponentOutput(OutputsZ[0], DestinationAttributeForZ, PCGBreakVectorTest::EPCGComponentToCheck::Z);
+				ValidateComponentOutput(OutputsZ[0], Settings->OutputTarget, PCGBreakVectorTest::EPCGComponentToCheck::Z);
 			}
 			else
 			{
@@ -322,11 +298,11 @@ bool FPCGMetadataBreakVectorTest::RunTest(const FString& Parameters)
 			}
 		}
 
-		if (Settings->ForceOutputConnections[3] && SourceAttributeBase->GetTypeId() == PCG::Private::MetadataTypes<FVector4>::Id)
+		if (Settings->ForceOutputConnections[3])
 		{
 			if (TestEqual("Appropriate number of outputs generated for W", OutputsW.Num(), 1))
 			{
-				ValidateComponentOutput(OutputsW[0], DestinationAttributeForW, PCGBreakVectorTest::EPCGComponentToCheck::W);
+				ValidateComponentOutput(OutputsW[0], Settings->OutputTarget, PCGBreakVectorTest::EPCGComponentToCheck::W);
 			}
 			else
 			{
@@ -368,11 +344,9 @@ bool FPCGMetadataBreakVectorTest::RunTest(const FString& Parameters)
 		Settings->ForceOutputConnections[2] = false;
 		Settings->ForceOutputConnections[3] = false;
 
-		Settings->InputSource.Selection = EPCGAttributePropertySelection::Attribute;
-
 		{
 			AddInfo(TEXT("Test with FVector2D as input attribute"));
-			Settings->InputSource.AttributeName = PCGBreakVectorTest::Vec2Attribute;
+			Settings->InputSource.SetAttributeName(PCGBreakVectorTest::Vec2Attribute);
 			bTestPassed &= ValidateMetadataBreakVector(*TestData);
 		}
 
@@ -380,13 +354,13 @@ bool FPCGMetadataBreakVectorTest::RunTest(const FString& Parameters)
 
 		{
 			AddInfo(TEXT("Test with FVector as input attribute"));
-			Settings->InputSource.AttributeName = PCGBreakVectorTest::Vec3Attribute;
+			Settings->InputSource.SetAttributeName(PCGBreakVectorTest::Vec3Attribute);
 			bTestPassed &= ValidateMetadataBreakVector(*TestData);
 		}
 
 		{
 			AddInfo(TEXT("Test with FVector4 as input attribute"));
-			Settings->InputSource.AttributeName = PCGBreakVectorTest::Vec4Attribute;
+			Settings->InputSource.SetAttributeName(PCGBreakVectorTest::Vec4Attribute);
 			Settings->ForceOutputConnections[3] = true;
 			bTestPassed &= ValidateMetadataBreakVector(*TestData);
 			Settings->ForceOutputConnections[3] = false;
@@ -394,13 +368,13 @@ bool FPCGMetadataBreakVectorTest::RunTest(const FString& Parameters)
 
 		{
 			AddInfo(TEXT("Test with FRotator as input attribute"));
-			Settings->InputSource.AttributeName = PCGBreakVectorTest::RotatorAttribute;
+			Settings->InputSource.SetAttributeName(PCGBreakVectorTest::RotatorAttribute);
 			bTestPassed &= ValidateMetadataBreakVector(*TestData);
 		}
 
 		{
 			AddInfo(TEXT("Test with Float as input attribute (invalid type)"));
-			Settings->InputSource.AttributeName = PCGBreakVectorTest::FloatAttribute;
+			Settings->InputSource.SetAttributeName(PCGBreakVectorTest::FloatAttribute);
 
 			bTestPassed &= ValidateMetadataBreakVector(*TestData, /*bIsValid=*/false);
 		}
diff --git a/Engine/Plugins/Experimental/PCG/Source/PCG/Private/Tests/Elements/Metadata/PCGMetadataMakeTransformTest.cpp b/Engine/Plugins/Experimental/PCG/Source/PCG/Private/Tests/Elements/Metadata/PCGMetadataMakeTransformTest.cpp
--- a/Engine/Plugins/Experimental/PCG/Source/PCG/Private/Tests/Elements/Metadata/PCGMetadataMakeTransformTest.cpp
+++ b/Engine/Plugins/Experimental/PCG/Source/PCG/Private/Tests/Elements/Metadata/PCGMetadataMakeTransformTest.cpp
@@ -12,6 +12,9 @@
 #include "Metadata/PCGMetadata.h"
 #include "Metadata/PCGMetadataAttribute.h"
 #include "Metadata/PCGMetadataAttributeTpl.h"
+#include "Metadata/Accessors/PCGAttributeAccessorKeys.h"
+#include "Metadata/Accessors/IPCGAttributeAccessor.h"
+#include "Metadata/Accessors/PCGAttributeAccessorHelpers.h"
 
 #if WITH_EDITOR
 
@@ -115,105 +118,86 @@ bool FPCGMetadataMakeTransformTest::RunTest(const FString& Parameters)
 
 		const FPCGTaggedData& Input = TranslationInputs[0];
 
-		UPCGMetadata* SourceMetadata = nullptr; 
-		if (const UPCGSpatialData* SpatialInput = Cast<const UPCGSpatialData>(Input.Data))
-		{
-			SourceMetadata = SpatialInput->Metadata;
-		}
-		else if (const UPCGParamData* ParamsInput = Cast<const UPCGParamData>(Input.Data))
-		{
-			SourceMetadata = ParamsInput->Metadata;
-		}
-		else
-		{
-			return false;
-		}
-
-		check(SourceMetadata);
+		TUniquePtr<const IPCGAttributeAccessor> InputAccessor1 = PCGAttributeAccessorHelpers::CreateConstAccessor(Input.Data, Settings->InputSource1);
+		TUniquePtr<const IPCGAttributeAccessorKeys> InputKeys1 = PCGAttributeAccessorHelpers::CreateConstKeys(Input.Data, Settings->InputSource1);
 
-		const FPCGMetadataAttributeBase* TranslationAttributeBase = SourceMetadata->GetConstAttribute(Settings->InputSource1.GetName());
-		check(TranslationAttributeBase);
+		TUniquePtr<const IPCGAttributeAccessor> InputAccessor2 = PCGAttributeAccessorHelpers::CreateConstAccessor(Input.Data, Settings->InputSource2);
+		TUniquePtr<const IPCGAttributeAccessorKeys> InputKeys2 = PCGAttributeAccessorHelpers::CreateConstKeys(Input.Data, Settings->InputSource2);
 
-		const FPCGMetadataAttributeBase* RotationAttributeBase = SourceMetadata->GetConstAttribute(Settings->InputSource2.GetName());
-		check(RotationAttributeBase);
+		TUniquePtr<const IPCGAttributeAccessor> InputAccessor3 = PCGAttributeAccessorHelpers::CreateConstAccessor(Input.Data, Settings->InputSource3);
+		TUniquePtr<const IPCGAttributeAccessorKeys> InputKeys3 = PCGAttributeAccessorHelpers::CreateConstKeys(Input.Data, Settings->InputSource3);
 
-		const FPCGMetadataAttributeBase* ScaleAttributeBase = SourceMetadata->GetConstAttribute(Settings->InputSource3.GetName());
-		check(ScaleAttributeBase);
-
-		const TArray<FPCGTaggedData> Outputs = Context->OutputData.GetInputsByPin(PCGPinConstants::DefaultOutputLabel);
-		if (!TestEqual("Right number of outputs", Outputs.Num(), 1))
+		if (!TestTrue("InputSource1 was found", InputAccessor1.IsValid() && InputKeys1.IsValid()))
 		{
 			return false;
 		}
 
-		const FPCGTaggedData& Output = Outputs[0];
-
-		if (!TestTrue("Valid output data", Output.Data != nullptr))
+		if (!TestTrue("InputSource2 was found", InputAccessor2.IsValid() && InputKeys2.IsValid()))
 		{
 			return false;
 		}
 
-		UPCGMetadata* OutMetadata = nullptr;
-		if (const UPCGSpatialData* SpatialInput = Cast<const UPCGSpatialData>(Output.Data))
+		if (!TestTrue("InputSource3 was found", InputAccessor3.IsValid() && InputKeys3.IsValid()))
 		{
-			OutMetadata = SpatialInput->Metadata;
+			return false;
 		}
-		else if (const UPCGParamData* ParamsInput = Cast<const UPCGParamData>(Output.Data))
+
+		if (!TestEqual("Same number of input keys between source 1 and 2", InputKeys1->GetNum(), InputKeys2->GetNum()))
 		{
-			OutMetadata = ParamsInput->Metadata;
+			return false;
 		}
-		else
+
+		if (!TestEqual("Same number of input keys between source 1 and 3", InputKeys1->GetNum(), InputKeys3->GetNum()))
 		{
 			return false;
 		}
 
-		if (!TestNotNull("Valid output metadata", OutMetadata))
+		const TArray<FPCGTaggedData> Outputs = Context->OutputData.GetInputsByPin(PCGPinConstants::DefaultOutputLabel);
+		if (!TestEqual("Right number of outputs", Outputs.Num(), 1))
 		{
 			return false;
 		}
 
-		const FPCGMetadataAttributeBase* OutAttributeBase = OutMetadata->GetConstAttribute(Settings->OutputTarget.GetName());
+		const FPCGTaggedData& Output = Outputs[0];
 
-		if (!TestNotNull("Valid output attribute", OutAttributeBase))
+		if (!TestTrue("Valid output data", Output.Data != nullptr))
 		{
 			return false;
 		}
 
-		const FPCGMetadataAttribute<FTransform>* OutAttribute = static_cast<const FPCGMetadataAttribute<FTransform>*>(OutAttributeBase);
-		check(OutAttribute);
+		TUniquePtr<const IPCGAttributeAccessor> OutputAccessor = PCGAttributeAccessorHelpers::CreateConstAccessor(Output.Data, Settings->OutputTarget);
+		TUniquePtr<const IPCGAttributeAccessorKeys> OutputKeys = PCGAttributeAccessorHelpers::CreateConstKeys(Output.Data, Settings->OutputTarget);
 
-		const PCGMetadataEntryKey InEntryKeyCount = SourceMetadata->GetItemCountForChild();
-		const PCGMetadataEntryKey OutEntryKeyCount = OutMetadata->GetItemCountForChild();
-
-		if (!TestEqual("Identical EntryKey counts", InEntryKeyCount, OutEntryKeyCount))
+		if (!TestTrue("OutputTarget was found", OutputAccessor.IsValid() && OutputKeys.IsValid()))
 		{
 			return false;
 		}
 
-		const PCGMetadataValueKey InValueKeyCount = TranslationAttributeBase->GetValueKeyOffsetForChild();
-		const PCGMetadataValueKey OutValueKeyCount = OutAttribute->GetValueKeyOffsetForChild();
+		if (!TestEqual("Output has a valid type (Transform)", OutputAccessor->GetUnderlyingType(), PCG::Private::MetadataTypes<FTransform>::Id))
+		{
+			return false;
+		}
 
-		if (!TestEqual("Identical ValueKey counts", InValueKeyCount, OutValueKeyCount))
+		if (!TestEqual("Identical EntryKeys count", InputKeys1->GetNum(), OutputKeys->GetNum()))
 		{
 			return false;
 		}
 
-		for (PCGMetadataEntryKey EntryKey = 0; EntryKey < InEntryKeyCount; ++EntryKey)
+		for (int32 i = 0; i < OutputKeys->GetNum(); ++i)
 		{
-			const PCGMetadataValueKey InValueKey = TranslationAttributeBase->GetValueKey(EntryKey);
-			const PCGMetadataValueKey OutValueKey = OutAttribute->GetValueKey(EntryKey);
+			FVector TranslationValue{};
+			FQuat QuatValue{};
+			FVector ScaleValue{};
+			FTransform OutValue{};
 
-			if (!TestEqual("Identical value keys", InValueKey, OutValueKey))
+			if (!InputAccessor1->Get<FVector>(TranslationValue, i, *InputKeys1)
+				|| !InputAccessor2->Get<FQuat>(QuatValue, i, *InputKeys2)
+				|| !InputAccessor3->Get<FVector>(ScaleValue, i, *InputKeys3)
+				|| !OutputAccessor->Get<FTransform>(OutValue, i, *OutputKeys))
 			{
 				return false;
 			}
-
-			FVector TranslationValue = static_cast<const FPCGMetadataAttribute<FVector>*>(TranslationAttributeBase)->GetValue(InValueKey);
-			FQuat QuatValue = static_cast<const FPCGMetadataAttribute<FQuat>*>(RotationAttributeBase)->GetValue(InValueKey);
-			FVector ScaleValue = static_cast<const FPCGMetadataAttribute<FVector>*>(ScaleAttributeBase)->GetValue(InValueKey);
 			
-			FTransform OutValue = OutAttribute->GetValue(OutValueKey);
-
 			if (!TestEqual("Mismatch between input and output for translation", OutValue.GetTranslation(), TranslationValue))
 			{
 				return false;
@@ -257,14 +241,10 @@ bool FPCGMetadataMakeTransformTest::RunTest(const FString& Parameters)
 		UPCGMetadataMakeTransformSettings* Settings = CastChecked<UPCGMetadataMakeTransformSettings>(TestData->Settings);
 
 		Settings->ForceOutputConnections[0] = true;
-		Settings->InputSource1.Selection = EPCGAttributePropertySelection::Attribute;
-		Settings->InputSource1.AttributeName = Vec3Attribute;
-
-		Settings->InputSource2.Selection = EPCGAttributePropertySelection::Attribute;
-		Settings->InputSource2.AttributeName = QuatAttribute;
 
-		Settings->InputSource3.Selection = EPCGAttributePropertySelection::Attribute;
-		Settings->InputSource3.AttributeName = Vec3Attribute;
+		Settings->InputSource1.SetAttributeName(Vec3Attribute);
+		Settings->InputSource2.SetAttributeName(QuatAttribute);
+		Settings->InputSource3.SetAttributeName(Vec3Attribute);
 
 		bTestPassed &= ValidateMetadataMakeVector(*TestData);
 	}
