[PCG] Cleanup pass on implicit lambda capture + fix issue with set attribute from property

* In SetAttributeFromPropertyHelper, if the entry key was invalid, it would have asserted. Fixed now

diff --git a/Engine/Plugins/Experimental/PCG/Source/PCG/Private/Elements/Metadata/PCGMetadataElement.cpp b/Engine/Plugins/Experimental/PCG/Source/PCG/Private/Elements/Metadata/PCGMetadataElement.cpp
--- a/Engine/Plugins/Experimental/PCG/Source/PCG/Private/Elements/Metadata/PCGMetadataElement.cpp
+++ b/Engine/Plugins/Experimental/PCG/Source/PCG/Private/Elements/Metadata/PCGMetadataElement.cpp
@@ -163,7 +163,7 @@ bool FPCGMetadataOperationElement::ExecuteInternal(FPCGContext* Context) const
 
 		if (Settings->OutputTarget.Selection == EPCGAttributePropertySelection::Attribute)
 		{
-			auto CreateAttribute = [&](auto Dummy)
+			auto CreateAttribute = [SampledData, DestinationAttribute](auto Dummy)
 			{
 				using AttributeType = decltype(Dummy);
 				return PCGMetadataElementCommon::ClearOrCreateAttribute(SampledData->Metadata, DestinationAttribute, AttributeType{}) != nullptr;
@@ -196,7 +196,7 @@ bool FPCGMetadataOperationElement::ExecuteInternal(FPCGContext* Context) const
 		}
 
 		// At this point, we are ready.
-		auto Operation = [&](auto Dummy)
+		auto Operation = [&InputAccessor, &InputKeys, &OutputAccessor, &OutputKeys](auto Dummy)
 		{
 			using OutputType = decltype(Dummy);
 			OutputType Value{};
diff --git a/Engine/Plugins/Experimental/PCG/Source/PCG/Private/Elements/Metadata/PCGMetadataOpElementBase.cpp b/Engine/Plugins/Experimental/PCG/Source/PCG/Private/Elements/Metadata/PCGMetadataOpElementBase.cpp
--- a/Engine/Plugins/Experimental/PCG/Source/PCG/Private/Elements/Metadata/PCGMetadataOpElementBase.cpp
+++ b/Engine/Plugins/Experimental/PCG/Source/PCG/Private/Elements/Metadata/PCGMetadataOpElementBase.cpp
@@ -244,12 +244,16 @@ bool FPCGMetadataElementBase::ExecuteInternal(FPCGContext* Context) const
 	{
 		OutputTarget.AttributeName = InputName;
 	}
+	
+	// TODO: Make it an option
+	const int32 InputToForward = 0;
 
+	// Use implicit capture, since we capture a lot
 	auto CreateAttribute = [&](uint32 OutputIndex, auto DummyOutValue) -> bool
 	{
 		using AttributeType = decltype(DummyOutValue);
 
-		FPCGTaggedData& OutputData = Outputs.Add_GetRef(InputTaggedData[0]);
+		FPCGTaggedData& OutputData = Outputs.Add_GetRef(InputTaggedData[InputToForward]);
 		OutputData.Pin = Settings->GetOutputPinLabel(OutputIndex);
 
 		UPCGMetadata* OutMetadata = nullptr;
@@ -259,7 +263,7 @@ bool FPCGMetadataElementBase::ExecuteInternal(FPCGContext* Context) const
 		if (OutputTarget.Selection == EPCGAttributePropertySelection::Attribute)
 		{
 			// In case of attribute, there is no point of failure before duplicating. So duplicate, create the attribute and then the accessor.
-			PCGMetadataElementCommon::DuplicateTaggedData(InputTaggedData[0], OutputData, OutMetadata);
+			PCGMetadataElementCommon::DuplicateTaggedData(InputTaggedData[InputToForward], OutputData, OutMetadata);
 			FPCGMetadataAttributeBase* OutputAttribute = PCGMetadataElementCommon::ClearOrCreateAttribute(OutMetadata, OutputName, AttributeType{});
 			if (!OutputAttribute)
 			{
@@ -267,9 +271,9 @@ bool FPCGMetadataElementBase::ExecuteInternal(FPCGContext* Context) const
 			}
 
 			// And copy the mapping from the original attribute, if it is not points
-			if (!InputTaggedData[0].Data->IsA<UPCGPointData>() && SourceMetadata[0] && SourceAttribute[0])
+			if (!InputTaggedData[InputToForward].Data->IsA<UPCGPointData>() && SourceMetadata[InputToForward] && SourceAttribute[InputToForward])
 			{
-				PCGMetadataElementCommon::CopyEntryToValueKeyMap(SourceMetadata[0], SourceAttribute[0], OutputAttribute);
+				PCGMetadataElementCommon::CopyEntryToValueKeyMap(SourceMetadata[InputToForward], SourceAttribute[InputToForward], OutputAttribute);
 			}
 
 			OperationData.OutputAccessors[OutputIndex] = PCGAttributeAccessorHelpers::CreateAccessor(Cast<UPCGData>(OutputData.Data), OutputTarget);
@@ -288,7 +292,7 @@ bool FPCGMetadataElementBase::ExecuteInternal(FPCGContext* Context) const
 					return false;
 				}
 
-				PCGMetadataElementCommon::DuplicateTaggedData(InputTaggedData[0], OutputData, OutMetadata);
+				PCGMetadataElementCommon::DuplicateTaggedData(InputTaggedData[InputToForward], OutputData, OutMetadata);
 			}
 		}
 
@@ -302,7 +306,7 @@ bool FPCGMetadataElementBase::ExecuteInternal(FPCGContext* Context) const
 		return OperationData.OutputKeys[OutputIndex].IsValid();
 	};
 
-	auto CreateAllSameAttributes = [&](auto DummyOutValue) -> bool
+	auto CreateAllSameAttributes = [NumberOfOutputs, &CreateAttribute](auto DummyOutValue) -> bool
 	{
 		for (uint32 i = 0; i < NumberOfOutputs; ++i)
 		{
@@ -331,7 +335,8 @@ bool FPCGMetadataElementBase::ExecuteInternal(FPCGContext* Context) const
 		for (uint32 i = 0; i < NumberOfOutputs && bCreateAttributeSucceeded; ++i)
 		{
 			bCreateAttributeSucceeded &= PCGMetadataAttribute::CallbackWithRightType(OutputTypes[i],
-				[&](auto DummyOutValue) -> bool {
+				[&CreateAttribute, i](auto DummyOutValue) -> bool 
+				{
 					return CreateAttribute(i, DummyOutValue);
 				});
 		}
diff --git a/Engine/Plugins/Experimental/PCG/Source/PCG/Private/Elements/PCGCopyPoints.cpp b/Engine/Plugins/Experimental/PCG/Source/PCG/Private/Elements/PCGCopyPoints.cpp
--- a/Engine/Plugins/Experimental/PCG/Source/PCG/Private/Elements/PCGCopyPoints.cpp
+++ b/Engine/Plugins/Experimental/PCG/Source/PCG/Private/Elements/PCGCopyPoints.cpp
@@ -127,6 +127,7 @@ bool FPCGCopyPointsElement::ExecuteInternal(FPCGContext* Context) const
 	TArray<TTuple<int64, int64>> AllMetadataEntries;
 	AllMetadataEntries.SetNum(SourcePoints.Num() * TargetPoints.Num());
 
+	// Use implicit capture, since we capture a lot
 	FPCGAsync::AsyncPointProcessing(Context, SourcePoints.Num() * TargetPoints.Num(), OutPoints, [&](int32 Index, FPCGPoint& OutPoint)
 	{
 		const FPCGPoint& SourcePoint = SourcePoints[Index / TargetPoints.Num()];
diff --git a/Engine/Plugins/Experimental/PCG/Source/PCG/Private/Elements/PCGDensityNoise.cpp b/Engine/Plugins/Experimental/PCG/Source/PCG/Private/Elements/PCGDensityNoise.cpp
--- a/Engine/Plugins/Experimental/PCG/Source/PCG/Private/Elements/PCGDensityNoise.cpp
+++ b/Engine/Plugins/Experimental/PCG/Source/PCG/Private/Elements/PCGDensityNoise.cpp
@@ -38,7 +38,7 @@ bool FPCGDensityNoiseElement::ExecuteInternal(FPCGContext* Context) const
 	// Precompute a seed based on the settings one and the component one
 	const int Seed = PCGSettingsHelpers::ComputeSeedWithOverride(Settings, Context->SourceComponent, Params);
 
-	ProcessPoints(Context, Inputs, Outputs, [&](const FPCGPoint& InPoint, FPCGPoint& OutPoint)
+	ProcessPoints(Context, Inputs, Outputs, [Seed, DensityNoiseMin, DensityNoiseMax, bInvertSourceDensity, DensityMode](const FPCGPoint& InPoint, FPCGPoint& OutPoint)
 	{
 		OutPoint = InPoint;
 		FRandomStream RandomSource(PCGHelpers::ComputeSeed(Seed, OutPoint.Seed));
diff --git a/Engine/Plugins/Experimental/PCG/Source/PCG/Private/Elements/PCGDensityRemapElement.cpp b/Engine/Plugins/Experimental/PCG/Source/PCG/Private/Elements/PCGDensityRemapElement.cpp
--- a/Engine/Plugins/Experimental/PCG/Source/PCG/Private/Elements/PCGDensityRemapElement.cpp
+++ b/Engine/Plugins/Experimental/PCG/Source/PCG/Private/Elements/PCGDensityRemapElement.cpp
@@ -53,7 +53,7 @@ bool FPCGDensityRemapElement::ExecuteInternal(FPCGContext* Context) const
 		Intercept = OutRangeMin;
 	}
 
-	ProcessPoints(Context, Inputs, Outputs, [&](const FPCGPoint& InPoint, FPCGPoint& OutPoint)
+	ProcessPoints(Context, Inputs, Outputs, [bExcludeValuesOutsideInputRange, InRangeTrueMin, InRangeTrueMax, Slope, InRangeMin, Intercept](const FPCGPoint& InPoint, FPCGPoint& OutPoint)
 	{
 		OutPoint = InPoint;
 		const float SourceDensity = InPoint.Density;
diff --git a/Engine/Plugins/Experimental/PCG/Source/PCG/Private/Elements/PCGSplineSampler.cpp b/Engine/Plugins/Experimental/PCG/Source/PCG/Private/Elements/PCGSplineSampler.cpp
--- a/Engine/Plugins/Experimental/PCG/Source/PCG/Private/Elements/PCGSplineSampler.cpp
+++ b/Engine/Plugins/Experimental/PCG/Source/PCG/Private/Elements/PCGSplineSampler.cpp
@@ -132,7 +132,7 @@ namespace PCGSplineSamplerHelpers
 		}
 
 		// Checks if a duplicate point should be counted as a valid intersection and adds it to the intersection list if so. Collinear points are also considered duplicates.
-		auto TestDuplicatePoint = [&](const FVector2D& IntersectionPoint, int32 Index) -> bool
+		auto TestDuplicatePoint = [PointCount, &PolygonPoints](const FVector2D& IntersectionPoint, int32 Index) -> bool
 		{
 			const int32 BackwardIndex = Index - 1;
 			const int32 ForwardIndex = Index + 1;
diff --git a/Engine/Plugins/Experimental/PCG/Source/PCG/Private/Elements/PCGTransformPoints.cpp b/Engine/Plugins/Experimental/PCG/Source/PCG/Private/Elements/PCGTransformPoints.cpp
--- a/Engine/Plugins/Experimental/PCG/Source/PCG/Private/Elements/PCGTransformPoints.cpp
+++ b/Engine/Plugins/Experimental/PCG/Source/PCG/Private/Elements/PCGTransformPoints.cpp
@@ -51,6 +51,7 @@ bool FPCGTransformPointsElement::ExecuteInternal(FPCGContext* Context) const
 
 	const int Seed = PCGSettingsHelpers::ComputeSeedWithOverride(Settings, Context->SourceComponent, Params);
 
+	// Use implicit capture, since we capture a lot
 	ProcessPoints(Context, Inputs, Outputs, [&](const FPCGPoint& InPoint, FPCGPoint& OutPoint)
 	{
 		OutPoint = InPoint;
diff --git a/Engine/Plugins/Experimental/PCG/Source/PCG/Private/Metadata/PCGMetadata.cpp b/Engine/Plugins/Experimental/PCG/Source/PCG/Private/Metadata/PCGMetadata.cpp
--- a/Engine/Plugins/Experimental/PCG/Source/PCG/Private/Metadata/PCGMetadata.cpp
+++ b/Engine/Plugins/Experimental/PCG/Source/PCG/Private/Metadata/PCGMetadata.cpp
@@ -437,29 +437,29 @@ void UPCGMetadata::CreateBoolAttribute(FName AttributeName, bool DefaultValue, b
 namespace PCGMetadata
 {
 	template<typename DataType>
-	bool SetAttributeFromPropertyHelper(UPCGMetadata& Metadata, FName AttributeName, PCGMetadataEntryKey& EntryKey, const DataType* DataPtr, const FProperty* InProperty, bool bCreate)
+	bool SetAttributeFromPropertyHelper(UPCGMetadata* Metadata, FName AttributeName, PCGMetadataEntryKey& EntryKey, const DataType* DataPtr, const FProperty* InProperty, bool bCreate)
 	{
-		if (!InProperty || !DataPtr)
+		if (!InProperty || !DataPtr || !Metadata)
 		{
 			return false;
 		}
 
 		// Check if an attribute already exists or not if we ask to create a new one
-		if (!bCreate && !Metadata.HasAttribute(AttributeName))
+		if (!bCreate && !Metadata->HasAttribute(AttributeName))
 		{
 			return false;
 		}
 
-		auto CreateAttributeAndSet = [AttributeName, &Metadata, bCreate, &EntryKey](auto&& PropertyValue) -> bool
+		auto CreateAttributeAndSet = [AttributeName, Metadata, bCreate, &EntryKey](auto&& PropertyValue) -> bool
 		{
 			using PropertyType = std::remove_const_t<std::remove_reference_t<decltype(PropertyValue)>>;
 
-			FPCGMetadataAttributeBase* BaseAttribute = Metadata.GetMutableAttribute(AttributeName);
+			FPCGMetadataAttributeBase* BaseAttribute = Metadata->GetMutableAttribute(AttributeName);
 
 			if (!BaseAttribute && bCreate)
 			{
 				// Interpolation is disabled and no parent override.
-				BaseAttribute = Metadata.CreateAttribute<PropertyType>(AttributeName, PropertyValue, false, false);
+				BaseAttribute = Metadata->CreateAttribute<PropertyType>(AttributeName, PropertyValue, false, false);
 			}
 
 			if (!BaseAttribute)
@@ -468,18 +468,20 @@ namespace PCGMetadata
 			}
 
 			// Allow to set the value if both type matches or if we can construct AttributeType from PropertyType.
-			return PCGMetadataAttribute::CallbackWithRightType(BaseAttribute->GetTypeId(), [&](auto AttributeValue) -> bool
+			return PCGMetadataAttribute::CallbackWithRightType(BaseAttribute->GetTypeId(), [&EntryKey, &PropertyValue, BaseAttribute, Metadata](auto AttributeValue) -> bool
 				{
 					using AttributeType = decltype(AttributeValue);
 					FPCGMetadataAttribute<AttributeType>* Attribute = static_cast<FPCGMetadataAttribute<AttributeType>*>(BaseAttribute);
 
 					if constexpr (std::is_same_v<AttributeType, PropertyType>)
 					{
+						Metadata->InitializeOnSet(EntryKey);
 						Attribute->SetValue(EntryKey, PropertyValue);
 						return true;
 					}
 					else if constexpr (std::is_constructible_v<AttributeType, PropertyType>)
 					{
+						Metadata->InitializeOnSet(EntryKey);
 						Attribute->SetValue(EntryKey, AttributeType(PropertyValue));
 						return true;
 					}
@@ -496,12 +498,12 @@ namespace PCGMetadata
 
 bool UPCGMetadata::SetAttributeFromProperty(FName AttributeName, PCGMetadataEntryKey& EntryKey, const UObject* Object, const FProperty* InProperty, bool bCreate)
 {
-	return PCGMetadata::SetAttributeFromPropertyHelper<UObject>(*this, AttributeName, EntryKey, Object, InProperty, bCreate);
+	return PCGMetadata::SetAttributeFromPropertyHelper<UObject>(this, AttributeName, EntryKey, Object, InProperty, bCreate);
 }
 
 bool UPCGMetadata::SetAttributeFromDataProperty(FName AttributeName, PCGMetadataEntryKey& EntryKey, const void* Data, const FProperty* InProperty, bool bCreate)
 {
-	return PCGMetadata::SetAttributeFromPropertyHelper<void>(*this, AttributeName, EntryKey, Data, InProperty, bCreate);
+	return PCGMetadata::SetAttributeFromPropertyHelper<void>(this, AttributeName, EntryKey, Data, InProperty, bCreate);
 }
 
 bool UPCGMetadata::CopyExistingAttribute(FName AttributeToCopy, FName NewAttributeName, bool bKeepParent)
diff --git a/Engine/Plugins/Experimental/PCG/Source/PCG/Private/PCGHelpers.cpp b/Engine/Plugins/Experimental/PCG/Source/PCG/Private/PCGHelpers.cpp
--- a/Engine/Plugins/Experimental/PCG/Source/PCG/Private/PCGHelpers.cpp
+++ b/Engine/Plugins/Experimental/PCG/Source/PCG/Private/PCGHelpers.cpp
@@ -104,7 +104,7 @@ namespace PCGHelpers
 
 		if (InActor)
 		{
-			InActor->ForEachComponent<UPrimitiveComponent>(bIncludeFromChildActors, [&](const UPrimitiveComponent* InPrimComp)
+			InActor->ForEachComponent<UPrimitiveComponent>(bIncludeFromChildActors, [bNonColliding, bIgnorePCGCreatedComponents, &Box](const UPrimitiveComponent* InPrimComp)
 			{
 				// Note: we omit the IsRegistered check here (e.g. InPrimComp->IsRegistered() )
 				// since this can be called in a scope where the components are temporarily unregistered
@@ -137,7 +137,7 @@ namespace PCGHelpers
 			const FTransform& ActorToWorld = InActor->GetTransform();
 			const FTransform WorldToActor = ActorToWorld.Inverse();
 
-			InActor->ForEachComponent<UPrimitiveComponent>(bIncludeFromChildActors, [&](const UPrimitiveComponent* InPrimComp)
+			InActor->ForEachComponent<UPrimitiveComponent>(bIncludeFromChildActors, [bNonColliding, bIgnorePCGCreatedComponents, &WorldToActor, &Box](const UPrimitiveComponent* InPrimComp)
 			{
 				if ((bNonColliding || InPrimComp->IsCollisionEnabled()) &&
 					(!bIgnorePCGCreatedComponents || !InPrimComp->ComponentTags.Contains(DefaultPCGTag)))
diff --git a/Engine/Plugins/Experimental/PCG/Source/PCGEditor/Private/SPCGEditorGraphFind.cpp b/Engine/Plugins/Experimental/PCG/Source/PCGEditor/Private/SPCGEditorGraphFind.cpp
--- a/Engine/Plugins/Experimental/PCG/Source/PCGEditor/Private/SPCGEditorGraphFind.cpp
+++ b/Engine/Plugins/Experimental/PCG/Source/PCGEditor/Private/SPCGEditorGraphFind.cpp
@@ -280,7 +280,7 @@ void SPCGEditorGraphFind::MatchTokens(const TArray<FString>& InTokens)
 		NodeSearchString = NodeSearchString.Replace(TEXT(" "), TEXT(""));
 
 		FPCGEditorGraphFindResultPtr NodeResult;
-		auto GetOrCreateNodeResult = [&]() -> FPCGEditorGraphFindResultPtr&
+		auto GetOrCreateNodeResult = [this, &NodeResult, &NodeName, &NodeType, Node]() -> FPCGEditorGraphFindResultPtr&
 		{
 			if (!NodeResult.IsValid())
 			{
