Removed usage of XY map in the landscape info to prevent in-game issues

diff --git a/Engine/Plugins/Experimental/PCG/Source/PCG/Private/Data/PCGLandscapeData.cpp b/Engine/Plugins/Experimental/PCG/Source/PCG/Private/Data/PCGLandscapeData.cpp
--- a/Engine/Plugins/Experimental/PCG/Source/PCG/Private/Data/PCGLandscapeData.cpp
+++ b/Engine/Plugins/Experimental/PCG/Source/PCG/Private/Data/PCGLandscapeData.cpp
@@ -133,48 +133,46 @@ bool UPCGLandscapeData::ProjectPoint(const FTransform& InTransform, const FBox&
 	const FVector LocalPoint = LandscapeTransform.InverseTransformPosition(InTransform.GetLocation());
 	const FIntPoint ComponentMapKey(FMath::FloorToInt(LocalPoint.X / LandscapeInfo->ComponentSizeQuads), FMath::FloorToInt(LocalPoint.Y / LandscapeInfo->ComponentSizeQuads));
 
-	if (ULandscapeComponent* LandscapeComponent = LandscapeInfo->XYtoComponentMap.FindRef(ComponentMapKey))
+#if WITH_EDITOR
+	ULandscapeComponent* LandscapeComponent = LandscapeInfo->XYtoComponentMap.FindRef(ComponentMapKey);
+	const FPCGLandscapeCacheEntry * LandscapeCacheEntry = (LandscapeComponent ? LandscapeCache->GetCacheEntry(LandscapeComponent, ComponentMapKey) : nullptr);
+#else
+	const FPCGLandscapeCacheEntry* LandscapeCacheEntry = LandscapeCache->GetCacheEntry(LandscapeInfo->LandscapeGuid, ComponentMapKey);
+#endif
+
+	if (!LandscapeCacheEntry)
 	{
-		const FPCGLandscapeCacheEntry* LandscapeCacheEntry = LandscapeCache->GetCacheEntry(LandscapeComponent, ComponentMapKey);
-
-		if (!LandscapeCacheEntry)
-		{
-			return false;
-		}
+		return false;
+	}
 
-		const FVector2D ComponentLocalPoint(LocalPoint.X - ComponentMapKey.X * LandscapeInfo->ComponentSizeQuads, LocalPoint.Y - ComponentMapKey.Y * LandscapeInfo->ComponentSizeQuads);
+	const FVector2D ComponentLocalPoint(LocalPoint.X - ComponentMapKey.X * LandscapeInfo->ComponentSizeQuads, LocalPoint.Y - ComponentMapKey.Y * LandscapeInfo->ComponentSizeQuads);
 
-		if (bHeightOnly)
-		{
-			LandscapeCacheEntry->GetInterpolatedPointHeightOnly(ComponentLocalPoint, OutPoint);
-		}
-		else
-		{
-			LandscapeCacheEntry->GetInterpolatedPoint(ComponentLocalPoint, OutPoint, bUseMetadata ? OutMetadata : nullptr);
-		}
+	if (bHeightOnly)
+	{
+		LandscapeCacheEntry->GetInterpolatedPointHeightOnly(ComponentLocalPoint, OutPoint);
+	}
+	else
+	{
+		LandscapeCacheEntry->GetInterpolatedPoint(ComponentLocalPoint, OutPoint, bUseMetadata ? OutMetadata : nullptr);
+	}
 
-		// Respect projection settings
-		if (!InParams.bProjectPositions)
-		{
-			OutPoint.Transform.SetLocation(InTransform.GetLocation());
-		}
+	// Respect projection settings
+	if (!InParams.bProjectPositions)
+	{
+		OutPoint.Transform.SetLocation(InTransform.GetLocation());
+	}
 		
-		if (!InParams.bProjectRotations)
-		{
-			OutPoint.Transform.SetRotation(InTransform.GetRotation());
-		}
-
-		if (!InParams.bProjectScales)
-		{
-			OutPoint.Transform.SetScale3D(InTransform.GetScale3D());
-		}
-
-		return true;
+	if (!InParams.bProjectRotations)
+	{
+		OutPoint.Transform.SetRotation(InTransform.GetRotation());
 	}
-	else
+
+	if (!InParams.bProjectScales)
 	{
-		return false;
+		OutPoint.Transform.SetScale3D(InTransform.GetScale3D());
 	}
+
+	return true;
 }
 
 const UPCGPointData* UPCGLandscapeData::CreatePointData(FPCGContext* Context, const FBox& InBounds) const
@@ -243,39 +241,41 @@ const UPCGPointData* UPCGLandscapeData::CreatePointData(FPCGContext* Context, co
 			for (int32 ComponentY = MinComponentY; ComponentY <= MaxComponentY; ++ComponentY)
 			{
 				FIntPoint ComponentMapKey(ComponentX, ComponentY);
-				if (ULandscapeComponent* LandscapeComponent = LandscapeInfo->XYtoComponentMap.FindRef(ComponentMapKey))
+#if WITH_EDITOR
+				ULandscapeComponent* LandscapeComponent = LandscapeInfo->XYtoComponentMap.FindRef(ComponentMapKey);
+				const FPCGLandscapeCacheEntry* LandscapeCacheEntry = LandscapeComponent ? LandscapeCache->GetCacheEntry(LandscapeComponent, ComponentMapKey) : nullptr;
+#else
+				const FPCGLandscapeCacheEntry* LandscapeCacheEntry = LandscapeCache->GetCacheEntry(LandscapeInfo->LandscapeGuid, ComponentMapKey);
+#endif
+
+				if (!LandscapeCacheEntry)
 				{
-					const FPCGLandscapeCacheEntry* LandscapeCacheEntry = LandscapeCache->GetCacheEntry(LandscapeComponent, ComponentMapKey);
-
-					if (!LandscapeCacheEntry)
-					{
-						continue;
-					}
+					continue;
+				}
 
-					// Rebase our bounds in the component referential
-					const int32 LocalMinX = FMath::Clamp(MinX - ComponentMapKey.X * ComponentSizeQuads, 0, ComponentSizeQuads - 1);
-					const int32 LocalMaxX = FMath::Clamp(MaxX - ComponentMapKey.X * ComponentSizeQuads, 0, ComponentSizeQuads - 1);
+				// Rebase our bounds in the component referential
+				const int32 LocalMinX = FMath::Clamp(MinX - ComponentMapKey.X * ComponentSizeQuads, 0, ComponentSizeQuads - 1);
+				const int32 LocalMaxX = FMath::Clamp(MaxX - ComponentMapKey.X * ComponentSizeQuads, 0, ComponentSizeQuads - 1);
 
-					const int32 LocalMinY = FMath::Clamp(MinY - ComponentMapKey.Y * ComponentSizeQuads, 0, ComponentSizeQuads - 1);
-					const int32 LocalMaxY = FMath::Clamp(MaxY - ComponentMapKey.Y * ComponentSizeQuads, 0, ComponentSizeQuads - 1);
+				const int32 LocalMinY = FMath::Clamp(MinY - ComponentMapKey.Y * ComponentSizeQuads, 0, ComponentSizeQuads - 1);
+				const int32 LocalMaxY = FMath::Clamp(MaxY - ComponentMapKey.Y * ComponentSizeQuads, 0, ComponentSizeQuads - 1);
 
-					// We can't really copy data from the component points wholesale because the component points have an additional boundary point.
-					// TODO: consider optimizing this, though it will impact the Sample then
-					for (int32 LocalX = LocalMinX; LocalX <= LocalMaxX; ++LocalX)
+				// We can't really copy data from the component points wholesale because the component points have an additional boundary point.
+				// TODO: consider optimizing this, though it will impact the Sample then
+				for (int32 LocalX = LocalMinX; LocalX <= LocalMaxX; ++LocalX)
+				{
+					for (int32 LocalY = LocalMinY; LocalY <= LocalMaxY; ++LocalY)
 					{
-						for (int32 LocalY = LocalMinY; LocalY <= LocalMaxY; ++LocalY)
+						const int32 PointIndex = LocalX + LocalY * (ComponentSizeQuads + 1);
+
+						FPCGPoint& Point = Points.Emplace_GetRef();
+						if (bHeightOnly)
+						{
+							LandscapeCacheEntry->GetPointHeightOnly(PointIndex, Point);
+						}
+						else
 						{
-							const int32 PointIndex = LocalX + LocalY * (ComponentSizeQuads + 1);
-
-							FPCGPoint& Point = Points.Emplace_GetRef();
-							if (bHeightOnly)
-							{
-								LandscapeCacheEntry->GetPointHeightOnly(PointIndex, Point);
-							}
-							else
-							{
-								LandscapeCacheEntry->GetPoint(PointIndex, Point, bUseMetadata ? Data->Metadata : nullptr);
-							}
+							LandscapeCacheEntry->GetPoint(PointIndex, Point, bUseMetadata ? Data->Metadata : nullptr);
 						}
 					}
 				}
diff --git a/Engine/Plugins/Experimental/PCG/Source/PCG/Private/Grid/PCGLandscapeCache.cpp b/Engine/Plugins/Experimental/PCG/Source/PCG/Private/Grid/PCGLandscapeCache.cpp
--- a/Engine/Plugins/Experimental/PCG/Source/PCG/Private/Grid/PCGLandscapeCache.cpp
+++ b/Engine/Plugins/Experimental/PCG/Source/PCG/Private/Grid/PCGLandscapeCache.cpp
@@ -578,25 +578,16 @@ void UPCGLandscapeCache::ClearCache()
 	CachedData.Reset();
 }
 
+#if WITH_EDITOR
 const FPCGLandscapeCacheEntry* UPCGLandscapeCache::GetCacheEntry(ULandscapeComponent* LandscapeComponent, const FIntPoint& ComponentCoordinate)
 {
-	const FPCGLandscapeCacheEntry* CacheEntry = nullptr;
-	TPair<FGuid, FIntPoint> ComponentKey(LandscapeComponent && LandscapeComponent->GetLandscapeProxy() ? LandscapeComponent->GetLandscapeProxy()->GetLandscapeGuid() : FGuid(), ComponentCoordinate);
+	const FGuid LandscapeGuid = (LandscapeComponent && LandscapeComponent->GetLandscapeProxy() ? LandscapeComponent->GetLandscapeProxy()->GetLandscapeGuid() : FGuid());
+	const FPCGLandscapeCacheEntry* CacheEntry = GetCacheEntry(LandscapeGuid, ComponentCoordinate);
 
-	{
-#if WITH_EDITOR
-		FReadScopeLock ScopeLock(CacheLock);
-#endif
-		if (FPCGLandscapeCacheEntry** FoundEntry = CachedData.Find(ComponentKey))
-		{
-			CacheEntry = *FoundEntry;
-		}
-	}
-
-#if WITH_EDITOR
 	if (!CacheEntry && LandscapeComponent && LandscapeComponent->GetLandscapeInfo())
 	{
 		FWriteScopeLock ScopeLock(CacheLock);
+		TPair<FGuid, FIntPoint> ComponentKey(LandscapeGuid, ComponentCoordinate);
 		if (FPCGLandscapeCacheEntry** FoundEntry = CachedData.Find(ComponentKey))
 		{
 			CacheEntry = *FoundEntry;
@@ -610,9 +601,35 @@ const FPCGLandscapeCacheEntry* UPCGLandscapeCache::GetCacheEntry(ULandscapeCompo
 			CacheEntry = NewEntry;
 			CachedData.Add(ComponentKey, NewEntry);
 		}
+
+		if (CacheEntry)
+		{
+			if (CacheEntry->TouchAndLoad(CacheTouch++))
+			{
+				CacheMemorySize += CacheEntry->GetMemorySize();
+			}
+		}
 	}
+
+	return CacheEntry;
+}
 #endif
 
+const FPCGLandscapeCacheEntry* UPCGLandscapeCache::GetCacheEntry(const FGuid& LandscapeGuid, const FIntPoint& ComponentCoordinate)
+{
+	const FPCGLandscapeCacheEntry* CacheEntry = nullptr;
+	TPair<FGuid, FIntPoint> ComponentKey(LandscapeGuid, ComponentCoordinate);
+
+	{
+#if WITH_EDITOR
+		FReadScopeLock ScopeLock(CacheLock);
+#endif
+		if (FPCGLandscapeCacheEntry** FoundEntry = CachedData.Find(ComponentKey))
+		{
+			CacheEntry = *FoundEntry;
+		}
+	}
+
 	if (CacheEntry)
 	{
 		if (CacheEntry->TouchAndLoad(CacheTouch++))
diff --git a/Engine/Plugins/Experimental/PCG/Source/PCG/Private/Helpers/PCGBlueprintHelpers.cpp b/Engine/Plugins/Experimental/PCG/Source/PCG/Private/Helpers/PCGBlueprintHelpers.cpp
--- a/Engine/Plugins/Experimental/PCG/Source/PCG/Private/Helpers/PCGBlueprintHelpers.cpp
+++ b/Engine/Plugins/Experimental/PCG/Source/PCG/Private/Helpers/PCGBlueprintHelpers.cpp
@@ -131,59 +131,59 @@ TArray<FPCGLandscapeLayerWeight> UPCGBlueprintHelpers::GetInterpolatedPCGLandsca
 {
 	UWorld* World = WorldContextObject ? WorldContextObject->GetWorld() : nullptr;
 	if (!World)
+	{
 		return {};
+	}
 
 	UPCGSubsystem* PCGSubSystem = UWorld::GetSubsystem<UPCGSubsystem>(World);
 	if (!PCGSubSystem)
+	{
 		return {};
+	}
 
-	FIntPoint ComponentKey{};
-
-	auto FindLandscapeComponent = [World, Location, &ComponentKey]() -> ULandscapeComponent* {
-		for (TObjectIterator<ALandscapeProxy> It; It; ++It)
-		{
-			if (It->GetWorld() == World)
-			{
-				FBox Box = It->GetComponentsBoundingBox();
-				if (Box.IsInsideOrOnXY(Location))
-				{
-					if (ULandscapeInfo* Info = It->GetLandscapeInfo())
-					{
-						const FVector ActorSpaceLocation = It->LandscapeActorToWorld().InverseTransformPosition(Location);
-						ComponentKey = FIntPoint(FMath::FloorToInt(ActorSpaceLocation.X / It->ComponentSizeQuads), FMath::FloorToInt(ActorSpaceLocation.Y / It->ComponentSizeQuads));
-
-						return Info->XYtoComponentMap.FindRef(ComponentKey);
-					}
-
-					return nullptr;
-				}
-			}
-		}
-
-		return nullptr;
-	};
-
-	ULandscapeComponent* LandscapeComponent = FindLandscapeComponent();
-	if (!LandscapeComponent)
+	UPCGLandscapeCache* LandscapeCache = PCGSubSystem->GetLandscapeCache();
+	if (!LandscapeCache)
+	{
 		return {};
+	}
 
-	const FVector ComponentSpaceLocation = LandscapeComponent->GetComponentToWorld().InverseTransformPosition(Location);
+	FBox Bounds(&Location, 1);
+	TArray<TWeakObjectPtr<ALandscapeProxy>> Landscapes = PCGHelpers::GetLandscapeProxies(World, Bounds);
 
-	UPCGLandscapeCache* LandscapeCache = PCGSubSystem->GetLandscapeCache();
-	if (!LandscapeCache)
+	if (Landscapes.IsEmpty() || !Landscapes[0].Get())
+	{
 		return {};
+	}
+
+	ALandscapeProxy* Landscape = Landscapes[0].Get();
+	check(Landscape);
+
+	ULandscapeInfo* LandscapeInfo = Landscape->GetLandscapeInfo();
+	check(LandscapeInfo);
+
+	const FVector LocalPoint = Landscape->GetTransform().InverseTransformPosition(Location);
+	const FIntPoint ComponentMapKey(FMath::FloorToInt(LocalPoint.X / LandscapeInfo->ComponentSizeQuads), FMath::FloorToInt(LocalPoint.Y / LandscapeInfo->ComponentSizeQuads));
+
+#if WITH_EDITOR
+	ULandscapeComponent* LandscapeComponent = LandscapeInfo->XYtoComponentMap.FindRef(ComponentMapKey);
+	const FPCGLandscapeCacheEntry* CacheEntry = LandscapeCache->GetCacheEntry(LandscapeComponent, ComponentMapKey);
+#else
+	const FPCGLandscapeCacheEntry* CacheEntry = LandscapeCache->GetCacheEntry(Landscape->GetLandscapeGuid(), ComponentMapKey);
+#endif
 
-	const FPCGLandscapeCacheEntry* CacheEntry = LandscapeCache->GetCacheEntry(LandscapeComponent, ComponentKey);
 	if (!CacheEntry)
+	{
 		return {};
+	}
 
+	const FVector2D ComponentLocalPoint(LocalPoint.X - ComponentMapKey.X * LandscapeInfo->ComponentSizeQuads, LocalPoint.Y - ComponentMapKey.Y * LandscapeInfo->ComponentSizeQuads);
+	
 	TArray<FPCGLandscapeLayerWeight> Result;
+	CacheEntry->GetInterpolatedLayerWeights(ComponentLocalPoint, Result);
 
-	CacheEntry->GetInterpolatedLayerWeights(FVector2D(ComponentSpaceLocation), Result);
-
-	Result.Sort([](const FPCGLandscapeLayerWeight& Lhs, const FPCGLandscapeLayerWeight& Rhs){
+	Result.Sort([](const FPCGLandscapeLayerWeight& Lhs, const FPCGLandscapeLayerWeight& Rhs) {
 		return Lhs.Weight > Rhs.Weight;
 	});
-					
+
 	return Result;
 }
diff --git a/Engine/Plugins/Experimental/PCG/Source/PCG/Public/Grid/PCGLandscapeCache.h b/Engine/Plugins/Experimental/PCG/Source/PCG/Public/Grid/PCGLandscapeCache.h
--- a/Engine/Plugins/Experimental/PCG/Source/PCG/Public/Grid/PCGLandscapeCache.h
+++ b/Engine/Plugins/Experimental/PCG/Source/PCG/Public/Grid/PCGLandscapeCache.h
@@ -87,7 +87,14 @@ public:
 	void ClearCache();
 	void Tick(float DeltaSeconds);
 
+#if WITH_EDITOR
+	/** Gets (and creates if needed) the cache entry - available only in Editor */
 	const FPCGLandscapeCacheEntry* GetCacheEntry(ULandscapeComponent* LandscapeComponent, const FIntPoint& ComponentKey);
+#endif
+
+	/** Gets landscape cache entry, works both in editor (but does not create) but works in game mode too. */
+	const FPCGLandscapeCacheEntry* GetCacheEntry(const FGuid& LandscapeGuid, const FIntPoint& ComponentKey);
+
 	TArray<FName> GetLayerNames(ALandscapeProxy* Landscape);
 
 private:
