[PCG] Fix issue with pin tooltips raised by static analysis.

Fallback to no pin description if no matching pin found.

diff --git a/Engine/Plugins/Experimental/PCG/Source/PCGEditor/Private/PCGEditorGraphNodeBase.cpp b/Engine/Plugins/Experimental/PCG/Source/PCGEditor/Private/PCGEditorGraphNodeBase.cpp
--- a/Engine/Plugins/Experimental/PCG/Source/PCGEditor/Private/PCGEditorGraphNodeBase.cpp
+++ b/Engine/Plugins/Experimental/PCG/Source/PCGEditor/Private/PCGEditorGraphNodeBase.cpp
@@ -482,7 +482,12 @@ void UPCGEditorGraphNodeBase::GetPinHoverText(const UEdGraphPin& Pin, FString& H
 
 	const FText DataTypeText = PinTypeToText(Pin.PinType.PinCategory, MatchingPin);
 	const FText DataSubtypeText = PinTypeToText(Pin.PinType.PinSubCategory, MatchingPin);
-	const FText Description = (MatchingPin && !MatchingPin->Properties.Tooltip.IsEmpty()) ? MatchingPin->Properties.Tooltip : FText::FromName(MatchingPin->Properties.Label);
+
+	FText Description;
+	if (MatchingPin)
+	{
+		Description = MatchingPin->Properties.Tooltip.IsEmpty() ? FText::FromName(MatchingPin->Properties.Label) : MatchingPin->Properties.Tooltip;
+	}
 
 	if (MatchingPin != nullptr && bIsInputPin)
 	{
