[PCG] Fix for Composite type - it was wrong to have it on its own isolated bit, because Concrete "is a" Composite, ie Composite was supposed to be a superset of Concrete.

It may be quirky that there is a new type bit ShapeBoolean that is not directly used in the wider code. I.e. Union still outputs Composite. I believe this is correct - a Union can create a boolean if it has multiple inputs, or if it only has one input it passes it through in which case it can produce Concrete. So it can output concrete or boolean and its type should stay at Composite. Therefore the ShapeBoolean bit is only there to differentiate type but is not needed in the wider code.

diff --git a/Engine/Plugins/Experimental/PCG/Source/PCG/Public/PCGCommon.h b/Engine/Plugins/Experimental/PCG/Source/PCG/Public/PCGCommon.h
--- a/Engine/Plugins/Experimental/PCG/Source/PCG/Public/PCGCommon.h
+++ b/Engine/Plugins/Experimental/PCG/Source/PCG/Public/PCGCommon.h
@@ -42,8 +42,11 @@ enum class EPCGDataType : uint32
 	/** Simple concrete data. */
 	Concrete = Point | PolyLine | Surface | Volume | Primitive,
 
-	/** Composite spatial data which can include boolean operations. Points can be sampled from the set without collapsing to concrete data. */
-	Composite = 1 << 9,
+	/** Boolean operations like union, difference, intersection. */
+	ShapeBoolean = 1 << 9 UMETA(Hidden),
+
+	/** Composite is any combination of concrete and booleans. */
+	Composite = ShapeBoolean | Concrete,
 
 	Spatial = Composite | Concrete,
 
