[PCG] Fix pin categorisation.

I broke this code in previous commits. I missed the distinction between setting PinCategory and PinSubCategory when I added sub-categories within Spatial. Made them clear and distinct now.

diff --git a/Engine/Plugins/Experimental/PCG/Source/PCGEditor/Private/PCGEditorGraphNodeBase.cpp b/Engine/Plugins/Experimental/PCG/Source/PCGEditor/Private/PCGEditorGraphNodeBase.cpp
--- a/Engine/Plugins/Experimental/PCG/Source/PCGEditor/Private/PCGEditorGraphNodeBase.cpp
+++ b/Engine/Plugins/Experimental/PCG/Source/PCGEditor/Private/PCGEditorGraphNodeBase.cpp
@@ -376,18 +376,22 @@ FEdGraphPinType UPCGEditorGraphNodeBase::GetPinType(const UPCGPin* InPin)
 
 	if (CheckType(EPCGDataType::Spatial))
 	{
-		EdPinType.PinCategory = FPCGEditorCommon::SpatialDataType;
-
-		// Assign subcategory if we have precise information
-		if (CheckType(EPCGDataType::Composite))
+		// Categories within spatial, checking from narrow -> wide
+		if (CheckType(EPCGDataType::Concrete))
+		{
+			EdPinType.PinCategory = FPCGEditorCommon::ConcreteDataType;
+		}
+		else if (CheckType(EPCGDataType::Composite))
 		{
 			EdPinType.PinCategory = FPCGEditorCommon::CompositeDataType;
 		}
-		else if (CheckType(EPCGDataType::Concrete))
+		else
 		{
-			EdPinType.PinCategory = FPCGEditorCommon::ConcreteDataType;
+			EdPinType.PinCategory = FPCGEditorCommon::SpatialDataType;
 		}
-		else if (CheckType(EPCGDataType::Point))
+
+		// Assign subcategory if we have precise information
+		if (CheckType(EPCGDataType::Point))
 		{
 			EdPinType.PinSubCategory = FPCGEditorCommon::PointDataType;
 		}
