Making the root set a bit more robust by rooting outermost instead of just data directly

diff --git a/Engine/Plugins/Experimental/PCG/Source/PCG/Private/PCGData.cpp b/Engine/Plugins/Experimental/PCG/Source/PCG/Private/PCGData.cpp
--- a/Engine/Plugins/Experimental/PCG/Source/PCG/Private/PCGData.cpp
+++ b/Engine/Plugins/Experimental/PCG/Source/PCG/Private/PCGData.cpp
@@ -20,16 +20,31 @@ void FPCGRootSet::Clear()
 void FPCGRootSet::Add(UObject* InObject)
 {
 	check(InObject);
+	AddInternal(InObject);
+}
+
+void FPCGRootSet::AddInternal(UObject* InObject)
+{
+	check(InObject && !InObject->IsA<UPackage>());
+
 	if (int32* Found = RootSet.Find(InObject))
 	{
 		(*Found)++;
 	}
-	else if(!InObject->IsRooted() && InObject->GetPackage() == GetTransientPackage())
+	else if (!InObject->IsRooted() && InObject->GetPackage() == GetTransientPackage())
 	{
 		InObject->AddToRoot();
 		RootSet.Emplace(InObject, 1);
 	}
+
+	// Recurse to outermost
+	UObject* OuterObject = InObject->GetOuter();
+	if (OuterObject && !OuterObject->IsA<UPackage>())
+	{
+		AddInternal(OuterObject);
+	}
 }
+
 void FPCGRootSet::Remove(UObject* InObject)
 {
 	if (!InObject)
@@ -37,7 +52,14 @@ void FPCGRootSet::Remove(UObject* InObject)
 		UE_LOG(LogPCG, Warning, TEXT("Trying to remove a null object from the rootset"));
 		return;
 	}
-	
+
+	RemoveInternal(InObject);
+}
+
+void FPCGRootSet::RemoveInternal(UObject* InObject)
+{
+	check(InObject && !InObject->IsA<UPackage>());
+
 	if (int32* Found = RootSet.Find(InObject))
 	{
 		check(InObject->IsRooted());
@@ -49,6 +71,13 @@ void FPCGRootSet::Remove(UObject* InObject)
 			RootSet.Remove(InObject);
 		}
 	}
+
+	// Recurse to outermost
+	UObject* OuterObject = InObject->GetOuter();
+	if (OuterObject && !OuterObject->IsA<UPackage>())
+	{
+		RemoveInternal(OuterObject);
+	}
 }
 
 bool FPCGTaggedData::operator==(const FPCGTaggedData& Other) const
diff --git a/Engine/Plugins/Experimental/PCG/Source/PCG/Public/PCGData.h b/Engine/Plugins/Experimental/PCG/Source/PCG/Public/PCGData.h
--- a/Engine/Plugins/Experimental/PCG/Source/PCG/Public/PCGData.h
+++ b/Engine/Plugins/Experimental/PCG/Source/PCG/Public/PCGData.h
@@ -54,6 +54,9 @@ struct FPCGRootSet
 	void Remove(UObject* InObject);
 
 	TMap<UObject*, int32> RootSet;
+private:
+	void AddInternal(UObject* InObject);
+	void RemoveInternal(UObject* InObject);
 };
 
 USTRUCT(BlueprintType)
