Fixed missing include

#rnx

diff --git a/Engine/Plugins/Experimental/PCG/Source/PCG/Private/PCGData.cpp b/Engine/Plugins/Experimental/PCG/Source/PCG/Private/PCGData.cpp
--- a/Engine/Plugins/Experimental/PCG/Source/PCG/Private/PCGData.cpp
+++ b/Engine/Plugins/Experimental/PCG/Source/PCG/Private/PCGData.cpp
@@ -5,6 +5,8 @@
 #include "PCGParamData.h"
 #include "Data/PCGSpatialData.h"
 
+#include "UObject/Package.h"
+
 #include UE_INLINE_GENERATED_CPP_BY_NAME(PCGData)
 
 void FPCGRootSet::Clear()
