Cleanup the way we do the unregister / tentative fix for crash on close

#rnx

diff --git a/Engine/Plugins/Experimental/PCG/Source/PCGEditor/Private/PCGEditorModule.cpp b/Engine/Plugins/Experimental/PCG/Source/PCGEditor/Private/PCGEditorModule.cpp
--- a/Engine/Plugins/Experimental/PCG/Source/PCGEditor/Private/PCGEditorModule.cpp
+++ b/Engine/Plugins/Experimental/PCG/Source/PCGEditor/Private/PCGEditorModule.cpp
@@ -88,20 +88,22 @@ void FPCGEditorModule::RegisterDetailsCustomizations()
 	PropertyEditor.RegisterCustomClassLayout(APCGVolume::StaticClass()->GetFName(), FOnGetDetailCustomizationInstance::CreateStatic(&FPCGVolumeDetails::MakeInstance));
 
 	PropertyEditor.RegisterCustomPropertyTypeLayout(FPCGAttributePropertySelector::StaticStruct()->GetFName(), FOnGetPropertyTypeCustomizationInstance::CreateStatic(&FPCGAttributePropertySelectorDetails::MakeInstance));
+
+	PropertyEditor.NotifyCustomizationModuleChanged();
 }
 
 void FPCGEditorModule::UnregisterDetailsCustomizations()
 {
-	if (FModuleManager::Get().IsModuleLoaded("PropertyEditor"))
+	if (FPropertyEditorModule* PropertyModule = FModuleManager::GetModulePtr<FPropertyEditorModule>("PropertyEditor"))
 	{
-		FPropertyEditorModule& PropertyModule = FModuleManager::LoadModuleChecked<FPropertyEditorModule>("PropertyEditor");
-		PropertyModule.UnregisterCustomClassLayout(UPCGBlueprintSettings::StaticClass()->GetFName());
-		PropertyModule.UnregisterCustomClassLayout(UPCGComponent::StaticClass()->GetFName());
-		PropertyModule.UnregisterCustomClassLayout(UPCGGraph::StaticClass()->GetFName());
-		PropertyModule.UnregisterCustomClassLayout(APCGVolume::StaticClass()->GetFName());
-
-		PropertyModule.UnregisterCustomPropertyTypeLayout(FPCGAttributePropertySelector::StaticStruct()->GetFName());
-		PropertyModule.NotifyCustomizationModuleChanged();
+		PropertyModule->UnregisterCustomClassLayout(UPCGBlueprintSettings::StaticClass()->GetFName());
+		PropertyModule->UnregisterCustomClassLayout(UPCGComponent::StaticClass()->GetFName());
+		PropertyModule->UnregisterCustomClassLayout(UPCGGraph::StaticClass()->GetFName());
+		PropertyModule->UnregisterCustomClassLayout(APCGVolume::StaticClass()->GetFName());
+
+		PropertyModule->UnregisterCustomPropertyTypeLayout(FPCGAttributePropertySelector::StaticStruct()->GetFName());
+
+		PropertyModule->NotifyCustomizationModuleChanged();
 	}
 }
 
@@ -121,18 +123,14 @@ void FPCGEditorModule::RegisterAssetTypeActions()
 
 void FPCGEditorModule::UnregisterAssetTypeActions()
 {
-	FAssetToolsModule* AssetToolsModule = FModuleManager::GetModulePtr<FAssetToolsModule>("AssetTools");
-
-	if (!AssetToolsModule)
+	if (FAssetToolsModule* AssetToolsModule = FModuleManager::GetModulePtr<FAssetToolsModule>("AssetTools"))
 	{
-		return;
-	}
-
-	IAssetTools& AssetTools = AssetToolsModule->Get();
+		IAssetTools& AssetTools = AssetToolsModule->Get();
 
-	for (auto Action : RegisteredAssetTypeActions)
-	{
-		AssetTools.UnregisterAssetTypeActions(Action);
+		for (auto Action : RegisteredAssetTypeActions)
+		{
+			AssetTools.UnregisterAssetTypeActions(Action);
+		}
 	}
 }
 
