[PCG] Add modulo for Attribute Maths op

diff --git a/Engine/Plugins/Experimental/PCG/Source/PCG/Private/Elements/Metadata/PCGMetadataMaths.inl b/Engine/Plugins/Experimental/PCG/Source/PCG/Private/Elements/Metadata/PCGMetadataMaths.inl
--- a/Engine/Plugins/Experimental/PCG/Source/PCG/Private/Elements/Metadata/PCGMetadataMaths.inl
+++ b/Engine/Plugins/Experimental/PCG/Source/PCG/Private/Elements/Metadata/PCGMetadataMaths.inl
@@ -393,4 +393,52 @@ namespace PCGMetadataMaths
 			Lerp(Value1.Z, Value2.Z, Ratio.Z),
 			Lerp(Value1.W, Value2.W, Ratio.W));
 	}
+
+	////////////////////
+	// Modulo
+	////////////////////
+	template <typename T>
+	inline T Modulo(const T& Value1, const T& Value2)
+	{
+		return (Value2 != 0) ? (Value1 % Value2) : 0;
+	}
+
+	template<>
+	inline float Modulo(const float& Value1, const float& Value2)
+	{
+		return static_cast<float>(Modulo(FMath::RoundToInt(Value1), FMath::RoundToInt(Value2)));
+	}
+
+	template<>
+	inline double Modulo(const double& Value1, const double& Value2)
+	{
+		return static_cast<double>(Modulo(FMath::RoundToInt(Value1), FMath::RoundToInt(Value2)));
+	}
+
+	template<>
+	inline FVector2D Modulo(const FVector2D& Value1, const FVector2D& Value2)
+	{
+		return FVector2D(
+			Modulo(Value1.X, Value2.X),
+			Modulo(Value1.Y, Value2.Y));
+	}
+
+	template<>
+	inline FVector Modulo(const FVector& Value1, const FVector& Value2)
+	{
+		return FVector(
+			Modulo(Value1.X, Value2.X),
+			Modulo(Value1.Y, Value2.Y),
+			Modulo(Value1.Z, Value2.Z));
+	}
+
+	template<>
+	inline FVector4 Modulo(const FVector4& Value1, const FVector4& Value2)
+	{
+		return FVector4(
+			Modulo(Value1.X, Value2.X),
+			Modulo(Value1.Y, Value2.Y),
+			Modulo(Value1.Z, Value2.Z),
+			Modulo(Value1.W, Value2.W));
+	}	
 }
\ No newline at end of file
diff --git a/Engine/Plugins/Experimental/PCG/Source/PCG/Private/Elements/Metadata/PCGMetadataMathsOpElement.cpp b/Engine/Plugins/Experimental/PCG/Source/PCG/Private/Elements/Metadata/PCGMetadataMathsOpElement.cpp
--- a/Engine/Plugins/Experimental/PCG/Source/PCG/Private/Elements/Metadata/PCGMetadataMathsOpElement.cpp
+++ b/Engine/Plugins/Experimental/PCG/Source/PCG/Private/Elements/Metadata/PCGMetadataMathsOpElement.cpp
@@ -134,6 +134,8 @@ namespace PCGMetadataMathsSettings
 			return PCGMetadataMaths::Clamp(Value1, Value1, Value2);
 		case EPCGMedadataMathsOperation::Pow:
 			return PCGMetadataMaths::Pow(Value1, Value2);
+		case EPCGMedadataMathsOperation::Modulo:
+			return PCGMetadataMaths::Modulo(Value1, Value2);
 		default:
 			return T{};
 		}
@@ -238,7 +240,7 @@ FPCGAttributePropertySelector UPCGMetadataMathsSettings::GetInputSource(uint32 I
 
 FName UPCGMetadataMathsSettings::AdditionalTaskName() const
 {
-	if (const UEnum* EnumPtr = FindObject<UEnum>(nullptr, TEXT("/Script/PCG.EPCGMedadataMathsOperation"), true))
+	if (const UEnum* EnumPtr = StaticEnum<EPCGMedadataMathsOperation>())
 	{
 		return FName(FString("Maths: ") + EnumPtr->GetNameStringByValue(static_cast<int>(Operation)));
 	}
diff --git a/Engine/Plugins/Experimental/PCG/Source/PCG/Public/Elements/Metadata/PCGMetadataMathsOpElement.h b/Engine/Plugins/Experimental/PCG/Source/PCG/Public/Elements/Metadata/PCGMetadataMathsOpElement.h
--- a/Engine/Plugins/Experimental/PCG/Source/PCG/Public/Elements/Metadata/PCGMetadataMathsOpElement.h
+++ b/Engine/Plugins/Experimental/PCG/Source/PCG/Public/Elements/Metadata/PCGMetadataMathsOpElement.h
@@ -33,6 +33,7 @@ enum class EPCGMedadataMathsOperation : uint16
 	Pow,
 	ClampMin,
 	ClampMax,
+	Modulo,
 
 	// Ternary op
 	TernaryOp = 1 << 12 UMETA(Hidden),
