[PCG] Fix issue where Attributes could be created with a given type, and parented to an attribute of another type.

* Also make CreateAttribute return the typed attribute.
* Add TypedCopy, GetMutableTypedAttribute, GetConstTypedAttribute and FindOrCreateAttribute, all templated.

diff --git a/Engine/Plugins/Experimental/PCG/Source/PCG/Public/Metadata/PCGMetadata.h b/Engine/Plugins/Experimental/PCG/Source/PCG/Public/Metadata/PCGMetadata.h
--- a/Engine/Plugins/Experimental/PCG/Source/PCG/Public/Metadata/PCGMetadata.h
+++ b/Engine/Plugins/Experimental/PCG/Source/PCG/Public/Metadata/PCGMetadata.h
@@ -142,6 +142,12 @@ public:
 	bool HasAttribute(FName AttributeName) const;
 	int32 GetAttributeCount() const;
 
+	template <typename T>
+	FPCGMetadataAttribute<T>* GetMutableTypedAttribute(FName AttributeName);
+
+	template <typename T>
+	const FPCGMetadataAttribute<T>* GetConstTypedAttribute(FName AttributeName) const;
+
 	UFUNCTION(BlueprintCallable, Category = "PCG|Metadata")
 	void GetAttributes(TArray<FName>& AttributeNames, TArray<EPCGMetadataTypes>& AttributeTypes) const;
 
@@ -235,7 +241,10 @@ public:
 	int64 GetItemCountForChild() const;
 
 	template<typename T>
-	FPCGMetadataAttributeBase* CreateAttribute(FName AttributeName, const T& DefaultValue, bool bAllowsInterpolation, bool bOverrideParent);
+	FPCGMetadataAttribute<T>* CreateAttribute(FName AttributeName, const T& DefaultValue, bool bAllowsInterpolation, bool bOverrideParent);
+
+	template<typename T>
+	FPCGMetadataAttribute<T>* FindOrCreateAttribute(FName AttributeName, const T& DefaultValue = T{}, bool bAllowsInterpolation = true, bool bOverrideParent = true);
 
 protected:
 	FPCGMetadataAttributeBase* CopyAttribute(FName AttributeToCopy, FName NewAttributeName, bool bKeepParent, bool bCopyEntries, bool bCopyValues);
@@ -266,7 +275,7 @@ protected:
 };
 
 template<typename T>
-FPCGMetadataAttributeBase* UPCGMetadata::CreateAttribute(FName AttributeName, const T& DefaultValue, bool bAllowsInterpolation, bool bOverrideParent)
+FPCGMetadataAttribute<T>* UPCGMetadata::CreateAttribute(FName AttributeName, const T& DefaultValue, bool bAllowsInterpolation, bool bOverrideParent)
 {
 	if (!FPCGMetadataAttributeBase::IsValidName(AttributeName))
 	{
@@ -281,14 +290,28 @@ FPCGMetadataAttributeBase* UPCGMetadata::CreateAttribute(FName AttributeName, co
 		ParentAttribute = Parent->GetConstAttribute(AttributeName);
 	}
 
-	FPCGMetadataAttributeBase* NewAttribute = new FPCGMetadataAttribute<T>(this, AttributeName, ParentAttribute, DefaultValue, bAllowsInterpolation);
+	if (ParentAttribute && (ParentAttribute->GetTypeId() != PCG::Private::MetadataTypes<T>::Id))
+	{
+		// Can't parent if the types doesn't match
+		ParentAttribute = nullptr;
+	}
+
+	FPCGMetadataAttribute<T>* NewAttribute = new FPCGMetadataAttribute<T>(this, AttributeName, ParentAttribute, DefaultValue, bAllowsInterpolation);
 
 	AttributeLock.WriteLock();
 	if (FPCGMetadataAttributeBase** ExistingAttribute = Attributes.Find(AttributeName))
 	{
-		UE_LOG(LogPCG, Warning, TEXT("Attribute %s already exists"), *AttributeName.ToString());
 		delete NewAttribute;
-		NewAttribute = *ExistingAttribute;
+		if ((*ExistingAttribute)->GetTypeId() != PCG::Private::MetadataTypes<T>::Id)
+		{
+			UE_LOG(LogPCG, Error, TEXT("Attribute %s already exists but is not the right type. Abort."), *AttributeName.ToString());
+			return nullptr;
+		}
+		else
+		{
+			UE_LOG(LogPCG, Warning, TEXT("Attribute %s already exists"), *AttributeName.ToString());
+			NewAttribute = static_cast<FPCGMetadataAttribute<T>*>(*ExistingAttribute);
+		}
 	}
 	else
 	{
@@ -298,4 +321,35 @@ FPCGMetadataAttributeBase* UPCGMetadata::CreateAttribute(FName AttributeName, co
 	AttributeLock.WriteUnlock();
 
 	return NewAttribute;
+}
+
+template<typename T>
+FPCGMetadataAttribute<T>* UPCGMetadata::FindOrCreateAttribute(FName AttributeName, const T& DefaultValue, bool bAllowsInterpolation, bool bOverrideParent)
+{
+	FPCGMetadataAttribute<T>* Attribute = GetMutableTypedAttribute<T>(AttributeName);
+	if (!Attribute)
+	{
+		Attribute = CreateAttribute<T>(AttributeName, DefaultValue, bAllowsInterpolation, bOverrideParent);
+	}
+
+	return Attribute;
+}
+
+template <typename T>
+FPCGMetadataAttribute<T>* UPCGMetadata::GetMutableTypedAttribute(FName AttributeName)
+{
+	FPCGMetadataAttributeBase* BaseAttribute = GetMutableAttribute(AttributeName);
+	return (BaseAttribute && (BaseAttribute->GetTypeId() == PCG::Private::MetadataTypes<T>::Id))
+		? static_cast<FPCGMetadataAttribute<T>*>(BaseAttribute)
+		: nullptr;
+
+}
+
+template <typename T>
+const FPCGMetadataAttribute<T>* UPCGMetadata::GetConstTypedAttribute(FName AttributeName) const
+{
+	const FPCGMetadataAttributeBase* BaseAttribute = GetConstAttribute(AttributeName);
+	return (BaseAttribute && (BaseAttribute->GetTypeId() == PCG::Private::MetadataTypes<T>::Id))
+		? static_cast<const FPCGMetadataAttribute<T>*>(BaseAttribute)
+		: nullptr;
 }
\ No newline at end of file
diff --git a/Engine/Plugins/Experimental/PCG/Source/PCG/Public/Metadata/PCGMetadataAttributeTpl.h b/Engine/Plugins/Experimental/PCG/Source/PCG/Public/Metadata/PCGMetadataAttributeTpl.h
--- a/Engine/Plugins/Experimental/PCG/Source/PCG/Public/Metadata/PCGMetadataAttributeTpl.h
+++ b/Engine/Plugins/Experimental/PCG/Source/PCG/Public/Metadata/PCGMetadataAttributeTpl.h
@@ -22,6 +22,9 @@ public:
 	{
 		TypeId = PCG::Private::MetadataTypes<T>::Id;
 
+		// Make sure we don't parent with the wrong type id
+		check(!Parent || Parent->GetTypeId() == TypeId)
+
 		if (GetParent())
 		{
 			ValueKeyOffset = GetParent()->GetValueKeyOffsetForChild();
@@ -50,6 +53,11 @@ public:
 
 	const FPCGMetadataAttribute* GetParent() const { return static_cast<const FPCGMetadataAttribute*>(Parent); }
 
+	FPCGMetadataAttribute* TypedCopy(FName NewName, UPCGMetadata* InMetadata, bool bKeepParent, bool bCopyEntries = true, bool bCopyValues = true)
+	{
+		return static_cast<FPCGMetadataAttribute*>(Copy(NewName, InMetadata, bKeepParent, bCopyEntries, bCopyValues));
+	}
+
 	virtual FPCGMetadataAttributeBase* Copy(FName NewName, UPCGMetadata* InMetadata, bool bKeepParent, bool bCopyEntries = true, bool bCopyValues = true) const override
 	{
 		// If we copy an attribute where we don't want to keep the parent, while copying entries and/or values, we'll lose data.
@@ -333,8 +341,9 @@ public:
 		}
 		else if (ValueKey >= ValueKeyOffset)
 		{
+			int32 Index = ValueKey - ValueKeyOffset;
 			FReadScopeLock ScopeLock(ValueLock);
-			return Values[ValueKey - ValueKeyOffset];
+			return Index < Values.Num() ? Values[Index] : DefaultValue;
 		}
 		else if (GetParent())
 		{
