[PCG] Tentative fix for Unregister in PCG Module Editor.

#rnx

diff --git a/Engine/Plugins/Experimental/PCG/Source/PCGEditor/Private/PCGEditorModule.cpp b/Engine/Plugins/Experimental/PCG/Source/PCGEditor/Private/PCGEditorModule.cpp
--- a/Engine/Plugins/Experimental/PCG/Source/PCGEditor/Private/PCGEditorModule.cpp
+++ b/Engine/Plugins/Experimental/PCG/Source/PCGEditor/Private/PCGEditorModule.cpp
@@ -4,15 +4,12 @@
 
 #include "AssetTypeActions/PCGGraphAssetTypeActions.h"
 #include "AssetTypeActions/PCGSettingsAssetTypeActions.h"
-#include "PCGComponent.h"
 #include "PCGEditorCommands.h"
 #include "PCGEditorGraphNodeFactory.h"
 #include "PCGEditorSettings.h"
 #include "PCGEditorStyle.h"
 #include "PCGEditorUtils.h"
-#include "PCGGraph.h"
 #include "PCGSubsystem.h"
-#include "PCGVolume.h"
 #include "PCGVolumeFactory.h"
 #include "PCGWorldActor.h"
 
@@ -21,8 +18,6 @@
 #include "Details/PCGGraphDetails.h"
 #include "Details/PCGVolumeDetails.h"
 #include "Details/PCGAttributePropertySelectorDetails.h"
-#include "Elements/PCGExecuteBlueprint.h"
-#include "Metadata/PCGAttributePropertySelector.h"
 
 #include "EdGraphUtilities.h"
 #include "Framework/MultiBox/MultiBoxBuilder.h"
@@ -82,12 +77,12 @@ bool FPCGEditorModule::SupportsDynamicReloading()
 void FPCGEditorModule::RegisterDetailsCustomizations()
 {
 	FPropertyEditorModule& PropertyEditor = FModuleManager::LoadModuleChecked<FPropertyEditorModule>("PropertyEditor");
-	PropertyEditor.RegisterCustomClassLayout(UPCGBlueprintSettings::StaticClass()->GetFName(), FOnGetDetailCustomizationInstance::CreateStatic(&FPCGBlueprintSettingsDetails::MakeInstance));
-	PropertyEditor.RegisterCustomClassLayout(UPCGComponent::StaticClass()->GetFName(), FOnGetDetailCustomizationInstance::CreateStatic(&FPCGComponentDetails::MakeInstance));
-	PropertyEditor.RegisterCustomClassLayout(UPCGGraph::StaticClass()->GetFName(), FOnGetDetailCustomizationInstance::CreateStatic(&FPCGGraphDetails::MakeInstance));
-	PropertyEditor.RegisterCustomClassLayout(APCGVolume::StaticClass()->GetFName(), FOnGetDetailCustomizationInstance::CreateStatic(&FPCGVolumeDetails::MakeInstance));
+	PropertyEditor.RegisterCustomClassLayout("PCGBlueprintSettings", FOnGetDetailCustomizationInstance::CreateStatic(&FPCGBlueprintSettingsDetails::MakeInstance));
+	PropertyEditor.RegisterCustomClassLayout("PCGComponent", FOnGetDetailCustomizationInstance::CreateStatic(&FPCGComponentDetails::MakeInstance));
+	PropertyEditor.RegisterCustomClassLayout("PCGGraph", FOnGetDetailCustomizationInstance::CreateStatic(&FPCGGraphDetails::MakeInstance));
+	PropertyEditor.RegisterCustomClassLayout("PCGVolume", FOnGetDetailCustomizationInstance::CreateStatic(&FPCGVolumeDetails::MakeInstance));
 
-	PropertyEditor.RegisterCustomPropertyTypeLayout(FPCGAttributePropertySelector::StaticStruct()->GetFName(), FOnGetPropertyTypeCustomizationInstance::CreateStatic(&FPCGAttributePropertySelectorDetails::MakeInstance));
+	PropertyEditor.RegisterCustomPropertyTypeLayout("PCGAttributePropertySelector", FOnGetPropertyTypeCustomizationInstance::CreateStatic(&FPCGAttributePropertySelectorDetails::MakeInstance));
 
 	PropertyEditor.NotifyCustomizationModuleChanged();
 }
@@ -96,12 +91,12 @@ void FPCGEditorModule::UnregisterDetailsCustomizations()
 {
 	if (FPropertyEditorModule* PropertyModule = FModuleManager::GetModulePtr<FPropertyEditorModule>("PropertyEditor"))
 	{
-		PropertyModule->UnregisterCustomClassLayout(UPCGBlueprintSettings::StaticClass()->GetFName());
-		PropertyModule->UnregisterCustomClassLayout(UPCGComponent::StaticClass()->GetFName());
-		PropertyModule->UnregisterCustomClassLayout(UPCGGraph::StaticClass()->GetFName());
-		PropertyModule->UnregisterCustomClassLayout(APCGVolume::StaticClass()->GetFName());
+		PropertyModule->UnregisterCustomClassLayout("PCGBlueprintSettings");
+		PropertyModule->UnregisterCustomClassLayout("PCGComponent");
+		PropertyModule->UnregisterCustomClassLayout("PCGGraph");
+		PropertyModule->UnregisterCustomClassLayout("PCGVolume");
 
-		PropertyModule->UnregisterCustomPropertyTypeLayout(FPCGAttributePropertySelector::StaticStruct()->GetFName());
+		PropertyModule->UnregisterCustomPropertyTypeLayout("PCGAttributePropertySelector");
 
 		PropertyModule->NotifyCustomizationModuleChanged();
 	}
