[PCG] Second iteration of disabled node behaviour.

Firstly tweak pass through behaviour - check types overlap and the input type(s) do not exceed what the output type supports. Some types are only known dynamically/at execution time (there are currently nodes that take Any and output points, for example, and we don't know what data will come in before execution), so type checking happens during pass through. An example of what this solves: a node that takes a surface as input and outputs points. This should *not* pass through a surface when the node is disabled, as it was supposed to produce points. Its job was to convert a surface to points, if it is disabled, it outputs 0 points.

Secondly when node is disabled, pass through all edges connected to pass-through pin by default. Basically assume node operation is "for each edge, ...". The only nodes that actually perform an operation across all edges are Union, Intersection and Merge and these will pass first edge only when disabled.

To support precise type checking, UPCGData::GetDataType narrowed to give most specific type.

diff --git a/Engine/Plugins/Experimental/PCG/Source/PCG/Private/PCGElement.cpp b/Engine/Plugins/Experimental/PCG/Source/PCG/Private/PCGElement.cpp
--- a/Engine/Plugins/Experimental/PCG/Source/PCG/Private/PCGElement.cpp
+++ b/Engine/Plugins/Experimental/PCG/Source/PCG/Private/PCGElement.cpp
@@ -200,6 +200,11 @@ void IPCGElement::PostExecute(FPCGContext* Context) const
 
 void IPCGElement::DisabledPassThroughData(FPCGContext* Context) const
 {
+	check(Context);
+
+	const UPCGSettings* Settings = Context->GetInputSettings<UPCGSettings>();
+	check(Settings);
+
 	// Copy as baseline
 	Context->OutputData = Context->InputData;
 
@@ -209,16 +214,15 @@ void IPCGElement::DisabledPassThroughData(FPCGContext* Context) const
 		return;
 	}
 
-	if (Context->Node->GetInputPins().Num() == 0)
+	if (Context->Node->GetInputPins().Num() == 0 || Context->Node->GetOutputPins().Num() == 0)
 	{
-		// No input pins, return nothing
+		// No input pins or not output pins, return nothing
 		Context->OutputData.TaggedData.Empty();
 
 		return;
 	}
 
 	const UPCGPin* PassThroughPin = Context->Node->GetPassThroughInputPin();
-
 	if (PassThroughPin == nullptr)
 	{
 		// No pin to grab pass through data from
@@ -227,24 +231,53 @@ void IPCGElement::DisabledPassThroughData(FPCGContext* Context) const
 		return;
 	}
 
-	// Find first incoming non-params data that is coming through the identified pin
-	TArray<FPCGTaggedData> InputsOnFirstPin = Context->InputData.GetInputsByPin(PassThroughPin->Properties.Label);
-	const int FirstNonParamsDataIndex = InputsOnFirstPin.IndexOfByPredicate([Context](const FPCGTaggedData& InData) { return Cast<UPCGParamData>(InData.Data) == nullptr; });
-	if (FirstNonParamsDataIndex != INDEX_NONE)
+	const EPCGDataType OutputType = Context->Node->GetOutputPins()[0]->Properties.AllowedTypes;
+
+	// Pass through input data if it is not params, and if the output type supports it (e.g. if we have a incoming
+	// surface connected to an input pin of type Any, do not pass the surface through to an output pin of type Point).
+	auto InputDataShouldPassThrough = [OutputType](const FPCGTaggedData& InData)
 	{
-		// Remove everything except the data we found above
-		for (int Index = Context->OutputData.TaggedData.Num() - 1; Index >= 0; --Index)
+		EPCGDataType InputType = InData.Data ? InData.Data->GetDataType() : EPCGDataType::None;
+		const bool bInputTypeNotWiderThanOutputType = !(InputType & ~OutputType);
+		return InputType != EPCGDataType::Param && bInputTypeNotWiderThanOutputType;
+	};
+
+	if (Settings->OnlyPassThroughOneEdgeWhenDisabled())
+	{
+		// Find first incoming non-params data that is coming through the pass through pin
+		TArray<FPCGTaggedData> InputsOnFirstPin = Context->InputData.GetInputsByPin(PassThroughPin->Properties.Label);
+		const int FirstNonParamsDataIndex = InputsOnFirstPin.IndexOfByPredicate(InputDataShouldPassThrough);
+
+		if (FirstNonParamsDataIndex != INDEX_NONE)
 		{
-			if (Index != FirstNonParamsDataIndex)
+			// Remove everything except the data we found above
+			for (int Index = Context->OutputData.TaggedData.Num() - 1; Index >= 0; --Index)
 			{
-				Context->OutputData.TaggedData.RemoveAt(Index);
+				if (Index != FirstNonParamsDataIndex)
+				{
+					Context->OutputData.TaggedData.RemoveAt(Index);
+				}
 			}
 		}
+		else
+		{
+			// No data found to return
+			Context->OutputData.TaggedData.Empty();
+		}
 	}
 	else
 	{
-		// No data found to return
-		Context->OutputData.TaggedData.Empty();
+		// Remove any incoming non-params data that is coming through the pass through pin
+		TArray<FPCGTaggedData> InputsOnFirstPin = Context->InputData.GetInputsByPin(PassThroughPin->Properties.Label);
+		for (int Index = InputsOnFirstPin.Num() - 1; Index >= 0; --Index)
+		{
+			const FPCGTaggedData& Data = InputsOnFirstPin[Index];
+
+			if (!InputDataShouldPassThrough(Data))
+			{
+				Context->OutputData.TaggedData.RemoveAt(Index);
+			}
+		}
 	}
 }
 
diff --git a/Engine/Plugins/Experimental/PCG/Source/PCG/Private/PCGNode.cpp b/Engine/Plugins/Experimental/PCG/Source/PCG/Private/PCGNode.cpp
--- a/Engine/Plugins/Experimental/PCG/Source/PCG/Private/PCGNode.cpp
+++ b/Engine/Plugins/Experimental/PCG/Source/PCG/Private/PCGNode.cpp
@@ -374,8 +374,10 @@ const UPCGPin* UPCGNode::GetPassThroughInputPin() const
 		return nullptr;
 	}
 
+	const EPCGDataType OutputType = GetOutputPins()[0]->Properties.AllowedTypes;
+
 	// We assume a node whose primary output is params is a param processing node.
-	const bool bNodeOutputsParams = (GetOutputPins()[0]->Properties.AllowedTypes == EPCGDataType::Param);
+	const bool bNodeOutputsParams = (OutputType == EPCGDataType::Param);
 
 	if (bNodeOutputsParams)
 	{
@@ -390,8 +392,94 @@ const UPCGPin* UPCGNode::GetPassThroughInputPin() const
 		// Params-only pins will be rejected/ignored.
 		auto FindNonParams = [](const TObjectPtr<UPCGPin>& InPin) { return InPin->Properties.AllowedTypes != EPCGDataType::Param; };
 		const TObjectPtr<UPCGPin>* FirstNonParamsPinPtr = Algo::FindByPredicate(GetInputPins(), FindNonParams);
-		return FirstNonParamsPinPtr ? *FirstNonParamsPinPtr : nullptr;
+
+		if (FirstNonParamsPinPtr)
+		{
+			// Finally in order to be a candidate for passing through, data must be compatible.
+			const UPCGPin* InputPin = *FirstNonParamsPinPtr;
+			const EPCGDataType InputType = InputPin->Properties.AllowedTypes;
+			const bool bTypesOverlap = !!(InputType & OutputType);
+			
+			// Misc note - it would be nice if we could be stricter, like line below this comment. However it would mean it will stop an Any
+			// input being passed through to a Point output, even if the incoming edge will be receiving points dynamically/during execution. If a
+			// user creates a BP node with an Any input, which they may do lazily or unknowingly, this blocks passthrough. So instead we'll indicate
+			// that this pin *may* be used as a passthrough, and during execution in DisabledPassThroughData() we check dynamic types.
+			//const bool bInputTypeNotWiderThanOutputType = !(InputType & ~OutputType);
+			
+			if (bTypesOverlap)
+			{
+				return *FirstNonParamsPinPtr;
+			}
+		}
+
+		return nullptr;
+	}
+}
+
+bool UPCGNode::IsPinUsedByNodeExecution(const UPCGPin* InPin) const
+{
+	check(InPin);
+
+	UPCGSettings* Settings = GetSettings();
+	if (!Settings)
+	{
+		// Safe default - assume used
+		return true;
+	}
+
+	// Disabled nodes only use the 'pass through pin'
+	if (!Settings->bEnabled)
+	{
+		return InPin == GetPassThroughInputPin();
+	}
+
+	// Let Settings signal whether it uses a pin or not
+	return Settings->IsPinUsedByNodeExecution(InPin);
+}
+
+bool UPCGNode::IsEdgeUsedByNodeExecution(const UPCGEdge* InEdge) const
+{
+	check(InEdge);
+
+	// Locate the pin on this node that the edge is connected to
+	const UPCGPin* Pin = nullptr;
+	if (InEdge->InputPin->Node == this)
+	{
+		Pin = InEdge->InputPin;
+	}
+	else if (InEdge->OutputPin->Node == this)
+	{
+		Pin = InEdge->OutputPin;
+	}
+
+	if (!ensure(Pin))
+	{
+		// Safe default - assume used
+		return true;
 	}
+
+	UPCGSettings* Settings = GetSettings();
+	if (!Settings)
+	{
+		// Safe default - assume used
+		return true;
+	}
+
+	// Disabled nodes only use the 'pass through pin'
+	if (!Settings->bEnabled)
+	{
+		// Only accept first edge if node is disabled
+		const bool bConnectedToPassThrough = (Pin == GetPassThroughInputPin());
+		const bool bMustBeFirstEdge = Settings->OnlyPassThroughOneEdgeWhenDisabled();
+		const bool bIsFirstEdge = Pin->Edges.Num() > 0 && Pin->Edges[0] == InEdge;
+		if (!bConnectedToPassThrough || (bMustBeFirstEdge && !bIsFirstEdge))
+		{
+			return false;
+		}
+	}
+
+	// Ask Settings if the pin is in use, good opportunity to gray out pins based on parameters
+	return Settings->IsPinUsedByNodeExecution(Pin);
 }
 
 void UPCGNode::SetSettingsInterface(UPCGSettingsInterface* InSettingsInterface, bool bUpdatePins)
diff --git a/Engine/Plugins/Experimental/PCG/Source/PCG/Public/Data/PCGLandscapeData.h b/Engine/Plugins/Experimental/PCG/Source/PCG/Public/Data/PCGLandscapeData.h
--- a/Engine/Plugins/Experimental/PCG/Source/PCG/Public/Data/PCGLandscapeData.h
+++ b/Engine/Plugins/Experimental/PCG/Source/PCG/Public/Data/PCGLandscapeData.h
@@ -27,7 +27,7 @@ public:
 	// ~End UObject interface
 
 	// ~Begin UPCGData interface
-	virtual EPCGDataType GetDataType() const override { return EPCGDataType::Landscape | Super::GetDataType(); }
+	virtual EPCGDataType GetDataType() const override { return EPCGDataType::Landscape; }
 	// ~End UPCGData interface
 
 	// ~Begin UPGCSpatialData interface
diff --git a/Engine/Plugins/Experimental/PCG/Source/PCG/Public/Data/PCGLandscapeSplineData.h b/Engine/Plugins/Experimental/PCG/Source/PCG/Public/Data/PCGLandscapeSplineData.h
--- a/Engine/Plugins/Experimental/PCG/Source/PCG/Public/Data/PCGLandscapeSplineData.h
+++ b/Engine/Plugins/Experimental/PCG/Source/PCG/Public/Data/PCGLandscapeSplineData.h
@@ -18,7 +18,7 @@ public:
 	void Initialize(ULandscapeSplinesComponent* InSplineComponent);
 
 	// ~Begin UPCGData interface
-	virtual EPCGDataType GetDataType() const override { return EPCGDataType::LandscapeSpline | Super::GetDataType(); }
+	virtual EPCGDataType GetDataType() const override { return EPCGDataType::LandscapeSpline; }
 	// ~End UPCGData interface
 
 	//~Begin UPCGPolyLineData interface
diff --git a/Engine/Plugins/Experimental/PCG/Source/PCG/Public/Data/PCGPointData.h b/Engine/Plugins/Experimental/PCG/Source/PCG/Public/Data/PCGPointData.h
--- a/Engine/Plugins/Experimental/PCG/Source/PCG/Public/Data/PCGPointData.h
+++ b/Engine/Plugins/Experimental/PCG/Source/PCG/Public/Data/PCGPointData.h
@@ -67,7 +67,7 @@ public:
 	typedef TOctree2<FPCGPointRef, FPCGPointRefSemantics> PointOctree;
 
 	// ~Begin UPCGData interface
-	virtual EPCGDataType GetDataType() const override { return EPCGDataType::Point | Super::GetDataType(); }
+	virtual EPCGDataType GetDataType() const override { return EPCGDataType::Point; }
 	// ~End UPCGData interface
 
 	// ~Begin UPCGSpatialData interface
diff --git a/Engine/Plugins/Experimental/PCG/Source/PCG/Public/Data/PCGPolyLineData.h b/Engine/Plugins/Experimental/PCG/Source/PCG/Public/Data/PCGPolyLineData.h
--- a/Engine/Plugins/Experimental/PCG/Source/PCG/Public/Data/PCGPolyLineData.h
+++ b/Engine/Plugins/Experimental/PCG/Source/PCG/Public/Data/PCGPolyLineData.h
@@ -14,7 +14,7 @@ class PCG_API UPCGPolyLineData : public UPCGSpatialDataWithPointCache
 
 public:
 	// ~Begin UPCGData interface
-	virtual EPCGDataType GetDataType() const override { return EPCGDataType::PolyLine | Super::GetDataType(); }
+	virtual EPCGDataType GetDataType() const override { return EPCGDataType::PolyLine; }
 	// ~End UPCGData interface
 
 	//~Begin UPCGSpatialData interface
diff --git a/Engine/Plugins/Experimental/PCG/Source/PCG/Public/Data/PCGPrimitiveData.h b/Engine/Plugins/Experimental/PCG/Source/PCG/Public/Data/PCGPrimitiveData.h
--- a/Engine/Plugins/Experimental/PCG/Source/PCG/Public/Data/PCGPrimitiveData.h
+++ b/Engine/Plugins/Experimental/PCG/Source/PCG/Public/Data/PCGPrimitiveData.h
@@ -15,7 +15,7 @@ public:
 	void Initialize(UPrimitiveComponent* InPrim);
 
 	// ~Begin UPCGData interface
-	virtual EPCGDataType GetDataType() const override { return EPCGDataType::Primitive | Super::GetDataType(); }
+	virtual EPCGDataType GetDataType() const override { return EPCGDataType::Primitive; }
 	// ~End UPCGData interface
 
 	// ~Begin UPCGSpatialData interface
diff --git a/Engine/Plugins/Experimental/PCG/Source/PCG/Public/Data/PCGRenderTargetData.h b/Engine/Plugins/Experimental/PCG/Source/PCG/Public/Data/PCGRenderTargetData.h
--- a/Engine/Plugins/Experimental/PCG/Source/PCG/Public/Data/PCGRenderTargetData.h
+++ b/Engine/Plugins/Experimental/PCG/Source/PCG/Public/Data/PCGRenderTargetData.h
@@ -16,7 +16,7 @@ class PCG_API UPCGRenderTargetData : public UPCGBaseTextureData
 
 public:
 	// ~Begin UPCGData interface
-	virtual EPCGDataType GetDataType() const override { return EPCGDataType::RenderTarget | Super::GetDataType(); }
+	virtual EPCGDataType GetDataType() const override { return EPCGDataType::RenderTarget; }
 	// ~End UPCGData interface
 
 
diff --git a/Engine/Plugins/Experimental/PCG/Source/PCG/Public/Data/PCGSpatialData.h b/Engine/Plugins/Experimental/PCG/Source/PCG/Public/Data/PCGSpatialData.h
--- a/Engine/Plugins/Experimental/PCG/Source/PCG/Public/Data/PCGSpatialData.h
+++ b/Engine/Plugins/Experimental/PCG/Source/PCG/Public/Data/PCGSpatialData.h
@@ -37,7 +37,7 @@ public:
 	UPCGSpatialData(const FObjectInitializer& ObjectInitializer);
 
 	// ~Begin UPCGData interface
-	virtual EPCGDataType GetDataType() const override { return EPCGDataType::Spatial | Super::GetDataType(); }
+	virtual EPCGDataType GetDataType() const override { return EPCGDataType::Spatial; }
 	// ~End UPCGData interface
 
 	/** Virtual call to allocate a new spacial data object, duplicate this spatial data into
diff --git a/Engine/Plugins/Experimental/PCG/Source/PCG/Public/Data/PCGSplineData.h b/Engine/Plugins/Experimental/PCG/Source/PCG/Public/Data/PCGSplineData.h
--- a/Engine/Plugins/Experimental/PCG/Source/PCG/Public/Data/PCGSplineData.h
+++ b/Engine/Plugins/Experimental/PCG/Source/PCG/Public/Data/PCGSplineData.h
@@ -22,7 +22,7 @@ public:
 	void Initialize(USplineComponent* InSpline);
 
 	// ~Begin UPCGData interface
-	virtual EPCGDataType GetDataType() const override { return EPCGDataType::Spline | Super::GetDataType(); }
+	virtual EPCGDataType GetDataType() const override { return EPCGDataType::Spline; }
 	// ~End UPCGData interface
 
 	//~Begin UPCGPolyLineData interface
diff --git a/Engine/Plugins/Experimental/PCG/Source/PCG/Public/Data/PCGSurfaceData.h b/Engine/Plugins/Experimental/PCG/Source/PCG/Public/Data/PCGSurfaceData.h
--- a/Engine/Plugins/Experimental/PCG/Source/PCG/Public/Data/PCGSurfaceData.h
+++ b/Engine/Plugins/Experimental/PCG/Source/PCG/Public/Data/PCGSurfaceData.h
@@ -14,7 +14,7 @@ class PCG_API UPCGSurfaceData : public UPCGSpatialDataWithPointCache
 
 public:
 	// ~Begin UPCGData interface
-	virtual EPCGDataType GetDataType() const override { return EPCGDataType::Surface | Super::GetDataType(); }
+	virtual EPCGDataType GetDataType() const override { return EPCGDataType::Surface; }
 	// ~End UPCGData interface
 
 	//~Begin UPCGSpatialData interface
diff --git a/Engine/Plugins/Experimental/PCG/Source/PCG/Public/Data/PCGTextureData.h b/Engine/Plugins/Experimental/PCG/Source/PCG/Public/Data/PCGTextureData.h
--- a/Engine/Plugins/Experimental/PCG/Source/PCG/Public/Data/PCGTextureData.h
+++ b/Engine/Plugins/Experimental/PCG/Source/PCG/Public/Data/PCGTextureData.h
@@ -30,7 +30,7 @@ class PCG_API UPCGBaseTextureData : public UPCGSurfaceData
 
 public:
 	// ~Begin UPCGData interface
-	virtual EPCGDataType GetDataType() const override { return EPCGDataType::BaseTexture | Super::GetDataType(); }
+	virtual EPCGDataType GetDataType() const override { return EPCGDataType::BaseTexture; }
 	// ~End UPCGData interface
 
 	//~Begin UPCGSpatialData interface
@@ -99,7 +99,7 @@ class PCG_API UPCGTextureData : public UPCGBaseTextureData
 
 public:
 	// ~Begin UPCGData interface
-	virtual EPCGDataType GetDataType() const override { return EPCGDataType::Texture | Super::GetDataType(); }
+	virtual EPCGDataType GetDataType() const override { return EPCGDataType::Texture; }
 	// ~End UPCGData interface
 
 	//~Begin UPCGSpatialData interface
diff --git a/Engine/Plugins/Experimental/PCG/Source/PCG/Public/Data/PCGVolumeData.h b/Engine/Plugins/Experimental/PCG/Source/PCG/Public/Data/PCGVolumeData.h
--- a/Engine/Plugins/Experimental/PCG/Source/PCG/Public/Data/PCGVolumeData.h
+++ b/Engine/Plugins/Experimental/PCG/Source/PCG/Public/Data/PCGVolumeData.h
@@ -19,7 +19,7 @@ public:
 	void Initialize(const FBox& InBounds, AActor* InTargetActor);
 
 	// ~Begin UPCGData interface
-	virtual EPCGDataType GetDataType() const override { return EPCGDataType::Volume | Super::GetDataType(); }
+	virtual EPCGDataType GetDataType() const override { return EPCGDataType::Volume; }
 	// ~End UPCGData interface
 
 	// ~Begin UPGCSpatialData interface
diff --git a/Engine/Plugins/Experimental/PCG/Source/PCG/Public/Data/PCGWorldData.h b/Engine/Plugins/Experimental/PCG/Source/PCG/Public/Data/PCGWorldData.h
--- a/Engine/Plugins/Experimental/PCG/Source/PCG/Public/Data/PCGWorldData.h
+++ b/Engine/Plugins/Experimental/PCG/Source/PCG/Public/Data/PCGWorldData.h
@@ -108,10 +108,6 @@ class PCG_API UPCGWorldRayHitData : public UPCGSurfaceData
 public:
 	void Initialize(UWorld* InWorld, const FBox& InBounds = FBox(EForceInit::ForceInit));
 
-	//~Begin UPCGData interface
-	virtual EPCGDataType GetDataType() const override { return Super::GetDataType(); }
-	//~End UPCGData interface
-
 	//~Begin UPCGSpatialData interface
 	virtual FBox GetBounds() const override { return Bounds; }
 	virtual FBox GetStrictBounds() const override { return Bounds; }
diff --git a/Engine/Plugins/Experimental/PCG/Source/PCG/Public/Elements/PCGIntersectionElement.h b/Engine/Plugins/Experimental/PCG/Source/PCG/Public/Elements/PCGIntersectionElement.h
--- a/Engine/Plugins/Experimental/PCG/Source/PCG/Public/Elements/PCGIntersectionElement.h
+++ b/Engine/Plugins/Experimental/PCG/Source/PCG/Public/Elements/PCGIntersectionElement.h
@@ -23,6 +23,9 @@ public:
 
 	virtual TArray<FPCGPinProperties> OutputPinProperties() const override;
 
+	// If node disabled, don't intersect - pass through first edge
+	virtual bool OnlyPassThroughOneEdgeWhenDisabled() const { return true; }
+
 protected:
 	virtual FPCGElementPtr CreateElement() const override;
 	//~End UPCGSettings interface
diff --git a/Engine/Plugins/Experimental/PCG/Source/PCG/Public/Elements/PCGMergeElement.h b/Engine/Plugins/Experimental/PCG/Source/PCG/Public/Elements/PCGMergeElement.h
--- a/Engine/Plugins/Experimental/PCG/Source/PCG/Public/Elements/PCGMergeElement.h
+++ b/Engine/Plugins/Experimental/PCG/Source/PCG/Public/Elements/PCGMergeElement.h
@@ -24,6 +24,9 @@ public:
 	virtual TArray<FPCGPinProperties> InputPinProperties() const override;
 	virtual TArray<FPCGPinProperties> OutputPinProperties() const override;
 
+	// If node disabled, don't merge - pass through first edge
+	virtual bool OnlyPassThroughOneEdgeWhenDisabled() const { return true; }
+
 protected:
 	virtual FPCGElementPtr CreateElement() const override;
 	//~End UPCGSettings interface
diff --git a/Engine/Plugins/Experimental/PCG/Source/PCG/Public/Elements/PCGUnionElement.h b/Engine/Plugins/Experimental/PCG/Source/PCG/Public/Elements/PCGUnionElement.h
--- a/Engine/Plugins/Experimental/PCG/Source/PCG/Public/Elements/PCGUnionElement.h
+++ b/Engine/Plugins/Experimental/PCG/Source/PCG/Public/Elements/PCGUnionElement.h
@@ -23,6 +23,9 @@ public:
 
 	virtual TArray<FPCGPinProperties> OutputPinProperties() const override;
 
+	// If node disabled, don't union - pass through first edge
+	virtual bool OnlyPassThroughOneEdgeWhenDisabled() const { return true; }
+
 protected:
 	virtual FPCGElementPtr CreateElement() const override;
 	//~End UPCGSettings interface
diff --git a/Engine/Plugins/Experimental/PCG/Source/PCG/Public/PCGNode.h b/Engine/Plugins/Experimental/PCG/Source/PCG/Public/PCGNode.h
--- a/Engine/Plugins/Experimental/PCG/Source/PCG/Public/PCGNode.h
+++ b/Engine/Plugins/Experimental/PCG/Source/PCG/Public/PCGNode.h
@@ -100,6 +100,12 @@ public:
 	/** Pin from which data is passed through when this node is disabled. */
 	virtual const UPCGPin* GetPassThroughInputPin() const;
 
+	/** True if the pin is being used by the node. UI will gray out unused pins. */
+	virtual bool IsPinUsedByNodeExecution(const UPCGPin* InPin) const;
+
+	/** True if the edge is being used by the node. UI will gray out unused pins. */
+	virtual bool IsEdgeUsedByNodeExecution(const UPCGEdge* InEdge) const;
+
 	const TArray<TObjectPtr<UPCGPin>>& GetInputPins() const { return InputPins; }
 	const TArray<TObjectPtr<UPCGPin>>& GetOutputPins() const { return OutputPins; }
 
diff --git a/Engine/Plugins/Experimental/PCG/Source/PCG/Public/PCGParamData.h b/Engine/Plugins/Experimental/PCG/Source/PCG/Public/PCGParamData.h
--- a/Engine/Plugins/Experimental/PCG/Source/PCG/Public/PCGParamData.h
+++ b/Engine/Plugins/Experimental/PCG/Source/PCG/Public/PCGParamData.h
@@ -19,7 +19,7 @@ public:
 	UPCGParamData(const FObjectInitializer& ObjectInitializer);
 
 	// ~Begin UPCGData interface
-	virtual EPCGDataType GetDataType() const override { return EPCGDataType::Param | Super::GetDataType(); }
+	virtual EPCGDataType GetDataType() const override { return EPCGDataType::Param; }
 	// ~End UPCGData interface
 
 	UFUNCTION(BlueprintCallable, Category = Metadata)
diff --git a/Engine/Plugins/Experimental/PCG/Source/PCG/Public/PCGSettings.h b/Engine/Plugins/Experimental/PCG/Source/PCG/Public/PCGSettings.h
--- a/Engine/Plugins/Experimental/PCG/Source/PCG/Public/PCGSettings.h
+++ b/Engine/Plugins/Experimental/PCG/Source/PCG/Public/PCGSettings.h
@@ -83,7 +83,7 @@ class PCG_API UPCGSettings : public UPCGSettingsInterface
 
 public:
 	// ~Begin UPCGData interface
-	virtual EPCGDataType GetDataType() const override { return EPCGDataType::Settings | Super::GetDataType(); }
+	virtual EPCGDataType GetDataType() const override { return EPCGDataType::Settings; }
 	// ~End UPCGData interface
 
 	// ~Begin UPCGSettingsInterface interface
@@ -127,6 +127,12 @@ public:
 	/** Derived classes can implement this to expose additional name information in the logs */
 	virtual FName AdditionalTaskName() const { return NAME_None; }
 
+	/** Returns true if InPin is in use by node (assuming node enabled). Can be used to communicate when a pin is not in use to user. */
+	virtual bool IsPinUsedByNodeExecution(const UPCGPin* InPin) const { return true; }
+
+	/** Returns true if only the first input edge is used from the primary pin when the node is disabled. */
+	virtual bool OnlyPassThroughOneEdgeWhenDisabled() const { return false; }
+
 	UPROPERTY(EditAnywhere, BlueprintReadWrite, Category = Settings, meta=(EditCondition=bUseSeed, EditConditionHides))
 	int Seed = 0xC35A9631; // random prime number
 
diff --git a/Engine/Plugins/Experimental/PCG/Source/PCGEditor/Private/PCGEditorGraphSchema.cpp b/Engine/Plugins/Experimental/PCG/Source/PCGEditor/Private/PCGEditorGraphSchema.cpp
--- a/Engine/Plugins/Experimental/PCG/Source/PCGEditor/Private/PCGEditorGraphSchema.cpp
+++ b/Engine/Plugins/Experimental/PCG/Source/PCGEditor/Private/PCGEditorGraphSchema.cpp
@@ -2,9 +2,11 @@
 
 #include "PCGEditorGraphSchema.h"
 
-#include "Elements/PCGExecuteBlueprint.h"
+#include "PCGEdge.h"
 #include "PCGGraph.h"
+#include "PCGPin.h"
 #include "PCGSettings.h"
+#include "Elements/PCGExecuteBlueprint.h"
 
 #include "PCGEditorCommon.h"
 #include "PCGEditorGraph.h"
@@ -13,6 +15,7 @@
 #include "PCGEditorSettings.h"
 #include "PCGEditorUtils.h"
 
+#include "Algo/Find.h"
 #include "AssetRegistry/AssetData.h"
 #include "Engine/Blueprint.h"
 #include "Framework/Application/SlateApplication.h"
@@ -420,23 +423,25 @@ void FPCGEditorConnectionDrawingPolicy::DetermineWiringStyle(UEdGraphPin* Output
 	}
 
 	// Desaturate and connection if the node is disabled and the data on this wire won't be used
-	if (InputPin)
+	if (InputPin && OutputPin)
 	{
-		if (const UPCGEditorGraphNodeBase* EditorNode = CastChecked<const UPCGEditorGraphNodeBase>(InputPin->GetOwningNode()))
+		const UPCGEditorGraphNodeBase* EditorNode = CastChecked<const UPCGEditorGraphNodeBase>(InputPin->GetOwningNode());
+		const UPCGNode* PCGNode = EditorNode ? EditorNode->GetPCGNode() : nullptr;
+		const UPCGPin* PCGPin = PCGNode ? PCGNode->GetInputPin(InputPin->GetFName()) : nullptr;
+		const UPCGEditorGraphNodeBase* UpstreamEditorNode = CastChecked<const UPCGEditorGraphNodeBase>(OutputPin->GetOwningNode());
+
+		if (PCGPin && UpstreamEditorNode)
 		{
-			const UPCGNode* PCGNode = EditorNode->GetPCGNode();
-			if (PCGNode && PCGNode->GetSettings() && !PCGNode->GetSettings()->bEnabled)
+			// Look for the PCG edge that correlates with passed in (OutputPin, InputPin) edge
+			const TObjectPtr<UPCGEdge>* PCGEdge = PCGPin->Edges.FindByPredicate([UpstreamEditorNode, OutputPin](const UPCGEdge* ConnectedPCGEdge)
+			{
+				return UpstreamEditorNode->GetPCGNode() == ConnectedPCGEdge->InputPin->Node && ConnectedPCGEdge->InputPin->Properties.Label == OutputPin->GetFName();
+			});
+
+			// If edge found and is not used, gray it out
+			if (PCGEdge && !PCGNode->IsEdgeUsedByNodeExecution(*PCGEdge))
 			{
-				const UPCGPin* PassThroughInputPin = PCGNode->GetPassThroughInputPin();
-
-				// If this wire does not connect to the pass-through pin, or if the wire is not the first link on that pin, then it will
-				// not be used by the node. In this case desaturate the wire.
-				const bool bConnectedToPassThrough = PassThroughInputPin && PassThroughInputPin->Properties.Label == InputPin->GetFName();
-				const bool bFirstWire = InputPin->LinkedTo[0] == OutputPin;
-				if (!bConnectedToPassThrough || !bFirstWire)
-				{
-					Params.WireColor = Params.WireColor.Desaturate(0.7f);
-				}
+				Params.WireColor = Params.WireColor.Desaturate(0.7f);
 			}
 		}
 	}
diff --git a/Engine/Plugins/Experimental/PCG/Source/PCGEditor/Private/SPCGEditorGraphNode.cpp b/Engine/Plugins/Experimental/PCG/Source/PCGEditor/Private/SPCGEditorGraphNode.cpp
--- a/Engine/Plugins/Experimental/PCG/Source/PCGEditor/Private/SPCGEditorGraphNode.cpp
+++ b/Engine/Plugins/Experimental/PCG/Source/PCGEditor/Private/SPCGEditorGraphNode.cpp
@@ -12,6 +12,46 @@
 #include "Styling/SlateBrush.h"
 #include "Widgets/Text/SInlineEditableTextBlock.h"
 
+/** PCG pin primarily to give more control over pin coloring. */
+class SPCGEditorGraphNodePin : public SGraphPin
+{
+public:
+	SLATE_BEGIN_ARGS(SPCGEditorGraphNodePin) {}
+	SLATE_END_ARGS()
+
+	void Construct(const FArguments& InArgs, UEdGraphPin* InPin);
+
+	virtual FSlateColor GetPinColor() const override;
+};
+
+void SPCGEditorGraphNodePin::Construct(const FArguments& InArgs, UEdGraphPin* InPin)
+{
+	// Required to give the first argument.
+	SGraphPin::Construct(SGraphPin::FArguments(), InPin);
+}
+
+// Adapted from SGraphPin::GetPinColor
+FSlateColor SPCGEditorGraphNodePin::GetPinColor() const
+{
+	FSlateColor Color = SGraphPin::GetPinColor();
+
+	UEdGraphPin* GraphPin = GetPinObj();
+	if (GraphPin && !GraphPin->IsPendingKill())
+	{
+		const UPCGEditorGraphNodeBase* EditorNode = CastChecked<const UPCGEditorGraphNodeBase>(GraphPinObj->GetOwningNode());
+		const UPCGNode* PCGNode = EditorNode ? EditorNode->GetPCGNode() : nullptr;
+		const UPCGPin* PCGPin = PCGNode ? PCGNode->GetInputPin(GraphPin->GetFName()) : nullptr;
+
+		// Desaturate if pin is unused - intended to happen whether disabled or not
+		if (PCGPin && !PCGNode->IsPinUsedByNodeExecution(PCGPin))
+		{
+			Color = Color.GetSpecifiedColor().Desaturate(0.7f);
+		}
+	}
+
+	return Color;
+}
+
 void SPCGEditorGraphNode::Construct(const FArguments& InArgs, UPCGEditorGraphNodeBase* InNode)
 {
 	GraphNode = InNode;
@@ -54,6 +94,11 @@ TSharedRef<SWidget> SPCGEditorGraphNode::CreateTitleWidget(TSharedPtr<SNodeTitle
 	return InlineEditableText.ToSharedRef();
 }
 
+TSharedPtr<SGraphPin> SPCGEditorGraphNode::CreatePinWidget(UEdGraphPin* Pin) const
+{
+	return SNew(SPCGEditorGraphNodePin, Pin);
+}
+
 void SPCGEditorGraphNode::AddPin(const TSharedRef<SGraphPin>& PinToAdd)
 {
 	check(PCGEditorGraphNode);
diff --git a/Engine/Plugins/Experimental/PCG/Source/PCGEditor/Private/SPCGEditorGraphNode.h b/Engine/Plugins/Experimental/PCG/Source/PCGEditor/Private/SPCGEditorGraphNode.h
--- a/Engine/Plugins/Experimental/PCG/Source/PCGEditor/Private/SPCGEditorGraphNode.h
+++ b/Engine/Plugins/Experimental/PCG/Source/PCGEditor/Private/SPCGEditorGraphNode.h
@@ -20,6 +20,7 @@ public:
 	virtual void RequestRenameOnSpawn() override { /* Empty to avoid the default behavior to rename on node spawn */ }
 	virtual void AddPin(const TSharedRef<SGraphPin>& PinToAdd) override;
 	virtual TSharedRef<SWidget> CreateTitleWidget(TSharedPtr<SNodeTitle> InNodeTitle) override;
+	virtual TSharedPtr<SGraphPin> CreatePinWidget(UEdGraphPin* Pin) const;
 	//~ End SGraphNode Interface
 
 	//~ Begin SNodePanel::SNode Interface
