Made it possible to override the DataTable and RowName in UCGDataTableRowToParamData
Added a unit test for that
Also moved some of the unit test helpers out to a new header to assist with some other node tests

diff --git a/Engine/Plugins/Experimental/PCG/Source/PCG/Private/Elements/PCGDataTableRowToParamData.cpp b/Engine/Plugins/Experimental/PCG/Source/PCG/Private/Elements/PCGDataTableRowToParamData.cpp
--- a/Engine/Plugins/Experimental/PCG/Source/PCG/Private/Elements/PCGDataTableRowToParamData.cpp
+++ b/Engine/Plugins/Experimental/PCG/Source/PCG/Private/Elements/PCGDataTableRowToParamData.cpp
@@ -13,10 +13,18 @@ FPCGElementPtr UPCGDataTableRowToParamDataSettings::CreateElement() const
 	return MakeShared<FPCGDataTableRowToParamData>();
 }
 
+TArray<FPCGPinProperties> UPCGDataTableRowToParamDataSettings::InputPinProperties() const
+{
+	TArray<FPCGPinProperties> PinProperties;
+	PinProperties.Emplace(PCGPinConstants::DefaultParamsLabel, EPCGDataType::Param, false);
+
+	return PinProperties;
+}
+
 TArray<FPCGPinProperties> UPCGDataTableRowToParamDataSettings::OutputPinProperties() const
 {
 	TArray<FPCGPinProperties> PinProperties;
-	PinProperties.Emplace(PCGPinConstants::DefaultOutputLabel, EPCGDataType::Param);
+	PinProperties.Emplace(PCGPinConstants::DefaultOutputLabel, EPCGDataType::Param, false);
 
 	return PinProperties;
 }
@@ -30,21 +38,33 @@ bool FPCGDataTableRowToParamData::ExecuteInternal(FPCGContext* Context) const
 {
 	TRACE_CPUPROFILER_EVENT_SCOPE(FPCGDataTableRowToParamData::Execute);
 
+	UPCGParamData* Params = Context->InputData.GetParams();
+
 	const UPCGDataTableRowToParamDataSettings* Settings = Context->GetInputSettings<UPCGDataTableRowToParamDataSettings>();
 	check(Settings);
 
-	const UDataTable* DataTable = Settings->DataTable.Get();
+	const FString PathOverride = PCGSettingsHelpers::GetValue(GET_MEMBER_NAME_CHECKED(UPCGDataTableRowToParamDataSettings, DataTable), FString(), Params);
+	const FName RowName = PCGSettingsHelpers::GetValue(GET_MEMBER_NAME_CHECKED(UPCGDataTableRowToParamDataSettings, RowName), Settings->RowName, Params);
+
+	TSoftObjectPtr<UDataTable> DataTablePtr = Settings->DataTable;
+
+	if (!PathOverride.IsEmpty())
+	{
+		DataTablePtr = FSoftObjectPath(PathOverride);
+	}
+
+	const UDataTable* DataTable = DataTablePtr.LoadSynchronous();
 	if (!DataTable)
 	{
 		PCGE_LOG(Error, TEXT("FPCGDataTableRowToParamData: Data table is invalid"));
 		return true;
 	}
 
-	const uint8* const* RowDataItr = DataTable->GetRowMap().Find(Settings->RowName);
+	const uint8* const* RowDataItr = DataTable->GetRowMap().Find(RowName);
 	const uint8* RowDataPtr = RowDataItr ? *RowDataItr : nullptr;
 	if (!RowDataPtr)
 	{
-		PCGE_LOG(Error, TEXT("FPCGDataTableRowToParamData: Data table '%s' has no row named '%s'"), *DataTable->GetPathName(), *Settings->RowName.ToString());
+		PCGE_LOG(Error, TEXT("FPCGDataTableRowToParamData: Data table '%s' has no row named '%s'"), *DataTable->GetPathName(), *RowName.ToString());
 		return true;
 	}
 
@@ -71,4 +91,3 @@ bool FPCGDataTableRowToParamData::ExecuteInternal(FPCGContext* Context) const
 
 	return true;
 }
-
diff --git a/Engine/Plugins/Experimental/PCG/Source/PCG/Private/Tests/Elements/PCGDataTableRowToParamDataTest.cpp b/Engine/Plugins/Experimental/PCG/Source/PCG/Private/Tests/Elements/PCGDataTableRowToParamDataTest.cpp
--- a/Engine/Plugins/Experimental/PCG/Source/PCG/Private/Tests/Elements/PCGDataTableRowToParamDataTest.cpp
+++ b/Engine/Plugins/Experimental/PCG/Source/PCG/Private/Tests/Elements/PCGDataTableRowToParamDataTest.cpp
@@ -10,54 +10,11 @@
 
 #include "Elements/PCGDataTableRowToParamData.h"
 #include "Tests/Elements/PCGDataTableRowToParamDataTestStruct.h"
+#include "Tests/Elements/PCGDataTableRowToParamDataTestHelpers.h"
 #include "Engine/DataTable.h"
 
 IMPLEMENT_CUSTOM_SIMPLE_AUTOMATION_TEST(FPCGDataTableRowToParamDataTest, FPCGTestBaseClass, "pcg.tests.DataTableRowToParamData.Basic", PCGTestsCommon::TestFlags)
 
-namespace PCGDataTableRowToParamDataTest
-{
-	template<typename StructType, typename PCGType>
-	bool TestAttribute(FPCGTestBaseClass& Self, void *Data, const UPCGParamData& Params, FName FieldName, PCGType ExpectedValue)
-	{
-		const FPCGMetadataAttribute<PCGType>* Attribute = static_cast<const FPCGMetadataAttribute<PCGType>*>(Params.Metadata->GetConstAttribute(FieldName));
-
-		if (!Self.TestNotNull(*FString::Printf(TEXT("'%s' Metadata attribute"), *FieldName.ToString()), Attribute))
-		{
-			return false;
-		}
-
-		if (!Self.TestEqual(*FString::Printf(TEXT("'%s' Metadata type"), *FieldName.ToString()), Attribute->GetTypeId(), PCG::Private::MetadataTypes<PCGType>::Id))
-		{
-			return false;
-		}
-
-		const FProperty* Property = FindFProperty<FProperty>(FPCGDataTableRowToParamDataTestStruct::StaticStruct(), FieldName);
-		if (!Self.TestNotNull(*FString::Printf(TEXT("'%s' Property not found"), *FieldName.ToString()), Property))
-		{
-			return false;
-		}
-
-		const StructType* RowValuePtr = Property->ContainerPtrToValuePtr<StructType>(Data);
-		if (!Self.TestNotNull(*FString::Printf(TEXT("'%s' Row Value not found"), *FieldName.ToString()), RowValuePtr))
-		{
-			return false;
-		}
-
-		if (!Self.TestEqual(*FString::Printf(TEXT("'%s' Row Value"), *FieldName.ToString()), *RowValuePtr, StructType(ExpectedValue)))
-		{
-			return false;
-		}
-		
-		PCGType Value = Attribute->GetValueFromItemKey(0);
-		if (!Self.TestEqual(*FString::Printf(TEXT("'%s' Attribute Value"), *FieldName.ToString()), Value, ExpectedValue))
-		{
-			return false;
-		}
-
-		return true;
-	}
-}
-
 bool FPCGDataTableRowToParamDataTest::RunTest(const FString& Parameters)
 {
 	PCGTestsCommon::FTestData TestData(PCGDeterminismTests::Defaults::Seed);
@@ -95,16 +52,66 @@ CCC,"C Name","C String","333","3333","3.3","3.33","(X=3.0,Y=3.0)","(X=3.0,Y=3.0,
 
 	// first template type is the value of the row struct member, second is the pcg attribute type
 	// this is because sometimes conversion occurs as PCG can't represent all types
-	bSuccess |= PCGDataTableRowToParamDataTest::TestAttribute<FName,FName>(*this, RowData, *Params, TEXT("Name"), TEXT("B Name"));
-	bSuccess |= PCGDataTableRowToParamDataTest::TestAttribute<FString,FString>(*this, RowData, *Params, TEXT("String"), TEXT("B String"));
-	bSuccess |= PCGDataTableRowToParamDataTest::TestAttribute<int32,int64>(*this, RowData, *Params, TEXT("I32"), 222);
-	bSuccess |= PCGDataTableRowToParamDataTest::TestAttribute<int64,int64>(*this, RowData, *Params, TEXT("I64"), 2222);
-	bSuccess |= PCGDataTableRowToParamDataTest::TestAttribute<float,double>(*this, RowData, *Params, TEXT("F32"), 2.2);
-	bSuccess |= PCGDataTableRowToParamDataTest::TestAttribute<double,double>(*this, RowData, *Params, TEXT("F64"), 2.22);
-	bSuccess |= PCGDataTableRowToParamDataTest::TestAttribute<FVector2D,FVector2D>(*this, RowData, *Params, TEXT("V2"), FVector2D(2.0, 2.0));
-	bSuccess |= PCGDataTableRowToParamDataTest::TestAttribute<FVector,FVector>(*this, RowData, *Params, TEXT("V3"), FVector(2.0, 2.0, 2.0));
-	bSuccess |= PCGDataTableRowToParamDataTest::TestAttribute<FVector4,FVector4>(*this, RowData, *Params, TEXT("V4"), FVector4(2.0, 2.0, 2.0, 2.0));
-	bSuccess |= PCGDataTableRowToParamDataTest::TestAttribute<FSoftObjectPath,FString>(*this, RowData, *Params, TEXT("SoftPath"), TEXT("/Script/PCG"));
+	bSuccess |= PCGDataTableRowToParamDataTestHelpers::TestAttribute<FName,FName>(*this, RowData, *Params, TEXT("Name"), TEXT("B Name"));
+	bSuccess |= PCGDataTableRowToParamDataTestHelpers::TestAttribute<FString,FString>(*this, RowData, *Params, TEXT("String"), TEXT("B String"));
+	bSuccess |= PCGDataTableRowToParamDataTestHelpers::TestAttribute<int32,int64>(*this, RowData, *Params, TEXT("I32"), 222);
+	bSuccess |= PCGDataTableRowToParamDataTestHelpers::TestAttribute<int64,int64>(*this, RowData, *Params, TEXT("I64"), 2222);
+	bSuccess |= PCGDataTableRowToParamDataTestHelpers::TestAttribute<float,double>(*this, RowData, *Params, TEXT("F32"), 2.2);
+	bSuccess |= PCGDataTableRowToParamDataTestHelpers::TestAttribute<double,double>(*this, RowData, *Params, TEXT("F64"), 2.22);
+	bSuccess |= PCGDataTableRowToParamDataTestHelpers::TestAttribute<FVector2D,FVector2D>(*this, RowData, *Params, TEXT("V2"), FVector2D(2.0, 2.0));
+	bSuccess |= PCGDataTableRowToParamDataTestHelpers::TestAttribute<FVector,FVector>(*this, RowData, *Params, TEXT("V3"), FVector(2.0, 2.0, 2.0));
+	bSuccess |= PCGDataTableRowToParamDataTestHelpers::TestAttribute<FVector4,FVector4>(*this, RowData, *Params, TEXT("V4"), FVector4(2.0, 2.0, 2.0, 2.0));
+	bSuccess |= PCGDataTableRowToParamDataTestHelpers::TestAttribute<FSoftObjectPath,FString>(*this, RowData, *Params, TEXT("SoftPath"), TEXT("/Script/PCG"));
 
 	return bSuccess;
 }
+
+IMPLEMENT_CUSTOM_SIMPLE_AUTOMATION_TEST(FPCGDataTableRowToParamDataTest_ParamsOverride, FPCGTestBaseClass, "pcg.tests.DataTableRowToParamData.ParamsOverride", PCGTestsCommon::TestFlags)
+
+bool FPCGDataTableRowToParamDataTest_ParamsOverride::RunTest(const FString& Parameters)
+{
+	PCGTestsCommon::FTestData TestData(PCGDeterminismTests::Defaults::Seed);
+	PCGDeterminismTests::GenerateSettings<UPCGDataTableRowToParamDataSettings>(TestData);
+	UPCGDataTableRowToParamDataSettings* Settings = CastChecked<UPCGDataTableRowToParamDataSettings>(TestData.Settings);
+
+	UDataTable *TestDataTable = NewObject<UDataTable>();
+	TestDataTable->RowStruct = FPCGDataTableRowToParamDataTestStruct::StaticStruct();
+	TArray<FString> Errors = TestDataTable->CreateTableFromCSVString(R"CSV(---,Name,String,I32,I64,F32,F64,V2,V3,V4,SoftPath
+AAA,"A Name","A String","111","1111","1.1","1.11","(X=1.0,Y=1.0)","(X=1.0,Y=1.0,Z=1.0)","(X=1.0,Y=1.0,Z=1.0,W=1.0)","/Script/PCG"
+BBB,"B Name","B String","222","2222","2.2","2.22","(X=2.0,Y=2.0)","(X=2.0,Y=2.0,Z=2.0)","(X=2.0,Y=2.0,Z=2.0,W=2.0)","/Script/PCG"
+CCC,"C Name","C String","333","3333","3.3","3.33","(X=3.0,Y=3.0)","(X=3.0,Y=3.0,Z=3.0)","(X=3.0,Y=3.0,Z=3.0,W=3.0)","/Script/PCG")CSV");
+
+	UTEST_EQUAL(*FString::Printf(TEXT("CSV Errors:\n%s"), *FString::Join(Errors, TEXT("\n"))), Errors.Num(), 0);
+
+	const FName RowNameOverride = TEXT("BBB");
+
+	{
+		UPCGParamData* InputParamData = NewObject<UPCGParamData>();
+		InputParamData->MutableMetadata()->CreateStringAttribute(TEXT("DataTable"), TestDataTable->GetPathName(), false);
+		InputParamData->MutableMetadata()->CreateStringAttribute(TEXT("RowName"), RowNameOverride.ToString(), false);
+		FPCGTaggedData& ParamInput = TestData.InputData.TaggedData.Emplace_GetRef();
+		ParamInput.Data = InputParamData;
+		ParamInput.Pin = PCGPinConstants::DefaultParamsLabel;
+	}
+
+	FPCGElementPtr TestElement = TestData.Settings->GetElement();
+
+	TUniquePtr<FPCGContext> Context(TestElement->Initialize(TestData.InputData, TestData.TestPCGComponent, nullptr));
+	Context->NumAvailableTasks = 1;
+
+	while (!TestElement->Execute(Context.Get())) {}
+
+	const UPCGParamData* Params = Context->OutputData.GetParams();
+	UTEST_NOT_NULL("Output params", Params);
+
+	UTEST_NOT_NULL("Params has metadata", Params->Metadata.Get());
+
+	void* RowData = TestDataTable->FindRow<FPCGDataTableRowToParamDataTestStruct>(RowNameOverride, TEXT("FPCGDataTableRowToParamDataTest"));
+	UTEST_NOT_NULL(*FString::Printf(TEXT("'%s' Row not found in data table"), *RowNameOverride.ToString()), RowData);
+
+	bool bSuccess = true;
+
+	bSuccess |= PCGDataTableRowToParamDataTestHelpers::TestAttribute<FName,FName>(*this, RowData, *Params, TEXT("Name"), TEXT("B Name"));
+
+	return bSuccess;
+}
\ No newline at end of file
diff --git a/Engine/Plugins/Experimental/PCG/Source/PCG/Private/Tests/Elements/PCGDataTableRowToParamDataTestHelpers.h b/Engine/Plugins/Experimental/PCG/Source/PCG/Private/Tests/Elements/PCGDataTableRowToParamDataTestHelpers.h
new file mode 100644
--- /dev/null
+++ b/Engine/Plugins/Experimental/PCG/Source/PCG/Private/Tests/Elements/PCGDataTableRowToParamDataTestHelpers.h
@@ -0,0 +1,49 @@
+// Copyright Epic Games, Inc. All Rights Reserved.
+
+#pragma once
+
+#include "CoreMinimal.h"
+
+namespace PCGDataTableRowToParamDataTestHelpers
+{
+	template<typename StructType, typename PCGType>
+	bool TestAttribute(FPCGTestBaseClass& Self, void *Data, const UPCGParamData& Params, FName FieldName, PCGType ExpectedValue)
+	{
+		const FPCGMetadataAttribute<PCGType>* Attribute = static_cast<const FPCGMetadataAttribute<PCGType>*>(Params.Metadata->GetConstAttribute(FieldName));
+
+		if (!Self.TestNotNull(*FString::Printf(TEXT("'%s' Metadata attribute"), *FieldName.ToString()), Attribute))
+		{
+			return false;
+		}
+
+		if (!Self.TestEqual(*FString::Printf(TEXT("'%s' Metadata type"), *FieldName.ToString()), Attribute->GetTypeId(), PCG::Private::MetadataTypes<PCGType>::Id))
+		{
+			return false;
+		}
+
+		const FProperty* Property = FindFProperty<FProperty>(FPCGDataTableRowToParamDataTestStruct::StaticStruct(), FieldName);
+		if (!Self.TestNotNull(*FString::Printf(TEXT("'%s' Property not found"), *FieldName.ToString()), Property))
+		{
+			return false;
+		}
+
+		const StructType* RowValuePtr = Property->ContainerPtrToValuePtr<StructType>(Data);
+		if (!Self.TestNotNull(*FString::Printf(TEXT("'%s' Row Value not found"), *FieldName.ToString()), RowValuePtr))
+		{
+			return false;
+		}
+
+		if (!Self.TestEqual(*FString::Printf(TEXT("'%s' Row Value"), *FieldName.ToString()), *RowValuePtr, StructType(ExpectedValue)))
+		{
+			return false;
+		}
+		
+		PCGType Value = Attribute->GetValueFromItemKey(0);
+		if (!Self.TestEqual(*FString::Printf(TEXT("'%s' Attribute Value"), *FieldName.ToString()), Value, ExpectedValue))
+		{
+			return false;
+		}
+
+		return true;
+	}
+}
\ No newline at end of file
diff --git a/Engine/Plugins/Experimental/PCG/Source/PCG/Public/Elements/PCGDataTableRowToParamData.h b/Engine/Plugins/Experimental/PCG/Source/PCG/Public/Elements/PCGDataTableRowToParamData.h
--- a/Engine/Plugins/Experimental/PCG/Source/PCG/Public/Elements/PCGDataTableRowToParamData.h
+++ b/Engine/Plugins/Experimental/PCG/Source/PCG/Public/Elements/PCGDataTableRowToParamData.h
@@ -26,7 +26,7 @@ public:
 	virtual FName AdditionalTaskName() const override;
 
 	virtual TArray<FPCGPinProperties> OutputPinProperties() const override;
-	virtual TArray<FPCGPinProperties> InputPinProperties() const override { return TArray<FPCGPinProperties>(); }
+	virtual TArray<FPCGPinProperties> InputPinProperties() const override;
 
 	// The name of the row to copy from
 	UPROPERTY(BlueprintReadWrite, EditAnywhere, Category = Settings)
@@ -34,7 +34,7 @@ public:
 
 	// the data table to copy from
 	UPROPERTY(BlueprintReadWrite, EditAnywhere, Category = Settings)
-	TObjectPtr<UDataTable> DataTable;
+	TSoftObjectPtr<UDataTable> DataTable;
 
 protected:
 	virtual FPCGElementPtr CreateElement() const override;
@@ -50,4 +50,3 @@ public:
 protected:
 	virtual bool ExecuteInternal(FPCGContext* Context) const override;
 };
-
