Exposed a bunch of APIs from PCG dll so they can be called in plugins

diff --git a/Engine/Plugins/Experimental/PCG/Source/PCG/Public/Helpers/PCGMetadataHelpers.h b/Engine/Plugins/Experimental/PCG/Source/PCG/Public/Helpers/PCGMetadataHelpers.h
--- a/Engine/Plugins/Experimental/PCG/Source/PCG/Public/Helpers/PCGMetadataHelpers.h
+++ b/Engine/Plugins/Experimental/PCG/Source/PCG/Public/Helpers/PCGMetadataHelpers.h
@@ -6,6 +6,6 @@ class UPCGMetadata;
 
 namespace PCGMetadataHelpers
 {
-	bool HasSameRoot(const UPCGMetadata* Metadata1, const UPCGMetadata* Metadata2);
-	const UPCGMetadata* GetParentMetadata(const UPCGMetadata* Metadata);
+	PCG_API bool HasSameRoot(const UPCGMetadata* Metadata1, const UPCGMetadata* Metadata2);
+	PCG_API const UPCGMetadata* GetParentMetadata(const UPCGMetadata* Metadata);
 }
\ No newline at end of file
diff --git a/Engine/Plugins/Experimental/PCG/Source/PCG/Public/Tests/PCGTestsCommon.h b/Engine/Plugins/Experimental/PCG/Source/PCG/Public/Tests/PCGTestsCommon.h
--- a/Engine/Plugins/Experimental/PCG/Source/PCG/Public/Tests/PCGTestsCommon.h
+++ b/Engine/Plugins/Experimental/PCG/Source/PCG/Public/Tests/PCGTestsCommon.h
@@ -39,27 +39,27 @@ namespace PCGTestsCommon
 		FRandomStream RandomStream;
 	};
 
-	AActor* CreateTemporaryActor();
-	UPCGPointData* CreateEmptyPointData();
-	UPCGParamData* CreateEmptyParamData();
+	PCG_API AActor* CreateTemporaryActor();
+	PCG_API UPCGPointData* CreateEmptyPointData();
+	PCG_API UPCGParamData* CreateEmptyParamData();
 
 	/** Creates a PointData with a single point at the origin */
-	UPCGPointData* CreatePointData();
+	PCG_API UPCGPointData* CreatePointData();
 	/** Creates a PointData with a single point at the provided location */
-	UPCGPointData* CreatePointData(const FVector& InLocation);
+	PCG_API UPCGPointData* CreatePointData(const FVector& InLocation);
 
 	/** Creates a PointData with PointCount many points, and randomizes the Transform and Color */
-	UPCGPointData* CreateRandomPointData(int32 PointCount, int32 Seed);
+	PCG_API UPCGPointData* CreateRandomPointData(int32 PointCount, int32 Seed);
 
-	UPCGPolyLineData* CreatePolyLineData();
-	UPCGSurfaceData* CreateSurfaceData();
-	UPCGVolumeData* CreateVolumeData(const FBox& InBounds = FBox::BuildAABB(FVector::ZeroVector, FVector::OneVector * 100));
-	UPCGPrimitiveData* CreatePrimitiveData();
+	PCG_API UPCGPolyLineData* CreatePolyLineData();
+	PCG_API UPCGSurfaceData* CreateSurfaceData();
+	PCG_API UPCGVolumeData* CreateVolumeData(const FBox& InBounds = FBox::BuildAABB(FVector::ZeroVector, FVector::OneVector * 100));
+	PCG_API UPCGPrimitiveData* CreatePrimitiveData();
 
-	TArray<FPCGDataCollection> GenerateAllowedData(const FPCGPinProperties& PinProperties);
+	PCG_API TArray<FPCGDataCollection> GenerateAllowedData(const FPCGPinProperties& PinProperties);
 
 	/** Validates that two Spatial Points are identical */
-	bool PointsAreIdentical(const FPCGPoint& FirstPoint, const FPCGPoint& SecondPoint);
+	PCG_API bool PointsAreIdentical(const FPCGPoint& FirstPoint, const FPCGPoint& SecondPoint);
 }
 
 class PCG_API FPCGTestBaseClass : public FAutomationTestBase
