Non-Unity build fix

diff --git a/Engine/Plugins/Experimental/PCG/Source/PCG/Public/Helpers/PCGMetadataHelpers.h b/Engine/Plugins/Experimental/PCG/Source/PCG/Public/Helpers/PCGMetadataHelpers.h
--- a/Engine/Plugins/Experimental/PCG/Source/PCG/Public/Helpers/PCGMetadataHelpers.h
+++ b/Engine/Plugins/Experimental/PCG/Source/PCG/Public/Helpers/PCGMetadataHelpers.h
@@ -2,6 +2,8 @@
 
 #pragma once
 
+#include "CoreMinimal.h"
+
 class UPCGMetadata;
 
 namespace PCGMetadataHelpers
