Expose additional metadata attribute existence test

#rnx

diff --git a/Engine/Plugins/Experimental/PCG/Source/PCG/Public/Metadata/PCGMetadata.h b/Engine/Plugins/Experimental/PCG/Source/PCG/Public/Metadata/PCGMetadata.h
--- a/Engine/Plugins/Experimental/PCG/Source/PCG/Public/Metadata/PCGMetadata.h
+++ b/Engine/Plugins/Experimental/PCG/Source/PCG/Public/Metadata/PCGMetadata.h
@@ -139,7 +139,10 @@ public:
 	FPCGMetadataAttributeBase* GetMutableAttribute(FName AttributeName);
 	const FPCGMetadataAttributeBase* GetConstAttribute(FName AttributeName) const;
 	const FPCGMetadataAttributeBase* GetConstAttributeById(int32 AttributeId) const;
+
+	UFUNCTION(BlueprintCallable, Category = "PCG|Metadata")
 	bool HasAttribute(FName AttributeName) const;
+
 	int32 GetAttributeCount() const;
 
 	template <typename T>
