[PCG] Change SelfPruning input/output types from Any to Point.

This node converts to points internally. In the future it may deal with more than just points, if so we can widen the types accordingly.

Add Param pin to catch first existing params connection (if any), and add deprecation to automatically set it up to not break existing graphs. Note that as per noding standard Params pin does not support multiple connections and only takes the first.

To support this, multiple changes made to node deprecation:

* Logic made less flimsy by asserting DataVersion is up to date after deprecation
* New deprecation hook added that is called before UpdatePins. This gives the node a chance to rearrange connections before UpdatePins removes any invalid connections. We can use this to find any connection to Params and move it to a params pin. (My other local change that splits the Surface Sampler pins also requires this deprecation hook).

diff --git a/Engine/Plugins/Experimental/PCG/Source/PCG/Private/Elements/PCGProjectionElement.cpp b/Engine/Plugins/Experimental/PCG/Source/PCG/Private/Elements/PCGProjectionElement.cpp
--- a/Engine/Plugins/Experimental/PCG/Source/PCG/Private/Elements/PCGProjectionElement.cpp
+++ b/Engine/Plugins/Experimental/PCG/Source/PCG/Private/Elements/PCGProjectionElement.cpp
@@ -108,7 +108,6 @@ bool FPCGProjectionElement::ExecuteInternal(FPCGContext* Context) const
 #if WITH_EDITOR
 void UPCGProjectionSettings::ApplyDeprecation(UPCGNode* InOutNode)
 {
-	Super::ApplyDeprecation(InOutNode);
 	check(InOutNode);
 
 	if (DataVersion < FPCGCustomVersion::SplitProjectionNodeInputs)
@@ -138,10 +137,8 @@ void UPCGProjectionSettings::ApplyDeprecation(UPCGNode* InOutNode)
 		}
 	}
 
-	// Signal that data is up to date at end of deprecation code.
-	DataVersion = FPCGCustomVersion::LatestVersion;
+	Super::ApplyDeprecation(InOutNode);
 }
 #endif // WITH_EDITOR
 
 #undef LOCTEXT_NAMESPACE
-
diff --git a/Engine/Plugins/Experimental/PCG/Source/PCG/Private/Elements/PCGSelfPruning.cpp b/Engine/Plugins/Experimental/PCG/Source/PCG/Private/Elements/PCGSelfPruning.cpp
--- a/Engine/Plugins/Experimental/PCG/Source/PCG/Private/Elements/PCGSelfPruning.cpp
+++ b/Engine/Plugins/Experimental/PCG/Source/PCG/Private/Elements/PCGSelfPruning.cpp
@@ -1,10 +1,14 @@
 // Copyright Epic Games, Inc. All Rights Reserved.
 
 #include "Elements/PCGSelfPruning.h"
+
+#include "PCGCustomVersion.h"
+#include "PCGEdge.h"
+#include "PCGHelpers.h"
 #include "Data/PCGPointData.h"
 #include "Data/PCGSpatialData.h"
 #include "Helpers/PCGSettingsHelpers.h"
-#include "PCGHelpers.h"
+
 #include "Math/RandomStream.h"
 
 #include UE_INLINE_GENERATED_CPP_BY_NAME(PCGSelfPruning)
@@ -242,6 +246,24 @@ namespace PCGSelfPruningElement
 	}
 }
 
+TArray<FPCGPinProperties> UPCGSelfPruningSettings::InputPinProperties() const
+{
+	TArray<FPCGPinProperties> PinProperties;
+	// TODO in the future type checking of edges will be stricter and a conversion node will be added to convert from other types
+	PinProperties.Emplace(PCGPinConstants::DefaultInputLabel, EPCGDataType::Point);
+	PinProperties.Emplace(PCGPinConstants::DefaultParamsLabel, EPCGDataType::Param, /*bInAllowMultipleConnections=*/false);
+
+	return PinProperties;
+}
+
+TArray<FPCGPinProperties> UPCGSelfPruningSettings::OutputPinProperties() const
+{
+	TArray<FPCGPinProperties> PinProperties;
+	PinProperties.Emplace(PCGPinConstants::DefaultOutputLabel, EPCGDataType::Point);
+
+	return PinProperties;
+}
+
 FPCGElementPtr UPCGSelfPruningSettings::CreateElement() const
 {
 	return MakeShared<FPCGSelfPruningElement>();
@@ -264,3 +286,54 @@ bool FPCGSelfPruningElement::ExecuteInternal(FPCGContext* Context) const
 
 	return true;
 }
+
+#if WITH_EDITOR
+void UPCGSelfPruningSettings::ApplyDeprecationBeforeUpdatePins(UPCGNode* InOutNode, TArray<TObjectPtr<UPCGPin>>& InputPins, TArray<TObjectPtr<UPCGPin>>& OutputPins)
+{
+	Super::ApplyDeprecationBeforeUpdatePins(InOutNode, InputPins, OutputPins);
+
+	check(InOutNode);
+
+	if (DataVersion < FPCGCustomVersion::MoveSelfPruningParamsOffFirstPin)
+	{
+		// A dedicated pin for params will be added when the pins are updated. Here we detect any params connections to the In pin
+		// and disconnect them, and move the first params connection to a new params pin.
+		
+		// Check basic conditions for which the code below should run.
+		check(InputPins.Num() == 1);
+		check(InputPins[0] && InputPins[0]->Properties.AllowedTypes == EPCGDataType::Any);
+
+		UPCGPin* InPin = InputPins[0];
+
+		// Add params pin with good defaults (UpdatePins will ensure pin details are correct later).
+		UPCGPin* NewParamsPin = NewObject<UPCGPin>(InOutNode);
+		NewParamsPin->Node = InOutNode;
+		NewParamsPin->Properties.AllowedTypes = EPCGDataType::Param;
+		NewParamsPin->Properties.Label = PCGPinConstants::DefaultParamsLabel;
+		NewParamsPin->Properties.bAllowMultipleConnections = false;
+		InputPins.Add(NewParamsPin);
+
+		// Make list of param pins that In pin is currently connected to.
+		TArray<UPCGPin*> UpstreamParamPins;
+		for (const UPCGEdge* Connection : InPin->Edges)
+		{
+			if (Connection->InputPin && Connection->InputPin->Properties.AllowedTypes == EPCGDataType::Param)
+			{
+				UpstreamParamPins.Add(Connection->InputPin);
+			}
+		}
+
+		// Break all connections to param pins, and connect the first such pin to the new params pin on this node.
+		for (UPCGPin* Pin : UpstreamParamPins)
+		{
+			InPin->BreakEdgeTo(Pin);
+
+			// Params never support multiple connections as a rule (user must merge params themselves), so just connect first.
+			if (!NewParamsPin->IsConnected())
+			{
+				NewParamsPin->AddEdgeTo(Pin);
+			}
+		}
+	}
+}
+#endif // WITH_EDITOR
diff --git a/Engine/Plugins/Experimental/PCG/Source/PCG/Private/PCGGraph.cpp b/Engine/Plugins/Experimental/PCG/Source/PCG/Private/PCGGraph.cpp
--- a/Engine/Plugins/Experimental/PCG/Source/PCG/Private/PCGGraph.cpp
+++ b/Engine/Plugins/Experimental/PCG/Source/PCG/Private/PCGGraph.cpp
@@ -94,23 +94,14 @@ void UPCGGraph::PostLoad()
 		ExtraNode->ConditionalPostLoad();
 	}
 
-	// Update pins on all nodes
-	InputNode->UpdatePins();
-	OutputNode->UpdatePins();
+	// Finally, apply deprecation that changes edges/rebinds
+	ForEachNode([](UPCGNode* InNode) { return InNode->ApplyDeprecationBeforeUpdatePins(); });
 
-	for (UPCGNode* Node : Nodes)
-	{
-		Node->UpdatePins();
-	}
+	// Update pins on all nodes
+	ForEachNode([](UPCGNode* InNode) { return InNode->UpdatePins(); });
 
 	// Finally, apply deprecation that changes edges/rebinds
-	InputNode->ApplyDeprecation();
-	OutputNode->ApplyDeprecation();
-	
-	for (UPCGNode* Node : Nodes)
-	{
-		Node->ApplyDeprecation();
-	}
+	ForEachNode([](UPCGNode* InNode) { return InNode->ApplyDeprecation(); });
 
 	InputNode->OnNodeChangedDelegate.AddUObject(this, &UPCGGraph::OnNodeChanged);
 	OutputNode->OnNodeChangedDelegate.AddUObject(this, &UPCGGraph::OnNodeChanged);
@@ -413,6 +404,17 @@ bool UPCGGraph::RemoveEdge(UPCGNode* From, const FName& FromLabel, UPCGNode* To,
 	return bChanged;
 }
 
+void UPCGGraph::ForEachNode(const TFunction<void(UPCGNode*)>& Action)
+{
+	Action(InputNode);
+	Action(OutputNode);
+
+	for (UPCGNode* Node : Nodes)
+	{
+		Action(Node);
+	}
+}
+
 bool UPCGGraph::RemoveAllInboundEdges(UPCGNode* InNode)
 {
 	check(InNode);
diff --git a/Engine/Plugins/Experimental/PCG/Source/PCG/Private/PCGNode.cpp b/Engine/Plugins/Experimental/PCG/Source/PCG/Private/PCGNode.cpp
--- a/Engine/Plugins/Experimental/PCG/Source/PCG/Private/PCGNode.cpp
+++ b/Engine/Plugins/Experimental/PCG/Source/PCG/Private/PCGNode.cpp
@@ -2,6 +2,7 @@
 
 #include "PCGNode.h"
 
+#include "PCGCustomVersion.h"
 #include "PCGEdge.h"
 #include "PCGGraph.h"
 #include "PCGSettings.h"
@@ -56,6 +57,20 @@ void UPCGNode::PostLoad()
 }
 
 #if WITH_EDITOR
+void UPCGNode::ApplyDeprecationBeforeUpdatePins()
+{
+	if (UPCGSettings* Settings = GetSettings())
+	{
+		const int VersionBefore = Settings->DataVersion;
+
+		Settings->ApplyDeprecationBeforeUpdatePins(this, InputPins, OutputPins);
+
+		// Version number should not be bumped in this before-update-pins migration, if it does
+		// we risk deprecation code not running in post-update-pins because version number is latest.
+		ensure(VersionBefore == Settings->DataVersion);
+	}
+}
+
 void UPCGNode::ApplyDeprecation()
 {
 	UPCGPin* DefaultOutputPin = OutputPins.IsEmpty() ? nullptr : OutputPins[0];
@@ -130,6 +145,9 @@ void UPCGNode::ApplyDeprecation()
 	if (UPCGSettings* Settings = GetSettings())
 	{
 		Settings->ApplyDeprecation(this);
+
+		// Once deprecation has run, version should be up to date.
+		ensure(Settings->DataVersion == FPCGCustomVersion::LatestVersion);
 	}
 }
 #endif
diff --git a/Engine/Plugins/Experimental/PCG/Source/PCG/Private/PCGPin.cpp b/Engine/Plugins/Experimental/PCG/Source/PCG/Private/PCGPin.cpp
--- a/Engine/Plugins/Experimental/PCG/Source/PCG/Private/PCGPin.cpp
+++ b/Engine/Plugins/Experimental/PCG/Source/PCG/Private/PCGPin.cpp
@@ -205,6 +205,8 @@ bool UPCGPin::IsCompatible(const UPCGPin* OtherPin) const
 	const bool bDownstreamComposite = !!(DownstreamPin->Properties.AllowedTypes & EPCGDataType::Composite);
 	if (bUpstreamConcrete && bDownstreamComposite)
 	{
+		// TODO: Bug - this allows any concrete to connect to any other concrete. We need conversion nodes
+		// where appropriate, or block incompatible edges.
 		return true;
 	}
 
diff --git a/Engine/Plugins/Experimental/PCG/Source/PCG/Private/PCGSettings.cpp b/Engine/Plugins/Experimental/PCG/Source/PCG/Private/PCGSettings.cpp
--- a/Engine/Plugins/Experimental/PCG/Source/PCG/Private/PCGSettings.cpp
+++ b/Engine/Plugins/Experimental/PCG/Source/PCG/Private/PCGSettings.cpp
@@ -53,6 +53,13 @@ uint32 UPCGSettings::GetCrc32() const
 	return Ar.Crc32(const_cast<UPCGSettings*>(this));
 }
 
+#if WITH_EDITOR
+void UPCGSettings::ApplyDeprecation(UPCGNode* InOutNode)
+{
+	DataVersion = FPCGCustomVersion::LatestVersion;
+}
+#endif // WITH_EDITOR
+
 void UPCGSettings::PostLoad()
 {
 	Super::PostLoad();
diff --git a/Engine/Plugins/Experimental/PCG/Source/PCG/Public/Elements/PCGSelfPruning.h b/Engine/Plugins/Experimental/PCG/Source/PCG/Public/Elements/PCGSelfPruning.h
--- a/Engine/Plugins/Experimental/PCG/Source/PCG/Public/Elements/PCGSelfPruning.h
+++ b/Engine/Plugins/Experimental/PCG/Source/PCG/Public/Elements/PCGSelfPruning.h
@@ -43,9 +43,11 @@ public:
 	//~Begin UPCGSettings interface
 	virtual FName GetDefaultNodeName() const override { return FName(TEXT("SelfPruning")); }
 	virtual EPCGSettingsType GetType() const override { return EPCGSettingsType::Filter; }
+	virtual void ApplyDeprecationBeforeUpdatePins(UPCGNode* InOutNode, TArray<TObjectPtr<UPCGPin>>& InputPins, TArray<TObjectPtr<UPCGPin>>& OutputPins);
 #endif
 
-	virtual TArray<FPCGPinProperties> OutputPinProperties() const override { return Super::DefaultPointOutputPinProperties(); }
+	virtual TArray<FPCGPinProperties> InputPinProperties() const override;
+	virtual TArray<FPCGPinProperties> OutputPinProperties() const override;
 
 protected:
 	virtual FPCGElementPtr CreateElement() const override;
diff --git a/Engine/Plugins/Experimental/PCG/Source/PCG/Public/PCGCustomVersion.h b/Engine/Plugins/Experimental/PCG/Source/PCG/Public/PCGCustomVersion.h
--- a/Engine/Plugins/Experimental/PCG/Source/PCG/Public/PCGCustomVersion.h
+++ b/Engine/Plugins/Experimental/PCG/Source/PCG/Public/PCGCustomVersion.h
@@ -16,6 +16,8 @@ struct PCG_API FPCGCustomVersion
 		// Split projection nodes inputs to separate source edges and target edge
 		SplitProjectionNodeInputs = 1,
 		
+		MoveSelfPruningParamsOffFirstPin = 2,
+
 		// -----<new versions can be added above this line>-------------------------------------------------
 		VersionPlusOne,
 		LatestVersion = VersionPlusOne - 1
diff --git a/Engine/Plugins/Experimental/PCG/Source/PCG/Public/PCGGraph.h b/Engine/Plugins/Experimental/PCG/Source/PCG/Public/PCGGraph.h
--- a/Engine/Plugins/Experimental/PCG/Source/PCG/Public/PCGGraph.h
+++ b/Engine/Plugins/Experimental/PCG/Source/PCG/Public/PCGGraph.h
@@ -96,6 +96,9 @@ public:
 	const TArray<UPCGNode*>& GetNodes() const { return Nodes; }
 	void AddNode(UPCGNode* InNode);
 
+	/** Calls the lambda on every node in graph. */
+	void ForEachNode(const TFunction<void(UPCGNode*)>& Action);
+
 	bool RemoveAllInboundEdges(UPCGNode* InNode);
 	bool RemoveAllOutboundEdges(UPCGNode* InNode);
 	bool RemoveInboundEdges(UPCGNode* InNode, const FName& InboundLabel);
diff --git a/Engine/Plugins/Experimental/PCG/Source/PCG/Public/PCGNode.h b/Engine/Plugins/Experimental/PCG/Source/PCG/Public/PCGNode.h
--- a/Engine/Plugins/Experimental/PCG/Source/PCG/Public/PCGNode.h
+++ b/Engine/Plugins/Experimental/PCG/Source/PCG/Public/PCGNode.h
@@ -39,7 +39,10 @@ public:
 	virtual void PostEditUndo() override;
 	/** ~End UObject interface */
 
-	/** Used to be able to force deprecation when things need to be deprecated at the graph level */
+	/** UpdatePins will kick off invalid edges, so this is useful for making pin changes graceful. */
+	void ApplyDeprecationBeforeUpdatePins();
+
+	/** Used to be able to force deprecation when things need to be deprecated at the graph level. */
 	void ApplyDeprecation();
 #endif
 
diff --git a/Engine/Plugins/Experimental/PCG/Source/PCG/Public/PCGSettings.h b/Engine/Plugins/Experimental/PCG/Source/PCG/Public/PCGSettings.h
--- a/Engine/Plugins/Experimental/PCG/Source/PCG/Public/PCGSettings.h
+++ b/Engine/Plugins/Experimental/PCG/Source/PCG/Public/PCGSettings.h
@@ -109,7 +109,10 @@ public:
 	bool UseSeed() const { return bUseSeed; }
 
 #if WITH_EDITOR
-	virtual void ApplyDeprecation(UPCGNode* InOutNode) {}
+	/** UpdatePins will kick off invalid edges, so this is useful for moving edges around in case of pin changes. */
+	virtual void ApplyDeprecationBeforeUpdatePins(UPCGNode* InOutNode, TArray<TObjectPtr<UPCGPin>>& InputPins, TArray<TObjectPtr<UPCGPin>>& OutputPins) {}
+	/** Any final migration/recovery that can be done after pins are finalized. This function should also set DataVersion to LatestVersion. */
+	virtual void ApplyDeprecation(UPCGNode* InOutNode);
 
 	virtual FName GetDefaultNodeName() const { return NAME_None; }
 	virtual FText GetNodeTooltipText() const { return FText::GetEmpty(); }
@@ -189,7 +192,7 @@ protected:
 	TArray<FPCGPinProperties> DefaultPointOutputPinProperties() const;
 
 #if WITH_EDITOR
-protected:
+public:
 	/** The version number of the data after load and after any data migration. */
 	int32 DataVersion = -1;
 
