[PCG] TransformPoints: Apply concrete, narrow types (types that are used internally in the node). This does not kick off any connections in existing graphs due to our lenient type checking.

If we extend this node to process more than just points in the future we can widen the types accordingly.

Also reorder pins so Params is second. Interestingly this does nothing to existing node usages in the graph, but will apply to new ones. Potentially UpdatePins() could be extended to ensure pin orders are correct, but having it apply to new node instances while preserving existing did not seem like the worst behaviour to me.

diff --git a/Engine/Plugins/Experimental/PCG/Source/PCG/Private/Elements/PCGTransformPoints.cpp b/Engine/Plugins/Experimental/PCG/Source/PCG/Private/Elements/PCGTransformPoints.cpp
--- a/Engine/Plugins/Experimental/PCG/Source/PCG/Private/Elements/PCGTransformPoints.cpp
+++ b/Engine/Plugins/Experimental/PCG/Source/PCG/Private/Elements/PCGTransformPoints.cpp
@@ -15,8 +15,17 @@ UPCGTransformPointsSettings::UPCGTransformPointsSettings()
 TArray<FPCGPinProperties> UPCGTransformPointsSettings::InputPinProperties() const
 {
 	TArray<FPCGPinProperties> PinProperties;
+	// TODO in the future type checking of edges will be stricter and a conversion node will be added to convert from other types
+	PinProperties.Emplace(PCGPinConstants::DefaultInputLabel, EPCGDataType::Point);
 	PinProperties.Emplace(TEXT("Params"), EPCGDataType::Param, /*bInAllowMultipleConnections=*/false);
-	PinProperties.Emplace(PCGPinConstants::DefaultInputLabel, EPCGDataType::Spatial);
+
+	return PinProperties;
+}
+
+TArray<FPCGPinProperties> UPCGTransformPointsSettings::OutputPinProperties() const
+{
+	TArray<FPCGPinProperties> PinProperties;
+	PinProperties.Emplace(PCGPinConstants::DefaultOutputLabel, EPCGDataType::Point);
 
 	return PinProperties;
 }
diff --git a/Engine/Plugins/Experimental/PCG/Source/PCG/Public/Elements/PCGTransformPoints.h b/Engine/Plugins/Experimental/PCG/Source/PCG/Public/Elements/PCGTransformPoints.h
--- a/Engine/Plugins/Experimental/PCG/Source/PCG/Public/Elements/PCGTransformPoints.h
+++ b/Engine/Plugins/Experimental/PCG/Source/PCG/Public/Elements/PCGTransformPoints.h
@@ -24,7 +24,7 @@ public:
 #endif
 
 	virtual TArray<FPCGPinProperties> InputPinProperties() const override;
-	virtual TArray<FPCGPinProperties> OutputPinProperties() const override { return Super::DefaultPointOutputPinProperties(); }
+	virtual TArray<FPCGPinProperties> OutputPinProperties() const override;
 
 protected:
 	virtual FPCGElementPtr CreateElement() const override;
