[PCG] Allow to set SoftObjectPaths in CreateAttribute + limit size of node titles

diff --git a/Engine/Plugins/Experimental/PCG/Source/PCG/Private/Metadata/PCGMetadataTypesConstantStruct.cpp b/Engine/Plugins/Experimental/PCG/Source/PCG/Private/Metadata/PCGMetadataTypesConstantStruct.cpp
--- a/Engine/Plugins/Experimental/PCG/Source/PCG/Private/Metadata/PCGMetadataTypesConstantStruct.cpp
+++ b/Engine/Plugins/Experimental/PCG/Source/PCG/Private/Metadata/PCGMetadataTypesConstantStruct.cpp
@@ -15,7 +15,14 @@ FString FPCGMetadataTypesConstantStruct::ToString() const
 	case EPCGMetadataTypes::Double:
 		return FString::Printf(TEXT("%.2f"), DoubleValue);
 	case EPCGMetadataTypes::String:
-		return FString::Printf(TEXT("\"%s\""), *StringValue);
+		if (bStringAsSoftObjectPath)
+		{
+			return FString::Printf(TEXT("\"%s\""), *(SoftObjectPathValue.GetAssetName()));
+		}
+		else
+		{
+			return FString::Printf(TEXT("\"%s\""), *StringValue);
+		}
 	case EPCGMetadataTypes::Name:
 		return FString::Printf(TEXT("N(\"%s\")"), *NameValue.ToString());
 	case EPCGMetadataTypes::Vector2:
diff --git a/Engine/Plugins/Experimental/PCG/Source/PCG/Public/Metadata/PCGMetadataTypesConstantStruct.h b/Engine/Plugins/Experimental/PCG/Source/PCG/Public/Metadata/PCGMetadataTypesConstantStruct.h
--- a/Engine/Plugins/Experimental/PCG/Source/PCG/Public/Metadata/PCGMetadataTypesConstantStruct.h
+++ b/Engine/Plugins/Experimental/PCG/Source/PCG/Public/Metadata/PCGMetadataTypesConstantStruct.h
@@ -26,6 +26,9 @@ public:
 	UPROPERTY(BlueprintReadWrite, EditAnywhere, Category = Settings, meta = (ValidEnumValues = "Float, Double, Integer32, Integer64, Vector2, Vector, Vector4, Quaternion, Transform, String, Boolean, Rotator, Name"))
 	EPCGMetadataTypes Type = EPCGMetadataTypes::Double;
 
+	UPROPERTY(BlueprintReadWrite, EditAnywhere, Category = Settings, meta = (EditCondition = "Type == EPCGMetadataTypes::String", EditConditionHides))
+	bool bStringAsSoftObjectPath = false;
+
 	// All different types
 	UPROPERTY(BlueprintReadWrite, EditAnywhere, Category = Settings, meta = (EditCondition = "Type == EPCGMetadataTypes::Float", EditConditionHides))
 	float FloatValue = 0.0f;
@@ -54,7 +57,7 @@ public:
 	UPROPERTY(BlueprintReadWrite, EditAnywhere, Category = Settings, meta = (EditCondition = "Type == EPCGMetadataTypes::Transform", EditConditionHides))
 	FTransform TransformValue = FTransform::Identity;
 
-	UPROPERTY(BlueprintReadWrite, EditAnywhere, Category = Settings, meta = (EditCondition = "Type == EPCGMetadataTypes::String", EditConditionHides))
+	UPROPERTY(BlueprintReadWrite, EditAnywhere, Category = Settings, meta = (EditCondition = "Type == EPCGMetadataTypes::String && !bStringAsSoftObjectPath", EditConditionHides))
 	FString StringValue = "";
 
 	UPROPERTY(BlueprintReadWrite, EditAnywhere, Category = Settings, meta = (EditCondition = "Type == EPCGMetadataTypes::Boolean", EditConditionHides))
@@ -65,6 +68,9 @@ public:
 
 	UPROPERTY(BlueprintReadWrite, EditAnywhere, Category = Settings, meta = (EditCondition = "Type == EPCGMetadataTypes::Name", EditConditionHides))
 	FName NameValue = NAME_None;
+
+	UPROPERTY(BlueprintReadWrite, EditAnywhere, Category = Settings, meta = (EditCondition = "Type == EPCGMetadataTypes::String && bStringAsSoftObjectPath", EditConditionHides))
+	FSoftClassPath SoftObjectPathValue;
 };
 
 template <typename Func>
@@ -93,7 +99,14 @@ decltype(auto) FPCGMetadataTypesConstantStruct::DispatcherWithOverride(const UPC
 	case EPCGMetadataTypes::Transform:
 		return Callback(PCG_GET_OVERRIDEN_VALUE(this, TransformValue, Params));
 	case EPCGMetadataTypes::String:
-		return Callback(PCG_GET_OVERRIDEN_VALUE(this, StringValue, Params));
+		if (bStringAsSoftObjectPath)
+		{
+			return Callback(PCG_GET_OVERRIDEN_VALUE(this, SoftObjectPathValue.ToString(), Params));
+		}
+		else
+		{
+			return Callback(PCG_GET_OVERRIDEN_VALUE(this, StringValue, Params));
+		}
 	case EPCGMetadataTypes::Boolean:
 		return Callback(PCG_GET_OVERRIDEN_VALUE(this, BoolValue, Params));
 	case EPCGMetadataTypes::Rotator:
diff --git a/Engine/Plugins/Experimental/PCG/Source/PCGEditor/Private/PCGEditorGraphNode.cpp b/Engine/Plugins/Experimental/PCG/Source/PCGEditor/Private/PCGEditorGraphNode.cpp
--- a/Engine/Plugins/Experimental/PCG/Source/PCGEditor/Private/PCGEditorGraphNode.cpp
+++ b/Engine/Plugins/Experimental/PCG/Source/PCGEditor/Private/PCGEditorGraphNode.cpp
@@ -21,9 +21,12 @@ UPCGEditorGraphNode::UPCGEditorGraphNode(const FObjectInitializer& ObjectInitial
 
 FText UPCGEditorGraphNode::GetNodeTitle(ENodeTitleType::Type TitleType) const
 {
+	constexpr int32 NodeTitleMaxSize = 70;
+
 	if (PCGNode)
 	{
-		return FText::FromName(PCGNode->GetNodeTitle());
+		FString Res = PCGNode->GetNodeTitle().ToString();
+		return FText::FromString((Res.Len() > NodeTitleMaxSize) ? Res.Left(NodeTitleMaxSize) : Res);
 	}
 	else
 	{
