[PCG] Change StringMode bool to enum in CreateAttribute to support future usecases

#rnx

diff --git a/Engine/Plugins/Experimental/PCG/Source/PCG/Private/Metadata/PCGMetadataTypesConstantStruct.cpp b/Engine/Plugins/Experimental/PCG/Source/PCG/Private/Metadata/PCGMetadataTypesConstantStruct.cpp
--- a/Engine/Plugins/Experimental/PCG/Source/PCG/Private/Metadata/PCGMetadataTypesConstantStruct.cpp
+++ b/Engine/Plugins/Experimental/PCG/Source/PCG/Private/Metadata/PCGMetadataTypesConstantStruct.cpp
@@ -15,14 +15,17 @@ FString FPCGMetadataTypesConstantStruct::ToString() const
 	case EPCGMetadataTypes::Double:
 		return FString::Printf(TEXT("%.2f"), DoubleValue);
 	case EPCGMetadataTypes::String:
-		if (bStringAsSoftObjectPath)
-		{
-			return FString::Printf(TEXT("\"%s\""), *(SoftObjectPathValue.GetAssetName()));
-		}
-		else
+	{
+		switch (StringMode)
 		{
+		case EPCGMetadataTypesConstantStructStringMode::String:
 			return FString::Printf(TEXT("\"%s\""), *StringValue);
+		case EPCGMetadataTypesConstantStructStringMode::SoftObjectPath:
+			return FString::Printf(TEXT("\"%s\""), *(SoftObjectPathValue.GetAssetName()));
+		default:
+			break;
 		}
+	}
 	case EPCGMetadataTypes::Name:
 		return FString::Printf(TEXT("N(\"%s\")"), *NameValue.ToString());
 	case EPCGMetadataTypes::Vector2:
@@ -40,6 +43,8 @@ FString FPCGMetadataTypesConstantStruct::ToString() const
 	case EPCGMetadataTypes::Boolean:
 		return FString::Printf(TEXT("%s"), (BoolValue ? TEXT("True") : TEXT("False")));
 	default:
-		return FString();
+		break;
 	}
+
+	return FString();
 }
\ No newline at end of file
diff --git a/Engine/Plugins/Experimental/PCG/Source/PCG/Public/Metadata/PCGMetadataTypesConstantStruct.h b/Engine/Plugins/Experimental/PCG/Source/PCG/Public/Metadata/PCGMetadataTypesConstantStruct.h
--- a/Engine/Plugins/Experimental/PCG/Source/PCG/Public/Metadata/PCGMetadataTypesConstantStruct.h
+++ b/Engine/Plugins/Experimental/PCG/Source/PCG/Public/Metadata/PCGMetadataTypesConstantStruct.h
@@ -9,6 +9,13 @@
 
 class UPCGParamData;
 
+UENUM()
+enum class EPCGMetadataTypesConstantStructStringMode
+{
+	String,
+	SoftObjectPath
+};
+
 /**
 * Struct to be re-used when you need to show constants types for a metadata type
 * It will store all our values, and will display nicely depending on the type chosen
@@ -27,7 +34,7 @@ public:
 	EPCGMetadataTypes Type = EPCGMetadataTypes::Double;
 
 	UPROPERTY(BlueprintReadWrite, EditAnywhere, Category = Settings, meta = (EditCondition = "Type == EPCGMetadataTypes::String", EditConditionHides))
-	bool bStringAsSoftObjectPath = false;
+	EPCGMetadataTypesConstantStructStringMode StringMode = EPCGMetadataTypesConstantStructStringMode::String;
 
 	// All different types
 	UPROPERTY(BlueprintReadWrite, EditAnywhere, Category = Settings, meta = (EditCondition = "Type == EPCGMetadataTypes::Float", EditConditionHides))
@@ -57,7 +64,7 @@ public:
 	UPROPERTY(BlueprintReadWrite, EditAnywhere, Category = Settings, meta = (EditCondition = "Type == EPCGMetadataTypes::Transform", EditConditionHides))
 	FTransform TransformValue = FTransform::Identity;
 
-	UPROPERTY(BlueprintReadWrite, EditAnywhere, Category = Settings, meta = (EditCondition = "Type == EPCGMetadataTypes::String && !bStringAsSoftObjectPath", EditConditionHides))
+	UPROPERTY(BlueprintReadWrite, EditAnywhere, Category = Settings, meta = (EditCondition = "Type == EPCGMetadataTypes::String && StringMode == EPCGMetadataTypesConstantStructStringMode::String", EditConditionHides))
 	FString StringValue = "";
 
 	UPROPERTY(BlueprintReadWrite, EditAnywhere, Category = Settings, meta = (EditCondition = "Type == EPCGMetadataTypes::Boolean", EditConditionHides))
@@ -69,7 +76,7 @@ public:
 	UPROPERTY(BlueprintReadWrite, EditAnywhere, Category = Settings, meta = (EditCondition = "Type == EPCGMetadataTypes::Name", EditConditionHides))
 	FName NameValue = NAME_None;
 
-	UPROPERTY(BlueprintReadWrite, EditAnywhere, Category = Settings, meta = (EditCondition = "Type == EPCGMetadataTypes::String && bStringAsSoftObjectPath", EditConditionHides))
+	UPROPERTY(BlueprintReadWrite, EditAnywhere, Category = Settings, meta = (EditCondition = "Type == EPCGMetadataTypes::String && StringMode == EPCGMetadataTypesConstantStructStringMode::SoftObjectPath", EditConditionHides))
 	FSoftClassPath SoftObjectPathValue;
 };
 
@@ -99,14 +106,17 @@ decltype(auto) FPCGMetadataTypesConstantStruct::DispatcherWithOverride(const UPC
 	case EPCGMetadataTypes::Transform:
 		return Callback(PCG_GET_OVERRIDEN_VALUE(this, TransformValue, Params));
 	case EPCGMetadataTypes::String:
-		if (bStringAsSoftObjectPath)
-		{
-			return Callback(PCG_GET_OVERRIDEN_VALUE(this, SoftObjectPathValue.ToString(), Params));
-		}
-		else
+	{
+		switch (StringMode)
 		{
+		case EPCGMetadataTypesConstantStructStringMode::String:
 			return Callback(PCG_GET_OVERRIDEN_VALUE(this, StringValue, Params));
+		case EPCGMetadataTypesConstantStructStringMode::SoftObjectPath:
+			return Callback(PCGSettingsHelpers::GetValue(GET_MEMBER_NAME_CHECKED(FPCGMetadataTypesConstantStruct, SoftObjectPathValue), SoftObjectPathValue.ToString(), Params));
+		default:
+			break;
 		}
+	}
 	case EPCGMetadataTypes::Boolean:
 		return Callback(PCG_GET_OVERRIDEN_VALUE(this, BoolValue, Params));
 	case EPCGMetadataTypes::Rotator:
@@ -114,14 +124,16 @@ decltype(auto) FPCGMetadataTypesConstantStruct::DispatcherWithOverride(const UPC
 	case EPCGMetadataTypes::Name:
 		return Callback(PCG_GET_OVERRIDEN_VALUE(this, NameValue, Params));
 	default:
-		// ReturnType{} is invalid if ReturnType is void
-		if constexpr (std::is_same_v<ReturnType, void>)
-		{
-			return;
-		}
-		else
-		{
-			return ReturnType{};
-		}
+		break;
+	}
+
+	// ReturnType{} is invalid if ReturnType is void
+	if constexpr (std::is_same_v<ReturnType, void>)
+	{
+		return;
+	}
+	else
+	{
+		return ReturnType{};
 	}
 }
\ No newline at end of file
