Added a unit test + fix for a bug with the landscape cache sampling having bad indicies
Also fixed crash when you leave the world bounds

diff --git a/Engine/Plugins/Experimental/PCG/Source/PCG/Private/Grid/PCGLandscapeCache.cpp b/Engine/Plugins/Experimental/PCG/Source/PCG/Private/Grid/PCGLandscapeCache.cpp
--- a/Engine/Plugins/Experimental/PCG/Source/PCG/Private/Grid/PCGLandscapeCache.cpp
+++ b/Engine/Plugins/Experimental/PCG/Source/PCG/Private/Grid/PCGLandscapeCache.cpp
@@ -36,65 +36,102 @@ static TAutoConsoleVariable<int32> CVarLandscapeCacheSizeThreshold(
 	64,
 	TEXT("Memory Threhold at which we start cleaning up the landscape cache"));
 
+namespace PCGLandscapeCache
+{
+	FSafeIndices CalcSafeIndices(FVector2D LocalPoint, int32 Stride)
+	{
+		check(Stride != 0);
+
+		const FVector2D ClampedLocalPoint = LocalPoint.ClampAxes(0.0, FVector2D::FReal(Stride-1));
+
+		FSafeIndices Result;
+		const int32 CellX0 = FMath::FloorToInt(ClampedLocalPoint.X);
+		const int32 CellY0 = FMath::FloorToInt(ClampedLocalPoint.Y);
+		const int32 CellX1 = FMath::Min(CellX0+1, Stride-1);
+		const int32 CellY1 = FMath::Min(CellY0+1, Stride-1);
+
+		Result.X0Y0 = CellX0 + CellY0 * Stride;
+		Result.X1Y0 = CellX1 + CellY0 * Stride;
+		Result.X0Y1 = CellX0 + CellY1 * Stride;
+		Result.X1Y1 = CellX1 + CellY1 * Stride;
+
+		Result.XFraction = FMath::Fractional(ClampedLocalPoint.X);
+		Result.YFraction = FMath::Fractional(ClampedLocalPoint.Y);
+
+		return Result;
+	}
+}
+
 #if WITH_EDITOR
-void FPCGLandscapeCacheEntry::BuildCacheData(ULandscapeInfo* LandscapeInfo, ULandscapeComponent* InComponent)
+FPCGLandscapeCacheEntry* FPCGLandscapeCacheEntry::CreateCacheEntry(ULandscapeInfo* LandscapeInfo, ULandscapeComponent* InComponent)
 {
-	TRACE_CPUPROFILER_EVENT_SCOPE(FPCGLandscapeCacheEntry::BuildCacheData);
-	check(!Component.Get() && InComponent && PositionsAndNormals.Num() == 0 && !bDataLoaded);
-	Component = InComponent;
+	TRACE_CPUPROFILER_EVENT_SCOPE(FPCGLandscapeCacheEntry::CreateCacheEntry);
+
+	FLandscapeComponentDataInterface CDI(InComponent, 0, /*bWorkInEditingLayer=*/false);
 
-	ALandscape* Landscape = Component->GetLandscapeActor();
-	if (!Landscape)
+	// CDI.GetRawHeightData() may be nullptr if InComponent has no texture data.
+	if (CDI.GetRawHeightData() == nullptr)
 	{
-		return;
+		return nullptr;
 	}
 
+	// The component has an extra vertex on the edge, for interpolation purposes
+	const int32 Stride = InComponent->ComponentSizeQuads + 1;
+
+	if (Stride <= 0)
+	{
+		return nullptr;
+	}
+
+	FPCGLandscapeCacheEntry *Result = new FPCGLandscapeCacheEntry();
+
+ 	Result->Component = InComponent;
+	Result->PointHalfSize = InComponent->GetComponentTransform().GetScale3D() * 0.5;
+	Result->Stride = Stride;
+
 	// Get landscape default layer (heightmap/tangents/normal)
 	{
-		FLandscapeComponentDataInterface CDI(InComponent, 0, /*bWorkInEditingLayer=*/false);
 		FVector WorldPos;
 		FVector WorldTangentX;
 		FVector WorldTangentY;
 		FVector WorldTangentZ;
 
-		PointHalfSize = InComponent->GetComponentTransform().GetScale3D() * 0.5;
+		const int32 NumVertices = FMath::Square(Stride);
 
-		// The component has an extra vertex on the edge, for interpolation purposes
-		const int32 ComponentSizeQuads = InComponent->ComponentSizeQuads + 1;
-		const int32 NumVertices = FMath::Square(ComponentSizeQuads);
+		Result->PositionsAndNormals.Reserve(2 * NumVertices);
 
-		Stride = ComponentSizeQuads;
-
-		PositionsAndNormals.Reserve(2 * NumVertices);
 		for (int32 Index = 0; Index < NumVertices; ++Index)
 		{
 			CDI.GetWorldPositionTangents(Index, WorldPos, WorldTangentX, WorldTangentY, WorldTangentZ);
-			PositionsAndNormals.Add(WorldPos);
-			PositionsAndNormals.Add(WorldTangentZ);
+			Result->PositionsAndNormals.Add(WorldPos);
+			Result->PositionsAndNormals.Add(WorldTangentZ);
 		}
 	}
-	
+
 	// Get other layers, push data into metadata attributes
-	TArray<uint8> LayerCache;
-	for (const FLandscapeInfoLayerSettings& Layer : LandscapeInfo->Layers)
 	{
-		ULandscapeLayerInfoObject* LayerInfo = Layer.LayerInfoObj;
-		if (!LayerInfo)
+		TArray<uint8> LayerCache;
+		for (const FLandscapeInfoLayerSettings& Layer : LandscapeInfo->Layers)
 		{
-			continue;
-		}
+			ULandscapeLayerInfoObject* LayerInfo = Layer.LayerInfoObj;
+			if (!LayerInfo)
+			{
+				continue;
+			}
 
-		FLandscapeComponentDataInterface CDI(InComponent, 0, /*bWorkInEditingLayer=*/false);
-		if (CDI.GetWeightmapTextureData(LayerInfo, LayerCache, /*bUseEditingLayer=*/false))
-		{
-			LayerDataNames.Add(Layer.LayerName);
-			LayerData.Add(LayerCache);
-		}
+			if (CDI.GetWeightmapTextureData(LayerInfo, LayerCache, /*bUseEditingLayer=*/false))
+			{
+				Result->LayerDataNames.Add(Layer.LayerName);
+				Result->LayerData.Emplace(std::move(LayerCache));
+			}
 
-		LayerCache.Reset();
+			LayerCache.Reset();
+		}
 	}
 	
-	bDataLoaded = true;
+	Result->bDataLoaded = true;
+
+	return Result;
 }
 #endif // WITH_EDITOR
 
@@ -161,13 +198,7 @@ void FPCGLandscapeCacheEntry::GetInterpolatedLayerWeights(const FVector2D& Local
 		return;
 	}
 
-	const int32 X0Y0 = FMath::FloorToInt(LocalPoint.X) + FMath::FloorToInt(LocalPoint.Y) * Stride;
-	const int32 X1Y0 = X0Y0 + 1;
-	const int32 X0Y1 = X0Y0 + Stride;
-	const int32 X1Y1 = X0Y1 + 1;
-
-	const float XFactor = FMath::Fractional(LocalPoint.X);
-	const float YFactor = FMath::Fractional(LocalPoint.Y);
+	const PCGLandscapeCache::FSafeIndices Indices = PCGLandscapeCache::CalcSafeIndices(LocalPoint, Stride);
 
 	if (!LayerData.IsEmpty())
 	{
@@ -180,44 +211,38 @@ void FPCGLandscapeCacheEntry::GetInterpolatedLayerWeights(const FVector2D& Local
 
 			const TArray<uint8>& CurrentLayerData = LayerData[LayerIndex];
 
-			const float Y0Data = FMath::Lerp((float)CurrentLayerData[X0Y0] / 255.0f, (float)CurrentLayerData[X1Y0] / 255.0f, XFactor);
-			const float Y1Data = FMath::Lerp((float)CurrentLayerData[X0Y1] / 255.0f, (float)CurrentLayerData[X1Y1] / 255.0f, XFactor);
+			const float Y0Data = FMath::Lerp((float)CurrentLayerData[Indices.X0Y0] / 255.0f, (float)CurrentLayerData[Indices.X1Y0] / 255.0f, Indices.XFraction);
+			const float Y1Data = FMath::Lerp((float)CurrentLayerData[Indices.X0Y1] / 255.0f, (float)CurrentLayerData[Indices.X1Y1] / 255.0f, Indices.XFraction);
 
-			OutLayer.Weight = FMath::Lerp(Y0Data, Y1Data, YFactor);
+			OutLayer.Weight = FMath::Lerp(Y0Data, Y1Data, Indices.YFraction);
 		}
 	}
 }
 
 void FPCGLandscapeCacheEntry::GetInterpolatedPoint(const FVector2D& LocalPoint, FPCGPoint& OutPoint, UPCGMetadata* OutMetadata) const
 {
-	check(bDataLoaded);
-
-	const int32 X0Y0 = FMath::FloorToInt(LocalPoint.X) + FMath::FloorToInt(LocalPoint.Y) * Stride;
-	const int32 X1Y0 = X0Y0 + 1;
-	const int32 X0Y1 = X0Y0 + Stride;
-	const int32 X1Y1 = X0Y1 + 1;
-
-	const float XFactor = FMath::Fractional(LocalPoint.X);
-	const float YFactor = FMath::Fractional(LocalPoint.Y);
+	check (bDataLoaded);
+	
+	const PCGLandscapeCache::FSafeIndices Indices = PCGLandscapeCache::CalcSafeIndices(LocalPoint, Stride);
 
-	check(X0Y0 >= 0 && 2 * X1Y1 < PositionsAndNormals.Num());
+	check(2 * Indices.X1Y1 < PositionsAndNormals.Num());
 
-	const FVector& PositionX0Y0 = PositionsAndNormals[2 * X0Y0];
-	const FVector& NormalX0Y0 = PositionsAndNormals[2 * X0Y0 + 1];
-	const FVector& PositionX1Y0 = PositionsAndNormals[2 * X1Y0];
-	const FVector& NormalX1Y0 = PositionsAndNormals[2 * X1Y0 + 1];
-	const FVector& PositionX0Y1 = PositionsAndNormals[2 * X0Y1];
-	const FVector& NormalX0Y1 = PositionsAndNormals[2 * X0Y1 + 1];
-	const FVector& PositionX1Y1 = PositionsAndNormals[2 * X1Y1];
-	const FVector& NormalX1Y1 = PositionsAndNormals[2 * X1Y1 + 1];
+	const FVector& PositionX0Y0 = PositionsAndNormals[2 * Indices.X0Y0];
+	const FVector& NormalX0Y0 = PositionsAndNormals[2 * Indices.X0Y0 + 1];
+	const FVector& PositionX1Y0 = PositionsAndNormals[2 * Indices.X1Y0];
+	const FVector& NormalX1Y0 = PositionsAndNormals[2 * Indices.X1Y0 + 1];
+	const FVector& PositionX0Y1 = PositionsAndNormals[2 * Indices.X0Y1];
+	const FVector& NormalX0Y1 = PositionsAndNormals[2 * Indices.X0Y1 + 1];
+	const FVector& PositionX1Y1 = PositionsAndNormals[2 * Indices.X1Y1];
+	const FVector& NormalX1Y1 = PositionsAndNormals[2 * Indices.X1Y1 + 1];
 
-	const FVector LerpPositionY0 = FMath::Lerp(PositionX0Y0, PositionX1Y0, XFactor);
-	const FVector LerpPositionY1 = FMath::Lerp(PositionX0Y1, PositionX1Y1, XFactor);
-	const FVector Position = FMath::Lerp(LerpPositionY0, LerpPositionY1, YFactor);
+	const FVector LerpPositionY0 = FMath::Lerp(PositionX0Y0, PositionX1Y0, Indices.XFraction);
+	const FVector LerpPositionY1 = FMath::Lerp(PositionX0Y1, PositionX1Y1, Indices.XFraction);
+	const FVector Position = FMath::Lerp(LerpPositionY0, LerpPositionY1, Indices.YFraction);
 
-	const FVector LerpNormalY0 = FMath::Lerp(NormalX0Y0.GetSafeNormal(), NormalX1Y0.GetSafeNormal(), XFactor).GetSafeNormal();
-	const FVector LerpNormalY1 = FMath::Lerp(NormalX0Y1.GetSafeNormal(), NormalX1Y1.GetSafeNormal(), XFactor).GetSafeNormal();
-	const FVector Normal = FMath::Lerp(LerpNormalY0, LerpNormalY1, YFactor);
+	const FVector LerpNormalY0 = FMath::Lerp(NormalX0Y0.GetSafeNormal(), NormalX1Y0.GetSafeNormal(), Indices.XFraction).GetSafeNormal();
+	const FVector LerpNormalY1 = FMath::Lerp(NormalX0Y1.GetSafeNormal(), NormalX1Y1.GetSafeNormal(), Indices.XFraction).GetSafeNormal();
+	const FVector Normal = FMath::Lerp(LerpNormalY0, LerpNormalY1, Indices.YFraction);
 
 	FVector TangentX;
 	FVector TangentY;
@@ -245,12 +270,12 @@ void FPCGLandscapeCacheEntry::GetInterpolatedPoint(const FVector2D& LocalPoint,
 			if (FPCGMetadataAttributeBase* Attribute = OutMetadata->GetMutableAttribute(CurrentLayerName))
 			{
 				check(Attribute->GetTypeId() == PCG::Private::MetadataTypes<float>::Id);
-				float Y0Data = FMath::Lerp((float)CurrentLayerData[X0Y0] / 255.0f, (float)CurrentLayerData[X1Y0] / 255.0f, XFactor);
-				float Y1Data = FMath::Lerp((float)CurrentLayerData[X0Y1] / 255.0f, (float)CurrentLayerData[X1Y1] / 255.0f, XFactor);
-				float Data = FMath::Lerp(Y0Data, Y1Data, YFactor);
+				float Y0Data = FMath::Lerp((float)CurrentLayerData[Indices.X0Y0] / 255.0f, (float)CurrentLayerData[Indices.X1Y0] / 255.0f, Indices.XFraction);
+				float Y1Data = FMath::Lerp((float)CurrentLayerData[Indices.X0Y1] / 255.0f, (float)CurrentLayerData[Indices.X1Y1] / 255.0f, Indices.XFraction);
+				float Data = FMath::Lerp(Y0Data, Y1Data, Indices.YFraction);
 
 				static_cast<FPCGMetadataAttribute<float>*>(Attribute)->SetValue(OutPoint.MetadataEntry, Data);
-			}
+			} 
 		}
 	}
 }
@@ -259,24 +284,18 @@ void FPCGLandscapeCacheEntry::GetInterpolatedPointHeightOnly(const FVector2D& Lo
 {
 	check(bDataLoaded);
 
-	const int32 X0Y0 = FMath::FloorToInt(LocalPoint.X) + FMath::FloorToInt(LocalPoint.Y) * Stride;
-	const int32 X1Y0 = X0Y0 + 1;
-	const int32 X0Y1 = X0Y0 + Stride;
-	const int32 X1Y1 = X0Y1 + 1;
-
-	const float XFactor = FMath::Fractional(LocalPoint.X);
-	const float YFactor = FMath::Fractional(LocalPoint.Y);
+	const PCGLandscapeCache::FSafeIndices Indices = PCGLandscapeCache::CalcSafeIndices(LocalPoint, Stride);
 
-	check(X0Y0 >= 0 && 2 * X1Y1 < PositionsAndNormals.Num());
+	check(2 * Indices.X1Y1 < PositionsAndNormals.Num());
 
-	const FVector& PositionX0Y0 = PositionsAndNormals[2 * X0Y0];
-	const FVector& PositionX1Y0 = PositionsAndNormals[2 * X1Y0];
-	const FVector& PositionX0Y1 = PositionsAndNormals[2 * X0Y1];
-	const FVector& PositionX1Y1 = PositionsAndNormals[2 * X1Y1];
+	const FVector& PositionX0Y0 = PositionsAndNormals[2 * Indices.X0Y0];
+	const FVector& PositionX1Y0 = PositionsAndNormals[2 * Indices.X1Y0];
+	const FVector& PositionX0Y1 = PositionsAndNormals[2 * Indices.X0Y1];
+	const FVector& PositionX1Y1 = PositionsAndNormals[2 * Indices.X1Y1];
 
-	const FVector LerpPositionY0 = FMath::Lerp(PositionX0Y0, PositionX1Y0, XFactor);
-	const FVector LerpPositionY1 = FMath::Lerp(PositionX0Y1, PositionX1Y1, XFactor);
-	const FVector Position = FMath::Lerp(LerpPositionY0, LerpPositionY1, YFactor);
+	const FVector LerpPositionY0 = FMath::Lerp(PositionX0Y0, PositionX1Y0, Indices.XFraction);
+	const FVector LerpPositionY1 = FMath::Lerp(PositionX0Y1, PositionX1Y1, Indices.XFraction);
+	const FVector Position = FMath::Lerp(LerpPositionY0, LerpPositionY1, Indices.YFraction);
 
 	const int32 Seed = UPCGBlueprintHelpers::ComputeSeedFromPosition(Position);
 	const float Density = 1;
@@ -554,9 +573,10 @@ void UPCGLandscapeCache::PrimeCache()
 
 					if (!CachedData.Contains(ComponentKey))
 					{
-						FPCGLandscapeCacheEntry* NewEntry = new FPCGLandscapeCacheEntry();
-						NewEntry->BuildCacheData(LandscapeInfo, LandscapeComponent);
-						CachedData.Add(ComponentKey, NewEntry);
+						if (FPCGLandscapeCacheEntry* NewEntry = FPCGLandscapeCacheEntry::CreateCacheEntry(LandscapeInfo, LandscapeComponent))
+						{
+							CachedData.Add(ComponentKey, NewEntry);
+						}
 					}
 				}
 			});
@@ -595,11 +615,11 @@ const FPCGLandscapeCacheEntry* UPCGLandscapeCache::GetCacheEntry(ULandscapeCompo
 		else
 		{
 			check(LandscapeComponent->SectionBaseX / LandscapeComponent->ComponentSizeQuads == ComponentKey.Value.X && LandscapeComponent->SectionBaseY / LandscapeComponent->ComponentSizeQuads == ComponentKey.Value.Y);
-			FPCGLandscapeCacheEntry* NewEntry = new FPCGLandscapeCacheEntry();
-			NewEntry->BuildCacheData(LandscapeComponent->GetLandscapeInfo(), LandscapeComponent);
-
-			CacheEntry = NewEntry;
-			CachedData.Add(ComponentKey, NewEntry);
+			if (FPCGLandscapeCacheEntry* NewEntry = FPCGLandscapeCacheEntry::CreateCacheEntry(LandscapeComponent->GetLandscapeInfo(), LandscapeComponent))
+			{
+				CacheEntry = NewEntry;
+				CachedData.Add(ComponentKey, NewEntry);				
+			}
 		}
 
 		if (CacheEntry)
diff --git a/Engine/Plugins/Experimental/PCG/Source/PCG/Private/Tests/Cache/PCGLandscapeCacheTest.cpp b/Engine/Plugins/Experimental/PCG/Source/PCG/Private/Tests/Cache/PCGLandscapeCacheTest.cpp
new file mode 100644
--- /dev/null
+++ b/Engine/Plugins/Experimental/PCG/Source/PCG/Private/Tests/Cache/PCGLandscapeCacheTest.cpp
@@ -0,0 +1,78 @@
+// Copyright Epic Games, Inc. All Rights Reserved.
+
+#include "Tests/PCGTestsCommon.h"
+
+#include "Grid/PCGLandscapeCache.h"
+
+#if WITH_EDITOR
+
+IMPLEMENT_CUSTOM_SIMPLE_AUTOMATION_TEST(FPCGLandscapeCache_CalcSafeIndices, FPCGTestBaseClass, "pcg.tests.LandscapeCache.CalcSafeIndices", PCGTestsCommon::TestFlags)
+
+bool FPCGLandscapeCache_CalcSafeIndices::RunTest(const FString& Parameters)
+{
+	PCGLandscapeCache::FSafeIndices TopLeft = PCGLandscapeCache::CalcSafeIndices(FVector2D(0.0), 32);
+
+	UTEST_EQUAL("TopLeft.X0Y0", TopLeft.X0Y0, 0 + 0*32);
+	UTEST_EQUAL("TopLeft.X1Y0", TopLeft.X1Y0, 1 + 0*32);
+	UTEST_EQUAL("TopLeft.X0Y1", TopLeft.X0Y1, 0 + 1*32);
+	UTEST_EQUAL("TopLeft.X1Y1", TopLeft.X1Y1, 1 + 1*32);
+	UTEST_EQUAL("TopLeft.XFraction", TopLeft.XFraction, 0.0f);
+	UTEST_EQUAL("TopLeft.YFraction", TopLeft.YFraction, 0.0f);
+
+	PCGLandscapeCache::FSafeIndices ClampTopLeft = PCGLandscapeCache::CalcSafeIndices(FVector2D(-0.5), 32);
+
+	UTEST_EQUAL("ClampTopLeft.X0Y0", ClampTopLeft.X0Y0, 0 + 0*32);
+	UTEST_EQUAL("ClampTopLeft.X1Y0", ClampTopLeft.X1Y0, 1 + 0*32);
+	UTEST_EQUAL("ClampTopLeft.X0Y1", ClampTopLeft.X0Y1, 0 + 1*32);
+	UTEST_EQUAL("ClampTopLeft.X1Y1", ClampTopLeft.X1Y1, 1 + 1*32);
+	UTEST_EQUAL("ClampTopLeft.XFraction", ClampTopLeft.XFraction, 0.0f);
+	UTEST_EQUAL("ClampTopLeft.YFraction", ClampTopLeft.YFraction, 0.0f);
+
+	PCGLandscapeCache::FSafeIndices FractionTopLeft = PCGLandscapeCache::CalcSafeIndices(FVector2D(0.5, 0.4), 32);
+
+	UTEST_EQUAL("FractionTopLeft.X0Y0", FractionTopLeft.X0Y0, 0 + 0*32);
+	UTEST_EQUAL("FractionTopLeft.X1Y0", FractionTopLeft.X1Y0, 1 + 0*32);
+	UTEST_EQUAL("FractionTopLeft.X0Y1", FractionTopLeft.X0Y1, 0 + 1*32);
+	UTEST_EQUAL("FractionTopLeft.X1Y1", FractionTopLeft.X1Y1, 1 + 1*32);
+	UTEST_EQUAL("FractionTopLeft.XFraction", FractionTopLeft.XFraction, 0.5f);
+	UTEST_EQUAL("FractionTopLeft.YFraction", FractionTopLeft.YFraction, 0.4f);
+
+	PCGLandscapeCache::FSafeIndices MiddleFraction = PCGLandscapeCache::CalcSafeIndices(FVector2D(4.5, 10.4), 32);
+
+	UTEST_EQUAL("MiddleFraction.X0Y0", MiddleFraction.X0Y0, 4 + 10*32);
+	UTEST_EQUAL("MiddleFraction.X1Y0", MiddleFraction.X1Y0, 5 + 10*32);
+	UTEST_EQUAL("MiddleFraction.X0Y1", MiddleFraction.X0Y1, 4 + 11*32);
+	UTEST_EQUAL("MiddleFraction.X1Y1", MiddleFraction.X1Y1, 5 + 11*32);
+	UTEST_EQUAL("MiddleFraction.XFraction", MiddleFraction.XFraction, 0.5f);
+	UTEST_EQUAL("MiddleFraction.YFraction", MiddleFraction.YFraction, 0.4f);
+
+	PCGLandscapeCache::FSafeIndices BottomRight = PCGLandscapeCache::CalcSafeIndices(FVector2D(31.0), 32);
+
+	UTEST_EQUAL("BottomRight.X0Y0", BottomRight.X0Y0, 32*32-1);
+	UTEST_EQUAL("BottomRight.X1Y0", BottomRight.X1Y0, 32*32-1);
+	UTEST_EQUAL("BottomRight.X0Y1", BottomRight.X0Y1, 32*32-1);
+	UTEST_EQUAL("BottomRight.X1Y1", BottomRight.X1Y1, 32*32-1);
+	UTEST_EQUAL("BottomRight.XFraction", BottomRight.XFraction, 0.0f);
+	UTEST_EQUAL("BottomRight.YFraction", BottomRight.YFraction, 0.0f);
+
+	PCGLandscapeCache::FSafeIndices ClampBottomRight = PCGLandscapeCache::CalcSafeIndices(FVector2D(33.5), 32);
+
+	UTEST_EQUAL("ClampBottomRight.X0Y0", ClampBottomRight.X0Y0, 32*32-1);
+	UTEST_EQUAL("ClampBottomRight.X1Y0", ClampBottomRight.X1Y0, 32*32-1);
+	UTEST_EQUAL("ClampBottomRight.X0Y1", ClampBottomRight.X0Y1, 32*32-1);
+	UTEST_EQUAL("ClampBottomRight.X1Y1", ClampBottomRight.X1Y1, 32*32-1);
+	UTEST_EQUAL("ClampBottomRight.XFraction", ClampBottomRight.XFraction, 0.0f);
+	UTEST_EQUAL("ClampBottomRight.YFraction", ClampBottomRight.YFraction, 0.0f);
+
+	PCGLandscapeCache::FSafeIndices FractionBottomRight = PCGLandscapeCache::CalcSafeIndices(FVector2D(30.4, 30.5), 32);
+
+	UTEST_EQUAL("FractionBottomRight.X0Y0", FractionBottomRight.X0Y0, 30 + 30*32);
+	UTEST_EQUAL("FractionBottomRight.X1Y0", FractionBottomRight.X1Y0, 31 + 30*32);
+	UTEST_EQUAL("FractionBottomRight.X0Y1", FractionBottomRight.X0Y1, 30 + 31*32);
+	UTEST_EQUAL("FractionBottomRight.X1Y1", FractionBottomRight.X1Y1, 31 + 31*32);
+	UTEST_EQUAL("FractionBottomRight.XFraction", FractionBottomRight.XFraction, 0.4f);
+	UTEST_EQUAL("FractionBottomRight.YFraction", FractionBottomRight.YFraction, 0.5f);
+
+	return true;
+}
+#endif
diff --git a/Engine/Plugins/Experimental/PCG/Source/PCG/Public/Grid/PCGLandscapeCache.h b/Engine/Plugins/Experimental/PCG/Source/PCG/Public/Grid/PCGLandscapeCache.h
--- a/Engine/Plugins/Experimental/PCG/Source/PCG/Public/Grid/PCGLandscapeCache.h
+++ b/Engine/Plugins/Experimental/PCG/Source/PCG/Public/Grid/PCGLandscapeCache.h
@@ -26,6 +26,23 @@ struct FPCGLandscapeLayerWeight
 	float Weight = 0.0f;
 };
 
+namespace PCGLandscapeCache
+{
+	struct FSafeIndices
+	{
+		int32 X0Y0 = 0;
+		int32 X1Y0 = 0;
+		int32 X0Y1 = 0;
+		int32 X1Y1 = 0;
+
+		float XFraction = 0;
+		float YFraction = 0;
+	};
+
+	// this will ensure all indicies are valid in a Stride*Stride sized array
+	FSafeIndices CalcSafeIndices(FVector2D LocalPosition, int32 Stride);
+}
+
 struct FPCGLandscapeCacheEntry
 {
 	friend class UPCGLandscapeCache;
@@ -44,7 +61,7 @@ private:
 	int32 GetMemorySize() const;
 
 #if WITH_EDITOR
-	void BuildCacheData(ULandscapeInfo* LandscapeInfo, ULandscapeComponent* InComponent);
+	static FPCGLandscapeCacheEntry* CreateCacheEntry(ULandscapeInfo* LandscapeInfo, ULandscapeComponent* InComponent);
 #endif
 
 	// Serialize called from the landscape cache
