[PCG] Add custom extractors for attribute selector

diff --git a/Engine/Plugins/Experimental/PCG/Source/PCG/Private/Elements/Metadata/PCGMetadataOpElementBase.cpp b/Engine/Plugins/Experimental/PCG/Source/PCG/Private/Elements/Metadata/PCGMetadataOpElementBase.cpp
--- a/Engine/Plugins/Experimental/PCG/Source/PCG/Private/Elements/Metadata/PCGMetadataOpElementBase.cpp
+++ b/Engine/Plugins/Experimental/PCG/Source/PCG/Private/Elements/Metadata/PCGMetadataOpElementBase.cpp
@@ -138,9 +138,6 @@ bool FPCGMetadataElementBase::ExecuteInternal(FPCGContext* Context) const
 		}
 
 		FPCGAttributePropertySelector InputSource = Settings->GetInputSource(i);
-		// Make sure to update it to transform it into a property if needed.
-		// TODO: Remove when it will be handled in the UI widget logic.
-		InputSource.Update();
 
 		if (InputSource.Selection == EPCGAttributePropertySelection::Attribute && InputSource.AttributeName == NAME_None)
 		{
@@ -237,8 +234,6 @@ bool FPCGMetadataElementBase::ExecuteInternal(FPCGContext* Context) const
 	OperationData.OutputKeys.SetNum(Settings->GetOutputPinNum());
 
 	FPCGAttributePropertySelector OutputTarget = Settings->OutputTarget;
-	// Make sure to update it to transform it into a property if needed.
-	OutputTarget.Update();
 
 	if (OutputTarget.Selection == EPCGAttributePropertySelection::Attribute && OutputTarget.AttributeName == NAME_None)
 	{
@@ -260,9 +255,12 @@ bool FPCGMetadataElementBase::ExecuteInternal(FPCGContext* Context) const
 
 		FName OutputName = OutputTarget.GetName();
 
-		if (OutputTarget.Selection == EPCGAttributePropertySelection::Attribute)
+		if (OutputTarget.Selection == EPCGAttributePropertySelection::Attribute && OutputTarget.ExtraNames.IsEmpty())
 		{
-			// In case of attribute, there is no point of failure before duplicating. So duplicate, create the attribute and then the accessor.
+			// In case of an attribute, we check if we have extra selectors. If not, we can just delete the attribute
+			// and create a new one of the right type.
+			// But if we have extra selectors, we need to handle it the same way as properties.
+			// There is no point of failure before duplicating. So duplicate, create the attribute and then the accessor.
 			PCGMetadataElementCommon::DuplicateTaggedData(InputTaggedData[InputToForward], OutputData, OutMetadata);
 			FPCGMetadataAttributeBase* OutputAttribute = PCGMetadataElementCommon::ClearOrCreateAttribute(OutMetadata, OutputName, AttributeType{});
 			if (!OutputAttribute)
@@ -278,21 +276,25 @@ bool FPCGMetadataElementBase::ExecuteInternal(FPCGContext* Context) const
 
 			OperationData.OutputAccessors[OutputIndex] = PCGAttributeAccessorHelpers::CreateAccessor(Cast<UPCGData>(OutputData.Data), OutputTarget);
 		}
-		else if (OutputTarget.Selection == EPCGAttributePropertySelection::PointProperty)
+		else
 		{
-			// In case of property, we need to validate that the property can accept the output type. Verify this before duplicating.
+			// In case of property or attribute with extra accessor, we need to validate that the property/attribute can accept the output type.
+			// Verify this before duplicating.
 			OperationData.OutputAccessors[OutputIndex] = PCGAttributeAccessorHelpers::CreateAccessor(Cast<UPCGData>(OutputData.Data), OutputTarget);
 
 			if (OperationData.OutputAccessors[OutputIndex].IsValid())
 			{
-				// We matched a property, check if the output type is valid
+				// We matched an attribute/property, check if the output type is valid.
 				if (!PCG::Private::IsBroadcastable(PCG::Private::MetadataTypes<AttributeType>::Id, OperationData.OutputAccessors[OutputIndex]->GetUnderlyingType()))
 				{
-					PCGE_LOG(Error, "Property %s cannot be broadcasted to match types for input", *OutputName.ToString());
+					PCGE_LOG(Error, "Attribute/Property %s cannot be broadcasted to match types for input", *OutputName.ToString());
 					return false;
 				}
 
 				PCGMetadataElementCommon::DuplicateTaggedData(InputTaggedData[InputToForward], OutputData, OutMetadata);
+
+				// Re-create the accessor to point to the right data (since we just duplicated the data)
+				OperationData.OutputAccessors[OutputIndex] = PCGAttributeAccessorHelpers::CreateAccessor(Cast<UPCGData>(OutputData.Data), OutputTarget);
 			}
 		}
 
diff --git a/Engine/Plugins/Experimental/PCG/Source/PCG/Private/Metadata/Accessors/PCGAttributeAccessorHelpers.cpp b/Engine/Plugins/Experimental/PCG/Source/PCG/Private/Metadata/Accessors/PCGAttributeAccessorHelpers.cpp
--- a/Engine/Plugins/Experimental/PCG/Source/PCG/Private/Metadata/Accessors/PCGAttributeAccessorHelpers.cpp
+++ b/Engine/Plugins/Experimental/PCG/Source/PCG/Private/Metadata/Accessors/PCGAttributeAccessorHelpers.cpp
@@ -2,17 +2,18 @@
 
 #include "Metadata/Accessors/PCGAttributeAccessorHelpers.h"
 
-#include "PCGPoint.h"
 #include "PCGParamData.h"
+#include "PCGPoint.h"
 #include "Data/PCGPointData.h"
 #include "Data/PCGSpatialData.h"
 #include "Helpers/PCGPropertyHelpers.h"
+#include "Metadata/PCGAttributePropertySelector.h"
 #include "Metadata/PCGMetadataAttributeTpl.h"
 #include "Metadata/PCGMetadataAttributeTraits.h"
-#include "Metadata/PCGAttributePropertySelector.h"
-#include "Metadata/Accessors/PCGPropertyAccessor.h"
 #include "Metadata/Accessors/PCGAttributeAccessor.h"
 #include "Metadata/Accessors/PCGAttributeAccessorKeys.h"
+#include "Metadata/Accessors/PCGPropertyAccessor.h"
+#include "Metadata/Accessors/PCGAttributeExtractor.h"
 
 #include "UObject/UnrealType.h"
 #include "UObject/EnumProperty.h"
@@ -53,6 +54,40 @@ namespace PCGAttributeAccessorHelpers
 		OutMetadata = Metadata;
 		OutAttribute = Attribute;
 	}
+
+	TUniquePtr<IPCGAttributeAccessor> CreateChainAccessor(TUniquePtr<IPCGAttributeAccessor> InAccessor, FName Name, bool& bOutSuccess)
+	{
+		if (!InAccessor.IsValid())
+		{
+			bOutSuccess = false;
+			return TUniquePtr<IPCGAttributeAccessor>();
+		}
+
+		auto Chain = [&Accessor = InAccessor, Name, &bOutSuccess](auto Dummy) -> TUniquePtr<IPCGAttributeAccessor>
+		{
+			using AccessorType = decltype(Dummy);
+
+			if constexpr (PCG::Private::IsOfTypes<AccessorType, FVector2D, FVector, FVector4, FQuat>())
+			{
+				return PCGAttributeExtractor::CreateVectorExtractor<AccessorType>(std::move(Accessor), Name, bOutSuccess);
+			}
+			if constexpr (PCG::Private::IsOfTypes<AccessorType, FTransform>())
+			{
+				return PCGAttributeExtractor::CreateTransformExtractor(std::move(Accessor), Name, bOutSuccess);
+			}
+			if constexpr (PCG::Private::IsOfTypes<AccessorType, FRotator>())
+			{
+				return PCGAttributeExtractor::CreateRotatorExtractor(std::move(Accessor), Name, bOutSuccess);
+			}
+			else
+			{
+				bOutSuccess = false;
+				return std::move(Accessor);
+			}
+		};
+
+		return PCGMetadataAttribute::CallbackWithRightType(InAccessor->GetUnderlyingType(), Chain);
+	}
 }
 
 TUniquePtr<IPCGAttributeAccessor> PCGAttributeAccessorHelpers::CreatePropertyAccessor(const FProperty* InProperty)
@@ -125,7 +160,8 @@ TUniquePtr<IPCGAttributeAccessor> PCGAttributeAccessorHelpers::CreatePropertyAcc
 
 TUniquePtr<const IPCGAttributeAccessor> PCGAttributeAccessorHelpers::CreateConstAccessor(const UPCGData* InData, const FPCGAttributePropertySelector& InSelector)
 {
-	FName Name = InSelector.GetName();
+	const FName Name = InSelector.GetName();
+	TUniquePtr<IPCGAttributeAccessor> Accessor;
 
 	if (InSelector.Selection == EPCGAttributePropertySelection::PointProperty)
 	{
@@ -133,47 +169,70 @@ TUniquePtr<const IPCGAttributeAccessor> PCGAttributeAccessorHelpers::CreateConst
 		{
 			if (const FProperty* Property = FPCGPoint::StaticStruct()->FindPropertyByName(Name))
 			{
-				return CreatePropertyAccessor(Property);
+				Accessor = CreatePropertyAccessor(Property);
 			}
 			else if (FPCGPoint::HasCustomPropertyGetterSetter(Name))
 			{
-				return FPCGPoint::CreateCustomPropertyAccessor(Name);
+				Accessor = FPCGPoint::CreateCustomPropertyAccessor(Name);
 			}
 		}
 	}
 
 	// At this point, it is not a point data or we didn't find a property.
 	// We can't continue if it is a property wanted.
-	if (InSelector.Selection == EPCGAttributePropertySelection::PointProperty)
+	if (InSelector.Selection == EPCGAttributePropertySelection::PointProperty && !Accessor.IsValid())
 	{
 		UE_LOG(LogPCG, Error, TEXT("[PCGAttributeAccessorHelpers::CreateConstAccessor] Expected to select a property but the data doesn't support this property."));
 		return TUniquePtr<const IPCGAttributeAccessor>();
 	}
 
-	const UPCGMetadata* Metadata = nullptr;
-	const FPCGMetadataAttributeBase* Attribute = nullptr;
+	if (InSelector.Selection == EPCGAttributePropertySelection::Attribute)
+	{
+		const UPCGMetadata* Metadata = nullptr;
+		const FPCGMetadataAttributeBase* Attribute = nullptr;
 
-	ExtractMetadataAtribute(InData, Name, Metadata, Attribute);
+		ExtractMetadataAtribute(InData, Name, Metadata, Attribute);
 
-	auto CreateTypedAccessor = [Attribute, Metadata](auto Dummy) -> TUniquePtr<const IPCGAttributeAccessor>
-	{
-		using AttributeType = decltype(Dummy);
-		return MakeUnique<FPCGAttributeAccessor<AttributeType>>(static_cast<const FPCGMetadataAttribute<AttributeType>*>(Attribute), Metadata);
-	};
+		auto CreateTypedAccessor = [Attribute, Metadata](auto Dummy) -> TUniquePtr<IPCGAttributeAccessor>
+		{
+			using AttributeType = decltype(Dummy);
+			return MakeUnique<FPCGAttributeAccessor<AttributeType>>(static_cast<const FPCGMetadataAttribute<AttributeType>*>(Attribute), Metadata);
+		};
 
-	if (Attribute && Metadata)
-	{
-		return PCGMetadataAttribute::CallbackWithRightType(Attribute->GetTypeId(), CreateTypedAccessor);
+		if (Attribute && Metadata)
+		{
+			Accessor = PCGMetadataAttribute::CallbackWithRightType(Attribute->GetTypeId(), CreateTypedAccessor);
+		}
+		else
+		{
+			return TUniquePtr<const IPCGAttributeAccessor>();
+		}
 	}
-	else
+
+	if (!Accessor.IsValid())
 	{
 		return TUniquePtr<const IPCGAttributeAccessor>();
 	}
+
+	// At this point, check if we have chain accessors
+	for (const FString& ExtraName : InSelector.ExtraNames)
+	{
+		bool bSuccess = false;
+		Accessor = CreateChainAccessor(std::move(Accessor), FName(ExtraName), bSuccess);
+		if (!bSuccess)
+		{
+			UE_LOG(LogPCG, Error, TEXT("[PCGAttributeAccessorHelpers::CreateConstAccessor] Extra selectors don't match existing properties."));
+			return TUniquePtr<const IPCGAttributeAccessor>();
+		}
+	}
+
+	return Accessor;
 }
 
 TUniquePtr<IPCGAttributeAccessor> PCGAttributeAccessorHelpers::CreateAccessor(UPCGData* InData, const FPCGAttributePropertySelector& InSelector)
 {
-	FName Name = InSelector.GetName();
+	const FName Name = InSelector.GetName();
+	TUniquePtr<IPCGAttributeAccessor> Accessor;
 
 	if (InSelector.Selection == EPCGAttributePropertySelection::PointProperty)
 	{
@@ -181,42 +240,64 @@ TUniquePtr<IPCGAttributeAccessor> PCGAttributeAccessorHelpers::CreateAccessor(UP
 		{
 			if (const FProperty* Property = FPCGPoint::StaticStruct()->FindPropertyByName(Name))
 			{
-				return CreatePropertyAccessor(Property);
+				Accessor = CreatePropertyAccessor(Property);
 			}
 			else if (FPCGPoint::HasCustomPropertyGetterSetter(Name))
 			{
-				return FPCGPoint::CreateCustomPropertyAccessor(Name);
+				Accessor = FPCGPoint::CreateCustomPropertyAccessor(Name);
 			}
 		}
 	}
 
 	// At this point, it is not a point data or we didn't find a property.
 	// We can't continue if it is a property wanted.
-	if (InSelector.Selection == EPCGAttributePropertySelection::PointProperty)
+	if (InSelector.Selection == EPCGAttributePropertySelection::PointProperty && !Accessor.IsValid())
 	{
 		UE_LOG(LogPCG, Error, TEXT("[PCGAttributeAccessorHelpers::CreateAccessor] Expected to select a property but the data doesn't support this property."));
 		return TUniquePtr<IPCGAttributeAccessor>();
 	}
 
-	UPCGMetadata* Metadata = nullptr;
-	FPCGMetadataAttributeBase* Attribute = nullptr;
+	if (InSelector.Selection == EPCGAttributePropertySelection::Attribute)
+	{
+		UPCGMetadata* Metadata = nullptr;
+		FPCGMetadataAttributeBase* Attribute = nullptr;
 
-	ExtractMetadataAtribute(InData, Name, Metadata, Attribute);
+		ExtractMetadataAtribute(InData, Name, Metadata, Attribute);
 
-	auto CreateTypedAccessor = [Attribute, Metadata](auto Dummy) -> TUniquePtr<IPCGAttributeAccessor>
-	{
-		using AttributeType = decltype(Dummy);
-		return MakeUnique<FPCGAttributeAccessor<AttributeType>>(static_cast<FPCGMetadataAttribute<AttributeType>*>(Attribute), Metadata);
-	};
+		auto CreateTypedAccessor = [Attribute, Metadata](auto Dummy) -> TUniquePtr<IPCGAttributeAccessor>
+		{
+			using AttributeType = decltype(Dummy);
+			return MakeUnique<FPCGAttributeAccessor<AttributeType>>(static_cast<FPCGMetadataAttribute<AttributeType>*>(Attribute), Metadata);
+		};
 
-	if (Attribute && Metadata)
-	{
-		return PCGMetadataAttribute::CallbackWithRightType(Attribute->GetTypeId(), CreateTypedAccessor);
+		if (Attribute && Metadata)
+		{
+			Accessor = PCGMetadataAttribute::CallbackWithRightType(Attribute->GetTypeId(), CreateTypedAccessor);
+		}
+		else
+		{
+			return TUniquePtr<IPCGAttributeAccessor>();
+		}
 	}
-	else
+
+	if (!Accessor.IsValid())
 	{
 		return TUniquePtr<IPCGAttributeAccessor>();
 	}
+
+	// At this point, check if we have chain accessors
+	for (const FString& ExtraName : InSelector.ExtraNames)
+	{
+		bool bSuccess = false;
+		Accessor = CreateChainAccessor(std::move(Accessor), FName(ExtraName), bSuccess);
+		if (!bSuccess)
+		{
+			UE_LOG(LogPCG, Error, TEXT("[PCGAttributeAccessorHelpers::CreateAccessor] Extra selectors don't match existing properties."));
+			return TUniquePtr<IPCGAttributeAccessor>();
+		}
+	}
+
+	return Accessor;
 }
 
 TUniquePtr<const IPCGAttributeAccessorKeys> PCGAttributeAccessorHelpers::CreateConstKeys(const UPCGData* InData, const FPCGAttributePropertySelector& InSelector)
diff --git a/Engine/Plugins/Experimental/PCG/Source/PCG/Private/Metadata/Accessors/PCGAttributeExtractor.cpp b/Engine/Plugins/Experimental/PCG/Source/PCG/Private/Metadata/Accessors/PCGAttributeExtractor.cpp
new file mode 100644
--- /dev/null
+++ b/Engine/Plugins/Experimental/PCG/Source/PCG/Private/Metadata/Accessors/PCGAttributeExtractor.cpp
@@ -0,0 +1,197 @@
+// Copyright Epic Games, Inc. All Rights Reserved.
+
+#include "Metadata/Accessors/PCGAttributeExtractor.h"
+#include "Metadata/Accessors/PCGCustomAccessor.h"
+
+#include "Internationalization/Regex.h"
+
+namespace PCGAttributeExtractor
+{
+	template <typename VectorType>
+	double GetAt(const VectorType& Value, int32 Index)
+	{
+		check(Index >= 0 && Index < 4);
+		return Value[Index];
+	}
+
+	template <>
+	double GetAt<FQuat>(const FQuat& Value, int32 Index)
+	{
+		check(Index >= 0 && Index < 4);
+		switch (Index)
+		{
+		case 0:
+			return Value.X;
+		case 1:
+			return Value.Y;
+		case 2:
+			return Value.Z;
+		default:
+			return Value.W;
+		}
+	}
+
+	template <typename VectorType>
+	void SetAt(VectorType& Value, double In, int32 Index)
+	{
+		check(Index >= 0 && Index < 4);
+		Value[Index] = In;
+	}
+
+	template <>
+	void SetAt<FQuat>(FQuat& Value, double In, int32 Index)
+	{
+		check(Index >= 0 && Index < 4);
+		switch (Index)
+		{
+		case 0:
+			Value.X = In;
+		case 1:
+			Value.Y = In;
+		case 2:
+			Value.Z = In;
+		default:
+			Value.W = In;
+		}
+	}
+
+	// Works for Vec2, Vec3, Vec4 and Quat (same as Vec4)
+	template <typename VectorType>
+	TUniquePtr<IPCGAttributeAccessor> CreateVectorExtractor(TUniquePtr<IPCGAttributeAccessor> InAccessor, FName Name, bool& bOutSuccess)
+	{
+		if ((Name == PCGAttributeExtractorConstants::VectorLength) || (Name == TEXT("Size")))
+		{
+			bOutSuccess = true;
+			return MakeUnique<FPCGChainAccessor<double, VectorType>>(std::move(InAccessor),
+				[](const VectorType& Value) -> double { return Value.Size(); });
+		}
+
+		FString MatchString;
+		if constexpr (std::is_same_v<FVector2D, VectorType>)
+		{
+			MatchString = TEXT("[XY]{1,4}");
+		}
+		else if constexpr (std::is_same_v<FVector, VectorType>)
+		{
+			MatchString = TEXT("[XYZ]{1,4}");
+		}
+		else
+		{
+			MatchString = TEXT("[XYZW]{1,4}");
+		}
+
+		const FString NameStr = Name.ToString();
+
+		const FRegexPattern RegexPattern(MatchString, ERegexPatternFlags::CaseInsensitive);
+		FRegexMatcher RegexMatcher(RegexPattern, *NameStr);
+		if (!RegexMatcher.FindNext())
+		{
+			// Failed
+			bOutSuccess = false;
+			return InAccessor;
+		}
+
+		TArray<int32, TInlineAllocator<4>> Indexes;
+		for (const TCHAR Char : NameStr)
+		{
+			int32 Index = static_cast<int32>(Char) - static_cast<int32>(TEXT('x'));
+			if (Index < 0)
+			{
+				Index = static_cast<int32>(Char) - static_cast<int32>(TEXT('X'));
+			}
+
+			Indexes.Add(Index);
+		}
+
+		bOutSuccess = true;
+
+		if (Indexes.Num() == 1)
+		{
+			return MakeUnique<FPCGChainAccessor<double, VectorType>>(std::move(InAccessor),
+				[Indexes](const VectorType& Value) -> double { return GetAt(Value, Indexes[0]); },
+				[Indexes](VectorType& Value, const double& In) -> void { SetAt(Value, In, Indexes[0]); });
+		}
+		else if (Indexes.Num() == 2)
+		{
+			return MakeUnique<FPCGChainAccessor<FVector2D, VectorType>>(std::move(InAccessor),
+				[Indexes](const VectorType& Value) -> FVector2D { return FVector2D(GetAt(Value, Indexes[0]), GetAt(Value, Indexes[1])); },
+				[Indexes](VectorType& Value, const FVector2D& In) -> void { SetAt(Value, In.X, Indexes[0]); SetAt(Value, In.Y, Indexes[1]);});
+		}
+		else if (Indexes.Num() == 3)
+		{
+			return MakeUnique<FPCGChainAccessor<FVector, VectorType>>(std::move(InAccessor),
+				[Indexes](const VectorType& Value) -> FVector { return FVector(GetAt(Value, Indexes[0]), GetAt(Value, Indexes[1]), GetAt(Value, Indexes[2])); },
+				[Indexes](VectorType& Value, const FVector& In) -> void { SetAt(Value, In.X, Indexes[0]); SetAt(Value, In.Y, Indexes[1]); SetAt(Value, In.Z, Indexes[2]); });
+		}
+		else
+		{
+			return MakeUnique<FPCGChainAccessor<FVector4, VectorType>>(std::move(InAccessor),
+				[Indexes](const VectorType& Value) -> FVector4 { return FVector4(GetAt(Value, Indexes[0]), GetAt(Value, Indexes[1]), GetAt(Value, Indexes[2]), GetAt(Value, Indexes[3])); },
+				[Indexes](VectorType& Value, const FVector4& In) -> void { SetAt(Value, In.X, Indexes[0]); SetAt(Value, In.Y, Indexes[1]); SetAt(Value, In.Z, Indexes[2]); SetAt(Value, In.W, Indexes[3]); });
+		}
+	}
+
+	TUniquePtr<IPCGAttributeAccessor> CreateRotatorExtractor(TUniquePtr<IPCGAttributeAccessor> InAccessor, FName Name, bool& bOutSuccess)
+	{
+		bOutSuccess = true;
+
+		if (Name == PCGAttributeExtractorConstants::RotatorPitch)
+		{
+			return MakeUnique<FPCGChainAccessor<double, FRotator>>(std::move(InAccessor),
+				[](const FRotator& Value) -> double { return Value.Pitch; },
+				[](FRotator& Value, const double& In) -> void { Value.Pitch = In; });
+		}
+
+		if (Name == PCGAttributeExtractorConstants::RotatorRoll)
+		{
+			return MakeUnique<FPCGChainAccessor<double, FRotator>>(std::move(InAccessor),
+				[](const FRotator& Value) -> double { return Value.Roll; },
+				[](FRotator& Value, const double& In) -> void { Value.Roll = In; });
+		}
+
+		if (Name == PCGAttributeExtractorConstants::RotatorYaw)
+		{
+			return MakeUnique<FPCGChainAccessor<double, FRotator>>(std::move(InAccessor),
+				[](const FRotator& Value) -> double { return Value.Yaw; },
+				[](FRotator& Value, const double& In) -> void { Value.Yaw = In; });
+		}
+
+		bOutSuccess = false;
+		return InAccessor;
+	}
+
+	TUniquePtr<IPCGAttributeAccessor> CreateTransformExtractor(TUniquePtr<IPCGAttributeAccessor> InAccessor, FName Name, bool& bOutSuccess)
+	{
+		bOutSuccess = true;
+
+		if ((Name == PCGAttributeExtractorConstants::TransformLocation) || (Name == TEXT("Position")))
+		{
+			return MakeUnique<FPCGChainAccessor<FVector, FTransform>>(std::move(InAccessor),
+				[](const FTransform& Value) -> FVector { return Value.GetLocation(); },
+				[](FTransform& Value, const FVector& In) -> void { Value.SetLocation(In); });
+		}
+
+		if ((Name == PCGAttributeExtractorConstants::TransformScale) || (Name == TEXT("Scale3D")))
+		{
+			return MakeUnique<FPCGChainAccessor<FVector, FTransform>>(std::move(InAccessor),
+				[](const FTransform& Value) -> FVector { return Value.GetScale3D(); },
+				[](FTransform& Value, const FVector& In) -> void { Value.SetScale3D(In); });
+		}
+
+		if (Name == PCGAttributeExtractorConstants::TransformRotation)
+		{
+			return MakeUnique<FPCGChainAccessor<FQuat, FTransform>>(std::move(InAccessor),
+				[](const FTransform& Value) -> FQuat { return Value.GetRotation(); },
+				[](FTransform& Value, const FQuat& In) -> void { Value.SetRotation(In); });
+		}
+
+		bOutSuccess = false;
+		return InAccessor;
+	}
+
+	// Template instantiation for all vectors types + quat
+	template TUniquePtr<IPCGAttributeAccessor> CreateVectorExtractor<FVector2D>(TUniquePtr<IPCGAttributeAccessor> InAccessor, FName Name, bool& bOutSuccess);
+	template TUniquePtr<IPCGAttributeAccessor> CreateVectorExtractor<FVector>(TUniquePtr<IPCGAttributeAccessor> InAccessor, FName Name, bool& bOutSuccess);
+	template TUniquePtr<IPCGAttributeAccessor> CreateVectorExtractor<FVector4>(TUniquePtr<IPCGAttributeAccessor> InAccessor, FName Name, bool& bOutSuccess);
+	template TUniquePtr<IPCGAttributeAccessor> CreateVectorExtractor<FQuat>(TUniquePtr<IPCGAttributeAccessor> InAccessor, FName Name, bool& bOutSuccess);
+}
\ No newline at end of file
diff --git a/Engine/Plugins/Experimental/PCG/Source/PCG/Private/Metadata/Accessors/PCGPropertyAccessor.cpp b/Engine/Plugins/Experimental/PCG/Source/PCG/Private/Metadata/Accessors/PCGPropertyAccessor.cpp
--- a/Engine/Plugins/Experimental/PCG/Source/PCG/Private/Metadata/Accessors/PCGPropertyAccessor.cpp
+++ b/Engine/Plugins/Experimental/PCG/Source/PCG/Private/Metadata/Accessors/PCGPropertyAccessor.cpp
@@ -4,7 +4,7 @@
 
 #include "UObject/EnumProperty.h"
 
-inline TArray<const void*> PCGPropertyAccessor::GetContainerKeys(int32 Index, int32 Range, const IPCGAttributeAccessorKeys& Keys)
+TArray<const void*> PCGPropertyAccessor::GetContainerKeys(int32 Index, int32 Range, const IPCGAttributeAccessorKeys& Keys)
 {
 	TArray<const void*> ContainerKeys;
 	ContainerKeys.SetNum(Range);
@@ -17,7 +17,7 @@ inline TArray<const void*> PCGPropertyAccessor::GetContainerKeys(int32 Index, in
 	return ContainerKeys;
 }
 
-inline TArray<void*> PCGPropertyAccessor::GetContainerKeys(int32 Index, int32 Range, IPCGAttributeAccessorKeys& Keys)
+TArray<void*> PCGPropertyAccessor::GetContainerKeys(int32 Index, int32 Range, IPCGAttributeAccessorKeys& Keys)
 {
 	TArray<void*> ContainerKeys;
 	ContainerKeys.SetNum(Range);
diff --git a/Engine/Plugins/Experimental/PCG/Source/PCG/Private/Metadata/PCGAttributePropertySelector.cpp b/Engine/Plugins/Experimental/PCG/Source/PCG/Private/Metadata/PCGAttributePropertySelector.cpp
--- a/Engine/Plugins/Experimental/PCG/Source/PCG/Private/Metadata/PCGAttributePropertySelector.cpp
+++ b/Engine/Plugins/Experimental/PCG/Source/PCG/Private/Metadata/PCGAttributePropertySelector.cpp
@@ -3,6 +3,14 @@
 #include "Metadata/PCGAttributePropertySelector.h"
 #include "Metadata/PCGMetadataAttribute.h"
 
+namespace PCGAttributePropertySelectorConstants
+{
+	const TCHAR* PropertyPrefix = TEXT("$");
+	const TCHAR* ExtraSeparator = TEXT(".");
+	const TCHAR PropertyPrefixChar = PropertyPrefix[0];
+	const TCHAR ExtraSeparatorChar = ExtraSeparator[0];
+}
+
 FName FPCGAttributePropertySelector::GetName() const
 {
 	switch (Selection)
@@ -24,12 +32,18 @@ FName FPCGAttributePropertySelector::GetName() const
 	}
 }
 
-bool FPCGAttributePropertySelector::SetPointProperty(EPCGPointProperties InPointProperty)
+bool FPCGAttributePropertySelector::SetPointProperty(EPCGPointProperties InPointProperty, bool bResetExtraNames)
 {
+	const bool bHasExtraNames = !ExtraNames.IsEmpty();
+	if (bResetExtraNames)
+	{
+		ExtraNames.Empty();
+	}
+
 	if (Selection == EPCGAttributePropertySelection::PointProperty && InPointProperty == PointProperty)
 	{
-		// Nothing changed
-		return false;
+		// Nothing changed, except perhaps the extra names
+		return (bHasExtraNames && bResetExtraNames);
 	}
 	else
 	{
@@ -39,12 +53,18 @@ bool FPCGAttributePropertySelector::SetPointProperty(EPCGPointProperties InPoint
 	}
 }
 
-bool FPCGAttributePropertySelector::SetAttributeName(FName InAttributeName)
+bool FPCGAttributePropertySelector::SetAttributeName(FName InAttributeName, bool bResetExtraNames)
 {
+	const bool bHasExtraNames = !ExtraNames.IsEmpty();
+	if (bResetExtraNames)
+	{
+		ExtraNames.Empty();
+	}
+
 	if (Selection == EPCGAttributePropertySelection::Attribute && AttributeName == InAttributeName)
 	{
-		// Nothing changed
-		return false;
+		// Nothing changed, except perhaps the extra names
+		return (bHasExtraNames && bResetExtraNames);
 	}
 	else
 	{
@@ -54,26 +74,6 @@ bool FPCGAttributePropertySelector::SetAttributeName(FName InAttributeName)
 	}
 }
 
-// TODO: To be removed when UI widget is done.
-void FPCGAttributePropertySelector::Update()
-{
-	const FString AttributeNameStr = AttributeName.ToString();
-
-	// If we attribute we select starts with a '$' and matches a property, we convert it to a property
-	if (!AttributeNameStr.IsEmpty() && AttributeNameStr[0] == TEXT('$'))
-	{
-		if (const UEnum* EnumPtr = StaticEnum<EPCGPointProperties>())
-		{
-			int32 Index = EnumPtr->GetIndexByNameString(AttributeNameStr.Mid(/*Start=*/1));
-			if (Index != INDEX_NONE)
-			{
-				SetPointProperty(static_cast<EPCGPointProperties>(EnumPtr->GetValueByIndex(Index)));
-				return;
-			}
-		}
-	}
-}
-
 #if WITH_EDITOR
 bool FPCGAttributePropertySelector::IsValid() const
 {
@@ -82,34 +82,65 @@ bool FPCGAttributePropertySelector::IsValid() const
 
 FText FPCGAttributePropertySelector::GetDisplayText() const
 {
+	FString Res;
 	const FName Name = GetName();
 
 	// Add a '$' if it is a property
 	if (Selection == EPCGAttributePropertySelection::PointProperty && (Name != NAME_None))
 	{
-		return FText::FromString(FString(TEXT("$")) + Name.ToString());
+		Res = FString(PCGAttributePropertySelectorConstants::PropertyPrefix) + Name.ToString();
 	}
 	else
 	{
-		return FText::FromName(Name);
+		Res = Name.ToString();
 	}
+
+	if (!ExtraNames.IsEmpty())
+	{
+		TArray<FString> AllNames;
+		AllNames.Add(Res);
+		AllNames.Append(ExtraNames);
+		Res = FString::Join(AllNames, PCGAttributePropertySelectorConstants::ExtraSeparator);
+	}
+
+	return FText::FromString(Res);
 }
 
 bool FPCGAttributePropertySelector::Update(FString NewValue)
 {
-	if (!NewValue.IsEmpty() && NewValue[0] == TEXT('$'))
+	TArray<FString> NewValues;
+	if (NewValue.IsEmpty())
+	{
+		NewValues.Add(NewValue);
+	}
+	else
+	{
+		NewValue.ParseIntoArray(NewValues, PCGAttributePropertySelectorConstants::ExtraSeparator, /*InCullEmpty=*/ false);
+	}
+
+	const FString NewName = NewValues[0];
+	TArray<FString> ExtraNamesTemp;
+	if (NewValues.Num() > 1)
+	{
+		ExtraNamesTemp.Append(&NewValues[1], NewValues.Num() - 1);
+	}
+
+	const bool bExtraChanged = ExtraNamesTemp != ExtraNames;
+	ExtraNames = MoveTemp(ExtraNamesTemp);
+
+	if (!NewName.IsEmpty() && NewName[0] == PCGAttributePropertySelectorConstants::PropertyPrefixChar)
 	{
 		if (const UEnum* EnumPtr = StaticEnum<EPCGPointProperties>())
 		{
-			int32 Index = EnumPtr->GetIndexByNameString(NewValue.Mid(/*Start=*/1));
+			int32 Index = EnumPtr->GetIndexByNameString(NewName.Mid(/*Start=*/1));
 			if (Index != INDEX_NONE)
 			{
-				return SetPointProperty(static_cast<EPCGPointProperties>(EnumPtr->GetValueByIndex(Index)));
+				return SetPointProperty(static_cast<EPCGPointProperties>(EnumPtr->GetValueByIndex(Index)), /*bResetExtraNames=*/ false) || bExtraChanged;
 			}
 		}
 	}
 
-	return SetAttributeName(NewValue.IsEmpty() ? NAME_None : FName(NewValue));
+	return SetAttributeName(NewName.IsEmpty() ? NAME_None : FName(NewName), /*bResetExtraNames=*/ false) || bExtraChanged;
 }
 #endif // WITH_EDITOR
 
diff --git a/Engine/Plugins/Experimental/PCG/Source/PCG/Private/Metadata/PCGMetadataAttribute.cpp b/Engine/Plugins/Experimental/PCG/Source/PCG/Private/Metadata/PCGMetadataAttribute.cpp
--- a/Engine/Plugins/Experimental/PCG/Source/PCG/Private/Metadata/PCGMetadataAttribute.cpp
+++ b/Engine/Plugins/Experimental/PCG/Source/PCG/Private/Metadata/PCGMetadataAttribute.cpp
@@ -77,7 +77,7 @@ void FPCGMetadataAttributeBase::ClearEntries()
 bool FPCGMetadataAttributeBase::IsValidName(const FString& Name)
 {
 	// A valid name is alphanumeric with some special characters allowed.
-	static const FString AllowedSpecialCharacters = TEXT("_.-");
+	static const FString AllowedSpecialCharacters = TEXT("_-");
 
 	for (int32 i = 0; i < Name.Len(); ++i)
 	{
diff --git a/Engine/Plugins/Experimental/PCG/Source/PCG/Public/Metadata/Accessors/IPCGAttributeAccessor.h b/Engine/Plugins/Experimental/PCG/Source/PCG/Public/Metadata/Accessors/IPCGAttributeAccessor.h
--- a/Engine/Plugins/Experimental/PCG/Source/PCG/Public/Metadata/Accessors/IPCGAttributeAccessor.h
+++ b/Engine/Plugins/Experimental/PCG/Source/PCG/Public/Metadata/Accessors/IPCGAttributeAccessor.h
@@ -134,6 +134,7 @@ public:
 	}
 
 	int16 GetUnderlyingType() const { return UnderlyingType; }
+	bool IsReadOnly() const { return bReadOnly; }
 
 protected:
 	IPCGAttributeAccessor(bool bInReadOnly, int16 InUnderlyingType)
diff --git a/Engine/Plugins/Experimental/PCG/Source/PCG/Public/Metadata/Accessors/PCGAttributeExtractor.h b/Engine/Plugins/Experimental/PCG/Source/PCG/Public/Metadata/Accessors/PCGAttributeExtractor.h
new file mode 100644
--- /dev/null
+++ b/Engine/Plugins/Experimental/PCG/Source/PCG/Public/Metadata/Accessors/PCGAttributeExtractor.h
@@ -0,0 +1,62 @@
+// Copyright Epic Games, Inc. All Rights Reserved.
+
+#pragma once
+
+#include "UObject/UnrealType.h"
+#include "Templates/UniquePtr.h"
+
+class IPCGAttributeAccessor;
+
+namespace PCGAttributeExtractorConstants
+{
+	// Vector/Quat
+	const FName VectorX = TEXT("X");
+	const FName VectorY = TEXT("Y");
+	const FName VectorZ = TEXT("Z");
+	const FName VectorW = TEXT("W");
+	const FName VectorLength = TEXT("Length");
+
+	const FName VectorExtractorNames[] =
+	{
+		VectorX,
+		VectorY,
+		VectorZ,
+		VectorW,
+		VectorLength
+	};
+
+	// Transform
+	const FName TransformLocation = TEXT("Location");
+	const FName TransformRotation = TEXT("Rotation");
+	const FName TransformScale = TEXT("Scale");
+
+	const FName TransformExtractorNames[] =
+	{
+		TransformLocation,
+		TransformRotation,
+		TransformScale
+	};
+
+	// Rotator
+	const FName RotatorPitch = TEXT("Pitch");
+	const FName RotatorYaw = TEXT("Yaw");
+	const FName RotatorRoll = TEXT("Roll");
+
+	const FName RotatorExtractorNames[] =
+	{
+		RotatorPitch,
+		RotatorYaw,
+		RotatorRoll
+	};
+}
+
+namespace PCGAttributeExtractor
+{
+	// Only instantiated for FVector2D, FVector, FVector4 and FQuat
+	template <typename VectorType>
+	TUniquePtr<IPCGAttributeAccessor> CreateVectorExtractor(TUniquePtr<IPCGAttributeAccessor> InAccessor, FName Name, bool& bOutSuccess);
+
+	TUniquePtr<IPCGAttributeAccessor> CreateRotatorExtractor(TUniquePtr<IPCGAttributeAccessor> InAccessor, FName Name, bool& bOutSuccess);
+
+	TUniquePtr<IPCGAttributeAccessor> CreateTransformExtractor(TUniquePtr<IPCGAttributeAccessor> InAccessor, FName Name, bool& bOutSuccess);
+}
diff --git a/Engine/Plugins/Experimental/PCG/Source/PCG/Public/Metadata/Accessors/PCGCustomAccessor.h b/Engine/Plugins/Experimental/PCG/Source/PCG/Public/Metadata/Accessors/PCGCustomAccessor.h
--- a/Engine/Plugins/Experimental/PCG/Source/PCG/Public/Metadata/Accessors/PCGCustomAccessor.h
+++ b/Engine/Plugins/Experimental/PCG/Source/PCG/Public/Metadata/Accessors/PCGCustomAccessor.h
@@ -98,4 +98,77 @@ public:
 
 private:
 	T Value;
-};
\ No newline at end of file
+};
+
+/**
+* To chain accessors. T is the type of this accessor. U is the type of the underlying accessor.
+* Key supported: Same as the underlying accessor
+*/
+template <typename T, typename U>
+class FPCGChainAccessor : public IPCGAttributeAccessorT<FPCGChainAccessor<T,U>>
+{
+public:
+	using ChainGetter = TFunction<T(const U&)>;
+	using ChainSetter = TFunction<void(U&, const T&)>;
+
+	using Type = T;
+	using OtherType = U;
+	using Super = IPCGAttributeAccessorT<FPCGChainAccessor<T, U>>;
+
+	FPCGChainAccessor(TUniquePtr<IPCGAttributeAccessor> InAccessor, const ChainGetter& InGetter)
+		: Super(/*bInReadOnly=*/ true)
+		, Accessor(std::move(InAccessor))
+		, Getter(InGetter)
+		, Setter()
+	{
+		check(Accessor);
+	}
+
+	FPCGChainAccessor(TUniquePtr<IPCGAttributeAccessor> InAccessor, const ChainGetter& InGetter, const ChainSetter& InSetter)
+		: Super(/*bInReadOnly=*/ !InAccessor || InAccessor->IsReadOnly())
+		, Accessor(std::move(InAccessor))
+		, Getter(InGetter)
+		, Setter(InSetter)
+	{
+		check(Accessor);
+	}
+
+	bool GetRangeImpl(TArrayView<T> OutValues, int32 Index, const IPCGAttributeAccessorKeys& Keys) const
+	{
+		TArray<U> TempValues;
+		TempValues.SetNumUninitialized(OutValues.Num());
+		if (!Accessor->GetRange<U>(TempValues, Index, Keys))
+		{
+			return false;
+		}
+
+		for (int32 i = 0; i < OutValues.Num(); ++i)
+		{
+			OutValues[i] = Getter(TempValues[i]);
+		}
+
+		return true;
+	}
+
+	bool SetRangeImpl(TArrayView<const T> InValues, int32 Index, IPCGAttributeAccessorKeys& Keys, EPCGAttributeAccessorFlags Flags)
+	{
+		TArray<U> TempValues;
+		TempValues.SetNumUninitialized(InValues.Num());
+		if (!Accessor->GetRange<U>(TempValues, Index, Keys))
+		{
+			return false;
+		}
+
+		for (int32 i = 0; i < InValues.Num(); ++i)
+		{
+			Setter(TempValues[i], InValues[i]);
+		}
+
+		return Accessor->SetRange<U>(TempValues, Index, Keys, Flags);
+	}
+
+private:
+	TUniquePtr<IPCGAttributeAccessor> Accessor;
+	ChainGetter Getter;
+	ChainSetter Setter;
+};
diff --git a/Engine/Plugins/Experimental/PCG/Source/PCG/Public/Metadata/PCGAttributePropertySelector.h b/Engine/Plugins/Experimental/PCG/Source/PCG/Public/Metadata/PCGAttributePropertySelector.h
--- a/Engine/Plugins/Experimental/PCG/Source/PCG/Public/Metadata/PCGAttributePropertySelector.h
+++ b/Engine/Plugins/Experimental/PCG/Source/PCG/Public/Metadata/PCGAttributePropertySelector.h
@@ -35,15 +35,15 @@ public:
 	UPROPERTY(BlueprintReadOnly, VisibleAnywhere, Category = "Settings", meta = (EditCondition = "Selection == EPCGAttributePropertySelection::PointProperty", EditConditionHides))
 	EPCGPointProperties PointProperty = EPCGPointProperties::Position;
 
+	UPROPERTY()
+	TArray<FString> ExtraNames;
+
 	// Setters, retrurn true if something changed.
-	bool SetPointProperty(EPCGPointProperties InPointProperty);
-	bool SetAttributeName(FName InAttributeName);
+	bool SetPointProperty(EPCGPointProperties InPointProperty, bool bResetExtraNames = true);
+	bool SetAttributeName(FName InAttributeName, bool bResetExtraNames = true);
 
 	FName GetName() const;
 
-	// TODO: To be removed when UI widget is done.
-	void Update();
-
 #if WITH_EDITOR
 	// Return true if the underlying name is valid.
 	bool IsValid() const;
diff --git a/Engine/Plugins/Experimental/PCG/Source/PCGEditor/Private/Details/PCGAttributePropertySelectorDetails.cpp b/Engine/Plugins/Experimental/PCG/Source/PCGEditor/Private/Details/PCGAttributePropertySelectorDetails.cpp
--- a/Engine/Plugins/Experimental/PCG/Source/PCGEditor/Private/Details/PCGAttributePropertySelectorDetails.cpp
+++ b/Engine/Plugins/Experimental/PCG/Source/PCGEditor/Private/Details/PCGAttributePropertySelectorDetails.cpp
@@ -36,7 +36,7 @@ void FPCGAttributePropertySelectorDetails::CustomizeHeader(TSharedRef<IPropertyH
 			PropertyHandle->CreatePropertyNameWidget()
 		]
 		.ValueContent()
-		.MinDesiredWidth(250.0f)
+		.MinDesiredWidth(350.0f)
 		[
 			SNew(SHorizontalBox)
 			+ SHorizontalBox::Slot()
