[PCG] Add more types to the pin/edge colours to add more granularity.

All will take the Spatial colour for now so appearance will not changed until the colours are set up.

Massages types again. Composite and Spatial were interchangeable which made things report Spatial as Composite. Composite now has its own bit and is distinct from Concrete.

diff --git a/Engine/Plugins/Experimental/PCG/Source/PCG/Public/PCGCommon.h b/Engine/Plugins/Experimental/PCG/Source/PCG/Public/PCGCommon.h
--- a/Engine/Plugins/Experimental/PCG/Source/PCG/Public/PCGCommon.h
+++ b/Engine/Plugins/Experimental/PCG/Source/PCG/Public/PCGCommon.h
@@ -43,10 +43,7 @@ enum class EPCGDataType : uint32
 	Concrete = Point | PolyLine | Surface | Volume | Primitive,
 
 	/** Boolean operations like union, difference, intersection. */
-	ShapeBoolean = 1 << 9 UMETA(Hidden),
-
-	/** Composite is any combination of concrete and booleans. */
-	Composite = ShapeBoolean | Concrete,
+	Composite = 1 << 9,
 
 	Spatial = Composite | Concrete,
 
diff --git a/Engine/Plugins/Experimental/PCG/Source/PCGEditor/Private/PCGEditorCommon.h b/Engine/Plugins/Experimental/PCG/Source/PCGEditor/Private/PCGEditorCommon.h
--- a/Engine/Plugins/Experimental/PCG/Source/PCGEditor/Private/PCGEditorCommon.h
+++ b/Engine/Plugins/Experimental/PCG/Source/PCGEditor/Private/PCGEditorCommon.h
@@ -38,6 +38,8 @@ namespace FPCGEditorCommon
 	const FName PointDataType = FName(TEXT("Point Data"));
 	const FName PolyLineDataType = FName(TEXT("Poly Line Data"));
 	const FName SurfaceDataType = FName(TEXT("Surface Data"));
+	const FName LandscapeDataType = FName(TEXT("Landscape Data"));
+	const FName TextureDataType = FName(TEXT("Texture Data"));
 	const FName RenderTargetDataType = FName(TEXT("Render Target Data"));
 	const FName VolumeDataType = FName(TEXT("Volume Data"));
 	const FName PrimitiveDataType = FName(TEXT("Primitive Data"));	
diff --git a/Engine/Plugins/Experimental/PCG/Source/PCGEditor/Private/PCGEditorGraphNodeBase.cpp b/Engine/Plugins/Experimental/PCG/Source/PCGEditor/Private/PCGEditorGraphNodeBase.cpp
--- a/Engine/Plugins/Experimental/PCG/Source/PCGEditor/Private/PCGEditorGraphNodeBase.cpp
+++ b/Engine/Plugins/Experimental/PCG/Source/PCGEditor/Private/PCGEditorGraphNodeBase.cpp
@@ -374,21 +374,9 @@ FEdGraphPinType UPCGEditorGraphNodeBase::GetPinType(const UPCGPin* InPin)
 		return !!(PinType & AllowedType) && !(PinType & ~AllowedType);
 	};
 
-	if (CheckType(EPCGDataType::Spatial))
+	if (CheckType(EPCGDataType::Concrete))
 	{
-		// Categories within spatial, checking from narrow -> wide
-		if (CheckType(EPCGDataType::Concrete))
-		{
-			EdPinType.PinCategory = FPCGEditorCommon::ConcreteDataType;
-		}
-		else if (CheckType(EPCGDataType::Composite))
-		{
-			EdPinType.PinCategory = FPCGEditorCommon::CompositeDataType;
-		}
-		else
-		{
-			EdPinType.PinCategory = FPCGEditorCommon::SpatialDataType;
-		}
+		EdPinType.PinCategory = FPCGEditorCommon::ConcreteDataType;
 
 		// Assign subcategory if we have precise information
 		if (CheckType(EPCGDataType::Point))
@@ -399,6 +387,14 @@ FEdGraphPinType UPCGEditorGraphNodeBase::GetPinType(const UPCGPin* InPin)
 		{
 			EdPinType.PinSubCategory = FPCGEditorCommon::PolyLineDataType;
 		}
+		else if (CheckType(EPCGDataType::Landscape))
+		{
+			EdPinType.PinSubCategory = FPCGEditorCommon::LandscapeDataType;
+		}
+		else if (CheckType(EPCGDataType::Texture))
+		{
+			EdPinType.PinSubCategory = FPCGEditorCommon::TextureDataType;
+		}
 		else if (CheckType(EPCGDataType::RenderTarget))
 		{
 			EdPinType.PinSubCategory = FPCGEditorCommon::RenderTargetDataType;
@@ -416,6 +412,14 @@ FEdGraphPinType UPCGEditorGraphNodeBase::GetPinType(const UPCGPin* InPin)
 			EdPinType.PinSubCategory = FPCGEditorCommon::PrimitiveDataType;
 		}
 	}
+	else if (CheckType(EPCGDataType::Composite))
+	{
+		EdPinType.PinCategory = FPCGEditorCommon::CompositeDataType;
+	}
+	else if (CheckType(EPCGDataType::Spatial))
+	{
+		EdPinType.PinCategory = FPCGEditorCommon::SpatialDataType;
+	}
 	else if (CheckType(EPCGDataType::Param))
 	{
 		EdPinType.PinCategory = FPCGEditorCommon::ParamDataType;
diff --git a/Engine/Plugins/Experimental/PCG/Source/PCGEditor/Private/PCGEditorSettings.cpp b/Engine/Plugins/Experimental/PCG/Source/PCGEditor/Private/PCGEditorSettings.cpp
--- a/Engine/Plugins/Experimental/PCG/Source/PCGEditor/Private/PCGEditorSettings.cpp
+++ b/Engine/Plugins/Experimental/PCG/Source/PCGEditor/Private/PCGEditorSettings.cpp
@@ -23,10 +23,19 @@ UPCGEditorSettings::UPCGEditorSettings(const FObjectInitializer& ObjectInitializ
 	DebugNodeColor = FLinearColor(1.0f, 0.0f, 1.0f);
 
 	DefaultPinColor = FLinearColor(1.0f, 1.0f, 1.0f);
+
 	SpatialDataPinColor = FLinearColor(0.2f, 0.2f, 1.0f);
 	CompositeDataPinColor = FLinearColor(0.2f, 0.2f, 1.0f);
 	ConcreteDataPinColor = FLinearColor(0.2f, 0.2f, 1.0f);
+	PointDataPinColor = FLinearColor(0.2f, 0.2f, 1.0f);
+	PolyLineDataPinColor = FLinearColor(0.2f, 0.2f, 1.0f);
+	LandscapeDataPinColor = FLinearColor(0.2f, 0.2f, 1.0f);
+	TextureDataPinColor = FLinearColor(0.2f, 0.2f, 1.0f);
 	RenderTargetDataPinColor = FLinearColor(1.0f, 0.3f, 0.f);
+	SurfaceDataPinColor = FLinearColor(0.2f, 0.2f, 1.0f);
+	VolumeDataPinColor = FLinearColor(0.2f, 0.2f, 1.0f);
+	PrimitiveDataPinColor = FLinearColor(0.2f, 0.2f, 1.0f);
+
 	ParamDataPinColor = FLinearColor(1.0f, 0.6f, 0.0f);
 	UnknownDataPinColor = FLinearColor(0.3f, 0.3f, 0.3f);
 
@@ -98,26 +107,54 @@ FLinearColor UPCGEditorSettings::GetColor(UPCGSettings* Settings) const
 
 FLinearColor UPCGEditorSettings::GetPinColor(const FEdGraphPinType& PinType) const
 {
-	if (PinType.PinCategory == FPCGEditorCommon::SpatialDataType || PinType.PinCategory == FPCGEditorCommon::ConcreteDataType)
+	if (PinType.PinCategory == FPCGEditorCommon::ConcreteDataType)
 	{
 		// Clauses below try to pick the narrowest type possible, falling back to Spatial
-		if (PinType.PinSubCategory == FPCGEditorCommon::RenderTargetDataType)
+		if (PinType.PinSubCategory == FPCGEditorCommon::PointDataType)
+		{
+			return PointDataPinColor;
+		}
+		else if (PinType.PinSubCategory == FPCGEditorCommon::PolyLineDataType)
+		{
+			return PolyLineDataPinColor;
+		}
+		else if (PinType.PinSubCategory == FPCGEditorCommon::LandscapeDataType)
+		{
+			return LandscapeDataPinColor;
+		}
+		else if (PinType.PinSubCategory == FPCGEditorCommon::TextureDataType)
+		{
+			return TextureDataPinColor;
+		}
+		else if (PinType.PinSubCategory == FPCGEditorCommon::RenderTargetDataType)
 		{
 			return RenderTargetDataPinColor;
 		}
-		else if (PinType.PinCategory == FPCGEditorCommon::ConcreteDataType)
+		else if (PinType.PinSubCategory == FPCGEditorCommon::SurfaceDataType)
 		{
-			return ConcreteDataPinColor;
+			return SurfaceDataPinColor;
+		}
+		else if (PinType.PinSubCategory == FPCGEditorCommon::VolumeDataType)
+		{
+			return VolumeDataPinColor;
+		}
+		else if (PinType.PinSubCategory == FPCGEditorCommon::PrimitiveDataType)
+		{
+			return PrimitiveDataPinColor;
 		}
 		else
 		{
-			return SpatialDataPinColor;
+			return ConcreteDataPinColor;
 		}
 	}
 	else if (PinType.PinCategory == FPCGEditorCommon::CompositeDataType)
 	{
 		return CompositeDataPinColor;
 	}
+	else if (PinType.PinCategory == FPCGEditorCommon::SpatialDataType)
+	{
+		return SpatialDataPinColor;
+	}
 	else if (PinType.PinCategory == FPCGEditorCommon::ParamDataType)
 	{
 		return ParamDataPinColor;
diff --git a/Engine/Plugins/Experimental/PCG/Source/PCGEditor/Public/PCGEditorSettings.h b/Engine/Plugins/Experimental/PCG/Source/PCGEditor/Public/PCGEditorSettings.h
--- a/Engine/Plugins/Experimental/PCG/Source/PCGEditor/Public/PCGEditorSettings.h
+++ b/Engine/Plugins/Experimental/PCG/Source/PCGEditor/Public/PCGEditorSettings.h
@@ -81,18 +81,46 @@ public:
 	UPROPERTY(EditAnywhere, config, Category = Node, meta = (HideAlphaChannel))
 	FLinearColor SpatialDataPinColor;
 
+	/** Color used for composite shape spatial data pins */
+	UPROPERTY(EditAnywhere, config, Category = Node, meta = (HideAlphaChannel))
+	FLinearColor CompositeDataPinColor;	
+	
 	/** Color used for concrete/simple spatial data pins */
 	UPROPERTY(EditAnywhere, config, Category = Node, meta = (HideAlphaChannel))
 	FLinearColor ConcreteDataPinColor;
 
-	/** Color used for composite shape spatial data pins */
+	/** Color used for data pins of type Point */
+	UPROPERTY(EditAnywhere, config, Category = Node, meta = (HideAlphaChannel))
+	FLinearColor PointDataPinColor;
+
+	/** Color used for data pins of type Spline */
+	UPROPERTY(EditAnywhere, config, Category = Node, meta = (HideAlphaChannel))
+	FLinearColor PolyLineDataPinColor;
+
+	/** Color used for data pins of type Surface */
 	UPROPERTY(EditAnywhere, config, Category = Node, meta = (HideAlphaChannel))
-	FLinearColor CompositeDataPinColor;
+	FLinearColor SurfaceDataPinColor;
 
-	/** Color used for render targets / GPU pins */
+	/** Color used for data pins of type Landscape */
+	UPROPERTY(EditAnywhere, config, Category = Node, meta = (HideAlphaChannel))
+	FLinearColor LandscapeDataPinColor;
+
+	/** Color used for data pins of type Texture */
+	UPROPERTY(EditAnywhere, config, Category = Node, meta = (HideAlphaChannel))
+	FLinearColor TextureDataPinColor;
+
+	/** Color used for data pins of type Render Target */
 	UPROPERTY(EditAnywhere, config, Category = Node, meta = (HideAlphaChannel))
 	FLinearColor RenderTargetDataPinColor;
 
+	/** Color used for data pins of type Volume */
+	UPROPERTY(EditAnywhere, config, Category = Node, meta = (HideAlphaChannel))
+	FLinearColor VolumeDataPinColor;
+
+	/** Color used for data pins of type Primitive */
+	UPROPERTY(EditAnywhere, config, Category = Node, meta = (HideAlphaChannel))
+	FLinearColor PrimitiveDataPinColor;
+
 	/** Color used for param data pins */
 	UPROPERTY(EditAnywhere, config, Category = Node, meta = (HideAlphaChannel))
 	FLinearColor ParamDataPinColor;
