[PCG] Added filter button to the attribute list view.

diff --git a/Engine/Plugins/Experimental/PCG/Source/PCGEditor/Private/SPCGEditorGraphAttributeListView.cpp b/Engine/Plugins/Experimental/PCG/Source/PCGEditor/Private/SPCGEditorGraphAttributeListView.cpp
--- a/Engine/Plugins/Experimental/PCG/Source/PCGEditor/Private/SPCGEditorGraphAttributeListView.cpp
+++ b/Engine/Plugins/Experimental/PCG/Source/PCGEditor/Private/SPCGEditorGraphAttributeListView.cpp
@@ -2,22 +2,28 @@
 
 #include "SPCGEditorGraphAttributeListView.h"
 
-#include "Data/PCGPointData.h"
-#include "Data/PCGSpatialData.h"
-#include "Metadata/PCGMetadataAttribute.h"
-#include "Metadata/PCGMetadataAttributeTpl.h"
 #include "PCGComponent.h"
 #include "PCGData.h"
 #include "PCGEditor.h"
+#include "PCGEditorGraphNodeBase.h"
 #include "PCGNode.h"
 #include "PCGParamData.h"
-#include "PCGEditorGraphNodeBase.h"
+#include "Data/PCGPointData.h"
+#include "Data/PCGSpatialData.h"
+#include "Metadata/PCGMetadata.h"
+#include "Metadata/PCGMetadataAttribute.h"
+#include "Metadata/PCGMetadataAttributeTpl.h"
 
+#include "Framework/MultiBox/MultiBoxBuilder.h"
 #include "Styling/AppStyle.h"
-#include "Widgets/Input/SSearchBox.h"
-#include "Widgets/Layout/SScrollBox.h"
+#include "Styling/SlateBrush.h"
 #include "Widgets/SBoxPanel.h"
 #include "Widgets/SWidget.h"
+#include "Widgets/Images/SImage.h"
+#include "Widgets/Images/SLayeredImage.h"
+#include "Widgets/Input/SComboButton.h"
+#include "Widgets/Layout/SScrollBox.h"
+#include "Widgets/Views/SHeaderRow.h"
 
 #define LOCTEXT_NAMESPACE "SPCGEditorGraphAttributeListView"
 
@@ -439,6 +445,22 @@ void SPCGEditorGraphAttributeListView::Construct(const FArguments& InArgs, TShar
 			.Text(this, &SPCGEditorGraphAttributeListView::OnGenerateSelectedDataText)
 		];
 
+	TSharedPtr<SLayeredImage> FilterImage = SNew(SLayeredImage)
+		.Image(FAppStyle::Get().GetBrush("Icons.Filter"))
+		.ColorAndOpacity(FSlateColor::UseForeground());
+
+	FilterImage->AddLayer(TAttribute<const FSlateBrush*>(this, &SPCGEditorGraphAttributeListView::GetFilterBadgeIcon));
+	
+	SAssignNew(FilterButton, SComboButton)
+		.ForegroundColor(FSlateColor::UseStyle())
+		.HasDownArrow(false)
+		.OnGetMenuContent(this, &SPCGEditorGraphAttributeListView::OnGenerateFilterMenu)
+		.ContentPadding(1)
+		.ButtonContent()
+		[
+			FilterImage.ToSharedRef()
+		];
+
 	this->ChildSlot
 	[
 		SNew(SVerticalBox)
@@ -447,6 +469,11 @@ void SPCGEditorGraphAttributeListView::Construct(const FArguments& InArgs, TShar
 		[
 			SNew(SHorizontalBox)
 			+SHorizontalBox::Slot()
+            .AutoWidth()
+            [
+            	FilterButton->AsShared()
+            ]
+			+SHorizontalBox::Slot()
 			.AutoWidth()
 			[
 				DataComboBox->AsShared()
@@ -506,9 +533,7 @@ void SPCGEditorGraphAttributeListView::Construct(const FArguments& InArgs, TShar
 
 TSharedRef<SHeaderRow> SPCGEditorGraphAttributeListView::CreateHeaderRowWidget() const
 {
-	return SNew(SHeaderRow)
-			.ResizeMode(ESplitterResizeMode::FixedPosition)
-			.CanSelectGeneratedColumn(true);
+	return SNew(SHeaderRow);
 }
 
 void SPCGEditorGraphAttributeListView::OnDebugObjectChanged(UPCGComponent* InPCGComponent)
@@ -567,6 +592,8 @@ void SPCGEditorGraphAttributeListView::OnGenerateUpdated(UPCGComponent* /*InPCGC
 
 void SPCGEditorGraphAttributeListView::RefreshAttributeList()
 {
+	HiddenAttributes = ListViewHeader->GetHiddenColumnIds();
+	
 	// Swapping to an empty item list to force a widget clear, otherwise the widgets will try to update during add column and access invalid data
 	static const TArray<PCGListviewItemPtr> EmptyList;
 	ListView->SetItemsSource(&EmptyList);
@@ -709,6 +736,54 @@ void SPCGEditorGraphAttributeListView::RefreshDataComboBox()
 	}
 }
 
+const FSlateBrush* SPCGEditorGraphAttributeListView::GetFilterBadgeIcon() const
+{
+	for (const SHeaderRow::FColumn& Column : ListViewHeader->GetColumns())
+	{
+		if (!Column.bIsVisible)
+		{
+			return FAppStyle::Get().GetBrush("Icons.BadgeModified"); 
+		}
+	}
+
+	return nullptr;
+}
+
+TSharedRef<SWidget> SPCGEditorGraphAttributeListView::OnGenerateFilterMenu()
+{
+	FMenuBuilder MenuBuilder(false, nullptr);
+
+	MenuBuilder.AddMenuEntry(
+		LOCTEXT("ToggleAllAttributes", "Toggle All"),
+		LOCTEXT("ToggleAllAttributesTooltip", "Toggle visibility for all attributes"),
+		FSlateIcon(),
+		FUIAction(FExecuteAction::CreateSP(this, &SPCGEditorGraphAttributeListView::ToggleAllAttributes),
+			FCanExecuteAction(),
+			FGetActionCheckState::CreateSP(this, &SPCGEditorGraphAttributeListView::GetAnyAttributeEnabledState)),
+		NAME_None,
+		EUserInterfaceActionType::ToggleButton);
+
+	MenuBuilder.AddSeparator();
+	
+	const TIndirectArray<SHeaderRow::FColumn>& Columns = ListViewHeader->GetColumns();
+	TArray<FName> HiddenColumns = ListViewHeader->GetHiddenColumnIds();
+
+	for (const SHeaderRow::FColumn& Column : Columns)
+	{	
+		MenuBuilder.AddMenuEntry(
+			Column.DefaultText,
+			Column.DefaultTooltip,
+			FSlateIcon(),
+			FUIAction(FExecuteAction::CreateSP(this, &SPCGEditorGraphAttributeListView::ToggleAttribute, Column.ColumnId),
+				FCanExecuteAction(),
+				FIsActionChecked::CreateSP(this, &SPCGEditorGraphAttributeListView::IsAttributeEnabled, Column.ColumnId)),
+			NAME_None,
+			EUserInterfaceActionType::ToggleButton);
+	}
+	
+	return MenuBuilder.MakeWidget();
+}
+
 TSharedRef<SWidget> SPCGEditorGraphAttributeListView::OnGenerateDataWidget(TSharedPtr<FName> InItem) const
 {
 	return SNew(STextBlock).Text(FText::FromName(InItem.IsValid() ? *InItem : NAME_None));
@@ -816,6 +891,61 @@ void SPCGEditorGraphAttributeListView::GenerateColumnsFromMetadata(const UPCGMet
 	}
 }
 
+void SPCGEditorGraphAttributeListView::ToggleAllAttributes()
+{
+	const TArray<FName> HiddenColumns = ListViewHeader->GetHiddenColumnIds();
+	if (HiddenColumns.Num() > 0)
+	{
+		for (const FName& HiddenColumn : HiddenColumns)
+		{
+			ListViewHeader->SetShowGeneratedColumn(HiddenColumn, /*InShow=*/true);
+		}
+	}
+	else
+	{
+		const TIndirectArray<SHeaderRow::FColumn>& Columns = ListViewHeader->GetColumns();
+		for (const SHeaderRow::FColumn& Column : Columns)
+		{	
+			ListViewHeader->SetShowGeneratedColumn(Column.ColumnId, /*InShow=*/false);
+		}
+	}
+}
+
+void SPCGEditorGraphAttributeListView::ToggleAttribute(FName InAttributeName)
+{
+	ListViewHeader->SetShowGeneratedColumn(InAttributeName, !ListViewHeader->IsColumnVisible(InAttributeName));
+}
+
+ECheckBoxState SPCGEditorGraphAttributeListView::GetAnyAttributeEnabledState() const
+{
+	bool bAllEnabled = true;
+	bool bAnyEnabled = false;
+	
+	for (const SHeaderRow::FColumn& Column : ListViewHeader->GetColumns())
+	{
+		bAllEnabled &= Column.bIsVisible;
+		bAnyEnabled |= Column.bIsVisible;
+	}
+
+	if (bAllEnabled)
+	{
+		return ECheckBoxState::Checked;
+	}
+	else if (bAnyEnabled)
+	{
+		return ECheckBoxState::Undetermined; 
+	}
+	else
+	{
+		return ECheckBoxState::Unchecked;
+	}
+}
+
+bool SPCGEditorGraphAttributeListView::IsAttributeEnabled(FName InAttributeName) const
+{
+	return ListViewHeader->IsColumnVisible(InAttributeName);
+}
+
 TSharedRef<ITableRow> SPCGEditorGraphAttributeListView::OnGenerateRow(PCGListviewItemPtr Item, const TSharedRef<STableViewBase>& OwnerTable) const
 {
 	return SNew(SPCGListViewItemRow, OwnerTable, Item);
@@ -832,20 +962,23 @@ void SPCGEditorGraphAttributeListView::OnItemDoubleClicked(PCGListviewItemPtr It
 	}
 }
 
-void SPCGEditorGraphAttributeListView::AddColumn(const FName& InColumnID, const FText& ColumnLabel, float ColumnWidth, EHorizontalAlignment HeaderHAlign, EHorizontalAlignment CellHAlign)
+void SPCGEditorGraphAttributeListView::AddColumn(const FName& InColumnId, const FText& ColumnLabel, float ColumnWidth, EHorizontalAlignment HeaderHAlign, EHorizontalAlignment CellHAlign)
 {
 	SHeaderRow::FColumn::FArguments Arguments;
-	Arguments.ColumnId(InColumnID);
+	Arguments.ColumnId(InColumnId);
 	Arguments.DefaultLabel(ColumnLabel);
 	Arguments.ManualWidth(ColumnWidth);
 	Arguments.HAlignHeader(HeaderHAlign);
 	Arguments.HAlignCell(CellHAlign);
-	ListViewHeader->AddColumn(Arguments);
+
+	SHeaderRow::FColumn* NewColumn = new SHeaderRow::FColumn(Arguments);
+	NewColumn->bIsVisible = !HiddenAttributes.Contains(InColumnId); 
+	ListViewHeader->AddColumn(*NewColumn);
 }
 
-void SPCGEditorGraphAttributeListView::RemoveColumn(const FName& InColumnID)
+void SPCGEditorGraphAttributeListView::RemoveColumn(const FName& InColumnId)
 {
-	ListViewHeader->RemoveColumn(InColumnID);
+	ListViewHeader->RemoveColumn(InColumnId);
 }
 
 void SPCGEditorGraphAttributeListView::AddIndexColumn()
@@ -943,7 +1076,10 @@ void SPCGEditorGraphAttributeListView::AddMetadataColumn(const FName& InColumnId
 	ColumnArguments.HAlignHeader(EHorizontalAlignment::HAlign_Center);
 	ColumnArguments.HAlignCell(EHorizontalAlignment::HAlign_Right);
 	ColumnArguments.FillWidth(1.0f);
-	ListViewHeader->AddColumn(ColumnArguments);
+
+	SHeaderRow::FColumn* NewColumn = new SHeaderRow::FColumn(ColumnArguments);
+	NewColumn->bIsVisible = !HiddenAttributes.Contains(InColumnId);
+	ListViewHeader->AddColumn(*NewColumn);
 }
 
 void SPCGEditorGraphAttributeListView::RemoveMetadataColumns()
diff --git a/Engine/Plugins/Experimental/PCG/Source/PCGEditor/Private/SPCGEditorGraphAttributeListView.h b/Engine/Plugins/Experimental/PCG/Source/PCGEditor/Private/SPCGEditorGraphAttributeListView.h
--- a/Engine/Plugins/Experimental/PCG/Source/PCGEditor/Private/SPCGEditorGraphAttributeListView.h
+++ b/Engine/Plugins/Experimental/PCG/Source/PCGEditor/Private/SPCGEditorGraphAttributeListView.h
@@ -3,18 +3,23 @@
 #pragma once
 
 #include "CoreMinimal.h"
-#include "Metadata/PCGMetadata.h"
-#include "Widgets/Input/SComboBox.h"
 #include "Widgets/SCompoundWidget.h"
-#include "Widgets/Views/SHeaderRow.h"
+#include "Widgets/Input/SComboBox.h"
 #include "Widgets/Views/SListView.h"
+#include "Widgets/Views/STableRow.h"
 
 class FPCGEditor;
+class FPCGMetadataAttributeBase;
 class UPCGComponent;
 class UPCGEditorGraphNodeBase;
+class UPCGMetadata;
 class UPCGParamData;
 struct FPCGPoint;
 
+class SComboButton;
+class SHeaderRow;
+struct FSlateBrush;
+
 struct FPCGMetadataInfo
 {
 	FName MetadataId = NAME_None;
@@ -71,17 +76,24 @@ private:
 	void RefreshAttributeList();
 	void RefreshDataComboBox();
 
+	const FSlateBrush* GetFilterBadgeIcon() const;
+	TSharedRef<SWidget> OnGenerateFilterMenu();
 	TSharedRef<SWidget> OnGenerateDataWidget(TSharedPtr<FName> InItem) const;
 	void OnSelectionChanged(TSharedPtr<FName> Item, ESelectInfo::Type SelectInfo);
 	FText OnGenerateSelectedDataText() const;
 	int32 GetSelectedDataIndex() const;
 	void GenerateColumnsFromMetadata(const UPCGMetadata* PCGMetadata);
 
+	void ToggleAllAttributes();
+	void ToggleAttribute(FName InAttributeName);
+	ECheckBoxState GetAnyAttributeEnabledState() const;
+	bool IsAttributeEnabled(FName InAttributeName) const;
+
 	TSharedRef<ITableRow> OnGenerateRow(PCGListviewItemPtr Item, const TSharedRef<STableViewBase>& OwnerTable) const;
 	void OnItemDoubleClicked(PCGListviewItemPtr Item) const;
 
-	void AddColumn(const FName& InColumnID, const FText& ColumnLabel, float ColumnWidth, EHorizontalAlignment HeaderHAlign = HAlign_Center, EHorizontalAlignment CellHAlign = HAlign_Right);
-	void RemoveColumn(const FName& InColumnID);
+	void AddColumn(const FName& InColumnId, const FText& ColumnLabel, float ColumnWidth, EHorizontalAlignment HeaderHAlign = HAlign_Center, EHorizontalAlignment CellHAlign = HAlign_Right);
+	void RemoveColumn(const FName& InColumnId);
 
 	void AddIndexColumn();
 	void RemoveIndexColumn();
@@ -110,6 +122,8 @@ private:
 	
 	TSharedPtr<STextBlock> NodeNameTextBlock;
 	TSharedPtr<STextBlock> InfoTextBlock;
+	TSharedPtr<SComboButton> FilterButton;
 
 	TMap<FName, FPCGMetadataInfo> MetadataInfos;
+	TArray<FName> HiddenAttributes;	
 };
