Improved disable behavior on subgraph/spawn actor nodes by having a different compilation path

#rnx

diff --git a/Engine/Plugins/Experimental/PCG/Source/PCG/Private/Graph/PCGGraphCompiler.cpp b/Engine/Plugins/Experimental/PCG/Source/PCG/Private/Graph/PCGGraphCompiler.cpp
--- a/Engine/Plugins/Experimental/PCG/Source/PCG/Private/Graph/PCGGraphCompiler.cpp
+++ b/Engine/Plugins/Experimental/PCG/Source/PCG/Private/Graph/PCGGraphCompiler.cpp
@@ -37,7 +37,8 @@ TArray<FPCGGraphTask> FPCGGraphCompiler::CompileGraph(UPCGGraph* InGraph, FPCGTa
 		const UPCGBaseSubgraphNode* SubgraphNode = Cast<const UPCGBaseSubgraphNode>(Node);
 		UPCGGraph* Subgraph = SubgraphNode ? SubgraphNode->GetSubgraph() : nullptr;
 
-		if (SubgraphNode != nullptr && !SubgraphNode->bDynamicGraph && Subgraph)
+		const bool bIsNonDynamicAndNonDisabledSubgraphNode = (SubgraphNode != nullptr && !SubgraphNode->bDynamicGraph && Subgraph && SubgraphNode->GetSettings() && SubgraphNode->GetSettings()->bEnabled);
+		if (bIsNonDynamicAndNonDisabledSubgraphNode)
 		{
 			const FPCGTaskId PreId = NextId++;
 
diff --git a/Engine/Plugins/Experimental/PCG/Source/PCG/Private/Graph/PCGGraphExecutor.h b/Engine/Plugins/Experimental/PCG/Source/PCG/Private/Graph/PCGGraphExecutor.h
--- a/Engine/Plugins/Experimental/PCG/Source/PCG/Private/Graph/PCGGraphExecutor.h
+++ b/Engine/Plugins/Experimental/PCG/Source/PCG/Private/Graph/PCGGraphExecutor.h
@@ -181,7 +181,7 @@ public:
 
 protected:
 	virtual bool ExecuteInternal(FPCGContext* Context) const override;
-	virtual bool IsPassthrough() const override { return true; }
+	virtual bool IsPassthrough(const UPCGSettings* InSettings) const override { return true; }
 };
 
 class FPCGGenericElement : public FSimplePCGElement
diff --git a/Engine/Plugins/Experimental/PCG/Source/PCG/Private/PCGElement.cpp b/Engine/Plugins/Experimental/PCG/Source/PCG/Private/PCGElement.cpp
--- a/Engine/Plugins/Experimental/PCG/Source/PCG/Private/PCGElement.cpp
+++ b/Engine/Plugins/Experimental/PCG/Source/PCG/Private/PCGElement.cpp
@@ -335,7 +335,9 @@ void IPCGElement::CleanupAndValidateOutput(FPCGContext* Context) const
 	const UPCGSettingsInterface* SettingsInterface = Context->GetInputSettingsInterface();
 	const UPCGSettings* Settings = SettingsInterface ? SettingsInterface->GetSettings() : nullptr;
 
-	if (!IsPassthrough() && Settings)
+	// Implementation note - disabled passthrough nodes can happen only in subgraphs/ spawn actor nodes
+	// which will behave properly when disabled. 
+	if (Settings && !IsPassthrough(Settings))
 	{
 		// Cleanup any residual labels if the node isn't supposed to produce them
 		// TODO: this is a bit of a crutch, could be refactored out if we review the way we push tagged data
diff --git a/Engine/Plugins/Experimental/PCG/Source/PCG/Private/PCGSettings.cpp b/Engine/Plugins/Experimental/PCG/Source/PCG/Private/PCGSettings.cpp
--- a/Engine/Plugins/Experimental/PCG/Source/PCG/Private/PCGSettings.cpp
+++ b/Engine/Plugins/Experimental/PCG/Source/PCG/Private/PCGSettings.cpp
@@ -32,6 +32,21 @@ bool UPCGSettingsInterface::IsInstance() const
 	return this != GetSettings();
 }
 
+void UPCGSettingsInterface::SetEnabled(bool bInEnabled)
+{
+	if (bEnabled != bInEnabled)
+	{
+		bEnabled = bInEnabled;
+#if WITH_EDITOR
+		if (UPCGSettings* Settings = GetSettings())
+		{
+			const bool bIsStructuralChange = Settings->IsStructuralProperty(GET_MEMBER_NAME_CHECKED(UPCGSettingsInterface, bEnabled));
+			OnSettingsChangedDelegate.Broadcast(Settings, (bIsStructuralChange ? EPCGChangeType::Structural : EPCGChangeType::None) | EPCGChangeType::Settings);
+		}
+#endif
+	}
+}
+
 bool UPCGSettings::operator==(const UPCGSettings& Other) const
 {
 	if (this == &Other)
diff --git a/Engine/Plugins/Experimental/PCG/Source/PCG/Private/PCGSubgraph.cpp b/Engine/Plugins/Experimental/PCG/Source/PCG/Private/PCGSubgraph.cpp
--- a/Engine/Plugins/Experimental/PCG/Source/PCG/Private/PCGSubgraph.cpp
+++ b/Engine/Plugins/Experimental/PCG/Source/PCG/Private/PCGSubgraph.cpp
@@ -66,6 +66,11 @@ void UPCGBaseSubgraphSettings::GetTrackedActorTags(FPCGTagToSettingsMap& OutTagT
 	}
 }
 
+bool UPCGBaseSubgraphSettings::IsStructuralProperty(const FName& InPropertyName) const
+{
+	return InPropertyName == GET_MEMBER_NAME_CHECKED(UPCGSettingsInterface, bEnabled) || Super::IsStructuralProperty(InPropertyName);
+}
+
 void UPCGBaseSubgraphSettings::OnSubgraphChanged(UPCGGraph* InGraph, EPCGChangeType ChangeType)
 {
 	if (InGraph == GetSubgraph())
diff --git a/Engine/Plugins/Experimental/PCG/Source/PCG/Public/Elements/PCGSpawnActor.h b/Engine/Plugins/Experimental/PCG/Source/PCG/Public/Elements/PCGSpawnActor.h
--- a/Engine/Plugins/Experimental/PCG/Source/PCG/Public/Elements/PCGSpawnActor.h
+++ b/Engine/Plugins/Experimental/PCG/Source/PCG/Public/Elements/PCGSpawnActor.h
@@ -95,7 +95,7 @@ class FPCGSpawnActorElement : public FSimplePCGElement
 public:
 	virtual bool CanExecuteOnlyOnMainThread(FPCGContext* Context) const override { return true; }
 	virtual bool IsCacheable(const UPCGSettings* InSettings) const override { return false; }
-	virtual bool IsPassthrough() const override { return true; }
+	virtual bool IsPassthrough(const UPCGSettings* InSettings) const override { return !InSettings || InSettings->bEnabled; }
 
 protected:
 	virtual bool ExecuteInternal(FPCGContext* Context) const override;
diff --git a/Engine/Plugins/Experimental/PCG/Source/PCG/Public/PCGElement.h b/Engine/Plugins/Experimental/PCG/Source/PCG/Public/PCGElement.h
--- a/Engine/Plugins/Experimental/PCG/Source/PCG/Public/PCGElement.h
+++ b/Engine/Plugins/Experimental/PCG/Source/PCG/Public/PCGElement.h
@@ -100,7 +100,7 @@ protected:
 	/** Controls whether an element can skip its execution wholly when the input data has the cancelled tag */
 	virtual bool IsCancellable() const { return true; }
 	/** Used to specify that the element passes through the data without any manipulation - used to correct target pins, etc. */
-	virtual bool IsPassthrough() const { return false; }
+	virtual bool IsPassthrough(const UPCGSettings* InSettings) const { return false; }
 
 	/** Passes through data when the element is Disabled. Can be implemented to override what gets passed through. */
 	virtual void DisabledPassThroughData(FPCGContext* Context) const;
diff --git a/Engine/Plugins/Experimental/PCG/Source/PCG/Public/PCGSettings.h b/Engine/Plugins/Experimental/PCG/Source/PCG/Public/PCGSettings.h
--- a/Engine/Plugins/Experimental/PCG/Source/PCG/Public/PCGSettings.h
+++ b/Engine/Plugins/Experimental/PCG/Source/PCG/Public/PCGSettings.h
@@ -56,6 +56,8 @@ public:
 	virtual const UPCGSettings* GetSettings() const PURE_VIRTUAL(UPCGSettingsInterface::GetSettings, return nullptr;);
 
 	bool IsInstance() const;
+	/** Dedicated method to change enable state because some nodes have more complex behavior on enable/disable (such as subgraphs) */
+	void SetEnabled(bool bInEnabled);
 
 #if WITH_EDITOR
 	FOnPCGSettingsChanged OnSettingsChangedDelegate;
@@ -81,6 +83,8 @@ class PCG_API UPCGSettings : public UPCGSettingsInterface
 {
 	GENERATED_BODY()
 
+	friend class UPCGSettingsInterface;
+
 public:
 	// ~Begin UPCGData interface
 	virtual EPCGDataType GetDataType() const override { return EPCGDataType::Settings; }
@@ -258,5 +262,5 @@ public:
 
 protected:
 	virtual bool ExecuteInternal(FPCGContext* Context) const override;
-	virtual bool IsPassthrough() const override { return true; }
+	virtual bool IsPassthrough(const UPCGSettings* InSettings) const override { return true; }
 };
diff --git a/Engine/Plugins/Experimental/PCG/Source/PCG/Public/PCGSubgraph.h b/Engine/Plugins/Experimental/PCG/Source/PCG/Public/PCGSubgraph.h
--- a/Engine/Plugins/Experimental/PCG/Source/PCG/Public/PCGSubgraph.h
+++ b/Engine/Plugins/Experimental/PCG/Source/PCG/Public/PCGSubgraph.h
@@ -27,6 +27,7 @@ protected:
 
 	//~Begin UPCGSettings interface
 	virtual void GetTrackedActorTags(FPCGTagToSettingsMap& OutTagToSettings, TArray<TObjectPtr<const UPCGGraph>>& OutVisitedGraphs) const override;
+	virtual bool IsStructuralProperty(const FName& InPropertyName) const override;
 #endif
 
 	TArray<FPCGPinProperties> InputPinProperties() const override;
@@ -109,7 +110,7 @@ public:
 
 protected:
 	virtual bool ExecuteInternal(FPCGContext* Context) const override;
-	virtual bool IsPassthrough() const override { return true; }
+	virtual bool IsPassthrough(const UPCGSettings* InSettings) const override { return !InSettings || InSettings->bEnabled; }
 };
 
 class PCG_API FPCGInputForwardingElement : public FSimplePCGElement
@@ -119,6 +120,6 @@ public:
 
 protected:
 	virtual bool ExecuteInternal(FPCGContext* Context) const override;
-	virtual bool IsPassthrough() const override { return true; }
+	virtual bool IsPassthrough(const UPCGSettings* InSettings) const override { return true; }
 	FPCGDataCollection Input;
 };
diff --git a/Engine/Plugins/Experimental/PCG/Source/PCGEditor/Private/PCGEditor.cpp b/Engine/Plugins/Experimental/PCG/Source/PCGEditor/Private/PCGEditor.cpp
--- a/Engine/Plugins/Experimental/PCG/Source/PCGEditor/Private/PCGEditor.cpp
+++ b/Engine/Plugins/Experimental/PCG/Source/PCGEditor/Private/PCGEditor.cpp
@@ -1227,15 +1227,14 @@ void FPCGEditor::OnToggleEnabled()
 			if (PCGSettingsInterface->bEnabled != bNewCheckState)
 			{
 				PCGSettingsInterface->Modify();
-				PCGSettingsInterface->bEnabled = bNewCheckState;
-				PCGNode->OnNodeChangedDelegate.Broadcast(PCGNode, EPCGChangeType::Settings);
+				PCGSettingsInterface->SetEnabled(bNewCheckState);
 				bChanged = true;
 			}
 		}
 		
 		if (bChanged)
 		{
-			GraphEditorWidget->NotifyGraphChanged();	
+			GraphEditorWidget->NotifyGraphChanged();
 		}
 	}
 }
