[PCG] PCG Graph Editor window: Disable Cancel Generation button when graph is not scheduled or executing.

Gives visual indicator of when graph is generating (and avoids user interpreting the button being enabled as a signal that the graph is generating, which confused me slightly at first).

Adds a new function that searches for any tasks related to a graph. Verified via logging that this called once per (editor) frame per visible button.

diff --git a/Engine/Plugins/Experimental/PCG/Source/PCG/Private/Graph/PCGGraphExecutor.cpp b/Engine/Plugins/Experimental/PCG/Source/PCG/Private/Graph/PCGGraphExecutor.cpp
--- a/Engine/Plugins/Experimental/PCG/Source/PCG/Private/Graph/PCGGraphExecutor.cpp
+++ b/Engine/Plugins/Experimental/PCG/Source/PCG/Private/Graph/PCGGraphExecutor.cpp
@@ -229,6 +229,22 @@ void FPCGGraphExecutor::Cancel(TFunctionRef<bool(TWeakObjectPtr<UPCGComponent>)>
 #endif
 }
 
+bool FPCGGraphExecutor::IsGraphCurrentlyExecuting(UPCGGraph* InGraph)
+{
+	bool bAnyPresent = false;
+
+	// This makes use of the Cancel function which runs over all tasks, but it always returns false so no tasks are cancelled.
+	auto CheckIfAnyGraphTasksPresent = [InGraph, &bAnyPresent](TWeakObjectPtr<UPCGComponent> Component)
+	{
+		bAnyPresent |= Component.IsValid() && Component->GetGraph() == InGraph;
+		return false;
+	};
+
+	Cancel(CheckIfAnyGraphTasksPresent);
+
+	return bAnyPresent;
+}
+
 FPCGTaskId FPCGGraphExecutor::ScheduleGeneric(TFunction<bool()> InOperation, UPCGComponent* InSourceComponent, const TArray<FPCGTaskId>& TaskDependencies)
 {
 	// Since we have no context, the generic task will consume no input
diff --git a/Engine/Plugins/Experimental/PCG/Source/PCG/Private/Graph/PCGGraphExecutor.h b/Engine/Plugins/Experimental/PCG/Source/PCG/Private/Graph/PCGGraphExecutor.h
--- a/Engine/Plugins/Experimental/PCG/Source/PCG/Private/Graph/PCGGraphExecutor.h
+++ b/Engine/Plugins/Experimental/PCG/Source/PCG/Private/Graph/PCGGraphExecutor.h
@@ -85,6 +85,9 @@ public:
 	/** Cancels all tasks */
 	TArray<UPCGComponent*> CancelAll();
 
+	/** Returns true if any task is scheduled or executing for the given graph. */
+	bool IsGraphCurrentlyExecuting(UPCGGraph* InGraph);
+
 	// Back compatibility function. Use ScheduleGenericWithContext
 	FPCGTaskId ScheduleGeneric(TFunction<bool()> InOperation, UPCGComponent* InSourceComponent, const TArray<FPCGTaskId>& TaskDependencies);
 
diff --git a/Engine/Plugins/Experimental/PCG/Source/PCG/Private/PCGSubsystem.cpp b/Engine/Plugins/Experimental/PCG/Source/PCG/Private/PCGSubsystem.cpp
--- a/Engine/Plugins/Experimental/PCG/Source/PCG/Private/PCGSubsystem.cpp
+++ b/Engine/Plugins/Experimental/PCG/Source/PCG/Private/PCGSubsystem.cpp
@@ -455,6 +455,18 @@ void UPCGSubsystem::CancelAllGeneration()
 	}
 }
 
+bool UPCGSubsystem::IsGraphCurrentlyExecuting(UPCGGraph* Graph)
+{
+	check(GraphExecutor);
+
+	if (!Graph)
+	{
+		return false;
+	}
+
+	return GraphExecutor->IsGraphCurrentlyExecuting(Graph);
+}
+
 TArray<FPCGTaskId> UPCGSubsystem::DispatchToRegisteredLocalComponents(UPCGComponent* OriginalComponent, const TFunction<FPCGTaskId(UPCGComponent*)>& InFunc) const
 {
 	TRACE_CPUPROFILER_EVENT_SCOPE(UPCGSubsystem::DispatchToRegisteredLocalComponents);
diff --git a/Engine/Plugins/Experimental/PCG/Source/PCG/Public/PCGSubsystem.h b/Engine/Plugins/Experimental/PCG/Source/PCG/Public/PCGSubsystem.h
--- a/Engine/Plugins/Experimental/PCG/Source/PCG/Public/PCGSubsystem.h
+++ b/Engine/Plugins/Experimental/PCG/Source/PCG/Public/PCGSubsystem.h
@@ -91,6 +91,9 @@ public:
 	/** Cancels currently running generation on given graph */
 	void CancelGeneration(UPCGGraph* Graph);
 
+	/** Returns true if there are any tasks for this graph currently scheduled or executing. */
+	bool IsGraphCurrentlyExecuting(UPCGGraph* Graph);
+
 	/** Cancels everything running */
 	void CancelAllGeneration();
 
diff --git a/Engine/Plugins/Experimental/PCG/Source/PCGEditor/Private/PCGEditor.cpp b/Engine/Plugins/Experimental/PCG/Source/PCGEditor/Private/PCGEditor.cpp
--- a/Engine/Plugins/Experimental/PCG/Source/PCGEditor/Private/PCGEditor.cpp
+++ b/Engine/Plugins/Experimental/PCG/Source/PCGEditor/Private/PCGEditor.cpp
@@ -368,7 +368,8 @@ void FPCGEditor::BindCommands()
 
 	ToolkitCommands->MapAction(
 		PCGEditorCommands.CancelExecution,
-		FExecuteAction::CreateSP(this, &FPCGEditor::OnCancelExecution_Clicked));
+		FExecuteAction::CreateSP(this, &FPCGEditor::OnCancelExecution_Clicked),
+		FCanExecuteAction::CreateSP(this, &FPCGEditor::IsCurrentlyGenerating));
 
 	ToolkitCommands->MapAction(
 		PCGEditorCommands.RunDeterminismGraphTest,
@@ -456,12 +457,24 @@ void FPCGEditor::OnForceGraphRegeneration_Clicked()
 
 void FPCGEditor::OnCancelExecution_Clicked()
 {
-	if (PCGEditorGraph)
+	UPCGSubsystem* Subsystem = GetSubsystem();
+	if (PCGEditorGraph && Subsystem)
 	{
-		GEditor->GetEditorWorldContext().World()->GetSubsystem<UPCGSubsystem>()->CancelGeneration(PCGEditorGraph->GetPCGGraph());
+		Subsystem->CancelGeneration(PCGEditorGraph->GetPCGGraph());
 	}
 }
 
+bool FPCGEditor::IsCurrentlyGenerating() const
+{
+	UPCGSubsystem* Subsystem = GetSubsystem();
+	if (PCGGraphBeingEdited && Subsystem)
+	{
+		return Subsystem->IsGraphCurrentlyExecuting(PCGGraphBeingEdited);
+	}
+
+	return false;
+}
+
 bool FPCGEditor::CanRunDeterminismNodeTest() const
 {
 	check(GraphEditorWidget.IsValid());
@@ -2016,6 +2029,12 @@ bool FPCGEditor::IsVisibleProperty(const FPropertyAndParent& InPropertyAndParent
 	return true;
 }
 
+UPCGSubsystem* FPCGEditor::GetSubsystem()
+{
+	UWorld* World = GEditor ? GEditor->GetEditorWorldContext().World() : nullptr;
+	return World ? World->GetSubsystem<UPCGSubsystem>() : nullptr;
+}
+
 TSharedRef<SDockTab> FPCGEditor::SpawnTab_GraphEditor(const FSpawnTabArgs& Args)
 {
 	return SNew(SDockTab)
diff --git a/Engine/Plugins/Experimental/PCG/Source/PCGEditor/Private/PCGEditor.h b/Engine/Plugins/Experimental/PCG/Source/PCGEditor/Private/PCGEditor.h
--- a/Engine/Plugins/Experimental/PCG/Source/PCGEditor/Private/PCGEditor.h
+++ b/Engine/Plugins/Experimental/PCG/Source/PCGEditor/Private/PCGEditor.h
@@ -122,6 +122,9 @@ private:
 	/** Cancels the current execution of the selected graph */
 	void OnCancelExecution_Clicked();
 
+	/* Returns true if inspected graph is currently scheduled or executing */
+	bool IsCurrentlyGenerating() const;
+
 	/** Can determinism be tested on the selected node(s) */
 	bool CanRunDeterminismNodeTest() const;
 	/** Run the determinism test on the selected node(s) */
@@ -247,6 +250,9 @@ private:
 	/** Returns whether a property should be visible (used for instance vs. settings properties) */
 	bool IsVisibleProperty(const FPropertyAndParent& InPropertyAndParent, IDetailsView* InDetailsView) const;
 
+	/** Helper to get to the subsystem. */
+	static class UPCGSubsystem* GetSubsystem();
+
 	TSharedRef<SDockTab> SpawnTab_GraphEditor(const FSpawnTabArgs& Args);
 	TSharedRef<SDockTab> SpawnTab_PropertyDetails(const FSpawnTabArgs& Args);
 	TSharedRef<SDockTab> SpawnTab_Palette(const FSpawnTabArgs& Args);
