[PCG] Node migration bits and pieces

Mostly just general type narrowing, but some noteworthy bits:

* Collapse - I think this operation only makes sense for Spatial. No params were being used that need to be migrated to a params pin.
* Copy Points - just moves Param pin to last in order. Does not mess up existing edges.
* Projection - Projection onto boolean sets is not well defined or easy. User can likely easily get what they want. Also projecting something that is not concrete is not trivial. Previously this was solved via collapse. Which is basically what is expressed by making the Source concrete - user will easily be able to collapse complex things into this input, and then they'll know a collapse has happened.

As a reminder, wide -> narrow connections are currently allowed. E.g. Any or Spatial can be connected to Point. This will change later but preserves existing connections for now.

diff --git a/Engine/Plugins/Experimental/PCG/Source/PCG/Private/Elements/PCGCollapseElement.cpp b/Engine/Plugins/Experimental/PCG/Source/PCG/Private/Elements/PCGCollapseElement.cpp
--- a/Engine/Plugins/Experimental/PCG/Source/PCG/Private/Elements/PCGCollapseElement.cpp
+++ b/Engine/Plugins/Experimental/PCG/Source/PCG/Private/Elements/PCGCollapseElement.cpp
@@ -2,11 +2,19 @@
 
 #include "Elements/PCGCollapseElement.h"
 
-#include "Data/PCGSpatialData.h"
 #include "Data/PCGPointData.h"
+#include "Data/PCGSpatialData.h"
 
 #include UE_INLINE_GENERATED_CPP_BY_NAME(PCGCollapseElement)
 
+TArray<FPCGPinProperties> UPCGCollapseSettings::InputPinProperties() const
+{
+	TArray<FPCGPinProperties> PinProperties;
+	PinProperties.Emplace(PCGPinConstants::DefaultInputLabel, EPCGDataType::Spatial);
+
+	return PinProperties;
+}
+
 FPCGElementPtr UPCGCollapseSettings::CreateElement() const
 {
 	return MakeShared<FPCGCollapseElement>();
diff --git a/Engine/Plugins/Experimental/PCG/Source/PCG/Private/Elements/PCGCopyPoints.cpp b/Engine/Plugins/Experimental/PCG/Source/PCG/Private/Elements/PCGCopyPoints.cpp
--- a/Engine/Plugins/Experimental/PCG/Source/PCG/Private/Elements/PCGCopyPoints.cpp
+++ b/Engine/Plugins/Experimental/PCG/Source/PCG/Private/Elements/PCGCopyPoints.cpp
@@ -13,9 +13,9 @@
 TArray<FPCGPinProperties> UPCGCopyPointsSettings::InputPinProperties() const
 {
 	TArray<FPCGPinProperties> PinProperties;
+	PinProperties.Emplace(PCGCopyPointsConstants::SourcePointsLabel, EPCGDataType::Point, /*bAllowMultipleConnections=*/false);
+	PinProperties.Emplace(PCGCopyPointsConstants::TargetPointsLabel, EPCGDataType::Point, /*bAllowMultipleConnections=*/false);
 	PinProperties.Emplace(PCGCopyPointsConstants::ParamsLabel, EPCGDataType::Param, /*bAllowMultipleConnections=*/false);
-	PinProperties.Emplace(PCGCopyPointsConstants::SourcePointsLabel, EPCGDataType::Spatial, /*bAllowMultipleConnections=*/false);
-	PinProperties.Emplace(PCGCopyPointsConstants::TargetPointsLabel, EPCGDataType::Spatial, /*bAllowMultipleConnections=*/false);
 
 	return PinProperties;
 }
diff --git a/Engine/Plugins/Experimental/PCG/Source/PCG/Private/Elements/PCGPointExtentsModifier.cpp b/Engine/Plugins/Experimental/PCG/Source/PCG/Private/Elements/PCGPointExtentsModifier.cpp
--- a/Engine/Plugins/Experimental/PCG/Source/PCG/Private/Elements/PCGPointExtentsModifier.cpp
+++ b/Engine/Plugins/Experimental/PCG/Source/PCG/Private/Elements/PCGPointExtentsModifier.cpp
@@ -15,18 +15,11 @@ FPCGElementPtr UPCGPointExtentsModifierSettings::CreateElement() const
 TArray<FPCGPinProperties> UPCGPointExtentsModifierSettings::InputPinProperties() const
 {
 	TArray<FPCGPinProperties> PinProperties;
-	PinProperties.Emplace(PCGPinConstants::DefaultInputLabel, EPCGDataType::Spatial);	
+	PinProperties.Emplace(PCGPinConstants::DefaultInputLabel, EPCGDataType::Point);	
 	PinProperties.Emplace(PCGPinConstants::DefaultParamsLabel, EPCGDataType::Param);
 	return PinProperties;
 }
 
-TArray<FPCGPinProperties> UPCGPointExtentsModifierSettings::OutputPinProperties() const
-{
-	TArray<FPCGPinProperties> PinProperties;
-	PinProperties.Emplace(PCGPinConstants::DefaultOutputLabel, EPCGDataType::Spatial);
-	return PinProperties;
-}
-
 bool FPCGPointExtentsModifier::ExecuteInternal(FPCGContext* Context) const
 {
 	TRACE_CPUPROFILER_EVENT_SCOPE(FPCGPointExtentsModifier::Execute);
@@ -88,4 +81,3 @@ bool FPCGPointExtentsModifier::ExecuteInternal(FPCGContext* Context) const
 	
 	return true;
 }
-
diff --git a/Engine/Plugins/Experimental/PCG/Source/PCG/Private/Elements/PCGProjectionElement.cpp b/Engine/Plugins/Experimental/PCG/Source/PCG/Private/Elements/PCGProjectionElement.cpp
--- a/Engine/Plugins/Experimental/PCG/Source/PCG/Private/Elements/PCGProjectionElement.cpp
+++ b/Engine/Plugins/Experimental/PCG/Source/PCG/Private/Elements/PCGProjectionElement.cpp
@@ -13,8 +13,8 @@
 TArray<FPCGPinProperties> UPCGProjectionSettings::InputPinProperties() const
 {
 	TArray<FPCGPinProperties> PinProperties;
-	PinProperties.Emplace(PCGPinConstants::DefaultInputLabel, EPCGDataType::Spatial, /*bInAllowMultipleConnections=*/true, LOCTEXT("ProjectionSourcePinTooltip", "The data to project."));
-	PinProperties.Emplace(PCGProjectionConstants::ProjectionTargetLabel, EPCGDataType::Spatial, /*bAllowMultipleConnections=*/false, LOCTEXT("ProjectionTargetPinTooltip", "The projection target."));
+	PinProperties.Emplace(PCGPinConstants::DefaultInputLabel, EPCGDataType::Concrete, /*bInAllowMultipleConnections=*/true, LOCTEXT("ProjectionSourcePinTooltip", "The data to project."));
+	PinProperties.Emplace(PCGProjectionConstants::ProjectionTargetLabel, EPCGDataType::Concrete, /*bAllowMultipleConnections=*/false, LOCTEXT("ProjectionTargetPinTooltip", "The projection target."));
 	PinProperties.Emplace(PCGPinConstants::DefaultParamsLabel, EPCGDataType::Param, /*bInAllowMultipleConnections=*/false);
 
 	return PinProperties;
diff --git a/Engine/Plugins/Experimental/PCG/Source/PCG/Public/Elements/PCGCollapseElement.h b/Engine/Plugins/Experimental/PCG/Source/PCG/Public/Elements/PCGCollapseElement.h
--- a/Engine/Plugins/Experimental/PCG/Source/PCG/Public/Elements/PCGCollapseElement.h
+++ b/Engine/Plugins/Experimental/PCG/Source/PCG/Public/Elements/PCGCollapseElement.h
@@ -19,6 +19,7 @@ public:
 	virtual EPCGSettingsType GetType() const override { return EPCGSettingsType::Spatial; }
 #endif
 
+	virtual TArray<FPCGPinProperties> InputPinProperties() const override;
 	virtual TArray<FPCGPinProperties> OutputPinProperties() const override { return Super::DefaultPointOutputPinProperties(); }
 
 protected:
diff --git a/Engine/Plugins/Experimental/PCG/Source/PCG/Public/Elements/PCGPointExtentsModifier.h b/Engine/Plugins/Experimental/PCG/Source/PCG/Public/Elements/PCGPointExtentsModifier.h
--- a/Engine/Plugins/Experimental/PCG/Source/PCG/Public/Elements/PCGPointExtentsModifier.h
+++ b/Engine/Plugins/Experimental/PCG/Source/PCG/Public/Elements/PCGPointExtentsModifier.h
@@ -34,8 +34,8 @@ public:
 	virtual EPCGSettingsType GetType() const override { return EPCGSettingsType::Spatial; }
 #endif
 
-	virtual TArray<FPCGPinProperties> InputPinProperties() const;
-	virtual TArray<FPCGPinProperties> OutputPinProperties() const;
+	virtual TArray<FPCGPinProperties> InputPinProperties() const override;
+	virtual TArray<FPCGPinProperties> OutputPinProperties() const override { return Super::DefaultPointOutputPinProperties(); }
 
 protected:
 	virtual FPCGElementPtr CreateElement() const override;
