Added pin variations (icons) for PCG graph editor. Some groundwork to separate single data from multi data as well.

diff --git a/Engine/Plugins/Experimental/PCG/Source/PCG/Private/Elements/PCGProjectionElement.cpp b/Engine/Plugins/Experimental/PCG/Source/PCG/Private/Elements/PCGProjectionElement.cpp
--- a/Engine/Plugins/Experimental/PCG/Source/PCG/Private/Elements/PCGProjectionElement.cpp
+++ b/Engine/Plugins/Experimental/PCG/Source/PCG/Private/Elements/PCGProjectionElement.cpp
@@ -13,9 +13,9 @@
 TArray<FPCGPinProperties> UPCGProjectionSettings::InputPinProperties() const
 {
 	TArray<FPCGPinProperties> PinProperties;
-	PinProperties.Emplace(PCGPinConstants::DefaultInputLabel, EPCGDataType::Concrete, /*bInAllowMultipleConnections=*/true, LOCTEXT("ProjectionSourcePinTooltip", "The data to project."));
-	PinProperties.Emplace(PCGProjectionConstants::ProjectionTargetLabel, EPCGDataType::Concrete, /*bAllowMultipleConnections=*/false, LOCTEXT("ProjectionTargetPinTooltip", "The projection target."));
-	PinProperties.Emplace(PCGPinConstants::DefaultParamsLabel, EPCGDataType::Param, /*bInAllowMultipleConnections=*/false);
+	PinProperties.Emplace(PCGPinConstants::DefaultInputLabel, EPCGDataType::Concrete, /*bInAllowMultipleConnections=*/true, /*bAllowMultipleData=*/true,  LOCTEXT("ProjectionSourcePinTooltip", "The data to project."));
+	PinProperties.Emplace(PCGProjectionConstants::ProjectionTargetLabel, EPCGDataType::Concrete, /*bAllowMultipleConnections=*/false, /*bAllowMultipleData=*/false, LOCTEXT("ProjectionTargetPinTooltip", "The projection target."));
+	PinProperties.Emplace(PCGPinConstants::DefaultParamsLabel, EPCGDataType::Param, /*bInAllowMultipleConnections=*/false, /*bAllowMultipleData=*/false);
 
 	return PinProperties;
 }
diff --git a/Engine/Plugins/Experimental/PCG/Source/PCG/Private/PCGInputOutputSettings.cpp b/Engine/Plugins/Experimental/PCG/Source/PCG/Private/PCGInputOutputSettings.cpp
--- a/Engine/Plugins/Experimental/PCG/Source/PCG/Private/PCGInputOutputSettings.cpp
+++ b/Engine/Plugins/Experimental/PCG/Source/PCG/Private/PCGInputOutputSettings.cpp
@@ -67,16 +67,18 @@ void UPCGGraphInputOutputSettings::PostLoad()
 TArray<FPCGPinProperties> UPCGGraphInputOutputSettings::GetPinProperties() const
 {
 	TArray<FPCGPinProperties> PinProperties;
+	const bool bIsInputPin = bIsInput;
 	const EPCGDataType DefaultPinDataType = bIsInput ? EPCGDataType::Composite : EPCGDataType::Any;
-	Algo::Transform(StaticLabels(), PinProperties, [DefaultPinDataType](const FLabelAndTooltip& InLabelAndTooltip) {
-		return FPCGPinProperties(InLabelAndTooltip.Label, DefaultPinDataType, true, InLabelAndTooltip.Tooltip);
+	Algo::Transform(StaticLabels(), PinProperties, [bIsInputPin, DefaultPinDataType](const FLabelAndTooltip& InLabelAndTooltip) {
+		return FPCGPinProperties(InLabelAndTooltip.Label, DefaultPinDataType, /*bMultiConnections=*/true, /*bMultiData=*/true, InLabelAndTooltip.Tooltip);
 	});
 	
 	if (bShowAdvancedPins)
 	{
 		Algo::Transform(StaticAdvancedLabels(), PinProperties, [DefaultPinDataType](const FLabelAndTooltip& InLabelAndTooltip) {
-			const EPCGDataType PinType = (InLabelAndTooltip.Label == PCGInputOutputConstants::DefaultLandscapeLabel || InLabelAndTooltip.Label == PCGInputOutputConstants::DefaultLandscapeHeightLabel) ? EPCGDataType::Surface : DefaultPinDataType;
-			return FPCGPinProperties(InLabelAndTooltip.Label, PinType, true, InLabelAndTooltip.Tooltip);
+			const bool bIsLandscapePin = (InLabelAndTooltip.Label == PCGInputOutputConstants::DefaultLandscapeLabel || InLabelAndTooltip.Label == PCGInputOutputConstants::DefaultLandscapeHeightLabel);
+			const EPCGDataType PinType = bIsLandscapePin ? EPCGDataType::Surface : DefaultPinDataType;
+			return FPCGPinProperties(InLabelAndTooltip.Label, PinType, /*bMultiConnection=*/true, /*bMultiData=*/false, InLabelAndTooltip.Tooltip);
 		});
 	}
 
diff --git a/Engine/Plugins/Experimental/PCG/Source/PCG/Private/PCGPin.cpp b/Engine/Plugins/Experimental/PCG/Source/PCG/Private/PCGPin.cpp
--- a/Engine/Plugins/Experimental/PCG/Source/PCG/Private/PCGPin.cpp
+++ b/Engine/Plugins/Experimental/PCG/Source/PCG/Private/PCGPin.cpp
@@ -6,13 +6,16 @@
 
 #include UE_INLINE_GENERATED_CPP_BY_NAME(PCGPin)
 
-FPCGPinProperties::FPCGPinProperties(const FName& InLabel, EPCGDataType InAllowedTypes, bool bInAllowMultipleConnections, const FText& InTooltip)
-	: Label(InLabel), AllowedTypes(InAllowedTypes), bAllowMultipleConnections(bInAllowMultipleConnections), Tooltip(InTooltip)
+FPCGPinProperties::FPCGPinProperties(const FName& InLabel, EPCGDataType InAllowedTypes, bool bInAllowMultipleConnections, bool bInAllowMultipleData, const FText& InTooltip)
+	: Label(InLabel), AllowedTypes(InAllowedTypes), bAllowMultipleData(bInAllowMultipleData), bAllowMultipleConnections(bInAllowMultipleConnections), Tooltip(InTooltip)
 {}
 
 bool FPCGPinProperties::operator==(const FPCGPinProperties& Other) const
 {
-	return Label == Other.Label && AllowedTypes == Other.AllowedTypes && bAllowMultipleConnections == Other.bAllowMultipleConnections;
+	return Label == Other.Label &&
+		AllowedTypes == Other.AllowedTypes &&
+		bAllowMultipleConnections == Other.bAllowMultipleConnections &&
+		bAllowMultipleData == Other.bAllowMultipleData;
 }
 
 UPCGPin::UPCGPin(const FObjectInitializer& ObjectInitializer)
diff --git a/Engine/Plugins/Experimental/PCG/Source/PCG/Public/PCGPin.h b/Engine/Plugins/Experimental/PCG/Source/PCG/Public/PCGPin.h
--- a/Engine/Plugins/Experimental/PCG/Source/PCG/Public/PCGPin.h
+++ b/Engine/Plugins/Experimental/PCG/Source/PCG/Public/PCGPin.h
@@ -16,7 +16,7 @@ struct PCG_API FPCGPinProperties
 	GENERATED_BODY()
 
 	FPCGPinProperties() = default;
-	explicit FPCGPinProperties(const FName& InLabel, EPCGDataType InAllowedTypes = EPCGDataType::Any, bool bInAllowMultipleConnections = true, const FText& InTooltip = FText::GetEmpty());
+	explicit FPCGPinProperties(const FName& InLabel, EPCGDataType InAllowedTypes = EPCGDataType::Any, bool bInAllowMultipleConnections = true, bool bAllowMultipleData = true, const FText& InTooltip = FText::GetEmpty());
 
 	bool operator==(const FPCGPinProperties& Other) const;
 
@@ -26,6 +26,9 @@ struct PCG_API FPCGPinProperties
 	UPROPERTY(BlueprintReadWrite, EditAnywhere, Category = Settings)
 	EPCGDataType AllowedTypes = EPCGDataType::Any;
 
+	UPROPERTY(BlueprintReadWrite, EditAnywhere, Category = Settings)
+	bool bAllowMultipleData = true;
+
 	UPROPERTY(BlueprintReadWrite, EditAnywhere, Category = Settings)
 	bool bAllowMultipleConnections = true;
 
diff --git a/Engine/Plugins/Experimental/PCG/Source/PCGEditor/PCGEditor.Build.cs b/Engine/Plugins/Experimental/PCG/Source/PCGEditor/PCGEditor.Build.cs
--- a/Engine/Plugins/Experimental/PCG/Source/PCGEditor/PCGEditor.Build.cs
+++ b/Engine/Plugins/Experimental/PCG/Source/PCGEditor/PCGEditor.Build.cs
@@ -37,6 +37,7 @@ namespace UnrealBuildTool.Rules
 					"GraphEditor",
 					"InputCore",
 					"Kismet",
+					"KismetWidgets",
 					"PCG",
 					"PropertyEditor",
 					"Slate",
diff --git a/Engine/Plugins/Experimental/PCG/Source/PCGEditor/Private/PCGEditorStyle.cpp b/Engine/Plugins/Experimental/PCG/Source/PCGEditor/Private/PCGEditorStyle.cpp
--- a/Engine/Plugins/Experimental/PCG/Source/PCGEditor/Private/PCGEditorStyle.cpp
+++ b/Engine/Plugins/Experimental/PCG/Source/PCGEditor/Private/PCGEditorStyle.cpp
@@ -2,6 +2,7 @@
 
 #include "PCGEditorStyle.h"
 
+#include "Interfaces/IPluginManager.h"
 #include "Misc/Paths.h"
 #include "Styling/AppStyle.h"
 #include "Styling/CoreStyle.h"
@@ -26,10 +27,42 @@ FPCGEditorStyle::FPCGEditorStyle() : FSlateStyleSet("PCGEditorStyle")
     static const FVector2D Icon20x20(20.0f, 20.0f);
 	
 	SetCoreContentRoot(FPaths::EngineContentDir() / TEXT("Editor/Slate"));
+	SetContentRoot(IPluginManager::Get().FindPlugin("PCG")->GetBaseDir() / TEXT("Content"));
 
 	Set("PCG.NodeOverlay.Debug", new CORE_IMAGE_BRUSH_SVG("Starship/Blueprints/Breakpoint_Valid", Icon20x20, FSlateColor(FColor::Cyan)));
 	Set("PCG.NodeOverlay.Inspect", new CORE_IMAGE_BRUSH_SVG("Starship/Blueprints/Breakpoint_Valid", Icon20x20, FSlateColor(FColor::Orange)));
 
+	static const FVector2D PinSize(22.0f, 22.0f);
+	Set(PCGEditorStyleConstants::Pin_SD_SC_IN_C, new IMAGE_BRUSH_SVG("Style/PCG_Graph_SD_SC_IN", PinSize));
+	Set(PCGEditorStyleConstants::Pin_SD_SC_IN_DC, new IMAGE_BRUSH_SVG("Style/PCG_Graph_SD_SC_IN_Unplugged", PinSize));
+	Set(PCGEditorStyleConstants::Pin_SD_SC_OUT_C, new IMAGE_BRUSH_SVG("Style/PCG_Graph_SD_SC_OUT", PinSize));
+	Set(PCGEditorStyleConstants::Pin_SD_SC_OUT_DC, new IMAGE_BRUSH_SVG("Style/PCG_Graph_SD_SC_OUT_Unplugged", PinSize));
+
+	Set(PCGEditorStyleConstants::Pin_SD_MC_IN_C, new IMAGE_BRUSH_SVG("Style/PCG_Graph_SD_MC_IN", PinSize));
+	Set(PCGEditorStyleConstants::Pin_SD_MC_IN_DC, new IMAGE_BRUSH_SVG("Style/PCG_Graph_SD_MC_IN_Unplugged", PinSize));
+	Set(PCGEditorStyleConstants::Pin_SD_MC_OUT_C, new IMAGE_BRUSH_SVG("Style/PCG_Graph_SD_MC_OUT", PinSize));
+	Set(PCGEditorStyleConstants::Pin_SD_MC_OUT_DC, new IMAGE_BRUSH_SVG("Style/PCG_Graph_SD_MC_OUT_Unplugged", PinSize));
+
+	Set(PCGEditorStyleConstants::Pin_MD_SC_IN_C, new IMAGE_BRUSH_SVG("Style/PCG_Graph_MD_SC_IN", PinSize));
+	Set(PCGEditorStyleConstants::Pin_MD_SC_IN_DC, new IMAGE_BRUSH_SVG("Style/PCG_Graph_MD_SC_IN_Unplugged", PinSize));
+	Set(PCGEditorStyleConstants::Pin_MD_SC_OUT_C, new IMAGE_BRUSH_SVG("Style/PCG_Graph_MD_SC_OUT", PinSize));
+	Set(PCGEditorStyleConstants::Pin_MD_SC_OUT_DC, new IMAGE_BRUSH_SVG("Style/PCG_Graph_MD_SC_OUT_Unplugged", PinSize));
+
+	Set(PCGEditorStyleConstants::Pin_MD_MC_IN_C, new IMAGE_BRUSH_SVG("Style/PCG_Graph_MD_MC_IN", PinSize));
+	Set(PCGEditorStyleConstants::Pin_MD_MC_IN_DC, new IMAGE_BRUSH_SVG("Style/PCG_Graph_MD_MC_IN_Unplugged", PinSize));
+	Set(PCGEditorStyleConstants::Pin_MD_MC_OUT_C, new IMAGE_BRUSH_SVG("Style/PCG_Graph_MD_MC_OUT", PinSize));
+	Set(PCGEditorStyleConstants::Pin_MD_MC_OUT_DC, new IMAGE_BRUSH_SVG("Style/PCG_Graph_MD_MC_OUT_Unplugged", PinSize));
+
+	Set(PCGEditorStyleConstants::Pin_Param_IN_C, new IMAGE_BRUSH_SVG("Style/PCG_Graph_Param_IN", PinSize));
+	Set(PCGEditorStyleConstants::Pin_Param_IN_DC, new IMAGE_BRUSH_SVG("Style/PCG_Graph_Param_IN_Unplugged", PinSize));
+	Set(PCGEditorStyleConstants::Pin_Param_OUT_C, new IMAGE_BRUSH_SVG("Style/PCG_Graph_Param_OUT", PinSize));
+	Set(PCGEditorStyleConstants::Pin_Param_OUT_DC, new IMAGE_BRUSH_SVG("Style/PCG_Graph_Param_OUT_Unplugged", PinSize));
+
+	Set(PCGEditorStyleConstants::Pin_Composite_IN_C, new IMAGE_BRUSH_SVG("Style/PCG_Graph_Composite_IN", PinSize));
+	Set(PCGEditorStyleConstants::Pin_Composite_IN_DC, new IMAGE_BRUSH_SVG("Style/PCG_Graph_Composite_IN_Unplugged", PinSize));
+	Set(PCGEditorStyleConstants::Pin_Composite_OUT_C, new IMAGE_BRUSH_SVG("Style/PCG_Graph_Composite_OUT", PinSize));
+	Set(PCGEditorStyleConstants::Pin_Composite_OUT_DC, new IMAGE_BRUSH_SVG("Style/PCG_Graph_Composite_OUT_Unplugged", PinSize));
+
 	FInlineEditableTextBlockStyle NodeTitleStyle = FInlineEditableTextBlockStyle(FAppStyle::Get().GetWidgetStyle<FInlineEditableTextBlockStyle>("Graph.Node.NodeTitleInlineEditableText"));
 	FTextBlockStyle GraphNodeItalicTitle = FTextBlockStyle(FAppStyle::Get().GetWidgetStyle<FTextBlockStyle>("Graph.Node.NodeTitle"))
 		.SetFont(DEFAULT_FONT("BoldCondensedItalic", 10));
diff --git a/Engine/Plugins/Experimental/PCG/Source/PCGEditor/Private/PCGEditorStyle.h b/Engine/Plugins/Experimental/PCG/Source/PCGEditor/Private/PCGEditorStyle.h
--- a/Engine/Plugins/Experimental/PCG/Source/PCGEditor/Private/PCGEditorStyle.h
+++ b/Engine/Plugins/Experimental/PCG/Source/PCGEditor/Private/PCGEditorStyle.h
@@ -15,3 +15,36 @@ public:
 private:	
 	FPCGEditorStyle();
 };
+
+namespace PCGEditorStyleConstants
+{
+	const FName Pin_SD_SC_IN_C = TEXT("PCG.Graph.SingleData.SingleCollection.In.Connected");
+	const FName Pin_SD_SC_IN_DC = TEXT("PCG.Graph.SingleData.SingleCollection.In.Disconnected");
+	const FName Pin_SD_SC_OUT_C = TEXT("PCG.Graph.SingleData.SingleCollection.Out.Connected");
+	const FName Pin_SD_SC_OUT_DC = TEXT("PCG.Graph.SingleData.SingleCollection.Out.Disconnected");
+	
+	const FName Pin_SD_MC_IN_C = TEXT("PCG.Graph.SingleData.MultiCollection.In.Connected");
+	const FName Pin_SD_MC_IN_DC = TEXT("PCG.Graph.SingleData.MultiCollection.In.Disconnected");
+	const FName Pin_SD_MC_OUT_C = TEXT("PCG.Graph.SingleData.MultiCollection.Out.Connected");
+	const FName Pin_SD_MC_OUT_DC = TEXT("PCG.Graph.SingleData.MultiCollection.Out.Disconnected");
+
+	const FName Pin_MD_SC_IN_C = TEXT("PCG.Graph.MultiData.SingleCollection.In.Connected");
+	const FName Pin_MD_SC_IN_DC = TEXT("PCG.Graph.MultiData.SingleCollection.In.Disconnected");
+	const FName Pin_MD_SC_OUT_C = TEXT("PCG.Graph.MultiData.SingleCollection.Out.Connected");
+	const FName Pin_MD_SC_OUT_DC = TEXT("PCG.Graph.MultiData.SingleCollection.Out.Disconnected");
+
+	const FName Pin_MD_MC_IN_C = TEXT("PCG.Graph.MultiData.MultiCollection.In.Connected");
+	const FName Pin_MD_MC_IN_DC = TEXT("PCG.Graph.MultiData.MultiCollection.In.Disconnected");
+	const FName Pin_MD_MC_OUT_C = TEXT("PCG.Graph.MultiData.MultiCollection.Out.Connected");
+	const FName Pin_MD_MC_OUT_DC = TEXT("PCG.Graph.MultiData.MultiCollection.Out.Disconnected");
+
+	const FName Pin_Param_IN_C = TEXT("PCG.Graph.Param.In.Connected");
+	const FName Pin_Param_IN_DC = TEXT("PCG.Graph.Param.In.Disconnected");
+	const FName Pin_Param_OUT_C = TEXT("PCG.Graph.Param.Out.Connected");
+	const FName Pin_Param_OUT_DC = TEXT("PCG.Graph.Param.Out.Disconnected");
+
+	const FName Pin_Composite_IN_C = TEXT("PCG.Graph.Composite.In.Connected");
+	const FName Pin_Composite_IN_DC = TEXT("PCG.Graph.Composite.In.Disconnected");
+	const FName Pin_Composite_OUT_C = TEXT("PCG.Graph.Composite.Out.Connected");
+	const FName Pin_Composite_OUT_DC = TEXT("PCG.Graph.Composite.Out.Disconnected");
+}
\ No newline at end of file
diff --git a/Engine/Plugins/Experimental/PCG/Source/PCGEditor/Private/SPCGEditorGraphNode.cpp b/Engine/Plugins/Experimental/PCG/Source/PCGEditor/Private/SPCGEditorGraphNode.cpp
--- a/Engine/Plugins/Experimental/PCG/Source/PCGEditor/Private/SPCGEditorGraphNode.cpp
+++ b/Engine/Plugins/Experimental/PCG/Source/PCGEditor/Private/SPCGEditorGraphNode.cpp
@@ -9,6 +9,8 @@
 #include "PCGSettings.h"
 
 #include "SGraphPin.h"
+#include "SLevelOfDetailBranchNode.h"
+#include "SPinTypeSelector.h"
 #include "Styling/SlateBrush.h"
 #include "Widgets/Text/SInlineEditableTextBlock.h"
 
@@ -16,7 +18,14 @@
 class SPCGEditorGraphNodePin : public SGraphPin
 {
 public:
-	SLATE_BEGIN_ARGS(SPCGEditorGraphNodePin) {}
+	SLATE_BEGIN_ARGS(SPCGEditorGraphNodePin)
+		: _PinLabelStyle(NAME_DefaultPinLabelStyle)
+		, _UsePinColorForText(false)
+		, _SideToSideMargin(5.0f)
+		{}
+		SLATE_ARGUMENT(FName, PinLabelStyle)
+		SLATE_ARGUMENT(bool, UsePinColorForText)
+		SLATE_ARGUMENT(float, SideToSideMargin)
 	SLATE_END_ARGS()
 
 	void Construct(const FArguments& InArgs, UEdGraphPin* InPin);
@@ -26,8 +35,187 @@ public:
 
 void SPCGEditorGraphNodePin::Construct(const FArguments& InArgs, UEdGraphPin* InPin)
 {
-	// Required to give the first argument.
-	SGraphPin::Construct(SGraphPin::FArguments(), InPin);
+	// TODO: replace this with base class when we have sufficient controls to change the padding
+	// IMPLEMENTATION NOTE: this is the code from SGraphPin::Construct
+	// e.g. SGraphPin::Construct(SGraphPin::FArguments().SideToSideMargin(0.0f), InPin);
+	// with additional padding exposed
+
+	bUsePinColorForText = InArgs._UsePinColorForText;
+	this->SetCursor(EMouseCursor::Default);
+
+	SetVisibility(MakeAttributeSP(this, &SPCGEditorGraphNodePin::GetPinVisiblity));
+
+	GraphPinObj = InPin;
+	check(GraphPinObj != NULL);
+
+	const UEdGraphSchema* Schema = GraphPinObj->GetSchema();
+	checkf(
+		Schema, 
+		TEXT("Missing schema for pin: %s with outer: %s of type %s"), 
+		*(GraphPinObj->GetName()),
+		GraphPinObj->GetOuter() ? *(GraphPinObj->GetOuter()->GetName()) : TEXT("NULL OUTER"), 
+		GraphPinObj->GetOuter() ? *(GraphPinObj->GetOuter()->GetClass()->GetName()) : TEXT("NULL OUTER")
+	);
+
+	const bool bIsInput = (GetDirection() == EGPD_Input);
+
+	// Create the pin icon widget
+	TSharedRef<SWidget> PinWidgetRef = SPinTypeSelector::ConstructPinTypeImage(
+		MakeAttributeSP(this, &SPCGEditorGraphNodePin::GetPinIcon ),
+		MakeAttributeSP(this, &SPCGEditorGraphNodePin::GetPinColor),
+		MakeAttributeSP(this, &SPCGEditorGraphNodePin::GetSecondaryPinIcon),
+		MakeAttributeSP(this, &SPCGEditorGraphNodePin::GetSecondaryPinColor));
+	PinImage = PinWidgetRef;
+
+	PinWidgetRef->SetCursor( 
+		TAttribute<TOptional<EMouseCursor::Type> >::Create (
+			TAttribute<TOptional<EMouseCursor::Type> >::FGetter::CreateRaw( this, &SPCGEditorGraphNodePin::GetPinCursor )
+		)
+	);
+
+	// Create the pin indicator widget (used for watched values)
+	static const FName NAME_NoBorder("NoBorder");
+	TSharedRef<SWidget> PinStatusIndicator =
+		SNew(SButton)
+		.ButtonStyle(FAppStyle::Get(), NAME_NoBorder)
+		.Visibility(this, &SPCGEditorGraphNodePin::GetPinStatusIconVisibility)
+		.ContentPadding(0)
+		.OnClicked(this, &SPCGEditorGraphNodePin::ClickedOnPinStatusIcon)
+		[
+			SNew(SImage)
+			.Image(this, &SPCGEditorGraphNodePin::GetPinStatusIcon)
+		];
+
+	TSharedRef<SWidget> LabelWidget = GetLabelWidget(InArgs._PinLabelStyle);
+
+	// Create the widget used for the pin body (status indicator, label, and value)
+	LabelAndValue =
+		SNew(SWrapBox)
+		.PreferredSize(150.f);
+
+	if (!bIsInput)
+	{
+		LabelAndValue->AddSlot()
+			.VAlign(VAlign_Center)
+			[
+				PinStatusIndicator
+			];
+
+		LabelAndValue->AddSlot()
+			.VAlign(VAlign_Center)
+			[
+				LabelWidget
+			];
+	}
+	else
+	{
+		LabelAndValue->AddSlot()
+			.VAlign(VAlign_Center)
+			[
+				LabelWidget
+			];
+
+		ValueWidget = GetDefaultValueWidget();
+
+		if (ValueWidget != SNullWidget::NullWidget)
+		{
+			TSharedPtr<SBox> ValueBox;
+			LabelAndValue->AddSlot()
+				.Padding(bIsInput ? FMargin(InArgs._SideToSideMargin, 0, 0, 0) : FMargin(0, 0, InArgs._SideToSideMargin, 0))
+				.VAlign(VAlign_Center)
+				[
+					SAssignNew(ValueBox, SBox)
+					.Padding(0.0f)
+					[
+						ValueWidget.ToSharedRef()
+					]
+				];
+
+			if (!DoesWidgetHandleSettingEditingEnabled())
+			{
+				ValueBox->SetEnabled(TAttribute<bool>(this, &SPCGEditorGraphNodePin::IsEditingEnabled));
+			}
+		}
+
+		LabelAndValue->AddSlot()
+			.VAlign(VAlign_Center)
+			[
+				PinStatusIndicator
+			];
+	}
+
+	TSharedPtr<SHorizontalBox> PinContent;
+	if (bIsInput)
+	{
+		// Input pin
+		FullPinHorizontalRowWidget = PinContent = 
+			SNew(SHorizontalBox)
+			+SHorizontalBox::Slot()
+			.AutoWidth()
+			.VAlign(VAlign_Center)
+			.Padding(0, 0, InArgs._SideToSideMargin, 0)
+			[
+				PinWidgetRef
+			]
+			+SHorizontalBox::Slot()
+			.AutoWidth()
+			.VAlign(VAlign_Center)
+			[
+				LabelAndValue.ToSharedRef()
+			];
+	}
+	else
+	{
+		// Output pin
+		FullPinHorizontalRowWidget = PinContent = SNew(SHorizontalBox)
+			+SHorizontalBox::Slot()
+			.AutoWidth()
+			.VAlign(VAlign_Center)
+			[
+				LabelAndValue.ToSharedRef()
+			]
+			+SHorizontalBox::Slot()
+			.AutoWidth()
+			.VAlign(VAlign_Center)
+			.Padding(InArgs._SideToSideMargin, 0, 0, 0)
+			[
+				PinWidgetRef
+			];
+	}
+
+	// Set up a hover for pins that is tinted the color of the pin.
+	
+	SBorder::Construct(SBorder::FArguments()
+		.BorderImage(this, &SPCGEditorGraphNodePin::GetPinBorder)
+		.BorderBackgroundColor(this, &SPCGEditorGraphNodePin::GetHighlightColor)
+		.OnMouseButtonDown(this, &SPCGEditorGraphNodePin::OnPinNameMouseDown)
+		.Padding(0) // NOTE: This is different from base class implementation
+		[
+			SNew(SBorder)
+			.BorderImage(CachedImg_Pin_DiffOutline)
+			.BorderBackgroundColor(this, &SPCGEditorGraphNodePin::GetPinDiffColor)
+			.Padding(0) // NOTE: This is different from base class implementation
+			[
+				SNew(SLevelOfDetailBranchNode)
+				.UseLowDetailSlot(this, &SPCGEditorGraphNodePin::UseLowDetailPinNames)
+				.LowDetail()
+				[
+					//@TODO: Try creating a pin-colored line replacement that doesn't measure text / call delegates but still renders
+					PinWidgetRef
+				]
+				.HighDetail()
+				[
+					PinContent.ToSharedRef()
+				]
+			]
+		]
+	);
+
+	TSharedPtr<IToolTip> TooltipWidget = SNew(SToolTip)
+		.Text(this, &SPCGEditorGraphNodePin::GetTooltipText);
+
+	SetToolTip(TooltipWidget);
+
 }
 
 // Adapted from SGraphPin::GetPinColor
@@ -103,14 +291,61 @@ void SPCGEditorGraphNode::AddPin(const TSharedRef<SGraphPin>& PinToAdd)
 {
 	check(PCGEditorGraphNode);
 	UPCGNode* PCGNode = PCGEditorGraphNode->GetPCGNode();
-	// Implementation note: we do not distinguish single/multiple pins on the output since that is not relevant
+
 	if (PCGNode && PinToAdd->GetPinObj())
 	{
-		if (UPCGPin* Pin = PCGNode->GetInputPin(PinToAdd->GetPinObj()->PinName))
+		const bool bIsInPin = PinToAdd->GetPinObj()->Direction == EEdGraphPinDirection::EGPD_Input;
+		const FName& PinName = PinToAdd->GetPinObj()->PinName;
+
+		if(UPCGPin* Pin = (bIsInPin ? PCGNode->GetInputPin(PinName) : PCGNode->GetOutputPin(PinName)))
 		{
-			if (Pin->Properties.bAllowMultipleConnections)
+			const bool bIsMultiData = Pin->Properties.bAllowMultipleData;
+			const bool bIsMultiConnections = Pin->Properties.bAllowMultipleConnections;
+
+			// Check for special types
+			if (Pin->Properties.AllowedTypes == EPCGDataType::Param)
+			{
+				const FSlateBrush* ConnectedBrush = FPCGEditorStyle::Get().GetBrush(bIsInPin ? PCGEditorStyleConstants::Pin_Param_IN_C : PCGEditorStyleConstants::Pin_Param_OUT_C);
+				const FSlateBrush* DisconnectedBrush = FPCGEditorStyle::Get().GetBrush(bIsInPin ? PCGEditorStyleConstants::Pin_Param_IN_DC : PCGEditorStyleConstants::Pin_Param_OUT_DC);
+
+				PinToAdd->SetCustomPinIcon(ConnectedBrush, DisconnectedBrush);
+			}
+			else if (Pin->Properties.AllowedTypes == EPCGDataType::Composite)
 			{
-				PinToAdd->SetCustomPinIcon(FAppStyle::GetBrush("Graph.ArrayPin.Connected"), FAppStyle::GetBrush("Graph.ArrayPin.Disconnected"));
+				const FSlateBrush* ConnectedBrush = FPCGEditorStyle::Get().GetBrush(bIsInPin ? PCGEditorStyleConstants::Pin_Composite_IN_C : PCGEditorStyleConstants::Pin_Composite_OUT_C);
+				const FSlateBrush* DisconnectedBrush = FPCGEditorStyle::Get().GetBrush(bIsInPin ? PCGEditorStyleConstants::Pin_Composite_IN_DC : PCGEditorStyleConstants::Pin_Composite_OUT_DC);
+
+				PinToAdd->SetCustomPinIcon(ConnectedBrush, DisconnectedBrush);
+			}
+			else
+			{
+				static const FName* PinBrushes[] =
+				{
+					&PCGEditorStyleConstants::Pin_SD_SC_IN_C,
+					&PCGEditorStyleConstants::Pin_SD_SC_IN_DC,
+					&PCGEditorStyleConstants::Pin_SD_MC_IN_C,
+					&PCGEditorStyleConstants::Pin_SD_MC_IN_DC,
+					&PCGEditorStyleConstants::Pin_MD_SC_IN_C,
+					&PCGEditorStyleConstants::Pin_MD_SC_IN_DC,
+					&PCGEditorStyleConstants::Pin_MD_MC_IN_C,
+					&PCGEditorStyleConstants::Pin_MD_MC_IN_DC,
+					&PCGEditorStyleConstants::Pin_SD_SC_OUT_C,
+					&PCGEditorStyleConstants::Pin_SD_SC_OUT_DC,
+					&PCGEditorStyleConstants::Pin_SD_MC_OUT_C,
+					&PCGEditorStyleConstants::Pin_SD_MC_OUT_DC,
+					&PCGEditorStyleConstants::Pin_MD_SC_OUT_C,
+					&PCGEditorStyleConstants::Pin_MD_SC_OUT_DC,
+					&PCGEditorStyleConstants::Pin_MD_MC_OUT_C,
+					&PCGEditorStyleConstants::Pin_MD_MC_OUT_DC
+				};
+
+				const int32 ConnectedIndex = (bIsInPin ? 0 : 8) + (bIsMultiData ? 4 : 0) + (bIsMultiConnections ? 2 : 0);
+				const int32 DisconnectedIndex = ConnectedIndex + 1;
+
+				const FSlateBrush* ConnectedBrush = FPCGEditorStyle::Get().GetBrush(*PinBrushes[ConnectedIndex]);
+				const FSlateBrush* DisconnectedBrush = FPCGEditorStyle::Get().GetBrush(*PinBrushes[DisconnectedIndex]);
+
+				PinToAdd->SetCustomPinIcon(ConnectedBrush, DisconnectedBrush);
 			}
 		}
 	}
