[PCG] SpawnActor can now define its own template and add a map from attribute to actor properties.

diff --git a/Engine/Plugins/Experimental/PCG/Source/PCG/Private/Elements/PCGSpawnActor.cpp b/Engine/Plugins/Experimental/PCG/Source/PCG/Private/Elements/PCGSpawnActor.cpp
--- a/Engine/Plugins/Experimental/PCG/Source/PCG/Private/Elements/PCGSpawnActor.cpp
+++ b/Engine/Plugins/Experimental/PCG/Source/PCG/Private/Elements/PCGSpawnActor.cpp
@@ -2,7 +2,6 @@
 
 #include "Elements/PCGSpawnActor.h"
 
-#include "Engine/Engine.h"
 #include "PCGComponent.h"
 #include "PCGHelpers.h"
 #include "PCGManagedResource.h"
@@ -11,14 +10,120 @@
 #include "Grid/PCGPartitionActor.h"
 #include "Helpers/PCGActorHelpers.h"
 #include "Helpers/PCGSettingsHelpers.h"
+#include "Metadata/Accessors/IPCGAttributeAccessor.h"
+#include "Metadata/Accessors/PCGAttributeAccessorKeys.h"
+#include "Metadata/Accessors/PCGAttributeAccessorHelpers.h"
 
 #include "Components/InstancedStaticMeshComponent.h"
 #include "GameFramework/Actor.h"
+#include "Engine/Blueprint.h"
+#include "Engine/BlueprintGeneratedClass.h"
+#include "Engine/Engine.h"
 #include "Engine/World.h"
 #include "ISMPartition/ISMComponentDescriptor.h"
+#include "UObject/Package.h"
+#include "UObject/UObjectGlobals.h"
 
 #include UE_INLINE_GENERATED_CPP_BY_NAME(PCGSpawnActor)
 
+namespace PCGSpawnActorHelpers
+{
+	struct FActorSingleOverride
+	{
+		using ApplyOverrideFunction = TFunction<void(int32, const IPCGAttributeAccessorKeys&, IPCGAttributeAccessorKeys&)>;
+
+		FActorSingleOverride(const FPCGAttributePropertySelector& InputSelector, const FString& OutputProperty, AActor* TemplateActor, const UPCGPointData* PointData)
+		{
+			ActorOverrideInputAccessor = PCGAttributeAccessorHelpers::CreateConstAccessor(PointData, InputSelector);
+			ActorOverrideOutputAccessor = PCGAttributeAccessorHelpers::CreatePropertyAccessor(FName(OutputProperty), TemplateActor->GetClass());
+
+			if (!ActorOverrideInputAccessor.IsValid() || !ActorOverrideOutputAccessor.IsValid())
+			{
+				UE_LOG(LogPCG, Warning, TEXT("ActorOverride from input %s or output %s is invalid or unsupported. Will be skipped."), *InputSelector.GetName().ToString(), *OutputProperty);
+				return;
+			}
+
+			if (!PCG::Private::IsBroadcastable(ActorOverrideInputAccessor->GetUnderlyingType(), ActorOverrideOutputAccessor->GetUnderlyingType()))
+			{
+				UE_LOG(LogPCG, Warning, TEXT("ActorOverride cannot set input %s to output %s. Types are incompatibles. Will be skipped."), *InputSelector.GetName().ToString(), *OutputProperty);
+				ActorOverrideInputAccessor.Reset();
+				ActorOverrideOutputAccessor.Reset();
+				return;
+			}
+
+			auto CreateGetterSetter = [this](auto Dummy)
+			{
+				using Type = decltype(Dummy);
+
+				ActorOverrideFunction = [this](int32 Index, const IPCGAttributeAccessorKeys& InputKeys, IPCGAttributeAccessorKeys& OutputKey)
+				{
+					if (!IsValid())
+					{
+						return;
+					}
+
+					Type Value{};
+					if (ActorOverrideInputAccessor->Get<Type>(Value, Index, InputKeys, EPCGAttributeAccessorFlags::AllowBroadcast))
+					{
+						ActorOverrideOutputAccessor->Set<Type>(Value, OutputKey);
+					}
+				};
+			};
+
+			PCGMetadataAttribute::CallbackWithRightType(ActorOverrideOutputAccessor->GetUnderlyingType(), CreateGetterSetter);
+		}
+
+		bool IsValid() const
+		{
+			return ActorOverrideInputAccessor.IsValid() && ActorOverrideOutputAccessor.IsValid() && ActorOverrideFunction;
+		}
+
+		void Apply(int32 Index, const IPCGAttributeAccessorKeys& InputKeys, IPCGAttributeAccessorKeys& OutputKey)
+		{
+			ActorOverrideFunction(Index, InputKeys, OutputKey);
+		}
+
+	private:
+		TUniquePtr<const IPCGAttributeAccessor> ActorOverrideInputAccessor;
+		TUniquePtr<IPCGAttributeAccessor> ActorOverrideOutputAccessor;
+		ApplyOverrideFunction ActorOverrideFunction;
+	};
+
+	struct FActorOverrides
+	{
+		FActorOverrides(const TArray<FPCGActorPropertyOverride>& Overrides, AActor* TemplateActor, const UPCGPointData* PointData)
+			: InputKeys(PointData->GetPoints())
+			, OutputKey(TemplateActor)
+		{
+			ActorSingleOverrides.Reserve(Overrides.Num());
+
+			for (int32 i = 0; i < Overrides.Num(); ++i)
+			{
+				const FPCGAttributePropertySelector& InputSelector = Overrides[i].InputSource;
+				const FString& OutputProperty = Overrides[i].PropertyTarget;
+
+				ActorSingleOverrides.Emplace(InputSelector, OutputProperty, TemplateActor, PointData);
+			}
+		}
+
+		void Apply(int32 Index)
+		{
+			for (FActorSingleOverride& ActorSingleOverride : ActorSingleOverrides)
+			{
+				if (ActorSingleOverride.IsValid())
+				{
+					ActorSingleOverride.Apply(Index, InputKeys, OutputKey);
+				}
+			}
+		}
+
+	private:
+		FPCGAttributeAccessorKeysPoints InputKeys;
+		FPCGAttributeAccessorKeysSingleObjectPtr<AActor> OutputKey;
+		TArray<FActorSingleOverride> ActorSingleOverrides;
+	};
+}
+
 UPCGNode* UPCGSpawnActorSettings::CreateNode() const
 {
 	return NewObject<UPCGSpawnActorNode>();
@@ -58,6 +163,58 @@ UPCGGraph* UPCGSpawnActorSettings::GetSubgraph() const
 	return nullptr;
 }
 
+void UPCGSpawnActorSettings::BeginDestroy()
+{
+#if WITH_EDITOR
+	TeardownBlueprintEvent();
+#endif // WITH_EDITOR
+
+	Super::BeginDestroy();
+}
+
+#if WITH_EDITOR
+void UPCGSpawnActorSettings::SetupBlueprintEvent()
+{
+	if (UBlueprintGeneratedClass* BlueprintClass = Cast<UBlueprintGeneratedClass>(TemplateActorClass))
+	{
+		if (UBlueprint* Blueprint = Cast<UBlueprint>(BlueprintClass->ClassGeneratedBy))
+		{
+			Blueprint->OnChanged().AddUObject(this, &UPCGSpawnActorSettings::OnBlueprintChanged);
+		}
+	}
+}
+
+void UPCGSpawnActorSettings::TeardownBlueprintEvent()
+{
+	if (UBlueprintGeneratedClass* BlueprintClass = Cast<UBlueprintGeneratedClass>(TemplateActorClass))
+	{
+		if (UBlueprint* Blueprint = Cast<UBlueprint>(BlueprintClass->ClassGeneratedBy))
+		{
+			Blueprint->OnChanged().RemoveAll(this);
+		}
+	}
+}
+#endif
+
+void UPCGSpawnActorSettings::PostLoad()
+{
+	Super::PostLoad();
+
+#if WITH_EDITOR
+	SetupBlueprintEvent();
+
+	if (TemplateActorClass)
+	{
+		if (TemplateActor)
+		{
+			TemplateActor->ConditionalPostLoad();
+		}
+
+		RefreshTemplateActor();
+	}
+#endif // WITH_EDITOR
+}
+
 #if WITH_EDITOR
 bool UPCGSpawnActorSettings::IsStructuralProperty(const FName& InPropertyName) const
 {
@@ -73,6 +230,77 @@ TObjectPtr<UPCGGraph> UPCGSpawnActorNode::GetSubgraph() const
 	return (Settings && Settings->Option != EPCGSpawnActorOption::NoMerging) ? Settings->GetSubgraph() : nullptr;
 }
 
+#if WITH_EDITOR
+void UPCGSpawnActorSettings::PreEditChange(FProperty* PropertyAboutToChange)
+{
+	if (PropertyAboutToChange && PropertyAboutToChange->GetFName() == GET_MEMBER_NAME_CHECKED(UPCGSpawnActorSettings, TemplateActorClass))
+	{
+		TeardownBlueprintEvent();
+	}
+}
+
+void UPCGSpawnActorSettings::PostEditChangeProperty(FPropertyChangedEvent& PropertyChangedEvent)
+{
+	if (PropertyChangedEvent.Property)
+	{
+		const FName& PropertyName = PropertyChangedEvent.Property->GetFName();
+
+		if (PropertyName == GET_MEMBER_NAME_CHECKED(UPCGSpawnActorSettings, TemplateActorClass))
+		{
+			SetupBlueprintEvent();
+			RefreshTemplateActor();
+		}
+	}
+
+	Super::PostEditChangeProperty(PropertyChangedEvent);
+}
+
+void UPCGSpawnActorSettings::PreEditUndo()
+{
+	TeardownBlueprintEvent();
+
+	Super::PreEditUndo();
+}
+
+void UPCGSpawnActorSettings::PostEditUndo()
+{
+	Super::PostEditUndo();
+
+	SetupBlueprintEvent();
+	RefreshTemplateActor();
+}
+
+void UPCGSpawnActorSettings::OnBlueprintChanged(UBlueprint* InBlueprint)
+{
+	RefreshTemplateActor();
+	DirtyCache();
+	OnSettingsChangedDelegate.Broadcast(this, EPCGChangeType::Settings);
+}
+#endif // WITH_EDITOR
+
+void UPCGSpawnActorSettings::RefreshTemplateActor()
+{
+	if (TemplateActorClass)
+	{
+		AActor* NewTemplateActor = NewObject<AActor>(GetTransientPackage(), TemplateActorClass, NAME_None, RF_ArchetypeObject | RF_Transactional | RF_Public);
+
+		if (TemplateActor)
+		{
+			UEngine::FCopyPropertiesForUnrelatedObjectsParams Options;
+			Options.bNotifyObjectReplacement = true;
+			UEngine::CopyPropertiesForUnrelatedObjects(TemplateActor, NewTemplateActor, Options);
+
+			TemplateActor->Rename(nullptr, GetTransientPackage(), REN_DontCreateRedirectors);
+		}
+
+		TemplateActor = NewTemplateActor;
+	}
+	else
+	{
+		TemplateActor = nullptr;
+	}
+}
+
 bool FPCGSpawnActorElement::ExecuteInternal(FPCGContext* Context) const
 {
 	TRACE_CPUPROFILER_EVENT_SCOPE(FPCSpawnActorElement::Execute);
@@ -87,6 +315,11 @@ bool FPCGSpawnActorElement::ExecuteInternal(FPCGContext* Context) const
 		return true;
 	}
 
+	if (!ensure(Settings->TemplateActor && Settings->TemplateActor->IsA(Settings->TemplateActorClass)))
+	{
+		return true;
+	}
+
 	TArray<UFunction*> PostSpawnFunctions;
 	for (FName PostSpawnFunctionName : Settings->PostSpawnFunctionNames)
 	{
@@ -137,7 +370,7 @@ bool FPCGSpawnActorElement::ExecuteInternal(FPCGContext* Context) const
 		}
 
 		const bool bHasAuthority = !Context->SourceComponent.IsValid() || (Context->SourceComponent->GetOwner() && Context->SourceComponent->GetOwner()->HasAuthority());
-		const bool bSpawnedActorsRequireAuthority = Settings->TemplateActorClass->GetDefaultObject<AActor>()->GetIsReplicated();
+		const bool bSpawnedActorsRequireAuthority = Settings->TemplateActor->GetIsReplicated();
 
 		// First, create target instance transforms
 		const UPCGPointData* PointData = SpatialData->ToPointData(Context);
@@ -232,8 +465,13 @@ bool FPCGSpawnActorElement::ExecuteInternal(FPCGContext* Context) const
 				TRACE_CPUPROFILER_EVENT_SCOPE(FPCGSpawnActorElement::ExecuteInternal::SpawnActors);
 				if (Settings->Option != EPCGSpawnActorOption::CollapseActors && (bHasAuthority || !bSpawnedActorsRequireAuthority))
 				{
+					AActor* TemplateActor = Settings->ActorOverrides.IsEmpty() ? Settings->TemplateActor.Get() : DuplicateObject(Settings->TemplateActor, GetTransientPackage());
+
+					PCGSpawnActorHelpers::FActorOverrides ActorOverrides(Settings->ActorOverrides, TemplateActor, PointData);
+
 					FActorSpawnParameters SpawnParams;
 					SpawnParams.Owner = TargetActor;
+					SpawnParams.Template = TemplateActor;
 					SpawnParams.SpawnCollisionHandlingOverride = ESpawnActorCollisionHandlingMethod::AlwaysSpawn;
 
 					if (PCGHelpers::IsRuntimeOrPIE())
@@ -254,9 +492,20 @@ bool FPCGSpawnActorElement::ExecuteInternal(FPCGContext* Context) const
 						Settings->bGenerationTrigger == EPCGSpawnActorGenerationTrigger::DoNotGenerateInEditor);
 #endif
 
-					for (const FPCGPoint& Point : Points)
+					for (int32 i = 0; i < Points.Num(); ++i)
 					{
+						const FPCGPoint& Point = Points[i];
+
+						ActorOverrides.Apply(i);
+
 						AActor* GeneratedActor = TargetActor->GetWorld()->SpawnActor(Settings->TemplateActorClass, &Point.Transform, SpawnParams);
+
+						if (!GeneratedActor)
+						{
+							PCGE_LOG(Error, "Failed to spawn actor on point with index %d", i);
+							continue;
+						}
+
 						// HACK: until UE-62747 is fixed, we have to force set the scale after spawning the actor
 						GeneratedActor->SetActorRelativeScale3D(Point.Transform.GetScale3D());
 						GeneratedActor->Tags.Append(NewActorTags);
diff --git a/Engine/Plugins/Experimental/PCG/Source/PCG/Private/Metadata/Accessors/PCGAttributeAccessorHelpers.cpp b/Engine/Plugins/Experimental/PCG/Source/PCG/Private/Metadata/Accessors/PCGAttributeAccessorHelpers.cpp
--- a/Engine/Plugins/Experimental/PCG/Source/PCG/Private/Metadata/Accessors/PCGAttributeAccessorHelpers.cpp
+++ b/Engine/Plugins/Experimental/PCG/Source/PCG/Private/Metadata/Accessors/PCGAttributeAccessorHelpers.cpp
@@ -92,6 +92,11 @@ namespace PCGAttributeAccessorHelpers
 
 TUniquePtr<IPCGAttributeAccessor> PCGAttributeAccessorHelpers::CreatePropertyAccessor(const FProperty* InProperty)
 {
+	if (!InProperty)
+	{
+		return TUniquePtr<IPCGAttributeAccessor>();
+	}
+
 	if (const FNumericProperty* NumericProperty = CastField<FNumericProperty>(InProperty))
 	{
 		if (NumericProperty->IsFloatingPoint())
@@ -158,6 +163,32 @@ TUniquePtr<IPCGAttributeAccessor> PCGAttributeAccessorHelpers::CreatePropertyAcc
 	return TUniquePtr<IPCGAttributeAccessor>();
 }
 
+TUniquePtr<IPCGAttributeAccessor> PCGAttributeAccessorHelpers::CreatePropertyAccessor(const FName InPropertyName, const UClass* InClass)
+{
+	if (InClass)
+	{
+		if (const FProperty* Property = InClass->FindPropertyByName(InPropertyName))
+		{
+			return CreatePropertyAccessor(Property);
+		}
+	}
+
+	return TUniquePtr<IPCGAttributeAccessor>();
+}
+
+TUniquePtr<IPCGAttributeAccessor> PCGAttributeAccessorHelpers::CreatePropertyAccessor(const FName InPropertyName, const UStruct* InStruct)
+{
+	if (InStruct)
+	{
+		if (const FProperty* Property = InStruct->FindPropertyByName(InPropertyName))
+		{
+			return CreatePropertyAccessor(Property);
+		}
+	}
+
+	return TUniquePtr<IPCGAttributeAccessor>();
+}
+
 TUniquePtr<const IPCGAttributeAccessor> PCGAttributeAccessorHelpers::CreateConstAccessor(const UPCGData* InData, const FPCGAttributePropertySelector& InSelector)
 {
 	const FName Name = InSelector.GetName();
diff --git a/Engine/Plugins/Experimental/PCG/Source/PCG/Public/Elements/PCGSpawnActor.h b/Engine/Plugins/Experimental/PCG/Source/PCG/Public/Elements/PCGSpawnActor.h
--- a/Engine/Plugins/Experimental/PCG/Source/PCG/Public/Elements/PCGSpawnActor.h
+++ b/Engine/Plugins/Experimental/PCG/Source/PCG/Public/Elements/PCGSpawnActor.h
@@ -6,6 +6,8 @@
 #include "PCGElement.h"
 #include "PCGSettings.h"
 #include "PCGSubgraph.h"
+#include "Metadata/PCGAttributePropertySelector.h"
+
 #include "Templates/SubclassOf.h"
 
 #include "PCGSpawnActor.generated.h"
@@ -29,35 +31,64 @@ enum class EPCGSpawnActorGenerationTrigger : uint8
 	DoNotGenerate // Does not call generate
 };
 
+USTRUCT(BlueprintType)
+struct FPCGActorPropertyOverride
+{
+	GENERATED_BODY()
+
+	UPROPERTY(EditAnywhere, Category = Settings)
+	FPCGAttributePropertySelector InputSource;
+
+	UPROPERTY(EditAnywhere, Category = Settings)
+	FString PropertyTarget;
+};
+
 UCLASS(BlueprintType, ClassGroup = (Procedural))
 class PCG_API UPCGSpawnActorSettings : public UPCGBaseSubgraphSettings
 {
 	GENERATED_BODY()
 
 public:
-	UPROPERTY(BlueprintReadWrite, EditAnywhere, Category = Settings)
+	UPROPERTY(BlueprintReadOnly, EditAnywhere, Category = Settings, meta = (OnlyPlaceable, DisallowCreateNew))
 	TSubclassOf<AActor> TemplateActorClass = nullptr;
 
 	/* Can specify a list of functions from the template class to be called on each actor spawned, in order. Need to be parameter-less and with "CallInEditor" flag enabled.*/
-	UPROPERTY(BlueprintReadWrite, EditAnywhere, Category = Settings, meta = (EditCondition = "Option!=EPCGSpawnActorOption::CollapseActors"))
+	UPROPERTY(BlueprintReadWrite, EditAnywhere, Category = Settings, meta = (EditCondition = "Option!=EPCGSpawnActorOption::CollapseActors", EditConditionHides))
 	TArray<FName> PostSpawnFunctionNames;
 
 	UPROPERTY(BlueprintReadWrite, EditAnywhere, Category = Settings)
 	EPCGSpawnActorOption Option = EPCGSpawnActorOption::CollapseActors;
 
-	UPROPERTY(BlueprintReadWrite, EditAnywhere, Category = Settings, meta = (EditCondition = "Option==EPCGSpawnActorOption::NoMerging"))
+	UPROPERTY(BlueprintReadWrite, EditAnywhere, Category = Settings, meta = (EditCondition = "Option==EPCGSpawnActorOption::NoMerging", EditConditionHides))
 	bool bForceDisableActorParsing = true;
 
-	UPROPERTY(BlueprintReadWrite, EditAnywhere, Category = Settings, meta = (EditCondition = "Option==EPCGSpawnActorOption::NoMerging"))
+	UPROPERTY(BlueprintReadWrite, EditAnywhere, Category = Settings, meta = (EditCondition = "Option==EPCGSpawnActorOption::NoMerging", EditConditionHides))
 	EPCGSpawnActorGenerationTrigger bGenerationTrigger = EPCGSpawnActorGenerationTrigger::Default;
 
 	/** Warning: inheriting parent actor tags work only in non-collapsed actor hierarchies */
-	UPROPERTY(BlueprintReadWrite, EditAnywhere, Category = Settings, meta = (EditCondition = "Option!=EPCGSpawnActorOption::CollapseActors"))
+	UPROPERTY(BlueprintReadWrite, EditAnywhere, Category = Settings, meta = (EditCondition = "Option!=EPCGSpawnActorOption::CollapseActors", EditConditionHides))
 	bool bInheritActorTags = false;
 
-	UPROPERTY(BlueprintReadWrite, EditAnywhere, Category = Settings, meta = (EditCondition = "Option!=EPCGSpawnActorOption::CollapseActors"))
+	UPROPERTY(BlueprintReadWrite, EditAnywhere, Category = Settings, meta = (EditCondition = "Option!=EPCGSpawnActorOption::CollapseActors", EditConditionHides))
 	TArray<FName> TagsToAddOnActors;
 
+	UPROPERTY(EditAnywhere, BlueprintReadWrite, Instanced, Category = Settings, meta = (ShowInnerProperties, EditCondition = "Option != EPCGSpawnActorOption::CollapseActors", EditConditionHides))
+	TObjectPtr<AActor> TemplateActor;
+
+	UPROPERTY(EditAnywhere, BlueprintReadWrite, Category = Settings, meta = (EditCondition = "Option != EPCGSpawnActorOption::CollapseActors", EditConditionHides))
+	TArray<FPCGActorPropertyOverride> ActorOverrides;
+
+	virtual void PostLoad() override;
+	virtual void BeginDestroy() override;
+
+#if WITH_EDITOR
+	void OnBlueprintChanged(UBlueprint* Blueprint);
+	virtual void PreEditChange(FProperty* PropertyAboutToChange) override;
+	virtual void PostEditChangeProperty(FPropertyChangedEvent& PropertyChangedEvent) override;
+	virtual void PreEditUndo() override;
+	virtual void PostEditUndo() override;
+#endif // WITH_EDITOR
+
 	//~Begin UCPGSettings interface
 	virtual UPCGNode* CreateNode() const override;
 
@@ -76,8 +107,14 @@ public:
 protected:
 #if WITH_EDITOR
 	virtual bool IsStructuralProperty(const FName& InPropertyName) const override;
-#endif
 	//~End UPCGBaseSubgraphSettings interface
+
+	void SetupBlueprintEvent();
+	void TeardownBlueprintEvent();
+#endif
+
+private:
+	void RefreshTemplateActor();
 };
 
 UCLASS(ClassGroup = (Procedural))
diff --git a/Engine/Plugins/Experimental/PCG/Source/PCG/Public/Metadata/Accessors/PCGAttributeAccessorHelpers.h b/Engine/Plugins/Experimental/PCG/Source/PCG/Public/Metadata/Accessors/PCGAttributeAccessorHelpers.h
--- a/Engine/Plugins/Experimental/PCG/Source/PCG/Public/Metadata/Accessors/PCGAttributeAccessorHelpers.h
+++ b/Engine/Plugins/Experimental/PCG/Source/PCG/Public/Metadata/Accessors/PCGAttributeAccessorHelpers.h
@@ -8,17 +8,20 @@ class IPCGAttributeAccessor;
 class IPCGAttributeAccessorKeyIterator;
 class IPCGAttributeAccessorKeys;
 class FProperty;
+class UClass;
 class UPCGData;
+class UStruct;
 struct FPCGAttributePropertySelector;
 
 namespace PCGAttributeAccessorHelpers
 {
 	TUniquePtr<IPCGAttributeAccessor> CreatePropertyAccessor(const FProperty* InProperty);
+	TUniquePtr<IPCGAttributeAccessor> CreatePropertyAccessor(const FName InPropertyName, const UClass* InClass);
+	TUniquePtr<IPCGAttributeAccessor> CreatePropertyAccessor(const FName InPropertyName, const UStruct* InStruct);
 
 	TUniquePtr<const IPCGAttributeAccessor> CreateConstAccessor(const UPCGData* InData, const FPCGAttributePropertySelector& InSelector);
 	TUniquePtr<IPCGAttributeAccessor> CreateAccessor(UPCGData* InData, const FPCGAttributePropertySelector& InSelector);
 
 	TUniquePtr<const IPCGAttributeAccessorKeys> CreateConstKeys(const UPCGData* InData, const FPCGAttributePropertySelector& InSelector);
 	TUniquePtr<IPCGAttributeAccessorKeys> CreateKeys(UPCGData* InData, const FPCGAttributePropertySelector& InSelector);
-
 }
\ No newline at end of file
