Default point intiialization from actor now sets local bounds

#rnx

diff --git a/Engine/Plugins/Experimental/PCG/Source/PCG/Private/Data/PCGPointData.cpp b/Engine/Plugins/Experimental/PCG/Source/PCG/Private/Data/PCGPointData.cpp
--- a/Engine/Plugins/Experimental/PCG/Source/PCG/Private/Data/PCGPointData.cpp
+++ b/Engine/Plugins/Experimental/PCG/Source/PCG/Private/Data/PCGPointData.cpp
@@ -308,6 +308,10 @@ void UPCGPointData::InitializeFromActor(AActor* InActor)
 	const FVector& Position = Points[0].Transform.GetLocation();
 	Points[0].Seed = PCGHelpers::ComputeSeed((int)Position.X, (int)Position.Y, (int)Position.Z);
 
+	const FBox LocalBounds = PCGHelpers::GetActorLocalBounds(InActor);
+	Points[0].BoundsMin = LocalBounds.Min;
+	Points[0].BoundsMax = LocalBounds.Max;
+
 	TargetActor = InActor;
 	Metadata = NewObject<UPCGMetadata>(this);
 }
