[PCG] Migrate some density-related nodes to newer standards.

Follows pattern established by SelfPruning node fixup - concrete types and separate Params pin.

Deprecation code common to all and factored out.

diff --git a/Engine/Plugins/Experimental/PCG/Source/PCG/Private/Elements/PCGDensityFilter.cpp b/Engine/Plugins/Experimental/PCG/Source/PCG/Private/Elements/PCGDensityFilter.cpp
--- a/Engine/Plugins/Experimental/PCG/Source/PCG/Private/Elements/PCGDensityFilter.cpp
+++ b/Engine/Plugins/Experimental/PCG/Source/PCG/Private/Elements/PCGDensityFilter.cpp
@@ -1,12 +1,25 @@
 // Copyright Epic Games, Inc. All Rights Reserved.
 
 #include "Elements/PCGDensityFilter.h"
+
+#include "PCGCustomVersion.h"
+#include "PCGEdge.h"
 #include "Data/PCGPointData.h"
 #include "Helpers/PCGAsync.h"
 #include "Helpers/PCGSettingsHelpers.h"
 
 #include UE_INLINE_GENERATED_CPP_BY_NAME(PCGDensityFilter)
 
+TArray<FPCGPinProperties> UPCGDensityFilterSettings::InputPinProperties() const
+{
+	TArray<FPCGPinProperties> PinProperties;
+	// TODO in the future type checking of edges will be stricter and a conversion node will be added to convert from other types
+	PinProperties.Emplace(PCGPinConstants::DefaultInputLabel, EPCGDataType::Point);
+	PinProperties.Emplace(PCGPinConstants::DefaultParamsLabel, EPCGDataType::Param, /*bInAllowMultipleConnections=*/false);
+
+	return PinProperties;
+}
+
 FPCGElementPtr UPCGDensityFilterSettings::CreateElement() const
 {
 	return MakeShared<FPCGDensityFilterElement>();
@@ -110,3 +123,17 @@ bool FPCGDensityFilterElement::ExecuteInternal(FPCGContext* Context) const
 
 	return true;
 }
+
+#if WITH_EDITOR
+void UPCGDensityFilterSettings::ApplyDeprecationBeforeUpdatePins(UPCGNode* InOutNode, TArray<TObjectPtr<UPCGPin>>& InputPins, TArray<TObjectPtr<UPCGPin>>& OutputPins)
+{
+	Super::ApplyDeprecationBeforeUpdatePins(InOutNode, InputPins, OutputPins);
+
+	check(InOutNode);
+
+	if (DataVersion < FPCGCustomVersion::MoveParamsOffFirstPinDensityNodes)
+	{
+		PCGSettingsHelpers::DeprecationBreakOutParamsToPin(InOutNode, InputPins, OutputPins);
+	}
+}
+#endif // WITH_EDITOR
diff --git a/Engine/Plugins/Experimental/PCG/Source/PCG/Private/Elements/PCGDensityNoise.cpp b/Engine/Plugins/Experimental/PCG/Source/PCG/Private/Elements/PCGDensityNoise.cpp
--- a/Engine/Plugins/Experimental/PCG/Source/PCG/Private/Elements/PCGDensityNoise.cpp
+++ b/Engine/Plugins/Experimental/PCG/Source/PCG/Private/Elements/PCGDensityNoise.cpp
@@ -2,6 +2,8 @@
 
 #include "Elements/PCGDensityNoise.h"
 
+#include "PCGCustomVersion.h"
+#include "PCGEdge.h"
 #include "PCGHelpers.h"
 #include "Helpers/PCGSettingsHelpers.h"
 
@@ -14,6 +16,16 @@ UPCGDensityNoiseSettings::UPCGDensityNoiseSettings()
 	bUseSeed = true;
 }
 
+TArray<FPCGPinProperties> UPCGDensityNoiseSettings::InputPinProperties() const
+{
+	TArray<FPCGPinProperties> PinProperties;
+	// TODO in the future type checking of edges will be stricter and a conversion node will be added to convert from other types
+	PinProperties.Emplace(PCGPinConstants::DefaultInputLabel, EPCGDataType::Point);
+	PinProperties.Emplace(PCGPinConstants::DefaultParamsLabel, EPCGDataType::Param, /*bInAllowMultipleConnections=*/false);
+
+	return PinProperties;
+}
+
 FPCGElementPtr UPCGDensityNoiseSettings::CreateElement() const
 {
 	return MakeShared<FPCGDensityNoiseElement>();
@@ -82,3 +94,16 @@ bool FPCGDensityNoiseElement::ExecuteInternal(FPCGContext* Context) const
 	return true;
 }
 
+#if WITH_EDITOR
+void UPCGDensityNoiseSettings::ApplyDeprecationBeforeUpdatePins(UPCGNode* InOutNode, TArray<TObjectPtr<UPCGPin>>& InputPins, TArray<TObjectPtr<UPCGPin>>& OutputPins)
+{
+	Super::ApplyDeprecationBeforeUpdatePins(InOutNode, InputPins, OutputPins);
+
+	check(InOutNode);
+
+	if (DataVersion < FPCGCustomVersion::MoveParamsOffFirstPinDensityNodes)
+	{
+		PCGSettingsHelpers::DeprecationBreakOutParamsToPin(InOutNode, InputPins, OutputPins);
+	}
+}
+#endif // WITH_EDITOR
diff --git a/Engine/Plugins/Experimental/PCG/Source/PCG/Private/Elements/PCGDensityRemapElement.cpp b/Engine/Plugins/Experimental/PCG/Source/PCG/Private/Elements/PCGDensityRemapElement.cpp
--- a/Engine/Plugins/Experimental/PCG/Source/PCG/Private/Elements/PCGDensityRemapElement.cpp
+++ b/Engine/Plugins/Experimental/PCG/Source/PCG/Private/Elements/PCGDensityRemapElement.cpp
@@ -2,6 +2,8 @@
 
 #include "Elements/PCGDensityRemapElement.h"
 
+#include "PCGCustomVersion.h"
+#include "PCGEdge.h"
 #include "PCGHelpers.h"
 #include "Helpers/PCGSettingsHelpers.h"
 
@@ -9,6 +11,16 @@
 
 #include UE_INLINE_GENERATED_CPP_BY_NAME(PCGDensityRemapElement)
 
+TArray<FPCGPinProperties> UPCGDensityRemapSettings::InputPinProperties() const
+{
+	TArray<FPCGPinProperties> PinProperties;
+	// TODO in the future type checking of edges will be stricter and a conversion node will be added to convert from other types
+	PinProperties.Emplace(PCGPinConstants::DefaultInputLabel, EPCGDataType::Point);
+	PinProperties.Emplace(PCGPinConstants::DefaultParamsLabel, EPCGDataType::Param, /*bInAllowMultipleConnections=*/false);
+
+	return PinProperties;
+}
+
 FPCGElementPtr UPCGDensityRemapSettings::CreateElement() const
 {
 	return MakeShared<FPCGDensityRemapElement>();
@@ -73,3 +85,16 @@ bool FPCGDensityRemapElement::ExecuteInternal(FPCGContext* Context) const
 	return true;
 }
 
+#if WITH_EDITOR
+void UPCGDensityRemapSettings::ApplyDeprecationBeforeUpdatePins(UPCGNode* InOutNode, TArray<TObjectPtr<UPCGPin>>& InputPins, TArray<TObjectPtr<UPCGPin>>& OutputPins)
+{
+	Super::ApplyDeprecationBeforeUpdatePins(InOutNode, InputPins, OutputPins);
+
+	check(InOutNode);
+
+	if (DataVersion < FPCGCustomVersion::MoveParamsOffFirstPinDensityNodes)
+	{
+		PCGSettingsHelpers::DeprecationBreakOutParamsToPin(InOutNode, InputPins, OutputPins);
+	}
+}
+#endif // WITH_EDITOR
diff --git a/Engine/Plugins/Experimental/PCG/Source/PCG/Private/Elements/PCGSelfPruning.cpp b/Engine/Plugins/Experimental/PCG/Source/PCG/Private/Elements/PCGSelfPruning.cpp
--- a/Engine/Plugins/Experimental/PCG/Source/PCG/Private/Elements/PCGSelfPruning.cpp
+++ b/Engine/Plugins/Experimental/PCG/Source/PCG/Private/Elements/PCGSelfPruning.cpp
@@ -256,14 +256,6 @@ TArray<FPCGPinProperties> UPCGSelfPruningSettings::InputPinProperties() const
 	return PinProperties;
 }
 
-TArray<FPCGPinProperties> UPCGSelfPruningSettings::OutputPinProperties() const
-{
-	TArray<FPCGPinProperties> PinProperties;
-	PinProperties.Emplace(PCGPinConstants::DefaultOutputLabel, EPCGDataType::Point);
-
-	return PinProperties;
-}
-
 FPCGElementPtr UPCGSelfPruningSettings::CreateElement() const
 {
 	return MakeShared<FPCGSelfPruningElement>();
@@ -296,44 +288,7 @@ void UPCGSelfPruningSettings::ApplyDeprecationBeforeUpdatePins(UPCGNode* InOutNo
 
 	if (DataVersion < FPCGCustomVersion::MoveSelfPruningParamsOffFirstPin)
 	{
-		// A dedicated pin for params will be added when the pins are updated. Here we detect any params connections to the In pin
-		// and disconnect them, and move the first params connection to a new params pin.
-		
-		// Check basic conditions for which the code below should run.
-		check(InputPins.Num() == 1);
-		check(InputPins[0] && InputPins[0]->Properties.AllowedTypes == EPCGDataType::Any);
-
-		UPCGPin* InPin = InputPins[0];
-
-		// Add params pin with good defaults (UpdatePins will ensure pin details are correct later).
-		UPCGPin* NewParamsPin = NewObject<UPCGPin>(InOutNode);
-		NewParamsPin->Node = InOutNode;
-		NewParamsPin->Properties.AllowedTypes = EPCGDataType::Param;
-		NewParamsPin->Properties.Label = PCGPinConstants::DefaultParamsLabel;
-		NewParamsPin->Properties.bAllowMultipleConnections = false;
-		InputPins.Add(NewParamsPin);
-
-		// Make list of param pins that In pin is currently connected to.
-		TArray<UPCGPin*> UpstreamParamPins;
-		for (const UPCGEdge* Connection : InPin->Edges)
-		{
-			if (Connection->InputPin && Connection->InputPin->Properties.AllowedTypes == EPCGDataType::Param)
-			{
-				UpstreamParamPins.Add(Connection->InputPin);
-			}
-		}
-
-		// Break all connections to param pins, and connect the first such pin to the new params pin on this node.
-		for (UPCGPin* Pin : UpstreamParamPins)
-		{
-			InPin->BreakEdgeTo(Pin);
-
-			// Params never support multiple connections as a rule (user must merge params themselves), so just connect first.
-			if (!NewParamsPin->IsConnected())
-			{
-				NewParamsPin->AddEdgeTo(Pin);
-			}
-		}
+		PCGSettingsHelpers::DeprecationBreakOutParamsToPin(InOutNode, InputPins, OutputPins);
 	}
 }
 #endif // WITH_EDITOR
diff --git a/Engine/Plugins/Experimental/PCG/Source/PCG/Private/Helpers/PCGSettingsHelpers.cpp b/Engine/Plugins/Experimental/PCG/Source/PCG/Private/Helpers/PCGSettingsHelpers.cpp
--- a/Engine/Plugins/Experimental/PCG/Source/PCG/Private/Helpers/PCGSettingsHelpers.cpp
+++ b/Engine/Plugins/Experimental/PCG/Source/PCG/Private/Helpers/PCGSettingsHelpers.cpp
@@ -3,6 +3,7 @@
 #include "Helpers/PCGSettingsHelpers.h"
 
 #include "PCGComponent.h"
+#include "PCGEdge.h"
 #include "PCGHelpers.h"
 #include "PCGSettings.h"
 
@@ -212,4 +213,43 @@ namespace PCGSettingsHelpers
 		const int SettingsSeed = InParams ? PCGSettingsHelpers::GetValue(GET_MEMBER_NAME_CHECKED(UPCGSettings, Seed), InSettings->Seed, InParams) : InSettings->Seed;
 		return InComponent ? PCGHelpers::ComputeSeed(SettingsSeed, InComponent->Seed) : SettingsSeed;
 	}
-}
\ No newline at end of file
+
+	void DeprecationBreakOutParamsToPin(UPCGNode* InOutNode, TArray<TObjectPtr<UPCGPin>>& InputPins, TArray<TObjectPtr<UPCGPin>>& OutputPins)
+	{
+		// Check basic conditions for which the code below should run.
+		check(InputPins.Num() == 1);
+		check(InputPins[0] && InputPins[0]->Properties.AllowedTypes == EPCGDataType::Any);
+
+		UPCGPin* InPin = InputPins[0];
+
+		// Add params pin with good defaults (UpdatePins will ensure pin details are correct later).
+		UPCGPin* NewParamsPin = NewObject<UPCGPin>(InOutNode);
+		NewParamsPin->Node = InOutNode;
+		NewParamsPin->Properties.AllowedTypes = EPCGDataType::Param;
+		NewParamsPin->Properties.Label = PCGPinConstants::DefaultParamsLabel;
+		NewParamsPin->Properties.bAllowMultipleConnections = false;
+		InputPins.Add(NewParamsPin);
+
+		// Make list of param pins that In pin is currently connected to.
+		TArray<UPCGPin*> UpstreamParamPins;
+		for (const UPCGEdge* Connection : InPin->Edges)
+		{
+			if (Connection->InputPin && Connection->InputPin->Properties.AllowedTypes == EPCGDataType::Param)
+			{
+				UpstreamParamPins.Add(Connection->InputPin);
+			}
+		}
+
+		// Break all connections to param pins, and connect the first such pin to the new params pin on this node.
+		for (UPCGPin* Pin : UpstreamParamPins)
+		{
+			InPin->BreakEdgeTo(Pin);
+
+			// Params never support multiple connections as a rule (user must merge params themselves), so just connect first.
+			if (!NewParamsPin->IsConnected())
+			{
+				NewParamsPin->AddEdgeTo(Pin);
+			}
+		}
+	}
+}
diff --git a/Engine/Plugins/Experimental/PCG/Source/PCG/Public/Elements/PCGDensityFilter.h b/Engine/Plugins/Experimental/PCG/Source/PCG/Public/Elements/PCGDensityFilter.h
--- a/Engine/Plugins/Experimental/PCG/Source/PCG/Public/Elements/PCGDensityFilter.h
+++ b/Engine/Plugins/Experimental/PCG/Source/PCG/Public/Elements/PCGDensityFilter.h
@@ -3,9 +3,10 @@
 #pragma once
 
 #include "CoreMinimal.h"
-#include "PCGSettings.h"
+
 #include "PCGElement.h"
 #include "PCGNode.h"
+#include "PCGSettings.h"
 
 #include "PCGDensityFilter.generated.h"
 
@@ -19,8 +20,10 @@ public:
 #if WITH_EDITOR
 	virtual FName GetDefaultNodeName() const override { return FName(TEXT("DensityFilter")); }
 	virtual EPCGSettingsType GetType() const override { return EPCGSettingsType::Filter; }
+	virtual void ApplyDeprecationBeforeUpdatePins(UPCGNode* InOutNode, TArray<TObjectPtr<UPCGPin>>& InputPins, TArray<TObjectPtr<UPCGPin>>& OutputPins) override;
 #endif
 
+	virtual TArray<FPCGPinProperties> InputPinProperties() const override;
 	virtual TArray<FPCGPinProperties> OutputPinProperties() const override { return Super::DefaultPointOutputPinProperties(); }
 
 protected:
diff --git a/Engine/Plugins/Experimental/PCG/Source/PCG/Public/Elements/PCGDensityNoise.h b/Engine/Plugins/Experimental/PCG/Source/PCG/Public/Elements/PCGDensityNoise.h
--- a/Engine/Plugins/Experimental/PCG/Source/PCG/Public/Elements/PCGDensityNoise.h
+++ b/Engine/Plugins/Experimental/PCG/Source/PCG/Public/Elements/PCGDensityNoise.h
@@ -33,8 +33,10 @@ public:
 #if WITH_EDITOR
 	virtual FName GetDefaultNodeName() const override { return FName(TEXT("DensityNoise")); }
 	virtual EPCGSettingsType GetType() const override { return EPCGSettingsType::Density; }
+	virtual void ApplyDeprecationBeforeUpdatePins(UPCGNode* InOutNode, TArray<TObjectPtr<UPCGPin>>& InputPins, TArray<TObjectPtr<UPCGPin>>& OutputPins) override;
 #endif
 
+	virtual TArray<FPCGPinProperties> InputPinProperties() const override;
 	virtual TArray<FPCGPinProperties> OutputPinProperties() const override { return Super::DefaultPointOutputPinProperties(); }
 
 protected:
diff --git a/Engine/Plugins/Experimental/PCG/Source/PCG/Public/Elements/PCGDensityRemapElement.h b/Engine/Plugins/Experimental/PCG/Source/PCG/Public/Elements/PCGDensityRemapElement.h
--- a/Engine/Plugins/Experimental/PCG/Source/PCG/Public/Elements/PCGDensityRemapElement.h
+++ b/Engine/Plugins/Experimental/PCG/Source/PCG/Public/Elements/PCGDensityRemapElement.h
@@ -21,8 +21,10 @@ public:
 #if WITH_EDITOR
 	virtual FName GetDefaultNodeName() const override { return FName(TEXT("DensityRemap")); }
 	virtual EPCGSettingsType GetType() const override { return EPCGSettingsType::Density; }
+	virtual void ApplyDeprecationBeforeUpdatePins(UPCGNode* InOutNode, TArray<TObjectPtr<UPCGPin>>& InputPins, TArray<TObjectPtr<UPCGPin>>& OutputPins) override;
 #endif
 
+	virtual TArray<FPCGPinProperties> InputPinProperties() const override;
 	virtual TArray<FPCGPinProperties> OutputPinProperties() const override { return Super::DefaultPointOutputPinProperties(); }
 
 protected:
diff --git a/Engine/Plugins/Experimental/PCG/Source/PCG/Public/Elements/PCGSelfPruning.h b/Engine/Plugins/Experimental/PCG/Source/PCG/Public/Elements/PCGSelfPruning.h
--- a/Engine/Plugins/Experimental/PCG/Source/PCG/Public/Elements/PCGSelfPruning.h
+++ b/Engine/Plugins/Experimental/PCG/Source/PCG/Public/Elements/PCGSelfPruning.h
@@ -43,11 +43,11 @@ public:
 	//~Begin UPCGSettings interface
 	virtual FName GetDefaultNodeName() const override { return FName(TEXT("SelfPruning")); }
 	virtual EPCGSettingsType GetType() const override { return EPCGSettingsType::Filter; }
-	virtual void ApplyDeprecationBeforeUpdatePins(UPCGNode* InOutNode, TArray<TObjectPtr<UPCGPin>>& InputPins, TArray<TObjectPtr<UPCGPin>>& OutputPins);
+	virtual void ApplyDeprecationBeforeUpdatePins(UPCGNode* InOutNode, TArray<TObjectPtr<UPCGPin>>& InputPins, TArray<TObjectPtr<UPCGPin>>& OutputPins) override;
 #endif
 
 	virtual TArray<FPCGPinProperties> InputPinProperties() const override;
-	virtual TArray<FPCGPinProperties> OutputPinProperties() const override;
+	virtual TArray<FPCGPinProperties> OutputPinProperties() const override { return Super::DefaultPointOutputPinProperties(); }
 
 protected:
 	virtual FPCGElementPtr CreateElement() const override;
diff --git a/Engine/Plugins/Experimental/PCG/Source/PCG/Public/Helpers/PCGSettingsHelpers.h b/Engine/Plugins/Experimental/PCG/Source/PCG/Public/Helpers/PCGSettingsHelpers.h
--- a/Engine/Plugins/Experimental/PCG/Source/PCG/Public/Helpers/PCGSettingsHelpers.h
+++ b/Engine/Plugins/Experimental/PCG/Source/PCG/Public/Helpers/PCGSettingsHelpers.h
@@ -9,6 +9,8 @@
 #include <type_traits>
 
 class UPCGComponent;
+class UPCGNode;
+class UPCGPin;
 class UPCGSettings;
 
 namespace PCGSettingsHelpers
@@ -128,6 +130,11 @@ namespace PCGSettingsHelpers
 	{
 		return ComputeSeedWithOverride(InSettings, InComponent.Get(), InParams);
 	}
+
+	/** Utility to call from before-node-update deprecation. A dedicated pin for params will be added when the pins are updated. Here we detect any params
+	*   connections to the In pin and disconnect them, and move the first params connection to a new params pin.
+	*/
+	void DeprecationBreakOutParamsToPin(UPCGNode* InOutNode, TArray<TObjectPtr<UPCGPin>>& InputPins, TArray<TObjectPtr<UPCGPin>>& OutputPins);
 }
 
-#define PCG_GET_OVERRIDEN_VALUE(Settings, Variable, Params) PCGSettingsHelpers::GetValue(GET_MEMBER_NAME_CHECKED(TRemovePointer<TRemoveConst<decltype(Settings)>::Type>::Type, Variable), (Settings)->Variable, Params)
\ No newline at end of file
+#define PCG_GET_OVERRIDEN_VALUE(Settings, Variable, Params) PCGSettingsHelpers::GetValue(GET_MEMBER_NAME_CHECKED(TRemovePointer<TRemoveConst<decltype(Settings)>::Type>::Type, Variable), (Settings)->Variable, Params)
diff --git a/Engine/Plugins/Experimental/PCG/Source/PCG/Public/PCGCustomVersion.h b/Engine/Plugins/Experimental/PCG/Source/PCG/Public/PCGCustomVersion.h
--- a/Engine/Plugins/Experimental/PCG/Source/PCG/Public/PCGCustomVersion.h
+++ b/Engine/Plugins/Experimental/PCG/Source/PCG/Public/PCGCustomVersion.h
@@ -18,6 +18,8 @@ struct PCG_API FPCGCustomVersion
 		
 		MoveSelfPruningParamsOffFirstPin = 2,
 
+		MoveParamsOffFirstPinDensityNodes = 3,
+
 		// -----<new versions can be added above this line>-------------------------------------------------
 		VersionPlusOne,
 		LatestVersion = VersionPlusOne - 1
