[PCG] Split surface sampler node pins into a Surface input and a Bounding Shape input.

For a description of the new behaviour, see the node and pin tooltips in the code.

Deprecation code connects the current 'In' input to the Bounding Shape which should preserve current graphs.

A new API is added to SpatialData that looks for a shape of a given dimension. This happened implicitly before but is explicit now. This could be extended in the future to return multiple shapes if there are bugs with it picking the first encountered. The implementation of this is fairly simple and seems to work in my cases, but notably it does not make use of some previous things like HasNonTrivialTransform(), if this sounds like trouble I could use a chat/test case that exhibits issues or limitations.

diff --git a/Engine/Plugins/Experimental/PCG/Source/PCG/Private/Data/PCGIntersectionData.cpp b/Engine/Plugins/Experimental/PCG/Source/PCG/Private/Data/PCGIntersectionData.cpp
--- a/Engine/Plugins/Experimental/PCG/Source/PCG/Private/Data/PCGIntersectionData.cpp
+++ b/Engine/Plugins/Experimental/PCG/Source/PCG/Private/Data/PCGIntersectionData.cpp
@@ -3,8 +3,9 @@
 #include "Data/PCGIntersectionData.h"
 #include "Data/PCGPointData.h"
 #include "Helpers/PCGAsync.h"
-#include "PCGHelpers.h"
+#include "Math/NumericLimits.h"
 #include "Metadata/PCGMetadataAccessor.h"
+#include "PCGHelpers.h"
 
 #include UE_INLINE_GENERATED_CPP_BY_NAME(PCGIntersectionData)
 
@@ -114,6 +115,47 @@ bool UPCGIntersectionData::HasNonTrivialTransform() const
 	return A->HasNonTrivialTransform() || B->HasNonTrivialTransform();
 }
 
+const UPCGSpatialData* UPCGIntersectionData::FindShapeFromNetwork(const int InDimension) const
+{
+	check(A && B);
+
+	if (InDimension >= 0)
+	{
+		// Return first candidate that matches Dimension
+		if (const UPCGSpatialData* CandidateA = A->FindShapeFromNetwork(InDimension))
+		{
+			return CandidateA;
+		}
+
+		if (const UPCGSpatialData* CandidateB = B->FindShapeFromNetwork(InDimension))
+		{
+			return CandidateB;
+		}
+	}
+	else
+	{
+		// Return lowest Dimension candidate
+		const UPCGSpatialData* Result = nullptr;
+
+		int LowestDimension = MAX_uint32;
+		const UPCGSpatialData* LowestDimensionalShape = nullptr;
+
+		if (const UPCGSpatialData* CandidateA = A->FindShapeFromNetwork(InDimension))
+		{
+			Result = CandidateA;
+			LowestDimension = CandidateA->GetDimension();
+		}
+
+		const UPCGSpatialData* CandidateB = B->FindShapeFromNetwork(InDimension);
+		if (CandidateB && CandidateB->GetDimension() < LowestDimension)
+		{
+			Result = CandidateB;
+		}
+	}
+
+	return nullptr;
+}
+
 const UPCGPointData* UPCGIntersectionData::CreatePointData(FPCGContext* Context) const
 {
 	check(A && B);
diff --git a/Engine/Plugins/Experimental/PCG/Source/PCG/Private/Data/PCGUnionData.cpp b/Engine/Plugins/Experimental/PCG/Source/PCG/Private/Data/PCGUnionData.cpp
--- a/Engine/Plugins/Experimental/PCG/Source/PCG/Private/Data/PCGUnionData.cpp
+++ b/Engine/Plugins/Experimental/PCG/Source/PCG/Private/Data/PCGUnionData.cpp
@@ -3,6 +3,7 @@
 #include "Data/PCGUnionData.h"
 #include "Data/PCGPointData.h"
 #include "Helpers/PCGAsync.h"
+#include "Math/NumericLimits.h"
 #include "PCGHelpers.h"
 
 #include UE_INLINE_GENERATED_CPP_BY_NAME(PCGUnionData)
@@ -160,6 +161,43 @@ bool UPCGUnionData::HasNonTrivialTransform() const
 	return (FirstNonTrivialTransformData != nullptr || Super::HasNonTrivialTransform());
 }
 
+const UPCGSpatialData* UPCGUnionData::FindShapeFromNetwork(const int InDimension) const
+{
+	if (InDimension >= 0)
+	{
+		// Return first candidate that matches Dimension
+		for (TObjectPtr<const UPCGSpatialData> Datum : Data)
+		{
+			const UPCGSpatialData* Candidate = Datum ? Datum->FindShapeFromNetwork(InDimension) : nullptr;
+			if (Candidate)
+			{
+				return Candidate;
+			}
+		}
+	}
+	else
+	{
+		// Return lowest Dimension candidate
+		const UPCGSpatialData* Result = nullptr;
+		int LowestDimension = MAX_uint32;
+
+		for (TObjectPtr<const UPCGSpatialData> Datum : Data)
+		{
+			const UPCGSpatialData* Candidate = Datum ? Datum->FindShapeFromNetwork(InDimension) : nullptr;
+
+			if (Candidate && Candidate->GetDimension() < LowestDimension)
+			{
+				LowestDimension = Candidate->GetDimension();
+				Result = Candidate;
+			}
+		}
+
+		return Result;
+	}
+
+	return nullptr;
+}
+
 const UPCGPointData* UPCGUnionData::CreatePointData(FPCGContext* Context) const
 {
 	TRACE_CPUPROFILER_EVENT_SCOPE(UPCGUnionData::CreatePointData);
diff --git a/Engine/Plugins/Experimental/PCG/Source/PCG/Private/Data/PCGWorldData.cpp b/Engine/Plugins/Experimental/PCG/Source/PCG/Private/Data/PCGWorldData.cpp
--- a/Engine/Plugins/Experimental/PCG/Source/PCG/Private/Data/PCGWorldData.cpp
+++ b/Engine/Plugins/Experimental/PCG/Source/PCG/Private/Data/PCGWorldData.cpp
@@ -252,7 +252,7 @@ const UPCGPointData* UPCGWorldRayHitData::CreatePointData(FPCGContext* Context,
 	PCGSurfaceSampler::FSurfaceSamplerSettings SamplerSettings;
 	if (SamplerSettings.Initialize(nullptr, Context, EffectiveBounds))
 	{
-		PCGSurfaceSampler::SampleSurface(Context, this, SamplerSettings, Data);
+		PCGSurfaceSampler::SampleSurface(Context, this, nullptr, SamplerSettings, Data);
 	}
 
 	return Data;
diff --git a/Engine/Plugins/Experimental/PCG/Source/PCG/Private/Elements/PCGSurfaceSampler.cpp b/Engine/Plugins/Experimental/PCG/Source/PCG/Private/Elements/PCGSurfaceSampler.cpp
--- a/Engine/Plugins/Experimental/PCG/Source/PCG/Private/Elements/PCGSurfaceSampler.cpp
+++ b/Engine/Plugins/Experimental/PCG/Source/PCG/Private/Elements/PCGSurfaceSampler.cpp
@@ -3,16 +3,26 @@
 #include "Elements/PCGSurfaceSampler.h"
 
 #include "PCGComponent.h"
+#include "PCGCustomVersion.h"
 #include "PCGHelpers.h"
 #include "Data/PCGPointData.h"
 #include "Data/PCGSpatialData.h"
 #include "Helpers/PCGAsync.h"
+#include "Helpers/PCGBlueprintHelpers.h"
 #include "Helpers/PCGSettingsHelpers.h"
 
 #include "Math/RandomStream.h"
 
 #include UE_INLINE_GENERATED_CPP_BY_NAME(PCGSurfaceSampler)
 
+#define LOCTEXT_NAMESPACE "PCGSurfaceSamplerElement"
+
+namespace PCGSurfaceSamplerConstants
+{
+	const FName SurfaceLabel = TEXT("Surface");
+	const FName BoundingShapeLabel = TEXT("Bounding Shape");
+}
+
 namespace PCGSurfaceSampler
 {
 	bool FSurfaceSamplerSettings::Initialize(const UPCGSurfaceSamplerSettings* InSettings, FPCGContext* Context, const FBox& InputBounds)
@@ -84,6 +94,8 @@ namespace PCGSurfaceSampler
 
 		Ratio = TargetPointCount / (FVector::FReal)CellCount;
 
+		InputBoundsMaxZ = InputBounds.Max.Z;
+
 		return true;
 	}
 
@@ -95,21 +107,25 @@ namespace PCGSurfaceSampler
 		return FIntVector2(CellMinX + (Index % CellCountX), CellMinY + (Index / CellCountX));
 	}
 
-	UPCGPointData* SampleSurface(FPCGContext* Context, const UPCGSpatialData* SpatialInput, const FSurfaceSamplerSettings& LoopData)
+	UPCGPointData* SampleSurface(FPCGContext* Context, const UPCGSpatialData* InSurface, const UPCGSpatialData* InBoundingShape, const FSurfaceSamplerSettings& LoopData)
 	{
 		UPCGPointData* SampledData = NewObject<UPCGPointData>();
-		SampledData->InitializeFromData(SpatialInput);
+		SampledData->InitializeFromData(InSurface);
 
-		SampleSurface(Context, SpatialInput, LoopData, SampledData);
+		SampleSurface(Context, InSurface, InBoundingShape, LoopData, SampledData);
 
 		return SampledData;
 	}
 
-	void SampleSurface(FPCGContext* Context, const UPCGSpatialData* SpatialInput, const FSurfaceSamplerSettings& LoopData, UPCGPointData* SampledData)
+	void SampleSurface(FPCGContext* Context, const UPCGSpatialData* InSurface, const UPCGSpatialData* InBoundingShape, const FSurfaceSamplerSettings& LoopData, UPCGPointData* SampledData)
 	{
+		check(InSurface);
+
 		TArray<FPCGPoint>& SampledPoints = SampledData->GetMutablePoints();
 
-		FPCGAsync::AsyncPointProcessing(Context, LoopData.CellCount, SampledPoints, [&LoopData, SampledData, SpatialInput](int32 Index, FPCGPoint& OutPoint)
+		FPCGProjectionParams ProjectionParams{};
+
+		FPCGAsync::AsyncPointProcessing(Context, LoopData.CellCount, SampledPoints, [&LoopData, SampledData, InBoundingShape, InSurface, &ProjectionParams](int32 Index, FPCGPoint& OutPoint)
 		{
 			const FIntVector2 Indices = LoopData.ComputeCellIndices(Index);
 
@@ -122,37 +138,47 @@ namespace PCGSurfaceSampler
 
 			const float Ratio = LoopData.Ratio;
 
-			if (Chance < Ratio)
+			if (Chance >= Ratio)
+			{
+				return false;
+			}
+
+			const float RandX = RandomSource.FRand();
+			const float RandY = RandomSource.FRand();
+			
+			const FVector TentativeLocation = FVector(CurrentX + RandX * InnerCellSize.X, CurrentY + RandY * InnerCellSize.Y, LoopData.InputBoundsMaxZ);
+			const FBox LocalBound(-LoopData.PointExtents, LoopData.PointExtents);
+
+			// Firstly project onto elected generating shape to move to final position.
+			if (!InSurface->ProjectPoint(FTransform(TentativeLocation), LocalBound, ProjectionParams, OutPoint, SampledData->Metadata))
 			{
-				const float RandX = RandomSource.FRand();
-				const float RandY = RandomSource.FRand();
-				 
-				const FVector TentativeLocation = FVector(CurrentX + RandX * InnerCellSize.X, CurrentY + RandY * InnerCellSize.Y, 0.0f);
-				const FBox LocalBound(-LoopData.PointExtents, LoopData.PointExtents);
+				return false;
+			}
 
+			// Now run gauntlet of shape network (if there is one) to accept or reject the point.
+			if (InBoundingShape)
+			{
+				FPCGPoint BoundingShapeSample;
 #if WITH_EDITORONLY_DATA
-				if (SpatialInput->SamplePoint(FTransform(TentativeLocation), LocalBound, OutPoint, SampledData->Metadata) || LoopData.bKeepZeroDensityPoints)
+				if (!InBoundingShape->SamplePoint(OutPoint.Transform, OutPoint.GetLocalBounds(), BoundingShapeSample, nullptr) && !LoopData.bKeepZeroDensityPoints)
 #else
-				if (SpatialInput->SamplePoint(FTransform(TentativeLocation), LocalBound, OutPoint, SampledData->Metadata))
+				if (!InBoundingShape->SamplePoint(OutPoint.Transform, OutPoint.GetLocalBounds(), BoundingShapeSample, nullptr))
 #endif
-				{
-					// Apply final parameters on the point
-					OutPoint.SetExtents(LoopData.PointExtents);
-					OutPoint.Density *= (LoopData.bApplyDensityToPoints ? ((Ratio - Chance) / Ratio) : 1.0f);
-					OutPoint.Steepness = LoopData.PointSteepness;
-					OutPoint.Seed = RandomSource.GetCurrentSeed();
-
-					return true;
-				}
-				else
 				{
 					return false;
 				}
+
+				// Produce smooth density field
+				OutPoint.Density *= BoundingShapeSample.Density;
 			}
-			else
-			{
-				return false;
-			}
+
+			// Apply final parameters on the point
+			OutPoint.SetExtents(LoopData.PointExtents);
+			OutPoint.Density *= (LoopData.bApplyDensityToPoints ? ((Ratio - Chance) / Ratio) : 1.0f);
+			OutPoint.Steepness = LoopData.PointSteepness;
+			OutPoint.Seed = RandomSource.GetCurrentSeed();
+
+			return true;
 		});
 
 		if (Context)
@@ -167,6 +193,38 @@ UPCGSurfaceSamplerSettings::UPCGSurfaceSamplerSettings()
 	bUseSeed = true;
 }
 
+#if WITH_EDITOR
+FText UPCGSurfaceSamplerSettings::GetNodeTooltipText() const
+{
+	return LOCTEXT("SurfaceSamplerNodeTooltip", "Generates points in two dimensional domain that sample the Surface input and lie within the Bounding Shape input.");
+}
+#endif
+
+TArray<FPCGPinProperties> UPCGSurfaceSamplerSettings::InputPinProperties() const
+{
+	TArray<FPCGPinProperties> PinProperties;
+	PinProperties.Emplace(PCGSurfaceSamplerConstants::SurfaceLabel, EPCGDataType::Surface, /*bAllowMultipleConnections=*/true, /*bAllowMultipleData=*/false, LOCTEXT("SurfaceSamplerSurfacePinTooltip",
+		"The surface to sample with points. Points will be generated in the two dimensional footprint of the combined bounds of the Surface and the Bounding Shape (if any) "
+		"and then projected onto this surface. If this input is omitted then the network of shapes connected to the Bounding Shape pin will be inspected for a surface "
+		"shape to use to project the points onto."
+	));
+	// Only one connection allowed, user can union multiple shapes
+	PinProperties.Emplace(PCGSurfaceSamplerConstants::BoundingShapeLabel, EPCGDataType::Spatial, /*bInAllowMultipleConnections=*/false, /*bAllowMultipleData=*/false, LOCTEXT("SurfaceSamplerBoundingShapePinTooltip",
+		"All sampled points must be contained within this shape. If this input is omitted then bounds will be taken from the actor so that points are contained within actor bounds. "
+		"The Unbounded property disables this and instead generates over the entire bounds of Surface."
+	));
+
+	return PinProperties;
+}
+
+TArray<FPCGPinProperties> UPCGSurfaceSamplerSettings::OutputPinProperties() const
+{
+	TArray<FPCGPinProperties> PinProperties;
+	PinProperties.Emplace(PCGPinConstants::DefaultOutputLabel, EPCGDataType::Point);
+
+	return PinProperties;
+}
+
 void UPCGSurfaceSamplerSettings::PostLoad()
 {
 	Super::PostLoad();
@@ -178,6 +236,13 @@ void UPCGSurfaceSamplerSettings::PostLoad()
 	}
 }
 
+#if WITH_EDITOR
+bool UPCGSurfaceSamplerSettings::IsPinUsedByNodeExecution(const UPCGPin* InPin) const
+{
+	return !bUnbounded || InPin->Properties.Label != PCGSurfaceSamplerConstants::BoundingShapeLabel;
+}
+#endif
+
 FPCGElementPtr UPCGSurfaceSamplerSettings::CreateElement() const
 {
 	return MakeShared<FPCGSurfaceSamplerElement>();
@@ -190,8 +255,63 @@ bool FPCGSurfaceSamplerElement::ExecuteInternal(FPCGContext* Context) const
 	check(Context);
 	const UPCGSurfaceSamplerSettings* Settings = Context->GetInputSettings<UPCGSurfaceSamplerSettings>();
 	check(Settings);
+	
+	TArray<FPCGTaggedData>& Outputs = Context->OutputData.TaggedData;
+
+	// Grab the Bounding Shape input if there is one.
+	TArray<FPCGTaggedData> BoundingShapeInputs = Context->InputData.GetInputsByPin(PCGSurfaceSamplerConstants::BoundingShapeLabel);
+	const UPCGSpatialData* BoundingShapeSpatialInput = nullptr;
+	if (!Settings->bUnbounded)
+	{
+		if (BoundingShapeInputs.Num() > 0)
+		{
+			ensure(BoundingShapeInputs.Num() == 1);
+			BoundingShapeSpatialInput = Cast<UPCGSpatialData>(BoundingShapeInputs[0].Data);
+		}
+		else if (Context->SourceComponent.IsValid())
+		{
+			// Fallback to getting bounds from actor
+			BoundingShapeSpatialInput = Cast<UPCGSpatialData>(Context->SourceComponent->GetActorPCGData());
+		}
+	}
+	else if (BoundingShapeInputs.Num() > 0)
+	{
+		PCGE_LOG(Verbose, "The bounds of the Bounding Shape input pin will be ignored because the Unbounded option is enabled.");
+	}
+
+	FBox BoundingShapeBounds(EForceInit::ForceInit);
+	if (BoundingShapeSpatialInput)
+	{
+		BoundingShapeBounds = BoundingShapeSpatialInput->GetBounds();
+	}
+
+	TArray<FPCGTaggedData> SurfaceInputs = Context->InputData.GetInputsByPin(PCGSurfaceSamplerConstants::SurfaceLabel);
+
+	// Construct a list of shapes to generate samples from. Prefer to get these directly from the first input pin.
+	TArray<const UPCGSpatialData*, TInlineAllocator<16>> GeneratingShapes;
+	for (FPCGTaggedData& TaggedData : SurfaceInputs)
+	{
+		if (UPCGSpatialData* SpatialData = Cast<UPCGSpatialData>(TaggedData.Data))
+		{
+			GeneratingShapes.Add(SpatialData);
+			Outputs.Add(TaggedData);
+		}
+	}
+
+	// If no shapes were obtained from the first input pin, try to find a shape to sample from nodes connected to the second pin.
+	const UPCGSpatialData* GeneratorFromBoundingShapeInput = nullptr;
+	if (GeneratingShapes.Num() == 0 && BoundingShapeSpatialInput)
+	{
+		GeneratorFromBoundingShapeInput = BoundingShapeSpatialInput->FindShapeFromNetwork(/*InDimension=*/2);
+		if (GeneratorFromBoundingShapeInput)
+		{
+			GeneratingShapes.Add(GeneratorFromBoundingShapeInput);
+
+			check(BoundingShapeInputs.Num() == 1);
+			Outputs.Add(BoundingShapeInputs[0]);
+		}
+	}
 
-	TArray<FPCGTaggedData> Inputs = Context->InputData.GetInputs();
 	UPCGParamData* Params = Context->InputData.GetParams();
 
 	// Early out on invalid settings
@@ -202,22 +322,20 @@ bool FPCGSurfaceSamplerElement::ExecuteInternal(FPCGContext* Context) const
 		PCGE_LOG(Warning, "Skipped - Invalid point extents");
 		return true;
 	}
-
-	TArray<FPCGTaggedData>& Outputs = Context->OutputData.TaggedData; 
-
+	
 	// TODO: embarassingly parallel loop
-	for (const FPCGTaggedData& Input : Inputs)
+	for (int GenerationIndex = 0; GenerationIndex < GeneratingShapes.Num(); ++GenerationIndex)
 	{
-		const UPCGSpatialData* SpatialInput = Cast<UPCGSpatialData>(Input.Data);
+		// If we have generating shape inputs, use them
+		const UPCGSpatialData* GeneratingShape = GeneratingShapes[GenerationIndex];
+		check(GeneratingShape);
 
-		if (!SpatialInput || !SpatialInput->TargetActor.IsValid())
+		// Calculate the intersection of bounds of the provided inputs
+		FBox InputBounds = GeneratingShape->GetBounds();
+		if (BoundingShapeBounds.IsValid)
 		{
-			PCGE_LOG(Error, "Invalid input data");
-			continue;
+			InputBounds = PCGHelpers::OverlapBounds(InputBounds, BoundingShapeBounds);
 		}
-
-		const FBox InputBounds = SpatialInput->GetBounds();
-
 		if (!InputBounds.IsValid)
 		{
 			PCGE_LOG(Warning, "Input data has invalid bounds");
@@ -230,10 +348,8 @@ bool FPCGSurfaceSamplerElement::ExecuteInternal(FPCGContext* Context) const
 			continue;
 		}
 
-		// Finally, create data
-		FPCGTaggedData& Output = Outputs.Emplace_GetRef();
-		Output = Input;
-		Output.Data = PCGSurfaceSampler::SampleSurface(Context, SpatialInput, LoopData);
+		// Sample surface
+		Outputs[GenerationIndex].Data = PCGSurfaceSampler::SampleSurface(Context, GeneratingShape, BoundingShapeSpatialInput, LoopData);
 	}
 
 	// Finally, forward any exclusions/settings
@@ -241,3 +357,23 @@ bool FPCGSurfaceSamplerElement::ExecuteInternal(FPCGContext* Context) const
 
 	return true;
 }
+
+#if WITH_EDITOR
+void UPCGSurfaceSamplerSettings::ApplyDeprecationBeforeUpdatePins(UPCGNode* InOutNode, TArray<TObjectPtr<UPCGPin>>& InputPins, TArray<TObjectPtr<UPCGPin>>& OutputPins)
+{
+	if (DataVersion < FPCGCustomVersion::SplitSamplerNodesInputs)
+	{
+		// Deprecation from a single pin node.
+		check(InputPins.Num() == 1);
+
+		// In prior versions this node had a single "In" pin. In later versions this is split. The node will function
+		// the same if we move all connections from "In" to "Bounding Shape". To make this happen, rename "In" to
+		// "Bounding Shape" just prior to pin update and the edges will be moved over.
+		InputPins[0]->Properties.Label = PCGSurfaceSamplerConstants::BoundingShapeLabel;
+	}
+
+	Super::ApplyDeprecationBeforeUpdatePins(InOutNode, InputPins, OutputPins);
+}
+#endif // WITH_EDITOR
+
+#undef LOCTEXT_NAMESPACE
diff --git a/Engine/Plugins/Experimental/PCG/Source/PCG/Public/Data/PCGDifferenceData.h b/Engine/Plugins/Experimental/PCG/Source/PCG/Public/Data/PCGDifferenceData.h
--- a/Engine/Plugins/Experimental/PCG/Source/PCG/Public/Data/PCGDifferenceData.h
+++ b/Engine/Plugins/Experimental/PCG/Source/PCG/Public/Data/PCGDifferenceData.h
@@ -52,6 +52,7 @@ public:
 	virtual FBox GetStrictBounds() const override;
 	virtual bool SamplePoint(const FTransform& Transform, const FBox& Bounds, FPCGPoint& OutPoint, UPCGMetadata* OutMetadata) const override;
 	virtual bool HasNonTrivialTransform() const override;
+	virtual const UPCGSpatialData* FindShapeFromNetwork(const int InDimension) const override { return Source ? Source->FindShapeFromNetwork(InDimension) : nullptr; }
 protected:
 	virtual UPCGSpatialData* CopyInternal() const override;
 	//~End UPCGSpatialData interface
diff --git a/Engine/Plugins/Experimental/PCG/Source/PCG/Public/Data/PCGIntersectionData.h b/Engine/Plugins/Experimental/PCG/Source/PCG/Public/Data/PCGIntersectionData.h
--- a/Engine/Plugins/Experimental/PCG/Source/PCG/Public/Data/PCGIntersectionData.h
+++ b/Engine/Plugins/Experimental/PCG/Source/PCG/Public/Data/PCGIntersectionData.h
@@ -30,6 +30,7 @@ public:
 	virtual FBox GetStrictBounds() const override;
 	virtual bool SamplePoint(const FTransform& Transform, const FBox& Bounds, FPCGPoint& OutPoint, UPCGMetadata* OutMetadata) const override;
 	virtual bool HasNonTrivialTransform() const override;
+	virtual const UPCGSpatialData* FindShapeFromNetwork(const int InDimension) const override;
 protected:
 	virtual UPCGSpatialData* CopyInternal() const override;
 	//~End UPCGSpatialData interface
diff --git a/Engine/Plugins/Experimental/PCG/Source/PCG/Public/Data/PCGSpatialData.h b/Engine/Plugins/Experimental/PCG/Source/PCG/Public/Data/PCGSpatialData.h
--- a/Engine/Plugins/Experimental/PCG/Source/PCG/Public/Data/PCGSpatialData.h
+++ b/Engine/Plugins/Experimental/PCG/Source/PCG/Public/Data/PCGSpatialData.h
@@ -124,6 +124,9 @@ public:
 	/** True if this operation does not have an inverse and cannot be queried analytically/implicitly, and therefore must be collapsed to an explicit point representation. */
 	virtual bool RequiresCollapseToSample() const { return false; }
 
+	/** A call that is made recursively up through the graph to find the best candidate shape for point generation. If InDimension is -1, finds lowest dimensional shape. */
+	virtual const UPCGSpatialData* FindShapeFromNetwork(const int InDimension) const { return (InDimension == -1 || GetDimension() == InDimension) ? this : nullptr; }
+
 	UPROPERTY(Transient, BlueprintReadWrite, EditAnywhere, Category = Data)
 	TWeakObjectPtr<AActor> TargetActor = nullptr;
 
diff --git a/Engine/Plugins/Experimental/PCG/Source/PCG/Public/Data/PCGUnionData.h b/Engine/Plugins/Experimental/PCG/Source/PCG/Public/Data/PCGUnionData.h
--- a/Engine/Plugins/Experimental/PCG/Source/PCG/Public/Data/PCGUnionData.h
+++ b/Engine/Plugins/Experimental/PCG/Source/PCG/Public/Data/PCGUnionData.h
@@ -43,6 +43,7 @@ public:
 	virtual FBox GetStrictBounds() const override;
 	virtual bool SamplePoint(const FTransform& Transform, const FBox& Bounds, FPCGPoint& OutPoint, UPCGMetadata* OutMetadata) const override;
 	virtual bool HasNonTrivialTransform() const override;
+	virtual const UPCGSpatialData* FindShapeFromNetwork(const int InDimension) const override;
 protected:
 	virtual UPCGSpatialData* CopyInternal() const override;
 	//~End UPCGSpatialData interface
diff --git a/Engine/Plugins/Experimental/PCG/Source/PCG/Public/Elements/PCGSurfaceSampler.h b/Engine/Plugins/Experimental/PCG/Source/PCG/Public/Elements/PCGSurfaceSampler.h
--- a/Engine/Plugins/Experimental/PCG/Source/PCG/Public/Elements/PCGSurfaceSampler.h
+++ b/Engine/Plugins/Experimental/PCG/Source/PCG/Public/Elements/PCGSurfaceSampler.h
@@ -3,13 +3,13 @@
 #pragma once
 
 #include "CoreMinimal.h"
-#include "PCGNode.h"
-#include "PCGSettings.h"
+
 #include "PCGElement.h"
+#include "PCGSettings.h"
 
 #include "PCGSurfaceSampler.generated.h"
 
-class FPCGSurfaceSamplerElement;
+class UPCGNode;
 class UPCGPointData;
 class UPCGSpatialData;
 
@@ -21,6 +21,30 @@ class PCG_API UPCGSurfaceSamplerSettings : public UPCGSettings
 public:
 	UPCGSurfaceSamplerSettings();
 
+	//~Begin UObject interface
+	virtual void PostLoad() override;
+	//~End UObject interface
+
+	//~Begin UPCGSettings interface
+#if WITH_EDITOR
+	virtual FName GetDefaultNodeName() const override { return FName(TEXT("SurfaceSampler")); }
+	virtual FText GetNodeTooltipText() const override;
+	virtual EPCGSettingsType GetType() const override { return EPCGSettingsType::Sampler; }
+	virtual bool IsPinUsedByNodeExecution(const UPCGPin* InPin) const override;
+#endif
+
+	virtual TArray<FPCGPinProperties> InputPinProperties() const override;
+	virtual TArray<FPCGPinProperties> OutputPinProperties() const override;
+
+#if WITH_EDITOR
+	virtual void ApplyDeprecationBeforeUpdatePins(UPCGNode* InOutNode, TArray<TObjectPtr<UPCGPin>>& InputPins, TArray<TObjectPtr<UPCGPin>>& OutputPins) override;
+#endif
+
+protected:
+	virtual FPCGElementPtr CreateElement() const override;
+	// ~End UPCGSettings interface
+
+public:
 	UPROPERTY(BlueprintReadWrite, EditAnywhere, Category = Settings, meta=(ClampMin="0"))
 	float PointsPerSquaredMeter = 0.1f;
 
@@ -33,6 +57,13 @@ public:
 	UPROPERTY(BlueprintReadWrite, EditAnywhere, Category = Settings, meta=(ClampMin="0"))
 	float Looseness = 1.0f;
 
+	/** If no Bounding Shape input is provided the actor bounds are used to limit the sample generation area.
+	* This option allows ignoring the actor bounds and generating over the entire surface. Use with caution as this
+	* may generate a lot of points.
+	*/
+	UPROPERTY(BlueprintReadWrite, EditAnywhere, Category = Settings)
+	bool bUnbounded = false;
+
 	UPROPERTY(BlueprintReadWrite, EditAnywhere, Category = "Settings|Points")
 	bool bApplyDensityToPoints = true;
 
@@ -43,22 +74,6 @@ public:
 	UPROPERTY(Transient, BlueprintReadWrite, EditAnywhere, Category = "Settings|Debug")
 	bool bKeepZeroDensityPoints = false;
 #endif
-
-	//~Begin UObject interface
-	virtual void PostLoad() override;
-	//~End UObject interface
-
-#if WITH_EDITOR
-	//~Begin UPCGSettings interface
-	virtual FName GetDefaultNodeName() const override { return FName(TEXT("SurfaceSampler")); }
-	virtual EPCGSettingsType GetType() const override { return EPCGSettingsType::Sampler; }
-#endif
-
-	virtual TArray<FPCGPinProperties> OutputPinProperties() const override { return Super::DefaultPointOutputPinProperties(); }
-
-protected:
-	virtual FPCGElementPtr CreateElement() const override;
-	// ~End UPCGSettings interface
 };
 
 class FPCGSurfaceSamplerElement : public FSimplePCGElement
@@ -98,8 +113,10 @@ namespace PCGSurfaceSampler
 		int64 TargetPointCount;
 		float Ratio;
 		int Seed;
+
+		FVector::FReal InputBoundsMaxZ;
 	};
 
-	UPCGPointData* SampleSurface(FPCGContext* Context, const UPCGSpatialData* SpatialInput, const FSurfaceSamplerSettings& LoopData);
-	void SampleSurface(FPCGContext* Context, const UPCGSpatialData* SpatialInput, const FSurfaceSamplerSettings& LoopData, UPCGPointData* SampledData);
-}
\ No newline at end of file
+	UPCGPointData* SampleSurface(FPCGContext* Context, const UPCGSpatialData* InSurface, const UPCGSpatialData* InBoundingShape, const FSurfaceSamplerSettings& LoopData);
+	void SampleSurface(FPCGContext* Context, const UPCGSpatialData* InSurface, const UPCGSpatialData* InBoundingShape, const FSurfaceSamplerSettings& LoopData, UPCGPointData* SampledData);
+}
diff --git a/Engine/Plugins/Experimental/PCG/Source/PCG/Public/PCGCustomVersion.h b/Engine/Plugins/Experimental/PCG/Source/PCG/Public/PCGCustomVersion.h
--- a/Engine/Plugins/Experimental/PCG/Source/PCG/Public/PCGCustomVersion.h
+++ b/Engine/Plugins/Experimental/PCG/Source/PCG/Public/PCGCustomVersion.h
@@ -20,6 +20,9 @@ struct PCG_API FPCGCustomVersion
 
 		MoveParamsOffFirstPinDensityNodes = 3,
 
+		// Split samplers to give a sampling shape and a bounding shape inputs
+		SplitSamplerNodesInputs = 4,
+
 		// -----<new versions can be added above this line>-------------------------------------------------
 		VersionPlusOne,
 		LatestVersion = VersionPlusOne - 1
