[PCG] Passthrough behaviour fixes.

Summary of types: Concrete is like points or splines. Spatial is concrete, or any boolean operation over concrete like unions etc. The Composite bit is used to distinguish Spatial and Concrete. Union data is of type Composite. The Union element output pin is of type Spatial, because it will output concrete if it has a single concrete input edge.

* Fixed case where Projection passed through both pins when disabled
* Boolean operations can either directly pass through concrete data, or can pass through booleans. Therefore their type should be Spatial. Without this, pass-through of some data through boolean nodes were being denied as it was deemed that the input type was too wide for the boolean output which was not the case. With this change booleans now output Composite which was not the original design, but is necessary/correct.
* For now, allow outputting Spatial out through a Concrete pin even if it is wider. Since we can right now connect Spatial to Concrete, it means e.g. a point processing node can receive a Spatial input, and when disabled it wants to pass through the Spatial to its Point output, which we let slide for now. One fix would be to do a Collapse, but the Real Fix is to require conversion/collapse upstream.

diff --git a/Engine/Plugins/Experimental/PCG/Source/PCG/Private/Elements/PCGDifferenceElement.cpp b/Engine/Plugins/Experimental/PCG/Source/PCG/Private/Elements/PCGDifferenceElement.cpp
--- a/Engine/Plugins/Experimental/PCG/Source/PCG/Private/Elements/PCGDifferenceElement.cpp
+++ b/Engine/Plugins/Experimental/PCG/Source/PCG/Private/Elements/PCGDifferenceElement.cpp
@@ -21,7 +21,7 @@ TArray<FPCGPinProperties> UPCGDifferenceSettings::InputPinProperties() const
 TArray<FPCGPinProperties> UPCGDifferenceSettings::OutputPinProperties() const
 {
 	TArray<FPCGPinProperties> PinProperties;
-	PinProperties.Emplace(PCGPinConstants::DefaultOutputLabel, EPCGDataType::Composite);
+	PinProperties.Emplace(PCGPinConstants::DefaultOutputLabel, EPCGDataType::Spatial);
 
 	return PinProperties;
 }
diff --git a/Engine/Plugins/Experimental/PCG/Source/PCG/Private/Elements/PCGIntersectionElement.cpp b/Engine/Plugins/Experimental/PCG/Source/PCG/Private/Elements/PCGIntersectionElement.cpp
--- a/Engine/Plugins/Experimental/PCG/Source/PCG/Private/Elements/PCGIntersectionElement.cpp
+++ b/Engine/Plugins/Experimental/PCG/Source/PCG/Private/Elements/PCGIntersectionElement.cpp
@@ -8,7 +8,7 @@
 TArray<FPCGPinProperties> UPCGIntersectionSettings::OutputPinProperties() const
 {
 	TArray<FPCGPinProperties> PinProperties;
-	PinProperties.Emplace(PCGPinConstants::DefaultOutputLabel, EPCGDataType::Composite);
+	PinProperties.Emplace(PCGPinConstants::DefaultOutputLabel, EPCGDataType::Spatial);
 
 	return PinProperties;
 }
diff --git a/Engine/Plugins/Experimental/PCG/Source/PCG/Private/Elements/PCGUnionElement.cpp b/Engine/Plugins/Experimental/PCG/Source/PCG/Private/Elements/PCGUnionElement.cpp
--- a/Engine/Plugins/Experimental/PCG/Source/PCG/Private/Elements/PCGUnionElement.cpp
+++ b/Engine/Plugins/Experimental/PCG/Source/PCG/Private/Elements/PCGUnionElement.cpp
@@ -8,7 +8,7 @@
 TArray<FPCGPinProperties> UPCGUnionSettings::OutputPinProperties() const
 {
 	TArray<FPCGPinProperties> PinProperties;
-	PinProperties.Emplace(PCGPinConstants::DefaultOutputLabel, EPCGDataType::Composite);
+	PinProperties.Emplace(PCGPinConstants::DefaultOutputLabel, EPCGDataType::Spatial);
 
 	return PinProperties;
 }
diff --git a/Engine/Plugins/Experimental/PCG/Source/PCG/Private/PCGElement.cpp b/Engine/Plugins/Experimental/PCG/Source/PCG/Private/PCGElement.cpp
--- a/Engine/Plugins/Experimental/PCG/Source/PCG/Private/PCGElement.cpp
+++ b/Engine/Plugins/Experimental/PCG/Source/PCG/Private/PCGElement.cpp
@@ -205,20 +205,16 @@ void IPCGElement::DisabledPassThroughData(FPCGContext* Context) const
 	const UPCGSettings* Settings = Context->GetInputSettings<UPCGSettings>();
 	check(Settings);
 
-	// Copy as baseline
-	Context->OutputData = Context->InputData;
-
 	if (!Context->Node)
 	{
 		// Full pass-through if we don't have a node
+		Context->OutputData = Context->InputData;
 		return;
 	}
 
 	if (Context->Node->GetInputPins().Num() == 0 || Context->Node->GetOutputPins().Num() == 0)
 	{
 		// No input pins or not output pins, return nothing
-		Context->OutputData.TaggedData.Empty();
-
 		return;
 	}
 
@@ -226,11 +222,12 @@ void IPCGElement::DisabledPassThroughData(FPCGContext* Context) const
 	if (PassThroughPin == nullptr)
 	{
 		// No pin to grab pass through data from
-		Context->OutputData.TaggedData.Empty();
-
 		return;
 	}
 
+	// Grab data from pass-through pin
+	Context->OutputData.TaggedData = Context->InputData.GetInputsByPin(PassThroughPin->Properties.Label);
+
 	const EPCGDataType OutputType = Context->Node->GetOutputPins()[0]->Properties.AllowedTypes;
 
 	// Pass through input data if it is not params, and if the output type supports it (e.g. if we have a incoming
@@ -239,9 +236,17 @@ void IPCGElement::DisabledPassThroughData(FPCGContext* Context) const
 	{
 		EPCGDataType InputType = InData.Data ? InData.Data->GetDataType() : EPCGDataType::None;
 		const bool bInputTypeNotWiderThanOutputType = !(InputType & ~OutputType);
-		return InputType != EPCGDataType::Param && bInputTypeNotWiderThanOutputType;
+
+		// Right now we allow edges from Spatial to Concrete. This can happen for example if a point processing node
+		// is receving a Spatial data, and the node is disabled, it will want to pass the Spatial data through. In the
+		// future we will force collapses/conversions. For now, allow an incoming Spatial to pass out through a Concrete.
+		// TODO remove!
+		const bool bAllowSpatialToConcrete = !!(InputType & EPCGDataType::Spatial) && !!(OutputType & EPCGDataType::Concrete);
+
+		return InputType != EPCGDataType::Param && (bInputTypeNotWiderThanOutputType || bAllowSpatialToConcrete);
 	};
 
+	// Now remove any non-params edges, and if only one edge should come through, remove the others
 	if (Settings->OnlyPassThroughOneEdgeWhenDisabled())
 	{
 		// Find first incoming non-params data that is coming through the pass through pin
diff --git a/Engine/Plugins/Experimental/PCG/Source/PCG/Private/PCGInputOutputSettings.cpp b/Engine/Plugins/Experimental/PCG/Source/PCG/Private/PCGInputOutputSettings.cpp
--- a/Engine/Plugins/Experimental/PCG/Source/PCG/Private/PCGInputOutputSettings.cpp
+++ b/Engine/Plugins/Experimental/PCG/Source/PCG/Private/PCGInputOutputSettings.cpp
@@ -68,7 +68,7 @@ TArray<FPCGPinProperties> UPCGGraphInputOutputSettings::GetPinProperties() const
 {
 	TArray<FPCGPinProperties> PinProperties;
 	const bool bIsInputPin = bIsInput;
-	const EPCGDataType DefaultPinDataType = bIsInput ? EPCGDataType::Composite : EPCGDataType::Any;
+	const EPCGDataType DefaultPinDataType = bIsInput ? EPCGDataType::Spatial : EPCGDataType::Any;
 	Algo::Transform(StaticLabels(), PinProperties, [bIsInputPin, DefaultPinDataType](const FLabelAndTooltip& InLabelAndTooltip) {
 		return FPCGPinProperties(InLabelAndTooltip.Label, DefaultPinDataType, /*bMultiConnections=*/true, /*bMultiData=*/true, InLabelAndTooltip.Tooltip);
 	});
diff --git a/Engine/Plugins/Experimental/PCG/Source/PCG/Private/PCGPin.cpp b/Engine/Plugins/Experimental/PCG/Source/PCG/Private/PCGPin.cpp
--- a/Engine/Plugins/Experimental/PCG/Source/PCG/Private/PCGPin.cpp
+++ b/Engine/Plugins/Experimental/PCG/Source/PCG/Private/PCGPin.cpp
@@ -208,8 +208,6 @@ bool UPCGPin::IsCompatible(const UPCGPin* OtherPin) const
 	const bool bDownstreamComposite = !!(DownstreamPin->Properties.AllowedTypes & EPCGDataType::Composite);
 	if (bUpstreamConcrete && bDownstreamComposite)
 	{
-		// TODO: Bug - this allows any concrete to connect to any other concrete. We need conversion nodes
-		// where appropriate, or block incompatible edges.
 		return true;
 	}
 
diff --git a/Engine/Plugins/Experimental/PCG/Source/PCG/Public/Data/PCGDifferenceData.h b/Engine/Plugins/Experimental/PCG/Source/PCG/Public/Data/PCGDifferenceData.h
--- a/Engine/Plugins/Experimental/PCG/Source/PCG/Public/Data/PCGDifferenceData.h
+++ b/Engine/Plugins/Experimental/PCG/Source/PCG/Public/Data/PCGDifferenceData.h
@@ -46,6 +46,10 @@ public:
 	virtual void PostEditChangeProperty(FPropertyChangedEvent& PropertyChangedEvent) override;
 #endif
 
+	// ~Begin UPCGData interface
+	virtual EPCGDataType GetDataType() const override { return EPCGDataType::Composite; }
+	// ~End UPCGData interface
+
 	//~Begin UPCGSpatialData interface
 	virtual int GetDimension() const override;
 	virtual FBox GetBounds() const override;
diff --git a/Engine/Plugins/Experimental/PCG/Source/PCG/Public/Data/PCGIntersectionData.h b/Engine/Plugins/Experimental/PCG/Source/PCG/Public/Data/PCGIntersectionData.h
--- a/Engine/Plugins/Experimental/PCG/Source/PCG/Public/Data/PCGIntersectionData.h
+++ b/Engine/Plugins/Experimental/PCG/Source/PCG/Public/Data/PCGIntersectionData.h
@@ -24,6 +24,10 @@ public:
 	UFUNCTION(BlueprintCallable, Category = SpatialData)
 	void Initialize(const UPCGSpatialData* InA, const UPCGSpatialData* InB);
 
+	// ~Begin UPCGData interface
+	virtual EPCGDataType GetDataType() const override { return EPCGDataType::Composite; }
+	// ~End UPCGData interface
+
 	//~Begin UPCGSpatialData interface
 	virtual int GetDimension() const override;
 	virtual FBox GetBounds() const override;
diff --git a/Engine/Plugins/Experimental/PCG/Source/PCG/Public/Data/PCGUnionData.h b/Engine/Plugins/Experimental/PCG/Source/PCG/Public/Data/PCGUnionData.h
--- a/Engine/Plugins/Experimental/PCG/Source/PCG/Public/Data/PCGUnionData.h
+++ b/Engine/Plugins/Experimental/PCG/Source/PCG/Public/Data/PCGUnionData.h
@@ -37,6 +37,10 @@ public:
 	void SetType(EPCGUnionType InUnionType) { UnionType = InUnionType; }
 	void SetDensityFunction(EPCGUnionDensityFunction InDensityFunction) { DensityFunction = InDensityFunction; }
 
+	// ~Begin UPCGData interface
+	virtual EPCGDataType GetDataType() const override { return EPCGDataType::Composite; }
+	// ~End UPCGData interface
+
 	//~Begin UPCGSpatialData interface
 	virtual int GetDimension() const override;
 	virtual FBox GetBounds() const override;
