Fix for FSoftClassPath + FSoftObjectPath metadata type shenanigans

diff --git a/Engine/Plugins/Experimental/PCG/Source/PCG/Private/Metadata/PCGMetadataTypesConstantStruct.cpp b/Engine/Plugins/Experimental/PCG/Source/PCG/Private/Metadata/PCGMetadataTypesConstantStruct.cpp
--- a/Engine/Plugins/Experimental/PCG/Source/PCG/Private/Metadata/PCGMetadataTypesConstantStruct.cpp
+++ b/Engine/Plugins/Experimental/PCG/Source/PCG/Private/Metadata/PCGMetadataTypesConstantStruct.cpp
@@ -22,6 +22,8 @@ FString FPCGMetadataTypesConstantStruct::ToString() const
 			return FString::Printf(TEXT("\"%s\""), *StringValue);
 		case EPCGMetadataTypesConstantStructStringMode::SoftObjectPath:
 			return FString::Printf(TEXT("\"%s\""), *(SoftObjectPathValue.GetAssetName()));
+		case EPCGMetadataTypesConstantStructStringMode::SoftClassPath:
+			return FString::Printf(TEXT("\"%s\""), *(SoftClassPathValue.GetAssetName()));
 		default:
 			break;
 		}
diff --git a/Engine/Plugins/Experimental/PCG/Source/PCG/Public/Metadata/PCGMetadataTypesConstantStruct.h b/Engine/Plugins/Experimental/PCG/Source/PCG/Public/Metadata/PCGMetadataTypesConstantStruct.h
--- a/Engine/Plugins/Experimental/PCG/Source/PCG/Public/Metadata/PCGMetadataTypesConstantStruct.h
+++ b/Engine/Plugins/Experimental/PCG/Source/PCG/Public/Metadata/PCGMetadataTypesConstantStruct.h
@@ -13,7 +13,8 @@ UENUM()
 enum class EPCGMetadataTypesConstantStructStringMode
 {
 	String,
-	SoftObjectPath
+	SoftObjectPath,
+	SoftClassPath,
 };
 
 /**
@@ -76,8 +77,11 @@ public:
 	UPROPERTY(BlueprintReadWrite, EditAnywhere, Category = Settings, meta = (EditCondition = "Type == EPCGMetadataTypes::Name", EditConditionHides))
 	FName NameValue = NAME_None;
 
+	UPROPERTY(BlueprintReadWrite, EditAnywhere, Category = Settings, meta = (EditCondition = "Type == EPCGMetadataTypes::String && StringMode == EPCGMetadataTypesConstantStructStringMode::SoftClassPath", EditConditionHides))
+	FSoftClassPath SoftClassPathValue;
+
 	UPROPERTY(BlueprintReadWrite, EditAnywhere, Category = Settings, meta = (EditCondition = "Type == EPCGMetadataTypes::String && StringMode == EPCGMetadataTypesConstantStructStringMode::SoftObjectPath", EditConditionHides))
-	FSoftClassPath SoftObjectPathValue;
+	FSoftObjectPath SoftObjectPathValue;
 };
 
 template <typename Func>
@@ -113,6 +117,8 @@ decltype(auto) FPCGMetadataTypesConstantStruct::DispatcherWithOverride(const UPC
 			return Callback(PCG_GET_OVERRIDEN_VALUE(this, StringValue, Params));
 		case EPCGMetadataTypesConstantStructStringMode::SoftObjectPath:
 			return Callback(PCGSettingsHelpers::GetValue(GET_MEMBER_NAME_CHECKED(FPCGMetadataTypesConstantStruct, SoftObjectPathValue), SoftObjectPathValue.ToString(), Params));
+		case EPCGMetadataTypesConstantStructStringMode::SoftClassPath:
+			return Callback(PCGSettingsHelpers::GetValue(GET_MEMBER_NAME_CHECKED(FPCGMetadataTypesConstantStruct, SoftClassPathValue), SoftClassPathValue.ToString(), Params));
 		default:
 			break;
 		}
