Fixed issues in the point filter, also split any first pin params to explict 3rd param pin.

#rnx

diff --git a/Engine/Plugins/Experimental/PCG/Source/PCG/Private/Elements/PCGPointFilter.cpp b/Engine/Plugins/Experimental/PCG/Source/PCG/Private/Elements/PCGPointFilter.cpp
--- a/Engine/Plugins/Experimental/PCG/Source/PCG/Private/Elements/PCGPointFilter.cpp
+++ b/Engine/Plugins/Experimental/PCG/Source/PCG/Private/Elements/PCGPointFilter.cpp
@@ -1,6 +1,8 @@
 // Copyright Epic Games, Inc. All Rights Reserved.
 
 #include "Elements/PCGPointFilter.h"
+
+#include "PCGCustomVersion.h"
 #include "Data/PCGSpatialData.h"
 #include "Data/PCGPointData.h"
 #include "Helpers/PCGSettingsHelpers.h"
@@ -64,13 +66,15 @@ namespace PCGPointFilterHelpers
 TArray<FPCGPinProperties> UPCGPointFilterSettings::InputPinProperties() const
 {
 	TArray<FPCGPinProperties> PinProperties;
-	PinProperties.Emplace(PCGPointFilterConstants::DataToFilterLabel, EPCGDataType::Any);
+	PinProperties.Emplace(PCGPointFilterConstants::DataToFilterLabel, EPCGDataType::Point);
 
 	if (!bUseConstantThreshold)
 	{
 		PinProperties.Emplace(PCGPointFilterConstants::FilterLabel, EPCGDataType::Any, /*bInAllowMultipleConnections=*/ false);
 	}
 
+	PinProperties.Emplace(PCGPinConstants::DefaultParamsLabel, EPCGDataType::Param, /*bAllowMultipleConnections=*/false, /*bAllowMultiData=*/false);
+
 	return PinProperties;
 }
 
@@ -175,6 +179,19 @@ PRAGMA_ENABLE_DEPRECATION_WARNINGS
 #endif // WITH_EDITOR
 }
 
+#if WITH_EDITOR
+void UPCGPointFilterSettings::ApplyDeprecationBeforeUpdatePins(UPCGNode* InOutNode, TArray<TObjectPtr<UPCGPin>>& InputPins, TArray<TObjectPtr<UPCGPin>>& OutputPins)
+{
+	Super::ApplyDeprecationBeforeUpdatePins(InOutNode, InputPins, OutputPins);
+
+	check(InOutNode);
+	if (DataVersion < FPCGCustomVersion::MovePointFilterParamsOffFirstPin)
+	{
+		PCGSettingsHelpers::DeprecationBreakOutParamsToPin(InOutNode, InputPins, OutputPins);
+	}
+}
+#endif // WITH_EDITOR
+
 bool FPCGPointFilterElement::ExecuteInternal(FPCGContext* Context) const
 {
 	TRACE_CPUPROFILER_EVENT_SCOPE(FPCGPointFilterElement::Execute);
@@ -195,7 +212,7 @@ bool FPCGPointFilterElement::ExecuteInternal(FPCGContext* Context) const
 	check(Settings);
 
 	TArray<FPCGTaggedData> DataToFilter = Context->InputData.GetInputsByPin(PCGPointFilterConstants::DataToFilterLabel);
-	UPCGParamData* Params = Context->InputData.GetParams(); //TODO: by pin?
+	UPCGParamData* Params = Context->InputData.GetParamsOnParamsPin();
 	TArray<FPCGTaggedData>& Outputs = Context->OutputData.TaggedData;
 	TArray<FPCGTaggedData> FilterData = Context->InputData.GetInputsByPin(PCGPointFilterConstants::FilterLabel);
 
@@ -203,8 +220,9 @@ bool FPCGPointFilterElement::ExecuteInternal(FPCGContext* Context) const
 	Outputs.Append(Context->InputData.GetAllSettings());
 
 	const EPCGPointFilterOperator Operator = PCGSettingsHelpers::GetValue(GET_MEMBER_NAME_CHECKED(UPCGPointFilterSettings, Operator), Settings->Operator, Params);
-
-	bool bUseSpatialQuery = PCGSettingsHelpers::GetValue(GET_MEMBER_NAME_CHECKED(UPCGPointFilterSettings, bUseSpatialQuery), Settings->bUseSpatialQuery, Params);
+	const bool bUseSpatialQuery = PCGSettingsHelpers::GetValue(GET_MEMBER_NAME_CHECKED(UPCGPointFilterSettings, bUseSpatialQuery), Settings->bUseSpatialQuery, Params);
+	const bool bUseConstantThreshold = PCGSettingsHelpers::GetValue(GET_MEMBER_NAME_CHECKED(UPCGPointFilterSettings, bUseConstantThreshold), Settings->bUseConstantThreshold, Params);
+	// TODO: allow selector overrides for ThresholdAttribute and TargetAttribute
 
 	// If there is no input, do nothing
 	if (DataToFilter.IsEmpty())
@@ -219,7 +237,7 @@ bool FPCGPointFilterElement::ExecuteInternal(FPCGContext* Context) const
 	UPCGPointData* ThresholdPointData = nullptr;
 	const UPCGSpatialData* ThresholdSpatialData = nullptr;
 
-	if (Settings->bUseConstantThreshold)
+	if (bUseConstantThreshold)
 	{
 		auto ConstantThreshold = [&ThresholdAccessor, &ThresholdKeys](auto&& Value)
 		{
@@ -243,7 +261,7 @@ bool FPCGPointFilterElement::ExecuteInternal(FPCGContext* Context) const
 				// If the threshold is spatial, and we use spatial query, it means we'll have to sample points.
 				// Don't create an accessor yet (ThresholdData = nullptr), it will be created further down.
 				// Otherwise, we convert it to point data and create the accessor.
-				if (Settings->bUseSpatialQuery)
+				if (bUseSpatialQuery)
 				{
 					ThresholdSpatialData = TentativeThresholdSpatialData;
 					ThresholdData = nullptr;
@@ -276,7 +294,11 @@ bool FPCGPointFilterElement::ExecuteInternal(FPCGContext* Context) const
 	for (const FPCGTaggedData& Input : DataToFilter)
 	{
 		const UPCGSpatialData* SpatialInput = Cast<const UPCGSpatialData>(Input.Data);
-		check(SpatialInput);
+		if (!SpatialInput)
+		{
+			PCGE_LOG(Error, "Please move non-spatial data to other pins");
+			continue;
+		}
 
 		const UPCGPointData* OriginalData = SpatialInput->ToPointData(Context);
 		if (!OriginalData)
@@ -350,7 +372,7 @@ bool FPCGPointFilterElement::ExecuteInternal(FPCGContext* Context) const
 		}
 
 		// And also validate that types are comparable
-		if (Settings->Operator != EPCGPointFilterOperator::Equal && Settings->Operator != EPCGPointFilterOperator::NotEqual)
+		if (Operator != EPCGPointFilterOperator::Equal && Operator != EPCGPointFilterOperator::NotEqual)
 		{
 			bool bCanCompare = PCGMetadataAttribute::CallbackWithRightType(TargetAccessor->GetUnderlyingType(), [](auto Dummy) -> bool
 			{
@@ -376,7 +398,7 @@ bool FPCGPointFilterElement::ExecuteInternal(FPCGContext* Context) const
 		InFilterPoints.Reserve(OriginalPoints.Num());
 		OutFilterPoints.Reserve(OriginalPoints.Num());
 
-		auto Operation = [&InFilterPoints, &OutFilterPoints, Settings, &ThresholdAccessor, &ThresholdKeys, &TargetAccessor, &TargetKeys, &OriginalPoints, ThresholdPointData, ThresholdSpatialData](auto Dummy) -> bool
+		auto Operation = [&InFilterPoints, &OutFilterPoints, &Operator, &ThresholdAccessor, &ThresholdKeys, &TargetAccessor, &TargetKeys, &OriginalPoints, ThresholdPointData, ThresholdSpatialData](auto Dummy) -> bool
 		{
 			using Type = decltype(Dummy);
 
@@ -434,7 +456,7 @@ bool FPCGPointFilterElement::ExecuteInternal(FPCGContext* Context) const
 
 				for (int32 j = 0; j < Range; ++j)
 				{
-					if (SkipTests[j] || PCGPointFilterHelpers::ApplyCompare(TargetValues[j], ThresholdValues[j], Settings->Operator))
+					if (SkipTests[j] || PCGPointFilterHelpers::ApplyCompare(TargetValues[j], ThresholdValues[j], Operator))
 					{
 						InFilterPoints.Add(OriginalPoints[StartIndex + j]);
 					}
diff --git a/Engine/Plugins/Experimental/PCG/Source/PCG/Private/Helpers/PCGSettingsHelpers.cpp b/Engine/Plugins/Experimental/PCG/Source/PCG/Private/Helpers/PCGSettingsHelpers.cpp
--- a/Engine/Plugins/Experimental/PCG/Source/PCG/Private/Helpers/PCGSettingsHelpers.cpp
+++ b/Engine/Plugins/Experimental/PCG/Source/PCG/Private/Helpers/PCGSettingsHelpers.cpp
@@ -217,8 +217,10 @@ namespace PCGSettingsHelpers
 	void DeprecationBreakOutParamsToPin(UPCGNode* InOutNode, TArray<TObjectPtr<UPCGPin>>& InputPins, TArray<TObjectPtr<UPCGPin>>& OutputPins)
 	{
 		// Check basic conditions for which the code below should run.
-		check(InputPins.Num() == 1);
-		check(InputPins[0] && InputPins[0]->Properties.AllowedTypes == EPCGDataType::Any);
+		if(InputPins.IsEmpty() || !InputPins[0] || InputPins[0]->Properties.AllowedTypes != EPCGDataType::Any)
+		{
+			return;
+		}
 
 		UPCGPin* InPin = InputPins[0];
 
diff --git a/Engine/Plugins/Experimental/PCG/Source/PCG/Private/PCGData.cpp b/Engine/Plugins/Experimental/PCG/Source/PCG/Private/PCGData.cpp
--- a/Engine/Plugins/Experimental/PCG/Source/PCG/Private/PCGData.cpp
+++ b/Engine/Plugins/Experimental/PCG/Source/PCG/Private/PCGData.cpp
@@ -157,6 +157,12 @@ UPCGParamData* FPCGDataCollection::GetParams() const
 	return nullptr;
 }
 
+UPCGParamData* FPCGDataCollection::GetParamsOnParamsPin() const
+{
+	TArray<FPCGTaggedData> ParamsOnDefaultPin = GetParamsByPin(PCGPinConstants::DefaultParamsLabel);
+	return (ParamsOnDefaultPin.IsEmpty() ? nullptr : Cast<UPCGParamData>(ParamsOnDefaultPin[0].Data));
+}
+
 const UPCGSettings* FPCGDataCollection::GetSettings(const UPCGSettings* InDefaultSettings) const
 {
 	if (!InDefaultSettings)
diff --git a/Engine/Plugins/Experimental/PCG/Source/PCG/Public/Elements/PCGPointFilter.h b/Engine/Plugins/Experimental/PCG/Source/PCG/Public/Elements/PCGPointFilter.h
--- a/Engine/Plugins/Experimental/PCG/Source/PCG/Public/Elements/PCGPointFilter.h
+++ b/Engine/Plugins/Experimental/PCG/Source/PCG/Public/Elements/PCGPointFilter.h
@@ -76,6 +76,7 @@ public:
 #if WITH_EDITOR
 	virtual FName GetDefaultNodeName() const override { return FName(TEXT("PointFilter")); }
 	virtual EPCGSettingsType GetType() const override { return EPCGSettingsType::Filter; }
+	virtual void ApplyDeprecationBeforeUpdatePins(UPCGNode* InOutNode, TArray<TObjectPtr<UPCGPin>>& InputPins, TArray<TObjectPtr<UPCGPin>>& OutputPins) override;
 #endif
 
 	virtual TArray<FPCGPinProperties> InputPinProperties() const override;
diff --git a/Engine/Plugins/Experimental/PCG/Source/PCG/Public/PCGCustomVersion.h b/Engine/Plugins/Experimental/PCG/Source/PCG/Public/PCGCustomVersion.h
--- a/Engine/Plugins/Experimental/PCG/Source/PCG/Public/PCGCustomVersion.h
+++ b/Engine/Plugins/Experimental/PCG/Source/PCG/Public/PCGCustomVersion.h
@@ -23,6 +23,8 @@ struct PCG_API FPCGCustomVersion
 		// Split samplers to give a sampling shape and a bounding shape inputs
 		SplitSamplerNodesInputs = 4,
 
+		MovePointFilterParamsOffFirstPin = 5,
+
 		// -----<new versions can be added above this line>-------------------------------------------------
 		VersionPlusOne,
 		LatestVersion = VersionPlusOne - 1
diff --git a/Engine/Plugins/Experimental/PCG/Source/PCG/Public/PCGData.h b/Engine/Plugins/Experimental/PCG/Source/PCG/Public/PCGData.h
--- a/Engine/Plugins/Experimental/PCG/Source/PCG/Public/PCGData.h
+++ b/Engine/Plugins/Experimental/PCG/Source/PCG/Public/PCGData.h
@@ -64,14 +64,24 @@ struct PCG_API FPCGDataCollection
 {
 	GENERATED_BODY()
 
+	/** Returns all spatial data in the collection */
 	TArray<FPCGTaggedData> GetInputs() const;
+	/** Returns all data on a given pin */
 	TArray<FPCGTaggedData> GetInputsByPin(const FName& InPinLabel) const;
+	/** Returns all spatial data in the collection with the given tag */
 	TArray<FPCGTaggedData> GetTaggedInputs(const FString& InTag) const;
+	/** Returns all settings in the collection */
 	TArray<FPCGTaggedData> GetAllSettings() const;
+	/** Returns all params in the collection */
 	TArray<FPCGTaggedData> GetAllParams() const;
+	/** Returns all params in the collection with a given tag */
 	TArray<FPCGTaggedData> GetTaggedParams(const FString& InTag) const;
+	/** Returns all params on a given pin */
 	TArray<FPCGTaggedData> GetParamsByPin(const FName& InPinLabel) const;
-	UPCGParamData* GetParams() const;
+	/** Returns the first params found in the collection */
+	UPCGParamData* GetParams() const; 
+	/** Returns the first/only param found on the default params pin */
+	UPCGParamData* GetParamsOnParamsPin() const;
 
 	const UPCGSettingsInterface* GetSettingsInterface() const;
 	const UPCGSettingsInterface* GetSettingsInterface(const UPCGSettingsInterface* InDefaultSettingsInterface) const;
