Minor tooltip improvement with single/multi data info.

#rnx

diff --git a/Engine/Plugins/Experimental/PCG/Source/PCGEditor/Private/PCGEditorGraphNodeBase.cpp b/Engine/Plugins/Experimental/PCG/Source/PCGEditor/Private/PCGEditorGraphNodeBase.cpp
--- a/Engine/Plugins/Experimental/PCG/Source/PCGEditor/Private/PCGEditorGraphNodeBase.cpp
+++ b/Engine/Plugins/Experimental/PCG/Source/PCGEditor/Private/PCGEditorGraphNodeBase.cpp
@@ -497,23 +497,28 @@ void UPCGEditorGraphNodeBase::GetPinHoverText(const UEdGraphPin& Pin, FString& H
 		Description = MatchingPin->Properties.Tooltip.IsEmpty() ? FText::FromName(MatchingPin->Properties.Label) : MatchingPin->Properties.Tooltip;
 	}
 
-	if (MatchingPin != nullptr && bIsInputPin)
-	{
-		const FText AdditionalInfo = MatchingPin->Properties.bAllowMultipleConnections ? FText(LOCTEXT("SupportsMultiInput", "Supports multiple inputs")) : FText(LOCTEXT("SingleInputOnly", "Supports only one input"));
+	FText MultiDataSupport;
+	FText MultiConnectionSupport;	
 
-		HoverTextOut = FText::Format(LOCTEXT("PinHoverToolTipFull", "{0}\n\nType: {1}\nSubtype: {2}\nAdditional information: {3}"),
-			Description,
-			DataTypeText,
-			DataSubtypeText,
-			AdditionalInfo).ToString();
-	}
-	else
+	if (MatchingPin)
 	{
-		HoverTextOut = FText::Format(LOCTEXT("PinHoverToolTipPartial", "{0}\n\nType: {1}\nSubtype: {2}"),
-			Description,
-			DataTypeText,
-			DataSubtypeText).ToString();
+		if (bIsInputPin)
+		{
+			MultiDataSupport = MatchingPin->Properties.bAllowMultipleData ? FText(LOCTEXT("SupportsMultiData", "Supports multiple data in input collection(s). ")) : FText(LOCTEXT("SingleDataOnly", "Supports only single data in collection(s). "));
+			MultiConnectionSupport = MatchingPin->Properties.bAllowMultipleConnections ? FText(LOCTEXT("SupportsMultiInput", "Supports multiple inputs.")) : FText(LOCTEXT("SingleInputOnly", "Supports only one input."));
+		}
+		else
+		{
+			MultiDataSupport = MatchingPin->Properties.bAllowMultipleData ? FText(LOCTEXT("SupportsMultiData", "Can generate multiple data.")) : FText(LOCTEXT("SingleDataOnly", "Generates only single data."));
+		}
 	}
+
+	HoverTextOut = FText::Format(LOCTEXT("PinHoverToolTipFull", "{0}\n\nType: {1}\nSubtype: {2}\nAdditional information: {3}{4}"),
+		Description,
+		DataTypeText,
+		DataSubtypeText,
+		MultiDataSupport,
+		MultiConnectionSupport).ToString();
 }
 
 UObject* UPCGEditorGraphNodeBase::GetJumpTargetForDoubleClick() const
