[PCG] Only use WITH_EDITORONLY_DATA around serialised data, WITH_EDITOR everywhere else

diff --git a/Engine/Plugins/Experimental/PCG/Source/PCG/Private/Data/PCGDifferenceData.cpp b/Engine/Plugins/Experimental/PCG/Source/PCG/Private/Data/PCGDifferenceData.cpp
--- a/Engine/Plugins/Experimental/PCG/Source/PCG/Private/Data/PCGDifferenceData.cpp
+++ b/Engine/Plugins/Experimental/PCG/Source/PCG/Private/Data/PCGDifferenceData.cpp
@@ -196,7 +196,7 @@ const UPCGPointData* UPCGDifferenceData::CreatePointData(FPCGContext* Context) c
 				Data->Metadata->MergePointAttributesSubset(Point, SourcePointData->Metadata, SourcePointData->Metadata, PointFromDiff, TempDiffMetadata, TempDiffMetadata, OutPoint, EPCGMetadataOp::Sub);
 			}
 
-#if WITH_EDITORONLY_DATA
+#if WITH_EDITOR
 			return OutPoint.Density > 0 || bKeepZeroDensityPoints;
 #else
 			return OutPoint.Density > 0;
diff --git a/Engine/Plugins/Experimental/PCG/Source/PCG/Private/Data/PCGIntersectionData.cpp b/Engine/Plugins/Experimental/PCG/Source/PCG/Private/Data/PCGIntersectionData.cpp
--- a/Engine/Plugins/Experimental/PCG/Source/PCG/Private/Data/PCGIntersectionData.cpp
+++ b/Engine/Plugins/Experimental/PCG/Source/PCG/Private/Data/PCGIntersectionData.cpp
@@ -207,7 +207,7 @@ UPCGPointData* UPCGIntersectionData::CreateAndFilterPointData(FPCGContext* Conte
 		const FPCGPoint& Point = SourcePoints[Index];
 
 		FPCGPoint PointFromY;
-#if WITH_EDITORONLY_DATA
+#if WITH_EDITOR
 		if (!Y->SamplePoint(Point.Transform, Point.GetLocalBounds(), PointFromY, TempYMetadata) && !bKeepZeroDensityPoints)
 #else
 		if (!Y->SamplePoint(Point.Transform, Point.GetLocalBounds(), PointFromY, TempYMetadata))
diff --git a/Engine/Plugins/Experimental/PCG/Source/PCG/Private/Data/PCGProjectionData.cpp b/Engine/Plugins/Experimental/PCG/Source/PCG/Private/Data/PCGProjectionData.cpp
--- a/Engine/Plugins/Experimental/PCG/Source/PCG/Private/Data/PCGProjectionData.cpp
+++ b/Engine/Plugins/Experimental/PCG/Source/PCG/Private/Data/PCGProjectionData.cpp
@@ -210,7 +210,7 @@ const UPCGPointData* UPCGProjectionData::CreatePointData(FPCGContext* Context) c
 		const FPCGPoint& SourcePoint = SourcePoints[Index];
 
 		FPCGPoint PointFromTarget;
-#if WITH_EDITORONLY_DATA
+#if WITH_EDITOR
 		if (!Target->ProjectPoint(SourcePoint.Transform, SourcePoint.GetLocalBounds(), ProjectionParams, PointFromTarget, TempTargetMetadata) && !bKeepZeroDensityPoints)
 #else
 		if (!Target->ProjectPoint(SourcePoint.Transform, SourcePoint.GetLocalBounds(), ProjectionParams, PointFromTarget, TempTargetMetadata))
diff --git a/Engine/Plugins/Experimental/PCG/Source/PCG/Private/Data/PCGTextureData.cpp b/Engine/Plugins/Experimental/PCG/Source/PCG/Private/Data/PCGTextureData.cpp
--- a/Engine/Plugins/Experimental/PCG/Source/PCG/Private/Data/PCGTextureData.cpp
+++ b/Engine/Plugins/Experimental/PCG/Source/PCG/Private/Data/PCGTextureData.cpp
@@ -217,7 +217,7 @@ const UPCGPointData* UPCGBaseTextureData::CreatePointData(FPCGContext* Context)
 		if (PCGTextureSampling::Sample<FLinearColor>(LocalCoordinate, Surface, this, Width, Height, Color, [this](int32 Index) { return ColorData[Index]; }))
 		{
 			const float Density = ((DensityFunction == EPCGTextureDensityFunction::Ignore) ? 1.0f : PCGTextureSampling::SampleFloatChannel(Color, ColorChannel));
-#if WITH_EDITORONLY_DATA
+#if WITH_EDITOR
 			if (Density > 0 || bKeepZeroDensityPoints)
 #else
 			if (Density > 0)
diff --git a/Engine/Plugins/Experimental/PCG/Source/PCG/Private/Elements/PCGDebugElement.cpp b/Engine/Plugins/Experimental/PCG/Source/PCG/Private/Elements/PCGDebugElement.cpp
--- a/Engine/Plugins/Experimental/PCG/Source/PCG/Private/Elements/PCGDebugElement.cpp
+++ b/Engine/Plugins/Experimental/PCG/Source/PCG/Private/Elements/PCGDebugElement.cpp
@@ -18,7 +18,7 @@ namespace PCGDebugElement
 	void ExecuteDebugDisplay(FPCGContext* Context)
 	{
 		TRACE_CPUPROFILER_EVENT_SCOPE(PCGDebugElement::ExecuteDebugDisplay);
-#if WITH_EDITORONLY_DATA
+#if WITH_EDITOR
 		// Early validation: if we don't have a valid PCG component, we're not going to add the debug display info
 		if (!Context->SourceComponent.IsValid())
 		{
diff --git a/Engine/Plugins/Experimental/PCG/Source/PCG/Private/Elements/PCGDensityFilter.cpp b/Engine/Plugins/Experimental/PCG/Source/PCG/Private/Elements/PCGDensityFilter.cpp
--- a/Engine/Plugins/Experimental/PCG/Source/PCG/Private/Elements/PCGDensityFilter.cpp
+++ b/Engine/Plugins/Experimental/PCG/Source/PCG/Private/Elements/PCGDensityFilter.cpp
@@ -42,7 +42,7 @@ bool FPCGDensityFilterElement::ExecuteInternal(FPCGContext* Context) const
 	const bool bInvertFilter = PCGSettingsHelpers::GetValue(GET_MEMBER_NAME_CHECKED(UPCGDensityFilterSettings, bInvertFilter), Settings->bInvertFilter, Params);
 	const float LowerBound = PCGSettingsHelpers::GetValue(GET_MEMBER_NAME_CHECKED(UPCGDensityFilterSettings, LowerBound), Settings->LowerBound, Params);
 	const float UpperBound = PCGSettingsHelpers::GetValue(GET_MEMBER_NAME_CHECKED(UPCGDensityFilterSettings, UpperBound), Settings->UpperBound, Params);
-#if WITH_EDITORONLY_DATA
+#if WITH_EDITOR
 	const bool bKeepZeroDensityPoints = PCGSettingsHelpers::GetValue(GET_MEMBER_NAME_CHECKED(UPCGDensityFilterSettings, bKeepZeroDensityPoints), Settings->bKeepZeroDensityPoints, Params);
 #else
 	const bool bKeepZeroDensityPoints = false;
@@ -104,7 +104,7 @@ bool FPCGDensityFilterElement::ExecuteInternal(FPCGContext* Context) const
 				OutPoint = Point;
 				return true;
 			}
-#if WITH_EDITORONLY_DATA
+#if WITH_EDITOR
 			else if (bKeepZeroDensityPoints)
 			{
 				OutPoint = Point;
diff --git a/Engine/Plugins/Experimental/PCG/Source/PCG/Private/Elements/PCGDifferenceElement.cpp b/Engine/Plugins/Experimental/PCG/Source/PCG/Private/Elements/PCGDifferenceElement.cpp
--- a/Engine/Plugins/Experimental/PCG/Source/PCG/Private/Elements/PCGDifferenceElement.cpp
+++ b/Engine/Plugins/Experimental/PCG/Source/PCG/Private/Elements/PCGDifferenceElement.cpp
@@ -112,7 +112,7 @@ bool FPCGDifferenceElement::ExecuteInternal(FPCGContext* Context) const
 					DifferenceData = (UnionData ? UnionData : FirstSpatialData)->Subtract(SpatialData);
 					DifferenceData->SetDensityFunction(Settings->DensityFunction);
 					DifferenceData->bDiffMetadata = Settings->bDiffMetadata;
-#if WITH_EDITORONLY_DATA
+#if WITH_EDITOR
 					DifferenceData->bKeepZeroDensityPoints = Settings->bKeepZeroDensityPoints;
 #endif
 
@@ -150,7 +150,7 @@ void FPCGDifferenceElement::LabellessProcessing(FPCGContext* Context) const
 	UPCGParamData* Params = Context->InputData.GetParams();
 
 	const EPCGDifferenceDensityFunction DensityFunction = PCGSettingsHelpers::GetValue(GET_MEMBER_NAME_CHECKED(UPCGDifferenceSettings, DensityFunction), Settings->DensityFunction, Params);
-#if WITH_EDITORONLY_DATA
+#if WITH_EDITOR
 	const bool bKeepZeroDensityPoints = PCGSettingsHelpers::GetValue(GET_MEMBER_NAME_CHECKED(UPCGDifferenceSettings, bKeepZeroDensityPoints), Settings->bKeepZeroDensityPoints, Params);
 #else
 	const bool bKeepZeroDensityPoints = false;
@@ -167,7 +167,7 @@ void FPCGDifferenceElement::LabellessProcessing(FPCGContext* Context) const
 		{
 			DifferenceData = FirstSpatialData->Subtract(SpatialData);
 			DifferenceData->SetDensityFunction(DensityFunction);
-#if WITH_EDITORONLY_DATA
+#if WITH_EDITOR
 			DifferenceData->bKeepZeroDensityPoints = bKeepZeroDensityPoints;
 #endif
 
diff --git a/Engine/Plugins/Experimental/PCG/Source/PCG/Private/Elements/PCGExecuteBlueprint.cpp b/Engine/Plugins/Experimental/PCG/Source/PCG/Private/Elements/PCGExecuteBlueprint.cpp
--- a/Engine/Plugins/Experimental/PCG/Source/PCG/Private/Elements/PCGExecuteBlueprint.cpp
+++ b/Engine/Plugins/Experimental/PCG/Source/PCG/Private/Elements/PCGExecuteBlueprint.cpp
@@ -43,7 +43,7 @@ void UPCGBlueprintElement::PostLoad()
 	Super::PostLoad();
 	Initialize();
 
-#if WITH_EDITORONLY_DATA
+#if WITH_EDITOR
 	if (!InputPinLabels_DEPRECATED.IsEmpty())
 	{
 		for (const FName& Label : InputPinLabels_DEPRECATED)
@@ -246,7 +246,7 @@ void UPCGBlueprintSettings::PostLoad()
 {
 	Super::PostLoad();
 
-#if WITH_EDITORONLY_DATA
+#if WITH_EDITOR
 	if (BlueprintElement_DEPRECATED && !BlueprintElementType)
 	{
 		BlueprintElementType = BlueprintElement_DEPRECATED;
@@ -268,13 +268,13 @@ void UPCGBlueprintSettings::PostLoad()
 	if (BlueprintElementInstance)
 	{
 		BlueprintElementInstance->ConditionalPostLoad();
-#if WITH_EDITORONLY_DATA
+#if WITH_EDITOR
 		BlueprintElementInstance->bCreatesArtifacts |= bCreatesArtifacts_DEPRECATED;
 		BlueprintElementInstance->bCanBeMultithreaded |= bCanBeMultithreaded_DEPRECATED;
 #endif
 	}
 
-#if WITH_EDITORONLY_DATA
+#if WITH_EDITOR
 	bCreatesArtifacts_DEPRECATED = false;
 	bCanBeMultithreaded_DEPRECATED = false;
 #endif
@@ -400,12 +400,12 @@ EPCGSettingsType UPCGBlueprintSettings::GetType() const
 
 void UPCGBlueprintSettings::GetTrackedActorTags(FPCGTagToSettingsMap& OutTagToSettings, TArray<TObjectPtr<const UPCGGraph>>& OutVisitedGraphs) const
 {
-#if WITH_EDITORONLY_DATA
+#if WITH_EDITOR
 	for (const FName& Tag : TrackedActorTags)
 	{
 		OutTagToSettings.FindOrAdd(Tag).Add(this);
 	}
-#endif // WITH_EDITORONLY_DATA
+#endif // WITH_EDITOR
 }
 
 UObject* UPCGBlueprintSettings::GetJumpTargetForDoubleClick() const
diff --git a/Engine/Plugins/Experimental/PCG/Source/PCG/Private/Elements/PCGIntersectionElement.cpp b/Engine/Plugins/Experimental/PCG/Source/PCG/Private/Elements/PCGIntersectionElement.cpp
--- a/Engine/Plugins/Experimental/PCG/Source/PCG/Private/Elements/PCGIntersectionElement.cpp
+++ b/Engine/Plugins/Experimental/PCG/Source/PCG/Private/Elements/PCGIntersectionElement.cpp
@@ -29,7 +29,7 @@ bool FPCGIntersectionElement::ExecuteInternal(FPCGContext* Context) const
 	UPCGParamData* Params = Context->InputData.GetParams();
 
 	const EPCGIntersectionDensityFunction DensityFunction = PCGSettingsHelpers::GetValue(GET_MEMBER_NAME_CHECKED(UPCGIntersectionSettings, DensityFunction), Settings->DensityFunction, Params);
-#if WITH_EDITORONLY_DATA
+#if WITH_EDITOR
 	const bool bKeepZeroDensityPoints = PCGSettingsHelpers::GetValue(GET_MEMBER_NAME_CHECKED(UPCGIntersectionSettings, bKeepZeroDensityPoints), Settings->bKeepZeroDensityPoints, Params);
 #else
 	const bool bKeepZeroDensityPoints = false;
@@ -64,7 +64,7 @@ bool FPCGIntersectionElement::ExecuteInternal(FPCGContext* Context) const
 		IntersectionData = (IntersectionData ? IntersectionData : FirstSpatialData)->IntersectWith(SpatialData);
 		// Propagate settings
 		IntersectionData->DensityFunction = DensityFunction;
-#if WITH_EDITORONLY_DATA
+#if WITH_EDITOR
 		IntersectionData->bKeepZeroDensityPoints = bKeepZeroDensityPoints;
 #endif
 
diff --git a/Engine/Plugins/Experimental/PCG/Source/PCG/Private/Elements/PCGPointSampler.cpp b/Engine/Plugins/Experimental/PCG/Source/PCG/Private/Elements/PCGPointSampler.cpp
--- a/Engine/Plugins/Experimental/PCG/Source/PCG/Private/Elements/PCGPointSampler.cpp
+++ b/Engine/Plugins/Experimental/PCG/Source/PCG/Private/Elements/PCGPointSampler.cpp
@@ -35,7 +35,7 @@ bool FPCGPointSamplerElement::ExecuteInternal(FPCGContext* Context) const
 	Outputs.Append(Context->InputData.GetAllSettings());
 
 	const float Ratio = PCGSettingsHelpers::GetValue(GET_MEMBER_NAME_CHECKED(UPCGPointSamplerSettings, Ratio), Settings->Ratio, Params);
-#if WITH_EDITORONLY_DATA
+#if WITH_EDITOR
 	const bool bKeepZeroDensityPoints = PCGSettingsHelpers::GetValue(GET_MEMBER_NAME_CHECKED(UPCGPointSamplerSettings, bKeepZeroDensityPoints), Settings->bKeepZeroDensityPoints, Params);
 #else
 	const bool bKeepZeroDensityPoints = false;
@@ -89,7 +89,7 @@ bool FPCGPointSamplerElement::ExecuteInternal(FPCGContext* Context) const
 		Output.Data = SampledData;
 
 		// TODO: randomize on the fractional number of points
-#if WITH_EDITORONLY_DATA
+#if WITH_EDITOR
 		int TargetNumPoints = (bKeepZeroDensityPoints ? OriginalPointCount : OriginalPointCount * Ratio);
 #else
 		int TargetNumPoints = OriginalPointCount * Ratio;
@@ -118,7 +118,7 @@ bool FPCGPointSamplerElement::ExecuteInternal(FPCGContext* Context) const
 					OutPoint = Point;
 					return true;
 				}
-#if WITH_EDITORONLY_DATA
+#if WITH_EDITOR
 				else if (bKeepZeroDensityPoints)
 				{
 					OutPoint = Point;
diff --git a/Engine/Plugins/Experimental/PCG/Source/PCG/Private/Elements/PCGProjectionElement.cpp b/Engine/Plugins/Experimental/PCG/Source/PCG/Private/Elements/PCGProjectionElement.cpp
--- a/Engine/Plugins/Experimental/PCG/Source/PCG/Private/Elements/PCGProjectionElement.cpp
+++ b/Engine/Plugins/Experimental/PCG/Source/PCG/Private/Elements/PCGProjectionElement.cpp
@@ -66,7 +66,7 @@ bool FPCGProjectionElement::ExecuteInternal(FPCGContext* Context) const
 
 	TArray<FPCGTaggedData>& Outputs = Context->OutputData.TaggedData;
 
-#if WITH_EDITORONLY_DATA
+#if WITH_EDITOR
 	const bool bKeepZeroDensityPoints = PCGSettingsHelpers::GetValue(GET_MEMBER_NAME_CHECKED(UPCGProjectionSettings, bKeepZeroDensityPoints), Settings->bKeepZeroDensityPoints, Params);
 #else
 	const bool bKeepZeroDensityPoints = false;
@@ -83,7 +83,7 @@ bool FPCGProjectionElement::ExecuteInternal(FPCGContext* Context) const
 		}
 
 		UPCGProjectionData* ProjectionData = ProjectionSource->ProjectOn(ProjectionTarget, ProjectionParams);
-#if WITH_EDITORONLY_DATA
+#if WITH_EDITOR
 		ProjectionData->bKeepZeroDensityPoints = bKeepZeroDensityPoints;
 #endif
 
diff --git a/Engine/Plugins/Experimental/PCG/Source/PCG/Private/Elements/PCGSurfaceSampler.cpp b/Engine/Plugins/Experimental/PCG/Source/PCG/Private/Elements/PCGSurfaceSampler.cpp
--- a/Engine/Plugins/Experimental/PCG/Source/PCG/Private/Elements/PCGSurfaceSampler.cpp
+++ b/Engine/Plugins/Experimental/PCG/Source/PCG/Private/Elements/PCGSurfaceSampler.cpp
@@ -39,7 +39,7 @@ namespace PCGSurfaceSampler
 			Looseness = PCGSettingsHelpers::GetValue(GET_MEMBER_NAME_CHECKED(UPCGSurfaceSamplerSettings, Looseness), Settings->Looseness, Params);
 			bApplyDensityToPoints = PCGSettingsHelpers::GetValue(GET_MEMBER_NAME_CHECKED(UPCGSurfaceSamplerSettings, bApplyDensityToPoints), Settings->bApplyDensityToPoints, Params);
 			PointSteepness = PCGSettingsHelpers::GetValue(GET_MEMBER_NAME_CHECKED(UPCGSurfaceSamplerSettings, PointSteepness), Settings->PointSteepness, Params);
-#if WITH_EDITORONLY_DATA
+#if WITH_EDITOR
 			bKeepZeroDensityPoints = PCGSettingsHelpers::GetValue(GET_MEMBER_NAME_CHECKED(UPCGSurfaceSamplerSettings, bKeepZeroDensityPoints), Settings->bKeepZeroDensityPoints, Params);
 #endif
 
@@ -159,7 +159,7 @@ namespace PCGSurfaceSampler
 			if (InBoundingShape)
 			{
 				FPCGPoint BoundingShapeSample;
-#if WITH_EDITORONLY_DATA
+#if WITH_EDITOR
 				if (!InBoundingShape->SamplePoint(OutPoint.Transform, OutPoint.GetLocalBounds(), BoundingShapeSample, nullptr) && !LoopData.bKeepZeroDensityPoints)
 #else
 				if (!InBoundingShape->SamplePoint(OutPoint.Transform, OutPoint.GetLocalBounds(), BoundingShapeSample, nullptr))
diff --git a/Engine/Plugins/Experimental/PCG/Source/PCG/Private/PCGComponent.cpp b/Engine/Plugins/Experimental/PCG/Source/PCG/Private/PCGComponent.cpp
--- a/Engine/Plugins/Experimental/PCG/Source/PCG/Private/PCGComponent.cpp
+++ b/Engine/Plugins/Experimental/PCG/Source/PCG/Private/PCGComponent.cpp
@@ -802,7 +802,7 @@ void UPCGComponent::PostLoad()
 {
 	Super::PostLoad();
 
-#if WITH_EDITORONLY_DATA
+#if WITH_EDITOR
 	if (!ExclusionTags_DEPRECATED.IsEmpty() && ExcludedTags.IsEmpty())
 	{
 		ExcludedTags.Append(ExclusionTags_DEPRECATED);
diff --git a/Engine/Plugins/Experimental/PCG/Source/PCG/Private/PCGGraph.cpp b/Engine/Plugins/Experimental/PCG/Source/PCG/Private/PCGGraph.cpp
--- a/Engine/Plugins/Experimental/PCG/Source/PCG/Private/PCGGraph.cpp
+++ b/Engine/Plugins/Experimental/PCG/Source/PCG/Private/PCGGraph.cpp
@@ -41,7 +41,7 @@ UPCGGraph::UPCGGraph(const FObjectInitializer& ObjectInitializer)
 	OutputNode->SetSettingsInterface(OutputSettings, /*bUpdatePins=*/false);
 	OutputNode->UpdatePins(PinAllocator);
 	
-#if WITH_EDITORONLY_DATA
+#if WITH_EDITOR
 	OutputNode->PositionX = 200;
 #endif
 
@@ -122,7 +122,7 @@ void UPCGGraph::PostLoad()
 #endif
 }
 
-#if WITH_EDITORONLY_DATA
+#if WITH_EDITOR
 void UPCGGraph::DeclareConstructClasses(TArray<FTopLevelAssetPath>& OutConstructClasses, const UClass* SpecificSubclass)
 {
 	Super::DeclareConstructClasses(OutConstructClasses, SpecificSubclass);
diff --git a/Engine/Plugins/Experimental/PCG/Source/PCG/Private/PCGSettings.cpp b/Engine/Plugins/Experimental/PCG/Source/PCG/Private/PCGSettings.cpp
--- a/Engine/Plugins/Experimental/PCG/Source/PCG/Private/PCGSettings.cpp
+++ b/Engine/Plugins/Experimental/PCG/Source/PCG/Private/PCGSettings.cpp
@@ -19,11 +19,11 @@ class FPCGSettingsObjectCrc32 : public FArchiveObjectCrc32
 public:
 	virtual bool ShouldSkipProperty(const FProperty* InProperty) const override
 	{
-#if WITH_EDITORONLY_DATA
+#if WITH_EDITOR
 		return InProperty && (InProperty->GetFName() == GET_MEMBER_NAME_CHECKED(UPCGSettings, DebugSettings) || InProperty->GetFName() == GET_MEMBER_NAME_CHECKED(UPCGSettings, DeterminismSettings));
 #else
 		return FArchiveObjectCrc32::ShouldSkipProperty(InProperty);
-#endif // WITH_EDITORONLY_DATA
+#endif // WITH_EDITOR
 	}
 };
 
@@ -79,7 +79,7 @@ void UPCGSettings::PostLoad()
 {
 	Super::PostLoad();
 
-#if WITH_EDITORONLY_DATA
+#if WITH_EDITOR
 	if (ExecutionMode_DEPRECATED != EPCGSettingsExecutionMode::Enabled)
 	{
 		bEnabled = ExecutionMode_DEPRECATED != EPCGSettingsExecutionMode::Disabled;
@@ -225,7 +225,7 @@ void UPCGSettingsInstance::PostLoad()
 	}
 #endif
 
-#if WITH_EDITORONLY_DATA
+#if WITH_EDITOR
 	OriginalSettings = Settings;
 #endif
 }
@@ -252,7 +252,7 @@ void UPCGSettingsInstance::SetSettings(UPCGSettings* InSettings)
 #endif
 
 	Settings = InSettings;
-#if WITH_EDITORONLY_DATA
+#if WITH_EDITOR
 	OriginalSettings = Settings;
 #endif
 
diff --git a/Engine/Plugins/Experimental/PCG/Source/PCG/Private/PCGWorldActor.cpp b/Engine/Plugins/Experimental/PCG/Source/PCG/Private/PCGWorldActor.cpp
--- a/Engine/Plugins/Experimental/PCG/Source/PCG/Private/PCGWorldActor.cpp
+++ b/Engine/Plugins/Experimental/PCG/Source/PCG/Private/PCGWorldActor.cpp
@@ -14,7 +14,7 @@
 APCGWorldActor::APCGWorldActor(const FObjectInitializer& ObjectInitializer)
 	: Super(ObjectInitializer)
 {
-#if WITH_EDITORONLY_DATA
+#if WITH_EDITOR
 	bIsSpatiallyLoaded = false;
 	bDefaultOutlinerExpansionState = false;
 #endif
diff --git a/Engine/Plugins/Experimental/PCG/Source/PCG/Public/PCGGraph.h b/Engine/Plugins/Experimental/PCG/Source/PCG/Public/PCGGraph.h
--- a/Engine/Plugins/Experimental/PCG/Source/PCG/Public/PCGGraph.h
+++ b/Engine/Plugins/Experimental/PCG/Source/PCG/Public/PCGGraph.h
@@ -26,7 +26,7 @@ public:
 	UPCGGraph(const FObjectInitializer& ObjectInitializer);
 	/** ~Begin UObject interface */
 	virtual void PostLoad() override;
-#if WITH_EDITORONLY_DATA
+#if WITH_EDITOR
 	static void DeclareConstructClasses(TArray<FTopLevelAssetPath>& OutConstructClasses, const UClass* SpecificSubclass);
 #endif
 
