[PCG] Prevent pin tooltips from serialising into game data

diff --git a/Engine/Plugins/Experimental/PCG/Source/PCG/Private/PCGPin.cpp b/Engine/Plugins/Experimental/PCG/Source/PCG/Private/PCGPin.cpp
--- a/Engine/Plugins/Experimental/PCG/Source/PCG/Private/PCGPin.cpp
+++ b/Engine/Plugins/Experimental/PCG/Source/PCG/Private/PCGPin.cpp
@@ -7,7 +7,10 @@
 #include UE_INLINE_GENERATED_CPP_BY_NAME(PCGPin)
 
 FPCGPinProperties::FPCGPinProperties(const FName& InLabel, EPCGDataType InAllowedTypes, bool bInAllowMultipleConnections, bool bInAllowMultipleData, const FText& InTooltip)
-	: Label(InLabel), AllowedTypes(InAllowedTypes), bAllowMultipleData(bInAllowMultipleData), bAllowMultipleConnections(bInAllowMultipleConnections), Tooltip(InTooltip)
+	: Label(InLabel), AllowedTypes(InAllowedTypes), bAllowMultipleData(bInAllowMultipleData), bAllowMultipleConnections(bInAllowMultipleConnections)
+#if WITH_EDITORONLY_DATA
+	, Tooltip(InTooltip)
+#endif
 {}
 
 bool FPCGPinProperties::operator==(const FPCGPinProperties& Other) const
@@ -35,6 +38,22 @@ void UPCGPin::PostLoad()
 	}
 }
 
+FText UPCGPin::GetTooltip() const
+{
+#if WITH_EDITOR
+	return Properties.Tooltip;
+#else
+	return FText::GetEmpty();
+#endif
+}
+
+void UPCGPin::SetTooltip(const FText& InTooltip)
+{
+#if WITH_EDITOR
+	Properties.Tooltip = InTooltip;
+#endif
+}
+
 bool UPCGPin::AddEdgeTo(UPCGPin* OtherPin)
 {
 	if (!OtherPin)
diff --git a/Engine/Plugins/Experimental/PCG/Source/PCG/Public/PCGPin.h b/Engine/Plugins/Experimental/PCG/Source/PCG/Public/PCGPin.h
--- a/Engine/Plugins/Experimental/PCG/Source/PCG/Public/PCGPin.h
+++ b/Engine/Plugins/Experimental/PCG/Source/PCG/Public/PCGPin.h
@@ -32,8 +32,10 @@ struct PCG_API FPCGPinProperties
 	UPROPERTY(BlueprintReadWrite, EditAnywhere, Category = Settings)
 	bool bAllowMultipleConnections = true;
 
-	UPROPERTY(BlueprintReadWrite, EditAnywhere, Category = Settings)
+#if WITH_EDITORONLY_DATA
+	UPROPERTY(EditAnywhere, Category = Settings)
 	FText Tooltip;
+#endif
 };
 
 UCLASS(ClassGroup = (Procedural))
@@ -60,6 +62,12 @@ public:
 	UPROPERTY(BlueprintReadOnly, EditAnywhere, Category = Settings, meta = (ShowOnlyInnerProperties))
 	FPCGPinProperties Properties;
 
+	UFUNCTION(BlueprintCallable, Category = Settings)
+	FText GetTooltip() const;
+	
+	UFUNCTION(BlueprintCallable, Category = Settings)
+	void SetTooltip(const FText& InTooltip);
+
 	bool IsCompatible(const UPCGPin* OtherPin) const;
 	bool CanConnect(const UPCGPin* OtherPin) const;
 	bool AddEdgeTo(UPCGPin* OtherPin);
