[PCG] Data migration improvements for Surface Sampler

If connected to Input node, then wire it directly to Landscape pin.

Also works for ancient data (before modern pin properties).

diff --git a/Engine/Plugins/Experimental/PCG/Source/PCG/Private/Elements/PCGSurfaceSampler.cpp b/Engine/Plugins/Experimental/PCG/Source/PCG/Private/Elements/PCGSurfaceSampler.cpp
--- a/Engine/Plugins/Experimental/PCG/Source/PCG/Private/Elements/PCGSurfaceSampler.cpp
+++ b/Engine/Plugins/Experimental/PCG/Source/PCG/Private/Elements/PCGSurfaceSampler.cpp
@@ -4,7 +4,10 @@
 
 #include "PCGComponent.h"
 #include "PCGCustomVersion.h"
+#include "PCGEdge.h"
+#include "PCGGraph.h"
 #include "PCGHelpers.h"
+#include "PCGInputOutputSettings.h"
 #include "Data/PCGPointData.h"
 #include "Data/PCGSpatialData.h"
 #include "Helpers/PCGAsync.h"
@@ -186,6 +189,21 @@ namespace PCGSurfaceSampler
 			PCGE_LOG_C(Verbose, Context, "Generated %d points in %d cells", SampledPoints.Num(), LoopData.CellCount);
 		}
 	}
+
+#if WITH_EDITOR
+	static bool IsPinOnlyConnectedToInputNode(UPCGPin* DownstreamPin, UPCGNode* GraphInputNode)
+	{
+		if (DownstreamPin->Edges.Num() == 1)
+		{
+			const UPCGEdge* Edge = DownstreamPin->Edges[0];
+			const UPCGNode* UpstreamNode = (Edge && Edge->InputPin) ? Edge->InputPin->Node : nullptr;
+			const bool bConnectedToInputNode = UpstreamNode && (GraphInputNode == UpstreamNode);
+			return bConnectedToInputNode && (Edge->InputPin->Properties.Label == FName(TEXT("In")) || Edge->InputPin->Properties.Label == FName(TEXT("Input")));
+		}
+
+		return false;
+	}
+#endif
 }
 
 UPCGSurfaceSamplerSettings::UPCGSurfaceSamplerSettings()
@@ -361,19 +379,72 @@ bool FPCGSurfaceSamplerElement::ExecuteInternal(FPCGContext* Context) const
 #if WITH_EDITOR
 void UPCGSurfaceSamplerSettings::ApplyDeprecationBeforeUpdatePins(UPCGNode* InOutNode, TArray<TObjectPtr<UPCGPin>>& InputPins, TArray<TObjectPtr<UPCGPin>>& OutputPins)
 {
-	if (DataVersion < FPCGCustomVersion::SplitSamplerNodesInputs)
+	if (DataVersion < FPCGCustomVersion::SplitSamplerNodesInputs && ensure(InOutNode))
 	{
-		// Deprecation from a single pin node.
-		check(InputPins.Num() == 1);
+		if (InputPins.Num() == 1)
+		{
+			// The node will function the same if we move all connections from "In" to "Bounding Shape". To make this happen, rename "In" to
+			// "Bounding Shape" just prior to pin update and the edges will be moved over. In ApplyDeprecation we'll see if we can do better than
+			// this baseline functional setup.
+			InputPins[0]->Properties.Label = PCGSurfaceSamplerConstants::BoundingShapeLabel;
+		}
 
-		// In prior versions this node had a single "In" pin. In later versions this is split. The node will function
-		// the same if we move all connections from "In" to "Bounding Shape". To make this happen, rename "In" to
-		// "Bounding Shape" just prior to pin update and the edges will be moved over.
-		InputPins[0]->Properties.Label = PCGSurfaceSamplerConstants::BoundingShapeLabel;
+		// If no pins or only 1 pin, then we will be doing rewiring later, so ask for advanced pins on graph input node
+		UPCGNode* GraphInputNode = InOutNode->GetGraph() ? InOutNode->GetGraph()->GetInputNode() : nullptr;
+		if (GraphInputNode && (InputPins.IsEmpty() || (InputPins.Num() == 1 && PCGSurfaceSampler::IsPinOnlyConnectedToInputNode(InputPins[0], GraphInputNode))))
+		{
+			// Ask for advanced params to be visible, as it will be useful later for wiring up the Surface input to the advanced Landscape pin
+			if (UPCGGraphInputOutputSettings* InputOutputSettings = Cast<UPCGGraphInputOutputSettings>(GraphInputNode->GetSettings()))
+			{
+				InputOutputSettings->SetShowAdvancedPins(true);
+			}
+		}
 	}
 
 	Super::ApplyDeprecationBeforeUpdatePins(InOutNode, InputPins, OutputPins);
 }
+
+void UPCGSurfaceSamplerSettings::ApplyDeprecation(UPCGNode* InOutNode)
+{
+	if (DataVersion < FPCGCustomVersion::SplitSamplerNodesInputs && ensure(InOutNode && InOutNode->GetInputPins().Num() >= 2))
+	{
+		UE_LOG(LogPCG, Log, TEXT("Surface Sampler node migrated from an older version. Review edges on the input pins and then save this graph to upgrade the data."));
+
+		UPCGPin* SurfacePin = InOutNode->GetInputPin(FName(TEXT("Surface")));
+		UPCGPin* BoundingShapePin = InOutNode->GetInputPin(FName(TEXT("Bounding Shape")));
+		UPCGNode* GraphInputNode = InOutNode->GetGraph() ? InOutNode->GetGraph()->GetInputNode() : nullptr;
+
+		if (SurfacePin && BoundingShapePin && GraphInputNode)
+		{
+			auto MoveEdgeOnInputNodeToLandscapePin = [InOutNode, GraphInputNode, SurfacePin](UPCGPin* DownstreamPin) {
+				// Detect if we're connected to the Input node.
+				if (PCGSurfaceSampler::IsPinOnlyConnectedToInputNode(DownstreamPin, GraphInputNode))
+				{
+					// If we are connected to the Input node, make just a connection from the Surface pin to the Landscape pin and rely on Unbounded setting to provide bounds.
+					if (UPCGPin* LandscapePin = GraphInputNode->GetOutputPin(FName(TEXT("Landscape"))))
+					{
+						DownstreamPin->BreakAllEdges();
+
+						LandscapePin->AddEdgeTo(SurfacePin);
+					}
+				}
+			};
+
+			// The input pin has been split into two. Detect if we have inputs on only one pin and are dealing with older data - if so there's a good chance we can rewire
+			// in a better way.
+			if (SurfacePin->Edges.Num() == 0 && BoundingShapePin->Edges.Num() > 0)
+			{
+				MoveEdgeOnInputNodeToLandscapePin(BoundingShapePin);
+			}
+			else if (SurfacePin->Edges.Num() > 0 && BoundingShapePin->Edges.Num() == 0)
+			{
+				MoveEdgeOnInputNodeToLandscapePin(SurfacePin);
+			}
+		}
+	}
+
+	Super::ApplyDeprecation(InOutNode);
+}
 #endif // WITH_EDITOR
 
 #undef LOCTEXT_NAMESPACE
diff --git a/Engine/Plugins/Experimental/PCG/Source/PCG/Public/Elements/PCGSurfaceSampler.h b/Engine/Plugins/Experimental/PCG/Source/PCG/Public/Elements/PCGSurfaceSampler.h
--- a/Engine/Plugins/Experimental/PCG/Source/PCG/Public/Elements/PCGSurfaceSampler.h
+++ b/Engine/Plugins/Experimental/PCG/Source/PCG/Public/Elements/PCGSurfaceSampler.h
@@ -38,6 +38,7 @@ public:
 
 #if WITH_EDITOR
 	virtual void ApplyDeprecationBeforeUpdatePins(UPCGNode* InOutNode, TArray<TObjectPtr<UPCGPin>>& InputPins, TArray<TObjectPtr<UPCGPin>>& OutputPins) override;
+	virtual void ApplyDeprecation(UPCGNode* InOutNode) override;
 #endif
 
 protected:
