Fixed warnings/errors on input/output nodes that aren't relevant as we hide those pins

#rnx

diff --git a/Engine/Plugins/Experimental/PCG/Source/PCGEditor/Private/PCGEditorGraph.cpp b/Engine/Plugins/Experimental/PCG/Source/PCGEditor/Private/PCGEditorGraph.cpp
--- a/Engine/Plugins/Experimental/PCG/Source/PCGEditor/Private/PCGEditorGraph.cpp
+++ b/Engine/Plugins/Experimental/PCG/Source/PCGEditor/Private/PCGEditorGraph.cpp
@@ -92,7 +92,11 @@ void UPCGEditorGraph::CreateLinks(UPCGEditorGraphNodeBase* GraphNode, bool bCrea
 
 			if (!InPin)
 			{
-				UE_LOG(LogPCGEditor, Error, TEXT("Invalid InputPin for %s"), *InputPin->Properties.Label.ToString());
+				if (!Cast<UPCGEditorGraphNodeInput>(GraphNode))
+				{
+					UE_LOG(LogPCGEditor, Error, TEXT("Invalid InputPin for %s"), *InputPin->Properties.Label.ToString());
+				}
+
 				continue;
 			}
 
@@ -129,7 +133,11 @@ void UPCGEditorGraph::CreateLinks(UPCGEditorGraphNodeBase* GraphNode, bool bCrea
 
 			if (!OutPin)
 			{
-				UE_LOG(LogPCGEditor, Error, TEXT("Invalid OutputPin for %s"), *OutputPin->Properties.Label.ToString());
+				if (!Cast<UPCGEditorGraphNodeOutput>(GraphNode))
+				{
+					UE_LOG(LogPCGEditor, Error, TEXT("Invalid OutputPin for %s"), *OutputPin->Properties.Label.ToString());
+				}
+				
 				continue;
 			}
 
