Fixed issue with union/keep all and merge nodes metadata merge.

#rnx

diff --git a/Engine/Plugins/Experimental/PCG/Source/PCG/Private/Data/PCGPointData.cpp b/Engine/Plugins/Experimental/PCG/Source/PCG/Private/Data/PCGPointData.cpp
--- a/Engine/Plugins/Experimental/PCG/Source/PCG/Private/Data/PCGPointData.cpp
+++ b/Engine/Plugins/Experimental/PCG/Source/PCG/Private/Data/PCGPointData.cpp
@@ -151,7 +151,10 @@ namespace PCGPointHelpers
 
 		auto CopyPoint = [&OutPoint, &OutMetadata, &SourceMetadata](const FPCGPoint& PointToCopy)
 		{
+			PCGMetadataEntryKey OutPointEntryKey = OutPoint.MetadataEntry;
 			OutPoint = PointToCopy;
+			OutPoint.MetadataEntry = OutPointEntryKey;
+
 			if (OutMetadata)
 			{
 				OutMetadata->SetPointAttributes(PointToCopy, SourceMetadata, OutPoint);
diff --git a/Engine/Plugins/Experimental/PCG/Source/PCG/Private/Data/PCGUnionData.cpp b/Engine/Plugins/Experimental/PCG/Source/PCG/Private/Data/PCGUnionData.cpp
--- a/Engine/Plugins/Experimental/PCG/Source/PCG/Private/Data/PCGUnionData.cpp
+++ b/Engine/Plugins/Experimental/PCG/Source/PCG/Private/Data/PCGUnionData.cpp
@@ -218,6 +218,7 @@ const UPCGPointData* UPCGUnionData::CreatePointData(FPCGContext* Context) const
 
 	UPCGPointData* PointData = NewObject<UPCGPointData>();
 	PointData->InitializeFromData(this, Data[0]->Metadata);
+	check(PointData->Metadata);
 
 	// Initialize metadata
 	for (TObjectPtr<const UPCGSpatialData> Datum : Data)
@@ -239,15 +240,28 @@ const UPCGPointData* UPCGUnionData::CreatePointData(FPCGContext* Context) const
 	case EPCGUnionType::KeepAll:
 		{
 			TArray<FPCGPoint>& TargetPoints = PointData->GetMutablePoints();
-			for (TObjectPtr<const UPCGSpatialData> Datum : Data)
+			for(int32 DataIndex = 0; DataIndex < Data.Num(); ++DataIndex)
 			{
+				TObjectPtr<const UPCGSpatialData> Datum = Data[DataIndex];
+
 				const UPCGPointData* DatumPointData = Datum->ToPointData(Context);
-				int32 TargetPointIndex = TargetPoints.Num();
+
+				int32 TargetPointOffset = TargetPoints.Num();
 				TargetPoints.Append(DatumPointData->GetPoints());
 
-				if (PointData->Metadata && DatumPointData->GetPoints().Num() > 0)
+				if (DataIndex > 0 && DatumPointData->GetPoints().Num() > 0)
 				{
-					PointData->Metadata->SetPointAttributes(MakeArrayView(DatumPointData->GetPoints()), DatumPointData->Metadata, MakeArrayView(&TargetPoints[TargetPointIndex], DatumPointData->GetPoints().Num()));
+					// TODO: could optimize case where there is a common parent between Data 0 and current data, for points that still point to common parent metadata.
+					TArrayView<FPCGPoint> TargetPointsSubset = MakeArrayView(&TargetPoints[TargetPointOffset], DatumPointData->GetPoints().Num());
+					for (FPCGPoint& Point : TargetPointsSubset)
+					{
+						Point.MetadataEntry = PCGInvalidEntryKey;
+					}
+
+					if (PointData->Metadata && DatumPointData->Metadata && DatumPointData->Metadata->GetAttributeCount() > 0)
+					{
+						PointData->Metadata->SetPointAttributes(MakeArrayView(DatumPointData->GetPoints()), DatumPointData->Metadata, TargetPointsSubset);
+					}
 				}
 			}
 
diff --git a/Engine/Plugins/Experimental/PCG/Source/PCG/Private/Elements/PCGMergeElement.cpp b/Engine/Plugins/Experimental/PCG/Source/PCG/Private/Elements/PCGMergeElement.cpp
--- a/Engine/Plugins/Experimental/PCG/Source/PCG/Private/Elements/PCGMergeElement.cpp
+++ b/Engine/Plugins/Experimental/PCG/Source/PCG/Private/Elements/PCGMergeElement.cpp
@@ -97,34 +97,32 @@ bool FPCGMergeElement::ExecuteInternal(FPCGContext* Context) const
 		return true;
 	}
 
-	bool bIsFirst = true;
 	TArray<FPCGPoint>& TargetPoints = TargetPointData->GetMutablePoints();
 	
-	for (FPCGTaggedData& Source : Sources)
+	for(int32 SourceIndex = 0; SourceIndex < Sources.Num(); ++SourceIndex)
 	{
-		const UPCGPointData* SourcePointData = Cast<const UPCGPointData>(Source.Data);
+		const UPCGPointData* SourcePointData = Cast<const UPCGPointData>(Sources[SourceIndex].Data);
 
 		if (!SourcePointData)
 		{
 			continue;
 		}
 
-		if (bIsFirst)
-		{
-			check(!bMergeMetadata || TargetPointData->Metadata->GetParent() == SourcePointData->Metadata);
-			TargetPoints.Append(SourcePointData->GetPoints());
-			bIsFirst = false;
-		}
-		else
+		int32 PointOffset = TargetPoints.Num();
+		TargetPoints.Append(SourcePointData->GetPoints());
+
+		if ((!bMergeMetadata || SourceIndex != 0) && !SourcePointData->GetPoints().IsEmpty())
 		{
-			int32 PointOffset = TargetPoints.Num();
-			TargetPoints.Append(SourcePointData->GetPoints());
+			TArrayView<FPCGPoint> TargetPointsSubset = MakeArrayView(&TargetPoints[PointOffset], SourcePointData->GetPoints().Num());
+			for (FPCGPoint& Point : TargetPointsSubset)
+			{
+				Point.MetadataEntry = PCGInvalidEntryKey;
+			}
 
-			if (bMergeMetadata && TargetPointData->Metadata && SourcePointData->Metadata && !SourcePointData->GetPoints().IsEmpty())
+			if (bMergeMetadata && TargetPointData->Metadata && SourcePointData->Metadata && SourcePointData->Metadata->GetAttributeCount() > 0)
 			{
-				TargetPointData->Metadata->SetPointAttributes(MakeArrayView(SourcePointData->GetPoints()), SourcePointData->Metadata, MakeArrayView(&TargetPoints[PointOffset], SourcePointData->GetPoints().Num()));
+				TargetPointData->Metadata->SetPointAttributes(MakeArrayView(SourcePointData->GetPoints()), SourcePointData->Metadata, TargetPointsSubset);
 			}
-			// TBD: should we null out the metadata entry keys on the vertices if we didn't setup the metadata?
 		}
 	}
 	
diff --git a/Engine/Plugins/Experimental/PCG/Source/PCG/Private/Metadata/PCGMetadata.cpp b/Engine/Plugins/Experimental/PCG/Source/PCG/Private/Metadata/PCGMetadata.cpp
--- a/Engine/Plugins/Experimental/PCG/Source/PCG/Private/Metadata/PCGMetadata.cpp
+++ b/Engine/Plugins/Experimental/PCG/Source/PCG/Private/Metadata/PCGMetadata.cpp
@@ -321,6 +321,29 @@ bool UPCGMetadata::HasAttribute(FName AttributeName) const
 	return Attributes.Contains(AttributeName);
 }
 
+bool UPCGMetadata::HasCommonAttributes(const UPCGMetadata* InMetadata) const
+{
+	if (!InMetadata)
+	{
+		return false;
+	}
+
+	bool bHasCommonAttribute = false;
+
+	AttributeLock.ReadLock();
+	for (const TPair<FName, FPCGMetadataAttributeBase*>& AttributePair : Attributes)
+	{
+		if (InMetadata->HasAttribute(AttributePair.Key))
+		{
+			bHasCommonAttribute = true;
+			break;
+		}
+	}
+	AttributeLock.ReadUnlock();
+
+	return bHasCommonAttribute;
+}
+
 int32 UPCGMetadata::GetAttributeCount() const
 {
 	FReadScopeLock ScopeLock(AttributeLock);
diff --git a/Engine/Plugins/Experimental/PCG/Source/PCG/Public/Metadata/PCGMetadata.h b/Engine/Plugins/Experimental/PCG/Source/PCG/Public/Metadata/PCGMetadata.h
--- a/Engine/Plugins/Experimental/PCG/Source/PCG/Public/Metadata/PCGMetadata.h
+++ b/Engine/Plugins/Experimental/PCG/Source/PCG/Public/Metadata/PCGMetadata.h
@@ -143,6 +143,9 @@ public:
 	UFUNCTION(BlueprintCallable, Category = "PCG|Metadata")
 	bool HasAttribute(FName AttributeName) const;
 
+	UFUNCTION(BlueprintCallable, Category = "PCG|Metadata")
+	bool HasCommonAttributes(const UPCGMetadata* InMetadata) const;
+
 	int32 GetAttributeCount() const;
 
 	template <typename T>
