On behalf of thomas.tedemalm:
Inspection now summons the inspection tab

#rnx

diff --git a/Engine/Plugins/Experimental/PCG/Source/PCGEditor/Private/PCGEditor.cpp b/Engine/Plugins/Experimental/PCG/Source/PCGEditor/Private/PCGEditor.cpp
--- a/Engine/Plugins/Experimental/PCG/Source/PCGEditor/Private/PCGEditor.cpp
+++ b/Engine/Plugins/Experimental/PCG/Source/PCGEditor/Private/PCGEditor.cpp
@@ -1196,6 +1196,8 @@ void FPCGEditor::OnStartInspectNode()
 	PCGGraphNodeBeingInspected = PCGGraphNodeBase;
 	PCGGraphNodeBeingInspected->SetInspected(true);
 	OnInspectedNodeChangedDelegate.Broadcast(PCGGraphNodeBeingInspected);
+
+	GetTabManager()->TryInvokeTab(FPCGEditor_private::AttributesID);
 }
 
 void FPCGEditor::OnStopInspectNode()
