Tentative fix for crash-on-load which might be related to bp assets

#rnx

diff --git a/Engine/Plugins/Experimental/PCG/Source/PCG/Private/PCGComponent.cpp b/Engine/Plugins/Experimental/PCG/Source/PCG/Private/PCGComponent.cpp
--- a/Engine/Plugins/Experimental/PCG/Source/PCG/Private/PCGComponent.cpp
+++ b/Engine/Plugins/Experimental/PCG/Source/PCG/Private/PCGComponent.cpp
@@ -609,7 +609,9 @@ void UPCGComponent::CleanupLocalImmediate(bool bRemoveComponents)
 		check(!GeneratedResourcesInaccessible);
 		for (int32 ResourceIndex = GeneratedResources.Num() - 1; ResourceIndex >= 0; --ResourceIndex)
 		{
-			if (GeneratedResources[ResourceIndex]->Release(bRemoveComponents, ActorsToDelete))
+			// Note: resources can be null here in some loading + bp object cases
+			UPCGManagedResource* Resource = GeneratedResources[ResourceIndex];
+			if (!Resource || Resource->Release(bRemoveComponents, ActorsToDelete))
 			{
 				GeneratedResources.RemoveAtSwap(ResourceIndex);
 			}
