Fixed localization warnings

#rnx

diff --git a/Engine/Plugins/Experimental/PCG/Source/PCGEditor/Private/PCGEditorGraphNodeBase.cpp b/Engine/Plugins/Experimental/PCG/Source/PCGEditor/Private/PCGEditorGraphNodeBase.cpp
--- a/Engine/Plugins/Experimental/PCG/Source/PCGEditor/Private/PCGEditorGraphNodeBase.cpp
+++ b/Engine/Plugins/Experimental/PCG/Source/PCGEditor/Private/PCGEditorGraphNodeBase.cpp
@@ -504,12 +504,12 @@ void UPCGEditorGraphNodeBase::GetPinHoverText(const UEdGraphPin& Pin, FString& H
 	{
 		if (bIsInputPin)
 		{
-			MultiDataSupport = MatchingPin->Properties.bAllowMultipleData ? FText(LOCTEXT("SupportsMultiData", "Supports multiple data in input collection(s). ")) : FText(LOCTEXT("SingleDataOnly", "Supports only single data in collection(s). "));
+			MultiDataSupport = MatchingPin->Properties.bAllowMultipleData ? FText(LOCTEXT("InputSupportsMultiData", "Supports multiple data in input collection(s). ")) : FText(LOCTEXT("InputSingleDataOnly", "Supports only single data in collection(s). "));
 			MultiConnectionSupport = MatchingPin->Properties.bAllowMultipleConnections ? FText(LOCTEXT("SupportsMultiInput", "Supports multiple inputs.")) : FText(LOCTEXT("SingleInputOnly", "Supports only one input."));
 		}
 		else
 		{
-			MultiDataSupport = MatchingPin->Properties.bAllowMultipleData ? FText(LOCTEXT("SupportsMultiData", "Can generate multiple data.")) : FText(LOCTEXT("SingleDataOnly", "Generates only single data."));
+			MultiDataSupport = MatchingPin->Properties.bAllowMultipleData ? FText(LOCTEXT("OutputSupportsMultiData", "Can generate multiple data.")) : FText(LOCTEXT("OutputSingleDataOnly", "Generates only single data."));
 		}
 	}
 
