Temporary fix - added metadata parenting check for rooting/unrooting

#rnx

diff --git a/Engine/Plugins/Experimental/PCG/Source/PCG/Private/PCGData.cpp b/Engine/Plugins/Experimental/PCG/Source/PCG/Private/PCGData.cpp
--- a/Engine/Plugins/Experimental/PCG/Source/PCG/Private/PCGData.cpp
+++ b/Engine/Plugins/Experimental/PCG/Source/PCG/Private/PCGData.cpp
@@ -45,6 +45,26 @@ void FPCGRootSet::AddInternal(UObject* InObject)
 	{
 		AddInternal(OuterObject);
 	}
+
+	// Recurse on metadata
+	UPCGMetadata* InMetadata = nullptr;
+	if (UPCGSpatialData* InSpatialData = Cast<UPCGSpatialData>(InObject))
+	{
+		InMetadata = InSpatialData->Metadata;
+	}
+	else if (UPCGParamData* InParamData = Cast<UPCGParamData>(InObject))
+	{
+		InMetadata = InParamData->Metadata;
+	}
+
+	if (InMetadata && InMetadata->GetParent())
+	{
+		UObject* MetadataParentOuter = InMetadata->GetParent()->GetOuter();
+		if (MetadataParentOuter && !MetadataParentOuter->IsA<UPackage>())
+		{
+			AddInternal(MetadataParentOuter);
+		}
+	}
 }
 
 void FPCGRootSet::Remove(UObject* InObject)
@@ -80,6 +100,27 @@ void FPCGRootSet::RemoveInternal(UObject* InObject)
 	{
 		RemoveInternal(OuterObject);
 	}
+
+	// Recurse on metadata
+	UPCGMetadata* InMetadata = nullptr;
+
+	if (UPCGSpatialData* InSpatialData = Cast<UPCGSpatialData>(InObject))
+	{
+		InMetadata = InSpatialData->Metadata;
+	}
+	else if (UPCGParamData* InParamData = Cast<UPCGParamData>(InObject))
+	{
+		InMetadata = InParamData->Metadata;
+	}
+
+	if (InMetadata && InMetadata->GetParent())
+	{
+		UObject* MetadataParentOuter = InMetadata->GetParent()->GetOuter();
+		if (MetadataParentOuter && !MetadataParentOuter->IsA<UPackage>())
+		{
+			RemoveInternal(MetadataParentOuter);
+		}
+	}
 }
 
 bool FPCGTaggedData::operator==(const FPCGTaggedData& Other) const
