[PCG] Fix static analysis warning (red herring but it also flagged a concern in review so does not hurt to explicit-ize it).

diff --git a/Engine/Plugins/Experimental/PCG/Source/PCG/Private/Elements/PCGSurfaceSampler.cpp b/Engine/Plugins/Experimental/PCG/Source/PCG/Private/Elements/PCGSurfaceSampler.cpp
--- a/Engine/Plugins/Experimental/PCG/Source/PCG/Private/Elements/PCGSurfaceSampler.cpp
+++ b/Engine/Plugins/Experimental/PCG/Source/PCG/Private/Elements/PCGSurfaceSampler.cpp
@@ -198,7 +198,8 @@ namespace PCGSurfaceSampler
 			const UPCGEdge* Edge = DownstreamPin->Edges[0];
 			const UPCGNode* UpstreamNode = (Edge && Edge->InputPin) ? Edge->InputPin->Node : nullptr;
 			const bool bConnectedToInputNode = UpstreamNode && (GraphInputNode == UpstreamNode);
-			return bConnectedToInputNode && (Edge->InputPin->Properties.Label == FName(TEXT("In")) || Edge->InputPin->Properties.Label == FName(TEXT("Input")));
+			const bool bConnectedToInputPin = Edge && (Edge->InputPin->Properties.Label == FName(TEXT("In")) || Edge->InputPin->Properties.Label == FName(TEXT("Input")));
+			return bConnectedToInputNode && bConnectedToInputPin;
 		}
 
 		return false;
