[PCG] Attribute operation with OutputTarget as None will try to write in the latest attribute instead of creating a None attribute

diff --git a/Engine/Plugins/Experimental/PCG/Source/PCG/Private/Elements/Metadata/PCGMetadataElement.cpp b/Engine/Plugins/Experimental/PCG/Source/PCG/Private/Elements/Metadata/PCGMetadataElement.cpp
--- a/Engine/Plugins/Experimental/PCG/Source/PCG/Private/Elements/Metadata/PCGMetadataElement.cpp
+++ b/Engine/Plugins/Experimental/PCG/Source/PCG/Private/Elements/Metadata/PCGMetadataElement.cpp
@@ -109,13 +109,16 @@ bool FPCGMetadataOperationElement::ExecuteInternal(FPCGContext* Context) const
 		}
 
 		FPCGAttributePropertySelector InputSource = Settings->InputSource;
+		FPCGAttributePropertySelector OutputTarget = Settings->OutputTarget;
 
 		const FName SourceAttribute = InputSource.GetName();
+		const FName DestinationAttribute = OutputTarget.GetName();
 		const FName LocalSourceAttribute = ((SourceAttribute != NAME_None) ? SourceAttribute : OriginalData->Metadata->GetLatestAttributeNameOrNone());
-		const FName DestinationAttribute = Settings->OutputTarget.GetName();
+		const FName LocalDestinationAttribute = ((DestinationAttribute != NAME_None) ? DestinationAttribute : OriginalData->Metadata->GetLatestAttributeNameOrNone());
 
 		// Make sure we use the right attribute name
 		InputSource.AttributeName = LocalSourceAttribute;
+		OutputTarget.AttributeName = LocalDestinationAttribute;
 
 		if (Settings->InputSource.Selection == EPCGAttributePropertySelection::Attribute && !OriginalData->Metadata->HasAttribute(LocalSourceAttribute))
 		{
@@ -136,7 +139,7 @@ bool FPCGMetadataOperationElement::ExecuteInternal(FPCGContext* Context) const
 		SampledPoints = Points;
 
 		// If it is attribute to attribute, just copy the attributes, if they exist and are valid
-		if (Settings->InputSource.Selection == EPCGAttributePropertySelection::Attribute && Settings->OutputTarget.Selection == EPCGAttributePropertySelection::Attribute)
+		if (Settings->InputSource.Selection == EPCGAttributePropertySelection::Attribute && OutputTarget.Selection == EPCGAttributePropertySelection::Attribute)
 		{
 			if (LocalSourceAttribute == DestinationAttribute)
 			{
@@ -161,17 +164,19 @@ bool FPCGMetadataOperationElement::ExecuteInternal(FPCGContext* Context) const
 			continue;
 		}
 
-		if (Settings->OutputTarget.Selection == EPCGAttributePropertySelection::Attribute)
+		// If the target is an attribute, only create a new one if the attribute doesn't already exist.
+		// If it exist, it will try to write to it.
+		if (OutputTarget.Selection == EPCGAttributePropertySelection::Attribute && !OriginalData->Metadata->HasAttribute(LocalDestinationAttribute))
 		{
-			auto CreateAttribute = [SampledData, DestinationAttribute](auto Dummy)
+			auto CreateAttribute = [SampledData, LocalDestinationAttribute](auto Dummy)
 			{
 				using AttributeType = decltype(Dummy);
-				return PCGMetadataElementCommon::ClearOrCreateAttribute(SampledData->Metadata, DestinationAttribute, AttributeType{}) != nullptr;
+				return PCGMetadataElementCommon::ClearOrCreateAttribute(SampledData->Metadata, LocalDestinationAttribute, AttributeType{}) != nullptr;
 			};
 			
 			if (!PCGMetadataAttribute::CallbackWithRightType(InputAccessor->GetUnderlyingType(), CreateAttribute))
 			{
-				PCGE_LOG(Warning, "Failed to create new attribute %s", *DestinationAttribute.ToString());
+				PCGE_LOG(Warning, "Failed to create new attribute %s", *LocalDestinationAttribute.ToString());
 				continue;
 			}
 		}
