[PCG] Unit test push - Attribute Selector

diff --git a/Engine/Plugins/Experimental/PCG/Source/PCG/Private/Tests/Accessors/PCGAttributeAccessorTest.cpp b/Engine/Plugins/Experimental/PCG/Source/PCG/Private/Tests/Accessors/PCGAttributeAccessorTest.cpp
new file mode 100644
--- /dev/null
+++ b/Engine/Plugins/Experimental/PCG/Source/PCG/Private/Tests/Accessors/PCGAttributeAccessorTest.cpp
@@ -0,0 +1,654 @@
+// Copyright Epic Games, Inc. All Rights Reserved.
+
+#include "Tests/PCGTestsCommon.h"
+
+#include "PCGPoint.h"
+#include "Data/PCGPointData.h"
+#include "Metadata/PCGMetadataAttributeTpl.h"
+#include "Metadata/PCGMetadata.h"
+#include "Metadata/Accessors/PCGAttributeAccessor.h"
+#include "Metadata/Accessors/PCGAttributeAccessorKeys.h"
+
+#include "UObject/UnrealType.h"
+
+#if WITH_EDITOR
+
+namespace FPCGAttributeAccessorTestHelpers
+{
+	constexpr int32 ValuesSize = 3;
+
+	const float FloatValues[ValuesSize] = { 0.57441f, -0.289f, 0.1587f};
+	const FVector VectorValues[ValuesSize] = { FVector(1.024f, 0.445f, 0.587f), FVector(4.6470f, 3.19874f, 9.6648f), FVector(4.690f, 7.874f, 8.668f) };
+	const FString StringValues[ValuesSize] = { FString(TEXT("Str1")), FString(TEXT("Foo")), FString(TEXT("MyStr")) };
+
+	struct AttributeData
+	{
+		AttributeData()
+			: PointData(NewObject<UPCGPointData>(GetTransientPackage(), NAME_None, RF_Transient))
+			, FloatAttribute(PointData->Metadata, TEXT("FloatAttribute"), nullptr, 0.0f, true)
+			, VectorAttribute(PointData->Metadata, TEXT("VectorAttribute"), nullptr, FVector::ZeroVector, true)
+			, StringAttribute(PointData->Metadata, TEXT("StringAttribute"), nullptr, FString(), true)
+		{
+			// All points will point to a given entry, in reverse order, except the last one, that will point to DefaultValue
+			TArray<FPCGPoint>& Points = PointData->GetMutablePoints();
+			Points.SetNum(ValuesSize + 1);
+
+			for (int32 i = 0; i < ValuesSize; ++i)
+			{
+				PCGMetadataEntryKey EntryKey = PointData->Metadata->AddEntry();
+				FloatAttribute.SetValue(EntryKey, FloatValues[i]);
+				VectorAttribute.SetValue(EntryKey, VectorValues[i]);
+				StringAttribute.SetValue(EntryKey, StringValues[i]);
+				Points[ValuesSize - i - 1].MetadataEntry = EntryKey;
+			}
+
+			FloatAccessor = MakeUnique<FPCGAttributeAccessor<float>>(&FloatAttribute, PointData->Metadata);
+			VectorAccessor = MakeUnique<FPCGAttributeAccessor<FVector>>(&VectorAttribute, PointData->Metadata);
+			StringAccessor = MakeUnique<FPCGAttributeAccessor<FString>>(&StringAttribute, PointData->Metadata);
+		}
+
+		void ResetTempFloats()
+		{
+			TempFloats[0] = 1.0f;
+			TempFloats[1] = 2.0f;
+			TempFloats[2] = 3.0f;
+		}
+
+		TObjectPtr<UPCGPointData> PointData;
+		FPCGMetadataAttribute<float> FloatAttribute;
+		FPCGMetadataAttribute<FVector> VectorAttribute;
+		FPCGMetadataAttribute<FString> StringAttribute;
+
+		TUniquePtr<IPCGAttributeAccessor> FloatAccessor;
+		TUniquePtr<IPCGAttributeAccessor> VectorAccessor;
+		TUniquePtr<IPCGAttributeAccessor> StringAccessor;
+
+		float TempFloats[FPCGAttributeAccessorTestHelpers::ValuesSize] = { 1.0f, 2.0f, 3.0f };
+		FVector TempVectors[FPCGAttributeAccessorTestHelpers::ValuesSize] = { FVector(0.1f, 0.1f, 0.1f), FVector(0.1f, 0.2f, 0.1f), FVector(0.1f, 0.2f, 0.3f) };
+		FString TempStrings[FPCGAttributeAccessorTestHelpers::ValuesSize] = { TEXT("Bla"), TEXT("BlaBla"), TEXT("Blou") };
+		int32 TempInvalidInts[FPCGAttributeAccessorTestHelpers::ValuesSize] = { 5, 6, 7 };
+	};
+
+	template <typename T, size_t N>
+	TArrayView<T> MakeArrayView(T(& Array)[N])
+	{
+		return TArrayView<T>(Array, N);
+	}
+
+	template <typename T, size_t N>
+	TArrayView<const T> MakeConstArrayView(T(&Array)[N])
+	{
+		return TArrayView<const T>(Array, N);
+	}
+}
+
+IMPLEMENT_CUSTOM_SIMPLE_AUTOMATION_TEST(FPCGAttributeAccessorSingleGetDefaultTest, FPCGTestBaseClass, "pcg.tests.Accessor.Attribute.SingleGetDefault", PCGTestsCommon::TestFlags)
+IMPLEMENT_CUSTOM_SIMPLE_AUTOMATION_TEST(FPCGAttributeAccessorSingleGetTest, FPCGTestBaseClass, "pcg.tests.Accessor.Attribute.SingleGet", PCGTestsCommon::TestFlags)
+IMPLEMENT_CUSTOM_SIMPLE_AUTOMATION_TEST(FPCGAttributeAccessorSingleGetOutsideRangeTest, FPCGTestBaseClass, "pcg.tests.Accessor.Attribute.SingleGetOutsideRange", PCGTestsCommon::TestFlags)
+
+IMPLEMENT_CUSTOM_SIMPLE_AUTOMATION_TEST(FPCGAttributeAccessorGetRangeDefaultTest, FPCGTestBaseClass, "pcg.tests.Accessor.Attribute.GetRangeDefault", PCGTestsCommon::TestFlags)
+IMPLEMENT_CUSTOM_SIMPLE_AUTOMATION_TEST(FPCGAttributeAccessorGetRangeTest, FPCGTestBaseClass, "pcg.tests.Accessor.Attribute.GetRange", PCGTestsCommon::TestFlags)
+IMPLEMENT_CUSTOM_SIMPLE_AUTOMATION_TEST(FPCGAttributeAccessorGetRangeOutsideRangeTest, FPCGTestBaseClass, "pcg.tests.Accessor.Attribute.GetRangeOutsideRange", PCGTestsCommon::TestFlags)
+
+IMPLEMENT_CUSTOM_SIMPLE_AUTOMATION_TEST(FPCGAttributeAccessorSingleSetDefaultTest, FPCGTestBaseClass, "pcg.tests.Accessor.Attribute.SingleSetDefault", PCGTestsCommon::TestFlags)
+IMPLEMENT_CUSTOM_SIMPLE_AUTOMATION_TEST(FPCGAttributeAccessorSingleSetTest, FPCGTestBaseClass, "pcg.tests.Accessor.Attribute.SingleSet", PCGTestsCommon::TestFlags)
+IMPLEMENT_CUSTOM_SIMPLE_AUTOMATION_TEST(FPCGAttributeAccessorSingleInvalidKeySetTest, FPCGTestBaseClass, "pcg.tests.Accessor.Attribute.SingleSetInvalidKey", PCGTestsCommon::TestFlags)
+
+IMPLEMENT_CUSTOM_SIMPLE_AUTOMATION_TEST(FPCGAttributeAccessorSetRangeTest, FPCGTestBaseClass, "pcg.tests.Accessor.Attribute.SetRange", PCGTestsCommon::TestFlags)
+IMPLEMENT_CUSTOM_SIMPLE_AUTOMATION_TEST(FPCGAttributeAccessorSetRangeOutsideRangeTest, FPCGTestBaseClass, "pcg.tests.Accessor.Attribute.SetRangeOutsideRange", PCGTestsCommon::TestFlags)
+
+IMPLEMENT_CUSTOM_SIMPLE_AUTOMATION_TEST(FPCGAttributeAccessorSingleGetPointsTest, FPCGTestBaseClass, "pcg.tests.Accessor.Attribute.SingleGetPoints", PCGTestsCommon::TestFlags)
+IMPLEMENT_CUSTOM_SIMPLE_AUTOMATION_TEST(FPCGAttributeAccessorGetRangePointsTest, FPCGTestBaseClass, "pcg.tests.Accessor.Attribute.GetRangePoints", PCGTestsCommon::TestFlags)
+
+IMPLEMENT_CUSTOM_SIMPLE_AUTOMATION_TEST(FPCGAttributeAccessorSingleSetPointsTest, FPCGTestBaseClass, "pcg.tests.Accessor.Attribute.SingleSetPoints", PCGTestsCommon::TestFlags)
+IMPLEMENT_CUSTOM_SIMPLE_AUTOMATION_TEST(FPCGAttributeAccessorSetRangePointsTest, FPCGTestBaseClass, "pcg.tests.Accessor.Attribute.SetRangePoints", PCGTestsCommon::TestFlags)
+
+
+bool FPCGAttributeAccessorSingleGetDefaultTest::RunTest(const FString& Parameters)
+{
+	FPCGAttributeAccessorTestHelpers::AttributeData AttributeData;
+	FPCGAttributeAccessorKeysEntries DefaultEntry(PCGInvalidEntryKey);
+
+	if (!TestEqual("There is only 1 entry", DefaultEntry.GetNum(), 1))
+	{
+		return false;
+	}
+
+	if (!TestTrue("Get default float attribute", AttributeData.FloatAccessor->Get(*AttributeData.TempFloats, DefaultEntry))
+		|| !TestEqual("Default float value", *AttributeData.TempFloats, 0.0f))
+	{
+		return false;
+	}
+
+	if (!TestTrue("Get default vector attribute", AttributeData.VectorAccessor->Get(*AttributeData.TempVectors, DefaultEntry))
+		|| !TestEqual("Default vector value", *AttributeData.TempVectors, FVector::ZeroVector))
+	{
+		return false;
+	}
+
+	if (!TestTrue("Get default string attribute", AttributeData.StringAccessor->Get(*AttributeData.TempStrings, DefaultEntry))
+		|| !TestEqual("Default string value", *AttributeData.TempStrings, FString{}))
+	{
+		return false;
+	}
+
+	if (!TestFalse("Get default float with int", AttributeData.FloatAccessor->Get(*AttributeData.TempInvalidInts, DefaultEntry))
+		|| !TestEqual("Invalid int value not modified", *AttributeData.TempInvalidInts, 5))
+	{
+		return false;
+	}
+
+	return true;
+}
+
+bool FPCGAttributeAccessorSingleGetTest::RunTest(const FString& Parameters)
+{
+	FPCGAttributeAccessorTestHelpers::AttributeData AttributeData;
+	FPCGAttributeAccessorKeysEntries Keys(&AttributeData.FloatAttribute);
+
+	if (!TestEqual("Right number of entries", Keys.GetNum(), FPCGAttributeAccessorTestHelpers::ValuesSize))
+	{
+		return false;
+	}
+
+	for (int32 i = 0; i < FPCGAttributeAccessorTestHelpers::ValuesSize; ++i)
+	{
+		if (!TestTrue(FString::Printf(TEXT("Get Float value %d"), i), AttributeData.FloatAccessor->Get(AttributeData.TempFloats[i], i, Keys))
+			|| !TestEqual(FString::Printf(TEXT("Float value %d"), i), AttributeData.TempFloats[i], FPCGAttributeAccessorTestHelpers::FloatValues[i]))
+		{
+			return false;
+		}
+
+		if (!TestTrue(FString::Printf(TEXT("Get Vector value %d"), i), AttributeData.VectorAccessor->Get(AttributeData.TempVectors[i], i, Keys))
+			|| !TestEqual(FString::Printf(TEXT("Vector value %d"), i), AttributeData.TempVectors[i], FPCGAttributeAccessorTestHelpers::VectorValues[i]))
+		{
+			return false;
+		}
+
+		if (!TestTrue(FString::Printf(TEXT("Get String value %d"), i), AttributeData.StringAccessor->Get(AttributeData.TempStrings[i], i, Keys))
+			|| !TestEqual(FString::Printf(TEXT("String value % d"), i), AttributeData.TempStrings[i], FPCGAttributeAccessorTestHelpers::StringValues[i]))
+		{
+			return false;
+		}
+
+		if (!TestFalse(FString::Printf(TEXT("Get Float with int %d"), i), AttributeData.FloatAccessor->Get(AttributeData.TempInvalidInts[i], i, Keys))
+			|| !TestEqual(FString::Printf(TEXT("Invalid int value not modified %d"), i), AttributeData.TempInvalidInts[i], 5 + i))
+		{
+			return false;
+		}
+	}
+
+	return true;
+}
+
+bool FPCGAttributeAccessorSingleGetOutsideRangeTest::RunTest(const FString& Parameters)
+{
+	FPCGAttributeAccessorTestHelpers::AttributeData AttributeData;
+	FPCGAttributeAccessorKeysEntries Keys(&AttributeData.FloatAttribute);
+
+	float DefaultFloat = 1.0f;
+
+	// 3 -> 7 wrapping to 0 -> 3 and back to 0
+	for (int32 i = 0; i < FPCGAttributeAccessorTestHelpers::ValuesSize + 1; ++i)
+	{
+		if (!TestTrue(FString::Printf(TEXT("Get Float attribute outside range %d"), i), AttributeData.FloatAccessor->Get(DefaultFloat, i + FPCGAttributeAccessorTestHelpers::ValuesSize, Keys))
+			|| !TestEqual(FString::Printf(TEXT("Float value outside range %d"), i), DefaultFloat, FPCGAttributeAccessorTestHelpers::FloatValues[i % FPCGAttributeAccessorTestHelpers::ValuesSize]))
+		{
+			return false;
+		}
+	}
+
+	return true;
+}
+
+bool FPCGAttributeAccessorGetRangeDefaultTest::RunTest(const FString& Parameters)
+{
+	FPCGAttributeAccessorTestHelpers::AttributeData AttributeData;
+	FPCGAttributeAccessorKeysEntries DefaultEntry(PCGInvalidEntryKey);
+
+	int32 Index = 0;
+
+	if (!TestTrue("Get default float attribute full range", AttributeData.FloatAccessor->GetRange(FPCGAttributeAccessorTestHelpers::MakeArrayView(AttributeData.TempFloats), Index, DefaultEntry))
+		|| !TestEqual("Default float value 0", AttributeData.TempFloats[0], 0.0f) 
+		|| !TestEqual("Default float value 1", AttributeData.TempFloats[1], 0.0f) 
+		|| !TestEqual("Default float value 2", AttributeData.TempFloats[2], 0.0f))
+	{
+		return false;
+	}
+
+	if (!TestTrue("Get default vector attribute full range", AttributeData.VectorAccessor->GetRange(FPCGAttributeAccessorTestHelpers::MakeArrayView(AttributeData.TempVectors), Index, DefaultEntry))
+		|| !TestEqual("Default vector value 0", AttributeData.TempVectors[0], FVector::ZeroVector) 
+		|| !TestEqual("Default vector value 1", AttributeData.TempVectors[1], FVector::ZeroVector) 
+		|| !TestEqual("Default vector value 2", AttributeData.TempVectors[2], FVector::ZeroVector))
+	{
+		return false;
+	}
+
+	if (!TestTrue("Get default string attribute full range", AttributeData.StringAccessor->GetRange(FPCGAttributeAccessorTestHelpers::MakeArrayView(AttributeData.TempStrings), Index, DefaultEntry))
+		|| !TestEqual("Default string value 0", AttributeData.TempStrings[0], FString{}) 
+		|| !TestEqual("Default string value 1", AttributeData.TempStrings[1], FString{}) 
+		|| !TestEqual("Default string value 2", AttributeData.TempStrings[2], FString{}))
+	{
+		return false;
+	}
+
+	if (!TestFalse("Get default float with int full range", AttributeData.FloatAccessor->GetRange(FPCGAttributeAccessorTestHelpers::MakeArrayView(AttributeData.TempInvalidInts), Index, DefaultEntry))
+		|| !TestEqual("Invalid int value 0 not modified", AttributeData.TempInvalidInts[0], 5)
+		|| !TestEqual("Invalid int value 1 not modified", AttributeData.TempInvalidInts[1], 6)
+		|| !TestEqual("Invalid int value 2 not modified", AttributeData.TempInvalidInts[2], 7))
+	{
+		return false;
+	}
+
+	AttributeData.ResetTempFloats();
+
+	Index = 1;
+
+	if (!TestTrue("Get default float attribute with a smaller range (2), from index 1", AttributeData.FloatAccessor->GetRange(TArrayView<float>(AttributeData.TempFloats, 2), Index, DefaultEntry))
+		|| !TestEqual("Default float value 0", AttributeData.TempFloats[0], 0.0f)
+		|| !TestEqual("Default float value 1", AttributeData.TempFloats[1], 0.0f)
+		|| !TestEqual("Unmodified float value 2", AttributeData.TempFloats[2], 3.0f))
+	{
+		return false;
+	}
+
+	return true;
+}
+
+bool FPCGAttributeAccessorGetRangeTest::RunTest(const FString& Parameters)
+{
+	FPCGAttributeAccessorTestHelpers::AttributeData AttributeData;
+	FPCGAttributeAccessorKeysEntries Keys(&AttributeData.FloatAttribute);
+
+	int32 Index = 0;
+
+	if (!TestTrue("Get float attribute full range", AttributeData.FloatAccessor->GetRange(FPCGAttributeAccessorTestHelpers::MakeArrayView(AttributeData.TempFloats), Index, Keys))
+		|| !TestEqual("float value 0", AttributeData.TempFloats[0], FPCGAttributeAccessorTestHelpers::FloatValues[0])
+		|| !TestEqual("float value 1", AttributeData.TempFloats[1], FPCGAttributeAccessorTestHelpers::FloatValues[1])
+		|| !TestEqual("float value 2", AttributeData.TempFloats[2], FPCGAttributeAccessorTestHelpers::FloatValues[2]))
+	{
+		return false;
+	}
+
+	if (!TestTrue("Get vector attribute full range", AttributeData.VectorAccessor->GetRange(FPCGAttributeAccessorTestHelpers::MakeArrayView(AttributeData.TempVectors), Index, Keys))
+		|| !TestEqual("vector value 0", AttributeData.TempVectors[0], FPCGAttributeAccessorTestHelpers::VectorValues[0])
+		|| !TestEqual("vector value 1", AttributeData.TempVectors[1], FPCGAttributeAccessorTestHelpers::VectorValues[1])
+		|| !TestEqual("vector value 2", AttributeData.TempVectors[2], FPCGAttributeAccessorTestHelpers::VectorValues[2]))
+	{
+		return false;
+	}
+
+	if (!TestTrue("Get string attribute full range", AttributeData.StringAccessor->GetRange(FPCGAttributeAccessorTestHelpers::MakeArrayView(AttributeData.TempStrings), Index, Keys))
+		|| !TestEqual("string value 0", AttributeData.TempStrings[0], FPCGAttributeAccessorTestHelpers::StringValues[0])
+		|| !TestEqual("string value 1", AttributeData.TempStrings[1], FPCGAttributeAccessorTestHelpers::StringValues[1])
+		|| !TestEqual("string value 2", AttributeData.TempStrings[2], FPCGAttributeAccessorTestHelpers::StringValues[2]))
+	{
+		return false;
+	}
+
+	if (!TestFalse("Get float with int full range", AttributeData.FloatAccessor->GetRange(FPCGAttributeAccessorTestHelpers::MakeArrayView(AttributeData.TempInvalidInts), Index, Keys))
+		|| !TestEqual("Invalid int value 0 not modified", AttributeData.TempInvalidInts[0], 5)
+		|| !TestEqual("Invalid int value 1 not modified", AttributeData.TempInvalidInts[1], 6)
+		|| !TestEqual("Invalid int value 2 not modified", AttributeData.TempInvalidInts[2], 7))
+	{
+		return false;
+	}
+
+	AttributeData.ResetTempFloats();
+
+	Index = 1;
+
+	if (!TestTrue("Get float attribute with a smaller range(2), from index 1", AttributeData.FloatAccessor->GetRange(TArrayView<float>(AttributeData.TempFloats, 2), Index, Keys))
+		|| !TestEqual("float value 0", AttributeData.TempFloats[0], FPCGAttributeAccessorTestHelpers::FloatValues[1])
+		|| !TestEqual("float value 1", AttributeData.TempFloats[1], FPCGAttributeAccessorTestHelpers::FloatValues[2])
+		|| !TestEqual("Unmodified float value 2", AttributeData.TempFloats[2], 3.0f))
+	{
+		return false;
+	}
+
+	return true;
+}
+
+bool FPCGAttributeAccessorGetRangeOutsideRangeTest::RunTest(const FString& Parameters)
+{
+	FPCGAttributeAccessorTestHelpers::AttributeData AttributeData;
+	FPCGAttributeAccessorKeysEntries Keys(&AttributeData.FloatAttribute);
+
+	int32 OutsideRangeIndex = FPCGAttributeAccessorTestHelpers::ValuesSize;
+
+	// 3 -> 6 wrapping to 0 -> 3
+	if (!TestTrue("Get float outside range", AttributeData.FloatAccessor->GetRange(FPCGAttributeAccessorTestHelpers::MakeArrayView(AttributeData.TempFloats), OutsideRangeIndex, Keys))
+		|| !TestEqual("Float value 0 outside range", AttributeData.TempFloats[0], FPCGAttributeAccessorTestHelpers::FloatValues[0])
+		|| !TestEqual("Float value 1 outside range", AttributeData.TempFloats[1], FPCGAttributeAccessorTestHelpers::FloatValues[1])
+		|| !TestEqual("Float value 2 outside range", AttributeData.TempFloats[2], FPCGAttributeAccessorTestHelpers::FloatValues[2]))
+	{
+		return false;
+	}
+
+	AttributeData.ResetTempFloats();
+
+	OutsideRangeIndex = FPCGAttributeAccessorTestHelpers::ValuesSize + 1;
+
+	// 4, 5, 7 wrapping to 1, 2, 0
+	if (!TestTrue("Get float outside range offset", AttributeData.FloatAccessor->GetRange(FPCGAttributeAccessorTestHelpers::MakeArrayView(AttributeData.TempFloats), OutsideRangeIndex, Keys))
+		|| !TestEqual("Float value 0 outside range offset", AttributeData.TempFloats[0], FPCGAttributeAccessorTestHelpers::FloatValues[1])
+		|| !TestEqual("Float value 1 outside range offset", AttributeData.TempFloats[1], FPCGAttributeAccessorTestHelpers::FloatValues[2])
+		|| !TestEqual("Float value 2 outside range offset", AttributeData.TempFloats[2], FPCGAttributeAccessorTestHelpers::FloatValues[0]))
+	{
+		return false;
+	}
+
+	// Gathering more values will wrap also
+	constexpr int32 MoreFloatValuesSize = FPCGAttributeAccessorTestHelpers::ValuesSize * 3;
+	float MoreFloatValues[MoreFloatValuesSize];
+	for (int32 i = 0; i < MoreFloatValuesSize; ++i)
+	{
+		MoreFloatValues[i] = float(i);
+	}
+
+	if (!TestTrue("Get more float outside range", AttributeData.FloatAccessor->GetRange(FPCGAttributeAccessorTestHelpers::MakeArrayView(MoreFloatValues), 0, Keys)))
+	{
+		return false;
+	}
+
+	for (int32 i = 0; i < MoreFloatValuesSize; ++i)
+	{
+		if (!TestEqual(FString::Printf(TEXT("More Float value %d outside range"), i), MoreFloatValues[i], FPCGAttributeAccessorTestHelpers::FloatValues[i % FPCGAttributeAccessorTestHelpers::ValuesSize]))
+		{
+			return false;
+		}
+	}
+
+	return true;
+}
+
+bool FPCGAttributeAccessorSingleSetDefaultTest::RunTest(const FString& Parameters)
+{
+	FPCGAttributeAccessorTestHelpers::AttributeData AttributeData;
+	FPCGAttributeAccessorKeysEntries DefaultEntry(PCGInvalidEntryKey);
+
+	float NewDefaultFloat = 0.5f;
+	FVector NewDefaultVector = FVector(0.5f, 0.2f, 0.1f);
+	FString NewDefaultString(TEXT("AAA"));
+
+	if (!TestTrue("Set new float default value, with flag", AttributeData.FloatAccessor->Set(NewDefaultFloat, DefaultEntry, EPCGAttributeAccessorFlags::AllowSetDefaultValue)
+		|| !TestEqual("New default float", AttributeData.FloatAttribute.GetValue(PCGDefaultValueKey), NewDefaultFloat)))
+	{
+		return false;
+	}
+
+	if (!TestTrue("Set new vector default value, with flag", AttributeData.VectorAccessor->Set(NewDefaultVector, DefaultEntry, EPCGAttributeAccessorFlags::AllowSetDefaultValue)
+		|| !TestEqual("New default vector", AttributeData.VectorAttribute.GetValue(PCGDefaultValueKey), NewDefaultVector)))
+	{
+		return false;
+	}
+
+	if (!TestTrue("Set new string default value, with flag", AttributeData.StringAccessor->Set(NewDefaultString, DefaultEntry, EPCGAttributeAccessorFlags::AllowSetDefaultValue)
+		|| !TestEqual("New default string", AttributeData.StringAttribute.GetValue(PCGDefaultValueKey), NewDefaultString)))
+	{
+		return false;
+	}
+
+	return true;
+}
+
+bool FPCGAttributeAccessorSingleSetTest::RunTest(const FString& Parameters)
+{
+	FPCGAttributeAccessorTestHelpers::AttributeData AttributeData;
+	FPCGAttributeAccessorKeysEntries Keys(&AttributeData.FloatAttribute);
+
+	if (!TestTrue("Set float attribute", AttributeData.FloatAccessor->Set(*AttributeData.TempFloats, Keys))
+		|| !TestEqual("float value 0", AttributeData.FloatAttribute.GetValueFromItemKey(0), *AttributeData.TempFloats))
+	{
+		return false;
+	}
+
+	if (!TestTrue("Set vector attribute", AttributeData.VectorAccessor->Set(*AttributeData.TempVectors, /*Index=*/ 1, Keys))
+		|| !TestEqual("vector value 1", AttributeData.VectorAttribute.GetValueFromItemKey(1), *AttributeData.TempVectors))
+	{
+		return false;
+	}
+
+	if (!TestTrue("Set string attribute", AttributeData.StringAccessor->Set(*AttributeData.TempStrings, /*Index=*/ 2, Keys))
+		|| !TestEqual("string value 2", AttributeData.StringAttribute.GetValueFromItemKey(2), *AttributeData.TempStrings))
+	{
+		return false;
+	}
+
+	float InvalidFloat = 2.147f;
+
+	if (!TestFalse("Set float attribute outside range is invalid", AttributeData.FloatAccessor->Set(InvalidFloat, /*Index=*/ 3, Keys)))
+	{
+		return false;
+	}
+
+	return true;
+}
+
+bool FPCGAttributeAccessorSingleInvalidKeySetTest::RunTest(const FString& Parameters)
+{
+	FPCGAttributeAccessorTestHelpers::AttributeData AttributeData;
+	FPCGAttributeAccessorKeysEntries DefaultEntry(PCGInvalidEntryKey);
+
+	float NewDefaultFloat = 0.5f;
+
+	if (!TestTrue("Set new float default value, without flag", AttributeData.FloatAccessor->Set(NewDefaultFloat, DefaultEntry)
+		|| !TestEqual("Default float still default", AttributeData.FloatAttribute.GetValue(PCGDefaultValueKey), 0.0f))
+		|| !TestEqual("New entry added", AttributeData.PointData->Metadata->GetItemCountForChild(), FPCGAttributeAccessorTestHelpers::ValuesSize + 1)
+		|| !TestEqual("New float value", AttributeData.FloatAttribute.GetValueFromItemKey(FPCGAttributeAccessorTestHelpers::ValuesSize), NewDefaultFloat))
+	{
+		return false;
+	}
+
+	return true;
+}
+
+bool FPCGAttributeAccessorSetRangeTest::RunTest(const FString& Parameters)
+{
+	FPCGAttributeAccessorTestHelpers::AttributeData AttributeData;
+	FPCGAttributeAccessorKeysEntries Keys(&AttributeData.FloatAttribute);
+
+	const int32 Index = 0;
+
+	if (!TestTrue("Set float attribute", AttributeData.FloatAccessor->SetRange(FPCGAttributeAccessorTestHelpers::MakeConstArrayView(AttributeData.TempFloats), Index, Keys))
+		|| !TestEqual("float value 0", AttributeData.FloatAttribute.GetValueFromItemKey(0), AttributeData.TempFloats[0])
+		|| !TestEqual("float value 1", AttributeData.FloatAttribute.GetValueFromItemKey(1), AttributeData.TempFloats[1])
+		|| !TestEqual("float value 2", AttributeData.FloatAttribute.GetValueFromItemKey(2), AttributeData.TempFloats[2]))
+	{
+		return false;
+	}
+
+	if (!TestTrue("Set vector attribute", AttributeData.VectorAccessor->SetRange(FPCGAttributeAccessorTestHelpers::MakeConstArrayView(AttributeData.TempVectors), Index, Keys))
+		|| !TestEqual("vector value 0", AttributeData.VectorAttribute.GetValueFromItemKey(0), AttributeData.TempVectors[0])
+		|| !TestEqual("vector value 1", AttributeData.VectorAttribute.GetValueFromItemKey(1), AttributeData.TempVectors[1])
+		|| !TestEqual("vector value 2", AttributeData.VectorAttribute.GetValueFromItemKey(2), AttributeData.TempVectors[2]))
+	{
+		return false;
+	}
+
+	if (!TestTrue("Set string attribute", AttributeData.StringAccessor->SetRange(FPCGAttributeAccessorTestHelpers::MakeConstArrayView(AttributeData.TempStrings), Index, Keys))
+		|| !TestEqual("string value 0", AttributeData.StringAttribute.GetValueFromItemKey(0), AttributeData.TempStrings[0])
+		|| !TestEqual("string value 1", AttributeData.StringAttribute.GetValueFromItemKey(1), AttributeData.TempStrings[1])
+		|| !TestEqual("string value 2", AttributeData.StringAttribute.GetValueFromItemKey(2), AttributeData.TempStrings[2]))
+	{
+		return false;
+	}
+
+	return true;
+}
+
+bool FPCGAttributeAccessorSetRangeOutsideRangeTest::RunTest(const FString& Parameters)
+{
+	FPCGAttributeAccessorTestHelpers::AttributeData AttributeData;
+	FPCGAttributeAccessorKeysEntries Keys(&AttributeData.FloatAttribute);
+
+	const int32 Index = 0;
+
+	float TooManyFloats[FPCGAttributeAccessorTestHelpers::ValuesSize + 1] = { 1.0f, 2.0f, 3.0f, 4.0f };
+
+	if (!TestFalse("Set too many float values", AttributeData.FloatAccessor->SetRange(FPCGAttributeAccessorTestHelpers::MakeConstArrayView(TooManyFloats), Index, Keys)))
+	{
+		return false;
+	}
+
+	if (!TestFalse("Set float attribute with index offset outside range is not valid", AttributeData.FloatAccessor->SetRange(FPCGAttributeAccessorTestHelpers::MakeConstArrayView(AttributeData.TempFloats), Index + 999, Keys)))
+	{
+		return false;
+	}
+
+	if (!TestFalse("Set float attribute with index offset in valid range, but number of values too high is not valid", AttributeData.FloatAccessor->SetRange(FPCGAttributeAccessorTestHelpers::MakeConstArrayView(AttributeData.TempFloats), Index + 1, Keys)))
+	{
+		return false;
+	}
+
+	return true;
+}
+
+// Testing only floats for points and in range, as the only thing that differ is the keys
+
+bool FPCGAttributeAccessorSingleGetPointsTest::RunTest(const FString& Parameters)
+{
+	FPCGAttributeAccessorTestHelpers::AttributeData AttributeData;
+	const TArray<FPCGPoint>& Points = AttributeData.PointData->GetPoints();
+	FPCGAttributeAccessorKeysPoints Keys(Points);
+
+	if (!TestEqual("Right number of keys", Keys.GetNum(), Points.Num()))
+	{
+		return false;
+	}
+
+	if (!TestTrue("Keys are read only", Keys.IsReadOnly()))
+	{
+		return false;
+	}
+
+	for (int32 i = 0; i < 4; ++i)
+	{
+		if (!TestTrue(FString::Printf(TEXT("Get Float value %d"), i), AttributeData.FloatAccessor->Get(*AttributeData.TempFloats, i, Keys))
+			|| !TestEqual(FString::Printf(TEXT("Float value %d"), i), *AttributeData.TempFloats, AttributeData.FloatAttribute.GetValueFromItemKey(Points[i].MetadataEntry)))
+		{
+			return false;
+		}
+	}
+
+	return true;
+}
+
+bool FPCGAttributeAccessorGetRangePointsTest::RunTest(const FString& Parameters)
+{
+	FPCGAttributeAccessorTestHelpers::AttributeData AttributeData;
+	const TArray<FPCGPoint>& Points = AttributeData.PointData->GetPoints();
+	FPCGAttributeAccessorKeysPoints Keys(Points);
+
+	if (!TestTrue("Keys are read only", Keys.IsReadOnly()))
+	{
+		return false;
+	}
+
+	float TempFloats[4] = { 1.0f, 2.0f, 3.0f, 4.0f };
+	const int32 Index = 0;
+
+	// Points entry are in reverse order, and last one is the default value.
+	if (!TestTrue("Get float attribute full range", AttributeData.FloatAccessor->GetRange(FPCGAttributeAccessorTestHelpers::MakeArrayView(TempFloats), Index, Keys))
+		|| !TestEqual("float value 0", TempFloats[0], FPCGAttributeAccessorTestHelpers::FloatValues[2])
+		|| !TestEqual("float value 1", TempFloats[1], FPCGAttributeAccessorTestHelpers::FloatValues[1])
+		|| !TestEqual("float value 2", TempFloats[2], FPCGAttributeAccessorTestHelpers::FloatValues[0])
+		|| !TestEqual("float value 3", TempFloats[3], 0.0f))
+	{
+		return false;
+	}
+
+	return true;
+}
+
+bool FPCGAttributeAccessorSingleSetPointsTest::RunTest(const FString& Parameters)
+{
+	FPCGAttributeAccessorTestHelpers::AttributeData AttributeData;
+	TArray<FPCGPoint>& Points = AttributeData.PointData->GetMutablePoints();
+	FPCGAttributeAccessorKeysPoints Keys(TArrayView<FPCGPoint>(Points.GetData(), Points.Num()));
+
+	if (!TestEqual("Right number of keys", Keys.GetNum(), Points.Num()))
+	{
+		return false;
+	}
+
+	if (!TestFalse("Keys are not read only", Keys.IsReadOnly()))
+	{
+		return false;
+	}
+
+	float TempFloats[4] = { 1.0f, 2.0f, 3.0f, 4.0f };
+	const int32 Index = 0;
+
+	for (int32 i = 0; i < 4; ++i)
+	{
+		if (!TestTrue(FString::Printf(TEXT("Set Float value %d"), i), AttributeData.FloatAccessor->Set(TempFloats[i], i, Keys))
+			|| !TestEqual(FString::Printf(TEXT("Float value %d"), i), AttributeData.FloatAttribute.GetValueFromItemKey(Points[i].MetadataEntry), TempFloats[i]))
+		{
+			return false;
+		}
+	}
+
+	// Points should not have a different entry key, if it was not invalid in the first place
+	if (!TestEqual("Point 0 should have entry 2", Points[0].MetadataEntry, PCGMetadataEntryKey(2))
+		|| !TestEqual("Point 1 should have entry 1", Points[1].MetadataEntry, PCGMetadataEntryKey(1))
+		|| !TestEqual("Point 2 should have entry 0", Points[2].MetadataEntry, PCGMetadataEntryKey(0)))
+	{
+		return false;
+	}
+
+	// Also last point was a default entry. Making sure that we have one more entry and was set correctly
+	if (!TestEqual("Metadata should have one more entry", AttributeData.PointData->Metadata->GetItemCountForChild(), FPCGAttributeAccessorTestHelpers::ValuesSize + 1))
+	{
+		return false;
+	}
+
+	if (!TestEqual("Last point should have a valid entry", Points[3].MetadataEntry, PCGMetadataEntryKey(3)))
+	{
+		return false;
+	}
+
+	return true;
+}
+
+bool FPCGAttributeAccessorSetRangePointsTest::RunTest(const FString& Parameters)
+{
+	FPCGAttributeAccessorTestHelpers::AttributeData AttributeData;
+	TArray<FPCGPoint>& Points = AttributeData.PointData->GetMutablePoints();
+	FPCGAttributeAccessorKeysPoints Keys(TArrayView<FPCGPoint>(Points.GetData(), Points.Num()));
+
+	float TempFloats[4] = { 1.0f, 2.0f, 3.0f, 4.0f };
+	const int32 Index = 0;
+
+	if (!TestTrue("Set Float attribute full range", AttributeData.FloatAccessor->SetRange(FPCGAttributeAccessorTestHelpers::MakeConstArrayView(TempFloats), Index, Keys))
+		|| !TestEqual("Float value 0", AttributeData.FloatAttribute.GetValueFromItemKey(Points[0].MetadataEntry), TempFloats[0])
+		|| !TestEqual("Float value 1", AttributeData.FloatAttribute.GetValueFromItemKey(Points[1].MetadataEntry), TempFloats[1])
+		|| !TestEqual("Float value 2", AttributeData.FloatAttribute.GetValueFromItemKey(Points[2].MetadataEntry), TempFloats[2])
+		|| !TestEqual("Float value 3", AttributeData.FloatAttribute.GetValueFromItemKey(Points[3].MetadataEntry), TempFloats[3]))
+	{
+		return false;
+	}
+
+	// Points should not have a different entry key, if it was not invalid in the first place
+	if (!TestEqual("Point 0 should have entry 2", Points[0].MetadataEntry, PCGMetadataEntryKey(2))
+		|| !TestEqual("Point 1 should have entry 1", Points[1].MetadataEntry, PCGMetadataEntryKey(1))
+		|| !TestEqual("Point 2 should have entry 0", Points[2].MetadataEntry, PCGMetadataEntryKey(0)))
+	{
+		return false;
+	}
+
+	// Also last point was a default entry. Making sure that we have one more entry and was set correctly
+	if (!TestEqual("Metadata should have one more entry", AttributeData.PointData->Metadata->GetItemCountForChild(), FPCGAttributeAccessorTestHelpers::ValuesSize + 1))
+	{
+		return false;
+	}
+
+	if (!TestEqual("Last point should have a valid entry", Points[3].MetadataEntry, PCGMetadataEntryKey(3)))
+	{
+		return false;
+	}
+
+	return true;
+}
+
+#endif // WITH_EDITOR
\ No newline at end of file
diff --git a/Engine/Plugins/Experimental/PCG/Source/PCG/Public/Metadata/Accessors/IPCGAttributeAccessorTpl.h b/Engine/Plugins/Experimental/PCG/Source/PCG/Public/Metadata/Accessors/IPCGAttributeAccessorTpl.h
--- a/Engine/Plugins/Experimental/PCG/Source/PCG/Public/Metadata/Accessors/IPCGAttributeAccessorTpl.h
+++ b/Engine/Plugins/Experimental/PCG/Source/PCG/Public/Metadata/Accessors/IPCGAttributeAccessorTpl.h
@@ -92,8 +92,8 @@ private:
 	{
 		TRACE_CPUPROFILER_EVENT_SCOPE(IPCGAttributeAccessorT::InternalSetRange);
 
-		// Can't set if read only, and discard any set outside the keys range.
-		if (bReadOnly || Index >= Keys.GetNum())
+		// Can't set if read only, and discard any set if there are too many values.
+		if (bReadOnly || (Index + InValues.Num()) > Keys.GetNum())
 		{
 			return false;
 		}
diff --git a/Engine/Plugins/Experimental/PCG/Source/PCG/Public/Metadata/Accessors/PCGAttributeAccessorKeys.h b/Engine/Plugins/Experimental/PCG/Source/PCG/Public/Metadata/Accessors/PCGAttributeAccessorKeys.h
--- a/Engine/Plugins/Experimental/PCG/Source/PCG/Public/Metadata/Accessors/PCGAttributeAccessorKeys.h
+++ b/Engine/Plugins/Experimental/PCG/Source/PCG/Public/Metadata/Accessors/PCGAttributeAccessorKeys.h
@@ -115,6 +115,8 @@ public:
 	*/
 	virtual int32 GetNum() const = 0;
 
+	bool IsReadOnly() const { return bIsReadOnly; }
+
 protected:
 	virtual bool GetPointKeys(int32 InStart, TArrayView<FPCGPoint*>& OutPoints) { return false; }
 	virtual bool GetPointKeys(int32 InStart, TArrayView<const FPCGPoint*>& OutPoints) const { return false; }
