Tentative fix for unity build

#rnx

diff --git a/Engine/Plugins/Experimental/PCG/Source/PCG/Private/Tests/Accessors/PCGAttributeAccessorTest.cpp b/Engine/Plugins/Experimental/PCG/Source/PCG/Private/Tests/Accessors/PCGAttributeAccessorTest.cpp
--- a/Engine/Plugins/Experimental/PCG/Source/PCG/Private/Tests/Accessors/PCGAttributeAccessorTest.cpp
+++ b/Engine/Plugins/Experimental/PCG/Source/PCG/Private/Tests/Accessors/PCGAttributeAccessorTest.cpp
@@ -9,7 +9,9 @@
 #include "Metadata/Accessors/PCGAttributeAccessor.h"
 #include "Metadata/Accessors/PCGAttributeAccessorKeys.h"
 
+#include "UObject/Package.h"
 #include "UObject/UnrealType.h"
+#include "UObject/UObjectGlobals.h"
 
 #if WITH_EDITOR
 
