Cleaner dependency management for metadata

diff --git a/Engine/Plugins/Experimental/PCG/Source/PCG/Private/Metadata/PCGMetadata.cpp b/Engine/Plugins/Experimental/PCG/Source/PCG/Private/Metadata/PCGMetadata.cpp
--- a/Engine/Plugins/Experimental/PCG/Source/PCG/Private/Metadata/PCGMetadata.cpp
+++ b/Engine/Plugins/Experimental/PCG/Source/PCG/Private/Metadata/PCGMetadata.cpp
@@ -59,7 +59,7 @@ void UPCGMetadata::Serialize(FArchive& InArchive)
 		// Therefore by construction, it should never be less than the number of attributes (but can be greater).
 		NextAttributeId = MaxAttributeId + 1;
 		check(NextAttributeId >= Attributes.Num());
-		ItemKeyOffset = (Parent.IsValid() ? Parent->GetItemCountForChild() : 0);
+		ItemKeyOffset = (Parent ? Parent->GetItemCountForChild() : 0);
 	}
 }
 
@@ -86,14 +86,14 @@ void UPCGMetadata::Initialize(const UPCGMetadata* InParent, bool bAddAttributesF
 
 void UPCGMetadata::InitializeWithAttributeFilter(const UPCGMetadata* InParent, const TSet<FName>& InFilteredAttributes, EPCGMetadataFilterMode InFilterMode)
 {
-	if (Parent.IsValid() || Attributes.Num() != 0)
+	if (Parent || Attributes.Num() != 0)
 	{
 		// Already initialized; note that while that might be construed as a warning, there are legit cases where this is correct
 		return;
 	}
 
 	Parent = ((InParent != this) ? InParent : nullptr);
-	ItemKeyOffset = Parent.IsValid() ? Parent->GetItemCountForChild() : 0;
+	ItemKeyOffset = Parent ? Parent->GetItemCountForChild() : 0;
 
 	// If we have been given an include list which is empty, then don't bother adding any attributes
 	const bool bSkipAddingAttributesFromParent = (InFilterMode == EPCGMetadataFilterMode::IncludeAttributes) && (InFilteredAttributes.Num() == 0);
@@ -111,7 +111,7 @@ void UPCGMetadata::InitializeAsCopy(const UPCGMetadata* InMetadataToCopy)
 	}
 
 	check(InMetadataToCopy);
-	if (Parent.IsValid() || Attributes.Num() != 0)
+	if (Parent || Attributes.Num() != 0)
 	{
 		UE_LOG(LogPCG, Error, TEXT("Metadata has already been initialized or already contains attributes"));
 		return;
@@ -232,7 +232,7 @@ void UPCGMetadata::CopyAttribute(const UPCGMetadata* InOther, FName AttributeToC
 
 const UPCGMetadata* UPCGMetadata::GetRoot() const
 {
-	if (Parent.IsValid())
+	if (Parent)
 	{
 		return Parent->GetRoot();
 	}
@@ -394,7 +394,7 @@ FName UPCGMetadata::GetLatestAttributeNameOrNone() const
 
 bool UPCGMetadata::ParentHasAttribute(FName AttributeName) const
 {
-	return Parent.IsValid() && Parent->HasAttribute(AttributeName);
+	return Parent && Parent->HasAttribute(AttributeName);
 }
 
 void UPCGMetadata::CreateInteger32Attribute(FName AttributeName, int32 DefaultValue, bool bAllowsInterpolation, bool bOverrideParent)
@@ -545,7 +545,7 @@ FPCGMetadataAttributeBase* UPCGMetadata::CopyAttribute(FName AttributeToCopy, FN
 	}
 	AttributeLock.ReadUnlock();
 
-	if (!OriginalAttribute && Parent.IsValid())
+	if (!OriginalAttribute && Parent)
 	{
 		OriginalAttribute = Parent->GetConstAttribute(AttributeToCopy);
 	}
diff --git a/Engine/Plugins/Experimental/PCG/Source/PCG/Private/PCGData.cpp b/Engine/Plugins/Experimental/PCG/Source/PCG/Private/PCGData.cpp
--- a/Engine/Plugins/Experimental/PCG/Source/PCG/Private/PCGData.cpp
+++ b/Engine/Plugins/Experimental/PCG/Source/PCG/Private/PCGData.cpp
@@ -38,33 +38,6 @@ void FPCGRootSet::AddInternal(UObject* InObject)
 		InObject->AddToRoot();
 		RootSet.Emplace(InObject, 1);
 	}
-
-	// Recurse to outermost
-	UObject* OuterObject = InObject->GetOuter();
-	if (OuterObject && !OuterObject->IsA<UPackage>())
-	{
-		AddInternal(OuterObject);
-	}
-
-	// Recurse on metadata
-	UPCGMetadata* InMetadata = nullptr;
-	if (UPCGSpatialData* InSpatialData = Cast<UPCGSpatialData>(InObject))
-	{
-		InMetadata = InSpatialData->Metadata;
-	}
-	else if (UPCGParamData* InParamData = Cast<UPCGParamData>(InObject))
-	{
-		InMetadata = InParamData->Metadata;
-	}
-
-	if (InMetadata && InMetadata->GetParent())
-	{
-		UObject* MetadataParentOuter = InMetadata->GetParent()->GetOuter();
-		if (MetadataParentOuter && !MetadataParentOuter->IsA<UPackage>())
-		{
-			AddInternal(MetadataParentOuter);
-		}
-	}
 }
 
 void FPCGRootSet::Remove(UObject* InObject)
@@ -93,34 +66,6 @@ void FPCGRootSet::RemoveInternal(UObject* InObject)
 			RootSet.Remove(InObject);
 		}
 	}
-
-	// Recurse to outermost
-	UObject* OuterObject = InObject->GetOuter();
-	if (OuterObject && !OuterObject->IsA<UPackage>())
-	{
-		RemoveInternal(OuterObject);
-	}
-
-	// Recurse on metadata
-	UPCGMetadata* InMetadata = nullptr;
-
-	if (UPCGSpatialData* InSpatialData = Cast<UPCGSpatialData>(InObject))
-	{
-		InMetadata = InSpatialData->Metadata;
-	}
-	else if (UPCGParamData* InParamData = Cast<UPCGParamData>(InObject))
-	{
-		InMetadata = InParamData->Metadata;
-	}
-
-	if (InMetadata && InMetadata->GetParent())
-	{
-		UObject* MetadataParentOuter = InMetadata->GetParent()->GetOuter();
-		if (MetadataParentOuter && !MetadataParentOuter->IsA<UPackage>())
-		{
-			RemoveInternal(MetadataParentOuter);
-		}
-	}
 }
 
 bool FPCGTaggedData::operator==(const FPCGTaggedData& Other) const
diff --git a/Engine/Plugins/Experimental/PCG/Source/PCG/Public/Data/PCGPrimitiveData.h b/Engine/Plugins/Experimental/PCG/Source/PCG/Public/Data/PCGPrimitiveData.h
--- a/Engine/Plugins/Experimental/PCG/Source/PCG/Public/Data/PCGPrimitiveData.h
+++ b/Engine/Plugins/Experimental/PCG/Source/PCG/Public/Data/PCGPrimitiveData.h
@@ -39,7 +39,7 @@ public:
 
 protected:
 	UPROPERTY(BlueprintReadOnly, VisibleAnywhere, Category = "Data")
-	TObjectPtr<UPrimitiveComponent> Primitive = nullptr;
+	TWeakObjectPtr<UPrimitiveComponent> Primitive = nullptr;
 
 	UPROPERTY()
 	FBox CachedBounds = FBox(EForceInit::ForceInit);
diff --git a/Engine/Plugins/Experimental/PCG/Source/PCG/Public/Metadata/PCGMetadata.h b/Engine/Plugins/Experimental/PCG/Source/PCG/Public/Metadata/PCGMetadata.h
--- a/Engine/Plugins/Experimental/PCG/Source/PCG/Public/Metadata/PCGMetadata.h
+++ b/Engine/Plugins/Experimental/PCG/Source/PCG/Public/Metadata/PCGMetadata.h
@@ -261,7 +261,7 @@ protected:
 	void RemoveAttributeInternal(FName AttributeName);
 
 	UPROPERTY()
-	TWeakObjectPtr<const UPCGMetadata> Parent;
+	TObjectPtr<const UPCGMetadata> Parent;
 
 	// Set of parents kept for streams relationship and GC collection
 	// But otherwise not used directly
@@ -291,7 +291,7 @@ FPCGMetadataAttribute<T>* UPCGMetadata::CreateAttribute(FName AttributeName, con
 
 	const FPCGMetadataAttributeBase* ParentAttribute = nullptr;
 
-	if (bOverrideParent && Parent.IsValid())
+	if (bOverrideParent && Parent)
 	{
 		ParentAttribute = Parent->GetConstAttribute(AttributeName);
 	}
