[PCG] Warn if no surface could be found on either input pin.

Only emits warning is node has at least one connection, to avoid it spamming log when new node created or similar events.

diff --git a/Engine/Plugins/Experimental/PCG/Source/PCG/Private/Elements/PCGSurfaceSampler.cpp b/Engine/Plugins/Experimental/PCG/Source/PCG/Private/Elements/PCGSurfaceSampler.cpp
--- a/Engine/Plugins/Experimental/PCG/Source/PCG/Private/Elements/PCGSurfaceSampler.cpp
+++ b/Engine/Plugins/Experimental/PCG/Source/PCG/Private/Elements/PCGSurfaceSampler.cpp
@@ -318,11 +318,9 @@ bool FPCGSurfaceSamplerElement::ExecuteInternal(FPCGContext* Context) const
 	}
 
 	// If no shapes were obtained from the first input pin, try to find a shape to sample from nodes connected to the second pin.
-	const UPCGSpatialData* GeneratorFromBoundingShapeInput = nullptr;
 	if (GeneratingShapes.Num() == 0 && BoundingShapeSpatialInput)
 	{
-		GeneratorFromBoundingShapeInput = BoundingShapeSpatialInput->FindShapeFromNetwork(/*InDimension=*/2);
-		if (GeneratorFromBoundingShapeInput)
+		if (const UPCGSpatialData* GeneratorFromBoundingShapeInput = BoundingShapeSpatialInput->FindShapeFromNetwork(/*InDimension=*/2))
 		{
 			GeneratingShapes.Add(GeneratorFromBoundingShapeInput);
 
@@ -331,6 +329,12 @@ bool FPCGSurfaceSamplerElement::ExecuteInternal(FPCGContext* Context) const
 		}
 	}
 
+	// Warn if something is connected but no shape could be obtained for sampling
+	if (GeneratingShapes.Num() == 0 && (BoundingShapeInputs.Num() > 0 || SurfaceInputs.Num() > 0))
+	{
+		PCGE_LOG(Warning, "No Surface input was provided, and no surface could be found in the Bounding Shape input for sampling. Connect the surface to be sampled to the Surface input.");
+	}
+
 	UPCGParamData* Params = Context->InputData.GetParams();
 
 	// Early out on invalid settings
