Added point metadata validation in the element, turned on with pcg.debug.validatepointmetadata

diff --git a/Engine/Plugins/Experimental/PCG/Source/PCG/Private/PCGElement.cpp b/Engine/Plugins/Experimental/PCG/Source/PCG/Private/PCGElement.cpp
--- a/Engine/Plugins/Experimental/PCG/Source/PCG/Private/PCGElement.cpp
+++ b/Engine/Plugins/Experimental/PCG/Source/PCG/Private/PCGElement.cpp
@@ -7,12 +7,19 @@
 #include "PCGModule.h"
 #include "PCGParamData.h"
 #include "PCGSettings.h"
+#include "Data/PCGPointData.h"
 #include "Elements/PCGDebugElement.h"
 #include "Elements/PCGSelfPruning.h"
 #include "Graph/PCGGraphCache.h"
+#include "HAL/IConsoleManager.h"
 
 #include "Algo/Find.h"
 
+static TAutoConsoleVariable<bool> CVarPCGValidatePointMetadata(
+	TEXT("pcg.debug.ValidatePointMetadata"),
+	true,
+	TEXT("Controls whether we validate that the metadata entry keys on the output point data are consistent"));
+
 bool IPCGElement::Execute(FPCGContext* Context) const
 {
 	TRACE_CPUPROFILER_EVENT_SCOPE(IPCGElement::Execute);
@@ -371,6 +378,27 @@ void IPCGElement::CleanupAndValidateOutput(FPCGContext* Context) const
 				{
 					PCGE_LOG(Warning, "Output generated for pin %s does not have a compatible type: %s", *TaggedData.Pin.ToString(), *UEnum::GetValueAsString(TaggedData.Data->GetDataType()));
 				}
+
+				if (CVarPCGValidatePointMetadata.GetValueOnAnyThread())
+				{
+					if (UPCGPointData* PointData = Cast<UPCGPointData>(TaggedData.Data))
+					{
+						const TArray<FPCGPoint>& NewPoints = PointData->GetPoints();
+						const int32 MaxMetadataEntry = PointData->Metadata ? PointData->Metadata->GetItemCountForChild() : 0;
+
+						bool bHasError = false;
+
+						for(int32 PointIndex = 0; PointIndex < NewPoints.Num() && !bHasError; ++PointIndex)
+						{
+							bHasError |= (NewPoints[PointIndex].MetadataEntry >= MaxMetadataEntry);
+						}
+
+						if (bHasError)
+						{
+							PCGE_LOG(Warning, "Output generated for pin %s does not have valid point metadata", *TaggedData.Pin.ToString());
+						}
+					}
+				}
 			}
 		}
 #endif
