Added metadata flattening to support independent serialization (needed to store results on components)

diff --git a/Engine/Plugins/Experimental/PCG/Source/PCG/Private/Metadata/PCGMetadata.cpp b/Engine/Plugins/Experimental/PCG/Source/PCG/Private/Metadata/PCGMetadata.cpp
--- a/Engine/Plugins/Experimental/PCG/Source/PCG/Private/Metadata/PCGMetadata.cpp
+++ b/Engine/Plugins/Experimental/PCG/Source/PCG/Private/Metadata/PCGMetadata.cpp
@@ -258,6 +258,44 @@ bool UPCGMetadata::HasParent(const UPCGMetadata* InTentativeParent) const
 	return HierarchicalParent == InTentativeParent;
 }
 
+void UPCGMetadata::Flatten()
+{
+	Modify();
+
+	const int32 NumEntries = GetItemCountForChild();
+
+	AttributeLock.WriteLock();
+	for (auto& AttributePair : Attributes)
+	{
+		FPCGMetadataAttributeBase* Attribute = AttributePair.Value;
+		check(Attribute);
+
+		// For all stored entries (from the root), we need to make sure that entries that should have a concrete value have it
+		// Optimization notes:
+		// - we could skip entries that existed prior to attribute existence, etc.
+		// - we could skip entries that have no parent, but that would require checking against the parent entries in the parent hierarchy
+		for (int64 EntryKey = 0; EntryKey < NumEntries; ++EntryKey)
+		{
+			// Get value using value inheritance as expected
+			PCGMetadataValueKey ValueKey = Attribute->GetValueKey(EntryKey);
+			if (ValueKey != PCGDefaultValueKey)
+			{
+				// Set concrete non-default value
+				Attribute->SetValueFromValueKey(EntryKey, ValueKey);
+			}
+		}
+
+		// Finally, flatten values
+		Attribute->Flatten();
+	}
+	AttributeLock.WriteUnlock();
+
+	Parent = nullptr;
+	ParentKeys.Reset();
+	ParentKeys.Init(PCGInvalidEntryKey, NumEntries);
+	ItemKeyOffset = 0;
+}
+
 void UPCGMetadata::AddAttributeInternal(FName AttributeName, FPCGMetadataAttributeBase* Attribute)
 {
 	// This call assumes we have a write lock on the attribute map.
diff --git a/Engine/Plugins/Experimental/PCG/Source/PCG/Private/Metadata/PCGMetadataAttribute.cpp b/Engine/Plugins/Experimental/PCG/Source/PCG/Private/Metadata/PCGMetadataAttribute.cpp
--- a/Engine/Plugins/Experimental/PCG/Source/PCG/Private/Metadata/PCGMetadataAttribute.cpp
+++ b/Engine/Plugins/Experimental/PCG/Source/PCG/Private/Metadata/PCGMetadataAttribute.cpp
@@ -19,9 +19,14 @@ void FPCGMetadataAttributeBase::Serialize(UPCGMetadata* InMetadata, FArchive& In
 	int32 ParentAttributeId = (Parent ? Parent->AttributeId : -1);
 	InArchive << ParentAttributeId;
 
-	if (InArchive.IsLoading() && ParentAttributeId >= 0 && Metadata->GetParent())
+	if (InArchive.IsLoading())
 	{
-		Parent = Metadata->GetParent()->GetConstAttributeById(ParentAttributeId);
+		ensure(ParentAttributeId < 0 || Metadata->GetParent());
+		if (ParentAttributeId >= 0 && Metadata->GetParent())
+		{
+			Parent = Metadata->GetParent()->GetConstAttributeById(ParentAttributeId);
+			check(Parent);
+		}
 	}
 
 	//Type id should already be known by then, so no need to serialize it
diff --git a/Engine/Plugins/Experimental/PCG/Source/PCG/Private/PCGComponent.cpp b/Engine/Plugins/Experimental/PCG/Source/PCG/Private/PCGComponent.cpp
--- a/Engine/Plugins/Experimental/PCG/Source/PCG/Private/PCGComponent.cpp
+++ b/Engine/Plugins/Experimental/PCG/Source/PCG/Private/PCGComponent.cpp
@@ -6,6 +6,7 @@
 #include "PCGGraph.h"
 #include "PCGHelpers.h"
 #include "PCGInputOutputSettings.h"
+#include "PCGParamData.h"
 #include "PCGVolume.h"
 #include "PCGManagedResource.h"
 #include "Data/PCGDifferenceData.h"
@@ -366,6 +367,22 @@ void UPCGComponent::PostProcessGraph(const FBox& InNewBounds, bool bInGenerated,
 				// TODO: outering the first layer might not be sufficient here - might need to expose
 				// some methods in the data to traverse all the data to outer everything for serialization
 				DuplicatedTaggedData.Data = Cast<UPCGData>(StaticDuplicateObject(TaggedData.Data, this));
+
+				UPCGMetadata* DuplicatedMetadata = nullptr;
+				if (UPCGSpatialData* DuplicatedSpatialData = Cast<UPCGSpatialData>(DuplicatedTaggedData.Data))
+				{
+					DuplicatedMetadata = DuplicatedSpatialData->Metadata;
+				}
+				else if (UPCGParamData* DuplicatedParamData = Cast<UPCGParamData>(DuplicatedTaggedData.Data))
+				{
+					DuplicatedMetadata = DuplicatedParamData->Metadata;
+				}
+
+				// Make sure the metadata can be serialized independently
+				if (DuplicatedMetadata)
+				{
+					DuplicatedMetadata->Flatten();
+				}
 			}
 
 			// If the original component is partitioned, local components have to forward
diff --git a/Engine/Plugins/Experimental/PCG/Source/PCG/Public/Metadata/PCGMetadata.h b/Engine/Plugins/Experimental/PCG/Source/PCG/Public/Metadata/PCGMetadata.h
--- a/Engine/Plugins/Experimental/PCG/Source/PCG/Public/Metadata/PCGMetadata.h
+++ b/Engine/Plugins/Experimental/PCG/Source/PCG/Public/Metadata/PCGMetadata.h
@@ -78,6 +78,9 @@ public:
 	const UPCGMetadata* GetRoot() const;
 	bool HasParent(const UPCGMetadata* InTentativeParent) const;
 
+	/** Unparents current metadata by flattening the attributes (values, entries, etc.) */
+	void Flatten();
+
 	/** Create new streams */
 	UFUNCTION(BlueprintCallable, Category = "PCG|Metadata")
 	void CreateInteger32Attribute(FName AttributeName, int32 DefaultValue, bool bAllowsInterpolation, bool bOverrideParent = true);
diff --git a/Engine/Plugins/Experimental/PCG/Source/PCG/Public/Metadata/PCGMetadataAttribute.h b/Engine/Plugins/Experimental/PCG/Source/PCG/Public/Metadata/PCGMetadataAttribute.h
--- a/Engine/Plugins/Experimental/PCG/Source/PCG/Public/Metadata/PCGMetadataAttribute.h
+++ b/Engine/Plugins/Experimental/PCG/Source/PCG/Public/Metadata/PCGMetadataAttribute.h
@@ -19,6 +19,8 @@ public:
 	virtual ~FPCGMetadataAttributeBase() = default;
 	virtual void Serialize(UPCGMetadata* InMetadata, FArchive& InArchive);
 
+	virtual void Flatten() = 0;
+
 	const UPCGMetadata* GetMetadata() const { return Metadata; }
 	int16 GetTypeId() const { return TypeId; }
 
diff --git a/Engine/Plugins/Experimental/PCG/Source/PCG/Public/Metadata/PCGMetadataAttributeTpl.h b/Engine/Plugins/Experimental/PCG/Source/PCG/Public/Metadata/PCGMetadataAttributeTpl.h
--- a/Engine/Plugins/Experimental/PCG/Source/PCG/Public/Metadata/PCGMetadataAttributeTpl.h
+++ b/Engine/Plugins/Experimental/PCG/Source/PCG/Public/Metadata/PCGMetadataAttributeTpl.h
@@ -51,6 +51,43 @@ public:
 		}
 	}
 
+	virtual void Flatten() override
+	{
+		// Implementation notes:
+		// We don't need to flatten the EntryToValueKeyMap - this will have been taken care of in the metadata flatten
+
+		// Flatten values, from root to current attribute
+		if(Parent)
+		{
+			FWriteScopeLock ScopeLock(ValueLock);
+
+			TArray<const TArray<T>*> OriginalValues;
+			int32 ValueCount = 0;
+
+			const FPCGMetadataAttribute<T>* Current = static_cast<const FPCGMetadataAttribute<T>*>(this);
+			while (Current)
+			{
+				ValueCount += Current->Values.Num();
+				OriginalValues.Add(&Current->Values);
+				Current = static_cast<const FPCGMetadataAttribute<T>*>(Current->Parent);
+			}
+
+			TArray<T> FlattenedValues;
+			FlattenedValues.Reserve(ValueCount);
+
+			for (int32 ValuesIndex = OriginalValues.Num() - 1; ValuesIndex >= 0; --ValuesIndex)
+			{
+				FlattenedValues.Append(*OriginalValues[ValuesIndex]);
+			}
+
+			Values = MoveTemp(FlattenedValues);
+		}
+		
+		// Reset value offset, and lose parent
+		ValueKeyOffset = 0;
+		Parent = nullptr;
+	}
+
 	const FPCGMetadataAttribute* GetParent() const { return static_cast<const FPCGMetadataAttribute*>(Parent); }
 
 	FPCGMetadataAttribute* TypedCopy(FName NewName, UPCGMetadata* InMetadata, bool bKeepParent, bool bCopyEntries = true, bool bCopyValues = true)
