[PCG] Change user-facing name of PCG plugin and massage description slightly to bring it more in line with other examples I browsed

diff --git a/Engine/Plugins/Experimental/PCG/PCG.uplugin b/Engine/Plugins/Experimental/PCG/PCG.uplugin
--- a/Engine/Plugins/Experimental/PCG/PCG.uplugin
+++ b/Engine/Plugins/Experimental/PCG/PCG.uplugin
@@ -2,8 +2,8 @@
 	"FileVersion": 3,
 	"Version": 8,
 	"VersionName": "0.1",
-	"FriendlyName": "PCG",
-	"Description": "Procedural generation framework for UE",
+	"FriendlyName": "Procedural Content Generation Framework",
+	"Description": "Visual scripting framework for procedurally populating worlds with content in editor and/or at run-time.",
 	"Category": "Editor",
 	"CreatedBy": "Epic Games, Inc.",
 	"CreatedByURL": "https://epicgames.com",
