Added object notification on generation/cleanup - can be used in other systems as needed, but we will use it internally as a dependency change

diff --git a/Engine/Plugins/Experimental/PCG/Source/PCG/Private/PCGComponent.cpp b/Engine/Plugins/Experimental/PCG/Source/PCG/Private/PCGComponent.cpp
--- a/Engine/Plugins/Experimental/PCG/Source/PCG/Private/PCGComponent.cpp
+++ b/Engine/Plugins/Experimental/PCG/Source/PCG/Private/PCGComponent.cpp
@@ -338,6 +338,8 @@ void UPCGComponent::PostProcessGraph(const FBox& InNewBounds, bool bInGenerated,
 {
 	LastGeneratedBounds = InNewBounds;
 
+	const bool bHadGeneratedOutputBefore = GeneratedGraphOutput.TaggedData.Num() > 0;
+
 	CleanupUnusedManagedResources();
 
 	GeneratedGraphOutput.Reset();
@@ -394,6 +396,19 @@ void UPCGComponent::PostProcessGraph(const FBox& InNewBounds, bool bInGenerated,
 			CallPostGenerateFunctions(Context);
 		}
 	}
+
+	// Trigger notification - will be used by other tracking mechanisms
+#if WITH_EDITOR
+	const bool bHasGeneratedOutputAfter = GeneratedGraphOutput.TaggedData.Num() > 0;
+
+	if (bHasGeneratedOutputAfter || bHadGeneratedOutputBefore)
+	{
+		FProperty* GeneratedOutputProperty = FindFProperty<FProperty>(UPCGComponent::StaticClass(), GET_MEMBER_NAME_CHECKED(UPCGComponent, GeneratedGraphOutput));
+		check(GeneratedOutputProperty);
+		FPropertyChangedEvent GeneratedOutputChangedEvent(GeneratedOutputProperty, EPropertyChangeType::ValueSet);
+		FCoreUObjectDelegates::OnObjectPropertyChanged.Broadcast(this, GeneratedOutputChangedEvent);
+	}
+#endif
 }
 
 void UPCGComponent::CallPostGenerateFunctions(FPCGContext* Context) const
@@ -453,11 +468,21 @@ void UPCGComponent::PostCleanupGraph()
 {
 	bGenerated = false;
 	CurrentCleanupTask = InvalidPCGTaskId;
+
+	const bool bHadGeneratedGraphOutput = GeneratedGraphOutput.TaggedData.Num() > 0;
 	GeneratedGraphOutput.Reset();
 
 #if WITH_EDITOR
 	OnPCGGraphCleanedDelegate.Broadcast(this);
 	bDirtyGenerated = false;
+
+	if (bHadGeneratedGraphOutput)
+	{
+		FProperty* GeneratedOutputProperty = FindFProperty<FProperty>(UPCGComponent::StaticClass(), GET_MEMBER_NAME_CHECKED(UPCGComponent, GeneratedGraphOutput));
+		check(GeneratedOutputProperty);
+		FPropertyChangedEvent GeneratedOutputChangedEvent(GeneratedOutputProperty, EPropertyChangeType::ValueSet);
+		FCoreUObjectDelegates::OnObjectPropertyChanged.Broadcast(this, GeneratedOutputChangedEvent);
+	}
 #endif
 }
 
