Added possibility to apply the transform on an attribute and not only on point transform.

diff --git a/Engine/Plugins/Experimental/PCG/Source/PCG/Private/Elements/PCGTransformPoints.cpp b/Engine/Plugins/Experimental/PCG/Source/PCG/Private/Elements/PCGTransformPoints.cpp
--- a/Engine/Plugins/Experimental/PCG/Source/PCG/Private/Elements/PCGTransformPoints.cpp
+++ b/Engine/Plugins/Experimental/PCG/Source/PCG/Private/Elements/PCGTransformPoints.cpp
@@ -3,6 +3,7 @@
 #include "Elements/PCGTransformPoints.h"
 
 #include "PCGHelpers.h"
+#include "Data/PCGPointData.h"
 #include "Helpers/PCGSettingsHelpers.h"
 
 #include UE_INLINE_GENERATED_CPP_BY_NAME(PCGTransformPoints)
@@ -46,6 +47,8 @@ bool FPCGTransformPointsElement::ExecuteInternal(FPCGContext* Context) const
 	TArray<FPCGTaggedData>& Outputs = Context->OutputData.TaggedData;
 
 	UPCGParamData* Params = Context->InputData.GetParams();
+	const bool bApplyToAttribute = PCG_GET_OVERRIDEN_VALUE(Settings, bApplyToAttribute, Params);
+	const FName AttributeName = PCG_GET_OVERRIDEN_VALUE(Settings, AttributeName, Params);
 	const FVector OffsetMin = PCG_GET_OVERRIDEN_VALUE(Settings, OffsetMin, Params);
 	const FVector OffsetMax = PCG_GET_OVERRIDEN_VALUE(Settings, OffsetMax, Params);
 	const bool bAbsoluteOffset = PCG_GET_OVERRIDEN_VALUE(Settings, bAbsoluteOffset, Params);
@@ -61,70 +64,165 @@ bool FPCGTransformPointsElement::ExecuteInternal(FPCGContext* Context) const
 	const int Seed = PCGSettingsHelpers::ComputeSeedWithOverride(Settings, Context->SourceComponent, Params);
 
 	// Use implicit capture, since we capture a lot
-	ProcessPoints(Context, Inputs, Outputs, [&](const FPCGPoint& InPoint, FPCGPoint& OutPoint)
+	//ProcessPoints(Context, Inputs, Outputs, [&](const FPCGPoint& InPoint, FPCGPoint& OutPoint)
+	for (const FPCGTaggedData& Input : Inputs)
 	{
-		OutPoint = InPoint;
+		TRACE_CPUPROFILER_EVENT_SCOPE(FPCGTransformPointsElement::Execute::InputLoop);
+		FPCGTaggedData& Output = Outputs.Add_GetRef(Input);
 
-		FRandomStream RandomSource(PCGHelpers::ComputeSeed(Seed, InPoint.Seed));
+		const UPCGSpatialData* SpatialData = Cast<UPCGSpatialData>(Input.Data);
 
-		const float OffsetX = RandomSource.FRandRange(OffsetMin.X, OffsetMax.X);
-		const float OffsetY = RandomSource.FRandRange(OffsetMin.Y, OffsetMax.Y);
-		const float OffsetZ = RandomSource.FRandRange(OffsetMin.Z, OffsetMax.Z);
-		const FVector RandomOffset(OffsetX, OffsetY, OffsetZ);
-
-		const float RotationX = RandomSource.FRandRange(RotationMin.Pitch, RotationMax.Pitch);
-		const float RotationY = RandomSource.FRandRange(RotationMin.Yaw, RotationMax.Yaw);
-		const float RotationZ = RandomSource.FRandRange(RotationMin.Roll, RotationMax.Roll);
-		const FQuat RandomRotation(FRotator(RotationX, RotationY, RotationZ).Quaternion());
-
-		FVector RandomScale;
-		if (bUniformScale)
-		{
-			RandomScale = FVector(RandomSource.FRandRange(ScaleMin.X, ScaleMax.X));
-		}
-		else
+		if (!SpatialData)
 		{
-			RandomScale.X = RandomSource.FRandRange(ScaleMin.X, ScaleMax.X);
-			RandomScale.Y = RandomSource.FRandRange(ScaleMin.Y, ScaleMax.Y);
-			RandomScale.Z = RandomSource.FRandRange(ScaleMin.Z, ScaleMax.Z);
+			PCGE_LOG(Error, "Unable to get SpatialData from input");
+			continue;
 		}
 
-		if (bAbsoluteOffset)
-		{
-			OutPoint.Transform.SetLocation(InPoint.Transform.GetLocation() + RandomOffset); 
-		}
-		else
-		{
-			const FTransform RotatedTransform(InPoint.Transform.GetRotation());
-			OutPoint.Transform.SetLocation(InPoint.Transform.GetLocation() + RotatedTransform.TransformPosition(RandomOffset)); 
-		}
+		const UPCGPointData* PointData = SpatialData->ToPointData(Context);
 
-		if (bAbsoluteRotation)
+		if (!PointData)
 		{
-			OutPoint.Transform.SetRotation(RandomRotation);
-		}
-		else
-		{
-			OutPoint.Transform.SetRotation(InPoint.Transform.GetRotation() * RandomRotation);
+			PCGE_LOG(Error, "Unable to get PointData from input");
+			continue;
 		}
 
-		if (bAbsoluteScale)
+		FName LocalAttributeName = AttributeName;
+		const FPCGMetadataAttribute<FTransform>* SourceAttribute = nullptr;
+
+		if (bApplyToAttribute)
 		{
-			OutPoint.Transform.SetScale3D(RandomScale);
+			const UPCGMetadata* PointMetadata = PointData->ConstMetadata();
+			check(PointMetadata);
+
+			if (LocalAttributeName == NAME_None)
+			{
+				LocalAttributeName = PointMetadata->GetLatestAttributeNameOrNone();
+			}
+
+			// Validate that the attribute has the proper type
+			const FPCGMetadataAttributeBase* FoundAttribute = PointMetadata->GetConstAttribute(LocalAttributeName);
+
+			if (!FoundAttribute || FoundAttribute->GetTypeId() != PCG::Private::MetadataTypes<FTransform>::Id)
+			{
+				PCGE_LOG(Error, "Attribute %s does not exist or is not a transform", *LocalAttributeName.ToString());
+				continue;
+			}
+
+			SourceAttribute = static_cast<const FPCGMetadataAttribute<FTransform>*>(FoundAttribute);
 		}
-		else
+
+		const TArray<FPCGPoint>& Points = PointData->GetPoints();
+
+		UPCGPointData* OutputData = NewObject<UPCGPointData>();
+		OutputData->InitializeFromData(PointData);
+		TArray<FPCGPoint>& OutputPoints = OutputData->GetMutablePoints();
+		Output.Data = OutputData;
+
+		FPCGMetadataAttribute<FTransform>* TargetAttribute = nullptr;
+		TArray<TTuple<int64, int64>> AllMetadataEntries;
+
+		if (bApplyToAttribute)
 		{
-			OutPoint.Transform.SetScale3D(InPoint.Transform.GetScale3D() * RandomScale);
+			check(SourceAttribute && OutputData && OutputData->Metadata);
+			TargetAttribute = OutputData->Metadata->GetMutableTypedAttribute<FTransform>(LocalAttributeName);
+			AllMetadataEntries.SetNum(Points.Num());
 		}
 
-		if (bRecomputeSeed)
+		FPCGAsync::AsyncPointProcessing(Context, Points.Num(), OutputPoints, [&](int32 Index, FPCGPoint& OutPoint)
 		{
-			const FVector& Position = OutPoint.Transform.GetLocation();
-			OutPoint.Seed = PCGHelpers::ComputeSeed((int)Position.X, (int)Position.Y, (int)Position.Z);
+			const FPCGPoint& InPoint = Points[Index];
+			OutPoint = InPoint;
+
+			FRandomStream RandomSource(PCGHelpers::ComputeSeed(Seed, InPoint.Seed));
+
+			const float OffsetX = RandomSource.FRandRange(OffsetMin.X, OffsetMax.X);
+			const float OffsetY = RandomSource.FRandRange(OffsetMin.Y, OffsetMax.Y);
+			const float OffsetZ = RandomSource.FRandRange(OffsetMin.Z, OffsetMax.Z);
+			const FVector RandomOffset(OffsetX, OffsetY, OffsetZ);
+
+			const float RotationX = RandomSource.FRandRange(RotationMin.Pitch, RotationMax.Pitch);
+			const float RotationY = RandomSource.FRandRange(RotationMin.Yaw, RotationMax.Yaw);
+			const float RotationZ = RandomSource.FRandRange(RotationMin.Roll, RotationMax.Roll);
+			const FQuat RandomRotation(FRotator(RotationX, RotationY, RotationZ).Quaternion());
+
+			FVector RandomScale;
+			if (bUniformScale)
+			{
+				RandomScale = FVector(RandomSource.FRandRange(ScaleMin.X, ScaleMax.X));
+			}
+			else
+			{
+				RandomScale.X = RandomSource.FRandRange(ScaleMin.X, ScaleMax.X);
+				RandomScale.Y = RandomSource.FRandRange(ScaleMin.Y, ScaleMax.Y);
+				RandomScale.Z = RandomSource.FRandRange(ScaleMin.Z, ScaleMax.Z);
+			}
+		
+			FTransform SourceTransform;
+
+			if (!bApplyToAttribute)
+			{
+				SourceTransform = InPoint.Transform;
+			}
+			else
+			{
+				SourceTransform = SourceAttribute->GetValue(InPoint.MetadataEntry);
+			}
+
+			FTransform FinalTransform = SourceTransform;
+
+			if (bAbsoluteOffset)
+			{
+				FinalTransform.SetLocation(SourceTransform.GetLocation() + RandomOffset); 
+			}
+			else
+			{
+				const FTransform RotatedTransform(SourceTransform.GetRotation());
+				FinalTransform.SetLocation(SourceTransform.GetLocation() + RotatedTransform.TransformPosition(RandomOffset)); 
+			}
+
+			if (bAbsoluteRotation)
+			{
+				FinalTransform.SetRotation(RandomRotation);
+			}
+			else
+			{
+				FinalTransform.SetRotation(SourceTransform.GetRotation() * RandomRotation);
+			}
+
+			if (bAbsoluteScale)
+			{
+				FinalTransform.SetScale3D(RandomScale);
+			}
+			else
+			{
+				FinalTransform.SetScale3D(SourceTransform.GetScale3D() * RandomScale);
+			}
+
+			if (!bApplyToAttribute)
+			{
+				OutPoint.Transform = FinalTransform;
+
+				if (bRecomputeSeed)
+				{
+					const FVector& Position = FinalTransform.GetLocation();
+					OutPoint.Seed = PCGHelpers::ComputeSeed((int)Position.X, (int)Position.Y, (int)Position.Z);
+				}
+			}
+			else
+			{
+				OutPoint.MetadataEntry = OutputData->Metadata->AddEntryPlaceholder();
+				AllMetadataEntries[Index] = MakeTuple(OutPoint.MetadataEntry, InPoint.MetadataEntry);
+				TargetAttribute->SetValue(OutPoint.MetadataEntry, FinalTransform);
+			}
+
+			return true;
+		});
+
+		if (TargetAttribute)
+		{
+			OutputData->Metadata->AddDelayedEntries(AllMetadataEntries);
 		}
-
-		return true;
-	});
+	}
 
 	// Forward any non-input data
 	Outputs.Append(Context->InputData.GetAllSettings());
diff --git a/Engine/Plugins/Experimental/PCG/Source/PCG/Public/Elements/PCGTransformPoints.h b/Engine/Plugins/Experimental/PCG/Source/PCG/Public/Elements/PCGTransformPoints.h
--- a/Engine/Plugins/Experimental/PCG/Source/PCG/Public/Elements/PCGTransformPoints.h
+++ b/Engine/Plugins/Experimental/PCG/Source/PCG/Public/Elements/PCGTransformPoints.h
@@ -31,6 +31,12 @@ protected:
 	//~End UPCGSettings interface
 
 public:
+	UPROPERTY(BlueprintReadWrite, EditAnywhere, Category = Settings)
+	bool bApplyToAttribute = false;
+
+	UPROPERTY(BlueprintReadWrite, EditAnywhere, Category = Settings, meta = (EditCondition="bApplyToAttribute"))
+	FName AttributeName;
+
 	UPROPERTY(BlueprintReadWrite, EditAnywhere, Category = Settings)
 	FVector OffsetMin = FVector::Zero();
 
@@ -66,11 +72,11 @@ public:
 	bool bUniformScale = true;
 
 	/** Recompute the seed for each new point using its new location */
-	UPROPERTY(BlueprintReadWrite, EditAnywhere, Category = Settings)
+	UPROPERTY(BlueprintReadWrite, EditAnywhere, Category = Settings, meta = (EditCondition = "!bApplyToAttribute"))
 	bool bRecomputeSeed = false;
 };
 
-class FPCGTransformPointsElement : public FPCGPointProcessingElementBase
+class FPCGTransformPointsElement : public FSimplePCGElement
 {
 protected:
 	virtual bool ExecuteInternal(FPCGContext* Context) const override;
