Fixed crash on deleted landscape

#rnx

diff --git a/Engine/Plugins/Experimental/PCG/Source/PCG/Private/Data/PCGLandscapeData.cpp b/Engine/Plugins/Experimental/PCG/Source/PCG/Private/Data/PCGLandscapeData.cpp
--- a/Engine/Plugins/Experimental/PCG/Source/PCG/Private/Data/PCGLandscapeData.cpp
+++ b/Engine/Plugins/Experimental/PCG/Source/PCG/Private/Data/PCGLandscapeData.cpp
@@ -122,7 +122,7 @@ bool UPCGLandscapeData::ProjectPoint(const FTransform& InTransform, const FBox&
 	}
 
 	const ULandscapeInfo* LandscapeInfo = GetLandscapeInfo(InTransform.GetLocation());
-	if (!LandscapeInfo)
+	if (!LandscapeInfo || !LandscapeInfo->GetLandscapeProxy())
 	{
 		return false;
 	}
@@ -205,6 +205,11 @@ const UPCGPointData* UPCGLandscapeData::CreatePointData(FPCGContext* Context, co
 	{
 		ULandscapeInfo* LandscapeInfo = LandscapeInfoPair.Value;
 
+		if (!LandscapeInfo || !LandscapeInfo->GetLandscapeProxy())
+		{
+			continue;
+		}
+
 		const FTransform& LandscapeTransform = LandscapeInfo->GetLandscapeProxy()->GetTransform();
 		const int32 ComponentSizeQuads = LandscapeInfo->ComponentSizeQuads;
 
