[PCG] Migrate existing params connections on old SurfaceSampler nodes to new params pin so it does not get lost.

Shoved this into same version as it does not make sense to add a new version - any up to date data will no longer have the params connection.

diff --git a/Engine/Plugins/Experimental/PCG/Source/PCG/Private/Elements/PCGDensityFilter.cpp b/Engine/Plugins/Experimental/PCG/Source/PCG/Private/Elements/PCGDensityFilter.cpp
--- a/Engine/Plugins/Experimental/PCG/Source/PCG/Private/Elements/PCGDensityFilter.cpp
+++ b/Engine/Plugins/Experimental/PCG/Source/PCG/Private/Elements/PCGDensityFilter.cpp
@@ -133,7 +133,7 @@ void UPCGDensityFilterSettings::ApplyDeprecationBeforeUpdatePins(UPCGNode* InOut
 
 	if (DataVersion < FPCGCustomVersion::MoveParamsOffFirstPinDensityNodes)
 	{
-		PCGSettingsHelpers::DeprecationBreakOutParamsToPin(InOutNode, InputPins, OutputPins);
+		PCGSettingsHelpers::DeprecationBreakOutParamsToNewPin(InOutNode, InputPins, OutputPins);
 	}
 }
 #endif // WITH_EDITOR
diff --git a/Engine/Plugins/Experimental/PCG/Source/PCG/Private/Elements/PCGDensityNoise.cpp b/Engine/Plugins/Experimental/PCG/Source/PCG/Private/Elements/PCGDensityNoise.cpp
--- a/Engine/Plugins/Experimental/PCG/Source/PCG/Private/Elements/PCGDensityNoise.cpp
+++ b/Engine/Plugins/Experimental/PCG/Source/PCG/Private/Elements/PCGDensityNoise.cpp
@@ -103,7 +103,7 @@ void UPCGDensityNoiseSettings::ApplyDeprecationBeforeUpdatePins(UPCGNode* InOutN
 
 	if (DataVersion < FPCGCustomVersion::MoveParamsOffFirstPinDensityNodes)
 	{
-		PCGSettingsHelpers::DeprecationBreakOutParamsToPin(InOutNode, InputPins, OutputPins);
+		PCGSettingsHelpers::DeprecationBreakOutParamsToNewPin(InOutNode, InputPins, OutputPins);
 	}
 }
 #endif // WITH_EDITOR
diff --git a/Engine/Plugins/Experimental/PCG/Source/PCG/Private/Elements/PCGDensityRemapElement.cpp b/Engine/Plugins/Experimental/PCG/Source/PCG/Private/Elements/PCGDensityRemapElement.cpp
--- a/Engine/Plugins/Experimental/PCG/Source/PCG/Private/Elements/PCGDensityRemapElement.cpp
+++ b/Engine/Plugins/Experimental/PCG/Source/PCG/Private/Elements/PCGDensityRemapElement.cpp
@@ -94,7 +94,7 @@ void UPCGDensityRemapSettings::ApplyDeprecationBeforeUpdatePins(UPCGNode* InOutN
 
 	if (DataVersion < FPCGCustomVersion::MoveParamsOffFirstPinDensityNodes)
 	{
-		PCGSettingsHelpers::DeprecationBreakOutParamsToPin(InOutNode, InputPins, OutputPins);
+		PCGSettingsHelpers::DeprecationBreakOutParamsToNewPin(InOutNode, InputPins, OutputPins);
 	}
 }
 #endif // WITH_EDITOR
diff --git a/Engine/Plugins/Experimental/PCG/Source/PCG/Private/Elements/PCGPointFilter.cpp b/Engine/Plugins/Experimental/PCG/Source/PCG/Private/Elements/PCGPointFilter.cpp
--- a/Engine/Plugins/Experimental/PCG/Source/PCG/Private/Elements/PCGPointFilter.cpp
+++ b/Engine/Plugins/Experimental/PCG/Source/PCG/Private/Elements/PCGPointFilter.cpp
@@ -187,7 +187,7 @@ void UPCGPointFilterSettings::ApplyDeprecationBeforeUpdatePins(UPCGNode* InOutNo
 	check(InOutNode);
 	if (DataVersion < FPCGCustomVersion::MovePointFilterParamsOffFirstPin)
 	{
-		PCGSettingsHelpers::DeprecationBreakOutParamsToPin(InOutNode, InputPins, OutputPins);
+		PCGSettingsHelpers::DeprecationBreakOutParamsToNewPin(InOutNode, InputPins, OutputPins);
 	}
 }
 #endif // WITH_EDITOR
diff --git a/Engine/Plugins/Experimental/PCG/Source/PCG/Private/Elements/PCGSelfPruning.cpp b/Engine/Plugins/Experimental/PCG/Source/PCG/Private/Elements/PCGSelfPruning.cpp
--- a/Engine/Plugins/Experimental/PCG/Source/PCG/Private/Elements/PCGSelfPruning.cpp
+++ b/Engine/Plugins/Experimental/PCG/Source/PCG/Private/Elements/PCGSelfPruning.cpp
@@ -288,7 +288,7 @@ void UPCGSelfPruningSettings::ApplyDeprecationBeforeUpdatePins(UPCGNode* InOutNo
 
 	if (DataVersion < FPCGCustomVersion::MoveSelfPruningParamsOffFirstPin)
 	{
-		PCGSettingsHelpers::DeprecationBreakOutParamsToPin(InOutNode, InputPins, OutputPins);
+		PCGSettingsHelpers::DeprecationBreakOutParamsToNewPin(InOutNode, InputPins, OutputPins);
 	}
 }
 #endif // WITH_EDITOR
diff --git a/Engine/Plugins/Experimental/PCG/Source/PCG/Private/Elements/PCGSurfaceSampler.cpp b/Engine/Plugins/Experimental/PCG/Source/PCG/Private/Elements/PCGSurfaceSampler.cpp
--- a/Engine/Plugins/Experimental/PCG/Source/PCG/Private/Elements/PCGSurfaceSampler.cpp
+++ b/Engine/Plugins/Experimental/PCG/Source/PCG/Private/Elements/PCGSurfaceSampler.cpp
@@ -405,6 +405,9 @@ void UPCGSurfaceSamplerSettings::ApplyDeprecationBeforeUpdatePins(UPCGNode* InOu
 				InputOutputSettings->SetShowAdvancedPins(true);
 			}
 		}
+
+		// A new params pin was added, migrate the first param connection there if any
+		PCGSettingsHelpers::DeprecationBreakOutParamsToNewPin(InOutNode, InputPins, OutputPins);
 	}
 
 	Super::ApplyDeprecationBeforeUpdatePins(InOutNode, InputPins, OutputPins);
diff --git a/Engine/Plugins/Experimental/PCG/Source/PCG/Private/Helpers/PCGSettingsHelpers.cpp b/Engine/Plugins/Experimental/PCG/Source/PCG/Private/Helpers/PCGSettingsHelpers.cpp
--- a/Engine/Plugins/Experimental/PCG/Source/PCG/Private/Helpers/PCGSettingsHelpers.cpp
+++ b/Engine/Plugins/Experimental/PCG/Source/PCG/Private/Helpers/PCGSettingsHelpers.cpp
@@ -214,7 +214,7 @@ namespace PCGSettingsHelpers
 		return InComponent ? PCGHelpers::ComputeSeed(SettingsSeed, InComponent->Seed) : SettingsSeed;
 	}
 
-	void DeprecationBreakOutParamsToPin(UPCGNode* InOutNode, TArray<TObjectPtr<UPCGPin>>& InputPins, TArray<TObjectPtr<UPCGPin>>& OutputPins)
+	void DeprecationBreakOutParamsToNewPin(UPCGNode* InOutNode, TArray<TObjectPtr<UPCGPin>>& InputPins, TArray<TObjectPtr<UPCGPin>>& OutputPins)
 	{
 		// Check basic conditions for which the code below should run.
 		if(InputPins.IsEmpty() || !InputPins[0] || InputPins[0]->Properties.AllowedTypes != EPCGDataType::Any)
diff --git a/Engine/Plugins/Experimental/PCG/Source/PCG/Public/Helpers/PCGSettingsHelpers.h b/Engine/Plugins/Experimental/PCG/Source/PCG/Public/Helpers/PCGSettingsHelpers.h
--- a/Engine/Plugins/Experimental/PCG/Source/PCG/Public/Helpers/PCGSettingsHelpers.h
+++ b/Engine/Plugins/Experimental/PCG/Source/PCG/Public/Helpers/PCGSettingsHelpers.h
@@ -134,7 +134,7 @@ namespace PCGSettingsHelpers
 	/** Utility to call from before-node-update deprecation. A dedicated pin for params will be added when the pins are updated. Here we detect any params
 	*   connections to the In pin and disconnect them, and move the first params connection to a new params pin.
 	*/
-	void DeprecationBreakOutParamsToPin(UPCGNode* InOutNode, TArray<TObjectPtr<UPCGPin>>& InputPins, TArray<TObjectPtr<UPCGPin>>& OutputPins);
+	void DeprecationBreakOutParamsToNewPin(UPCGNode* InOutNode, TArray<TObjectPtr<UPCGPin>>& InputPins, TArray<TObjectPtr<UPCGPin>>& OutputPins);
 }
 
 #define PCG_GET_OVERRIDEN_VALUE(Settings, Variable, Params) PCGSettingsHelpers::GetValue(GET_MEMBER_NAME_CHECKED(TRemovePointer<TRemoveConst<decltype(Settings)>::Type>::Type, Variable), (Settings)->Variable, Params)
