[PCG] New functionnalities for profiling and log views

* Moving timer and logging storage to component
* Extract scoped structs in another file
* Add support for subgraphs in Profiling and log
* Profiling/Log will focus on a given component.
* Remove "Refresh" button since it will automatically refresh when we change components/generate.

diff --git a/Engine/Plugins/Experimental/PCG/Source/PCG/Private/PCGElement.cpp b/Engine/Plugins/Experimental/PCG/Source/PCG/Private/PCGElement.cpp
--- a/Engine/Plugins/Experimental/PCG/Source/PCG/Private/PCGElement.cpp
+++ b/Engine/Plugins/Experimental/PCG/Source/PCG/Private/PCGElement.cpp
@@ -12,6 +12,7 @@
 #include "Elements/PCGSelfPruning.h"
 #include "Graph/PCGGraphCache.h"
 #include "HAL/IConsoleManager.h"
+#include "Utils/PCGExtraCapture.h"
 
 #include "Algo/Find.h"
 
@@ -28,6 +29,7 @@ bool IPCGElement::Execute(FPCGContext* Context) const
 
 	while (Context->CurrentPhase != EPCGExecutionPhase::Done)
 	{
+		PCGUtils::FScopedCall ScopedCall(*this, Context);
 		bool bExecutionPostponed = false;
 
 		switch (Context->CurrentPhase)
@@ -40,7 +42,6 @@ bool IPCGElement::Execute(FPCGContext* Context) const
 
 			case EPCGExecutionPhase::PrepareData:
 			{
-				FScopedCall ScopedCall(*this, Context);
 				if (PrepareDataInternal(Context))
 				{
 					Context->CurrentPhase = EPCGExecutionPhase::Execute;
@@ -54,7 +55,6 @@ bool IPCGElement::Execute(FPCGContext* Context) const
 
 			case EPCGExecutionPhase::Execute:
 			{
-				FScopedCall ScopedCall(*this, Context);
 				if (ExecuteInternal(Context))
 				{
 					Context->CurrentPhase = EPCGExecutionPhase::PostExecute;
@@ -164,13 +164,6 @@ void IPCGElement::PostExecute(FPCGContext* Context) const
 	// Cleanup and validate output
 	CleanupAndValidateOutput(Context);
 
-#if WITH_EDITOR
-	{
-		FScopeLock Lock(&CapturedDataLock);
-		PCGE_LOG(Verbose, "Executed in (%f)s and (%d) frames(s)", Timers[CurrentTimerIndex].ExecutionTime, Timers[CurrentTimerIndex].ExecutionFrameCount);
-	}
-#endif
-
 	const UPCGSettingsInterface* SettingsInterface = Context->GetInputSettingsInterface();
 	const UPCGSettings* Settings = SettingsInterface ? SettingsInterface->GetSettings() : nullptr;
 
@@ -312,33 +305,6 @@ void IPCGElement::DebugDisplay(FPCGContext* Context) const
 	}
 }
 
-
-TArray<IPCGElement::FCallTime> IPCGElement::GetTimers() const
-{
-	FScopeLock Lock(&CapturedDataLock);
-	return Timers;
-}
-
-TArray<IPCGElement::FCapturedMessage> IPCGElement::GetCapturedMessages() const
-{
-	FScopeLock Lock(&CapturedDataLock);
-	return CapturedMessages;
-
-}
-
-void IPCGElement::ResetTimers()
-{
-	FScopeLock Lock(&CapturedDataLock);
-	Timers.Empty();
-	CurrentTimerIndex = 0;
-}
-
-void IPCGElement::ResetMessages()
-{
-	FScopeLock Lock(&CapturedDataLock);
-	CapturedMessages.Reset();
-}
-
 #endif // WITH_EDITOR
 
 void IPCGElement::CleanupAndValidateOutput(FPCGContext* Context) const
@@ -425,79 +391,7 @@ bool IPCGElement::IsCacheableInstance(const UPCGSettingsInterface* InSettingsInt
 }
 
 #if WITH_EDITOR
-IPCGElement::FScopedCall::FScopedCall(const IPCGElement& InOwner, FPCGContext* InContext)
-	: Owner(InOwner)
-	, Context(InContext)
-	, Phase(InContext->CurrentPhase)
-	, ThreadID(FPlatformTLS::GetCurrentThreadId())
-{
-	StartTime = FPlatformTime::Seconds();
-
-	GLog->AddOutputDevice(this);
-}
-
-IPCGElement::FScopedCall::~FScopedCall()
-{
-	GLog->RemoveOutputDevice(this);
-
-	const double ThisFrameTime = FPlatformTime::Seconds() - StartTime;
-
-	FScopeLock Lock(&Owner.CapturedDataLock);
 
-	constexpr int32 MaxNumberOfTrackedTimers = 100;
-	if (Phase == EPCGExecutionPhase::PrepareData)
-	{
-		if (Owner.Timers.Num() < MaxNumberOfTrackedTimers)
-		{
-			// first time here
-			Owner.Timers.Add({});
-			Owner.CurrentTimerIndex = Owner.Timers.Num()-1;
-		}
-		else
-		{
-			Owner.CurrentTimerIndex = (Owner.CurrentTimerIndex+1) % MaxNumberOfTrackedTimers;
-			Owner.Timers[Owner.CurrentTimerIndex] = FCallTime();
-		}
-
-		FCallTime& Time = Owner.Timers[Owner.CurrentTimerIndex];
-		Time.PrepareDataTime = ThisFrameTime;
-	}
-	else if (Phase == EPCGExecutionPhase::Execute)
-	{
-		FCallTime& Time = Owner.Timers[Owner.CurrentTimerIndex];
-
-		Time.ExecutionTime += ThisFrameTime;
-		Time.ExecutionFrameCount++;
-
-		Time.MaxExecutionFrameTime = FMath::Max(Time.MaxExecutionFrameTime, ThisFrameTime);
-		Time.MinExecutionFrameTime = FMath::Min(Time.MinExecutionFrameTime, ThisFrameTime);
-	}
-	else if (Phase == EPCGExecutionPhase::PostExecute)
-	{
-		FCallTime& Time = Owner.Timers[Owner.CurrentTimerIndex];
-
-		Time.PostExecuteTime = ThisFrameTime;
-	}
-
-	Owner.CapturedMessages += std::move(CapturedMessages);
-}
-
-void IPCGElement::FScopedCall::Serialize(const TCHAR* V, ELogVerbosity::Type Verbosity, const class FName& Category)
-{
-	// TODO: this thread id check will also filter out messages spawned from threads spawned inside of nodes. To improve that,
-	// perhaps set at TLS bit on things from here and inside of PCGAsync spawned jobs. If this was done, CapturedMessages below also will
-	// need protection
-	if (Verbosity > ELogVerbosity::Warning || FPlatformTLS::GetCurrentThreadId() != ThreadID)
-	{
-		// ignore
-		return;
-	}
-
-	// this is a dumb counter just so messages can be sorted in a similar order as when they were logged
-	static volatile int32 MessageCounter = 0;
-
-	CapturedMessages.Add(FCapturedMessage { MessageCounter++, Category, V, Verbosity});
-}
 
 #endif // WITH_EDITOR
 
diff --git a/Engine/Plugins/Experimental/PCG/Source/PCG/Private/PCGSubsystem.cpp b/Engine/Plugins/Experimental/PCG/Source/PCG/Private/PCGSubsystem.cpp
--- a/Engine/Plugins/Experimental/PCG/Source/PCG/Private/PCGSubsystem.cpp
+++ b/Engine/Plugins/Experimental/PCG/Source/PCG/Private/PCGSubsystem.cpp
@@ -467,6 +467,17 @@ bool UPCGSubsystem::IsGraphCurrentlyExecuting(UPCGGraph* Graph)
 	return GraphExecutor->IsGraphCurrentlyExecuting(Graph);
 }
 
+void UPCGSubsystem::ForAllRegisteredLocalComponents(UPCGComponent* OriginalComponent, const TFunction<void(UPCGComponent*)>& InFunc) const
+{
+	auto WrapperFunc = [&InFunc](UPCGComponent* Component) -> FPCGTaskId
+	{
+		InFunc(Component);
+		return InvalidPCGTaskId;
+	};
+
+	DispatchToRegisteredLocalComponents(OriginalComponent, WrapperFunc);
+}
+
 TArray<FPCGTaskId> UPCGSubsystem::DispatchToRegisteredLocalComponents(UPCGComponent* OriginalComponent, const TFunction<FPCGTaskId(UPCGComponent*)>& InFunc) const
 {
 	TRACE_CPUPROFILER_EVENT_SCOPE(UPCGSubsystem::DispatchToRegisteredLocalComponents);
diff --git a/Engine/Plugins/Experimental/PCG/Source/PCG/Private/Utils/PCGExtraCapture.cpp b/Engine/Plugins/Experimental/PCG/Source/PCG/Private/Utils/PCGExtraCapture.cpp
new file mode 100644
--- /dev/null
+++ b/Engine/Plugins/Experimental/PCG/Source/PCG/Private/Utils/PCGExtraCapture.cpp
@@ -0,0 +1,110 @@
+// Copyright Epic Games, Inc. All Rights Reserved.
+
+#include "Utils/PCGExtraCapture.h"
+
+#include "PCGComponent.h"
+
+#if WITH_EDITOR
+
+void PCGUtils::FExtraCapture::ResetTimers()
+{
+	Timers.Empty();
+}
+
+void PCGUtils::FExtraCapture::ResetCapturedMessages()
+{
+	CapturedMessages.Empty();
+}
+
+void PCGUtils::FExtraCapture::Update(const PCGUtils::FScopedCall& InScopedCall)
+{
+	if (!InScopedCall.Context->Node)
+	{
+		return;
+	}
+
+	FScopeLock ScopedLock(&Lock);
+
+	TArray<FCallTime>* TimesPtr = Timers.Find(InScopedCall.Context->Node);
+	TArray<FCapturedMessage>* CapturedMessagesPtr = CapturedMessages.Find(InScopedCall.Context->Node);
+
+	if (!TimesPtr)
+	{
+		TimesPtr = &Timers.Add(InScopedCall.Context->Node);
+	}
+
+	if (!CapturedMessagesPtr)
+	{
+		CapturedMessagesPtr = &CapturedMessages.Add(InScopedCall.Context->Node);
+	}
+
+	check(TimesPtr && CapturedMessagesPtr);
+
+	const double ThisFrameTime = FPlatformTime::Seconds() - InScopedCall.StartTime;
+
+	constexpr int32 MaxNumberOfTrackedTimers = 100;
+
+	FCallTime* Time = !TimesPtr->IsEmpty() ? &TimesPtr->Last() : nullptr;
+	check(Time || InScopedCall.Phase == EPCGExecutionPhase::NotExecuted);
+
+	switch (InScopedCall.Phase)
+	{
+	case EPCGExecutionPhase::NotExecuted:
+		TimesPtr->Emplace();
+		break;
+	case EPCGExecutionPhase::PrepareData:
+		Time->PrepareDataTime = ThisFrameTime;
+		break;
+	case EPCGExecutionPhase::Execute:
+		Time->ExecutionTime += ThisFrameTime;
+		Time->ExecutionFrameCount++;
+
+		Time->MaxExecutionFrameTime = FMath::Max(Time->MaxExecutionFrameTime, ThisFrameTime);
+		Time->MinExecutionFrameTime = FMath::Min(Time->MinExecutionFrameTime, ThisFrameTime);
+		break;
+	case EPCGExecutionPhase::PostExecute:
+		Time->PostExecuteTime = ThisFrameTime;
+		break;
+	}
+
+	CapturedMessagesPtr->Append(InScopedCall.CapturedMessages);
+}
+
+PCGUtils::FScopedCall::FScopedCall(const IPCGElement& InOwner, FPCGContext* InContext)
+	: Owner(InOwner)
+	, Context(InContext)
+	, Phase(InContext->CurrentPhase)
+	, ThreadID(FPlatformTLS::GetCurrentThreadId())
+{
+	StartTime = FPlatformTime::Seconds();
+
+	GLog->AddOutputDevice(this);
+}
+
+PCGUtils::FScopedCall::~FScopedCall()
+{
+	GLog->RemoveOutputDevice(this);
+	if (Context && Context->SourceComponent.IsValid())
+	{
+		Context->SourceComponent->ExtraCapture.Update(*this);
+	}
+}
+
+void PCGUtils::FScopedCall::Serialize(const TCHAR* V, ELogVerbosity::Type Verbosity, const class FName& Category)
+{
+	// TODO: this thread id check will also filter out messages spawned from threads spawned inside of nodes. To improve that,
+	// perhaps set at TLS bit on things from here and inside of PCGAsync spawned jobs. If this was done, CapturedMessages below also will
+	// need protection
+	if (Verbosity > ELogVerbosity::Warning || FPlatformTLS::GetCurrentThreadId() != ThreadID)
+	{
+		// ignore
+		return;
+	}
+
+	// this is a dumb counter just so messages can be sorted in a similar order as when they were logged
+	static volatile int32 MessageCounter = 0;
+
+	CapturedMessages.Add(PCGUtils::FCapturedMessage{ MessageCounter++, Category, V, Verbosity });
+}
+
+#endif // WITH_EDITOR
\ No newline at end of file
diff --git a/Engine/Plugins/Experimental/PCG/Source/PCG/Public/PCGComponent.h b/Engine/Plugins/Experimental/PCG/Source/PCG/Public/PCGComponent.h
--- a/Engine/Plugins/Experimental/PCG/Source/PCG/Public/PCGComponent.h
+++ b/Engine/Plugins/Experimental/PCG/Source/PCG/Public/PCGComponent.h
@@ -5,6 +5,7 @@
 #include "CoreMinimal.h"
 #include "Components/ActorComponent.h"
 #include "PCGSettings.h"
+#include "Utils/PCGExtraCapture.h"
 
 #include "PCGComponent.generated.h"
 
@@ -396,6 +397,11 @@ private:
 #endif
 
 	mutable FCriticalSection GeneratedResourcesLock;
+
+#if WITH_EDITOR
+public:
+	mutable PCGUtils::FExtraCapture ExtraCapture;
+#endif // WITH_EDITOR
 };
 
 /** Used to store generated resources data during RerunConstructionScripts */
diff --git a/Engine/Plugins/Experimental/PCG/Source/PCG/Public/PCGElement.h b/Engine/Plugins/Experimental/PCG/Source/PCG/Public/PCGElement.h
--- a/Engine/Plugins/Experimental/PCG/Source/PCG/Public/PCGElement.h
+++ b/Engine/Plugins/Experimental/PCG/Source/PCG/Public/PCGElement.h
@@ -57,34 +57,7 @@ public:
 
 	/** Note: the following methods must be called from the main thread */
 #if WITH_EDITOR
-	struct FCallTime
-	{
-		// sum of all frames
-		double ExecutionTime = 0;
-		// how many frames
-		int32 ExecutionFrameCount = 0;
-		double MinExecutionFrameTime = MAX_dbl;
-		double MaxExecutionFrameTime = 0.0;
-		double PrepareDataTime = 0.0;
-		double PostExecuteTime = 0.0;
-	};
-
-	struct FCapturedMessage
-	{
-		int32 Index=0;
-		FName Namespace;
-		FString Message;
-		ELogVerbosity::Type Verbosity;
-	};
-
 	void DebugDisplay(FPCGContext* Context) const;
-
-	// returns a copy because they can be modified on another thread otherwise
-	TArray<FCallTime> GetTimers() const;
-	TArray<FCapturedMessage> GetCapturedMessages() const;
-
-	void ResetTimers();
-	void ResetMessages();
 #endif
 
 protected:
@@ -111,45 +84,6 @@ protected:
 
 private:
 	void CleanupAndValidateOutput(FPCGContext* Context) const;
-
-#if WITH_EDITOR
-
-	struct FScopedCall : public FOutputDevice
-	{
-		FScopedCall(const IPCGElement& InOwner, FPCGContext* InContext);
-		virtual ~FScopedCall();
-
-		const IPCGElement& Owner;
-		FPCGContext* Context;
-		double StartTime;
-		EPCGExecutionPhase Phase;
-		const uint32 ThreadID;
-		TArray<FCapturedMessage> CapturedMessages;
-
-		// FOutputDevice
-		virtual bool IsMemoryOnly() const override { return true; }
-		virtual bool CanBeUsedOnMultipleThreads() const override { return true; }
-		virtual bool CanBeUsedOnAnyThread() const override { return true; }
-		virtual void Serialize(const TCHAR* V, ELogVerbosity::Type Verbosity, const class FName& Category) override;
-	};
-
-	// Set mutable because we need to modify them in the execute call, which is const
-	// TODO: Should be a map with PCG Components. We need a mechanism to make sure that this map is cleaned up when component doesn't exist anymore.
-	// For now, it will track all calls to execute (excluding call where the result is already in cache).
-	mutable TArray<FCallTime> Timers;
-	mutable int32 CurrentTimerIndex = 0;
-	mutable TArray<FCapturedMessage> CapturedMessages;
-
-	// Perhaps overkill but there is a slight chance that we need to protect the timers array. If we call reset from the UI while an element is executing,
-	// it could crash while writing to the timers array.
-	mutable FCriticalSection CapturedDataLock;
-#else // !WITH_EDITOR
-	struct FScopedCall
-	{
-		FScopedCall(const IPCGElement& InOwner, FPCGContext* InContext) {}
-		void AdvanceTimerIndex() {}
-	};
-#endif // WITH_EDITOR
 };
 
 /**
diff --git a/Engine/Plugins/Experimental/PCG/Source/PCG/Public/PCGSubsystem.h b/Engine/Plugins/Experimental/PCG/Source/PCG/Public/PCGSubsystem.h
--- a/Engine/Plugins/Experimental/PCG/Source/PCG/Public/PCGSubsystem.h
+++ b/Engine/Plugins/Experimental/PCG/Source/PCG/Public/PCGSubsystem.h
@@ -120,6 +120,9 @@ public:
 	/** Flushes the graph cache completely, use only for debugging */
 	void FlushCache();
 
+	/* Call the InFunc function to all local component registered to the original component. Thread safe*/
+	void ForAllRegisteredLocalComponents(UPCGComponent* OriginalComponent, const TFunction<void(UPCGComponent*)>& InFunc) const;
+
 #if WITH_EDITOR
 public:
 	/** Schedule refresh on the current or next frame */
diff --git a/Engine/Plugins/Experimental/PCG/Source/PCG/Public/Utils/PCGExtraCapture.h b/Engine/Plugins/Experimental/PCG/Source/PCG/Public/Utils/PCGExtraCapture.h
new file mode 100644
--- /dev/null
+++ b/Engine/Plugins/Experimental/PCG/Source/PCG/Public/Utils/PCGExtraCapture.h
@@ -0,0 +1,77 @@
+// Copyright Epic Games, Inc. All Rights Reserved.
+
+#pragma once
+
+#include "PCGContext.h"
+#include "UObject/UnrealType.h"
+
+class IPCGElement;
+
+namespace PCGUtils
+{
+#if WITH_EDITOR
+	struct PCG_API FCallTime
+	{
+		// sum of all frames
+		double ExecutionTime = 0;
+		// how many frames
+		int32 ExecutionFrameCount = 0;
+		double MinExecutionFrameTime = MAX_dbl;
+		double MaxExecutionFrameTime = 0.0;
+		double PrepareDataTime = 0.0;
+		double PostExecuteTime = 0.0;
+	};
+
+	struct PCG_API FCapturedMessage
+	{
+		int32 Index = 0;
+		FName Namespace;
+		FString Message;
+		ELogVerbosity::Type Verbosity;
+	};
+
+	struct FScopedCall : public FOutputDevice
+	{
+		FScopedCall(const IPCGElement& InOwner, FPCGContext* InContext);
+		virtual ~FScopedCall();
+
+		const IPCGElement& Owner;
+		FPCGContext* Context;
+		double StartTime;
+		EPCGExecutionPhase Phase;
+		const uint32 ThreadID;
+		TArray<FCapturedMessage> CapturedMessages;
+
+		// FOutputDevice
+		virtual bool IsMemoryOnly() const override { return true; }
+		virtual bool CanBeUsedOnMultipleThreads() const override { return true; }
+		virtual bool CanBeUsedOnAnyThread() const override { return true; }
+		virtual void Serialize(const TCHAR* V, ELogVerbosity::Type Verbosity, const class FName& Category) override;
+	};
+
+	class PCG_API FExtraCapture
+	{
+	public:
+		void Update(const FScopedCall& InScopedCall);
+
+		void ResetTimers();
+		void ResetCapturedMessages();
+
+		using TTimersMap = TMap<TWeakObjectPtr<const UPCGNode>, TArray<FCallTime>>;
+		using TCapturedMessageMap = TMap<TWeakObjectPtr<const UPCGNode>, TArray<FCapturedMessage>>;
+
+		const TTimersMap& GetTimers() const { return Timers; }
+		const TCapturedMessageMap& GetCapturedMessages() const { return CapturedMessages; }
+
+	private:
+		TTimersMap Timers;
+		TCapturedMessageMap CapturedMessages;
+		mutable FCriticalSection Lock;
+	};
+#else
+	struct FScopedCall
+	{
+		FScopedCall(const IPCGElement& InOwner, FPCGContext* InContext) {}
+	};
+#endif // WITH_EDITOR
+}
\ No newline at end of file
diff --git a/Engine/Plugins/Experimental/PCG/Source/PCGEditor/Private/SPCGEditorGraphLogView.cpp b/Engine/Plugins/Experimental/PCG/Source/PCGEditor/Private/SPCGEditorGraphLogView.cpp
--- a/Engine/Plugins/Experimental/PCG/Source/PCGEditor/Private/SPCGEditorGraphLogView.cpp
+++ b/Engine/Plugins/Experimental/PCG/Source/PCGEditor/Private/SPCGEditorGraphLogView.cpp
@@ -7,6 +7,7 @@
 #include "PCGEditorGraph.h"
 #include "PCGEditorGraphNode.h"
 #include "PCGGraph.h"
+#include "PCGSubgraph.h"
 
 #include "Styling/AppStyle.h"
 #include "Widgets/Input/SSearchBox.h"
@@ -103,6 +104,14 @@ void SPCGEditorGraphLogView::OnItemDoubleClicked(PCGLogListViewItemPtr Item)
 	PCGEditor->JumpToNode(Item->EditorNode);
 }
 
+SPCGEditorGraphLogView::~SPCGEditorGraphLogView()
+{
+	if (PCGEditorPtr.IsValid())
+	{
+		PCGEditorPtr.Pin()->OnDebugObjectChangedDelegate.RemoveAll(this);
+	}
+}
+
 void SPCGEditorGraphLogView::Construct(const FArguments& InArgs, TSharedPtr<FPCGEditor> InPCGEditor)
 {
 	PCGEditorPtr = InPCGEditor;
@@ -111,6 +120,9 @@ void SPCGEditorGraphLogView::Construct(const FArguments& InArgs, TSharedPtr<FPCG
 	if (PCGEditor)
 	{
 		PCGEditorGraph = PCGEditor->GetPCGEditorGraph();
+		PCGComponent = PCGEditor->GetPCGComponentBeingDebugged();
+
+		PCGEditor->OnDebugObjectChangedDelegate.AddSP(this, &SPCGEditorGraphLogView::OnDebugObjectChanged);
 	}
 
 	ListViewHeader = CreateHeaderRowWidget();
@@ -141,13 +153,6 @@ void SPCGEditorGraphLogView::Construct(const FArguments& InArgs, TSharedPtr<FPCG
 			SNew(SHorizontalBox)
 			+SHorizontalBox::Slot()
 			.AutoWidth()
-			[
-				SNew(SButton)
-				.Text(LOCTEXT("RefreshButton", "Refresh"))
-				.OnClicked(this, &SPCGEditorGraphLogView::Refresh)
-			]
-			+SHorizontalBox::Slot()
-			.AutoWidth()
 			[
 				SNew(SButton)
 				.Text(LOCTEXT("ClearButton", "Clear"))
@@ -263,28 +268,12 @@ FReply SPCGEditorGraphLogView::Clear()
 		return FReply::Handled();
 	}
 
-	TArray<UPCGEditorGraphNode*> EditorNodes;
-	PCGEditorGraph->GetNodesOfClass<UPCGEditorGraphNode>(EditorNodes);
-
-	for (UPCGEditorGraphNode* PCGEditorNode : EditorNodes)
+	if (const UPCGComponent* Component = PCGComponent.Get())
 	{
-		if (PCGEditorNode)
-		{
-			if (UPCGNode* PCGNode = PCGEditorNode->GetPCGNode())
-			{
-				if (UPCGSettings* Settings = PCGNode->GetSettings())
-				{
-					if (IPCGElement* Element = Settings->GetElement().Get())
-					{
-						Element->ResetMessages();
-					}
-				}
-			}
-		}
+		Component->ExtraCapture.ResetCapturedMessages();
+		Refresh();
 	}
 
-	Refresh();
-
 	return FReply::Handled();
 }
 
@@ -304,6 +293,13 @@ FReply SPCGEditorGraphLogView::Refresh()
 		return FReply::Handled();
 	}
 
+	const UPCGComponent* Component = PCGComponent.Get();
+
+	if (!Component)
+	{
+		return FReply::Handled();
+	}
+
 	TArray<UPCGEditorGraphNode*> EditorNodes;
 	PCGEditorGraph->GetNodesOfClass<UPCGEditorGraphNode>(EditorNodes);
 
@@ -312,25 +308,7 @@ FReply SPCGEditorGraphLogView::Refresh()
 		const UPCGNode* PCGNode = PCGEditorNode ? PCGEditorNode->GetPCGNode() : nullptr;
 		if (PCGNode)
 		{
-			if (const UPCGSettings* Settings = PCGNode->GetSettings())
-			{
-				if (const IPCGElement* Element = Settings->GetElement().Get())
-				{
-					for (const IPCGElement::FCapturedMessage& Message : Element->GetCapturedMessages())
-					{
-						PCGLogListViewItemPtr ListViewItem = MakeShared<FPCGLogListViewItem>();
-						ListViewItem->EditorNode = PCGEditorNode;
-						ListViewItem->PCGNode = PCGNode;
-						ListViewItem->NodeName = PCGNode->GetNodeTitle();
-						ListViewItem->Order = Message.Index;
-						ListViewItem->Namespace = Message.Namespace;
-						ListViewItem->Message = Message.Message;
-						ListViewItem->Verbosity = Message.Verbosity;
-
-						this->ListViewItems.Add(ListViewItem);
-					};
-				}
-			}
+			CreateAndAddItem(PCGEditorNode, PCGNode, PCGNode->GetNodeTitle());
 		}
 	}
 
@@ -365,6 +343,82 @@ FReply SPCGEditorGraphLogView::Refresh()
 	return FReply::Handled();
 }
 
+void SPCGEditorGraphLogView::CreateAndAddItem(const UPCGEditorGraphNode* InPCGEditorNode, const UPCGNode* InPCGNode, const FName& InName)
+{
+	const TSharedPtr<FPCGEditor> PCGEditor = PCGEditorPtr.Pin();
+
+	check(PCGEditor);
+
+	const UPCGComponent* Component = PCGComponent.Get();
+
+	if (!Component || !InPCGEditorNode || !InPCGNode)
+	{
+		return;
+	}
+
+	const TArray<PCGUtils::FCapturedMessage>* CapturedMessagesPtr = Component->ExtraCapture.GetCapturedMessages().Find(InPCGNode);
+	if (!CapturedMessagesPtr)
+	{
+		return;
+	}
+
+	for (const PCGUtils::FCapturedMessage& Message : *CapturedMessagesPtr)
+	{
+		PCGLogListViewItemPtr ListViewItem = MakeShared<FPCGLogListViewItem>();
+		ListViewItem->EditorNode = InPCGEditorNode;
+		ListViewItem->PCGNode = InPCGNode;
+		ListViewItem->NodeName = InName;
+		ListViewItem->Order = Message.Index;
+		ListViewItem->Namespace = Message.Namespace;
+		ListViewItem->Message = Message.Message;
+		ListViewItem->Verbosity = Message.Verbosity;
+
+		ListViewItems.Add(ListViewItem);
+	}
+
+	// Also if the node is a subgraph, recurse on all child nodes
+	const UPCGSettings* Settings = InPCGNode->GetSettings();
+	if (Settings && Settings->IsA<UPCGBaseSubgraphSettings>())
+	{
+		if (const UPCGGraph* Graph = CastChecked<const UPCGBaseSubgraphSettings>(Settings)->GetSubgraph())
+		{
+			FString Prefix = InName.ToString() + "/";
+			for (const UPCGNode* ChildNode : Graph->GetNodes())
+			{
+				if (!ChildNode)
+				{
+					continue;
+				}
+
+				FName ChildName = FName(Prefix + ChildNode->GetNodeTitle().ToString());
+				CreateAndAddItem(InPCGEditorNode, ChildNode, ChildName);
+			}
+		}
+	}
+}
+
+void SPCGEditorGraphLogView::OnDebugObjectChanged(UPCGComponent* InPCGComponent)
+{
+	if (PCGComponent.IsValid())
+	{
+		PCGComponent->OnPCGGraphGeneratedDelegate.RemoveAll(this);
+	}
+
+	PCGComponent = InPCGComponent;
+
+	if (PCGComponent.IsValid())
+	{
+		PCGComponent->OnPCGGraphGeneratedDelegate.AddSP(this, &SPCGEditorGraphLogView::OnGenerateUpdated);
+	}
+
+	Refresh();
+}
+
+void SPCGEditorGraphLogView::OnGenerateUpdated(UPCGComponent* InPCGComponent)
+{
+	Refresh();
+}
+
 TSharedRef<ITableRow> SPCGEditorGraphLogView::OnGenerateRow(PCGLogListViewItemPtr Item, const TSharedRef<STableViewBase>& OwnerTable) const
 {
 	return SNew(SPCGLogListViewItemRow, OwnerTable, Item);
diff --git a/Engine/Plugins/Experimental/PCG/Source/PCGEditor/Private/SPCGEditorGraphLogView.h b/Engine/Plugins/Experimental/PCG/Source/PCGEditor/Private/SPCGEditorGraphLogView.h
--- a/Engine/Plugins/Experimental/PCG/Source/PCGEditor/Private/SPCGEditorGraphLogView.h
+++ b/Engine/Plugins/Experimental/PCG/Source/PCGEditor/Private/SPCGEditorGraphLogView.h
@@ -57,10 +57,16 @@ public:
 	SLATE_BEGIN_ARGS(SPCGEditorGraphLogView) { }
 	SLATE_END_ARGS()
 
+	~SPCGEditorGraphLogView();
+
 	void Construct(const FArguments& InArgs, TSharedPtr<FPCGEditor> InPCGEditor);
 
 private:
 	TSharedRef<SHeaderRow> CreateHeaderRowWidget();
+	void CreateAndAddItem(const UPCGEditorGraphNode* InPCGEditorNode, const UPCGNode* InPCGNode, const FName& InName);
+
+	void OnDebugObjectChanged(UPCGComponent* InPCGComponent);
+	void OnGenerateUpdated(UPCGComponent* InPCGComponent);
 	
 	// Callbacks
 	FReply Refresh();
@@ -76,6 +82,9 @@ private:
 	/** Cached PCGGraph being viewed */
 	UPCGEditorGraph* PCGEditorGraph = nullptr;
 
+	/** Cached PCGComponent being viewed */
+	TWeakObjectPtr<UPCGComponent> PCGComponent;
+
 	TSharedPtr<SHeaderRow> ListViewHeader;
 	TSharedPtr<SListView<PCGLogListViewItemPtr>> ListView;
 	TArray<PCGLogListViewItemPtr> ListViewItems;
diff --git a/Engine/Plugins/Experimental/PCG/Source/PCGEditor/Private/SPCGEditorGraphProfilingView.cpp b/Engine/Plugins/Experimental/PCG/Source/PCGEditor/Private/SPCGEditorGraphProfilingView.cpp
--- a/Engine/Plugins/Experimental/PCG/Source/PCGEditor/Private/SPCGEditorGraphProfilingView.cpp
+++ b/Engine/Plugins/Experimental/PCG/Source/PCGEditor/Private/SPCGEditorGraphProfilingView.cpp
@@ -7,12 +7,14 @@
 #include "PCGEditorGraph.h"
 #include "PCGEditorGraphNode.h"
 #include "PCGGraph.h"
+#include "PCGSubgraph.h"
+#include "PCGSubsystem.h"
 
 #include "Styling/AppStyle.h"
-#include "Widgets/Input/SSearchBox.h"
-#include "Widgets/Layout/SScrollBox.h"
 #include "Widgets/SBoxPanel.h"
 #include "Widgets/SWidget.h"
+#include "Widgets/Input/SCheckBox.h"
+#include "Widgets/Layout/SScrollBox.h"
 
 #define LOCTEXT_NAMESPACE "SPCGEditorGraphProfilingView"
 
@@ -76,17 +78,16 @@ TSharedRef<SWidget> SPCGProfilingListViewItemRow::GenerateWidgetForColumn(const
 		}
 		else if (ColumnId == PCGEditorGraphProfilingView::NAME_NbExecutionFrames)
 		{
-			ColumnData = FText::AsNumber(InternalItem->NbExecutionFrames);
+			ColumnData = ((InternalItem->NbExecutionFrames >= 0) ? FText::AsNumber(InternalItem->NbExecutionFrames) : FText());
 		}
 		else if (ColumnId == PCGEditorGraphProfilingView::NAME_MinNbExecutionFrames)
 		{
-			ColumnData = FText::AsNumber(InternalItem->MinNbExecutionFrames);
+			ColumnData = ((InternalItem->MinNbExecutionFrames >= 0) ? FText::AsNumber(InternalItem->MinNbExecutionFrames) : FText());
 		}
 		else if (ColumnId == PCGEditorGraphProfilingView::NAME_MaxNbExecutionFrames)
 		{
-			ColumnData = FText::AsNumber(InternalItem->MaxNbExecutionFrames);
+			ColumnData = ((InternalItem->MaxNbExecutionFrames >= 0) ? FText::AsNumber(InternalItem->MaxNbExecutionFrames) : FText());
 		}
-		// For all other values, if we don't have data, just write "N/A"
 		else if (!InternalItem->HasData)
 		{
 			ColumnData = PCGEditorGraphProfilingView::NoDataAvailableText;
@@ -99,27 +100,27 @@ TSharedRef<SWidget> SPCGProfilingListViewItemRow::GenerateWidgetForColumn(const
 		else if (ColumnId == PCGEditorGraphProfilingView::NAME_MinExecutionTime)
 		{
 			// In ms
-			ColumnData = FText::AsNumber(InternalItem->MinExecutionTime * 1000.0);
+			ColumnData = ((InternalItem->MinExecutionTime >= 0) ? FText::AsNumber(InternalItem->MinExecutionTime * 1000.0) : FText());
 		}
 		else if (ColumnId == PCGEditorGraphProfilingView::NAME_MaxExecutionTime)
 		{
 			// In ms
-			ColumnData = FText::AsNumber(InternalItem->MaxExecutionTime * 1000.0);
+			ColumnData = ((InternalItem->MaxExecutionTime >= 0) ? FText::AsNumber(InternalItem->MaxExecutionTime * 1000.0) : FText());
 		}
 		else if (ColumnId == PCGEditorGraphProfilingView::NAME_MinExecutionFrameTime)
 		{
 			// In ms
-			ColumnData = FText::AsNumber(InternalItem->MinExecutionFrameTime * 1000.0);
+			ColumnData = ((InternalItem->MinExecutionFrameTime >= 0) ? FText::AsNumber(InternalItem->MinExecutionFrameTime * 1000.0) : FText());
 		}
 		else if (ColumnId == PCGEditorGraphProfilingView::NAME_MaxExecutionFrameTime)
 		{
 			// In ms
-			ColumnData = FText::AsNumber(InternalItem->MaxExecutionFrameTime * 1000.0);
+			ColumnData = ((InternalItem->MaxExecutionFrameTime >= 0) ? FText::AsNumber(InternalItem->MaxExecutionFrameTime * 1000.0) : FText());
 		}
 		else if (ColumnId == PCGEditorGraphProfilingView::NAME_StdExecutionTime)
 		{
 			// In ms
-			ColumnData = FText::AsNumber(InternalItem->StdExecutionTime * 1000.0);
+			ColumnData = ((InternalItem->StdExecutionTime >= 0) ? FText::AsNumber(InternalItem->StdExecutionTime * 1000.0) : FText());
 		}
 		else if (ColumnId == PCGEditorGraphProfilingView::NAME_TotalExecutionTime)
 		{
@@ -155,6 +156,14 @@ void SPCGEditorGraphProfilingView::OnItemDoubleClicked(PCGProfilingListViewItemP
 	PCGEditor->JumpToNode(Item->EditorNode);
 }
 
+SPCGEditorGraphProfilingView::~SPCGEditorGraphProfilingView()
+{
+	if (PCGEditorPtr.IsValid())
+	{
+		PCGEditorPtr.Pin()->OnDebugObjectChangedDelegate.RemoveAll(this);
+	}
+}
+
 void SPCGEditorGraphProfilingView::Construct(const FArguments& InArgs, TSharedPtr<FPCGEditor> InPCGEditor)
 {
 	PCGEditorPtr = InPCGEditor;
@@ -163,6 +172,9 @@ void SPCGEditorGraphProfilingView::Construct(const FArguments& InArgs, TSharedPt
 	if (PCGEditor)
 	{
 		PCGEditorGraph = PCGEditor->GetPCGEditorGraph();
+		PCGComponent = PCGEditor->GetPCGComponentBeingDebugged();
+
+		PCGEditor->OnDebugObjectChangedDelegate.AddSP(this, &SPCGEditorGraphProfilingView::OnDebugObjectChanged);
 	}
 
 	ListViewHeader = CreateHeaderRowWidget();
@@ -195,15 +207,21 @@ void SPCGEditorGraphProfilingView::Construct(const FArguments& InArgs, TSharedPt
 			.AutoWidth()
 			[
 				SNew(SButton)
-				.Text(LOCTEXT("RefreshButton", "Refresh"))
-				.OnClicked(this, &SPCGEditorGraphProfilingView::Refresh)
+				.Text(LOCTEXT("ResetButton", "Reset"))
+				.OnClicked(this, &SPCGEditorGraphProfilingView::ResetTimers)
 			]
 			+SHorizontalBox::Slot()
 			.AutoWidth()
 			[
-				SNew(SButton)
-				.Text(LOCTEXT("ResetButton", "Reset"))
-				.OnClicked(this, &SPCGEditorGraphProfilingView::ResetTimers)
+				SNew(STextBlock)
+				.Text(LOCTEXT("ExpandSubgraph", "Expand Subgraph"))
+			]
+			+SHorizontalBox::Slot()
+			.AutoWidth()
+			[
+				SNew(SCheckBox)
+				.IsChecked(this, &SPCGEditorGraphProfilingView::IsSubgraphExpanded)
+				.OnCheckStateChanged(this, &SPCGEditorGraphProfilingView::OnSubgraphExpandedChanged)
 			]
 		]
 		+SVerticalBox::Slot()
@@ -341,6 +359,17 @@ void SPCGEditorGraphProfilingView::OnSortColumnHeader(const EColumnSortPriority:
 	Refresh();
 }
 
+ECheckBoxState SPCGEditorGraphProfilingView::IsSubgraphExpanded() const
+{
+	return bExpandSubgraph ? ECheckBoxState::Checked : ECheckBoxState::Unchecked;
+}
+
+void SPCGEditorGraphProfilingView::OnSubgraphExpandedChanged(ECheckBoxState InNewState)
+{
+	bExpandSubgraph = (InNewState == ECheckBoxState::Checked);
+	Refresh();
+}
+
 EColumnSortMode::Type SPCGEditorGraphProfilingView::GetColumnSortMode(const FName ColumnId) const
 {
 	if (SortingColumn != ColumnId)
@@ -364,28 +393,12 @@ FReply SPCGEditorGraphProfilingView::ResetTimers()
 		return FReply::Handled();
 	}
 
-	TArray<UPCGEditorGraphNode*> EditorNodes;
-	PCGEditorGraph->GetNodesOfClass<UPCGEditorGraphNode>(EditorNodes);
-
-	for (UPCGEditorGraphNode* PCGEditorNode : EditorNodes)
+	if (const UPCGComponent* Component = PCGComponent.Get())
 	{
-		if (PCGEditorNode)
-		{
-			if (UPCGNode* PCGNode = PCGEditorNode->GetPCGNode())
-			{
-				if (UPCGSettings* Settings = PCGNode->GetSettings())
-				{
-					if (IPCGElement* Element = Settings->GetElement().Get())
-					{
-						Element->ResetTimers();
-					}
-				}
-			}
-		}
+		Component->ExtraCapture.ResetTimers();
+		Refresh();
 	}
 
-	Refresh();
-
 	return FReply::Handled();
 }
 
@@ -409,64 +422,83 @@ FReply SPCGEditorGraphProfilingView::Refresh()
 	PCGEditorGraph->GetNodesOfClass<UPCGEditorGraphNode>(EditorNodes);
 	ListViewItems.Reserve(EditorNodes.Num());
 
+	const UPCGComponent* Component = PCGComponent.Get();
+
+	if (!Component)
+	{
+		return FReply::Handled();
+	}
+
 	for (const UPCGEditorGraphNode* PCGEditorNode : EditorNodes)
 	{
-		if (PCGEditorNode)
+		const UPCGNode* PCGNode = PCGEditorNode ? PCGEditorNode->GetPCGNode() : nullptr;
+		if (PCGNode)
 		{
 			PCGProfilingListViewItemPtr ListViewItem = MakeShared<FPCGProfilingListViewItem>();
-			ListViewItem->EditorNode = PCGEditorNode;
-			ListViewItem->PCGNode = PCGEditorNode->GetPCGNode();
 
-			if (ListViewItem->PCGNode)
-			{
-				ListViewItem->Name = ListViewItem->PCGNode->GetNodeTitle();
+			// If component is partitioned, we will sum all data that makes sense to sum.
 
-				if (const UPCGSettings* Settings = ListViewItem->PCGNode->GetSettings())
+			if (Component->IsPartitioned())
+			{
+				if (UPCGSubsystem* Subsystem = Component->GetSubsystem())
 				{
-					if (const IPCGElement* Element = Settings->GetElement().Get())
+					ListViewItem->PCGNode = PCGNode;
+					ListViewItem->EditorNode = PCGEditorNode;
+					ListViewItem->Name = PCGNode->GetNodeTitle();
+					ListViewItem->HasData = true;
+					ListViewItem->StdExecutionTime = -1;
+
+					ListViewItem->MinExecutionTime = std::numeric_limits<double>::max();
+					ListViewItem->MaxExecutionTime = std::numeric_limits<double>::min();
+					ListViewItem->MinExecutionFrameTime = std::numeric_limits<double>::max();
+					ListViewItem->MaxExecutionFrameTime = std::numeric_limits<double>::min();
+
+					uint32 Count = 0;
+
+					Subsystem->ForAllRegisteredLocalComponents(const_cast<UPCGComponent*>(Component),
+						[this, PCGEditorNode, PCGNode, &ListViewItem, &Count](UPCGComponent* LocalComponent)
 					{
-						const TArray<IPCGElement::FCallTime>& Timers = Element->GetTimers();
-						if (Timers.IsEmpty())
-						{
-							continue;
-						}
-
-						ListViewItem->MinExecutionTime = std::numeric_limits<double>::max();
-						ListViewItem->MaxExecutionTime = std::numeric_limits<double>::min();
-						ListViewItem->MinExecutionFrameTime = std::numeric_limits<double>::max();
-						ListViewItem->MaxExecutionFrameTime = std::numeric_limits<double>::min();
-
-						for(const IPCGElement::FCallTime& Timer : Timers)
-						{
-							ListViewItem->PrepareDataTime += Timer.PrepareDataTime;
-
-							ListViewItem->MinExecutionFrameTime = FMath::Min(ListViewItem->MinExecutionFrameTime, Timer.MinExecutionFrameTime);
-							ListViewItem->MaxExecutionFrameTime = FMath::Max(ListViewItem->MaxExecutionFrameTime, Timer.MaxExecutionFrameTime);
-
-							ListViewItem->ExecutionTime += Timer.ExecutionTime;
-							ListViewItem->MinExecutionTime = FMath::Min(ListViewItem->MinExecutionTime, Timer.ExecutionTime);
-							ListViewItem->MaxExecutionTime = FMath::Max(ListViewItem->MaxExecutionTime, Timer.ExecutionTime);
-
-							ListViewItem->NbExecutionFrames += Timer.ExecutionFrameCount;
-							ListViewItem->MaxNbExecutionFrames = FMath::Max(ListViewItem->MaxNbExecutionFrames, Timer.ExecutionFrameCount);
-							ListViewItem->MinNbExecutionFrames = FMath::Min(ListViewItem->MinNbExecutionFrames, Timer.ExecutionFrameCount);
-						}
-
-						ListViewItem->NbCalls = Timers.Num();
-						ListViewItem->AvgExecutionTime = ListViewItem->ExecutionTime / ListViewItem->NbCalls;
-
-						for (const IPCGElement::FCallTime& Timer : Timers)
-						{
-							ListViewItem->StdExecutionTime += FMath::Square(ListViewItem->AvgExecutionTime - Timer.ExecutionTime);
-						}
-
-						ListViewItem->StdExecutionTime = FMath::Sqrt(ListViewItem->StdExecutionTime / ListViewItem->NbCalls);
-						ListViewItem->HasData = true;
+						FPCGProfilingListViewItem TempItem;
+						FillItem(LocalComponent, PCGEditorNode, PCGNode, PCGNode->GetNodeTitle(), TempItem);
+
+						ListViewItem->ExecutionTime += TempItem.ExecutionTime;
+						ListViewItem->AvgExecutionTime += TempItem.AvgExecutionTime;
+						ListViewItem->NbCalls = FMath::Max(ListViewItem->NbCalls, TempItem.NbCalls);
+
+						ListViewItem->PrepareDataTime += TempItem.PrepareDataTime;
+						ListViewItem->PostExecuteTime += TempItem.PostExecuteTime;
+
+						ListViewItem->MinExecutionFrameTime = FMath::Min(ListViewItem->MinExecutionFrameTime, TempItem.MinExecutionFrameTime);
+						ListViewItem->MaxExecutionFrameTime = FMath::Max(ListViewItem->MaxExecutionFrameTime, TempItem.MaxExecutionFrameTime);
+
+						ListViewItem->MinExecutionTime = FMath::Min(ListViewItem->MinExecutionTime, TempItem.ExecutionTime);
+						ListViewItem->MaxExecutionTime = FMath::Max(ListViewItem->MaxExecutionTime, TempItem.ExecutionTime);
+
+						ListViewItem->NbExecutionFrames += TempItem.NbExecutionFrames;
+						ListViewItem->MaxNbExecutionFrames = FMath::Max(ListViewItem->MaxNbExecutionFrames, TempItem.MaxNbExecutionFrames);
+						ListViewItem->MinNbExecutionFrames = FMath::Min(ListViewItem->MinNbExecutionFrames, TempItem.MinNbExecutionFrames);
+
+						Count++;
+						return InvalidPCGTaskId;
+					});
+
+					// "Amortized" cost for mean an nb frames. Not exact because of multithread and some overhead.
+					if (Count > 0)
+					{
+						ListViewItem->AvgExecutionTime /= Count;
+						ListViewItem->NbExecutionFrames /= Count;
 					}
+
+					ListViewItems.Add(ListViewItem);
+				}
+			}
+			else
+			{
+				if (FillItem(Component, PCGEditorNode, PCGNode, PCGNode->GetNodeTitle(), *ListViewItem))
+				{
+					ListViewItems.Add(ListViewItem);
 				}
 			}
-
-			ListViewItems.Add(ListViewItem);
 		}
 	}
 
@@ -537,6 +569,137 @@ FReply SPCGEditorGraphProfilingView::Refresh()
 	return FReply::Handled();
 }
 
+bool SPCGEditorGraphProfilingView::FillItem(const UPCGComponent* InComponent, const UPCGEditorGraphNode* InPCGEditorNode, const UPCGNode* InPCGNode, const FName& InName, FPCGProfilingListViewItem& OutItem)
+{
+	const TSharedPtr<FPCGEditor> PCGEditor = PCGEditorPtr.Pin();
+
+	check(PCGEditor);
+
+	if (!InComponent || !InPCGEditorNode || !InPCGNode)
+	{
+		return false;
+	}
+
+	const TArray<PCGUtils::FCallTime>* TimersPtr = InComponent->ExtraCapture.GetTimers().Find(InPCGNode);
+	if (!TimersPtr)
+	{
+		return false;
+	}
+
+	OutItem.EditorNode = InPCGEditorNode;
+	OutItem.PCGNode = InPCGNode;
+
+	OutItem.Name = InName;
+	OutItem.MinExecutionTime = std::numeric_limits<double>::max();
+	OutItem.MaxExecutionTime = std::numeric_limits<double>::min();
+	OutItem.MinExecutionFrameTime = std::numeric_limits<double>::max();
+	OutItem.MaxExecutionFrameTime = std::numeric_limits<double>::min();
+
+	OutItem.NbCalls = TimersPtr->Num();
+
+	for (const PCGUtils::FCallTime& Timer : *TimersPtr)
+	{
+		OutItem.PrepareDataTime += Timer.PrepareDataTime;
+		OutItem.PostExecuteTime += Timer.PostExecuteTime;
+
+		OutItem.MinExecutionFrameTime = FMath::Min(OutItem.MinExecutionFrameTime, Timer.MinExecutionFrameTime);
+		OutItem.MaxExecutionFrameTime = FMath::Max(OutItem.MaxExecutionFrameTime, Timer.MaxExecutionFrameTime);
+
+		OutItem.ExecutionTime += Timer.ExecutionTime;
+		OutItem.MinExecutionTime = FMath::Min(OutItem.MinExecutionTime, Timer.ExecutionTime);
+		OutItem.MaxExecutionTime = FMath::Max(OutItem.MaxExecutionTime, Timer.ExecutionTime);
+
+		OutItem.NbExecutionFrames += Timer.ExecutionFrameCount;
+		OutItem.MaxNbExecutionFrames = FMath::Max(OutItem.MaxNbExecutionFrames, Timer.ExecutionFrameCount);
+		OutItem.MinNbExecutionFrames = FMath::Min(OutItem.MinNbExecutionFrames, Timer.ExecutionFrameCount);
+	}
+
+	OutItem.AvgExecutionTime = OutItem.ExecutionTime / OutItem.NbCalls;
+
+	for (const PCGUtils::FCallTime& Timer : *TimersPtr)
+	{
+		OutItem.StdExecutionTime += FMath::Square(OutItem.AvgExecutionTime - Timer.ExecutionTime);
+	}
+
+	OutItem.StdExecutionTime = FMath::Sqrt(OutItem.StdExecutionTime / OutItem.NbCalls);
+
+	OutItem.HasData = true;
+
+	// Also if the node is a subgraph, we need to add all stats from the child nodes
+	const UPCGSettings* Settings = InPCGNode->GetSettings();
+	if (Settings && Settings->IsA<UPCGBaseSubgraphSettings>())
+	{
+		if (const UPCGGraph* Graph = CastChecked<const UPCGBaseSubgraphSettings>(Settings)->GetSubgraph())
+		{
+			FString Prefix = InName.ToString() + "/";
+			for (const UPCGNode* ChildNode : Graph->GetNodes())
+			{
+				if (!ChildNode)
+				{
+					continue;
+				}
+
+				FPCGProfilingListViewItem TempItem;
+				PCGProfilingListViewItemPtr ChildItemPtr;
+
+				FPCGProfilingListViewItem* ChildItem = &TempItem;
+				if (bExpandSubgraph)
+				{
+					ChildItemPtr = MakeShared<FPCGProfilingListViewItem>();
+					ChildItem = ChildItemPtr.Get();
+				}
+
+				FName ChildName = FName(Prefix + ChildNode->GetNodeTitle().ToString());
+				if (FillItem(InComponent, InPCGEditorNode, ChildNode, ChildName, *ChildItem) && bExpandSubgraph)
+				{
+					ListViewItems.Add(ChildItemPtr);
+				}
+
+				OutItem.PrepareDataTime += ChildItem->PrepareDataTime;
+				OutItem.PostExecuteTime += ChildItem->PostExecuteTime;
+				OutItem.ExecutionTime += ChildItem->ExecutionTime;
+			}
+
+			OutItem.AvgExecutionTime = OutItem.ExecutionTime / OutItem.NbCalls;
+
+			// Other values can't be easily computed as we don't have enough granularity in subgraph scoped call.
+			// Setting them to -1 will hide them.
+			OutItem.MaxExecutionFrameTime = -1;
+			OutItem.MinExecutionFrameTime = -1;
+			OutItem.MinExecutionTime = -1;
+			OutItem.MaxExecutionTime = -1;
+			OutItem.StdExecutionTime = -1;
+			OutItem.MinNbExecutionFrames = -1;
+			OutItem.MaxNbExecutionFrames = -1;
+			OutItem.NbExecutionFrames = -1;
+		}
+	}
+
+	return true;
+}
+
+void SPCGEditorGraphProfilingView::OnDebugObjectChanged(UPCGComponent* InPCGComponent)
+{
+	if (PCGComponent.IsValid())
+	{
+		PCGComponent->OnPCGGraphGeneratedDelegate.RemoveAll(this);
+	}
+
+	PCGComponent = InPCGComponent;
+
+	if (PCGComponent.IsValid())
+	{
+		PCGComponent->OnPCGGraphGeneratedDelegate.AddSP(this, &SPCGEditorGraphProfilingView::OnGenerateUpdated);
+	}
+
+	Refresh();
+}
+
+void SPCGEditorGraphProfilingView::OnGenerateUpdated(UPCGComponent* InPCGComponent)
+{
+	Refresh();
+}
+
 TSharedRef<ITableRow> SPCGEditorGraphProfilingView::OnGenerateRow(PCGProfilingListViewItemPtr Item, const TSharedRef<STableViewBase>& OwnerTable) const
 {
 	return SNew(SPCGProfilingListViewItemRow, OwnerTable, Item);
diff --git a/Engine/Plugins/Experimental/PCG/Source/PCGEditor/Private/SPCGEditorGraphProfilingView.h b/Engine/Plugins/Experimental/PCG/Source/PCGEditor/Private/SPCGEditorGraphProfilingView.h
--- a/Engine/Plugins/Experimental/PCG/Source/PCGEditor/Private/SPCGEditorGraphProfilingView.h
+++ b/Engine/Plugins/Experimental/PCG/Source/PCGEditor/Private/SPCGEditorGraphProfilingView.h
@@ -31,7 +31,7 @@ struct FPCGProfilingListViewItem
 	double StdExecutionTime = 0.0;
 	double PostExecuteTime = 0.0;
 	int32 NbCalls = 0;
-	int32 NbExecutionFrames = 0;
+	double NbExecutionFrames = 0;
 	int32 MaxNbExecutionFrames = 0;
 	int32 MinNbExecutionFrames = 0;
 	bool HasData = false;
@@ -59,10 +59,19 @@ public:
 	SLATE_BEGIN_ARGS(SPCGEditorGraphProfilingView) { }
 	SLATE_END_ARGS()
 
+	~SPCGEditorGraphProfilingView();
+
 	void Construct(const FArguments& InArgs, TSharedPtr<FPCGEditor> InPCGEditor);
 
 private:
 	TSharedRef<SHeaderRow> CreateHeaderRowWidget();
+	bool FillItem(const UPCGComponent* InComponent, const UPCGEditorGraphNode* InPCGEditorNode, const UPCGNode* InPCGNode, const FName& InName, FPCGProfilingListViewItem& OutItem);
+
+	ECheckBoxState IsSubgraphExpanded() const;
+	void OnSubgraphExpandedChanged(ECheckBoxState InNewState);
+
+	void OnDebugObjectChanged(UPCGComponent* InPCGComponent);
+	void OnGenerateUpdated(UPCGComponent* InPCGComponent);
 	
 	// Callbacks
 	FReply Refresh();
@@ -78,6 +87,9 @@ private:
 	/** Cached PCGGraph being viewed */
 	UPCGEditorGraph* PCGEditorGraph = nullptr;
 
+	/** Cached PCGComponent being viewed */
+	TWeakObjectPtr<UPCGComponent> PCGComponent;
+
 	TSharedPtr<SHeaderRow> ListViewHeader;
 	TSharedPtr<SListView<PCGProfilingListViewItemPtr>> ListView;
 	TArray<PCGProfilingListViewItemPtr> ListViewItems;
@@ -85,4 +97,6 @@ private:
 	// To allow sorting
 	FName SortingColumn = NAME_None;
 	EColumnSortMode::Type SortMode = EColumnSortMode::Type::None;
+
+	bool bExpandSubgraph = true;
 };
