In non-editor PIE, use more time for processing so users can get results faster.

diff --git a/Engine/Plugins/Experimental/PCG/Source/PCG/Private/Graph/PCGGraphExecutor.cpp b/Engine/Plugins/Experimental/PCG/Source/PCG/Private/Graph/PCGGraphExecutor.cpp
--- a/Engine/Plugins/Experimental/PCG/Source/PCG/Private/Graph/PCGGraphExecutor.cpp
+++ b/Engine/Plugins/Experimental/PCG/Source/PCG/Private/Graph/PCGGraphExecutor.cpp
@@ -31,6 +31,11 @@ static TAutoConsoleVariable<float> CVarTimePerFrame(
 	1000.0f / 60.0f,
 	TEXT("Allocated time in ms per frame"));
 
+static TAutoConsoleVariable<float> CVarEditorTimePerFrame(
+	TEXT("pcg.EditorFrameTime"),
+	1000.0f / 20.0f,
+	TEXT("Allocated time in ms per frame when running in editor (non pie)"));
+
 static TAutoConsoleVariable<bool> CVarGraphMultithreading(
 	TEXT("pcg.GraphMultithreading"),
 	false,
@@ -351,7 +356,17 @@ void FPCGGraphExecutor::Execute()
 	bool bAnyTaskEnded = false;
 
 	const double StartTime = FPlatformTime::Seconds();
-	const double EndTime = StartTime + (CVarTimePerFrame.GetValueOnAnyThread() / 1000.0);
+
+	double VarTimePerFrame = CVarTimePerFrame.GetValueOnAnyThread() / 1000.0;
+
+#if WITH_EDITOR
+	if (GEditor && !GEditor->IsPlaySessionInProgress())
+	{
+		VarTimePerFrame = CVarEditorTimePerFrame.GetValueOnAnyThread() / 1000.0;
+	}
+#endif
+
+	const double EndTime = StartTime + VarTimePerFrame;
 	const int32 MaxNumThreads = FMath::Max(0, FMath::Min(FPlatformMisc::NumberOfCoresIncludingHyperthreads() - 2, CVarMaxNumTasks.GetValueOnAnyThread() - 1));
 	const bool bAllowMultiDispatch = CVarGraphMultithreading.GetValueOnAnyThread();
 
