[PCG] Missing Super::PostLoad

diff --git a/Engine/Plugins/Experimental/PCG/Source/PCG/Private/Data/PCGLandscapeData.cpp b/Engine/Plugins/Experimental/PCG/Source/PCG/Private/Data/PCGLandscapeData.cpp
--- a/Engine/Plugins/Experimental/PCG/Source/PCG/Private/Data/PCGLandscapeData.cpp
+++ b/Engine/Plugins/Experimental/PCG/Source/PCG/Private/Data/PCGLandscapeData.cpp
@@ -69,6 +69,8 @@ void UPCGLandscapeData::Initialize(const TArray<TWeakObjectPtr<ALandscapeProxy>>
 
 void UPCGLandscapeData::PostLoad()
 {
+	Super::PostLoad();
+
 	ALandscapeProxy* FirstLandscape = nullptr;
 
 	for (TSoftObjectPtr<ALandscapeProxy> Landscape : Landscapes)
