[PCG] Fixed an issue where the input node was not updated properly after undo/redo causing broken UI or even crashes.

diff --git a/Engine/Plugins/Experimental/PCG/Source/PCGEditor/Private/PCGEditorGraphNodeBase.cpp b/Engine/Plugins/Experimental/PCG/Source/PCGEditor/Private/PCGEditorGraphNodeBase.cpp
--- a/Engine/Plugins/Experimental/PCG/Source/PCGEditor/Private/PCGEditorGraphNodeBase.cpp
+++ b/Engine/Plugins/Experimental/PCG/Source/PCGEditor/Private/PCGEditorGraphNodeBase.cpp
@@ -63,6 +63,8 @@ void UPCGEditorGraphNodeBase::PostTransacted(const FTransactionObjectEvent& Tran
 	{
 		UpdatePosition();
 	}
+
+	ReconstructNode();
 }
 
 void UPCGEditorGraphNodeBase::GetNodeContextMenuActions(UToolMenu* Menu, class UGraphNodeContextMenuContext* Context) const
