[PCG] Add support for advanced pins within nodes using BP paradigm

* Input/Output nodes has now advanced pins always created.
* Remove `Show Advanced Pins` in Input/Output node
* Add `bIsAdvancedPin` in PinProperties
* Any node that has at least 1 advanced pin will have an arrow at the bottom and will hide all advanced pins.

diff --git a/Engine/Plugins/Experimental/PCG/Source/PCG/Private/Elements/PCGSurfaceSampler.cpp b/Engine/Plugins/Experimental/PCG/Source/PCG/Private/Elements/PCGSurfaceSampler.cpp
--- a/Engine/Plugins/Experimental/PCG/Source/PCG/Private/Elements/PCGSurfaceSampler.cpp
+++ b/Engine/Plugins/Experimental/PCG/Source/PCG/Private/Elements/PCGSurfaceSampler.cpp
@@ -396,17 +396,6 @@ void UPCGSurfaceSamplerSettings::ApplyDeprecationBeforeUpdatePins(UPCGNode* InOu
 			InputPins[0]->Properties.Label = PCGSurfaceSamplerConstants::BoundingShapeLabel;
 		}
 
-		// If no pins or only 1 pin, then we will be doing rewiring later, so ask for advanced pins on graph input node
-		UPCGNode* GraphInputNode = InOutNode->GetGraph() ? InOutNode->GetGraph()->GetInputNode() : nullptr;
-		if (GraphInputNode && (InputPins.IsEmpty() || (InputPins.Num() == 1 && PCGSurfaceSampler::IsPinOnlyConnectedToInputNode(InputPins[0], GraphInputNode))))
-		{
-			// Ask for advanced params to be visible, as it will be useful later for wiring up the Surface input to the advanced Landscape pin
-			if (UPCGGraphInputOutputSettings* InputOutputSettings = Cast<UPCGGraphInputOutputSettings>(GraphInputNode->GetSettings()))
-			{
-				InputOutputSettings->SetShowAdvancedPins(true);
-			}
-		}
-
 		// A new params pin was added, migrate the first param connection there if any
 		PCGSettingsHelpers::DeprecationBreakOutParamsToNewPin(InOutNode, InputPins, OutputPins);
 	}
diff --git a/Engine/Plugins/Experimental/PCG/Source/PCG/Private/PCGInputOutputSettings.cpp b/Engine/Plugins/Experimental/PCG/Source/PCG/Private/PCGInputOutputSettings.cpp
--- a/Engine/Plugins/Experimental/PCG/Source/PCG/Private/PCGInputOutputSettings.cpp
+++ b/Engine/Plugins/Experimental/PCG/Source/PCG/Private/PCGInputOutputSettings.cpp
@@ -73,14 +73,15 @@ TArray<FPCGPinProperties> UPCGGraphInputOutputSettings::GetPinProperties() const
 		return FPCGPinProperties(InLabelAndTooltip.Label, DefaultPinDataType, /*bMultiConnections=*/true, /*bMultiData=*/true, InLabelAndTooltip.Tooltip);
 	});
 	
-	if (bShowAdvancedPins)
-	{
-		Algo::Transform(StaticAdvancedLabels(), PinProperties, [DefaultPinDataType](const FLabelAndTooltip& InLabelAndTooltip) {
-			const bool bIsLandscapePin = (InLabelAndTooltip.Label == PCGInputOutputConstants::DefaultLandscapeLabel || InLabelAndTooltip.Label == PCGInputOutputConstants::DefaultLandscapeHeightLabel);
-			const EPCGDataType PinType = bIsLandscapePin ? EPCGDataType::Surface : DefaultPinDataType;
-			return FPCGPinProperties(InLabelAndTooltip.Label, PinType, /*bMultiConnection=*/true, /*bMultiData=*/false, InLabelAndTooltip.Tooltip);
-		});
-	}
+
+	Algo::Transform(StaticAdvancedLabels(), PinProperties, [this, DefaultPinDataType](const FLabelAndTooltip& InLabelAndTooltip) {
+		const bool bIsLandscapePin = (InLabelAndTooltip.Label == PCGInputOutputConstants::DefaultLandscapeLabel || InLabelAndTooltip.Label == PCGInputOutputConstants::DefaultLandscapeHeightLabel);
+		const EPCGDataType PinType = bIsLandscapePin ? EPCGDataType::Surface : DefaultPinDataType;
+		FPCGPinProperties Res(InLabelAndTooltip.Label, PinType, /*bMultiConnection=*/true, /*bMultiData=*/false, InLabelAndTooltip.Tooltip);
+		Res.bAdvancedPin = true;
+		return Res;
+	});
+
 
 	PinProperties.Append(CustomPins);
 	return PinProperties;
@@ -96,24 +97,16 @@ TArray<FPCGPinProperties> UPCGGraphInputOutputSettings::OutputPinProperties() co
 	return GetPinProperties();
 }
 
-bool UPCGGraphInputOutputSettings::IsPinAdvanced(const UPCGPin* Pin) const
-{
-	return Pin && StaticAdvancedLabels().FindByPredicate([Pin](const FLabelAndTooltip& InLabelAndTooltip) -> bool { return InLabelAndTooltip.Label == Pin->Properties.Label; });
-}
-
-void UPCGGraphInputOutputSettings::SetShowAdvancedPins(bool bValue)
-{
-	if (bValue != bShowAdvancedPins)
-	{
-		Modify();
-		bShowAdvancedPins = bValue;
-	}
-}
-
 void UPCGGraphInputOutputSettings::AddCustomPin(const FPCGPinProperties& NewCustomPinProperties)
 {
 	Modify();
 	CustomPins.Add(NewCustomPinProperties);
 }
 
+bool UPCGGraphInputOutputSettings::IsCustomPin(const UPCGPin* InPin) const
+{
+	check(InPin);
+	return CustomPins.Contains(InPin->Properties);
+}
+
 #undef LOCTEXT_NAMESPACE
diff --git a/Engine/Plugins/Experimental/PCG/Source/PCG/Private/PCGNode.cpp b/Engine/Plugins/Experimental/PCG/Source/PCG/Private/PCGNode.cpp
--- a/Engine/Plugins/Experimental/PCG/Source/PCG/Private/PCGNode.cpp
+++ b/Engine/Plugins/Experimental/PCG/Source/PCG/Private/PCGNode.cpp
@@ -660,30 +660,60 @@ EPCGChangeType UPCGNode::UpdatePins(TFunctionRef<UPCGPin*(UPCGNode*)> PinAllocat
 		}
 		else
 		{
-			if(!UnmatchedPins.IsEmpty() || !UnmatchedProperties.IsEmpty())
+			// Verification that we don't have 2 pins with the same name
+			// If so, mark them to be removed.
+			TSet<FName> AllPinNames;
+			TArray<UPCGPin*> DuplicatedNamePins;
+
+			for (UPCGPin* Pin : Pins)
+			{
+				if (AllPinNames.Contains(Pin->Properties.Label))
+				{
+					DuplicatedNamePins.Add(Pin);
+				}
+
+				AllPinNames.Add(Pin->Properties.Label);
+			}
+
+			if(!UnmatchedPins.IsEmpty() || !UnmatchedProperties.IsEmpty() || !DuplicatedNamePins.IsEmpty())
 			{
 				Modify();
 				bChangedPins = true;
 			}
 
-			// Remove old pins
-			for (int32 UnmatchedPinIndex = UnmatchedPins.Num() - 1; UnmatchedPinIndex >= 0; --UnmatchedPinIndex)
+			auto RemovePins = [&Pins, &AllPinNames, &bAppliedEdgeChanges](TArray<UPCGPin*>& PinsToRemove, bool bRemoveFromAllNames)
 			{
-				const int32 PinIndex = Pins.IndexOfByKey(UnmatchedPins[UnmatchedPinIndex]);
-				check(PinIndex >= 0);
+				for (int32 RemovedPinIndex = PinsToRemove.Num() - 1; RemovedPinIndex >= 0; --RemovedPinIndex)
+				{
+					const int32 PinIndex = Pins.IndexOfByKey(PinsToRemove[RemovedPinIndex]);
+					check(PinIndex >= 0);
 
-				bAppliedEdgeChanges |= Pins[PinIndex]->BreakAllEdges();
-				Pins.RemoveAt(PinIndex);
-			}
+					if (bRemoveFromAllNames)
+					{
+						AllPinNames.Remove(Pins[PinIndex]->Properties.Label);
+					}
+
+					bAppliedEdgeChanges |= Pins[PinIndex]->BreakAllEdges();
+					Pins.RemoveAt(PinIndex);
+				}
+			};
+
+			RemovePins(UnmatchedPins, /*bRemoveFromAllNames=*/ true);
+			RemovePins(DuplicatedNamePins, /*bRemoveFromAllNames=*/ false);
 
 			// Add new pins
 			for (const FPCGPinProperties& UnmatchedProperty : UnmatchedProperties)
 			{
-				const int32 InsertIndex = PinProperties.IndexOfByKey(UnmatchedProperty);
-				UPCGPin* NewPin = PinAllocator(this);
-				NewPin->Node = this;
-				NewPin->Properties = UnmatchedProperty;
-				Pins.Insert(NewPin, InsertIndex);
+				if (ensure(!AllPinNames.Contains(UnmatchedProperty.Label)))
+				{
+					AllPinNames.Add(UnmatchedProperty.Label);
+
+					const int32 InsertIndex = FMath::Min(PinProperties.IndexOfByKey(UnmatchedProperty), Pins.Num());
+					UPCGPin* NewPin = PinAllocator(this);
+					NewPin->Node = this;
+					NewPin->Properties = UnmatchedProperty;
+					Pins.Insert(NewPin, InsertIndex);
+				}
 			}
 		}
 
diff --git a/Engine/Plugins/Experimental/PCG/Source/PCG/Private/PCGPin.cpp b/Engine/Plugins/Experimental/PCG/Source/PCG/Private/PCGPin.cpp
--- a/Engine/Plugins/Experimental/PCG/Source/PCG/Private/PCGPin.cpp
+++ b/Engine/Plugins/Experimental/PCG/Source/PCG/Private/PCGPin.cpp
@@ -18,7 +18,8 @@ bool FPCGPinProperties::operator==(const FPCGPinProperties& Other) const
 	return Label == Other.Label &&
 		AllowedTypes == Other.AllowedTypes &&
 		bAllowMultipleConnections == Other.bAllowMultipleConnections &&
-		bAllowMultipleData == Other.bAllowMultipleData;
+		bAllowMultipleData == Other.bAllowMultipleData &&
+		bAdvancedPin == Other.bAdvancedPin;
 }
 
 UPCGPin::UPCGPin(const FObjectInitializer& ObjectInitializer)
diff --git a/Engine/Plugins/Experimental/PCG/Source/PCG/Public/PCGInputOutputSettings.h b/Engine/Plugins/Experimental/PCG/Source/PCG/Public/PCGInputOutputSettings.h
--- a/Engine/Plugins/Experimental/PCG/Source/PCG/Public/PCGInputOutputSettings.h
+++ b/Engine/Plugins/Experimental/PCG/Source/PCG/Public/PCGInputOutputSettings.h
@@ -47,8 +47,7 @@ public:
 
 	void SetInput(bool bInIsInput) { bIsInput = bInIsInput; }
 
-	bool IsPinAdvanced(const UPCGPin* Pin) const;
-	void SetShowAdvancedPins(bool bValue);
+	bool IsCustomPin(const UPCGPin* InPin) const;
 
 	void AddCustomPin(const FPCGPinProperties& NewCustomPinProperties);
 
@@ -79,9 +78,6 @@ protected:
 	UPROPERTY(EditAnywhere, BlueprintReadWrite, Category = "Settings|Input")
 	TArray<FPCGPinProperties> CustomPins;
 
-	UPROPERTY(EditAnywhere, BlueprintReadWrite, Category = "Settings|Input")
-	bool bShowAdvancedPins = false;
-
 	TArray<FLabelAndTooltip> StaticInLabels;
 	TArray<FLabelAndTooltip> StaticAdvancedInLabels;
 	TArray<FLabelAndTooltip> StaticOutLabels;
diff --git a/Engine/Plugins/Experimental/PCG/Source/PCG/Public/PCGPin.h b/Engine/Plugins/Experimental/PCG/Source/PCG/Public/PCGPin.h
--- a/Engine/Plugins/Experimental/PCG/Source/PCG/Public/PCGPin.h
+++ b/Engine/Plugins/Experimental/PCG/Source/PCG/Public/PCGPin.h
@@ -32,6 +32,9 @@ struct PCG_API FPCGPinProperties
 	UPROPERTY(BlueprintReadWrite, EditAnywhere, Category = Settings)
 	bool bAllowMultipleConnections = true;
 
+	UPROPERTY(BlueprintReadWrite, EditAnywhere, Category = Settings)
+	bool bAdvancedPin = false;
+
 #if WITH_EDITORONLY_DATA
 	UPROPERTY(EditAnywhere, Category = Settings)
 	FText Tooltip;
diff --git a/Engine/Plugins/Experimental/PCG/Source/PCGEditor/Private/PCGEditor.cpp b/Engine/Plugins/Experimental/PCG/Source/PCGEditor/Private/PCGEditor.cpp
--- a/Engine/Plugins/Experimental/PCG/Source/PCGEditor/Private/PCGEditor.cpp
+++ b/Engine/Plugins/Experimental/PCG/Source/PCGEditor/Private/PCGEditor.cpp
@@ -804,26 +804,6 @@ void FPCGEditor::OnCollapseNodesInSubgraph()
 	InputFromSubgraphPins = InputFromSubgraphPins.Difference(AllSubgraphPins);
 	OutputToSubgraphPins = OutputToSubgraphPins.Difference(AllSubgraphPins);
 
-	// Function that check if the pin is part of the advanced pins of the input/output node of the PCGGraph.
-	// If so, enable advanced pins on the subgraph.
-	auto IsAdvancedAndSet = [&PCGGraph, &NewPCGGraph](UPCGPin* Pin, bool bIsInput)
-	{
-		TObjectPtr<UPCGNode> InputOutputNode = bIsInput ? PCGGraph->GetInputNode() : PCGGraph->GetOutputNode();
-		TObjectPtr<UPCGNode> NewInputOutputNode = bIsInput ? NewPCGGraph->GetInputNode() : NewPCGGraph->GetOutputNode();
-
-		TObjectPtr<UPCGGraphInputOutputSettings> InputOutputSettings = CastChecked<UPCGGraphInputOutputSettings>(InputOutputNode->GetSettings());
-		if (InputOutputSettings->IsPinAdvanced(Pin))
-		{
-			TObjectPtr<UPCGGraphInputOutputSettings> NewInputOutputSettings = CastChecked<UPCGGraphInputOutputSettings>(NewInputOutputNode->GetSettings());
-			NewInputOutputSettings->SetShowAdvancedPins(true);
-			NewInputOutputNode->UpdateAfterSettingsChangeDuringCreation();
-
-			return true;
-		}
-
-		return false;
-	};
-
 	// Function that will create a new pin in the input/output node of the subgraph, with a unique name that will
 	// be the related to the pin passed as argument. It will be formatted like this:
 	// "{NodeName} {PinName} {OptionalIndex}"
@@ -915,7 +895,9 @@ void FPCGEditor::OnCollapseNodesInSubgraph()
 
 			if (Edge->InputPin->Node == PCGGraph->GetInputNode())
 			{
-				if (Edge->InputPin == PCGGraph->GetInputNode()->GetOutputPins()[0] || IsAdvancedAndSet(Edge->InputPin, true))
+				const UPCGGraphInputOutputSettings* Settings = Cast<const UPCGGraphInputOutputSettings>(Edge->InputPin->Node->GetSettings());
+
+				if (Settings && !Settings->IsCustomPin(Edge->InputPin))
 				{
 					OutsideSubgraphEdge.OutputPinLabel = Edge->InputPin->Properties.Label;
 					InsideSubgraphEdge.InputPin = NewPCGGraph->GetInputNode()->GetOutputPin(Edge->InputPin->Properties.Label);
@@ -946,7 +928,9 @@ void FPCGEditor::OnCollapseNodesInSubgraph()
 
 			if (Edge->OutputPin->Node == PCGGraph->GetOutputNode())
 			{
-				if (Edge->OutputPin == PCGGraph->GetOutputNode()->GetInputPins()[0] || IsAdvancedAndSet(Edge->OutputPin, false))
+				const UPCGGraphInputOutputSettings* Settings = Cast<const UPCGGraphInputOutputSettings>(Edge->OutputPin->Node->GetSettings());
+
+				if (Settings && !Settings->IsCustomPin(Edge->OutputPin))
 				{
 					OutsideSubgraphEdge.InputPinLabel = Edge->OutputPin->Properties.Label;
 					InsideSubgraphEdge.OutputPin = NewPCGGraph->GetOutputNode()->GetInputPin(Edge->OutputPin->Properties.Label);
diff --git a/Engine/Plugins/Experimental/PCG/Source/PCGEditor/Private/PCGEditorGraphNode.cpp b/Engine/Plugins/Experimental/PCG/Source/PCGEditor/Private/PCGEditorGraphNode.cpp
--- a/Engine/Plugins/Experimental/PCG/Source/PCGEditor/Private/PCGEditorGraphNode.cpp
+++ b/Engine/Plugins/Experimental/PCG/Source/PCGEditor/Private/PCGEditorGraphNode.cpp
@@ -54,15 +54,7 @@ void UPCGEditorGraphNode::AllocateDefaultPins()
 {
 	if (PCGNode)
 	{
-		for (const UPCGPin* InputPin : PCGNode->GetInputPins())
-		{
-			CreatePin(EEdGraphPinDirection::EGPD_Input, GetPinType(InputPin), InputPin->Properties.Label);
-		}
-
-		for (const UPCGPin* OutputPin : PCGNode->GetOutputPins())
-		{
-			CreatePin(EEdGraphPinDirection::EGPD_Output, GetPinType(OutputPin), OutputPin->Properties.Label);
-		}
+		CreatePins(PCGNode->GetInputPins(), PCGNode->GetOutputPins());
 	}
 }
 
diff --git a/Engine/Plugins/Experimental/PCG/Source/PCGEditor/Private/PCGEditorGraphNodeBase.cpp b/Engine/Plugins/Experimental/PCG/Source/PCGEditor/Private/PCGEditorGraphNodeBase.cpp
--- a/Engine/Plugins/Experimental/PCG/Source/PCGEditor/Private/PCGEditorGraphNodeBase.cpp
+++ b/Engine/Plugins/Experimental/PCG/Source/PCGEditor/Private/PCGEditorGraphNodeBase.cpp
@@ -572,4 +572,32 @@ void UPCGEditorGraphNodeBase::UpdatePosition()
 	}
 }
 
+void UPCGEditorGraphNodeBase::CreatePins(const TArray<UPCGPin*>& InInputPins, const TArray<UPCGPin*>& InOutputPins)
+{
+	bool bHasAdvancedPin = false;
+
+	for (const UPCGPin* InputPin : InInputPins)
+	{
+		UEdGraphPin* Pin = CreatePin(EEdGraphPinDirection::EGPD_Input, GetPinType(InputPin), InputPin->Properties.Label);
+		Pin->bAdvancedView = InputPin->Properties.bAdvancedPin;
+		bHasAdvancedPin |= Pin->bAdvancedView;
+	}
+
+	for (const UPCGPin* OutputPin : InOutputPins)
+	{
+		UEdGraphPin* Pin = CreatePin(EEdGraphPinDirection::EGPD_Output, GetPinType(OutputPin), OutputPin->Properties.Label);
+		Pin->bAdvancedView = OutputPin->Properties.bAdvancedPin;
+		bHasAdvancedPin |= Pin->bAdvancedView;
+	}
+
+	if (bHasAdvancedPin && AdvancedPinDisplay == ENodeAdvancedPins::NoPins)
+	{
+		AdvancedPinDisplay = ENodeAdvancedPins::Hidden;
+	}
+	else if (!bHasAdvancedPin)
+	{
+		AdvancedPinDisplay = ENodeAdvancedPins::NoPins;
+	}
+}
+
 #undef LOCTEXT_NAMESPACE
diff --git a/Engine/Plugins/Experimental/PCG/Source/PCGEditor/Private/PCGEditorGraphNodeBase.h b/Engine/Plugins/Experimental/PCG/Source/PCGEditor/Private/PCGEditorGraphNodeBase.h
--- a/Engine/Plugins/Experimental/PCG/Source/PCGEditor/Private/PCGEditorGraphNodeBase.h
+++ b/Engine/Plugins/Experimental/PCG/Source/PCGEditor/Private/PCGEditorGraphNodeBase.h
@@ -63,6 +63,8 @@ protected:
 	void UpdateCommentBubblePinned();
 	void UpdatePosition();
 
+	void CreatePins(const TArray<UPCGPin*>& InInputPins, const TArray<UPCGPin*>& InOutputPins);
+
 	UPROPERTY()
 	TObjectPtr<UPCGNode> PCGNode = nullptr;
 
diff --git a/Engine/Plugins/Experimental/PCG/Source/PCGEditor/Private/PCGEditorGraphNodeInput.cpp b/Engine/Plugins/Experimental/PCG/Source/PCGEditor/Private/PCGEditorGraphNodeInput.cpp
--- a/Engine/Plugins/Experimental/PCG/Source/PCGEditor/Private/PCGEditorGraphNodeInput.cpp
+++ b/Engine/Plugins/Experimental/PCG/Source/PCGEditor/Private/PCGEditorGraphNodeInput.cpp
@@ -17,10 +17,7 @@ void UPCGEditorGraphNodeInput::AllocateDefaultPins()
 {
 	if (PCGNode)
 	{
-		for (const UPCGPin* OutPin : PCGNode->GetOutputPins())
-		{
-			CreatePin(EEdGraphPinDirection::EGPD_Output, GetPinType(OutPin), OutPin->Properties.Label);
-		}
+		CreatePins(/*InInputPins=*/{}, PCGNode->GetOutputPins());
 	}
 }
 
diff --git a/Engine/Plugins/Experimental/PCG/Source/PCGEditor/Private/PCGEditorGraphNodeOutput.cpp b/Engine/Plugins/Experimental/PCG/Source/PCGEditor/Private/PCGEditorGraphNodeOutput.cpp
--- a/Engine/Plugins/Experimental/PCG/Source/PCGEditor/Private/PCGEditorGraphNodeOutput.cpp
+++ b/Engine/Plugins/Experimental/PCG/Source/PCGEditor/Private/PCGEditorGraphNodeOutput.cpp
@@ -17,10 +17,7 @@ void UPCGEditorGraphNodeOutput::AllocateDefaultPins()
 {
 	if (PCGNode)
 	{
-		for (const UPCGPin* InputPin : PCGNode->GetInputPins())
-		{
-			CreatePin(EEdGraphPinDirection::EGPD_Input, GetPinType(InputPin), InputPin->Properties.Label);
-		}
+		CreatePins(PCGNode->GetInputPins(), /*InOutputPins=*/{});
 	}
 }
 
