Added PCG Distance node to calculate a signed distance between point arrays.

diff --git a/Engine/Plugins/Experimental/PCG/Source/PCG/Private/Elements/PCGDistance.cpp b/Engine/Plugins/Experimental/PCG/Source/PCG/Private/Elements/PCGDistance.cpp
new file mode 100644
--- /dev/null
+++ b/Engine/Plugins/Experimental/PCG/Source/PCG/Private/Elements/PCGDistance.cpp
@@ -0,0 +1,198 @@
+// Copyright Epic Games, Inc. All Rights Reserved.
+
+#include "Elements/PCGDistance.h"
+#include "Helpers/PCGSettingsHelpers.h"
+#include "Math/Box2D.h"
+#include "Data/PCGPointData.h"
+
+namespace PCGDistance
+{
+	const FName SourceLabel = TEXT("Source");
+	const FName TargetLabel = TEXT("Target");
+
+	FVector CalcPosition(PCGDistanceShape Shape, const FPCGPoint& SourcePoint, const FPCGPoint& TargetPoint, const FVector SourceCenter, const FVector TargetCenter)
+	{
+		if (Shape == PCGDistanceShape::SphereBounds)
+		{
+			FVector Dir = TargetCenter - SourceCenter;
+			Dir.Normalize();
+
+			return SourceCenter + Dir * ((SourcePoint.BoundsMax-SourcePoint.BoundsMin) * SourcePoint.Transform.GetScale3D()).Length() * 0.5;
+		}
+		else if (Shape == PCGDistanceShape::BoxBounds)
+		{
+			const FVector LocalTargetCenter = SourcePoint.Transform.Inverse().TransformPosition(TargetCenter);
+
+			const double DistanceSquared = ComputeSquaredDistanceFromBoxToPoint(SourcePoint.BoundsMin, SourcePoint.BoundsMax, LocalTargetCenter);
+
+			FVector Dir = -LocalTargetCenter;
+			Dir.Normalize();			
+
+			const FVector LocalClosestPoint = LocalTargetCenter + Dir * FMath::Sqrt(DistanceSquared);
+
+			return SourcePoint.Transform.TransformPosition(LocalClosestPoint);
+		}
+
+		// PCGDistanceShape::Center
+		return SourceCenter;
+	}
+}
+	
+TArray<FPCGPinProperties> UPCGDistanceSettings::InputPinProperties() const
+{
+	TArray<FPCGPinProperties> PinProperties;
+	PinProperties.Emplace(PCGDistance::SourceLabel, EPCGDataType::Spatial);
+	PinProperties.Emplace(PCGDistance::TargetLabel, EPCGDataType::Spatial);
+	PinProperties.Emplace(PCGPinConstants::DefaultParamsLabel, EPCGDataType::Param, /*bInAllowMultipleConnections*/ false);
+
+	return PinProperties;
+}
+
+TArray<FPCGPinProperties> UPCGDistanceSettings::OutputPinProperties() const
+{
+	TArray<FPCGPinProperties> PinProperties;
+	PinProperties.Emplace(PCGPinConstants::DefaultOutputLabel, EPCGDataType::Spatial);
+
+	return PinProperties;
+}
+
+FPCGElementPtr UPCGDistanceSettings::CreateElement() const
+{
+	return MakeShared<FPCGDistanceElement>();
+}
+
+bool FPCGDistanceElement::ExecuteInternal(FPCGContext* Context) const
+{
+	TRACE_CPUPROFILER_EVENT_SCOPE(FPCGDistanceElement::Execute);
+
+	const UPCGDistanceSettings* Settings = Context->GetInputSettings<UPCGDistanceSettings>();
+	check(Settings);
+
+	UPCGParamData* Params = Context->InputData.GetParams();
+
+	const FName AttributeName = PCGSettingsHelpers::GetValue(GET_MEMBER_NAME_CHECKED(UPCGDistanceSettings, AttributeName), Settings->AttributeName, Params);
+	const bool bSetDensity = PCGSettingsHelpers::GetValue(GET_MEMBER_NAME_CHECKED(UPCGDistanceSettings, bSetDensity), Settings->bSetDensity, Params);
+	const double MaximumDistance = PCGSettingsHelpers::GetValue(GET_MEMBER_NAME_CHECKED(UPCGDistanceSettings, MaximumDistance), Settings->MaximumDistance, Params);
+	const PCGDistanceShape SourceShape = PCGSettingsHelpers::GetValue(GET_MEMBER_NAME_CHECKED(UPCGDistanceSettings, SourceShape), Settings->SourceShape, Params);
+	const PCGDistanceShape TargetShape = PCGSettingsHelpers::GetValue(GET_MEMBER_NAME_CHECKED(UPCGDistanceSettings, TargetShape), Settings->TargetShape, Params);
+
+	TArray<FPCGTaggedData> Sources = Context->InputData.GetInputsByPin(PCGDistance::SourceLabel);
+	TArray<FPCGTaggedData> Targets = Context->InputData.GetInputsByPin(PCGDistance::TargetLabel);
+	TArray<FPCGTaggedData>& Outputs = Context->OutputData.TaggedData;
+
+	FBox InputBounds(EForceInit::ForceInit);
+	int32 InputPointCount = 0;
+
+	TArray<const UPCGPointData*> TargetPointDatas;
+	TargetPointDatas.Reserve(Targets.Num());
+
+	for (const FPCGTaggedData& Target : Targets)
+	{
+
+		const UPCGSpatialData* TargetData = Cast<UPCGSpatialData>(Target.Data);
+
+		if (!TargetData)
+		{
+			PCGE_LOG(Error, TEXT("Target must be spatial data, found '%s'"), *Target.Data->GetClass()->GetName());
+			continue;
+		}
+
+		const UPCGPointData* TargetPointData = TargetData->ToPointData(Context);
+		if (!TargetPointData)
+		{
+			PCGE_LOG(Error, TEXT("Cannot convert target '%s' into point data"), *Target.Data->GetClass()->GetName());
+			continue;			
+		}
+
+		TargetPointDatas.Add(TargetPointData);
+	}
+
+	// first find the total Input bounds which will determine the size of each cell
+	for (const FPCGTaggedData& Source : Sources) 
+	{
+		// add the point bounds to the input cell
+
+		const UPCGSpatialData* SourceData = Cast<UPCGSpatialData>(Source.Data);
+
+		if (!SourceData)
+		{
+			PCGE_LOG(Error, TEXT("Invalid input data"));
+			continue;
+		}
+
+		const UPCGPointData* SourcePointData = SourceData->ToPointData(Context);
+		if (!SourcePointData)
+		{
+			PCGE_LOG(Error, TEXT("Cannot convert input spatial data to point data"));
+			continue;			
+		}
+
+		UPCGPointData* OutputData = NewObject<UPCGPointData>();
+		OutputData->InitializeFromData(SourcePointData);
+		Outputs.Add_GetRef(Source).Data = OutputData;
+
+		FPCGMetadataAttribute<float>* Attribute = (AttributeName != NAME_None) ? OutputData->Metadata->FindOrCreateAttribute<float>(AttributeName, 0.0f) : nullptr;
+
+		FPCGAsync::AsyncPointProcessing(Context, SourcePointData->GetPoints(), OutputData->GetMutablePoints(),
+			[OutputData, SourceShape, TargetShape, &TargetPointDatas, MaximumDistance, Attribute, bSetDensity](const FPCGPoint& SourcePoint, FPCGPoint& OutPoint) {
+
+				OutPoint = SourcePoint;
+
+				const FBoxSphereBounds SourceQueryBounds = FBoxSphereBounds(FBox(SourcePoint.BoundsMin - FVector(MaximumDistance), SourcePoint.BoundsMax + FVector(MaximumDistance))).TransformBy(SourcePoint.Transform);
+
+				const FVector SourceCenter = SourcePoint.Transform.TransformPosition(SourcePoint.GetLocalCenter());
+
+				double DistanceSquared = MaximumDistance*MaximumDistance;
+
+				auto CalculateSDF = [&DistanceSquared, &SourcePoint, SourceCenter, SourceShape, TargetShape](const FPCGPointRef& TargetPointRef)
+				{
+					const FPCGPoint& TargetPoint = *TargetPointRef.Point;
+
+					const FVector TargetCenter = TargetPoint.Transform.TransformPosition(TargetPoint.GetLocalCenter());
+
+					const FVector SourceShapePos = PCGDistance::CalcPosition(SourceShape, SourcePoint, TargetPoint, SourceCenter, TargetCenter);
+					const FVector TargetShapePos = PCGDistance::CalcPosition(TargetShape, TargetPoint, SourcePoint, TargetCenter, SourceCenter);
+
+					const FVector ToTargetShapeDir = TargetShapePos - SourceShapePos;
+					const FVector ToTargetCenterDir = TargetCenter - SourceCenter;
+
+					const double Sign = ToTargetShapeDir.Dot(ToTargetCenterDir) > 0 ? 1.0 : -1.0;
+					const double ThisDistanceSquared = ToTargetShapeDir.SquaredLength() * Sign;
+
+					DistanceSquared = FMath::Min(DistanceSquared, ThisDistanceSquared);
+				};
+
+				for (const UPCGPointData* TargetPointData : TargetPointDatas)
+				{
+					const UPCGPointData::PointOctree& Octree = TargetPointData->GetOctree();
+
+					Octree.FindElementsWithBoundsTest(
+						FBoxCenterAndExtent(SourceQueryBounds.Origin, SourceQueryBounds.BoxExtent),
+						CalculateSDF
+					);
+				}
+
+				const float Distance = FMath::Sqrt(FMath::Abs(DistanceSquared)) * (DistanceSquared < 0 ? -1.0 : 1.0);
+
+				if (Attribute)
+				{
+					OutputData->Metadata->InitializeOnSet(OutPoint.MetadataEntry);
+					Attribute->SetValue(OutPoint.MetadataEntry, Distance);
+				}
+				
+				if (bSetDensity)
+				{
+					// set density instead
+					OutPoint.Density = FMath::Clamp(Distance, -MaximumDistance, MaximumDistance) / MaximumDistance;
+				}
+
+				return true;
+			}
+		);		
+	}
+
+	// Forward any non-input data
+	Outputs.Append(Context->InputData.GetAllSettings());
+	return true;
+
+}
diff --git a/Engine/Plugins/Experimental/PCG/Source/PCG/Private/Tests/Elements/PCGDistanceTest.cpp b/Engine/Plugins/Experimental/PCG/Source/PCG/Private/Tests/Elements/PCGDistanceTest.cpp
new file mode 100644
--- /dev/null
+++ b/Engine/Plugins/Experimental/PCG/Source/PCG/Private/Tests/Elements/PCGDistanceTest.cpp
@@ -0,0 +1,260 @@
+// Copyright Epic Games, Inc. All Rights Reserved.
+
+#include "Tests/Determinism/PCGDeterminismTestsCommon.h"
+#include "Tests/PCGTestsCommon.h"
+
+#include "PCGComponent.h"
+#include "PCGData.h"
+#include "PCGHelpers.h"
+#include "PCGParamData.h"
+#include "Data/PCGPointData.h"
+
+#include "Elements/PCGDistance.h"
+
+IMPLEMENT_CUSTOM_SIMPLE_AUTOMATION_TEST(FPCGDistanceTest_PointToPoint, FPCGTestBaseClass, "pcg.tests.Distance.PointToPoint", PCGTestsCommon::TestFlags)
+
+bool FPCGDistanceTest_PointToPoint::RunTest(const FString& Parameters)
+{
+	PCGTestsCommon::FTestData TestData(PCGDeterminismTests::Defaults::Seed);
+	PCGDeterminismTests::GenerateSettings<UPCGDistanceSettings>(TestData);
+	UPCGDistanceSettings* Settings = CastChecked<UPCGDistanceSettings>(TestData.Settings);
+
+	Settings->SourceShape = PCGDistanceShape::Center;
+	Settings->TargetShape = PCGDistanceShape::Center;
+
+	FPCGElementPtr TestElement = TestData.Settings->GetElement();
+
+	{
+		UPCGPointData *SourceData = PCGTestsCommon::CreateRandomPointData(2, 42);
+
+		TArray<FPCGPoint>& SourcePoints = SourceData->GetMutablePoints();
+		SourcePoints[0].Transform.SetTranslation(FVector(100, 0, 0));
+		SourcePoints[1].Transform.SetTranslation(FVector(0, 50, 0));
+
+		FPCGTaggedData& SourcePin = TestData.InputData.TaggedData.Emplace_GetRef();
+		SourcePin.Pin = PCGDistance::SourceLabel;
+		SourcePin.Data = SourceData;
+	}
+
+	{
+		FPCGTaggedData& TargetPin = TestData.InputData.TaggedData.Emplace_GetRef();
+		TargetPin.Pin = PCGDistance::TargetLabel;
+		TargetPin.Data = PCGTestsCommon::CreatePointData(FVector(0, 0, 0));		
+	}
+
+	TUniquePtr<FPCGContext> Context(TestElement->Initialize(TestData.InputData, TestData.TestPCGComponent, nullptr));
+	Context->NumAvailableTasks = 1;
+
+	while (!TestElement->Execute(Context.Get())) {}
+
+	// test our point data
+	const TArray<FPCGTaggedData>& Outputs = Context->OutputData.GetInputsByPin(PCGPinConstants::DefaultOutputLabel);
+
+	UTEST_EQUAL("Output count", Outputs.Num(), 2);
+
+	UPCGPointData *OutPointData = Cast<UPCGPointData>(Outputs[0].Data);
+
+	UTEST_NOT_NULL("Output point data", OutPointData);
+
+	const TArray<FPCGPoint>& OutPoints = OutPointData->GetPoints();
+
+	UTEST_EQUAL("Output point count", OutPoints.Num(), 2);
+
+	const FPCGMetadataAttribute<float> *DistanceAttribute = OutPointData->Metadata->GetConstTypedAttribute<float>(Settings->AttributeName);
+	UTEST_NOT_NULL("Distance attribute", DistanceAttribute);
+
+	UTEST_EQUAL_TOLERANCE("Point 0 distance", DistanceAttribute->GetValue(OutPoints[0].MetadataEntry), 100.f, 0.01f);
+	UTEST_EQUAL_TOLERANCE("Point 1 distance", DistanceAttribute->GetValue(OutPoints[1].MetadataEntry), 50.f, 0.01f);
+
+	return true;
+}
+
+IMPLEMENT_CUSTOM_SIMPLE_AUTOMATION_TEST(FPCGDistanceTest_SetDensity, FPCGTestBaseClass, "pcg.tests.Distance.SetDensity", PCGTestsCommon::TestFlags)
+
+bool FPCGDistanceTest_SetDensity::RunTest(const FString& Parameters)
+{
+	PCGTestsCommon::FTestData TestData(PCGDeterminismTests::Defaults::Seed);
+	PCGDeterminismTests::GenerateSettings<UPCGDistanceSettings>(TestData);
+	UPCGDistanceSettings* Settings = CastChecked<UPCGDistanceSettings>(TestData.Settings);
+
+	Settings->SourceShape = PCGDistanceShape::Center;
+	Settings->TargetShape = PCGDistanceShape::Center;
+	Settings->AttributeName = NAME_None;
+	Settings->bSetDensity = true;
+	Settings->MaximumDistance = 200.f;
+
+	FPCGElementPtr TestElement = TestData.Settings->GetElement();
+
+	{
+		UPCGPointData *SourceData = PCGTestsCommon::CreateRandomPointData(2, 42);
+
+		TArray<FPCGPoint>& SourcePoints = SourceData->GetMutablePoints();
+		SourcePoints[0].Transform.SetTranslation(FVector(100, 0, 0));
+		SourcePoints[1].Transform.SetTranslation(FVector(0, 50, 0));
+
+		FPCGTaggedData& SourcePin = TestData.InputData.TaggedData.Emplace_GetRef();
+		SourcePin.Pin = PCGDistance::SourceLabel;
+		SourcePin.Data = SourceData;
+	}
+
+	{
+		FPCGTaggedData& TargetPin = TestData.InputData.TaggedData.Emplace_GetRef();
+		TargetPin.Pin = PCGDistance::TargetLabel;
+		TargetPin.Data = PCGTestsCommon::CreatePointData(FVector(0, 0, 0));		
+	}
+
+	TUniquePtr<FPCGContext> Context(TestElement->Initialize(TestData.InputData, TestData.TestPCGComponent, nullptr));
+	Context->NumAvailableTasks = 1;
+
+	while (!TestElement->Execute(Context.Get())) {}
+
+	// test our point data
+	const TArray<FPCGTaggedData>& Outputs = Context->OutputData.GetInputsByPin(PCGPinConstants::DefaultOutputLabel);
+
+	UTEST_EQUAL("Output count", Outputs.Num(), 2);
+
+	UPCGPointData *OutPointData = Cast<UPCGPointData>(Outputs[0].Data);
+
+	UTEST_NOT_NULL("Output point data", OutPointData);
+
+	const TArray<FPCGPoint>& OutPoints = OutPointData->GetPoints();
+
+	UTEST_EQUAL("Output point count", OutPoints.Num(), 2);
+
+	UTEST_EQUAL_TOLERANCE("Point 0 distance", OutPoints[0].Density, 0.5f, 0.01f);
+	UTEST_EQUAL_TOLERANCE("Point 1 distance", OutPoints[1].Density, 0.25f, 0.01f);
+
+	return true;
+}
+
+IMPLEMENT_CUSTOM_SIMPLE_AUTOMATION_TEST(FPCGDistanceTest_PointToSphere, FPCGTestBaseClass, "pcg.tests.Distance.PointToSphere", PCGTestsCommon::TestFlags)
+
+bool FPCGDistanceTest_PointToSphere::RunTest(const FString& Parameters)
+{
+	PCGTestsCommon::FTestData TestData(PCGDeterminismTests::Defaults::Seed);
+	PCGDeterminismTests::GenerateSettings<UPCGDistanceSettings>(TestData);
+	UPCGDistanceSettings* Settings = CastChecked<UPCGDistanceSettings>(TestData.Settings);
+
+	Settings->SourceShape = PCGDistanceShape::Center;
+	Settings->TargetShape = PCGDistanceShape::SphereBounds;
+
+	FPCGElementPtr TestElement = TestData.Settings->GetElement();
+
+	{
+		UPCGPointData *SourceData = PCGTestsCommon::CreateRandomPointData(2, 42);
+
+		TArray<FPCGPoint>& SourcePoints = SourceData->GetMutablePoints();
+		SourcePoints[0].Transform.SetTranslation(FVector(100, 0, 0));
+		SourcePoints[1].Transform.SetTranslation(FVector(0, 50, 0));
+
+		FPCGTaggedData& SourcePin = TestData.InputData.TaggedData.Emplace_GetRef();
+		SourcePin.Pin = PCGDistance::SourceLabel;
+		SourcePin.Data = SourceData;
+	}
+
+	const FVector TargetExtents(10.0f);
+	const float TargetPointRadius = TargetExtents.Length();
+
+	{
+		FPCGTaggedData& TargetPin = TestData.InputData.TaggedData.Emplace_GetRef();
+		TargetPin.Pin = PCGDistance::TargetLabel;
+
+		UPCGPointData *TargetData = PCGTestsCommon::CreatePointData(FVector(0, 0, 0));
+		TArray<FPCGPoint>& TargetPoints = TargetData->GetMutablePoints();
+		TargetPoints[0].SetExtents(TargetExtents);
+
+		TargetPin.Data = TargetData;
+	}
+
+	TUniquePtr<FPCGContext> Context(TestElement->Initialize(TestData.InputData, TestData.TestPCGComponent, nullptr));
+	Context->NumAvailableTasks = 1;
+
+	while (!TestElement->Execute(Context.Get())) {}
+
+	// test our point data
+	const TArray<FPCGTaggedData>& Outputs = Context->OutputData.GetInputsByPin(PCGPinConstants::DefaultOutputLabel);
+
+	UTEST_EQUAL("Output count", Outputs.Num(), 2);
+
+	UPCGPointData *OutPointData = Cast<UPCGPointData>(Outputs[0].Data);
+
+	UTEST_NOT_NULL("Output point data", OutPointData);
+
+	const TArray<FPCGPoint>& OutPoints = OutPointData->GetPoints();
+
+	UTEST_EQUAL("Output point count", OutPoints.Num(), 2);
+
+	const FPCGMetadataAttribute<float> *DistanceAttribute = OutPointData->Metadata->GetConstTypedAttribute<float>(Settings->AttributeName);
+	UTEST_NOT_NULL("Distance attribute", DistanceAttribute);
+
+	UTEST_EQUAL_TOLERANCE("Point 0 distance", DistanceAttribute->GetValue(OutPoints[0].MetadataEntry), 100.f-TargetPointRadius, 0.01f);
+	UTEST_EQUAL_TOLERANCE("Point 1 distance", DistanceAttribute->GetValue(OutPoints[1].MetadataEntry), 50.f-TargetPointRadius, 0.01f);
+
+	return true;
+}
+
+IMPLEMENT_CUSTOM_SIMPLE_AUTOMATION_TEST(FPCGDistanceTest_PointToBox, FPCGTestBaseClass, "pcg.tests.Distance.PointToBox", PCGTestsCommon::TestFlags)
+
+bool FPCGDistanceTest_PointToBox::RunTest(const FString& Parameters)
+{
+	PCGTestsCommon::FTestData TestData(PCGDeterminismTests::Defaults::Seed);
+	PCGDeterminismTests::GenerateSettings<UPCGDistanceSettings>(TestData);
+	UPCGDistanceSettings* Settings = CastChecked<UPCGDistanceSettings>(TestData.Settings);
+
+	Settings->SourceShape = PCGDistanceShape::Center;
+	Settings->TargetShape = PCGDistanceShape::BoxBounds;
+
+	FPCGElementPtr TestElement = TestData.Settings->GetElement();
+
+	{
+		UPCGPointData *SourceData = PCGTestsCommon::CreateRandomPointData(2, 42);
+
+		TArray<FPCGPoint>& SourcePoints = SourceData->GetMutablePoints();
+		SourcePoints[0].Transform.SetTranslation(FVector(100, 0, 0));
+		SourcePoints[1].Transform.SetTranslation(FVector(0, 50, 0));
+
+		FPCGTaggedData& SourcePin = TestData.InputData.TaggedData.Emplace_GetRef();
+		SourcePin.Pin = PCGDistance::SourceLabel;
+		SourcePin.Data = SourceData;
+	}
+
+	// should put the box 10 units away from each point
+	const FVector TargetExtents(90.0f, 40.f, 10.f);
+
+	{
+		FPCGTaggedData& TargetPin = TestData.InputData.TaggedData.Emplace_GetRef();
+		TargetPin.Pin = PCGDistance::TargetLabel;
+
+		UPCGPointData *TargetData = PCGTestsCommon::CreatePointData(FVector(0, 0, 0));
+		TArray<FPCGPoint>& TargetPoints = TargetData->GetMutablePoints();
+		TargetPoints[0].SetExtents(TargetExtents);
+
+		TargetPin.Data = TargetData;
+	}
+
+	TUniquePtr<FPCGContext> Context(TestElement->Initialize(TestData.InputData, TestData.TestPCGComponent, nullptr));
+	Context->NumAvailableTasks = 1;
+
+	while (!TestElement->Execute(Context.Get())) {}
+
+	// test our point data
+	const TArray<FPCGTaggedData>& Outputs = Context->OutputData.GetInputsByPin(PCGPinConstants::DefaultOutputLabel);
+
+	UTEST_EQUAL("Output count", Outputs.Num(), 2);
+
+	UPCGPointData *OutPointData = Cast<UPCGPointData>(Outputs[0].Data);
+
+	UTEST_NOT_NULL("Output point data", OutPointData);
+
+	const TArray<FPCGPoint>& OutPoints = OutPointData->GetPoints();
+
+	UTEST_EQUAL("Output point count", OutPoints.Num(), 2);
+
+	const FPCGMetadataAttribute<float> *DistanceAttribute = OutPointData->Metadata->GetConstTypedAttribute<float>(Settings->AttributeName);
+	UTEST_NOT_NULL("Distance attribute", DistanceAttribute);
+
+	UTEST_EQUAL_TOLERANCE("Point 0 distance", DistanceAttribute->GetValue(OutPoints[0].MetadataEntry), 10.f, 0.01f);
+	UTEST_EQUAL_TOLERANCE("Point 1 distance", DistanceAttribute->GetValue(OutPoints[1].MetadataEntry), 10.f, 0.01f);
+
+	return true;
+}
diff --git a/Engine/Plugins/Experimental/PCG/Source/PCG/Public/Elements/PCGDistance.h b/Engine/Plugins/Experimental/PCG/Source/PCG/Public/Elements/PCGDistance.h
new file mode 100644
--- /dev/null
+++ b/Engine/Plugins/Experimental/PCG/Source/PCG/Public/Elements/PCGDistance.h
@@ -0,0 +1,74 @@
+// Copyright Epic Games, Inc. All Rights Reserved.
+#pragma once
+
+#include "CoreMinimal.h"
+#include "PCGSettings.h"
+#include "Elements/PCGPointProcessingElementBase.h"
+
+#include "PCGDistance.generated.h"
+
+UENUM()
+enum class PCGDistanceShape
+{
+	SphereBounds,
+	BoxBounds,
+	Center,
+};
+
+namespace PCGDistance
+{
+	extern const FName SourceLabel;
+	extern const FName TargetLabel;
+}
+
+/**
+ * Calculates the distance between two points (inherently a n*n operation)
+ */
+UCLASS(BlueprintType, ClassGroup = (Procedural))
+class UPCGDistanceSettings : public UPCGSettings
+{
+	GENERATED_BODY()
+
+public:
+	//~Begin UPCGSettings interface
+#if WITH_EDITOR
+	virtual FName GetDefaultNodeName() const override { return FName(TEXT("Distance")); }
+	virtual EPCGSettingsType GetType() const override { return EPCGSettingsType::Spatial; }
+#endif
+
+	virtual TArray<FPCGPinProperties> InputPinProperties() const override;
+	virtual TArray<FPCGPinProperties> OutputPinProperties() const override;
+
+protected:
+	virtual FPCGElementPtr CreateElement() const override;
+	//~End UPCGSettings interface
+
+public:
+
+	// The name of the attribute to store on the point. Use 'None' to disable
+	UPROPERTY(BlueprintReadWrite, EditAnywhere, Category = Settings)
+	FName AttributeName = TEXT("Distance");
+
+	// If true, will also set the density to be 0-1 based on MaximumDistance
+	UPROPERTY(BlueprintReadWrite, EditAnywhere, Category = Settings)
+	bool bSetDensity = false;
+
+	// A maximum distance to search, which is used as an optimization
+	UPROPERTY(BlueprintReadWrite, EditAnywhere, Category = Settings, meta = (ClampMin = "1"))
+	double MaximumDistance = 20000.0;
+
+	// What shape is used on the 'source' points
+	UPROPERTY(BlueprintReadWrite, EditAnywhere, Category = Settings)
+	PCGDistanceShape SourceShape = PCGDistanceShape::SphereBounds;
+
+	// What shape is used on the 'target' points
+	UPROPERTY(BlueprintReadWrite, EditAnywhere, Category = Settings)
+	PCGDistanceShape TargetShape = PCGDistanceShape::SphereBounds;
+};
+
+class FPCGDistanceElement : public FPCGPointProcessingElementBase
+{
+protected:
+	virtual bool ExecuteInternal(FPCGContext* Context) const override;
+};
+
