Added a match & set node (& some variations, and an extensible pattern) to implement generic look-ups.

diff --git a/Engine/Plugins/Experimental/PCG/Source/PCG/Private/Elements/PCGPointMatchAndSet.cpp b/Engine/Plugins/Experimental/PCG/Source/PCG/Private/Elements/PCGPointMatchAndSet.cpp
new file mode 100644
--- /dev/null
+++ b/Engine/Plugins/Experimental/PCG/Source/PCG/Private/Elements/PCGPointMatchAndSet.cpp
@@ -0,0 +1,187 @@
+// Copyright Epic Games, Inc. All Rights Reserved.
+
+#include "Elements/PCGPointMatchAndSet.h"
+
+#include "PCGHelpers.h"
+#include "Data/PCGPointData.h"
+#include "MatchAndSet/PCGMatchAndSetWeighted.h"
+#include "Metadata/Accessors/IPCGAttributeAccessor.h"
+#include "Metadata/Accessors/PCGAttributeAccessorHelpers.h"
+
+#include UE_INLINE_GENERATED_CPP_BY_NAME(PCGPointMatchAndSet)
+
+#define LOCTEXT_NAMESPACE "PCGPointMatchAndSetElement"
+
+UPCGPointMatchAndSetSettings::UPCGPointMatchAndSetSettings(const FObjectInitializer& ObjectInitializer)
+{
+	MatchAndSetType = UPCGMatchAndSetWeighted::StaticClass();
+	MatchAndSetInstance = ObjectInitializer.CreateDefaultSubobject<UPCGMatchAndSetWeighted>(this, TEXT("DefaultMatchAndSetInstance"));
+	bUseSeed = MatchAndSetInstance && MatchAndSetInstance->UsesRandomProcess();
+}
+
+#if WITH_EDITOR
+FText UPCGPointMatchAndSetSettings::GetNodeTooltipText() const
+{
+	return LOCTEXT("PointMatchAndSetNodeTooltip", "For all points, if a match is found (e.g. some attribute is equal to some value), sets a value on the point (e.g. another attribute).");
+}
+#endif // WITH_EDITOR
+
+TArray<FPCGPinProperties> UPCGPointMatchAndSetSettings::InputPinProperties() const
+{
+	TArray<FPCGPinProperties> PinProperties;
+
+	PinProperties.Emplace(PCGPinConstants::DefaultInputLabel, EPCGDataType::Point);
+	// TODO Add param support?
+
+	return PinProperties;
+}
+
+TArray<FPCGPinProperties> UPCGPointMatchAndSetSettings::OutputPinProperties() const
+{
+	TArray<FPCGPinProperties> PinProperties;
+	PinProperties.Emplace(PCGPinConstants::DefaultOutputLabel, EPCGDataType::Point);
+
+	return PinProperties;
+}
+
+FPCGElementPtr UPCGPointMatchAndSetSettings::CreateElement() const
+{
+	return MakeShared<FPCGPointMatchAndSetElement>();
+}
+
+#if WITH_EDITOR
+
+void UPCGPointMatchAndSetSettings::PostEditChangeProperty(FPropertyChangedEvent& PropertyChangedEvent)
+{
+	if (PropertyChangedEvent.Property)
+	{
+		const FName& PropertyName = PropertyChangedEvent.Property->GetFName();
+
+		if (PropertyName == GET_MEMBER_NAME_CHECKED(UPCGPointMatchAndSetSettings, MatchAndSetType))
+		{
+			RefreshMatchAndSet();
+		}
+		else if (PropertyName == GET_MEMBER_NAME_CHECKED(UPCGPointMatchAndSetSettings, SetTarget))
+		{
+			bSetTargetIsAttribute = (SetTarget.Selection == EPCGAttributePropertySelection::Attribute);
+
+			// If not targetting an attribute, but a property, assign the new type to the OutputType accordingly
+			if (!bSetTargetIsAttribute)
+			{
+				TUniquePtr<const IPCGAttributeAccessor> SetTargetAccessor = PCGAttributeAccessorHelpers::CreateConstAccessor(nullptr, SetTarget);
+				if (SetTargetAccessor)
+				{
+					SetTargetType = static_cast<EPCGMetadataTypes>(SetTargetAccessor->GetUnderlyingType());
+
+					if (MatchAndSetInstance)
+					{
+						MatchAndSetInstance->SetType(SetTargetType, SetTargetStringMode);
+					}
+				}
+			}
+		}
+		else if (PropertyName == GET_MEMBER_NAME_CHECKED(UPCGPointMatchAndSetSettings, SetTargetType) ||
+			PropertyName == GET_MEMBER_NAME_CHECKED(UPCGPointMatchAndSetSettings, SetTargetStringMode))
+		{
+			if (MatchAndSetInstance)
+			{
+				MatchAndSetInstance->SetType(SetTargetType, SetTargetStringMode);
+			}
+		}
+	}
+
+	Super::PostEditChangeProperty(PropertyChangedEvent);
+}
+
+#endif // WITH_EDITOR
+
+void UPCGPointMatchAndSetSettings::SetMatchAndSetType(TSubclassOf<UPCGMatchAndSetBase> InMatchAndSetType)
+{
+	if (!MatchAndSetInstance || InMatchAndSetType != MatchAndSetType)
+	{
+		MatchAndSetType = InMatchAndSetType;
+		RefreshMatchAndSet();
+	}
+}
+
+void UPCGPointMatchAndSetSettings::RefreshMatchAndSet()
+{
+	if (MatchAndSetType)
+	{
+		MatchAndSetInstance = NewObject<UPCGMatchAndSetBase>(this, MatchAndSetType);
+		check(MatchAndSetInstance);
+		MatchAndSetInstance->SetType(SetTargetType, SetTargetStringMode);
+	}
+	else
+	{
+		MatchAndSetInstance = nullptr;
+	}
+
+	bUseSeed = MatchAndSetInstance && MatchAndSetInstance->UsesRandomProcess();
+}
+
+bool FPCGPointMatchAndSetElement::ExecuteInternal(FPCGContext* Context) const
+{
+	TRACE_CPUPROFILER_EVENT_SCOPE(FPCGPointMatchAndSetElement::Execute);
+
+	const UPCGPointMatchAndSetSettings* Settings = Context->GetInputSettings<UPCGPointMatchAndSetSettings>();
+	check(Settings);
+
+	const UPCGMatchAndSetBase* MatchAndSet = Settings->MatchAndSetInstance;
+
+	if (!MatchAndSet)
+	{
+		PCGE_LOG(Error, "Invalid MatchAndSet instance");
+		return true;
+	}
+
+	TArray<FPCGTaggedData> Inputs = Context->InputData.GetInputs();
+	TArray<FPCGTaggedData>& Outputs = Context->OutputData.TaggedData;
+
+	// Forward any non-input data
+	Outputs.Append(Context->InputData.GetAllSettings());
+
+	for (const FPCGTaggedData& Input : Inputs)
+	{
+		FPCGTaggedData& Output = Outputs.Add_GetRef(Input);
+
+		const UPCGPointData* InPointData = Cast<UPCGPointData>(Input.Data);
+		if (!InPointData)
+		{
+			PCGE_LOG(Error, "Invalid input data type");
+			continue;
+		}
+
+		// Perform validation - the MatchAndSet is responsible to validate if the input
+		// conforms to what is expected (esp. on the metadata side)
+		if (!MatchAndSet->ValidatePreconditions(InPointData))
+		{
+			PCGE_LOG(Error, "MatchAndSet failed to validate preconditions on input data");
+			continue;
+		}
+
+		UPCGPointData* OutPointData = NewObject<UPCGPointData>();
+		Output.Data = OutPointData;
+
+		OutPointData->InitializeFromData(InPointData);
+		// Copy all points
+		OutPointData->GetMutablePoints() = InPointData->GetPoints();
+
+		// Apply MatchAndSet
+		MatchAndSet->MatchAndSet(*Context, Settings, InPointData, OutPointData);
+
+		if (MatchAndSet->ShouldMutateSeed())
+		{
+			TArray<FPCGPoint>& OutPoints = OutPointData->GetMutablePoints();
+			for (int32 PointIndex = 0; PointIndex < OutPoints.Num(); ++PointIndex)
+			{
+				FPCGPoint& Point = OutPoints[PointIndex];
+				Point.Seed = PCGHelpers::ComputeSeed(Point.Seed, PointIndex);
+			}
+		}
+	}
+
+	return true;
+}
+
+#undef LOCTEXT_NAMESPACE
\ No newline at end of file
diff --git a/Engine/Plugins/Experimental/PCG/Source/PCG/Private/MatchAndSet/PCGMatchAndSetBase.cpp b/Engine/Plugins/Experimental/PCG/Source/PCG/Private/MatchAndSet/PCGMatchAndSetBase.cpp
new file mode 100644
--- /dev/null
+++ b/Engine/Plugins/Experimental/PCG/Source/PCG/Private/MatchAndSet/PCGMatchAndSetBase.cpp
@@ -0,0 +1,51 @@
+// Copyright Epic Games, Inc. All Rights Reserved.
+
+#include "MatchAndSet/PCGMatchAndSetBase.h"
+
+#include "Data/PCGPointData.h"
+#include "Elements/PCGPointMatchAndSet.h"
+#include "Elements/Metadata/PCGMetadataElementCommon.h"
+#include "Metadata/PCGAttributePropertySelector.h"
+#include "Metadata/PCGMetadata.h"
+#include "Metadata/PCGMetadataTypesConstantStruct.h"
+
+#include UE_INLINE_GENERATED_CPP_BY_NAME(PCGMatchAndSetBase)
+
+void UPCGMatchAndSetBase::SetType(EPCGMetadataTypes InType, EPCGMetadataTypesConstantStructStringMode InStringMode)
+{
+	Type = InType;
+	StringMode = InStringMode;
+}
+
+bool UPCGMatchAndSetBase::CreateAttributeIfNeeded(FPCGContext& Context, const FPCGAttributePropertySelector& Selector, const FPCGMetadataTypesConstantStruct& ConstantValue, UPCGPointData* OutPointData, const UPCGPointMatchAndSetSettings* InSettings) const
+{
+	check(OutPointData && OutPointData->Metadata);
+
+	check(OutPointData->Metadata);
+	if (Selector.Selection == EPCGAttributePropertySelection::Attribute)
+	{
+		FName DestinationAttribute = Selector.GetName();
+		if (DestinationAttribute == NAME_None)
+		{
+			DestinationAttribute = OutPointData->Metadata->GetLatestAttributeNameOrNone();
+		}
+
+		if (!OutPointData->Metadata->HasAttribute(DestinationAttribute) ||
+			OutPointData->Metadata->GetConstAttribute(DestinationAttribute)->GetTypeId() != static_cast<uint16>(InSettings->SetTargetType))
+		{
+			auto CreateAttribute = [OutPointData, &DestinationAttribute](auto&& Value)
+			{
+				using ConstantType = std::decay_t<decltype(Value)>;
+				return PCGMetadataElementCommon::ClearOrCreateAttribute(OutPointData->Metadata, DestinationAttribute, ConstantType{}) != nullptr;
+			};
+
+			if (!ConstantValue.DispatcherWithOverride(nullptr, CreateAttribute))
+			{
+				PCGE_LOG_C(Error, &Context, "Unable to create attribute %s on point data", *DestinationAttribute.ToString());
+				return false;
+			}
+		}
+	}
+
+	return true;
+}
\ No newline at end of file
diff --git a/Engine/Plugins/Experimental/PCG/Source/PCG/Private/MatchAndSet/PCGMatchAndSetByAttribute.cpp b/Engine/Plugins/Experimental/PCG/Source/PCG/Private/MatchAndSet/PCGMatchAndSetByAttribute.cpp
new file mode 100644
--- /dev/null
+++ b/Engine/Plugins/Experimental/PCG/Source/PCG/Private/MatchAndSet/PCGMatchAndSetByAttribute.cpp
@@ -0,0 +1,211 @@
+// Copyright Epic Games, Inc. All Rights Reserved.
+
+#include "MatchAndSet/PCGMatchAndSetByAttribute.h"
+
+#include "Data/PCGPointData.h"
+#include "Elements/PCGPointMatchAndSet.h"
+#include "Metadata/PCGMetadata.h"
+#include "Metadata/PCGMetadataAttribute.h"
+#include "Metadata/Accessors/IPCGAttributeAccessor.h"
+#include "Metadata/Accessors/PCGAttributeAccessorHelpers.h"
+#include "Metadata/Accessors/PCGAttributeAccessorKeys.h"
+#include "Metadata/Accessors/PCGCustomAccessor.h"
+
+#include UE_INLINE_GENERATED_CPP_BY_NAME(PCGMatchAndSetByAttribute)
+
+FPCGMatchAndSetByAttributeEntry::FPCGMatchAndSetByAttributeEntry()
+{
+	ValueToMatch.bAllowsTypeChange = false;
+	Value.bAllowsTypeChange = false;
+}
+
+void UPCGMatchAndSetByAttribute::SetType(EPCGMetadataTypes InType, EPCGMetadataTypesConstantStructStringMode InStringMode)
+{
+	for (FPCGMatchAndSetByAttributeEntry& Entry : Entries)
+	{
+		Entry.Value.Type = InType;
+		Entry.Value.StringMode = InStringMode;
+	}
+
+	Super::SetType(InType, InStringMode);
+}
+
+void UPCGMatchAndSetByAttribute::SetSourceType(EPCGMetadataTypes InType, EPCGMetadataTypesConstantStructStringMode InStringMode)
+{
+	for (FPCGMatchAndSetByAttributeEntry& Entry : Entries)
+	{
+		Entry.ValueToMatch.Type = InType;
+		Entry.ValueToMatch.StringMode = InStringMode;
+	}
+}
+
+#if WITH_EDITOR
+void UPCGMatchAndSetByAttribute::PostEditChangeProperty(FPropertyChangedEvent& PropertyChangedEvent)
+{
+	if (PropertyChangedEvent.Property)
+	{
+		const FName& PropertyName = PropertyChangedEvent.Property->GetFName();
+
+		if (PropertyName == GET_MEMBER_NAME_CHECKED(UPCGMatchAndSetByAttribute, MatchSourceType) ||
+			PropertyName == GET_MEMBER_NAME_CHECKED(UPCGMatchAndSetByAttribute, MatchSourceStringMode))
+		{
+			SetSourceType(MatchSourceType, MatchSourceStringMode);
+		}
+		else if (PropertyName == GET_MEMBER_NAME_CHECKED(UPCGMatchAndSetByAttribute, Entries))
+		{
+			// Some changes in the entries (add, insert) might require us to re-set the type
+			SetType(Type, StringMode);
+			SetSourceType(MatchSourceType, MatchSourceStringMode);
+		}
+	}
+
+	Super::PostEditChangeProperty(PropertyChangedEvent);
+}
+#endif // WITH_EDITOR
+
+void UPCGMatchAndSetByAttribute::MatchAndSet_Implementation(
+	FPCGContext& Context,
+	const UPCGPointMatchAndSetSettings* InSettings,
+	const UPCGPointData* InPointData,
+	UPCGPointData* OutPointData) const
+{
+	check(InSettings && InPointData && OutPointData);
+	check(OutPointData->Metadata);
+	check(InPointData->GetPoints().Num() == OutPointData->GetMutablePoints().Num());
+
+	if (Entries.Num() == 0)
+	{
+		return;
+	}
+
+	const FPCGAttributePropertySelector& SetTarget = InSettings->SetTarget;
+
+	// Create attribute if needed
+	if (!CreateAttributeIfNeeded(Context, SetTarget, Entries[0].Value, OutPointData, InSettings))
+	{
+		return; // Failed adding attribute
+	}
+
+	TArray<TUniquePtr<const IPCGAttributeAccessor>> EntryMatchSourceAccessors;
+	TArray<TUniquePtr<const IPCGAttributeAccessor>> EntrySetValueAccessors;
+	EntryMatchSourceAccessors.Reserve(Entries.Num());
+	EntrySetValueAccessors.Reserve(Entries.Num());
+
+	TUniquePtr<const IPCGAttributeAccessorKeys> EntryKeys = MakeUnique<FPCGAttributeAccessorKeysSingleObjectPtr<void>>();
+
+	auto CreateConstantMatchSourceAccessor = [&EntryMatchSourceAccessors](auto&& Value)
+	{
+		using ConstantType = std::decay_t<decltype(Value)>;
+		EntryMatchSourceAccessors.Emplace(MakeUnique<FPCGConstantValueAccessor<ConstantType>>(std::forward<ConstantType>(Value)));
+	};
+
+	auto CreateConstantSetValueAccessor = [&EntrySetValueAccessors](auto&& Value)
+	{
+		using ConstantType = std::decay_t<decltype(Value)>;
+		EntrySetValueAccessors.Emplace(MakeUnique<FPCGConstantValueAccessor<ConstantType>>(std::forward<ConstantType>(Value)));
+	};
+
+	for (const FPCGMatchAndSetByAttributeEntry& Entry : Entries)
+	{
+		Entry.ValueToMatch.DispatcherWithOverride(nullptr, CreateConstantMatchSourceAccessor);
+		Entry.Value.DispatcherWithOverride(nullptr, CreateConstantSetValueAccessor);
+	}
+
+	//TODO: implement async loop?
+	FPCGAttributePropertySelector InputSource;
+	InputSource.Selection = EPCGAttributePropertySelection::Attribute;
+	InputSource.AttributeName = (MatchSourceAttribute == NAME_None ? InPointData->Metadata->GetLatestAttributeNameOrNone() : MatchSourceAttribute);
+
+	TUniquePtr<const IPCGAttributeAccessor> InputAccessor = PCGAttributeAccessorHelpers::CreateConstAccessor(InPointData, InputSource);
+	TUniquePtr<const IPCGAttributeAccessorKeys> InputKeys = PCGAttributeAccessorHelpers::CreateConstKeys(InPointData, InputSource);
+
+	if (!InputAccessor.IsValid() || !InputKeys.IsValid())
+	{
+		PCGE_LOG_C(Warning, &Context, "Failed to create input accessor or iterator in MatchAndSet");
+		return;
+	}
+
+	// TODO: this should be more complex type - see what's done in the compare element
+	TUniquePtr<IPCGAttributeAccessor> SetTargetAccessor = PCGAttributeAccessorHelpers::CreateAccessor(OutPointData, SetTarget);
+	TUniquePtr<IPCGAttributeAccessorKeys> SetTargetKeys = PCGAttributeAccessorHelpers::CreateKeys(OutPointData, SetTarget);
+
+	auto MatchAndSetOperation = [&InputAccessor, &InputKeys, &EntryMatchSourceAccessors, &EntrySetValueAccessors, &EntryKeys, &SetTargetAccessor, &SetTargetKeys](auto MatchDummy)
+	{
+		using MatchType = decltype(MatchDummy);
+		auto GetValueIndex = [&InputAccessor, &InputKeys, &EntryMatchSourceAccessors, &EntryKeys](int32 Index) -> int32
+		{
+			MatchType Value{};
+			EPCGAttributeAccessorFlags Flags = EPCGAttributeAccessorFlags::AllowBroadcast;
+
+			InputAccessor->Get<MatchType>(Value, Index, *InputKeys, Flags);
+
+			for (int32 SourceIndex = 0; SourceIndex < EntryMatchSourceAccessors.Num(); ++SourceIndex)
+			{
+				MatchType SourceValue{};
+				EntryMatchSourceAccessors[SourceIndex]->Get<MatchType>(SourceValue, *EntryKeys, Flags);
+
+				if (PCG::Private::MetadataTraits<MatchType>::Equal(Value, SourceValue))
+				{
+					return SourceIndex;
+				}
+			}
+
+			return -1;
+		};
+
+		auto SetValueOperation = [&GetValueIndex, &EntrySetValueAccessors, &EntryKeys, &SetTargetAccessor, &SetTargetKeys](auto SetValueDummy)
+		{
+			using OutputType = decltype(SetValueDummy);
+			EPCGAttributeAccessorFlags Flags = EPCGAttributeAccessorFlags::AllowBroadcast;
+
+			// Cache values to set
+			TArray<OutputType> Values;
+			Values.Reserve(EntrySetValueAccessors.Num());
+
+			for (TUniquePtr<const IPCGAttributeAccessor>& InputAccessor : EntrySetValueAccessors)
+			{
+				InputAccessor->Get<OutputType>(Values.Emplace_GetRef(), *EntryKeys, Flags);
+			}
+
+			// Set values on points
+			const int32 NumberOfElements = SetTargetKeys->GetNum();
+			for (int32 ElementIndex = 0; ElementIndex < NumberOfElements; ++ElementIndex)
+			{
+				const int32 ValueIndex = GetValueIndex(ElementIndex);
+				if (ValueIndex >= 0)
+				{
+					SetTargetAccessor->Set<OutputType>(Values[ValueIndex], ElementIndex, *SetTargetKeys, Flags);
+				}
+			}
+
+			return true;
+		};
+
+		return PCGMetadataAttribute::CallbackWithRightType(SetTargetAccessor->GetUnderlyingType(), SetValueOperation);
+	};
+
+	if (!PCGMetadataAttribute::CallbackWithRightType(InputAccessor->GetUnderlyingType(), MatchAndSetOperation))
+	{
+		PCGE_LOG_C(Warning, &Context, "Error while getting/setting values in the MatchAndSet");
+	}
+}
+
+bool UPCGMatchAndSetByAttribute::ValidatePreconditions_Implementation(const UPCGPointData* InPointData) const
+{
+	if (!InPointData || !InPointData->Metadata)
+	{
+		return false;
+	}
+
+	// Check if source attribute exists
+	FName AttributeName = (MatchSourceAttribute == NAME_None ? InPointData->Metadata->GetLatestAttributeNameOrNone() : MatchSourceAttribute);
+
+	if (!InPointData->Metadata->HasAttribute(AttributeName))
+	{
+		return false;
+	}
+
+	// Check if type matches
+	const FPCGMetadataAttributeBase* AttributeToMatch = InPointData->Metadata->GetConstAttribute(AttributeName);
+	return AttributeToMatch && AttributeToMatch->GetTypeId() == static_cast<uint16>(MatchSourceType);
+}
\ No newline at end of file
diff --git a/Engine/Plugins/Experimental/PCG/Source/PCG/Private/MatchAndSet/PCGMatchAndSetWeighted.cpp b/Engine/Plugins/Experimental/PCG/Source/PCG/Private/MatchAndSet/PCGMatchAndSetWeighted.cpp
new file mode 100644
--- /dev/null
+++ b/Engine/Plugins/Experimental/PCG/Source/PCG/Private/MatchAndSet/PCGMatchAndSetWeighted.cpp
@@ -0,0 +1,152 @@
+// Copyright Epic Games, Inc. All Rights Reserved.
+
+#include "MatchAndSet/PCGMatchAndSetWeighted.h"
+
+#include "Data/PCGPointData.h"
+#include "Elements/PCGPointMatchAndSet.h"
+#include "Helpers/PCGBlueprintHelpers.h"
+#include "Metadata/PCGMetadata.h"
+#include "Metadata/PCGMetadataAttribute.h"
+#include "Metadata/Accessors/IPCGAttributeAccessor.h"
+#include "Metadata/Accessors/PCGAttributeAccessorHelpers.h"
+#include "Metadata/Accessors/PCGAttributeAccessorKeys.h"
+#include "Metadata/Accessors/PCGCustomAccessor.h"
+
+#include UE_INLINE_GENERATED_CPP_BY_NAME(PCGMatchAndSetWeighted)
+
+FPCGMatchAndSetWeightedEntry::FPCGMatchAndSetWeightedEntry()
+{
+	Value.bAllowsTypeChange = false;
+}
+
+void UPCGMatchAndSetWeighted::SetType(EPCGMetadataTypes InType, EPCGMetadataTypesConstantStructStringMode InStringMode)
+{
+	for (FPCGMatchAndSetWeightedEntry& Entry : Entries)
+	{
+		Entry.Value.Type = InType;
+		Entry.Value.StringMode = InStringMode;
+	}
+
+	Super::SetType(InType, InStringMode);
+}
+
+#if WITH_EDITOR
+void UPCGMatchAndSetWeighted::PostEditChangeProperty(FPropertyChangedEvent& PropertyChangedEvent)
+{
+	if (PropertyChangedEvent.Property)
+	{
+		const FName& PropertyName = PropertyChangedEvent.Property->GetFName();
+
+		if (PropertyName == GET_MEMBER_NAME_CHECKED(UPCGMatchAndSetWeighted, Entries))
+		{
+			// Some changes in the array might (such as insert or new) might require us to re-set the type
+			SetType(Type, StringMode);
+		}
+	}
+
+	Super::PostEditChangeProperty(PropertyChangedEvent);
+}
+#endif // WITH_EDITOR
+
+void UPCGMatchAndSetWeighted::MatchAndSet_Implementation(
+	FPCGContext& Context,
+	const UPCGPointMatchAndSetSettings* InSettings,
+	const UPCGPointData* InPointData,
+	UPCGPointData* OutPointData) const
+{
+	check(InSettings && InPointData && OutPointData);
+	check(OutPointData->Metadata);
+	check(InPointData->GetPoints().Num() == OutPointData->GetMutablePoints().Num());
+
+	if (Entries.Num() == 0)
+	{
+		return;
+	}
+
+	const FPCGAttributePropertySelector& SetTarget = InSettings->SetTarget;
+
+	// Create attribute if needed
+	if (!CreateAttributeIfNeeded(Context, SetTarget, Entries[0].Value, OutPointData, InSettings))
+	{
+		return; // Failed adding attribute
+	}
+
+	TArray<TUniquePtr<const IPCGAttributeAccessor>> EntryValueAccessors;
+	EntryValueAccessors.Reserve(Entries.Num());
+	TUniquePtr<const IPCGAttributeAccessorKeys> EntryKeys = MakeUnique<FPCGAttributeAccessorKeysSingleObjectPtr<void>>();
+
+	auto CreateConstantAccessor = [&EntryValueAccessors](auto&& Value)
+	{
+		using ConstantType = std::decay_t<decltype(Value)>;
+		EntryValueAccessors.Emplace(MakeUnique<FPCGConstantValueAccessor<ConstantType>>(std::forward<ConstantType>(Value)));
+	};
+
+	TArray<int> CumulativeWeights;
+
+	int TotalWeight = 0;
+
+	for (const FPCGMatchAndSetWeightedEntry& Entry : Entries)
+	{
+		TotalWeight += Entry.Weight;
+		CumulativeWeights.Add(TotalWeight);
+		Entry.Value.DispatcherWithOverride(nullptr, CreateConstantAccessor);
+	}
+
+	if (TotalWeight <= 0)
+	{
+		return;
+	}
+
+	// TODO: implement async loop?
+	const UPCGComponent* SourceComponent = Context.SourceComponent.Get();
+	const TArray<FPCGPoint>& InPoints = InPointData->GetPoints();
+
+	TUniquePtr<IPCGAttributeAccessor> SetTargetAccessor = PCGAttributeAccessorHelpers::CreateAccessor(OutPointData, SetTarget);
+	TUniquePtr<IPCGAttributeAccessorKeys> SetTargetKeys = PCGAttributeAccessorHelpers::CreateKeys(OutPointData, SetTarget);
+
+	auto GetValueIndexFromElementIndex = [&InPoints, InSettings, SourceComponent, &CumulativeWeights, TotalWeight](int32 Index) -> int32
+	{
+		const FPCGPoint& Point = InPoints[Index];
+
+		FRandomStream RandomSource = UPCGBlueprintHelpers::GetRandomStream(Point, InSettings, SourceComponent);
+		int RandomWeightedPick = RandomSource.RandRange(0, TotalWeight - 1);
+
+		int RandomPick = 0;
+		while (RandomPick < CumulativeWeights.Num() && CumulativeWeights[RandomPick] <= RandomWeightedPick)
+		{
+			++RandomPick;
+		}
+
+		return RandomPick;
+	};
+
+	auto MatchAndSetOperation = [&GetValueIndexFromElementIndex, &EntryKeys, &EntryValueAccessors, &SetTargetAccessor, &SetTargetKeys](auto Dummy)
+	{
+		using OutputType = decltype(Dummy);
+		EPCGAttributeAccessorFlags Flags = EPCGAttributeAccessorFlags::AllowBroadcast;
+
+		// Cache values to set
+		TArray<OutputType> Values;
+		Values.Reserve(EntryValueAccessors.Num());
+
+		for (TUniquePtr<const IPCGAttributeAccessor>& InputAccessor : EntryValueAccessors)
+		{
+			InputAccessor->Get<OutputType>(Values.Emplace_GetRef(), *EntryKeys, Flags);
+		}
+
+		// Set values on points
+		const int32 NumberOfElements = SetTargetKeys->GetNum();
+		for (int32 ElementIndex = 0; ElementIndex < NumberOfElements; ++ElementIndex)
+		{
+			const int32 ValueIndex = GetValueIndexFromElementIndex(ElementIndex);
+			SetTargetAccessor->Set<OutputType>(Values[ValueIndex], ElementIndex, *SetTargetKeys, Flags);
+		}
+
+		return true;
+	};
+
+	if (!PCGMetadataAttribute::CallbackWithRightType(SetTargetAccessor->GetUnderlyingType(), MatchAndSetOperation))
+	{
+		PCGE_LOG_C(Warning, &Context, "Error while getting/setting values in the MatchAndSet");
+	}
+}
diff --git a/Engine/Plugins/Experimental/PCG/Source/PCG/Private/MatchAndSet/PCGMatchAndSetWeightedByCategory.cpp b/Engine/Plugins/Experimental/PCG/Source/PCG/Private/MatchAndSet/PCGMatchAndSetWeightedByCategory.cpp
new file mode 100644
--- /dev/null
+++ b/Engine/Plugins/Experimental/PCG/Source/PCG/Private/MatchAndSet/PCGMatchAndSetWeightedByCategory.cpp
@@ -0,0 +1,296 @@
+// Copyright Epic Games, Inc. All Rights Reserved.
+
+#include "MatchAndSet/PCGMatchAndSetWeightedByCategory.h"
+
+#include "Data/PCGPointData.h"
+#include "Elements/PCGPointMatchAndSet.h"
+#include "Helpers/PCGBlueprintHelpers.h"
+#include "Metadata/PCGMetadata.h"
+#include "Metadata/PCGMetadataAttribute.h"
+#include "Metadata/Accessors/IPCGAttributeAccessor.h"
+#include "Metadata/Accessors/PCGAttributeAccessorHelpers.h"
+#include "Metadata/Accessors/PCGAttributeAccessorKeys.h"
+#include "Metadata/Accessors/PCGCustomAccessor.h"
+
+#include UE_INLINE_GENERATED_CPP_BY_NAME(PCGMatchAndSetWeightedByCategory)
+
+FPCGMatchAndSetWeightedByCategoryEntryList::FPCGMatchAndSetWeightedByCategoryEntryList()
+{
+	CategoryValue.bAllowsTypeChange = false;
+}
+
+void FPCGMatchAndSetWeightedByCategoryEntryList::SetType(EPCGMetadataTypes InType, EPCGMetadataTypesConstantStructStringMode InStringMode)
+{
+	for (FPCGMatchAndSetWeightedEntry& Entry : WeightedEntries)
+	{
+		Entry.Value.Type = InType;
+		Entry.Value.StringMode = InStringMode;
+	}
+}
+
+int FPCGMatchAndSetWeightedByCategoryEntryList::GetTotalWeight() const
+{
+	int TotalWeight = 0;
+
+	for (const FPCGMatchAndSetWeightedEntry& Entry : WeightedEntries)
+	{
+		TotalWeight += Entry.Weight;
+	}
+
+	return TotalWeight;
+}
+
+void UPCGMatchAndSetWeightedByCategory::SetType(EPCGMetadataTypes InType, EPCGMetadataTypesConstantStructStringMode InStringMode)
+{
+	for (FPCGMatchAndSetWeightedByCategoryEntryList& Category : Categories)
+	{
+		Category.SetType(InType, InStringMode);
+	}
+
+	Super::SetType(InType, InStringMode);
+}
+
+void UPCGMatchAndSetWeightedByCategory::SetCategoryType(EPCGMetadataTypes InType, EPCGMetadataTypesConstantStructStringMode InStringMode)
+{
+	for (FPCGMatchAndSetWeightedByCategoryEntryList& Category : Categories)
+	{
+		Category.CategoryValue.Type = InType;
+		Category.CategoryValue.StringMode = InStringMode;
+	}
+}
+
+#if WITH_EDITOR
+void UPCGMatchAndSetWeightedByCategory::PostEditChangeProperty(FPropertyChangedEvent& PropertyChangedEvent)
+{
+	if (PropertyChangedEvent.Property)
+	{
+		const FName& PropertyName = PropertyChangedEvent.Property->GetFName();
+
+		if (PropertyName == GET_MEMBER_NAME_CHECKED(UPCGMatchAndSetWeightedByCategory, CategoryType) ||
+			PropertyName == GET_MEMBER_NAME_CHECKED(UPCGMatchAndSetWeightedByCategory, CategoryStringMode))
+		{
+			SetCategoryType(CategoryType, CategoryStringMode);
+		}
+		else if (PropertyName == GET_MEMBER_NAME_CHECKED(FPCGMatchAndSetWeightedByCategoryEntryList, WeightedEntries))
+		{
+			// Changes in any weighted entries (which store the end values) - due to add or insert - need to update type
+			SetType(Type, StringMode);
+		}
+		else if (PropertyName == GET_MEMBER_NAME_CHECKED(UPCGMatchAndSetWeightedByCategory, Categories))
+		{
+			SetCategoryType(CategoryType, CategoryStringMode);
+			SetType(Type, StringMode);
+		}
+	}
+
+	Super::PostEditChangeProperty(PropertyChangedEvent);
+}
+#endif // WITH_EDITOR
+
+void UPCGMatchAndSetWeightedByCategory::MatchAndSet_Implementation(
+	FPCGContext& Context,
+	const UPCGPointMatchAndSetSettings* InSettings,
+	const UPCGPointData* InPointData,
+	UPCGPointData* OutPointData) const
+{
+	// Early validation
+	check(InSettings && InPointData && OutPointData);
+	check(OutPointData->Metadata);
+	check(InPointData->GetPoints().Num() == OutPointData->GetMutablePoints().Num());
+
+	bool bHasAtLeastOneNonEmptyEntry = false;
+	int32 DefaultCategoryIndex = INDEX_NONE;
+	const FPCGMetadataTypesConstantStruct* FirstValueForTypeExtraction = nullptr;
+	
+	for(int32 CategoryIndex = 0; CategoryIndex < Categories.Num(); ++CategoryIndex)
+	{
+		const FPCGMatchAndSetWeightedByCategoryEntryList& Category = Categories[CategoryIndex];
+
+		if (Category.bIsDefault && DefaultCategoryIndex == INDEX_NONE)
+		{
+			DefaultCategoryIndex = CategoryIndex;
+		}
+
+		if (Category.GetTotalWeight() > 0)
+		{
+			bHasAtLeastOneNonEmptyEntry = true;
+
+			check(Category.WeightedEntries.Num() > 0);
+			FirstValueForTypeExtraction = &Category.WeightedEntries[0].Value;
+		}
+	}
+
+	if (!bHasAtLeastOneNonEmptyEntry)
+	{
+		return;
+	}
+
+	const FPCGAttributePropertySelector& SetTarget = InSettings->SetTarget;
+
+	// Create attribute if needed
+	if (!CreateAttributeIfNeeded(Context, SetTarget, *FirstValueForTypeExtraction, OutPointData, InSettings))
+	{
+		return; // Failed adding attribute
+	}
+
+	TArray<TUniquePtr<const IPCGAttributeAccessor>> EntryCategoryAccessors;
+	TArray<TArray<TUniquePtr<const IPCGAttributeAccessor>>> EntryValueAccessors;
+	EntryCategoryAccessors.Reserve(Categories.Num());
+	EntryValueAccessors.Reserve(Categories.Num());
+
+	TUniquePtr<const IPCGAttributeAccessorKeys> EntryKeys = MakeUnique<FPCGAttributeAccessorKeysSingleObjectPtr<void>>();
+
+	auto CreateConstantCategoryAccessor = [&EntryCategoryAccessors](auto&& Value)
+	{
+		using ConstantType = std::decay_t<decltype(Value)>;
+		EntryCategoryAccessors.Emplace(MakeUnique<FPCGConstantValueAccessor<ConstantType>>(std::forward<ConstantType>(Value)));
+	};
+
+	for (const FPCGMatchAndSetWeightedByCategoryEntryList& Category : Categories)
+	{
+		Category.CategoryValue.DispatcherWithOverride(nullptr, CreateConstantCategoryAccessor);
+
+		TArray<TUniquePtr<const IPCGAttributeAccessor>>& ValuesAccessor = EntryValueAccessors.Emplace_GetRef();
+		for (const FPCGMatchAndSetWeightedEntry& WeightedEntry : Category.WeightedEntries)
+		{
+			WeightedEntry.Value.DispatcherWithOverride(nullptr, [&ValuesAccessor](auto&& Value)
+			{
+				using ConstantType = std::decay_t<decltype(Value)>;
+				ValuesAccessor.Emplace(MakeUnique<FPCGConstantValueAccessor<ConstantType>>(std::forward<ConstantType>(Value)));
+			});
+		}
+	}
+
+	//TODO: implement async loop?
+	const UPCGComponent* SourceComponent = Context.SourceComponent.Get();
+	const TArray<FPCGPoint>& InPoints = InPointData->GetPoints();
+
+	FPCGAttributePropertySelector InputSource;
+	InputSource.Selection = EPCGAttributePropertySelection::Attribute;
+	InputSource.AttributeName = (CategoryAttribute == NAME_None ? InPointData->Metadata->GetLatestAttributeNameOrNone() : CategoryAttribute);
+
+	TUniquePtr<const IPCGAttributeAccessor> InputAccessor = PCGAttributeAccessorHelpers::CreateConstAccessor(InPointData, InputSource);
+	TUniquePtr<const IPCGAttributeAccessorKeys> InputKeys = PCGAttributeAccessorHelpers::CreateConstKeys(InPointData, InputSource);
+
+	if (!InputAccessor.IsValid() || !InputKeys.IsValid())
+	{
+		PCGE_LOG_C(Warning, &Context, "Failed to create input accessor or iterator in MatchAndSet");
+		return;
+	}
+
+	TUniquePtr<IPCGAttributeAccessor> SetTargetAccessor = PCGAttributeAccessorHelpers::CreateAccessor(OutPointData, SetTarget);
+	TUniquePtr<IPCGAttributeAccessorKeys> SetTargetKeys = PCGAttributeAccessorHelpers::CreateKeys(OutPointData, SetTarget);
+
+	auto GetValueIndex = [this, InSettings, SourceComponent, &InPoints](int32 CategoryIndex, int32 PointIndex) -> int32
+	{
+		const FPCGMatchAndSetWeightedByCategoryEntryList& Category = Categories[CategoryIndex];
+		const int32 TotalWeight = Category.GetTotalWeight();
+
+		if (TotalWeight == 0)
+		{
+			return INDEX_NONE;
+		}
+
+		const FPCGPoint& Point = InPoints[PointIndex];
+
+		FRandomStream RandomSource = UPCGBlueprintHelpers::GetRandomStream(Point, InSettings, SourceComponent);
+		int RandomWeightedPick = RandomSource.RandRange(0, TotalWeight - 1);
+
+		int RandomPick = 0;
+		while (RandomPick < Category.WeightedEntries.Num() && RandomWeightedPick >= Category.WeightedEntries[RandomPick].Weight)
+		{
+			RandomWeightedPick -= Category.WeightedEntries[RandomPick].Weight;
+			++RandomPick;
+		}
+
+		return RandomPick;
+	};
+
+	auto MatchAndSetOperation = [&GetValueIndex, &InputAccessor, &InputKeys, &EntryCategoryAccessors, &EntryValueAccessors, &EntryKeys, DefaultCategoryIndex, &SetTargetAccessor, &SetTargetKeys](auto MatchDummy)
+	{
+		using MatchType = decltype(MatchDummy);
+		auto GetCategoryIndex = [&InputAccessor, &InputKeys, &EntryCategoryAccessors, &EntryKeys, DefaultCategoryIndex](int32 Index) -> int32
+		{
+			EPCGAttributeAccessorFlags Flags = EPCGAttributeAccessorFlags::AllowBroadcast;
+
+			MatchType Value{};
+			InputAccessor->Get<MatchType>(Value, Index, *InputKeys, Flags);
+
+			for (int32 CategoryIndex = 0; CategoryIndex < EntryCategoryAccessors.Num(); ++CategoryIndex)
+			{
+				MatchType SourceValue{};
+				EntryCategoryAccessors[CategoryIndex]->Get<MatchType>(SourceValue, *EntryKeys, Flags);
+
+				if (PCG::Private::MetadataTraits<MatchType>::Equal(Value, SourceValue))
+				{
+					return CategoryIndex;
+				}
+			}
+
+			return DefaultCategoryIndex;
+		};
+
+		auto SetValueOperation = [&GetCategoryIndex, &GetValueIndex, &EntryValueAccessors, &EntryKeys, &SetTargetAccessor, &SetTargetKeys](auto SetValueDummy)
+		{
+			using OutputType = decltype(SetValueDummy);
+			EPCGAttributeAccessorFlags Flags = EPCGAttributeAccessorFlags::AllowBroadcast;
+
+			// Cache values to set
+			TArray<TArray<OutputType>> Values;
+			Values.Reserve(EntryValueAccessors.Num());
+			for (const TArray<TUniquePtr<const IPCGAttributeAccessor>>& CategoryValueAccessor : EntryValueAccessors)
+			{
+				TArray<OutputType>& CategoryValues = Values.Emplace_GetRef();
+				CategoryValues.Reserve(CategoryValueAccessor.Num());
+
+				for (const TUniquePtr<const IPCGAttributeAccessor>& ValueAccessor : CategoryValueAccessor)
+				{
+					ValueAccessor->Get<OutputType>(CategoryValues.Emplace_GetRef(), *EntryKeys, Flags);
+				}
+			}
+
+			// Set values on points
+			const int32 NumberOfElements = SetTargetKeys->GetNum();
+			for (int32 ElementIndex = 0; ElementIndex < NumberOfElements; ++ElementIndex)
+			{
+				const int32 CategoryIndex = GetCategoryIndex(ElementIndex);
+				if (CategoryIndex >= 0)
+				{
+					const int32 ValueIndex = GetValueIndex(CategoryIndex, ElementIndex);
+					if (ValueIndex >= 0)
+					{
+						SetTargetAccessor->Set<OutputType>(Values[CategoryIndex][ValueIndex], ElementIndex, *SetTargetKeys, Flags);
+					}
+				}
+			}
+
+			return true;
+		};
+
+		return PCGMetadataAttribute::CallbackWithRightType(SetTargetAccessor->GetUnderlyingType(), SetValueOperation);
+	};
+
+	if (!PCGMetadataAttribute::CallbackWithRightType(InputAccessor->GetUnderlyingType(), MatchAndSetOperation))
+	{
+		PCGE_LOG_C(Warning, &Context, "Error while getting/setting values in the MatchAndSet");
+	}
+}
+
+bool UPCGMatchAndSetWeightedByCategory::ValidatePreconditions_Implementation(const UPCGPointData* InPointData) const
+{
+	if (!InPointData || !InPointData->Metadata)
+	{
+		return false;
+	}
+
+	FName AttributeName = (CategoryAttribute == NAME_None ? InPointData->Metadata->GetLatestAttributeNameOrNone() : CategoryAttribute);
+	
+	if (!InPointData->Metadata->HasAttribute(AttributeName))
+	{
+		return false;
+	}
+
+	// Check if type matches
+	const FPCGMetadataAttributeBase* AttributeToMatch = InPointData->Metadata->GetConstAttribute(AttributeName);
+	return AttributeToMatch && AttributeToMatch->GetTypeId() == static_cast<uint16>(CategoryType);
+}
\ No newline at end of file
diff --git a/Engine/Plugins/Experimental/PCG/Source/PCG/Private/Metadata/Accessors/PCGAttributeAccessorHelpers.cpp b/Engine/Plugins/Experimental/PCG/Source/PCG/Private/Metadata/Accessors/PCGAttributeAccessorHelpers.cpp
--- a/Engine/Plugins/Experimental/PCG/Source/PCG/Private/Metadata/Accessors/PCGAttributeAccessorHelpers.cpp
+++ b/Engine/Plugins/Experimental/PCG/Source/PCG/Private/Metadata/Accessors/PCGAttributeAccessorHelpers.cpp
@@ -191,7 +191,7 @@ TUniquePtr<const IPCGAttributeAccessor> PCGAttributeAccessorHelpers::CreateConst
 
 	if (InSelector.Selection == EPCGAttributePropertySelection::PointProperty)
 	{
-		if (const UPCGPointData* PointData = Cast<const UPCGPointData>(InData))
+		if(!InData || Cast<const UPCGPointData>(InData))
 		{
 			if (const FProperty* Property = FPCGPoint::StaticStruct()->FindPropertyByName(Name))
 			{
@@ -262,7 +262,7 @@ TUniquePtr<IPCGAttributeAccessor> PCGAttributeAccessorHelpers::CreateAccessor(UP
 
 	if (InSelector.Selection == EPCGAttributePropertySelection::PointProperty)
 	{
-		if (const UPCGPointData* PointData = Cast<const UPCGPointData>(InData))
+		if (!InData || Cast<const UPCGPointData>(InData))
 		{
 			if (const FProperty* Property = FPCGPoint::StaticStruct()->FindPropertyByName(Name))
 			{
diff --git a/Engine/Plugins/Experimental/PCG/Source/PCG/Public/Elements/PCGPointMatchAndSet.h b/Engine/Plugins/Experimental/PCG/Source/PCG/Public/Elements/PCGPointMatchAndSet.h
new file mode 100644
--- /dev/null
+++ b/Engine/Plugins/Experimental/PCG/Source/PCG/Public/Elements/PCGPointMatchAndSet.h
@@ -0,0 +1,83 @@
+// Copyright Epic Games, Inc. All Rights Reserved.
+
+#pragma once
+
+#include "PCGSettings.h"
+#include "MatchAndSet/PCGMatchAndSetBase.h"
+#include "Metadata/PCGAttributePropertySelector.h"
+
+#include "PCGPointMatchAndSet.generated.h"
+
+/** This settings class is used to create a PCG node that will apply a "Match and Set" operation
+* on the point data it consumes as input.
+* E.g. for a given point, if it matches with something in the Match & Set object, it will set a value on the point.
+*/
+UCLASS(BlueprintType, ClassGroup = (Procedural))
+class PCG_API UPCGPointMatchAndSetSettings : public UPCGSettings
+{
+	GENERATED_BODY()
+
+public:
+	UPCGPointMatchAndSetSettings(const FObjectInitializer& ObjectInitializer);
+
+	// ~Begin UPCGSettings interface
+#if WITH_EDITOR
+	virtual FName GetDefaultNodeName() const override { return FName(TEXT("PointMatchAndSet")); }
+	virtual FText GetNodeTooltipText() const override;
+	virtual EPCGSettingsType GetType() const override { return EPCGSettingsType::Metadata; }
+#endif // WITH_EDITOR
+
+	virtual TArray<FPCGPinProperties> InputPinProperties() const override;
+	virtual TArray<FPCGPinProperties> OutputPinProperties() const override;
+
+protected:
+	virtual FPCGElementPtr CreateElement() const override;
+
+public:
+	// ~Begin UObject interface
+#if WITH_EDITOR
+	virtual void PostEditChangeProperty(FPropertyChangedEvent& PropertyChangedEvent) override;
+#endif
+	// ~End UObject interface
+
+	/** Recreates the match & set instance stored in this settings object if needed. */
+	UFUNCTION(BlueprintCallable, Category = Settings)
+	void SetMatchAndSetType(TSubclassOf<UPCGMatchAndSetBase> InMatchAndSetType);
+
+protected:
+	void RefreshMatchAndSet();
+
+public:
+	/** Defines the type of Match & Set object to use. */
+	UPROPERTY(BlueprintReadOnly, EditAnywhere, NoClear, Category = Settings)
+	TSubclassOf<UPCGMatchAndSetBase> MatchAndSetType;
+
+	/** Instance of MatchAndSetType, stores the data that will be used in these settings. */
+	UPROPERTY(BlueprintReadOnly, VisibleAnywhere, Instanced, Category = Settings)
+	TObjectPtr<UPCGMatchAndSetBase> MatchAndSetInstance;
+
+	/** "Set" part of the Match & Set - defines what will be changed in the operation */
+	UPROPERTY(BlueprintReadWrite, EditAnywhere, Category = Settings)
+	FPCGAttributePropertySelector SetTarget;
+
+	/** If the "Set" part is an attribute, then the type must be provided */
+	UPROPERTY(BlueprintReadWrite, EditAnywhere, Category = Settings, meta = (EditCondition = "bSetTargetIsAttribute", HideEditConditionToggle, EditConditionHides, ValidEnumValues = "Float, Double, Integer32, Integer64, Vector2, Vector, Vector4, Quaternion, Transform, String, Boolean, Rotator, Name"))
+	EPCGMetadataTypes SetTargetType = EPCGMetadataTypes::Double;
+
+	/** For string types, the subtype is used to cleanup the UI. */
+	UPROPERTY(BlueprintReadWrite, EditAnywhere, Category = Settings, meta = (EditCondition = "bSetTargetIsAttribute && SetTargetType == EPCGMetadataTypes::String", EditConditionHides))
+	EPCGMetadataTypesConstantStructStringMode SetTargetStringMode = EPCGMetadataTypesConstantStructStringMode::String;
+
+private:
+#if WITH_EDITORONLY_DATA
+	// Property used to sidestep edit condition issues - reflects SetTarget.Selection == Attribute
+	UPROPERTY()
+	bool bSetTargetIsAttribute = true;
+#endif
+};
+
+class FPCGPointMatchAndSetElement : public FSimplePCGElement
+{
+public:
+	virtual bool ExecuteInternal(FPCGContext* Context) const override;
+};
\ No newline at end of file
diff --git a/Engine/Plugins/Experimental/PCG/Source/PCG/Public/MatchAndSet/PCGMatchAndSetBase.h b/Engine/Plugins/Experimental/PCG/Source/PCG/Public/MatchAndSet/PCGMatchAndSetBase.h
new file mode 100644
--- /dev/null
+++ b/Engine/Plugins/Experimental/PCG/Source/PCG/Public/MatchAndSet/PCGMatchAndSetBase.h
@@ -0,0 +1,71 @@
+// Copyright Epic Games, Inc. All Rights Reserved.
+
+#pragma once
+
+#include "CoreMinimal.h"
+#include "Metadata/PCGMetadataTypesConstantStruct.h"
+
+#include "PCGMatchAndSetBase.generated.h"
+
+struct FPCGAttributePropertySelector;
+struct FPCGContext;
+struct FPCGMetadataTypesConstantStruct;
+struct FPCGPoint;
+class UPCGMetadata;
+class UPCGPointMatchAndSetSettings;
+
+/** Base class for Match & Set objects. Note that while it currently deals with points, it might be extended in the future.
+* This class is extensible and can be implemented in different ways, but its role should be simple:
+* For a given point, if it matches some criteria ("Match"), apply it some value ("Set").
+* It can be a lookup, a random process or something more involved. 
+*/
+UCLASS(Abstract, BlueprintType, Blueprintable, ClassGroup = (Procedural))
+class PCG_API UPCGMatchAndSetBase : public UObject
+{
+	GENERATED_BODY()
+
+public:
+	/** Queries whether this object uses a random process - note that this is expected to be static through the life of the object */
+	virtual bool UsesRandomProcess() const { return false; }
+
+	/** Queries whether we should mutate the seeds as a post-process */
+	virtual bool ShouldMutateSeed() const { return false; }
+
+	/** Sets & propagates type change from the owner settings object */
+	virtual void SetType(EPCGMetadataTypes InType, EPCGMetadataTypesConstantStructStringMode InStringMode);
+
+	/** Early check to prevent issues when the data does not contain the required information to perform the operation */
+	UFUNCTION(BlueprintNativeEvent, Category = Selection)
+	bool ValidatePreconditions(const UPCGPointData* InPointData) const;
+
+	virtual bool ValidatePreconditions_Implementation(const UPCGPointData* InPointData) const { return true; }
+
+	/** Main function to process points, and pass them through the Match & Set logic. */
+	UFUNCTION(BlueprintNativeEvent, Category = Selection)
+	void MatchAndSet(
+		UPARAM(ref) FPCGContext& Context,
+		const UPCGPointMatchAndSetSettings* InSettings,
+		const UPCGPointData* InPointData,
+		UPCGPointData* OutPointData) const;
+
+	virtual void MatchAndSet_Implementation(
+		FPCGContext& Context,
+		const UPCGPointMatchAndSetSettings* InSettings,
+		const UPCGPointData* InPointData,
+		UPCGPointData* OutPointData) const PURE_VIRTUAL(UPCGMatchAndSetBase::MatchAndSet_Implementation);
+
+protected:
+	bool CreateAttributeIfNeeded(
+		FPCGContext& Context,
+		const FPCGAttributePropertySelector& Selector,
+		const FPCGMetadataTypesConstantStruct& Value,
+		UPCGPointData* OutPointData,
+		const UPCGPointMatchAndSetSettings* InSettings) const;
+
+	/** For the sake of managing internal state a bit better, we keep a copy of the Set type & string subtype. */
+	UPROPERTY()
+	EPCGMetadataTypes Type = EPCGMetadataTypes::Double;
+
+	UPROPERTY()
+	EPCGMetadataTypesConstantStructStringMode StringMode = EPCGMetadataTypesConstantStructStringMode::String;
+};
\ No newline at end of file
diff --git a/Engine/Plugins/Experimental/PCG/Source/PCG/Public/MatchAndSet/PCGMatchAndSetByAttribute.h b/Engine/Plugins/Experimental/PCG/Source/PCG/Public/MatchAndSet/PCGMatchAndSetByAttribute.h
new file mode 100644
--- /dev/null
+++ b/Engine/Plugins/Experimental/PCG/Source/PCG/Public/MatchAndSet/PCGMatchAndSetByAttribute.h
@@ -0,0 +1,67 @@
+// Copyright Epic Games, Inc. All Rights Reserved.
+
+#pragma once
+
+#include "MatchAndSet/PCGMatchAndSetBase.h"
+#include "Metadata/PCGMetadataTypesConstantStruct.h"
+
+#include "PCGMatchAndSetByAttribute.generated.h"
+
+USTRUCT(BlueprintType)
+struct PCG_API FPCGMatchAndSetByAttributeEntry
+{
+	GENERATED_BODY()
+
+	FPCGMatchAndSetByAttributeEntry();
+
+	UPROPERTY(BlueprintReadWrite, EditAnywhere, Category = Settings, meta = (ShowOnlyInnerProperties))
+	FPCGMetadataTypesConstantStruct ValueToMatch;
+
+	UPROPERTY(BlueprintReadWrite, EditAnywhere, Category = Settings, meta = (ShowOnlyInnerProperties))
+	FPCGMetadataTypesConstantStruct Value;
+};
+
+/** 
+* This Match & Set object looks up an attribute on a given point,
+* then looks up its entries to find a match; if there is one, then it sets it value. 
+*/
+UCLASS(BlueprintType, ClassGroup = (Procedural))
+class PCG_API UPCGMatchAndSetByAttribute : public UPCGMatchAndSetBase
+{
+	GENERATED_BODY()
+
+public:
+	virtual void SetType(EPCGMetadataTypes InType, EPCGMetadataTypesConstantStructStringMode InStringMode) override;
+
+	/** Propagates (does not set) the Match type to the entries */
+	virtual void SetSourceType(EPCGMetadataTypes InType, EPCGMetadataTypesConstantStructStringMode InStringMode);
+
+	virtual void MatchAndSet_Implementation(
+		FPCGContext& Context,
+		const UPCGPointMatchAndSetSettings* InSettings,
+		const UPCGPointData* InPointData,
+		UPCGPointData* OutPointData) const override;
+
+	virtual bool ValidatePreconditions_Implementation(const UPCGPointData* InPointData) const override;
+
+#if WITH_EDITOR
+	virtual void PostEditChangeProperty(FPropertyChangedEvent& PropertyChangedEvent) override;
+#endif
+
+public:
+	/** Attribute to match on the data */
+	UPROPERTY(BlueprintReadWrite, EditAnywhere, Category = Settings)
+	FName MatchSourceAttribute;
+
+	/** Type of the attribute to match against. */
+	UPROPERTY(BlueprintReadWrite, EditAnywhere, Category = Settings, meta = (ValidEnumValues = "Float, Double, Integer32, Integer64, Vector2, Vector, Vector4, Quaternion, Transform, String, Boolean, Rotator, Name"))
+	EPCGMetadataTypes MatchSourceType = EPCGMetadataTypes::Double;
+
+	/** String type of the attribute to match against (if required). */
+	UPROPERTY(BlueprintReadWrite, EditAnywhere, Category = Settings, meta = (EditCondition = "MatchSourceType == EPCGMetadataTypes::String", EditConditionHides))
+	EPCGMetadataTypesConstantStructStringMode MatchSourceStringMode = EPCGMetadataTypesConstantStructStringMode::String;
+
+	/** Lookup entries (key-value pairs) */
+	UPROPERTY(BlueprintReadWrite, EditAnywhere, Category = Settings)
+	TArray<FPCGMatchAndSetByAttributeEntry> Entries;
+};
\ No newline at end of file
diff --git a/Engine/Plugins/Experimental/PCG/Source/PCG/Public/MatchAndSet/PCGMatchAndSetWeighted.h b/Engine/Plugins/Experimental/PCG/Source/PCG/Public/MatchAndSet/PCGMatchAndSetWeighted.h
new file mode 100644
--- /dev/null
+++ b/Engine/Plugins/Experimental/PCG/Source/PCG/Public/MatchAndSet/PCGMatchAndSetWeighted.h
@@ -0,0 +1,57 @@
+// Copyright Epic Games, Inc. All Rights Reserved.
+
+#pragma once
+
+#include "MatchAndSet/PCGMatchAndSetBase.h"
+#include "Metadata/PCGMetadataTypesConstantStruct.h"
+
+#include "PCGMatchAndSetWeighted.generated.h"
+
+USTRUCT(BlueprintType)
+struct PCG_API FPCGMatchAndSetWeightedEntry
+{
+	GENERATED_BODY()
+
+	FPCGMatchAndSetWeightedEntry();
+
+	UPROPERTY(BlueprintReadWrite, EditAnywhere, Category = Settings, meta = (ShowOnlyInnerProperties))
+	FPCGMetadataTypesConstantStruct Value;
+
+	/** Relative weight of this entry */
+	UPROPERTY(BlueprintReadWrite, EditAnywhere, Category = Settings, meta = (ClampMin = "0"))
+	int Weight = 1;
+};
+
+/**
+* This Match & Set object assigns randomly a value based on weighted ratios,
+* provided in the entries.
+*/
+UCLASS(BlueprintType, ClassGroup = (Procedural))
+class PCG_API UPCGMatchAndSetWeighted : public UPCGMatchAndSetBase
+{
+	GENERATED_BODY()
+
+public:
+	virtual bool UsesRandomProcess() const { return true; }
+	virtual bool ShouldMutateSeed() const { return bShouldMutateSeed; }
+	virtual void SetType(EPCGMetadataTypes InType, EPCGMetadataTypesConstantStructStringMode InStringMode) override;
+
+	virtual void MatchAndSet_Implementation(
+		FPCGContext& Context,
+		const UPCGPointMatchAndSetSettings* InSettings,
+		const UPCGPointData* InPointData,
+		UPCGPointData* OutPointData) const override;
+
+#if WITH_EDITOR
+	virtual void PostEditChangeProperty(FPropertyChangedEvent& PropertyChangedEvent) override;
+#endif
+
+public:
+	/** Values and their respective weights */
+	UPROPERTY(BlueprintReadWrite, EditAnywhere, Category = Settings)
+	TArray<FPCGMatchAndSetWeightedEntry> Entries;
+
+	/** Controls whether the output data should mutate its seed - prevents issues when doing multiple random processes in a row */
+	UPROPERTY(BlueprintReadWrite, EditAnywhere, Category = Settings)
+	bool bShouldMutateSeed = true;
+};
\ No newline at end of file
diff --git a/Engine/Plugins/Experimental/PCG/Source/PCG/Public/MatchAndSet/PCGMatchAndSetWeightedByCategory.h b/Engine/Plugins/Experimental/PCG/Source/PCG/Public/MatchAndSet/PCGMatchAndSetWeightedByCategory.h
new file mode 100644
--- /dev/null
+++ b/Engine/Plugins/Experimental/PCG/Source/PCG/Public/MatchAndSet/PCGMatchAndSetWeightedByCategory.h
@@ -0,0 +1,77 @@
+// Copyright Epic Games, Inc. All Rights Reserved.
+
+#pragma once
+
+#include "MatchAndSet/PCGMatchAndSetBase.h"
+#include "MatchAndSet/PCGMatchAndSetWeighted.h"
+#include "Metadata/PCGMetadataTypesConstantStruct.h"
+
+#include "PCGMatchAndSetWeightedByCategory.generated.h"
+
+USTRUCT(BlueprintType)
+struct PCG_API FPCGMatchAndSetWeightedByCategoryEntryList
+{
+	GENERATED_BODY()
+
+	FPCGMatchAndSetWeightedByCategoryEntryList();
+
+	void SetType(EPCGMetadataTypes InType, EPCGMetadataTypesConstantStructStringMode InStringMode);
+	int GetTotalWeight() const;
+
+	UPROPERTY(BlueprintReadWrite, EditAnywhere, Category = Settings, meta = (ShowOnlyInnerProperties))
+	FPCGMetadataTypesConstantStruct CategoryValue;
+
+	/** If the category is the default, if the input does not match to anything, it will use this category. */
+	UPROPERTY(BlueprintReadWrite, EditAnywhere, Category = Settings)
+	bool bIsDefault = false;
+
+	/** Values and their weights */
+	UPROPERTY(BlueprintReadWrite, EditAnywhere, Category = Settings)
+	TArray<FPCGMatchAndSetWeightedEntry> WeightedEntries;
+};
+
+UCLASS(BlueprintType, ClassGroup = (Procedural))
+class PCG_API UPCGMatchAndSetWeightedByCategory : public UPCGMatchAndSetBase
+{
+	GENERATED_BODY()
+
+public:
+	virtual bool UsesRandomProcess() const { return true; }
+	virtual bool ShouldMutateSeed() const { return bShouldMutateSeed; }
+	virtual void SetType(EPCGMetadataTypes InType, EPCGMetadataTypesConstantStructStringMode InStringMode) override;
+	/** Propagates (does not set) category (e.g. Match) type to entries. */
+	virtual void SetCategoryType(EPCGMetadataTypes InType, EPCGMetadataTypesConstantStructStringMode InStringMode);
+
+	virtual void MatchAndSet_Implementation(
+		FPCGContext& Context,
+		const UPCGPointMatchAndSetSettings* InSettings,
+		const UPCGPointData* InPointData,
+		UPCGPointData* OutPointData) const override;
+
+	virtual bool ValidatePreconditions_Implementation(const UPCGPointData* InPointData) const override;
+
+#if WITH_EDITOR
+	virtual void PostEditChangeProperty(FPropertyChangedEvent& PropertyChangedEvent) override;
+#endif
+
+public:
+	/** Attribute to match against */
+	UPROPERTY(BlueprintReadWrite, EditAnywhere, Category = Settings)
+	FName CategoryAttribute;
+
+	/** Type of the attribute to match against. */
+	UPROPERTY(BlueprintReadWrite, EditAnywhere, Category = Settings, meta = (ValidEnumValues = "Float, Double, Integer32, Integer64, Vector2, Vector, Vector4, Quaternion, Transform, String, Boolean, Rotator, Name"))
+	EPCGMetadataTypes CategoryType = EPCGMetadataTypes::Double;
+
+	/** String subtype of the attribute to match against (if required). */
+	UPROPERTY(BlueprintReadWrite, EditAnywhere, Category = Settings, meta = (EditCondition = "CategoryType == EPCGMetadataTypes::String", EditConditionHides))
+	EPCGMetadataTypesConstantStructStringMode CategoryStringMode = EPCGMetadataTypesConstantStructStringMode::String;
+
+	/** Lookup entries (key -> weighted list) */
+	UPROPERTY(BlueprintReadWrite, EditAnywhere, Category = Settings)
+	TArray<FPCGMatchAndSetWeightedByCategoryEntryList> Categories;
+
+	/** Controls whether the output data should mutate its seed - prevents issues when doing multiple random processes in a row */
+	UPROPERTY(BlueprintReadWrite, EditAnywhere, Category = Settings)
+	bool bShouldMutateSeed = true;
+};
\ No newline at end of file
diff --git a/Engine/Plugins/Experimental/PCG/Source/PCG/Public/Metadata/Accessors/PCGAttributeAccessorHelpers.h b/Engine/Plugins/Experimental/PCG/Source/PCG/Public/Metadata/Accessors/PCGAttributeAccessorHelpers.h
--- a/Engine/Plugins/Experimental/PCG/Source/PCG/Public/Metadata/Accessors/PCGAttributeAccessorHelpers.h
+++ b/Engine/Plugins/Experimental/PCG/Source/PCG/Public/Metadata/Accessors/PCGAttributeAccessorHelpers.h
@@ -19,7 +19,18 @@ namespace PCGAttributeAccessorHelpers
 	TUniquePtr<IPCGAttributeAccessor> CreatePropertyAccessor(const FName InPropertyName, const UClass* InClass);
 	TUniquePtr<IPCGAttributeAccessor> CreatePropertyAccessor(const FName InPropertyName, const UStruct* InStruct);
 
+	/** 
+	* Creates a const accessor to the property or attribute pointed at by the InSelector.
+	* Note that InData must not be null if the selector points to an attribute,
+	* but in the case of properties, it either has to be the appropriate type or null.
+	*/
 	TUniquePtr<const IPCGAttributeAccessor> CreateConstAccessor(const UPCGData* InData, const FPCGAttributePropertySelector& InSelector);
+	
+	/** 
+	* Creates a non-const accessor to the property or attribute pointed at by the InSelector.
+	* Note that InData must not be null if the selector points to an attribute,
+	* but in the case of properties, it either has to be the appropriate type or null.
+	*/
 	TUniquePtr<IPCGAttributeAccessor> CreateAccessor(UPCGData* InData, const FPCGAttributePropertySelector& InSelector);
 
 	TUniquePtr<const IPCGAttributeAccessorKeys> CreateConstKeys(const UPCGData* InData, const FPCGAttributePropertySelector& InSelector);
diff --git a/Engine/Plugins/Experimental/PCG/Source/PCG/Public/Metadata/PCGMetadataAttributeTpl.h b/Engine/Plugins/Experimental/PCG/Source/PCG/Public/Metadata/PCGMetadataAttributeTpl.h
--- a/Engine/Plugins/Experimental/PCG/Source/PCG/Public/Metadata/PCGMetadataAttributeTpl.h
+++ b/Engine/Plugins/Experimental/PCG/Source/PCG/Public/Metadata/PCGMetadataAttributeTpl.h
@@ -615,15 +615,15 @@ namespace PCGMetadataAttribute
 #undef AllocatePCGMetadataAttributeOnType
 	}
 
-	template <typename Func>
-	inline decltype(auto) CallbackWithRightType(uint16 TypeId, Func Callback)
+	template <typename Func, typename... Args>
+	inline decltype(auto) CallbackWithRightType(uint16 TypeId, Func Callback, Args&& ...InArgs)
 	{
-		using ReturnType = decltype(Callback(double{}));
+		using ReturnType = decltype(Callback(double{}, std::forward<Args>(InArgs)...));
 
 		switch (TypeId)
 		{
 
-#define PCG_CALLBACKWITHRIGHTTYPE_DECL(T) case (uint16)(PCG::Private::MetadataTypes<T>::Id): return Callback(T{});
+#define PCG_CALLBACKWITHRIGHTTYPE_DECL(T) case (uint16)(PCG::Private::MetadataTypes<T>::Id): return Callback(T{}, std::forward<Args>(InArgs)...);
 		PCG_FOREACH_SUPPORTEDTYPES(PCG_CALLBACKWITHRIGHTTYPE_DECL)
 #undef PCG_CALLBACKWITHRIGHTTYPE_DECL
 
diff --git a/Engine/Plugins/Experimental/PCG/Source/PCG/Public/Metadata/PCGMetadataTypesConstantStruct.h b/Engine/Plugins/Experimental/PCG/Source/PCG/Public/Metadata/PCGMetadataTypesConstantStruct.h
--- a/Engine/Plugins/Experimental/PCG/Source/PCG/Public/Metadata/PCGMetadataTypesConstantStruct.h
+++ b/Engine/Plugins/Experimental/PCG/Source/PCG/Public/Metadata/PCGMetadataTypesConstantStruct.h
@@ -30,10 +30,10 @@ public:
 
 	FString ToString() const;
 
-	UPROPERTY(BlueprintReadWrite, EditAnywhere, Category = Settings, meta = (ValidEnumValues = "Float, Double, Integer32, Integer64, Vector2, Vector, Vector4, Quaternion, Transform, String, Boolean, Rotator, Name"))
+	UPROPERTY(BlueprintReadWrite, EditAnywhere, Category = Settings, meta = (EditCondition = "bAllowsTypeChange", EditConditionHides, HideEditConditionToggle, ValidEnumValues = "Float, Double, Integer32, Integer64, Vector2, Vector, Vector4, Quaternion, Transform, String, Boolean, Rotator, Name"))
 	EPCGMetadataTypes Type = EPCGMetadataTypes::Double;
 
-	UPROPERTY(BlueprintReadWrite, EditAnywhere, Category = Settings, meta = (EditCondition = "Type == EPCGMetadataTypes::String", EditConditionHides))
+	UPROPERTY(BlueprintReadWrite, EditAnywhere, Category = Settings, meta = (EditCondition = "bAllowsTypeChange && Type == EPCGMetadataTypes::String", EditConditionHides))
 	EPCGMetadataTypesConstantStructStringMode StringMode = EPCGMetadataTypesConstantStructStringMode::String;
 
 	// All different types
@@ -81,6 +81,9 @@ public:
 
 	UPROPERTY(BlueprintReadWrite, EditAnywhere, Category = Settings, meta = (EditCondition = "Type == EPCGMetadataTypes::String && StringMode == EPCGMetadataTypesConstantStructStringMode::SoftObjectPath", EditConditionHides))
 	FSoftObjectPath SoftObjectPathValue;
+
+	UPROPERTY()
+	bool bAllowsTypeChange = true;
 };
 
 template <typename Func>
