Tentative fix for Non-unity build take2

#rnx

diff --git a/Engine/Plugins/Experimental/PCG/Source/PCG/Private/Elements/PCGPointMatchAndSet.cpp b/Engine/Plugins/Experimental/PCG/Source/PCG/Private/Elements/PCGPointMatchAndSet.cpp
--- a/Engine/Plugins/Experimental/PCG/Source/PCG/Private/Elements/PCGPointMatchAndSet.cpp
+++ b/Engine/Plugins/Experimental/PCG/Source/PCG/Private/Elements/PCGPointMatchAndSet.cpp
@@ -4,6 +4,7 @@
 
 #include "PCGContext.h"
 #include "PCGHelpers.h"
+#include "PCGPin.h"
 #include "Data/PCGPointData.h"
 #include "MatchAndSet/PCGMatchAndSetWeighted.h"
 #include "Metadata/Accessors/IPCGAttributeAccessor.h"
