Small QoL changes - flush cache runs gc, ctrl+click flushes cache on generation, etc.

#rnx

diff --git a/Engine/Plugins/Experimental/PCG/Source/PCG/Private/PCGSubsystem.cpp b/Engine/Plugins/Experimental/PCG/Source/PCG/Private/PCGSubsystem.cpp
--- a/Engine/Plugins/Experimental/PCG/Source/PCG/Private/PCGSubsystem.cpp
+++ b/Engine/Plugins/Experimental/PCG/Source/PCG/Private/PCGSubsystem.cpp
@@ -1623,4 +1623,12 @@ void UPCGSubsystem::FlushCache()
 	{
 		GraphExecutor->GetCache().ClearCache();
 	}
+
+#if WITH_EDITOR
+	// Garbage collection is very seldom run in the editor, but we currently can consume a lot of memory in the cache.
+	if (!PCGHelpers::IsRuntimeOrPIE())
+	{
+		CollectGarbage(RF_NoFlags, true);
+	}	
+#endif
 }
diff --git a/Engine/Plugins/Experimental/PCG/Source/PCGEditor/Private/PCGEditor.cpp b/Engine/Plugins/Experimental/PCG/Source/PCGEditor/Private/PCGEditor.cpp
--- a/Engine/Plugins/Experimental/PCG/Source/PCGEditor/Private/PCGEditor.cpp
+++ b/Engine/Plugins/Experimental/PCG/Source/PCGEditor/Private/PCGEditor.cpp
@@ -40,6 +40,7 @@
 #include "Editor/UnrealEdEngine.h"
 #include "EditorAssetLibrary.h"
 #include "Fonts/FontMeasure.h"
+#include "Framework/Application/SlateApplication.h"
 #include "Framework/Commands/GenericCommands.h"
 #include "Framework/Notifications/NotificationManager.h"
 #include "GraphEditorActions.h"
@@ -447,6 +448,15 @@ void FPCGEditor::OnForceGraphRegeneration_Clicked()
 {
 	if (PCGGraphBeingEdited)
 	{
+		FModifierKeysState ModifierKeys = FSlateApplication::Get().GetModifierKeys();
+		if (ModifierKeys.IsControlDown())
+		{
+			if (UPCGSubsystem* Subsystem = GetSubsystem())
+			{
+				Subsystem->FlushCache();
+			}
+		}
+
 		PCGGraphBeingEdited->ForceNotificationForEditor();
 	}
 }
@@ -2029,7 +2039,7 @@ bool FPCGEditor::IsVisibleProperty(const FPropertyAndParent& InPropertyAndParent
 
 UPCGSubsystem* FPCGEditor::GetSubsystem()
 {
-	UWorld* World = GEditor ? GEditor->GetEditorWorldContext().World() : nullptr;
+	UWorld* World = (GEditor ? (GEditor->PlayWorld ? GEditor->PlayWorld.Get() : GEditor->GetEditorWorldContext().World()) : nullptr);
 	return World ? World->GetSubsystem<UPCGSubsystem>() : nullptr;
 }
 
diff --git a/Engine/Plugins/Experimental/PCG/Source/PCGEditor/Private/PCGEditorCommands.cpp b/Engine/Plugins/Experimental/PCG/Source/PCGEditor/Private/PCGEditorCommands.cpp
--- a/Engine/Plugins/Experimental/PCG/Source/PCGEditor/Private/PCGEditorCommands.cpp
+++ b/Engine/Plugins/Experimental/PCG/Source/PCGEditor/Private/PCGEditorCommands.cpp
@@ -22,7 +22,7 @@ void FPCGEditorCommands::RegisterCommands()
 	UI_COMMAND(ConvertToStandaloneNodes, "Convert to standalone Nodes", "Converts instanced nodes to standalone nodes.", EUserInterfaceActionType::Button, FInputChord());
 	UI_COMMAND(Find, "Find", "Finds PCG nodes and comments in the current graph", EUserInterfaceActionType::Button, FInputChord(EModifierKey::Control, EKeys::F));
 	UI_COMMAND(PauseAutoRegeneration, "Pause Regen", "Pause automatic regeneration of the current graph.", EUserInterfaceActionType::ToggleButton, FInputChord(EModifierKey::Alt, EKeys::R));
-	UI_COMMAND(ForceGraphRegeneration, "Force Regen", "Manually force a regeneration of the current graph.", EUserInterfaceActionType::Button, FInputChord());
+	UI_COMMAND(ForceGraphRegeneration, "Force Regen", "Manually force a regeneration of the current graph.\nCtrl-click will also perform a flush cache before the regeneration.", EUserInterfaceActionType::Button, FInputChord());
 	UI_COMMAND(RunDeterminismNodeTest, "Run Determinism Test on Node", "Evaluate the current node for determinism.", EUserInterfaceActionType::Button, FInputChord(EModifierKey::Alt, EKeys::D));
 	UI_COMMAND(RunDeterminismGraphTest, "Graph Determinism Test", "Evaluate the current graph for determinism.", EUserInterfaceActionType::Button, FInputChord());
 	UI_COMMAND(EditGraphSettings, "Graph Settings", "Edit the graph settings.", EUserInterfaceActionType::ToggleButton, FInputChord());
