Fixed static analysis warnings

#rnx

diff --git a/Engine/Plugins/Experimental/PCG/Source/PCG/Private/MatchAndSet/PCGMatchAndSetByAttribute.cpp b/Engine/Plugins/Experimental/PCG/Source/PCG/Private/MatchAndSet/PCGMatchAndSetByAttribute.cpp
--- a/Engine/Plugins/Experimental/PCG/Source/PCG/Private/MatchAndSet/PCGMatchAndSetByAttribute.cpp
+++ b/Engine/Plugins/Experimental/PCG/Source/PCG/Private/MatchAndSet/PCGMatchAndSetByAttribute.cpp
@@ -115,7 +115,7 @@ void UPCGMatchAndSetByAttribute::MatchAndSet_Implementation(
 	//TODO: implement async loop?
 	FPCGAttributePropertySelector InputSource;
 	InputSource.Selection = EPCGAttributePropertySelection::Attribute;
-	InputSource.AttributeName = (MatchSourceAttribute == NAME_None ? InPointData->Metadata->GetLatestAttributeNameOrNone() : MatchSourceAttribute);
+	InputSource.AttributeName = ((MatchSourceAttribute == NAME_None) ? InPointData->Metadata->GetLatestAttributeNameOrNone() : MatchSourceAttribute);
 
 	TUniquePtr<const IPCGAttributeAccessor> InputAccessor = PCGAttributeAccessorHelpers::CreateConstAccessor(InPointData, InputSource);
 	TUniquePtr<const IPCGAttributeAccessorKeys> InputKeys = PCGAttributeAccessorHelpers::CreateConstKeys(InPointData, InputSource);
@@ -199,7 +199,7 @@ bool UPCGMatchAndSetByAttribute::ValidatePreconditions_Implementation(const UPCG
 	}
 
 	// Check if source attribute exists
-	FName AttributeName = (MatchSourceAttribute == NAME_None ? InPointData->Metadata->GetLatestAttributeNameOrNone() : MatchSourceAttribute);
+	FName AttributeName = ((MatchSourceAttribute == NAME_None) ? InPointData->Metadata->GetLatestAttributeNameOrNone() : MatchSourceAttribute);
 
 	if (!InPointData->Metadata->HasAttribute(AttributeName))
 	{
diff --git a/Engine/Plugins/Experimental/PCG/Source/PCG/Private/MatchAndSet/PCGMatchAndSetWeightedByCategory.cpp b/Engine/Plugins/Experimental/PCG/Source/PCG/Private/MatchAndSet/PCGMatchAndSetWeightedByCategory.cpp
--- a/Engine/Plugins/Experimental/PCG/Source/PCG/Private/MatchAndSet/PCGMatchAndSetWeightedByCategory.cpp
+++ b/Engine/Plugins/Experimental/PCG/Source/PCG/Private/MatchAndSet/PCGMatchAndSetWeightedByCategory.cpp
@@ -168,7 +168,7 @@ void UPCGMatchAndSetWeightedByCategory::MatchAndSet_Implementation(
 
 	FPCGAttributePropertySelector InputSource;
 	InputSource.Selection = EPCGAttributePropertySelection::Attribute;
-	InputSource.AttributeName = (CategoryAttribute == NAME_None ? InPointData->Metadata->GetLatestAttributeNameOrNone() : CategoryAttribute);
+	InputSource.AttributeName = ((CategoryAttribute == NAME_None) ? InPointData->Metadata->GetLatestAttributeNameOrNone() : CategoryAttribute);
 
 	TUniquePtr<const IPCGAttributeAccessor> InputAccessor = PCGAttributeAccessorHelpers::CreateConstAccessor(InPointData, InputSource);
 	TUniquePtr<const IPCGAttributeAccessorKeys> InputKeys = PCGAttributeAccessorHelpers::CreateConstKeys(InPointData, InputSource);
@@ -284,7 +284,7 @@ bool UPCGMatchAndSetWeightedByCategory::ValidatePreconditions_Implementation(con
 		return false;
 	}
 
-	FName AttributeName = (CategoryAttribute == NAME_None ? InPointData->Metadata->GetLatestAttributeNameOrNone() : CategoryAttribute);
+	FName AttributeName = ((CategoryAttribute == NAME_None) ? InPointData->Metadata->GetLatestAttributeNameOrNone() : CategoryAttribute);
 	
 	if (!InPointData->Metadata->HasAttribute(AttributeName))
 	{
