[PCG] Updated jump to target doubleclick functionality for the execute blueprint node. It will now jump to source if no blueprint element type is specified.

diff --git a/Engine/Plugins/Experimental/PCG/Source/PCG/Private/Elements/PCGExecuteBlueprint.cpp b/Engine/Plugins/Experimental/PCG/Source/PCG/Private/Elements/PCGExecuteBlueprint.cpp
--- a/Engine/Plugins/Experimental/PCG/Source/PCG/Private/Elements/PCGExecuteBlueprint.cpp
+++ b/Engine/Plugins/Experimental/PCG/Source/PCG/Private/Elements/PCGExecuteBlueprint.cpp
@@ -411,7 +411,12 @@ void UPCGBlueprintSettings::GetTrackedActorTags(FPCGTagToSettingsMap& OutTagToSe
 
 UObject* UPCGBlueprintSettings::GetJumpTargetForDoubleClick() const
 {
-	return BlueprintElementType ? BlueprintElementType->ClassGeneratedBy : nullptr;
+	if (BlueprintElementType)
+	{
+		return BlueprintElementType->ClassGeneratedBy;
+	}
+	
+	return Super::GetJumpTargetForDoubleClick();
 }
 #endif // WITH_EDITOR
 
