[PCG] Made BlueprintElementInstance transactionable so that undo/redo works properly.

diff --git a/Engine/Plugins/Experimental/PCG/Source/PCG/Private/Elements/PCGExecuteBlueprint.cpp b/Engine/Plugins/Experimental/PCG/Source/PCG/Private/Elements/PCGExecuteBlueprint.cpp
--- a/Engine/Plugins/Experimental/PCG/Source/PCG/Private/Elements/PCGExecuteBlueprint.cpp
+++ b/Engine/Plugins/Experimental/PCG/Source/PCG/Private/Elements/PCGExecuteBlueprint.cpp
@@ -269,6 +269,7 @@ void UPCGBlueprintSettings::PostLoad()
 	if (BlueprintElementInstance)
 	{
 		BlueprintElementInstance->ConditionalPostLoad();
+		BlueprintElementInstance->SetFlags(RF_Transactional);
 #if WITH_EDITOR
 		BlueprintElementInstance->bCreatesArtifacts |= bCreatesArtifacts_DEPRECATED;
 		BlueprintElementInstance->bCanBeMultithreaded |= bCanBeMultithreaded_DEPRECATED;
@@ -364,7 +365,7 @@ void UPCGBlueprintSettings::RefreshBlueprintElement()
 
 	if (BlueprintElementType)
 	{
-		BlueprintElementInstance = NewObject<UPCGBlueprintElement>(this, BlueprintElementType);
+		BlueprintElementInstance = NewObject<UPCGBlueprintElement>(this, BlueprintElementType, NAME_None, RF_Transactional);
 		BlueprintElementInstance->Initialize();
 		SetupBlueprintElementEvent();
 	}
