Fixed crash on static mesh spawner copy & paste.
Fixed similar issue in the Match & Set

#rnx

diff --git a/Engine/Plugins/Experimental/PCG/Source/PCG/Private/Elements/PCGPointMatchAndSet.cpp b/Engine/Plugins/Experimental/PCG/Source/PCG/Private/Elements/PCGPointMatchAndSet.cpp
--- a/Engine/Plugins/Experimental/PCG/Source/PCG/Private/Elements/PCGPointMatchAndSet.cpp
+++ b/Engine/Plugins/Experimental/PCG/Source/PCG/Private/Elements/PCGPointMatchAndSet.cpp
@@ -17,7 +17,6 @@
 UPCGPointMatchAndSetSettings::UPCGPointMatchAndSetSettings(const FObjectInitializer& ObjectInitializer)
 {
 	MatchAndSetType = UPCGMatchAndSetWeighted::StaticClass();
-	MatchAndSetInstance = ObjectInitializer.CreateDefaultSubobject<UPCGMatchAndSetWeighted>(this, TEXT("DefaultMatchAndSetInstance"));
 	bUseSeed = MatchAndSetInstance && MatchAndSetInstance->UsesRandomProcess();
 }
 
@@ -51,6 +50,16 @@ FPCGElementPtr UPCGPointMatchAndSetSettings::CreateElement() const
 	return MakeShared<FPCGPointMatchAndSetElement>();
 }
 
+void UPCGPointMatchAndSetSettings::PostLoad()
+{
+	Super::PostLoad();
+
+	if (!MatchAndSetInstance)
+	{
+		RefreshMatchAndSet();
+	}	
+}
+
 #if WITH_EDITOR
 
 void UPCGPointMatchAndSetSettings::PostEditChangeProperty(FPropertyChangedEvent& PropertyChangedEvent)
diff --git a/Engine/Plugins/Experimental/PCG/Source/PCG/Private/Elements/PCGStaticMeshSpawner.cpp b/Engine/Plugins/Experimental/PCG/Source/PCG/Private/Elements/PCGStaticMeshSpawner.cpp
--- a/Engine/Plugins/Experimental/PCG/Source/PCG/Private/Elements/PCGStaticMeshSpawner.cpp
+++ b/Engine/Plugins/Experimental/PCG/Source/PCG/Private/Elements/PCGStaticMeshSpawner.cpp
@@ -20,7 +20,13 @@ UPCGStaticMeshSpawnerSettings::UPCGStaticMeshSpawnerSettings(const FObjectInitia
 	bUseSeed = true;
 
 	MeshSelectorType = UPCGMeshSelectorWeighted::StaticClass();
-	MeshSelectorInstance = ObjectInitializer.CreateDefaultSubobject<UPCGMeshSelectorWeighted>(this, TEXT("DefaultSelectorInstance"));
+	// Implementation note: this should not have been done here (it should have been null), as it causes issues with copy & paste
+	// when the thing to paste does not have that class for its instance.
+	// However, removing it makes it that any object actually using the instance created by default would be lost.
+	if (!this->HasAnyFlags(RF_ClassDefaultObject))
+	{
+		MeshSelectorInstance = ObjectInitializer.CreateDefaultSubobject<UPCGMeshSelectorWeighted>(this, TEXT("DefaultSelectorInstance"));
+	}
 }
 
 FPCGElementPtr UPCGStaticMeshSpawnerSettings::CreateElement() const
diff --git a/Engine/Plugins/Experimental/PCG/Source/PCG/Public/Elements/PCGPointMatchAndSet.h b/Engine/Plugins/Experimental/PCG/Source/PCG/Public/Elements/PCGPointMatchAndSet.h
--- a/Engine/Plugins/Experimental/PCG/Source/PCG/Public/Elements/PCGPointMatchAndSet.h
+++ b/Engine/Plugins/Experimental/PCG/Source/PCG/Public/Elements/PCGPointMatchAndSet.h
@@ -35,6 +35,8 @@ protected:
 
 public:
 	// ~Begin UObject interface
+	void PostLoad();
+
 #if WITH_EDITOR
 	virtual void PostEditChangeProperty(FPropertyChangedEvent& PropertyChangedEvent) override;
 #endif
