[PCG] Redo - Add support for advanced pins within nodes using BP paradigm

* Input/Output nodes has now advanced pins always created.
* Remove `Show Advanced Pins` in Input/Output node
* Add `bIsAdvancedPin` in PinProperties
* Any node that has at least 1 advanced pin will have an arrow at the bottom and will hide all advanced pins.
* Adding a new concept of `CreateDefaultPins` to make sure the graph always create the same default input/output node.

diff --git a/Engine/Plugins/Experimental/PCG/Source/PCG/Private/Elements/PCGSurfaceSampler.cpp b/Engine/Plugins/Experimental/PCG/Source/PCG/Private/Elements/PCGSurfaceSampler.cpp
--- a/Engine/Plugins/Experimental/PCG/Source/PCG/Private/Elements/PCGSurfaceSampler.cpp
+++ b/Engine/Plugins/Experimental/PCG/Source/PCG/Private/Elements/PCGSurfaceSampler.cpp
@@ -396,17 +396,6 @@ void UPCGSurfaceSamplerSettings::ApplyDeprecationBeforeUpdatePins(UPCGNode* InOu
 			InputPins[0]->Properties.Label = PCGSurfaceSamplerConstants::BoundingShapeLabel;
 		}
 
-		// If no pins or only 1 pin, then we will be doing rewiring later, so ask for advanced pins on graph input node
-		UPCGNode* GraphInputNode = InOutNode->GetGraph() ? InOutNode->GetGraph()->GetInputNode() : nullptr;
-		if (GraphInputNode && (InputPins.IsEmpty() || (InputPins.Num() == 1 && PCGSurfaceSampler::IsPinOnlyConnectedToInputNode(InputPins[0], GraphInputNode))))
-		{
-			// Ask for advanced params to be visible, as it will be useful later for wiring up the Surface input to the advanced Landscape pin
-			if (UPCGGraphInputOutputSettings* InputOutputSettings = Cast<UPCGGraphInputOutputSettings>(GraphInputNode->GetSettings()))
-			{
-				InputOutputSettings->SetShowAdvancedPins(true);
-			}
-		}
-
 		// A new params pin was added, migrate the first param connection there if any
 		PCGSettingsHelpers::DeprecationBreakOutParamsToNewPin(InOutNode, InputPins, OutputPins);
 	}
diff --git a/Engine/Plugins/Experimental/PCG/Source/PCG/Private/PCGGraph.cpp b/Engine/Plugins/Experimental/PCG/Source/PCG/Private/PCGGraph.cpp
--- a/Engine/Plugins/Experimental/PCG/Source/PCG/Private/PCGGraph.cpp
+++ b/Engine/Plugins/Experimental/PCG/Source/PCG/Private/PCGGraph.cpp
@@ -23,7 +23,7 @@ UPCGGraph::UPCGGraph(const FObjectInitializer& ObjectInitializer)
 	// Since pins would be allocated after initializing the input/output nodes, we must make sure to allocate them using the object initializer
 	int NumAllocatedPins = 1;
 	auto PinAllocator = [&ObjectInitializer, &NumAllocatedPins](UPCGNode* Node)
-	{ 
+	{
 		FName DefaultPinName = TEXT("DefaultPin");
 		DefaultPinName.SetNumber(NumAllocatedPins++);
 		return ObjectInitializer.CreateDefaultSubobject<UPCGPin>(Node, DefaultPinName);
@@ -32,16 +32,34 @@ UPCGGraph::UPCGGraph(const FObjectInitializer& ObjectInitializer)
 	UPCGGraphInputOutputSettings* InputSettings = ObjectInitializer.CreateDefaultSubobject<UPCGGraphInputOutputSettings>(this, TEXT("DefaultInputNodeSettings"));
 	InputSettings->SetInput(true);
 	InputNode->SetSettingsInterface(InputSettings, /*bUpdatePins=*/false);
-	InputNode->UpdatePins(PinAllocator);
-	
+
+	// Only allocate default pins if this is the default object
+	if (this->HasAnyFlags(RF_ClassDefaultObject))
+	{
+		InputNode->CreateDefaultPins(PinAllocator);
+	}
+	else
+	{
+		InputNode->UpdatePins();
+	}
+
 	OutputNode = ObjectInitializer.CreateDefaultSubobject<UPCGNode>(this, TEXT("DefaultOutputNode"));
 	OutputNode->SetFlags(RF_Transactional);
 
 	UPCGGraphInputOutputSettings* OutputSettings = ObjectInitializer.CreateDefaultSubobject<UPCGGraphInputOutputSettings>(this, TEXT("DefaultOutputNodeSettings"));
 	OutputSettings->SetInput(false);
 	OutputNode->SetSettingsInterface(OutputSettings, /*bUpdatePins=*/false);
-	OutputNode->UpdatePins(PinAllocator);
-	
+
+	// Only allocate default pins if this is the default object
+	if (this->HasAnyFlags(RF_ClassDefaultObject))
+	{
+		OutputNode->CreateDefaultPins(PinAllocator);
+	}
+	else
+	{
+		OutputNode->UpdatePins();
+	}
+
 #if WITH_EDITOR
 	OutputNode->PositionX = 200;
 #endif
diff --git a/Engine/Plugins/Experimental/PCG/Source/PCG/Private/PCGInputOutputSettings.cpp b/Engine/Plugins/Experimental/PCG/Source/PCG/Private/PCGInputOutputSettings.cpp
--- a/Engine/Plugins/Experimental/PCG/Source/PCG/Private/PCGInputOutputSettings.cpp
+++ b/Engine/Plugins/Experimental/PCG/Source/PCG/Private/PCGInputOutputSettings.cpp
@@ -73,14 +73,15 @@ TArray<FPCGPinProperties> UPCGGraphInputOutputSettings::GetPinProperties() const
 		return FPCGPinProperties(InLabelAndTooltip.Label, DefaultPinDataType, /*bMultiConnections=*/true, /*bMultiData=*/true, InLabelAndTooltip.Tooltip);
 	});
 	
-	if (bShowAdvancedPins)
-	{
-		Algo::Transform(StaticAdvancedLabels(), PinProperties, [DefaultPinDataType](const FLabelAndTooltip& InLabelAndTooltip) {
-			const bool bIsLandscapePin = (InLabelAndTooltip.Label == PCGInputOutputConstants::DefaultLandscapeLabel || InLabelAndTooltip.Label == PCGInputOutputConstants::DefaultLandscapeHeightLabel);
-			const EPCGDataType PinType = bIsLandscapePin ? EPCGDataType::Surface : DefaultPinDataType;
-			return FPCGPinProperties(InLabelAndTooltip.Label, PinType, /*bMultiConnection=*/true, /*bMultiData=*/false, InLabelAndTooltip.Tooltip);
-		});
-	}
+
+	Algo::Transform(StaticAdvancedLabels(), PinProperties, [this, DefaultPinDataType](const FLabelAndTooltip& InLabelAndTooltip) {
+		const bool bIsLandscapePin = (InLabelAndTooltip.Label == PCGInputOutputConstants::DefaultLandscapeLabel || InLabelAndTooltip.Label == PCGInputOutputConstants::DefaultLandscapeHeightLabel);
+		const EPCGDataType PinType = bIsLandscapePin ? EPCGDataType::Surface : DefaultPinDataType;
+		FPCGPinProperties Res(InLabelAndTooltip.Label, PinType, /*bMultiConnection=*/true, /*bMultiData=*/false, InLabelAndTooltip.Tooltip);
+		Res.bAdvancedPin = true;
+		return Res;
+	});
+
 
 	PinProperties.Append(CustomPins);
 	return PinProperties;
@@ -96,18 +97,22 @@ TArray<FPCGPinProperties> UPCGGraphInputOutputSettings::OutputPinProperties() co
 	return GetPinProperties();
 }
 
-bool UPCGGraphInputOutputSettings::IsPinAdvanced(const UPCGPin* Pin) const
+TArray<FPCGPinProperties> UPCGGraphInputOutputSettings::DefaultInputPinProperties() const
 {
-	return Pin && StaticAdvancedLabels().FindByPredicate([Pin](const FLabelAndTooltip& InLabelAndTooltip) -> bool { return InLabelAndTooltip.Label == Pin->Properties.Label; });
+	// It is important for serialization that this is not modified, or it could break existing graphs.
+	TArray<FPCGPinProperties> PinProperties;
+	const bool bIsInputPin = bIsInput;
+	const EPCGDataType DefaultPinDataType = bIsInput ? EPCGDataType::Spatial : EPCGDataType::Any;
+	Algo::Transform(StaticLabels(), PinProperties, [bIsInputPin, DefaultPinDataType](const FLabelAndTooltip& InLabelAndTooltip) {
+		return FPCGPinProperties(InLabelAndTooltip.Label, DefaultPinDataType, /*bMultiConnections=*/true, /*bMultiData=*/true, InLabelAndTooltip.Tooltip);
+	});
+
+	return PinProperties;
 }
 
-void UPCGGraphInputOutputSettings::SetShowAdvancedPins(bool bValue)
+TArray<FPCGPinProperties> UPCGGraphInputOutputSettings::DefaultOutputPinProperties() const
 {
-	if (bValue != bShowAdvancedPins)
-	{
-		Modify();
-		bShowAdvancedPins = bValue;
-	}
+	return DefaultInputPinProperties();
 }
 
 void UPCGGraphInputOutputSettings::AddCustomPin(const FPCGPinProperties& NewCustomPinProperties)
@@ -116,4 +121,10 @@ void UPCGGraphInputOutputSettings::AddCustomPin(const FPCGPinProperties& NewCust
 	CustomPins.Add(NewCustomPinProperties);
 }
 
+bool UPCGGraphInputOutputSettings::IsCustomPin(const UPCGPin* InPin) const
+{
+	check(InPin);
+	return CustomPins.Contains(InPin->Properties);
+}
+
 #undef LOCTEXT_NAMESPACE
diff --git a/Engine/Plugins/Experimental/PCG/Source/PCG/Private/PCGNode.cpp b/Engine/Plugins/Experimental/PCG/Source/PCG/Private/PCGNode.cpp
--- a/Engine/Plugins/Experimental/PCG/Source/PCG/Private/PCGNode.cpp
+++ b/Engine/Plugins/Experimental/PCG/Source/PCG/Private/PCGNode.cpp
@@ -578,6 +578,25 @@ EPCGChangeType UPCGNode::UpdatePins()
 	return UpdatePins([](UPCGNode* Node){ return NewObject<UPCGPin>(Node); });
 }
 
+void UPCGNode::CreateDefaultPins(TFunctionRef<UPCGPin* (UPCGNode*)> PinAllocator)
+{
+	auto CreatePins = [this, &PinAllocator](TArray<UPCGPin*>& Pins, const TArray<FPCGPinProperties>& PinProperties)
+	{
+		for (const FPCGPinProperties& Properties : PinProperties)
+		{
+			UPCGPin* NewPin = PinAllocator(this);
+			NewPin->Node = this;
+			NewPin->Properties = Properties;
+			Pins.Add(NewPin);
+		};
+	};
+
+	UPCGSettings* Settings = GetSettings();
+	check(Settings);
+	CreatePins(InputPins, Settings->DefaultInputPinProperties());
+	CreatePins(OutputPins, Settings->DefaultOutputPinProperties());
+}
+
 EPCGChangeType UPCGNode::UpdatePins(TFunctionRef<UPCGPin*(UPCGNode*)> PinAllocator)
 {
 	if (!GetSettings())
@@ -660,30 +679,60 @@ EPCGChangeType UPCGNode::UpdatePins(TFunctionRef<UPCGPin*(UPCGNode*)> PinAllocat
 		}
 		else
 		{
-			if(!UnmatchedPins.IsEmpty() || !UnmatchedProperties.IsEmpty())
+			// Verification that we don't have 2 pins with the same name
+			// If so, mark them to be removed.
+			TSet<FName> AllPinNames;
+			TArray<UPCGPin*> DuplicatedNamePins;
+
+			for (UPCGPin* Pin : Pins)
+			{
+				if (AllPinNames.Contains(Pin->Properties.Label))
+				{
+					DuplicatedNamePins.Add(Pin);
+				}
+
+				AllPinNames.Add(Pin->Properties.Label);
+			}
+
+			if(!UnmatchedPins.IsEmpty() || !UnmatchedProperties.IsEmpty() || !DuplicatedNamePins.IsEmpty())
 			{
 				Modify();
 				bChangedPins = true;
 			}
 
-			// Remove old pins
-			for (int32 UnmatchedPinIndex = UnmatchedPins.Num() - 1; UnmatchedPinIndex >= 0; --UnmatchedPinIndex)
+			auto RemovePins = [&Pins, &AllPinNames, &bAppliedEdgeChanges](TArray<UPCGPin*>& PinsToRemove, bool bRemoveFromAllNames)
 			{
-				const int32 PinIndex = Pins.IndexOfByKey(UnmatchedPins[UnmatchedPinIndex]);
-				check(PinIndex >= 0);
+				for (int32 RemovedPinIndex = PinsToRemove.Num() - 1; RemovedPinIndex >= 0; --RemovedPinIndex)
+				{
+					const int32 PinIndex = Pins.IndexOfByKey(PinsToRemove[RemovedPinIndex]);
+					check(PinIndex >= 0);
 
-				bAppliedEdgeChanges |= Pins[PinIndex]->BreakAllEdges();
-				Pins.RemoveAt(PinIndex);
-			}
+					if (bRemoveFromAllNames)
+					{
+						AllPinNames.Remove(Pins[PinIndex]->Properties.Label);
+					}
+
+					bAppliedEdgeChanges |= Pins[PinIndex]->BreakAllEdges();
+					Pins.RemoveAt(PinIndex);
+				}
+			};
+
+			RemovePins(UnmatchedPins, /*bRemoveFromAllNames=*/ true);
+			RemovePins(DuplicatedNamePins, /*bRemoveFromAllNames=*/ false);
 
 			// Add new pins
 			for (const FPCGPinProperties& UnmatchedProperty : UnmatchedProperties)
 			{
-				const int32 InsertIndex = PinProperties.IndexOfByKey(UnmatchedProperty);
-				UPCGPin* NewPin = PinAllocator(this);
-				NewPin->Node = this;
-				NewPin->Properties = UnmatchedProperty;
-				Pins.Insert(NewPin, InsertIndex);
+				if (ensure(!AllPinNames.Contains(UnmatchedProperty.Label)))
+				{
+					AllPinNames.Add(UnmatchedProperty.Label);
+
+					const int32 InsertIndex = FMath::Min(PinProperties.IndexOfByKey(UnmatchedProperty), Pins.Num());
+					UPCGPin* NewPin = PinAllocator(this);
+					NewPin->Node = this;
+					NewPin->Properties = UnmatchedProperty;
+					Pins.Insert(NewPin, InsertIndex);
+				}
 			}
 		}
 
diff --git a/Engine/Plugins/Experimental/PCG/Source/PCG/Private/PCGPin.cpp b/Engine/Plugins/Experimental/PCG/Source/PCG/Private/PCGPin.cpp
--- a/Engine/Plugins/Experimental/PCG/Source/PCG/Private/PCGPin.cpp
+++ b/Engine/Plugins/Experimental/PCG/Source/PCG/Private/PCGPin.cpp
@@ -18,7 +18,8 @@ bool FPCGPinProperties::operator==(const FPCGPinProperties& Other) const
 	return Label == Other.Label &&
 		AllowedTypes == Other.AllowedTypes &&
 		bAllowMultipleConnections == Other.bAllowMultipleConnections &&
-		bAllowMultipleData == Other.bAllowMultipleData;
+		bAllowMultipleData == Other.bAllowMultipleData &&
+		bAdvancedPin == Other.bAdvancedPin;
 }
 
 UPCGPin::UPCGPin(const FObjectInitializer& ObjectInitializer)
diff --git a/Engine/Plugins/Experimental/PCG/Source/PCG/Public/PCGInputOutputSettings.h b/Engine/Plugins/Experimental/PCG/Source/PCG/Public/PCGInputOutputSettings.h
--- a/Engine/Plugins/Experimental/PCG/Source/PCG/Public/PCGInputOutputSettings.h
+++ b/Engine/Plugins/Experimental/PCG/Source/PCG/Public/PCGInputOutputSettings.h
@@ -44,10 +44,12 @@ public:
 	TArray<FPCGPinProperties> InputPinProperties() const override;
 	TArray<FPCGPinProperties> OutputPinProperties() const override;
 
+	virtual TArray<FPCGPinProperties> DefaultInputPinProperties() const override;
+	virtual TArray<FPCGPinProperties> DefaultOutputPinProperties() const override;
+
 	void SetInput(bool bInIsInput) { bIsInput = bInIsInput; }
 
-	bool IsPinAdvanced(const UPCGPin* Pin) const;
-	void SetShowAdvancedPins(bool bValue);
+	bool IsCustomPin(const UPCGPin* InPin) const;
 
 	void AddCustomPin(const FPCGPinProperties& NewCustomPinProperties);
 
@@ -78,9 +80,6 @@ protected:
 	UPROPERTY(EditAnywhere, BlueprintReadWrite, Category = "Settings|Input")
 	TArray<FPCGPinProperties> CustomPins;
 
-	UPROPERTY(EditAnywhere, BlueprintReadWrite, Category = "Settings|Input")
-	bool bShowAdvancedPins = false;
-
 	TArray<FLabelAndTooltip> StaticInLabels;
 	TArray<FLabelAndTooltip> StaticAdvancedInLabels;
 	TArray<FLabelAndTooltip> StaticOutLabels;
diff --git a/Engine/Plugins/Experimental/PCG/Source/PCG/Public/PCGNode.h b/Engine/Plugins/Experimental/PCG/Source/PCG/Public/PCGNode.h
--- a/Engine/Plugins/Experimental/PCG/Source/PCG/Public/PCGNode.h
+++ b/Engine/Plugins/Experimental/PCG/Source/PCG/Public/PCGNode.h
@@ -165,6 +165,14 @@ protected:
 	EPCGChangeType UpdatePins();
 	EPCGChangeType UpdatePins(TFunctionRef<UPCGPin* (UPCGNode*)> PinAllocator);
 
+	// When we create a new graph, we initialize the input/output nodes as default, with default pins.
+	// Those default pins are not serialized, therefore if we change the default pins, combined with the use
+	// of recycling objects in Unreal, can lead to pins that are garbage or even worse: valid pins but not the right
+	// one, potentially making the edges connecting wrong pins together!
+	// That is why we have a specific function to create default pins, and we have to make sure that those
+	// default pins are always created the same way.
+	void CreateDefaultPins(TFunctionRef<UPCGPin* (UPCGNode*)> PinAllocator);
+
 #if WITH_EDITOR
 	virtual void PostEditChangeProperty(struct FPropertyChangedEvent& PropertyChangedEvent) override;
 	void OnSettingsChanged(UPCGSettings* InSettings, EPCGChangeType ChangeType);
diff --git a/Engine/Plugins/Experimental/PCG/Source/PCG/Public/PCGPin.h b/Engine/Plugins/Experimental/PCG/Source/PCG/Public/PCGPin.h
--- a/Engine/Plugins/Experimental/PCG/Source/PCG/Public/PCGPin.h
+++ b/Engine/Plugins/Experimental/PCG/Source/PCG/Public/PCGPin.h
@@ -31,6 +31,10 @@ struct PCG_API FPCGPinProperties
 	UPROPERTY(BlueprintReadWrite, EditAnywhere, Category = Settings)
 	bool bAllowMultipleConnections = true;
 
+	/* Advanced pin will be hidden by default in the UI and will be shown only if the user extend the node (in the UI) to see advanced pins. */
+	UPROPERTY(BlueprintReadWrite, EditAnywhere, Category = Settings)
+	bool bAdvancedPin = false;
+
 #if WITH_EDITORONLY_DATA
 	UPROPERTY(EditAnywhere, Category = Settings)
 	FText Tooltip;
diff --git a/Engine/Plugins/Experimental/PCG/Source/PCG/Public/PCGSettings.h b/Engine/Plugins/Experimental/PCG/Source/PCG/Public/PCGSettings.h
--- a/Engine/Plugins/Experimental/PCG/Source/PCG/Public/PCGSettings.h
+++ b/Engine/Plugins/Experimental/PCG/Source/PCG/Public/PCGSettings.h
@@ -3,14 +3,14 @@
 #pragma once
 
 #include "PCGData.h"
-#include "PCGElement.h"
 #include "PCGDebug.h"
+#include "PCGElement.h"
+#include "PCGPin.h"
 #include "Tests/Determinism/PCGDeterminismSettings.h"
 
 #include "PCGSettings.generated.h"
 
 class UPCGPin;
-struct FPCGPinProperties;
 struct FPropertyChangedEvent;
 
 class UPCGGraph;
@@ -111,6 +111,11 @@ public:
 	virtual TArray<FPCGPinProperties> InputPinProperties() const;
 	virtual TArray<FPCGPinProperties> OutputPinProperties() const;
 
+	// Internal functions, should not be used by any user.
+	// Return a different subset for for input/output pin properties, in case of a default object.
+	virtual TArray<FPCGPinProperties> DefaultInputPinProperties() const { return InputPinProperties(); }
+	virtual TArray<FPCGPinProperties> DefaultOutputPinProperties() const { return OutputPinProperties(); }
+
 	bool operator==(const UPCGSettings& Other) const;
 	uint32 GetCrc32() const;
 	bool UseSeed() const { return bUseSeed; }
diff --git a/Engine/Plugins/Experimental/PCG/Source/PCGEditor/Private/PCGEditor.cpp b/Engine/Plugins/Experimental/PCG/Source/PCGEditor/Private/PCGEditor.cpp
--- a/Engine/Plugins/Experimental/PCG/Source/PCGEditor/Private/PCGEditor.cpp
+++ b/Engine/Plugins/Experimental/PCG/Source/PCGEditor/Private/PCGEditor.cpp
@@ -814,26 +814,6 @@ void FPCGEditor::OnCollapseNodesInSubgraph()
 	InputFromSubgraphPins = InputFromSubgraphPins.Difference(AllSubgraphPins);
 	OutputToSubgraphPins = OutputToSubgraphPins.Difference(AllSubgraphPins);
 
-	// Function that check if the pin is part of the advanced pins of the input/output node of the PCGGraph.
-	// If so, enable advanced pins on the subgraph.
-	auto IsAdvancedAndSet = [&PCGGraph, &NewPCGGraph](UPCGPin* Pin, bool bIsInput)
-	{
-		TObjectPtr<UPCGNode> InputOutputNode = bIsInput ? PCGGraph->GetInputNode() : PCGGraph->GetOutputNode();
-		TObjectPtr<UPCGNode> NewInputOutputNode = bIsInput ? NewPCGGraph->GetInputNode() : NewPCGGraph->GetOutputNode();
-
-		TObjectPtr<UPCGGraphInputOutputSettings> InputOutputSettings = CastChecked<UPCGGraphInputOutputSettings>(InputOutputNode->GetSettings());
-		if (InputOutputSettings->IsPinAdvanced(Pin))
-		{
-			TObjectPtr<UPCGGraphInputOutputSettings> NewInputOutputSettings = CastChecked<UPCGGraphInputOutputSettings>(NewInputOutputNode->GetSettings());
-			NewInputOutputSettings->SetShowAdvancedPins(true);
-			NewInputOutputNode->UpdateAfterSettingsChangeDuringCreation();
-
-			return true;
-		}
-
-		return false;
-	};
-
 	// Function that will create a new pin in the input/output node of the subgraph, with a unique name that will
 	// be the related to the pin passed as argument. It will be formatted like this:
 	// "{NodeName} {PinName} {OptionalIndex}"
@@ -925,7 +905,9 @@ void FPCGEditor::OnCollapseNodesInSubgraph()
 
 			if (Edge->InputPin->Node == PCGGraph->GetInputNode())
 			{
-				if (Edge->InputPin == PCGGraph->GetInputNode()->GetOutputPins()[0] || IsAdvancedAndSet(Edge->InputPin, true))
+				const UPCGGraphInputOutputSettings* Settings = Cast<const UPCGGraphInputOutputSettings>(Edge->InputPin->Node->GetSettings());
+
+				if (Settings && !Settings->IsCustomPin(Edge->InputPin))
 				{
 					OutsideSubgraphEdge.OutputPinLabel = Edge->InputPin->Properties.Label;
 					InsideSubgraphEdge.InputPin = NewPCGGraph->GetInputNode()->GetOutputPin(Edge->InputPin->Properties.Label);
@@ -956,7 +938,9 @@ void FPCGEditor::OnCollapseNodesInSubgraph()
 
 			if (Edge->OutputPin->Node == PCGGraph->GetOutputNode())
 			{
-				if (Edge->OutputPin == PCGGraph->GetOutputNode()->GetInputPins()[0] || IsAdvancedAndSet(Edge->OutputPin, false))
+				const UPCGGraphInputOutputSettings* Settings = Cast<const UPCGGraphInputOutputSettings>(Edge->OutputPin->Node->GetSettings());
+
+				if (Settings && !Settings->IsCustomPin(Edge->OutputPin))
 				{
 					OutsideSubgraphEdge.InputPinLabel = Edge->OutputPin->Properties.Label;
 					InsideSubgraphEdge.OutputPin = NewPCGGraph->GetOutputNode()->GetInputPin(Edge->OutputPin->Properties.Label);
diff --git a/Engine/Plugins/Experimental/PCG/Source/PCGEditor/Private/PCGEditorGraphNode.cpp b/Engine/Plugins/Experimental/PCG/Source/PCGEditor/Private/PCGEditorGraphNode.cpp
--- a/Engine/Plugins/Experimental/PCG/Source/PCGEditor/Private/PCGEditorGraphNode.cpp
+++ b/Engine/Plugins/Experimental/PCG/Source/PCGEditor/Private/PCGEditorGraphNode.cpp
@@ -54,15 +54,7 @@ void UPCGEditorGraphNode::AllocateDefaultPins()
 {
 	if (PCGNode)
 	{
-		for (const UPCGPin* InputPin : PCGNode->GetInputPins())
-		{
-			CreatePin(EEdGraphPinDirection::EGPD_Input, GetPinType(InputPin), InputPin->Properties.Label);
-		}
-
-		for (const UPCGPin* OutputPin : PCGNode->GetOutputPins())
-		{
-			CreatePin(EEdGraphPinDirection::EGPD_Output, GetPinType(OutputPin), OutputPin->Properties.Label);
-		}
+		CreatePins(PCGNode->GetInputPins(), PCGNode->GetOutputPins());
 	}
 }
 
diff --git a/Engine/Plugins/Experimental/PCG/Source/PCGEditor/Private/PCGEditorGraphNodeBase.cpp b/Engine/Plugins/Experimental/PCG/Source/PCGEditor/Private/PCGEditorGraphNodeBase.cpp
--- a/Engine/Plugins/Experimental/PCG/Source/PCGEditor/Private/PCGEditorGraphNodeBase.cpp
+++ b/Engine/Plugins/Experimental/PCG/Source/PCGEditor/Private/PCGEditorGraphNodeBase.cpp
@@ -572,4 +572,32 @@ void UPCGEditorGraphNodeBase::UpdatePosition()
 	}
 }
 
+void UPCGEditorGraphNodeBase::CreatePins(const TArray<UPCGPin*>& InInputPins, const TArray<UPCGPin*>& InOutputPins)
+{
+	bool bHasAdvancedPin = false;
+
+	for (const UPCGPin* InputPin : InInputPins)
+	{
+		UEdGraphPin* Pin = CreatePin(EEdGraphPinDirection::EGPD_Input, GetPinType(InputPin), InputPin->Properties.Label);
+		Pin->bAdvancedView = InputPin->Properties.bAdvancedPin;
+		bHasAdvancedPin |= Pin->bAdvancedView;
+	}
+
+	for (const UPCGPin* OutputPin : InOutputPins)
+	{
+		UEdGraphPin* Pin = CreatePin(EEdGraphPinDirection::EGPD_Output, GetPinType(OutputPin), OutputPin->Properties.Label);
+		Pin->bAdvancedView = OutputPin->Properties.bAdvancedPin;
+		bHasAdvancedPin |= Pin->bAdvancedView;
+	}
+
+	if (bHasAdvancedPin && AdvancedPinDisplay == ENodeAdvancedPins::NoPins)
+	{
+		AdvancedPinDisplay = ENodeAdvancedPins::Hidden;
+	}
+	else if (!bHasAdvancedPin)
+	{
+		AdvancedPinDisplay = ENodeAdvancedPins::NoPins;
+	}
+}
+
 #undef LOCTEXT_NAMESPACE
diff --git a/Engine/Plugins/Experimental/PCG/Source/PCGEditor/Private/PCGEditorGraphNodeBase.h b/Engine/Plugins/Experimental/PCG/Source/PCGEditor/Private/PCGEditorGraphNodeBase.h
--- a/Engine/Plugins/Experimental/PCG/Source/PCGEditor/Private/PCGEditorGraphNodeBase.h
+++ b/Engine/Plugins/Experimental/PCG/Source/PCGEditor/Private/PCGEditorGraphNodeBase.h
@@ -63,6 +63,8 @@ protected:
 	void UpdateCommentBubblePinned();
 	void UpdatePosition();
 
+	void CreatePins(const TArray<UPCGPin*>& InInputPins, const TArray<UPCGPin*>& InOutputPins);
+
 	UPROPERTY()
 	TObjectPtr<UPCGNode> PCGNode = nullptr;
 
diff --git a/Engine/Plugins/Experimental/PCG/Source/PCGEditor/Private/PCGEditorGraphNodeInput.cpp b/Engine/Plugins/Experimental/PCG/Source/PCGEditor/Private/PCGEditorGraphNodeInput.cpp
--- a/Engine/Plugins/Experimental/PCG/Source/PCGEditor/Private/PCGEditorGraphNodeInput.cpp
+++ b/Engine/Plugins/Experimental/PCG/Source/PCGEditor/Private/PCGEditorGraphNodeInput.cpp
@@ -17,10 +17,7 @@ void UPCGEditorGraphNodeInput::AllocateDefaultPins()
 {
 	if (PCGNode)
 	{
-		for (const UPCGPin* OutPin : PCGNode->GetOutputPins())
-		{
-			CreatePin(EEdGraphPinDirection::EGPD_Output, GetPinType(OutPin), OutPin->Properties.Label);
-		}
+		CreatePins(/*InInputPins=*/{}, PCGNode->GetOutputPins());
 	}
 }
 
diff --git a/Engine/Plugins/Experimental/PCG/Source/PCGEditor/Private/PCGEditorGraphNodeOutput.cpp b/Engine/Plugins/Experimental/PCG/Source/PCGEditor/Private/PCGEditorGraphNodeOutput.cpp
--- a/Engine/Plugins/Experimental/PCG/Source/PCGEditor/Private/PCGEditorGraphNodeOutput.cpp
+++ b/Engine/Plugins/Experimental/PCG/Source/PCGEditor/Private/PCGEditorGraphNodeOutput.cpp
@@ -17,10 +17,7 @@ void UPCGEditorGraphNodeOutput::AllocateDefaultPins()
 {
 	if (PCGNode)
 	{
-		for (const UPCGPin* InputPin : PCGNode->GetInputPins())
-		{
-			CreatePin(EEdGraphPinDirection::EGPD_Input, GetPinType(InputPin), InputPin->Properties.Label);
-		}
+		CreatePins(PCGNode->GetInputPins(), /*InOutputPins=*/{});
 	}
 }
 
