Nativized bounds modifier node

#rnx

diff --git a/Engine/Plugins/Experimental/PCG/Source/PCG/Private/Elements/PCGBoundsModifier.cpp b/Engine/Plugins/Experimental/PCG/Source/PCG/Private/Elements/PCGBoundsModifier.cpp
new file mode 100644
--- /dev/null
+++ b/Engine/Plugins/Experimental/PCG/Source/PCG/Private/Elements/PCGBoundsModifier.cpp
@@ -0,0 +1,132 @@
+// Copyright Epic Games, Inc. All Rights Reserved.
+
+#include "Elements/PCGBoundsModifier.h"
+
+#include "Helpers/PCGSettingsHelpers.h"
+#include "PCGContext.h"
+#include "PCGPoint.h"
+
+#include UE_INLINE_GENERATED_CPP_BY_NAME(PCGBoundsModifier)
+
+#define LOCTEXT_NAMESPACE "PCGBoundsModifier"
+
+FPCGElementPtr UPCGBoundsModifierSettings::CreateElement() const
+{
+	return MakeShared<FPCGBoundsModifier>();
+}
+
+TArray<FPCGPinProperties> UPCGBoundsModifierSettings::InputPinProperties() const
+{
+	TArray<FPCGPinProperties> PinProperties;
+	PinProperties.Emplace(PCGPinConstants::DefaultInputLabel, EPCGDataType::Point);
+	PinProperties.Emplace(PCGPinConstants::DefaultParamsLabel, EPCGDataType::Param);
+	return PinProperties;
+}
+
+#if WITH_EDITOR
+FText UPCGBoundsModifierSettings::GetNodeTooltipText() const
+{
+	return LOCTEXT("BoundsModifierNodeTooltip", "Applies a transformation on the point bounds & optionally its steepness.");
+}
+#endif // WITH_EDITOR
+
+bool FPCGBoundsModifier::ExecuteInternal(FPCGContext* Context) const
+{
+	TRACE_CPUPROFILER_EVENT_SCOPE(FPCGBoundsModifier::Execute);
+
+	const UPCGBoundsModifierSettings* Settings = Context->GetInputSettings<UPCGBoundsModifierSettings>();
+	check(Settings);
+
+	TArray<FPCGTaggedData> Inputs = Context->InputData.GetInputsByPin(PCGPinConstants::DefaultInputLabel);
+	TArray<FPCGTaggedData>& Outputs = Context->OutputData.TaggedData;
+	UPCGParamData* Params = Context->InputData.GetParams();
+
+	const EPCGBoundsModifierMode Mode = PCG_GET_OVERRIDEN_VALUE(Settings, Mode, Params);
+	const FVector BoundsMin = PCG_GET_OVERRIDEN_VALUE(Settings, BoundsMin, Params);
+	const FVector BoundsMax = PCG_GET_OVERRIDEN_VALUE(Settings, BoundsMax, Params);
+	const bool bAffectSteepness = PCG_GET_OVERRIDEN_VALUE(Settings, bAffectSteepness, Params);
+	const float Steepness = PCG_GET_OVERRIDEN_VALUE(Settings, Steepness, Params);
+
+	const FBox Bounds(BoundsMin, BoundsMax);
+
+	switch (Mode)
+	{
+	case EPCGBoundsModifierMode::Intersect:
+		ProcessPoints(Context, Inputs, Outputs, [&Bounds, bAffectSteepness, Steepness](const FPCGPoint& InPoint, FPCGPoint& OutPoint){
+			OutPoint = InPoint;
+			OutPoint.SetLocalBounds(InPoint.GetLocalBounds().Overlap(Bounds));
+			
+			if (bAffectSteepness)
+			{
+				OutPoint.Steepness = FMath::Min(InPoint.Steepness, Steepness);
+			}
+
+			return true;
+		});
+		break;
+
+	case EPCGBoundsModifierMode::Include:
+		ProcessPoints(Context, Inputs, Outputs, [&Bounds, bAffectSteepness, Steepness](const FPCGPoint& InPoint, FPCGPoint& OutPoint){
+			OutPoint = InPoint;
+			OutPoint.SetLocalBounds(InPoint.GetLocalBounds() + Bounds);
+
+			if (bAffectSteepness)
+			{
+				OutPoint.Steepness = FMath::Max(InPoint.Steepness, Steepness);
+			}
+
+			return true;
+		});
+		break;
+
+	case EPCGBoundsModifierMode::Translate:
+		ProcessPoints(Context, Inputs, Outputs, [&BoundsMin, &BoundsMax, bAffectSteepness, Steepness](const FPCGPoint& InPoint, FPCGPoint& OutPoint){
+			OutPoint = InPoint;
+			OutPoint.BoundsMin += BoundsMin;
+			OutPoint.BoundsMax += BoundsMax;
+
+			if (bAffectSteepness)
+			{
+				OutPoint.Steepness = FMath::Clamp(InPoint.Steepness + Steepness, 0.0f, 1.0f);
+			}
+
+			return true;
+		});
+		break;
+
+	case EPCGBoundsModifierMode::Scale:
+		ProcessPoints(Context, Inputs, Outputs, [&BoundsMin, &BoundsMax, bAffectSteepness, Steepness](const FPCGPoint& InPoint, FPCGPoint& OutPoint){
+			OutPoint = InPoint;
+			OutPoint.BoundsMin *= BoundsMin;
+			OutPoint.BoundsMax *= BoundsMax;
+
+			if (bAffectSteepness)
+			{
+				OutPoint.Steepness = FMath::Clamp(InPoint.Steepness * Steepness, 0.0f, 1.0f);
+			}
+
+			return true;
+		});
+		break;
+
+	case EPCGBoundsModifierMode::Set:
+		ProcessPoints(Context, Inputs, Outputs, [&Bounds, bAffectSteepness, Steepness](const FPCGPoint& InPoint, FPCGPoint& OutPoint){
+			OutPoint = InPoint;
+			OutPoint.SetLocalBounds(Bounds);
+
+			if (bAffectSteepness)
+			{
+				OutPoint.Steepness = Steepness;
+			}
+
+			return true;
+		});
+		break;
+	}
+
+	Outputs.Append(Context->InputData.GetAllSettings());
+	
+	return true;
+}
+
+#undef LOCTEXT_NAMESPACE
\ No newline at end of file
diff --git a/Engine/Plugins/Experimental/PCG/Source/PCG/Public/Elements/PCGBoundsModifier.h b/Engine/Plugins/Experimental/PCG/Source/PCG/Public/Elements/PCGBoundsModifier.h
new file mode 100644
--- /dev/null
+++ b/Engine/Plugins/Experimental/PCG/Source/PCG/Public/Elements/PCGBoundsModifier.h
@@ -0,0 +1,67 @@
+// Copyright Epic Games, Inc. All Rights Reserved.
+
+#pragma once
+
+#include "PCGPin.h"
+#include "PCGSettings.h"
+
+#include "Elements/PCGPointProcessingElementBase.h"
+
+#include "PCGBoundsModifier.generated.h"
+
+UENUM()
+enum class EPCGBoundsModifierMode : uint8
+{
+	Set,
+	Intersect,
+	Include,
+	Translate,
+	Scale
+};
+
+/**
+* This class controls/sets up a node that modifies the min/max bounds of the input points.
+*/
+UCLASS(BlueprintType, ClassGroup = (Procedural))
+class PCG_API UPCGBoundsModifierSettings : public UPCGSettings
+{
+	GENERATED_BODY()
+
+public:
+
+	// ~Begin UPCGSettings interface
+#if WITH_EDITOR
+	virtual FName GetDefaultNodeName() const override { return FName(TEXT("BoundsModifier")); }
+	virtual FText GetNodeTooltipText() const override;
+	virtual EPCGSettingsType GetType() const override { return EPCGSettingsType::Spatial; }
+#endif
+
+	virtual TArray<FPCGPinProperties> InputPinProperties() const override;
+	virtual TArray<FPCGPinProperties> OutputPinProperties() const override { return Super::DefaultPointOutputPinProperties(); }
+
+protected:
+	virtual FPCGElementPtr CreateElement() const override;
+	//~End UPCGSettings interface
+
+public:
+	UPROPERTY(BlueprintReadWrite, EditAnywhere, Category = Settings)
+	EPCGBoundsModifierMode Mode = EPCGBoundsModifierMode::Scale;
+
+	UPROPERTY(BlueprintReadWrite, EditAnywhere, Category = Settings)
+	FVector BoundsMin = FVector::One();
+
+	UPROPERTY(BlueprintReadWrite, EditAnywhere, Category = Settings)
+	FVector BoundsMax = FVector::One();
+
+	UPROPERTY(BlueprintReadWrite, EditAnywhere, Category = Settings)
+	bool bAffectSteepness = false;
+
+	UPROPERTY(BlueprintReadWrite, EditAnywhere, Category = Settings, meta = (EditCondition = "bAffectSteepness", ClampMin = "0", ClampMax = "1"))
+	float Steepness = 1.0f;
+};
+
+class FPCGBoundsModifier : public FPCGPointProcessingElementBase
+{
+protected:
+	virtual bool ExecuteInternal(FPCGContext* Context) const;
+};
\ No newline at end of file
