[PCG] Add name clash protection on custom pins for input/output

diff --git a/Engine/Plugins/Experimental/PCG/Source/PCG/Private/PCGInputOutputSettings.cpp b/Engine/Plugins/Experimental/PCG/Source/PCG/Private/PCGInputOutputSettings.cpp
--- a/Engine/Plugins/Experimental/PCG/Source/PCG/Private/PCGInputOutputSettings.cpp
+++ b/Engine/Plugins/Experimental/PCG/Source/PCG/Private/PCGInputOutputSettings.cpp
@@ -115,10 +115,12 @@ TArray<FPCGPinProperties> UPCGGraphInputOutputSettings::DefaultOutputPinProperti
 	return DefaultInputPinProperties();
 }
 
-void UPCGGraphInputOutputSettings::AddCustomPin(const FPCGPinProperties& NewCustomPinProperties)
+const FPCGPinProperties& UPCGGraphInputOutputSettings::AddCustomPin(const FPCGPinProperties& NewCustomPinProperties)
 {
 	Modify();
-	CustomPins.Add(NewCustomPinProperties);
+	int32 Index = CustomPins.Add(NewCustomPinProperties);
+	FixCustomPinProperties();
+	return CustomPins[Index];
 }
 
 bool UPCGGraphInputOutputSettings::IsCustomPin(const UPCGPin* InPin) const
@@ -127,4 +129,85 @@ bool UPCGGraphInputOutputSettings::IsCustomPin(const UPCGPin* InPin) const
 	return CustomPins.Contains(InPin->Properties);
 }
 
+#if WITH_EDITOR
+void UPCGGraphInputOutputSettings::PostEditChangeProperty(struct FPropertyChangedEvent& PropertyChangedEvent)
+{
+	if (PropertyChangedEvent.Property)
+	{
+		// Label has changed if we have modified "CustomPins" array
+		bool bLabelChanged = PropertyChangedEvent.Property->GetFName() == GET_MEMBER_NAME_CHECKED(UPCGGraphInputOutputSettings, CustomPins);
+
+		if (!bLabelChanged)
+		{
+			// Or might has changed if an element of the array "CustomPin" was modified.
+			bLabelChanged = PropertyChangedEvent.MemberProperty &&
+				PropertyChangedEvent.MemberProperty->GetFName() == GET_MEMBER_NAME_CHECKED(UPCGGraphInputOutputSettings, CustomPins) &&
+				PropertyChangedEvent.Property->GetFName() == GET_MEMBER_NAME_CHECKED(FPCGPinProperties, Label);
+		}
+
+		if (bLabelChanged)
+		{
+			FixCustomPinProperties();
+		}
+	}
+
+	Super::PostEditChangeProperty(PropertyChangedEvent);
+}
+#endif // WITH_EDITOR
+
+void UPCGGraphInputOutputSettings::FixCustomPinProperties()
+{
+	// No need to fix if we have no custom pins
+	if (CustomPins.IsEmpty())
+	{
+		return;
+	}
+
+	TSet<FName> AllLabels;
+
+	// First gather all our static labels
+	for (FLabelAndTooltip LabelAndTooltip : StaticLabels())
+	{
+		AllLabels.Emplace(LabelAndTooltip.Label);
+	}
+
+	for (FLabelAndTooltip LabelAndTooltip : StaticAdvancedLabels())
+	{
+		AllLabels.Emplace(LabelAndTooltip.Label);
+	}
+
+	bool bWasModified = false;
+
+	for (FPCGPinProperties& CustomPinProperties : CustomPins)
+	{
+		// Avoid "None" pin label
+		if (CustomPinProperties.Label == NAME_None)
+		{
+			if (!bWasModified)
+			{
+				bWasModified = true;
+				Modify();
+			}
+
+			CustomPinProperties.Label = PCGInputOutputConstants::DefaultNewCustomPinName;
+		}
+
+		uint32 Count = 1;
+		FString OriginalLabel = CustomPinProperties.Label.ToString();
+
+		while (AllLabels.Contains(CustomPinProperties.Label))
+		{
+			if (!bWasModified)
+			{
+				bWasModified = true;
+				Modify();
+			}
+
+			CustomPinProperties.Label = FName(FString::Printf(TEXT("%s%d"), *OriginalLabel, Count++));
+		}
+
+		AllLabels.Emplace(CustomPinProperties.Label);
+	}
+}
+
 #undef LOCTEXT_NAMESPACE
diff --git a/Engine/Plugins/Experimental/PCG/Source/PCG/Public/PCGInputOutputSettings.h b/Engine/Plugins/Experimental/PCG/Source/PCG/Public/PCGInputOutputSettings.h
--- a/Engine/Plugins/Experimental/PCG/Source/PCG/Public/PCGInputOutputSettings.h
+++ b/Engine/Plugins/Experimental/PCG/Source/PCG/Public/PCGInputOutputSettings.h
@@ -14,6 +14,7 @@ namespace PCGInputOutputConstants
 	const FName DefaultLandscapeLabel = TEXT("Landscape");
 	const FName DefaultLandscapeHeightLabel = TEXT("Landscape Height");
 	const FName DefaultExcludedActorsLabel = TEXT("ExcludedActors");
+	const FName DefaultNewCustomPinName = TEXT("NewPin");
 }
 
 class PCG_API FPCGInputOutputElement : public FSimplePCGElement
@@ -39,6 +40,7 @@ public:
 #if WITH_EDITOR
 	virtual FName GetDefaultNodeName() const override { return bIsInput ? FName(TEXT("InputNode")) : FName(TEXT("OutputNode")); }
 	virtual EPCGSettingsType GetType() const override { return EPCGSettingsType::InputOutput; }
+	virtual void PostEditChangeProperty(struct FPropertyChangedEvent& PropertyChangedEvent);
 #endif
 
 	TArray<FPCGPinProperties> InputPinProperties() const override;
@@ -51,7 +53,10 @@ public:
 
 	bool IsCustomPin(const UPCGPin* InPin) const;
 
-	void AddCustomPin(const FPCGPinProperties& NewCustomPinProperties);
+	// Add a new custom pin
+	// Note that you should use the return value of this function, since it can be different from
+	// the one passed as argument. It will change if its label collides with existing pins.
+	const FPCGPinProperties& AddCustomPin(const FPCGPinProperties& NewCustomPinProperties);
 
 protected:
 	virtual FPCGElementPtr CreateElement() const override { return MakeShared<FPCGInputOutputElement>(); }
@@ -73,6 +78,8 @@ protected:
 
 	TArray<FPCGPinProperties> GetPinProperties() const;
 
+	void FixCustomPinProperties();
+
 protected:
 	UPROPERTY()
 	TSet<FName> PinLabels_DEPRECATED;
diff --git a/Engine/Plugins/Experimental/PCG/Source/PCGEditor/Private/PCGEditor.cpp b/Engine/Plugins/Experimental/PCG/Source/PCGEditor/Private/PCGEditor.cpp
--- a/Engine/Plugins/Experimental/PCG/Source/PCGEditor/Private/PCGEditor.cpp
+++ b/Engine/Plugins/Experimental/PCG/Source/PCGEditor/Private/PCGEditor.cpp
@@ -834,7 +834,7 @@ void FPCGEditor::OnCollapseNodesInSubgraph()
 
 		FPCGPinProperties NewProperties = Pin->Properties;
 		NewProperties.Label = FName(NewName);
-		NewInputOutputSettings->AddCustomPin(NewProperties);
+		NewProperties = NewInputOutputSettings->AddCustomPin(NewProperties);
 		NewInputOutputNode->UpdateAfterSettingsChangeDuringCreation();
 		return NewProperties.Label;
 	};
