[PCG] Add UE_NODISCARD to AddCustomPin to enforce the use of the return value.

#rnx

diff --git a/Engine/Plugins/Experimental/PCG/Source/PCG/Public/PCGInputOutputSettings.h b/Engine/Plugins/Experimental/PCG/Source/PCG/Public/PCGInputOutputSettings.h
--- a/Engine/Plugins/Experimental/PCG/Source/PCG/Public/PCGInputOutputSettings.h
+++ b/Engine/Plugins/Experimental/PCG/Source/PCG/Public/PCGInputOutputSettings.h
@@ -56,7 +56,7 @@ public:
 	// Add a new custom pin
 	// Note that you should use the return value of this function, since it can be different from
 	// the one passed as argument. It will change if its label collides with existing pins.
-	const FPCGPinProperties& AddCustomPin(const FPCGPinProperties& NewCustomPinProperties);
+	UE_NODISCARD const FPCGPinProperties& AddCustomPin(const FPCGPinProperties& NewCustomPinProperties);
 
 protected:
 	virtual FPCGElementPtr CreateElement() const override { return MakeShared<FPCGInputOutputElement>(); }
