Add negative scale support in PCG isms

#rnx

diff --git a/Engine/Plugins/Experimental/PCG/Source/PCG/Private/Elements/PCGDebugElement.cpp b/Engine/Plugins/Experimental/PCG/Source/PCG/Private/Elements/PCGDebugElement.cpp
--- a/Engine/Plugins/Experimental/PCG/Source/PCG/Private/Elements/PCGDebugElement.cpp
+++ b/Engine/Plugins/Experimental/PCG/Source/PCG/Private/Elements/PCGDebugElement.cpp
@@ -114,14 +114,16 @@ namespace PCGDebugElement
 
 			const int NumCustomData = 8;
 
-			TArray<FTransform> Instances;
+			TArray<FTransform> ForwardInstances;
+			TArray<FTransform> ReverseInstances;
 			TArray<float> InstanceCustomData;
 
-			Instances.Reserve(Points.Num());
+			ForwardInstances.Reserve(Points.Num());
 			InstanceCustomData.Reserve(NumCustomData);
 
 			// First, create target instance transforms
 			const float PointScale = DebugSettings.PointScale;
+			const bool bIsAbsolute = DebugSettings.ScaleMethod == EPCGDebugVisScaleMethod::Absolute;
 			const bool bIsRelative = DebugSettings.ScaleMethod == EPCGDebugVisScaleMethod::Relative;
 			const bool bScaleWithExtents = DebugSettings.ScaleMethod == EPCGDebugVisScaleMethod::Extents;
 			const FVector MeshExtents = Mesh->GetBoundingBox().GetExtent();
@@ -129,6 +131,7 @@ namespace PCGDebugElement
 
 			for (const FPCGPoint& Point : Points)
 			{
+				TArray<FTransform>& Instances = ((bIsAbsolute || Point.Transform.GetDeterminant() >= 0) ? ForwardInstances : ReverseInstances);
 				FTransform& InstanceTransform = Instances.Add_GetRef(Point.Transform);
 				if (bIsRelative)
 				{
@@ -154,34 +157,45 @@ namespace PCGDebugElement
 			// Note: In the future we may consider enabling culling for performance reasons, but for now culling disabled.
 			Params.CullStartDistance = Params.CullEndDistance = 0;
 
-			UInstancedStaticMeshComponent* ISMC = UPCGActorHelpers::GetOrCreateISMC(TargetActor, Context->SourceComponent.Get(), Params);
-			check(ISMC);
-			
-			ISMC->ComponentTags.AddUnique(PCGHelpers::DefaultPCGDebugTag);
-			ISMC->NumCustomDataFloats = NumCustomData;
-			const int32 PreExistingInstanceCount = ISMC->GetInstanceCount();
-			ISMC->AddInstances(Instances, /*bShouldReturnIndices=*/false, /*bWorldSpace=*/true);
-
-			// Then get & assign custom data
-			for (int32 PointIndex = 0; PointIndex < Points.Num(); ++PointIndex)
+			for (int32 Direction = 0; Direction < 2; ++Direction)
 			{
-				const FPCGPoint& Point = Points[PointIndex];
-				InstanceCustomData.Add(Point.Density);
-				const FVector Extents = Point.GetExtents();
-				InstanceCustomData.Add(Extents[0]);
-				InstanceCustomData.Add(Extents[1]);
-				InstanceCustomData.Add(Extents[2]);
-				InstanceCustomData.Add(Point.Color[0]);
-				InstanceCustomData.Add(Point.Color[1]);
-				InstanceCustomData.Add(Point.Color[2]);
-				InstanceCustomData.Add(Point.Color[3]);
-
-				ISMC->SetCustomData(PointIndex + PreExistingInstanceCount, InstanceCustomData);
-
-				InstanceCustomData.Reset();
-			}
+				TArray<FTransform>& Instances = (Direction == 0 ? ForwardInstances : ReverseInstances);
+
+				if (Instances.IsEmpty())
+				{
+					continue;
+				}
+
+				Params.bIsLocalToWorldDeterminantNegative = (Direction != 0);
+				UInstancedStaticMeshComponent* ISMC = UPCGActorHelpers::GetOrCreateISMC(TargetActor, Context->SourceComponent.Get(), Params);
+				check(ISMC);
 
-			ISMC->UpdateBounds();
+				ISMC->ComponentTags.AddUnique(PCGHelpers::DefaultPCGDebugTag);
+				ISMC->NumCustomDataFloats = NumCustomData;
+				const int32 PreExistingInstanceCount = ISMC->GetInstanceCount();
+				ISMC->AddInstances(Instances, /*bShouldReturnIndices=*/false, /*bWorldSpace=*/true);
+
+				// Then get & assign custom data
+				for (int32 PointIndex = 0; PointIndex < Points.Num(); ++PointIndex)
+				{
+					const FPCGPoint& Point = Points[PointIndex];
+					InstanceCustomData.Add(Point.Density);
+					const FVector Extents = Point.GetExtents();
+					InstanceCustomData.Add(Extents[0]);
+					InstanceCustomData.Add(Extents[1]);
+					InstanceCustomData.Add(Extents[2]);
+					InstanceCustomData.Add(Point.Color[0]);
+					InstanceCustomData.Add(Point.Color[1]);
+					InstanceCustomData.Add(Point.Color[2]);
+					InstanceCustomData.Add(Point.Color[3]);
+
+					ISMC->SetCustomData(PointIndex + PreExistingInstanceCount, InstanceCustomData);
+
+					InstanceCustomData.Reset();
+				}
+
+				ISMC->UpdateBounds();
+			}
 		}
 #endif
 	}
diff --git a/Engine/Plugins/Experimental/PCG/Source/PCG/Private/Elements/PCGStaticMeshSpawner.cpp b/Engine/Plugins/Experimental/PCG/Source/PCG/Private/Elements/PCGStaticMeshSpawner.cpp
--- a/Engine/Plugins/Experimental/PCG/Source/PCG/Private/Elements/PCGStaticMeshSpawner.cpp
+++ b/Engine/Plugins/Experimental/PCG/Source/PCG/Private/Elements/PCGStaticMeshSpawner.cpp
@@ -199,10 +199,10 @@ void FPCGStaticMeshSpawnerElement::SpawnStaticMeshInstances(FPCGContext* Context
 		Params.MaterialOverrides = InstanceList.MaterialOverrides;
 	}
 
+	Params.NumCustomDataFloats = PackedCustomData.NumCustomDataFloats;
 	Params.CullStartDistance = InstanceList.CullStartDistance;
 	Params.CullEndDistance = InstanceList.CullEndDistance;
-
-	Params.NumCustomDataFloats = PackedCustomData.NumCustomDataFloats;
+	Params.bIsLocalToWorldDeterminantNegative = InstanceList.bIsLocalToWorldDeterminantNegative;
 
 	UInstancedStaticMeshComponent* ISMC = UPCGActorHelpers::GetOrCreateISMC(TargetActor, Context->SourceComponent.Get(), Params);
 
diff --git a/Engine/Plugins/Experimental/PCG/Source/PCG/Private/Helpers/PCGActorHelpers.cpp b/Engine/Plugins/Experimental/PCG/Source/PCG/Private/Helpers/PCGActorHelpers.cpp
--- a/Engine/Plugins/Experimental/PCG/Source/PCG/Private/Helpers/PCGActorHelpers.cpp
+++ b/Engine/Plugins/Experimental/PCG/Source/PCG/Private/Helpers/PCGActorHelpers.cpp
@@ -123,7 +123,8 @@ UInstancedStaticMeshComponent* UPCGActorHelpers::GetOrCreateISMC(AActor* InTarge
 				(ISMC->bUseDefaultCollision && InParams.CollisionProfile != TEXT("Default")) ||
 				(!ISMC->bUseDefaultCollision && ISMC->GetCollisionProfileName() != InParams.CollisionProfile) ||
 				ISMC->InstanceStartCullDistance != InParams.CullStartDistance ||
-				ISMC->InstanceEndCullDistance != InParams.CullEndDistance)
+				ISMC->InstanceEndCullDistance != InParams.CullEndDistance ||
+				ISMC->bReverseCulling != InParams.bIsLocalToWorldDeterminantNegative)
 			{
 				continue;
 			}
@@ -202,6 +203,7 @@ UInstancedStaticMeshComponent* UPCGActorHelpers::GetOrCreateISMC(AActor* InTarge
 
 	ISMC->InstanceStartCullDistance = InParams.CullStartDistance;
 	ISMC->InstanceEndCullDistance = InParams.CullEndDistance;
+	ISMC->bReverseCulling = (InParams.bIsLocalToWorldDeterminantNegative ? 1 : 0);
 
 	// Implementation note: we use the relative for position only to prevent issues at the rendering level, but world for scale & rotation otherwise we run into issues when we have non-uniform scales esp. in volumes
 	ISMC->AttachToComponent(InTargetActor->GetRootComponent(), FAttachmentTransformRules(EAttachmentRule::KeepRelative, EAttachmentRule::KeepWorld, EAttachmentRule::KeepWorld, false));
diff --git a/Engine/Plugins/Experimental/PCG/Source/PCG/Private/MeshSelectors/PCGMeshSelectorBase.cpp b/Engine/Plugins/Experimental/PCG/Source/PCG/Private/MeshSelectors/PCGMeshSelectorBase.cpp
--- a/Engine/Plugins/Experimental/PCG/Source/PCG/Private/MeshSelectors/PCGMeshSelectorBase.cpp
+++ b/Engine/Plugins/Experimental/PCG/Source/PCG/Private/MeshSelectors/PCGMeshSelectorBase.cpp
@@ -13,7 +13,8 @@ int32 UPCGMeshSelectorBase::FindOrAddInstanceList(
 	bool bOverrideMaterials,
 	const TArray<UMaterialInterface*>& MaterialOverrides,
 	const float InCullStartDistance,
-	const float InCullEndDistance) const
+	const float InCullEndDistance, 
+	const bool bInIsLocalToWorldDeterminantNegative) const
 {
 	for (int Index = 0; Index < OutInstanceLists.Num(); ++Index)
 	{
@@ -27,6 +28,11 @@ int32 UPCGMeshSelectorBase::FindOrAddInstanceList(
 			continue;
 		}
 
+		if (OutInstanceLists[Index].bIsLocalToWorldDeterminantNegative != bInIsLocalToWorldDeterminantNegative)
+		{
+			continue;
+		}
+
 		if (OutInstanceLists[Index].bOverrideCollisionProfile && !OutInstanceLists[Index].CollisionProfile.Name.IsEqual(CollisionProfile.Name))
 		{
 			continue;
@@ -40,6 +46,6 @@ int32 UPCGMeshSelectorBase::FindOrAddInstanceList(
 		return Index;
 	}
 
-	return OutInstanceLists.Emplace(Mesh, bOverrideCollisionProfile, CollisionProfile, bOverrideMaterials, MaterialOverrides, InCullStartDistance, InCullEndDistance);
+	return OutInstanceLists.Emplace(Mesh, bOverrideCollisionProfile, CollisionProfile, bOverrideMaterials, MaterialOverrides, InCullStartDistance, InCullEndDistance, bInIsLocalToWorldDeterminantNegative);
 }
 
diff --git a/Engine/Plugins/Experimental/PCG/Source/PCG/Private/MeshSelectors/PCGMeshSelectorByAttribute.cpp b/Engine/Plugins/Experimental/PCG/Source/PCG/Private/MeshSelectors/PCGMeshSelectorByAttribute.cpp
--- a/Engine/Plugins/Experimental/PCG/Source/PCG/Private/MeshSelectors/PCGMeshSelectorByAttribute.cpp
+++ b/Engine/Plugins/Experimental/PCG/Source/PCG/Private/MeshSelectors/PCGMeshSelectorByAttribute.cpp
@@ -90,7 +90,8 @@ void UPCGMeshSelectorByAttribute::SelectInstances_Implementation(
 			continue;
 		}
 
-		const int32 Index = FindOrAddInstanceList(OutMeshInstances, Mesh, bOverrideCollisionProfile, CollisionProfile, bOverrideMaterials, MaterialOverrides, CullStartDistance, CullEndDistance);
+		const bool bReverseTransform = (Point.Transform.GetDeterminant() < 0);
+		const int32 Index = FindOrAddInstanceList(OutMeshInstances, Mesh, bOverrideCollisionProfile, CollisionProfile, bOverrideMaterials, MaterialOverrides, CullStartDistance, CullEndDistance, bReverseTransform);
 		check(Index != INDEX_NONE);
 		
 		OutMeshInstances[Index].Instances.Emplace(Point);
diff --git a/Engine/Plugins/Experimental/PCG/Source/PCG/Private/MeshSelectors/PCGMeshSelectorWeighted.cpp b/Engine/Plugins/Experimental/PCG/Source/PCG/Private/MeshSelectors/PCGMeshSelectorWeighted.cpp
--- a/Engine/Plugins/Experimental/PCG/Source/PCG/Private/MeshSelectors/PCGMeshSelectorWeighted.cpp
+++ b/Engine/Plugins/Experimental/PCG/Source/PCG/Private/MeshSelectors/PCGMeshSelectorWeighted.cpp
@@ -19,6 +19,7 @@ void UPCGMeshSelectorWeighted::SelectInstances_Implementation(
 	TArray<FPCGMeshInstanceList>& OutMeshInstances,
 	UPCGPointData* OutPointData) const
 {
+	TArray<FPCGMeshInstanceList> OutReverseMeshInstances;
 	TArray<int> CumulativeWeights;
 
 	int TotalWeight = 0;
@@ -31,7 +32,8 @@ void UPCGMeshSelectorWeighted::SelectInstances_Implementation(
 			continue;
 		}
 
-		FindOrAddInstanceList(OutMeshInstances, Entry.Mesh, Entry.bOverrideCollisionProfile, Entry.CollisionProfile, Entry.bOverrideMaterials, Entry.MaterialOverrides, Entry.CullStartDistance, Entry.CullEndDistance);
+		FindOrAddInstanceList(OutMeshInstances, Entry.Mesh, Entry.bOverrideCollisionProfile, Entry.CollisionProfile, Entry.bOverrideMaterials, Entry.MaterialOverrides, Entry.CullStartDistance, Entry.CullEndDistance, /*bReverseCulling=*/false);
+		FindOrAddInstanceList(OutReverseMeshInstances, Entry.Mesh, Entry.bOverrideCollisionProfile, Entry.CollisionProfile, Entry.bOverrideMaterials, Entry.MaterialOverrides, Entry.CullStartDistance, Entry.CullEndDistance, /*bReverseCulling=*/true);
 		TotalWeight += Entry.Weight;
 		CumulativeWeights.Add(TotalWeight);
 	}
@@ -100,7 +102,8 @@ void UPCGMeshSelectorWeighted::SelectInstances_Implementation(
 
 			if (RandomPick < OutMeshInstances.Num())
 			{
-				FPCGMeshInstanceList& InstanceList = OutMeshInstances[RandomPick];
+				const bool bNeedsReverseCulling = (Point.Transform.GetDeterminant() < 0);
+				FPCGMeshInstanceList& InstanceList = (bNeedsReverseCulling ? OutReverseMeshInstances[RandomPick] : OutMeshInstances[RandomPick]);
 				InstanceList.Instances.Emplace(Point);
 
 				const TSoftObjectPtr<UStaticMesh>& Mesh = InstanceList.Mesh;
@@ -123,5 +126,11 @@ void UPCGMeshSelectorWeighted::SelectInstances_Implementation(
 			}
 		}
 	}
+
+	// Collapse reverse entries to the OutMeshInstances
+	for (int32 OutIndex = 0; OutIndex < OutReverseMeshInstances.Num(); ++OutIndex)
+	{
+		OutMeshInstances.Emplace(MoveTemp(OutReverseMeshInstances[OutIndex]));
+	}
 }
 
diff --git a/Engine/Plugins/Experimental/PCG/Source/PCG/Private/MeshSelectors/PCGMeshSelectorWeightedByCategory.cpp b/Engine/Plugins/Experimental/PCG/Source/PCG/Private/MeshSelectors/PCGMeshSelectorWeightedByCategory.cpp
--- a/Engine/Plugins/Experimental/PCG/Source/PCG/Private/MeshSelectors/PCGMeshSelectorWeightedByCategory.cpp
+++ b/Engine/Plugins/Experimental/PCG/Source/PCG/Private/MeshSelectors/PCGMeshSelectorWeightedByCategory.cpp
@@ -15,6 +15,7 @@
 struct FPCGInstancesAndWeights
 {
 	TArray<int> InstanceListIndices;
+	TArray<int> ReverseInstanceListIndices;
 	TArray<int> CumulativeWeights;
 };
 
@@ -25,6 +26,7 @@ void UPCGMeshSelectorWeightedByCategory::SelectInstances_Implementation(
 	TArray<FPCGMeshInstanceList>& OutMeshInstances,
 	UPCGPointData* OutPointData) const
 {
+	TArray<FPCGMeshInstanceList> OutReverseMeshInstances;
 	const UPCGPointData* PointData = InSpatialData->ToPointData(&Context);
 
 	if (!PointData)
@@ -110,13 +112,15 @@ void UPCGMeshSelectorWeightedByCategory::SelectInstances_Implementation(
 				continue;
 			}
 			
-			const int32 Index = FindOrAddInstanceList(OutMeshInstances, WeightedEntry.Mesh, WeightedEntry.bOverrideCollisionProfile, WeightedEntry.CollisionProfile, WeightedEntry.bOverrideMaterials, WeightedEntry.MaterialOverrides, WeightedEntry.CullStartDistance, WeightedEntry.CullEndDistance);
-			check(Index != INDEX_NONE);
+			const int32 Index = FindOrAddInstanceList(OutMeshInstances, WeightedEntry.Mesh, WeightedEntry.bOverrideCollisionProfile, WeightedEntry.CollisionProfile, WeightedEntry.bOverrideMaterials, WeightedEntry.MaterialOverrides, WeightedEntry.CullStartDistance, WeightedEntry.CullEndDistance, /*bReverseCulling=*/false);
+			const int32 ReverseIndex = FindOrAddInstanceList(OutReverseMeshInstances, WeightedEntry.Mesh, WeightedEntry.bOverrideCollisionProfile, WeightedEntry.CollisionProfile, WeightedEntry.bOverrideMaterials, WeightedEntry.MaterialOverrides, WeightedEntry.CullStartDistance, WeightedEntry.CullEndDistance, /*bReverseCulling=*/true);
+			check(Index != INDEX_NONE && ReverseIndex != INDEX_NONE);
 
 			// precompute the weights
 			TotalWeight += WeightedEntry.Weight;
 			InstancesAndWeights->CumulativeWeights.Add(TotalWeight);
 			InstancesAndWeights->InstanceListIndices.Add(Index);
+			InstancesAndWeights->ReverseInstanceListIndices.Add(ReverseIndex);
 		}
 
 		// if no weighted entries were collected, discard this InstancesAndWeights
@@ -195,8 +199,8 @@ void UPCGMeshSelectorWeightedByCategory::SelectInstances_Implementation(
 
 		if (RandomPick < InstancesAndWeights->InstanceListIndices.Num())
 		{
-			const int32 Index = InstancesAndWeights->InstanceListIndices[RandomPick];
-			FPCGMeshInstanceList& InstanceList = OutMeshInstances[Index];
+			const bool bUseReverse = (Point.Transform.GetDeterminant() < 0);
+			FPCGMeshInstanceList& InstanceList = (bUseReverse ? OutReverseMeshInstances[InstancesAndWeights->ReverseInstanceListIndices[RandomPick]] : OutMeshInstances[InstancesAndWeights->InstanceListIndices[RandomPick]]);
 			InstanceList.Instances.Emplace(Point);
 
 			const TSoftObjectPtr<UStaticMesh>& Mesh = InstanceList.Mesh;
@@ -218,5 +222,10 @@ void UPCGMeshSelectorWeightedByCategory::SelectInstances_Implementation(
 			}
 		}
 	}
-}
 
+	// Collapse reverse in the out mesh instances array
+	for (int32 Index = 0; Index < OutReverseMeshInstances.Num(); ++Index)
+	{
+		OutMeshInstances.Emplace(MoveTemp(OutReverseMeshInstances[Index]));
+	}
+}
diff --git a/Engine/Plugins/Experimental/PCG/Source/PCG/Public/Helpers/PCGActorHelpers.h b/Engine/Plugins/Experimental/PCG/Source/PCG/Public/Helpers/PCGActorHelpers.h
--- a/Engine/Plugins/Experimental/PCG/Source/PCG/Public/Helpers/PCGActorHelpers.h
+++ b/Engine/Plugins/Experimental/PCG/Source/PCG/Public/Helpers/PCGActorHelpers.h
@@ -28,6 +28,7 @@ struct FPCGISMCBuilderParameters
 	int32 NumCustomDataFloats = 0;
 	float CullStartDistance = 0;
 	float CullEndDistance = 0;
+	bool bIsLocalToWorldDeterminantNegative = false;
 };
 
 UCLASS(BlueprintType)
diff --git a/Engine/Plugins/Experimental/PCG/Source/PCG/Public/MeshSelectors/PCGMeshSelectorBase.h b/Engine/Plugins/Experimental/PCG/Source/PCG/Public/MeshSelectors/PCGMeshSelectorBase.h
--- a/Engine/Plugins/Experimental/PCG/Source/PCG/Public/MeshSelectors/PCGMeshSelectorBase.h
+++ b/Engine/Plugins/Experimental/PCG/Source/PCG/Public/MeshSelectors/PCGMeshSelectorBase.h
@@ -23,8 +23,8 @@ struct FPCGMeshInstanceList
 
 	FPCGMeshInstanceList() = default;
 	
-	FPCGMeshInstanceList(const TSoftObjectPtr<UStaticMesh>& InMesh, bool bInOverrideCollisionProfile, const FCollisionProfileName& InCollisionProfile, bool bInOverrideMaterials, const TArray<UMaterialInterface*>& InMaterialOverrides, const float InCullStartDistance, const float InCullEndDistance)
-		: Mesh(InMesh), bOverrideCollisionProfile(bInOverrideCollisionProfile), CollisionProfile(InCollisionProfile), bOverrideMaterials(bInOverrideMaterials), MaterialOverrides(InMaterialOverrides), CullStartDistance(InCullStartDistance), CullEndDistance(InCullEndDistance)
+	FPCGMeshInstanceList(const TSoftObjectPtr<UStaticMesh>& InMesh, bool bInOverrideCollisionProfile, const FCollisionProfileName& InCollisionProfile, bool bInOverrideMaterials, const TArray<UMaterialInterface*>& InMaterialOverrides, const float InCullStartDistance, const float InCullEndDistance, const bool bInIsLocalToWorldDeterminantNegative)
+		: Mesh(InMesh), bOverrideCollisionProfile(bInOverrideCollisionProfile), CollisionProfile(InCollisionProfile), bOverrideMaterials(bInOverrideMaterials), MaterialOverrides(InMaterialOverrides), CullStartDistance(InCullStartDistance), CullEndDistance(InCullEndDistance), bIsLocalToWorldDeterminantNegative(bInIsLocalToWorldDeterminantNegative)
 	{}
 
 	UPROPERTY(BlueprintReadWrite, EditAnywhere, Category = Settings)
@@ -52,6 +52,10 @@ struct FPCGMeshInstanceList
 	/** Distance at which instances are culled. Use 0 to disable. */
 	UPROPERTY(BlueprintReadWrite, EditAnywhere, Category = Settings)
 	float CullEndDistance = 0;
+
+	/** Whether the culling should be reversed or not (needed to support negative scales) */
+	UPROPERTY(BlueprintReadWrite, EditAnywhere, Category = Settings)
+	bool bIsLocalToWorldDeterminantNegative = false;
 };
 
 UCLASS(Abstract, BlueprintType, Blueprintable, ClassGroup = (Procedural))
@@ -85,7 +89,8 @@ public:
 		bool bOverrideMaterials,
 		const TArray<UMaterialInterface*>& MaterialOverrides,
 		const float InCullStartDistance,
-		const float InCullEndDistance) const;
+		const float InCullEndDistance, 
+		const bool bInIsLocalToWorldDeterminantNegative) const;
 };
 
 #if UE_ENABLE_INCLUDE_ORDER_DEPRECATED_IN_5_2
