[PCG] Initialized collision profiles to NoCollision to avoid an extra command being done when adding elements to the mesh selector array.

diff --git a/Engine/Plugins/Experimental/PCG/Source/PCG/Public/MeshSelectors/PCGMeshSelectorBase.h b/Engine/Plugins/Experimental/PCG/Source/PCG/Public/MeshSelectors/PCGMeshSelectorBase.h
--- a/Engine/Plugins/Experimental/PCG/Source/PCG/Public/MeshSelectors/PCGMeshSelectorBase.h
+++ b/Engine/Plugins/Experimental/PCG/Source/PCG/Public/MeshSelectors/PCGMeshSelectorBase.h
@@ -34,7 +34,7 @@ struct FPCGMeshInstanceList
 	bool bOverrideCollisionProfile = false;
 
 	UPROPERTY(BlueprintReadWrite, EditAnywhere, Category = Settings)
-	FCollisionProfileName CollisionProfile;
+	FCollisionProfileName CollisionProfile = UCollisionProfile::NoCollision_ProfileName;
 
 	UPROPERTY(BlueprintReadWrite, EditAnywhere, Category = Settings)
 	bool bOverrideMaterials = false;
diff --git a/Engine/Plugins/Experimental/PCG/Source/PCG/Public/MeshSelectors/PCGMeshSelectorByAttribute.h b/Engine/Plugins/Experimental/PCG/Source/PCG/Public/MeshSelectors/PCGMeshSelectorByAttribute.h
--- a/Engine/Plugins/Experimental/PCG/Source/PCG/Public/MeshSelectors/PCGMeshSelectorByAttribute.h
+++ b/Engine/Plugins/Experimental/PCG/Source/PCG/Public/MeshSelectors/PCGMeshSelectorByAttribute.h
@@ -4,6 +4,8 @@
 
 #include "PCGMeshSelectorBase.h"
 
+#include "Engine/CollisionProfile.h"
+
 #include "PCGMeshSelectorByAttribute.generated.h"
 
 UCLASS(BlueprintType, ClassGroup = (Procedural))
@@ -27,7 +29,7 @@ public:
 	bool bOverrideCollisionProfile = false;
 
 	UPROPERTY(BlueprintReadWrite, EditAnywhere, Category = Settings)
-	FCollisionProfileName CollisionProfile;
+	FCollisionProfileName CollisionProfile = UCollisionProfile::NoCollision_ProfileName;
 
 	UPROPERTY(BlueprintReadWrite, EditAnywhere, Category = Settings)
 	bool bOverrideMaterials = false;
@@ -43,4 +45,3 @@ public:
 	UPROPERTY(BlueprintReadWrite, EditAnywhere, Category = Settings)
 	float CullEndDistance = 0;
 };
-
diff --git a/Engine/Plugins/Experimental/PCG/Source/PCG/Public/MeshSelectors/PCGMeshSelectorWeighted.h b/Engine/Plugins/Experimental/PCG/Source/PCG/Public/MeshSelectors/PCGMeshSelectorWeighted.h
--- a/Engine/Plugins/Experimental/PCG/Source/PCG/Public/MeshSelectors/PCGMeshSelectorWeighted.h
+++ b/Engine/Plugins/Experimental/PCG/Source/PCG/Public/MeshSelectors/PCGMeshSelectorWeighted.h
@@ -4,6 +4,8 @@
 
 #include "PCGMeshSelectorBase.h"
 
+#include "Engine/CollisionProfile.h"
+
 #include "PCGMeshSelectorWeighted.generated.h"
 
 USTRUCT(BlueprintType)
@@ -24,7 +26,7 @@ struct PCG_API FPCGMeshSelectorWeightedEntry
 	bool bOverrideCollisionProfile = false;
 
 	UPROPERTY(BlueprintReadWrite, EditAnywhere, Category = Settings)
-	FCollisionProfileName CollisionProfile;
+	FCollisionProfileName CollisionProfile = UCollisionProfile::NoCollision_ProfileName;
 
 	UPROPERTY(BlueprintReadWrite, EditAnywhere, Category = Settings)
 	bool bOverrideMaterials = false;
@@ -61,4 +63,3 @@ public:
 	UPROPERTY(BlueprintReadWrite, EditAnywhere, Category = Settings)
 	TArray<FPCGMeshSelectorWeightedEntry> MeshEntries;
 };
-
