[PCG]
* Re-removed include from header and changed virtual method to be implemented in cpp file.

diff --git a/Engine/Plugins/Experimental/PCG/Source/PCG/Private/PCGSettings.cpp b/Engine/Plugins/Experimental/PCG/Source/PCG/Private/PCGSettings.cpp
--- a/Engine/Plugins/Experimental/PCG/Source/PCG/Private/PCGSettings.cpp
+++ b/Engine/Plugins/Experimental/PCG/Source/PCG/Private/PCGSettings.cpp
@@ -160,6 +160,18 @@ TArray<FPCGPinProperties> UPCGSettings::OutputPinProperties() const
 	return PinProperties;
 }
 
+TArray<FPCGPinProperties>
+UPCGSettings::DefaultInputPinProperties() const
+{
+	return InputPinProperties();
+}
+
+TArray<FPCGPinProperties>
+UPCGSettings::DefaultOutputPinProperties() const
+{
+	return OutputPinProperties();
+}
+
 FPCGElementPtr UPCGSettings::GetElement() const
 {
 	if (!CachedElement)
diff --git a/Engine/Plugins/Experimental/PCG/Source/PCG/Public/PCGSettings.h b/Engine/Plugins/Experimental/PCG/Source/PCG/Public/PCGSettings.h
--- a/Engine/Plugins/Experimental/PCG/Source/PCG/Public/PCGSettings.h
+++ b/Engine/Plugins/Experimental/PCG/Source/PCG/Public/PCGSettings.h
@@ -5,12 +5,12 @@
 #include "PCGData.h"
 #include "PCGDebug.h"
 #include "PCGElement.h"
-#include "PCGPin.h"
 #include "Tests/Determinism/PCGDeterminismSettings.h"
 
 #include "PCGSettings.generated.h"
 
 class UPCGPin;
+struct FPCGPinProperties;
 struct FPropertyChangedEvent;
 
 class UPCGGraph;
@@ -113,8 +113,8 @@ public:
 
 	// Internal functions, should not be used by any user.
 	// Return a different subset for for input/output pin properties, in case of a default object.
-	virtual TArray<FPCGPinProperties> DefaultInputPinProperties() const { return InputPinProperties(); }
-	virtual TArray<FPCGPinProperties> DefaultOutputPinProperties() const { return OutputPinProperties(); }
+	virtual TArray<FPCGPinProperties> DefaultInputPinProperties() const;
+	virtual TArray<FPCGPinProperties> DefaultOutputPinProperties() const;
 
 	bool operator==(const UPCGSettings& Other) const;
 	uint32 GetCrc32() const;
