Added max depth in UpdateTrackedActor - was already done this way in the same file

#rnx

diff --git a/Engine/Plugins/Experimental/PCG/Source/PCG/Private/PCGComponent.cpp b/Engine/Plugins/Experimental/PCG/Source/PCG/Private/PCGComponent.cpp
--- a/Engine/Plugins/Experimental/PCG/Source/PCG/Private/PCGComponent.cpp
+++ b/Engine/Plugins/Experimental/PCG/Source/PCG/Private/PCGComponent.cpp
@@ -2223,7 +2223,7 @@ bool UPCGComponent::UpdateTrackedActor(AActor* InActor)
 		if (!CachedTrackedActorToTags.FindOrAdd(InActor).Find(Tag))
 		{
 			CachedTrackedActorToTags[InActor].Add(Tag);
-			PCGHelpers::GatherDependencies(InActor, CachedTrackedActorToDependencies.FindOrAdd(InActor));
+			PCGHelpers::GatherDependencies(InActor, CachedTrackedActorToDependencies.FindOrAdd(InActor), 1);
 			DirtyCacheFromTag(Tag);
 			bAppliedChange = true;
 		}
