[PCG] Ignore bAllowMultipleConnection for output pins.

#rnx

diff --git a/Engine/Plugins/Experimental/PCG/Source/PCG/Private/PCGGraph.cpp b/Engine/Plugins/Experimental/PCG/Source/PCG/Private/PCGGraph.cpp
--- a/Engine/Plugins/Experimental/PCG/Source/PCG/Private/PCGGraph.cpp
+++ b/Engine/Plugins/Experimental/PCG/Source/PCG/Private/PCGGraph.cpp
@@ -330,7 +330,7 @@ bool UPCGGraph::AddLabeledEdge(UPCGNode* From, const FName& FromPinLabel, UPCGNo
 	bool bToPinBrokeOtherEdges = false;
 	
 	// Add an edge to a pin that doesn't allow multiple connections requires to do some cleanup
-	if (!ToPin->Properties.bAllowMultipleConnections)
+	if (!ToPin->AllowMultipleConnections())
 	{
 		bToPinBrokeOtherEdges = ToPin->BreakAllIncompatibleEdges();
 	}
diff --git a/Engine/Plugins/Experimental/PCG/Source/PCG/Private/PCGPin.cpp b/Engine/Plugins/Experimental/PCG/Source/PCG/Private/PCGPin.cpp
--- a/Engine/Plugins/Experimental/PCG/Source/PCG/Private/PCGPin.cpp
+++ b/Engine/Plugins/Experimental/PCG/Source/PCG/Private/PCGPin.cpp
@@ -149,7 +149,7 @@ bool UPCGPin::BreakAllIncompatibleEdges()
 		UPCGEdge* Edge = Edges[EdgeIndex];
 		UPCGPin* OtherPin = Edge->GetOtherPin(this);
 
-		bool bRemoveEdge = !IsCompatible(OtherPin) || (!Properties.bAllowMultipleConnections && bHasAValidEdge);
+		bool bRemoveEdge = !IsCompatible(OtherPin) || (!AllowMultipleConnections() && bHasAValidEdge);
 
 		if (bRemoveEdge)
 		{
@@ -244,7 +244,13 @@ bool UPCGPin::IsCompatible(const UPCGPin* OtherPin) const
 	return !!(Properties.AllowedTypes & OtherPin->Properties.AllowedTypes);
 }
 
+bool UPCGPin::AllowMultipleConnections() const
+{
+	// Always allow multiple connection on output pin
+	return IsOutputPin() || Properties.bAllowMultipleConnections;
+}
+
 bool UPCGPin::CanConnect(const UPCGPin* OtherPin) const
 {
-	return OtherPin && (Properties.bAllowMultipleConnections || Edges.IsEmpty());
+	return OtherPin && (Edges.IsEmpty() || AllowMultipleConnections());
 }
diff --git a/Engine/Plugins/Experimental/PCG/Source/PCG/Public/PCGPin.h b/Engine/Plugins/Experimental/PCG/Source/PCG/Public/PCGPin.h
--- a/Engine/Plugins/Experimental/PCG/Source/PCG/Public/PCGPin.h
+++ b/Engine/Plugins/Experimental/PCG/Source/PCG/Public/PCGPin.h
@@ -71,6 +71,7 @@ public:
 	UFUNCTION(BlueprintCallable, Category = Settings)
 	void SetTooltip(const FText& InTooltip);
 
+	bool AllowMultipleConnections() const;
 	bool IsCompatible(const UPCGPin* OtherPin) const;
 	bool CanConnect(const UPCGPin* OtherPin) const;
 	bool AddEdgeTo(UPCGPin* OtherPin);
diff --git a/Engine/Plugins/Experimental/PCG/Source/PCGEditor/Private/SPCGEditorGraphNode.cpp b/Engine/Plugins/Experimental/PCG/Source/PCGEditor/Private/SPCGEditorGraphNode.cpp
--- a/Engine/Plugins/Experimental/PCG/Source/PCGEditor/Private/SPCGEditorGraphNode.cpp
+++ b/Engine/Plugins/Experimental/PCG/Source/PCGEditor/Private/SPCGEditorGraphNode.cpp
@@ -304,7 +304,7 @@ void SPCGEditorGraphNode::AddPin(const TSharedRef<SGraphPin>& PinToAdd)
 		if(UPCGPin* Pin = (bIsInPin ? PCGNode->GetInputPin(PinName) : PCGNode->GetOutputPin(PinName)))
 		{
 			const bool bIsMultiData = Pin->Properties.bAllowMultipleData;
-			const bool bIsMultiConnections = Pin->Properties.bAllowMultipleConnections;
+			const bool bIsMultiConnections = Pin->AllowMultipleConnections();
 
 			// Check for special types
 			if (Pin->Properties.AllowedTypes == EPCGDataType::Param)
