Fixed issue in the TransformPoints to attribute

#rnx

diff --git a/Engine/Plugins/Experimental/PCG/Source/PCG/Private/Elements/PCGTransformPoints.cpp b/Engine/Plugins/Experimental/PCG/Source/PCG/Private/Elements/PCGTransformPoints.cpp
--- a/Engine/Plugins/Experimental/PCG/Source/PCG/Private/Elements/PCGTransformPoints.cpp
+++ b/Engine/Plugins/Experimental/PCG/Source/PCG/Private/Elements/PCGTransformPoints.cpp
@@ -170,7 +170,7 @@ bool FPCGTransformPointsElement::ExecuteInternal(FPCGContext* Context) const
 			}
 			else
 			{
-				SourceTransform = SourceAttribute->GetValue(InPoint.MetadataEntry);
+				SourceTransform = SourceAttribute->GetValueFromItemKey(InPoint.MetadataEntry);
 			}
 
 			FTransform FinalTransform = SourceTransform;
