[PCG] Workaround for missing points when running a surface sampler on a volume input. Points are placed at max extent of volume, but PCGHelpers::IsInsideBounds() does a one-sided check (< MaxZ) and the points fail this test.

For now drop the points by epsilon to workaround. Points touching the sides (+X, +Y) of the volume will still be culled but that is not specific to this node and is left to potential future work.

The one-sided check is done to ensure overlapping PCG volumes will only generate one set of points on touching faces. Perhaps this can be fixed another way in the future, either by pulling in the component bounds by epsilon, or by optionally culling output points by the one-sided volume. This would isolate this behaviour to this aspect of the system rather than touching everything that uses the general purpose PCGHelpers::IsInsideBounds().

diff --git a/Engine/Plugins/Experimental/PCG/Source/PCG/Private/Elements/PCGSurfaceSampler.cpp b/Engine/Plugins/Experimental/PCG/Source/PCG/Private/Elements/PCGSurfaceSampler.cpp
--- a/Engine/Plugins/Experimental/PCG/Source/PCG/Private/Elements/PCGSurfaceSampler.cpp
+++ b/Engine/Plugins/Experimental/PCG/Source/PCG/Private/Elements/PCGSurfaceSampler.cpp
@@ -98,6 +98,7 @@ namespace PCGSurfaceSampler
 
 		Ratio = TargetPointCount / (FVector::FReal)CellCount;
 
+		InputBoundsMinZ = InputBounds.Min.Z;
 		InputBoundsMaxZ = InputBounds.Max.Z;
 
 		return true;
@@ -129,7 +130,17 @@ namespace PCGSurfaceSampler
 
 		FPCGProjectionParams ProjectionParams{};
 
-		FPCGAsync::AsyncPointProcessing(Context, LoopData.CellCount, SampledPoints, [&LoopData, SampledData, InBoundingShape, InSurface, &ProjectionParams](int32 Index, FPCGPoint& OutPoint)
+		// Drop points slightly by an epsilon otherwise point can be culled. If the sampler has a volume connected as the Bounding Shape,
+		// the volume will call through to PCGHelpers::IsInsideBounds() which is a one sided test and points at the top of the volume
+		// will fail it. TODO perhaps the one-sided check can be isolated to component-bounds
+		const FVector::FReal ZMultiplier = 1.0 - UE_DOUBLE_SMALL_NUMBER;
+		// Try to use a multiplier instead of a simply offset to combat loss of precision in floats. However if MaxZ is very small,
+		// then multiplier will not work, so just use an offset.
+		FVector::FReal SampleZ = (FMath::Abs(LoopData.InputBoundsMaxZ) > UE_DOUBLE_SMALL_NUMBER) ? LoopData.InputBoundsMaxZ * ZMultiplier : -UE_DOUBLE_SMALL_NUMBER;
+		// Make sure we're still in bounds though!
+		SampleZ = FMath::Max(SampleZ, LoopData.InputBoundsMinZ);
+
+		FPCGAsync::AsyncPointProcessing(Context, LoopData.CellCount, SampledPoints, [&LoopData, SampledData, InBoundingShape, InSurface, &ProjectionParams, SampleZ](int32 Index, FPCGPoint& OutPoint)
 		{
 			const FIntVector2 Indices = LoopData.ComputeCellIndices(Index);
 
@@ -149,8 +160,8 @@ namespace PCGSurfaceSampler
 
 			const float RandX = RandomSource.FRand();
 			const float RandY = RandomSource.FRand();
-			
-			const FVector TentativeLocation = FVector(CurrentX + RandX * InnerCellSize.X, CurrentY + RandY * InnerCellSize.Y, LoopData.InputBoundsMaxZ);
+
+			const FVector TentativeLocation = FVector(CurrentX + RandX * InnerCellSize.X, CurrentY + RandY * InnerCellSize.Y, SampleZ);
 			const FBox LocalBound(-LoopData.PointExtents, LoopData.PointExtents);
 
 			// Firstly project onto elected generating shape to move to final position.
diff --git a/Engine/Plugins/Experimental/PCG/Source/PCG/Public/Elements/PCGSurfaceSampler.h b/Engine/Plugins/Experimental/PCG/Source/PCG/Public/Elements/PCGSurfaceSampler.h
--- a/Engine/Plugins/Experimental/PCG/Source/PCG/Public/Elements/PCGSurfaceSampler.h
+++ b/Engine/Plugins/Experimental/PCG/Source/PCG/Public/Elements/PCGSurfaceSampler.h
@@ -114,6 +114,7 @@ namespace PCGSurfaceSampler
 		float Ratio;
 		int Seed;
 
+		FVector::FReal InputBoundsMinZ;
 		FVector::FReal InputBoundsMaxZ;
 	};
 
