Submitted on behalf of thomas.tedemalm : [PCG] Fixed an issue when ctrl+draging connections caused by reconstruct node not copying peristent data between pins.

#rnx

diff --git a/Engine/Plugins/Experimental/PCG/Source/PCGEditor/Private/PCGEditorGraphNodeBase.cpp b/Engine/Plugins/Experimental/PCG/Source/PCGEditor/Private/PCGEditorGraphNodeBase.cpp
--- a/Engine/Plugins/Experimental/PCG/Source/PCGEditor/Private/PCGEditorGraphNodeBase.cpp
+++ b/Engine/Plugins/Experimental/PCG/Source/PCGEditor/Private/PCGEditorGraphNodeBase.cpp
@@ -285,28 +285,29 @@ void UPCGEditorGraphNodeBase::ReconstructNode()
 		return;
 	}
 
-	// Remove all current pins
-	TArray<UEdGraphPin*> OldPins = Pins;
+	// Store copy of old pins
+	TArray<UEdGraphPin*> OldPins = MoveTemp(Pins);
+	Pins.Reset();
 
+	// Generate new pins
+	AllocateDefaultPins();
+
+	// Transfer persistent data from old to new pins
 	for (UEdGraphPin* OldPin : OldPins)
 	{
-		OldPin->BreakAllPinLinks();
-		RemovePin(OldPin);
+		const FName& OldPinName = OldPin->PinName;
+		if (UEdGraphPin** NewPin = Pins.FindByPredicate([&OldPinName](UEdGraphPin* InPin) { return InPin->PinName == OldPinName; }))
+		{
+			(*NewPin)->MovePersistentDataFromOldPin(*OldPin);
+		}
 	}
-	check(Pins.IsEmpty());
 
-	// Generate new pins
-	AllocateDefaultPins();
-
-	// Generate new links
-	// TODO: we should either keep a map in the PCGEditorGraph or do this elsewhere
-	// TODO: this will not work if we have non-PCG nodes in the graph
-	if (PCGNode)
+	// Remove old pins
+	for (UEdGraphPin* OldPin : OldPins)
 	{
-		UPCGEditorGraph* PCGEditorGraph = CastChecked<UPCGEditorGraph>(GetGraph());
-		PCGEditorGraph->CreateLinks(this, /*bCreateInbound=*/true, /*bCreateOutbound=*/true);
+		RemovePin(OldPin);
 	}
-	
+
 	// Notify editor
 	OnNodeChangedDelegate.ExecuteIfBound();
 }
