Simplified code path when not parsing components on actor to prevent special cases

#rnx

diff --git a/Engine/Plugins/Experimental/PCG/Source/PCG/Private/PCGComponent.cpp b/Engine/Plugins/Experimental/PCG/Source/PCG/Private/PCGComponent.cpp
--- a/Engine/Plugins/Experimental/PCG/Source/PCG/Private/PCGComponent.cpp
+++ b/Engine/Plugins/Experimental/PCG/Source/PCG/Private/PCGComponent.cpp
@@ -1837,6 +1837,12 @@ UPCGData* UPCGComponent::CreateActorPCGData(AActor* Actor, const UPCGComponent*
 
 		return Data;
 	}
+	else if (!bParseActor)
+	{
+		UPCGPointData* Data = NewObject<UPCGPointData>();
+		Data->InitializeFromActor(Actor);
+		return Data;
+	}
 	else if (AVolume* Volume = Cast<AVolume>(Actor))
 	{
 		UPCGVolumeData* Data = NewObject<UPCGVolumeData>();
@@ -1844,7 +1850,7 @@ UPCGData* UPCGComponent::CreateActorPCGData(AActor* Actor, const UPCGComponent*
 
 		return Data;
 	}
-	else if (bParseActor)// Prepare data on a component basis
+	else // Prepare data on a component basis
 	{
 		TInlineComponentArray<ULandscapeSplinesComponent*, 1> LandscapeSplines;
 		Actor->GetComponents(LandscapeSplines);
