Improved task successor check so it happens once and not at every following execution

#rnx

diff --git a/Engine/Plugins/Experimental/PCG/Source/PCG/Private/Graph/PCGGraphExecutor.cpp b/Engine/Plugins/Experimental/PCG/Source/PCG/Private/Graph/PCGGraphExecutor.cpp
--- a/Engine/Plugins/Experimental/PCG/Source/PCG/Private/Graph/PCGGraphExecutor.cpp
+++ b/Engine/Plugins/Experimental/PCG/Source/PCG/Private/Graph/PCGGraphExecutor.cpp
@@ -746,7 +746,11 @@ void FPCGGraphExecutor::Execute()
 		// keep a cache of intermediate results.
 		if (ReadyTasks.Num() == 0 && ActiveTasks.Num() == 0 && SleepingTasks.Num() == 0 && Tasks.Num() == 0)
 		{
-			check(TaskSuccessors.IsEmpty());
+			if (!ensure(TaskSuccessors.IsEmpty()))
+			{
+				TaskSuccessors.Reset();
+			}
+			
 			ClearResults();
 
 #if WITH_EDITOR
