Fixed bad condition check in surface sampler

#rnx

diff --git a/Engine/Plugins/Experimental/PCG/Source/PCG/Private/Elements/PCGSurfaceSampler.cpp b/Engine/Plugins/Experimental/PCG/Source/PCG/Private/Elements/PCGSurfaceSampler.cpp
--- a/Engine/Plugins/Experimental/PCG/Source/PCG/Private/Elements/PCGSurfaceSampler.cpp
+++ b/Engine/Plugins/Experimental/PCG/Source/PCG/Private/Elements/PCGSurfaceSampler.cpp
@@ -380,6 +380,9 @@ bool FPCGSurfaceSamplerElement::ExecuteInternal(FPCGContext* Context) const
 		PCGSurfaceSampler::FSurfaceSamplerSettings LoopData;
 		if (!LoopData.Initialize(Settings, Context, InputBounds))
 		{
+			Outputs.RemoveAt(GenerationIndex);
+			GeneratingShapes.RemoveAt(GenerationIndex);
+			--GenerationIndex;
 			continue;
 		}
 
