Marking  FPCGPointProcessingElementBase PCG_API so it's usable from other dlls

diff --git a/Engine/Plugins/Experimental/PCG/Source/PCG/Public/Elements/PCGPointProcessingElementBase.h b/Engine/Plugins/Experimental/PCG/Source/PCG/Public/Elements/PCGPointProcessingElementBase.h
--- a/Engine/Plugins/Experimental/PCG/Source/PCG/Public/Elements/PCGPointProcessingElementBase.h
+++ b/Engine/Plugins/Experimental/PCG/Source/PCG/Public/Elements/PCGPointProcessingElementBase.h
@@ -8,7 +8,7 @@
 struct FPCGPoint;
 struct FPCGTaggedData;
 
-class FPCGPointProcessingElementBase : public FSimplePCGElement
+class PCG_API FPCGPointProcessingElementBase : public FSimplePCGElement
 {
 protected:
 	virtual bool ExecuteInternal(FPCGContext* Context) const = 0;
