Added 'Gather' node, to be used for ui purposes make it easier to deal with combined pins

diff --git a/Engine/Plugins/Experimental/PCG/Source/PCG/Private/Elements/PCGGather.cpp b/Engine/Plugins/Experimental/PCG/Source/PCG/Private/Elements/PCGGather.cpp
new file mode 100644
--- /dev/null
+++ b/Engine/Plugins/Experimental/PCG/Source/PCG/Private/Elements/PCGGather.cpp
@@ -0,0 +1,37 @@
+// Copyright Epic Games, Inc. All Rights Reserved.
+
+#include "Elements/PCGGather.h"
+#include "Helpers/PCGSettingsHelpers.h"
+#include "PCGContext.h"
+#include "PCGPin.h"
+	
+TArray<FPCGPinProperties> UPCGGatherSettings::InputPinProperties() const
+{
+	TArray<FPCGPinProperties> PinProperties;
+	PinProperties.Emplace(PCGPinConstants::DefaultInputLabel, EPCGDataType::Any);
+
+	return PinProperties;
+}
+
+TArray<FPCGPinProperties> UPCGGatherSettings::OutputPinProperties() const
+{
+	TArray<FPCGPinProperties> PinProperties;
+	PinProperties.Emplace(PCGPinConstants::DefaultOutputLabel, EPCGDataType::Any);
+
+	return PinProperties;
+}
+
+FPCGElementPtr UPCGGatherSettings::CreateElement() const
+{
+	return MakeShared<FPCGGatherElement>();
+}
+
+bool FPCGGatherElement::ExecuteInternal(FPCGContext* Context) const
+{
+	TRACE_CPUPROFILER_EVENT_SCOPE(FPCGGatherElement::Execute);
+
+	Context->OutputData = Context->InputData;
+
+	return true;
+
+}
diff --git a/Engine/Plugins/Experimental/PCG/Source/PCG/Private/Tests/Elements/PCGGatherTest.cpp b/Engine/Plugins/Experimental/PCG/Source/PCG/Private/Tests/Elements/PCGGatherTest.cpp
new file mode 100644
--- /dev/null
+++ b/Engine/Plugins/Experimental/PCG/Source/PCG/Private/Tests/Elements/PCGGatherTest.cpp
@@ -0,0 +1,51 @@
+// Copyright Epic Games, Inc. All Rights Reserved.
+
+#include "Tests/Determinism/PCGDeterminismTestsCommon.h"
+#include "Tests/PCGTestsCommon.h"
+
+#include "PCGComponent.h"
+#include "PCGContext.h"
+#include "PCGData.h"
+#include "PCGHelpers.h"
+#include "PCGParamData.h"
+#include "Data/PCGPointData.h"
+
+#include "Elements/PCGGather.h"
+
+IMPLEMENT_CUSTOM_SIMPLE_AUTOMATION_TEST(FPCGGatherTest_Basic, FPCGTestBaseClass, "pcg.tests.Gather.Basic", PCGTestsCommon::TestFlags)
+
+bool FPCGGatherTest_Basic::RunTest(const FString& Parameters)
+{
+	PCGTestsCommon::FTestData TestData(PCGDeterminismTests::Defaults::Seed);
+	PCGDeterminismTests::GenerateSettings<UPCGGatherSettings>(TestData);
+	UPCGGatherSettings* Settings = CastChecked<UPCGGatherSettings>(TestData.Settings);
+
+	FPCGElementPtr TestElement = TestData.Settings->GetElement();
+
+	{
+		FPCGTaggedData& SourcePin = TestData.InputData.TaggedData.Emplace_GetRef();
+		SourcePin.Pin = PCGPinConstants::DefaultInputLabel;
+		SourcePin.Data = PCGTestsCommon::CreateRandomPointData(100, 42);
+	}
+
+	TUniquePtr<FPCGContext> Context(TestElement->Initialize(TestData.InputData, TestData.TestPCGComponent, nullptr));
+	Context->NumAvailableTasks = 1;
+
+	while (!TestElement->Execute(Context.Get())) {}
+
+	// test our point data
+	const TArray<FPCGTaggedData>& Outputs = Context->OutputData.GetInputs();
+
+	UTEST_EQUAL("Output count", Outputs.Num(), 1);
+
+	UPCGPointData *OutPointData = Cast<UPCGPointData>(Outputs[0].Data);
+
+	UTEST_NOT_NULL("Output point data", OutPointData);
+
+	const TArray<FPCGPoint>& OutPoints = OutPointData->GetPoints();
+
+	UTEST_EQUAL("Output point count", OutPoints.Num(), 100);
+
+	return true;
+}
+
diff --git a/Engine/Plugins/Experimental/PCG/Source/PCG/Public/Elements/PCGGather.h b/Engine/Plugins/Experimental/PCG/Source/PCG/Public/Elements/PCGGather.h
new file mode 100644
--- /dev/null
+++ b/Engine/Plugins/Experimental/PCG/Source/PCG/Public/Elements/PCGGather.h
@@ -0,0 +1,38 @@
+// Copyright Epic Games, Inc. All Rights Reserved.
+#pragma once
+
+#include "CoreMinimal.h"
+#include "PCGSettings.h"
+#include "PCGElement.h"
+
+#include "PCGGather.generated.h"
+
+/**
+ * Used to wrangle multiple input wires into one output wire for organizational purposes.
+ */
+UCLASS(BlueprintType, ClassGroup = (Procedural))
+class UPCGGatherSettings : public UPCGSettings
+{
+	GENERATED_BODY()
+
+public:
+	//~Begin UPCGSettings interface
+#if WITH_EDITOR
+	virtual FName GetDefaultNodeName() const override { return FName(TEXT("Gather")); }
+	virtual EPCGSettingsType GetType() const override { return EPCGSettingsType::Spatial; }
+#endif
+
+	virtual TArray<FPCGPinProperties> InputPinProperties() const override;
+	virtual TArray<FPCGPinProperties> OutputPinProperties() const override;
+
+protected:
+	virtual FPCGElementPtr CreateElement() const override;
+	//~End UPCGSettings interface
+};
+
+class FPCGGatherElement : public FSimplePCGElement
+{
+protected:
+	virtual bool ExecuteInternal(FPCGContext* Context) const override;
+};
+
