[PCG] Use node tooltips as node palette tooltip

diff --git a/Engine/Plugins/Experimental/PCG/Source/PCG/Private/Elements/PCGProjectionElement.cpp b/Engine/Plugins/Experimental/PCG/Source/PCG/Private/Elements/PCGProjectionElement.cpp
--- a/Engine/Plugins/Experimental/PCG/Source/PCG/Private/Elements/PCGProjectionElement.cpp
+++ b/Engine/Plugins/Experimental/PCG/Source/PCG/Private/Elements/PCGProjectionElement.cpp
@@ -27,7 +27,7 @@ TArray<FPCGPinProperties> UPCGProjectionSettings::InputPinProperties() const
 #if WITH_EDITOR
 FText UPCGProjectionSettings::GetNodeTooltipText() const
 {
-	return LOCTEXT("ProjectionNodeTooltip", "Projects each of the inputs connected to In onto the Projection Target and concatenates all of the results to Out. Overrides to the projection settings can be specified using the Params.");
+	return LOCTEXT("ProjectionNodeTooltip", "Projects each of the inputs connected to In onto the Projection Target and concatenates all of the results to Out.");
 }
 #endif
 
diff --git a/Engine/Plugins/Experimental/PCG/Source/PCGEditor/Private/PCGEditorGraphSchema.cpp b/Engine/Plugins/Experimental/PCG/Source/PCGEditor/Private/PCGEditorGraphSchema.cpp
--- a/Engine/Plugins/Experimental/PCG/Source/PCGEditor/Private/PCGEditorGraphSchema.cpp
+++ b/Engine/Plugins/Experimental/PCG/Source/PCGEditor/Private/PCGEditorGraphSchema.cpp
@@ -220,7 +220,7 @@ void UPCGEditorGraphSchema::GetNativeElementActions(FGraphActionMenuBuilder& Act
 		{
 			const FText MenuDesc = FText::FromName(PCGSettings->GetDefaultNodeName());
 			const FText Category = StaticEnum<EPCGSettingsType>()->GetDisplayNameTextByValue(static_cast<__underlying_type(EPCGSettingsType)>(PCGSettings->GetType()));
-			const FText Description = FText::GetEmpty();
+			const FText Description = PCGSettings->GetNodeTooltipText();
 
 			TSharedPtr<FPCGEditorGraphSchemaAction_NewNativeElement> NewAction(new FPCGEditorGraphSchemaAction_NewNativeElement(Category, MenuDesc, Description, 0));
 			NewAction->SettingsClass = SettingsClass;
