Changes to the graph profiling view. The profiler is now task based, not node based so per call times on graph instance children may be recorded.

diff --git a/Engine/Plugins/Experimental/PCG/Source/PCG/Private/Graph/PCGGraphCompiler.cpp b/Engine/Plugins/Experimental/PCG/Source/PCG/Private/Graph/PCGGraphCompiler.cpp
--- a/Engine/Plugins/Experimental/PCG/Source/PCG/Private/Graph/PCGGraphCompiler.cpp
+++ b/Engine/Plugins/Experimental/PCG/Source/PCG/Private/Graph/PCGGraphCompiler.cpp
@@ -54,7 +54,7 @@ TArray<FPCGGraphTask> FPCGGraphCompiler::CompileGraph(UPCGGraph* InGraph, FPCGTa
 			GraphDependenciesLock.Unlock();
 #endif // WITH_EDITOR
 
-			OffsetNodeIds(Subtasks, NextId);
+			OffsetNodeIds(Subtasks, NextId, PreId);
 			NextId += Subtasks.Num();
 
 			const UPCGNode* SubgraphInputNode = Subgraph->GetInputNode();
@@ -211,13 +211,23 @@ void FPCGGraphCompiler::Compile(UPCGGraph* InGraph)
 	// Store back the results in the cache if it's valid
 	if (!CompiledTasks.IsEmpty())
 	{
-	GraphToTaskMapLock.WriteLock();
-	if (!GraphToTaskMap.Contains(InGraph))
-	{
-		GraphToTaskMap.Add(InGraph, MoveTemp(CompiledTasks));
+		GraphToTaskMapLock.WriteLock();
+		if (!GraphToTaskMap.Contains(InGraph))
+		{
+			GraphToTaskMap.Add(InGraph, MoveTemp(CompiledTasks));
+		}
+		GraphToTaskMapLock.WriteUnlock();
 	}
-	GraphToTaskMapLock.WriteUnlock();
 }
+
+TArray<FPCGGraphTask> FPCGGraphCompiler::GetPrecompiledTasks(UPCGGraph* InGraph, bool bIsTopGraph) const
+{
+	// Get compiled tasks in a threadsafe way
+	FReadScopeLock ReadLock(GraphToTaskMapLock);
+
+	const TArray<FPCGGraphTask>* ExistingTasks = (bIsTopGraph ? TopGraphToTaskMap : GraphToTaskMap).Find(InGraph);
+
+	return ExistingTasks ? *ExistingTasks : TArray<FPCGGraphTask>();
 }
 
 TArray<FPCGGraphTask> FPCGGraphCompiler::GetCompiledTasks(UPCGGraph* InGraph, bool bIsTopGraph)
@@ -254,12 +264,21 @@ TArray<FPCGGraphTask> FPCGGraphCompiler::GetCompiledTasks(UPCGGraph* InGraph, bo
 	return CompiledTasks;
 }
 
-void FPCGGraphCompiler::OffsetNodeIds(TArray<FPCGGraphTask>& Tasks, FPCGTaskId Offset)
+void FPCGGraphCompiler::OffsetNodeIds(TArray<FPCGGraphTask>& Tasks, FPCGTaskId Offset, FPCGTaskId ParentId)
 {
 	for (FPCGGraphTask& Task : Tasks)
 	{
 		Task.NodeId += Offset;
 
+		if  (Task.ParentId == InvalidPCGTaskId)
+		{
+			Task.ParentId = ParentId;
+		}
+		else
+		{
+			Task.ParentId += Offset;
+		}
+
 		for(FPCGGraphTaskInput& Input : Task.Inputs)
 		{
 			Input.TaskId += Offset;
@@ -302,6 +321,8 @@ void FPCGGraphCompiler::CompileTopGraph(UPCGGraph* InGraph)
 		{
 			Task.Inputs.Emplace(PreExecuteTaskId, nullptr, nullptr);
 		}
+
+		Task.CompiledTaskId = Task.NodeId;
 	}
 
 	FPCGGraphTask& PostExecuteTask = CompiledTasks.Emplace_GetRef();
diff --git a/Engine/Plugins/Experimental/PCG/Source/PCG/Private/Graph/PCGGraphCompiler.h b/Engine/Plugins/Experimental/PCG/Source/PCG/Private/Graph/PCGGraphCompiler.h
--- a/Engine/Plugins/Experimental/PCG/Source/PCG/Private/Graph/PCGGraphCompiler.h
+++ b/Engine/Plugins/Experimental/PCG/Source/PCG/Private/Graph/PCGGraphCompiler.h
@@ -16,8 +16,9 @@ class FPCGGraphCompiler
 public:
 	void Compile(UPCGGraph* InGraph);
 	TArray<FPCGGraphTask> GetCompiledTasks(UPCGGraph* InGraph, bool bIsTopGraph = true);
+	TArray<FPCGGraphTask> GetPrecompiledTasks(UPCGGraph* InGraph, bool bIsTopGraph = true) const;
 
-	static void OffsetNodeIds(TArray<FPCGGraphTask>& Tasks, FPCGTaskId Offset);
+	static void OffsetNodeIds(TArray<FPCGGraphTask>& Tasks, FPCGTaskId Offset, FPCGTaskId ParentId);
 
 #if WITH_EDITOR
 	void NotifyGraphChanged(UPCGGraph* InGraph);
@@ -27,7 +28,7 @@ private:
 	TArray<FPCGGraphTask> CompileGraph(UPCGGraph* InGraph, FPCGTaskId& NextId);
 	void CompileTopGraph(UPCGGraph* InGraph);
 
-	FRWLock GraphToTaskMapLock;
+	mutable FRWLock GraphToTaskMapLock;
 	TMap<UPCGGraph*, TArray<FPCGGraphTask>> GraphToTaskMap;
 	TMap<UPCGGraph*, TArray<FPCGGraphTask>> TopGraphToTaskMap;
 
diff --git a/Engine/Plugins/Experimental/PCG/Source/PCG/Private/Graph/PCGGraphExecutor.cpp b/Engine/Plugins/Experimental/PCG/Source/PCG/Private/Graph/PCGGraphExecutor.cpp
--- a/Engine/Plugins/Experimental/PCG/Source/PCG/Private/Graph/PCGGraphExecutor.cpp
+++ b/Engine/Plugins/Experimental/PCG/Source/PCG/Private/Graph/PCGGraphExecutor.cpp
@@ -106,7 +106,7 @@ FPCGTaskId FPCGGraphExecutor::Schedule(UPCGGraph* Graph, UPCGComponent* SourceCo
 		ScheduledTask.SourceComponent = SourceComponent;
 
 		// Offset task node ids
-		FPCGGraphCompiler::OffsetNodeIds(ScheduledTask.Tasks, NextTaskId);
+		FPCGGraphCompiler::OffsetNodeIds(ScheduledTask.Tasks, NextTaskId, InvalidPCGTaskId);
 		NextTaskId += ScheduledTask.Tasks.Num();
 		ScheduledId = NextTaskId - 1; // This is true because the last task is from the output node or is the post-execute task
 
@@ -516,6 +516,7 @@ void FPCGGraphExecutor::Execute()
 				{
 					Task.Context = Task.Element->Initialize(TaskInput, Task.SourceComponent, Task.Node);
 					Task.Context->TaskId = Task.NodeId;
+					Task.Context->CompiledTaskId = Task.CompiledTaskId;
 				}
 
 				// Validate that we can start this task now
diff --git a/Engine/Plugins/Experimental/PCG/Source/PCG/Private/Graph/PCGGraphExecutor.h b/Engine/Plugins/Experimental/PCG/Source/PCG/Private/Graph/PCGGraphExecutor.h
--- a/Engine/Plugins/Experimental/PCG/Source/PCG/Private/Graph/PCGGraphExecutor.h
+++ b/Engine/Plugins/Experimental/PCG/Source/PCG/Private/Graph/PCGGraphExecutor.h
@@ -43,6 +43,8 @@ struct FPCGGraphTask
 	FPCGElementPtr Element; // Added to have tasks that aren't node-bound
 	FPCGContext* Context = nullptr;
 	FPCGTaskId NodeId = InvalidPCGTaskId;
+	FPCGTaskId CompiledTaskId = InvalidPCGTaskId; // the task id as it exists when compiled
+	FPCGTaskId ParentId = InvalidPCGTaskId; // represents the parent sub object graph task, if we were called from one
 };
 
 struct FPCGGraphScheduleTask
@@ -108,6 +110,10 @@ public:
 
 	/** Notify compiler that graph has changed so it'll be removed from the cache */
 	void NotifyGraphChanged(UPCGGraph* InGraph);
+
+	/** So the profiler can decode graph task ids **/
+	FPCGGraphCompiler* GetCompiler() const { return GraphCompiler.Get(); }
+
 #endif
 
 	/** "Tick" of the graph executor. This call is NOT THREADSAFE */
diff --git a/Engine/Plugins/Experimental/PCG/Source/PCG/Private/PCGComponent.cpp b/Engine/Plugins/Experimental/PCG/Source/PCG/Private/PCGComponent.cpp
--- a/Engine/Plugins/Experimental/PCG/Source/PCG/Private/PCGComponent.cpp
+++ b/Engine/Plugins/Experimental/PCG/Source/PCG/Private/PCGComponent.cpp
@@ -533,6 +533,11 @@ FPCGTaskId UPCGComponent::CleanupInternal(bool bRemoveComponents, bool bSave, co
 
 	Modify();
 
+#if WITH_EDITOR
+	ExtraCapture.ResetTimers();
+	ExtraCapture.ResetCapturedMessages();
+#endif
+
 	CurrentCleanupTask = GetSubsystem()->ScheduleCleanup(this, bRemoveComponents, bSave, Dependencies);
 	return CurrentCleanupTask;
 }
diff --git a/Engine/Plugins/Experimental/PCG/Source/PCG/Private/PCGSubsystem.cpp b/Engine/Plugins/Experimental/PCG/Source/PCG/Private/PCGSubsystem.cpp
--- a/Engine/Plugins/Experimental/PCG/Source/PCG/Private/PCGSubsystem.cpp
+++ b/Engine/Plugins/Experimental/PCG/Source/PCG/Private/PCGSubsystem.cpp
@@ -1608,6 +1608,16 @@ void UPCGSubsystem::ClearLandscapeCache()
 	}
 }
 
+const FPCGGraphCompiler* UPCGSubsystem::GetGraphCompiler() const
+{
+	if (GraphExecutor)
+	{
+		return GraphExecutor->GetCompiler();
+	}
+
+	return nullptr;
+}
+
 void UPCGSubsystem::ResetPartitionActorsMap()
 {
 	PartitionActorsMapLock.WriteLock();
diff --git a/Engine/Plugins/Experimental/PCG/Source/PCG/Private/Utils/PCGExtraCapture.cpp b/Engine/Plugins/Experimental/PCG/Source/PCG/Private/Utils/PCGExtraCapture.cpp
--- a/Engine/Plugins/Experimental/PCG/Source/PCG/Private/Utils/PCGExtraCapture.cpp
+++ b/Engine/Plugins/Experimental/PCG/Source/PCG/Private/Utils/PCGExtraCapture.cpp
@@ -4,6 +4,8 @@
 
 #include "PCGComponent.h"
 #include "PCGContext.h"
+#include "Graph/PCGGraphCompiler.h"
+#include "Graph/PCGGraphExecutor.h"
 
 #if WITH_EDITOR
 
@@ -24,51 +26,172 @@ void PCGUtils::FExtraCapture::Update(const PCGUtils::FScopedCall& InScopedCall)
 		return;
 	}
 
-	FScopeLock ScopedLock(&Lock);
-
-	TArray<FCallTime>* TimesPtr = Timers.Find(InScopedCall.Context->Node);
-	TArray<FCapturedMessage>* CapturedMessagesPtr = CapturedMessages.Find(InScopedCall.Context->Node);
-
-	if (!TimesPtr)
-	{
-		TimesPtr = &Timers.Add(InScopedCall.Context->Node);
-	}
-
-	if (!CapturedMessagesPtr)
-	{
-		CapturedMessagesPtr = &CapturedMessages.Add(InScopedCall.Context->Node);
-	}
-
-	check(TimesPtr && CapturedMessagesPtr);
-
 	const double ThisFrameTime = FPlatformTime::Seconds() - InScopedCall.StartTime;
 
-	constexpr int32 MaxNumberOfTrackedTimers = 100;
+	FScopeLock ScopedLock(&Lock);
 
-	FCallTime* Time = !TimesPtr->IsEmpty() ? &TimesPtr->Last() : nullptr;
-	check(Time || InScopedCall.Phase == EPCGExecutionPhase::NotExecuted);
+	FCallTime& Timer = Timers.FindOrAdd(InScopedCall.Context->CompiledTaskId);
 
 	switch (InScopedCall.Phase)
 	{
 	case EPCGExecutionPhase::NotExecuted:
-		TimesPtr->Emplace();
+		Timer = FCallTime(); // reset it
 		break;
 	case EPCGExecutionPhase::PrepareData:
-		Time->PrepareDataTime = ThisFrameTime;
+		Timer.PrepareDataTime = ThisFrameTime;
 		break;
 	case EPCGExecutionPhase::Execute:
-		Time->ExecutionTime += ThisFrameTime;
-		Time->ExecutionFrameCount++;
+		Timer.ExecutionTime += ThisFrameTime;
+		Timer.ExecutionFrameCount++;
 
-		Time->MaxExecutionFrameTime = FMath::Max(Time->MaxExecutionFrameTime, ThisFrameTime);
-		Time->MinExecutionFrameTime = FMath::Min(Time->MinExecutionFrameTime, ThisFrameTime);
+		Timer.MaxExecutionFrameTime = FMath::Max(Timer.MaxExecutionFrameTime, ThisFrameTime);
+		Timer.MinExecutionFrameTime = FMath::Min(Timer.MinExecutionFrameTime, ThisFrameTime);
 		break;
 	case EPCGExecutionPhase::PostExecute:
-		Time->PostExecuteTime = ThisFrameTime;
+		Timer.PostExecuteTime = ThisFrameTime;
 		break;
 	}
 
-	CapturedMessagesPtr->Append(InScopedCall.CapturedMessages);
+	if (!InScopedCall.CapturedMessages.IsEmpty())
+	{
+		TArray<FCapturedMessage>& InstanceMessages = CapturedMessages.FindOrAdd(InScopedCall.Context->Node);
+		InstanceMessages.Append(std::move(InScopedCall.CapturedMessages));
+	}
+}
+
+namespace PCGUtils
+{
+	void AddTimers(FCallTreeInfo& RootInfo, const FExtraCapture::TTimersMap& Timers, const TMap<FPCGTaskId, const FPCGGraphTask*>& TaskLookup)
+	{
+		// re-use this map
+		TArray<FPCGTaskId> PathIds;
+		TArray<const UPCGNode*> PathNodes;
+
+		for (const auto& Pair : Timers)
+		{
+			const FPCGTaskId TaskId = Pair.Key;
+			const FCallTime& CallTime = Pair.Value;
+
+			PathIds.Reset();
+			PathNodes.Reset();
+
+			// build a list of parent ids that we need to populate our tree with
+			PathIds.Add(TaskId);
+			while (true)
+			{
+				const FPCGGraphTask*const * TaskItr = TaskLookup.Find(PathIds.Last());
+				if (!TaskItr || !*TaskItr)
+				{
+					// can't find an entry in our list ignore it
+					PathIds.Reset();
+					PathNodes.Reset();
+					break;
+				}
+
+				const FPCGGraphTask& Task = **TaskItr;
+
+				if (PathNodes.Num() < PathIds.Num())
+				{
+					PathNodes.Add(Task.Node);
+				}
+
+				if (Task.ParentId == InvalidPCGTaskId)
+				{
+					break;
+				}
+
+				PathIds.Add(Task.ParentId);
+			}
+
+			FCallTreeInfo* CurrentInfo = &RootInfo;
+
+			// now add the entries, the order is reversed b/c top level ones were pushed to the end
+			for (int32 Idx = PathIds.Num()-1; Idx >= 0; --Idx)
+			{
+				const FPCGTaskId CurrentId = PathIds[Idx];
+
+				FCallTreeInfo* FoundInfo = nullptr;
+				for (FCallTreeInfo& Child : CurrentInfo->Children)
+				{
+					if (Child.TaskId == CurrentId)
+					{
+						FoundInfo = &Child;
+						break;
+					}
+				}
+
+				if (!FoundInfo)
+				{
+					// didn't find an existing entry, add one
+					FCallTreeInfo& Child = CurrentInfo->Children.Emplace_GetRef();
+					Child.TaskId = CurrentId;
+					Child.Node = PathNodes[Idx];
+
+					CurrentInfo = &Child;
+				}
+				else
+				{
+					CurrentInfo = FoundInfo;
+				}
+			}
+
+			CurrentInfo->CallTime = CallTime;
+		}
+	}
+
+	void BuildTreeInfo(FCallTreeInfo& Info)
+	{
+		for (FCallTreeInfo& Child : Info.Children)
+		{
+			BuildTreeInfo(Child);
+
+			Info.CallTime.PrepareDataTime += Child.CallTime.PrepareDataTime;
+			Info.CallTime.ExecutionTime += Child.CallTime.ExecutionTime;
+			Info.CallTime.PostExecuteTime += Child.CallTime.PostExecuteTime;
+			Info.CallTime.MinExecutionFrameTime = FMath::Min(Info.CallTime.MinExecutionFrameTime, Child.CallTime.MinExecutionFrameTime);
+			Info.CallTime.MaxExecutionFrameTime = FMath::Max(Info.CallTime.MaxExecutionFrameTime, Child.CallTime.MaxExecutionFrameTime);
+		}
+	}
+}
+
+PCGUtils::FCallTreeInfo PCGUtils::FExtraCapture::CalculateCallTreeInfo(const UPCGComponent* Component) const
+{
+	UPCGSubsystem* PCGSubsystem = Component ? Component->GetSubsystem() : nullptr;
+
+	if (!PCGSubsystem)
+	{
+		return {};
+	}
+
+	const FPCGGraphCompiler* Compiler = PCGSubsystem->GetGraphCompiler();
+	if (!Compiler)
+	{
+		return {};
+	}
+
+	TArray<FPCGGraphTask> CompiledTasks = Compiler->GetPrecompiledTasks(Component->GetGraph());
+	if (CompiledTasks.IsEmpty())
+	{
+		return {};
+	}
+
+	// the last task on a top level graph is post execute task that doesn't need to be in the call tree
+	CompiledTasks.Pop(); 
+
+	TMap<FPCGTaskId, const FPCGGraphTask*> TaskLookup;
+
+	// build some lookup maps
+	for (const FPCGGraphTask& Task : CompiledTasks)
+	{
+		TaskLookup.Add(Task.CompiledTaskId, &Task);
+	}
+
+	FCallTreeInfo RootInfo;
+
+	PCGUtils::AddTimers(RootInfo, Timers, TaskLookup);
+	PCGUtils::BuildTreeInfo(RootInfo);
+
+	return RootInfo;
 }
 
 PCGUtils::FScopedCall::FScopedCall(const IPCGElement& InOwner, FPCGContext* InContext)
@@ -108,4 +231,4 @@ void PCGUtils::FScopedCall::Serialize(const TCHAR* V, ELogVerbosity::Type Verbos
 	CapturedMessages.Add(PCGUtils::FCapturedMessage{ MessageCounter++, Category, V, Verbosity });
 }
 
-#endif // WITH_EDITOR
+#endif // WITH_EDITOR
\ No newline at end of file
diff --git a/Engine/Plugins/Experimental/PCG/Source/PCG/Public/PCGContext.h b/Engine/Plugins/Experimental/PCG/Source/PCG/Public/PCGContext.h
--- a/Engine/Plugins/Experimental/PCG/Source/PCG/Public/PCGContext.h
+++ b/Engine/Plugins/Experimental/PCG/Source/PCG/Public/PCGContext.h
@@ -53,6 +53,7 @@ struct PCG_API FPCGContext
 	// TODO: replace this by a better identification mechanism
 	const UPCGNode* Node = nullptr;
 	FPCGTaskId TaskId = InvalidPCGTaskId;
+	FPCGTaskId CompiledTaskId = InvalidPCGTaskId;
 	bool bIsPaused = false;
 
 	EPCGExecutionPhase CurrentPhase = EPCGExecutionPhase::NotExecuted;
diff --git a/Engine/Plugins/Experimental/PCG/Source/PCG/Public/PCGSubsystem.h b/Engine/Plugins/Experimental/PCG/Source/PCG/Public/PCGSubsystem.h
--- a/Engine/Plugins/Experimental/PCG/Source/PCG/Public/PCGSubsystem.h
+++ b/Engine/Plugins/Experimental/PCG/Source/PCG/Public/PCGSubsystem.h
@@ -10,21 +10,23 @@
 
 #include "PCGSubsystem.generated.h"
 
-enum class EPCGComponentDirtyFlag : uint8;
-enum class ETickableTickType : uint8;
-struct FPCGContext;
 
-class FPCGGraphExecutor;
 class APCGPartitionActor;
 class APCGWorldActor;
 class UPCGGraph;
-struct FPCGDataCollection;
 class UPCGLandscapeCache;
 
+enum class EPCGComponentDirtyFlag : uint8;
+enum class ETickableTickType : uint8;
+
+class FPCGGraphCompiler;
+class FPCGGraphExecutor;
+struct FPCGContext;
+struct FPCGDataCollection;
+
 class IPCGElement;
 typedef TSharedPtr<IPCGElement, ESPMode::ThreadSafe> FPCGElementPtr;
 
-
 /**
 * UPCGSubsystem
 */
@@ -159,6 +161,9 @@ public:
 	/** Clears the landscape data cache */
 	void ClearLandscapeCache();
 
+	/** Returns the graph compiler so we can figure out task info in the profiler view **/
+	const FPCGGraphCompiler* GetGraphCompiler() const;
+
 private:
 	enum class EOperation : uint32
 	{
diff --git a/Engine/Plugins/Experimental/PCG/Source/PCG/Public/Utils/PCGExtraCapture.h b/Engine/Plugins/Experimental/PCG/Source/PCG/Public/Utils/PCGExtraCapture.h
--- a/Engine/Plugins/Experimental/PCG/Source/PCG/Public/Utils/PCGExtraCapture.h
+++ b/Engine/Plugins/Experimental/PCG/Source/PCG/Public/Utils/PCGExtraCapture.h
@@ -2,8 +2,10 @@
 
 #pragma once
 
+#include "PCGCommon.h"
 #include "UObject/WeakObjectPtr.h" // IWYU pragma: keep
 
+class UPCGComponent;
 class UPCGNode;
 enum class EPCGExecutionPhase : uint8;
 struct FPCGContext;
@@ -33,6 +35,16 @@ namespace PCGUtils
 		ELogVerbosity::Type Verbosity;
 	};
 
+	struct PCG_API FCallTreeInfo
+	{
+		FPCGTaskId TaskId;
+		FString Name; // name of the task, Node may be nullptr for generated tasks
+		const UPCGNode* Node = nullptr;
+		FCallTime CallTime;
+
+		TArray<FCallTreeInfo> Children;
+	};
+
 	struct FScopedCall : public FOutputDevice
 	{
 		FScopedCall(const IPCGElement& InOwner, FPCGContext* InContext);
@@ -60,12 +72,14 @@ namespace PCGUtils
 		void ResetTimers();
 		void ResetCapturedMessages();
 
-		using TTimersMap = TMap<TWeakObjectPtr<const UPCGNode>, TArray<FCallTime>>;
+		using TTimersMap = TMap<FPCGTaskId, FCallTime>;
 		using TCapturedMessageMap = TMap<TWeakObjectPtr<const UPCGNode>, TArray<FCapturedMessage>>;
 
 		const TTimersMap& GetTimers() const { return Timers; }
 		const TCapturedMessageMap& GetCapturedMessages() const { return CapturedMessages; }
 
+		FCallTreeInfo CalculateCallTreeInfo(const UPCGComponent* Component) const;
+
 	private:
 		TTimersMap Timers;
 		TCapturedMessageMap CapturedMessages;
diff --git a/Engine/Plugins/Experimental/PCG/Source/PCGEditor/Private/SPCGEditorGraphProfilingView.cpp b/Engine/Plugins/Experimental/PCG/Source/PCGEditor/Private/SPCGEditorGraphProfilingView.cpp
--- a/Engine/Plugins/Experimental/PCG/Source/PCGEditor/Private/SPCGEditorGraphProfilingView.cpp
+++ b/Engine/Plugins/Experimental/PCG/Source/PCGEditor/Private/SPCGEditorGraphProfilingView.cpp
@@ -25,31 +25,19 @@ namespace PCGEditorGraphProfilingView
 	/** Names of the columns in the attribute list */
 	const FName NAME_Node = FName(TEXT("Node"));
 	const FName NAME_PrepareDataTime = FName(TEXT("PrepareDataTime"));
-	const FName NAME_AvgExecutionTime = FName(TEXT("AvgExecutionTime"));
-	const FName NAME_MinExecutionTime = FName(TEXT("MinExecutionTime"));
-	const FName NAME_MaxExecutionTime = FName(TEXT("MaxExecutionTime"));
 	const FName NAME_MinExecutionFrameTime = FName(TEXT("MinFrameTime"));
 	const FName NAME_MaxExecutionFrameTime = FName(TEXT("MaxFrameTime"));
-	const FName NAME_MinNbExecutionFrames = FName(TEXT("MinNbExecutionFrames"));
-	const FName NAME_MaxNbExecutionFrames = FName(TEXT("MaxNbExecutionFrames"));
-	const FName NAME_StdExecutionTime = FName(TEXT("StdExecutionTime"));
 	const FName NAME_TotalExecutionTime = FName(TEXT("TotalExecutionTime"));
-	const FName NAME_NbCalls = FName(TEXT("NbCalls"));
 	const FName NAME_NbExecutionFrames = FName(TEXT("NbExecutionFrames"));
 	const FName NAME_PostExecuteTime = FName(TEXT("PostExecuteTime"));
 
 	/** Labels of the columns */
 	const FText TEXT_NodeLabel = LOCTEXT("NodeLabel", "Node");
 	const FText TEXT_PrepareDataTimeLabel = LOCTEXT("PrepareDataTimeLabel", "PrepareData (ms)");
-	const FText TEXT_PostExecuteTimeLabel = LOCTEXT("PrepareDataTimeLabel", "PrepareData (ms)");
-	const FText TEXT_AvgExecutionTimeLabel = LOCTEXT("AvgExecutionTimeLabel", "Avg Time(ms)");
-	const FText TEXT_MinExecutionTimeLabel = LOCTEXT("MinExecutionTimeLabel", "Min Time(ms)");
-	const FText TEXT_MaxExecutionTimeLabel = LOCTEXT("MaxExecutionTimeLabel", "Max Time(ms)");
+	const FText TEXT_PostExecuteTimeLabel = LOCTEXT("PostExecuteTimeLabel", "PostExecute (ms)");
 	const FText TEXT_MinExecutionFrameTimeLabel = LOCTEXT("MinExecutionFrameTimeLabel", "Min Frame Time(ms)");
 	const FText TEXT_MaxExecutionFrameTimeLabel = LOCTEXT("MaxExecutionFrameTimeLabel", "Max Frame Time(ms)");
-	const FText TEXT_StdExecutionTimeLabel = LOCTEXT("StdExecutionTimeLabel", "Std(ms)");
 	const FText TEXT_TotalExecutionTimeLabel = LOCTEXT("TotalExecutionTimeLabel", "Total time(s)");
-	const FText TEXT_NbCallsLabel = LOCTEXT("NbCallsLabel", "Calls");
 	const FText TEXT_NbExecutionFramesLabel = LOCTEXT("NbExecutionFramesLabel", "Exec frames");
 }
 
@@ -70,70 +58,38 @@ TSharedRef<SWidget> SPCGProfilingListViewItemRow::GenerateWidgetForColumn(const
 	{
 		if (ColumnId == PCGEditorGraphProfilingView::NAME_Node)
 		{
-			ColumnData = FText::FromName(InternalItem->Name);
-		}
-		else if (ColumnId == PCGEditorGraphProfilingView::NAME_NbCalls)
-		{
-			ColumnData = FText::AsNumber(InternalItem->NbCalls);
+			ColumnData = FText::FromString(InternalItem->Name);
 		}
 		else if (ColumnId == PCGEditorGraphProfilingView::NAME_NbExecutionFrames)
 		{
-			ColumnData = ((InternalItem->NbExecutionFrames >= 0) ? FText::AsNumber(InternalItem->NbExecutionFrames) : FText());
-		}
-		else if (ColumnId == PCGEditorGraphProfilingView::NAME_MinNbExecutionFrames)
-		{
-			ColumnData = ((InternalItem->MinNbExecutionFrames >= 0) ? FText::AsNumber(InternalItem->MinNbExecutionFrames) : FText());
-		}
-		else if (ColumnId == PCGEditorGraphProfilingView::NAME_MaxNbExecutionFrames)
-		{
-			ColumnData = ((InternalItem->MaxNbExecutionFrames >= 0) ? FText::AsNumber(InternalItem->MaxNbExecutionFrames) : FText());
+			ColumnData = ((InternalItem->CallTime.ExecutionFrameCount >= 0) ? FText::AsNumber(InternalItem->CallTime.ExecutionFrameCount) : FText());
 		}
 		else if (!InternalItem->HasData)
 		{
 			ColumnData = PCGEditorGraphProfilingView::NoDataAvailableText;
 		}
-		else if (ColumnId == PCGEditorGraphProfilingView::NAME_AvgExecutionTime)
-		{
-			// In ms
-			ColumnData = FText::AsNumber(InternalItem->AvgExecutionTime * 1000.0);
-		}
-		else if (ColumnId == PCGEditorGraphProfilingView::NAME_MinExecutionTime)
-		{
-			// In ms
-			ColumnData = ((InternalItem->MinExecutionTime >= 0) ? FText::AsNumber(InternalItem->MinExecutionTime * 1000.0) : FText());
-		}
-		else if (ColumnId == PCGEditorGraphProfilingView::NAME_MaxExecutionTime)
-		{
-			// In ms
-			ColumnData = ((InternalItem->MaxExecutionTime >= 0) ? FText::AsNumber(InternalItem->MaxExecutionTime * 1000.0) : FText());
-		}
 		else if (ColumnId == PCGEditorGraphProfilingView::NAME_MinExecutionFrameTime)
 		{
 			// In ms
-			ColumnData = ((InternalItem->MinExecutionFrameTime >= 0) ? FText::AsNumber(InternalItem->MinExecutionFrameTime * 1000.0) : FText());
+			ColumnData = ((InternalItem->CallTime.MinExecutionFrameTime >= 0) ? FText::AsNumber(InternalItem->CallTime.MinExecutionFrameTime * 1000.0) : FText());
 		}
 		else if (ColumnId == PCGEditorGraphProfilingView::NAME_MaxExecutionFrameTime)
 		{
 			// In ms
-			ColumnData = ((InternalItem->MaxExecutionFrameTime >= 0) ? FText::AsNumber(InternalItem->MaxExecutionFrameTime * 1000.0) : FText());
-		}
-		else if (ColumnId == PCGEditorGraphProfilingView::NAME_StdExecutionTime)
-		{
-			// In ms
-			ColumnData = ((InternalItem->StdExecutionTime >= 0) ? FText::AsNumber(InternalItem->StdExecutionTime * 1000.0) : FText());
+			ColumnData = ((InternalItem->CallTime.MaxExecutionFrameTime >= 0) ? FText::AsNumber(InternalItem->CallTime.MaxExecutionFrameTime * 1000.0) : FText());
 		}
 		else if (ColumnId == PCGEditorGraphProfilingView::NAME_TotalExecutionTime)
 		{
 			// In s
-			ColumnData = FText::AsNumber(InternalItem->ExecutionTime);
+			ColumnData = FText::AsNumber(InternalItem->CallTime.ExecutionTime);
 		}
 		else if (ColumnId == PCGEditorGraphProfilingView::NAME_PrepareDataTime)
 		{
-			ColumnData = FText::AsNumber(InternalItem->PrepareDataTime * 1000.0);
+			ColumnData = FText::AsNumber(InternalItem->CallTime.PrepareDataTime * 1000.0);
 		}
 		else if (ColumnId == PCGEditorGraphProfilingView::NAME_PostExecuteTime)
 		{
-			ColumnData = FText::AsNumber(InternalItem->PostExecuteTime * 1000.0);
+			ColumnData = FText::AsNumber(InternalItem->CallTime.PostExecuteTime * 1000.0);
 		}
 	}
 
@@ -177,6 +133,8 @@ void SPCGEditorGraphProfilingView::Construct(const FArguments& InArgs, TSharedPt
 		PCGEditor->OnDebugObjectChangedDelegate.AddSP(this, &SPCGEditorGraphProfilingView::OnDebugObjectChanged);
 	}
 
+	SortingColumn = PCGEditorGraphProfilingView::NAME_TotalExecutionTime;
+	SortMode = EColumnSortMode::Descending;
 	ListViewHeader = CreateHeaderRowWidget();
 
 	const TSharedRef<SScrollBar> HorizontalScrollBar = SNew(SScrollBar)
@@ -223,6 +181,20 @@ void SPCGEditorGraphProfilingView::Construct(const FArguments& InArgs, TSharedPt
 				.IsChecked(this, &SPCGEditorGraphProfilingView::IsSubgraphExpanded)
 				.OnCheckStateChanged(this, &SPCGEditorGraphProfilingView::OnSubgraphExpandedChanged)
 			]
+			+SHorizontalBox::Slot()
+			.Padding(FMargin(30.f, 0.f, 0.f, 5.f))
+			.AutoWidth()
+			[
+				SNew(STextBlock)
+				.Text(LOCTEXT("TotalExecutionTime", "Total Time:"))
+			]
+			+SHorizontalBox::Slot()
+			.AutoWidth()
+			[
+				SNew(STextBlock)
+				.Text(this, &SPCGEditorGraphProfilingView::GetTotalTimeLabel)
+				.MinDesiredWidth(50.0f)
+			]
 		]
 		+SVerticalBox::Slot()
 		.FillHeight(1.0f)
@@ -271,13 +243,6 @@ TSharedRef<SHeaderRow> SPCGEditorGraphProfilingView::CreateHeaderRowWidget()
 		.HAlignCell(HAlign_Left)
 		.SortMode(this, &SPCGEditorGraphProfilingView::GetColumnSortMode, PCGEditorGraphProfilingView::NAME_Node)
 		.OnSort(this, &SPCGEditorGraphProfilingView::OnSortColumnHeader)
-		+ SHeaderRow::Column(PCGEditorGraphProfilingView::NAME_NbCalls)
-		.ManualWidth(80)
-		.DefaultLabel(PCGEditorGraphProfilingView::TEXT_NbCallsLabel)
-		.HAlignHeader(HAlign_Center)
-		.HAlignCell(HAlign_Right)
-		.SortMode(this, &SPCGEditorGraphProfilingView::GetColumnSortMode, PCGEditorGraphProfilingView::NAME_NbCalls)
-		.OnSort(this, &SPCGEditorGraphProfilingView::OnSortColumnHeader)
 		+ SHeaderRow::Column(PCGEditorGraphProfilingView::NAME_PrepareDataTime)
 		.ManualWidth(125)
 		.DefaultLabel(PCGEditorGraphProfilingView::TEXT_PrepareDataTimeLabel)
@@ -306,34 +271,6 @@ TSharedRef<SHeaderRow> SPCGEditorGraphProfilingView::CreateHeaderRowWidget()
 		.HAlignCell(HAlign_Right)
 		.SortMode(this, &SPCGEditorGraphProfilingView::GetColumnSortMode, PCGEditorGraphProfilingView::NAME_MaxExecutionFrameTime)
 		.OnSort(this, &SPCGEditorGraphProfilingView::OnSortColumnHeader)
-		+ SHeaderRow::Column(PCGEditorGraphProfilingView::NAME_AvgExecutionTime)
-		.ManualWidth(100)
-		.DefaultLabel(PCGEditorGraphProfilingView::TEXT_AvgExecutionTimeLabel)
-		.HAlignHeader(HAlign_Center)
-		.HAlignCell(HAlign_Right)
-		.SortMode(this, &SPCGEditorGraphProfilingView::GetColumnSortMode, PCGEditorGraphProfilingView::NAME_AvgExecutionTime)
-		.OnSort(this, &SPCGEditorGraphProfilingView::OnSortColumnHeader)
-		+ SHeaderRow::Column(PCGEditorGraphProfilingView::NAME_MinExecutionTime)
-		.ManualWidth(100)
-		.DefaultLabel(PCGEditorGraphProfilingView::TEXT_MinExecutionTimeLabel)
-		.HAlignHeader(HAlign_Center)
-		.HAlignCell(HAlign_Right)
-		.SortMode(this, &SPCGEditorGraphProfilingView::GetColumnSortMode, PCGEditorGraphProfilingView::NAME_MinExecutionTime)
-		.OnSort(this, &SPCGEditorGraphProfilingView::OnSortColumnHeader)
-		+ SHeaderRow::Column(PCGEditorGraphProfilingView::NAME_MaxExecutionTime)
-		.ManualWidth(100)
-		.DefaultLabel(PCGEditorGraphProfilingView::TEXT_MaxExecutionTimeLabel)
-		.HAlignHeader(HAlign_Center)
-		.HAlignCell(HAlign_Right)
-		.SortMode(this, &SPCGEditorGraphProfilingView::GetColumnSortMode, PCGEditorGraphProfilingView::NAME_MaxExecutionTime)
-		.OnSort(this, &SPCGEditorGraphProfilingView::OnSortColumnHeader)
-		+ SHeaderRow::Column(PCGEditorGraphProfilingView::NAME_StdExecutionTime)
-		.ManualWidth(100)
-		.DefaultLabel(PCGEditorGraphProfilingView::TEXT_StdExecutionTimeLabel)
-		.HAlignHeader(HAlign_Center)
-		.HAlignCell(HAlign_Right)
-		.SortMode(this, &SPCGEditorGraphProfilingView::GetColumnSortMode, PCGEditorGraphProfilingView::NAME_StdExecutionTime)
-		.OnSort(this, &SPCGEditorGraphProfilingView::OnSortColumnHeader)
 		+ SHeaderRow::Column(PCGEditorGraphProfilingView::NAME_TotalExecutionTime)
 		.ManualWidth(100)
 		.DefaultLabel(PCGEditorGraphProfilingView::TEXT_TotalExecutionTimeLabel)
@@ -343,6 +280,11 @@ TSharedRef<SHeaderRow> SPCGEditorGraphProfilingView::CreateHeaderRowWidget()
 		.OnSort(this, &SPCGEditorGraphProfilingView::OnSortColumnHeader);
 }
 
+FText SPCGEditorGraphProfilingView::GetTotalTimeLabel() const
+{
+	return FText::Format(LOCTEXT("GraphTotalTimeLabel", "{0} s"), TotalTime);
+}
+
 void SPCGEditorGraphProfilingView::OnSortColumnHeader(const EColumnSortPriority::Type SortPriority, const FName& ColumnId, const EColumnSortMode::Type NewSortMode)
 {
 	if (SortingColumn == ColumnId)
@@ -402,6 +344,50 @@ FReply SPCGEditorGraphProfilingView::ResetTimers()
 	return FReply::Handled();
 }
 
+namespace PCGEditorGraphProfilingView
+{
+	void AddListItems(
+		TArray<PCGProfilingListViewItemPtr>& OutListViewItems,
+		const TArray<PCGUtils::FCallTreeInfo>& TreeInfo,
+		const TMap<const UPCGNode*, UPCGEditorGraphNode*>& EditorNodeLookup,
+		bool bExpandSubgraph,
+		FString FolderName,
+		UPCGEditorGraphNode* CurrentEditorNode = nullptr)
+	{
+		for (const PCGUtils::FCallTreeInfo& Info : TreeInfo)
+		{
+			UPCGEditorGraphNode*const* EditorNodeItr = EditorNodeLookup.Find(Info.Node);
+			UPCGEditorGraphNode* EditorNode = EditorNodeItr ? *EditorNodeItr : CurrentEditorNode;
+
+			FString Fullname = FolderName;
+			if (Info.Node)
+			{
+				Fullname += Info.Node->GetNodeTitle().ToString();
+			}
+
+			// don't show inclusive times when bExpandSubgraph is on or the total sum of everything will mismatch reality
+			if ((Info.Children.IsEmpty() || bExpandSubgraph == false) && Info.CallTime.MaxExecutionFrameTime > 0)
+			{
+				PCGProfilingListViewItemPtr ListViewItem = MakeShared<FPCGProfilingListViewItem>();
+
+				ListViewItem->PCGNode = Info.Node;
+				ListViewItem->EditorNode = EditorNode;
+				ListViewItem->Name = Fullname;
+				ListViewItem->HasData = true;
+
+				ListViewItem->CallTime = Info.CallTime;
+
+				OutListViewItems.Add(ListViewItem);
+			}
+
+			if (bExpandSubgraph && !Info.Children.IsEmpty())
+			{
+				AddListItems(OutListViewItems, Info.Children, EditorNodeLookup, /* bExpandSubgraph */ true, Fullname + "/", EditorNode);
+			}
+		}
+	}
+} 
+
 FReply SPCGEditorGraphProfilingView::Refresh()
 {
 	ListViewItems.Empty();
@@ -418,89 +404,29 @@ FReply SPCGEditorGraphProfilingView::Refresh()
 		return FReply::Handled();
 	}
 
-	TArray<UPCGEditorGraphNode*> EditorNodes;
-	PCGEditorGraph->GetNodesOfClass<UPCGEditorGraphNode>(EditorNodes);
-	ListViewItems.Reserve(EditorNodes.Num());
-
 	const UPCGComponent* Component = PCGComponent.Get();
-
 	if (!Component)
 	{
 		return FReply::Handled();
 	}
 
-	for (const UPCGEditorGraphNode* PCGEditorNode : EditorNodes)
+	TArray<UPCGEditorGraphNode*> EditorNodes;
+	PCGEditorGraph->GetNodesOfClass<UPCGEditorGraphNode>(EditorNodes);
+
+	TMap<const UPCGNode*, UPCGEditorGraphNode*> EditorNodeLookup;
+	for (UPCGEditorGraphNode* EditorNode : EditorNodes)
 	{
-		const UPCGNode* PCGNode = PCGEditorNode ? PCGEditorNode->GetPCGNode() : nullptr;
-		if (PCGNode)
-		{
-			PCGProfilingListViewItemPtr ListViewItem = MakeShared<FPCGProfilingListViewItem>();
+		EditorNodeLookup.Add(EditorNode->GetPCGNode(), EditorNode);
+	}
 
-			// If component is partitioned, we will sum all data that makes sense to sum.
+	PCGUtils::FCallTreeInfo TreeInfo = Component->ExtraCapture.CalculateCallTreeInfo(Component);
 
-			if (Component->IsPartitioned())
-			{
-				if (UPCGSubsystem* Subsystem = Component->GetSubsystem())
-				{
-					ListViewItem->PCGNode = PCGNode;
-					ListViewItem->EditorNode = PCGEditorNode;
-					ListViewItem->Name = PCGNode->GetNodeTitle();
-					ListViewItem->HasData = true;
-					ListViewItem->StdExecutionTime = -1;
-
-					ListViewItem->MinExecutionTime = std::numeric_limits<double>::max();
-					ListViewItem->MaxExecutionTime = std::numeric_limits<double>::min();
-					ListViewItem->MinExecutionFrameTime = std::numeric_limits<double>::max();
-					ListViewItem->MaxExecutionFrameTime = std::numeric_limits<double>::min();
-
-					uint32 Count = 0;
-
-					Subsystem->ForAllRegisteredLocalComponents(const_cast<UPCGComponent*>(Component),
-						[this, PCGEditorNode, PCGNode, &ListViewItem, &Count](UPCGComponent* LocalComponent)
-					{
-						FPCGProfilingListViewItem TempItem;
-						FillItem(LocalComponent, PCGEditorNode, PCGNode, PCGNode->GetNodeTitle(), TempItem);
-
-						ListViewItem->ExecutionTime += TempItem.ExecutionTime;
-						ListViewItem->AvgExecutionTime += TempItem.AvgExecutionTime;
-						ListViewItem->NbCalls = FMath::Max(ListViewItem->NbCalls, TempItem.NbCalls);
-
-						ListViewItem->PrepareDataTime += TempItem.PrepareDataTime;
-						ListViewItem->PostExecuteTime += TempItem.PostExecuteTime;
-
-						ListViewItem->MinExecutionFrameTime = FMath::Min(ListViewItem->MinExecutionFrameTime, TempItem.MinExecutionFrameTime);
-						ListViewItem->MaxExecutionFrameTime = FMath::Max(ListViewItem->MaxExecutionFrameTime, TempItem.MaxExecutionFrameTime);
-
-						ListViewItem->MinExecutionTime = FMath::Min(ListViewItem->MinExecutionTime, TempItem.ExecutionTime);
-						ListViewItem->MaxExecutionTime = FMath::Max(ListViewItem->MaxExecutionTime, TempItem.ExecutionTime);
-
-						ListViewItem->NbExecutionFrames += TempItem.NbExecutionFrames;
-						ListViewItem->MaxNbExecutionFrames = FMath::Max(ListViewItem->MaxNbExecutionFrames, TempItem.MaxNbExecutionFrames);
-						ListViewItem->MinNbExecutionFrames = FMath::Min(ListViewItem->MinNbExecutionFrames, TempItem.MinNbExecutionFrames);
-
-						Count++;
-						return InvalidPCGTaskId;
-					});
-
-					// "Amortized" cost for mean an nb frames. Not exact because of multithread and some overhead.
-					if (Count > 0)
-					{
-						ListViewItem->AvgExecutionTime /= Count;
-						ListViewItem->NbExecutionFrames /= Count;
-					}
-
-					ListViewItems.Add(ListViewItem);
-				}
-			}
-			else
-			{
-				if (FillItem(Component, PCGEditorNode, PCGNode, PCGNode->GetNodeTitle(), *ListViewItem))
-				{
-					ListViewItems.Add(ListViewItem);
-				}
-			}
-		}
-	}
+	TotalTime = TreeInfo.CallTime.ExecutionTime;
+
+	ListViewItems.Reserve(TreeInfo.Children.Num());
+
+	//TODO: could turn this into a tree instead of expanding into a list
+	PCGEditorGraphProfilingView::AddListItems(ListViewItems, TreeInfo.Children, EditorNodeLookup, bExpandSubgraph, FString());
 
 	if (SortingColumn != NAME_None && SortMode != EColumnSortMode::None)
 	{
@@ -509,55 +435,27 @@ FReply SPCGEditorGraphProfilingView::Refresh()
 				bool isLess = false;
 				if (SortingColumn == PCGEditorGraphProfilingView::NAME_Node)
 				{
-					isLess = A->Name.FastLess(B->Name);
+					isLess = A->Name < B->Name;
 				}
 				else if (SortingColumn == PCGEditorGraphProfilingView::NAME_PrepareDataTime)
 				{
-					isLess = A->PrepareDataTime < B->PrepareDataTime;
-				}
-				else if (SortingColumn == PCGEditorGraphProfilingView::NAME_AvgExecutionTime)
-				{
-					isLess = A->AvgExecutionTime < B->AvgExecutionTime;
-				}
-				else if (SortingColumn == PCGEditorGraphProfilingView::NAME_MinExecutionTime)
-				{
-					isLess = A->MinExecutionTime < B->MinExecutionTime;
-				}
-				else if (SortingColumn == PCGEditorGraphProfilingView::NAME_MaxExecutionTime)
-				{
-					isLess = A->MaxExecutionTime < B->MaxExecutionTime;
+					isLess = A->CallTime.PrepareDataTime < B->CallTime.PrepareDataTime;
 				}
 				else if (SortingColumn == PCGEditorGraphProfilingView::NAME_MinExecutionFrameTime)
 				{
-					isLess = A->MinExecutionFrameTime < B->MinExecutionFrameTime;
+					isLess = A->CallTime.MinExecutionFrameTime < B->CallTime.MinExecutionFrameTime;
 				}
 				else if (SortingColumn == PCGEditorGraphProfilingView::NAME_MaxExecutionFrameTime)
 				{
-					isLess = A->MaxExecutionFrameTime < B->MaxExecutionFrameTime;
-				}
-				else if (SortingColumn == PCGEditorGraphProfilingView::NAME_StdExecutionTime)
-				{
-					isLess = A->StdExecutionTime < B->StdExecutionTime;
+					isLess = A->CallTime.MaxExecutionFrameTime < B->CallTime.MaxExecutionFrameTime;
 				}
 				else if (SortingColumn == PCGEditorGraphProfilingView::NAME_TotalExecutionTime)
 				{
-					isLess = A->ExecutionTime < B->ExecutionTime;
-				}
-				else if (SortingColumn == PCGEditorGraphProfilingView::NAME_NbCalls)
-				{
-					isLess = A->NbCalls < B->NbCalls;
+					isLess = A->CallTime.ExecutionTime < B->CallTime.ExecutionTime;
 				}
 				else if (SortingColumn == PCGEditorGraphProfilingView::NAME_NbExecutionFrames)
 				{
-					isLess = A->NbExecutionFrames < B->NbExecutionFrames;
-				}
-				else if (SortingColumn == PCGEditorGraphProfilingView::NAME_MaxNbExecutionFrames)
-				{
-					isLess = A->MaxNbExecutionFrames < B->MaxNbExecutionFrames;
-				}
-				else if (SortingColumn == PCGEditorGraphProfilingView::NAME_MinNbExecutionFrames)
-				{
-					isLess = A->MinNbExecutionFrames < B->MinNbExecutionFrames;
+					isLess = A->CallTime.ExecutionFrameCount < B->CallTime.ExecutionFrameCount;
 				}
 
 				return SortMode == EColumnSortMode::Ascending ? isLess : !isLess;
@@ -569,115 +467,6 @@ FReply SPCGEditorGraphProfilingView::Refresh()
 	return FReply::Handled();
 }
 
-bool SPCGEditorGraphProfilingView::FillItem(const UPCGComponent* InComponent, const UPCGEditorGraphNode* InPCGEditorNode, const UPCGNode* InPCGNode, const FName& InName, FPCGProfilingListViewItem& OutItem)
-{
-	const TSharedPtr<FPCGEditor> PCGEditor = PCGEditorPtr.Pin();
-
-	check(PCGEditor);
-
-	if (!InComponent || !InPCGEditorNode || !InPCGNode)
-	{
-		return false;
-	}
-
-	const TArray<PCGUtils::FCallTime>* TimersPtr = InComponent->ExtraCapture.GetTimers().Find(InPCGNode);
-	if (!TimersPtr)
-	{
-		return false;
-	}
-
-	OutItem.EditorNode = InPCGEditorNode;
-	OutItem.PCGNode = InPCGNode;
-
-	OutItem.Name = InName;
-	OutItem.MinExecutionTime = std::numeric_limits<double>::max();
-	OutItem.MaxExecutionTime = std::numeric_limits<double>::min();
-	OutItem.MinExecutionFrameTime = std::numeric_limits<double>::max();
-	OutItem.MaxExecutionFrameTime = std::numeric_limits<double>::min();
-
-	OutItem.NbCalls = TimersPtr->Num();
-
-	for (const PCGUtils::FCallTime& Timer : *TimersPtr)
-	{
-		OutItem.PrepareDataTime += Timer.PrepareDataTime;
-		OutItem.PostExecuteTime += Timer.PostExecuteTime;
-
-		OutItem.MinExecutionFrameTime = FMath::Min(OutItem.MinExecutionFrameTime, Timer.MinExecutionFrameTime);
-		OutItem.MaxExecutionFrameTime = FMath::Max(OutItem.MaxExecutionFrameTime, Timer.MaxExecutionFrameTime);
-
-		OutItem.ExecutionTime += Timer.ExecutionTime;
-		OutItem.MinExecutionTime = FMath::Min(OutItem.MinExecutionTime, Timer.ExecutionTime);
-		OutItem.MaxExecutionTime = FMath::Max(OutItem.MaxExecutionTime, Timer.ExecutionTime);
-
-		OutItem.NbExecutionFrames += Timer.ExecutionFrameCount;
-		OutItem.MaxNbExecutionFrames = FMath::Max(OutItem.MaxNbExecutionFrames, Timer.ExecutionFrameCount);
-		OutItem.MinNbExecutionFrames = FMath::Min(OutItem.MinNbExecutionFrames, Timer.ExecutionFrameCount);
-	}
-
-	OutItem.AvgExecutionTime = OutItem.ExecutionTime / OutItem.NbCalls;
-
-	for (const PCGUtils::FCallTime& Timer : *TimersPtr)
-	{
-		OutItem.StdExecutionTime += FMath::Square(OutItem.AvgExecutionTime - Timer.ExecutionTime);
-	}
-
-	OutItem.StdExecutionTime = FMath::Sqrt(OutItem.StdExecutionTime / OutItem.NbCalls);
-
-	OutItem.HasData = true;
-
-	// Also if the node is a subgraph, we need to add all stats from the child nodes
-	const UPCGSettings* Settings = InPCGNode->GetSettings();
-	if (Settings && Settings->IsA<UPCGBaseSubgraphSettings>())
-	{
-		if (const UPCGGraph* Graph = CastChecked<const UPCGBaseSubgraphSettings>(Settings)->GetSubgraph())
-		{
-			FString Prefix = InName.ToString() + "/";
-			for (const UPCGNode* ChildNode : Graph->GetNodes())
-			{
-				if (!ChildNode)
-				{
-					continue;
-				}
-
-				FPCGProfilingListViewItem TempItem;
-				PCGProfilingListViewItemPtr ChildItemPtr;
-
-				FPCGProfilingListViewItem* ChildItem = &TempItem;
-				if (bExpandSubgraph)
-				{
-					ChildItemPtr = MakeShared<FPCGProfilingListViewItem>();
-					ChildItem = ChildItemPtr.Get();
-				}
-
-				FName ChildName = FName(Prefix + ChildNode->GetNodeTitle().ToString());
-				if (FillItem(InComponent, InPCGEditorNode, ChildNode, ChildName, *ChildItem) && bExpandSubgraph)
-				{
-					ListViewItems.Add(ChildItemPtr);
-				}
-
-				OutItem.PrepareDataTime += ChildItem->PrepareDataTime;
-				OutItem.PostExecuteTime += ChildItem->PostExecuteTime;
-				OutItem.ExecutionTime += ChildItem->ExecutionTime;
-			}
-
-			OutItem.AvgExecutionTime = OutItem.ExecutionTime / OutItem.NbCalls;
-
-			// Other values can't be easily computed as we don't have enough granularity in subgraph scoped call.
-			// Setting them to -1 will hide them.
-			OutItem.MaxExecutionFrameTime = -1;
-			OutItem.MinExecutionFrameTime = -1;
-			OutItem.MinExecutionTime = -1;
-			OutItem.MaxExecutionTime = -1;
-			OutItem.StdExecutionTime = -1;
-			OutItem.MinNbExecutionFrames = -1;
-			OutItem.MaxNbExecutionFrames = -1;
-			OutItem.NbExecutionFrames = -1;
-		}
-	}
-
-	return true;
-}
-
 void SPCGEditorGraphProfilingView::OnDebugObjectChanged(UPCGComponent* InPCGComponent)
 {
 	if (PCGComponent.IsValid())
diff --git a/Engine/Plugins/Experimental/PCG/Source/PCGEditor/Private/SPCGEditorGraphProfilingView.h b/Engine/Plugins/Experimental/PCG/Source/PCGEditor/Private/SPCGEditorGraphProfilingView.h
--- a/Engine/Plugins/Experimental/PCG/Source/PCGEditor/Private/SPCGEditorGraphProfilingView.h
+++ b/Engine/Plugins/Experimental/PCG/Source/PCGEditor/Private/SPCGEditorGraphProfilingView.h
@@ -2,6 +2,7 @@
 
 #pragma once
 
+#include "Utils/PCGExtraCapture.h"
 #include "Widgets/Views/ITableRow.h"
 #include "Widgets/Views/STableRow.h"
 
@@ -18,20 +19,10 @@ struct FPCGProfilingListViewItem
 	const UPCGNode* PCGNode = nullptr;
 	const UPCGEditorGraphNode* EditorNode = nullptr;
 
-	FName Name = NAME_None;
-	double PrepareDataTime = 0.0;
-	double ExecutionTime = 0.0;
-	double AvgExecutionTime = 0.0;
-	double MinExecutionTime = 0.0;
-	double MaxExecutionTime = 0.0;
-	double MinExecutionFrameTime = 0.0;
-	double MaxExecutionFrameTime = 0.0;
-	double StdExecutionTime = 0.0;
-	double PostExecuteTime = 0.0;
-	int32 NbCalls = 0;
-	double NbExecutionFrames = 0;
-	int32 MaxNbExecutionFrames = 0;
-	int32 MinNbExecutionFrames = 0;
+	FString Name;
+
+	PCGUtils::FCallTime CallTime;
+
 	bool HasData = false;
 };
 
@@ -63,7 +54,6 @@ public:
 
 private:
 	TSharedRef<SHeaderRow> CreateHeaderRowWidget();
-	bool FillItem(const UPCGComponent* InComponent, const UPCGEditorGraphNode* InPCGEditorNode, const UPCGNode* InPCGNode, const FName& InName, FPCGProfilingListViewItem& OutItem);
 
 	ECheckBoxState IsSubgraphExpanded() const;
 	void OnSubgraphExpandedChanged(ECheckBoxState InNewState);
@@ -78,6 +68,7 @@ private:
 	void OnItemDoubleClicked(PCGProfilingListViewItemPtr Item);
 	void OnSortColumnHeader(const EColumnSortPriority::Type SortPriority, const FName& ColumnId, const EColumnSortMode::Type NewSortMode);
 	EColumnSortMode::Type GetColumnSortMode(const FName ColumnId) const;
+	FText GetTotalTimeLabel() const;
 
 	/** Pointer back to the PCG editor that owns us */
 	TWeakPtr<FPCGEditor> PCGEditorPtr;
@@ -96,5 +87,7 @@ private:
 	FName SortingColumn = NAME_None;
 	EColumnSortMode::Type SortMode = EColumnSortMode::Type::None;
 
+	double TotalTime = 0.0;
+
 	bool bExpandSubgraph = true;
 };
