Allow yet-unloaded but pending instances in local to original map in the PCG partition actors. It does not behave perfectly (will not clean up properly on deleted originals) but will fix loading issues in the meantime.

#rnx

diff --git a/Engine/Plugins/Experimental/PCG/Source/PCG/Private/Grid/PCGPartitionActor.cpp b/Engine/Plugins/Experimental/PCG/Source/PCG/Private/Grid/PCGPartitionActor.cpp
--- a/Engine/Plugins/Experimental/PCG/Source/PCG/Private/Grid/PCGPartitionActor.cpp
+++ b/Engine/Plugins/Experimental/PCG/Source/PCG/Private/Grid/PCGPartitionActor.cpp
@@ -60,7 +60,7 @@ void APCGPartitionActor::PostLoad()
 #endif
 
 	// Make sure that we don't track objects that do not exist anymore
-	CleanupDeadGraphInstances();
+	CleanupDeadGraphInstances(/*bRemoveNullOnly=*/true);
 
 #if WITH_EDITOR
 	// Mark all our local components as local
@@ -226,7 +226,7 @@ UPCGComponent* APCGPartitionActor::GetOriginalComponent(const UPCGComponent* Loc
 	return OriginalComponent ? (*OriginalComponent).Get() : nullptr;
 }
 
-void APCGPartitionActor::CleanupDeadGraphInstances()
+void APCGPartitionActor::CleanupDeadGraphInstances(bool bRemoveNonNullOnly)
 {
 	// First find if we have any local dead instance (= nullptr) hooked to an original component.
 	TSet<TSoftObjectPtr<UPCGComponent>> DeadOriginalInstances;
@@ -254,7 +254,8 @@ void APCGPartitionActor::CleanupDeadGraphInstances()
 	TSet<TObjectPtr<UPCGComponent>> DeadLocalInstances;
 	for (const auto& LocalToOriginalItem : LocalToOriginal)
 	{
-		if (!LocalToOriginalItem.Value.IsValid())
+		if((bRemoveNonNullOnly && LocalToOriginalItem.Value.IsNull()) ||
+			(!bRemoveNonNullOnly && !LocalToOriginalItem.Value.IsValid()))
 		{
 			DeadLocalInstances.Add(LocalToOriginalItem.Key);
 		}
diff --git a/Engine/Plugins/Experimental/PCG/Source/PCG/Private/Grid/PCGPartitionActor.h b/Engine/Plugins/Experimental/PCG/Source/PCG/Private/Grid/PCGPartitionActor.h
--- a/Engine/Plugins/Experimental/PCG/Source/PCG/Private/Grid/PCGPartitionActor.h
+++ b/Engine/Plugins/Experimental/PCG/Source/PCG/Private/Grid/PCGPartitionActor.h
@@ -57,7 +57,7 @@ public:
 	void AddGraphInstance(UPCGComponent* OriginalComponent);
 	void RemapGraphInstance(const UPCGComponent* OldOriginalComponent, UPCGComponent* NewOriginalComponent);
 	bool RemoveGraphInstance(UPCGComponent* OriginalComponent);
-	void CleanupDeadGraphInstances();
+	void CleanupDeadGraphInstances(bool bRemoveNonNullOnly = false);
 
 #if WITH_EDITOR
 	/** To be called after the creation of a new actor to copy the GridSize property (Editor only) into the PCGGridSize property */
