[PCG] Fixed an issue were the settings objects transactional flag was set too late causing default sub objects to be created without it. Also made the MeshSelectorInstance and InstancePackInstance transactable.

diff --git a/Engine/Plugins/Experimental/PCG/Source/PCG/Private/Elements/PCGPointMatchAndSet.cpp b/Engine/Plugins/Experimental/PCG/Source/PCG/Private/Elements/PCGPointMatchAndSet.cpp
--- a/Engine/Plugins/Experimental/PCG/Source/PCG/Private/Elements/PCGPointMatchAndSet.cpp
+++ b/Engine/Plugins/Experimental/PCG/Source/PCG/Private/Elements/PCGPointMatchAndSet.cpp
@@ -57,7 +57,12 @@ void UPCGPointMatchAndSetSettings::PostLoad()
 	if (!MatchAndSetInstance)
 	{
 		RefreshMatchAndSet();
-	}	
+	}
+	else
+	{
+		const EObjectFlags Flags = GetMaskedFlags(RF_PropagateToSubObjects) | RF_Transactional;
+		MatchAndSetInstance->SetFlags(Flags);
+	}
 }
 
 #if WITH_EDITOR
@@ -119,7 +124,8 @@ void UPCGPointMatchAndSetSettings::RefreshMatchAndSet()
 {
 	if (MatchAndSetType)
 	{
-		MatchAndSetInstance = NewObject<UPCGMatchAndSetBase>(this, MatchAndSetType);
+		const EObjectFlags Flags = GetMaskedFlags(RF_PropagateToSubObjects);
+		MatchAndSetInstance = NewObject<UPCGMatchAndSetBase>(this, MatchAndSetType, NAME_None, Flags);
 		check(MatchAndSetInstance);
 		MatchAndSetInstance->SetType(SetTargetType, SetTargetStringMode);
 	}
diff --git a/Engine/Plugins/Experimental/PCG/Source/PCG/Private/Elements/PCGStaticMeshSpawner.cpp b/Engine/Plugins/Experimental/PCG/Source/PCG/Private/Elements/PCGStaticMeshSpawner.cpp
--- a/Engine/Plugins/Experimental/PCG/Source/PCG/Private/Elements/PCGStaticMeshSpawner.cpp
+++ b/Engine/Plugins/Experimental/PCG/Source/PCG/Private/Elements/PCGStaticMeshSpawner.cpp
@@ -263,15 +263,25 @@ void UPCGStaticMeshSpawnerSettings::PostLoad()
 	}
 #endif
 
+	const EObjectFlags Flags = GetMaskedFlags(RF_PropagateToSubObjects) | RF_Transactional;
+	
 	if (!MeshSelectorInstance)
 	{
 		RefreshMeshSelector();
 	}
+	else
+	{
+		MeshSelectorInstance->SetFlags(Flags);
+	}
 
 	if (!InstancePackerInstance)
 	{
 		RefreshInstancePacker();
 	}
+	else
+	{
+		InstancePackerInstance->SetFlags(Flags);
+	}
 }
 
 #if WITH_EDITOR
@@ -325,7 +335,8 @@ void UPCGStaticMeshSpawnerSettings::RefreshMeshSelector()
 {
 	if (MeshSelectorType)
 	{
-		MeshSelectorInstance = NewObject<UPCGMeshSelectorBase>(this, MeshSelectorType);
+		const EObjectFlags Flags = GetMaskedFlags(RF_PropagateToSubObjects);
+		MeshSelectorInstance = NewObject<UPCGMeshSelectorBase>(this, MeshSelectorType, NAME_None, Flags);
 	}
 	else
 	{
@@ -337,7 +348,8 @@ void UPCGStaticMeshSpawnerSettings::RefreshInstancePacker()
 {
 	if (InstancePackerType)
 	{
-		InstancePackerInstance = NewObject<UPCGInstancePackerBase>(this, InstancePackerType);
+		const EObjectFlags Flags = GetMaskedFlags(RF_PropagateToSubObjects);
+		InstancePackerInstance = NewObject<UPCGInstancePackerBase>(this, InstancePackerType, NAME_None, Flags);
 	}
 	else
 	{
diff --git a/Engine/Plugins/Experimental/PCG/Source/PCG/Private/PCGGraph.cpp b/Engine/Plugins/Experimental/PCG/Source/PCG/Private/PCGGraph.cpp
--- a/Engine/Plugins/Experimental/PCG/Source/PCG/Private/PCGGraph.cpp
+++ b/Engine/Plugins/Experimental/PCG/Source/PCG/Private/PCGGraph.cpp
@@ -178,7 +178,7 @@ void UPCGGraph::BeginDestroy()
 
 UPCGNode* UPCGGraph::AddNodeOfType(TSubclassOf<class UPCGSettings> InSettingsClass, UPCGSettings*& OutDefaultNodeSettings)
 {
-	UPCGSettings* Settings = NewObject<UPCGSettings>(GetTransientPackage(), InSettingsClass);
+	UPCGSettings* Settings = NewObject<UPCGSettings>(GetTransientPackage(), InSettingsClass, NAME_None, RF_Transactional);
 
 	if (!Settings)
 	{
@@ -190,7 +190,6 @@ UPCGNode* UPCGGraph::AddNodeOfType(TSubclassOf<class UPCGSettings> InSettingsCla
 	if (Node)
 	{
 		Settings->Rename(nullptr, Node);
-		Settings->SetFlags(RF_Transactional);
 	}
 
 	OutDefaultNodeSettings = Settings;
@@ -268,7 +267,6 @@ UPCGNode* UPCGGraph::AddNodeCopy(UPCGSettings* InSettings, UPCGSettings*& Defaul
 	if (SettingsCopy)
 	{
 		SettingsCopy->Rename(nullptr, NewNode);
-		SettingsCopy->SetFlags(RF_Transactional);
 	}
 
 	DefaultNodeSettings = SettingsCopy;
