Improved sampling using volumes.
Fixed color averaging in point data sampling/projection.

#rnx

diff --git a/Engine/Plugins/Experimental/PCG/Source/PCG/Private/Data/PCGPointData.cpp b/Engine/Plugins/Experimental/PCG/Source/PCG/Private/Data/PCGPointData.cpp
--- a/Engine/Plugins/Experimental/PCG/Source/PCG/Private/Data/PCGPointData.cpp
+++ b/Engine/Plugins/Experimental/PCG/Source/PCG/Private/Data/PCGPointData.cpp
@@ -42,7 +42,7 @@ namespace PCGPointHelpers
 		return true;
 	}
 
-	float ManhattanDensity(const FPCGPoint& InPoint, const FVector& InPosition)
+	FVector::FReal ManhattanDensity(const FPCGPoint& InPoint, const FVector& InPosition)
 	{
 		FVector Ratios;
 		if (GetDistanceRatios(InPoint, InPosition, Ratios))
@@ -55,12 +55,12 @@ namespace PCGPointHelpers
 		}
 	}
 
-	float InverseEuclidianDistance(const FPCGPoint& InPoint, const FVector& InPosition)
+	FVector::FReal InverseEuclidianDistance(const FPCGPoint& InPoint, const FVector& InPosition)
 	{
 		FVector Ratios;
 		if (GetDistanceRatios(InPoint, InPosition, Ratios))
 		{
-			return 1 - Ratios.Length();
+			return 1.0 - Ratios.Length();
 		}
 		else
 		{
@@ -72,22 +72,22 @@ namespace PCGPointHelpers
 	* Note that this assumes that either data set is homogeneous in its points dimension (either 0d, 1d, 2d, 3d)
 	* Otherwise there will be some artifacts from our assumption here (namely using a 1.0 value for the additional coordinates).
 	*/
-	float ComputeOverlapRatio(const FBox& Numerator, const FBox& Denominator)
+	FVector::FReal ComputeOverlapRatio(const FBox& Numerator, const FBox& Denominator)
 	{
 		const FVector NumeratorExtent = Numerator.GetExtent();
 		const FVector DenominatorExtent = Denominator.GetExtent();
 
-		return (float)((DenominatorExtent.X > 0 ? NumeratorExtent.X / DenominatorExtent.X : 1.0) *
+		return (FVector::FReal)((DenominatorExtent.X > 0 ? NumeratorExtent.X / DenominatorExtent.X : 1.0) *
 			(DenominatorExtent.Y > 0 ? NumeratorExtent.Y / DenominatorExtent.Y : 1.0) *
 			(DenominatorExtent.Z > 0 ? NumeratorExtent.Z / DenominatorExtent.Z : 1.0));
 	}
 
-	float VolumeOverlap(const FPCGPoint& InPoint, const FBox& InBounds, const FTransform& InTransform)
+	FVector::FReal VolumeOverlap(const FPCGPoint& InPoint, const FBox& InBounds, const FTransform& InTransform)
 	{
 		// This is similar in idea to SAT considering we have two boxes - since we will test all 6 axes.
 		// However, there is some uncertainty due to rotation, and using the overlap value as-is is an overestimation, which might not be critical in this case
 		// TODO: investigate if we should do a 8-pt test instead (would be more precise, but significantly more costly).
-		const FBox PointBounds = InPoint.GetLocalBounds();
+		const FBox PointBounds = InPoint.GetLocalDensityBounds();
 		const FTransform& PointTransform = InPoint.Transform;
 
 		const FBox FirstOverlap = PointBounds.Overlap(InBounds.TransformBy(InTransform.GetRelativeTransform(PointTransform)));
@@ -102,7 +102,7 @@ namespace PCGPointHelpers
 			return 0;
 		}
 		
-		return FMath::Min(ComputeOverlapRatio(FirstOverlap, PointBounds), ComputeOverlapRatio(SecondOverlap, InBounds));
+		return FMath::Min(ComputeOverlapRatio(FirstOverlap, InBounds), ComputeOverlapRatio(SecondOverlap, InBounds));
 	}
 
 	/** Helper function for additive blending of quaternions (copied from ControlRig) */
@@ -350,7 +350,7 @@ bool UPCGPointData::ProjectPoint(const FTransform& InTransform, const FBox& InBo
 		RebuildOctree();
 	}
 
-	TArray<TPair<const FPCGPoint*, float>, TInlineAllocator<4>> Contributions;
+	TArray<TPair<const FPCGPoint*, FVector::FReal>, TInlineAllocator<4>> Contributions;
 	const bool bSampleInVolume = (InBounds.GetExtent() != FVector::ZeroVector);
 
 	if (!bSampleInVolume)
@@ -372,11 +372,11 @@ bool UPCGPointData::ProjectPoint(const FTransform& InTransform, const FBox& InBo
 		});
 	}
 
-	float SumContributions = 0;
-	float MaxContribution = 0;
+	FVector::FReal SumContributions = 0;
+	FVector::FReal MaxContribution = 0;
 	const FPCGPoint* MaxContributor = nullptr;
 
-	for (const TPair<const FPCGPoint*, float>& Contribution : Contributions)
+	for (const TPair<const FPCGPoint*, FVector::FReal>& Contribution : Contributions)
 	{
 		SumContributions += Contribution.Value;
 
@@ -392,6 +392,15 @@ bool UPCGPointData::ProjectPoint(const FTransform& InTransform, const FBox& InBo
 		return false;
 	}
 
+	// Rationale: 
+	// When doing volume-to-volume intersection, we want the final density to reflect the amount of overlap
+	// if any - hence the volume overlap computation before.
+	// But, considering that some points may/will overlap (incl. due to steepness), we want to make sure we do not
+	// sum up to more than the total volume. 
+	// Note that this might create some artifacts on the edges in some instances, but we will revisit this once we have a
+	// better and sufficiently efficient solution.
+	const FVector::FReal DensityNormalizationFactor = ((SumContributions > 1.0) ? (1.0 / SumContributions) : 1.0);
+
 	TArray<TPair<const FPCGPoint*, float>, TInlineAllocator<4>> ContributionsForMetadata;
 
 	// Computed weighted average of spatial properties
@@ -401,13 +410,13 @@ bool UPCGPointData::ProjectPoint(const FTransform& InTransform, const FBox& InBo
 	float WeightedDensity = 0;
 	FVector WeightedBoundsMin = FVector::ZeroVector;
 	FVector WeightedBoundsMax = FVector::ZeroVector;
-	FVector WeightedColor = FVector::ZeroVector;
+	FVector4 WeightedColor = FVector4::Zero();
 	float WeightedSteepness = 0;
 
-	for (const TPair<const FPCGPoint*, float> Contribution : Contributions)
+	for (const TPair<const FPCGPoint*, FVector::FReal> Contribution : Contributions)
 	{
 		const FPCGPoint& SourcePoint = *Contribution.Key;
-		const float Weight = Contribution.Value / SumContributions;
+		const FVector::FReal Weight = Contribution.Value / SumContributions;
 
 		WeightedPosition += SourcePoint.Transform.GetLocation() * Weight;
 		WeightedQuat = PCGPointHelpers::AddQuatWithWeight(WeightedQuat, SourcePoint.Transform.GetRotation(), Weight);
@@ -419,7 +428,7 @@ bool UPCGPointData::ProjectPoint(const FTransform& InTransform, const FBox& InBo
 		}
 		else
 		{
-			WeightedDensity += SourcePoint.Density * Weight * Contribution.Value;
+			WeightedDensity += SourcePoint.Density * Contribution.Value * DensityNormalizationFactor;
 		}
 
 		WeightedBoundsMin += SourcePoint.BoundsMin * Weight;
@@ -427,11 +436,10 @@ bool UPCGPointData::ProjectPoint(const FTransform& InTransform, const FBox& InBo
 		WeightedColor += SourcePoint.Color * Weight;
 		WeightedSteepness += SourcePoint.Steepness * Weight;
 
-		ContributionsForMetadata.Emplace(Contribution.Key, Weight);
+		ContributionsForMetadata.Emplace(Contribution.Key, static_cast<float>(Weight));
 	}
 
-	// Finally, apply changes to point
-	
+	// Finally, apply changes to point, based on the projection settings
 	if (InParams.bProjectPositions)
 	{
 		OutPoint.Transform.SetLocation(bSampleInVolume ? WeightedPosition : InTransform.GetLocation());
diff --git a/Engine/Plugins/Experimental/PCG/Source/PCG/Private/PCGPoint.cpp b/Engine/Plugins/Experimental/PCG/Source/PCG/Private/PCGPoint.cpp
--- a/Engine/Plugins/Experimental/PCG/Source/PCG/Private/PCGPoint.cpp
+++ b/Engine/Plugins/Experimental/PCG/Source/PCG/Private/PCGPoint.cpp
@@ -29,6 +29,11 @@ FBox FPCGPoint::GetLocalBounds() const
 	return FBox(BoundsMin, BoundsMax);
 }
 
+FBox FPCGPoint::GetLocalDensityBounds() const
+{
+	return FBox((2 - Steepness) * BoundsMin, (2 - Steepness) * BoundsMax);
+}
+
 void FPCGPoint::SetLocalBounds(const FBox& InBounds)
 {
 	BoundsMin = InBounds.Min;
@@ -37,7 +42,7 @@ void FPCGPoint::SetLocalBounds(const FBox& InBounds)
 
 FBoxSphereBounds FPCGPoint::GetDensityBounds() const
 {
-	return FBoxSphereBounds(FBox((2 - Steepness) * BoundsMin, (2 - Steepness) * BoundsMax).TransformBy(Transform));
+	return FBoxSphereBounds(GetLocalDensityBounds().TransformBy(Transform));
 }
 
 bool FPCGPoint::HasCustomPropertyGetterSetter(FName Name)
diff --git a/Engine/Plugins/Experimental/PCG/Source/PCG/Public/PCGPoint.h b/Engine/Plugins/Experimental/PCG/Source/PCG/Public/PCGPoint.h
--- a/Engine/Plugins/Experimental/PCG/Source/PCG/Public/PCGPoint.h
+++ b/Engine/Plugins/Experimental/PCG/Source/PCG/Public/PCGPoint.h
@@ -33,6 +33,7 @@ public:
 	FPCGPoint(const FTransform& InTransform, float InDensity, int32 InSeed);
 
 	FBox GetLocalBounds() const;
+	FBox GetLocalDensityBounds() const;
 	void SetLocalBounds(const FBox& InBounds);
 	FBoxSphereBounds GetDensityBounds() const;
 
