[PCG] Introduce Crc-based system for element caching

Objects of type inheriting from PCGData have a globally unique ID. This allows us to unambiguously identify an instance of data and we use that for Crcs to identify data. As next steps, doing a more finegrained Crc would be much more accurate but this should have parity with previous behaviour for now.

A Dependency Crc is computed from all the inputs and settings to a node. This is used as the "key" to lookup/store node output in the cache.

If a node is cacheable, use that Dependency Crc as its output data Crc passed downstream. If a node is not cacheable, compute a Crc from its output data using the UID for now.

Choosing what to include/exclude from the dependency Crc allows fine-grained cache invalidation and will allow to ignore data from unused pins or unused settings or etc (simple example included in this CL in FPCGTextureSamplerElement.cpp - the data output from this element will depend on the PCG actor transform if the bAbsoluteTransform==false).

Also, added debug tooling toggled on CVar. Log has trace of cache actions and execution actions. Crc is displayed in Inspection window. Nodes display the number of their cache entries in the titles so the change propagation can be inspected quickly at graph level.

diff --git a/Engine/Plugins/Experimental/PCG/Source/PCG/Private/Elements/PCGSurfaceSampler.cpp b/Engine/Plugins/Experimental/PCG/Source/PCG/Private/Elements/PCGSurfaceSampler.cpp
--- a/Engine/Plugins/Experimental/PCG/Source/PCG/Private/Elements/PCGSurfaceSampler.cpp
+++ b/Engine/Plugins/Experimental/PCG/Source/PCG/Private/Elements/PCGSurfaceSampler.cpp
@@ -396,6 +396,31 @@ bool FPCGSurfaceSamplerElement::ExecuteInternal(FPCGContext* Context) const
 	return true;
 }
 
+void FPCGSurfaceSamplerElement::GetDependenciesCrc(const FPCGDataCollection& InInput, const UPCGSettings* InSettings, UPCGComponent* InComponent, FPCGCrc& OutCrc) const
+{
+	FPCGCrc Crc;
+	FSimplePCGElement::GetDependenciesCrc(InInput, InSettings, InComponent, Crc);
+
+	if (const UPCGSurfaceSamplerSettings* Settings = Cast<UPCGSurfaceSamplerSettings>(InSettings))
+	{
+		const UPCGParamData* Params = InInput.GetParams();
+		const bool bUnbounded = PCG_GET_OVERRIDEN_VALUE(Settings, bUnbounded, Params);
+		const bool bBoundsConnected = InInput.GetInputsByPin(PCGSurfaceSamplerConstants::BoundingShapeLabel).Num() > 0;
+
+		// If we're operating in bounded mode and there is no bounding shape connected then we'll use actor bounds, and therefore take
+		// dependency on actor data.
+		if (!bUnbounded && !bBoundsConnected && InComponent)
+		{
+			if (const UPCGData* Data = InComponent->GetActorPCGData())
+			{
+				Crc.Combine(Data->ComputeCrc());
+			}
+		}
+	}
+
+	OutCrc = Crc;
+}
+
 #if WITH_EDITOR
 void UPCGSurfaceSamplerSettings::ApplyDeprecationBeforeUpdatePins(UPCGNode* InOutNode, TArray<TObjectPtr<UPCGPin>>& InputPins, TArray<TObjectPtr<UPCGPin>>& OutputPins)
 {
diff --git a/Engine/Plugins/Experimental/PCG/Source/PCG/Private/Elements/PCGTextureSampler.cpp b/Engine/Plugins/Experimental/PCG/Source/PCG/Private/Elements/PCGTextureSampler.cpp
--- a/Engine/Plugins/Experimental/PCG/Source/PCG/Private/Elements/PCGTextureSampler.cpp
+++ b/Engine/Plugins/Experimental/PCG/Source/PCG/Private/Elements/PCGTextureSampler.cpp
@@ -2,14 +2,16 @@
 
 #include "Elements/PCGTextureSampler.h"
 
-#include "Engine/Texture2D.h"
-#include "GameFramework/Actor.h"
 #include "PCGComponent.h"
 #include "PCGContext.h"
+#include "PCGCrc.h"
 #include "PCGHelpers.h"
 #include "Helpers/PCGBlueprintHelpers.h"
 #include "Helpers/PCGSettingsHelpers.h"
 
+#include "Engine/Texture2D.h"
+#include "GameFramework/Actor.h"
+
 #include UE_INLINE_GENERATED_CPP_BY_NAME(PCGTextureSampler)
 
 FPCGElementPtr UPCGTextureSamplerSettings::CreateElement() const
@@ -25,14 +27,6 @@ TArray<FPCGPinProperties> UPCGTextureSamplerSettings::InputPinProperties() const
 	return PinProperties;
 }
 
-bool FPCGTextureSamplerElement::IsCacheable(const UPCGSettings* InSettings) const
-{
-	// It is currently possible to cache the texture sampler only if using an absolute transform
-	// since otherwise it depends on the source component.
-	// However, more tricky here is that the bUseAbsoluteTransform can be overriden, in which case we can't take a decision
-	return false;
-}
-
 bool FPCGTextureSamplerElement::ExecuteInternal(FPCGContext* Context) const
 {
 	TRACE_CPUPROFILER_EVENT_SCOPE(FPCGTextureSamplerElement::Execute);
@@ -92,3 +86,25 @@ bool FPCGTextureSamplerElement::ExecuteInternal(FPCGContext* Context) const
 
 	return true;
 }
+
+void FPCGTextureSamplerElement::GetDependenciesCrc(const FPCGDataCollection& InInput, const UPCGSettings* InSettings, UPCGComponent* InComponent, FPCGCrc& OutCrc) const
+{
+	FPCGCrc Crc;
+	FSimplePCGElement::GetDependenciesCrc(InInput, InSettings, InComponent, Crc);
+
+	if (const UPCGTextureSamplerSettings* Settings = Cast<UPCGTextureSamplerSettings>(InSettings))
+	{
+		// If not using absolute transform, depend on actor transform and bounds, and therefore take dependency on actor data.
+		const UPCGParamData* Params = InInput.GetParams();
+		const bool bUseAbsoluteTransform = PCG_GET_OVERRIDEN_VALUE(Settings, bUseAbsoluteTransform, Params);
+		if (!bUseAbsoluteTransform && InComponent)
+		{
+			if (const UPCGData* Data = InComponent->GetActorPCGData())
+			{
+				Crc.Combine(Data->ComputeCrc());
+			}
+		}
+	}
+
+	OutCrc = Crc;
+}
diff --git a/Engine/Plugins/Experimental/PCG/Source/PCG/Private/Graph/PCGGraphCache.cpp b/Engine/Plugins/Experimental/PCG/Source/PCG/Private/Graph/PCGGraphCache.cpp
--- a/Engine/Plugins/Experimental/PCG/Source/PCG/Private/Graph/PCGGraphCache.cpp
+++ b/Engine/Plugins/Experimental/PCG/Source/PCG/Private/Graph/PCGGraphCache.cpp
@@ -1,44 +1,26 @@
 // Copyright Epic Games, Inc. All Rights Reserved.
 
-#include "PCGGraphCache.h"
+#include "Graph/PCGGraphCache.h"
+
 #include "PCGComponent.h"
+#include "PCGModule.h"
 
+#include "HAL/IConsoleManager.h"
+#include "GameFramework/Actor.h"
 #include "Misc/ScopeRWLock.h"
 
-namespace PCGGraphCache
-{
-	constexpr int32 NullComponentSeed = 0;
-	constexpr int32 NullSettingsCrc32 = 0;
-
-	// Component seed is only get if we have a component and either no settings or the settings use seed.
-	// TODO: Perhaps we should not get the component seed if we have no settings...
-	int32 GetComponentSeed(const UPCGSettings* InSettings, const UPCGComponent* InComponent)
-	{
-		if ((!InSettings || InSettings->UseSeed()) && InComponent)
-		{
-			return InComponent->Seed;
-		}
+static TAutoConsoleVariable<bool> CVarCacheDebugging(
+	TEXT("pcg.CacheDebugging"),
+	false,
+	TEXT("Enable various features for debugging the graph cache system."));
 
-		return PCGGraphCache::NullComponentSeed;
-	}
-}
-
-FPCGGraphCacheEntry::FPCGGraphCacheEntry(const FPCGDataCollection& InInput, const UPCGSettings* InSettings, const UPCGComponent* InComponent, const FPCGDataCollection& InOutput, TWeakObjectPtr<UObject> InOwner, FPCGRootSet& OutRootSet)
-	: Input(InInput)
-	, Output(InOutput)
+FPCGGraphCacheEntry::FPCGGraphCacheEntry(const FPCGCrc& InDependenciesCrc, const FPCGDataCollection& InOutput, FPCGRootSet& OutRootSet)
+	: Output(InOutput)
+	, DependenciesCrc(InDependenciesCrc)
 {
-	SettingsCrc32 = InSettings ? InSettings->GetCrc32() : PCGGraphCache::NullSettingsCrc32;
-	ComponentSeed = PCGGraphCache::GetComponentSeed(InSettings, InComponent);
-
-	Input.AddToRootSet(OutRootSet);
 	Output.AddToRootSet(OutRootSet);
 }
 
-bool FPCGGraphCacheEntry::Matches(const FPCGDataCollection& InInput, int32 InSettingsCrc32, int32 InComponentSeed) const
-{
-	return (SettingsCrc32 == InSettingsCrc32) && (Input == InInput) && (ComponentSeed == InComponentSeed);
-}
-
 FPCGGraphCache::FPCGGraphCache(TWeakObjectPtr<UObject> InOwner, FPCGRootSet* InRootSet)
 	: Owner(InOwner), RootSet(InRootSet)
 {
@@ -50,39 +32,61 @@ FPCGGraphCache::~FPCGGraphCache()
 	ClearCache();
 }
 
-bool FPCGGraphCache::GetFromCache(const IPCGElement* InElement, const FPCGDataCollection& InInput, const UPCGSettings* InSettings, const UPCGComponent* InComponent, FPCGDataCollection& OutOutput) const
+bool FPCGGraphCache::GetFromCache(const UPCGNode* InNode, const IPCGElement* InElement, const FPCGCrc& InDependenciesCrc, const FPCGDataCollection& InInput, const UPCGSettings* InSettings, const UPCGComponent* InComponent, FPCGDataCollection& OutOutput) const
 {
 	if (!Owner.IsValid())
 	{
 		return false;
 	}
 
+	if(!InDependenciesCrc.IsValid())
+	{
+		UE_LOG(LogPCG, Warning, TEXT("Invalid dependencies passed to FPCGGraphCache::GetFromCache(), lookup aborted."));
+		return false;
+	}
+
+	const bool bDebuggingEnabled = IsDebuggingEnabled() && InComponent && InComponent->GetOwner() && InNode;
+
 	TRACE_CPUPROFILER_EVENT_SCOPE(FPCGGraphCache::GetFromCache);
 	FReadScopeLock ScopedReadLock(CacheLock);
 
 	if (const FPCGGraphCacheEntries* Entries = CacheData.Find(InElement))
 	{
-		int32 InSettingsCrc32 = (InSettings ? InSettings->GetCrc32() : PCGGraphCache::NullSettingsCrc32);
-		int32 InComponentSeed = PCGGraphCache::GetComponentSeed(InSettings, InComponent);
-
 		for (const FPCGGraphCacheEntry& Entry : *Entries)
 		{
-			if (Entry.Matches(InInput, InSettingsCrc32, InComponentSeed))
+			if (Entry.DependenciesCrc == InDependenciesCrc)
 			{
 				OutOutput = Entry.Output;
+
+				if (bDebuggingEnabled)
+				{
+					// Leading spaces to align log content with warnings below - helps readability a lot.
+					UE_LOG(LogPCG, Log, TEXT("         [%s] %s\t\tCACHE HIT %u"), *InComponent->GetOwner()->GetName(), *InNode->GetNodeTitle().ToString(), InDependenciesCrc.GetValue());
+				}
+
 				return true;
 			}
 		}
 
+		if (bDebuggingEnabled)
+		{
+			UE_LOG(LogPCG, Warning, TEXT("[%s] %s\t\tCACHE MISS %u"), *InComponent->GetOwner()->GetName() , *InNode->GetNodeTitle().ToString(), InDependenciesCrc.GetValue());
+		}
+
 		return false;
 	}
 
+	if (bDebuggingEnabled)
+	{
+		UE_LOG(LogPCG, Warning, TEXT("[%s] %s\t\tCACHE MISS NOELEMENT"), *InComponent->GetOwner()->GetName() , *InNode->GetNodeTitle().ToString());
+	}
+
 	return false;
 }
 
-void FPCGGraphCache::StoreInCache(const IPCGElement* InElement, const FPCGDataCollection& InInput, const UPCGSettings* InSettings, const UPCGComponent* InComponent, const FPCGDataCollection& InOutput)
+void FPCGGraphCache::StoreInCache(const IPCGElement* InElement, const FPCGCrc& InDependenciesCrc, const FPCGDataCollection& InInput, const UPCGSettings* InSettings, const UPCGComponent* InComponent, const FPCGDataCollection& InOutput)
 {
-	if (!Owner.IsValid())
+	if (!Owner.IsValid() || !ensure(InDependenciesCrc.IsValid()))
 	{
 		return;
 	}
@@ -96,7 +100,7 @@ void FPCGGraphCache::StoreInCache(const IPCGElement* InElement, const FPCGDataCo
 		Entries = &(CacheData.Add(InElement));
 	}
 
-	Entries->Emplace(InInput, InSettings, InComponent, InOutput, Owner, *RootSet);
+	Entries->Emplace(InDependenciesCrc, InOutput, *RootSet);
 }
 
 void FPCGGraphCache::ClearCache()
@@ -108,7 +112,6 @@ void FPCGGraphCache::ClearCache()
 	{
 		for (FPCGGraphCacheEntry& Entry : CacheEntry.Value)
 		{
-			Entry.Input.RemoveFromRootSet(*RootSet);
 			Entry.Output.RemoveFromRootSet(*RootSet);
 		}
 	}
@@ -118,20 +121,24 @@ void FPCGGraphCache::ClearCache()
 }
 
 #if WITH_EDITOR
-void FPCGGraphCache::CleanFromCache(const IPCGElement* InElement)
+void FPCGGraphCache::CleanFromCache(const IPCGElement* InElement, const UPCGSettings* InSettings/*= nullptr*/)
 {
 	if (!InElement)
 	{
 		return;
 	}
 
+	if (IsDebuggingEnabled() && InSettings)
+	{
+		UE_LOG(LogPCG, Warning, TEXT("CACHE: PURGED [%s]"), *InSettings->GetDefaultNodeName().ToString());
+	}
+
 	FWriteScopeLock ScopeWriteLock(CacheLock);
 	FPCGGraphCacheEntries* Entries = CacheData.Find(InElement);
 	if (Entries)
 	{
 		for (FPCGGraphCacheEntry& Entry : *Entries)
 		{
-			Entry.Input.RemoveFromRootSet(*RootSet);
 			Entry.Output.RemoveFromRootSet(*RootSet);
 		}
 	}
@@ -139,4 +146,22 @@ void FPCGGraphCache::CleanFromCache(const IPCGElement* InElement)
 	// Finally, remove all entries matching that element
 	CacheData.Remove(InElement);
 }
+
+uint32 FPCGGraphCache::GetGraphCacheEntryCount(IPCGElement* InElement) const
+{
+	TRACE_CPUPROFILER_EVENT_SCOPE(FPCGGraphCache::GetFromCache);
+	FReadScopeLock ScopedReadLock(CacheLock);
+
+	if (const FPCGGraphCacheEntries* Entries = CacheData.Find(InElement))
+	{
+		return Entries->Num();
+	}
+
+	return 0;
+}
 #endif // WITH_EDITOR
+
+bool FPCGGraphCache::IsDebuggingEnabled() const
+{
+	return CVarCacheDebugging.GetValueOnAnyThread();
+}
diff --git a/Engine/Plugins/Experimental/PCG/Source/PCG/Private/Graph/PCGGraphCache.h b/Engine/Plugins/Experimental/PCG/Source/PCG/Private/Graph/PCGGraphCache.h
--- a/Engine/Plugins/Experimental/PCG/Source/PCG/Private/Graph/PCGGraphCache.h
+++ b/Engine/Plugins/Experimental/PCG/Source/PCG/Private/Graph/PCGGraphCache.h
@@ -2,25 +2,24 @@
 
 #pragma once
 
+#include "PCGCrc.h"
 #include "PCGData.h"
 
 class IPCGElement;
-
-class UPCGSettings;
 class UPCGComponent;
+class UPCGNode;
+class UPCGSettings;
 class AActor;
 
 struct FPCGGraphCacheEntry
 {
 	FPCGGraphCacheEntry() = default;
-	FPCGGraphCacheEntry(const FPCGDataCollection& InInput, const UPCGSettings* InSettings, const UPCGComponent* InComponent, const FPCGDataCollection& InOutput, TWeakObjectPtr<UObject> InOwner, FPCGRootSet& OutRootSet);
-
-	bool Matches(const FPCGDataCollection& InInput, int32 InSettingsCrc32, int32 InComponentSeed) const;
+	FPCGGraphCacheEntry(const FPCGCrc& InCrc, const FPCGDataCollection& InOutput, FPCGRootSet& OutRootSet);
 
-	FPCGDataCollection Input;
 	FPCGDataCollection Output;
-	int32 SettingsCrc32;
-	int32 ComponentSeed;
+
+	/** A Crc value that encapsulates all information that can affect the result of an element. */
+	FPCGCrc DependenciesCrc;
 };
 
 // TODO: investigate if we need a more evolved data structure here
@@ -36,18 +35,24 @@ struct FPCGGraphCache
 	FPCGGraphCache(TWeakObjectPtr<UObject> InOwner, FPCGRootSet* InRootSet);
 	~FPCGGraphCache();
 
-	/** Returns true if data was found from the cache, in which case the outputs are written in OutOutput */
-	bool GetFromCache(const IPCGElement* InElement, const FPCGDataCollection& InInput, const UPCGSettings* InSettings, const UPCGComponent* InComponent, FPCGDataCollection& OutOutput) const;
+	/** Returns true if data was found from the cache, in which case the outputs are written in OutOutput. InNode is optional and for logging only. */
+	bool GetFromCache(const UPCGNode* InNode, const IPCGElement* InElement, const FPCGCrc& InCrc, const FPCGDataCollection& InInput, const UPCGSettings* InSettings, const UPCGComponent* InComponent, FPCGDataCollection& OutOutput) const;
 
 	/** Stores data in the cache for later use */
-	void StoreInCache(const IPCGElement* InElement, const FPCGDataCollection& InInput, const UPCGSettings* InSettings, const UPCGComponent* InComponent, const FPCGDataCollection& InOutput);
+	void StoreInCache(const IPCGElement* InElement, const FPCGCrc& InCrc, const FPCGDataCollection& InInput, const UPCGSettings* InSettings, const UPCGComponent* InComponent, const FPCGDataCollection& InOutput);
 
 	/** Removes all entries from the cache, unroots data, etc. */
 	void ClearCache();
 
+	/** True if debugging features enabled. Exposes a CVar so it can be called from editor module. */
+	bool IsDebuggingEnabled() const;
+
 #if WITH_EDITOR
-	/** Clears any cache entry for the given data */
-	void CleanFromCache(const IPCGElement* InElement);
+	/** Clears any cache entry for the given data. InSettings is optional and for logging only. */
+	void CleanFromCache(const IPCGElement* InElement, const UPCGSettings* InSettings = nullptr);
+
+	/** Returns number of copies of data cached for InElement. */
+	uint32 GetGraphCacheEntryCount(IPCGElement* InElement) const;
 #endif
 
 private:
diff --git a/Engine/Plugins/Experimental/PCG/Source/PCG/Private/Graph/PCGGraphExecutor.cpp b/Engine/Plugins/Experimental/PCG/Source/PCG/Private/Graph/PCGGraphExecutor.cpp
--- a/Engine/Plugins/Experimental/PCG/Source/PCG/Private/Graph/PCGGraphExecutor.cpp
+++ b/Engine/Plugins/Experimental/PCG/Source/PCG/Private/Graph/PCGGraphExecutor.cpp
@@ -1,13 +1,15 @@
 // Copyright Epic Games, Inc. All Rights Reserved.
 
 #include "PCGGraphExecutor.h"
+
 #include "PCGComponent.h"
-#include "Graph/PCGGraphCache.h"
-#include "PCGGraph.h"
 #include "PCGContext.h"
+#include "PCGCrc.h"
+#include "PCGGraph.h"
 #include "PCGGraphCompiler.h"
 #include "PCGHelpers.h"
 #include "PCGInputOutputSettings.h"
+#include "Graph/PCGGraphCache.h"
 
 #include "Async/Async.h"
 #include "PCGModule.h"
@@ -80,6 +82,11 @@ FPCGTaskId FPCGGraphExecutor::Schedule(UPCGComponent* Component, const TArray<FP
 
 FPCGTaskId FPCGGraphExecutor::Schedule(UPCGGraph* Graph, UPCGComponent* SourceComponent, FPCGElementPtr InputElement, const TArray<FPCGTaskId>& ExternalDependencies)
 {
+	if (SourceComponent && IsGraphCacheDebuggingEnabled())
+	{
+		UE_LOG(LogPCG, Log, TEXT("[%s] --- SCHEDULE GRAPH ---"), *SourceComponent->GetOwner()->GetName());
+	}
+
 	FPCGTaskId ScheduledId = InvalidPCGTaskId;
 
 	// Get compiled tasks from compiler
@@ -370,6 +377,7 @@ void FPCGGraphExecutor::Execute()
 	const double EndTime = StartTime + VarTimePerFrame;
 	const int32 MaxNumThreads = FMath::Max(0, FMath::Min(FPlatformMisc::NumberOfCoresIncludingHyperthreads() - 2, CVarMaxNumTasks.GetValueOnAnyThread() - 1));
 	const bool bAllowMultiDispatch = CVarGraphMultithreading.GetValueOnAnyThread();
+	const bool bGraphCacheDebuggingEnabled = IsGraphCacheDebuggingEnabled();
 
 #if WITH_EDITOR
 	UpdateGenerationNotification();
@@ -482,8 +490,17 @@ void FPCGGraphExecutor::Execute()
 				// there is an execution mode that would prevent us from doing so.
 				const UPCGSettingsInterface* TaskSettingsInterface = TaskInput.GetSettingsInterface(Task.Node ? Task.Node->GetSettingsInterface() : nullptr);
 				const UPCGSettings* TaskSettings = TaskSettingsInterface ? TaskSettingsInterface->GetSettings() : nullptr;
+				const bool bCacheable = Task.Element->IsCacheableInstance(TaskSettingsInterface);
+
+				// Calculate Crc of dependencies (input data Crcs, settings) and use this as the key in the cache lookup
+				FPCGCrc DependenciesCrc;
+				if (TaskSettings && bCacheable)
+				{
+					Task.Element->GetDependenciesCrc(TaskInput, TaskSettings, Task.SourceComponent.Get(), DependenciesCrc);
+				}
+
 				FPCGDataCollection CachedOutput;
-				const bool bResultAlreadyInCache = Task.Element->IsCacheableInstance(TaskSettingsInterface) && GraphCache.GetFromCache(Task.Element.Get(), TaskInput, TaskSettings, Task.SourceComponent.Get(), CachedOutput);
+				const bool bResultAlreadyInCache = bCacheable && DependenciesCrc.IsValid() && GraphCache.GetFromCache(Task.Node, Task.Element.Get(), DependenciesCrc, TaskInput, TaskSettings, Task.SourceComponent.Get(), CachedOutput);
 #if WITH_EDITOR
 				const bool bNeedsToCreateActiveTask = !bResultAlreadyInCache || TaskSettingsInterface->bDebug;
 #else
@@ -511,12 +528,19 @@ void FPCGGraphExecutor::Execute()
 					continue;
 				}
 
+				if (bGraphCacheDebuggingEnabled && Task.SourceComponent.Get() && Task.Node)
+				{
+					UE_LOG(LogPCG, Log, TEXT("         [%s] %s\t\tACTIVETASK"), *Task.SourceComponent->GetOwner()->GetName(), *Task.Node->GetNodeTitle().ToString());
+				}
+
 				// Allocate context if not previously done
 				if (!Task.Context)
 				{
 					Task.Context = Task.Element->Initialize(TaskInput, Task.SourceComponent, Task.Node);
 					Task.Context->TaskId = Task.NodeId;
 					Task.Context->CompiledTaskId = Task.CompiledTaskId;
+					Task.Context->InputData.Crc = TaskInput.Crc;
+					Task.Context->DependenciesCrc = DependenciesCrc;
 				}
 
 				// Validate that we can start this task now
@@ -604,7 +628,7 @@ void FPCGGraphExecutor::Execute()
 			}
 		}
 
-		auto PostTaskExecute = [this, &bAnyTaskEnded](int32 TaskIndex)
+		auto PostTaskExecute = [this, &bAnyTaskEnded, bGraphCacheDebuggingEnabled](int32 TaskIndex)
 		{
 			FPCGGraphActiveTask& ActiveTask = ActiveTasks[TaskIndex];
 
@@ -612,12 +636,17 @@ void FPCGGraphExecutor::Execute()
 			if (!ActiveTask.bIsBypassed)
 #endif
 			{
+				if (bGraphCacheDebuggingEnabled && ActiveTask.Context->SourceComponent.Get() && ActiveTask.Context->Node)
+				{
+					UE_LOG(LogPCG, Log, TEXT("         [%s] %s\t\tOUTPUT CRC %u"), *ActiveTask.Context->SourceComponent->GetOwner()->GetName(), *ActiveTask.Context->Node->GetNodeTitle().ToString(), ActiveTask.Context->OutputData.Crc.GetValue());
+				}
+
 				// Store result in cache as needed - done here because it needs to be done on the main thread
 				const UPCGSettingsInterface* ActiveTaskSettingsInterface = ActiveTask.Context->GetInputSettingsInterface();
 				if (ActiveTaskSettingsInterface && ActiveTask.Element->IsCacheableInstance(ActiveTaskSettingsInterface))
 				{
 					const UPCGSettings* ActiveTaskSettings = ActiveTaskSettingsInterface ? ActiveTaskSettingsInterface->GetSettings() : nullptr;
-					GraphCache.StoreInCache(ActiveTask.Element.Get(), ActiveTask.Context->InputData, ActiveTaskSettings, ActiveTask.Context->SourceComponent.Get(), ActiveTask.Context->OutputData);
+					GraphCache.StoreInCache(ActiveTask.Element.Get(), ActiveTask.Context->DependenciesCrc, ActiveTask.Context->InputData, ActiveTaskSettings, ActiveTask.Context->SourceComponent.Get(), ActiveTask.Context->OutputData);
 				}
 			}
 
@@ -758,6 +787,11 @@ void FPCGGraphExecutor::Execute()
 			// Call the notification update here to prevent it from sticking around - needed because we early out before this
 			UpdateGenerationNotification();
 #endif
+
+			if (bGraphCacheDebuggingEnabled)
+			{
+				UE_LOG(LogPCG, Log, TEXT("--- FINISH FPCGGRAPHEXECUTOR::EXECUTE ---"));
+			}
 		}
 
 #if WITH_EDITOR
@@ -827,6 +861,10 @@ void FPCGGraphExecutor::CancelNextTasks(FPCGTaskId CancelledTask)
 void FPCGGraphExecutor::BuildTaskInput(const FPCGGraphTask& Task, FPCGDataCollection& TaskInput)
 {
 	TRACE_CPUPROFILER_EVENT_SCOPE(FPCGGraphExecutor::BuildTaskInput);
+
+	// Initialize a Crc onto which each input Crc will be combined.
+	FPCGCrc Crc(Task.Inputs.Num());
+
 	for (const FPCGGraphTaskInput& Input : Task.Inputs)
 	{
 		check(OutputData.Contains(Input.TaskId));
@@ -846,6 +884,9 @@ void FPCGGraphExecutor::BuildTaskInput(const FPCGGraphTask& Task, FPCGDataCollec
 		if (Input.InPin)
 		{
 			TaskInput.TaggedData.Append(InputCollection.GetInputsByPin(Input.InPin->Properties.Label));
+
+			// Write pin name Crc to uniquely identify inputs per-pin.
+			Crc.Combine(Input.InPin->Properties.Label);
 		}
 		else
 		{
@@ -857,6 +898,12 @@ void FPCGGraphExecutor::BuildTaskInput(const FPCGGraphTask& Task, FPCGDataCollec
 			TaskInput.bCancelExecution = true;
 		}
 
+		// This chains the Crc of each input to produce a Crc that covers all of them.
+		if (InputCollection.Crc.IsValid())
+		{
+			Crc.Combine(InputCollection.Crc);
+		}
+
 		// Apply labelling on data; technically, we should ensure that we do this only for pass-through nodes,
 		// Otherwise we could also null out the label on the input...
 		if (Input.OutPin)
@@ -867,6 +914,8 @@ void FPCGGraphExecutor::BuildTaskInput(const FPCGGraphTask& Task, FPCGDataCollec
 			}
 		}
 	}
+
+	TaskInput.Crc = Crc;
 }
 
 void FPCGGraphExecutor::StoreResults(FPCGTaskId InTaskId, const FPCGDataCollection& InTaskOutput)
diff --git a/Engine/Plugins/Experimental/PCG/Source/PCG/Private/Graph/PCGGraphExecutor.h b/Engine/Plugins/Experimental/PCG/Source/PCG/Private/Graph/PCGGraphExecutor.h
--- a/Engine/Plugins/Experimental/PCG/Source/PCG/Private/Graph/PCGGraphExecutor.h
+++ b/Engine/Plugins/Experimental/PCG/Source/PCG/Private/Graph/PCGGraphExecutor.h
@@ -3,8 +3,8 @@
 #pragma once
 
 #include "PCGElement.h"
-#include "PCGGraphCache.h"
 #include "PCGSubsystem.h"
+#include "Graph/PCGGraphCache.h"
 
 #if WITH_EDITOR
 #include "AsyncCompilationHelpers.h"
@@ -114,6 +114,8 @@ public:
 	/** So the profiler can decode graph task ids **/
 	FPCGGraphCompiler* GetCompiler() const { return GraphCompiler.Get(); }
 
+	/** Returns the number of entries currently in the cache for InElement. */
+	uint32 GetGraphCacheEntryCount(IPCGElement* InElement) const { return GraphCache.GetGraphCacheEntryCount(InElement); }
 #endif
 
 	/** "Tick" of the graph executor. This call is NOT THREADSAFE */
@@ -122,6 +124,9 @@ public:
 	/** Expose cache so it can be dirtied */
 	FPCGGraphCache& GetCache() { return GraphCache; }
 
+	/** True if graph cache debugging is enabled. */
+	bool IsGraphCacheDebuggingEnabled() const { return GraphCache.IsDebuggingEnabled(); }
+
 private:
 	void Cancel(TFunctionRef<bool(TWeakObjectPtr<UPCGComponent>)> CancelFilter);
 	void ClearAllTasks();
diff --git a/Engine/Plugins/Experimental/PCG/Source/PCG/Private/PCGComponent.cpp b/Engine/Plugins/Experimental/PCG/Source/PCG/Private/PCGComponent.cpp
--- a/Engine/Plugins/Experimental/PCG/Source/PCG/Private/PCGComponent.cpp
+++ b/Engine/Plugins/Experimental/PCG/Source/PCG/Private/PCGComponent.cpp
@@ -185,7 +185,7 @@ void UPCGComponent::SetPropertiesFromOriginal(const UPCGComponent* Original)
 #if WITH_EDITOR
 	const bool bHasDirtyInput = InputType != NewInputType;
 	const bool bHasDirtyExclusions = !(ExcludedTags.Num() == Original->ExcludedTags.Num() && ExcludedTags.Includes(Original->ExcludedTags));
-	const bool bIsDirty = bHasDirtyInput || bHasDirtyExclusions || Seed != Original->Seed || Graph != Original->Graph;
+	const bool bIsDirty = bHasDirtyInput || bHasDirtyExclusions || Graph != Original->Graph;
 
 	if (bHasDirtyExclusions)
 	{
@@ -1506,6 +1506,11 @@ void UPCGComponent::OnActorChanged(AActor* Actor, UObject* InObject, bool bActor
 
 void UPCGComponent::DirtyGenerated(EPCGComponentDirtyFlag DirtyFlag)
 {
+	if (GetSubsystem() && GetSubsystem()->IsGraphCacheDebuggingEnabled())
+	{
+		UE_LOG(LogPCG, Log, TEXT("[%s] --- DIRTY GENERATED ---"), *GetOwner()->GetName());
+	}
+
 	bDirtyGenerated = true;
 
 	// Dirty data as a waterfall from basic values
@@ -1666,6 +1671,11 @@ UPCGData* UPCGComponent::GetPCGData()
 	if (!CachedPCGData)
 	{
 		CachedPCGData = CreatePCGData();
+
+		if (GetSubsystem() && GetSubsystem()->IsGraphCacheDebuggingEnabled())
+		{
+			UE_LOG(LogPCG, Log, TEXT("         [%s] CACHE REFRESH CachedPCGData"), *GetOwner()->GetName());
+		}
 	}
 
 	return CachedPCGData;
@@ -1676,6 +1686,11 @@ UPCGData* UPCGComponent::GetInputPCGData()
 	if (!CachedInputData)
 	{
 		CachedInputData = CreateInputPCGData();
+
+		if (GetSubsystem() && GetSubsystem()->IsGraphCacheDebuggingEnabled())
+		{
+			UE_LOG(LogPCG, Log, TEXT("         [%s] CACHE REFRESH CachedInputData"), *GetOwner()->GetName());
+		}
 	}
 
 	return CachedInputData;
@@ -1687,6 +1702,11 @@ UPCGData* UPCGComponent::GetActorPCGData()
 	if (!CachedActorData || IsLandscapeCachedDataDirty(CachedActorData))
 	{
 		CachedActorData = CreateActorPCGData();
+
+		if (GetSubsystem() && GetSubsystem()->IsGraphCacheDebuggingEnabled())
+		{
+			UE_LOG(LogPCG, Log, TEXT("         [%s] CACHE REFRESH CachedActorData"), *GetOwner()->GetName());
+		}
 	}
 
 	return CachedActorData;
@@ -1697,6 +1717,11 @@ UPCGData* UPCGComponent::GetLandscapePCGData()
 	if (!CachedLandscapeData || IsLandscapeCachedDataDirty(CachedLandscapeData))
 	{
 		CachedLandscapeData = CreateLandscapePCGData(/*bHeightOnly=*/false);
+
+		if (GetSubsystem() && GetSubsystem()->IsGraphCacheDebuggingEnabled())
+		{
+			UE_LOG(LogPCG, Log, TEXT("         [%s] CACHE REFRESH CachedLandscapeData"), *GetOwner()->GetName());
+		}
 	}
 
 	return CachedLandscapeData;
@@ -1707,6 +1732,11 @@ UPCGData* UPCGComponent::GetLandscapeHeightPCGData()
 	if (!CachedLandscapeHeightData || IsLandscapeCachedDataDirty(CachedLandscapeHeightData))
 	{
 		CachedLandscapeHeightData = CreateLandscapePCGData(/*bHeightOnly=*/true);
+
+		if (GetSubsystem() && GetSubsystem()->IsGraphCacheDebuggingEnabled())
+		{
+			UE_LOG(LogPCG, Log, TEXT("         [%s] CACHE REFRESH CachedLandscapeHeightData"), *GetOwner()->GetName());
+		}
 	}
 
 	return CachedLandscapeHeightData;
@@ -2282,7 +2312,7 @@ void UPCGComponent::DirtyCacheFromTag(const FName& InTag)
 		{
 			if (Settings.IsValid() && GetSubsystem())
 			{
-				GetSubsystem()->CleanFromCache(Settings->GetElement().Get());
+				GetSubsystem()->CleanFromCache(Settings->GetElement().Get(), Settings.Get());
 			}
 		}
 	}
@@ -2296,7 +2326,7 @@ void UPCGComponent::DirtyCacheForAllTrackedTags()
 		{
 			if (Settings.IsValid() && GetSubsystem())
 			{
-				GetSubsystem()->CleanFromCache(Settings->GetElement().Get());
+				GetSubsystem()->CleanFromCache(Settings->GetElement().Get(), Settings.Get());
 			}
 		}
 	}
diff --git a/Engine/Plugins/Experimental/PCG/Source/PCG/Private/PCGData.cpp b/Engine/Plugins/Experimental/PCG/Source/PCG/Private/PCGData.cpp
--- a/Engine/Plugins/Experimental/PCG/Source/PCG/Private/PCGData.cpp
+++ b/Engine/Plugins/Experimental/PCG/Source/PCG/Private/PCGData.cpp
@@ -1,10 +1,13 @@
 // Copyright Epic Games, Inc. All Rights Reserved.
 
 #include "PCGData.h"
-#include "PCGSettings.h"
+
 #include "PCGParamData.h"
+#include "PCGSettings.h"
+#include "Data/PCGPointData.h"
 #include "Data/PCGSpatialData.h"
 
+#include "Serialization/ArchiveCrc32.h"
 #include "UObject/Package.h"
 
 #include UE_INLINE_GENERATED_CPP_BY_NAME(PCGData)
@@ -68,6 +71,29 @@ void FPCGRootSet::RemoveInternal(UObject* InObject)
 	}
 }
 
+UPCGData::UPCGData(const FObjectInitializer& ObjectInitializer)
+	: Super(ObjectInitializer)
+{
+	if (!HasAnyFlags(RF_ClassDefaultObject))
+	{
+		UID = ++UIDCounter;
+	}
+}
+
+FPCGCrc UPCGData::ComputeCrc() const
+{
+	FArchiveCrc32 Ar;
+	AddToCrc(Ar);
+	return FPCGCrc(Ar.GetCrc());
+}
+
+void UPCGData::AddToCrc(FArchiveCrc32& Ar) const
+{
+	// Fallback implementation uses UID to ensure every object returns a different Crc.
+	uint64 UIDValue = UID;
+	Ar << UIDValue;
+}
+
 bool FPCGTaggedData::operator==(const FPCGTaggedData& Other) const
 {
 	return Data == Other.Data &&
@@ -243,6 +269,26 @@ void FPCGDataCollection::RemoveFromRootSet(FPCGRootSet& RootSet) const
 	}
 }
 
+FPCGCrc FPCGDataCollection::ComputeCrc()
+{
+	// If there is no data, will return valid Crc==0 which is fine. No such thing as an invalid FPCGDataCollection
+	FArchiveCrc32 Ar;
+
+	for (FPCGTaggedData& Data : TaggedData)
+	{
+		Ar << Data.Pin;
+
+		if (Data.Data)
+		{
+			Data.Data->AddToCrc(Ar);
+		}
+
+		Ar << Data.Tags;
+	}
+
+	return FPCGCrc(Ar.GetCrc());
+}
+
 void FPCGDataCollection::Reset()
 {
 	// Implementation note: We are assuming that there is no need to remove the data from the rootset here.
diff --git a/Engine/Plugins/Experimental/PCG/Source/PCG/Private/PCGElement.cpp b/Engine/Plugins/Experimental/PCG/Source/PCG/Private/PCGElement.cpp
--- a/Engine/Plugins/Experimental/PCG/Source/PCG/Private/PCGElement.cpp
+++ b/Engine/Plugins/Experimental/PCG/Source/PCG/Private/PCGElement.cpp
@@ -2,14 +2,16 @@
 
 #include "PCGElement.h"
 
+#include "PCGComponent.h"
 #include "PCGContext.h"
+#include "PCGGraph.h"
 #include "Data/PCGPointData.h"
 #include "Elements/PCGDebugElement.h"
 #include "Elements/PCGSelfPruning.h"
+
 #include "HAL/IConsoleManager.h"
 #include "Utils/PCGExtraCapture.h"
 
-
 static TAutoConsoleVariable<bool> CVarPCGValidatePointMetadata(
 	TEXT("pcg.debug.ValidatePointMetadata"),
 	true,
@@ -171,6 +173,21 @@ void IPCGElement::PostExecute(FPCGContext* Context) const
 		}
 	}
 
+	if (IsCacheableInstance(SettingsInterface))
+	{
+		// If cacheable, then assume that the dependencies Crc is a good Crc to use from the output data - assumes changes
+		// in inputs and settings are 1:1 with changes in output. In general this is not the case as some settings changes will
+		// have no impact on the output. And since we are unlikely to do full Crc's of landscapes, a landscape data may produce
+		// a different Crc on each regenerate, triggering unnecessary executions. Removing this branch (always doing a full Crc)
+		// would avoid such triggers propagating.
+		Context->OutputData.Crc = Context->DependenciesCrc;
+	}
+	else
+	{
+		// Compute Crc from output data
+		Context->OutputData.Crc = Context->OutputData.ComputeCrc();
+	}
+
 	// Additional debug things (check for duplicates),
 #if WITH_EDITOR
 	if (SettingsInterface && SettingsInterface->DebugSettings.bCheckForDuplicates)
@@ -384,10 +401,22 @@ bool IPCGElement::IsCacheableInstance(const UPCGSettingsInterface* InSettingsInt
 	}
 }
 
-#if WITH_EDITOR
+void IPCGElement::GetDependenciesCrc(const FPCGDataCollection& InInput, const UPCGSettings* InSettings, UPCGComponent* InComponent, FPCGCrc& OutCrc) const
+{
+	FPCGCrc Crc = InInput.Crc;
 
+	if (InSettings)
+	{
+		Crc.Combine(InSettings->GetCrc32());
+	}
 
-#endif // WITH_EDITOR
+	if (InComponent)
+	{
+		Crc.Combine(InComponent->Seed);
+	}
+
+	OutCrc = Crc;
+}
 
 FPCGContext* FSimplePCGElement::Initialize(const FPCGDataCollection& InputData, TWeakObjectPtr<UPCGComponent> SourceComponent, const UPCGNode* Node)
 {
diff --git a/Engine/Plugins/Experimental/PCG/Source/PCG/Private/PCGSettings.cpp b/Engine/Plugins/Experimental/PCG/Source/PCG/Private/PCGSettings.cpp
--- a/Engine/Plugins/Experimental/PCG/Source/PCG/Private/PCGSettings.cpp
+++ b/Engine/Plugins/Experimental/PCG/Source/PCG/Private/PCGSettings.cpp
@@ -1,10 +1,12 @@
 // Copyright Epic Games, Inc. All Rights Reserved.
 
 #include "PCGSettings.h"
+
 #include "PCGContext.h"
 #include "PCGCustomVersion.h"
 #include "PCGPin.h"
 #include "PCGSubsystem.h"
+
 #include "Serialization/ArchiveObjectCrc32.h"
 #include "UObject/ObjectSaveContext.h"
 
@@ -14,18 +16,24 @@
 #include "Editor.h"
 #endif
 
-/** In order to reuse the cache when only debug settings change, we must make sure to ignore these from the CRC check */
+/** Custom Crc computation that ignores properties that will not affect the computed result of a node. */
 class FPCGSettingsObjectCrc32 : public FArchiveObjectCrc32
 {
 public:
+#if WITH_EDITOR
 	virtual bool ShouldSkipProperty(const FProperty* InProperty) const override
 	{
-#if WITH_EDITOR
-		return InProperty && (InProperty->GetFName() == GET_MEMBER_NAME_CHECKED(UPCGSettings, DebugSettings) || InProperty->GetFName() == GET_MEMBER_NAME_CHECKED(UPCGSettings, DeterminismSettings));
-#else
-		return FArchiveObjectCrc32::ShouldSkipProperty(InProperty);
-#endif // WITH_EDITOR
+		const bool bSkip = InProperty && (
+			InProperty->GetFName() == GET_MEMBER_NAME_CHECKED(UPCGSettings, DebugSettings)
+			|| InProperty->GetFName() == GET_MEMBER_NAME_CHECKED(UPCGSettings, DeterminismSettings)
+			|| InProperty->GetFName() == GET_MEMBER_NAME_CHECKED(UPCGSettings, bDebug)
+			|| InProperty->GetFName() == GET_MEMBER_NAME_CHECKED(UPCGSettings, Category)
+			|| InProperty->GetFName() == GET_MEMBER_NAME_CHECKED(UPCGSettings, Description)
+			|| InProperty->GetFName() == GET_MEMBER_NAME_CHECKED(UPCGSettings, bExposeToLibrary)
+			);
+		return bSkip;
 	}
+#endif // WITH_EDITOR
 };
 
 bool UPCGSettingsInterface::IsInstance() const
@@ -160,14 +168,12 @@ TArray<FPCGPinProperties> UPCGSettings::OutputPinProperties() const
 	return PinProperties;
 }
 
-TArray<FPCGPinProperties>
-UPCGSettings::DefaultInputPinProperties() const
+TArray<FPCGPinProperties> UPCGSettings::DefaultInputPinProperties() const
 {
 	return InputPinProperties();
 }
 
-TArray<FPCGPinProperties>
-UPCGSettings::DefaultOutputPinProperties() const
+TArray<FPCGPinProperties> UPCGSettings::DefaultOutputPinProperties() const
 {
 	return OutputPinProperties();
 }
@@ -213,7 +219,7 @@ void UPCGSettings::DirtyCache()
 		UPCGSubsystem* PCGSubsystem = World ? World->GetSubsystem<UPCGSubsystem>() : nullptr;
 		if (PCGSubsystem)
 		{
-			PCGSubsystem->CleanFromCache(GetElement().Get());
+			PCGSubsystem->CleanFromCache(GetElement().Get(), this);
 		}
 	}
 }
diff --git a/Engine/Plugins/Experimental/PCG/Source/PCG/Private/PCGSubsystem.cpp b/Engine/Plugins/Experimental/PCG/Source/PCG/Private/PCGSubsystem.cpp
--- a/Engine/Plugins/Experimental/PCG/Source/PCG/Private/PCGSubsystem.cpp
+++ b/Engine/Plugins/Experimental/PCG/Source/PCG/Private/PCGSubsystem.cpp
@@ -482,6 +482,11 @@ void UPCGSubsystem::ForAllRegisteredLocalComponents(UPCGComponent* OriginalCompo
 	DispatchToRegisteredLocalComponents(OriginalComponent, WrapperFunc);
 }
 
+bool UPCGSubsystem::IsGraphCacheDebuggingEnabled() const
+{
+	return GraphExecutor && GraphExecutor->IsGraphCacheDebuggingEnabled();
+}
+
 TArray<FPCGTaskId> UPCGSubsystem::DispatchToRegisteredLocalComponents(UPCGComponent* OriginalComponent, const TFunction<FPCGTaskId(UPCGComponent*)>& InFunc) const
 {
 	TRACE_CPUPROFILER_EVENT_SCOPE(UPCGSubsystem::DispatchToRegisteredLocalComponents);
@@ -1579,11 +1584,11 @@ void UPCGSubsystem::NotifyGraphChanged(UPCGGraph* InGraph)
 	}
 }
 
-void UPCGSubsystem::CleanFromCache(const IPCGElement* InElement)
+void UPCGSubsystem::CleanFromCache(const IPCGElement* InElement, const UPCGSettings* InSettings/*= nullptr*/)
 {
 	if (GraphExecutor)
 	{
-		GraphExecutor->GetCache().CleanFromCache(InElement);
+		GraphExecutor->GetCache().CleanFromCache(InElement, InSettings);
 	}
 }
 
@@ -1618,6 +1623,11 @@ const FPCGGraphCompiler* UPCGSubsystem::GetGraphCompiler() const
 	return nullptr;
 }
 
+uint32 UPCGSubsystem::GetGraphCacheEntryCount(IPCGElement* InElement) const
+{
+	return GraphExecutor ? GraphExecutor->GetGraphCacheEntryCount(InElement) : 0;
+}
+
 void UPCGSubsystem::ResetPartitionActorsMap()
 {
 	PartitionActorsMapLock.WriteLock();
diff --git a/Engine/Plugins/Experimental/PCG/Source/PCG/Public/Elements/PCGSurfaceSampler.h b/Engine/Plugins/Experimental/PCG/Source/PCG/Public/Elements/PCGSurfaceSampler.h
--- a/Engine/Plugins/Experimental/PCG/Source/PCG/Public/Elements/PCGSurfaceSampler.h
+++ b/Engine/Plugins/Experimental/PCG/Source/PCG/Public/Elements/PCGSurfaceSampler.h
@@ -78,6 +78,9 @@ public:
 
 class FPCGSurfaceSamplerElement : public FSimplePCGElement
 {
+public:
+	virtual void GetDependenciesCrc(const FPCGDataCollection& InInput, const UPCGSettings* InSettings, UPCGComponent* InComponent, FPCGCrc& OutCrc) const override;
+
 protected:
 	virtual bool ExecuteInternal(FPCGContext* Context) const override;
 };
diff --git a/Engine/Plugins/Experimental/PCG/Source/PCG/Public/Elements/PCGTextureSampler.h b/Engine/Plugins/Experimental/PCG/Source/PCG/Public/Elements/PCGTextureSampler.h
--- a/Engine/Plugins/Experimental/PCG/Source/PCG/Public/Elements/PCGTextureSampler.h
+++ b/Engine/Plugins/Experimental/PCG/Source/PCG/Public/Elements/PCGTextureSampler.h
@@ -76,7 +76,9 @@ public:
 
 class FPCGTextureSamplerElement : public FSimplePCGElement
 {
+public:
+	virtual void GetDependenciesCrc(const FPCGDataCollection& InInput, const UPCGSettings* InSettings, UPCGComponent* InComponent, FPCGCrc& OutCrc) const override;
+
 protected:
-	virtual bool IsCacheable(const UPCGSettings* InSettings) const;
 	virtual bool ExecuteInternal(FPCGContext* Context) const override;
 };
diff --git a/Engine/Plugins/Experimental/PCG/Source/PCG/Public/PCGContext.h b/Engine/Plugins/Experimental/PCG/Source/PCG/Public/PCGContext.h
--- a/Engine/Plugins/Experimental/PCG/Source/PCG/Public/PCGContext.h
+++ b/Engine/Plugins/Experimental/PCG/Source/PCG/Public/PCGContext.h
@@ -49,6 +49,8 @@ struct PCG_API FPCGContext
 	TWeakObjectPtr<UPCGComponent> SourceComponent = nullptr;
 	int32 NumAvailableTasks = 0;
 
+	FPCGCrc DependenciesCrc;
+
 	// TODO: add RNG source
 	// TODO: replace this by a better identification mechanism
 	const UPCGNode* Node = nullptr;
diff --git a/Engine/Plugins/Experimental/PCG/Source/PCG/Public/PCGCrc.h b/Engine/Plugins/Experimental/PCG/Source/PCG/Public/PCGCrc.h
new file mode 100644
--- /dev/null
+++ b/Engine/Plugins/Experimental/PCG/Source/PCG/Public/PCGCrc.h
@@ -0,0 +1,73 @@
+// Copyright Epic Games, Inc. All Rights Reserved.
+
+#pragma once
+
+#include "Misc/Crc.h"
+#include "Templates/TypeHash.h"
+#include "UObject/NameTypes.h"
+
+/** Crc with valid flag and helper functionality. */
+struct PCG_API FPCGCrc
+{
+	/** Creates an invalid Crc. */
+	FPCGCrc() = default;
+
+	/** Initializes a valid Crc set to InValue. */
+	explicit FPCGCrc(uint32 InValue)
+		: Value(InValue)
+		, bValid(true)
+	{
+	}
+
+	bool IsValid() const { return bValid; }
+
+	uint32 GetValue() const { return Value; }
+
+	/** Combines another Crc into this Crc to chain them. */
+	void Combine(const FPCGCrc& InOtherCrc)
+	{
+#if WITH_EDITOR
+		ensure(IsValid() && InOtherCrc.IsValid());
+#endif
+
+		Value = HashCombineFast(Value, InOtherCrc.Value);
+	}
+
+	/** Combines another Crc value into this Crc to chain them. */
+	void Combine(uint32 InOtherCrcValue)
+	{
+#if WITH_EDITOR
+		ensure(IsValid());
+#endif
+
+		Value = HashCombineFast(Value, InOtherCrcValue);
+	}
+
+	/** Computes Crc of given FName and combines it into this Crc. */
+	void Combine(FName InName)
+	{
+#if WITH_EDITOR
+		ensure(IsValid());
+#endif
+
+		uint32 Crc = 0;
+		FCrc::StrCrc32(*InName.ToString(), Crc);
+		Combine(Crc);
+	}
+
+	/** Compares Crc. This and other Crc must be valid. */
+	bool operator==(const FPCGCrc& InOtherCrc) const
+	{
+#if WITH_EDITOR
+		ensure(IsValid() && InOtherCrc.IsValid());
+#endif
+
+		return IsValid() && InOtherCrc.IsValid() && Value == InOtherCrc.Value;
+	}
+
+private:
+	/** Crc32 value. */
+	uint32 Value = 0;
+
+	bool bValid = false;
+};
diff --git a/Engine/Plugins/Experimental/PCG/Source/PCG/Public/PCGData.h b/Engine/Plugins/Experimental/PCG/Source/PCG/Public/PCGData.h
--- a/Engine/Plugins/Experimental/PCG/Source/PCG/Public/PCGData.h
+++ b/Engine/Plugins/Experimental/PCG/Source/PCG/Public/PCGData.h
@@ -2,12 +2,14 @@
 
 #pragma once
 
-#include "Kismet/BlueprintFunctionLibrary.h"
-
 #include "PCGCommon.h"
+#include "PCGCrc.h"
+
+#include "Kismet/BlueprintFunctionLibrary.h"
 
 #include "PCGData.generated.h"
 
+class FArchiveCrc32;
 class UPCGSettings;
 class UPCGSettingsInterface;
 class UPCGParamData;
@@ -23,7 +25,22 @@ class PCG_API UPCGData : public UObject
 	GENERATED_BODY()
 
 public:
+	UPCGData(const FObjectInitializer& ObjectInitializer);
 	virtual EPCGDataType GetDataType() const { return EPCGDataType::None; }
+
+	/** Computes a Crc for this data. */
+	FPCGCrc ComputeCrc() const;
+
+	/** Contributes to Crc. Fallback implementation that writes object instance UID. */
+	virtual void AddToCrc(FArchiveCrc32& Ar) const;
+
+	/** Unique ID for this object instance. */
+	UPROPERTY(Transient)
+	uint64 UID = 0;
+
+private:
+	/** Serves unique ID values to instances of this object. */
+	static inline std::atomic<uint64> UIDCounter{ 1 };
 };
 
 USTRUCT(BlueprintType)
@@ -94,6 +111,9 @@ struct PCG_API FPCGDataCollection
 	void AddToRootSet(FPCGRootSet& RootSet) const;
 	void RemoveFromRootSet(FPCGRootSet& RootSet) const;
 
+	/** Computes Crc for this data. */
+	FPCGCrc ComputeCrc();
+
 	/** Cleans up the collection, but does not unroot any previously rooted data. */
 	void Reset();
 
@@ -105,6 +125,9 @@ struct PCG_API FPCGDataCollection
 
 	/** This flag is used to cancel further computation or for the debug/isolate feature */
 	bool bCancelExecution = false;
+
+	/** A snapshot of the internal state of the data. If any dependency (setting, node input or external data) changes then this value should change. */
+	FPCGCrc Crc;
 };
 
 template<typename SettingsType>
diff --git a/Engine/Plugins/Experimental/PCG/Source/PCG/Public/PCGElement.h b/Engine/Plugins/Experimental/PCG/Source/PCG/Public/PCGElement.h
--- a/Engine/Plugins/Experimental/PCG/Source/PCG/Public/PCGElement.h
+++ b/Engine/Plugins/Experimental/PCG/Source/PCG/Public/PCGElement.h
@@ -7,6 +7,7 @@
 
 class UPCGSettingsInterface;
 struct FPCGContext;
+struct FPCGCrc;
 struct FPCGDataCollection;
 
 class IPCGElement;
@@ -50,6 +51,13 @@ public:
 	/** Returns true if the node can be cached (e.g. does not create artifacts & does not depend on untracked data */
 	virtual bool IsCacheable(const UPCGSettings* InSettings) const { return true; }
 
+	/**
+	 * Calculate a Crc that provides a receipt for the input data that can be paired with output data from the cache. If any dependency (setting, node input or
+	 * external data) changes then this value should change. For some elements it is inefficient or not possible to output a Crc here. These can return an invalid
+	 * Crc and the Crc can either be computed during execution, or afterwards based on output data.
+	 */
+	virtual void GetDependenciesCrc(const FPCGDataCollection& InInput, const UPCGSettings* InSettings, UPCGComponent* InComponent, FPCGCrc& OutCrc) const;
+
 	/** Public function that executes the element on the appropriately created context.
 	* The caller should call the Execute function until it returns true.
 	*/
diff --git a/Engine/Plugins/Experimental/PCG/Source/PCG/Public/PCGSettings.h b/Engine/Plugins/Experimental/PCG/Source/PCG/Public/PCGSettings.h
--- a/Engine/Plugins/Experimental/PCG/Source/PCG/Public/PCGSettings.h
+++ b/Engine/Plugins/Experimental/PCG/Source/PCG/Public/PCGSettings.h
@@ -196,7 +196,7 @@ protected:
 	void DirtyCache();
 #endif
 
-	// By default, settings won't use a seed. Set this bool at true in the child ctor to allow edition and use it.
+	// By default, settings won't use a seed. Set this bool to true in the child ctor to allow edition and use it.
 	UPROPERTY(VisibleAnywhere, Transient, Category = Settings, meta = (EditCondition = false, EditConditionHides))
 	bool bUseSeed = false;
 
diff --git a/Engine/Plugins/Experimental/PCG/Source/PCG/Public/PCGSubsystem.h b/Engine/Plugins/Experimental/PCG/Source/PCG/Public/PCGSubsystem.h
--- a/Engine/Plugins/Experimental/PCG/Source/PCG/Public/PCGSubsystem.h
+++ b/Engine/Plugins/Experimental/PCG/Source/PCG/Public/PCGSubsystem.h
@@ -23,6 +23,7 @@ class FPCGGraphCompiler;
 class FPCGGraphExecutor;
 struct FPCGContext;
 struct FPCGDataCollection;
+class UPCGSettings;
 
 class IPCGElement;
 typedef TSharedPtr<IPCGElement, ESPMode::ThreadSafe> FPCGElementPtr;
@@ -128,6 +129,9 @@ public:
 	/* Call the InFunc function to all local component registered to the original component. Thread safe*/
 	void ForAllRegisteredLocalComponents(UPCGComponent* OriginalComponent, const TFunction<void(UPCGComponent*)>& InFunc) const;
 
+	/** True if graph cache debugging is enabled. */
+	bool IsGraphCacheDebuggingEnabled() const;
+
 #if WITH_EDITOR
 public:
 	/** Schedule refresh on the current or next frame */
@@ -146,8 +150,8 @@ public:
 	/** Propagate to the graph compiler graph changes */
 	void NotifyGraphChanged(UPCGGraph* InGraph);
 
-	/** Cleans up the graph cache on an element basis */
-	void CleanFromCache(const IPCGElement* InElement);
+	/** Cleans up the graph cache on an element basis. InSettings is used for debugging and is optional. */
+	void CleanFromCache(const IPCGElement* InElement, const UPCGSettings* InSettings = nullptr);
 
 	/** Move all resources from sub actors to a new actor */
 	void ClearPCGLink(UPCGComponent* InComponent, const FBox& InBounds, AActor* InNewActor);
@@ -164,6 +168,9 @@ public:
 	/** Returns the graph compiler so we can figure out task info in the profiler view **/
 	const FPCGGraphCompiler* GetGraphCompiler() const;
 
+	/** Returns how many times InElement is present in the cache. */
+	uint32 GetGraphCacheEntryCount(IPCGElement* InElement) const;
+
 private:
 	enum class EOperation : uint32
 	{
diff --git a/Engine/Plugins/Experimental/PCG/Source/PCGEditor/Private/PCGEditorGraphNode.cpp b/Engine/Plugins/Experimental/PCG/Source/PCGEditor/Private/PCGEditorGraphNode.cpp
--- a/Engine/Plugins/Experimental/PCG/Source/PCGEditor/Private/PCGEditorGraphNode.cpp
+++ b/Engine/Plugins/Experimental/PCG/Source/PCGEditor/Private/PCGEditorGraphNode.cpp
@@ -4,6 +4,8 @@
 
 #include "PCGEditorModule.h"
 #include "PCGNode.h"
+#include "PCGSettings.h"
+#include "PCGSubsystem.h"
 
 #include "Framework/Commands/GenericCommands.h"
 #include "PCGPin.h"
@@ -12,6 +14,23 @@
 
 #define LOCTEXT_NAMESPACE "PCGEditorGraphNode"
 
+namespace UPCGEditorGraphNodeHelpers
+{
+	// Info to aid element cache analysis / debugging
+	void GetGraphCacheDebugInfo(const UPCGNode* InNode, bool& bOutDebuggingEnabled, uint32& OutNumCacheEntries)
+	{
+		UWorld* World = GEditor ? (GEditor->PlayWorld ? GEditor->PlayWorld.Get() : GEditor->GetEditorWorldContext().World()) : nullptr;
+		UPCGSubsystem* Subsystem = World ? World->GetSubsystem<UPCGSubsystem>() : nullptr;
+		bOutDebuggingEnabled = Subsystem && Subsystem->IsGraphCacheDebuggingEnabled();
+
+		if (bOutDebuggingEnabled)
+		{
+			IPCGElement* Element = (InNode && InNode->GetSettings()) ? InNode->GetSettings()->GetElement().Get() : nullptr;
+			OutNumCacheEntries = Element ? Subsystem->GetGraphCacheEntryCount(Element) : 0;
+		}
+	}
+}
+
 UPCGEditorGraphNode::UPCGEditorGraphNode(const FObjectInitializer& ObjectInitializer)
 	: Super(ObjectInitializer)
 {
@@ -22,15 +41,27 @@ FText UPCGEditorGraphNode::GetNodeTitle(ENodeTitleType::Type TitleType) const
 {
 	constexpr int32 NodeTitleMaxSize = 70;
 
+	FString Result;
 	if (PCGNode)
 	{
-		FString Res = PCGNode->GetNodeTitle().ToString();
-		return FText::FromString((Res.Len() > NodeTitleMaxSize) ? Res.Left(NodeTitleMaxSize) : Res);
+		Result = PCGNode->GetNodeTitle().ToString();
+		Result = (Result.Len() > NodeTitleMaxSize) ? Result.Left(NodeTitleMaxSize) : Result;
 	}
 	else
 	{
-		return FText::FromName(TEXT("Unnamed node"));
+		Result = FString(TEXT("Unnamed node"));
 	}
+
+	// Debug info - append how many copies of this element are currently in the cache to the node title
+	bool bDebuggingEnabled;
+	uint32 NumCacheEntries;
+	UPCGEditorGraphNodeHelpers::GetGraphCacheDebugInfo(PCGNode.Get(), bDebuggingEnabled, NumCacheEntries);
+	if (bDebuggingEnabled)
+	{
+		Result = FString::Format(TEXT("{0} [{1}]"), { Result, NumCacheEntries });
+	}
+
+	return FText::FromString(Result);
 }
 
 void UPCGEditorGraphNode::GetNodeContextMenuActions(UToolMenu* Menu, class UGraphNodeContextMenuContext* Context) const
diff --git a/Engine/Plugins/Experimental/PCG/Source/PCGEditor/Private/SPCGEditorGraphAttributeListView.cpp b/Engine/Plugins/Experimental/PCG/Source/PCGEditor/Private/SPCGEditorGraphAttributeListView.cpp
--- a/Engine/Plugins/Experimental/PCG/Source/PCGEditor/Private/SPCGEditorGraphAttributeListView.cpp
+++ b/Engine/Plugins/Experimental/PCG/Source/PCGEditor/Private/SPCGEditorGraphAttributeListView.cpp
@@ -7,6 +7,7 @@
 #include "PCGEditor.h"
 #include "PCGEditorGraphNodeBase.h"
 #include "PCGParamData.h"
+#include "PCGSubsystem.h"
 #include "Data/PCGPointData.h"
 #include "Data/PCGSpatialData.h"
 
@@ -81,6 +82,13 @@ namespace PCGEditorGraphAttributeListView
 	const FText TEXT_PointSeedLabel = LOCTEXT("PointSeedLabel", "Seed");
 	const FText TEXT_PointMetadataEntryLabel = LOCTEXT("PointMetadataEntryLabel", "Entry Key");
 	const FText TEXT_PointMetadataEntryParentLabel = LOCTEXT("PointMetadataEntryParentLabel", "Parent Key");
+
+	bool IsGraphCacheDebuggingEnabled()
+	{
+		UWorld* World = GEditor ? (GEditor->PlayWorld ? GEditor->PlayWorld.Get() : GEditor->GetEditorWorldContext().World()) : nullptr;
+		UPCGSubsystem* Subsystem = World ? World->GetSubsystem<UPCGSubsystem>() : nullptr;
+		return Subsystem && Subsystem->IsGraphCacheDebuggingEnabled();
+	}
 }
 
 void SPCGListViewItemRow::Construct(const FArguments& InArgs, const TSharedRef<STableViewBase>& InOwnerTableView, const PCGListviewItemPtr& Item)
@@ -644,7 +652,15 @@ void SPCGEditorGraphAttributeListView::RefreshAttributeList()
 				ListViewItems.Add(ListViewItem);
 			}
 
-			InfoTextBlock->SetText(FText::Format(LOCTEXT("MetadataInfoTextBlockFmt", "Number of metadata: {0}"), ItemKeyUpperBound - ItemKeyLowerBound));
+			if (!PCGEditorGraphAttributeListView::IsGraphCacheDebuggingEnabled())
+			{
+				InfoTextBlock->SetText(FText::Format(LOCTEXT("MetadataInfoTextBlockFmt", "Number of metadata: {0}"), ItemKeyUpperBound - ItemKeyLowerBound));
+			}
+			else
+			{
+				// If cache debugging enabled, write CRC to help diagnose missed-dependency issues
+				InfoTextBlock->SetText(FText::Format(LOCTEXT("MetadataInfoTextBlockWithCrcFmt", "Number of metadata: {0}  CRC: {1}"), ItemKeyUpperBound - ItemKeyLowerBound, InspectionData->Crc.GetValue()));
+			}
 		}
 	}
 	else if (const UPCGSpatialData* PCGSpatialData = Cast<const UPCGSpatialData>(PCGData))
@@ -671,8 +687,16 @@ void SPCGEditorGraphAttributeListView::RefreshAttributeList()
 				ListViewItem->MetadataInfos = &MetadataInfos;
 				ListViewItems.Add(ListViewItem);
 			}
-
-			InfoTextBlock->SetText(FText::Format(LOCTEXT("PointInfoTextBlockFmt", "Number of points: {0}"), NumPoints));
+			
+			if (!PCGEditorGraphAttributeListView::IsGraphCacheDebuggingEnabled())
+			{
+				InfoTextBlock->SetText(FText::Format(LOCTEXT("PointInfoTextBlockFmt", "Number of points: {0}"), NumPoints));
+			}
+			else
+			{
+				// If cache debugging enabled, write CRC to help diagnose missed-dependency issues
+				InfoTextBlock->SetText(FText::Format(LOCTEXT("PointInfoTextBlockFmt", "Number of points: {0}, CRC: {1}"), NumPoints, InspectionData->Crc.GetValue()));
+			}
 		}
 	}
 
