[PCG] Overhaul of the param override mechanic.

The goal is to automate some logic to avoid supporting different ways to do things.

* Overridable params are now using UProperty metadata. (PCG_Overridable)
* Params pin is now created automatically when a node has at least one overriable param
* Params pin is now multi connection/data. It will do an auto combine.
* Every overridable param is also added as an advanced pin.
* Deprecation to move params away from the input pin generalized to all nodes with params
* Add some safeguard to avoid name clashing between pins and overridable params pins.
* Settings are automatically overriden, and we can use `GetInputSettings` from the context to get the overriden settings.
* Deprecating GetValue and GetParams, since they are not useful anymore
* Still support (for a limited time) params coming from the input pins, for backcompatibility.
* Overriable param pins connected will grey out the edit in the detail view for this param.

diff --git a/Engine/Plugins/Experimental/PCG/Config/BasePCG.ini b/Engine/Plugins/Experimental/PCG/Config/BasePCG.ini
new file mode 100644
--- /dev/null
+++ b/Engine/Plugins/Experimental/PCG/Config/BasePCG.ini
@@ -0,0 +1,3 @@
+[CoreRedirects]
++PropertyRedirects=(OldName="PCGSplineSamplerSettings.Params",NewName="PCGSplineSamplerSettings.SamplerParams")
++PropertyRedirects=(OldName="PCGProjectionSettings.Params",NewName="PCGProjectionSettings.ProjectionParams")
diff --git a/Engine/Plugins/Experimental/PCG/Source/PCG/Private/Elements/Metadata/PCGMetadataElement.cpp b/Engine/Plugins/Experimental/PCG/Source/PCG/Private/Elements/Metadata/PCGMetadataElement.cpp
--- a/Engine/Plugins/Experimental/PCG/Source/PCG/Private/Elements/Metadata/PCGMetadataElement.cpp
+++ b/Engine/Plugins/Experimental/PCG/Source/PCG/Private/Elements/Metadata/PCGMetadataElement.cpp
@@ -71,7 +71,6 @@ bool FPCGMetadataOperationElement::ExecuteInternal(FPCGContext* Context) const
 	check(Settings);
 
 	TArray<FPCGTaggedData> Inputs = Context->InputData.GetInputs();
-	UPCGParamData* Params = Context->InputData.GetParams();
 
 	TArray<FPCGTaggedData>& Outputs = Context->OutputData.TaggedData;
 
diff --git a/Engine/Plugins/Experimental/PCG/Source/PCG/Private/Elements/Metadata/PCGMetadataPartition.cpp b/Engine/Plugins/Experimental/PCG/Source/PCG/Private/Elements/Metadata/PCGMetadataPartition.cpp
--- a/Engine/Plugins/Experimental/PCG/Source/PCG/Private/Elements/Metadata/PCGMetadataPartition.cpp
+++ b/Engine/Plugins/Experimental/PCG/Source/PCG/Private/Elements/Metadata/PCGMetadataPartition.cpp
@@ -3,7 +3,6 @@
 #include "Elements/Metadata/PCGMetadataPartition.h"
 #include "Data/PCGSpatialData.h"
 #include "Data/PCGPointData.h"
-#include "Helpers/PCGSettingsHelpers.h"
 
 #include "Algo/Find.h"
 #include "PCGContext.h"
@@ -24,13 +23,12 @@ bool FPCGMetadataPartitionElement::ExecuteInternal(FPCGContext* Context) const
 	check(Settings);
 
 	TArray<FPCGTaggedData> Inputs = Context->InputData.GetInputs();
-	UPCGParamData* Params = Context->InputData.GetParams();
 	TArray<FPCGTaggedData>& Outputs = Context->OutputData.TaggedData;
 
 	// Forward any non-input data
 	Outputs.Append(Context->InputData.GetAllSettings());
 
-	const FName PartitionAttribute = PCGSettingsHelpers::GetValue(GET_MEMBER_NAME_CHECKED(UPCGMetadataPartitionSettings, PartitionAttribute), Settings->PartitionAttribute, Params);
+	const FName PartitionAttribute = Settings->PartitionAttribute;
 
 	for (const FPCGTaggedData& Input : Inputs)
 	{
diff --git a/Engine/Plugins/Experimental/PCG/Source/PCG/Private/Elements/Metadata/PCGMetadataRenameElement.cpp b/Engine/Plugins/Experimental/PCG/Source/PCG/Private/Elements/Metadata/PCGMetadataRenameElement.cpp
--- a/Engine/Plugins/Experimental/PCG/Source/PCG/Private/Elements/Metadata/PCGMetadataRenameElement.cpp
+++ b/Engine/Plugins/Experimental/PCG/Source/PCG/Private/Elements/Metadata/PCGMetadataRenameElement.cpp
@@ -2,22 +2,14 @@
 
 #include "Elements/Metadata/PCGMetadataRenameElement.h"
 
-#include "Data/PCGSpatialData.h"
-#include "Elements/Metadata/PCGMetadataElementCommon.h"
-#include "Helpers/PCGSettingsHelpers.h"
 #include "PCGContext.h"
 #include "PCGPin.h"
+#include "PCGParamData.h"
+#include "Data/PCGSpatialData.h"
+#include "Elements/Metadata/PCGMetadataElementCommon.h"
 
 #include UE_INLINE_GENERATED_CPP_BY_NAME(PCGMetadataRenameElement)
 
-TArray<FPCGPinProperties> UPCGMetadataRenameSettings::InputPinProperties() const
-{
-	TArray<FPCGPinProperties> PinProperties;
-	PinProperties.Emplace(PCGPinConstants::DefaultInputLabel, EPCGDataType::Any, /*bInAllowMultipleConnections=*/ true);
-	PinProperties.Emplace(PCGPinConstants::DefaultParamsLabel, EPCGDataType::Param, /*bInAllowMultipleConnections=*/ false);
-
-	return PinProperties;
-}
 
 TArray<FPCGPinProperties> UPCGMetadataRenameSettings::OutputPinProperties() const
 {
@@ -40,13 +32,10 @@ bool FPCGMetadataRenameElement::ExecuteInternal(FPCGContext* Context) const
 	check(Settings);
 
 	TArray<FPCGTaggedData> Inputs = Context->InputData.GetInputsByPin(PCGPinConstants::DefaultInputLabel);
-	TArray<FPCGTaggedData> InputParams = Context->InputData.GetInputsByPin(PCGPinConstants::DefaultParamsLabel);
-	const UPCGParamData* Params = InputParams.IsEmpty() ? nullptr : Cast<const UPCGParamData>(InputParams[0].Data);
-
 	TArray<FPCGTaggedData>& Outputs = Context->OutputData.TaggedData;
 	
-	const FName AttributeToRename = PCGSettingsHelpers::GetValue(GET_MEMBER_NAME_CHECKED(UPCGMetadataRenameSettings, AttributeToRename), Settings->AttributeToRename, Params);
-	const FName NewAttributeName = PCGSettingsHelpers::GetValue(GET_MEMBER_NAME_CHECKED(UPCGMetadataRenameSettings, NewAttributeName), Settings->NewAttributeName, Params);
+	const FName AttributeToRename = Settings->AttributeToRename;
+	const FName NewAttributeName = Settings->NewAttributeName;
 
 	if (NewAttributeName == NAME_None)
 	{
diff --git a/Engine/Plugins/Experimental/PCG/Source/PCG/Private/Elements/PCGAttributeGetFromPointIndexElement.cpp b/Engine/Plugins/Experimental/PCG/Source/PCG/Private/Elements/PCGAttributeGetFromPointIndexElement.cpp
--- a/Engine/Plugins/Experimental/PCG/Source/PCG/Private/Elements/PCGAttributeGetFromPointIndexElement.cpp
+++ b/Engine/Plugins/Experimental/PCG/Source/PCG/Private/Elements/PCGAttributeGetFromPointIndexElement.cpp
@@ -2,13 +2,13 @@
 
 #include "Elements/PCGAttributeGetFromPointIndexElement.h"
 
+#include "PCGContext.h"
+#include "PCGPin.h"
+#include "PCGParamData.h"
 #include "Data/PCGPointData.h"
-#include "Helpers/PCGSettingsHelpers.h"
 #include "Metadata/Accessors/IPCGAttributeAccessor.h"
 #include "Metadata/Accessors/PCGAttributeAccessorKeys.h"
 #include "Metadata/Accessors/PCGAttributeAccessorHelpers.h"
-#include "PCGContext.h"
-#include "PCGPin.h"
 
 #include UE_INLINE_GENERATED_CPP_BY_NAME(PCGAttributeGetFromPointIndexElement)
 
@@ -36,7 +36,6 @@ TArray<FPCGPinProperties> UPCGAttributeGetFromPointIndexSettings::InputPinProper
 {
 	TArray<FPCGPinProperties> PinProperties;
 	PinProperties.Emplace(PCGPinConstants::DefaultInputLabel, EPCGDataType::Point, /*bInAllowMultipleConnections=*/ false);
-	PinProperties.Emplace(PCGPinConstants::DefaultParamsLabel, EPCGDataType::Param, /*bInAllowMultipleConnections=*/ false);
 
 	return PinProperties;
 }
@@ -80,9 +79,7 @@ bool FPCGAttributeGetFromPointIndexElement::ExecuteInternal(FPCGContext* Context
 		return true;
 	}
 
-	const UPCGParamData* ParamData = Context->InputData.GetParams();
-
-	const int32 Index = PCG_GET_OVERRIDEN_VALUE(Settings, Index, ParamData);
+	const int32 Index = Settings->Index;
 
 	if (Index < 0 || Index >= PointData->GetPoints().Num())
 	{
diff --git a/Engine/Plugins/Experimental/PCG/Source/PCG/Private/Elements/PCGAttributeSelectElement.cpp b/Engine/Plugins/Experimental/PCG/Source/PCG/Private/Elements/PCGAttributeSelectElement.cpp
--- a/Engine/Plugins/Experimental/PCG/Source/PCG/Private/Elements/PCGAttributeSelectElement.cpp
+++ b/Engine/Plugins/Experimental/PCG/Source/PCG/Private/Elements/PCGAttributeSelectElement.cpp
@@ -227,7 +227,6 @@ TArray<FPCGPinProperties> UPCGAttributeSelectSettings::InputPinProperties() cons
 {
 	TArray<FPCGPinProperties> PinProperties;
 	PinProperties.Emplace(PCGPinConstants::DefaultInputLabel, EPCGDataType::Spatial, /*bInAllowMultipleConnections=*/ false);
-	PinProperties.Emplace(PCGPinConstants::DefaultParamsLabel, EPCGDataType::Param, /*bInAllowMultipleConnections=*/ false);
 
 	return PinProperties;
 }
diff --git a/Engine/Plugins/Experimental/PCG/Source/PCG/Private/Elements/PCGBoundsModifier.cpp b/Engine/Plugins/Experimental/PCG/Source/PCG/Private/Elements/PCGBoundsModifier.cpp
--- a/Engine/Plugins/Experimental/PCG/Source/PCG/Private/Elements/PCGBoundsModifier.cpp
+++ b/Engine/Plugins/Experimental/PCG/Source/PCG/Private/Elements/PCGBoundsModifier.cpp
@@ -2,7 +2,6 @@
 
 #include "Elements/PCGBoundsModifier.h"
 
-#include "Helpers/PCGSettingsHelpers.h"
 #include "PCGContext.h"
 #include "PCGPoint.h"
 
@@ -19,7 +18,6 @@ TArray<FPCGPinProperties> UPCGBoundsModifierSettings::InputPinProperties() const
 {
 	TArray<FPCGPinProperties> PinProperties;
 	PinProperties.Emplace(PCGPinConstants::DefaultInputLabel, EPCGDataType::Point);
-	PinProperties.Emplace(PCGPinConstants::DefaultParamsLabel, EPCGDataType::Param);
 	return PinProperties;
 }
 
@@ -39,13 +37,12 @@ bool FPCGBoundsModifier::ExecuteInternal(FPCGContext* Context) const
 
 	TArray<FPCGTaggedData> Inputs = Context->InputData.GetInputsByPin(PCGPinConstants::DefaultInputLabel);
 	TArray<FPCGTaggedData>& Outputs = Context->OutputData.TaggedData;
-	UPCGParamData* Params = Context->InputData.GetParams();
 
-	const EPCGBoundsModifierMode Mode = PCG_GET_OVERRIDEN_VALUE(Settings, Mode, Params);
-	const FVector BoundsMin = PCG_GET_OVERRIDEN_VALUE(Settings, BoundsMin, Params);
-	const FVector BoundsMax = PCG_GET_OVERRIDEN_VALUE(Settings, BoundsMax, Params);
-	const bool bAffectSteepness = PCG_GET_OVERRIDEN_VALUE(Settings, bAffectSteepness, Params);
-	const float Steepness = PCG_GET_OVERRIDEN_VALUE(Settings, Steepness, Params);
+	const EPCGBoundsModifierMode Mode = Settings->Mode;
+	const FVector& BoundsMin = Settings->BoundsMin;
+	const FVector& BoundsMax = Settings->BoundsMax;
+	const bool bAffectSteepness = Settings->bAffectSteepness;
+	const float Steepness = Settings->Steepness;
 
 	const FBox Bounds(BoundsMin, BoundsMax);
 
diff --git a/Engine/Plugins/Experimental/PCG/Source/PCG/Private/Elements/PCGCopyPoints.cpp b/Engine/Plugins/Experimental/PCG/Source/PCG/Private/Elements/PCGCopyPoints.cpp
--- a/Engine/Plugins/Experimental/PCG/Source/PCG/Private/Elements/PCGCopyPoints.cpp
+++ b/Engine/Plugins/Experimental/PCG/Source/PCG/Private/Elements/PCGCopyPoints.cpp
@@ -7,7 +7,6 @@
 #include "Data/PCGPointData.h"
 #include "Data/PCGSpatialData.h"
 #include "Helpers/PCGAsync.h"
-#include "Helpers/PCGSettingsHelpers.h"
 
 #include UE_INLINE_GENERATED_CPP_BY_NAME(PCGCopyPoints)
 
@@ -16,8 +15,6 @@ TArray<FPCGPinProperties> UPCGCopyPointsSettings::InputPinProperties() const
 	TArray<FPCGPinProperties> PinProperties;
 	PinProperties.Emplace(PCGCopyPointsConstants::SourcePointsLabel, EPCGDataType::Point, /*bAllowMultipleConnections=*/false);
 	PinProperties.Emplace(PCGCopyPointsConstants::TargetPointsLabel, EPCGDataType::Point, /*bAllowMultipleConnections=*/false);
-	PinProperties.Emplace(PCGCopyPointsConstants::ParamsLabel, EPCGDataType::Param, /*bAllowMultipleConnections=*/false);
-
 	return PinProperties;
 }
 
@@ -33,12 +30,11 @@ bool FPCGCopyPointsElement::ExecuteInternal(FPCGContext* Context) const
 	const UPCGCopyPointsSettings* Settings = Context->GetInputSettings<UPCGCopyPointsSettings>();
 	check(Settings);
 
-	UPCGParamData* Params = Context->InputData.GetParams();
-	const EPCGCopyPointsInheritanceMode RotationInheritance = PCG_GET_OVERRIDEN_VALUE(Settings, RotationInheritance, Params);
-	const EPCGCopyPointsInheritanceMode ScaleInheritance = PCG_GET_OVERRIDEN_VALUE(Settings, ScaleInheritance, Params);
-	const EPCGCopyPointsInheritanceMode ColorInheritance = PCG_GET_OVERRIDEN_VALUE(Settings, ColorInheritance, Params);
-	const EPCGCopyPointsInheritanceMode SeedInheritance = PCG_GET_OVERRIDEN_VALUE(Settings, SeedInheritance, Params);
-	const EPCGCopyPointsMetadataInheritanceMode AttributeInheritance = PCG_GET_OVERRIDEN_VALUE(Settings, AttributeInheritance, Params);
+	const EPCGCopyPointsInheritanceMode RotationInheritance = Settings->RotationInheritance;
+	const EPCGCopyPointsInheritanceMode ScaleInheritance = Settings->ScaleInheritance;
+	const EPCGCopyPointsInheritanceMode ColorInheritance = Settings->ColorInheritance;
+	const EPCGCopyPointsInheritanceMode SeedInheritance = Settings->SeedInheritance;
+	const EPCGCopyPointsMetadataInheritanceMode AttributeInheritance = Settings->AttributeInheritance;
 
 	const TArray<FPCGTaggedData> Sources = Context->InputData.GetInputsByPin(PCGCopyPointsConstants::SourcePointsLabel);
 	const TArray<FPCGTaggedData> Targets = Context->InputData.GetInputsByPin(PCGCopyPointsConstants::TargetPointsLabel);
diff --git a/Engine/Plugins/Experimental/PCG/Source/PCG/Private/Elements/PCGCreateAttribute.cpp b/Engine/Plugins/Experimental/PCG/Source/PCG/Private/Elements/PCGCreateAttribute.cpp
--- a/Engine/Plugins/Experimental/PCG/Source/PCG/Private/Elements/PCGCreateAttribute.cpp
+++ b/Engine/Plugins/Experimental/PCG/Source/PCG/Private/Elements/PCGCreateAttribute.cpp
@@ -2,10 +2,11 @@
 
 #include "Elements/PCGCreateAttribute.h"
 
-#include "Data/PCGSpatialData.h"
-#include "Elements/Metadata/PCGMetadataElementCommon.h"
 #include "PCGContext.h"
+#include "PCGParamData.h"
 #include "PCGPin.h"
+#include "Data/PCGSpatialData.h"
+#include "Elements/Metadata/PCGMetadataElementCommon.h"
 
 #include UE_INLINE_GENERATED_CPP_BY_NAME(PCGCreateAttribute)
 
@@ -195,7 +196,7 @@ bool FPCGCreateAttributeElement::ExecuteInternal(FPCGContext* Context) const
 		}
 		else
 		{
-			Attribute = ClearOrCreateAttribute(Settings, Metadata, nullptr, &OutputAttributeName);
+			Attribute = ClearOrCreateAttribute(Settings, Metadata, &OutputAttributeName);
 		}
 
 		if (!Attribute)
@@ -214,42 +215,42 @@ bool FPCGCreateAttributeElement::ExecuteInternal(FPCGContext* Context) const
 			// If the metadata is empty, we need to add a new entry, so set it to PCGInvalidEntryKey.
 			// Otherwise, use the entry key 0.
 			PCGMetadataEntryKey EntryKey = Metadata->GetItemCountForChild() == 0 ? PCGInvalidEntryKey : 0;
-			SetAttribute(Settings, Attribute, Metadata, EntryKey, nullptr);
+			SetAttribute(Settings, Attribute, Metadata, EntryKey);
 		}
 	}
 
 	return true;
 }
 
-FPCGMetadataAttributeBase* FPCGCreateAttributeElement::ClearOrCreateAttribute(const UPCGCreateAttributeSettings* Settings, UPCGMetadata* Metadata, const UPCGParamData* Params, const FName* OutputAttributeNameOverride) const
+FPCGMetadataAttributeBase* FPCGCreateAttributeElement::ClearOrCreateAttribute(const UPCGCreateAttributeSettings* Settings, UPCGMetadata* Metadata, const FName* OutputAttributeNameOverride) const
 {
 	check(Metadata);
 
 	auto CreateAttribute = [Settings, Metadata, OutputAttributeNameOverride](auto&& Value) -> FPCGMetadataAttributeBase*
 	{
-		return PCGMetadataElementCommon::ClearOrCreateAttribute(Metadata, OutputAttributeNameOverride ? *OutputAttributeNameOverride : Settings->OutputAttributeName, Value);
+		return PCGMetadataElementCommon::ClearOrCreateAttribute(Metadata, OutputAttributeNameOverride ? *OutputAttributeNameOverride : Settings->OutputAttributeName, std::forward<decltype(Value)>(Value));
 	};
 
-	return Settings->AttributeTypes.DispatcherWithOverride(Params, CreateAttribute);
+	return Settings->AttributeTypes.Dispatcher(CreateAttribute);
 }
 
-PCGMetadataEntryKey FPCGCreateAttributeElement::SetAttribute(const UPCGCreateAttributeSettings* Settings, FPCGMetadataAttributeBase* Attribute, UPCGMetadata* Metadata, PCGMetadataEntryKey EntryKey, const UPCGParamData* Params) const
+PCGMetadataEntryKey FPCGCreateAttributeElement::SetAttribute(const UPCGCreateAttributeSettings* Settings, FPCGMetadataAttributeBase* Attribute, UPCGMetadata* Metadata, PCGMetadataEntryKey EntryKey) const
 {
 	check(Attribute && Metadata);
 
 	auto SetAttribute = [Attribute, EntryKey, Metadata](auto&& Value) -> PCGMetadataEntryKey
 	{
-		using AttributeType = std::remove_reference_t<decltype(Value)>;
+		using AttributeType = std::decay_t<decltype(Value)>;
 
 		check(Attribute->GetTypeId() == PCG::Private::MetadataTypes<AttributeType>::Id);
 
 		const PCGMetadataEntryKey FinalKey = (EntryKey == PCGInvalidEntryKey) ? Metadata->AddEntry() : EntryKey;
 
-		static_cast<FPCGMetadataAttribute<AttributeType>*>(Attribute)->SetValue(FinalKey, Value);
+		static_cast<FPCGMetadataAttribute<AttributeType>*>(Attribute)->SetValue(FinalKey, std::forward<decltype(Value)>(Value));
 
 		return FinalKey;
 	};
 
-	return Settings->AttributeTypes.DispatcherWithOverride(Params, SetAttribute);
+	return Settings->AttributeTypes.Dispatcher(SetAttribute);
 }
 
diff --git a/Engine/Plugins/Experimental/PCG/Source/PCG/Private/Elements/PCGDataTableRowToParamData.cpp b/Engine/Plugins/Experimental/PCG/Source/PCG/Private/Elements/PCGDataTableRowToParamData.cpp
--- a/Engine/Plugins/Experimental/PCG/Source/PCG/Private/Elements/PCGDataTableRowToParamData.cpp
+++ b/Engine/Plugins/Experimental/PCG/Source/PCG/Private/Elements/PCGDataTableRowToParamData.cpp
@@ -1,10 +1,14 @@
 // Copyright Epic Games, Inc. All Rights Reserved.
 
 #include "Elements/PCGDataTableRowToParamData.h"
-#include "Helpers/PCGSettingsHelpers.h"
-#include "Engine/DataTable.h"
+
 #include "PCGContext.h"
+#include "PCGParamData.h"
 #include "PCGPin.h"
+#include "PCGModule.h"
+#include "Metadata/PCGMetadata.h"
+
+#include "Engine/DataTable.h"
 
 #include UE_INLINE_GENERATED_CPP_BY_NAME(PCGDataTableRowToParamData)
 
@@ -15,10 +19,7 @@ FPCGElementPtr UPCGDataTableRowToParamDataSettings::CreateElement() const
 
 TArray<FPCGPinProperties> UPCGDataTableRowToParamDataSettings::InputPinProperties() const
 {
-	TArray<FPCGPinProperties> PinProperties;
-	PinProperties.Emplace(PCGPinConstants::DefaultParamsLabel, EPCGDataType::Param, false);
-
-	return PinProperties;
+	return TArray<FPCGPinProperties>();
 }
 
 TArray<FPCGPinProperties> UPCGDataTableRowToParamDataSettings::OutputPinProperties() const
@@ -38,13 +39,11 @@ bool FPCGDataTableRowToParamData::ExecuteInternal(FPCGContext* Context) const
 {
 	TRACE_CPUPROFILER_EVENT_SCOPE(FPCGDataTableRowToParamData::Execute);
 
-	UPCGParamData* Params = Context->InputData.GetParams();
-
 	const UPCGDataTableRowToParamDataSettings* Settings = Context->GetInputSettings<UPCGDataTableRowToParamDataSettings>();
 	check(Settings);
 
-	const FString PathOverride = PCGSettingsHelpers::GetValue(GET_MEMBER_NAME_CHECKED(UPCGDataTableRowToParamDataSettings, DataTable), FString(), Params);
-	const FName RowName = PCGSettingsHelpers::GetValue(GET_MEMBER_NAME_CHECKED(UPCGDataTableRowToParamDataSettings, RowName), Settings->RowName, Params);
+	const FString& PathOverride = Settings->PathOverride;
+	const FName RowName = Settings->RowName;
 
 	TSoftObjectPtr<UDataTable> DataTablePtr = Settings->DataTable;
 
diff --git a/Engine/Plugins/Experimental/PCG/Source/PCG/Private/Elements/PCGDensityFilter.cpp b/Engine/Plugins/Experimental/PCG/Source/PCG/Private/Elements/PCGDensityFilter.cpp
--- a/Engine/Plugins/Experimental/PCG/Source/PCG/Private/Elements/PCGDensityFilter.cpp
+++ b/Engine/Plugins/Experimental/PCG/Source/PCG/Private/Elements/PCGDensityFilter.cpp
@@ -6,7 +6,6 @@
 #include "PCGCustomVersion.h"
 #include "Data/PCGPointData.h"
 #include "Helpers/PCGAsync.h"
-#include "Helpers/PCGSettingsHelpers.h"
 #include "PCGContext.h"
 
 #include UE_INLINE_GENERATED_CPP_BY_NAME(PCGDensityFilter)
@@ -16,7 +15,6 @@ TArray<FPCGPinProperties> UPCGDensityFilterSettings::InputPinProperties() const
 	TArray<FPCGPinProperties> PinProperties;
 	// TODO in the future type checking of edges will be stricter and a conversion node will be added to convert from other types
 	PinProperties.Emplace(PCGPinConstants::DefaultInputLabel, EPCGDataType::Point);
-	PinProperties.Emplace(PCGPinConstants::DefaultParamsLabel, EPCGDataType::Param, /*bInAllowMultipleConnections=*/false);
 
 	return PinProperties;
 }
@@ -35,16 +33,15 @@ bool FPCGDensityFilterElement::ExecuteInternal(FPCGContext* Context) const
 
 	TArray<FPCGTaggedData> Inputs = Context->InputData.GetInputs();
 	TArray<FPCGTaggedData>& Outputs = Context->OutputData.TaggedData;
-	UPCGParamData* Params = Context->InputData.GetParams();
 
 	// Forward any non-input data
 	Outputs.Append(Context->InputData.GetAllSettings());
 
-	const bool bInvertFilter = PCGSettingsHelpers::GetValue(GET_MEMBER_NAME_CHECKED(UPCGDensityFilterSettings, bInvertFilter), Settings->bInvertFilter, Params);
-	const float LowerBound = PCGSettingsHelpers::GetValue(GET_MEMBER_NAME_CHECKED(UPCGDensityFilterSettings, LowerBound), Settings->LowerBound, Params);
-	const float UpperBound = PCGSettingsHelpers::GetValue(GET_MEMBER_NAME_CHECKED(UPCGDensityFilterSettings, UpperBound), Settings->UpperBound, Params);
+	const bool bInvertFilter = Settings->bInvertFilter;
+	const float LowerBound = Settings->LowerBound;
+	const float UpperBound = Settings->UpperBound;
 #if WITH_EDITOR
-	const bool bKeepZeroDensityPoints = PCGSettingsHelpers::GetValue(GET_MEMBER_NAME_CHECKED(UPCGDensityFilterSettings, bKeepZeroDensityPoints), Settings->bKeepZeroDensityPoints, Params);
+	const bool bKeepZeroDensityPoints = Settings->bKeepZeroDensityPoints;
 #else
 	const bool bKeepZeroDensityPoints = false;
 #endif
@@ -123,18 +120,4 @@ bool FPCGDensityFilterElement::ExecuteInternal(FPCGContext* Context) const
 	}
 
 	return true;
-}
-
-#if WITH_EDITOR
-void UPCGDensityFilterSettings::ApplyDeprecationBeforeUpdatePins(UPCGNode* InOutNode, TArray<TObjectPtr<UPCGPin>>& InputPins, TArray<TObjectPtr<UPCGPin>>& OutputPins)
-{
-	Super::ApplyDeprecationBeforeUpdatePins(InOutNode, InputPins, OutputPins);
-
-	check(InOutNode);
-
-	if (DataVersion < FPCGCustomVersion::MoveParamsOffFirstPinDensityNodes)
-	{
-		PCGSettingsHelpers::DeprecationBreakOutParamsToNewPin(InOutNode, InputPins, OutputPins);
-	}
-}
-#endif // WITH_EDITOR
+}
\ No newline at end of file
diff --git a/Engine/Plugins/Experimental/PCG/Source/PCG/Private/Elements/PCGDensityNoise.cpp b/Engine/Plugins/Experimental/PCG/Source/PCG/Private/Elements/PCGDensityNoise.cpp
--- a/Engine/Plugins/Experimental/PCG/Source/PCG/Private/Elements/PCGDensityNoise.cpp
+++ b/Engine/Plugins/Experimental/PCG/Source/PCG/Private/Elements/PCGDensityNoise.cpp
@@ -5,7 +5,6 @@
 #include "PCGContext.h"
 #include "PCGCustomVersion.h"
 #include "PCGHelpers.h"
-#include "Helpers/PCGSettingsHelpers.h"
 
 #include "Math/RandomStream.h"
 #include "PCGPoint.h"
@@ -22,7 +21,6 @@ TArray<FPCGPinProperties> UPCGDensityNoiseSettings::InputPinProperties() const
 	TArray<FPCGPinProperties> PinProperties;
 	// TODO in the future type checking of edges will be stricter and a conversion node will be added to convert from other types
 	PinProperties.Emplace(PCGPinConstants::DefaultInputLabel, EPCGDataType::Point);
-	PinProperties.Emplace(PCGPinConstants::DefaultParamsLabel, EPCGDataType::Param, /*bInAllowMultipleConnections=*/false);
 
 	return PinProperties;
 }
@@ -41,15 +39,14 @@ bool FPCGDensityNoiseElement::ExecuteInternal(FPCGContext* Context) const
 
 	TArray<FPCGTaggedData> Inputs = Context->InputData.GetInputs();
 	TArray<FPCGTaggedData>& Outputs = Context->OutputData.TaggedData;
-	UPCGParamData* Params = Context->InputData.GetParams();
 
-	const EPCGDensityNoiseMode DensityMode = PCG_GET_OVERRIDEN_VALUE(Settings, DensityMode, Params);
-	const float DensityNoiseMin = PCG_GET_OVERRIDEN_VALUE(Settings, DensityNoiseMin, Params);
-	const float DensityNoiseMax = PCG_GET_OVERRIDEN_VALUE(Settings, DensityNoiseMax, Params);
-	const bool bInvertSourceDensity = PCG_GET_OVERRIDEN_VALUE(Settings, bInvertSourceDensity, Params);
+	const EPCGDensityNoiseMode DensityMode = Settings->DensityMode;
+	const float DensityNoiseMin = Settings->DensityNoiseMin;
+	const float DensityNoiseMax = Settings->DensityNoiseMax;
+	const bool bInvertSourceDensity = Settings->bInvertSourceDensity;
 
 	// Precompute a seed based on the settings one and the component one
-	const int Seed = PCGSettingsHelpers::ComputeSeedWithOverride(Settings, Context->SourceComponent, Params);
+	const int Seed = Context->GetSeed();
 
 	ProcessPoints(Context, Inputs, Outputs, [Seed, DensityNoiseMin, DensityNoiseMax, bInvertSourceDensity, DensityMode](const FPCGPoint& InPoint, FPCGPoint& OutPoint)
 	{
@@ -93,18 +90,4 @@ bool FPCGDensityNoiseElement::ExecuteInternal(FPCGContext* Context) const
 	Outputs.Append(Context->InputData.GetAllSettings());
 
 	return true;
-}
-
-#if WITH_EDITOR
-void UPCGDensityNoiseSettings::ApplyDeprecationBeforeUpdatePins(UPCGNode* InOutNode, TArray<TObjectPtr<UPCGPin>>& InputPins, TArray<TObjectPtr<UPCGPin>>& OutputPins)
-{
-	Super::ApplyDeprecationBeforeUpdatePins(InOutNode, InputPins, OutputPins);
-
-	check(InOutNode);
-
-	if (DataVersion < FPCGCustomVersion::MoveParamsOffFirstPinDensityNodes)
-	{
-		PCGSettingsHelpers::DeprecationBreakOutParamsToNewPin(InOutNode, InputPins, OutputPins);
-	}
-}
-#endif // WITH_EDITOR
+}
\ No newline at end of file
diff --git a/Engine/Plugins/Experimental/PCG/Source/PCG/Private/Elements/PCGDensityRemap.cpp b/Engine/Plugins/Experimental/PCG/Source/PCG/Private/Elements/PCGDensityRemap.cpp
--- a/Engine/Plugins/Experimental/PCG/Source/PCG/Private/Elements/PCGDensityRemap.cpp
+++ b/Engine/Plugins/Experimental/PCG/Source/PCG/Private/Elements/PCGDensityRemap.cpp
@@ -1,16 +1,26 @@
 // Copyright Epic Games, Inc. All Rights Reserved.
 
 #include "Elements/PCGDensityRemap.h"
-#include "Data/PCGSpatialData.h"
+
+#include "PCGContext.h"
+#include "PCGCustomVersion.h"
 #include "PCGHelpers.h"
 #include "Data/PCGPointData.h"
+#include "Data/PCGSpatialData.h"
 #include "Helpers/PCGAsync.h"
-#include "Helpers/PCGSettingsHelpers.h"
+
 #include "Math/RandomStream.h"
-#include "PCGContext.h"
 
 #include UE_INLINE_GENERATED_CPP_BY_NAME(PCGDensityRemap)
 
+TArray<FPCGPinProperties> UPCGLinearDensityRemapSettings::InputPinProperties() const
+{
+	TArray<FPCGPinProperties> PinProperties;
+	PinProperties.Emplace(PCGPinConstants::DefaultInputLabel, EPCGDataType::Spatial, /*bInAllowMultipleConnections=*/ true);
+
+	return PinProperties;
+}
+
 UPCGLinearDensityRemapSettings::UPCGLinearDensityRemapSettings()
 {
 	bUseSeed = true;
@@ -30,16 +40,15 @@ bool FPCGLinearDensityRemapElement::ExecuteInternal(FPCGContext* Context) const
 
 	TArray<FPCGTaggedData> Inputs = Context->InputData.GetInputs();
 	TArray<FPCGTaggedData>& Outputs = Context->OutputData.TaggedData;
-	UPCGParamData* Params = Context->InputData.GetParams();
 
 	// Forward any non-input data
 	Outputs.Append(Context->InputData.GetAllSettings());
 
-	const float SettingsRemapMin = PCGSettingsHelpers::GetValue(GET_MEMBER_NAME_CHECKED(UPCGLinearDensityRemapSettings, RemapMin), Settings->RemapMin, Params);
-	const float SettingsRemapMax = PCGSettingsHelpers::GetValue(GET_MEMBER_NAME_CHECKED(UPCGLinearDensityRemapSettings, RemapMax), Settings->RemapMax, Params);
-	const bool bMultiplyDensity = PCGSettingsHelpers::GetValue(GET_MEMBER_NAME_CHECKED(UPCGLinearDensityRemapSettings, bMultiplyDensity), Settings->bMultiplyDensity, Params);
+	const float SettingsRemapMin = Settings->RemapMin;
+	const float SettingsRemapMax = Settings->RemapMax;
+	const bool bMultiplyDensity = Settings->bMultiplyDensity;
 
-	const int Seed = PCGSettingsHelpers::ComputeSeedWithOverride(Settings, Context->SourceComponent, Params);
+	const int Seed = Context->GetSeed();
 
 	const float RemapMin = FMath::Min(SettingsRemapMin, SettingsRemapMax);
 	const float RemapMax = FMath::Max(SettingsRemapMin, SettingsRemapMax);
diff --git a/Engine/Plugins/Experimental/PCG/Source/PCG/Private/Elements/PCGDensityRemapElement.cpp b/Engine/Plugins/Experimental/PCG/Source/PCG/Private/Elements/PCGDensityRemapElement.cpp
--- a/Engine/Plugins/Experimental/PCG/Source/PCG/Private/Elements/PCGDensityRemapElement.cpp
+++ b/Engine/Plugins/Experimental/PCG/Source/PCG/Private/Elements/PCGDensityRemapElement.cpp
@@ -4,7 +4,6 @@
 
 #include "PCGContext.h"
 #include "PCGCustomVersion.h"
-#include "Helpers/PCGSettingsHelpers.h"
 
 #include "PCGPoint.h"
 
@@ -15,7 +14,6 @@ TArray<FPCGPinProperties> UPCGDensityRemapSettings::InputPinProperties() const
 	TArray<FPCGPinProperties> PinProperties;
 	// TODO in the future type checking of edges will be stricter and a conversion node will be added to convert from other types
 	PinProperties.Emplace(PCGPinConstants::DefaultInputLabel, EPCGDataType::Point);
-	PinProperties.Emplace(PCGPinConstants::DefaultParamsLabel, EPCGDataType::Param, /*bInAllowMultipleConnections=*/false);
 
 	return PinProperties;
 }
@@ -34,13 +32,12 @@ bool FPCGDensityRemapElement::ExecuteInternal(FPCGContext* Context) const
 
 	TArray<FPCGTaggedData> Inputs = Context->InputData.GetInputs();
 	TArray<FPCGTaggedData>& Outputs = Context->OutputData.TaggedData;
-	UPCGParamData* Params = Context->InputData.GetParams();
 
-	const float InRangeMin = PCG_GET_OVERRIDEN_VALUE(Settings, InRangeMin, Params);
-	const float InRangeMax = PCG_GET_OVERRIDEN_VALUE(Settings, InRangeMax, Params);
-	const float OutRangeMin = PCG_GET_OVERRIDEN_VALUE(Settings, OutRangeMin, Params);
-	const float OutRangeMax = PCG_GET_OVERRIDEN_VALUE(Settings, OutRangeMax, Params);
-	const bool bExcludeValuesOutsideInputRange = PCG_GET_OVERRIDEN_VALUE(Settings, bExcludeValuesOutsideInputRange, Params);
+	const float InRangeMin = Settings->InRangeMin;
+	const float InRangeMax = Settings->InRangeMax;
+	const float OutRangeMin = Settings->OutRangeMin;
+	const float OutRangeMax = Settings->OutRangeMax;
+	const bool bExcludeValuesOutsideInputRange = Settings->bExcludeValuesOutsideInputRange;
 
 	// used to determine if a density value lies between TrueMin and TrueMax
 	const float InRangeTrueMin = FMath::Min(InRangeMin, InRangeMax);
@@ -83,17 +80,3 @@ bool FPCGDensityRemapElement::ExecuteInternal(FPCGContext* Context) const
 
 	return true;
 }
-
-#if WITH_EDITOR
-void UPCGDensityRemapSettings::ApplyDeprecationBeforeUpdatePins(UPCGNode* InOutNode, TArray<TObjectPtr<UPCGPin>>& InputPins, TArray<TObjectPtr<UPCGPin>>& OutputPins)
-{
-	Super::ApplyDeprecationBeforeUpdatePins(InOutNode, InputPins, OutputPins);
-
-	check(InOutNode);
-
-	if (DataVersion < FPCGCustomVersion::MoveParamsOffFirstPinDensityNodes)
-	{
-		PCGSettingsHelpers::DeprecationBreakOutParamsToNewPin(InOutNode, InputPins, OutputPins);
-	}
-}
-#endif // WITH_EDITOR
diff --git a/Engine/Plugins/Experimental/PCG/Source/PCG/Private/Elements/PCGDifferenceElement.cpp b/Engine/Plugins/Experimental/PCG/Source/PCG/Private/Elements/PCGDifferenceElement.cpp
--- a/Engine/Plugins/Experimental/PCG/Source/PCG/Private/Elements/PCGDifferenceElement.cpp
+++ b/Engine/Plugins/Experimental/PCG/Source/PCG/Private/Elements/PCGDifferenceElement.cpp
@@ -2,19 +2,19 @@
 
 #include "Elements/PCGDifferenceElement.h"
 
+#include "PCGContext.h"
+#include "PCGCustomVersion.h"
+#include "PCGPin.h"
 #include "Data/PCGPointData.h"
 #include "Data/PCGSpatialData.h"
 #include "Data/PCGUnionData.h"
-#include "Helpers/PCGSettingsHelpers.h"
-#include "PCGContext.h"
-#include "PCGPin.h"
 
 #include UE_INLINE_GENERATED_CPP_BY_NAME(PCGDifferenceElement)
 
 TArray<FPCGPinProperties> UPCGDifferenceSettings::InputPinProperties() const
 {
 	TArray<FPCGPinProperties> PinProperties;
-	PinProperties.Emplace(PCGDifferenceConstants::SourceLabel, EPCGDataType::Any);
+	PinProperties.Emplace(PCGDifferenceConstants::SourceLabel, EPCGDataType::Spatial);
 	PinProperties.Emplace(PCGDifferenceConstants::DifferencesLabel, EPCGDataType::Spatial);
 
 	return PinProperties;
@@ -149,11 +149,10 @@ void FPCGDifferenceElement::LabellessProcessing(FPCGContext* Context) const
 	check(Settings);
 
 	TArray<FPCGTaggedData> Inputs = Context->InputData.GetInputs();
-	UPCGParamData* Params = Context->InputData.GetParams();
 
-	const EPCGDifferenceDensityFunction DensityFunction = PCGSettingsHelpers::GetValue(GET_MEMBER_NAME_CHECKED(UPCGDifferenceSettings, DensityFunction), Settings->DensityFunction, Params);
+	const EPCGDifferenceDensityFunction DensityFunction = Settings->DensityFunction;
 #if WITH_EDITOR
-	const bool bKeepZeroDensityPoints = PCGSettingsHelpers::GetValue(GET_MEMBER_NAME_CHECKED(UPCGDifferenceSettings, bKeepZeroDensityPoints), Settings->bKeepZeroDensityPoints, Params);
+	const bool bKeepZeroDensityPoints = Settings->bKeepZeroDensityPoints;
 #else
 	const bool bKeepZeroDensityPoints = false;
 #endif
diff --git a/Engine/Plugins/Experimental/PCG/Source/PCG/Private/Elements/PCGDistance.cpp b/Engine/Plugins/Experimental/PCG/Source/PCG/Private/Elements/PCGDistance.cpp
--- a/Engine/Plugins/Experimental/PCG/Source/PCG/Private/Elements/PCGDistance.cpp
+++ b/Engine/Plugins/Experimental/PCG/Source/PCG/Private/Elements/PCGDistance.cpp
@@ -3,7 +3,6 @@
 #include "Elements/PCGDistance.h"
 #include "Data/PCGSpatialData.h"
 #include "Helpers/PCGAsync.h"
-#include "Helpers/PCGSettingsHelpers.h"
 #include "Data/PCGPointData.h"
 #include "PCGContext.h"
 #include "PCGPin.h"
@@ -71,13 +70,11 @@ bool FPCGDistanceElement::ExecuteInternal(FPCGContext* Context) const
 	const UPCGDistanceSettings* Settings = Context->GetInputSettings<UPCGDistanceSettings>();
 	check(Settings);
 
-	UPCGParamData* Params = Context->InputData.GetParams();
-
-	const FName AttributeName = PCGSettingsHelpers::GetValue(GET_MEMBER_NAME_CHECKED(UPCGDistanceSettings, AttributeName), Settings->AttributeName, Params);
-	const bool bSetDensity = PCGSettingsHelpers::GetValue(GET_MEMBER_NAME_CHECKED(UPCGDistanceSettings, bSetDensity), Settings->bSetDensity, Params);
-	const double MaximumDistance = PCGSettingsHelpers::GetValue(GET_MEMBER_NAME_CHECKED(UPCGDistanceSettings, MaximumDistance), Settings->MaximumDistance, Params);
-	const PCGDistanceShape SourceShape = PCGSettingsHelpers::GetValue(GET_MEMBER_NAME_CHECKED(UPCGDistanceSettings, SourceShape), Settings->SourceShape, Params);
-	const PCGDistanceShape TargetShape = PCGSettingsHelpers::GetValue(GET_MEMBER_NAME_CHECKED(UPCGDistanceSettings, TargetShape), Settings->TargetShape, Params);
+	const FName AttributeName = Settings->AttributeName;
+	const bool bSetDensity = Settings->bSetDensity;
+	const double MaximumDistance = Settings->MaximumDistance;
+	const PCGDistanceShape SourceShape = Settings->SourceShape;
+	const PCGDistanceShape TargetShape = Settings->TargetShape;
 
 	TArray<FPCGTaggedData> Sources = Context->InputData.GetInputsByPin(PCGDistance::SourceLabel);
 	TArray<FPCGTaggedData> Targets = Context->InputData.GetInputsByPin(PCGDistance::TargetLabel);
diff --git a/Engine/Plugins/Experimental/PCG/Source/PCG/Private/Elements/PCGExecuteBlueprint.cpp b/Engine/Plugins/Experimental/PCG/Source/PCG/Private/Elements/PCGExecuteBlueprint.cpp
--- a/Engine/Plugins/Experimental/PCG/Source/PCG/Private/Elements/PCGExecuteBlueprint.cpp
+++ b/Engine/Plugins/Experimental/PCG/Source/PCG/Private/Elements/PCGExecuteBlueprint.cpp
@@ -166,25 +166,7 @@ TSet<FName> UPCGBlueprintElement::OutputLabels() const
 
 int UPCGBlueprintElement::GetSeed(FPCGContext& InContext) const 
 {
-	const UPCGComponent* SourceComponent = InContext.SourceComponent.Get();
-	const UPCGBlueprintSettings* Settings = InContext.GetInputSettings<UPCGBlueprintSettings>();
-
-	if (SourceComponent && Settings)
-	{
-		return PCGHelpers::ComputeSeed(Settings->Seed, SourceComponent->Seed);
-	}
-	else if (Settings)
-	{
-		return Settings->Seed;
-	}
-	else if (SourceComponent)
-	{
-		return SourceComponent->Seed;
-	}
-	else
-	{
-		return 42;
-	}
+	return InContext.GetSeed();
 }
 
 FRandomStream UPCGBlueprintElement::GetRandomStream(FPCGContext& InContext) const
@@ -441,18 +423,14 @@ TArray<FPCGPinProperties> UPCGBlueprintSettings::InputPinProperties() const
 {
 	TArray<FPCGPinProperties> PinProperties;
 
-	if (BlueprintElementInstance)
+	if (!BlueprintElementInstance || BlueprintElementInstance->bHasDefaultInPin)
 	{
-		if (BlueprintElementInstance->bHasDefaultInPin)
-		{
-			PinProperties.Append(Super::InputPinProperties());
-		}
-
-		PinProperties.Append(BlueprintElementInstance->CustomInputPins);
+		PinProperties.Append(Super::InputPinProperties());
 	}
-	else
+
+	if (BlueprintElementInstance)
 	{
-		PinProperties = Super::InputPinProperties();
+		PinProperties.Append(BlueprintElementInstance->CustomInputPins);
 	}
 
 	return PinProperties;
@@ -508,7 +486,7 @@ bool FPCGExecuteBlueprintElement::ExecuteInternal(FPCGContext* InContext) const
 #endif
 
 		/** Apply params overrides to variables if any */
-		if (UPCGParamData* Params = Context->InputData.GetParams())
+		if (UPCGParamData* Params = Context->InputData.GetParamsWithDeprecation(Context->Node))
 		{
 			for (TFieldIterator<FProperty> PropertyIt(BPClass); PropertyIt; ++PropertyIt)
 			{
diff --git a/Engine/Plugins/Experimental/PCG/Source/PCG/Private/Elements/PCGIntersectionElement.cpp b/Engine/Plugins/Experimental/PCG/Source/PCG/Private/Elements/PCGIntersectionElement.cpp
--- a/Engine/Plugins/Experimental/PCG/Source/PCG/Private/Elements/PCGIntersectionElement.cpp
+++ b/Engine/Plugins/Experimental/PCG/Source/PCG/Private/Elements/PCGIntersectionElement.cpp
@@ -2,12 +2,19 @@
 
 #include "Elements/PCGIntersectionElement.h"
 #include "Data/PCGSpatialData.h"
-#include "Helpers/PCGSettingsHelpers.h"
 #include "PCGContext.h"
 #include "PCGPin.h"
 
 #include UE_INLINE_GENERATED_CPP_BY_NAME(PCGIntersectionElement)
 
+TArray<FPCGPinProperties> UPCGIntersectionSettings::InputPinProperties() const
+{
+	TArray<FPCGPinProperties> PinProperties;
+	PinProperties.Emplace(PCGPinConstants::DefaultInputLabel, EPCGDataType::Spatial);
+
+	return PinProperties;
+}
+
 TArray<FPCGPinProperties> UPCGIntersectionSettings::OutputPinProperties() const
 {
 	TArray<FPCGPinProperties> PinProperties;
@@ -29,11 +36,10 @@ bool FPCGIntersectionElement::ExecuteInternal(FPCGContext* Context) const
 	check(Settings);
 
 	TArray<FPCGTaggedData> Inputs = Context->InputData.GetInputs();
-	UPCGParamData* Params = Context->InputData.GetParams();
 
-	const EPCGIntersectionDensityFunction DensityFunction = PCGSettingsHelpers::GetValue(GET_MEMBER_NAME_CHECKED(UPCGIntersectionSettings, DensityFunction), Settings->DensityFunction, Params);
+	const EPCGIntersectionDensityFunction DensityFunction = Settings->DensityFunction;
 #if WITH_EDITOR
-	const bool bKeepZeroDensityPoints = PCGSettingsHelpers::GetValue(GET_MEMBER_NAME_CHECKED(UPCGIntersectionSettings, bKeepZeroDensityPoints), Settings->bKeepZeroDensityPoints, Params);
+	const bool bKeepZeroDensityPoints = Settings->bKeepZeroDensityPoints;
 #else
 	const bool bKeepZeroDensityPoints = false;
 #endif
diff --git a/Engine/Plugins/Experimental/PCG/Source/PCG/Private/Elements/PCGPointExtentsModifier.cpp b/Engine/Plugins/Experimental/PCG/Source/PCG/Private/Elements/PCGPointExtentsModifier.cpp
--- a/Engine/Plugins/Experimental/PCG/Source/PCG/Private/Elements/PCGPointExtentsModifier.cpp
+++ b/Engine/Plugins/Experimental/PCG/Source/PCG/Private/Elements/PCGPointExtentsModifier.cpp
@@ -2,7 +2,6 @@
 
 #include "Elements/PCGPointExtentsModifier.h"
 
-#include "Helpers/PCGSettingsHelpers.h"
 #include "PCGContext.h"
 #include "PCGPoint.h"
 
@@ -17,7 +16,6 @@ TArray<FPCGPinProperties> UPCGPointExtentsModifierSettings::InputPinProperties()
 {
 	TArray<FPCGPinProperties> PinProperties;
 	PinProperties.Emplace(PCGPinConstants::DefaultInputLabel, EPCGDataType::Point);	
-	PinProperties.Emplace(PCGPinConstants::DefaultParamsLabel, EPCGDataType::Param);
 	return PinProperties;
 }
 
@@ -30,10 +28,9 @@ bool FPCGPointExtentsModifier::ExecuteInternal(FPCGContext* Context) const
 
 	TArray<FPCGTaggedData> Inputs = Context->InputData.GetInputsByPin(PCGPinConstants::DefaultInputLabel);
 	TArray<FPCGTaggedData>& Outputs = Context->OutputData.TaggedData;
-	UPCGParamData* Params = Context->InputData.GetParams();
 
-	const EPCGPointExtentsModifierMode Mode = PCG_GET_OVERRIDEN_VALUE(Settings, Mode, Params);
-	const FVector Extents = PCG_GET_OVERRIDEN_VALUE(Settings, Extents, Params);
+	const EPCGPointExtentsModifierMode Mode = Settings->Mode;
+	const FVector& Extents = Settings->Extents;
 
 	switch (Mode)
 	{
diff --git a/Engine/Plugins/Experimental/PCG/Source/PCG/Private/Elements/PCGPointFilter.cpp b/Engine/Plugins/Experimental/PCG/Source/PCG/Private/Elements/PCGPointFilter.cpp
--- a/Engine/Plugins/Experimental/PCG/Source/PCG/Private/Elements/PCGPointFilter.cpp
+++ b/Engine/Plugins/Experimental/PCG/Source/PCG/Private/Elements/PCGPointFilter.cpp
@@ -70,8 +70,6 @@ TArray<FPCGPinProperties> UPCGPointFilterSettings::InputPinProperties() const
 		PinProperties.Emplace(PCGPointFilterConstants::FilterLabel, EPCGDataType::Any, /*bInAllowMultipleConnections=*/ false);
 	}
 
-	PinProperties.Emplace(PCGPinConstants::DefaultParamsLabel, EPCGDataType::Param, /*bAllowMultipleConnections=*/false, /*bAllowMultiData=*/false);
-
 	return PinProperties;
 }
 
@@ -176,19 +174,6 @@ PRAGMA_ENABLE_DEPRECATION_WARNINGS
 #endif // WITH_EDITOR
 }
 
-#if WITH_EDITOR
-void UPCGPointFilterSettings::ApplyDeprecationBeforeUpdatePins(UPCGNode* InOutNode, TArray<TObjectPtr<UPCGPin>>& InputPins, TArray<TObjectPtr<UPCGPin>>& OutputPins)
-{
-	Super::ApplyDeprecationBeforeUpdatePins(InOutNode, InputPins, OutputPins);
-
-	check(InOutNode);
-	if (DataVersion < FPCGCustomVersion::MovePointFilterParamsOffFirstPin)
-	{
-		PCGSettingsHelpers::DeprecationBreakOutParamsToNewPin(InOutNode, InputPins, OutputPins);
-	}
-}
-#endif // WITH_EDITOR
-
 bool FPCGPointFilterElement::ExecuteInternal(FPCGContext* Context) const
 {
 	TRACE_CPUPROFILER_EVENT_SCOPE(FPCGPointFilterElement::Execute);
@@ -209,16 +194,15 @@ bool FPCGPointFilterElement::ExecuteInternal(FPCGContext* Context) const
 	check(Settings);
 
 	TArray<FPCGTaggedData> DataToFilter = Context->InputData.GetInputsByPin(PCGPointFilterConstants::DataToFilterLabel);
-	UPCGParamData* Params = Context->InputData.GetParamsOnParamsPin();
 	TArray<FPCGTaggedData>& Outputs = Context->OutputData.TaggedData;
 	TArray<FPCGTaggedData> FilterData = Context->InputData.GetInputsByPin(PCGPointFilterConstants::FilterLabel);
 
 	// Forward any non-input data
 	Outputs.Append(Context->InputData.GetAllSettings());
 
-	const EPCGPointFilterOperator Operator = PCGSettingsHelpers::GetValue(GET_MEMBER_NAME_CHECKED(UPCGPointFilterSettings, Operator), Settings->Operator, Params);
-	const bool bUseSpatialQuery = PCGSettingsHelpers::GetValue(GET_MEMBER_NAME_CHECKED(UPCGPointFilterSettings, bUseSpatialQuery), Settings->bUseSpatialQuery, Params);
-	const bool bUseConstantThreshold = PCGSettingsHelpers::GetValue(GET_MEMBER_NAME_CHECKED(UPCGPointFilterSettings, bUseConstantThreshold), Settings->bUseConstantThreshold, Params);
+	const EPCGPointFilterOperator Operator = Settings->Operator;
+	const bool bUseSpatialQuery = Settings->bUseSpatialQuery;
+	const bool bUseConstantThreshold = Settings->bUseConstantThreshold;
 	// TODO: allow selector overrides for ThresholdAttribute and TargetAttribute
 
 	// If there is no input, do nothing
@@ -240,12 +224,12 @@ bool FPCGPointFilterElement::ExecuteInternal(FPCGContext* Context) const
 		{
 			using ConstantType = std::decay_t<decltype(Value)>;
 
-			ThresholdAccessor = MakeUnique<FPCGConstantValueAccessor<ConstantType>>(std::forward<ConstantType>(Value));
+			ThresholdAccessor = MakeUnique<FPCGConstantValueAccessor<ConstantType>>(std::forward<decltype(Value)>(Value));
 			// Dummy keys
 			ThresholdKeys = MakeUnique<FPCGAttributeAccessorKeysSingleObjectPtr<void>>();
 		};
 
-		Settings->AttributeTypes.DispatcherWithOverride(Params, ConstantThreshold);
+		Settings->AttributeTypes.Dispatcher(ConstantThreshold);
 	}
 	else
 	{
diff --git a/Engine/Plugins/Experimental/PCG/Source/PCG/Private/Elements/PCGPointSampler.cpp b/Engine/Plugins/Experimental/PCG/Source/PCG/Private/Elements/PCGPointSampler.cpp
--- a/Engine/Plugins/Experimental/PCG/Source/PCG/Private/Elements/PCGPointSampler.cpp
+++ b/Engine/Plugins/Experimental/PCG/Source/PCG/Private/Elements/PCGPointSampler.cpp
@@ -5,7 +5,6 @@
 #include "PCGHelpers.h"
 #include "Data/PCGPointData.h"
 #include "Helpers/PCGAsync.h"
-#include "Helpers/PCGSettingsHelpers.h"
 #include "Math/RandomStream.h"
 #include "PCGContext.h"
 
@@ -29,21 +28,20 @@ bool FPCGPointSamplerElement::ExecuteInternal(FPCGContext* Context) const
 	check(Settings);
 
 	TArray<FPCGTaggedData> Inputs = Context->InputData.GetInputs();
-	UPCGParamData* Params = Context->InputData.GetParams();
 
 	TArray<FPCGTaggedData>& Outputs = Context->OutputData.TaggedData;
 
 	// Forward any non-input data, excluding params
 	Outputs.Append(Context->InputData.GetAllSettings());
 
-	const float Ratio = PCGSettingsHelpers::GetValue(GET_MEMBER_NAME_CHECKED(UPCGPointSamplerSettings, Ratio), Settings->Ratio, Params);
+	const float Ratio = Settings->Ratio;
 #if WITH_EDITOR
-	const bool bKeepZeroDensityPoints = PCGSettingsHelpers::GetValue(GET_MEMBER_NAME_CHECKED(UPCGPointSamplerSettings, bKeepZeroDensityPoints), Settings->bKeepZeroDensityPoints, Params);
+	const bool bKeepZeroDensityPoints = Settings->bKeepZeroDensityPoints;
 #else
 	const bool bKeepZeroDensityPoints = false;
 #endif
 
-	const int Seed = PCGSettingsHelpers::ComputeSeedWithOverride(Settings, Context->SourceComponent, Params);
+	const int Seed = Context->GetSeed();
 
 	const bool bNoSampling = (Ratio <= 0.0f);
 	const bool bTrivialSampling = (Ratio >= 1.0f);
diff --git a/Engine/Plugins/Experimental/PCG/Source/PCG/Private/Elements/PCGProjectionElement.cpp b/Engine/Plugins/Experimental/PCG/Source/PCG/Private/Elements/PCGProjectionElement.cpp
--- a/Engine/Plugins/Experimental/PCG/Source/PCG/Private/Elements/PCGProjectionElement.cpp
+++ b/Engine/Plugins/Experimental/PCG/Source/PCG/Private/Elements/PCGProjectionElement.cpp
@@ -6,7 +6,6 @@
 #include "PCGCustomVersion.h"
 #include "Data/PCGSpatialData.h"
 #include "PCGEdge.h"
-#include "Helpers/PCGSettingsHelpers.h"
 #include "PCGContext.h"
 #include "PCGPin.h"
 
@@ -19,7 +18,6 @@ TArray<FPCGPinProperties> UPCGProjectionSettings::InputPinProperties() const
 	TArray<FPCGPinProperties> PinProperties;
 	PinProperties.Emplace(PCGPinConstants::DefaultInputLabel, EPCGDataType::Concrete, /*bInAllowMultipleConnections=*/true, /*bAllowMultipleData=*/true,  LOCTEXT("ProjectionSourcePinTooltip", "The data to project."));
 	PinProperties.Emplace(PCGProjectionConstants::ProjectionTargetLabel, EPCGDataType::Concrete, /*bAllowMultipleConnections=*/false, /*bAllowMultipleData=*/false, LOCTEXT("ProjectionTargetPinTooltip", "The projection target."));
-	PinProperties.Emplace(PCGPinConstants::DefaultParamsLabel, EPCGDataType::Param, /*bInAllowMultipleConnections=*/false, /*bAllowMultipleData=*/false);
 
 	return PinProperties;
 }
@@ -60,18 +58,12 @@ bool FPCGProjectionElement::ExecuteInternal(FPCGContext* Context) const
 		return true;
 	}
 
-	UPCGParamData* Params = Context->InputData.GetParams();
-
-	FPCGProjectionParams ProjectionParams = Settings->Params;
-	ProjectionParams.bProjectPositions = PCGSettingsHelpers::GetValue(GET_MEMBER_NAME_CHECKED(FPCGProjectionParams, bProjectPositions), ProjectionParams.bProjectPositions, Params);
-	ProjectionParams.bProjectRotations = PCGSettingsHelpers::GetValue(GET_MEMBER_NAME_CHECKED(FPCGProjectionParams, bProjectRotations), ProjectionParams.bProjectRotations, Params);
-	ProjectionParams.bProjectScales = PCGSettingsHelpers::GetValue(GET_MEMBER_NAME_CHECKED(FPCGProjectionParams, bProjectScales), ProjectionParams.bProjectScales, Params);
-	ProjectionParams.bProjectColors = PCGSettingsHelpers::GetValue(GET_MEMBER_NAME_CHECKED(FPCGProjectionParams, bProjectColors), ProjectionParams.bProjectColors, Params);
+	const FPCGProjectionParams& ProjectionParams = Settings->ProjectionParams;
 
 	TArray<FPCGTaggedData>& Outputs = Context->OutputData.TaggedData;
 
 #if WITH_EDITOR
-	const bool bKeepZeroDensityPoints = PCGSettingsHelpers::GetValue(GET_MEMBER_NAME_CHECKED(UPCGProjectionSettings, bKeepZeroDensityPoints), Settings->bKeepZeroDensityPoints, Params);
+	const bool bKeepZeroDensityPoints = Settings->bKeepZeroDensityPoints;
 #else
 	const bool bKeepZeroDensityPoints = false;
 #endif
diff --git a/Engine/Plugins/Experimental/PCG/Source/PCG/Private/Elements/PCGSelfPruning.cpp b/Engine/Plugins/Experimental/PCG/Source/PCG/Private/Elements/PCGSelfPruning.cpp
--- a/Engine/Plugins/Experimental/PCG/Source/PCG/Private/Elements/PCGSelfPruning.cpp
+++ b/Engine/Plugins/Experimental/PCG/Source/PCG/Private/Elements/PCGSelfPruning.cpp
@@ -6,7 +6,6 @@
 #include "PCGCustomVersion.h"
 #include "Data/PCGPointData.h"
 #include "Data/PCGSpatialData.h"
-#include "Helpers/PCGSettingsHelpers.h"
 
 
 #include UE_INLINE_GENERATED_CPP_BY_NAME(PCGSelfPruning)
@@ -249,7 +248,6 @@ TArray<FPCGPinProperties> UPCGSelfPruningSettings::InputPinProperties() const
 	TArray<FPCGPinProperties> PinProperties;
 	// TODO in the future type checking of edges will be stricter and a conversion node will be added to convert from other types
 	PinProperties.Emplace(PCGPinConstants::DefaultInputLabel, EPCGDataType::Point);
-	PinProperties.Emplace(PCGPinConstants::DefaultParamsLabel, EPCGDataType::Param, /*bInAllowMultipleConnections=*/false);
 
 	return PinProperties;
 }
@@ -266,27 +264,11 @@ bool FPCGSelfPruningElement::ExecuteInternal(FPCGContext* Context) const
 	const UPCGSelfPruningSettings* Settings = Context->GetInputSettings<UPCGSelfPruningSettings>();
 	check(Settings);
 
-	UPCGParamData* Params = Context->InputData.GetParams();
-
-	const EPCGSelfPruningType PruningType = PCGSettingsHelpers::GetValue(GET_MEMBER_NAME_CHECKED(UPCGSelfPruningSettings, PruningType), Settings->PruningType, Params);
-	const float RadiusSimilarityFactor = PCGSettingsHelpers::GetValue(GET_MEMBER_NAME_CHECKED(UPCGSelfPruningSettings, RadiusSimilarityFactor), Settings->RadiusSimilarityFactor, Params);
-	const bool bRandomizedPruning = PCGSettingsHelpers::GetValue(GET_MEMBER_NAME_CHECKED(UPCGSelfPruningSettings, bRandomizedPruning), Settings->bRandomizedPruning, Params);
+	const EPCGSelfPruningType PruningType = Settings->PruningType;
+	const float RadiusSimilarityFactor = Settings->RadiusSimilarityFactor;
+	const bool bRandomizedPruning = Settings->bRandomizedPruning;
 
 	PCGSelfPruningElement::Execute(Context, PruningType, RadiusSimilarityFactor, bRandomizedPruning);
 
 	return true;
 }
-
-#if WITH_EDITOR
-void UPCGSelfPruningSettings::ApplyDeprecationBeforeUpdatePins(UPCGNode* InOutNode, TArray<TObjectPtr<UPCGPin>>& InputPins, TArray<TObjectPtr<UPCGPin>>& OutputPins)
-{
-	Super::ApplyDeprecationBeforeUpdatePins(InOutNode, InputPins, OutputPins);
-
-	check(InOutNode);
-
-	if (DataVersion < FPCGCustomVersion::MoveSelfPruningParamsOffFirstPin)
-	{
-		PCGSettingsHelpers::DeprecationBreakOutParamsToNewPin(InOutNode, InputPins, OutputPins);
-	}
-}
-#endif // WITH_EDITOR
diff --git a/Engine/Plugins/Experimental/PCG/Source/PCG/Private/Elements/PCGSplineSampler.cpp b/Engine/Plugins/Experimental/PCG/Source/PCG/Private/Elements/PCGSplineSampler.cpp
--- a/Engine/Plugins/Experimental/PCG/Source/PCG/Private/Elements/PCGSplineSampler.cpp
+++ b/Engine/Plugins/Experimental/PCG/Source/PCG/Private/Elements/PCGSplineSampler.cpp
@@ -9,7 +9,6 @@
 #include "Data/PCGSplineData.h"
 #include "Data/PCGSpatialData.h"
 #include "Helpers/PCGBlueprintHelpers.h"
-#include "Helpers/PCGSettingsHelpers.h"
 
 #include "Components/SplineComponent.h"
 #include "PCGContext.h"
@@ -902,20 +901,8 @@ bool FPCGSplineSamplerElement::ExecuteInternal(FPCGContext* Context) const
 	check(Settings);
 
 	TArray<FPCGTaggedData> Inputs = Context->InputData.GetInputs();
-	UPCGParamData* Params = Context->InputData.GetParams();
-
-	FPCGSplineSamplerParams SamplerParams = Settings->Params;
-	SamplerParams.Mode = PCG_GET_OVERRIDEN_VALUE(&SamplerParams, Mode, Params);
-	SamplerParams.Dimension = PCG_GET_OVERRIDEN_VALUE(&SamplerParams, Dimension, Params);
-	SamplerParams.Fill = PCG_GET_OVERRIDEN_VALUE(&SamplerParams, Fill, Params);
-	SamplerParams.SubdivisionsPerSegment = PCG_GET_OVERRIDEN_VALUE(&SamplerParams, SubdivisionsPerSegment, Params);
-	SamplerParams.DistanceIncrement = PCG_GET_OVERRIDEN_VALUE(&SamplerParams, DistanceIncrement, Params);
-	SamplerParams.NumPlanarSubdivisions = PCG_GET_OVERRIDEN_VALUE(&SamplerParams, NumPlanarSubdivisions, Params);
-	SamplerParams.NumHeightSubdivisions = PCG_GET_OVERRIDEN_VALUE(&SamplerParams, NumHeightSubdivisions, Params);
-	SamplerParams.InteriorSampleSpacing = PCG_GET_OVERRIDEN_VALUE(&SamplerParams, InteriorSampleSpacing, Params);
-	SamplerParams.InteriorBorderSampleSpacing = PCG_GET_OVERRIDEN_VALUE(&SamplerParams, InteriorBorderSampleSpacing, Params);
-	SamplerParams.InteriorOrientation = PCG_GET_OVERRIDEN_VALUE(&SamplerParams, InteriorOrientation, Params);
-	SamplerParams.bProjectOntoSurface = PCG_GET_OVERRIDEN_VALUE(&SamplerParams, bProjectOntoSurface, Params);
+
+	const FPCGSplineSamplerParams& SamplerParams = Settings->SamplerParams;
 
 	TArray<FPCGTaggedData>& Outputs = Context->OutputData.TaggedData;
 
diff --git a/Engine/Plugins/Experimental/PCG/Source/PCG/Private/Elements/PCGSurfaceSampler.cpp b/Engine/Plugins/Experimental/PCG/Source/PCG/Private/Elements/PCGSurfaceSampler.cpp
--- a/Engine/Plugins/Experimental/PCG/Source/PCG/Private/Elements/PCGSurfaceSampler.cpp
+++ b/Engine/Plugins/Experimental/PCG/Source/PCG/Private/Elements/PCGSurfaceSampler.cpp
@@ -34,25 +34,19 @@ namespace PCGSurfaceSampler
 
 		if (Settings)
 		{
-			UPCGParamData* Params = Context ? Context->InputData.GetParams() : nullptr;
-			
 			// Compute used values
-			PointsPerSquaredMeter = PCGSettingsHelpers::GetValue(GET_MEMBER_NAME_CHECKED(UPCGSurfaceSamplerSettings, PointsPerSquaredMeter), Settings->PointsPerSquaredMeter, Params);
-			PointExtents = PCGSettingsHelpers::GetValue(GET_MEMBER_NAME_CHECKED(UPCGSurfaceSamplerSettings, PointExtents), Settings->PointExtents, Params);
-			Looseness = PCGSettingsHelpers::GetValue(GET_MEMBER_NAME_CHECKED(UPCGSurfaceSamplerSettings, Looseness), Settings->Looseness, Params);
-			bApplyDensityToPoints = PCGSettingsHelpers::GetValue(GET_MEMBER_NAME_CHECKED(UPCGSurfaceSamplerSettings, bApplyDensityToPoints), Settings->bApplyDensityToPoints, Params);
-			PointSteepness = PCGSettingsHelpers::GetValue(GET_MEMBER_NAME_CHECKED(UPCGSurfaceSamplerSettings, PointSteepness), Settings->PointSteepness, Params);
+			PointsPerSquaredMeter = Settings->PointsPerSquaredMeter;
+			PointExtents = Settings->PointExtents;
+			Looseness = Settings->Looseness;
+			bApplyDensityToPoints = Settings->bApplyDensityToPoints;
+			PointSteepness = Settings->PointSteepness;
 #if WITH_EDITOR
-			bKeepZeroDensityPoints = PCGSettingsHelpers::GetValue(GET_MEMBER_NAME_CHECKED(UPCGSurfaceSamplerSettings, bKeepZeroDensityPoints), Settings->bKeepZeroDensityPoints, Params);
+			bKeepZeroDensityPoints = Settings->bKeepZeroDensityPoints;
 #endif
-
-			Seed = PCGSettingsHelpers::ComputeSeedWithOverride(InSettings, Context ? Context->SourceComponent : nullptr, Params);
-		}
-		else
-		{
-			Seed = (Context && Context->SourceComponent.IsValid()) ? Context->SourceComponent->Seed : 42;
 		}
 
+		Seed = Context->GetSeed();
+
 		// Conceptually, we will break down the surface bounds in a N x M grid
 		InterstitialDistance = PointExtents * 2;
 		InnerCellSize = InterstitialDistance * Looseness;
@@ -243,7 +237,6 @@ TArray<FPCGPinProperties> UPCGSurfaceSamplerSettings::InputPinProperties() const
 		"All sampled points must be contained within this shape. If this input is omitted then bounds will be taken from the actor so that points are contained within actor bounds. "
 		"The Unbounded property disables this and instead generates over the entire bounds of Surface."
 	));
-	PinProperties.Emplace(PCGPinConstants::DefaultParamsLabel, EPCGDataType::Param, /*bInAllowMultipleConnections=*/false);
 
 	return PinProperties;
 }
@@ -260,11 +253,13 @@ void UPCGSurfaceSamplerSettings::PostLoad()
 {
 	Super::PostLoad();
 
+#if WITH_EDITOR
 	if (PointRadius_DEPRECATED != 0)
 	{
 		PointExtents = FVector(PointRadius_DEPRECATED);
 		PointRadius_DEPRECATED = 0;
 	}
+#endif // WITH_EDITOR
 }
 
 #if WITH_EDITOR
@@ -347,11 +342,9 @@ bool FPCGSurfaceSamplerElement::ExecuteInternal(FPCGContext* Context) const
 		PCGE_LOG(Warning, "No Surface input was provided, and no surface could be found in the Bounding Shape input for sampling. Connect the surface to be sampled to the Surface input.");
 	}
 
-	UPCGParamData* Params = Context->InputData.GetParams();
-
 	// Early out on invalid settings
 	// TODO: we could compute an approximate radius based on the points per squared meters if that's useful
-	const FVector PointExtents = PCGSettingsHelpers::GetValue(GET_MEMBER_NAME_CHECKED(UPCGSurfaceSamplerSettings, PointExtents), Settings->PointExtents, Params);
+	const FVector& PointExtents = Settings->PointExtents;
 	if(PointExtents.X <= 0 || PointExtents.Y <= 0)
 	{
 		PCGE_LOG(Warning, "Skipped - Invalid point extents");
@@ -403,8 +396,8 @@ void FPCGSurfaceSamplerElement::GetDependenciesCrc(const FPCGDataCollection& InI
 
 	if (const UPCGSurfaceSamplerSettings* Settings = Cast<UPCGSurfaceSamplerSettings>(InSettings))
 	{
-		const UPCGParamData* Params = InInput.GetParams();
-		const bool bUnbounded = PCG_GET_OVERRIDEN_VALUE(Settings, bUnbounded, Params);
+		bool bUnbounded;
+		PCGSettingsHelpers::GetOverrideValue(InInput, Settings, GET_MEMBER_NAME_CHECKED(UPCGSurfaceSamplerSettings, bUnbounded), Settings->bUnbounded, bUnbounded);
 		const bool bBoundsConnected = InInput.GetInputsByPin(PCGSurfaceSamplerConstants::BoundingShapeLabel).Num() > 0;
 
 		// If we're operating in bounded mode and there is no bounding shape connected then we'll use actor bounds, and therefore take
diff --git a/Engine/Plugins/Experimental/PCG/Source/PCG/Private/Elements/PCGTextureSampler.cpp b/Engine/Plugins/Experimental/PCG/Source/PCG/Private/Elements/PCGTextureSampler.cpp
--- a/Engine/Plugins/Experimental/PCG/Source/PCG/Private/Elements/PCGTextureSampler.cpp
+++ b/Engine/Plugins/Experimental/PCG/Source/PCG/Private/Elements/PCGTextureSampler.cpp
@@ -21,10 +21,7 @@ FPCGElementPtr UPCGTextureSamplerSettings::CreateElement() const
 
 TArray<FPCGPinProperties> UPCGTextureSamplerSettings::InputPinProperties() const
 {
-	TArray<FPCGPinProperties> PinProperties;
-	PinProperties.Emplace(PCGPinConstants::DefaultParamsLabel, EPCGDataType::Param, /*bInAllowMultipleConnections=*/false);
-
-	return PinProperties;
+	return TArray<FPCGPinProperties>();
 }
 
 bool FPCGTextureSamplerElement::ExecuteInternal(FPCGContext* Context) const
@@ -34,22 +31,19 @@ bool FPCGTextureSamplerElement::ExecuteInternal(FPCGContext* Context) const
 	const UPCGTextureSamplerSettings* Settings = Context->GetInputSettings<UPCGTextureSamplerSettings>();
 	check(Settings);
 
-	UPCGParamData* Params = Context->InputData.GetParams();
-
-	FTransform Transform = PCG_GET_OVERRIDEN_VALUE(Settings, Transform, Params);
-	bool bUseAbsoluteTransform = PCG_GET_OVERRIDEN_VALUE(Settings, bUseAbsoluteTransform, Params);
-	//TObjectPtr<UTexture2D> PCG_OVERRIDEABLE_VALUE(Texture);
+	const FTransform& Transform = Settings->Transform;
+	const bool bUseAbsoluteTransform = Settings->bUseAbsoluteTransform;
 	TObjectPtr<UTexture2D> Texture = Settings->Texture;
-	EPCGTextureDensityFunction DensityFunction = PCG_GET_OVERRIDEN_VALUE(Settings, DensityFunction, Params);
-	EPCGTextureColorChannel ColorChannel = PCG_GET_OVERRIDEN_VALUE(Settings, ColorChannel, Params);
-	float TexelSize = PCG_GET_OVERRIDEN_VALUE(Settings, TexelSize, Params);
-	bool bUseAdvancedTiling = PCG_GET_OVERRIDEN_VALUE(Settings, bUseAdvancedTiling, Params);
-	FVector2D Tiling = PCG_GET_OVERRIDEN_VALUE(Settings, Tiling, Params);
-	FVector2D CenterOffset = PCG_GET_OVERRIDEN_VALUE(Settings, CenterOffset, Params);
-	float Rotation = PCG_GET_OVERRIDEN_VALUE(Settings, Rotation, Params);
-	bool bUseTileBounds = PCG_GET_OVERRIDEN_VALUE(Settings, bUseTileBounds, Params);
-	FVector2D TileBoundsMin = PCG_GET_OVERRIDEN_VALUE(Settings, TileBoundsMin, Params);
-	FVector2D TileBoundsMax = PCG_GET_OVERRIDEN_VALUE(Settings, TileBoundsMax, Params);
+	const EPCGTextureDensityFunction DensityFunction = Settings->DensityFunction;
+	const EPCGTextureColorChannel ColorChannel = Settings->ColorChannel;
+	const float TexelSize = Settings->TexelSize;
+	const bool bUseAdvancedTiling = Settings->bUseAdvancedTiling;
+	const FVector2D& Tiling = Settings->Tiling;
+	const FVector2D& CenterOffset = Settings->CenterOffset;
+	const float Rotation = Settings->Rotation;
+	const bool bUseTileBounds = Settings->bUseTileBounds;
+	const FVector2D& TileBoundsMin = Settings->TileBoundsMin;
+	const FVector2D& TileBoundsMax = Settings->TileBoundsMax;
 
 	TArray<FPCGTaggedData>& Outputs = Context->OutputData.TaggedData;
 	FPCGTaggedData& Output = Outputs.Emplace_GetRef();
@@ -95,8 +89,8 @@ void FPCGTextureSamplerElement::GetDependenciesCrc(const FPCGDataCollection& InI
 	if (const UPCGTextureSamplerSettings* Settings = Cast<UPCGTextureSamplerSettings>(InSettings))
 	{
 		// If not using absolute transform, depend on actor transform and bounds, and therefore take dependency on actor data.
-		const UPCGParamData* Params = InInput.GetParams();
-		const bool bUseAbsoluteTransform = PCG_GET_OVERRIDEN_VALUE(Settings, bUseAbsoluteTransform, Params);
+		bool bUseAbsoluteTransform;
+		PCGSettingsHelpers::GetOverrideValue(InInput, Settings, GET_MEMBER_NAME_CHECKED(UPCGTextureSamplerSettings, bUseAbsoluteTransform), Settings->bUseAbsoluteTransform, bUseAbsoluteTransform);
 		if (!bUseAbsoluteTransform && InComponent)
 		{
 			if (const UPCGData* Data = InComponent->GetActorPCGData())
diff --git a/Engine/Plugins/Experimental/PCG/Source/PCG/Private/Elements/PCGTransformPoints.cpp b/Engine/Plugins/Experimental/PCG/Source/PCG/Private/Elements/PCGTransformPoints.cpp
--- a/Engine/Plugins/Experimental/PCG/Source/PCG/Private/Elements/PCGTransformPoints.cpp
+++ b/Engine/Plugins/Experimental/PCG/Source/PCG/Private/Elements/PCGTransformPoints.cpp
@@ -7,7 +7,6 @@
 #include "Helpers/PCGAsync.h"
 #include "PCGHelpers.h"
 #include "Data/PCGPointData.h"
-#include "Helpers/PCGSettingsHelpers.h"
 #include "PCGContext.h"
 #include "PCGPin.h"
 
@@ -23,7 +22,6 @@ TArray<FPCGPinProperties> UPCGTransformPointsSettings::InputPinProperties() cons
 	TArray<FPCGPinProperties> PinProperties;
 	// TODO in the future type checking of edges will be stricter and a conversion node will be added to convert from other types
 	PinProperties.Emplace(PCGPinConstants::DefaultInputLabel, EPCGDataType::Point);
-	PinProperties.Emplace(TEXT("Params"), EPCGDataType::Param, /*bInAllowMultipleConnections=*/false);
 
 	return PinProperties;
 }
@@ -51,22 +49,21 @@ bool FPCGTransformPointsElement::ExecuteInternal(FPCGContext* Context) const
 	const TArray<FPCGTaggedData> Inputs = Context->InputData.GetInputs();
 	TArray<FPCGTaggedData>& Outputs = Context->OutputData.TaggedData;
 
-	UPCGParamData* Params = Context->InputData.GetParams();
-	const bool bApplyToAttribute = PCG_GET_OVERRIDEN_VALUE(Settings, bApplyToAttribute, Params);
-	const FName AttributeName = PCG_GET_OVERRIDEN_VALUE(Settings, AttributeName, Params);
-	const FVector OffsetMin = PCG_GET_OVERRIDEN_VALUE(Settings, OffsetMin, Params);
-	const FVector OffsetMax = PCG_GET_OVERRIDEN_VALUE(Settings, OffsetMax, Params);
-	const bool bAbsoluteOffset = PCG_GET_OVERRIDEN_VALUE(Settings, bAbsoluteOffset, Params);
-	const FRotator RotationMin = PCG_GET_OVERRIDEN_VALUE(Settings, RotationMin, Params);
-	const FRotator RotationMax = PCG_GET_OVERRIDEN_VALUE(Settings, RotationMax, Params);
-	const bool bAbsoluteRotation = PCG_GET_OVERRIDEN_VALUE(Settings, bAbsoluteRotation, Params);
-	const FVector ScaleMin = PCG_GET_OVERRIDEN_VALUE(Settings, ScaleMin, Params);
-	const FVector ScaleMax = PCG_GET_OVERRIDEN_VALUE(Settings, ScaleMax, Params);
-	const bool bAbsoluteScale = PCG_GET_OVERRIDEN_VALUE(Settings, bAbsoluteScale, Params);
-	const bool bUniformScale = PCG_GET_OVERRIDEN_VALUE(Settings, bUniformScale, Params);
-	const bool bRecomputeSeed = PCG_GET_OVERRIDEN_VALUE(Settings, bRecomputeSeed, Params);
-
-	const int Seed = PCGSettingsHelpers::ComputeSeedWithOverride(Settings, Context->SourceComponent, Params);
+	const bool bApplyToAttribute = Settings->bApplyToAttribute;
+	const FName AttributeName = Settings->AttributeName;
+	const FVector& OffsetMin = Settings->OffsetMin;
+	const FVector& OffsetMax = Settings->OffsetMax;
+	const bool bAbsoluteOffset = Settings->bAbsoluteOffset;
+	const FRotator& RotationMin = Settings->RotationMin;
+	const FRotator& RotationMax = Settings->RotationMax;
+	const bool bAbsoluteRotation = Settings->bAbsoluteRotation;
+	const FVector& ScaleMin = Settings->ScaleMin;
+	const FVector& ScaleMax = Settings->ScaleMax;
+	const bool bAbsoluteScale = Settings->bAbsoluteScale;
+	const bool bUniformScale = Settings->bUniformScale;
+	const bool bRecomputeSeed = Settings->bRecomputeSeed;
+
+	const int Seed = Context->GetSeed();
 
 	// Use implicit capture, since we capture a lot
 	//ProcessPoints(Context, Inputs, Outputs, [&](const FPCGPoint& InPoint, FPCGPoint& OutPoint)
diff --git a/Engine/Plugins/Experimental/PCG/Source/PCG/Private/Elements/PCGUnionElement.cpp b/Engine/Plugins/Experimental/PCG/Source/PCG/Private/Elements/PCGUnionElement.cpp
--- a/Engine/Plugins/Experimental/PCG/Source/PCG/Private/Elements/PCGUnionElement.cpp
+++ b/Engine/Plugins/Experimental/PCG/Source/PCG/Private/Elements/PCGUnionElement.cpp
@@ -2,12 +2,19 @@
 
 #include "Elements/PCGUnionElement.h"
 #include "Data/PCGSpatialData.h"
-#include "Helpers/PCGSettingsHelpers.h"
 #include "PCGContext.h"
 #include "PCGPin.h"
 
 #include UE_INLINE_GENERATED_CPP_BY_NAME(PCGUnionElement)
 
+TArray<FPCGPinProperties> UPCGUnionSettings::InputPinProperties() const
+{
+	TArray<FPCGPinProperties> PinProperties;
+	PinProperties.Emplace(PCGPinConstants::DefaultInputLabel, EPCGDataType::Spatial);
+
+	return PinProperties;
+}
+
 TArray<FPCGPinProperties> UPCGUnionSettings::OutputPinProperties() const
 {
 	TArray<FPCGPinProperties> PinProperties;
@@ -29,10 +36,9 @@ bool FPCGUnionElement::ExecuteInternal(FPCGContext* Context) const
 	check(Settings);
 
 	TArray<FPCGTaggedData> Inputs = Context->InputData.GetInputs();
-	UPCGParamData* Params = Context->InputData.GetParams();
 
-	const EPCGUnionType Type = PCGSettingsHelpers::GetValue(GET_MEMBER_NAME_CHECKED(UPCGUnionSettings, Type), Settings->Type, Params);
-	const EPCGUnionDensityFunction DensityFunction = PCGSettingsHelpers::GetValue(GET_MEMBER_NAME_CHECKED(UPCGUnionSettings, DensityFunction), Settings->DensityFunction, Params);
+	const EPCGUnionType Type = Settings->Type;
+	const EPCGUnionDensityFunction DensityFunction = Settings->DensityFunction;
 
 	TArray<FPCGTaggedData>& Outputs = Context->OutputData.TaggedData;
 
diff --git a/Engine/Plugins/Experimental/PCG/Source/PCG/Private/Elements/PCGVolumeSampler.cpp b/Engine/Plugins/Experimental/PCG/Source/PCG/Private/Elements/PCGVolumeSampler.cpp
--- a/Engine/Plugins/Experimental/PCG/Source/PCG/Private/Elements/PCGVolumeSampler.cpp
+++ b/Engine/Plugins/Experimental/PCG/Source/PCG/Private/Elements/PCGVolumeSampler.cpp
@@ -4,7 +4,6 @@
 #include "Data/PCGPointData.h"
 #include "Data/PCGSpatialData.h"
 #include "Helpers/PCGAsync.h"
-#include "Helpers/PCGSettingsHelpers.h"
 #include "PCGContext.h"
 #include "PCGHelpers.h"
 
@@ -76,6 +75,14 @@ namespace PCGVolumeSampler
 	}
 }
 
+TArray<FPCGPinProperties> UPCGVolumeSamplerSettings::InputPinProperties() const
+{
+	TArray<FPCGPinProperties> PinProperties;
+	PinProperties.Emplace(PCGPinConstants::DefaultInputLabel, EPCGDataType::Spatial);
+
+	return PinProperties;
+}
+
 FPCGElementPtr UPCGVolumeSamplerSettings::CreateElement() const
 {
 	return MakeShared<FPCGVolumeSamplerElement>();
@@ -89,9 +96,8 @@ bool FPCGVolumeSamplerElement::ExecuteInternal(FPCGContext* Context) const
 	check(Settings);
 
 	TArray<FPCGTaggedData> Inputs = Context->InputData.GetInputs();
-	UPCGParamData* Params = Context->InputData.GetParams();
 
-	const FVector VoxelSize = PCGSettingsHelpers::GetValue(GET_MEMBER_NAME_CHECKED(UPCGVolumeSamplerSettings, VoxelSize), Settings->VoxelSize, Params);
+	const FVector& VoxelSize = Settings->VoxelSize;
 	if (VoxelSize.X <= 0 || VoxelSize.Y <= 0 || VoxelSize.Z <= 0)
 	{
 		PCGE_LOG(Warning, "Skipped - Invalid voxel size");
diff --git a/Engine/Plugins/Experimental/PCG/Source/PCG/Private/Graph/PCGGraphExecutor.cpp b/Engine/Plugins/Experimental/PCG/Source/PCG/Private/Graph/PCGGraphExecutor.cpp
--- a/Engine/Plugins/Experimental/PCG/Source/PCG/Private/Graph/PCGGraphExecutor.cpp
+++ b/Engine/Plugins/Experimental/PCG/Source/PCG/Private/Graph/PCGGraphExecutor.cpp
@@ -9,11 +9,13 @@
 #include "PCGGraphCompiler.h"
 #include "PCGHelpers.h"
 #include "PCGInputOutputSettings.h"
+#include "PCGModule.h"
+#include "PCGParamData.h"
+#include "PCGPin.h"
 #include "Graph/PCGGraphCache.h"
+#include "Metadata/PCGMetadata.h"
 
 #include "Async/Async.h"
-#include "PCGModule.h"
-#include "PCGPin.h"
 
 #if WITH_EDITOR
 #include "Editor.h"
@@ -537,6 +539,7 @@ void FPCGGraphExecutor::Execute()
 				if (!Task.Context)
 				{
 					Task.Context = Task.Element->Initialize(TaskInput, Task.SourceComponent, Task.Node);
+					Task.Context->InitializeSettings();
 					Task.Context->TaskId = Task.NodeId;
 					Task.Context->CompiledTaskId = Task.CompiledTaskId;
 					Task.Context->InputData.Crc = TaskInput.Crc;
@@ -672,6 +675,13 @@ void FPCGGraphExecutor::Execute()
 			QueueNextTasks(ActiveTask.NodeId);
 			bAnyTaskEnded = true;
 
+			// Un-root temporary input data
+			if (FPCGDataCollection* TemporaryInput = InputTemporaryData.Find(ActiveTask.NodeId))
+			{
+				TemporaryInput->RemoveFromRootSet(DataRootSet);
+				InputTemporaryData.Remove(ActiveTask.NodeId);
+			}
+
 			// Remove current active task from list
 			ActiveTasks.RemoveAtSwap(TaskIndex);
 		};
@@ -915,9 +925,54 @@ void FPCGGraphExecutor::BuildTaskInput(const FPCGGraphTask& Task, FPCGDataCollec
 		}
 	}
 
+	// Then combine params if needed
+	CombineParams(Task.NodeId, TaskInput);
+
 	TaskInput.Crc = Crc;
 }
 
+void FPCGGraphExecutor::CombineParams(FPCGTaskId InTaskId, FPCGDataCollection& InTaskInput)
+{
+	TRACE_CPUPROFILER_EVENT_SCOPE(FPCGGraphExecutor::CombineParams);
+
+	TArray<FPCGTaggedData> AllParamsData = InTaskInput.GetParamsByPin(PCGPinConstants::DefaultParamsLabel);
+	if (AllParamsData.Num() > 1)
+	{
+		UPCGParamData* CombinedParamData = NewObject<UPCGParamData>();
+		for (const FPCGTaggedData& TaggedDatum : AllParamsData)
+		{
+			const UPCGParamData* ParamData = CastChecked<UPCGParamData>(TaggedDatum.Data);
+			CombinedParamData->Metadata->AddAttributes(ParamData->Metadata);
+		}
+
+		const int32 NewNumberOfInputs = InTaskInput.TaggedData.Num() - AllParamsData.Num() + 1;
+		check(NewNumberOfInputs >= 1);
+
+		TArray<FPCGTaggedData> TempTaggedData{};
+		TempTaggedData.Reserve(NewNumberOfInputs);
+		for (FPCGTaggedData& TaggedData : InTaskInput.TaggedData)
+		{
+			if (TaggedData.Pin != PCGPinConstants::DefaultParamsLabel)
+			{
+				TempTaggedData.Add(std::move(TaggedData));
+			}
+		}
+
+		// Add to the root set since we created a new object, that needs to be kept alive for the duration of the task.
+		DataRootSet.Add(CombinedParamData);
+		FPCGTaggedData CombineParams{};
+		CombineParams.Data = CombinedParamData;
+		CombineParams.Pin = PCGPinConstants::DefaultParamsLabel;
+		TempTaggedData.Add(CombineParams);
+
+		InTaskInput.TaggedData = std::move(TempTaggedData);
+
+		// Also store it into the TemporaryMap to track it down
+		FPCGDataCollection& TemporaryInput = InputTemporaryData.FindOrAdd(InTaskId);
+		TemporaryInput.TaggedData.Add(CombineParams);
+	}
+}
+
 void FPCGGraphExecutor::StoreResults(FPCGTaskId InTaskId, const FPCGDataCollection& InTaskOutput)
 {
 	TRACE_CPUPROFILER_EVENT_SCOPE(FPCGGraphExecutor::StoreResults);
@@ -939,6 +994,12 @@ void FPCGGraphExecutor::ClearResults()
 		NextTaskId = 0;
 	}
 
+	for (const TPair<FPCGTaskId, FPCGDataCollection>& TaskTemporaryInput : InputTemporaryData)
+	{
+		TaskTemporaryInput.Value.RemoveFromRootSet(DataRootSet);
+	}
+	InputTemporaryData.Reset();
+
 	for (const TPair<FPCGTaskId, FPCGDataCollection>& TaskOutput : OutputData)
 	{
 		TaskOutput.Value.RemoveFromRootSet(DataRootSet);
diff --git a/Engine/Plugins/Experimental/PCG/Source/PCG/Private/Graph/PCGGraphExecutor.h b/Engine/Plugins/Experimental/PCG/Source/PCG/Private/Graph/PCGGraphExecutor.h
--- a/Engine/Plugins/Experimental/PCG/Source/PCG/Private/Graph/PCGGraphExecutor.h
+++ b/Engine/Plugins/Experimental/PCG/Source/PCG/Private/Graph/PCGGraphExecutor.h
@@ -133,6 +133,8 @@ private:
 	void QueueNextTasks(FPCGTaskId FinishedTask);
 	void CancelNextTasks(FPCGTaskId CancelledTask);
 	void BuildTaskInput(const FPCGGraphTask& Task, FPCGDataCollection& TaskInput);
+	/** Combine all param data into one on the Params pin, if any.*/
+	void CombineParams(FPCGTaskId InTaskId, FPCGDataCollection& InTaskInput);
 	void StoreResults(FPCGTaskId InTaskId, const FPCGDataCollection& InTaskOutput);
 	void ClearResults();
 
@@ -169,6 +171,8 @@ private:
 	/** Map of node instances to their output, could be cleared once execution is done */
 	/** Note: this should at some point unload based on loaded/unloaded proxies, otherwise memory cost will be unbounded */
 	TMap<FPCGTaskId, FPCGDataCollection> OutputData;
+	/** Map of node instances to their temporary input, could be cleared once execution is done */
+	TMap<FPCGTaskId, FPCGDataCollection> InputTemporaryData;
 	/** Monotonically increasing id. Should be reset once all tasks are executed, should be protected by the ScheduleLock */
 	FPCGTaskId NextTaskId = 0;
 
diff --git a/Engine/Plugins/Experimental/PCG/Source/PCG/Private/Helpers/PCGSettingsHelpers.cpp b/Engine/Plugins/Experimental/PCG/Source/PCG/Private/Helpers/PCGSettingsHelpers.cpp
--- a/Engine/Plugins/Experimental/PCG/Source/PCG/Private/Helpers/PCGSettingsHelpers.cpp
+++ b/Engine/Plugins/Experimental/PCG/Source/PCG/Private/Helpers/PCGSettingsHelpers.cpp
@@ -6,6 +6,7 @@
 #include "PCGEdge.h"
 #include "PCGHelpers.h"
 #include "PCGPin.h"
+
 #include "UObject/EnumProperty.h"
 
 namespace PCGSettingsHelpers
@@ -207,6 +208,7 @@ namespace PCGSettingsHelpers
 		}
 	}
 
+	PRAGMA_DISABLE_DEPRECATION_WARNINGS
 	int ComputeSeedWithOverride(const UPCGSettings* InSettings, const UPCGComponent* InComponent, UPCGParamData* InParams)
 	{
 		check(InSettings);
@@ -214,11 +216,24 @@ namespace PCGSettingsHelpers
 		const int SettingsSeed = InParams ? PCGSettingsHelpers::GetValue(GET_MEMBER_NAME_CHECKED(UPCGSettings, Seed), InSettings->Seed, InParams) : InSettings->Seed;
 		return InComponent ? PCGHelpers::ComputeSeed(SettingsSeed, InComponent->Seed) : SettingsSeed;
 	}
+	PRAGMA_ENABLE_DEPRECATION_WARNINGS
 
 	void DeprecationBreakOutParamsToNewPin(UPCGNode* InOutNode, TArray<TObjectPtr<UPCGPin>>& InputPins, TArray<TObjectPtr<UPCGPin>>& OutputPins)
 	{
 		// Check basic conditions for which the code below should run.
-		if(InputPins.IsEmpty() || !InputPins[0] || InputPins[0]->Properties.AllowedTypes != EPCGDataType::Any)
+		if(!InOutNode || InputPins.IsEmpty() || !InputPins[0] || InputPins[0]->Properties.AllowedTypes != EPCGDataType::Any)
+		{
+			return;
+		}
+
+		// Check if the node already has a param pin, if so, nothing to do.
+		if (InOutNode->GetInputPin(PCGPinConstants::DefaultParamsLabel))
+		{
+			return;
+		}
+
+		// Also no need to add a param pin if it has no overriable params
+		if (!InOutNode->GetSettings() || InOutNode->GetSettings()->OverridableParams().IsEmpty())
 		{
 			return;
 		}
@@ -230,7 +245,8 @@ namespace PCGSettingsHelpers
 		NewParamsPin->Node = InOutNode;
 		NewParamsPin->Properties.AllowedTypes = EPCGDataType::Param;
 		NewParamsPin->Properties.Label = PCGPinConstants::DefaultParamsLabel;
-		NewParamsPin->Properties.bAllowMultipleConnections = false;
+		NewParamsPin->Properties.bAllowMultipleConnections = true;
+		NewParamsPin->Properties.bAllowMultipleData = true;
 		InputPins.Add(NewParamsPin);
 
 		// Make list of param pins that In pin is currently connected to.
diff --git a/Engine/Plugins/Experimental/PCG/Source/PCG/Private/MatchAndSet/PCGMatchAndSetBase.cpp b/Engine/Plugins/Experimental/PCG/Source/PCG/Private/MatchAndSet/PCGMatchAndSetBase.cpp
--- a/Engine/Plugins/Experimental/PCG/Source/PCG/Private/MatchAndSet/PCGMatchAndSetBase.cpp
+++ b/Engine/Plugins/Experimental/PCG/Source/PCG/Private/MatchAndSet/PCGMatchAndSetBase.cpp
@@ -37,7 +37,7 @@ bool UPCGMatchAndSetBase::CreateAttributeIfNeeded(FPCGContext& Context, const FP
 				return PCGMetadataElementCommon::ClearOrCreateAttribute(OutPointData->Metadata, DestinationAttribute, ConstantType{}) != nullptr;
 			};
 
-			if (!ConstantValue.DispatcherWithOverride(nullptr, CreateAttribute))
+			if (!ConstantValue.Dispatcher(CreateAttribute))
 			{
 				PCGE_LOG_C(Error, &Context, "Unable to create attribute %s on point data", *DestinationAttribute.ToString());
 				return false;
diff --git a/Engine/Plugins/Experimental/PCG/Source/PCG/Private/MatchAndSet/PCGMatchAndSetByAttribute.cpp b/Engine/Plugins/Experimental/PCG/Source/PCG/Private/MatchAndSet/PCGMatchAndSetByAttribute.cpp
--- a/Engine/Plugins/Experimental/PCG/Source/PCG/Private/MatchAndSet/PCGMatchAndSetByAttribute.cpp
+++ b/Engine/Plugins/Experimental/PCG/Source/PCG/Private/MatchAndSet/PCGMatchAndSetByAttribute.cpp
@@ -94,19 +94,19 @@ void UPCGMatchAndSetByAttribute::MatchAndSet_Implementation(
 	auto CreateConstantMatchSourceAccessor = [&EntryMatchSourceAccessors](auto&& Value)
 	{
 		using ConstantType = std::decay_t<decltype(Value)>;
-		EntryMatchSourceAccessors.Emplace(MakeUnique<FPCGConstantValueAccessor<ConstantType>>(std::forward<ConstantType>(Value)));
+		EntryMatchSourceAccessors.Emplace(MakeUnique<FPCGConstantValueAccessor<ConstantType>>(std::forward<decltype(Value)>(Value)));
 	};
 
 	auto CreateConstantSetValueAccessor = [&EntrySetValueAccessors](auto&& Value)
 	{
 		using ConstantType = std::decay_t<decltype(Value)>;
-		EntrySetValueAccessors.Emplace(MakeUnique<FPCGConstantValueAccessor<ConstantType>>(std::forward<ConstantType>(Value)));
+		EntrySetValueAccessors.Emplace(MakeUnique<FPCGConstantValueAccessor<ConstantType>>(std::forward<decltype(Value)>(Value)));
 	};
 
 	for (const FPCGMatchAndSetByAttributeEntry& Entry : Entries)
 	{
-		Entry.ValueToMatch.DispatcherWithOverride(nullptr, CreateConstantMatchSourceAccessor);
-		Entry.Value.DispatcherWithOverride(nullptr, CreateConstantSetValueAccessor);
+		Entry.ValueToMatch.Dispatcher(CreateConstantMatchSourceAccessor);
+		Entry.Value.Dispatcher(CreateConstantSetValueAccessor);
 	}
 
 	//TODO: implement async loop?
diff --git a/Engine/Plugins/Experimental/PCG/Source/PCG/Private/MatchAndSet/PCGMatchAndSetWeighted.cpp b/Engine/Plugins/Experimental/PCG/Source/PCG/Private/MatchAndSet/PCGMatchAndSetWeighted.cpp
--- a/Engine/Plugins/Experimental/PCG/Source/PCG/Private/MatchAndSet/PCGMatchAndSetWeighted.cpp
+++ b/Engine/Plugins/Experimental/PCG/Source/PCG/Private/MatchAndSet/PCGMatchAndSetWeighted.cpp
@@ -77,7 +77,7 @@ void UPCGMatchAndSetWeighted::MatchAndSet_Implementation(
 	auto CreateConstantAccessor = [&EntryValueAccessors](auto&& Value)
 	{
 		using ConstantType = std::decay_t<decltype(Value)>;
-		EntryValueAccessors.Emplace(MakeUnique<FPCGConstantValueAccessor<ConstantType>>(std::forward<ConstantType>(Value)));
+		EntryValueAccessors.Emplace(MakeUnique<FPCGConstantValueAccessor<ConstantType>>(std::forward<decltype(Value)>(Value)));
 	};
 
 	TArray<int> CumulativeWeights;
@@ -88,7 +88,7 @@ void UPCGMatchAndSetWeighted::MatchAndSet_Implementation(
 	{
 		TotalWeight += Entry.Weight;
 		CumulativeWeights.Add(TotalWeight);
-		Entry.Value.DispatcherWithOverride(nullptr, CreateConstantAccessor);
+		Entry.Value.Dispatcher(CreateConstantAccessor);
 	}
 
 	if (TotalWeight <= 0)
diff --git a/Engine/Plugins/Experimental/PCG/Source/PCG/Private/MatchAndSet/PCGMatchAndSetWeightedByCategory.cpp b/Engine/Plugins/Experimental/PCG/Source/PCG/Private/MatchAndSet/PCGMatchAndSetWeightedByCategory.cpp
--- a/Engine/Plugins/Experimental/PCG/Source/PCG/Private/MatchAndSet/PCGMatchAndSetWeightedByCategory.cpp
+++ b/Engine/Plugins/Experimental/PCG/Source/PCG/Private/MatchAndSet/PCGMatchAndSetWeightedByCategory.cpp
@@ -143,20 +143,20 @@ void UPCGMatchAndSetWeightedByCategory::MatchAndSet_Implementation(
 	auto CreateConstantCategoryAccessor = [&EntryCategoryAccessors](auto&& Value)
 	{
 		using ConstantType = std::decay_t<decltype(Value)>;
-		EntryCategoryAccessors.Emplace(MakeUnique<FPCGConstantValueAccessor<ConstantType>>(std::forward<ConstantType>(Value)));
+		EntryCategoryAccessors.Emplace(MakeUnique<FPCGConstantValueAccessor<ConstantType>>(std::forward<decltype(Value)>(Value)));
 	};
 
 	for (const FPCGMatchAndSetWeightedByCategoryEntryList& Category : Categories)
 	{
-		Category.CategoryValue.DispatcherWithOverride(nullptr, CreateConstantCategoryAccessor);
+		Category.CategoryValue.Dispatcher(CreateConstantCategoryAccessor);
 
 		TArray<TUniquePtr<const IPCGAttributeAccessor>>& ValuesAccessor = EntryValueAccessors.Emplace_GetRef();
 		for (const FPCGMatchAndSetWeightedEntry& WeightedEntry : Category.WeightedEntries)
 		{
-			WeightedEntry.Value.DispatcherWithOverride(nullptr, [&ValuesAccessor](auto&& Value)
+			WeightedEntry.Value.Dispatcher([&ValuesAccessor](auto&& Value)
 			{
 				using ConstantType = std::decay_t<decltype(Value)>;
-				ValuesAccessor.Emplace(MakeUnique<FPCGConstantValueAccessor<ConstantType>>(std::forward<ConstantType>(Value)));
+				ValuesAccessor.Emplace(MakeUnique<FPCGConstantValueAccessor<ConstantType>>(std::forward<decltype(Value)>(Value)));
 			});
 		}
 	}
diff --git a/Engine/Plugins/Experimental/PCG/Source/PCG/Private/Metadata/Accessors/PCGAttributeAccessorHelpers.cpp b/Engine/Plugins/Experimental/PCG/Source/PCG/Private/Metadata/Accessors/PCGAttributeAccessorHelpers.cpp
--- a/Engine/Plugins/Experimental/PCG/Source/PCG/Private/Metadata/Accessors/PCGAttributeAccessorHelpers.cpp
+++ b/Engine/Plugins/Experimental/PCG/Source/PCG/Private/Metadata/Accessors/PCGAttributeAccessorHelpers.cpp
@@ -2,11 +2,13 @@
 
 #include "Metadata/Accessors/PCGAttributeAccessorHelpers.h"
 
-#include "Metadata/Accessors/IPCGAttributeAccessorTpl.h"
+#include "PCGData.h"
 #include "PCGParamData.h"
+#include "PCGSettings.h"
 #include "Data/PCGPointData.h"
 #include "Data/PCGSpatialData.h"
 #include "Metadata/PCGAttributePropertySelector.h"
+#include "Metadata/Accessors/IPCGAttributeAccessorTpl.h"
 #include "Metadata/Accessors/PCGAttributeAccessor.h"
 #include "Metadata/Accessors/PCGPropertyAccessor.h"
 #include "Metadata/Accessors/PCGAttributeExtractor.h"
@@ -83,105 +85,188 @@ namespace PCGAttributeAccessorHelpers
 
 		return PCGMetadataAttribute::CallbackWithRightType(InAccessor->GetUnderlyingType(), Chain);
 	}
-}
 
-TUniquePtr<IPCGAttributeAccessor> PCGAttributeAccessorHelpers::CreatePropertyAccessor(const FProperty* InProperty)
-{
-	if (!InProperty)
+	// Empty signature to passthrough types to functors.
+	template <typename T>
+	struct Signature
 	{
-		return TUniquePtr<IPCGAttributeAccessor>();
-	}
+		using Type = T;
+	};
 
-	if (const FNumericProperty* NumericProperty = CastField<FNumericProperty>(InProperty))
-	{
-		if (NumericProperty->IsFloatingPoint())
-		{
-			return MakeUnique<FPCGNumericPropertyAccessor<double>>(NumericProperty);
-		}
-		else if (NumericProperty->IsInteger())
-		{
-			return MakeUnique<FPCGNumericPropertyAccessor<int64>>(NumericProperty);
-		}
-	}
-	else if (const FBoolProperty* BoolProperty = CastField<FBoolProperty>(InProperty))
-	{
-		return MakeUnique<FPCGPropertyAccessor<bool, FBoolProperty>>(BoolProperty);
-	}
-	else if (const FStrProperty* StringProperty = CastField<FStrProperty>(InProperty))
-	{
-		return MakeUnique<FPCGPropertyAccessor<FString, FStrProperty>>(StringProperty);
-	}
-	else if (const FNameProperty* NameProperty = CastField<FNameProperty>(InProperty))
-	{
-		return MakeUnique<FPCGPropertyAccessor<FName, FNameProperty>>(NameProperty);
-	}
-	else if (const FEnumProperty* EnumProperty = CastField<FEnumProperty>(InProperty))
-	{
-		return MakeUnique<FPCGEnumPropertyAccessor>(EnumProperty);
-	}
-	else if (const FStructProperty* StructProperty = CastField<FStructProperty>(InProperty))
+	template <typename Func>
+	decltype(auto) DispatchPropertyTypes(const FProperty* InProperty, Func&& Functor)
 	{
-		if (StructProperty->Struct == TBaseStructure<FVector>::Get())
+		// Use the FPCGPropertyPathAccessor (on soft object path) as dummy type to get the functor return type
+		// because this accessor takes a generic FProperty, while others take a more specialized type (like FNumericProperty).
+		using ReturnType = decltype(Functor(Signature<FPCGPropertyPathAccessor<FSoftObjectPath>>{}, InProperty));
+
+		if (!InProperty)
 		{
-			return MakeUnique<FPCGPropertyStructAccessor<FVector>>(StructProperty);
+			return ReturnType{};
 		}
-		else if (StructProperty->Struct == TBaseStructure<FVector4>::Get())
+
+		if (const FNumericProperty* NumericProperty = CastField<FNumericProperty>(InProperty))
 		{
-			return MakeUnique<FPCGPropertyStructAccessor<FVector4>>(StructProperty);
+			if (NumericProperty->IsFloatingPoint())
+			{
+				return Functor(Signature<FPCGNumericPropertyAccessor<double>>{}, NumericProperty);
+			}
+			else if (NumericProperty->IsInteger())
+			{
+				return Functor(Signature<FPCGNumericPropertyAccessor<int64>>{}, NumericProperty);
+			}
 		}
-		else if (StructProperty->Struct == TBaseStructure<FQuat>::Get())
+		else if (const FBoolProperty* BoolProperty = CastField<FBoolProperty>(InProperty))
 		{
-			return MakeUnique<FPCGPropertyStructAccessor<FQuat>>(StructProperty);
+			return Functor(Signature<FPCGPropertyAccessor<bool, FBoolProperty>>{}, BoolProperty);
 		}
-		else if (StructProperty->Struct == TBaseStructure<FTransform>::Get())
+		else if (const FStrProperty* StringProperty = CastField<FStrProperty>(InProperty))
 		{
-			return MakeUnique<FPCGPropertyStructAccessor<FTransform>>(StructProperty);
+			return Functor(Signature<FPCGPropertyAccessor<FString, FStrProperty>>{}, StringProperty);
 		}
-		else if (StructProperty->Struct == TBaseStructure<FRotator>::Get())
+		else if (const FNameProperty* NameProperty = CastField<FNameProperty>(InProperty))
 		{
-			return MakeUnique<FPCGPropertyStructAccessor<FRotator>>(StructProperty);
+			return Functor(Signature<FPCGPropertyAccessor<FName, FNameProperty>>{}, NameProperty);
 		}
-		else if (StructProperty->Struct == TBaseStructure<FVector2D>::Get())
+		else if (const FEnumProperty* EnumProperty = CastField<FEnumProperty>(InProperty))
 		{
-			return MakeUnique<FPCGPropertyStructAccessor<FVector2D>>(StructProperty);
+			return Functor(Signature<FPCGEnumPropertyAccessor>{}, EnumProperty);
 		}
-		else if (StructProperty->Struct == TBaseStructure<FSoftObjectPath>::Get())
+		else if (const FStructProperty* StructProperty = CastField<FStructProperty>(InProperty))
 		{
-			return MakeUnique<FPCGPropertyPathAccessor<FSoftObjectPath>>(InProperty);
+			if (StructProperty->Struct == TBaseStructure<FVector>::Get())
+			{
+				return Functor(Signature<FPCGPropertyStructAccessor<FVector>>{}, StructProperty);
+			}
+			else if (StructProperty->Struct == TBaseStructure<FVector4>::Get())
+			{
+				return Functor(Signature<FPCGPropertyStructAccessor<FVector4>>{}, StructProperty);
+			}
+			else if (StructProperty->Struct == TBaseStructure<FQuat>::Get())
+			{
+				return Functor(Signature<FPCGPropertyStructAccessor<FQuat>>{}, StructProperty);
+			}
+			else if (StructProperty->Struct == TBaseStructure<FTransform>::Get())
+			{
+				return Functor(Signature<FPCGPropertyStructAccessor<FTransform>>{}, StructProperty);
+			}
+			else if (StructProperty->Struct == TBaseStructure<FRotator>::Get())
+			{
+				return Functor(Signature<FPCGPropertyStructAccessor<FRotator>>{}, StructProperty);
+			}
+			else if (StructProperty->Struct == TBaseStructure<FVector2D>::Get())
+			{
+				return Functor(Signature<FPCGPropertyStructAccessor<FVector2D>>{}, StructProperty);
+			}
+			else if (StructProperty->Struct == TBaseStructure<FSoftObjectPath>::Get())
+			{
+				return Functor(Signature<FPCGPropertyPathAccessor<FSoftObjectPath>>{}, InProperty);
+			}
+			else if (StructProperty->Struct == TBaseStructure<FSoftClassPath>::Get())
+			{
+				return Functor(Signature<FPCGPropertyPathAccessor<FSoftClassPath>>{}, InProperty);
+			}
 		}
-		else if (StructProperty->Struct == TBaseStructure<FSoftClassPath>::Get())
+
+		return ReturnType{};
+	}
+
+	template <typename ClassType, typename Func>
+	decltype(auto) DispatchPropertyTypes(const FName InPropertyName, const ClassType* InClass, Func&& Functor)
+	{
+		if (InClass)
 		{
-			return MakeUnique<FPCGPropertyPathAccessor<FSoftClassPath>>(InProperty);
+			if (const FProperty* Property = InClass->FindPropertyByName(InPropertyName))
+			{
+				return DispatchPropertyTypes(Property, std::forward<Func>(Functor));
+			}
 		}
+
+		using ReturnType = decltype(Functor(static_cast<FPCGPropertyPathAccessor<FSoftObjectPath>*>(nullptr), static_cast<const FProperty*>(nullptr)));
+		return ReturnType{};
 	}
+}
 
-	return TUniquePtr<IPCGAttributeAccessor>();
+TUniquePtr<IPCGAttributeAccessor> PCGAttributeAccessorHelpers::CreatePropertyAccessor(const FProperty* InProperty)
+{
+	return PCGAttributeAccessorHelpers::DispatchPropertyTypes(InProperty, [](auto SignatureDummy, const auto* TypedProperty) -> TUniquePtr<IPCGAttributeAccessor>
+	{
+		using TypedAccessor = typename decltype(SignatureDummy)::Type;
+		return MakeUnique<TypedAccessor>(TypedProperty);
+	});
 }
 
 TUniquePtr<IPCGAttributeAccessor> PCGAttributeAccessorHelpers::CreatePropertyAccessor(const FName InPropertyName, const UClass* InClass)
 {
-	if (InClass)
+	return PCGAttributeAccessorHelpers::DispatchPropertyTypes(InPropertyName, InClass, [](auto SignatureDummy, const auto* TypedProperty) -> TUniquePtr<IPCGAttributeAccessor>
 	{
-		if (const FProperty* Property = InClass->FindPropertyByName(InPropertyName))
-		{
-			return CreatePropertyAccessor(Property);
-		}
-	}
-
-	return TUniquePtr<IPCGAttributeAccessor>();
+		using TypedAccessor = typename decltype(SignatureDummy)::Type;
+		return MakeUnique<TypedAccessor>(TypedProperty);
+	});
 }
 
 TUniquePtr<IPCGAttributeAccessor> PCGAttributeAccessorHelpers::CreatePropertyAccessor(const FName InPropertyName, const UStruct* InStruct)
 {
-	if (InStruct)
+	return PCGAttributeAccessorHelpers::DispatchPropertyTypes(InPropertyName, InStruct, [](auto SignatureDummy, const auto* TypedProperty) -> TUniquePtr<IPCGAttributeAccessor>
+	{
+		using TypedAccessor = typename decltype(SignatureDummy)::Type;
+		return MakeUnique<TypedAccessor>(TypedProperty);
+	});
+}
+
+bool PCGAttributeAccessorHelpers::IsPropertyAccessorSupported(const FProperty* InProperty)
+{
+	return PCGAttributeAccessorHelpers::DispatchPropertyTypes(InProperty, [](auto, const auto*) -> bool
+	{
+		return true;
+	});
+}
+
+bool PCGAttributeAccessorHelpers::IsPropertyAccessorSupported(const FName InPropertyName, const UClass* InClass)
+{
+	return PCGAttributeAccessorHelpers::DispatchPropertyTypes(InPropertyName, InClass, [](auto, const auto*) -> bool
+	{
+		return true;
+	});
+}
+
+bool PCGAttributeAccessorHelpers::IsPropertyAccessorSupported(const FName InPropertyName, const UStruct* InStruct)
+{
+	return PCGAttributeAccessorHelpers::DispatchPropertyTypes(InPropertyName, InStruct, [](auto, const auto*) -> bool
 	{
-		if (const FProperty* Property = InStruct->FindPropertyByName(InPropertyName))
+		return true;
+	});
+}
+
+TUniquePtr<const IPCGAttributeAccessor> PCGAttributeAccessorHelpers::CreateConstAccessorForOverrideParam(const FPCGDataCollection& InInputData, const FPCGSettingsOverridableParam& InParam, FName* OutAttributeName)
+{
+	bool bFromGlobalParamsPin = false;
+	TArray<FPCGTaggedData> InputParamData = InInputData.GetParamsByPin(InParam.Label);
+	UPCGParamData* ParamData = !InputParamData.IsEmpty() ? CastChecked<UPCGParamData>(InputParamData[0].Data) : nullptr;
+
+	// If it is empty, try with the param pin
+	if (!ParamData)
+	{
+		bFromGlobalParamsPin = true;
+		ParamData = InInputData.GetFirstParamsOnParamsPin();
+	}
+
+	if (ParamData && ParamData->Metadata && ParamData->Metadata->GetAttributeCount() > 0)
+	{
+		// If the param only has a single attribute and is not from the global Params pin, use this one. Otherwise we need perfect name matching.
+		const FName AttributeName = (ParamData->Metadata->GetAttributeCount() == 1 && !bFromGlobalParamsPin) ? ParamData->Metadata->GetLatestAttributeNameOrNone() : InParam.Properties.Last()->GetFName();
+
+		if (OutAttributeName)
 		{
-			return CreatePropertyAccessor(Property);
+			*OutAttributeName = AttributeName;
 		}
+
+		FPCGAttributePropertySelector InputSelector{};
+		InputSelector.SetAttributeName(AttributeName);
+		return PCGAttributeAccessorHelpers::CreateConstAccessor(ParamData, InputSelector);
 	}
 
-	return TUniquePtr<IPCGAttributeAccessor>();
+	return TUniquePtr<const IPCGAttributeAccessor>{};
 }
 
 TUniquePtr<const IPCGAttributeAccessor> PCGAttributeAccessorHelpers::CreateConstAccessor(const UPCGData* InData, const FPCGAttributePropertySelector& InSelector)
diff --git a/Engine/Plugins/Experimental/PCG/Source/PCG/Private/PCGContext.cpp b/Engine/Plugins/Experimental/PCG/Source/PCG/Private/PCGContext.cpp
--- a/Engine/Plugins/Experimental/PCG/Source/PCG/Private/PCGContext.cpp
+++ b/Engine/Plugins/Experimental/PCG/Source/PCG/Private/PCGContext.cpp
@@ -2,10 +2,32 @@
 
 #include "PCGContext.h"
 #include "PCGComponent.h"
+#include "PCGHelpers.h"
+#include "PCGParamData.h"
+#include "PCGPin.h"
+#include "Metadata/PCGAttributePropertySelector.h"
+#include "Metadata/PCGMetadata.h"
+#include "Metadata/PCGMetadataAttribute.h"
+#include "Metadata/PCGMetadataAttributeTpl.h"
+#include "Metadata/Accessors/IPCGAttributeAccessor.h"
+#include "Metadata/Accessors/PCGAttributeAccessorKeys.h"
+#include "Metadata/Accessors/PCGAttributeAccessorHelpers.h"
+
 #include "GameFramework/Actor.h"
+#include "UObject/Package.h"
 
 #include UE_INLINE_GENERATED_CPP_BY_NAME(PCGContext)
 
+FPCGContext::~FPCGContext()
+{
+	if (SettingsWithOverride)
+	{
+		SettingsWithOverride->RemoveFromRoot();
+		SettingsWithOverride->MarkAsGarbage();
+		SettingsWithOverride = nullptr;
+	}
+}
+
 FString FPCGContext::GetTaskName() const
 {
 	if (Node)
@@ -30,6 +52,22 @@ FString FPCGContext::GetTaskName() const
 	}
 }
 
+int FPCGContext::GetSeed() const
+{
+	if (const UPCGSettings* Settings = GetInputSettings<UPCGSettings>())
+	{
+		return Settings->GetSeed(SourceComponent.Get());
+	}
+	else if (SourceComponent.IsValid())
+	{
+		return SourceComponent->Seed;
+	}
+	else
+	{
+		return 42;
+	}
+}
+
 FString FPCGContext::GetComponentName() const
 {
 	return SourceComponent.IsValid() && SourceComponent->GetOwner() ? SourceComponent->GetOwner()->GetFName().ToString() : TEXT("Non-PCG Component");
@@ -51,3 +89,98 @@ const UPCGSettingsInterface* FPCGContext::GetInputSettingsInterface() const
 		return InputData.GetSettingsInterface();
 	}
 }
+
+void FPCGContext::InitializeSettings()
+{
+	TRACE_CPUPROFILER_EVENT_SCOPE(FPCGContext::InitializeSettings);
+
+	if (SettingsWithOverride)
+	{
+		return;
+	}
+
+	if (const UPCGSettings* NodeSettings = GetOriginalSettings<UPCGSettings>())
+	{
+		// Only duplicate the settings if we have overriable params and we have at least one param pin connected.
+		const TArray<FPCGSettingsOverridableParam>& OverridableParams = NodeSettings->OverridableParams();
+		if (!OverridableParams.IsEmpty())
+		{
+			bool bHasParamConnected = InputData.GetParamsWithDeprecation(Node) != nullptr;
+
+			int32 Index = 0;
+			while (!bHasParamConnected && Index < OverridableParams.Num())
+			{
+				bHasParamConnected |= !InputData.GetParamsByPin(OverridableParams[Index++].Label).IsEmpty();
+			}
+
+			if (bHasParamConnected)
+			{
+				SettingsWithOverride = Cast<UPCGSettings>(StaticDuplicateObject(NodeSettings, GetTransientPackage(), NAME_None, RF_Transient));
+				SettingsWithOverride->AddToRoot();
+			}
+		}
+	}
+}
+
+void FPCGContext::OverrideSettings()
+{
+	TRACE_CPUPROFILER_EVENT_SCOPE(FPCGContext::OverrideSettings);
+
+	// Use original settings to avoid recomputing OverridableParams() everytime
+	const UPCGSettings* OriginalSettings = GetOriginalSettings<UPCGSettings>();
+
+	if (!SettingsWithOverride || !OriginalSettings)
+	{
+		return;
+	}
+
+	for (const FPCGSettingsOverridableParam& Param : OriginalSettings->OverridableParams())
+	{
+		check(!Param.Properties.IsEmpty());
+
+		// TODO: Support override on subobjects
+		if (!SettingsWithOverride->IsA(Param.Properties[0]->GetOwnerClass()))
+		{
+			continue;
+		}
+
+		FName AttributeName = NAME_None;
+		TUniquePtr<const IPCGAttributeAccessor> AttributeAccessor = PCGAttributeAccessorHelpers::CreateConstAccessorForOverrideParam(InputData, Param, &AttributeName);
+
+		// Attribute doesn't exist
+		if (!AttributeAccessor)
+		{
+			continue;
+		}
+
+		TUniquePtr<IPCGAttributeAccessor> PropertyAccessor = PCGAttributeAccessorHelpers::CreatePropertyAccessor(Param.Properties.Last());
+		check(PropertyAccessor.IsValid());
+
+		PCGMetadataAttribute::CallbackWithRightType(PropertyAccessor->GetUnderlyingType(), [this, &AttributeAccessor, &PropertyAccessor, &Param, &AttributeName](auto Dummy) -> bool
+		{
+			using PropertyType = decltype(Dummy);
+
+			// Override were using the first entry (0) by default.
+			FPCGAttributeAccessorKeysEntries FirstEntry(PCGMetadataEntryKey(0));
+					
+			PropertyType Value{};
+			if (!AttributeAccessor->Get<PropertyType>(Value, FirstEntry, EPCGAttributeAccessorFlags::AllowBroadcast | EPCGAttributeAccessorFlags::AllowConstructible))
+			{
+				PCGE_LOG_C(Warning, this, "%s param can't be converted from %s attribute", *Param.Label.ToString(), *AttributeName.ToString());
+				return false;
+			}
+			// TODO: Perhaps factorise this code in another property accessor.
+			void* PropertyObjectPtr = SettingsWithOverride.Get();
+			for (int32 i = 0; i < Param.Properties.Num() - 1; ++i)
+			{
+				PropertyObjectPtr = Param.Properties[i]->ContainerPtrToValuePtr<void>(PropertyObjectPtr);
+			}
+
+			FPCGAttributeAccessorKeysSingleObjectPtr PropertyObjectKey(PropertyObjectPtr);
+
+			PropertyAccessor->Set<PropertyType>(Value, PropertyObjectKey);
+
+			return true;
+		});
+	}
+}
diff --git a/Engine/Plugins/Experimental/PCG/Source/PCG/Private/PCGData.cpp b/Engine/Plugins/Experimental/PCG/Source/PCG/Private/PCGData.cpp
--- a/Engine/Plugins/Experimental/PCG/Source/PCG/Private/PCGData.cpp
+++ b/Engine/Plugins/Experimental/PCG/Source/PCG/Private/PCGData.cpp
@@ -2,7 +2,9 @@
 
 #include "PCGData.h"
 
+#include "PCGNode.h"
 #include "PCGParamData.h"
+#include "PCGPin.h"
 #include "PCGSettings.h"
 #include "Data/PCGPointData.h"
 #include "Data/PCGSpatialData.h"
@@ -156,6 +158,41 @@ TArray<FPCGTaggedData> FPCGDataCollection::GetTaggedParams(const FString& InTag)
 		});
 }
 
+UPCGParamData* FPCGDataCollection::GetParamsWithDeprecation(const UPCGNode* Node) const
+{
+	// First try with the param pin
+	if (UPCGParamData* ParamData = GetFirstParamsOnParamsPin())
+	{
+		return ParamData;
+	}
+
+	// If there is nothing on the params pin, temporarily supports reading param from the input pin, to avoid breaking graphs with nodes that doesn't have a param pin
+	// Log a warning too
+	if (Node)
+	{
+		for (const UPCGPin* InputPin : Node->GetInputPins())
+		{
+			// To avoid matching the individual parameters pins added by the override system, we'll ignore
+			// any input pin that only accepts params.
+			if (InputPin->Properties.AllowedTypes == EPCGDataType::Param)
+			{
+				continue;
+			}
+
+			for (const FPCGTaggedData& TaggedDatum : GetParamsByPin(InputPin->Properties.Label))
+			{
+				if (UPCGParamData* Params = Cast<UPCGParamData>(TaggedDatum.Data))
+				{
+					UE_LOG(LogPCG, Warning, TEXT("[%s] Found a param data on an input pin that should not accept params. Make sure to re-wire it to the param pin if it is used for overrides"), *Node->GetNodeTitle().ToString());
+					return Params;
+				}
+			}
+		}
+	}
+
+	return nullptr;
+}
+
 UPCGParamData* FPCGDataCollection::GetParams() const
 {
 	for (const FPCGTaggedData& TaggedDatum : TaggedData)
@@ -169,7 +206,7 @@ UPCGParamData* FPCGDataCollection::GetParams() const
 	return nullptr;
 }
 
-UPCGParamData* FPCGDataCollection::GetParamsOnParamsPin() const
+UPCGParamData* FPCGDataCollection::GetFirstParamsOnParamsPin() const
 {
 	TArray<FPCGTaggedData> ParamsOnDefaultPin = GetParamsByPin(PCGPinConstants::DefaultParamsLabel);
 	return (ParamsOnDefaultPin.IsEmpty() ? nullptr : Cast<UPCGParamData>(ParamsOnDefaultPin[0].Data));
diff --git a/Engine/Plugins/Experimental/PCG/Source/PCG/Private/PCGElement.cpp b/Engine/Plugins/Experimental/PCG/Source/PCG/Private/PCGElement.cpp
--- a/Engine/Plugins/Experimental/PCG/Source/PCG/Private/PCGElement.cpp
+++ b/Engine/Plugins/Experimental/PCG/Source/PCG/Private/PCGElement.cpp
@@ -38,6 +38,9 @@ bool IPCGElement::Execute(FPCGContext* Context) const
 
 			case EPCGExecutionPhase::PrepareData:
 			{
+				// Will override the settings if there is any override.
+				Context->OverrideSettings();
+
 				if (PrepareDataInternal(Context))
 				{
 					Context->CurrentPhase = EPCGExecutionPhase::Execute;
@@ -330,7 +333,7 @@ void IPCGElement::CleanupAndValidateOutput(FPCGContext* Context) const
 	{
 		// Cleanup any residual labels if the node isn't supposed to produce them
 		// TODO: this is a bit of a crutch, could be refactored out if we review the way we push tagged data
-		TArray<FPCGPinProperties> OutputPinProperties = Settings->OutputPinProperties();
+		TArray<FPCGPinProperties> OutputPinProperties = Settings->AllOutputPinProperties();
 		if(OutputPinProperties.Num() == 1)
 		{
 			for (FPCGTaggedData& TaggedData : Context->OutputData.TaggedData)
diff --git a/Engine/Plugins/Experimental/PCG/Source/PCG/Private/PCGNode.cpp b/Engine/Plugins/Experimental/PCG/Source/PCG/Private/PCGNode.cpp
--- a/Engine/Plugins/Experimental/PCG/Source/PCG/Private/PCGNode.cpp
+++ b/Engine/Plugins/Experimental/PCG/Source/PCG/Private/PCGNode.cpp
@@ -631,8 +631,8 @@ EPCGChangeType UPCGNode::UpdatePins(TFunctionRef<UPCGPin*(UPCGNode*)> PinAllocat
 	}
 
 	UPCGSettings* Settings = GetSettings();
-	TArray<FPCGPinProperties> InboundPinProperties = Settings->InputPinProperties();
-	TArray<FPCGPinProperties> OutboundPinProperties = Settings->OutputPinProperties();
+	TArray<FPCGPinProperties> InboundPinProperties = Settings->AllInputPinProperties();
+	TArray<FPCGPinProperties> OutboundPinProperties = Settings->AllOutputPinProperties();
 
 	auto UpdatePins = [this, &PinAllocator](TArray<UPCGPin*>& Pins, const TArray<FPCGPinProperties>& PinProperties)
 	{
diff --git a/Engine/Plugins/Experimental/PCG/Source/PCG/Private/PCGSettings.cpp b/Engine/Plugins/Experimental/PCG/Source/PCG/Private/PCGSettings.cpp
--- a/Engine/Plugins/Experimental/PCG/Source/PCG/Private/PCGSettings.cpp
+++ b/Engine/Plugins/Experimental/PCG/Source/PCG/Private/PCGSettings.cpp
@@ -2,10 +2,17 @@
 
 #include "PCGSettings.h"
 
+#include "PCGComponent.h"
 #include "PCGContext.h"
 #include "PCGCustomVersion.h"
+#include "PCGGraph.h"
+#include "PCGHelpers.h"
+#include "PCGModule.h"
 #include "PCGPin.h"
 #include "PCGSubsystem.h"
+#include "Helpers/PCGSettingsHelpers.h"
+#include "Metadata/Accessors/PCGAttributeAccessorHelpers.h"
+#include "Metadata/Accessors/IPCGAttributeAccessor.h"
 
 #include "Serialization/ArchiveObjectCrc32.h"
 #include "UObject/ObjectSaveContext.h"
@@ -16,6 +23,8 @@
 #include "Editor.h"
 #endif
 
+#define LOCTEXT_NAMESPACE "PCGSettings"
+
 /** Custom Crc computation that ignores properties that will not affect the computed result of a node. */
 class FPCGSettingsObjectCrc32 : public FArchiveObjectCrc32
 {
@@ -88,6 +97,11 @@ void UPCGSettings::PostLoad()
 {
 	Super::PostLoad();
 
+	if (!HasAnyFlags(RF_ClassDefaultObject))
+	{
+		InitializeCachedOverridableParams();
+	}
+
 #if WITH_EDITOR
 	if (ExecutionMode_DEPRECATED != EPCGSettingsExecutionMode::Enabled)
 	{
@@ -98,6 +112,18 @@ void UPCGSettings::PostLoad()
 #endif
 }
 
+void UPCGSettings::PostInitProperties()
+{
+	Super::PostInitProperties();
+
+#if WITH_EDITOR
+	if (!HasAnyFlags(RF_ClassDefaultObject))
+	{
+		InitializeCachedOverridableParams();
+	}
+#endif //WITH_EDITOR
+}
+
 void UPCGSettings::Serialize(FArchive& Ar)
 {
 	Super::Serialize(Ar);
@@ -168,6 +194,104 @@ TArray<FPCGPinProperties> UPCGSettings::OutputPinProperties() const
 	return PinProperties;
 }
 
+void UPCGSettings::FillOverridableParamsPins(TArray<FPCGPinProperties>& OutPins) const
+{
+	if (!HasOverridableParams())
+	{
+		return;
+	}
+
+	// Validating that we are not clashing with existing pins
+	TMap<FName, EPCGDataType> InputPinsLabelsAndTypes;
+
+	for (const FPCGPinProperties& InputPinProperties : OutPins)
+	{
+		InputPinsLabelsAndTypes.Emplace(InputPinProperties.Label, InputPinProperties.AllowedTypes);
+	}
+
+	// For debugging
+	FString GraphName;
+	FString NodeName;
+
+	if (const UPCGNode* Node = Cast<UPCGNode>(GetOuter()))
+	{
+		NodeName = ((Node->NodeTitle != NAME_None) ? Node->NodeTitle : Node->GetFName()).ToString();
+
+		if (Node->GetGraph())
+		{
+			GraphName = Node->GetGraph()->GetName();
+		}
+	}
+	else
+	{
+		NodeName = GetName();
+	}
+
+	// Adding the multi-pin connection for params.
+	// If it already exists (and is the correct type), we can keep it.
+	const EPCGDataType* PinType = InputPinsLabelsAndTypes.Find(PCGPinConstants::DefaultParamsLabel);
+	if (PinType && (*PinType != EPCGDataType::Param))
+	{
+		const FString ParamsName = PCGPinConstants::DefaultParamsLabel.ToString();
+		UE_LOG(LogPCG, Error, TEXT("[%s-%s] While adding %s pin, we found another %s pin with not the same allowed type (Param). "
+			"Please rename this pin if you want to take advantage of automatic override. Until then it will probably break your graph."), *GraphName, *NodeName, *ParamsName, *ParamsName);
+	}
+	else
+	{
+		FPCGPinProperties& ParamPin = OutPins.Emplace_GetRef(PCGPinConstants::DefaultParamsLabel, EPCGDataType::Param, /*bInAllowMultipleConnections=*/ true, /*bAllowMultipleData=*/ true);
+		ParamPin.bAdvancedPin = false;
+
+#if WITH_EDITOR
+		ParamPin.Tooltip = LOCTEXT("GlobalParamPinTooltip", "Can bundle multiple param data to override multiple parameters at the same time. Names need to match perfectly.");
+#endif // WITH_EDITOR
+	}
+
+	InputPinsLabelsAndTypes.Emplace(PCGPinConstants::DefaultParamsLabel, EPCGDataType::Param);
+
+	for (const FPCGSettingsOverridableParam& OverridableParam : OverridableParams())
+	{
+		if (InputPinsLabelsAndTypes.Contains(OverridableParam.Label))
+		{
+			const FString ParamsName = OverridableParam.Label.ToString();
+			UE_LOG(LogPCG, Warning, TEXT("[%s-%s] While automatically adding overriable param pins, we found a %s pin. "
+				"Please rename this pin if you want to take advantage of automatic override. Until then, we will not add a %s pin."), *GraphName, *NodeName, *ParamsName, *ParamsName);
+			continue;
+		}
+
+		InputPinsLabelsAndTypes.Emplace(OverridableParam.Label, EPCGDataType::Param);
+
+		FPCGPinProperties& ParamPin = OutPins.Emplace_GetRef(OverridableParam.Label, EPCGDataType::Param, /*bInAllowMultipleConnections=*/ false, /*bAllowMultipleData=*/ false);
+		ParamPin.bAdvancedPin = true;
+#if WITH_EDITOR
+		const FProperty* Property = OverridableParam.Properties.Last();
+		check(Property);
+		static FName TooltipMetadata("Tooltip");
+		FString Tooltip;
+		if (const FString* TooltipPtr = Property->FindMetaData(TooltipMetadata))
+		{
+			Tooltip = *TooltipPtr + TEXT("\n");
+		}
+
+		ParamPin.Tooltip = FText::Format(LOCTEXT("OverridableParamPinTooltip", "{0}Param type is \"{1}\" and its exact name is \"{2}\""),
+			FText::FromString(Tooltip),
+			FText::FromString(Property->GetCPPType()),
+			FText::FromName(Property->GetFName()));
+#endif // WITH_EDITOR
+	}
+}
+
+TArray<FPCGPinProperties> UPCGSettings::AllInputPinProperties() const
+{
+	TArray<FPCGPinProperties> InputPins = InputPinProperties();
+	FillOverridableParamsPins(InputPins);
+	return InputPins;
+}
+
+TArray<FPCGPinProperties> UPCGSettings::AllOutputPinProperties() const
+{
+	return OutputPinProperties();
+}
+
 TArray<FPCGPinProperties> UPCGSettings::DefaultInputPinProperties() const
 {
 	return InputPinProperties();
@@ -200,6 +324,11 @@ UPCGNode* UPCGSettings::CreateNode() const
 	return NewObject<UPCGNode>();
 }
 
+int UPCGSettings::GetSeed(const UPCGComponent* InSourceComponent) const
+{
+	return !bUseSeed ? 42 : (InSourceComponent ? PCGHelpers::ComputeSeed(Seed, InSourceComponent->Seed) : Seed);
+}
+
 #if WITH_EDITOR
 void UPCGSettings::PostEditChangeProperty(FPropertyChangedEvent& PropertyChangedEvent)
 {
@@ -223,7 +352,158 @@ void UPCGSettings::DirtyCache()
 		}
 	}
 }
+
+bool UPCGSettings::CanEditChange(const FProperty* InProperty) const
+{
+	if (!InProperty || !Super::CanEditChange(InProperty))
+	{
+		return false;
+	}
+
+	if (!InProperty->HasMetaData(PCGObjectMetadata::Overridable))
+	{
+		return true;
+	}
+
+	if (const UPCGNode* Node = Cast<UPCGNode>(GetOuter()))
+	{
+		const FPCGSettingsOverridableParam* Param = OverridableParams().FindByPredicate([InProperty](const FPCGSettingsOverridableParam& ParamToCheck)
+		{
+			// In OverridableParam, the array of properties is the chain of properties from the Settings class to the wanted param.
+			// Therefore the property we are editing would match the latest property of the array.
+			return ParamToCheck.Properties.Last() == InProperty;
+		});
+
+		if (Param)
+		{
+			if (const UPCGPin* Pin = Node->GetInputPin(Param->Label))
+			{
+				return !Pin->IsConnected();
+			}
+		}
+	}
+
+	return true;
+}
+
+void UPCGSettings::ApplyDeprecationBeforeUpdatePins(UPCGNode* InOutNode, TArray<TObjectPtr<UPCGPin>>& InputPins, TArray<TObjectPtr<UPCGPin>>& OutputPins)
+{
+	check(InOutNode);
+
+	if (DataVersion < FPCGCustomVersion::AddParamPinToOverridableNodes)
+	{
+		PCGSettingsHelpers::DeprecationBreakOutParamsToNewPin(InOutNode, InputPins, OutputPins);
+	}
+}
+#endif // WITH_EDITOR
+
+namespace PCGSettings
+{
+#if WITH_EDITOR
+	template <typename ClassType>
+	TArray<FPCGSettingsOverridableParam> GetAllParams(const ClassType* Class, bool bCheckMetadata, bool bUseSeed)
+	{
+		// TODO: Was not a concern until now, and we didn't have a solution, but this function
+		// only worked if we don't have names clashes in overriable parameters.
+		// The previous override solution was flatenning structs, and only override use the struct member name,
+		// not prefixed by the struct name or anything else.
+		// We cannot prefix it now, because it will break existing node that were assuming the flatenning.
+		// We'll keep this behavior for now, as it might be solved by passing structs instead of param data,
+		// but we'll still at least raise a warning if there is a clash.
+		TSet<FName> LabelCache;
+
+		TArray<FPCGSettingsOverridableParam> Res;
+
+		for (TFieldIterator<FProperty> InputIt(Class, EFieldIteratorFlags::IncludeSuper, EFieldIteratorFlags::ExcludeDeprecated); InputIt; ++InputIt)
+		{
+			const FProperty* Property = *InputIt;
+			if (Property && (!bCheckMetadata || Property->HasMetaData(PCGObjectMetadata::Overridable)))
+			{
+				// Don't allow to override the seed if the settings doesn't use the seed.
+				if (Property->GetFName() == GET_MEMBER_NAME_CHECKED(UPCGSettings, Seed) && !bUseSeed)
+				{
+					continue;
+				}
+
+				// Validating that the property can be overriden by params
+				if (PCGAttributeAccessorHelpers::IsPropertyAccessorSupported(Property))
+				{
+					FName Label = *Property->GetDisplayNameText().ToString();
+					if (LabelCache.Contains(Label))
+					{
+						UE_LOG(LogPCG, Warning, TEXT("%s property clashes with another property already found. It is a limitation at the moment and this property will be ignored (ie. will not be overridable)"), *Label.ToString());
+						continue;
+					}
+
+					LabelCache.Add(Label);
+
+					FPCGSettingsOverridableParam& Param = Res.Emplace_GetRef();
+					Param.Label = Label;
+					Param.PropertiesNames.Add(Property->GetFName());
+				}
+				else if (const FStructProperty* StructProperty = CastField<FStructProperty>(Property))
+				{
+					FString PropertyName = Property->GetDisplayNameText().ToString();
+					for (const FPCGSettingsOverridableParam& ChildParam : GetAllParams(StructProperty->Struct, /*bCheckMetadata=*/false, /*bUseSeed=*/true))
+					{
+						FName Label = ChildParam.Label;
+						if (LabelCache.Contains(Label))
+						{
+							UE_LOG(LogPCG, Warning, TEXT("%s property clashes with another property already found. It is a limitation at the moment and this property will be ignored (ie. will not be overridable)"), *Label.ToString());
+							continue;
+						}
+
+						LabelCache.Add(Label);
+
+						FPCGSettingsOverridableParam& Param = Res.Emplace_GetRef();
+						Param.Label = Label;
+						Param.PropertiesNames.Add(Property->GetFName());
+						Param.PropertiesNames.Append(ChildParam.PropertiesNames);
+					}
+				}
+			}
+		}
+
+		return Res;
+	}
 #endif // WITH_EDITOR
+}
+
+void UPCGSettings::InitializeCachedOverridableParams()
+{
+	TRACE_CPUPROFILER_EVENT_SCOPE(UPCGSettings::InitializeCachedOverridableParams);
+
+#if WITH_EDITOR
+	if (CachedOverridableParams.IsEmpty())
+	{
+		CachedOverridableParams = PCGSettings::GetAllParams(GetClass(), true, bUseSeed);
+	}
+#endif // WITH_EDITOR
+
+	for (FPCGSettingsOverridableParam& Param : CachedOverridableParams)
+	{
+		check(!Param.PropertiesNames.IsEmpty());
+		Param.Properties.Reset(Param.PropertiesNames.Num());
+
+		// Some properties might not be available at runtime. Ignore them.
+		if (const FProperty* Property = GetClass()->FindPropertyByName(Param.PropertiesNames[0]))
+		{
+			Param.Properties.Add(Property);
+
+			for (int32 i = 1; i < Param.PropertiesNames.Num(); ++i)
+			{
+				// If we have multiple depth properties, it should be Struct properties by construction
+				const FStructProperty* StructProperty = CastField<FStructProperty>(Property);
+				if (ensure(StructProperty))
+				{
+					Property = StructProperty->Struct->FindPropertyByName(Param.PropertiesNames[i]);
+					check(Property);
+					Param.Properties.Add(Property);
+				}
+			}
+		}
+	}
+}
 
 TArray<FPCGPinProperties> UPCGSettings::DefaultPointOutputPinProperties() const
 {
@@ -312,3 +592,5 @@ bool FPCGTrivialElement::ExecuteInternal(FPCGContext* Context) const
 	Context->OutputData = Context->InputData;
 	return true;
 }
+
+#undef LOCTEXT_NAMESPACE
\ No newline at end of file
diff --git a/Engine/Plugins/Experimental/PCG/Source/PCG/Private/Tests/Determinism/PCGDeterminismTestBlueprintBase.cpp b/Engine/Plugins/Experimental/PCG/Source/PCG/Private/Tests/Determinism/PCGDeterminismTestBlueprintBase.cpp
--- a/Engine/Plugins/Experimental/PCG/Source/PCG/Private/Tests/Determinism/PCGDeterminismTestBlueprintBase.cpp
+++ b/Engine/Plugins/Experimental/PCG/Source/PCG/Private/Tests/Determinism/PCGDeterminismTestBlueprintBase.cpp
@@ -11,6 +11,6 @@ void UPCGDeterminismTestBlueprintBase::ExecuteTest_Implementation(const UPCGNode
 	// The default case if the user tries to run the blueprint test unimplemented
 	InOutTestResult.TestResultTitle = InPCGNode->GetNodeTitle();
 	InOutTestResult.TestResultName = InPCGNode->GetName();
-	InOutTestResult.Seed = InPCGNode->GetSettings()->Seed;
+	InOutTestResult.Seed = InPCGNode->GetSettings()->GetSeed();
 	PCGDeterminismTests::RunDeterminismTest(InPCGNode, InOutTestResult, PCGDeterminismTests::Defaults::DeterminismBasicTestInfo);
 }
diff --git a/Engine/Plugins/Experimental/PCG/Source/PCG/Private/Tests/Elements/PCGDataTableRowToParamDataTest.cpp b/Engine/Plugins/Experimental/PCG/Source/PCG/Private/Tests/Elements/PCGDataTableRowToParamDataTest.cpp
--- a/Engine/Plugins/Experimental/PCG/Source/PCG/Private/Tests/Elements/PCGDataTableRowToParamDataTest.cpp
+++ b/Engine/Plugins/Experimental/PCG/Source/PCG/Private/Tests/Elements/PCGDataTableRowToParamDataTest.cpp
@@ -36,7 +36,8 @@ CCC,"C Name","C String","333","3333","3.3","3.33","(X=3.0,Y=3.0)","(X=3.0,Y=3.0,
 
 	while (!TestElement->Execute(Context.Get())) {}
 
-	const UPCGParamData* Params = Context->OutputData.GetParams();
+	TArray<FPCGTaggedData> OutputParamData = Context->OutputData.GetParamsByPin(PCGPinConstants::DefaultOutputLabel);
+	const UPCGParamData* Params = OutputParamData.IsEmpty() ? nullptr : Cast<UPCGParamData>(OutputParamData[0].Data);
 	UTEST_NOT_NULL("Output params", Params);
 
 	UTEST_NOT_NULL("Params has metadata", Params->Metadata.Get());
@@ -97,7 +98,7 @@ CCC,"C Name","C String","333","3333","3.3","3.33","(X=3.0,Y=3.0)","(X=3.0,Y=3.0,
 
 	while (!TestElement->Execute(Context.Get())) {}
 
-	const UPCGParamData* Params = Context->OutputData.GetParams();
+	const UPCGParamData* Params = Context->OutputData.GetFirstParamsOnParamsPin();
 	UTEST_NOT_NULL("Output params", Params);
 
 	UTEST_NOT_NULL("Params has metadata", Params->Metadata.Get());
diff --git a/Engine/Plugins/Experimental/PCG/Source/PCG/Private/Tests/PCGTestsCommon.cpp b/Engine/Plugins/Experimental/PCG/Source/PCG/Private/Tests/PCGTestsCommon.cpp
--- a/Engine/Plugins/Experimental/PCG/Source/PCG/Private/Tests/PCGTestsCommon.cpp
+++ b/Engine/Plugins/Experimental/PCG/Source/PCG/Private/Tests/PCGTestsCommon.cpp
@@ -253,7 +253,7 @@ bool FPCGTestBaseClass::SmokeTestAnyValidInput(UPCGSettings* InSettings, TFuncti
 		return false;
 	}
 
-	TArray<FPCGPinProperties> InputProperties = InSettings->InputPinProperties();
+	TArray<FPCGPinProperties> InputProperties = InSettings->AllInputPinProperties();
 	// For each pin: take nothing, take 1 of any supported type, take 2 of any supported types (if enabled)
 	TArray<TArray<FPCGDataCollection>> InputsPerProperties;
 	TArray<uint32> InputIndices;
diff --git a/Engine/Plugins/Experimental/PCG/Source/PCG/Public/Elements/Metadata/PCGMetadataPartition.h b/Engine/Plugins/Experimental/PCG/Source/PCG/Public/Elements/Metadata/PCGMetadataPartition.h
--- a/Engine/Plugins/Experimental/PCG/Source/PCG/Public/Elements/Metadata/PCGMetadataPartition.h
+++ b/Engine/Plugins/Experimental/PCG/Source/PCG/Public/Elements/Metadata/PCGMetadataPartition.h
@@ -26,7 +26,7 @@ protected:
 	//~End UPCGSettings interface
 
 public:
-	UPROPERTY(BlueprintReadWrite, EditAnywhere, Category = Settings)
+	UPROPERTY(BlueprintReadWrite, EditAnywhere, Category = Settings, meta = (PCG_Overridable))
 	FName PartitionAttribute = NAME_None;
 };
 
diff --git a/Engine/Plugins/Experimental/PCG/Source/PCG/Public/Elements/Metadata/PCGMetadataRenameElement.h b/Engine/Plugins/Experimental/PCG/Source/PCG/Public/Elements/Metadata/PCGMetadataRenameElement.h
--- a/Engine/Plugins/Experimental/PCG/Source/PCG/Public/Elements/Metadata/PCGMetadataRenameElement.h
+++ b/Engine/Plugins/Experimental/PCG/Source/PCG/Public/Elements/Metadata/PCGMetadataRenameElement.h
@@ -18,7 +18,6 @@ public:
 	virtual EPCGSettingsType GetType() const override { return EPCGSettingsType::Metadata; }
 #endif
 
-	virtual TArray<FPCGPinProperties> InputPinProperties() const override;
 	virtual TArray<FPCGPinProperties> OutputPinProperties() const override;
 
 protected:
@@ -26,10 +25,10 @@ protected:
 	//~End UPCGSettings interface
 
 public:
-	UPROPERTY(BlueprintReadWrite, EditAnywhere, Category = Settings)
+	UPROPERTY(BlueprintReadWrite, EditAnywhere, Category = Settings, meta = (PCG_Overridable))
 	FName AttributeToRename = NAME_None;
 
-	UPROPERTY(BlueprintReadWrite, EditAnywhere, Category = Settings)
+	UPROPERTY(BlueprintReadWrite, EditAnywhere, Category = Settings, meta = (PCG_Overridable))
 	FName NewAttributeName = NAME_None;
 };
 
diff --git a/Engine/Plugins/Experimental/PCG/Source/PCG/Public/Elements/PCGAttributeGetFromPointIndexElement.h b/Engine/Plugins/Experimental/PCG/Source/PCG/Public/Elements/PCGAttributeGetFromPointIndexElement.h
--- a/Engine/Plugins/Experimental/PCG/Source/PCG/Public/Elements/PCGAttributeGetFromPointIndexElement.h
+++ b/Engine/Plugins/Experimental/PCG/Source/PCG/Public/Elements/PCGAttributeGetFromPointIndexElement.h
@@ -43,7 +43,7 @@ public:
 	UPROPERTY(BlueprintReadWrite, EditAnywhere, Category = Settings)
 	FPCGAttributePropertySelector InputSource;
 
-	UPROPERTY(BlueprintReadWrite, EditAnywhere, Category = Settings)
+	UPROPERTY(BlueprintReadWrite, EditAnywhere, Category = Settings, meta = (PCG_Overridable))
 	int32 Index = 0;
 
 	UPROPERTY(BlueprintReadWrite, EditAnywhere, Category = Settings)
diff --git a/Engine/Plugins/Experimental/PCG/Source/PCG/Public/Elements/PCGAttributeSelectElement.h b/Engine/Plugins/Experimental/PCG/Source/PCG/Public/Elements/PCGAttributeSelectElement.h
--- a/Engine/Plugins/Experimental/PCG/Source/PCG/Public/Elements/PCGAttributeSelectElement.h
+++ b/Engine/Plugins/Experimental/PCG/Source/PCG/Public/Elements/PCGAttributeSelectElement.h
@@ -72,13 +72,13 @@ public:
 	UPROPERTY(BlueprintReadWrite, EditAnywhere, Category = Settings)
 	FName OutputAttributeName = NAME_None;
 
-	UPROPERTY(BlueprintReadWrite, EditAnywhere, Category = Settings)
+	UPROPERTY(BlueprintReadWrite, EditAnywhere, Category = Settings, meta = (PCG_Overridable))
 	EPCGAttributeSelectOperation Operation = EPCGAttributeSelectOperation::Min;
 
-	UPROPERTY(BlueprintReadWrite, EditAnywhere, Category = Settings)
+	UPROPERTY(BlueprintReadWrite, EditAnywhere, Category = Settings, meta = (PCG_Overridable))
 	EPCGAttributeSelectAxis Axis = EPCGAttributeSelectAxis::X;
 
-	UPROPERTY(BlueprintReadWrite, EditAnywhere, Category = Settings, meta = (EditCondition = "Axis == EPCGAttributeSelectAxis::CustomAxis", EditConditionHides))
+	UPROPERTY(BlueprintReadWrite, EditAnywhere, Category = Settings, meta = (EditCondition = "Axis == EPCGAttributeSelectAxis::CustomAxis", PCG_Overridable))
 	FVector4 CustomAxis = FVector4::Zero();
 
 #if WITH_EDITORONLY_DATA
diff --git a/Engine/Plugins/Experimental/PCG/Source/PCG/Public/Elements/PCGBoundsModifier.h b/Engine/Plugins/Experimental/PCG/Source/PCG/Public/Elements/PCGBoundsModifier.h
--- a/Engine/Plugins/Experimental/PCG/Source/PCG/Public/Elements/PCGBoundsModifier.h
+++ b/Engine/Plugins/Experimental/PCG/Source/PCG/Public/Elements/PCGBoundsModifier.h
@@ -44,19 +44,19 @@ protected:
 	//~End UPCGSettings interface
 
 public:
-	UPROPERTY(BlueprintReadWrite, EditAnywhere, Category = Settings)
+	UPROPERTY(BlueprintReadWrite, EditAnywhere, Category = Settings, meta = (PCG_Overridable))
 	EPCGBoundsModifierMode Mode = EPCGBoundsModifierMode::Scale;
 
-	UPROPERTY(BlueprintReadWrite, EditAnywhere, Category = Settings)
+	UPROPERTY(BlueprintReadWrite, EditAnywhere, Category = Settings, meta = (PCG_Overridable))
 	FVector BoundsMin = FVector::One();
 
-	UPROPERTY(BlueprintReadWrite, EditAnywhere, Category = Settings)
+	UPROPERTY(BlueprintReadWrite, EditAnywhere, Category = Settings, meta = (PCG_Overridable))
 	FVector BoundsMax = FVector::One();
 
-	UPROPERTY(BlueprintReadWrite, EditAnywhere, Category = Settings)
+	UPROPERTY(BlueprintReadWrite, EditAnywhere, Category = Settings, meta = (PCG_Overridable))
 	bool bAffectSteepness = false;
 
-	UPROPERTY(BlueprintReadWrite, EditAnywhere, Category = Settings, meta = (EditCondition = "bAffectSteepness", ClampMin = "0", ClampMax = "1"))
+	UPROPERTY(BlueprintReadWrite, EditAnywhere, Category = Settings, meta = (EditCondition = "bAffectSteepness", ClampMin = "0", ClampMax = "1", PCG_Overridable))
 	float Steepness = 1.0f;
 };
 
diff --git a/Engine/Plugins/Experimental/PCG/Source/PCG/Public/Elements/PCGCopyPoints.h b/Engine/Plugins/Experimental/PCG/Source/PCG/Public/Elements/PCGCopyPoints.h
--- a/Engine/Plugins/Experimental/PCG/Source/PCG/Public/Elements/PCGCopyPoints.h
+++ b/Engine/Plugins/Experimental/PCG/Source/PCG/Public/Elements/PCGCopyPoints.h
@@ -9,7 +9,6 @@
 
 namespace PCGCopyPointsConstants
 {
-	const FName ParamsLabel = TEXT("Params");
 	const FName SourcePointsLabel = TEXT("Source");
 	const FName TargetPointsLabel = TEXT("Target");
 }
@@ -50,23 +49,23 @@ protected:
 
 public:
 	/** The method used to determine output point rotation */
-	UPROPERTY(BlueprintReadWrite, EditAnywhere, Category = Settings)
+	UPROPERTY(BlueprintReadWrite, EditAnywhere, Category = Settings, meta = (PCG_Overridable))
 	EPCGCopyPointsInheritanceMode RotationInheritance = EPCGCopyPointsInheritanceMode::Relative;
 
 	/** The method used to determine output point scale */
-	UPROPERTY(BlueprintReadWrite, EditAnywhere, Category = Settings)
+	UPROPERTY(BlueprintReadWrite, EditAnywhere, Category = Settings, meta = (PCG_Overridable))
 	EPCGCopyPointsInheritanceMode ScaleInheritance = EPCGCopyPointsInheritanceMode::Relative;
 
 	/** The method used to determine output point color */
-	UPROPERTY(BlueprintReadWrite, EditAnywhere, Category = Settings)
+	UPROPERTY(BlueprintReadWrite, EditAnywhere, Category = Settings, meta = (PCG_Overridable))
 	EPCGCopyPointsInheritanceMode ColorInheritance = EPCGCopyPointsInheritanceMode::Relative;
 
 	/** The method used to determine output seed values. Relative recomputes the seed from the new location. */
-	UPROPERTY(BlueprintReadWrite, EditAnywhere, Category = Settings)
+	UPROPERTY(BlueprintReadWrite, EditAnywhere, Category = Settings, meta = (PCG_Overridable))
 	EPCGCopyPointsInheritanceMode SeedInheritance = EPCGCopyPointsInheritanceMode::Relative;
 
 	/** The method used to determine output data attributes */
-	UPROPERTY(BlueprintReadWrite, EditAnywhere, Category = Settings)
+	UPROPERTY(BlueprintReadWrite, EditAnywhere, Category = Settings, meta = (PCG_Overridable))
 	EPCGCopyPointsMetadataInheritanceMode AttributeInheritance = EPCGCopyPointsMetadataInheritanceMode::Source;
 };
 
diff --git a/Engine/Plugins/Experimental/PCG/Source/PCG/Public/Elements/PCGCreateAttribute.h b/Engine/Plugins/Experimental/PCG/Source/PCG/Public/Elements/PCGCreateAttribute.h
--- a/Engine/Plugins/Experimental/PCG/Source/PCG/Public/Elements/PCGCreateAttribute.h
+++ b/Engine/Plugins/Experimental/PCG/Source/PCG/Public/Elements/PCGCreateAttribute.h
@@ -4,6 +4,7 @@
 
 #include "PCGSettings.h"
 #include "Metadata/PCGMetadataTypesConstantStruct.h"
+#include "Metadata/PCGMetadataCommon.h"
 
 #include "PCGCreateAttribute.generated.h"
 
@@ -109,10 +110,10 @@ protected:
 
 private:
 	/* Create (or clear) an attribute named by OutputAttributeName and depending on the selected type. Value can be overridden by params. Default value will be set to the specified value. */
-	FPCGMetadataAttributeBase* ClearOrCreateAttribute(const UPCGCreateAttributeSettings* Settings, UPCGMetadata* Metadata, const UPCGParamData* Params = nullptr, const FName* OutputAttributeNameOverride = nullptr) const;
+	FPCGMetadataAttributeBase* ClearOrCreateAttribute(const UPCGCreateAttributeSettings* Settings, UPCGMetadata* Metadata, const FName* OutputAttributeNameOverride = nullptr) const;
 
 	/* Set an entry defined by EntryKey (or create one if it is invalid) in the Attribute depending on the selected type. Value can be overridden by params. */
-	PCGMetadataEntryKey SetAttribute(const UPCGCreateAttributeSettings* Settings, FPCGMetadataAttributeBase* Attribute, UPCGMetadata* Metadata, PCGMetadataEntryKey EntryKey, const UPCGParamData* Params = nullptr) const;
+	PCGMetadataEntryKey SetAttribute(const UPCGCreateAttributeSettings* Settings, FPCGMetadataAttributeBase* Attribute, UPCGMetadata* Metadata, PCGMetadataEntryKey EntryKey) const;
 };
 
 #if UE_ENABLE_INCLUDE_ORDER_DEPRECATED_IN_5_2
diff --git a/Engine/Plugins/Experimental/PCG/Source/PCG/Public/Elements/PCGDataTableRowToParamData.h b/Engine/Plugins/Experimental/PCG/Source/PCG/Public/Elements/PCGDataTableRowToParamData.h
--- a/Engine/Plugins/Experimental/PCG/Source/PCG/Public/Elements/PCGDataTableRowToParamData.h
+++ b/Engine/Plugins/Experimental/PCG/Source/PCG/Public/Elements/PCGDataTableRowToParamData.h
@@ -26,9 +26,13 @@ public:
 	virtual TArray<FPCGPinProperties> InputPinProperties() const override;
 
 	// The name of the row to copy from
-	UPROPERTY(BlueprintReadWrite, EditAnywhere, Category = Settings)
+	UPROPERTY(BlueprintReadWrite, EditAnywhere, Category = Settings, meta = (PCG_Overridable))
 	FName RowName = NAME_None;	
 
+	// Path override, hidden to be only presented as param pin
+	UPROPERTY(BlueprintReadWrite, Category = Settings, meta = (PCG_Overridable))
+	FString PathOverride = FString();
+
 	// the data table to copy from
 	UPROPERTY(BlueprintReadWrite, EditAnywhere, Category = Settings)
 	TSoftObjectPtr<UDataTable> DataTable;
diff --git a/Engine/Plugins/Experimental/PCG/Source/PCG/Public/Elements/PCGDensityFilter.h b/Engine/Plugins/Experimental/PCG/Source/PCG/Public/Elements/PCGDensityFilter.h
--- a/Engine/Plugins/Experimental/PCG/Source/PCG/Public/Elements/PCGDensityFilter.h
+++ b/Engine/Plugins/Experimental/PCG/Source/PCG/Public/Elements/PCGDensityFilter.h
@@ -18,7 +18,6 @@ public:
 #if WITH_EDITOR
 	virtual FName GetDefaultNodeName() const override { return FName(TEXT("DensityFilter")); }
 	virtual EPCGSettingsType GetType() const override { return EPCGSettingsType::Filter; }
-	virtual void ApplyDeprecationBeforeUpdatePins(UPCGNode* InOutNode, TArray<TObjectPtr<UPCGPin>>& InputPins, TArray<TObjectPtr<UPCGPin>>& OutputPins) override;
 #endif
 
 	virtual TArray<FPCGPinProperties> InputPinProperties() const override;
@@ -29,17 +28,17 @@ protected:
 	//~End UPCGSettings interface
 
 public:
-	UPROPERTY(BlueprintReadWrite, EditAnywhere, Category = Settings, meta = (ClampMin = "0", ClampMax = "1"))
+	UPROPERTY(BlueprintReadWrite, EditAnywhere, Category = Settings, meta = (ClampMin = "0", ClampMax = "1", PCG_Overridable))
 	float LowerBound = 0.5f;
 
-	UPROPERTY(BlueprintReadWrite, EditAnywhere, Category = Settings, meta = (ClampMin = "0", ClampMax = "1"))
+	UPROPERTY(BlueprintReadWrite, EditAnywhere, Category = Settings, meta = (ClampMin = "0", ClampMax = "1", PCG_Overridable))
 	float UpperBound = 1.0f;
 
-	UPROPERTY(BlueprintReadWrite, EditAnywhere, Category = Settings)
+	UPROPERTY(BlueprintReadWrite, EditAnywhere, Category = Settings, meta = (PCG_Overridable))
 	bool bInvertFilter = false;
 
 #if WITH_EDITORONLY_DATA
-	UPROPERTY(Transient, BlueprintReadWrite, EditAnywhere, Category = "Settings|Debug")
+	UPROPERTY(Transient, BlueprintReadWrite, EditAnywhere, Category = "Settings|Debug", meta = (PCG_Overridable))
 	bool bKeepZeroDensityPoints = false;
 #endif
 };
diff --git a/Engine/Plugins/Experimental/PCG/Source/PCG/Public/Elements/PCGDensityNoise.h b/Engine/Plugins/Experimental/PCG/Source/PCG/Public/Elements/PCGDensityNoise.h
--- a/Engine/Plugins/Experimental/PCG/Source/PCG/Public/Elements/PCGDensityNoise.h
+++ b/Engine/Plugins/Experimental/PCG/Source/PCG/Public/Elements/PCGDensityNoise.h
@@ -31,7 +31,6 @@ public:
 #if WITH_EDITOR
 	virtual FName GetDefaultNodeName() const override { return FName(TEXT("DensityNoise")); }
 	virtual EPCGSettingsType GetType() const override { return EPCGSettingsType::Density; }
-	virtual void ApplyDeprecationBeforeUpdatePins(UPCGNode* InOutNode, TArray<TObjectPtr<UPCGPin>>& InputPins, TArray<TObjectPtr<UPCGPin>>& OutputPins) override;
 #endif
 
 	virtual TArray<FPCGPinProperties> InputPinProperties() const override;
@@ -43,17 +42,17 @@ protected:
 
 public:
 	/** Density = (OriginalDensity op DensityNoise), DensityNoise in [DensityNoiseMin, DensityNoiseMax] */
-	UPROPERTY(BlueprintReadWrite, EditAnywhere, Category = Settings)
+	UPROPERTY(BlueprintReadWrite, EditAnywhere, Category = Settings, meta = (PCG_Overridable))
 	EPCGDensityNoiseMode DensityMode = EPCGDensityNoiseMode::Set;
 
-	UPROPERTY(BlueprintReadWrite, EditAnywhere, Category = Settings)
+	UPROPERTY(BlueprintReadWrite, EditAnywhere, Category = Settings, meta = (PCG_Overridable))
 	float DensityNoiseMin = 0.f;
 
-	UPROPERTY(BlueprintReadWrite, EditAnywhere, Category = Settings)
+	UPROPERTY(BlueprintReadWrite, EditAnywhere, Category = Settings, meta = (PCG_Overridable))
 	float DensityNoiseMax = 1.f;
 
 	/** Density = 1 - Density before applying the DensityMode operation */
-	UPROPERTY(BlueprintReadWrite, EditAnywhere, Category = Settings)
+	UPROPERTY(BlueprintReadWrite, EditAnywhere, Category = Settings, meta = (PCG_Overridable))
 	bool bInvertSourceDensity = false;
 };
 
diff --git a/Engine/Plugins/Experimental/PCG/Source/PCG/Public/Elements/PCGDensityRemap.h b/Engine/Plugins/Experimental/PCG/Source/PCG/Public/Elements/PCGDensityRemap.h
--- a/Engine/Plugins/Experimental/PCG/Source/PCG/Public/Elements/PCGDensityRemap.h
+++ b/Engine/Plugins/Experimental/PCG/Source/PCG/Public/Elements/PCGDensityRemap.h
@@ -21,6 +21,7 @@ public:
 	virtual EPCGSettingsType GetType() const override { return EPCGSettingsType::Density; }
 #endif
 
+	virtual TArray<FPCGPinProperties> InputPinProperties() const override;
 	virtual TArray<FPCGPinProperties> OutputPinProperties() const override { return Super::DefaultPointOutputPinProperties(); }
 
 protected:
@@ -28,13 +29,13 @@ protected:
 	//~End UPCGSettings interface
 
 public:
-	UPROPERTY(BlueprintReadWrite, EditAnywhere, Category = Settings, meta = (ClampMin = "0", ClampMax = "1"))
+	UPROPERTY(BlueprintReadWrite, EditAnywhere, Category = Settings, meta = (ClampMin = "0", ClampMax = "1", PCG_Overridable))
 	float RemapMin = 0.5f;
 
-	UPROPERTY(BlueprintReadWrite, EditAnywhere, Category = Settings, meta = (ClampMin = "0", ClampMax = "1"))
+	UPROPERTY(BlueprintReadWrite, EditAnywhere, Category = Settings, meta = (ClampMin = "0", ClampMax = "1", PCG_Overridable))
 	float RemapMax = 1.0f;
 
-	UPROPERTY(BlueprintReadWrite, EditAnywhere, Category = Settings)
+	UPROPERTY(BlueprintReadWrite, EditAnywhere, Category = Settings, meta = (PCG_Overridable))
 	bool bMultiplyDensity = true;
 };
 
diff --git a/Engine/Plugins/Experimental/PCG/Source/PCG/Public/Elements/PCGDensityRemapElement.h b/Engine/Plugins/Experimental/PCG/Source/PCG/Public/Elements/PCGDensityRemapElement.h
--- a/Engine/Plugins/Experimental/PCG/Source/PCG/Public/Elements/PCGDensityRemapElement.h
+++ b/Engine/Plugins/Experimental/PCG/Source/PCG/Public/Elements/PCGDensityRemapElement.h
@@ -19,7 +19,6 @@ public:
 #if WITH_EDITOR
 	virtual FName GetDefaultNodeName() const override { return FName(TEXT("DensityRemap")); }
 	virtual EPCGSettingsType GetType() const override { return EPCGSettingsType::Density; }
-	virtual void ApplyDeprecationBeforeUpdatePins(UPCGNode* InOutNode, TArray<TObjectPtr<UPCGPin>>& InputPins, TArray<TObjectPtr<UPCGPin>>& OutputPins) override;
 #endif
 
 	virtual TArray<FPCGPinProperties> InputPinProperties() const override;
@@ -31,21 +30,21 @@ protected:
 
 public:
 	/** If InRangeMin = InRangeMax, then that density value is mapped to the average of OutRangeMin and OutRangeMax */
-	UPROPERTY(BlueprintReadWrite, EditAnywhere, Category = Settings, meta = (ClampMin = "0", ClampMax = "1"))
+	UPROPERTY(BlueprintReadWrite, EditAnywhere, Category = Settings, meta = (ClampMin = "0", ClampMax = "1", PCG_Overridable))
 	float InRangeMin = 0.f;
 
 	/** If InRangeMin = InRangeMax, then that density value is mapped to the average of OutRangeMin and OutRangeMax */
-	UPROPERTY(BlueprintReadWrite, EditAnywhere, Category = Settings, meta = (ClampMin = "0", ClampMax = "1"))
+	UPROPERTY(BlueprintReadWrite, EditAnywhere, Category = Settings, meta = (ClampMin = "0", ClampMax = "1", PCG_Overridable))
 	float InRangeMax = 1.f;
 
-	UPROPERTY(BlueprintReadWrite, EditAnywhere, Category = Settings, meta = (ClampMin = "0", ClampMax = "1"))
+	UPROPERTY(BlueprintReadWrite, EditAnywhere, Category = Settings, meta = (ClampMin = "0", ClampMax = "1", PCG_Overridable))
 	float OutRangeMin = 0.f;
 
-	UPROPERTY(BlueprintReadWrite, EditAnywhere, Category = Settings, meta = (ClampMin = "0", ClampMax = "1"))
+	UPROPERTY(BlueprintReadWrite, EditAnywhere, Category = Settings, meta = (ClampMin = "0", ClampMax = "1", PCG_Overridable))
 	float OutRangeMax = 1.f;
 
 	/** Density values outside of the input range will be unaffected by the remapping */
-	UPROPERTY(BlueprintReadWrite, EditAnywhere, Category = Settings)
+	UPROPERTY(BlueprintReadWrite, EditAnywhere, Category = Settings, meta = (PCG_Overridable))
 	bool bExcludeValuesOutsideInputRange = false;
 };
 
diff --git a/Engine/Plugins/Experimental/PCG/Source/PCG/Public/Elements/PCGDifferenceElement.h b/Engine/Plugins/Experimental/PCG/Source/PCG/Public/Elements/PCGDifferenceElement.h
--- a/Engine/Plugins/Experimental/PCG/Source/PCG/Public/Elements/PCGDifferenceElement.h
+++ b/Engine/Plugins/Experimental/PCG/Source/PCG/Public/Elements/PCGDifferenceElement.h
@@ -33,7 +33,7 @@ protected:
 	//~End UPCGSettings interface
 
 public:
-	UPROPERTY(BlueprintReadWrite, EditAnywhere, Category = Settings)
+	UPROPERTY(BlueprintReadWrite, EditAnywhere, Category = Settings, meta = (PCG_Overridable))
 	EPCGDifferenceDensityFunction DensityFunction = EPCGDifferenceDensityFunction::Minimum;
 
 	UPROPERTY(BlueprintReadWrite, EditAnywhere, Category = Settings)
@@ -43,7 +43,7 @@ public:
 	bool bDiffMetadata = true;
 
 #if WITH_EDITORONLY_DATA
-	UPROPERTY(Transient, BlueprintReadWrite, EditAnywhere, Category = "Settings|Debug")
+	UPROPERTY(Transient, BlueprintReadWrite, EditAnywhere, Category = "Settings|Debug", meta = (PCG_Overridable))
 	bool bKeepZeroDensityPoints = false;
 #endif
 };
diff --git a/Engine/Plugins/Experimental/PCG/Source/PCG/Public/Elements/PCGDistance.h b/Engine/Plugins/Experimental/PCG/Source/PCG/Public/Elements/PCGDistance.h
--- a/Engine/Plugins/Experimental/PCG/Source/PCG/Public/Elements/PCGDistance.h
+++ b/Engine/Plugins/Experimental/PCG/Source/PCG/Public/Elements/PCGDistance.h
@@ -45,23 +45,23 @@ protected:
 public:
 
 	// The name of the attribute to store on the point. Use 'None' to disable
-	UPROPERTY(BlueprintReadWrite, EditAnywhere, Category = Settings)
+	UPROPERTY(BlueprintReadWrite, EditAnywhere, Category = Settings, meta = (PCG_Overridable))
 	FName AttributeName = TEXT("Distance");
 
 	// If true, will also set the density to be 0-1 based on MaximumDistance
-	UPROPERTY(BlueprintReadWrite, EditAnywhere, Category = Settings)
+	UPROPERTY(BlueprintReadWrite, EditAnywhere, Category = Settings, meta = (PCG_Overridable))
 	bool bSetDensity = false;
 
 	// A maximum distance to search, which is used as an optimization
-	UPROPERTY(BlueprintReadWrite, EditAnywhere, Category = Settings, meta = (ClampMin = "1"))
+	UPROPERTY(BlueprintReadWrite, EditAnywhere, Category = Settings, meta = (ClampMin = "1", PCG_Overridable))
 	double MaximumDistance = 20000.0;
 
 	// What shape is used on the 'source' points
-	UPROPERTY(BlueprintReadWrite, EditAnywhere, Category = Settings)
+	UPROPERTY(BlueprintReadWrite, EditAnywhere, Category = Settings, meta = (PCG_Overridable))
 	PCGDistanceShape SourceShape = PCGDistanceShape::SphereBounds;
 
 	// What shape is used on the 'target' points
-	UPROPERTY(BlueprintReadWrite, EditAnywhere, Category = Settings)
+	UPROPERTY(BlueprintReadWrite, EditAnywhere, Category = Settings, meta = (PCG_Overridable))
 	PCGDistanceShape TargetShape = PCGDistanceShape::SphereBounds;
 };
 
diff --git a/Engine/Plugins/Experimental/PCG/Source/PCG/Public/Elements/PCGExecuteBlueprint.h b/Engine/Plugins/Experimental/PCG/Source/PCG/Public/Elements/PCGExecuteBlueprint.h
--- a/Engine/Plugins/Experimental/PCG/Source/PCG/Public/Elements/PCGExecuteBlueprint.h
+++ b/Engine/Plugins/Experimental/PCG/Source/PCG/Public/Elements/PCGExecuteBlueprint.h
@@ -189,6 +189,8 @@ public:
 	virtual TArray<FPCGPinProperties> InputPinProperties() const override;
 	virtual TArray<FPCGPinProperties> OutputPinProperties() const override;
 
+	/** To be removed when we support automatic override of BP params. For now always return true to force params pin. */
+	virtual bool HasOverridableParams() const override { return true; }
 protected:
 	virtual FPCGElementPtr CreateElement() const override;
 	// ~End UPCGSettings interface
diff --git a/Engine/Plugins/Experimental/PCG/Source/PCG/Public/Elements/PCGIntersectionElement.h b/Engine/Plugins/Experimental/PCG/Source/PCG/Public/Elements/PCGIntersectionElement.h
--- a/Engine/Plugins/Experimental/PCG/Source/PCG/Public/Elements/PCGIntersectionElement.h
+++ b/Engine/Plugins/Experimental/PCG/Source/PCG/Public/Elements/PCGIntersectionElement.h
@@ -19,6 +19,7 @@ public:
 	virtual EPCGSettingsType GetType() const override { return EPCGSettingsType::Spatial; }
 #endif
 
+	virtual TArray<FPCGPinProperties> InputPinProperties() const override;
 	virtual TArray<FPCGPinProperties> OutputPinProperties() const override;
 
 	// If node disabled, don't intersect - pass through first edge
@@ -29,11 +30,11 @@ protected:
 	//~End UPCGSettings interface
 
 public:
-	UPROPERTY(BlueprintReadWrite, EditAnywhere, Category = Settings)
+	UPROPERTY(BlueprintReadWrite, EditAnywhere, Category = Settings, meta = (PCG_Overridable))
 	EPCGIntersectionDensityFunction DensityFunction = EPCGIntersectionDensityFunction::Multiply;
 
 #if WITH_EDITORONLY_DATA
-	UPROPERTY(Transient, BlueprintReadWrite, EditAnywhere, Category = "Settings|Debug")
+	UPROPERTY(Transient, BlueprintReadWrite, EditAnywhere, Category = "Settings|Debug", meta = (PCG_Overridable))
 	bool bKeepZeroDensityPoints = false;
 #endif
 };
diff --git a/Engine/Plugins/Experimental/PCG/Source/PCG/Public/Elements/PCGMergeElement.h b/Engine/Plugins/Experimental/PCG/Source/PCG/Public/Elements/PCGMergeElement.h
--- a/Engine/Plugins/Experimental/PCG/Source/PCG/Public/Elements/PCGMergeElement.h
+++ b/Engine/Plugins/Experimental/PCG/Source/PCG/Public/Elements/PCGMergeElement.h
@@ -32,7 +32,7 @@ protected:
 
 public:
 	/** Controls whether the resulting merge data will have any metadata */
-	UPROPERTY(BlueprintReadWrite, EditAnywhere, Category = Settings)
+	UPROPERTY(BlueprintReadWrite, EditAnywhere, Category = Settings, meta = (PCG_Overridable))
 	bool bMergeMetadata = true;
 };
 
diff --git a/Engine/Plugins/Experimental/PCG/Source/PCG/Public/Elements/PCGPointExtentsModifier.h b/Engine/Plugins/Experimental/PCG/Source/PCG/Public/Elements/PCGPointExtentsModifier.h
--- a/Engine/Plugins/Experimental/PCG/Source/PCG/Public/Elements/PCGPointExtentsModifier.h
+++ b/Engine/Plugins/Experimental/PCG/Source/PCG/Public/Elements/PCGPointExtentsModifier.h
@@ -40,10 +40,10 @@ protected:
 	//~End UPCGSettings interface
 
 public:
-	UPROPERTY(BlueprintReadWrite, EditAnywhere, Category = Settings)
+	UPROPERTY(BlueprintReadWrite, EditAnywhere, Category = Settings, meta = (PCG_Overridable))
 	FVector Extents = FVector::One();
 
-	UPROPERTY(BlueprintReadWrite, EditAnywhere, Category = Settings)
+	UPROPERTY(BlueprintReadWrite, EditAnywhere, Category = Settings, meta = (PCG_Overridable))
 	EPCGPointExtentsModifierMode Mode = EPCGPointExtentsModifierMode::Set;
 };
 
diff --git a/Engine/Plugins/Experimental/PCG/Source/PCG/Public/Elements/PCGPointFilter.h b/Engine/Plugins/Experimental/PCG/Source/PCG/Public/Elements/PCGPointFilter.h
--- a/Engine/Plugins/Experimental/PCG/Source/PCG/Public/Elements/PCGPointFilter.h
+++ b/Engine/Plugins/Experimental/PCG/Source/PCG/Public/Elements/PCGPointFilter.h
@@ -74,7 +74,6 @@ public:
 #if WITH_EDITOR
 	virtual FName GetDefaultNodeName() const override { return FName(TEXT("PointFilter")); }
 	virtual EPCGSettingsType GetType() const override { return EPCGSettingsType::Filter; }
-	virtual void ApplyDeprecationBeforeUpdatePins(UPCGNode* InOutNode, TArray<TObjectPtr<UPCGPin>>& InputPins, TArray<TObjectPtr<UPCGPin>>& OutputPins) override;
 #endif
 
 	virtual TArray<FPCGPinProperties> InputPinProperties() const override;
@@ -85,7 +84,7 @@ protected:
 	//~End UPCGSettings interface
 
 public:
-	UPROPERTY(BlueprintReadWrite, EditAnywhere, Category = Settings)
+	UPROPERTY(BlueprintReadWrite, EditAnywhere, Category = Settings, meta = (PCG_Overridable))
 	EPCGPointFilterOperator Operator = EPCGPointFilterOperator::Greater;
 
 	/** Target property/attribute related properties */
@@ -93,13 +92,13 @@ public:
 	FPCGAttributePropertySelector TargetAttribute;
 
 	/** Threshold property/attribute/constant related properties */
-	UPROPERTY(BlueprintReadWrite, EditAnywhere, Category = Settings)
+	UPROPERTY(BlueprintReadWrite, EditAnywhere, Category = Settings, meta = (PCG_Overridable))
 	bool bUseConstantThreshold = false;
 
 	UPROPERTY(BlueprintReadWrite, EditAnywhere, Category = Settings, meta = (EditCondition = "!bUseConstantThreshold", EditConditionHides))
 	FPCGAttributePropertySelector ThresholdAttribute;
 
-	UPROPERTY(BlueprintReadWrite, EditAnywhere, Category = Settings, meta = (EditCondition = "!bUseConstantThreshold", EditConditionHides))
+	UPROPERTY(BlueprintReadWrite, EditAnywhere, Category = Settings, meta = (EditCondition = "!bUseConstantThreshold", EditConditionHides, PCG_Overridable))
 	bool bUseSpatialQuery = true;
 
 	UPROPERTY(BlueprintReadWrite, EditAnywhere, Category = Constant, meta = (EditCondition = "bUseConstantThreshold", EditConditionHides, ShowOnlyInnerProperties, DisplayAfter = "bUseConstantThreshold"))
diff --git a/Engine/Plugins/Experimental/PCG/Source/PCG/Public/Elements/PCGPointSampler.h b/Engine/Plugins/Experimental/PCG/Source/PCG/Public/Elements/PCGPointSampler.h
--- a/Engine/Plugins/Experimental/PCG/Source/PCG/Public/Elements/PCGPointSampler.h
+++ b/Engine/Plugins/Experimental/PCG/Source/PCG/Public/Elements/PCGPointSampler.h
@@ -28,11 +28,11 @@ protected:
 	//~End UPCGSettings interface
 
 public:
-	UPROPERTY(BlueprintReadWrite, EditAnywhere, Category = Settings, meta=(ClampMin="0", ClampMax="1"))
+	UPROPERTY(BlueprintReadWrite, EditAnywhere, Category = Settings, meta=(ClampMin="0", ClampMax="1", PCG_Overridable))
 	float Ratio = 0.1f;
 
 #if WITH_EDITORONLY_DATA
-	UPROPERTY(Transient, BlueprintReadWrite, EditAnywhere, Category = "Settings|Debug")
+	UPROPERTY(Transient, BlueprintReadWrite, EditAnywhere, Category = "Settings|Debug", meta = (PCG_Overridable))
 	bool bKeepZeroDensityPoints = false;
 #endif
 };
diff --git a/Engine/Plugins/Experimental/PCG/Source/PCG/Public/Elements/PCGProjectionElement.h b/Engine/Plugins/Experimental/PCG/Source/PCG/Public/Elements/PCGProjectionElement.h
--- a/Engine/Plugins/Experimental/PCG/Source/PCG/Public/Elements/PCGProjectionElement.h
+++ b/Engine/Plugins/Experimental/PCG/Source/PCG/Public/Elements/PCGProjectionElement.h
@@ -28,11 +28,11 @@ public:
 
 	virtual TArray<FPCGPinProperties> InputPinProperties() const override;
 
-	UPROPERTY(BlueprintReadWrite, EditAnywhere, Category = Settings, meta = (ShowOnlyInnerProperties))
-	FPCGProjectionParams Params;
+	UPROPERTY(BlueprintReadWrite, EditAnywhere, Category = Settings, meta = (ShowOnlyInnerProperties, PCG_Overridable))
+	FPCGProjectionParams ProjectionParams;
 
 #if WITH_EDITORONLY_DATA
-	UPROPERTY(Transient, BlueprintReadWrite, EditAnywhere, Category = "Settings|Debug")
+	UPROPERTY(Transient, BlueprintReadWrite, EditAnywhere, Category = "Settings|Debug", meta = (PCG_Overridable))
 	bool bKeepZeroDensityPoints = false;
 #endif
 
diff --git a/Engine/Plugins/Experimental/PCG/Source/PCG/Public/Elements/PCGSelfPruning.h b/Engine/Plugins/Experimental/PCG/Source/PCG/Public/Elements/PCGSelfPruning.h
--- a/Engine/Plugins/Experimental/PCG/Source/PCG/Public/Elements/PCGSelfPruning.h
+++ b/Engine/Plugins/Experimental/PCG/Source/PCG/Public/Elements/PCGSelfPruning.h
@@ -28,20 +28,19 @@ class PCG_API UPCGSelfPruningSettings : public UPCGSettings
 	GENERATED_BODY()
 
 public:
-	UPROPERTY(BlueprintReadWrite, EditAnywhere, Category = Settings)
+	UPROPERTY(BlueprintReadWrite, EditAnywhere, Category = Settings, meta = (PCG_Overridable))
 	EPCGSelfPruningType PruningType = EPCGSelfPruningType::LargeToSmall;
 
-	UPROPERTY(BlueprintReadWrite, EditAnywhere, Category = Settings, meta=(ClampMin=0.0f, EditCondition="PruningType == EPCGSelfPruningType::LargeToSmall || PruningType == EPCGSelfPruningType::SmallToLarge"))
+	UPROPERTY(BlueprintReadWrite, EditAnywhere, Category = Settings, meta=(ClampMin=0.0f, EditCondition="PruningType == EPCGSelfPruningType::LargeToSmall || PruningType == EPCGSelfPruningType::SmallToLarge", PCG_Overridable))
 	float RadiusSimilarityFactor = 0.25f;
 
-	UPROPERTY(BlueprintReadWrite, EditAnywhere, Category = Settings)
+	UPROPERTY(BlueprintReadWrite, EditAnywhere, Category = Settings, meta = (PCG_Overridable))
 	bool bRandomizedPruning = true;
 
 #if WITH_EDITOR
 	//~Begin UPCGSettings interface
 	virtual FName GetDefaultNodeName() const override { return FName(TEXT("SelfPruning")); }
 	virtual EPCGSettingsType GetType() const override { return EPCGSettingsType::Filter; }
-	virtual void ApplyDeprecationBeforeUpdatePins(UPCGNode* InOutNode, TArray<TObjectPtr<UPCGPin>>& InputPins, TArray<TObjectPtr<UPCGPin>>& OutputPins) override;
 #endif
 
 	virtual TArray<FPCGPinProperties> InputPinProperties() const override;
diff --git a/Engine/Plugins/Experimental/PCG/Source/PCG/Public/Elements/PCGSplineSampler.h b/Engine/Plugins/Experimental/PCG/Source/PCG/Public/Elements/PCGSplineSampler.h
--- a/Engine/Plugins/Experimental/PCG/Source/PCG/Public/Elements/PCGSplineSampler.h
+++ b/Engine/Plugins/Experimental/PCG/Source/PCG/Public/Elements/PCGSplineSampler.h
@@ -128,8 +128,8 @@ protected:
 	// ~End UPCGSettings interface
 
 public:
-	UPROPERTY(BlueprintReadWrite, EditAnywhere, Category = Settings, meta = (ShowOnlyInnerProperties))
-	FPCGSplineSamplerParams Params;
+	UPROPERTY(BlueprintReadWrite, EditAnywhere, Category = Settings, meta = (ShowOnlyInnerProperties, PCG_Overridable))
+	FPCGSplineSamplerParams SamplerParams;
 };
 
 class FPCGSplineSamplerElement : public FSimplePCGElement
diff --git a/Engine/Plugins/Experimental/PCG/Source/PCG/Public/Elements/PCGSurfaceSampler.h b/Engine/Plugins/Experimental/PCG/Source/PCG/Public/Elements/PCGSurfaceSampler.h
--- a/Engine/Plugins/Experimental/PCG/Source/PCG/Public/Elements/PCGSurfaceSampler.h
+++ b/Engine/Plugins/Experimental/PCG/Source/PCG/Public/Elements/PCGSurfaceSampler.h
@@ -45,34 +45,34 @@ protected:
 	// ~End UPCGSettings interface
 
 public:
-	UPROPERTY(BlueprintReadWrite, EditAnywhere, Category = Settings, meta=(ClampMin="0"))
+	UPROPERTY(BlueprintReadWrite, EditAnywhere, Category = Settings, meta=(ClampMin="0", PCG_Overridable))
 	float PointsPerSquaredMeter = 0.1f;
 
-	UPROPERTY()
-	float PointRadius_DEPRECATED = 0.0f;
-
-	UPROPERTY(BlueprintReadWrite, EditAnywhere, Category = Settings)
+	UPROPERTY(BlueprintReadWrite, EditAnywhere, Category = Settings, meta = (PCG_Overridable))
 	FVector PointExtents = FVector(100.0f);
 
-	UPROPERTY(BlueprintReadWrite, EditAnywhere, Category = Settings, meta=(ClampMin="0"))
+	UPROPERTY(BlueprintReadWrite, EditAnywhere, Category = Settings, meta=(ClampMin="0", PCG_Overridable))
 	float Looseness = 1.0f;
 
 	/** If no Bounding Shape input is provided the actor bounds are used to limit the sample generation area.
 	* This option allows ignoring the actor bounds and generating over the entire surface. Use with caution as this
 	* may generate a lot of points.
 	*/
-	UPROPERTY(BlueprintReadWrite, EditAnywhere, Category = Settings)
+	UPROPERTY(BlueprintReadWrite, EditAnywhere, Category = Settings, meta = (PCG_Overridable))
 	bool bUnbounded = false;
 
-	UPROPERTY(BlueprintReadWrite, EditAnywhere, Category = "Settings|Points")
+	UPROPERTY(BlueprintReadWrite, EditAnywhere, Category = "Settings|Points", meta = (PCG_Overridable))
 	bool bApplyDensityToPoints = true;
 
-	UPROPERTY(BlueprintReadWrite, EditAnywhere, Category = "Settings|Points", meta=(ClampMin="0", ClampMax="1"))
+	UPROPERTY(BlueprintReadWrite, EditAnywhere, Category = "Settings|Points", meta=(ClampMin="0", ClampMax="1", PCG_Overridable))
 	float PointSteepness = 0.5f;
 
 #if WITH_EDITORONLY_DATA
-	UPROPERTY(Transient, BlueprintReadWrite, EditAnywhere, Category = "Settings|Debug")
+	UPROPERTY(Transient, BlueprintReadWrite, EditAnywhere, Category = "Settings|Debug", meta = (PCG_Overridable))
 	bool bKeepZeroDensityPoints = false;
+
+	UPROPERTY()
+	float PointRadius_DEPRECATED = 0.0f;
 #endif
 };
 
diff --git a/Engine/Plugins/Experimental/PCG/Source/PCG/Public/Elements/PCGTextureSampler.h b/Engine/Plugins/Experimental/PCG/Source/PCG/Public/Elements/PCGTextureSampler.h
--- a/Engine/Plugins/Experimental/PCG/Source/PCG/Public/Elements/PCGTextureSampler.h
+++ b/Engine/Plugins/Experimental/PCG/Source/PCG/Public/Elements/PCGTextureSampler.h
@@ -29,10 +29,10 @@ protected:
 
 public:
 	// Surface transform
-	UPROPERTY(BlueprintReadWrite, EditAnywhere, Category = Settings)
+	UPROPERTY(BlueprintReadWrite, EditAnywhere, Category = Settings, meta = (PCG_Overridable))
 	FTransform Transform = FTransform::Identity;
 
-	UPROPERTY(BlueprintReadWrite, EditAnywhere, Category = Settings)
+	UPROPERTY(BlueprintReadWrite, EditAnywhere, Category = Settings, meta = (PCG_Overridable))
 	bool bUseAbsoluteTransform = false;
 
 	// Texture specific parameters
@@ -40,37 +40,37 @@ public:
 	TObjectPtr<UTexture2D> Texture = nullptr;
 
 	// Common members in BaseTextureData
-	UPROPERTY(BlueprintReadWrite, VisibleAnywhere, Category = SpatialData)
+	UPROPERTY(BlueprintReadWrite, VisibleAnywhere, Category = SpatialData, meta = (PCG_Overridable))
 	EPCGTextureDensityFunction DensityFunction = EPCGTextureDensityFunction::Multiply;
 
-	UPROPERTY(BlueprintReadWrite, EditAnywhere, Category = Settings)
+	UPROPERTY(BlueprintReadWrite, EditAnywhere, Category = Settings, meta = (PCG_Overridable))
 	EPCGTextureColorChannel ColorChannel = EPCGTextureColorChannel::Alpha;
 
 	/** The size of one texel in cm, used when calling ToPointData. */
-	UPROPERTY(BlueprintReadWrite, EditAnywhere, Category = Settings, meta = (UIMin = "1.0", ClampMin = "1.0"))
+	UPROPERTY(BlueprintReadWrite, EditAnywhere, Category = Settings, meta = (UIMin = "1.0", ClampMin = "1.0", PCG_Overridable))
 	float TexelSize = 50.0f;
 
 	/** Whether to tile the source or to stretch it to fit target area. */
-	UPROPERTY(BlueprintReadWrite, EditAnywhere, Category = Settings)
+	UPROPERTY(BlueprintReadWrite, EditAnywhere, Category = Settings, meta = (PCG_Overridable))
 	bool bUseAdvancedTiling = false;
 
-	UPROPERTY(BlueprintReadWrite, EditAnywhere, Category = "Settings|Tiling", meta = (EditCondition = "bUseAdvancedTiling"))
+	UPROPERTY(BlueprintReadWrite, EditAnywhere, Category = "Settings|Tiling", meta = (EditCondition = "bUseAdvancedTiling", PCG_Overridable))
 	FVector2D Tiling = FVector2D(1.0, 1.0);
 
-	UPROPERTY(BlueprintReadWrite, EditAnywhere, Category = "Settings|Tiling", meta = (EditCondition = "bUseAdvancedTiling"))
+	UPROPERTY(BlueprintReadWrite, EditAnywhere, Category = "Settings|Tiling", meta = (EditCondition = "bUseAdvancedTiling", PCG_Overridable))
 	FVector2D CenterOffset = FVector2D::ZeroVector;
 
 	/** Rotation to apply when sampling texture. */
-	UPROPERTY(BlueprintReadWrite, EditAnywhere, Category = "Settings|Tiling", meta = (UIMin = -360, ClampMin = -360, UIMax = 360, ClampMax = 360, Units = deg, EditCondition = "bUseAdvancedTiling"))
+	UPROPERTY(BlueprintReadWrite, EditAnywhere, Category = "Settings|Tiling", meta = (UIMin = -360, ClampMin = -360, UIMax = 360, ClampMax = 360, Units = deg, EditCondition = "bUseAdvancedTiling", PCG_Overridable))
 	float Rotation = 0;
 
-	UPROPERTY(BlueprintReadWrite, EditAnywhere, Category = "Settings|Tiling", meta = (EditionCondition = "bUseAdvancedTiling"))
+	UPROPERTY(BlueprintReadWrite, EditAnywhere, Category = "Settings|Tiling", meta = (EditionCondition = "bUseAdvancedTiling", PCG_Overridable))
 	bool bUseTileBounds = false;
 
-	UPROPERTY(BlueprintReadWrite, EditAnywhere, Category = "Settings|Tiling", meta = (EditCondition = "bUseAdvancedTiling && bUseTileBounds"))
+	UPROPERTY(BlueprintReadWrite, EditAnywhere, Category = "Settings|Tiling", meta = (EditCondition = "bUseAdvancedTiling && bUseTileBounds", PCG_Overridable))
 	FVector2D TileBoundsMin = FVector2D(-0.5, -0.5);
 
-	UPROPERTY(BlueprintReadWrite, EditAnywhere, Category = "Settings|Tiling", meta = (EditCondition = "bUseAdvancedTiling && bUseTileBounds"))
+	UPROPERTY(BlueprintReadWrite, EditAnywhere, Category = "Settings|Tiling", meta = (EditCondition = "bUseAdvancedTiling && bUseTileBounds", PCG_Overridable))
 	FVector2D TileBoundsMax = FVector2D(0.5, 0.5);
 };
 
diff --git a/Engine/Plugins/Experimental/PCG/Source/PCG/Public/Elements/PCGTransformPoints.h b/Engine/Plugins/Experimental/PCG/Source/PCG/Public/Elements/PCGTransformPoints.h
--- a/Engine/Plugins/Experimental/PCG/Source/PCG/Public/Elements/PCGTransformPoints.h
+++ b/Engine/Plugins/Experimental/PCG/Source/PCG/Public/Elements/PCGTransformPoints.h
@@ -29,48 +29,48 @@ protected:
 	//~End UPCGSettings interface
 
 public:
-	UPROPERTY(BlueprintReadWrite, EditAnywhere, Category = Settings)
+	UPROPERTY(BlueprintReadWrite, EditAnywhere, Category = Settings, meta = (PCG_Overridable))
 	bool bApplyToAttribute = false;
 
-	UPROPERTY(BlueprintReadWrite, EditAnywhere, Category = Settings, meta = (EditCondition="bApplyToAttribute"))
+	UPROPERTY(BlueprintReadWrite, EditAnywhere, Category = Settings, meta = (EditCondition="bApplyToAttribute", PCG_Overridable))
 	FName AttributeName;
 
-	UPROPERTY(BlueprintReadWrite, EditAnywhere, Category = Settings)
+	UPROPERTY(BlueprintReadWrite, EditAnywhere, Category = Settings, meta = (PCG_Overridable))
 	FVector OffsetMin = FVector::Zero();
 
-	UPROPERTY(BlueprintReadWrite, EditAnywhere, Category = Settings)
+	UPROPERTY(BlueprintReadWrite, EditAnywhere, Category = Settings, meta = (PCG_Overridable))
 	FVector OffsetMax = FVector::Zero();
 
 	/** Set offset in world space */
-	UPROPERTY(BlueprintReadWrite, EditAnywhere, Category = Settings)
+	UPROPERTY(BlueprintReadWrite, EditAnywhere, Category = Settings, meta = (PCG_Overridable))
 	bool bAbsoluteOffset = false;
 
-	UPROPERTY(BlueprintReadWrite, EditAnywhere, Category = Settings)
+	UPROPERTY(BlueprintReadWrite, EditAnywhere, Category = Settings, meta = (PCG_Overridable))
 	FRotator RotationMin = FRotator::ZeroRotator;
 
-	UPROPERTY(BlueprintReadWrite, EditAnywhere, Category = Settings)
+	UPROPERTY(BlueprintReadWrite, EditAnywhere, Category = Settings, meta = (PCG_Overridable))
 	FRotator RotationMax = FRotator::ZeroRotator;
 
 	/** Set rotation directly instead of additively */
-	UPROPERTY(BlueprintReadWrite, EditAnywhere, Category = Settings)
+	UPROPERTY(BlueprintReadWrite, EditAnywhere, Category = Settings, meta = (PCG_Overridable))
 	bool bAbsoluteRotation = false;
 
-	UPROPERTY(BlueprintReadWrite, EditAnywhere, Category = Settings, meta = (AllowPreserveRatio))
+	UPROPERTY(BlueprintReadWrite, EditAnywhere, Category = Settings, meta = (AllowPreserveRatio, PCG_Overridable))
 	FVector ScaleMin = FVector::One();
 
-	UPROPERTY(BlueprintReadWrite, EditAnywhere, Category = Settings, meta = (AllowPreserveRatio))
+	UPROPERTY(BlueprintReadWrite, EditAnywhere, Category = Settings, meta = (AllowPreserveRatio, PCG_Overridable))
 	FVector ScaleMax = FVector::One();
 
 	/** Set scale directly instead of multiplicatively */
-	UPROPERTY(BlueprintReadWrite, EditAnywhere, Category = Settings)
+	UPROPERTY(BlueprintReadWrite, EditAnywhere, Category = Settings, meta = (PCG_Overridable))
 	bool bAbsoluteScale = false;
 
 	/** Scale uniformly on each axis. Uses the X component of ScaleMin and ScaleMax. */
-	UPROPERTY(BlueprintReadWrite, EditAnywhere, Category = Settings)
+	UPROPERTY(BlueprintReadWrite, EditAnywhere, Category = Settings, meta = (PCG_Overridable))
 	bool bUniformScale = true;
 
 	/** Recompute the seed for each new point using its new location */
-	UPROPERTY(BlueprintReadWrite, EditAnywhere, Category = Settings, meta = (EditCondition = "!bApplyToAttribute"))
+	UPROPERTY(BlueprintReadWrite, EditAnywhere, Category = Settings, meta = (EditCondition = "!bApplyToAttribute", PCG_Overridable))
 	bool bRecomputeSeed = false;
 };
 
diff --git a/Engine/Plugins/Experimental/PCG/Source/PCG/Public/Elements/PCGUnionElement.h b/Engine/Plugins/Experimental/PCG/Source/PCG/Public/Elements/PCGUnionElement.h
--- a/Engine/Plugins/Experimental/PCG/Source/PCG/Public/Elements/PCGUnionElement.h
+++ b/Engine/Plugins/Experimental/PCG/Source/PCG/Public/Elements/PCGUnionElement.h
@@ -19,6 +19,7 @@ public:
 	virtual EPCGSettingsType GetType() const override { return EPCGSettingsType::Spatial; }
 #endif
 
+	virtual TArray<FPCGPinProperties> InputPinProperties() const override;
 	virtual TArray<FPCGPinProperties> OutputPinProperties() const override;
 
 	// If node disabled, don't union - pass through first edge
@@ -29,10 +30,10 @@ protected:
 	//~End UPCGSettings interface
 
 public:
-	UPROPERTY(BlueprintReadWrite, EditAnywhere, Category = Settings)
+	UPROPERTY(BlueprintReadWrite, EditAnywhere, Category = Settings, meta = (PCG_Overridable))
 	EPCGUnionType Type = EPCGUnionType::LeftToRightPriority;
 
-	UPROPERTY(BlueprintReadWrite, EditAnywhere, Category = Settings)
+	UPROPERTY(BlueprintReadWrite, EditAnywhere, Category = Settings, meta = (PCG_Overridable))
 	EPCGUnionDensityFunction DensityFunction = EPCGUnionDensityFunction::Maximum;
 };
 
diff --git a/Engine/Plugins/Experimental/PCG/Source/PCG/Public/Elements/PCGVolumeSampler.h b/Engine/Plugins/Experimental/PCG/Source/PCG/Public/Elements/PCGVolumeSampler.h
--- a/Engine/Plugins/Experimental/PCG/Source/PCG/Public/Elements/PCGVolumeSampler.h
+++ b/Engine/Plugins/Experimental/PCG/Source/PCG/Public/Elements/PCGVolumeSampler.h
@@ -28,7 +28,7 @@ class PCG_API UPCGVolumeSamplerSettings : public UPCGSettings
 	GENERATED_BODY()
 
 public:
-	UPROPERTY(BlueprintReadWrite, EditAnywhere, Category = "Data")
+	UPROPERTY(BlueprintReadWrite, EditAnywhere, Category = "Data", meta = (PCG_Overridable))
 	FVector VoxelSize = FVector(100.0, 100.0, 100.0);
 
 	//~Begin UPCGSettings interface
@@ -37,6 +37,7 @@ public:
 	virtual EPCGSettingsType GetType() const override { return EPCGSettingsType::Sampler; }
 #endif
 
+	virtual TArray<FPCGPinProperties> InputPinProperties() const override;
 	virtual TArray<FPCGPinProperties> OutputPinProperties() const override { return Super::DefaultPointOutputPinProperties(); }
 	
 protected:
diff --git a/Engine/Plugins/Experimental/PCG/Source/PCG/Public/Helpers/PCGSettingsHelpers.h b/Engine/Plugins/Experimental/PCG/Source/PCG/Public/Helpers/PCGSettingsHelpers.h
--- a/Engine/Plugins/Experimental/PCG/Source/PCG/Public/Helpers/PCGSettingsHelpers.h
+++ b/Engine/Plugins/Experimental/PCG/Source/PCG/Public/Helpers/PCGSettingsHelpers.h
@@ -2,13 +2,19 @@
 
 #pragma once
 
+#include "PCGModule.h"
 #include "PCGParamData.h"
-
+#include "PCGSettings.h"
+#include "Metadata/PCGMetadataAttributeTpl.h"
+#include "Metadata/PCGMetadataAttributeTraits.h"
+#include "Metadata/Accessors/PCGAttributeAccessorHelpers.h"
+#include "Metadata/Accessors/IPCGAttributeAccessor.h"
+#include "Metadata/Accessors/PCGAttributeAccessorKeys.h"
 
 class UPCGComponent;
 class UPCGNode;
 class UPCGPin;
-class UPCGSettings;
+struct FPCGDataCollection;
 
 namespace PCGSettingsHelpers
 {
@@ -19,6 +25,7 @@ namespace PCGSettingsHelpers
 	* @param InKey - Metadata Entry Key to get the value from.
 	*/
 	template<typename T, typename TEnableIf<!TIsEnumClass<T>::Value>::Type* = nullptr>
+	UE_DEPRECATED(5.2, "GetValue is deprecated as overrides are now automatically applied on the settings. You should just query the settings value.")
 	T GetValue(const FName& InName, const T& InValue, const UPCGParamData* InParams, PCGMetadataEntryKey InKey)
 	{
 
@@ -61,13 +68,17 @@ namespace PCGSettingsHelpers
 		}
 	}
 
+PRAGMA_DISABLE_DEPRECATION_WARNINGS
+
 	template<typename T, typename TEnableIf<!TIsEnumClass<T>::Value>::Type* = nullptr>
+	UE_DEPRECATED(5.2, "GetValue is deprecated as overrides are now automatically applied on the settings. You should just query the settings value.")
 	T GetValue(const FName& InName, const T& InValue, const UPCGParamData* InParams)
 	{
 		return GetValue(InName, InValue, InParams, 0);
 	}
 
 	template<typename T, typename TEnableIf<!TIsEnumClass<T>::Value>::Type* = nullptr>
+	UE_DEPRECATED(5.2, "GetValue is deprecated as overrides are now automatically applied on the settings. You should just query the settings value.")
 	T GetValue(const FName& InName, const T& InValue, const UPCGParamData* InParams, const FName& InParamName)
 	{
 		if (InParams && InParamName != NAME_None)
@@ -82,18 +93,21 @@ namespace PCGSettingsHelpers
 
 	/** Specialized versions for enums */
 	template<typename T, typename TEnableIf<TIsEnumClass<T>::Value>::Type* = nullptr>
+	UE_DEPRECATED(5.2, "GetValue is deprecated as overrides are now automatically applied on the settings. You should just query the settings value.")
 	T GetValue(const FName& InName, const T& InValue, const UPCGParamData* InParams, PCGMetadataEntryKey InKey)
 	{
 		return static_cast<T>(GetValue(InName, static_cast<__underlying_type(T)>(InValue), InParams, InKey));
 	}
 
 	template<typename T, typename TEnableIf<TIsEnumClass<T>::Value>::Type* = nullptr>
+	UE_DEPRECATED(5.2, "GetValue is deprecated as overrides are now automatically applied on the settings. You should just query the settings value.")
 	T GetValue(const FName& InName, const T& InValue, const UPCGParamData* InParams)
 	{
 		return static_cast<T>(GetValue(InName, static_cast<__underlying_type(T)>(InValue), InParams));
 	}
 
 	template<typename T, typename TEnableIf<TIsEnumClass<T>::Value>::Type* = nullptr>
+	UE_DEPRECATED(5.2, "GetValue is deprecated as overrides are now automatically applied on the settings. You should just query the settings value.")
 	T GetValue(const FName& InName, const T& InValue, const UPCGParamData* InParams, const FName& InParamName)
 	{
 		return static_cast<T>(GetValue(InName, static_cast<__underlying_type(T)>(InValue), InParams, InParamName));
@@ -101,18 +115,21 @@ namespace PCGSettingsHelpers
 
 	/** Specialized version for names */
 	template<>
+	UE_DEPRECATED(5.2, "GetValue is deprecated as overrides are now automatically applied on the settings. You should just query the settings value.")
 	FORCEINLINE FName GetValue(const FName& InName, const FName& InValue, const UPCGParamData* InParams, PCGMetadataEntryKey InKey)
 	{
 		return FName(GetValue(InName, InValue.ToString(), InParams, InKey));
 	}
 
 	template<>
+	UE_DEPRECATED(5.2, "GetValue is deprecated as overrides are now automatically applied on the settings. You should just query the settings value.")
 	FORCEINLINE FName GetValue(const FName& InName, const FName& InValue, const UPCGParamData* InParams)
 	{
 		return FName(GetValue(InName, InValue.ToString(), InParams));
 	}
 
 	template<>
+	UE_DEPRECATED(5.2, "GetValue is deprecated as overrides are now automatically applied on the settings. You should just query the settings value.")
 	FORCEINLINE FName GetValue(const FName& InName, const FName& InValue, const UPCGParamData* InParams, const FName& InParamName)
 	{
 		return FName(GetValue(InName, InValue.ToString(), InParams, InParamName));
@@ -121,17 +138,70 @@ namespace PCGSettingsHelpers
 	/** Sets data from the params to a given property, matched on a name basis */
 	void SetValue(UPCGParamData* Params, UObject* Object, FProperty* Property);
 
+	UE_DEPRECATED(5.2, "ComputeSeedWithOverride is deprecated as overrides are now automatically applied on the settings. You should just query the settings value.")
 	int ComputeSeedWithOverride(const UPCGSettings* InSettings, const UPCGComponent* InComponent, UPCGParamData* InParams);
 
+	UE_DEPRECATED(5.2, "ComputeSeedWithOverride is deprecated as overrides are now automatically applied on the settings. You should just query the settings value.")
 	FORCEINLINE int ComputeSeedWithOverride(const UPCGSettings* InSettings, TWeakObjectPtr<UPCGComponent> InComponent, UPCGParamData* InParams)
 	{
 		return ComputeSeedWithOverride(InSettings, InComponent.Get(), InParams);
 	}
 
+PRAGMA_ENABLE_DEPRECATION_WARNINGS
+
 	/** Utility to call from before-node-update deprecation. A dedicated pin for params will be added when the pins are updated. Here we detect any params
 	*   connections to the In pin and disconnect them, and move the first params connection to a new params pin.
 	*/
 	void DeprecationBreakOutParamsToNewPin(UPCGNode* InOutNode, TArray<TObjectPtr<UPCGPin>>& InputPins, TArray<TObjectPtr<UPCGPin>>& OutputPins);
+
+	/**
+	* Advanced method to gather override params when you don't have access to FPCGContext (and therefore don't have access to automatic
+	* param override).
+	* Limitation: Only support metadata types for T.
+	*/
+	template <typename T>
+	bool GetOverrideValue(const FPCGDataCollection& InInputData, const UPCGSettings* InSettings, const FName InPropertyName, const T& InDefaultValue, T& OutValue)
+	{
+		check(InSettings);
+
+		// Limitation: Only support metadata types
+		static_assert(PCG::Private::IsPCGType<T>());
+
+		// Try to find the override param associated with the property.
+		const FPCGSettingsOverridableParam* Param = InSettings->OverridableParams().FindByPredicate([InPropertyName](const FPCGSettingsOverridableParam& InParam) { return !InParam.PropertiesNames.IsEmpty() && (InParam.PropertiesNames.Last() == InPropertyName);});
+
+		OutValue = InDefaultValue;
+
+		if (!Param)
+		{
+			return false;
+		}
+
+		FName AttributeName = NAME_None;
+		TUniquePtr<const IPCGAttributeAccessor> AttributeAccessor = PCGAttributeAccessorHelpers::CreateConstAccessorForOverrideParam(InInputData, *Param, &AttributeName);
+
+		if (!AttributeAccessor)
+		{
+			return false;
+		}
+
+		return PCGMetadataAttribute::CallbackWithRightType(PCG::Private::MetadataTypes<T>::Id, [&AttributeAccessor, &Param, &AttributeName, &InDefaultValue, &OutValue](auto Dummy) -> bool
+		{
+			using PropertyType = decltype(Dummy);
+
+			// Override were using the first entry (0) by default.
+			FPCGAttributeAccessorKeysEntries FirstEntry(PCGMetadataEntryKey(0));
+
+			if (!AttributeAccessor->Get<T>(OutValue, FirstEntry, EPCGAttributeAccessorFlags::AllowBroadcast | EPCGAttributeAccessorFlags::AllowConstructible))
+			{
+				UE_LOG(LogPCG, Warning, TEXT("[PCGSettingsHelpers::GetOverrideValue] %s param can't be converted from %s attribute, incompatible types."), *Param->Label.ToString(), *AttributeName.ToString());
+				return false;
+			}
+
+			return true;
+		});
+	}
 }
 
-#define PCG_GET_OVERRIDEN_VALUE(Settings, Variable, Params) PCGSettingsHelpers::GetValue(GET_MEMBER_NAME_CHECKED(TRemovePointer<TRemoveConst<decltype(Settings)>::Type>::Type, Variable), (Settings)->Variable, Params)
+// Deprecated macro, not necessary anymore. Cf. GetValue
+#define PCG_GET_OVERRIDEN_VALUE(Settings, Variable, Params) PCGSettingsHelpers::GetValue(GET_MEMBER_NAME_CHECKED(TRemovePointer<TRemoveConst<decltype(Settings)>::Type>::Type, Variable), (Settings)->Variable, Params)
\ No newline at end of file
diff --git a/Engine/Plugins/Experimental/PCG/Source/PCG/Public/Metadata/Accessors/PCGAttributeAccessorHelpers.h b/Engine/Plugins/Experimental/PCG/Source/PCG/Public/Metadata/Accessors/PCGAttributeAccessorHelpers.h
--- a/Engine/Plugins/Experimental/PCG/Source/PCG/Public/Metadata/Accessors/PCGAttributeAccessorHelpers.h
+++ b/Engine/Plugins/Experimental/PCG/Source/PCG/Public/Metadata/Accessors/PCGAttributeAccessorHelpers.h
@@ -12,13 +12,25 @@ class UClass;
 class UPCGData;
 class UStruct;
 struct FPCGAttributePropertySelector;
+struct FPCGDataCollection;
+struct FPCGSettingsOverridableParam;
+
 
 namespace PCGAttributeAccessorHelpers
 {
+	bool IsPropertyAccessorSupported(const FProperty* InProperty);
+	bool IsPropertyAccessorSupported(const FName InPropertyName, const UClass* InClass);
+	bool IsPropertyAccessorSupported(const FName InPropertyName, const UStruct* InStruct);
+
 	TUniquePtr<IPCGAttributeAccessor> CreatePropertyAccessor(const FProperty* InProperty);
 	TUniquePtr<IPCGAttributeAccessor> CreatePropertyAccessor(const FName InPropertyName, const UClass* InClass);
 	TUniquePtr<IPCGAttributeAccessor> CreatePropertyAccessor(const FName InPropertyName, const UStruct* InStruct);
 
+	/**
+	* Create a const accessor depending on an overridable param
+	*/
+	TUniquePtr<const IPCGAttributeAccessor> CreateConstAccessorForOverrideParam(const FPCGDataCollection& InInputData, const FPCGSettingsOverridableParam& InParam, FName* OutAttributeName = nullptr);
+
 	/** 
 	* Creates a const accessor to the property or attribute pointed at by the InSelector.
 	* Note that InData must not be null if the selector points to an attribute,
diff --git a/Engine/Plugins/Experimental/PCG/Source/PCG/Public/Metadata/PCGMetadataTypesConstantStruct.h b/Engine/Plugins/Experimental/PCG/Source/PCG/Public/Metadata/PCGMetadataTypesConstantStruct.h
--- a/Engine/Plugins/Experimental/PCG/Source/PCG/Public/Metadata/PCGMetadataTypesConstantStruct.h
+++ b/Engine/Plugins/Experimental/PCG/Source/PCG/Public/Metadata/PCGMetadataTypesConstantStruct.h
@@ -2,7 +2,7 @@
 
 #pragma once
 
-#include "Helpers/PCGSettingsHelpers.h"
+#include "Metadata/PCGMetadataAttributeTraits.h"
 
 #include "PCGMetadataTypesConstantStruct.generated.h"
 
@@ -26,7 +26,7 @@ struct PCG_API FPCGMetadataTypesConstantStruct
 	GENERATED_BODY()
 public:
 	template <typename Func>
-	decltype(auto) DispatcherWithOverride(const UPCGParamData* Params, Func Callback) const;
+	decltype(auto) Dispatcher(Func Callback) const;
 
 	FString ToString() const;
 
@@ -87,50 +87,50 @@ public:
 };
 
 template <typename Func>
-decltype(auto) FPCGMetadataTypesConstantStruct::DispatcherWithOverride(const UPCGParamData* Params, Func Callback) const
+decltype(auto) FPCGMetadataTypesConstantStruct::Dispatcher(Func Callback) const
 {
 	using ReturnType = decltype(Callback(double{}));
 
 	switch (Type)
 	{
 	case EPCGMetadataTypes::Integer64:
-		return Callback(PCG_GET_OVERRIDEN_VALUE(this, IntValue, Params));
+		return Callback(IntValue);
 	case EPCGMetadataTypes::Integer32:
-		return Callback(PCG_GET_OVERRIDEN_VALUE(this, Int32Value, Params));
+		return Callback(Int32Value);
 	case EPCGMetadataTypes::Float:
-		return Callback(PCG_GET_OVERRIDEN_VALUE(this, FloatValue, Params));
+		return Callback(FloatValue);
 	case EPCGMetadataTypes::Double:
-		return Callback(PCG_GET_OVERRIDEN_VALUE(this, DoubleValue, Params));
+		return Callback(DoubleValue);
 	case EPCGMetadataTypes::Vector2:
-		return Callback(PCG_GET_OVERRIDEN_VALUE(this, Vector2Value, Params));
+		return Callback(Vector2Value);
 	case EPCGMetadataTypes::Vector:
-		return Callback(PCG_GET_OVERRIDEN_VALUE(this, VectorValue, Params));
+		return Callback(VectorValue);
 	case EPCGMetadataTypes::Vector4:
-		return Callback(PCG_GET_OVERRIDEN_VALUE(this, Vector4Value, Params));
+		return Callback(Vector4Value);
 	case EPCGMetadataTypes::Quaternion:
-		return Callback(PCG_GET_OVERRIDEN_VALUE(this, QuatValue, Params));
+		return Callback(QuatValue);
 	case EPCGMetadataTypes::Transform:
-		return Callback(PCG_GET_OVERRIDEN_VALUE(this, TransformValue, Params));
+		return Callback(TransformValue);
 	case EPCGMetadataTypes::String:
 	{
 		switch (StringMode)
 		{
 		case EPCGMetadataTypesConstantStructStringMode::String:
-			return Callback(PCG_GET_OVERRIDEN_VALUE(this, StringValue, Params));
+			return Callback(StringValue);
 		case EPCGMetadataTypesConstantStructStringMode::SoftObjectPath:
-			return Callback(PCGSettingsHelpers::GetValue(GET_MEMBER_NAME_CHECKED(FPCGMetadataTypesConstantStruct, SoftObjectPathValue), SoftObjectPathValue.ToString(), Params));
+			return Callback(SoftObjectPathValue.ToString());
 		case EPCGMetadataTypesConstantStructStringMode::SoftClassPath:
-			return Callback(PCGSettingsHelpers::GetValue(GET_MEMBER_NAME_CHECKED(FPCGMetadataTypesConstantStruct, SoftClassPathValue), SoftClassPathValue.ToString(), Params));
+			return Callback(SoftClassPathValue.ToString());
 		default:
 			break;
 		}
 	}
 	case EPCGMetadataTypes::Boolean:
-		return Callback(PCG_GET_OVERRIDEN_VALUE(this, BoolValue, Params));
+		return Callback(BoolValue);
 	case EPCGMetadataTypes::Rotator:
-		return Callback(PCG_GET_OVERRIDEN_VALUE(this, RotatorValue, Params));
+		return Callback(RotatorValue);
 	case EPCGMetadataTypes::Name:
-		return Callback(PCG_GET_OVERRIDEN_VALUE(this, NameValue, Params));
+		return Callback(NameValue);
 	default:
 		break;
 	}
diff --git a/Engine/Plugins/Experimental/PCG/Source/PCG/Public/PCGCommon.h b/Engine/Plugins/Experimental/PCG/Source/PCG/Public/PCGCommon.h
--- a/Engine/Plugins/Experimental/PCG/Source/PCG/Public/PCGCommon.h
+++ b/Engine/Plugins/Experimental/PCG/Source/PCG/Public/PCGCommon.h
@@ -59,4 +59,17 @@ namespace PCGPinConstants
 	const FName DefaultInputLabel = TEXT("In");
 	const FName DefaultOutputLabel = TEXT("Out");
 	const FName DefaultParamsLabel = TEXT("Params");
-}
\ No newline at end of file
+}
+
+// Metadata used by PCG
+namespace PCGObjectMetadata
+{
+	const FName Overridable = TEXT("PCG_Overridable");
+
+	// Metadata usable in UPROPERTY for customizing the behavior when displaying the property in a property panel or graph node
+	enum
+	{
+		/// [PropertyMetadata] Indicates that the property is overridable by params.
+		PCG_Overridable
+	};
+}
diff --git a/Engine/Plugins/Experimental/PCG/Source/PCG/Public/PCGContext.h b/Engine/Plugins/Experimental/PCG/Source/PCG/Public/PCGContext.h
--- a/Engine/Plugins/Experimental/PCG/Source/PCG/Public/PCGContext.h
+++ b/Engine/Plugins/Experimental/PCG/Source/PCG/Public/PCGContext.h
@@ -42,7 +42,7 @@ struct PCG_API FPCGContext
 {
 	GENERATED_BODY()
 
-	virtual ~FPCGContext() = default;
+	virtual ~FPCGContext();
 
 	FPCGDataCollection InputData;
 	FPCGDataCollection OutputData;
@@ -51,7 +51,6 @@ struct PCG_API FPCGContext
 
 	FPCGCrc DependenciesCrc;
 
-	// TODO: add RNG source
 	// TODO: replace this by a better identification mechanism
 	const UPCGNode* Node = nullptr;
 	FPCGTaskId TaskId = InvalidPCGTaskId;
@@ -68,14 +67,37 @@ struct PCG_API FPCGContext
 	bool bIsRunningAsyncCall = false;
 
 	const UPCGSettingsInterface* GetInputSettingsInterface() const;
+	
+	// After initializing the context, we can call this method to prepare for parameter override
+	// It will create a copy of the original settings if there is indeed a possible override.
+	void InitializeSettings();
 
+	// If we any any parameter override, it will read from the params and override matching values
+	// in the settings copy.
+	void OverrideSettings();
+
+	// Return the seed, possibly overriden by params, and combined with the source component (if any).
+	int GetSeed() const;
+
+	// Return the settings casted in the wanted type.
+	// If there is any override, those settings will already contains all the overriden values.
 	template<typename SettingsType>
 	const SettingsType* GetInputSettings() const
 	{
-		return PCGContextHelpers::GetInputSettings<SettingsType>(Node, InputData);
+		return SettingsWithOverride ? Cast<SettingsType>(SettingsWithOverride) : GetOriginalSettings<SettingsType>();
 	}
 
 	FString GetTaskName() const;
 	FString GetComponentName() const;
 	bool ShouldStop() const;
+
+private:
+	template<typename SettingsType>
+	const SettingsType* GetOriginalSettings() const 
+	{
+		return PCGContextHelpers::GetInputSettings<SettingsType>(Node, InputData);
+	}
+
+	// Copy of the settings that will be used to apply overrides.
+	TObjectPtr<UPCGSettings> SettingsWithOverride = nullptr;
 };
diff --git a/Engine/Plugins/Experimental/PCG/Source/PCG/Public/PCGCustomVersion.h b/Engine/Plugins/Experimental/PCG/Source/PCG/Public/PCGCustomVersion.h
--- a/Engine/Plugins/Experimental/PCG/Source/PCG/Public/PCGCustomVersion.h
+++ b/Engine/Plugins/Experimental/PCG/Source/PCG/Public/PCGCustomVersion.h
@@ -26,6 +26,9 @@ struct PCG_API FPCGCustomVersion
 
 		MovePointFilterParamsOffFirstPin = 5,
 
+		// Add param pin for all nodes that have override and were using the default input pin.
+		AddParamPinToOverridableNodes = 6,
+
 		// -----<new versions can be added above this line>-------------------------------------------------
 		VersionPlusOne,
 		LatestVersion = VersionPlusOne - 1
diff --git a/Engine/Plugins/Experimental/PCG/Source/PCG/Public/PCGData.h b/Engine/Plugins/Experimental/PCG/Source/PCG/Public/PCGData.h
--- a/Engine/Plugins/Experimental/PCG/Source/PCG/Public/PCGData.h
+++ b/Engine/Plugins/Experimental/PCG/Source/PCG/Public/PCGData.h
@@ -10,6 +10,7 @@
 #include "PCGData.generated.h"
 
 class FArchiveCrc32;
+class UPCGNode;
 class UPCGSettings;
 class UPCGSettingsInterface;
 class UPCGParamData;
@@ -93,10 +94,17 @@ struct PCG_API FPCGDataCollection
 	TArray<FPCGTaggedData> GetTaggedParams(const FString& InTag) const;
 	/** Returns all params on a given pin */
 	TArray<FPCGTaggedData> GetParamsByPin(const FName& InPinLabel) const;
+
+	UE_DEPRECATED(5.2, "GetParams is deprecated, please use GetParamsByPin or GetFirstParamsOnParamsPin.")
 	/** Returns the first params found in the collection */
-	UPCGParamData* GetParams() const; 
+	UPCGParamData* GetParams() const;
+
+	// Only used as a temporary solution for old graph with nodes that didn't have params pins.
+	// Should NOT be used with new nodes.
+	UPCGParamData* GetParamsWithDeprecation(const UPCGNode* Node) const;
+
 	/** Returns the first/only param found on the default params pin */
-	UPCGParamData* GetParamsOnParamsPin() const;
+	UPCGParamData* GetFirstParamsOnParamsPin() const;
 
 	const UPCGSettingsInterface* GetSettingsInterface() const;
 	const UPCGSettingsInterface* GetSettingsInterface(const UPCGSettingsInterface* InDefaultSettingsInterface) const;
diff --git a/Engine/Plugins/Experimental/PCG/Source/PCG/Public/PCGSettings.h b/Engine/Plugins/Experimental/PCG/Source/PCG/Public/PCGSettings.h
--- a/Engine/Plugins/Experimental/PCG/Source/PCG/Public/PCGSettings.h
+++ b/Engine/Plugins/Experimental/PCG/Source/PCG/Public/PCGSettings.h
@@ -9,6 +9,7 @@
 
 #include "PCGSettings.generated.h"
 
+class UPCGComponent;
 class UPCGPin;
 struct FPCGPinProperties;
 struct FPropertyChangedEvent;
@@ -49,6 +50,22 @@ enum class EPCGSettingsType : uint8
 	DECLARE_MULTICAST_DELEGATE_TwoParams(FOnPCGSettingsChanged, UPCGSettings*, EPCGChangeType);
 #endif
 
+USTRUCT()
+struct FPCGSettingsOverridableParam
+{
+	GENERATED_BODY()
+
+	UPROPERTY()
+	FName Label = NAME_None;
+
+	UPROPERTY()
+	TArray<FName> PropertiesNames;
+
+	// Transient
+	TArray<const FProperty*> Properties;
+};
+
+
 UCLASS(Abstract)
 class PCG_API UPCGSettingsInterface : public UPCGData
 {
@@ -102,14 +119,18 @@ public:
 	virtual void PostLoad() override;
 	virtual void Serialize(FArchive& Ar) override;
 	virtual void PostSaveRoot(FObjectPostSaveRootContext ObjectSaveContext) override;
+	virtual void PostInitProperties() override;
 	//~End UObject interface
 
 	// TODO: check if we need this to be virtual, we don't really need if we're always caching
 	/*virtual*/ FPCGElementPtr GetElement() const;
 	virtual UPCGNode* CreateNode() const;
+	
+	/** Return the concatenation of InputPinPropertiesand FillOverridableParamsPins */
+	TArray<FPCGPinProperties> AllInputPinProperties() const;
 
-	virtual TArray<FPCGPinProperties> InputPinProperties() const;
-	virtual TArray<FPCGPinProperties> OutputPinProperties() const;
+	/** For symmetry reason, do the same with output pins. For now forward just the call to OutputPinProperties */
+	TArray<FPCGPinProperties> AllOutputPinProperties() const;
 
 	// Internal functions, should not be used by any user.
 	// Return a different subset for for input/output pin properties, in case of a default object.
@@ -120,9 +141,12 @@ public:
 	uint32 GetCrc32() const;
 	bool UseSeed() const { return bUseSeed; }
 
+	// Get the seed, combined with optional PCGComponent seed
+	int GetSeed(const UPCGComponent* InSourceComponent = nullptr) const;
+
 #if WITH_EDITOR
 	/** UpdatePins will kick off invalid edges, so this is useful for moving edges around in case of pin changes. */
-	virtual void ApplyDeprecationBeforeUpdatePins(UPCGNode* InOutNode, TArray<TObjectPtr<UPCGPin>>& InputPins, TArray<TObjectPtr<UPCGPin>>& OutputPins) {}
+	virtual void ApplyDeprecationBeforeUpdatePins(UPCGNode* InOutNode, TArray<TObjectPtr<UPCGPin>>& InputPins, TArray<TObjectPtr<UPCGPin>>& OutputPins);
 	/** Any final migration/recovery that can be done after pins are finalized. This function should also set DataVersion to LatestVersion. */
 	virtual void ApplyDeprecation(UPCGNode* InOutNode);
 
@@ -148,7 +172,7 @@ public:
 	/** Returns true if only the first input edge is used from the primary pin when the node is disabled. */
 	virtual bool OnlyPassThroughOneEdgeWhenDisabled() const { return false; }
 
-	UPROPERTY(EditAnywhere, BlueprintReadWrite, Category = Settings, meta=(EditCondition=bUseSeed, EditConditionHides))
+	UPROPERTY(EditAnywhere, BlueprintReadWrite, Category = Settings, meta=(EditCondition=bUseSeed, EditConditionHides, PCG_Overridable))
 	int Seed = 0xC35A9631; // random prime number
 
 	/** TODO: Remove this - Placeholder feature until we have a nodegraph */
@@ -183,6 +207,10 @@ public:
 #endif
 
 protected:
+	// Returns an array of all the input pin properties. You should not add manually a "params" pin, it is handled automatically by FillOverridableParamsPins
+	virtual TArray<FPCGPinProperties> InputPinProperties() const;
+	virtual TArray<FPCGPinProperties> OutputPinProperties() const;
+
 	virtual FPCGElementPtr CreateElement() const PURE_VIRTUAL(UPCGSettings::CreateElement, return nullptr;);
 
 	/** An additional custom version number that external system users can use to track versions. This version will be serialized into the asset and will be provided by UserDataVersion after load. */
@@ -194,6 +222,8 @@ protected:
 
 	/** Method that can be called to dirty the cache data from this settings objects if the operator== does not allow to detect changes */
 	void DirtyCache();
+
+	virtual bool CanEditChange(const FProperty* InProperty) const override;
 #endif
 
 	// By default, settings won't use a seed. Set this bool to true in the child ctor to allow edition and use it.
@@ -215,6 +245,26 @@ public:
 private:
 	mutable FPCGElementPtr CachedElement;
 	mutable FCriticalSection CacheLock;
+
+	// Overridable param section
+public:
+	/** List of all the overridable params available for this settings. */
+	virtual const TArray<FPCGSettingsOverridableParam>& OverridableParams() const { return CachedOverridableParams; }
+	/** Check if we have some override. Can be overriden to force params pin for example */
+	virtual bool HasOverridableParams() const { return !CachedOverridableParams.IsEmpty(); }
+protected:
+	/** Iterate over OverridableParams to automatically add param pins to the list. */
+	void FillOverridableParamsPins(TArray<FPCGPinProperties>& OutPins) const;
+private:
+	/** Called after intialization to construct the list of all overridable params. 
+	* In Editor, it will first gather all the overridable properties names and labels, based on their metadata.
+	* And then, in Editor and Runtime, will gather the FProperty*.
+	*/
+	void InitializeCachedOverridableParams();
+
+	/** Needs to be serialized since property metadata (used to populate this array) is not available at runtime. */
+	UPROPERTY()
+	TArray<FPCGSettingsOverridableParam> CachedOverridableParams;
 };
 
 UCLASS(BlueprintType, ClassGroup = (Procedural))
diff --git a/Engine/Plugins/Experimental/PCG/Source/PCGEditor/Private/PCGEditor.cpp b/Engine/Plugins/Experimental/PCG/Source/PCGEditor/Private/PCGEditor.cpp
--- a/Engine/Plugins/Experimental/PCG/Source/PCGEditor/Private/PCGEditor.cpp
+++ b/Engine/Plugins/Experimental/PCG/Source/PCGEditor/Private/PCGEditor.cpp
@@ -534,7 +534,7 @@ void FPCGEditor::OnDeterminismNodeTest()
 				TSharedPtr<FDeterminismTestResult> NodeResult = MakeShared<FDeterminismTestResult>();
 				NodeResult->TestResultTitle = PCGNode->GetNodeTitle();
 				NodeResult->TestResultName = PCGNode->GetName();
-				NodeResult->Seed = PCGNode->GetSettings()->Seed;
+				NodeResult->Seed = PCGNode->GetSettings()->GetSeed();
 
 				if (PCGNode->GetSettings()->DeterminismSettings.bNativeTests)
 				{
diff --git a/Engine/Plugins/Experimental/PCG/Source/PCGEditor/Private/PCGEditorGraphNodeBase.cpp b/Engine/Plugins/Experimental/PCG/Source/PCGEditor/Private/PCGEditorGraphNodeBase.cpp
--- a/Engine/Plugins/Experimental/PCG/Source/PCGEditor/Private/PCGEditorGraphNodeBase.cpp
+++ b/Engine/Plugins/Experimental/PCG/Source/PCGEditor/Private/PCGEditorGraphNodeBase.cpp
@@ -573,12 +573,22 @@ void UPCGEditorGraphNodeBase::UpdatePosition()
 	}
 }
 
+bool UPCGEditorGraphNodeBase::ShouldCreatePin(const UPCGPin* InPin) const
+{
+	return true;
+}
+
 void UPCGEditorGraphNodeBase::CreatePins(const TArray<UPCGPin*>& InInputPins, const TArray<UPCGPin*>& InOutputPins)
 {
 	bool bHasAdvancedPin = false;
 
 	for (const UPCGPin* InputPin : InInputPins)
 	{
+		if (!ShouldCreatePin(InputPin))
+		{
+			continue;
+		}
+
 		UEdGraphPin* Pin = CreatePin(EEdGraphPinDirection::EGPD_Input, GetPinType(InputPin), InputPin->Properties.Label);
 		Pin->bAdvancedView = InputPin->Properties.bAdvancedPin;
 		bHasAdvancedPin |= Pin->bAdvancedView;
@@ -586,6 +596,11 @@ void UPCGEditorGraphNodeBase::CreatePins(const TArray<UPCGPin*>& InInputPins, co
 
 	for (const UPCGPin* OutputPin : InOutputPins)
 	{
+		if (!ShouldCreatePin(OutputPin))
+		{
+			continue;
+		}
+
 		UEdGraphPin* Pin = CreatePin(EEdGraphPinDirection::EGPD_Output, GetPinType(OutputPin), OutputPin->Properties.Label);
 		Pin->bAdvancedView = OutputPin->Properties.bAdvancedPin;
 		bHasAdvancedPin |= Pin->bAdvancedView;
diff --git a/Engine/Plugins/Experimental/PCG/Source/PCGEditor/Private/PCGEditorGraphNodeBase.h b/Engine/Plugins/Experimental/PCG/Source/PCGEditor/Private/PCGEditorGraphNodeBase.h
--- a/Engine/Plugins/Experimental/PCG/Source/PCGEditor/Private/PCGEditorGraphNodeBase.h
+++ b/Engine/Plugins/Experimental/PCG/Source/PCGEditor/Private/PCGEditorGraphNodeBase.h
@@ -65,6 +65,10 @@ protected:
 
 	void CreatePins(const TArray<UPCGPin*>& InInputPins, const TArray<UPCGPin*>& InOutputPins);
 
+	// Custom logic to hide some pins to the user (by not creating a UI pin, even if the model pin exists).
+	// Useful for deprecation
+	virtual bool ShouldCreatePin(const UPCGPin* InPin) const;
+
 	UPROPERTY()
 	TObjectPtr<UPCGNode> PCGNode = nullptr;
 
