Fixed automatic re-outering issue in blueprint elements - could end up removing data saved elsewhere and could cause data loss.

#rnx

diff --git a/Engine/Plugins/Experimental/PCG/Source/PCG/Private/Elements/PCGExecuteBlueprint.cpp b/Engine/Plugins/Experimental/PCG/Source/PCG/Private/Elements/PCGExecuteBlueprint.cpp
--- a/Engine/Plugins/Experimental/PCG/Source/PCG/Private/Elements/PCGExecuteBlueprint.cpp
+++ b/Engine/Plugins/Experimental/PCG/Source/PCG/Private/Elements/PCGExecuteBlueprint.cpp
@@ -526,9 +526,23 @@ bool FPCGExecuteBlueprintElement::ExecuteInternal(FPCGContext* InContext) const
 			// Important implementation note:
 			// Any data that was created by the user in the blueprint will have that data parented to this blueprint element instance
 			// Which will cause issues wrt to reference leaks. We need to fix this here.
-			if (Output.Data)
+			// Note that we will recurse up the outer tree to make sure we catch every case.
+			if(Output.Data)
 			{
-				const_cast<UPCGData*>(Output.Data.Get())->Rename(nullptr, GetTransientPackage(), REN_ForceNoResetLoaders | REN_DoNotDirty | REN_DontCreateRedirectors | REN_NonTransactional);
+				UObject* ThisData = const_cast<UPCGData*>(Output.Data.Get());
+
+				bool bHasInstanceAsOuter = false;
+				UObject* CurrentObject = ThisData;
+				while (CurrentObject && !bHasInstanceAsOuter)
+				{
+					bHasInstanceAsOuter = (CurrentObject->GetOuter() == Context->BlueprintElementInstance);
+					CurrentObject = CurrentObject->GetOuter();
+				}
+
+				if (bHasInstanceAsOuter)
+				{
+					ThisData->Rename(nullptr, GetTransientPackage(), REN_ForceNoResetLoaders | REN_DoNotDirty | REN_DontCreateRedirectors | REN_NonTransactional);
+				}
 			}
 		}
 	}
