[PCG] Use2DGrid is now true by default option in PCG World Actor. Note that previous maps that already has PartitionActors created with the flag at False, will force the flag in PCG World Actor at False.

diff --git a/Engine/Plugins/Experimental/PCG/Source/PCG/Private/Grid/PCGPartitionActor.cpp b/Engine/Plugins/Experimental/PCG/Source/PCG/Private/Grid/PCGPartitionActor.cpp
--- a/Engine/Plugins/Experimental/PCG/Source/PCG/Private/Grid/PCGPartitionActor.cpp
+++ b/Engine/Plugins/Experimental/PCG/Source/PCG/Private/Grid/PCGPartitionActor.cpp
@@ -387,7 +387,7 @@ void APCGPartitionActor::PostCreation()
 	}
 	else
 	{
-		bUse2DGrid = false;
+		bUse2DGrid = true;
 	}
 
 #if WITH_EDITOR
diff --git a/Engine/Plugins/Experimental/PCG/Source/PCG/Private/PCGSubsystem.cpp b/Engine/Plugins/Experimental/PCG/Source/PCG/Private/PCGSubsystem.cpp
--- a/Engine/Plugins/Experimental/PCG/Source/PCG/Private/PCGSubsystem.cpp
+++ b/Engine/Plugins/Experimental/PCG/Source/PCG/Private/PCGSubsystem.cpp
@@ -818,6 +818,16 @@ void UPCGSubsystem::RegisterPartitionActor(APCGPartitionActor* Actor, bool bDoCo
 		PartitionActorsMap.Add(GridCoord, Actor);
 	}
 
+	// For deprecration: bUse2DGrid is now true by default. But if we already have Partition Actors that were created when the flag was false by default,
+	// we keep this flag
+	if (APCGWorldActor* WorldActor = GetPCGWorldActor())
+	{
+		if (WorldActor->bUse2DGrid != Actor->IsUsing2DGrid())
+		{
+			WorldActor->bUse2DGrid = Actor->IsUsing2DGrid();
+		}
+	}
+
 	// And then register itself to all the components that intersect with it
 	{
 		FWriteScopeLock WriteLock(ComponentToPartitionActorsMapLock);
diff --git a/Engine/Plugins/Experimental/PCG/Source/PCG/Public/PCGWorldActor.h b/Engine/Plugins/Experimental/PCG/Source/PCG/Public/PCGWorldActor.h
--- a/Engine/Plugins/Experimental/PCG/Source/PCG/Public/PCGWorldActor.h
+++ b/Engine/Plugins/Experimental/PCG/Source/PCG/Public/PCGWorldActor.h
@@ -49,7 +49,7 @@ public:
 
 	/** Disable creation of Partition Actors on the Z axis. Can improve performances if 3D partitioning is not needed. */
 	UPROPERTY(config, EditAnywhere, Category = WorldPartition)
-	bool bUse2DGrid = false;
+	bool bUse2DGrid = true;
 
 private:
 	void RegisterToSubsystem();
