[PCG] Crash fix - stop assuming entries in GeneratedResources on PCGComponent are non-null.

Adds paranoia checks and actively just to boot nulls out of the array.

There might be a BP-related path which means resources are getting added that are not owned by the PCGComponent and are null'd out.

diff --git a/Engine/Plugins/Experimental/PCG/Source/PCG/Private/PCGComponent.cpp b/Engine/Plugins/Experimental/PCG/Source/PCG/Private/PCGComponent.cpp
--- a/Engine/Plugins/Experimental/PCG/Source/PCG/Private/PCGComponent.cpp
+++ b/Engine/Plugins/Experimental/PCG/Source/PCG/Private/PCGComponent.cpp
@@ -136,7 +136,10 @@ void UPCGComponent::ForEachManagedResource(TFunctionRef<void(UPCGManagedResource
 	check(!GeneratedResourcesInaccessible);
 	for (TObjectPtr<UPCGManagedResource> ManagedResource : GeneratedResources)
 	{
-		Func(ManagedResource);
+		if (ManagedResource)
+		{
+			Func(ManagedResource);
+		}
 	}
 }
 
@@ -624,11 +627,17 @@ bool UPCGComponent::MoveResourcesToNewActor(AActor* InNewActor, bool bCreateChil
 		check(!GeneratedResourcesInaccessible);
 		for (TObjectPtr<UPCGManagedResource>& GeneratedResource : GeneratedResources)
 		{
-			check(GeneratedResource);
-			GeneratedResource->MoveResourceToNewActor(NewActor);
-			TSet<TSoftObjectPtr<AActor>> Dummy;
-			GeneratedResource->ReleaseIfUnused(Dummy);
-			bHasMovedResources = true;
+			if (GeneratedResource)
+			{
+				GeneratedResource->MoveResourceToNewActor(NewActor);
+				TSet<TSoftObjectPtr<AActor>> Dummy;
+				GeneratedResource->ReleaseIfUnused(Dummy);
+				bHasMovedResources = true;
+			}
+			else if (GetOwner())
+			{
+				UE_LOG(LogPCG, Error, TEXT("[UPCGComponent::MoveResourcesToNewActor] Null generated resource encountered on actor \"%s\" and will be skipped."), *GetOwner()->GetFName().ToString());
+			}
 		}
 
 		GeneratedResources.Empty();
@@ -726,9 +735,13 @@ FPCGTaskId UPCGComponent::CreateCleanupTask(bool bRemoveComponents, const TArray
 			if (Context->ResourceIndex >= 0)
 			{
 				UPCGManagedResource* Resource = ThisComponent->GeneratedResources[Context->ResourceIndex];
-				check(Resource);
 
-				if (Resource->Release(bRemoveComponents, Context->ActorsToDelete))
+				if (!Resource && ThisComponent->GetOwner())
+				{
+					UE_LOG(LogPCG, Error, TEXT("[UPCGComponent::CreateCleanupTask] Null generated resource encountered on actor \"%s\"."), *ThisComponent->GetOwner()->GetFName().ToString());
+				}
+
+				if (!Resource || Resource->Release(bRemoveComponents, Context->ActorsToDelete))
 				{
 					ThisComponent->GeneratedResources.RemoveAtSwap(Context->ResourceIndex);
 				}
@@ -764,8 +777,14 @@ void UPCGComponent::CleanupUnusedManagedResources()
 		check(!GeneratedResourcesInaccessible);
 		for (int32 ResourceIndex = GeneratedResources.Num() - 1; ResourceIndex >= 0; --ResourceIndex)
 		{
-			check(GeneratedResources[ResourceIndex]);
-			if (GeneratedResources[ResourceIndex]->ReleaseIfUnused(ActorsToDelete))
+			UPCGManagedResource* Resource = GeneratedResources[ResourceIndex];
+
+			if (!Resource && GetOwner())
+			{
+				UE_LOG(LogPCG, Error, TEXT("[UPCGComponent::CleanupUnusedManagedResources] Null generated resource encountered on actor \"%s\"."), *GetOwner()->GetFName().ToString());
+			}
+
+			if (!Resource || Resource->ReleaseIfUnused(ActorsToDelete))
 			{
 				GeneratedResources.RemoveAtSwap(ResourceIndex);
 			}
@@ -2346,6 +2365,15 @@ void UPCGComponent::SetManagedResources(const TArray<TObjectPtr<UPCGManagedResou
 	FScopeLock ResourcesLock(&GeneratedResourcesLock);
 	check(GeneratedResources.IsEmpty());
 	GeneratedResources = Resources;
+
+	// Remove any null entries
+	for (int32 ResourceIndex = GeneratedResources.Num() - 1; ResourceIndex >= 0; --ResourceIndex)
+	{
+		if (!GeneratedResources[ResourceIndex])
+		{
+			GeneratedResources.RemoveAtSwap(ResourceIndex);
+		}
+	}
 }
 
 void UPCGComponent::GetManagedResources(TArray<TObjectPtr<UPCGManagedResource>>& Resources) const
