Added attribute-driven material overrides in the static mesh spawner node

diff --git a/Engine/Plugins/Experimental/PCG/Source/PCG/Private/Elements/PCGStaticMeshSpawner.cpp b/Engine/Plugins/Experimental/PCG/Source/PCG/Private/Elements/PCGStaticMeshSpawner.cpp
--- a/Engine/Plugins/Experimental/PCG/Source/PCG/Private/Elements/PCGStaticMeshSpawner.cpp
+++ b/Engine/Plugins/Experimental/PCG/Source/PCG/Private/Elements/PCGStaticMeshSpawner.cpp
@@ -12,6 +12,7 @@
 
 #include "Components/InstancedStaticMeshComponent.h"
 #include "Engine/StaticMesh.h"
+#include "Materials/MaterialInterface.h"
 
 #include UE_INLINE_GENERATED_CPP_BY_NAME(PCGStaticMeshSpawner)
 
@@ -196,7 +197,11 @@ void FPCGStaticMeshSpawnerElement::SpawnStaticMeshInstances(FPCGContext* Context
 
 	if (InstanceList.bOverrideMaterials)
 	{
-		Params.MaterialOverrides = InstanceList.MaterialOverrides;
+		Params.MaterialOverrides.Reserve(InstanceList.MaterialOverrides.Num());
+		for (TSoftObjectPtr<UMaterialInterface> MaterialOverride : InstanceList.MaterialOverrides)
+		{
+			Params.MaterialOverrides.Add(MaterialOverride.LoadSynchronous());
+		}
 	}
 
 	Params.NumCustomDataFloats = PackedCustomData.NumCustomDataFloats;
diff --git a/Engine/Plugins/Experimental/PCG/Source/PCG/Private/MeshSelectors/PCGMeshSelectorBase.cpp b/Engine/Plugins/Experimental/PCG/Source/PCG/Private/MeshSelectors/PCGMeshSelectorBase.cpp
--- a/Engine/Plugins/Experimental/PCG/Source/PCG/Private/MeshSelectors/PCGMeshSelectorBase.cpp
+++ b/Engine/Plugins/Experimental/PCG/Source/PCG/Private/MeshSelectors/PCGMeshSelectorBase.cpp
@@ -1,7 +1,7 @@
 // Copyright Epic Games, Inc. All Rights Reserved.
 
 #include "MeshSelectors/PCGMeshSelectorBase.h"
-
+#include "PCGElement.h"
 
 #include UE_INLINE_GENERATED_CPP_BY_NAME(PCGMeshSelectorBase)
 
@@ -11,10 +11,10 @@ int32 UPCGMeshSelectorBase::FindOrAddInstanceList(
 	bool bOverrideCollisionProfile,
 	const FCollisionProfileName& CollisionProfile,
 	bool bOverrideMaterials,
-	const TArray<UMaterialInterface*>& MaterialOverrides,
+	const TArray<TSoftObjectPtr<UMaterialInterface>>& MaterialOverrides,
 	const float InCullStartDistance,
 	const float InCullEndDistance, 
-	const bool bInIsLocalToWorldDeterminantNegative) const
+	const bool bInIsLocalToWorldDeterminantNegative)
 {
 	for (int Index = 0; Index < OutInstanceLists.Num(); ++Index)
 	{
@@ -49,3 +49,109 @@ int32 UPCGMeshSelectorBase::FindOrAddInstanceList(
 	return OutInstanceLists.Emplace(Mesh, bOverrideCollisionProfile, CollisionProfile, bOverrideMaterials, MaterialOverrides, InCullStartDistance, InCullEndDistance, bInIsLocalToWorldDeterminantNegative);
 }
 
+FPCGMeshMaterialOverrideHelper::FPCGMeshMaterialOverrideHelper(
+	FPCGContext& InContext,
+	EPCGMeshSelectorMaterialOverrideMode InMaterialOverrideMode,
+	const TArray<TSoftObjectPtr<UMaterialInterface>>& InStaticMaterialOverrides,
+	const TArray<FName>& InMaterialOverrideAttributeNames,
+	const UPCGMetadata* InMetadata
+)
+	: MaterialOverrideMode(InMaterialOverrideMode)
+	, StaticMaterialOverrides(InStaticMaterialOverrides)
+	, MaterialOverrideAttributeNames(InMaterialOverrideAttributeNames)
+	, Metadata(InMetadata)
+{
+	Initialize(InContext);
+}
+
+FPCGMeshMaterialOverrideHelper::FPCGMeshMaterialOverrideHelper(
+	FPCGContext& InContext,
+	bool bInByAttributeOverride,
+	const TArray<FName>& InMaterialOverrideAttributeNames,
+	const UPCGMetadata* InMetadata
+)
+	: MaterialOverrideMode(bInByAttributeOverride ? EPCGMeshSelectorMaterialOverrideMode::ByAttributeOverride : EPCGMeshSelectorMaterialOverrideMode::NoOverride)
+	, StaticMaterialOverrides(EmptyArray)
+	, MaterialOverrideAttributeNames(InMaterialOverrideAttributeNames)
+	, Metadata(InMetadata)
+{
+	Initialize(InContext);
+}
+
+void FPCGMeshMaterialOverrideHelper::Initialize(FPCGContext& InContext)
+{
+	// Perform data setup & validation up-front
+	if (MaterialOverrideMode == EPCGMeshSelectorMaterialOverrideMode::ByAttributeOverride)
+	{
+		if (!Metadata)
+		{
+			PCGE_LOG_C(Error, &InContext, "Data has no metadata");
+			return;
+		}
+
+		for (const FName& MaterialOverrideAttributeName : MaterialOverrideAttributeNames)
+		{
+			if (!Metadata->HasAttribute(MaterialOverrideAttributeName))
+			{
+				PCGE_LOG_C(Error, &InContext, "Attribute %s for material overrides is not present in the metadata", *MaterialOverrideAttributeName.ToString());
+				return;
+			}
+
+			const FPCGMetadataAttributeBase* MaterialAttributeBase = Metadata->GetConstAttribute(MaterialOverrideAttributeName);
+			check(MaterialAttributeBase);
+
+			if (MaterialAttributeBase->GetTypeId() != PCG::Private::MetadataTypes<FString>::Id)
+			{
+				PCGE_LOG_C(Error, &InContext, "Material override attribute is not of valid type");
+				return;
+			}
+
+			MaterialAttributes.Add(static_cast<const FPCGMetadataAttribute<FString>*>(MaterialAttributeBase));
+		}
+
+		ValueKeyToOverrideMaterials.SetNum(MaterialOverrideAttributeNames.Num());
+		WorkingMaterialOverrides.Reserve(MaterialOverrideAttributeNames.Num());
+	}
+
+	bIsValid = true;
+}
+
+const TArray<TSoftObjectPtr<UMaterialInterface>>& FPCGMeshMaterialOverrideHelper::GetMaterialOverrides(PCGMetadataEntryKey EntryKey)
+{
+	check(bIsValid);
+	if (MaterialOverrideMode == EPCGMeshSelectorMaterialOverrideMode::ByAttributeOverride)
+	{
+		WorkingMaterialOverrides.Reset();
+
+		for (int32 MaterialIndex = 0; MaterialIndex < MaterialAttributes.Num(); ++MaterialIndex)
+		{
+			const FPCGMetadataAttribute<FString>* MaterialAttribute = MaterialAttributes[MaterialIndex];
+			PCGMetadataValueKey MaterialValueKey = MaterialAttribute->GetValueKey(EntryKey);
+			TSoftObjectPtr<UMaterialInterface>* NewMaterial = ValueKeyToOverrideMaterials[MaterialIndex].Find(MaterialValueKey);
+			TSoftObjectPtr<UMaterialInterface> Material = nullptr;
+
+			if (!NewMaterial)
+			{
+				FSoftObjectPath MaterialPath(MaterialAttribute->GetValue(MaterialValueKey));
+				Material = TSoftObjectPtr<UMaterialInterface>(MaterialPath);
+				ValueKeyToOverrideMaterials[MaterialIndex].Add(MaterialValueKey, Material);
+			}
+			else
+			{
+				Material = *NewMaterial;
+			}
+
+			WorkingMaterialOverrides.Add(Material);
+		}
+
+		return WorkingMaterialOverrides;
+	}
+	else if(MaterialOverrideMode == EPCGMeshSelectorMaterialOverrideMode::StaticOverride)
+	{
+		return StaticMaterialOverrides;
+	}
+	else // no override
+	{
+		return EmptyArray;
+	}
+}
diff --git a/Engine/Plugins/Experimental/PCG/Source/PCG/Private/MeshSelectors/PCGMeshSelectorByAttribute.cpp b/Engine/Plugins/Experimental/PCG/Source/PCG/Private/MeshSelectors/PCGMeshSelectorByAttribute.cpp
--- a/Engine/Plugins/Experimental/PCG/Source/PCG/Private/MeshSelectors/PCGMeshSelectorByAttribute.cpp
+++ b/Engine/Plugins/Experimental/PCG/Source/PCG/Private/MeshSelectors/PCGMeshSelectorByAttribute.cpp
@@ -9,6 +9,19 @@
 
 #include UE_INLINE_GENERATED_CPP_BY_NAME(PCGMeshSelectorByAttribute)
 
+void UPCGMeshSelectorByAttribute::PostLoad()
+{
+	Super::PostLoad();
+
+#if WITH_EDITOR
+	if (bOverrideMaterials_DEPRECATED)
+	{
+		MaterialOverrideMode = EPCGMeshSelectorMaterialOverrideMode::StaticOverride;
+		bOverrideMaterials_DEPRECATED = false;
+	}
+#endif
+}
+
 void UPCGMeshSelectorByAttribute::SelectInstances_Implementation(
 	FPCGContext& Context, 
 	const UPCGStaticMeshSpawnerSettings* Settings, 
@@ -47,6 +60,13 @@ void UPCGMeshSelectorByAttribute::SelectInstances_Implementation(
 
 	const FPCGMetadataAttribute<FString>* Attribute = static_cast<const FPCGMetadataAttribute<FString>*>(AttributeBase);
 
+	FPCGMeshMaterialOverrideHelper MaterialOverrideHelper(Context, MaterialOverrideMode, MaterialOverrides, MaterialOverrideAttributes, PointData->Metadata);
+
+	if (!MaterialOverrideHelper.IsValid())
+	{
+		return;
+	}
+
 	TMap<PCGMetadataValueKey, TSoftObjectPtr<UStaticMesh>> ValueKeyToMesh;
 
 	// ByAttribute takes in SoftObjectPaths per point in the metadata, so we can pass those directly into the outgoing pin if it exists
@@ -91,7 +111,7 @@ void UPCGMeshSelectorByAttribute::SelectInstances_Implementation(
 		}
 
 		const bool bReverseTransform = (Point.Transform.GetDeterminant() < 0);
-		const int32 Index = FindOrAddInstanceList(OutMeshInstances, Mesh, bOverrideCollisionProfile, CollisionProfile, bOverrideMaterials, MaterialOverrides, CullStartDistance, CullEndDistance, bReverseTransform);
+		const int32 Index = FindOrAddInstanceList(OutMeshInstances, Mesh, bOverrideCollisionProfile, CollisionProfile, MaterialOverrideHelper.OverridesMaterials(), MaterialOverrideHelper.GetMaterialOverrides(Point.MetadataEntry), CullStartDistance, CullEndDistance, bReverseTransform);
 		check(Index != INDEX_NONE);
 		
 		OutMeshInstances[Index].Instances.Emplace(Point);
diff --git a/Engine/Plugins/Experimental/PCG/Source/PCG/Private/MeshSelectors/PCGMeshSelectorWeighted.cpp b/Engine/Plugins/Experimental/PCG/Source/PCG/Private/MeshSelectors/PCGMeshSelectorWeighted.cpp
--- a/Engine/Plugins/Experimental/PCG/Source/PCG/Private/MeshSelectors/PCGMeshSelectorWeighted.cpp
+++ b/Engine/Plugins/Experimental/PCG/Source/PCG/Private/MeshSelectors/PCGMeshSelectorWeighted.cpp
@@ -12,6 +12,63 @@
 
 #include UE_INLINE_GENERATED_CPP_BY_NAME(PCGMeshSelectorWeighted)
 
+namespace PCGMeshSelectorWeighted
+{
+	FPCGMeshInstanceList& GetInstanceList(
+		TArray<FPCGMeshInstanceList>& InstanceLists,
+		bool bUseAttributeMaterialOverrides,
+		const TArray<TSoftObjectPtr<UMaterialInterface>>& InMaterialOverrides,
+		bool bInIsLocalToWorldDeterminantNegative)
+	{
+		check(InstanceLists.Num() > 0);
+		const bool bOverrideMaterials = (bUseAttributeMaterialOverrides || InstanceLists[0].bOverrideMaterials);
+
+		// First look through previously existing values - note that we scope this to prevent issues with the 0 index access which might become invalid below
+		{
+			const TArray<TSoftObjectPtr<UMaterialInterface>>& MaterialOverrides = (bUseAttributeMaterialOverrides ? InMaterialOverrides : InstanceLists[0].MaterialOverrides);
+
+			for (int Index = 0; Index < InstanceLists.Num(); ++Index)
+			{
+				if (InstanceLists[Index].bOverrideMaterials != bOverrideMaterials)
+				{
+					continue;
+				}
+
+				if (InstanceLists[Index].bIsLocalToWorldDeterminantNegative != bInIsLocalToWorldDeterminantNegative)
+				{
+					continue;
+				}
+
+				if (bOverrideMaterials && InstanceLists[Index].MaterialOverrides != MaterialOverrides)
+				{
+					continue;
+				}
+
+				return InstanceLists[Index];
+			}
+		}
+
+		// If not found, then copy first entry which is our "clean" version and apply the new values
+		{
+			FPCGMeshInstanceList& NewInstanceList = InstanceLists.Emplace_GetRef();
+			const FPCGMeshInstanceList& SourceInstanceList = InstanceLists[0];
+			NewInstanceList.Mesh = SourceInstanceList.Mesh;
+			NewInstanceList.bOverrideCollisionProfile = SourceInstanceList.bOverrideCollisionProfile;
+			NewInstanceList.CollisionProfile = SourceInstanceList.CollisionProfile;
+			NewInstanceList.CullStartDistance = SourceInstanceList.CullStartDistance;
+			NewInstanceList.CullEndDistance = SourceInstanceList.CullEndDistance;
+
+			NewInstanceList.bOverrideMaterials = bOverrideMaterials;
+
+			const TArray<TSoftObjectPtr<UMaterialInterface>>& MaterialOverrides = (bUseAttributeMaterialOverrides ? InMaterialOverrides : SourceInstanceList.MaterialOverrides);
+			NewInstanceList.MaterialOverrides = MaterialOverrides;
+			NewInstanceList.bIsLocalToWorldDeterminantNegative = bInIsLocalToWorldDeterminantNegative;
+
+			return NewInstanceList;
+		}
+	}
+}
+
 void UPCGMeshSelectorWeighted::SelectInstances_Implementation(
 	FPCGContext& Context, 
 	const UPCGStaticMeshSpawnerSettings* Settings, 
@@ -19,11 +76,12 @@ void UPCGMeshSelectorWeighted::SelectInstances_Implementation(
 	TArray<FPCGMeshInstanceList>& OutMeshInstances,
 	UPCGPointData* OutPointData) const
 {
-	TArray<FPCGMeshInstanceList> OutReverseMeshInstances;
+	TArray<TArray<FPCGMeshInstanceList>> MeshInstances;
 	TArray<int> CumulativeWeights;
 
 	int TotalWeight = 0;
 
+	// Prepare common mesh setups which we will use as a kind of map
 	for (const FPCGMeshSelectorWeightedEntry& Entry : MeshEntries)
 	{
 		if (Entry.Weight <= 0)
@@ -32,8 +90,9 @@ void UPCGMeshSelectorWeighted::SelectInstances_Implementation(
 			continue;
 		}
 
-		FindOrAddInstanceList(OutMeshInstances, Entry.Mesh, Entry.bOverrideCollisionProfile, Entry.CollisionProfile, Entry.bOverrideMaterials, Entry.MaterialOverrides, Entry.CullStartDistance, Entry.CullEndDistance, /*bReverseCulling=*/false);
-		FindOrAddInstanceList(OutReverseMeshInstances, Entry.Mesh, Entry.bOverrideCollisionProfile, Entry.CollisionProfile, Entry.bOverrideMaterials, Entry.MaterialOverrides, Entry.CullStartDistance, Entry.CullEndDistance, /*bReverseCulling=*/true);
+		TArray<FPCGMeshInstanceList>& PickEntry = MeshInstances.Emplace_GetRef();
+		FindOrAddInstanceList(PickEntry, Entry.Mesh, Entry.bOverrideCollisionProfile, Entry.CollisionProfile, Entry.bOverrideMaterials, Entry.MaterialOverrides, Entry.CullStartDistance, Entry.CullEndDistance, /*bReverseCulling=*/false);
+
 		TotalWeight += Entry.Weight;
 		CumulativeWeights.Add(TotalWeight);
 	}
@@ -55,6 +114,13 @@ void UPCGMeshSelectorWeighted::SelectInstances_Implementation(
 	FPCGMetadataAttribute<FString>* OutAttribute = nullptr;
 	TMap<TSoftObjectPtr<UStaticMesh>, PCGMetadataValueKey> MeshToValueKey;
 
+	FPCGMeshMaterialOverrideHelper MaterialOverrideHelper(Context, bUseAttributeMaterialOverrides, MaterialOverrideAttributes, PointData->Metadata);
+
+	if (!MaterialOverrideHelper.IsValid())
+	{
+		return;
+	}
+
 	if (OutPointData)
 	{
 		check(OutPointData->Metadata);
@@ -95,15 +161,15 @@ void UPCGMeshSelectorWeighted::SelectInstances_Implementation(
 			int RandomWeightedPick = RandomSource.RandRange(0, TotalWeight - 1);
 
 			int RandomPick = 0;
-			while (RandomPick < OutMeshInstances.Num() && CumulativeWeights[RandomPick] <= RandomWeightedPick)
+			while(RandomPick < MeshInstances.Num() && CumulativeWeights[RandomPick] <= RandomWeightedPick)
 			{
 				++RandomPick;
 			}
 
-			if (RandomPick < OutMeshInstances.Num())
+			if(RandomPick < MeshInstances.Num())
 			{
 				const bool bNeedsReverseCulling = (Point.Transform.GetDeterminant() < 0);
-				FPCGMeshInstanceList& InstanceList = (bNeedsReverseCulling ? OutReverseMeshInstances[RandomPick] : OutMeshInstances[RandomPick]);
+				FPCGMeshInstanceList& InstanceList = PCGMeshSelectorWeighted::GetInstanceList(MeshInstances[RandomPick], bUseAttributeMaterialOverrides, MaterialOverrideHelper.GetMaterialOverrides(Point.MetadataEntry), bNeedsReverseCulling);
 				InstanceList.Instances.Emplace(Point);
 
 				const TSoftObjectPtr<UStaticMesh>& Mesh = InstanceList.Mesh;
@@ -127,10 +193,13 @@ void UPCGMeshSelectorWeighted::SelectInstances_Implementation(
 		}
 	}
 
-	// Collapse reverse entries to the OutMeshInstances
-	for (int32 OutIndex = 0; OutIndex < OutReverseMeshInstances.Num(); ++OutIndex)
+	// Collapse to OutMeshInstances
+	for (TArray<FPCGMeshInstanceList>& PickedMeshInstances : MeshInstances)
 	{
-		OutMeshInstances.Emplace(MoveTemp(OutReverseMeshInstances[OutIndex]));
+		for (FPCGMeshInstanceList& PickedMeshInstanceEntry : PickedMeshInstances)
+		{
+			OutMeshInstances.Emplace(MoveTemp(PickedMeshInstanceEntry));
+		}
 	}
 }
 
diff --git a/Engine/Plugins/Experimental/PCG/Source/PCG/Private/MeshSelectors/PCGMeshSelectorWeightedByCategory.cpp b/Engine/Plugins/Experimental/PCG/Source/PCG/Private/MeshSelectors/PCGMeshSelectorWeightedByCategory.cpp
--- a/Engine/Plugins/Experimental/PCG/Source/PCG/Private/MeshSelectors/PCGMeshSelectorWeightedByCategory.cpp
+++ b/Engine/Plugins/Experimental/PCG/Source/PCG/Private/MeshSelectors/PCGMeshSelectorWeightedByCategory.cpp
@@ -14,8 +14,7 @@
 
 struct FPCGInstancesAndWeights
 {
-	TArray<int> InstanceListIndices;
-	TArray<int> ReverseInstanceListIndices;
+	TArray<TArray<FPCGMeshInstanceList>> MeshInstances;
 	TArray<int> CumulativeWeights;
 };
 
@@ -26,7 +25,6 @@ void UPCGMeshSelectorWeightedByCategory::SelectInstances_Implementation(
 	TArray<FPCGMeshInstanceList>& OutMeshInstances,
 	UPCGPointData* OutPointData) const
 {
-	TArray<FPCGMeshInstanceList> OutReverseMeshInstances;
 	const UPCGPointData* PointData = InSpatialData->ToPointData(&Context);
 
 	if (!PointData)
@@ -112,15 +110,12 @@ void UPCGMeshSelectorWeightedByCategory::SelectInstances_Implementation(
 				continue;
 			}
 			
-			const int32 Index = FindOrAddInstanceList(OutMeshInstances, WeightedEntry.Mesh, WeightedEntry.bOverrideCollisionProfile, WeightedEntry.CollisionProfile, WeightedEntry.bOverrideMaterials, WeightedEntry.MaterialOverrides, WeightedEntry.CullStartDistance, WeightedEntry.CullEndDistance, /*bReverseCulling=*/false);
-			const int32 ReverseIndex = FindOrAddInstanceList(OutReverseMeshInstances, WeightedEntry.Mesh, WeightedEntry.bOverrideCollisionProfile, WeightedEntry.CollisionProfile, WeightedEntry.bOverrideMaterials, WeightedEntry.MaterialOverrides, WeightedEntry.CullStartDistance, WeightedEntry.CullEndDistance, /*bReverseCulling=*/true);
-			check(Index != INDEX_NONE && ReverseIndex != INDEX_NONE);
+			TArray<FPCGMeshInstanceList>& PickEntry = InstancesAndWeights->MeshInstances.Emplace_GetRef();
+			FindOrAddInstanceList(PickEntry, WeightedEntry.Mesh, WeightedEntry.bOverrideCollisionProfile, WeightedEntry.CollisionProfile, WeightedEntry.bOverrideMaterials, WeightedEntry.MaterialOverrides, WeightedEntry.CullStartDistance, WeightedEntry.CullEndDistance, /*bReverseCulling=*/false);
 
 			// precompute the weights
 			TotalWeight += WeightedEntry.Weight;
 			InstancesAndWeights->CumulativeWeights.Add(TotalWeight);
-			InstancesAndWeights->InstanceListIndices.Add(Index);
-			InstancesAndWeights->ReverseInstanceListIndices.Add(ReverseIndex);
 		}
 
 		// if no weighted entries were collected, discard this InstancesAndWeights
@@ -134,6 +129,13 @@ void UPCGMeshSelectorWeightedByCategory::SelectInstances_Implementation(
 	FPCGMetadataAttribute<FString>* OutAttribute = nullptr;
 	TMap<TSoftObjectPtr<UStaticMesh>, PCGMetadataValueKey> MeshToValueKey;
 
+	FPCGMeshMaterialOverrideHelper MaterialOverrideHelper(Context, bUseAttributeMaterialOverrides, MaterialOverrideAttributes, PointData->Metadata);
+
+	if (!MaterialOverrideHelper.IsValid())
+	{
+		return;
+	}
+
 	if (OutPointData)
 	{
 		check(OutPointData->Metadata);
@@ -192,15 +194,15 @@ void UPCGMeshSelectorWeightedByCategory::SelectInstances_Implementation(
 		int RandomWeightedPick = RandomSource.RandRange(0, TotalWeight - 1);
 
 		int RandomPick = 0;
-		while (RandomPick < InstancesAndWeights->InstanceListIndices.Num() && InstancesAndWeights->CumulativeWeights[RandomPick] <= RandomWeightedPick)
+		while(RandomPick < InstancesAndWeights->MeshInstances.Num() && InstancesAndWeights->CumulativeWeights[RandomPick] <= RandomWeightedPick)
 		{
 			++RandomPick;
 		}
 
-		if (RandomPick < InstancesAndWeights->InstanceListIndices.Num())
+		if(RandomPick < InstancesAndWeights->MeshInstances.Num())
 		{
-			const bool bUseReverse = (Point.Transform.GetDeterminant() < 0);
-			FPCGMeshInstanceList& InstanceList = (bUseReverse ? OutReverseMeshInstances[InstancesAndWeights->ReverseInstanceListIndices[RandomPick]] : OutMeshInstances[InstancesAndWeights->InstanceListIndices[RandomPick]]);
+			const bool bNeedsReverseCulling = (Point.Transform.GetDeterminant() < 0);
+			FPCGMeshInstanceList& InstanceList = PCGMeshSelectorWeighted::GetInstanceList(InstancesAndWeights->MeshInstances[RandomPick], bUseAttributeMaterialOverrides, MaterialOverrideHelper.GetMaterialOverrides(Point.MetadataEntry), bNeedsReverseCulling);
 			InstanceList.Instances.Emplace(Point);
 
 			const TSoftObjectPtr<UStaticMesh>& Mesh = InstanceList.Mesh;
@@ -223,9 +225,15 @@ void UPCGMeshSelectorWeightedByCategory::SelectInstances_Implementation(
 		}
 	}
 
-	// Collapse reverse in the out mesh instances array
-	for (int32 Index = 0; Index < OutReverseMeshInstances.Num(); ++Index)
+	// Collapse to OutMeshInstances
+	for (TPair<PCGMetadataValueKey, FPCGInstancesAndWeights>& Entry : CategoryEntryToInstancesAndWeights)
 	{
-		OutMeshInstances.Emplace(MoveTemp(OutReverseMeshInstances[Index]));
+		for (TArray<FPCGMeshInstanceList>& PickedMeshInstances : Entry.Value.MeshInstances)
+		{
+			for (FPCGMeshInstanceList& PickedMeshInstanceEntry : PickedMeshInstances)
+			{
+				OutMeshInstances.Emplace(MoveTemp(PickedMeshInstanceEntry));
+			}
+		}
 	}
 }
diff --git a/Engine/Plugins/Experimental/PCG/Source/PCG/Public/MeshSelectors/PCGMeshSelectorBase.h b/Engine/Plugins/Experimental/PCG/Source/PCG/Public/MeshSelectors/PCGMeshSelectorBase.h
--- a/Engine/Plugins/Experimental/PCG/Source/PCG/Public/MeshSelectors/PCGMeshSelectorBase.h
+++ b/Engine/Plugins/Experimental/PCG/Source/PCG/Public/MeshSelectors/PCGMeshSelectorBase.h
@@ -3,6 +3,7 @@
 #pragma once
 
 #include "PCGPoint.h"
+#include "Metadata/PCGMetadata.h"
 
 #include "Engine/CollisionProfile.h"
 
@@ -22,8 +23,8 @@ struct FPCGMeshInstanceList
 	GENERATED_BODY()
 
 	FPCGMeshInstanceList() = default;
-	
-	FPCGMeshInstanceList(const TSoftObjectPtr<UStaticMesh>& InMesh, bool bInOverrideCollisionProfile, const FCollisionProfileName& InCollisionProfile, bool bInOverrideMaterials, const TArray<UMaterialInterface*>& InMaterialOverrides, const float InCullStartDistance, const float InCullEndDistance, const bool bInIsLocalToWorldDeterminantNegative)
+
+	FPCGMeshInstanceList(const TSoftObjectPtr<UStaticMesh>& InMesh, bool bInOverrideCollisionProfile, const FCollisionProfileName& InCollisionProfile, bool bInOverrideMaterials, const TArray<TSoftObjectPtr<UMaterialInterface>>& InMaterialOverrides, const float InCullStartDistance, const float InCullEndDistance, const bool bInIsLocalToWorldDeterminantNegative)
 		: Mesh(InMesh), bOverrideCollisionProfile(bInOverrideCollisionProfile), CollisionProfile(InCollisionProfile), bOverrideMaterials(bInOverrideMaterials), MaterialOverrides(InMaterialOverrides), CullStartDistance(InCullStartDistance), CullEndDistance(InCullEndDistance), bIsLocalToWorldDeterminantNegative(bInIsLocalToWorldDeterminantNegative)
 	{}
 
@@ -40,7 +41,7 @@ struct FPCGMeshInstanceList
 	bool bOverrideMaterials = false;
 
 	UPROPERTY(BlueprintReadWrite, EditAnywhere, Category = Settings)
-	TArray<TObjectPtr<UMaterialInterface>> MaterialOverrides;
+	TArray<TSoftObjectPtr<UMaterialInterface>> MaterialOverrides;
 
 	UPROPERTY(BlueprintReadWrite, EditAnywhere, Category = Settings)
     TArray<FPCGPoint> Instances;
@@ -58,6 +59,54 @@ struct FPCGMeshInstanceList
 	bool bIsLocalToWorldDeterminantNegative = false;
 };
 
+UENUM()
+enum class EPCGMeshSelectorMaterialOverrideMode : uint8
+{
+	NoOverride UMETA(Tooltip = "Does not apply any material overrides to the spawned mesh(es)"),
+	StaticOverride UMETA(Tooltip = "Applies the material overrides provided in the Static Material Overrides array"),
+	ByAttributeOverride UMETA(Tooltip = "Applies the materials overrides using the point data attribute(s) specified in the By Attribute Material Overrides array")
+};
+
+/** Struct used to efficiently gather overrides and cache them during instance packing */
+struct FPCGMeshMaterialOverrideHelper
+{
+	// Use this constructor when you have a 1:1 mapping between attributes or static overrides
+	FPCGMeshMaterialOverrideHelper(
+		FPCGContext& InContext,
+		EPCGMeshSelectorMaterialOverrideMode InMaterialOverrideMode,
+		const TArray<TSoftObjectPtr<UMaterialInterface>>& InStaticMaterialOverrides,
+		const TArray<FName>& InMaterialOverrideAttributeNames,
+		const UPCGMetadata* InMetadata);
+
+	// Use this constructor when you have common attribute usage or separate static overrides
+	FPCGMeshMaterialOverrideHelper(
+		FPCGContext& InContext,
+		bool bInByAttributeOverride,
+		const TArray<FName>& InMaterialOverrideAttributeNames,
+		const UPCGMetadata* InMetadata);
+
+	bool IsValid() const { return bIsValid; }
+	bool OverridesMaterials() const { return MaterialOverrideMode != EPCGMeshSelectorMaterialOverrideMode::NoOverride; }
+	const TArray<TSoftObjectPtr<UMaterialInterface>>& GetMaterialOverrides(PCGMetadataEntryKey EntryKey);
+
+	// Cached data
+	TArray<const FPCGMetadataAttribute<FString>*> MaterialAttributes;
+	TArray<TMap<PCGMetadataValueKey, TSoftObjectPtr<UMaterialInterface>>> ValueKeyToOverrideMaterials;
+	TArray<TSoftObjectPtr<UMaterialInterface>> WorkingMaterialOverrides;
+	TArray<TSoftObjectPtr<UMaterialInterface>> EmptyArray;
+
+	// Data needed to perform operations
+	bool bIsValid = false;
+	EPCGMeshSelectorMaterialOverrideMode MaterialOverrideMode = EPCGMeshSelectorMaterialOverrideMode::NoOverride;
+
+	const TArray<TSoftObjectPtr<UMaterialInterface>>& StaticMaterialOverrides;
+	const TArray<FName>& MaterialOverrideAttributeNames;
+	const UPCGMetadata* Metadata = nullptr;
+
+private:
+	void Initialize(FPCGContext& InContext);
+};
+
 UCLASS(Abstract, BlueprintType, Blueprintable, ClassGroup = (Procedural))
 class PCG_API UPCGMeshSelectorBase : public UObject 
 {
@@ -81,16 +130,16 @@ public:
 
 	/** Searches OutInstanceLists for an InstanceList matching the given parameters. If nothing is found, creates a new InstanceList and adds to OutInstanceLists. Returns index of the matching instance list. */
 	UFUNCTION(BlueprintCallable, Category = MeshSelection)
-	int32 FindOrAddInstanceList(
+	static int32 FindOrAddInstanceList(
 		TArray<FPCGMeshInstanceList>& OutInstanceLists,
 		const TSoftObjectPtr<UStaticMesh>& Mesh,
 		bool bOverrideCollisionProfile,
 		const FCollisionProfileName& CollisionProfile,
 		bool bOverrideMaterials,
-		const TArray<UMaterialInterface*>& MaterialOverrides,
+		const TArray<TSoftObjectPtr<UMaterialInterface>>& MaterialOverrides,
 		const float InCullStartDistance,
 		const float InCullEndDistance, 
-		const bool bInIsLocalToWorldDeterminantNegative) const;
+		const bool bInIsLocalToWorldDeterminantNegative);
 };
 
 #if UE_ENABLE_INCLUDE_ORDER_DEPRECATED_IN_5_2
diff --git a/Engine/Plugins/Experimental/PCG/Source/PCG/Public/MeshSelectors/PCGMeshSelectorByAttribute.h b/Engine/Plugins/Experimental/PCG/Source/PCG/Public/MeshSelectors/PCGMeshSelectorByAttribute.h
--- a/Engine/Plugins/Experimental/PCG/Source/PCG/Public/MeshSelectors/PCGMeshSelectorByAttribute.h
+++ b/Engine/Plugins/Experimental/PCG/Source/PCG/Public/MeshSelectors/PCGMeshSelectorByAttribute.h
@@ -13,6 +13,10 @@ class PCG_API UPCGMeshSelectorByAttribute : public UPCGMeshSelectorBase
 	GENERATED_BODY()
 
 public:
+	// ~Begin UObject interface
+	void PostLoad() override;
+	// ~End UObject interface
+
 	virtual void SelectInstances_Implementation(
 		UPARAM(ref) FPCGContext& Context,
 		const UPCGStaticMeshSpawnerSettings* Settings,
@@ -27,14 +31,22 @@ public:
 	UPROPERTY(BlueprintReadWrite, EditAnywhere, Category = Settings)
 	bool bOverrideCollisionProfile = false;
 
-	UPROPERTY(BlueprintReadWrite, EditAnywhere, Category = Settings)
+	UPROPERTY(BlueprintReadWrite, EditAnywhere, Category = Settings, meta = (EditCondition = "bOverrideCollisionProfile", EditConditionHides))
 	FCollisionProfileName CollisionProfile = UCollisionProfile::NoCollision_ProfileName;
 
-	UPROPERTY(BlueprintReadWrite, EditAnywhere, Category = Settings)
-	bool bOverrideMaterials = false;
+#if WITH_EDITORONLY_DATA
+	UPROPERTY()
+	bool bOverrideMaterials_DEPRECATED = false;
+#endif
 
 	UPROPERTY(BlueprintReadWrite, EditAnywhere, Category = Settings)
-	TArray<TObjectPtr<UMaterialInterface>> MaterialOverrides;
+	EPCGMeshSelectorMaterialOverrideMode MaterialOverrideMode = EPCGMeshSelectorMaterialOverrideMode::NoOverride;
+
+	UPROPERTY(BlueprintReadWrite, EditAnywhere, DisplayName = "Static Material Overrides", Category = Settings, meta = (EditCondition = "MaterialOverrideMode==EPCGMeshSelectorMaterialOverrideMode::StaticOverride", EditConditionHides))
+	TArray<TSoftObjectPtr<UMaterialInterface>> MaterialOverrides;
+
+	UPROPERTY(BlueprintReadWrite, EditAnywhere, DisplayName = "By Attribute Material Overrides", Category = Settings, meta = (EditCondition = "MaterialOverrideMode==EPCGMeshSelectorMaterialOverrideMode::ByAttributeOverride", EditConditionHides))
+	TArray<FName> MaterialOverrideAttributes;
 
 	/** Distance at which instances begin to fade. */
 	UPROPERTY(BlueprintReadWrite, EditAnywhere, Category = Settings)
diff --git a/Engine/Plugins/Experimental/PCG/Source/PCG/Public/MeshSelectors/PCGMeshSelectorWeighted.h b/Engine/Plugins/Experimental/PCG/Source/PCG/Public/MeshSelectors/PCGMeshSelectorWeighted.h
--- a/Engine/Plugins/Experimental/PCG/Source/PCG/Public/MeshSelectors/PCGMeshSelectorWeighted.h
+++ b/Engine/Plugins/Experimental/PCG/Source/PCG/Public/MeshSelectors/PCGMeshSelectorWeighted.h
@@ -4,9 +4,17 @@
 
 #include "PCGMeshSelectorBase.h"
 
-
 #include "PCGMeshSelectorWeighted.generated.h"
 
+namespace PCGMeshSelectorWeighted
+{
+	FPCGMeshInstanceList& GetInstanceList(
+		TArray<FPCGMeshInstanceList>& InstanceLists,
+		bool bUseAttributeMaterialOverrides,
+		const TArray<TSoftObjectPtr<UMaterialInterface>>& InMaterialOverrides,
+		bool bInIsLocalToWorldDeterminantNegative);
+}
+
 USTRUCT(BlueprintType)
 struct PCG_API FPCGMeshSelectorWeightedEntry
 {
@@ -24,14 +32,14 @@ struct PCG_API FPCGMeshSelectorWeightedEntry
 	UPROPERTY(BlueprintReadWrite, EditAnywhere, Category = Settings)
 	bool bOverrideCollisionProfile = false;
 
-	UPROPERTY(BlueprintReadWrite, EditAnywhere, Category = Settings)
+	UPROPERTY(BlueprintReadWrite, EditAnywhere, Category = Settings, meta = (EditCondition = "bOverrideCollisionProfile", EditConditionHides))
 	FCollisionProfileName CollisionProfile = UCollisionProfile::NoCollision_ProfileName;
 
-	UPROPERTY(BlueprintReadWrite, EditAnywhere, Category = Settings)
+	UPROPERTY(BlueprintReadWrite, EditAnywhere, Category = Settings, meta = (InlineEditConditionToggle))
 	bool bOverrideMaterials = false;
 
-	UPROPERTY(BlueprintReadWrite, EditAnywhere, Category = Settings)
-	TArray<TObjectPtr<UMaterialInterface>> MaterialOverrides;
+	UPROPERTY(BlueprintReadWrite, EditAnywhere, Category = Settings, DisplayName = "Static Material Overrides", meta = (EditCondition = "bOverrideMaterials"))
+	TArray<TSoftObjectPtr<UMaterialInterface>> MaterialOverrides;
 
 	/** Distance at which instances begin to fade. */
 	UPROPERTY(BlueprintReadWrite, EditAnywhere, Category = Settings)
@@ -61,4 +69,10 @@ public:
 public:
 	UPROPERTY(BlueprintReadWrite, EditAnywhere, Category = Settings)
 	TArray<FPCGMeshSelectorWeightedEntry> MeshEntries;
+
+	UPROPERTY(BlueprintReadWrite, EditAnywhere, Category = Settings, meta = (InlineEditConditionToggle))
+	bool bUseAttributeMaterialOverrides = false;
+
+	UPROPERTY(BlueprintReadWrite, EditAnywhere, DisplayName = "By Attribute Material Overrides", Category = Settings, meta = (EditCondition = "bUseAttributeMaterialOverrides"))
+	TArray<FName> MaterialOverrideAttributes;
 };
diff --git a/Engine/Plugins/Experimental/PCG/Source/PCG/Public/MeshSelectors/PCGMeshSelectorWeightedByCategory.h b/Engine/Plugins/Experimental/PCG/Source/PCG/Public/MeshSelectors/PCGMeshSelectorWeightedByCategory.h
--- a/Engine/Plugins/Experimental/PCG/Source/PCG/Public/MeshSelectors/PCGMeshSelectorWeightedByCategory.h
+++ b/Engine/Plugins/Experimental/PCG/Source/PCG/Public/MeshSelectors/PCGMeshSelectorWeightedByCategory.h
@@ -47,4 +47,10 @@ public:
 
 	UPROPERTY(BlueprintReadWrite, EditAnywhere, Category = Settings)
 	TArray<FPCGWeightedByCategoryEntryList> Entries;
+
+	UPROPERTY(BlueprintReadWrite, EditAnywhere, Category = Settings, meta = (InlineEditConditionToggle))
+	bool bUseAttributeMaterialOverrides = false;
+
+	UPROPERTY(BlueprintReadWrite, EditAnywhere, DisplayName = "By Attribute Material Overrides", Category = Settings, meta = (EditCondition = "bUseAttributeMaterialOverrides"))
+	TArray<FName> MaterialOverrideAttributes;
 };
