[PCG] Fix static analysis warning that popped up after a recent commit (this issue predates that commit)

diff --git a/Engine/Plugins/Experimental/PCG/Source/PCG/Private/PCGComponent.cpp b/Engine/Plugins/Experimental/PCG/Source/PCG/Private/PCGComponent.cpp
--- a/Engine/Plugins/Experimental/PCG/Source/PCG/Private/PCGComponent.cpp
+++ b/Engine/Plugins/Experimental/PCG/Source/PCG/Private/PCGComponent.cpp
@@ -608,6 +608,12 @@ bool UPCGComponent::MoveResourcesToNewActor(AActor* InNewActor, bool bCreateChil
 		return false;
 	}
 
+	if (!GetOwner())
+	{
+		UE_LOG(LogPCG, Error, TEXT("[UPCGComponent::MoveResourcesToNewActor] Owner is null, child actor not created."));
+		return false;
+	}
+
 	check(InNewActor);
 	AActor* NewActor = InNewActor;
 
