[PCG] Fix crash if user attempts to project a high dimensional source onto a low dimensional target.

Also robust against projection onto null target.

Failed projection => return projection source not-projected.

Changing the signature does not appear to break existing links in BP.

diff --git a/Engine/Plugins/Experimental/PCG/Source/PCG/Private/Data/PCGSpatialData.cpp b/Engine/Plugins/Experimental/PCG/Source/PCG/Private/Data/PCGSpatialData.cpp
--- a/Engine/Plugins/Experimental/PCG/Source/PCG/Private/Data/PCGSpatialData.cpp
+++ b/Engine/Plugins/Experimental/PCG/Source/PCG/Private/Data/PCGSpatialData.cpp
@@ -121,8 +121,21 @@ UPCGIntersectionData* UPCGSpatialData::IntersectWith(const UPCGSpatialData* InOt
 	return IntersectionData;
 }
 
-UPCGProjectionData* UPCGSpatialData::ProjectOn(const UPCGSpatialData* InOther, const FPCGProjectionParams& InParams) const
+UPCGSpatialData* UPCGSpatialData::ProjectOn(const UPCGSpatialData* InOther, const FPCGProjectionParams& InParams) const
 {
+	// Check necessary conditions. Fail to project -> return copy of projection source, i.e. projection not performed.
+	if (!InOther)
+	{
+		UE_LOG(LogPCG, Warning, TEXT("No projection target specified, no projection will occur"));
+		return DuplicateData();
+	}
+
+	if (GetDimension() > InOther->GetDimension())
+	{
+		UE_LOG(LogPCG, Error, TEXT("Dimension of projection source (%d) must be less than or equal to that of the projection target (%d)"), GetDimension(), InOther->GetDimension());
+		return DuplicateData();
+	}
+
 	UPCGProjectionData* ProjectionData = NewObject<UPCGProjectionData>();
 	ProjectionData->Initialize(this, InOther, InParams);
 
diff --git a/Engine/Plugins/Experimental/PCG/Source/PCG/Private/Data/PCGSplineData.cpp b/Engine/Plugins/Experimental/PCG/Source/PCG/Private/Data/PCGSplineData.cpp
--- a/Engine/Plugins/Experimental/PCG/Source/PCG/Private/Data/PCGSplineData.cpp
+++ b/Engine/Plugins/Experimental/PCG/Source/PCG/Private/Data/PCGSplineData.cpp
@@ -110,9 +110,9 @@ bool UPCGSplineData::SamplePoint(const FTransform& InTransform, const FBox& InBo
 	}
 }
 
-UPCGProjectionData* UPCGSplineData::ProjectOn(const UPCGSpatialData* InOther, const FPCGProjectionParams& InParams) const
+UPCGSpatialData* UPCGSplineData::ProjectOn(const UPCGSpatialData* InOther, const FPCGProjectionParams& InParams) const
 {
-	if(InOther->GetDimension() == 2)
+	if (InOther->GetDimension() == 2)
 	{
 		UPCGSplineProjectionData* SplineProjectionData = NewObject<UPCGSplineProjectionData>();
 		SplineProjectionData->Initialize(this, InOther, InParams);
diff --git a/Engine/Plugins/Experimental/PCG/Source/PCG/Private/Elements/PCGProjectionElement.cpp b/Engine/Plugins/Experimental/PCG/Source/PCG/Private/Elements/PCGProjectionElement.cpp
--- a/Engine/Plugins/Experimental/PCG/Source/PCG/Private/Elements/PCGProjectionElement.cpp
+++ b/Engine/Plugins/Experimental/PCG/Source/PCG/Private/Elements/PCGProjectionElement.cpp
@@ -78,7 +78,7 @@ bool FPCGProjectionElement::ExecuteInternal(FPCGContext* Context) const
 			continue;
 		}
 
-		UPCGProjectionData* ProjectionData = ProjectionSource->ProjectOn(ProjectionTarget, ProjectionParams);
+		UPCGSpatialData* ProjectionData = ProjectionSource->ProjectOn(ProjectionTarget, ProjectionParams);
 #if WITH_EDITOR
 		ProjectionData->bKeepZeroDensityPoints = bKeepZeroDensityPoints;
 #endif
diff --git a/Engine/Plugins/Experimental/PCG/Source/PCG/Public/Data/PCGSpatialData.h b/Engine/Plugins/Experimental/PCG/Source/PCG/Public/Data/PCGSpatialData.h
--- a/Engine/Plugins/Experimental/PCG/Source/PCG/Public/Data/PCGSpatialData.h
+++ b/Engine/Plugins/Experimental/PCG/Source/PCG/Public/Data/PCGSpatialData.h
@@ -96,9 +96,9 @@ public:
 	UFUNCTION(BlueprintCallable, Category = SpatialData)
 	virtual UPCGIntersectionData* IntersectWith(const UPCGSpatialData* InOther) const;
 
-	/** Returns a specialized data to project this on another data of equal or higher dimension */
+	/** Returns a specialized data to project this on another data of equal or higher dimension. Returns copy of this data if projection fails. */
 	UFUNCTION(BlueprintCallable, Category = SpatialData, meta=(AutoCreateRefTerm="InParams"))
-	virtual UPCGProjectionData* ProjectOn(const UPCGSpatialData* InOther, const FPCGProjectionParams& InParams = FPCGProjectionParams()) const;
+	virtual UPCGSpatialData* ProjectOn(const UPCGSpatialData* InOther, const FPCGProjectionParams& InParams = FPCGProjectionParams()) const;
 
 	/** Returns a specialized data to union this with another data */
 	UFUNCTION(BlueprintCallable, Category = SpatialData)
diff --git a/Engine/Plugins/Experimental/PCG/Source/PCG/Public/Data/PCGSplineData.h b/Engine/Plugins/Experimental/PCG/Source/PCG/Public/Data/PCGSplineData.h
--- a/Engine/Plugins/Experimental/PCG/Source/PCG/Public/Data/PCGSplineData.h
+++ b/Engine/Plugins/Experimental/PCG/Source/PCG/Public/Data/PCGSplineData.h
@@ -40,7 +40,7 @@ public:
 	//~Begin UPCGSpatialData interface
 	virtual FBox GetBounds() const override;
 	virtual bool SamplePoint(const FTransform& Transform, const FBox& Bounds, FPCGPoint& OutPoint, UPCGMetadata* OutMetadata) const override;
-	virtual UPCGProjectionData* ProjectOn(const UPCGSpatialData* InOther, const FPCGProjectionParams& InParams = FPCGProjectionParams()) const override;
+	virtual UPCGSpatialData* ProjectOn(const UPCGSpatialData* InOther, const FPCGProjectionParams& InParams = FPCGProjectionParams()) const override;
 protected:
 	virtual UPCGSpatialData* CopyInternal() const override;
 	//~End UPCGSpatialData interface
