[PCG] Fix localization warning

diff --git a/Engine/Plugins/Experimental/PCG/Source/PCGEditor/Private/SPCGEditorGraphAttributeListView.cpp b/Engine/Plugins/Experimental/PCG/Source/PCGEditor/Private/SPCGEditorGraphAttributeListView.cpp
--- a/Engine/Plugins/Experimental/PCG/Source/PCGEditor/Private/SPCGEditorGraphAttributeListView.cpp
+++ b/Engine/Plugins/Experimental/PCG/Source/PCGEditor/Private/SPCGEditorGraphAttributeListView.cpp
@@ -695,7 +695,7 @@ void SPCGEditorGraphAttributeListView::RefreshAttributeList()
 			else
 			{
 				// If cache debugging enabled, write CRC to help diagnose missed-dependency issues
-				InfoTextBlock->SetText(FText::Format(LOCTEXT("PointInfoTextBlockFmt", "Number of points: {0}, CRC: {1}"), NumPoints, InspectionData->Crc.GetValue()));
+				InfoTextBlock->SetText(FText::Format(LOCTEXT("PointInfoTextBlockWithCrcFmt", "Number of points: {0}, CRC: {1}"), NumPoints, InspectionData->Crc.GetValue()));
 			}
 		}
 	}
