C++ port of BP NormalToDensity

diff --git a/Engine/Plugins/Experimental/PCG/Source/PCG/Private/Elements/PCGNormalToDensity.cpp b/Engine/Plugins/Experimental/PCG/Source/PCG/Private/Elements/PCGNormalToDensity.cpp
new file mode 100644
--- /dev/null
+++ b/Engine/Plugins/Experimental/PCG/Source/PCG/Private/Elements/PCGNormalToDensity.cpp
@@ -0,0 +1,146 @@
+// Copyright Epic Games, Inc. All Rights Reserved.
+
+#include "Elements/PCGNormalToDensity.h"
+#include "Helpers/PCGSettingsHelpers.h"
+#include "Data/PCGPointData.h"
+#include "Kismet/KismetMathLibrary.h"
+#include "PCGContext.h"
+#include "PCGPin.h"
+
+TArray<FPCGPinProperties> UPCGNormalToDensitySettings::InputPinProperties() const
+{
+	TArray<FPCGPinProperties> PinProperties;
+	PinProperties.Emplace(PCGPinConstants::DefaultInputLabel, EPCGDataType::Spatial);
+	PinProperties.Emplace(PCGPinConstants::DefaultParamsLabel, EPCGDataType::Param, /*bInAllowMultipleConnections*/ false);
+
+	return PinProperties;
+}
+
+TArray<FPCGPinProperties> UPCGNormalToDensitySettings::OutputPinProperties() const
+{
+	TArray<FPCGPinProperties> PinProperties;
+	PinProperties.Emplace(PCGPinConstants::DefaultOutputLabel, EPCGDataType::Spatial);
+
+	return PinProperties;
+}
+
+FPCGElementPtr UPCGNormalToDensitySettings::CreateElement() const
+{
+	return MakeShared<FPCGNormalToDensityElement>();
+}
+
+bool FPCGNormalToDensityElement::ExecuteInternal(FPCGContext* Context) const
+{
+	TRACE_CPUPROFILER_EVENT_SCOPE(FPCGNormalToDensityElement::Execute);
+
+	const UPCGNormalToDensitySettings* Settings = Context->GetInputSettings<UPCGNormalToDensitySettings>();
+	check(Settings);
+
+	UPCGParamData* Params = Context->InputData.GetParams();
+
+	const FVector Normal = PCGSettingsHelpers::GetValue(GET_MEMBER_NAME_CHECKED(UPCGNormalToDensitySettings, Normal), Settings->Normal, Params);
+	const double Offset = PCGSettingsHelpers::GetValue(GET_MEMBER_NAME_CHECKED(UPCGNormalToDensitySettings, Offset), Settings->Offset, Params);
+	const double Strength = PCGSettingsHelpers::GetValue(GET_MEMBER_NAME_CHECKED(UPCGNormalToDensitySettings, Strength), Settings->Strength, Params);
+	const PCGNormalToDensityMode DensityMode = PCGSettingsHelpers::GetValue(GET_MEMBER_NAME_CHECKED(UPCGNormalToDensitySettings, DensityMode), Settings->DensityMode, Params);
+
+	TArray<FPCGTaggedData> Inputs = Context->InputData.GetInputsByPin(PCGPinConstants::DefaultInputLabel);
+	TArray<FPCGTaggedData>& Outputs = Context->OutputData.TaggedData;
+
+	const double InvStrength = 1.0/FMath::Max(0.0001, Strength);
+
+	const auto CalcValue = [Normal, Offset, InvStrength](const FPCGPoint& InPoint)
+	{
+		const FVector Up = InPoint.Transform.GetScaledAxis(EAxis::Z);
+		return FMath::Pow(FMath::Clamp(Up.Dot(Normal) + Offset, 0.0, 1.0), InvStrength);
+	};
+
+	switch (DensityMode)
+	{
+	case PCGNormalToDensityMode::Set:
+		ProcessPoints(Context, Inputs, Outputs, [CalcValue](const FPCGPoint& InPoint, FPCGPoint& OutPoint)->bool{
+
+			OutPoint = InPoint;
+
+			OutPoint.Density = CalcValue(InPoint);
+
+			return true;
+
+		});
+		break;
+
+	case PCGNormalToDensityMode::Minimum:
+		ProcessPoints(Context, Inputs, Outputs, [CalcValue](const FPCGPoint& InPoint, FPCGPoint& OutPoint)->bool{
+
+			OutPoint = InPoint;
+
+			OutPoint.Density = FMath::Min(InPoint.Density, CalcValue(InPoint));
+
+			return true;
+
+		});
+		break;
+
+	case PCGNormalToDensityMode::Maximum:
+		ProcessPoints(Context, Inputs, Outputs, [CalcValue](const FPCGPoint& InPoint, FPCGPoint& OutPoint)->bool{
+
+			OutPoint = InPoint;
+
+			OutPoint.Density = FMath::Max(InPoint.Density, CalcValue(InPoint));
+
+			return true;
+
+		});
+		break;
+
+	case PCGNormalToDensityMode::Add:
+		ProcessPoints(Context, Inputs, Outputs, [CalcValue](const FPCGPoint& InPoint, FPCGPoint& OutPoint)->bool{
+
+			OutPoint = InPoint;
+
+			OutPoint.Density += CalcValue(InPoint);
+
+			return true;
+
+		});
+		break;
+
+	case PCGNormalToDensityMode::Subtract:
+		ProcessPoints(Context, Inputs, Outputs, [CalcValue](const FPCGPoint& InPoint, FPCGPoint& OutPoint)->bool{
+
+			OutPoint = InPoint;
+
+			OutPoint.Density -= CalcValue(InPoint);
+
+			return true;
+
+		});
+		break;
+
+	case PCGNormalToDensityMode::Multiply:
+		ProcessPoints(Context, Inputs, Outputs, [CalcValue](const FPCGPoint& InPoint, FPCGPoint& OutPoint)->bool{
+
+			OutPoint = InPoint;
+
+			OutPoint.Density *= CalcValue(InPoint);
+
+			return true;
+
+		});
+		break;
+
+	case PCGNormalToDensityMode::Divide:
+		ProcessPoints(Context, Inputs, Outputs, [CalcValue](const FPCGPoint& InPoint, FPCGPoint& OutPoint)->bool{
+
+			OutPoint = InPoint;
+
+			OutPoint.Density = UKismetMathLibrary::SafeDivide(OutPoint.Density, CalcValue(InPoint));
+
+			return true;
+
+		});
+		break;
+	}
+
+	return true;
+
+}
diff --git a/Engine/Plugins/Experimental/PCG/Source/PCG/Private/Tests/Elements/PCGNormalToDensityTest.cpp b/Engine/Plugins/Experimental/PCG/Source/PCG/Private/Tests/Elements/PCGNormalToDensityTest.cpp
new file mode 100644
--- /dev/null
+++ b/Engine/Plugins/Experimental/PCG/Source/PCG/Private/Tests/Elements/PCGNormalToDensityTest.cpp
@@ -0,0 +1,451 @@
+// Copyright Epic Games, Inc. All Rights Reserved.
+
+#include "Tests/Determinism/PCGDeterminismTestsCommon.h"
+
+#include "Data/PCGSpatialData.h"
+#include "PCGComponent.h"
+
+#include "Data/PCGPointData.h"
+#include "Elements/PCGNormalToDensity.h"
+#include "PCGContext.h"
+
+IMPLEMENT_CUSTOM_SIMPLE_AUTOMATION_TEST(FPCGNormalToDensityTest_Set, FPCGTestBaseClass, "pcg.tests.NormalToDensity.Set", PCGTestsCommon::TestFlags)
+
+bool FPCGNormalToDensityTest_Set::RunTest(const FString& Parameters)
+{
+	PCGTestsCommon::FTestData TestData(PCGDeterminismTests::Defaults::Seed);
+	PCGDeterminismTests::GenerateSettings<UPCGNormalToDensitySettings>(TestData);
+	UPCGNormalToDensitySettings* Settings = CastChecked<UPCGNormalToDensitySettings>(TestData.Settings);
+	FPCGElementPtr Element = TestData.Settings->GetElement();
+
+	Settings->DensityMode = PCGNormalToDensityMode::Set;
+
+	{
+		TObjectPtr<UPCGPointData> PointData = PCGTestsCommon::CreatePointData();
+		FPCGPoint& Point = PointData->GetMutablePoints()[0];
+
+		Point.Transform.SetRotation(FQuat::MakeFromEuler(FVector(45.0, 0, 0)));
+		Point.Density = 0.5f;
+
+		FPCGTaggedData& TaggedData = TestData.InputData.TaggedData.Emplace_GetRef(FPCGTaggedData());
+		TaggedData.Data = PointData;
+		TaggedData.Pin = PCGPinConstants::DefaultInputLabel;
+	}
+
+	TUniquePtr<FPCGContext> Context(Element->Initialize(TestData.InputData, TestData.TestPCGComponent, nullptr));
+	Context->NumAvailableTasks = 1;
+
+	while (!Element->Execute(Context.Get())) {}
+
+	// test our point data
+	const TArray<FPCGTaggedData>& Outputs = Context->OutputData.GetInputs();
+
+	UTEST_EQUAL("Output count", Outputs.Num(), 1);
+
+	UPCGSpatialData* OutputData = Cast<UPCGSpatialData>(Outputs[0].Data);
+
+	UTEST_NOT_NULL("Output data", OutputData);
+
+	const UPCGPointData* OutputPointData = OutputData->ToPointData(Context.Get());
+	UTEST_NOT_NULL("Output point data", OutputPointData);
+
+	const TArray<FPCGPoint>& OutputPoints = OutputPointData->GetPoints();
+
+	UTEST_EQUAL("Output point count", OutputPoints.Num(), 1);
+
+	UTEST_EQUAL_TOLERANCE("Output point density", OutputPoints[0].Density, 0.707f, 0.001f);
+
+	return true;
+}
+
+IMPLEMENT_CUSTOM_SIMPLE_AUTOMATION_TEST(FPCGNormalToDensityTest_CustomNormal, FPCGTestBaseClass, "pcg.tests.NormalToDensity.CustomNormal", PCGTestsCommon::TestFlags)
+
+bool FPCGNormalToDensityTest_CustomNormal::RunTest(const FString& Parameters)
+{
+	PCGTestsCommon::FTestData TestData(PCGDeterminismTests::Defaults::Seed);
+	PCGDeterminismTests::GenerateSettings<UPCGNormalToDensitySettings>(TestData);
+	UPCGNormalToDensitySettings* Settings = CastChecked<UPCGNormalToDensitySettings>(TestData.Settings);
+	FPCGElementPtr Element = TestData.Settings->GetElement();
+
+	Settings->DensityMode = PCGNormalToDensityMode::Set;
+	Settings->Normal = FVector(0.707, 0, 0.707);
+
+	{
+		TObjectPtr<UPCGPointData> PointData = PCGTestsCommon::CreatePointData();
+		FPCGPoint& Point = PointData->GetMutablePoints()[0];
+		Point.Density = 0.5f;
+
+		FPCGTaggedData& TaggedData = TestData.InputData.TaggedData.Emplace_GetRef(FPCGTaggedData());
+		TaggedData.Data = PointData;
+		TaggedData.Pin = PCGPinConstants::DefaultInputLabel;
+	}
+
+	TUniquePtr<FPCGContext> Context(Element->Initialize(TestData.InputData, TestData.TestPCGComponent, nullptr));
+	Context->NumAvailableTasks = 1;
+
+	while (!Element->Execute(Context.Get())) {}
+
+	// test our point data
+	const TArray<FPCGTaggedData>& Outputs = Context->OutputData.GetInputs();
+
+	UTEST_EQUAL("Output count", Outputs.Num(), 1);
+
+	UPCGSpatialData* OutputData = Cast<UPCGSpatialData>(Outputs[0].Data);
+
+	UTEST_NOT_NULL("Output data", OutputData);
+
+	const UPCGPointData* OutputPointData = OutputData->ToPointData(Context.Get());
+	UTEST_NOT_NULL("Output point data", OutputPointData);
+
+	const TArray<FPCGPoint>& OutputPoints = OutputPointData->GetPoints();
+
+	UTEST_EQUAL("Output point count", OutputPoints.Num(), 1);
+
+	UTEST_EQUAL_TOLERANCE("Output point density", OutputPoints[0].Density, 0.707f, 0.001f);
+
+	return true;
+}
+
+IMPLEMENT_CUSTOM_SIMPLE_AUTOMATION_TEST(FPCGNormalToDensityTest_Strength, FPCGTestBaseClass, "pcg.tests.NormalToDensity.Strength", PCGTestsCommon::TestFlags)
+
+bool FPCGNormalToDensityTest_Strength::RunTest(const FString& Parameters)
+{
+	PCGTestsCommon::FTestData TestData(PCGDeterminismTests::Defaults::Seed);
+	PCGDeterminismTests::GenerateSettings<UPCGNormalToDensitySettings>(TestData);
+	UPCGNormalToDensitySettings* Settings = CastChecked<UPCGNormalToDensitySettings>(TestData.Settings);
+	FPCGElementPtr Element = TestData.Settings->GetElement();
+
+	Settings->DensityMode = PCGNormalToDensityMode::Set;
+	Settings->Strength = 2;
+	Settings->Normal = FVector(0.866, 0, 0.5); // so the dot with up is 0.5
+
+	{
+		TObjectPtr<UPCGPointData> PointData = PCGTestsCommon::CreatePointData();
+		FPCGPoint& Point = PointData->GetMutablePoints()[0];
+		Point.Density = 0.5f;
+		FPCGTaggedData& TaggedData = TestData.InputData.TaggedData.Emplace_GetRef(FPCGTaggedData());
+		TaggedData.Data = PointData;
+		TaggedData.Pin = PCGPinConstants::DefaultInputLabel;
+	}
+
+	TUniquePtr<FPCGContext> Context(Element->Initialize(TestData.InputData, TestData.TestPCGComponent, nullptr));
+	Context->NumAvailableTasks = 1;
+
+	while (!Element->Execute(Context.Get())) {}
+
+	// test our point data
+	const TArray<FPCGTaggedData>& Outputs = Context->OutputData.GetInputs();
+
+	UTEST_EQUAL("Output count", Outputs.Num(), 1);
+
+	UPCGSpatialData* OutputData = Cast<UPCGSpatialData>(Outputs[0].Data);
+
+	UTEST_NOT_NULL("Output data", OutputData);
+
+	const UPCGPointData* OutputPointData = OutputData->ToPointData(Context.Get());
+	UTEST_NOT_NULL("Output point data", OutputPointData);
+
+	const TArray<FPCGPoint>& OutputPoints = OutputPointData->GetPoints();
+
+	UTEST_EQUAL("Output point count", OutputPoints.Num(), 1);
+
+	UTEST_EQUAL_TOLERANCE("Output point density", OutputPoints[0].Density, 0.707f, 0.001f);
+
+	return true;
+}
+
+IMPLEMENT_CUSTOM_SIMPLE_AUTOMATION_TEST(FPCGNormalToDensityTest_Minimum, FPCGTestBaseClass, "pcg.tests.NormalToDensity.Minimum", PCGTestsCommon::TestFlags)
+
+bool FPCGNormalToDensityTest_Minimum::RunTest(const FString& Parameters)
+{
+	PCGTestsCommon::FTestData TestData(PCGDeterminismTests::Defaults::Seed);
+	PCGDeterminismTests::GenerateSettings<UPCGNormalToDensitySettings>(TestData);
+	UPCGNormalToDensitySettings* Settings = CastChecked<UPCGNormalToDensitySettings>(TestData.Settings);
+	FPCGElementPtr Element = TestData.Settings->GetElement();
+
+	Settings->DensityMode = PCGNormalToDensityMode::Minimum;
+
+	{
+		TObjectPtr<UPCGPointData> PointData = PCGTestsCommon::CreatePointData();
+		FPCGPoint& Point = PointData->GetMutablePoints()[0];
+
+		Point.Transform.SetRotation(FQuat::MakeFromEuler(FVector(45.0, 0, 0)));
+		Point.Density = 0.5f;
+
+		FPCGTaggedData& TaggedData = TestData.InputData.TaggedData.Emplace_GetRef(FPCGTaggedData());
+		TaggedData.Data = PointData;
+		TaggedData.Pin = PCGPinConstants::DefaultInputLabel;
+	}
+
+	TUniquePtr<FPCGContext> Context(Element->Initialize(TestData.InputData, TestData.TestPCGComponent, nullptr));
+	Context->NumAvailableTasks = 1;
+
+	while (!Element->Execute(Context.Get())) {}
+
+	// test our point data
+	const TArray<FPCGTaggedData>& Outputs = Context->OutputData.GetInputs();
+
+	UTEST_EQUAL("Output count", Outputs.Num(), 1);
+
+	UPCGSpatialData* OutputData = Cast<UPCGSpatialData>(Outputs[0].Data);
+
+	UTEST_NOT_NULL("Output data", OutputData);
+
+	const UPCGPointData* OutputPointData = OutputData->ToPointData(Context.Get());
+	UTEST_NOT_NULL("Output point data", OutputPointData);
+
+	const TArray<FPCGPoint>& OutputPoints = OutputPointData->GetPoints();
+
+	UTEST_EQUAL("Output point count", OutputPoints.Num(), 1);
+
+	UTEST_EQUAL_TOLERANCE("Output point density", OutputPoints[0].Density, 0.5f, 0.001f);
+
+	return true;
+}
+
+IMPLEMENT_CUSTOM_SIMPLE_AUTOMATION_TEST(FPCGNormalToDensityTest_Maximum, FPCGTestBaseClass, "pcg.tests.NormalToDensity.Maximum", PCGTestsCommon::TestFlags)
+
+bool FPCGNormalToDensityTest_Maximum::RunTest(const FString& Parameters)
+{
+	PCGTestsCommon::FTestData TestData(PCGDeterminismTests::Defaults::Seed);
+	PCGDeterminismTests::GenerateSettings<UPCGNormalToDensitySettings>(TestData);
+	UPCGNormalToDensitySettings* Settings = CastChecked<UPCGNormalToDensitySettings>(TestData.Settings);
+	FPCGElementPtr Element = TestData.Settings->GetElement();
+
+	Settings->DensityMode = PCGNormalToDensityMode::Maximum;
+
+	{
+		TObjectPtr<UPCGPointData> PointData = PCGTestsCommon::CreatePointData();
+		FPCGPoint& Point = PointData->GetMutablePoints()[0];
+
+		Point.Transform.SetRotation(FQuat::MakeFromEuler(FVector(45.0, 0, 0)));
+		Point.Density = 0.5f;
+
+		FPCGTaggedData& TaggedData = TestData.InputData.TaggedData.Emplace_GetRef(FPCGTaggedData());
+		TaggedData.Data = PointData;
+		TaggedData.Pin = PCGPinConstants::DefaultInputLabel;
+	}
+
+	TUniquePtr<FPCGContext> Context(Element->Initialize(TestData.InputData, TestData.TestPCGComponent, nullptr));
+	Context->NumAvailableTasks = 1;
+
+	while (!Element->Execute(Context.Get())) {}
+
+	// test our point data
+	const TArray<FPCGTaggedData>& Outputs = Context->OutputData.GetInputs();
+
+	UTEST_EQUAL("Output count", Outputs.Num(), 1);
+
+	UPCGSpatialData* OutputData = Cast<UPCGSpatialData>(Outputs[0].Data);
+
+	UTEST_NOT_NULL("Output data", OutputData);
+
+	const UPCGPointData* OutputPointData = OutputData->ToPointData(Context.Get());
+	UTEST_NOT_NULL("Output point data", OutputPointData);
+
+	const TArray<FPCGPoint>& OutputPoints = OutputPointData->GetPoints();
+
+	UTEST_EQUAL("Output point count", OutputPoints.Num(), 1);
+
+	UTEST_EQUAL_TOLERANCE("Output point density", OutputPoints[0].Density, 0.707f, 0.001f);
+
+	return true;
+}
+
+
+IMPLEMENT_CUSTOM_SIMPLE_AUTOMATION_TEST(FPCGNormalToDensityTest_Add, FPCGTestBaseClass, "pcg.tests.NormalToDensity.Add", PCGTestsCommon::TestFlags)
+
+bool FPCGNormalToDensityTest_Add::RunTest(const FString& Parameters)
+{
+	PCGTestsCommon::FTestData TestData(PCGDeterminismTests::Defaults::Seed);
+	PCGDeterminismTests::GenerateSettings<UPCGNormalToDensitySettings>(TestData);
+	UPCGNormalToDensitySettings* Settings = CastChecked<UPCGNormalToDensitySettings>(TestData.Settings);
+	FPCGElementPtr Element = TestData.Settings->GetElement();
+
+	Settings->DensityMode = PCGNormalToDensityMode::Add;
+
+	{
+		TObjectPtr<UPCGPointData> PointData = PCGTestsCommon::CreatePointData();
+		FPCGPoint& Point = PointData->GetMutablePoints()[0];
+
+		Point.Transform.SetRotation(FQuat::MakeFromEuler(FVector(45.0, 0, 0)));
+		Point.Density = 0.1f;
+
+		FPCGTaggedData& TaggedData = TestData.InputData.TaggedData.Emplace_GetRef(FPCGTaggedData());
+		TaggedData.Data = PointData;
+		TaggedData.Pin = PCGPinConstants::DefaultInputLabel;
+	}
+
+	TUniquePtr<FPCGContext> Context(Element->Initialize(TestData.InputData, TestData.TestPCGComponent, nullptr));
+	Context->NumAvailableTasks = 1;
+
+	while (!Element->Execute(Context.Get())) {}
+
+	// test our point data
+	const TArray<FPCGTaggedData>& Outputs = Context->OutputData.GetInputs();
+
+	UTEST_EQUAL("Output count", Outputs.Num(), 1);
+
+	UPCGSpatialData* OutputData = Cast<UPCGSpatialData>(Outputs[0].Data);
+
+	UTEST_NOT_NULL("Output data", OutputData);
+
+	const UPCGPointData* OutputPointData = OutputData->ToPointData(Context.Get());
+	UTEST_NOT_NULL("Output point data", OutputPointData);
+
+	const TArray<FPCGPoint>& OutputPoints = OutputPointData->GetPoints();
+
+	UTEST_EQUAL("Output point count", OutputPoints.Num(), 1);
+
+	UTEST_EQUAL_TOLERANCE("Output point density", OutputPoints[0].Density, 0.807f, 0.001f);
+
+	return true;
+}
+
+IMPLEMENT_CUSTOM_SIMPLE_AUTOMATION_TEST(FPCGNormalToDensityTest_Subtract, FPCGTestBaseClass, "pcg.tests.NormalToDensity.Subtract", PCGTestsCommon::TestFlags)
+
+bool FPCGNormalToDensityTest_Subtract::RunTest(const FString& Parameters)
+{
+	PCGTestsCommon::FTestData TestData(PCGDeterminismTests::Defaults::Seed);
+	PCGDeterminismTests::GenerateSettings<UPCGNormalToDensitySettings>(TestData);
+	UPCGNormalToDensitySettings* Settings = CastChecked<UPCGNormalToDensitySettings>(TestData.Settings);
+	FPCGElementPtr Element = TestData.Settings->GetElement();
+
+	Settings->DensityMode = PCGNormalToDensityMode::Subtract;
+
+	{
+		TObjectPtr<UPCGPointData> PointData = PCGTestsCommon::CreatePointData();
+		FPCGPoint& Point = PointData->GetMutablePoints()[0];
+
+		Point.Transform.SetRotation(FQuat::MakeFromEuler(FVector(45.0, 0, 0)));
+		Point.Density = 0.9f;
+
+		FPCGTaggedData& TaggedData = TestData.InputData.TaggedData.Emplace_GetRef(FPCGTaggedData());
+		TaggedData.Data = PointData;
+		TaggedData.Pin = PCGPinConstants::DefaultInputLabel;
+	}
+
+	TUniquePtr<FPCGContext> Context(Element->Initialize(TestData.InputData, TestData.TestPCGComponent, nullptr));
+	Context->NumAvailableTasks = 1;
+
+	while (!Element->Execute(Context.Get())) {}
+
+	// test our point data
+	const TArray<FPCGTaggedData>& Outputs = Context->OutputData.GetInputs();
+
+	UTEST_EQUAL("Output count", Outputs.Num(), 1);
+
+	UPCGSpatialData* OutputData = Cast<UPCGSpatialData>(Outputs[0].Data);
+
+	UTEST_NOT_NULL("Output data", OutputData);
+
+	const UPCGPointData* OutputPointData = OutputData->ToPointData(Context.Get());
+	UTEST_NOT_NULL("Output point data", OutputPointData);
+
+	const TArray<FPCGPoint>& OutputPoints = OutputPointData->GetPoints();
+
+	UTEST_EQUAL("Output point count", OutputPoints.Num(), 1);
+
+	UTEST_EQUAL_TOLERANCE("Output point density", OutputPoints[0].Density, 0.9f-0.707f, 0.001f);
+
+	return true;
+}
+
+IMPLEMENT_CUSTOM_SIMPLE_AUTOMATION_TEST(FPCGNormalToDensityTest_Multiply, FPCGTestBaseClass, "pcg.tests.NormalToDensity.Multiply", PCGTestsCommon::TestFlags)
+
+bool FPCGNormalToDensityTest_Multiply::RunTest(const FString& Parameters)
+{
+	PCGTestsCommon::FTestData TestData(PCGDeterminismTests::Defaults::Seed);
+	PCGDeterminismTests::GenerateSettings<UPCGNormalToDensitySettings>(TestData);
+	UPCGNormalToDensitySettings* Settings = CastChecked<UPCGNormalToDensitySettings>(TestData.Settings);
+	FPCGElementPtr Element = TestData.Settings->GetElement();
+
+	Settings->DensityMode = PCGNormalToDensityMode::Multiply;
+
+	{
+		TObjectPtr<UPCGPointData> PointData = PCGTestsCommon::CreatePointData();
+		FPCGPoint& Point = PointData->GetMutablePoints()[0];
+
+		Point.Transform.SetRotation(FQuat::MakeFromEuler(FVector(45.0, 0, 0)));
+		Point.Density = 0.5f;
+
+		FPCGTaggedData& TaggedData = TestData.InputData.TaggedData.Emplace_GetRef(FPCGTaggedData());
+		TaggedData.Data = PointData;
+		TaggedData.Pin = PCGPinConstants::DefaultInputLabel;
+	}
+
+	TUniquePtr<FPCGContext> Context(Element->Initialize(TestData.InputData, TestData.TestPCGComponent, nullptr));
+	Context->NumAvailableTasks = 1;
+
+	while (!Element->Execute(Context.Get())) {}
+
+	// test our point data
+	const TArray<FPCGTaggedData>& Outputs = Context->OutputData.GetInputs();
+
+	UTEST_EQUAL("Output count", Outputs.Num(), 1);
+
+	UPCGSpatialData* OutputData = Cast<UPCGSpatialData>(Outputs[0].Data);
+
+	UTEST_NOT_NULL("Output data", OutputData);
+
+	const UPCGPointData* OutputPointData = OutputData->ToPointData(Context.Get());
+	UTEST_NOT_NULL("Output point data", OutputPointData);
+
+	const TArray<FPCGPoint>& OutputPoints = OutputPointData->GetPoints();
+
+	UTEST_EQUAL("Output point count", OutputPoints.Num(), 1);
+
+	UTEST_EQUAL_TOLERANCE("Output point density", OutputPoints[0].Density, 0.5f*0.707f, 0.001f);
+
+	return true;
+}
+
+IMPLEMENT_CUSTOM_SIMPLE_AUTOMATION_TEST(FPCGNormalToDensityTest_Divide, FPCGTestBaseClass, "pcg.tests.NormalToDensity.Divide", PCGTestsCommon::TestFlags)
+
+bool FPCGNormalToDensityTest_Divide::RunTest(const FString& Parameters)
+{
+	PCGTestsCommon::FTestData TestData(PCGDeterminismTests::Defaults::Seed);
+	PCGDeterminismTests::GenerateSettings<UPCGNormalToDensitySettings>(TestData);
+	UPCGNormalToDensitySettings* Settings = CastChecked<UPCGNormalToDensitySettings>(TestData.Settings);
+	FPCGElementPtr Element = TestData.Settings->GetElement();
+
+	Settings->DensityMode = PCGNormalToDensityMode::Divide;
+
+	{
+		TObjectPtr<UPCGPointData> PointData = PCGTestsCommon::CreatePointData();
+		FPCGPoint& Point = PointData->GetMutablePoints()[0];
+
+		Point.Transform.SetRotation(FQuat::MakeFromEuler(FVector(45.0, 0, 0)));
+		Point.Density = 0.5f;
+
+		FPCGTaggedData& TaggedData = TestData.InputData.TaggedData.Emplace_GetRef(FPCGTaggedData());
+		TaggedData.Data = PointData;
+		TaggedData.Pin = PCGPinConstants::DefaultInputLabel;
+	}
+
+	TUniquePtr<FPCGContext> Context(Element->Initialize(TestData.InputData, TestData.TestPCGComponent, nullptr));
+	Context->NumAvailableTasks = 1;
+
+	while (!Element->Execute(Context.Get())) {}
+
+	// test our point data
+	const TArray<FPCGTaggedData>& Outputs = Context->OutputData.GetInputs();
+
+	UTEST_EQUAL("Output count", Outputs.Num(), 1);
+
+	UPCGSpatialData* OutputData = Cast<UPCGSpatialData>(Outputs[0].Data);
+
+	UTEST_NOT_NULL("Output data", OutputData);
+
+	const UPCGPointData* OutputPointData = OutputData->ToPointData(Context.Get());
+	UTEST_NOT_NULL("Output point data", OutputPointData);
+
+	const TArray<FPCGPoint>& OutputPoints = OutputPointData->GetPoints();
+
+	UTEST_EQUAL("Output point count", OutputPoints.Num(), 1);
+
+	UTEST_EQUAL_TOLERANCE("Output point density", OutputPoints[0].Density, 0.5f/0.707f, 0.001f);
+
+	return true;
+}
+
diff --git a/Engine/Plugins/Experimental/PCG/Source/PCG/Public/Elements/PCGNormalToDensity.h b/Engine/Plugins/Experimental/PCG/Source/PCG/Public/Elements/PCGNormalToDensity.h
new file mode 100644
--- /dev/null
+++ b/Engine/Plugins/Experimental/PCG/Source/PCG/Public/Elements/PCGNormalToDensity.h
@@ -0,0 +1,67 @@
+// Copyright Epic Games, Inc. All Rights Reserved.
+#pragma once
+
+#include "PCGSettings.h"
+#include "Elements/PCGPointProcessingElementBase.h"
+
+#include "PCGNormalToDensity.generated.h"
+
+UENUM()
+enum class PCGNormalToDensityMode : int8
+{
+	Set,
+	Minimum,
+	Maximum,
+	Add,
+	Subtract,
+	Multiply,
+	Divide
+};
+
+/**
+ * Finds the angle against the specified direction and applies that to the density
+ */
+UCLASS(BlueprintType, ClassGroup = (Procedural))
+class UPCGNormalToDensitySettings : public UPCGSettings
+{
+	GENERATED_BODY()
+
+public:
+	//~Begin UPCGSettings interface
+#if WITH_EDITOR
+	virtual FName GetDefaultNodeName() const override { return FName(TEXT("NormalToDensity")); }
+	virtual EPCGSettingsType GetType() const override { return EPCGSettingsType::Spatial; }
+#endif
+
+	virtual TArray<FPCGPinProperties> InputPinProperties() const override;
+	virtual TArray<FPCGPinProperties> OutputPinProperties() const override;
+
+protected:
+	virtual FPCGElementPtr CreateElement() const override;
+	//~End UPCGSettings interface
+
+public:
+
+	// The normal to compare against
+	UPROPERTY(BlueprintReadWrite, EditAnywhere, Category = Settings, meta = (PCG_Overridable))
+	FVector Normal = {0, 0.0, 1.0};
+
+	// This is biases the value towards or against the normal (positive or negative)
+	UPROPERTY(BlueprintReadWrite, EditAnywhere, Category = Settings, meta = (PCG_Overridable))
+	double Offset = 0.0;
+
+	// This applies a curve to scale the result density with Result = Result^(1/Strength)
+	UPROPERTY(BlueprintReadWrite, EditAnywhere, Category = Settings, meta = (PCG_Overridable, ClampMin = "0.0001", ClampMax = "100.0"))
+	double Strength = 1.0;
+
+	// The operator to apply to the output density 
+	UPROPERTY(BlueprintReadWrite, EditAnywhere, Category = Settings)
+	PCGNormalToDensityMode DensityMode = PCGNormalToDensityMode::Set;
+};
+
+class FPCGNormalToDensityElement : public FPCGPointProcessingElementBase
+{
+protected:
+	virtual bool ExecuteInternal(FPCGContext* Context) const override;
+};
+
