[PCG] Static Analysis fixes

diff --git a/Engine/Plugins/Experimental/PCG/Source/PCG/Private/Elements/PCGSurfaceSampler.cpp b/Engine/Plugins/Experimental/PCG/Source/PCG/Private/Elements/PCGSurfaceSampler.cpp
--- a/Engine/Plugins/Experimental/PCG/Source/PCG/Private/Elements/PCGSurfaceSampler.cpp
+++ b/Engine/Plugins/Experimental/PCG/Source/PCG/Private/Elements/PCGSurfaceSampler.cpp
@@ -30,6 +30,7 @@ namespace PCGSurfaceSampler
 {
 	bool FSurfaceSamplerSettings::Initialize(const UPCGSurfaceSamplerSettings* InSettings, FPCGContext* Context, const FBox& InputBounds)
 	{
+		check(Context);
 		Settings = InSettings;
 
 		if (Settings)
diff --git a/Engine/Plugins/Experimental/PCG/Source/PCG/Private/Tests/Elements/PCGDataTableRowToParamDataTest.cpp b/Engine/Plugins/Experimental/PCG/Source/PCG/Private/Tests/Elements/PCGDataTableRowToParamDataTest.cpp
--- a/Engine/Plugins/Experimental/PCG/Source/PCG/Private/Tests/Elements/PCGDataTableRowToParamDataTest.cpp
+++ b/Engine/Plugins/Experimental/PCG/Source/PCG/Private/Tests/Elements/PCGDataTableRowToParamDataTest.cpp
@@ -40,6 +40,12 @@ CCC,"C Name","C String","333","3333","3.3","3.33","(X=3.0,Y=3.0)","(X=3.0,Y=3.0,
 	const UPCGParamData* Params = OutputParamData.IsEmpty() ? nullptr : Cast<UPCGParamData>(OutputParamData[0].Data);
 	UTEST_NOT_NULL("Output params", Params);
 
+	// For static analysis
+	if (!Params)
+	{
+		return false;
+	}
+
 	UTEST_NOT_NULL("Params has metadata", Params->Metadata.Get());
 
 	void* RowData = TestDataTable->FindRow<FPCGDataTableRowToParamDataTestStruct>(Settings->RowName, TEXT("FPCGDataTableRowToParamDataTest"));
@@ -101,6 +107,12 @@ CCC,"C Name","C String","333","3333","3.3","3.33","(X=3.0,Y=3.0)","(X=3.0,Y=3.0,
 	const UPCGParamData* Params = Context->OutputData.GetFirstParamsOnParamsPin();
 	UTEST_NOT_NULL("Output params", Params);
 
+	// For static analysis
+	if (!Params)
+	{
+		return false;
+	}
+
 	UTEST_NOT_NULL("Params has metadata", Params->Metadata.Get());
 
 	void* RowData = TestDataTable->FindRow<FPCGDataTableRowToParamDataTestStruct>(RowNameOverride, TEXT("FPCGDataTableRowToParamDataTest"));
