Fix for some UE deprecation errors

diff --git a/Engine/Plugins/Experimental/PCG/Source/PCG/Private/Elements/PCGNormalToDensity.cpp b/Engine/Plugins/Experimental/PCG/Source/PCG/Private/Elements/PCGNormalToDensity.cpp
--- a/Engine/Plugins/Experimental/PCG/Source/PCG/Private/Elements/PCGNormalToDensity.cpp
+++ b/Engine/Plugins/Experimental/PCG/Source/PCG/Private/Elements/PCGNormalToDensity.cpp
@@ -36,12 +36,10 @@ bool FPCGNormalToDensityElement::ExecuteInternal(FPCGContext* Context) const
 	const UPCGNormalToDensitySettings* Settings = Context->GetInputSettings<UPCGNormalToDensitySettings>();
 	check(Settings);
 
-	UPCGParamData* Params = Context->InputData.GetParams();
-
-	const FVector Normal = PCGSettingsHelpers::GetValue(GET_MEMBER_NAME_CHECKED(UPCGNormalToDensitySettings, Normal), Settings->Normal, Params);
-	const double Offset = PCGSettingsHelpers::GetValue(GET_MEMBER_NAME_CHECKED(UPCGNormalToDensitySettings, Offset), Settings->Offset, Params);
-	const double Strength = PCGSettingsHelpers::GetValue(GET_MEMBER_NAME_CHECKED(UPCGNormalToDensitySettings, Strength), Settings->Strength, Params);
-	const PCGNormalToDensityMode DensityMode = PCGSettingsHelpers::GetValue(GET_MEMBER_NAME_CHECKED(UPCGNormalToDensitySettings, DensityMode), Settings->DensityMode, Params);
+	const FVector Normal = Settings->Normal;
+	const double Offset = Settings->Offset;
+	const double Strength = Settings->Strength;
+	const PCGNormalToDensityMode DensityMode = Settings->DensityMode;
 
 	TArray<FPCGTaggedData> Inputs = Context->InputData.GetInputsByPin(PCGPinConstants::DefaultInputLabel);
 	TArray<FPCGTaggedData>& Outputs = Context->OutputData.TaggedData;
