[PCG] Ensure landscape cache initialized to fix issue where landscape layer data can be missing from graph.

This initialization was happening via PCGSubsystem::Tick() (once a APCGWorldActor was created). The APCGWorldActor is created just as generate is clicked for first time and generate starts without running the Tick(), and the landscape cache will not be initialized with the layer names.

I can't spot anywhere where layer names are cached/init'd when not in editor. I have not dug too deeply and I have kept faithful to how the code was prior. PIE with a PCG volume on a landscape worked with this change.

diff --git a/Engine/Plugins/Experimental/PCG/Source/PCG/Private/Grid/PCGLandscapeCache.cpp b/Engine/Plugins/Experimental/PCG/Source/PCG/Private/Grid/PCGLandscapeCache.cpp
--- a/Engine/Plugins/Experimental/PCG/Source/PCG/Private/Grid/PCGLandscapeCache.cpp
+++ b/Engine/Plugins/Experimental/PCG/Source/PCG/Private/Grid/PCGLandscapeCache.cpp
@@ -467,7 +467,7 @@ void UPCGLandscapeCache::Serialize(FArchive& Archive)
 	}
 }
 
-void UPCGLandscapeCache::Tick(float DeltaSeconds)
+void UPCGLandscapeCache::Initialize()
 {
 #if WITH_EDITOR
 	if (!bInitialized && GetWorld())
@@ -476,7 +476,14 @@ void UPCGLandscapeCache::Tick(float DeltaSeconds)
 		CacheLayerNames();
 		bInitialized = true;
 	}
-#else
+#endif
+}
+
+void UPCGLandscapeCache::Tick(float DeltaSeconds)
+{
+	Initialize();
+
+#if !WITH_EDITOR
 	// Important implementation note:
 	// If the threshold is too low, it could lead to some issues - namely the check(bDataLoaded) in the cache entries
 	// as we currently do not have a state in the landscape cache to know whether something is still under use.
diff --git a/Engine/Plugins/Experimental/PCG/Source/PCG/Private/PCGWorldActor.cpp b/Engine/Plugins/Experimental/PCG/Source/PCG/Private/PCGWorldActor.cpp
--- a/Engine/Plugins/Experimental/PCG/Source/PCG/Private/PCGWorldActor.cpp
+++ b/Engine/Plugins/Experimental/PCG/Source/PCG/Private/PCGWorldActor.cpp
@@ -33,6 +33,21 @@ void APCGWorldActor::BeginCacheForCookedPlatformData(const ITargetPlatform* Targ
 }
 #endif
 
+
+void APCGWorldActor::PostInitProperties()
+{
+	Super::PostInitProperties();
+
+	if (!HasAnyFlags(RF_ClassDefaultObject))
+	{
+		if (LandscapeCacheObject.Get())
+		{
+			// Make sure landscape cache is ready to provide data immediately.
+			LandscapeCacheObject->Initialize();
+		}
+	}
+}
+
 void APCGWorldActor::BeginPlay()
 {
 	Super::BeginPlay();
diff --git a/Engine/Plugins/Experimental/PCG/Source/PCG/Public/Grid/PCGLandscapeCache.h b/Engine/Plugins/Experimental/PCG/Source/PCG/Public/Grid/PCGLandscapeCache.h
--- a/Engine/Plugins/Experimental/PCG/Source/PCG/Public/Grid/PCGLandscapeCache.h
+++ b/Engine/Plugins/Experimental/PCG/Source/PCG/Public/Grid/PCGLandscapeCache.h
@@ -100,6 +100,9 @@ public:
 	virtual void Serialize(FArchive& Archive) override;
 	//~End UObject interface
 
+	/** Initialize cache. Can be safely called multiple times. */
+	void Initialize();
+
 	void PrimeCache();
 	void ClearCache();
 	void Tick(float DeltaSeconds);
diff --git a/Engine/Plugins/Experimental/PCG/Source/PCG/Public/PCGWorldActor.h b/Engine/Plugins/Experimental/PCG/Source/PCG/Public/PCGWorldActor.h
--- a/Engine/Plugins/Experimental/PCG/Source/PCG/Public/PCGWorldActor.h
+++ b/Engine/Plugins/Experimental/PCG/Source/PCG/Public/PCGWorldActor.h
@@ -19,6 +19,7 @@ public:
 	APCGWorldActor(const FObjectInitializer& ObjectInitializer);
 
 	//~Begin AActor Interface
+	virtual void PostInitProperties() override;
 	virtual void BeginPlay() override;
 	virtual void EndPlay(const EEndPlayReason::Type EndPlayReason) override;
 
