[PCG] Tentative fix for Static Analysis

diff --git a/Engine/Plugins/Experimental/PCG/Source/PCG/Private/Elements/PCGSurfaceSampler.cpp b/Engine/Plugins/Experimental/PCG/Source/PCG/Private/Elements/PCGSurfaceSampler.cpp
--- a/Engine/Plugins/Experimental/PCG/Source/PCG/Private/Elements/PCGSurfaceSampler.cpp
+++ b/Engine/Plugins/Experimental/PCG/Source/PCG/Private/Elements/PCGSurfaceSampler.cpp
@@ -30,7 +30,11 @@ namespace PCGSurfaceSampler
 {
 	bool FSurfaceSamplerSettings::Initialize(const UPCGSurfaceSamplerSettings* InSettings, FPCGContext* Context, const FBox& InputBounds)
 	{
-		check(Context);
+		if (!Context)
+		{
+			return false;
+		}
+
 		Settings = InSettings;
 
 		if (Settings)
