[PCG] Property accessor unit tests

* Accessing a property in a struct
* Accessing components in maths objects (vector/transform...)
* Accessing multiple depth structs (structs in structs)
* Fixing a bug in components accessors

diff --git a/Engine/Plugins/Experimental/PCG/Source/PCG/Private/Metadata/Accessors/PCGAttributeExtractor.cpp b/Engine/Plugins/Experimental/PCG/Source/PCG/Private/Metadata/Accessors/PCGAttributeExtractor.cpp
--- a/Engine/Plugins/Experimental/PCG/Source/PCG/Private/Metadata/Accessors/PCGAttributeExtractor.cpp
+++ b/Engine/Plugins/Experimental/PCG/Source/PCG/Private/Metadata/Accessors/PCGAttributeExtractor.cpp
@@ -60,7 +60,7 @@ namespace PCGAttributeExtractor
 	template <typename VectorType>
 	TUniquePtr<IPCGAttributeAccessor> CreateVectorExtractor(TUniquePtr<IPCGAttributeAccessor> InAccessor, FName Name, bool& bOutSuccess)
 	{
-		if ((Name == PCGAttributeExtractorConstants::VectorLength) || (Name == TEXT("Size")))
+		if ((Name == PCGAttributeExtractorConstants::VectorLength) || (Name == PCGAttributeExtractorConstants::VectorSize))
 		{
 			bOutSuccess = true;
 			return MakeUnique<FPCGChainAccessor<double, VectorType>>(std::move(InAccessor),
@@ -95,10 +95,26 @@ namespace PCGAttributeExtractor
 		TArray<int32, TInlineAllocator<4>> Indexes;
 		for (const TCHAR Char : NameStr)
 		{
-			int32 Index = static_cast<int32>(Char) - static_cast<int32>(TEXT('x'));
-			if (Index < 0)
+			int32 Index = -1;
+
+			if (Char == TEXT('w') || Char == TEXT('W'))
+			{
+				Index = 3;
+			}
+			else
+			{
+				Index = static_cast<int32>(Char) - static_cast<int32>(TEXT('x'));
+				if (Index < 0)
+				{
+					Index = static_cast<int32>(Char) - static_cast<int32>(TEXT('X'));
+				}
+			}
+
+			// Safeguard, should be caught by not matching the regex above
+			if (!ensure(Index >= 0 && Index < 4))
 			{
-				Index = static_cast<int32>(Char) - static_cast<int32>(TEXT('X'));
+				bOutSuccess = false;
+				return InAccessor;
 			}
 
 			Indexes.Add(Index);
diff --git a/Engine/Plugins/Experimental/PCG/Source/PCG/Private/Tests/Accessors/PCGAttributeExtractorTest.cpp b/Engine/Plugins/Experimental/PCG/Source/PCG/Private/Tests/Accessors/PCGAttributeExtractorTest.cpp
new file mode 100644
--- /dev/null
+++ b/Engine/Plugins/Experimental/PCG/Source/PCG/Private/Tests/Accessors/PCGAttributeExtractorTest.cpp
@@ -0,0 +1,491 @@
+// Copyright Epic Games, Inc. All Rights Reserved.
+
+#include "Tests/Accessor/PCGAttributeExtractorTest.h"
+
+#include "Metadata/PCGMetadataAttributeTpl.h"
+#include "Metadata/PCGMetadata.h"
+#include "Metadata/Accessors/IPCGAttributeAccessor.h"
+#include "Metadata/Accessors/PCGAttributeAccessorHelpers.h"
+#include "Metadata/Accessors/PCGAttributeAccessorKeys.h"
+#include "Metadata/Accessors/PCGAttributeExtractor.h"
+#include "Metadata/PCGMetadataTypesConstantStruct.h"
+#include "Tests/PCGTestsCommon.h"
+
+#include "UObject/UnrealType.h"
+
+#if WITH_EDITOR
+
+namespace FPCGAttributeExtractorTestHelpers
+{
+	FPCGMetadataTypesConstantStruct CreateStruct(UObject* InObject = nullptr)
+	{
+		FPCGMetadataTypesConstantStruct Struct;
+
+		Struct.FloatValue = 1.25f;
+		Struct.DoubleValue = 2.0;
+		Struct.Int32Value = 42;
+		Struct.IntValue = 666;
+		Struct.Vector2Value = FVector2D(1.23, 1.58);
+		Struct.VectorValue = FVector(3.23, 5.58, 2.69);
+		Struct.Vector4Value = FVector4(2.23, 9.58, 4.21, 8.01);
+		Struct.BoolValue = true;
+		Struct.NameValue = TEXT("Foo");
+		Struct.StringValue = TEXT("Bar");
+		Struct.RotatorValue = FRotator(0.5, 0.9, 0.4);
+		Struct.QuatValue = FQuat(0.145, 0.254, 0.369, 0.478);
+		Struct.TransformValue = FTransform(Struct.QuatValue, Struct.VectorValue, FVector(Struct.Vector2Value, 1.9874));
+		Struct.SoftClassPathValue = FSoftClassPath(UPCGMetadata::StaticClass());
+		if (InObject)
+		{
+			Struct.SoftObjectPathValue = FSoftObjectPath(InObject);
+		}
+
+		return Struct;
+	}
+}
+
+
+IMPLEMENT_CUSTOM_SIMPLE_AUTOMATION_TEST(FPCGAttributeSingleGetPropertyTest, FPCGTestBaseClass, "pcg.tests.Accessor.Property.SimpleGetProperty", PCGTestsCommon::TestFlags)
+IMPLEMENT_CUSTOM_SIMPLE_AUTOMATION_TEST(FPCGAttributeVectorPropertyExtractorTest, FPCGTestBaseClass, "pcg.tests.Accessor.Property.VectorExtractor", PCGTestsCommon::TestFlags)
+IMPLEMENT_CUSTOM_SIMPLE_AUTOMATION_TEST(FPCGAttributeRotatorPropertyExtractorTest, FPCGTestBaseClass, "pcg.tests.Accessor.Property.RotatorExtractor", PCGTestsCommon::TestFlags)
+IMPLEMENT_CUSTOM_SIMPLE_AUTOMATION_TEST(FPCGAttributeTransformPropertyExtractorTest, FPCGTestBaseClass, "pcg.tests.Accessor.Property.TransformExtractor", PCGTestsCommon::TestFlags)
+IMPLEMENT_CUSTOM_SIMPLE_AUTOMATION_TEST(FPCGAttributeInvalidPropertyExtractorTest, FPCGTestBaseClass, "pcg.tests.Accessor.Property.InvalidExtractor", PCGTestsCommon::TestFlags)
+IMPLEMENT_CUSTOM_SIMPLE_AUTOMATION_TEST(FPCGAttributePropertyMultipleDepthTest, FPCGTestBaseClass, "pcg.tests.Accessor.Property.MultipleDepth", PCGTestsCommon::TestFlags)
+IMPLEMENT_CUSTOM_SIMPLE_AUTOMATION_TEST(FPCGAttributePropertyMultipleDepthRangeTest, FPCGTestBaseClass, "pcg.tests.Accessor.Property.MultipleDepthRange", PCGTestsCommon::TestFlags)
+
+bool FPCGAttributeSingleGetPropertyTest::RunTest(const FString& Parameters)
+{
+	UPCGMetadata* TempMetadata = NewObject<UPCGMetadata>();
+
+	FPCGMetadataTypesConstantStruct Struct = FPCGAttributeExtractorTestHelpers::CreateStruct(TempMetadata);
+
+	auto Verify = [this, &Struct](const auto& ExpectedValue, const FName PropertyName) -> bool
+	{
+		using PropertyType = std::decay_t<decltype(ExpectedValue)>;
+
+		UTEST_TRUE(TEXT("Property type supported"), PCGAttributeAccessorHelpers::IsPropertyAccessorSupported(PropertyName, FPCGMetadataTypesConstantStruct::StaticStruct()));
+
+		TUniquePtr<IPCGAttributeAccessor> Accessor = PCGAttributeAccessorHelpers::CreatePropertyAccessor(PropertyName, FPCGMetadataTypesConstantStruct::StaticStruct());
+		UTEST_TRUE(FString::Printf(TEXT("Valid accessor for property %s"), *PropertyName.ToString()), Accessor.IsValid());
+		UTEST_EQUAL(TEXT("Accessor type"), PCG::Private::MetadataTypes<PropertyType>::Id, Accessor->GetUnderlyingType());
+
+		FPCGAttributeAccessorKeysSingleObjectPtr ObjectKey(&Struct);
+
+		PropertyType Value{};
+		UTEST_TRUE(TEXT("Getting a value in range"), Accessor->Get<PropertyType>(Value, ObjectKey));
+		UTEST_EQUAL(TEXT("SingleGet: Value is equal to expected"), Value, ExpectedValue);
+
+		PropertyType Value2{};
+		UTEST_TRUE(TEXT("Getting a value outside range"), Accessor->Get<PropertyType>(Value2, 5, ObjectKey));
+		UTEST_EQUAL(TEXT("SingleGet Outside Range: Value is equal to expected"), Value2, ExpectedValue);
+
+		constexpr int NbValues = 3;
+		PropertyType Values[NbValues];
+		UTEST_TRUE(TEXT("Getting the same value thrice"), Accessor->GetRange<PropertyType>(MakeArrayView(Values, NbValues), 0, ObjectKey));
+		for (int i = 0; i < NbValues; ++i)
+		{
+			UTEST_EQUAL(FString::Printf(TEXT("GetRange: Value %d is equal to expected"), i), Values[i], ExpectedValue);
+		}
+
+		PropertyType Values2[NbValues];
+		UTEST_TRUE(TEXT("Getting the same value thrice outside range"), Accessor->GetRange<PropertyType>(MakeArrayView(Values2, NbValues), 5, ObjectKey));
+		for (int i = 0; i < NbValues; ++i)
+		{
+			UTEST_EQUAL(FString::Printf(TEXT("GetRange Outside Range: Value %d is equal to expected"), i), Values2[i], ExpectedValue);
+		}
+
+		return true;
+	};
+
+	bool bTestPassed = true;
+
+	bTestPassed &= Verify(double(Struct.FloatValue), GET_MEMBER_NAME_CHECKED(FPCGMetadataTypesConstantStruct, FloatValue));
+	bTestPassed &= Verify(Struct.DoubleValue, GET_MEMBER_NAME_CHECKED(FPCGMetadataTypesConstantStruct, DoubleValue));
+	bTestPassed &= Verify(int64(Struct.Int32Value), GET_MEMBER_NAME_CHECKED(FPCGMetadataTypesConstantStruct, Int32Value));
+	bTestPassed &= Verify(Struct.IntValue, GET_MEMBER_NAME_CHECKED(FPCGMetadataTypesConstantStruct, IntValue));
+	bTestPassed &= Verify(Struct.Vector2Value, GET_MEMBER_NAME_CHECKED(FPCGMetadataTypesConstantStruct, Vector2Value));
+	bTestPassed &= Verify(Struct.VectorValue, GET_MEMBER_NAME_CHECKED(FPCGMetadataTypesConstantStruct, VectorValue));
+	bTestPassed &= Verify(Struct.Vector4Value, GET_MEMBER_NAME_CHECKED(FPCGMetadataTypesConstantStruct, Vector4Value));
+	bTestPassed &= Verify(Struct.BoolValue, GET_MEMBER_NAME_CHECKED(FPCGMetadataTypesConstantStruct, BoolValue));
+	bTestPassed &= Verify(Struct.NameValue, GET_MEMBER_NAME_CHECKED(FPCGMetadataTypesConstantStruct, NameValue));
+	bTestPassed &= Verify(Struct.StringValue, GET_MEMBER_NAME_CHECKED(FPCGMetadataTypesConstantStruct, StringValue));
+	bTestPassed &= Verify(Struct.RotatorValue, GET_MEMBER_NAME_CHECKED(FPCGMetadataTypesConstantStruct, RotatorValue));
+	bTestPassed &= Verify(Struct.QuatValue, GET_MEMBER_NAME_CHECKED(FPCGMetadataTypesConstantStruct, QuatValue));
+	bTestPassed &= Verify(Struct.TransformValue, GET_MEMBER_NAME_CHECKED(FPCGMetadataTypesConstantStruct, TransformValue));
+	bTestPassed &= Verify(Struct.SoftClassPathValue.ToString(), GET_MEMBER_NAME_CHECKED(FPCGMetadataTypesConstantStruct, SoftClassPathValue));
+	bTestPassed &= Verify(Struct.SoftObjectPathValue.ToString(), GET_MEMBER_NAME_CHECKED(FPCGMetadataTypesConstantStruct, SoftObjectPathValue));
+
+	TempMetadata->MarkAsGarbage();
+
+	return true;
+}
+
+bool FPCGAttributeVectorPropertyExtractorTest::RunTest(const FString& Parameters)
+{
+	FPCGMetadataTypesConstantStruct Struct = FPCGAttributeExtractorTestHelpers::CreateStruct();
+
+	auto Verify = [this, &Struct](const auto& ExpectedValue, const FName PropertyName, const FName ExtractorName) -> bool
+	{
+		using PropertyType = std::decay_t<decltype(ExpectedValue)>;
+
+		TUniquePtr<IPCGAttributeAccessor> Accessor = PCGAttributeAccessorHelpers::CreatePropertyAccessor(PropertyName, FPCGMetadataTypesConstantStruct::StaticStruct());
+		UTEST_TRUE(FString::Printf(TEXT("Valid accessor for property %s"), *PropertyName.ToString()), Accessor.IsValid());
+
+		bool bSuccess = false;
+		TUniquePtr<IPCGAttributeAccessor> ExtractorAccessor = PCGAttributeAccessorHelpers::CreateChainAccessor(std::move(Accessor), ExtractorName, bSuccess);
+		UTEST_TRUE(TEXT("Valid extractor accessor"), (bSuccess && ExtractorAccessor.IsValid()));
+		UTEST_TRUE(TEXT("Invalid accessor (ownership moved)"), !Accessor.IsValid());
+		UTEST_EQUAL(TEXT("Extractror accessor type"), PCG::Private::MetadataTypes<PropertyType>::Id, ExtractorAccessor->GetUnderlyingType());
+
+		FPCGAttributeAccessorKeysSingleObjectPtr ObjectKey(&Struct);
+
+		PropertyType Value{};
+		UTEST_TRUE(FString::Printf(TEXT("Getting a value for property %s and extractor %s"), *PropertyName.ToString(), *ExtractorName.ToString()), ExtractorAccessor->Get<PropertyType>(Value, ObjectKey));
+		UTEST_EQUAL(FString::Printf(TEXT("Value for property %s and extractor %s is equal to expected"), *PropertyName.ToString(), *ExtractorName.ToString()), Value, ExpectedValue);
+
+		return true;
+	};
+
+	bool bTestPassed = true;
+
+	// Vector Lengths
+	bTestPassed &= Verify(Struct.Vector2Value.Length(), GET_MEMBER_NAME_CHECKED(FPCGMetadataTypesConstantStruct, Vector2Value), PCGAttributeExtractorConstants::VectorLength);
+	bTestPassed &= Verify(Struct.VectorValue.Length(), GET_MEMBER_NAME_CHECKED(FPCGMetadataTypesConstantStruct, VectorValue), PCGAttributeExtractorConstants::VectorLength);
+	bTestPassed &= Verify(Struct.Vector4Value.Size(), GET_MEMBER_NAME_CHECKED(FPCGMetadataTypesConstantStruct, Vector4Value), PCGAttributeExtractorConstants::VectorLength);
+
+	bTestPassed &= Verify(Struct.Vector2Value.Size(), GET_MEMBER_NAME_CHECKED(FPCGMetadataTypesConstantStruct, Vector2Value), PCGAttributeExtractorConstants::VectorSize);
+	bTestPassed &= Verify(Struct.VectorValue.Size(), GET_MEMBER_NAME_CHECKED(FPCGMetadataTypesConstantStruct, VectorValue), PCGAttributeExtractorConstants::VectorSize);
+	bTestPassed &= Verify(Struct.Vector4Value.Size(), GET_MEMBER_NAME_CHECKED(FPCGMetadataTypesConstantStruct, Vector4Value), PCGAttributeExtractorConstants::VectorSize);
+
+	// Vector2 Composents
+	auto Vec2Composents = [&Verify, &bTestPassed](const auto& Value, const FName PropertyName) -> bool
+	{
+		bTestPassed &= Verify(Value.X, PropertyName, FName(TEXT("X")));
+		bTestPassed &= Verify(Value.Y, PropertyName, FName(TEXT("Y")));
+
+		bTestPassed &= Verify(FVector2D{ Value.X, Value.Y }, PropertyName, FName(TEXT("XY")));
+		bTestPassed &= Verify(FVector2D{ Value.Y, Value.X }, PropertyName, FName(TEXT("YX")));
+		bTestPassed &= Verify(FVector2D{ Value.Y, Value.Y }, PropertyName, FName(TEXT("YY")));
+
+		return true;
+	};
+
+	Vec2Composents(Struct.Vector2Value, GET_MEMBER_NAME_CHECKED(FPCGMetadataTypesConstantStruct, Vector2Value));
+	Vec2Composents(Struct.VectorValue, GET_MEMBER_NAME_CHECKED(FPCGMetadataTypesConstantStruct, VectorValue));
+	Vec2Composents(Struct.Vector4Value, GET_MEMBER_NAME_CHECKED(FPCGMetadataTypesConstantStruct, Vector4Value));
+	Vec2Composents(Struct.QuatValue, GET_MEMBER_NAME_CHECKED(FPCGMetadataTypesConstantStruct, QuatValue));
+
+	// Vector3 Composents
+	auto Vec3Composents = [&Verify, &bTestPassed](const auto& Value, const FName PropertyName) -> bool
+	{
+		bTestPassed &= Verify(Value.Z, PropertyName, FName(TEXT("Z")));
+
+		bTestPassed &= Verify(FVector2D{ Value.Z, Value.X }, PropertyName, FName(TEXT("ZX")));
+		bTestPassed &= Verify(FVector2D{ Value.Y, Value.Z }, PropertyName, FName(TEXT("YZ")));
+
+		bTestPassed &= Verify(FVector{ Value.Z, Value.X, Value.Y }, PropertyName, FName(TEXT("ZXY")));
+		bTestPassed &= Verify(FVector{ Value.Y, Value.Y, Value.Y }, PropertyName, FName(TEXT("YYY")));
+		bTestPassed &= Verify(FVector{ Value.Y, Value.Y, Value.X }, PropertyName, FName(TEXT("YYX")));
+		bTestPassed &= Verify(FVector{ Value.X, Value.Y, Value.Z }, PropertyName, FName(TEXT("XYZ")));
+
+		return true;
+	};
+
+	Vec3Composents(Struct.VectorValue, GET_MEMBER_NAME_CHECKED(FPCGMetadataTypesConstantStruct, VectorValue));
+	Vec3Composents(Struct.Vector4Value, GET_MEMBER_NAME_CHECKED(FPCGMetadataTypesConstantStruct, Vector4Value));
+	Vec3Composents(Struct.QuatValue, GET_MEMBER_NAME_CHECKED(FPCGMetadataTypesConstantStruct, QuatValue));
+
+	// Vector4 Composents
+	auto Vec4Composents = [&Verify, &bTestPassed](const auto& Value, const FName PropertyName) -> bool
+	{
+		bTestPassed &= Verify(Value.W, PropertyName, FName(TEXT("W")));
+
+		bTestPassed &= Verify(FVector2D{ Value.Z, Value.W }, PropertyName, FName(TEXT("ZW")));
+		bTestPassed &= Verify(FVector2D{ Value.W, Value.Y }, PropertyName, FName(TEXT("WY")));
+
+		bTestPassed &= Verify(FVector{ Value.Z, Value.W, Value.Y }, PropertyName, FName(TEXT("ZWY")));
+		bTestPassed &= Verify(FVector{ Value.W, Value.W, Value.W }, PropertyName, FName(TEXT("WWW")));
+		bTestPassed &= Verify(FVector{ Value.Y, Value.W, Value.X }, PropertyName, FName(TEXT("YWX")));
+		bTestPassed &= Verify(FVector{ Value.X, Value.Y, Value.W }, PropertyName, FName(TEXT("XYW")));
+
+		bTestPassed &= Verify(FVector4{ Value.Z, Value.X, Value.Y, Value.W }, PropertyName, FName(TEXT("ZXYW")));
+		bTestPassed &= Verify(FVector4{ Value.Z, Value.Z, Value.Z, Value.Z }, PropertyName, FName(TEXT("ZZZZ")));
+		bTestPassed &= Verify(FVector4{ Value.X, Value.Y, Value.X, Value.Y }, PropertyName, FName(TEXT("XYXY")));
+		bTestPassed &= Verify(FVector4{ Value.X, Value.Y, Value.Z, Value.W }, PropertyName, FName(TEXT("XYZW")));
+
+		return true;
+	};
+
+	Vec4Composents(Struct.Vector4Value, GET_MEMBER_NAME_CHECKED(FPCGMetadataTypesConstantStruct, Vector4Value));
+	Vec4Composents(Struct.QuatValue, GET_MEMBER_NAME_CHECKED(FPCGMetadataTypesConstantStruct, QuatValue));
+
+	return bTestPassed;
+}
+
+bool FPCGAttributeRotatorPropertyExtractorTest::RunTest(const FString & Parameters)
+{
+	FPCGMetadataTypesConstantStruct Struct = FPCGAttributeExtractorTestHelpers::CreateStruct();
+
+	auto Verify = [this, &Struct](double ExpectedValue, const FName PropertyName, const FName ExtractorName) -> bool
+	{
+		TUniquePtr<IPCGAttributeAccessor> Accessor = PCGAttributeAccessorHelpers::CreatePropertyAccessor(PropertyName, FPCGMetadataTypesConstantStruct::StaticStruct());
+		UTEST_TRUE(FString::Printf(TEXT("Valid accessor for property %s"), *PropertyName.ToString()), Accessor.IsValid());
+
+		bool bSuccess = false;
+		TUniquePtr<IPCGAttributeAccessor> ExtractorAccessor = PCGAttributeAccessorHelpers::CreateChainAccessor(std::move(Accessor), ExtractorName, bSuccess);
+		UTEST_TRUE(TEXT("Valid extractor accessor"), (bSuccess && ExtractorAccessor.IsValid()));
+		UTEST_TRUE(TEXT("Invalid accessor (ownership moved)"), !Accessor.IsValid());
+		UTEST_EQUAL(TEXT("Extractror accessor type"), PCG::Private::MetadataTypes<double>::Id, ExtractorAccessor->GetUnderlyingType());
+
+		FPCGAttributeAccessorKeysSingleObjectPtr ObjectKey(&Struct);
+
+		double Value = 1.0;
+		UTEST_TRUE(FString::Printf(TEXT("Getting a value for property %s and extractor %s"), *PropertyName.ToString(), *ExtractorName.ToString()), ExtractorAccessor->Get<double>(Value, ObjectKey));
+		UTEST_EQUAL(FString::Printf(TEXT("Value for property %s and extractor %s is equal to expected"), *PropertyName.ToString(), *ExtractorName.ToString()), Value, ExpectedValue);
+
+		return true;
+	};
+
+	bool bTestPassed = true;
+
+	bTestPassed &= Verify(Struct.RotatorValue.Pitch, GET_MEMBER_NAME_CHECKED(FPCGMetadataTypesConstantStruct, RotatorValue), PCGAttributeExtractorConstants::RotatorPitch);
+	bTestPassed &= Verify(Struct.RotatorValue.Yaw, GET_MEMBER_NAME_CHECKED(FPCGMetadataTypesConstantStruct, RotatorValue), PCGAttributeExtractorConstants::RotatorYaw);
+	bTestPassed &= Verify(Struct.RotatorValue.Roll, GET_MEMBER_NAME_CHECKED(FPCGMetadataTypesConstantStruct, RotatorValue), PCGAttributeExtractorConstants::RotatorRoll);
+
+	return bTestPassed;
+}
+
+bool FPCGAttributeTransformPropertyExtractorTest::RunTest(const FString & Parameters)
+{
+	FPCGMetadataTypesConstantStruct Struct = FPCGAttributeExtractorTestHelpers::CreateStruct();
+
+	auto Verify = [this, &Struct](const auto& ExpectedValue, const FName PropertyName, const TArray<FName>& ExtractorNames) -> bool
+	{
+		using PropertyType = std::decay_t<decltype(ExpectedValue)>;
+
+		TUniquePtr<IPCGAttributeAccessor> Accessor = PCGAttributeAccessorHelpers::CreatePropertyAccessor(PropertyName, FPCGMetadataTypesConstantStruct::StaticStruct());
+		UTEST_TRUE(FString::Printf(TEXT("Valid accessor for property %s"), *PropertyName.ToString()), Accessor.IsValid());
+
+		bool bSuccess = false;
+		TUniquePtr<IPCGAttributeAccessor> ExtractorAccessor = std::move(Accessor);
+		for (const FName ExtractorName : ExtractorNames)
+		{
+			ExtractorAccessor = PCGAttributeAccessorHelpers::CreateChainAccessor(std::move(ExtractorAccessor), ExtractorName, bSuccess);
+			UTEST_TRUE(FString::Printf(TEXT("Valid extractor accessor %s"), *ExtractorName.ToString()), (bSuccess && ExtractorAccessor.IsValid()));
+		}
+
+		UTEST_EQUAL("Extractror accessor type", PCG::Private::MetadataTypes<PropertyType>::Id, ExtractorAccessor->GetUnderlyingType());
+
+		FPCGAttributeAccessorKeysSingleObjectPtr ObjectKey(&Struct);
+
+
+		PropertyType Value{};
+		FString ExtractorStr = FString::JoinBy(ExtractorNames, TEXT("."), [](const FName& Name) -> FString { return Name.ToString(); });
+		UTEST_TRUE(FString::Printf(TEXT("Getting a value for property %s and extractor %s"), *PropertyName.ToString(), *ExtractorStr), ExtractorAccessor->Get<PropertyType>(Value, ObjectKey));
+		UTEST_EQUAL(FString::Printf(TEXT("Value for property %s and extractor %s is equal to expected"), *PropertyName.ToString(), *ExtractorStr), Value, ExpectedValue);
+
+		return true;
+	};
+
+	bool bTestPassed = true;
+
+	bTestPassed &= Verify(Struct.TransformValue.GetLocation(), GET_MEMBER_NAME_CHECKED(FPCGMetadataTypesConstantStruct, TransformValue), { PCGAttributeExtractorConstants::TransformLocation });
+	bTestPassed &= Verify(Struct.TransformValue.GetRotation(), GET_MEMBER_NAME_CHECKED(FPCGMetadataTypesConstantStruct, TransformValue), { PCGAttributeExtractorConstants::TransformRotation });
+	bTestPassed &= Verify(Struct.TransformValue.GetScale3D(), GET_MEMBER_NAME_CHECKED(FPCGMetadataTypesConstantStruct, TransformValue), { PCGAttributeExtractorConstants::TransformScale });
+
+	// Also testing double chain
+	bTestPassed &= Verify(Struct.TransformValue.GetLocation().X, GET_MEMBER_NAME_CHECKED(FPCGMetadataTypesConstantStruct, TransformValue), { PCGAttributeExtractorConstants::TransformLocation, PCGAttributeExtractorConstants::VectorX });
+	bTestPassed &= Verify(Struct.TransformValue.GetRotation().W, GET_MEMBER_NAME_CHECKED(FPCGMetadataTypesConstantStruct, TransformValue), { PCGAttributeExtractorConstants::TransformRotation, PCGAttributeExtractorConstants::VectorW });
+	bTestPassed &= Verify(Struct.TransformValue.GetScale3D().Y, GET_MEMBER_NAME_CHECKED(FPCGMetadataTypesConstantStruct, TransformValue), { PCGAttributeExtractorConstants::TransformScale, PCGAttributeExtractorConstants::VectorY });
+
+	return true;
+}
+
+bool FPCGAttributeInvalidPropertyExtractorTest::RunTest(const FString& Parameters)
+{
+	const FName InvalidName = FName(TEXT("Invalid"));
+	TUniquePtr<IPCGAttributeAccessor> Accessor = PCGAttributeAccessorHelpers::CreatePropertyAccessor(InvalidName, FPCGMetadataTypesConstantStruct::StaticStruct());
+	UTEST_TRUE(TEXT("Invalid accessor with name \"Invalid\""), !Accessor.IsValid());
+
+	auto VerifyInvalidChain = [this](const FName PropertyName, const TArray<FName>& ExtractorNames)
+	{
+		TUniquePtr<IPCGAttributeAccessor> ExtractorAccessor = PCGAttributeAccessorHelpers::CreatePropertyAccessor(PropertyName, FPCGMetadataTypesConstantStruct::StaticStruct());
+		UTEST_TRUE(TEXT("Valid accessor"), ExtractorAccessor.IsValid());
+
+		bool bSuccess = false;
+
+		for (int32 i = 0; i < ExtractorNames.Num(); ++i)
+		{
+			ExtractorAccessor = PCGAttributeAccessorHelpers::CreateChainAccessor(std::move(ExtractorAccessor), ExtractorNames[i], bSuccess);
+			// All accessors valid until the last one
+			if (i != ExtractorNames.Num() - 1)
+			{
+				UTEST_TRUE(FString::Printf(TEXT("Valid extractor accessor %s"), *ExtractorNames[i].ToString()), (bSuccess && ExtractorAccessor.IsValid()));
+			}
+			else
+			{
+				// ExtractorAccessor should still be valid (moving back the previous accessor), but success should be false.
+				UTEST_TRUE(FString::Printf(TEXT("Failed to created extractor accessor %s"), *ExtractorNames[i].ToString()), (!bSuccess && ExtractorAccessor.IsValid()));
+			}
+		}
+
+		return true;
+	};
+
+	bool bTestPassed = true;
+
+	VerifyInvalidChain(GET_MEMBER_NAME_CHECKED(FPCGMetadataTypesConstantStruct, DoubleValue), { InvalidName });
+	VerifyInvalidChain(GET_MEMBER_NAME_CHECKED(FPCGMetadataTypesConstantStruct, Vector2Value), { InvalidName });
+	VerifyInvalidChain(GET_MEMBER_NAME_CHECKED(FPCGMetadataTypesConstantStruct, VectorValue), { InvalidName });
+	VerifyInvalidChain(GET_MEMBER_NAME_CHECKED(FPCGMetadataTypesConstantStruct, Vector4Value), { InvalidName });
+	VerifyInvalidChain(GET_MEMBER_NAME_CHECKED(FPCGMetadataTypesConstantStruct, QuatValue), { InvalidName });
+	VerifyInvalidChain(GET_MEMBER_NAME_CHECKED(FPCGMetadataTypesConstantStruct, RotatorValue), { InvalidName });
+	VerifyInvalidChain(GET_MEMBER_NAME_CHECKED(FPCGMetadataTypesConstantStruct, TransformValue), { InvalidName });
+
+	// Invalid double chain
+	VerifyInvalidChain(GET_MEMBER_NAME_CHECKED(FPCGMetadataTypesConstantStruct, TransformValue), { PCGAttributeExtractorConstants::TransformLocation, InvalidName });
+	VerifyInvalidChain(GET_MEMBER_NAME_CHECKED(FPCGMetadataTypesConstantStruct, TransformValue), { PCGAttributeExtractorConstants::TransformRotation, InvalidName });
+	VerifyInvalidChain(GET_MEMBER_NAME_CHECKED(FPCGMetadataTypesConstantStruct, TransformValue), { PCGAttributeExtractorConstants::TransformScale, InvalidName });
+
+	return bTestPassed;
+}
+
+bool FPCGAttributePropertyMultipleDepthTest::RunTest(const FString& Parameters)
+{
+	FPCGAttributeExtractorTestStruct TestStruct{};
+	TestStruct.Object = NewObject<UPCGAttributeExtractorTestObject>();
+	TestStruct.Object->DoubleValue = 0.5;
+
+	TestStruct.DepthStruct.FloatValue = 0.324f;
+	TestStruct.DepthStruct.Depth2Struct.IntValue = 5;
+
+	const FObjectPtrProperty* ObjectProperty = CastField<FObjectPtrProperty>(FPCGAttributeExtractorTestStruct::StaticStruct()->FindPropertyByName(GET_MEMBER_NAME_CHECKED(FPCGAttributeExtractorTestStruct, Object)));
+	const FProperty* DepthStructProperty = FPCGAttributeExtractorTestStruct::StaticStruct()->FindPropertyByName(GET_MEMBER_NAME_CHECKED(FPCGAttributeExtractorTestStruct, DepthStruct));
+	const FProperty* Depth2StructProperty = FPCGAttributeExtractorTestStructDepth1::StaticStruct()->FindPropertyByName(GET_MEMBER_NAME_CHECKED(FPCGAttributeExtractorTestStructDepth1, Depth2Struct));
+
+	UTEST_NOT_NULL(TEXT("Valid object property"), ObjectProperty);
+	UTEST_NOT_NULL(TEXT("Valid depth struct property"), DepthStructProperty);
+	UTEST_NOT_NULL(TEXT("Valid depth struct 2 property"), Depth2StructProperty);
+
+	// Static analysis caution
+	if (!ObjectProperty || !DepthStructProperty || !Depth2StructProperty)
+	{
+		return false;
+	}
+
+	TUniquePtr<IPCGAttributeAccessor> ObjectAccessor = PCGAttributeAccessorHelpers::CreatePropertyAccessor(ObjectProperty);
+	TUniquePtr<IPCGAttributeAccessor> ObjectDoubleAccessor = PCGAttributeAccessorHelpers::CreatePropertyAccessor(GET_MEMBER_NAME_CHECKED(UPCGAttributeExtractorTestObject, DoubleValue), UPCGAttributeExtractorTestObject::StaticClass());
+	TUniquePtr<IPCGAttributeAccessor> DepthStructFloatAccessor = PCGAttributeAccessorHelpers::CreatePropertyAccessor(GET_MEMBER_NAME_CHECKED(FPCGAttributeExtractorTestStructDepth1, FloatValue), FPCGAttributeExtractorTestStructDepth1::StaticStruct());
+	TUniquePtr<IPCGAttributeAccessor> Depth2StructIntAccessor = PCGAttributeAccessorHelpers::CreatePropertyAccessor(GET_MEMBER_NAME_CHECKED(FPCGAttributeExtractorTestStructDepth2, IntValue), FPCGAttributeExtractorTestStructDepth2::StaticStruct());
+
+	// Objects not yet supported
+	UTEST_TRUE(TEXT("Invalid object accessor"), !ObjectAccessor.IsValid());
+	UTEST_TRUE(TEXT("Valid object double accessor"), ObjectDoubleAccessor.IsValid());
+	UTEST_TRUE(TEXT("Valid depth struct float accessor"), DepthStructFloatAccessor.IsValid());
+	UTEST_TRUE(TEXT("Valid deptch struct 2 int accessor"), Depth2StructIntAccessor.IsValid());
+
+	// Then create different keys, manually for now
+	// For object it is a bit more work
+	const FObjectPtr& ObjectPtr = (FObjectPtr&)ObjectProperty->GetPropertyValue_InContainer(&TestStruct);
+	const UPCGAttributeExtractorTestObject* ObjectDataPtr = Cast<UPCGAttributeExtractorTestObject>(ObjectPtr.Get());
+	UTEST_NOT_NULL(TEXT("Object is not null"), ObjectDataPtr);
+
+	// Static analysis precaution
+	if (!ObjectDataPtr)
+	{
+		TestStruct.Object->MarkAsGarbage();
+		return false;
+	}
+
+	FPCGAttributeAccessorKeysSingleObjectPtr ObjectKey(ObjectDataPtr);
+
+	const FPCGAttributeExtractorTestStructDepth1* DepthStructPtr = DepthStructProperty->ContainerPtrToValuePtr<FPCGAttributeExtractorTestStructDepth1>(&TestStruct);
+
+	// Static analysis precaution
+	if (!DepthStructPtr)
+	{
+		TestStruct.Object->MarkAsGarbage();
+		return false;
+	}
+
+	FPCGAttributeAccessorKeysSingleObjectPtr DepthStructKey(DepthStructPtr);
+	FPCGAttributeAccessorKeysSingleObjectPtr DepthStruct2Key(Depth2StructProperty->ContainerPtrToValuePtr<FPCGAttributeExtractorTestStructDepth2>(DepthStructPtr));
+
+	double DoubleValue = 1.2;
+	int IntValue = 9;
+	float FloatValue = 2.14f;
+
+	UTEST_TRUE(TEXT("Get double value on object"), ObjectDoubleAccessor->Get<double>(DoubleValue, ObjectKey));
+
+	// Allow constructible for both, as property wrappers will create double/int64 accessors
+	UTEST_TRUE(TEXT("Get float value on depth struct"), DepthStructFloatAccessor->Get<float>(FloatValue, DepthStructKey, EPCGAttributeAccessorFlags::AllowConstructible));
+	UTEST_TRUE(TEXT("Get int value on depth 2 struct"), Depth2StructIntAccessor->Get<int>(IntValue, DepthStruct2Key, EPCGAttributeAccessorFlags::AllowConstructible));
+
+	UTEST_EQUAL(TEXT("Double value the same"), DoubleValue, TestStruct.Object->DoubleValue);
+	UTEST_EQUAL(TEXT("Float value the same"), FloatValue, TestStruct.DepthStruct.FloatValue);
+	UTEST_EQUAL(TEXT("Int value the same"), IntValue, TestStruct.DepthStruct.Depth2Struct.IntValue);
+
+	TestStruct.Object->MarkAsGarbage();
+	return true;
+}
+
+bool FPCGAttributePropertyMultipleDepthRangeTest::RunTest(const FString& Parameters)
+{
+	constexpr int NbItems = 5;
+	TArray<FPCGAttributeExtractorTestStruct, TInlineAllocator<NbItems>> TestStructs{};
+	TestStructs.SetNum(NbItems);
+
+	for (int i = 0; i < NbItems; ++i)
+	{
+		TestStructs[i].DepthStruct.Depth2Struct.IntValue = i + 1;
+	}
+
+	const FProperty* DepthStructProperty = FPCGAttributeExtractorTestStruct::StaticStruct()->FindPropertyByName(GET_MEMBER_NAME_CHECKED(FPCGAttributeExtractorTestStruct, DepthStruct));
+	const FProperty* Depth2StructProperty = FPCGAttributeExtractorTestStructDepth1::StaticStruct()->FindPropertyByName(GET_MEMBER_NAME_CHECKED(FPCGAttributeExtractorTestStructDepth1, Depth2Struct));
+
+	UTEST_NOT_NULL(TEXT("Valid depth struct property"), DepthStructProperty);
+	UTEST_NOT_NULL(TEXT("Valid depth struct 2 property"), Depth2StructProperty);
+
+	// Static analysis caution
+	if (!DepthStructProperty || !Depth2StructProperty)
+	{
+		return false;
+	}
+
+	TUniquePtr<IPCGAttributeAccessor> Depth2StructIntAccessor = PCGAttributeAccessorHelpers::CreatePropertyAccessor(GET_MEMBER_NAME_CHECKED(FPCGAttributeExtractorTestStructDepth2, IntValue), FPCGAttributeExtractorTestStructDepth2::StaticStruct());
+	UTEST_TRUE(TEXT("Valid depth struct 2 int accessor"), Depth2StructIntAccessor.IsValid());
+
+	// We will erase the type of struct along the way.
+	// We explicitly know the types, but in a runtime case, we won't be able to know (not known at compile time).
+	// So test this usecase here.
+	TArray<const void*, TInlineAllocator<NbItems>> Depth2StructPtrs;
+	Depth2StructPtrs.SetNum(NbItems);
+
+	for (int i = 0; i < NbItems; ++i)
+	{
+		const void* DepthStructPtr = DepthStructProperty->ContainerPtrToValuePtr<void>(&TestStructs[i]);
+
+		// Static analysis precaution
+		if (!DepthStructPtr)
+		{
+			return false;
+		}
+
+		Depth2StructPtrs[i] = Depth2StructProperty->ContainerPtrToValuePtr<void>(DepthStructPtr);
+	}
+
+	FPCGAttributeAccessorKeysGenericPtrs GenericKeys(Depth2StructPtrs);
+
+	TArray<int, TInlineAllocator<NbItems>> IntValues;
+	IntValues.SetNum(NbItems);
+
+	UTEST_TRUE(TEXT("Get int values on depth 2 struct"), Depth2StructIntAccessor->GetRange<int>(IntValues, 0, GenericKeys, EPCGAttributeAccessorFlags::AllowConstructible));
+	for (int i = 0; i < NbItems; ++i)
+	{
+		UTEST_EQUAL(FString::Printf(TEXT("Int value #%d is the same"), i), IntValues[i], TestStructs[i].DepthStruct.Depth2Struct.IntValue);
+	}
+
+	return true;
+}
+
+#endif // WITH_EDITOR
\ No newline at end of file
diff --git a/Engine/Plugins/Experimental/PCG/Source/PCG/Public/Metadata/Accessors/PCGAttributeAccessorHelpers.h b/Engine/Plugins/Experimental/PCG/Source/PCG/Public/Metadata/Accessors/PCGAttributeAccessorHelpers.h
--- a/Engine/Plugins/Experimental/PCG/Source/PCG/Public/Metadata/Accessors/PCGAttributeAccessorHelpers.h
+++ b/Engine/Plugins/Experimental/PCG/Source/PCG/Public/Metadata/Accessors/PCGAttributeAccessorHelpers.h
@@ -26,6 +26,8 @@ namespace PCGAttributeAccessorHelpers
 	TUniquePtr<IPCGAttributeAccessor> CreatePropertyAccessor(const FName InPropertyName, const UClass* InClass);
 	TUniquePtr<IPCGAttributeAccessor> CreatePropertyAccessor(const FName InPropertyName, const UStruct* InStruct);
 
+	TUniquePtr<IPCGAttributeAccessor> CreateChainAccessor(TUniquePtr<IPCGAttributeAccessor> InAccessor, FName Name, bool& bOutSuccess);
+
 	/**
 	* Create a const accessor depending on an overridable param
 	*/
diff --git a/Engine/Plugins/Experimental/PCG/Source/PCG/Public/Metadata/Accessors/PCGAttributeAccessorKeys.h b/Engine/Plugins/Experimental/PCG/Source/PCG/Public/Metadata/Accessors/PCGAttributeAccessorKeys.h
--- a/Engine/Plugins/Experimental/PCG/Source/PCG/Public/Metadata/Accessors/PCGAttributeAccessorKeys.h
+++ b/Engine/Plugins/Experimental/PCG/Source/PCG/Public/Metadata/Accessors/PCGAttributeAccessorKeys.h
@@ -290,6 +290,41 @@ protected:
 
 /////////////////////////////////////////////////////////////////
 
+/**
+* Type erasing generic keys. Allow to store void* keys, if we are dealing with addresses instead
+* of plain objects.
+* We can't use FPCGAttributeAccessorKeysGeneric since it has a constructor taking a reference on a object,
+* and you can't have void&.
+*/
+class FPCGAttributeAccessorKeysGenericPtrs : public IPCGAttributeAccessorKeys
+{
+public:
+	FPCGAttributeAccessorKeysGenericPtrs(const TArrayView<void*>& InPtrs)
+		: IPCGAttributeAccessorKeys(/*bInReadOnly=*/ false)
+		, Ptrs(InPtrs)
+	{}
+
+	FPCGAttributeAccessorKeysGenericPtrs(const TArrayView<const void*>& InPtrs)
+		: IPCGAttributeAccessorKeys(/*bInReadOnly=*/ true)
+		, Ptrs(const_cast<void**>(InPtrs.GetData()), InPtrs.Num())
+	{}
+
+	virtual int32 GetNum() const override { return Ptrs.Num(); }
+
+protected:
+	virtual bool GetGenericObjectKeys(int32 InStart, TArrayView<void*>& OutObjects) override
+	{
+		return PCGAttributeAccessorKeys::GetKeys(Ptrs, InStart, OutObjects, [](void* Ptr) -> void* { return Ptr; });
+	}
+
+	virtual bool GetGenericObjectKeys(int32 InStart, TArrayView<const void*>& OutObjects) const override
+	{
+		return PCGAttributeAccessorKeys::GetKeys(Ptrs, InStart, OutObjects, [](const void* Ptr) -> const void* { return Ptr; });
+	}
+
+	TArrayView<void*> Ptrs;
+};
+
 template <typename ObjectType>
 inline bool IPCGAttributeAccessorKeys::GetKeys(int32 InStart, TArrayView<ObjectType*>& OutKeys)
 {
diff --git a/Engine/Plugins/Experimental/PCG/Source/PCG/Public/Metadata/Accessors/PCGAttributeExtractor.h b/Engine/Plugins/Experimental/PCG/Source/PCG/Public/Metadata/Accessors/PCGAttributeExtractor.h
--- a/Engine/Plugins/Experimental/PCG/Source/PCG/Public/Metadata/Accessors/PCGAttributeExtractor.h
+++ b/Engine/Plugins/Experimental/PCG/Source/PCG/Public/Metadata/Accessors/PCGAttributeExtractor.h
@@ -15,6 +15,7 @@ namespace PCGAttributeExtractorConstants
 	const FName VectorZ = TEXT("Z");
 	const FName VectorW = TEXT("W");
 	const FName VectorLength = TEXT("Length");
+	const FName VectorSize = TEXT("Size");
 
 	const FName VectorExtractorNames[] =
 	{
@@ -22,7 +23,8 @@ namespace PCGAttributeExtractorConstants
 		VectorY,
 		VectorZ,
 		VectorW,
-		VectorLength
+		VectorLength,
+		VectorSize,
 	};
 
 	// Transform
diff --git a/Engine/Plugins/Experimental/PCG/Source/PCG/Public/Tests/Accessor/PCGAttributeExtractorTest.h b/Engine/Plugins/Experimental/PCG/Source/PCG/Public/Tests/Accessor/PCGAttributeExtractorTest.h
new file mode 100644
--- /dev/null
+++ b/Engine/Plugins/Experimental/PCG/Source/PCG/Public/Tests/Accessor/PCGAttributeExtractorTest.h
@@ -0,0 +1,50 @@
+// Copyright Epic Games, Inc. All Rights Reserved.
+
+#pragma once
+
+#include "UObject/UnrealType.h"
+
+#include "PCGAttributeExtractorTest.generated.h"
+
+UCLASS(meta = (Hidden))
+class UPCGAttributeExtractorTestObject : public UObject
+{
+	GENERATED_BODY()
+
+public:
+	UPROPERTY()
+	double DoubleValue = 0;
+};
+
+USTRUCT(meta = (Hidden))
+struct FPCGAttributeExtractorTestStructDepth2
+{
+	GENERATED_BODY()
+
+	UPROPERTY()
+	int32 IntValue = 0;
+};
+
+USTRUCT(meta = (Hidden))
+struct FPCGAttributeExtractorTestStructDepth1
+{
+	GENERATED_BODY()
+
+	UPROPERTY()
+	FPCGAttributeExtractorTestStructDepth2 Depth2Struct;
+
+	UPROPERTY()
+	float FloatValue;
+};
+
+USTRUCT(meta = (Hidden))
+struct FPCGAttributeExtractorTestStruct
+{
+	GENERATED_BODY()
+
+	UPROPERTY()
+	FPCGAttributeExtractorTestStructDepth1 DepthStruct;
+
+	UPROPERTY()
+	TObjectPtr<UPCGAttributeExtractorTestObject> Object;
+};
\ No newline at end of file
