[PCG] Fixing undo/redo issue with OverridableParams

* Properties list was transient and need to be recomputed after a undo/redo

diff --git a/Engine/Plugins/Experimental/PCG/Source/PCG/Private/PCGSettings.cpp b/Engine/Plugins/Experimental/PCG/Source/PCG/Private/PCGSettings.cpp
--- a/Engine/Plugins/Experimental/PCG/Source/PCG/Private/PCGSettings.cpp
+++ b/Engine/Plugins/Experimental/PCG/Source/PCG/Private/PCGSettings.cpp
@@ -91,16 +91,22 @@ void UPCGSettings::ApplyDeprecation(UPCGNode* InOutNode)
 {
 	DataVersion = FPCGCustomVersion::LatestVersion;
 }
+
+void UPCGSettings::PostEditUndo()
+{
+	// CachedOverridableParams was reset to previous value
+	// Therefore we need to rebuild the properties array since it is transient.
+	InitializeCachedOverridableParams();
+
+	Super::PostEditUndo();
+}
 #endif // WITH_EDITOR
 
 void UPCGSettings::PostLoad()
 {
 	Super::PostLoad();
 
-	if (!HasAnyFlags(RF_ClassDefaultObject))
-	{
-		InitializeCachedOverridableParams();
-	}
+	InitializeCachedOverridableParams();
 
 #if WITH_EDITOR
 	if (ExecutionMode_DEPRECATED != EPCGSettingsExecutionMode::Enabled)
@@ -117,10 +123,7 @@ void UPCGSettings::PostInitProperties()
 	Super::PostInitProperties();
 
 #if WITH_EDITOR
-	if (!HasAnyFlags(RF_ClassDefaultObject))
-	{
-		InitializeCachedOverridableParams();
-	}
+	InitializeCachedOverridableParams();
 #endif //WITH_EDITOR
 }
 
@@ -334,6 +337,12 @@ void UPCGSettings::PostEditChangeProperty(FPropertyChangedEvent& PropertyChanged
 {
 	Super::PostEditChangeProperty(PropertyChangedEvent);
 
+	if (PropertyChangedEvent.Property && (PropertyChangedEvent.Property->GetFName() == GET_MEMBER_NAME_CHECKED(UPCGSettings, CachedOverridableParams)))
+	{
+		// Need to rebuild properties, if it ever changes.
+		InitializeCachedOverridableParams();
+	}
+
 	if (PropertyChangedEvent.GetPropertyName() != GET_MEMBER_NAME_CHECKED(UPCGSettings, DeterminismSettings))
 	{
 		OnSettingsChangedDelegate.Broadcast(this, IsStructuralProperty(PropertyChangedEvent.GetPropertyName()) ? EPCGChangeType::Structural : EPCGChangeType::Settings);
@@ -471,6 +480,12 @@ namespace PCGSettings
 
 void UPCGSettings::InitializeCachedOverridableParams()
 {
+	// Don't do it for default object
+	if (HasAnyFlags(RF_ClassDefaultObject))
+	{
+		return;
+	}
+
 	TRACE_CPUPROFILER_EVENT_SCOPE(UPCGSettings::InitializeCachedOverridableParams);
 
 #if WITH_EDITOR
diff --git a/Engine/Plugins/Experimental/PCG/Source/PCG/Public/PCGSettings.h b/Engine/Plugins/Experimental/PCG/Source/PCG/Public/PCGSettings.h
--- a/Engine/Plugins/Experimental/PCG/Source/PCG/Public/PCGSettings.h
+++ b/Engine/Plugins/Experimental/PCG/Source/PCG/Public/PCGSettings.h
@@ -120,6 +120,9 @@ public:
 	virtual void Serialize(FArchive& Ar) override;
 	virtual void PostSaveRoot(FObjectPostSaveRootContext ObjectSaveContext) override;
 	virtual void PostInitProperties() override;
+#if WITH_EDITOR
+	virtual void PostEditUndo() override;
+#endif // WITH_EDITOR
 	//~End UObject interface
 
 	// TODO: check if we need this to be virtual, we don't really need if we're always caching
