[PCG] Was still adding Params pin even if it existed, making it assert

#rnx

diff --git a/Engine/Plugins/Experimental/PCG/Source/PCG/Private/PCGSettings.cpp b/Engine/Plugins/Experimental/PCG/Source/PCG/Private/PCGSettings.cpp
--- a/Engine/Plugins/Experimental/PCG/Source/PCG/Private/PCGSettings.cpp
+++ b/Engine/Plugins/Experimental/PCG/Source/PCG/Private/PCGSettings.cpp
@@ -233,11 +233,14 @@ void UPCGSettings::FillOverridableParamsPins(TArray<FPCGPinProperties>& OutPins)
 	// Adding the multi-pin connection for params.
 	// If it already exists (and is the correct type), we can keep it.
 	const EPCGDataType* PinType = InputPinsLabelsAndTypes.Find(PCGPinConstants::DefaultParamsLabel);
-	if (PinType && (*PinType != EPCGDataType::Param))
+	if (PinType)
 	{
-		const FString ParamsName = PCGPinConstants::DefaultParamsLabel.ToString();
-		UE_LOG(LogPCG, Error, TEXT("[%s-%s] While adding %s pin, we found another %s pin with not the same allowed type (Param). "
-			"Please rename this pin if you want to take advantage of automatic override. Until then it will probably break your graph."), *GraphName, *NodeName, *ParamsName, *ParamsName);
+		if (*PinType != EPCGDataType::Param)
+		{
+			const FString ParamsName = PCGPinConstants::DefaultParamsLabel.ToString();
+			UE_LOG(LogPCG, Error, TEXT("[%s-%s] While adding %s pin, we found another %s pin with not the same allowed type (Param). "
+				"Please rename this pin if you want to take advantage of automatic override. Until then it will probably break your graph."), *GraphName, *NodeName, *ParamsName, *ParamsName);
+		}
 	}
 	else
 	{
