[PCG] Second attempt at static analysis fix - cache the owner pointer..

diff --git a/Engine/Plugins/Experimental/PCG/Source/PCG/Private/PCGComponent.cpp b/Engine/Plugins/Experimental/PCG/Source/PCG/Private/PCGComponent.cpp
--- a/Engine/Plugins/Experimental/PCG/Source/PCG/Private/PCGComponent.cpp
+++ b/Engine/Plugins/Experimental/PCG/Source/PCG/Private/PCGComponent.cpp
@@ -608,7 +608,8 @@ bool UPCGComponent::MoveResourcesToNewActor(AActor* InNewActor, bool bCreateChil
 		return false;
 	}
 
-	if (!GetOwner())
+	const AActor* Owner = GetOwner();
+	if (!Owner)
 	{
 		UE_LOG(LogPCG, Error, TEXT("[UPCGComponent::MoveResourcesToNewActor] Owner is null, child actor not created."));
 		return false;
@@ -623,7 +624,7 @@ bool UPCGComponent::MoveResourcesToNewActor(AActor* InNewActor, bool bCreateChil
 
 	if (bCreateChild)
 	{
-		NewActor = UPCGActorHelpers::SpawnDefaultActor(GetWorld(), NewActor->GetClass(), TEXT("PCGStampChild"), GetOwner()->GetTransform(), NewActor);
+		NewActor = UPCGActorHelpers::SpawnDefaultActor(GetWorld(), NewActor->GetClass(), TEXT("PCGStampChild"), Owner->GetTransform(), NewActor);
 		check(NewActor);
 	}
 
@@ -640,9 +641,9 @@ bool UPCGComponent::MoveResourcesToNewActor(AActor* InNewActor, bool bCreateChil
 				GeneratedResource->ReleaseIfUnused(Dummy);
 				bHasMovedResources = true;
 			}
-			else if (GetOwner())
+			else
 			{
-				UE_LOG(LogPCG, Error, TEXT("[UPCGComponent::MoveResourcesToNewActor] Null generated resource encountered on actor \"%s\" and will be skipped."), *GetOwner()->GetFName().ToString());
+				UE_LOG(LogPCG, Error, TEXT("[UPCGComponent::MoveResourcesToNewActor] Null generated resource encountered on actor \"%s\" and will be skipped."), *Owner->GetFName().ToString());
 			}
 		}
 
