[PCG] Split pins on Volume Sampler to meet newer standards.

Also some fixes spotted for Surface Sampler. The deprecation code stopped working because the new params setup adds more input pins before the deprecation function is called.

diff --git a/Engine/Plugins/Experimental/PCG/Source/PCG/Private/Data/PCGWorldData.cpp b/Engine/Plugins/Experimental/PCG/Source/PCG/Private/Data/PCGWorldData.cpp
--- a/Engine/Plugins/Experimental/PCG/Source/PCG/Private/Data/PCGWorldData.cpp
+++ b/Engine/Plugins/Experimental/PCG/Source/PCG/Private/Data/PCGWorldData.cpp
@@ -128,7 +128,7 @@ const UPCGPointData* UPCGWorldVolumetricData::CreatePointData(FPCGContext* Conte
 	PCGVolumeSampler::FVolumeSamplerSettings SamplerSettings;
 	SamplerSettings.VoxelSize = VoxelSize;
 
-	PCGVolumeSampler::SampleVolume(Context, this, SamplerSettings, Data, EffectiveBounds);
+	PCGVolumeSampler::SampleVolume(Context, this, nullptr, EffectiveBounds, SamplerSettings, Data);
 	UE_LOG(LogPCG, Verbose, TEXT("Volumetric world extracted %d points"), Data->GetPoints().Num());
 	
 	return Data;
diff --git a/Engine/Plugins/Experimental/PCG/Source/PCG/Private/Elements/PCGSurfaceSampler.cpp b/Engine/Plugins/Experimental/PCG/Source/PCG/Private/Elements/PCGSurfaceSampler.cpp
--- a/Engine/Plugins/Experimental/PCG/Source/PCG/Private/Elements/PCGSurfaceSampler.cpp
+++ b/Engine/Plugins/Experimental/PCG/Source/PCG/Private/Elements/PCGSurfaceSampler.cpp
@@ -8,13 +8,13 @@
 #include "PCGEdge.h"
 #include "PCGGraph.h"
 #include "PCGHelpers.h"
+#include "PCGPin.h"
 #include "Data/PCGPointData.h"
 #include "Data/PCGSpatialData.h"
 #include "Helpers/PCGAsync.h"
 #include "Helpers/PCGSettingsHelpers.h"
 
 #include "Math/RandomStream.h"
-#include "PCGPin.h"
 
 #include UE_INLINE_GENERATED_CPP_BY_NAME(PCGSurfaceSampler)
 
@@ -232,7 +232,7 @@ FText UPCGSurfaceSamplerSettings::GetNodeTooltipText() const
 TArray<FPCGPinProperties> UPCGSurfaceSamplerSettings::InputPinProperties() const
 {
 	TArray<FPCGPinProperties> PinProperties;
-	PinProperties.Emplace(PCGSurfaceSamplerConstants::SurfaceLabel, EPCGDataType::Surface, /*bAllowMultipleConnections=*/true, /*bAllowMultipleData=*/false, LOCTEXT("SurfaceSamplerSurfacePinTooltip",
+	PinProperties.Emplace(PCGSurfaceSamplerConstants::SurfaceLabel, EPCGDataType::Surface, /*bAllowMultipleConnections=*/true, /*bAllowMultipleData=*/true, LOCTEXT("SurfaceSamplerSurfacePinTooltip",
 		"The surface to sample with points. Points will be generated in the two dimensional footprint of the combined bounds of the Surface and the Bounding Shape (if any) "
 		"and then projected onto this surface. If this input is omitted then the network of shapes connected to the Bounding Shape pin will be inspected for a surface "
 		"shape to use to project the points onto."
@@ -336,7 +336,7 @@ bool FPCGSurfaceSamplerElement::ExecuteInternal(FPCGContext* Context) const
 		{
 			GeneratingShapes.Add(GeneratorFromBoundingShapeInput);
 
-			check(BoundingShapeInputs.Num() == 1);
+			check(BoundingShapeInputs.Num() > 0);
 			Outputs.Add(BoundingShapeInputs[0]);
 		}
 	}
@@ -369,15 +369,15 @@ bool FPCGSurfaceSamplerElement::ExecuteInternal(FPCGContext* Context) const
 		{
 			InputBounds = PCGHelpers::OverlapBounds(InputBounds, BoundingShapeBounds);
 		}
-		if (!InputBounds.IsValid)
-		{
-			PCGE_LOG(Warning, "Input data has invalid bounds");
-			continue;
-		}
 
 		PCGSurfaceSampler::FSurfaceSamplerSettings LoopData;
-		if (!LoopData.Initialize(Settings, Context, InputBounds))
+		if (!InputBounds.IsValid || !LoopData.Initialize(Settings, Context, InputBounds))
 		{
+			if (!InputBounds.IsValid)
+			{
+				PCGE_LOG(Verbose, "Input data has invalid bounds");
+			}
+
 			Outputs.RemoveAt(GenerationIndex);
 			GeneratingShapes.RemoveAt(GenerationIndex);
 			--GenerationIndex;
@@ -424,7 +424,7 @@ void UPCGSurfaceSamplerSettings::ApplyDeprecationBeforeUpdatePins(UPCGNode* InOu
 {
 	if (DataVersion < FPCGCustomVersion::SplitSamplerNodesInputs && ensure(InOutNode))
 	{
-		if (InputPins.Num() == 1)
+		if (InputPins.Num() > 0 && InputPins[0])
 		{
 			// The node will function the same if we move all connections from "In" to "Bounding Shape". To make this happen, rename "In" to
 			// "Bounding Shape" just prior to pin update and the edges will be moved over. In ApplyDeprecation we'll see if we can do better than
diff --git a/Engine/Plugins/Experimental/PCG/Source/PCG/Private/Elements/PCGVolumeSampler.cpp b/Engine/Plugins/Experimental/PCG/Source/PCG/Private/Elements/PCGVolumeSampler.cpp
--- a/Engine/Plugins/Experimental/PCG/Source/PCG/Private/Elements/PCGVolumeSampler.cpp
+++ b/Engine/Plugins/Experimental/PCG/Source/PCG/Private/Elements/PCGVolumeSampler.cpp
@@ -1,52 +1,67 @@
 // Copyright Epic Games, Inc. All Rights Reserved.
 
 #include "Elements/PCGVolumeSampler.h"
+
+#include "PCGContext.h"
+#include "PCGCustomVersion.h"
+#include "PCGHelpers.h"
 #include "Data/PCGPointData.h"
 #include "Data/PCGSpatialData.h"
 #include "Helpers/PCGAsync.h"
-#include "PCGContext.h"
-#include "PCGHelpers.h"
 
 #include UE_INLINE_GENERATED_CPP_BY_NAME(PCGVolumeSampler)
 
+#define LOCTEXT_NAMESPACE "PCGVolumeSamplerElement"
+
+namespace PCGVolumeSamplerConstants
+{
+	const FName VolumeLabel = TEXT("Volume");
+	const FName BoundingShapeLabel = TEXT("Bounding Shape");
+}
+
 namespace PCGVolumeSampler
 {
-	UPCGPointData* SampleVolume(FPCGContext* Context, const UPCGSpatialData* SpatialData, const FVolumeSamplerSettings& SamplerSettings)
+	UPCGPointData* SampleVolume(FPCGContext* InContext, const UPCGSpatialData* InVolume, const FVolumeSamplerSettings& InSamplerSettings)
+	{
+		const FBox Bounds = InVolume->GetBounds();
+		return SampleVolume(InContext, InVolume, nullptr, Bounds, InSamplerSettings);
+	}
+
+	UPCGPointData* SampleVolume(FPCGContext* InContext, const UPCGSpatialData* InVolume, const UPCGSpatialData* InBoundingShape, const FBox& InBounds, const FVolumeSamplerSettings& InSamplerSettings)
 	{
 		UPCGPointData* Data = NewObject<UPCGPointData>();
-		Data->InitializeFromData(SpatialData);
-		const FBox Bounds = SpatialData->GetBounds();
+		Data->InitializeFromData(InVolume);
 
-		SampleVolume(Context, SpatialData, SamplerSettings, Data, Bounds);
+		SampleVolume(InContext, InVolume, InBoundingShape, InBounds, InSamplerSettings, Data);
 
 		return Data;
 	}
 
-	void SampleVolume(FPCGContext* Context, const UPCGSpatialData* SpatialData, const FVolumeSamplerSettings& SamplerSettings, UPCGPointData* OutputData, const FBox& Bounds)
+	void SampleVolume(FPCGContext* InContext, const UPCGSpatialData* InVolume, const UPCGSpatialData* InBoundingShape, const FBox& InBounds, const FVolumeSamplerSettings& InSamplerSettings, UPCGPointData* OutputData)
 	{
-		check(SpatialData && OutputData);
+		check(InVolume && OutputData);
 
 		// Early out
-		if (!Bounds.IsValid)
+		if (!InBounds.IsValid)
 		{
 			return;
 		}
 
 		TArray<FPCGPoint>& Points = OutputData->GetMutablePoints();
-		const FVector& VoxelSize = SamplerSettings.VoxelSize;
+		const FVector& VoxelSize = InSamplerSettings.VoxelSize;
 
-		const int32 MinX = FMath::CeilToInt(Bounds.Min.X / VoxelSize.X);
-		const int32 MaxX = FMath::FloorToInt(Bounds.Max.X / VoxelSize.X);
-		const int32 MinY = FMath::CeilToInt(Bounds.Min.Y / VoxelSize.Y);
-		const int32 MaxY = FMath::FloorToInt(Bounds.Max.Y / VoxelSize.Y);
-		const int32 MinZ = FMath::CeilToInt(Bounds.Min.Z / VoxelSize.Z);
-		const int32 MaxZ = FMath::FloorToInt(Bounds.Max.Z / VoxelSize.Z);
+		const int32 MinX = FMath::CeilToInt(InBounds.Min.X / VoxelSize.X);
+		const int32 MaxX = FMath::FloorToInt(InBounds.Max.X / VoxelSize.X);
+		const int32 MinY = FMath::CeilToInt(InBounds.Min.Y / VoxelSize.Y);
+		const int32 MaxY = FMath::FloorToInt(InBounds.Max.Y / VoxelSize.Y);
+		const int32 MinZ = FMath::CeilToInt(InBounds.Min.Z / VoxelSize.Z);
+		const int32 MaxZ = FMath::FloorToInt(InBounds.Max.Z / VoxelSize.Z);
 
 		if (MinX > MaxX || MinY > MaxY || MinZ > MaxZ)
 		{
-			if (Context)
+			if (InContext)
 			{
-				PCGE_LOG_C(Verbose, Context, "Skipped - invalid cell bounds");
+				PCGE_LOG_C(Verbose, InContext, "Skipped - invalid cell bounds");
 			}
 
 			return;
@@ -54,7 +69,7 @@ namespace PCGVolumeSampler
 
 		const int32 NumIterations = (MaxX - MinX) * (MaxY - MinY) * (MaxZ - MinZ);
 
-		FPCGAsync::AsyncPointProcessing(Context, NumIterations, Points, [SpatialData, VoxelSize, MinX, MaxX, MinY, MaxY, MinZ, MaxZ](int32 Index, FPCGPoint& OutPoint)
+		FPCGAsync::AsyncPointProcessing(InContext, NumIterations, Points, [InVolume, InBoundingShape, VoxelSize, MinX, MaxX, MinY, MaxY, MinZ, MaxZ](int32 Index, FPCGPoint& OutPoint)
 		{
 			const int X = MinX + (Index % (MaxX - MinX));
 			const int Y = MinY + (Index / (MaxX - MinX) % (MaxY - MinY));
@@ -62,8 +77,19 @@ namespace PCGVolumeSampler
 
 			const FVector SampleLocation(X * VoxelSize.X, Y * VoxelSize.Y, Z * VoxelSize.Z);
 			const FBox VoxelBox(VoxelSize * -0.5, VoxelSize * 0.5);
-			if (SpatialData->SamplePoint(FTransform(SampleLocation), VoxelBox, OutPoint, nullptr))
+
+			const FTransform SampleTransform = FTransform(SampleLocation);
+			if (InVolume->SamplePoint(SampleTransform, VoxelBox, OutPoint, nullptr))
 			{
+				if (InBoundingShape)
+				{
+					FPCGPoint BoundingShapeSample;
+					if (!InBoundingShape->SamplePoint(SampleTransform, VoxelBox, BoundingShapeSample, nullptr))
+					{
+						return false;
+					}
+				}
+
 				OutPoint.Seed = PCGHelpers::ComputeSeed(X, Y, Z);
 				return true;
 			}
@@ -75,10 +101,24 @@ namespace PCGVolumeSampler
 	}
 }
 
+#if WITH_EDITOR
+FText UPCGVolumeSamplerSettings::GetNodeTooltipText() const
+{
+	return LOCTEXT("VolumeSamplerNodeTooltip", "Generates points in the three dimensional bounds of the Volume input and within the Bounding Shape input if provided.");
+}
+#endif
+
 TArray<FPCGPinProperties> UPCGVolumeSamplerSettings::InputPinProperties() const
 {
 	TArray<FPCGPinProperties> PinProperties;
-	PinProperties.Emplace(PCGPinConstants::DefaultInputLabel, EPCGDataType::Spatial);
+	// Spatial is ok - volume sampling just needs bounds.
+	PinProperties.Emplace(PCGVolumeSamplerConstants::VolumeLabel, EPCGDataType::Spatial, /*bAllowMultipleConnections=*/true, /*bAllowMultipleData=*/true, LOCTEXT("VolumeSamplerVolumePinTooltip",
+		"The volume to sample with points. Can be any spatial data that can provide bounds."
+	));
+	// Only one connection allowed, user can union multiple shapes.
+	PinProperties.Emplace(PCGVolumeSamplerConstants::BoundingShapeLabel, EPCGDataType::Spatial, /*bInAllowMultipleConnections=*/false, /*bAllowMultipleData=*/false, LOCTEXT("VolumeSamplerBoundingShapePinTooltip",
+		"Optional. All sampled points must be contained within this shape."
+	));
 
 	return PinProperties;
 }
@@ -95,7 +135,7 @@ bool FPCGVolumeSamplerElement::ExecuteInternal(FPCGContext* Context) const
 	const UPCGVolumeSamplerSettings* Settings = Context->GetInputSettings<UPCGVolumeSamplerSettings>();
 	check(Settings);
 
-	TArray<FPCGTaggedData> Inputs = Context->InputData.GetInputs();
+	TArray<FPCGTaggedData> VolumeInputs = Context->InputData.GetInputsByPin(PCGVolumeSamplerConstants::VolumeLabel);
 
 	const FVector& VoxelSize = Settings->VoxelSize;
 	if (VoxelSize.X <= 0 || VoxelSize.Y <= 0 || VoxelSize.Z <= 0)
@@ -109,21 +149,78 @@ bool FPCGVolumeSamplerElement::ExecuteInternal(FPCGContext* Context) const
 
 	TArray<FPCGTaggedData>& Outputs = Context->OutputData.TaggedData;
 
+	// Grab the Bounding Shape input if there is one.
+	TArray<FPCGTaggedData> BoundingShapeInputs = Context->InputData.GetInputsByPin(PCGVolumeSamplerConstants::BoundingShapeLabel);
+	const UPCGSpatialData* BoundingShapeSpatialInput = nullptr;
+	if (BoundingShapeInputs.Num() > 0)
+	{
+		ensure(BoundingShapeInputs.Num() == 1);
+		BoundingShapeSpatialInput = Cast<UPCGSpatialData>(BoundingShapeInputs[0].Data);
+
+		if (!BoundingShapeSpatialInput)
+		{
+			PCGE_LOG(Warning, "Bounding Shape input is missing or of unsupported type and will not be used.");
+		}
+	}
+
+	// Compute bounds of bounding shape input
+	FBox BoundingShapeBounds(EForceInit::ForceInit);
+	if (BoundingShapeSpatialInput)
+	{
+		BoundingShapeBounds = BoundingShapeSpatialInput->GetBounds();
+	}
+
+	// Construct a list of shapes to generate samples from. Prefer to get these directly from the first input pin.
+	TArray<const UPCGSpatialData*, TInlineAllocator<16>> GeneratingShapes;
+	for (FPCGTaggedData& TaggedData : VolumeInputs)
+	{
+		if (UPCGSpatialData* SpatialData = Cast<UPCGSpatialData>(TaggedData.Data))
+		{
+			GeneratingShapes.Add(SpatialData);
+			Outputs.Add(TaggedData);
+		}
+	}
+
+	// If no shapes were obtained from the first input pin, try to find a shape to sample from nodes connected to the second pin.
+	if (GeneratingShapes.Num() == 0 && BoundingShapeSpatialInput)
+	{
+		GeneratingShapes.Add(BoundingShapeSpatialInput);
+		check(BoundingShapeInputs.Num() > 0);
+		Outputs.Add(BoundingShapeInputs[0]);
+	}
+
+	// Warn if something is connected but no spatial data could be obtained for sampling
+	if (GeneratingShapes.Num() == 0 && (BoundingShapeInputs.Num() > 0 || VolumeInputs.Num() > 0))
+	{
+		PCGE_LOG(Warning, "No spatial data shape was provided for sampling, no points will be produced.");
+	}
+
 	// TODO: embarassingly parallel loop
-	for (const FPCGTaggedData& Input : Inputs)
+	for (int GenerationIndex = 0; GenerationIndex < GeneratingShapes.Num(); ++GenerationIndex)
 	{
-		const UPCGSpatialData* SpatialInput = Cast<UPCGSpatialData>(Input.Data);
+		const UPCGSpatialData* GeneratingShape = GeneratingShapes[GenerationIndex];
+		check(GeneratingShape);
 
-		if (!SpatialInput)
+		// Calculate the intersection of bounds of the provided inputs
+		FBox InputBounds = GeneratingShape->GetBounds();
+		if (BoundingShapeBounds.IsValid)
 		{
-			PCGE_LOG(Error, "Invalid input data");
+			InputBounds = PCGHelpers::OverlapBounds(InputBounds, BoundingShapeBounds);
+		}
+
+		if (!InputBounds.IsValid)
+		{
+			PCGE_LOG(Verbose, "Input data has invalid bounds");
+
+			Outputs.RemoveAt(GenerationIndex);
+			GeneratingShapes.RemoveAt(GenerationIndex);
+			--GenerationIndex;
 			continue;
 		}
 
-		FPCGTaggedData& Output = Outputs.Emplace_GetRef();
-		Output = Input;
-		UPCGPointData* SampledData = PCGVolumeSampler::SampleVolume(Context, SpatialInput, SamplerSettings);
-		Output.Data = SampledData;
+		// Sample volume
+		const UPCGPointData* SampledData = PCGVolumeSampler::SampleVolume(Context, GeneratingShape, BoundingShapeSpatialInput, InputBounds, SamplerSettings);
+		Outputs[GenerationIndex].Data = SampledData;
 
 		if (SampledData)
 		{
@@ -136,3 +233,21 @@ bool FPCGVolumeSamplerElement::ExecuteInternal(FPCGContext* Context) const
 
 	return true;
 }
+
+#if WITH_EDITOR
+void UPCGVolumeSamplerSettings::ApplyDeprecationBeforeUpdatePins(UPCGNode* InOutNode, TArray<TObjectPtr<UPCGPin>>& InputPins, TArray<TObjectPtr<UPCGPin>>& OutputPins)
+{
+	if (DataVersion < FPCGCustomVersion::SplitVolumeSamplerNodeInputs && ensure(InOutNode))
+	{
+		if (InputPins.Num() > 0 && InputPins[0])
+		{
+			// First pin renamed in this version. Rename here so that edges won't get culled in UpdatePins later.
+			InputPins[0]->Properties.Label = PCGVolumeSamplerConstants::VolumeLabel;
+		}
+	}
+
+	Super::ApplyDeprecationBeforeUpdatePins(InOutNode, InputPins, OutputPins);
+}
+#endif // WITH_EDITOR
+
+#undef LOCTEXT_NAMESPACE
diff --git a/Engine/Plugins/Experimental/PCG/Source/PCG/Public/Elements/PCGVolumeSampler.h b/Engine/Plugins/Experimental/PCG/Source/PCG/Public/Elements/PCGVolumeSampler.h
--- a/Engine/Plugins/Experimental/PCG/Source/PCG/Public/Elements/PCGVolumeSampler.h
+++ b/Engine/Plugins/Experimental/PCG/Source/PCG/Public/Elements/PCGVolumeSampler.h
@@ -18,8 +18,10 @@ namespace PCGVolumeSampler
 		FVector VoxelSize;
 	};
 
-	UPCGPointData* SampleVolume(FPCGContext* Context, const UPCGSpatialData* SpatialData, const FVolumeSamplerSettings& SamplerSettings);
-	void SampleVolume(FPCGContext* Context, const UPCGSpatialData* SpatialData, const FVolumeSamplerSettings& SamplerSettings, UPCGPointData* OutputData, const FBox& Bounds);
+	UPCGPointData* SampleVolume(FPCGContext* InContext, const UPCGSpatialData* InVolume, const FVolumeSamplerSettings& InSamplerSettings);
+	UPCGPointData* SampleVolume(FPCGContext* InContext, const UPCGSpatialData* InVolume, const UPCGSpatialData* InBoundingShape, const FBox& InBounds, const FVolumeSamplerSettings& InSamplerSettings);
+
+	void SampleVolume(FPCGContext* InContext, const UPCGSpatialData* InVolume, const UPCGSpatialData* InBoundingShape, const FBox& InBounds, const FVolumeSamplerSettings& InSamplerSettings, UPCGPointData* InOutputData);
 }
 
 UCLASS(BlueprintType, ClassGroup = (Procedural))
@@ -33,13 +35,18 @@ public:
 
 	//~Begin UPCGSettings interface
 #if WITH_EDITOR
-	virtual FName GetDefaultNodeName() const override { return FName(TEXT("VolumeSampler")); }
+	virtual FName GetDefaultNodeName() const override { return FName(TEXT("Volume Sampler")); }
+	virtual FText GetNodeTooltipText() const override;
 	virtual EPCGSettingsType GetType() const override { return EPCGSettingsType::Sampler; }
 #endif
 
 	virtual TArray<FPCGPinProperties> InputPinProperties() const override;
 	virtual TArray<FPCGPinProperties> OutputPinProperties() const override { return Super::DefaultPointOutputPinProperties(); }
 	
+#if WITH_EDITOR
+	virtual void ApplyDeprecationBeforeUpdatePins(UPCGNode* InOutNode, TArray<TObjectPtr<UPCGPin>>& InputPins, TArray<TObjectPtr<UPCGPin>>& OutputPins) override;
+#endif
+
 protected:
 	virtual FPCGElementPtr CreateElement() const override;
 	//~End UPCGSettings interface
diff --git a/Engine/Plugins/Experimental/PCG/Source/PCG/Public/PCGCustomVersion.h b/Engine/Plugins/Experimental/PCG/Source/PCG/Public/PCGCustomVersion.h
--- a/Engine/Plugins/Experimental/PCG/Source/PCG/Public/PCGCustomVersion.h
+++ b/Engine/Plugins/Experimental/PCG/Source/PCG/Public/PCGCustomVersion.h
@@ -29,6 +29,9 @@ struct PCG_API FPCGCustomVersion
 		// Add param pin for all nodes that have override and were using the default input pin.
 		AddParamPinToOverridableNodes = 6,
 
+		// Sampling shape and bounding shape inputs.
+		SplitVolumeSamplerNodeInputs = 7,
+
 		// -----<new versions can be added above this line>-------------------------------------------------
 		VersionPlusOne,
 		LatestVersion = VersionPlusOne - 1
