[PCG] Modernize Spline Sampler, add Bounding Shape input

diff --git a/Engine/Plugins/Experimental/PCG/Source/PCG/Private/Data/PCGLandscapeSplineData.cpp b/Engine/Plugins/Experimental/PCG/Source/PCG/Private/Data/PCGLandscapeSplineData.cpp
--- a/Engine/Plugins/Experimental/PCG/Source/PCG/Private/Data/PCGLandscapeSplineData.cpp
+++ b/Engine/Plugins/Experimental/PCG/Source/PCG/Private/Data/PCGLandscapeSplineData.cpp
@@ -134,7 +134,7 @@ const UPCGPointData* UPCGLandscapeSplineData::CreatePointData(FPCGContext* Conte
 	SamplerParams.Mode = EPCGSplineSamplingMode::Distance;
 	SamplerParams.Dimension = EPCGSplineSamplingDimension::OnHorizontal;
 
-	PCGSplineSampler::SampleLineData(this, this, SamplerParams, Data);
+	PCGSplineSampler::SampleLineData(this, this, nullptr, SamplerParams, Data);
 
 	UE_LOG(LogPCG, Verbose, TEXT("Landscape spline %s generated %d points"), *Spline->GetFName().ToString(), Points.Num());
 
diff --git a/Engine/Plugins/Experimental/PCG/Source/PCG/Private/Data/PCGSplineData.cpp b/Engine/Plugins/Experimental/PCG/Source/PCG/Private/Data/PCGSplineData.cpp
--- a/Engine/Plugins/Experimental/PCG/Source/PCG/Private/Data/PCGSplineData.cpp
+++ b/Engine/Plugins/Experimental/PCG/Source/PCG/Private/Data/PCGSplineData.cpp
@@ -67,7 +67,7 @@ const UPCGPointData* UPCGSplineData::CreatePointData(FPCGContext* Context) const
 	FPCGSplineSamplerParams SamplerParams;
 	SamplerParams.Mode = EPCGSplineSamplingMode::Distance;
 
-	PCGSplineSampler::SampleLineData(this, this, SamplerParams, Data);
+	PCGSplineSampler::SampleLineData(this, this, nullptr, SamplerParams, Data);
 
 	UE_LOG(LogPCG, Verbose, TEXT("Spline %s generated %d points"), *Spline->GetFName().ToString(), Data->GetPoints().Num());
 
diff --git a/Engine/Plugins/Experimental/PCG/Source/PCG/Private/Elements/PCGSplineSampler.cpp b/Engine/Plugins/Experimental/PCG/Source/PCG/Private/Elements/PCGSplineSampler.cpp
--- a/Engine/Plugins/Experimental/PCG/Source/PCG/Private/Elements/PCGSplineSampler.cpp
+++ b/Engine/Plugins/Experimental/PCG/Source/PCG/Private/Elements/PCGSplineSampler.cpp
@@ -2,6 +2,8 @@
 
 #include "Elements/PCGSplineSampler.h"
 
+#include "PCGContext.h"
+#include "PCGCustomVersion.h"
 #include "Data/PCGIntersectionData.h"
 #include "Data/PCGLandscapeSplineData.h"
 #include "Data/PCGPointData.h"
@@ -11,11 +13,18 @@
 #include "Helpers/PCGBlueprintHelpers.h"
 
 #include "Components/SplineComponent.h"
-#include "PCGContext.h"
 #include "Voronoi/Voronoi.h"
 
 #include UE_INLINE_GENERATED_CPP_BY_NAME(PCGSplineSampler)
 
+#define LOCTEXT_NAMESPACE "PCGSplineSamplerElement"
+
+namespace PCGSplineSamplerConstants
+{
+	const FName SplineLabel = TEXT("Spline");
+	const FName BoundingShapeLabel = TEXT("Bounding Shape");
+}
+
 namespace PCGSplineSamplerHelpers
 {
 	/**
@@ -413,17 +422,19 @@ namespace PCGSplineSampler
 
 	struct FDimensionSampler
 	{
-		FDimensionSampler(const UPCGPolyLineData* InLineData, const UPCGSpatialData* InSpatialData)
+		FDimensionSampler(const UPCGPolyLineData* InLineData, const UPCGSpatialData* InSpatialData, const UPCGSpatialData* InBoundingShapeData)
 		{
 			check(InLineData);
 			LineData = InLineData;
 			SpatialData = InSpatialData;
+			BoundingShapeData = InBoundingShapeData;
 		}
 
 		virtual void Sample(const FTransform& InTransform, const FBox& InBox, UPCGPointData* OutPointData)
 		{
-			FPCGPoint TrivialPoint;
-			if(SpatialData->SamplePoint(InTransform, InBox, TrivialPoint, OutPointData->Metadata))
+			FPCGPoint TrivialPoint, BoundsTestPoint;
+			if(SpatialData->SamplePoint(InTransform, InBox, TrivialPoint, OutPointData->Metadata)
+				&& (!BoundingShapeData || BoundingShapeData->SamplePoint(InTransform, InBox, BoundsTestPoint, nullptr)))
 			{
 				OutPointData->GetMutablePoints().Add(TrivialPoint);
 			}
@@ -431,12 +442,14 @@ namespace PCGSplineSampler
 
 		const UPCGPolyLineData* LineData;
 		const UPCGSpatialData* SpatialData;
+		const UPCGSpatialData* BoundingShapeData;
 	};
 
+	/** Samples in a volume surrounding the poly line. */
 	struct FVolumeSampler : public FDimensionSampler
 	{
-		FVolumeSampler(const UPCGPolyLineData* InLineData, const UPCGSpatialData* InSpatialData, const FPCGSplineSamplerParams& Params)
-			: FDimensionSampler(InLineData, InSpatialData)
+		FVolumeSampler(const UPCGPolyLineData* InLineData, const UPCGSpatialData* InSpatialData, const UPCGSpatialData* InBoundingShapeData, const FPCGSplineSamplerParams& Params)
+			: FDimensionSampler(InLineData, InSpatialData, InBoundingShapeData)
 		{
 			Fill = Params.Fill;
 			NumPlanarSteps = 1 + ((Params.Dimension == EPCGSplineSamplingDimension::OnVertical) ? 0 : Params.NumPlanarSubdivisions);
@@ -502,8 +515,10 @@ namespace PCGSplineSampler
 							FTransform TentativeTransform = InTransform;
 							TentativeTransform.SetLocation(TentativeLocation);
 
-							FPCGPoint OutPoint;
-							if (SpatialData->SamplePoint(TentativeTransform, SubBox, OutPoint, OutPointData->Metadata))
+							// Sample point and if successful check against bounding shape (if provided).
+							FPCGPoint OutPoint, BoundsTestPoint;
+							if (SpatialData->SamplePoint(TentativeTransform, SubBox, OutPoint, OutPointData->Metadata)
+								&& (!BoundingShapeData || BoundingShapeData->SamplePoint(TentativeTransform, SubBox, BoundsTestPoint, nullptr)))
 							{
 								OutPointData->GetMutablePoints().Add(OutPoint);
 							}
@@ -522,7 +537,7 @@ namespace PCGSplineSampler
 		int NumHeightSteps;
 	};
 
-	void SampleLineData(const UPCGPolyLineData* LineData, const UPCGSpatialData* SpatialData, const FPCGSplineSamplerParams& Params, UPCGPointData* OutPointData)
+	void SampleLineData(const UPCGPolyLineData* LineData, const UPCGSpatialData* SpatialData, const UPCGSpatialData* InBoundingShapeData, const FPCGSplineSamplerParams& Params, UPCGPointData* OutPointData)
 	{
 		check(LineData && OutPointData);
 
@@ -531,8 +546,8 @@ namespace PCGSplineSampler
 
 		FStepSampler* Sampler = ((Params.Mode == EPCGSplineSamplingMode::Subdivision) ? static_cast<FStepSampler*>(&SubdivisionSampler) : static_cast<FStepSampler*>(&DistanceSampler));
 
-		FDimensionSampler TrivialDimensionSampler(LineData, SpatialData);
-		FVolumeSampler VolumeSampler(LineData, SpatialData, Params);
+		FDimensionSampler TrivialDimensionSampler(LineData, SpatialData, InBoundingShapeData);
+		FVolumeSampler VolumeSampler(LineData, SpatialData, InBoundingShapeData, Params);
 
 		FDimensionSampler* ExtentsSampler = ((Params.Dimension == EPCGSplineSamplingDimension::OnSpline) ? &TrivialDimensionSampler : static_cast<FDimensionSampler*>(&VolumeSampler));
 
@@ -554,7 +569,7 @@ namespace PCGSplineSampler
 		}
 	}
  
-	void SampleInteriorData(const UPCGPolyLineData* LineData, const UPCGSpatialData* SpatialData, const FPCGSplineSamplerParams& Params, UPCGPointData* OutPointData)
+	void SampleInteriorData(const UPCGPolyLineData* LineData, const UPCGSpatialData* SpatialData, const UPCGSpatialData* InBoundingShape, const FPCGSplineSamplerParams& Params, UPCGPointData* OutPointData)
 	{
 		check(LineData && OutPointData);
 
@@ -880,19 +895,30 @@ namespace PCGSplineSampler
 	}
 }
 
-FPCGElementPtr UPCGSplineSamplerSettings::CreateElement() const
+#if WITH_EDITOR
+FText UPCGSplineSamplerSettings::GetNodeTooltipText() const
 {
-	return MakeShared<FPCGSplineSamplerElement>();
+	return LOCTEXT("SplineSamplerNodeTooltip", "Generates points along the given Spline, and within the Bounding Shape if provided.");
 }
+#endif
 
 TArray<FPCGPinProperties> UPCGSplineSamplerSettings::InputPinProperties() const
 {
 	TArray<FPCGPinProperties> PinProperties;
-	PinProperties.Emplace(PCGPinConstants::DefaultInputLabel, /*EPCGDataType::Point |*/ EPCGDataType::Spline | EPCGDataType::LandscapeSpline);
+	PinProperties.Emplace(PCGSplineSamplerConstants::SplineLabel, EPCGDataType::PolyLine, /*bAllowMultipleConnections=*/true, /*bAllowMultipleData=*/true);
+	// Only one connection allowed, user can union multiple shapes.
+	PinProperties.Emplace(PCGSplineSamplerConstants::BoundingShapeLabel, EPCGDataType::Spatial, /*bInAllowMultipleConnections=*/false, /*bAllowMultipleData=*/false, LOCTEXT("SplineSamplerBoundingShapePinTooltip",
+		"Optional. All sampled points must be contained within this shape."
+	));
 
 	return PinProperties;
 }
 
+FPCGElementPtr UPCGSplineSamplerSettings::CreateElement() const
+{
+	return MakeShared<FPCGSplineSamplerElement>();
+}
+
 bool FPCGSplineSamplerElement::ExecuteInternal(FPCGContext* Context) const
 {
 	TRACE_CPUPROFILER_EVENT_SCOPE(FPCGSplineSamplerElement::Execute);
@@ -900,13 +926,27 @@ bool FPCGSplineSamplerElement::ExecuteInternal(FPCGContext* Context) const
 	const UPCGSplineSamplerSettings* Settings = Context->GetInputSettings<UPCGSplineSamplerSettings>();
 	check(Settings);
 
-	TArray<FPCGTaggedData> Inputs = Context->InputData.GetInputs();
+	TArray<FPCGTaggedData> SplineInputs = Context->InputData.GetInputsByPin(PCGSplineSamplerConstants::SplineLabel);
+
+	// Grab the Bounding Shape input if there is one.
+	TArray<FPCGTaggedData> BoundingShapeInputs = Context->InputData.GetInputsByPin(PCGSplineSamplerConstants::BoundingShapeLabel);
+	const UPCGSpatialData* BoundingShapeSpatialInput = nullptr;
+	if (BoundingShapeInputs.Num() > 0)
+	{
+		ensure(BoundingShapeInputs.Num() == 1);
+		BoundingShapeSpatialInput = Cast<UPCGSpatialData>(BoundingShapeInputs[0].Data);
+
+		if (!BoundingShapeSpatialInput)
+		{
+			PCGE_LOG(Warning, "Bounding Shape input is missing or of unsupported type and will not be used.");
+		}
+	}
 
 	const FPCGSplineSamplerParams& SamplerParams = Settings->SamplerParams;
 
 	TArray<FPCGTaggedData>& Outputs = Context->OutputData.TaggedData;
 
-	for (FPCGTaggedData& Input : Inputs)
+	for (FPCGTaggedData& Input : SplineInputs)
 	{
 		const UPCGSpatialData* SpatialData = Cast<const UPCGSpatialData>(Input.Data);
 
@@ -932,14 +972,31 @@ bool FPCGSplineSamplerElement::ExecuteInternal(FPCGContext* Context) const
 
 		if (SamplerParams.Dimension == EPCGSplineSamplingDimension::OnInterior)
 		{
-			PCGSplineSampler::SampleInteriorData(LineData, SpatialData, SamplerParams, SampledPointData);
+			PCGSplineSampler::SampleInteriorData(LineData, SpatialData, BoundingShapeSpatialInput, SamplerParams, SampledPointData);
 		}
 		else
 		{
-			PCGSplineSampler::SampleLineData(LineData, SpatialData, SamplerParams, SampledPointData);
+			PCGSplineSampler::SampleLineData(LineData, SpatialData, BoundingShapeSpatialInput, SamplerParams, SampledPointData);
 		}
 	}
 
 	return true;
 }
 
+#if WITH_EDITOR
+void UPCGSplineSamplerSettings::ApplyDeprecationBeforeUpdatePins(UPCGNode* InOutNode, TArray<TObjectPtr<UPCGPin>>& InputPins, TArray<TObjectPtr<UPCGPin>>& OutputPins)
+{
+	if (DataVersion < FPCGCustomVersion::SplineSamplerUpdatedNodeInputs && ensure(InOutNode))
+	{
+		if (InputPins.Num() > 0 && InputPins[0])
+		{
+			// First pin renamed in this version. Rename here so that edges won't get culled in UpdatePins later.
+			InputPins[0]->Properties.Label = PCGSplineSamplerConstants::SplineLabel;
+		}
+	}
+
+	Super::ApplyDeprecationBeforeUpdatePins(InOutNode, InputPins, OutputPins);
+}
+#endif // WITH_EDITOR
+
+#undef LOCTEXT_NAMESPACE
diff --git a/Engine/Plugins/Experimental/PCG/Source/PCG/Public/Elements/PCGSplineSampler.h b/Engine/Plugins/Experimental/PCG/Source/PCG/Public/Elements/PCGSplineSampler.h
--- a/Engine/Plugins/Experimental/PCG/Source/PCG/Public/Elements/PCGSplineSampler.h
+++ b/Engine/Plugins/Experimental/PCG/Source/PCG/Public/Elements/PCGSplineSampler.h
@@ -103,8 +103,9 @@ namespace PCGSplineSamplerHelpers
 
 namespace PCGSplineSampler
 {
-	void SampleLineData(const UPCGPolyLineData* LineData, const UPCGSpatialData* SpatialData, const FPCGSplineSamplerParams& Params, UPCGPointData* OutPointData);
-	void SampleInteriorData(const UPCGPolyLineData* LineData, const UPCGSpatialData* SpatialData, const FPCGSplineSamplerParams& Params, UPCGPointData* OutPointData);
+	void SampleLineData(const UPCGPolyLineData* LineData, const UPCGSpatialData* SpatialData, const UPCGSpatialData* InBoundingShape, const FPCGSplineSamplerParams& Params, UPCGPointData* OutPointData);
+
+	void SampleInteriorData(const UPCGPolyLineData* LineData, const UPCGSpatialData* SpatialData, const UPCGSpatialData* InBoundingShape, const FPCGSplineSamplerParams& Params, UPCGPointData* OutPointData);
 	const UPCGPolyLineData* GetPolyLineData(const UPCGSpatialData* InSpatialData);
 }
 
@@ -116,13 +117,18 @@ class PCG_API UPCGSplineSamplerSettings : public UPCGSettings
 public:
 	// ~Begin UPCGSettings interface
 #if WITH_EDITOR
-	virtual FName GetDefaultNodeName() const override { return FName(TEXT("SplineSampler")); }
+	virtual FName GetDefaultNodeName() const override { return FName(TEXT("Spline Sampler")); }
+	virtual FText GetNodeTooltipText() const override;
 	virtual EPCGSettingsType GetType() const override { return EPCGSettingsType::Sampler; }
 #endif
 
 	virtual TArray<FPCGPinProperties> InputPinProperties() const override;
 	virtual TArray<FPCGPinProperties> OutputPinProperties() const override { return Super::DefaultPointOutputPinProperties(); }
 
+#if WITH_EDITOR
+	virtual void ApplyDeprecationBeforeUpdatePins(UPCGNode* InOutNode, TArray<TObjectPtr<UPCGPin>>& InputPins, TArray<TObjectPtr<UPCGPin>>& OutputPins) override;
+#endif
+
 protected:
 	virtual FPCGElementPtr CreateElement() const override;
 	// ~End UPCGSettings interface
diff --git a/Engine/Plugins/Experimental/PCG/Source/PCG/Public/PCGCustomVersion.h b/Engine/Plugins/Experimental/PCG/Source/PCG/Public/PCGCustomVersion.h
--- a/Engine/Plugins/Experimental/PCG/Source/PCG/Public/PCGCustomVersion.h
+++ b/Engine/Plugins/Experimental/PCG/Source/PCG/Public/PCGCustomVersion.h
@@ -32,6 +32,9 @@ struct PCG_API FPCGCustomVersion
 		// Sampling shape and bounding shape inputs.
 		SplitVolumeSamplerNodeInputs = 7,
 
+		// Renamed spline input and added bounding shape input.
+		SplineSamplerUpdatedNodeInputs = 8,
+
 		// -----<new versions can be added above this line>-------------------------------------------------
 		VersionPlusOne,
 		LatestVersion = VersionPlusOne - 1
