[PCG] Tentative fix for non-Unity build, removed header by IWYU

#rnx

diff --git a/Engine/Plugins/Experimental/PCG/Source/PCG/Public/Metadata/PCGMetadataAttributeTraits.h b/Engine/Plugins/Experimental/PCG/Source/PCG/Public/Metadata/PCGMetadataAttributeTraits.h
--- a/Engine/Plugins/Experimental/PCG/Source/PCG/Public/Metadata/PCGMetadataAttributeTraits.h
+++ b/Engine/Plugins/Experimental/PCG/Source/PCG/Public/Metadata/PCGMetadataAttributeTraits.h
@@ -2,6 +2,8 @@
 
 #pragma once
 
+// We need core minimal because we are calling methods from Vector, Quat, etc...
+#include "CoreMinimal.h" // IWYU pragma: keep
 
 #include "PCGMetadataAttributeTraits.generated.h"
 
@@ -781,7 +783,3 @@ namespace PCG
 		};
 	}
 }
-
-#if UE_ENABLE_INCLUDE_ORDER_DEPRECATED_IN_5_2
-#include "CoreMinimal.h"
-#endif
