[PCG] Improvements to Attribute extractor

* Fixing regex for vector extractor
* Adding R,G,B and A as components for vector (cannot mix with XYZW though)
* Adding Forward, Right, Up to Rotator extractor
* Adding support for rotator extraction on quaternion.
* Added tests to cover those new possibilities.

diff --git a/Engine/Plugins/Experimental/PCG/Source/PCG/Private/Metadata/Accessors/PCGAttributeAccessorHelpers.cpp b/Engine/Plugins/Experimental/PCG/Source/PCG/Private/Metadata/Accessors/PCGAttributeAccessorHelpers.cpp
--- a/Engine/Plugins/Experimental/PCG/Source/PCG/Private/Metadata/Accessors/PCGAttributeAccessorHelpers.cpp
+++ b/Engine/Plugins/Experimental/PCG/Source/PCG/Private/Metadata/Accessors/PCGAttributeAccessorHelpers.cpp
@@ -64,7 +64,7 @@ namespace PCGAttributeAccessorHelpers
 		{
 			using AccessorType = decltype(Dummy);
 
-			if constexpr (PCG::Private::IsOfTypes<AccessorType, FVector2D, FVector, FVector4, FQuat>())
+			if constexpr (PCG::Private::IsOfTypes<AccessorType, FVector2D, FVector, FVector4>())
 			{
 				return PCGAttributeExtractor::CreateVectorExtractor<AccessorType>(std::move(Accessor), Name, bOutSuccess);
 			}
@@ -72,6 +72,10 @@ namespace PCGAttributeAccessorHelpers
 			{
 				return PCGAttributeExtractor::CreateTransformExtractor(std::move(Accessor), Name, bOutSuccess);
 			}
+			if constexpr (PCG::Private::IsOfTypes<AccessorType, FQuat>())
+			{
+				return PCGAttributeExtractor::CreateQuatExtractor(std::move(Accessor), Name, bOutSuccess);
+			}
 			if constexpr (PCG::Private::IsOfTypes<AccessorType, FRotator>())
 			{
 				return PCGAttributeExtractor::CreateRotatorExtractor(std::move(Accessor), Name, bOutSuccess);
diff --git a/Engine/Plugins/Experimental/PCG/Source/PCG/Private/Metadata/Accessors/PCGAttributeExtractor.cpp b/Engine/Plugins/Experimental/PCG/Source/PCG/Private/Metadata/Accessors/PCGAttributeExtractor.cpp
--- a/Engine/Plugins/Experimental/PCG/Source/PCG/Private/Metadata/Accessors/PCGAttributeExtractor.cpp
+++ b/Engine/Plugins/Experimental/PCG/Source/PCG/Private/Metadata/Accessors/PCGAttributeExtractor.cpp
@@ -67,57 +67,78 @@ namespace PCGAttributeExtractor
 				[](const VectorType& Value) -> double { return Value.Size(); });
 		}
 
-		FString MatchString;
+		const FString NameStr = Name.ToString();
+		const int32 NameLen = NameStr.Len();
+
+		// Name should have at least 1 char and max 4
+		if (NameLen <= 0 || NameLen > 4)
+		{
+			// Failed
+			bOutSuccess = false;
+			return InAccessor;
+		}
+
+		// We should match only X,Y,Z or W, for the whole name.
+		// We also support R, G, B and A. But we cannot mix both.
+		FString MatchStringXYZW;
+		FString MatchStringRGBA;
 		if constexpr (std::is_same_v<FVector2D, VectorType>)
 		{
-			MatchString = TEXT("[XY]{1,4}");
+			MatchStringXYZW = FString::Printf(TEXT("[XY]{%d}"), NameLen);
+			MatchStringRGBA = FString::Printf(TEXT("[RG]{%d}"), NameLen);
 		}
 		else if constexpr (std::is_same_v<FVector, VectorType>)
 		{
-			MatchString = TEXT("[XYZ]{1,4}");
+			MatchStringXYZW = FString::Printf(TEXT("[XYZ]{%d}"), NameLen);
+			MatchStringRGBA = FString::Printf(TEXT("[RGB]{%d}"), NameLen);
 		}
 		else
 		{
-			MatchString = TEXT("[XYZW]{1,4}");
+			MatchStringXYZW = FString::Printf(TEXT("[XYZW]{%d}"), NameLen);
+			MatchStringRGBA = FString::Printf(TEXT("[RGBA]{%d}"), NameLen);
 		}
 
-		const FString NameStr = Name.ToString();
-
-		const FRegexPattern RegexPattern(MatchString, ERegexPatternFlags::CaseInsensitive);
-		FRegexMatcher RegexMatcher(RegexPattern, *NameStr);
-		if (!RegexMatcher.FindNext())
-		{
-			// Failed
-			bOutSuccess = false;
-			return InAccessor;
-		}
+		const FRegexPattern RegexPatternXYZW(MatchStringXYZW, ERegexPatternFlags::CaseInsensitive);
+		const FRegexPattern RegexPatternRGBA(MatchStringRGBA, ERegexPatternFlags::CaseInsensitive);
 
 		TArray<int32, TInlineAllocator<4>> Indexes;
-		for (const TCHAR Char : NameStr)
-		{
-			int32 Index = -1;
 
-			if (Char == TEXT('w') || Char == TEXT('W'))
-			{
-				Index = 3;
-			}
-			else
+		FRegexMatcher RegexMatcherXYZW(RegexPatternXYZW, *NameStr);
+		FRegexMatcher RegexMatcherRGBA(RegexPatternRGBA, *NameStr);
+
+		if (RegexMatcherXYZW.FindNext() || RegexMatcherRGBA.FindNext())
+		{
+			for (const TCHAR Char : NameStr)
 			{
-				Index = static_cast<int32>(Char) - static_cast<int32>(TEXT('x'));
-				if (Index < 0)
+				if (Char == 'R' || Char == 'r' || Char == 'X' || Char == 'x')
 				{
-					Index = static_cast<int32>(Char) - static_cast<int32>(TEXT('X'));
+					Indexes.Add(0);
+				}
+				else if (Char == 'G' || Char == 'g' || Char == 'Y' || Char == 'y')
+				{
+					Indexes.Add(1);
+				}
+				else if (Char == 'B' || Char == 'b' || Char == 'Z' || Char == 'z')
+				{
+					Indexes.Add(2);
+				}
+				else if (Char == 'A' || Char == 'a' || Char == 'W' || Char == 'w')
+				{
+					Indexes.Add(3);
+				}
+				else // Safeguard, should be caught by the regex
+				{
+					ensure(false);
+					bOutSuccess = false;
+					return InAccessor;
 				}
 			}
-
-			// Safeguard, should be caught by not matching the regex above
-			if (!ensure(Index >= 0 && Index < 4))
-			{
-				bOutSuccess = false;
-				return InAccessor;
-			}
-
-			Indexes.Add(Index);
+		}
+		else
+		{
+			// Failed
+			bOutSuccess = false;
+			return InAccessor;
 		}
 
 		bOutSuccess = true;
@@ -173,10 +194,77 @@ namespace PCGAttributeExtractor
 				[](FRotator& Value, const double& In) -> void { Value.Yaw = In; });
 		}
 
+		// Read-only
+		if (Name == PCGAttributeExtractorConstants::RotatorForward)
+		{
+			return MakeUnique<FPCGChainAccessor<FVector, FRotator>>(std::move(InAccessor),
+				[](const FRotator& Value) -> FVector { return Value.Quaternion().GetForwardVector(); });
+		}
+
+		if (Name == PCGAttributeExtractorConstants::RotatorRight)
+		{
+			return MakeUnique<FPCGChainAccessor<FVector, FRotator>>(std::move(InAccessor),
+				[](const FRotator& Value) -> FVector { return Value.Quaternion().GetRightVector(); });
+		}
+
+		if (Name == PCGAttributeExtractorConstants::RotatorUp)
+		{
+			return MakeUnique<FPCGChainAccessor<FVector, FRotator>>(std::move(InAccessor),
+				[](const FRotator& Value) -> FVector { return Value.Quaternion().GetUpVector(); });
+		}
+
 		bOutSuccess = false;
 		return InAccessor;
 	}
 
+	TUniquePtr<IPCGAttributeAccessor> CreateQuatExtractor(TUniquePtr<IPCGAttributeAccessor> InAccessor, FName Name, bool& bOutSuccess)
+	{
+		bOutSuccess = true;
+
+		if (Name == PCGAttributeExtractorConstants::RotatorPitch)
+		{
+			return MakeUnique<FPCGChainAccessor<double, FQuat>>(std::move(InAccessor),
+				[](const FQuat& Value) -> double { return Value.Rotator().Pitch; },
+				[](FQuat& Value, const double& In) -> void { FRotator Temp = Value.Rotator(); Temp.Pitch = In; Value = Temp.Quaternion(); });
+		}
+
+		if (Name == PCGAttributeExtractorConstants::RotatorRoll)
+		{
+			return MakeUnique<FPCGChainAccessor<double, FQuat>>(std::move(InAccessor),
+				[](const FQuat& Value) -> double { return Value.Rotator().Roll; },
+				[](FQuat& Value, const double& In) -> void { FRotator Temp = Value.Rotator(); Temp.Roll = In; Value = Temp.Quaternion(); });
+		}
+
+		if (Name == PCGAttributeExtractorConstants::RotatorYaw)
+		{
+			return MakeUnique<FPCGChainAccessor<double, FQuat>>(std::move(InAccessor),
+				[](const FQuat& Value) -> double { return Value.Rotator().Yaw; },
+				[](FQuat& Value, const double& In) -> void { FRotator Temp = Value.Rotator(); Temp.Yaw = In; Value = Temp.Quaternion(); });
+		}
+
+		// Read-only
+		if (Name == PCGAttributeExtractorConstants::RotatorForward)
+		{
+			return MakeUnique<FPCGChainAccessor<FVector, FQuat>>(std::move(InAccessor),
+				[](const FQuat& Value) -> FVector { return Value.GetForwardVector(); });
+		}
+
+		if (Name == PCGAttributeExtractorConstants::RotatorRight)
+		{
+			return MakeUnique<FPCGChainAccessor<FVector, FQuat>>(std::move(InAccessor),
+				[](const FQuat& Value) -> FVector { return Value.GetRightVector(); });
+		}
+
+		if (Name == PCGAttributeExtractorConstants::RotatorUp)
+		{
+			return MakeUnique<FPCGChainAccessor<FVector, FQuat>>(std::move(InAccessor),
+				[](const FQuat& Value) -> FVector { return Value.GetUpVector(); });
+		}
+
+		// Quat also support vector extractor
+		return CreateVectorExtractor<FQuat>(std::move(InAccessor), Name, bOutSuccess);
+	}
+
 	TUniquePtr<IPCGAttributeAccessor> CreateTransformExtractor(TUniquePtr<IPCGAttributeAccessor> InAccessor, FName Name, bool& bOutSuccess)
 	{
 		bOutSuccess = true;
diff --git a/Engine/Plugins/Experimental/PCG/Source/PCG/Private/Tests/Accessors/PCGAttributeExtractorTest.cpp b/Engine/Plugins/Experimental/PCG/Source/PCG/Private/Tests/Accessors/PCGAttributeExtractorTest.cpp
--- a/Engine/Plugins/Experimental/PCG/Source/PCG/Private/Tests/Accessors/PCGAttributeExtractorTest.cpp
+++ b/Engine/Plugins/Experimental/PCG/Source/PCG/Private/Tests/Accessors/PCGAttributeExtractorTest.cpp
@@ -44,6 +44,7 @@ namespace FPCGAttributeExtractorTestHelpers
 	}
 }
 
+#define PCG_STRUCT_NAME_CHECK(Name) GET_MEMBER_NAME_CHECKED(FPCGMetadataTypesConstantStruct, Name)
 
 IMPLEMENT_CUSTOM_SIMPLE_AUTOMATION_TEST(FPCGAttributeSingleGetPropertyTest, FPCGTestBaseClass, "pcg.tests.Accessor.Property.SimpleGetProperty", PCGTestsCommon::TestFlags)
 IMPLEMENT_CUSTOM_SIMPLE_AUTOMATION_TEST(FPCGAttributeVectorPropertyExtractorTest, FPCGTestBaseClass, "pcg.tests.Accessor.Property.VectorExtractor", PCGTestsCommon::TestFlags)
@@ -99,21 +100,21 @@ bool FPCGAttributeSingleGetPropertyTest::RunTest(const FString& Parameters)
 
 	bool bTestPassed = true;
 
-	bTestPassed &= Verify(double(Struct.FloatValue), GET_MEMBER_NAME_CHECKED(FPCGMetadataTypesConstantStruct, FloatValue));
-	bTestPassed &= Verify(Struct.DoubleValue, GET_MEMBER_NAME_CHECKED(FPCGMetadataTypesConstantStruct, DoubleValue));
-	bTestPassed &= Verify(int64(Struct.Int32Value), GET_MEMBER_NAME_CHECKED(FPCGMetadataTypesConstantStruct, Int32Value));
-	bTestPassed &= Verify(Struct.IntValue, GET_MEMBER_NAME_CHECKED(FPCGMetadataTypesConstantStruct, IntValue));
-	bTestPassed &= Verify(Struct.Vector2Value, GET_MEMBER_NAME_CHECKED(FPCGMetadataTypesConstantStruct, Vector2Value));
-	bTestPassed &= Verify(Struct.VectorValue, GET_MEMBER_NAME_CHECKED(FPCGMetadataTypesConstantStruct, VectorValue));
-	bTestPassed &= Verify(Struct.Vector4Value, GET_MEMBER_NAME_CHECKED(FPCGMetadataTypesConstantStruct, Vector4Value));
-	bTestPassed &= Verify(Struct.BoolValue, GET_MEMBER_NAME_CHECKED(FPCGMetadataTypesConstantStruct, BoolValue));
-	bTestPassed &= Verify(Struct.NameValue, GET_MEMBER_NAME_CHECKED(FPCGMetadataTypesConstantStruct, NameValue));
-	bTestPassed &= Verify(Struct.StringValue, GET_MEMBER_NAME_CHECKED(FPCGMetadataTypesConstantStruct, StringValue));
-	bTestPassed &= Verify(Struct.RotatorValue, GET_MEMBER_NAME_CHECKED(FPCGMetadataTypesConstantStruct, RotatorValue));
-	bTestPassed &= Verify(Struct.QuatValue, GET_MEMBER_NAME_CHECKED(FPCGMetadataTypesConstantStruct, QuatValue));
-	bTestPassed &= Verify(Struct.TransformValue, GET_MEMBER_NAME_CHECKED(FPCGMetadataTypesConstantStruct, TransformValue));
-	bTestPassed &= Verify(Struct.SoftClassPathValue.ToString(), GET_MEMBER_NAME_CHECKED(FPCGMetadataTypesConstantStruct, SoftClassPathValue));
-	bTestPassed &= Verify(Struct.SoftObjectPathValue.ToString(), GET_MEMBER_NAME_CHECKED(FPCGMetadataTypesConstantStruct, SoftObjectPathValue));
+	bTestPassed &= Verify(double(Struct.FloatValue), PCG_STRUCT_NAME_CHECK(FloatValue));
+	bTestPassed &= Verify(Struct.DoubleValue, PCG_STRUCT_NAME_CHECK(DoubleValue));
+	bTestPassed &= Verify(int64(Struct.Int32Value), PCG_STRUCT_NAME_CHECK(Int32Value));
+	bTestPassed &= Verify(Struct.IntValue, PCG_STRUCT_NAME_CHECK(IntValue));
+	bTestPassed &= Verify(Struct.Vector2Value, PCG_STRUCT_NAME_CHECK(Vector2Value));
+	bTestPassed &= Verify(Struct.VectorValue, PCG_STRUCT_NAME_CHECK(VectorValue));
+	bTestPassed &= Verify(Struct.Vector4Value, PCG_STRUCT_NAME_CHECK(Vector4Value));
+	bTestPassed &= Verify(Struct.BoolValue, PCG_STRUCT_NAME_CHECK(BoolValue));
+	bTestPassed &= Verify(Struct.NameValue, PCG_STRUCT_NAME_CHECK(NameValue));
+	bTestPassed &= Verify(Struct.StringValue, PCG_STRUCT_NAME_CHECK(StringValue));
+	bTestPassed &= Verify(Struct.RotatorValue, PCG_STRUCT_NAME_CHECK(RotatorValue));
+	bTestPassed &= Verify(Struct.QuatValue, PCG_STRUCT_NAME_CHECK(QuatValue));
+	bTestPassed &= Verify(Struct.TransformValue, PCG_STRUCT_NAME_CHECK(TransformValue));
+	bTestPassed &= Verify(Struct.SoftClassPathValue.ToString(), PCG_STRUCT_NAME_CHECK(SoftClassPathValue));
+	bTestPassed &= Verify(Struct.SoftObjectPathValue.ToString(), PCG_STRUCT_NAME_CHECK(SoftObjectPathValue));
 
 	TempMetadata->MarkAsGarbage();
 
@@ -149,13 +150,13 @@ bool FPCGAttributeVectorPropertyExtractorTest::RunTest(const FString& Parameters
 	bool bTestPassed = true;
 
 	// Vector Lengths
-	bTestPassed &= Verify(Struct.Vector2Value.Length(), GET_MEMBER_NAME_CHECKED(FPCGMetadataTypesConstantStruct, Vector2Value), PCGAttributeExtractorConstants::VectorLength);
-	bTestPassed &= Verify(Struct.VectorValue.Length(), GET_MEMBER_NAME_CHECKED(FPCGMetadataTypesConstantStruct, VectorValue), PCGAttributeExtractorConstants::VectorLength);
-	bTestPassed &= Verify(Struct.Vector4Value.Size(), GET_MEMBER_NAME_CHECKED(FPCGMetadataTypesConstantStruct, Vector4Value), PCGAttributeExtractorConstants::VectorLength);
+	bTestPassed &= Verify(Struct.Vector2Value.Length(), PCG_STRUCT_NAME_CHECK(Vector2Value), PCGAttributeExtractorConstants::VectorLength);
+	bTestPassed &= Verify(Struct.VectorValue.Length(), PCG_STRUCT_NAME_CHECK(VectorValue), PCGAttributeExtractorConstants::VectorLength);
+	bTestPassed &= Verify(Struct.Vector4Value.Size(), PCG_STRUCT_NAME_CHECK(Vector4Value), PCGAttributeExtractorConstants::VectorLength);
 
-	bTestPassed &= Verify(Struct.Vector2Value.Size(), GET_MEMBER_NAME_CHECKED(FPCGMetadataTypesConstantStruct, Vector2Value), PCGAttributeExtractorConstants::VectorSize);
-	bTestPassed &= Verify(Struct.VectorValue.Size(), GET_MEMBER_NAME_CHECKED(FPCGMetadataTypesConstantStruct, VectorValue), PCGAttributeExtractorConstants::VectorSize);
-	bTestPassed &= Verify(Struct.Vector4Value.Size(), GET_MEMBER_NAME_CHECKED(FPCGMetadataTypesConstantStruct, Vector4Value), PCGAttributeExtractorConstants::VectorSize);
+	bTestPassed &= Verify(Struct.Vector2Value.Size(), PCG_STRUCT_NAME_CHECK(Vector2Value), PCGAttributeExtractorConstants::VectorSize);
+	bTestPassed &= Verify(Struct.VectorValue.Size(), PCG_STRUCT_NAME_CHECK(VectorValue), PCGAttributeExtractorConstants::VectorSize);
+	bTestPassed &= Verify(Struct.Vector4Value.Size(), PCG_STRUCT_NAME_CHECK(Vector4Value), PCGAttributeExtractorConstants::VectorSize);
 
 	// Vector2 Composents
 	auto Vec2Composents = [&Verify, &bTestPassed](const auto& Value, const FName PropertyName) -> bool
@@ -163,61 +164,92 @@ bool FPCGAttributeVectorPropertyExtractorTest::RunTest(const FString& Parameters
 		bTestPassed &= Verify(Value.X, PropertyName, FName(TEXT("X")));
 		bTestPassed &= Verify(Value.Y, PropertyName, FName(TEXT("Y")));
 
+		bTestPassed &= Verify(Value.X, PropertyName, FName(TEXT("R")));
+		bTestPassed &= Verify(Value.Y, PropertyName, FName(TEXT("G")));
+
 		bTestPassed &= Verify(FVector2D{ Value.X, Value.Y }, PropertyName, FName(TEXT("XY")));
 		bTestPassed &= Verify(FVector2D{ Value.Y, Value.X }, PropertyName, FName(TEXT("YX")));
 		bTestPassed &= Verify(FVector2D{ Value.Y, Value.Y }, PropertyName, FName(TEXT("YY")));
 
+		bTestPassed &= Verify(FVector2D{ Value.X, Value.Y }, PropertyName, FName(TEXT("RG")));
+		bTestPassed &= Verify(FVector2D{ Value.Y, Value.X }, PropertyName, FName(TEXT("GR")));
+		bTestPassed &= Verify(FVector2D{ Value.Y, Value.Y }, PropertyName, FName(TEXT("GG")));
+
 		return true;
 	};
 
-	Vec2Composents(Struct.Vector2Value, GET_MEMBER_NAME_CHECKED(FPCGMetadataTypesConstantStruct, Vector2Value));
-	Vec2Composents(Struct.VectorValue, GET_MEMBER_NAME_CHECKED(FPCGMetadataTypesConstantStruct, VectorValue));
-	Vec2Composents(Struct.Vector4Value, GET_MEMBER_NAME_CHECKED(FPCGMetadataTypesConstantStruct, Vector4Value));
-	Vec2Composents(Struct.QuatValue, GET_MEMBER_NAME_CHECKED(FPCGMetadataTypesConstantStruct, QuatValue));
+	Vec2Composents(Struct.Vector2Value, PCG_STRUCT_NAME_CHECK(Vector2Value));
+	Vec2Composents(Struct.VectorValue, PCG_STRUCT_NAME_CHECK(VectorValue));
+	Vec2Composents(Struct.Vector4Value, PCG_STRUCT_NAME_CHECK(Vector4Value));
+	Vec2Composents(Struct.QuatValue, PCG_STRUCT_NAME_CHECK(QuatValue));
 
 	// Vector3 Composents
 	auto Vec3Composents = [&Verify, &bTestPassed](const auto& Value, const FName PropertyName) -> bool
 	{
 		bTestPassed &= Verify(Value.Z, PropertyName, FName(TEXT("Z")));
+		bTestPassed &= Verify(Value.Z, PropertyName, FName(TEXT("B")));
 
 		bTestPassed &= Verify(FVector2D{ Value.Z, Value.X }, PropertyName, FName(TEXT("ZX")));
 		bTestPassed &= Verify(FVector2D{ Value.Y, Value.Z }, PropertyName, FName(TEXT("YZ")));
 
+		bTestPassed &= Verify(FVector2D{ Value.Z, Value.X }, PropertyName, FName(TEXT("BR")));
+		bTestPassed &= Verify(FVector2D{ Value.Y, Value.Z }, PropertyName, FName(TEXT("GB")));
+
 		bTestPassed &= Verify(FVector{ Value.Z, Value.X, Value.Y }, PropertyName, FName(TEXT("ZXY")));
 		bTestPassed &= Verify(FVector{ Value.Y, Value.Y, Value.Y }, PropertyName, FName(TEXT("YYY")));
 		bTestPassed &= Verify(FVector{ Value.Y, Value.Y, Value.X }, PropertyName, FName(TEXT("YYX")));
 		bTestPassed &= Verify(FVector{ Value.X, Value.Y, Value.Z }, PropertyName, FName(TEXT("XYZ")));
 
+		bTestPassed &= Verify(FVector{ Value.Z, Value.X, Value.Y }, PropertyName, FName(TEXT("BRG")));
+		bTestPassed &= Verify(FVector{ Value.Y, Value.Y, Value.Y }, PropertyName, FName(TEXT("GGG")));
+		bTestPassed &= Verify(FVector{ Value.Y, Value.Y, Value.X }, PropertyName, FName(TEXT("GGR")));
+		bTestPassed &= Verify(FVector{ Value.X, Value.Y, Value.Z }, PropertyName, FName(TEXT("RGB")));
+
 		return true;
 	};
 
-	Vec3Composents(Struct.VectorValue, GET_MEMBER_NAME_CHECKED(FPCGMetadataTypesConstantStruct, VectorValue));
-	Vec3Composents(Struct.Vector4Value, GET_MEMBER_NAME_CHECKED(FPCGMetadataTypesConstantStruct, Vector4Value));
-	Vec3Composents(Struct.QuatValue, GET_MEMBER_NAME_CHECKED(FPCGMetadataTypesConstantStruct, QuatValue));
+	Vec3Composents(Struct.VectorValue, PCG_STRUCT_NAME_CHECK(VectorValue));
+	Vec3Composents(Struct.Vector4Value, PCG_STRUCT_NAME_CHECK(Vector4Value));
+	Vec3Composents(Struct.QuatValue, PCG_STRUCT_NAME_CHECK(QuatValue));
 
 	// Vector4 Composents
 	auto Vec4Composents = [&Verify, &bTestPassed](const auto& Value, const FName PropertyName) -> bool
 	{
 		bTestPassed &= Verify(Value.W, PropertyName, FName(TEXT("W")));
 
+		bTestPassed &= Verify(Value.W, PropertyName, FName(TEXT("A")));
+
 		bTestPassed &= Verify(FVector2D{ Value.Z, Value.W }, PropertyName, FName(TEXT("ZW")));
 		bTestPassed &= Verify(FVector2D{ Value.W, Value.Y }, PropertyName, FName(TEXT("WY")));
 
+		bTestPassed &= Verify(FVector2D{ Value.Z, Value.W }, PropertyName, FName(TEXT("BA")));
+		bTestPassed &= Verify(FVector2D{ Value.W, Value.Y }, PropertyName, FName(TEXT("AG")));
+
 		bTestPassed &= Verify(FVector{ Value.Z, Value.W, Value.Y }, PropertyName, FName(TEXT("ZWY")));
 		bTestPassed &= Verify(FVector{ Value.W, Value.W, Value.W }, PropertyName, FName(TEXT("WWW")));
 		bTestPassed &= Verify(FVector{ Value.Y, Value.W, Value.X }, PropertyName, FName(TEXT("YWX")));
 		bTestPassed &= Verify(FVector{ Value.X, Value.Y, Value.W }, PropertyName, FName(TEXT("XYW")));
 
+		bTestPassed &= Verify(FVector{ Value.Z, Value.W, Value.Y }, PropertyName, FName(TEXT("BAG")));
+		bTestPassed &= Verify(FVector{ Value.W, Value.W, Value.W }, PropertyName, FName(TEXT("AAA")));
+		bTestPassed &= Verify(FVector{ Value.Y, Value.W, Value.X }, PropertyName, FName(TEXT("GAR")));
+		bTestPassed &= Verify(FVector{ Value.X, Value.Y, Value.W }, PropertyName, FName(TEXT("RGA")));
+
 		bTestPassed &= Verify(FVector4{ Value.Z, Value.X, Value.Y, Value.W }, PropertyName, FName(TEXT("ZXYW")));
 		bTestPassed &= Verify(FVector4{ Value.Z, Value.Z, Value.Z, Value.Z }, PropertyName, FName(TEXT("ZZZZ")));
 		bTestPassed &= Verify(FVector4{ Value.X, Value.Y, Value.X, Value.Y }, PropertyName, FName(TEXT("XYXY")));
 		bTestPassed &= Verify(FVector4{ Value.X, Value.Y, Value.Z, Value.W }, PropertyName, FName(TEXT("XYZW")));
 
+		bTestPassed &= Verify(FVector4{ Value.Z, Value.X, Value.Y, Value.W }, PropertyName, FName(TEXT("BRGA")));
+		bTestPassed &= Verify(FVector4{ Value.Z, Value.Z, Value.Z, Value.Z }, PropertyName, FName(TEXT("BBBB")));
+		bTestPassed &= Verify(FVector4{ Value.X, Value.Y, Value.X, Value.Y }, PropertyName, FName(TEXT("RGRG")));
+		bTestPassed &= Verify(FVector4{ Value.X, Value.Y, Value.Z, Value.W }, PropertyName, FName(TEXT("RGBA")));
+
 		return true;
 	};
 
-	Vec4Composents(Struct.Vector4Value, GET_MEMBER_NAME_CHECKED(FPCGMetadataTypesConstantStruct, Vector4Value));
-	Vec4Composents(Struct.QuatValue, GET_MEMBER_NAME_CHECKED(FPCGMetadataTypesConstantStruct, QuatValue));
+	Vec4Composents(Struct.Vector4Value, PCG_STRUCT_NAME_CHECK(Vector4Value));
+	Vec4Composents(Struct.QuatValue, PCG_STRUCT_NAME_CHECK(QuatValue));
 
 	return bTestPassed;
 }
@@ -226,8 +258,10 @@ bool FPCGAttributeRotatorPropertyExtractorTest::RunTest(const FString & Paramete
 {
 	FPCGMetadataTypesConstantStruct Struct = FPCGAttributeExtractorTestHelpers::CreateStruct();
 
-	auto Verify = [this, &Struct](double ExpectedValue, const FName PropertyName, const FName ExtractorName) -> bool
+	auto Verify = [this, &Struct](auto ExpectedValue, const FName PropertyName, const FName ExtractorName) -> bool
 	{
+		using ExpectedType = decltype(ExpectedValue);
+
 		TUniquePtr<IPCGAttributeAccessor> Accessor = PCGAttributeAccessorHelpers::CreatePropertyAccessor(PropertyName, FPCGMetadataTypesConstantStruct::StaticStruct());
 		UTEST_TRUE(FString::Printf(TEXT("Valid accessor for property %s"), *PropertyName.ToString()), Accessor.IsValid());
 
@@ -235,12 +269,12 @@ bool FPCGAttributeRotatorPropertyExtractorTest::RunTest(const FString & Paramete
 		TUniquePtr<IPCGAttributeAccessor> ExtractorAccessor = PCGAttributeAccessorHelpers::CreateChainAccessor(std::move(Accessor), ExtractorName, bSuccess);
 		UTEST_TRUE(TEXT("Valid extractor accessor"), (bSuccess && ExtractorAccessor.IsValid()));
 		UTEST_TRUE(TEXT("Invalid accessor (ownership moved)"), !Accessor.IsValid());
-		UTEST_EQUAL(TEXT("Extractror accessor type"), PCG::Private::MetadataTypes<double>::Id, ExtractorAccessor->GetUnderlyingType());
+		UTEST_EQUAL(TEXT("Extractror accessor type"), PCG::Private::MetadataTypes<ExpectedType>::Id, ExtractorAccessor->GetUnderlyingType());
 
 		FPCGAttributeAccessorKeysSingleObjectPtr ObjectKey(&Struct);
 
-		double Value = 1.0;
-		UTEST_TRUE(FString::Printf(TEXT("Getting a value for property %s and extractor %s"), *PropertyName.ToString(), *ExtractorName.ToString()), ExtractorAccessor->Get<double>(Value, ObjectKey));
+		ExpectedType Value{};
+		UTEST_TRUE(FString::Printf(TEXT("Getting a value for property %s and extractor %s"), *PropertyName.ToString(), *ExtractorName.ToString()), ExtractorAccessor->Get<ExpectedType>(Value, ObjectKey));
 		UTEST_EQUAL(FString::Printf(TEXT("Value for property %s and extractor %s is equal to expected"), *PropertyName.ToString(), *ExtractorName.ToString()), Value, ExpectedValue);
 
 		return true;
@@ -248,9 +282,21 @@ bool FPCGAttributeRotatorPropertyExtractorTest::RunTest(const FString & Paramete
 
 	bool bTestPassed = true;
 
-	bTestPassed &= Verify(Struct.RotatorValue.Pitch, GET_MEMBER_NAME_CHECKED(FPCGMetadataTypesConstantStruct, RotatorValue), PCGAttributeExtractorConstants::RotatorPitch);
-	bTestPassed &= Verify(Struct.RotatorValue.Yaw, GET_MEMBER_NAME_CHECKED(FPCGMetadataTypesConstantStruct, RotatorValue), PCGAttributeExtractorConstants::RotatorYaw);
-	bTestPassed &= Verify(Struct.RotatorValue.Roll, GET_MEMBER_NAME_CHECKED(FPCGMetadataTypesConstantStruct, RotatorValue), PCGAttributeExtractorConstants::RotatorRoll);
+	// Rotators
+	bTestPassed &= Verify(Struct.RotatorValue.Pitch, PCG_STRUCT_NAME_CHECK(RotatorValue), PCGAttributeExtractorConstants::RotatorPitch);
+	bTestPassed &= Verify(Struct.RotatorValue.Yaw, PCG_STRUCT_NAME_CHECK(RotatorValue), PCGAttributeExtractorConstants::RotatorYaw);
+	bTestPassed &= Verify(Struct.RotatorValue.Roll, PCG_STRUCT_NAME_CHECK(RotatorValue), PCGAttributeExtractorConstants::RotatorRoll);
+	bTestPassed &= Verify(Struct.RotatorValue.Quaternion().GetForwardVector(), PCG_STRUCT_NAME_CHECK(RotatorValue), PCGAttributeExtractorConstants::RotatorForward);
+	bTestPassed &= Verify(Struct.RotatorValue.Quaternion().GetRightVector(), PCG_STRUCT_NAME_CHECK(RotatorValue), PCGAttributeExtractorConstants::RotatorRight);
+	bTestPassed &= Verify(Struct.RotatorValue.Quaternion().GetUpVector(), PCG_STRUCT_NAME_CHECK(RotatorValue), PCGAttributeExtractorConstants::RotatorUp);
+
+	// Quaternions
+	bTestPassed &= Verify(Struct.QuatValue.Rotator().Pitch, PCG_STRUCT_NAME_CHECK(QuatValue), PCGAttributeExtractorConstants::RotatorPitch);
+	bTestPassed &= Verify(Struct.QuatValue.Rotator().Yaw, PCG_STRUCT_NAME_CHECK(QuatValue), PCGAttributeExtractorConstants::RotatorYaw);
+	bTestPassed &= Verify(Struct.QuatValue.Rotator().Roll, PCG_STRUCT_NAME_CHECK(QuatValue), PCGAttributeExtractorConstants::RotatorRoll);
+	bTestPassed &= Verify(Struct.QuatValue.GetForwardVector(), PCG_STRUCT_NAME_CHECK(QuatValue), PCGAttributeExtractorConstants::RotatorForward);
+	bTestPassed &= Verify(Struct.QuatValue.GetRightVector(), PCG_STRUCT_NAME_CHECK(QuatValue), PCGAttributeExtractorConstants::RotatorRight);
+	bTestPassed &= Verify(Struct.QuatValue.GetUpVector(), PCG_STRUCT_NAME_CHECK(QuatValue), PCGAttributeExtractorConstants::RotatorUp);
 
 	return bTestPassed;
 }
@@ -289,14 +335,14 @@ bool FPCGAttributeTransformPropertyExtractorTest::RunTest(const FString & Parame
 
 	bool bTestPassed = true;
 
-	bTestPassed &= Verify(Struct.TransformValue.GetLocation(), GET_MEMBER_NAME_CHECKED(FPCGMetadataTypesConstantStruct, TransformValue), { PCGAttributeExtractorConstants::TransformLocation });
-	bTestPassed &= Verify(Struct.TransformValue.GetRotation(), GET_MEMBER_NAME_CHECKED(FPCGMetadataTypesConstantStruct, TransformValue), { PCGAttributeExtractorConstants::TransformRotation });
-	bTestPassed &= Verify(Struct.TransformValue.GetScale3D(), GET_MEMBER_NAME_CHECKED(FPCGMetadataTypesConstantStruct, TransformValue), { PCGAttributeExtractorConstants::TransformScale });
+	bTestPassed &= Verify(Struct.TransformValue.GetLocation(), PCG_STRUCT_NAME_CHECK(TransformValue), { PCGAttributeExtractorConstants::TransformLocation });
+	bTestPassed &= Verify(Struct.TransformValue.GetRotation(), PCG_STRUCT_NAME_CHECK(TransformValue), { PCGAttributeExtractorConstants::TransformRotation });
+	bTestPassed &= Verify(Struct.TransformValue.GetScale3D(), PCG_STRUCT_NAME_CHECK(TransformValue), { PCGAttributeExtractorConstants::TransformScale });
 
 	// Also testing double chain
-	bTestPassed &= Verify(Struct.TransformValue.GetLocation().X, GET_MEMBER_NAME_CHECKED(FPCGMetadataTypesConstantStruct, TransformValue), { PCGAttributeExtractorConstants::TransformLocation, PCGAttributeExtractorConstants::VectorX });
-	bTestPassed &= Verify(Struct.TransformValue.GetRotation().W, GET_MEMBER_NAME_CHECKED(FPCGMetadataTypesConstantStruct, TransformValue), { PCGAttributeExtractorConstants::TransformRotation, PCGAttributeExtractorConstants::VectorW });
-	bTestPassed &= Verify(Struct.TransformValue.GetScale3D().Y, GET_MEMBER_NAME_CHECKED(FPCGMetadataTypesConstantStruct, TransformValue), { PCGAttributeExtractorConstants::TransformScale, PCGAttributeExtractorConstants::VectorY });
+	bTestPassed &= Verify(Struct.TransformValue.GetLocation().X, PCG_STRUCT_NAME_CHECK(TransformValue), { PCGAttributeExtractorConstants::TransformLocation, PCGAttributeExtractorConstants::VectorX });
+	bTestPassed &= Verify(Struct.TransformValue.GetRotation().W, PCG_STRUCT_NAME_CHECK(TransformValue), { PCGAttributeExtractorConstants::TransformRotation, PCGAttributeExtractorConstants::VectorW });
+	bTestPassed &= Verify(Struct.TransformValue.GetScale3D().Y, PCG_STRUCT_NAME_CHECK(TransformValue), { PCGAttributeExtractorConstants::TransformScale, PCGAttributeExtractorConstants::VectorY });
 
 	return true;
 }
@@ -334,18 +380,36 @@ bool FPCGAttributeInvalidPropertyExtractorTest::RunTest(const FString& Parameter
 
 	bool bTestPassed = true;
 
-	VerifyInvalidChain(GET_MEMBER_NAME_CHECKED(FPCGMetadataTypesConstantStruct, DoubleValue), { InvalidName });
-	VerifyInvalidChain(GET_MEMBER_NAME_CHECKED(FPCGMetadataTypesConstantStruct, Vector2Value), { InvalidName });
-	VerifyInvalidChain(GET_MEMBER_NAME_CHECKED(FPCGMetadataTypesConstantStruct, VectorValue), { InvalidName });
-	VerifyInvalidChain(GET_MEMBER_NAME_CHECKED(FPCGMetadataTypesConstantStruct, Vector4Value), { InvalidName });
-	VerifyInvalidChain(GET_MEMBER_NAME_CHECKED(FPCGMetadataTypesConstantStruct, QuatValue), { InvalidName });
-	VerifyInvalidChain(GET_MEMBER_NAME_CHECKED(FPCGMetadataTypesConstantStruct, RotatorValue), { InvalidName });
-	VerifyInvalidChain(GET_MEMBER_NAME_CHECKED(FPCGMetadataTypesConstantStruct, TransformValue), { InvalidName });
+	VerifyInvalidChain(PCG_STRUCT_NAME_CHECK(DoubleValue), { InvalidName });
+	VerifyInvalidChain(PCG_STRUCT_NAME_CHECK(Vector2Value), { InvalidName });
+	VerifyInvalidChain(PCG_STRUCT_NAME_CHECK(VectorValue), { InvalidName });
+	VerifyInvalidChain(PCG_STRUCT_NAME_CHECK(Vector4Value), { InvalidName });
+	VerifyInvalidChain(PCG_STRUCT_NAME_CHECK(QuatValue), { InvalidName });
+	VerifyInvalidChain(PCG_STRUCT_NAME_CHECK(RotatorValue), { InvalidName });
+	VerifyInvalidChain(PCG_STRUCT_NAME_CHECK(TransformValue), { InvalidName });
+
+	// Special case for Vectors
+	// Invalid letters for this type
+	VerifyInvalidChain(PCG_STRUCT_NAME_CHECK(Vector2Value), { "XYZ" });
+	VerifyInvalidChain(PCG_STRUCT_NAME_CHECK(Vector2Value), { "RGB" });
+	VerifyInvalidChain(PCG_STRUCT_NAME_CHECK(VectorValue), { "XYZW" });
+	VerifyInvalidChain(PCG_STRUCT_NAME_CHECK(VectorValue), { "RGBA" });
+	VerifyInvalidChain(PCG_STRUCT_NAME_CHECK(Vector4Value), { "XYZP" });
+	VerifyInvalidChain(PCG_STRUCT_NAME_CHECK(Vector4Value), { "RGBT" });
+	// Too long
+	VerifyInvalidChain(PCG_STRUCT_NAME_CHECK(Vector4Value), { "YYYYY" });
+	// Mixing RGBA and XYZW
+	VerifyInvalidChain(PCG_STRUCT_NAME_CHECK(Vector2Value), { "RX" });
+	VerifyInvalidChain(PCG_STRUCT_NAME_CHECK(VectorValue), { "BBYA" });
+	VerifyInvalidChain(PCG_STRUCT_NAME_CHECK(Vector4Value), { "BYZA" });
+
+	// Yaw will work for quaternions but not vector 4
+	VerifyInvalidChain(PCG_STRUCT_NAME_CHECK(Vector4Value), { "Yaw" });
 
 	// Invalid double chain
-	VerifyInvalidChain(GET_MEMBER_NAME_CHECKED(FPCGMetadataTypesConstantStruct, TransformValue), { PCGAttributeExtractorConstants::TransformLocation, InvalidName });
-	VerifyInvalidChain(GET_MEMBER_NAME_CHECKED(FPCGMetadataTypesConstantStruct, TransformValue), { PCGAttributeExtractorConstants::TransformRotation, InvalidName });
-	VerifyInvalidChain(GET_MEMBER_NAME_CHECKED(FPCGMetadataTypesConstantStruct, TransformValue), { PCGAttributeExtractorConstants::TransformScale, InvalidName });
+	VerifyInvalidChain(PCG_STRUCT_NAME_CHECK(TransformValue), { PCGAttributeExtractorConstants::TransformLocation, InvalidName });
+	VerifyInvalidChain(PCG_STRUCT_NAME_CHECK(TransformValue), { PCGAttributeExtractorConstants::TransformRotation, InvalidName });
+	VerifyInvalidChain(PCG_STRUCT_NAME_CHECK(TransformValue), { PCGAttributeExtractorConstants::TransformScale, InvalidName });
 
 	return bTestPassed;
 }
@@ -488,4 +552,6 @@ bool FPCGAttributePropertyMultipleDepthRangeTest::RunTest(const FString& Paramet
 	return true;
 }
 
-#endif // WITH_EDITOR
\ No newline at end of file
+#undef PCG_STRUCT_NAME_CHECK
+
+#endif // WITH_EDITOR
diff --git a/Engine/Plugins/Experimental/PCG/Source/PCG/Public/Metadata/Accessors/PCGAttributeExtractor.h b/Engine/Plugins/Experimental/PCG/Source/PCG/Public/Metadata/Accessors/PCGAttributeExtractor.h
--- a/Engine/Plugins/Experimental/PCG/Source/PCG/Public/Metadata/Accessors/PCGAttributeExtractor.h
+++ b/Engine/Plugins/Experimental/PCG/Source/PCG/Public/Metadata/Accessors/PCGAttributeExtractor.h
@@ -14,6 +14,10 @@ namespace PCGAttributeExtractorConstants
 	const FName VectorY = TEXT("Y");
 	const FName VectorZ = TEXT("Z");
 	const FName VectorW = TEXT("W");
+	const FName VectorR = TEXT("R");
+	const FName VectorG = TEXT("G");
+	const FName VectorB = TEXT("B");
+	const FName VectorA = TEXT("A");
 	const FName VectorLength = TEXT("Length");
 	const FName VectorSize = TEXT("Size");
 
@@ -23,6 +27,10 @@ namespace PCGAttributeExtractorConstants
 		VectorY,
 		VectorZ,
 		VectorW,
+		VectorR,
+		VectorG,
+		VectorB,
+		VectorA,
 		VectorLength,
 		VectorSize,
 	};
@@ -39,16 +47,22 @@ namespace PCGAttributeExtractorConstants
 		TransformScale
 	};
 
-	// Rotator
+	// Rotator/Quat
 	const FName RotatorPitch = TEXT("Pitch");
 	const FName RotatorYaw = TEXT("Yaw");
 	const FName RotatorRoll = TEXT("Roll");
+	const FName RotatorForward = TEXT("Forward");
+	const FName RotatorRight = TEXT("Right");
+	const FName RotatorUp = TEXT("Up");
 
 	const FName RotatorExtractorNames[] =
 	{
 		RotatorPitch,
 		RotatorYaw,
-		RotatorRoll
+		RotatorRoll,
+		RotatorForward,
+		RotatorRight,
+		RotatorUp
 	};
 }
 
@@ -59,6 +73,7 @@ namespace PCGAttributeExtractor
 	TUniquePtr<IPCGAttributeAccessor> CreateVectorExtractor(TUniquePtr<IPCGAttributeAccessor> InAccessor, FName Name, bool& bOutSuccess);
 
 	TUniquePtr<IPCGAttributeAccessor> CreateRotatorExtractor(TUniquePtr<IPCGAttributeAccessor> InAccessor, FName Name, bool& bOutSuccess);
+	TUniquePtr<IPCGAttributeAccessor> CreateQuatExtractor(TUniquePtr<IPCGAttributeAccessor> InAccessor, FName Name, bool& bOutSuccess);
 
 	TUniquePtr<IPCGAttributeAccessor> CreateTransformExtractor(TUniquePtr<IPCGAttributeAccessor> InAccessor, FName Name, bool& bOutSuccess);
 }
