[PCG] Collapse in subgraph bug fixes

* Fix a crash when collapsing node where pins has changed between original and copy
* Fix issue where pins were not connected correctly if pin order has changed

#rnx

diff --git a/Engine/Plugins/Experimental/PCG/Source/PCG/Private/PCGNode.cpp b/Engine/Plugins/Experimental/PCG/Source/PCG/Private/PCGNode.cpp
--- a/Engine/Plugins/Experimental/PCG/Source/PCG/Private/PCGNode.cpp
+++ b/Engine/Plugins/Experimental/PCG/Source/PCG/Private/PCGNode.cpp
@@ -705,15 +705,16 @@ EPCGChangeType UPCGNode::UpdatePins(TFunctionRef<UPCGPin*(UPCGNode*)> PinAllocat
 				for (int32 RemovedPinIndex = PinsToRemove.Num() - 1; RemovedPinIndex >= 0; --RemovedPinIndex)
 				{
 					const int32 PinIndex = Pins.IndexOfByKey(PinsToRemove[RemovedPinIndex]);
-					check(PinIndex >= 0);
-
-					if (bRemoveFromAllNames)
+					if (PinIndex >= 0)
 					{
-						AllPinNames.Remove(Pins[PinIndex]->Properties.Label);
-					}
+						if (bRemoveFromAllNames)
+						{
+							AllPinNames.Remove(Pins[PinIndex]->Properties.Label);
+						}
 
-					bAppliedEdgeChanges |= Pins[PinIndex]->BreakAllEdges();
-					Pins.RemoveAt(PinIndex);
+						bAppliedEdgeChanges |= Pins[PinIndex]->BreakAllEdges();
+						Pins.RemoveAt(PinIndex);
+					}
 				}
 			};
 
diff --git a/Engine/Plugins/Experimental/PCG/Source/PCGEditor/Private/PCGEditor.cpp b/Engine/Plugins/Experimental/PCG/Source/PCGEditor/Private/PCGEditor.cpp
--- a/Engine/Plugins/Experimental/PCG/Source/PCGEditor/Private/PCGEditor.cpp
+++ b/Engine/Plugins/Experimental/PCG/Source/PCGEditor/Private/PCGEditor.cpp
@@ -848,15 +848,39 @@ void FPCGEditor::OnCollapseNodesInSubgraph()
 
 		// Reconstruct a new node, same as PCGNode, but without any edges in the new graph 
 		TObjectPtr<UPCGNode> NewNode = NewPCGGraph->ReconstructNewNode(PCGNode);
-		
-		for (int i = 0; i < PCGNode->GetInputPins().Num(); ++i)
-		{
-			PinMapping.Emplace(PCGNode->GetInputPins()[i], NewNode->GetInputPins()[i]);
-		}
-		for (int i = 0; i < PCGNode->GetOutputPins().Num(); ++i)
+
+		// Safeguard: We should have a 1 for 1 matching between pins labels between the original node
+		// and the copied node. If for some reason we don't (perhaps the node was not updated correctly after pins were added/removed)
+		// we will log an error and try to connect as best as we can (probably breaking some edges on the process).
+
+		auto Mapping = [&PinMapping, PCGNode](const TArray<UPCGPin*>& OriginalPins, const TArray<UPCGPin*>& NewPins)
 		{
-			PinMapping.Emplace(PCGNode->GetOutputPins()[i], NewNode->GetOutputPins()[i]);
-		}
+			TSet<FName> UnmatchedOriginal;
+			TMap<FName, UPCGPin*> NewMapping;
+
+			for (UPCGPin* NewPin : NewPins)
+			{
+				NewMapping.Emplace(NewPin->Properties.Label, NewPin);
+			}
+
+			for (UPCGPin* OriginalPin : OriginalPins)
+			{
+				FName PinLabel = OriginalPin->Properties.Label;
+				if (UPCGPin** NewPinPtr = NewMapping.Find(PinLabel))
+				{
+					PinMapping.Emplace(OriginalPin, *NewPinPtr);
+				}
+				else if (OriginalPin->IsConnected())
+				{
+					// It is only problematic if the pin was connected
+					UE_LOG(LogPCG, Error, TEXT("[CollapseInSubgraph - %s] %s pin %s does not exist anymore. Edges will be broken."),
+						*PCGNode->GetNodeTitle().ToString(), (OriginalPin->IsOutputPin() ? TEXT("Output") : TEXT("Input")), *PinLabel.ToString());
+				}
+			}
+		};
+		
+		Mapping(PCGNode->GetInputPins(), NewNode->GetInputPins());
+		Mapping(PCGNode->GetOutputPins(), NewNode->GetOutputPins());
 	}
 
 	// Also duplicate the extra nodes and assign them to the new graph
