[PCG] Tentative fix for non-unity build

* Ran a preflight for non-Unity. Still a lot of errors, but no more in PCG, as far as I can tell.

#rnx

diff --git a/Engine/Plugins/Experimental/PCG/Source/PCG/Private/Graph/PCGGraphExecutor.cpp b/Engine/Plugins/Experimental/PCG/Source/PCG/Private/Graph/PCGGraphExecutor.cpp
--- a/Engine/Plugins/Experimental/PCG/Source/PCG/Private/Graph/PCGGraphExecutor.cpp
+++ b/Engine/Plugins/Experimental/PCG/Source/PCG/Private/Graph/PCGGraphExecutor.cpp
@@ -16,6 +16,7 @@
 #include "Metadata/PCGMetadata.h"
 
 #include "Async/Async.h"
+#include "GameFramework/Actor.h"
 
 #if WITH_EDITOR
 #include "Editor.h"
diff --git a/Engine/Plugins/Experimental/PCG/Source/PCG/Private/Helpers/PCGPropertyHelpers.cpp b/Engine/Plugins/Experimental/PCG/Source/PCG/Private/Helpers/PCGPropertyHelpers.cpp
--- a/Engine/Plugins/Experimental/PCG/Source/PCG/Private/Helpers/PCGPropertyHelpers.cpp
+++ b/Engine/Plugins/Experimental/PCG/Source/PCG/Private/Helpers/PCGPropertyHelpers.cpp
@@ -2,6 +2,8 @@
 
 #include "Helpers/PCGPropertyHelpers.h"
 
+#include "UObject/Field.h"
+
 EPCGMetadataTypes PCGPropertyHelpers::GetMetadataTypeFromProperty(const FProperty* InProperty)
 {
 	if (!InProperty)
diff --git a/Engine/Plugins/Experimental/PCG/Source/PCG/Public/Helpers/PCGPropertyHelpers.h b/Engine/Plugins/Experimental/PCG/Source/PCG/Public/Helpers/PCGPropertyHelpers.h
--- a/Engine/Plugins/Experimental/PCG/Source/PCG/Public/Helpers/PCGPropertyHelpers.h
+++ b/Engine/Plugins/Experimental/PCG/Source/PCG/Public/Helpers/PCGPropertyHelpers.h
@@ -7,6 +7,8 @@
 #include "Metadata/Accessors/PCGAttributeAccessorHelpers.h"
 #include "Metadata/Accessors/PCGAttributeAccessorKeys.h"
 
+#include "UObject/UnrealType.h" // IWYU pragma: keep
+
 enum class EPCGMetadataTypes : uint8;
 
 class UPCGData;
diff --git a/Engine/Plugins/Experimental/PCG/Source/PCG/Public/Metadata/Accessors/PCGPropertyAccessor.h b/Engine/Plugins/Experimental/PCG/Source/PCG/Public/Metadata/Accessors/PCGPropertyAccessor.h
--- a/Engine/Plugins/Experimental/PCG/Source/PCG/Public/Metadata/Accessors/PCGPropertyAccessor.h
+++ b/Engine/Plugins/Experimental/PCG/Source/PCG/Public/Metadata/Accessors/PCGPropertyAccessor.h
@@ -4,6 +4,9 @@
 
 #include "IPCGAttributeAccessorTpl.h"
 
+#include "Containers/UnrealString.h" // IWYU pragma: keep
+#include "UObject/SoftObjectPath.h" // IWYU pragma: keep
+#include "UObject/UnrealType.h" // IWYU pragma: keep
 
 class FEnumProperty;
 
diff --git a/Engine/Plugins/Experimental/PCG/Source/PCG/Public/Metadata/PCGMetadataAttribute.h b/Engine/Plugins/Experimental/PCG/Source/PCG/Public/Metadata/PCGMetadataAttribute.h
--- a/Engine/Plugins/Experimental/PCG/Source/PCG/Public/Metadata/PCGMetadataAttribute.h
+++ b/Engine/Plugins/Experimental/PCG/Source/PCG/Public/Metadata/PCGMetadataAttribute.h
@@ -4,6 +4,7 @@
 
 #include "PCGMetadataCommon.h"
 
+#include "UObject/ObjectPtr.h" // IWYU pragma: keep
 
 class UPCGMetadata;
 
