[PCG] First refinement of crc caching system - evaluate crc on point data, and default to not using dependencies crc as output crc.

Previously dependencies crc was used as both a caching lookup key and an output crc to pass to downstream nodes. The issue with this is that regeneration currently produces new crcs for input data (regardless of whether the data itself actually changed), and long term it may always do this (we may never crc all landscape data, so it may always be approximated with a counter-based UID). Propagating such spuriously changing crcs downstream leads to unnecessary executions.

To combat this, this change enables accurate crc of point data, so 'spurious' crcs will be replaced with precise crcs once data has collapsed to points. This reduced regenerating a heavy scene (warm cache) from 43s to 18s.

Crc'ing itself is not cheap, so the Crc value is cached within the PCGData to avoid recomputing when data passed from node to node. This reduced the regenerate time further from 18s to 12s. Overall 3.6x speedup for warm cache (as would be the case while authoring).

Final note - I expected the metadata crc to be costly, but in my tests its a tiny proportion of the crc'ing time (4% of total crc time in metadata). This needs to be monitored going forwards.

diff --git a/Engine/Plugins/Experimental/PCG/Source/PCG/Private/Data/PCGPointData.cpp b/Engine/Plugins/Experimental/PCG/Source/PCG/Private/Data/PCGPointData.cpp
--- a/Engine/Plugins/Experimental/PCG/Source/PCG/Private/Data/PCGPointData.cpp
+++ b/Engine/Plugins/Experimental/PCG/Source/PCG/Private/Data/PCGPointData.cpp
@@ -2,11 +2,13 @@
 
 #include "Data/PCGPointData.h"
 
-#include "Data/PCGSpatialData.h"
 #include "PCGHelpers.h"
 #include "Metadata/PCGMetadataAccessor.h"
+#include "Metadata/Accessors/PCGAttributeAccessorHelpers.h"
+#include "Metadata/Accessors/PCGAttributeAccessorKeys.h"
 
 #include "GameFramework/Actor.h"
+#include "Serialization/ArchiveCrc32.h"
 
 #include UE_INLINE_GENERATED_CPP_BY_NAME(PCGPointData)
 
@@ -250,6 +252,68 @@ const UPCGPointData::PointOctree& UPCGPointData::GetOctree() const
 	return Octree;
 }
 
+void UPCGPointData::AddToCrc(FArchiveCrc32& Ar) const
+{
+	uint32 ThisType = static_cast<uint32>(GetDataType());
+	Ar << ThisType;
+
+	if (Points.Num() == 0)
+	{
+		return;
+	}
+
+	// Crc point data.
+	{
+		// Create copy so we can zero-out the metadata keys which are non-deterministic.
+		TArray<FPCGPoint> PointsCopy = Points;
+		for (FPCGPoint& Point : PointsCopy)
+		{
+			Point.MetadataEntry = 0;
+		}
+
+		Ar.Serialize(PointsCopy.GetData(), PointsCopy.Num() * PointsCopy.GetTypeSize());
+	}
+
+	// Crc metadata.
+	if (const UPCGMetadata* PCGMetadata = ConstMetadata())
+	{
+		FPCGAttributeAccessorKeysPoints AccessorKeys(Points);
+
+		TArray<FName> AttributeNames;
+		{
+			TArray<EPCGMetadataTypes> AttributeTypes;
+			PCGMetadata->GetAttributes(AttributeNames, AttributeTypes);
+		}
+
+		// Attribute names might come in different orders for e.g. if edge order changes.
+		Algo::Sort(AttributeNames, [this](const FName& A, const FName& B) { return A.LexicalLess(B); });
+
+		for (FName AttributeName : AttributeNames)
+		{
+			Ar << AttributeName;
+
+			if (const FPCGMetadataAttributeBase* Attribute = PCGMetadata->GetConstAttribute(AttributeName))
+			{
+				for (const FPCGPoint& Point : Points)
+				{
+					auto Callback = [Attribute, PCGMetadata, &Ar, &Point](auto ValueWithType)
+					{
+						using AttributeType = decltype(ValueWithType);
+
+						if (const FPCGMetadataAttribute<AttributeType>* TypedAttribute = static_cast<const FPCGMetadataAttribute<AttributeType>*>(Attribute))
+						{
+							ValueWithType = TypedAttribute->GetValueFromItemKey(Point.MetadataEntry);
+							Ar << ValueWithType;
+						}
+					};
+
+					PCGMetadataAttribute::CallbackWithRightType(Attribute->GetTypeId(), Callback);
+				}
+			}
+		}
+	}
+}
+
 FBox UPCGPointData::GetBounds() const
 {
 	if (bBoundsAreDirty)
diff --git a/Engine/Plugins/Experimental/PCG/Source/PCG/Private/Elements/PCGSurfaceSampler.cpp b/Engine/Plugins/Experimental/PCG/Source/PCG/Private/Elements/PCGSurfaceSampler.cpp
--- a/Engine/Plugins/Experimental/PCG/Source/PCG/Private/Elements/PCGSurfaceSampler.cpp
+++ b/Engine/Plugins/Experimental/PCG/Source/PCG/Private/Elements/PCGSurfaceSampler.cpp
@@ -411,7 +411,7 @@ void FPCGSurfaceSamplerElement::GetDependenciesCrc(const FPCGDataCollection& InI
 		{
 			if (const UPCGData* Data = InComponent->GetActorPCGData())
 			{
-				Crc.Combine(Data->ComputeCrc());
+				Crc.Combine(Data->GetOrComputeCrc());
 			}
 		}
 	}
diff --git a/Engine/Plugins/Experimental/PCG/Source/PCG/Private/Elements/PCGTextureSampler.cpp b/Engine/Plugins/Experimental/PCG/Source/PCG/Private/Elements/PCGTextureSampler.cpp
--- a/Engine/Plugins/Experimental/PCG/Source/PCG/Private/Elements/PCGTextureSampler.cpp
+++ b/Engine/Plugins/Experimental/PCG/Source/PCG/Private/Elements/PCGTextureSampler.cpp
@@ -95,7 +95,7 @@ void FPCGTextureSamplerElement::GetDependenciesCrc(const FPCGDataCollection& InI
 		{
 			if (const UPCGData* Data = InComponent->GetActorPCGData())
 			{
-				Crc.Combine(Data->ComputeCrc());
+				Crc.Combine(Data->GetOrComputeCrc());
 			}
 		}
 	}
diff --git a/Engine/Plugins/Experimental/PCG/Source/PCG/Private/Graph/PCGGraphCache.cpp b/Engine/Plugins/Experimental/PCG/Source/PCG/Private/Graph/PCGGraphCache.cpp
--- a/Engine/Plugins/Experimental/PCG/Source/PCG/Private/Graph/PCGGraphCache.cpp
+++ b/Engine/Plugins/Experimental/PCG/Source/PCG/Private/Graph/PCGGraphCache.cpp
@@ -10,7 +10,7 @@
 #include "Misc/ScopeRWLock.h"
 
 static TAutoConsoleVariable<bool> CVarCacheDebugging(
-	TEXT("pcg.CacheDebugging"),
+	TEXT("pcg.Cache.EnableDebugging"),
 	false,
 	TEXT("Enable various features for debugging the graph cache system."));
 
@@ -91,16 +91,19 @@ void FPCGGraphCache::StoreInCache(const IPCGElement* InElement, const FPCGCrc& I
 		return;
 	}
 
-	TRACE_CPUPROFILER_EVENT_SCOPE(FPCGGraphCache::StoreInCache);
-	FWriteScopeLock ScopedWriteLock(CacheLock);
-
-	FPCGGraphCacheEntries* Entries = CacheData.Find(InElement);
-	if(!Entries)
 	{
-		Entries = &(CacheData.Add(InElement));
-	}
+		TRACE_CPUPROFILER_EVENT_SCOPE(FPCGGraphCache::StoreInCache);
+		FWriteScopeLock ScopedWriteLock(CacheLock);
 
-	Entries->Emplace(InDependenciesCrc, InOutput, *RootSet);
+		FPCGGraphCacheEntries* Entries = CacheData.Find(InElement);
+		if (!Entries)
+		{
+			Entries = &(CacheData.Add(InElement));
+		}
+
+
+		Entries->Emplace(InDependenciesCrc, InOutput, *RootSet);
+	}
 }
 
 void FPCGGraphCache::ClearCache()
diff --git a/Engine/Plugins/Experimental/PCG/Source/PCG/Private/Graph/PCGGraphExecutor.cpp b/Engine/Plugins/Experimental/PCG/Source/PCG/Private/Graph/PCGGraphExecutor.cpp
--- a/Engine/Plugins/Experimental/PCG/Source/PCG/Private/Graph/PCGGraphExecutor.cpp
+++ b/Engine/Plugins/Experimental/PCG/Source/PCG/Private/Graph/PCGGraphExecutor.cpp
@@ -504,6 +504,11 @@ void FPCGGraphExecutor::Execute()
 					Task.Element->GetDependenciesCrc(TaskInput, TaskSettings, Task.SourceComponent.Get(), DependenciesCrc);
 				}
 
+				if (bGraphCacheDebuggingEnabled && !bCacheable && Task.SourceComponent.Get() && Task.Node)
+				{
+					UE_LOG(LogPCG, Warning, TEXT("[%s] %s\t\tCACHING DISABLED"), *Task.SourceComponent->GetOwner()->GetName(), *Task.Node->GetNodeTitle().ToString());
+				}
+
 				FPCGDataCollection CachedOutput;
 				const bool bResultAlreadyInCache = bCacheable && DependenciesCrc.IsValid() && GraphCache.GetFromCache(Task.Node, Task.Element.Get(), DependenciesCrc, TaskInput, TaskSettings, Task.SourceComponent.Get(), CachedOutput);
 #if WITH_EDITOR
@@ -533,11 +538,6 @@ void FPCGGraphExecutor::Execute()
 					continue;
 				}
 
-				if (bGraphCacheDebuggingEnabled && Task.SourceComponent.Get() && Task.Node)
-				{
-					UE_LOG(LogPCG, Log, TEXT("         [%s] %s\t\tACTIVETASK"), *Task.SourceComponent->GetOwner()->GetName(), *Task.Node->GetNodeTitle().ToString());
-				}
-
 				// Allocate context if not previously done
 				if (!Task.Context)
 				{
diff --git a/Engine/Plugins/Experimental/PCG/Source/PCG/Private/PCGData.cpp b/Engine/Plugins/Experimental/PCG/Source/PCG/Private/PCGData.cpp
--- a/Engine/Plugins/Experimental/PCG/Source/PCG/Private/PCGData.cpp
+++ b/Engine/Plugins/Experimental/PCG/Source/PCG/Private/PCGData.cpp
@@ -82,11 +82,19 @@ UPCGData::UPCGData(const FObjectInitializer& ObjectInitializer)
 	}
 }
 
-FPCGCrc UPCGData::ComputeCrc() const
+FPCGCrc UPCGData::GetOrComputeCrc() const
 {
+	if (Crc.IsValid())
+	{
+		return Crc;
+	}
+
 	FArchiveCrc32 Ar;
 	AddToCrc(Ar);
-	return FPCGCrc(Ar.GetCrc());
+
+	Crc = FPCGCrc(Ar.GetCrc());
+
+	return Crc;
 }
 
 void UPCGData::AddToCrc(FArchiveCrc32& Ar) const
@@ -317,7 +325,9 @@ FPCGCrc FPCGDataCollection::ComputeCrc()
 
 		if (Data.Data)
 		{
-			Data.Data->AddToCrc(Ar);
+			const FPCGCrc DataCrc = Data.Data->GetOrComputeCrc();
+			uint32 Result = DataCrc.GetValue();
+			Ar << Result;
 		}
 
 		Ar << Data.Tags;
diff --git a/Engine/Plugins/Experimental/PCG/Source/PCG/Private/PCGElement.cpp b/Engine/Plugins/Experimental/PCG/Source/PCG/Private/PCGElement.cpp
--- a/Engine/Plugins/Experimental/PCG/Source/PCG/Private/PCGElement.cpp
+++ b/Engine/Plugins/Experimental/PCG/Source/PCG/Private/PCGElement.cpp
@@ -17,9 +17,13 @@ static TAutoConsoleVariable<bool> CVarPCGValidatePointMetadata(
 	true,
 	TEXT("Controls whether we validate that the metadata entry keys on the output point data are consistent"));
 
+static TAutoConsoleVariable<bool> CVarCacheUseDependenciesCrc(
+	TEXT("pcg.Cache.UseDependenciesCrc"),
+	false,
+	TEXT("If element is cacheable, use dependencies crc instead of calculating crc from output data."));
+
 bool IPCGElement::Execute(FPCGContext* Context) const
 {
-	TRACE_CPUPROFILER_EVENT_SCOPE(IPCGElement::Execute);
 	check(Context && Context->NumAvailableTasks > 0 && Context->CurrentPhase < EPCGExecutionPhase::Done);
 	check(Context->bIsRunningOnMainThread || !CanExecuteOnlyOnMainThread(Context));
 
@@ -32,12 +36,14 @@ bool IPCGElement::Execute(FPCGContext* Context) const
 		{
 			case EPCGExecutionPhase::NotExecuted: // Fall-through
 			{
+				TRACE_CPUPROFILER_EVENT_SCOPE(EPCGExecutionPhase::NotExecuted);
 				PreExecute(Context);
 				break;
 			}
 
 			case EPCGExecutionPhase::PrepareData:
 			{
+				TRACE_CPUPROFILER_EVENT_SCOPE(EPCGExecutionPhase::PrepareData);
 				// Will override the settings if there is any override.
 				Context->OverrideSettings();
 
@@ -54,6 +60,7 @@ bool IPCGElement::Execute(FPCGContext* Context) const
 
 			case EPCGExecutionPhase::Execute:
 			{
+				TRACE_CPUPROFILER_EVENT_SCOPE(EPCGExecutionPhase::Execute);
 				if (ExecuteInternal(Context))
 				{
 					Context->CurrentPhase = EPCGExecutionPhase::PostExecute;
@@ -67,6 +74,7 @@ bool IPCGElement::Execute(FPCGContext* Context) const
 
 			case EPCGExecutionPhase::PostExecute:
 			{
+				TRACE_CPUPROFILER_EVENT_SCOPE(EPCGExecutionPhase::PostExecute);
 				PostExecute(Context);
 				break;
 			}
@@ -175,20 +183,24 @@ void IPCGElement::PostExecute(FPCGContext* Context) const
 			Context->OutputData.TaggedData[TaggedDataIdx].Tags.Append(Settings->TagsAppliedOnOutput);
 		}
 	}
-
-	if (IsCacheableInstance(SettingsInterface))
-	{
-		// If cacheable, then assume that the dependencies Crc is a good Crc to use from the output data - assumes changes
-		// in inputs and settings are 1:1 with changes in output. In general this is not the case as some settings changes will
-		// have no impact on the output. And since we are unlikely to do full Crc's of landscapes, a landscape data may produce
-		// a different Crc on each regenerate, triggering unnecessary executions. Removing this branch (always doing a full Crc)
-		// would avoid such triggers propagating.
-		Context->OutputData.Crc = Context->DependenciesCrc;
-	}
-	else
+	
+	// Output data crc.
 	{
-		// Compute Crc from output data
-		Context->OutputData.Crc = Context->OutputData.ComputeCrc();
+		if (CVarCacheUseDependenciesCrc.GetValueOnAnyThread() && IsCacheableInstance(SettingsInterface))
+		{
+			// If cacheable, then assume that the dependencies Crc is a good Crc to use from the output data - assumes changes
+			// in inputs and settings are 1:1 with changes in output. In general this is not the case as some settings changes will
+			// have no impact on the output. And since we are unlikely to do full Crc's of landscapes, a landscape data may produce
+			// a different Crc on each regenerate, triggering unnecessary executions.
+			Context->OutputData.Crc = Context->DependenciesCrc;
+		}
+		else
+		{
+			TRACE_CPUPROFILER_EVENT_SCOPE(IPCGElement::PostExecute::CRC);
+
+			// Compute Crc from output data
+			Context->OutputData.Crc = Context->OutputData.ComputeCrc();
+		}
 	}
 
 	// Additional debug things (check for duplicates),
diff --git a/Engine/Plugins/Experimental/PCG/Source/PCG/Public/Data/PCGPointData.h b/Engine/Plugins/Experimental/PCG/Source/PCG/Public/Data/PCGPointData.h
--- a/Engine/Plugins/Experimental/PCG/Source/PCG/Public/Data/PCGPointData.h
+++ b/Engine/Plugins/Experimental/PCG/Source/PCG/Public/Data/PCGPointData.h
@@ -69,6 +69,7 @@ public:
 
 	// ~Begin UPCGData interface
 	virtual EPCGDataType GetDataType() const override { return EPCGDataType::Point; }
+	virtual void AddToCrc(FArchiveCrc32& Ar) const override;
 	// ~End UPCGData interface
 
 	// ~Begin UPCGSpatialData interface
@@ -77,6 +78,7 @@ public:
 	virtual const UPCGPointData* ToPointData(FPCGContext* Context, const FBox& InBounds = FBox(EForceInit::ForceInit)) const { return this; }
 	virtual bool SamplePoint(const FTransform& Transform, const FBox& Bounds, FPCGPoint& OutPoint, UPCGMetadata* OutMetadata) const override;
 	virtual bool ProjectPoint(const FTransform& InTransform, const FBox& InBounds, const FPCGProjectionParams& InParams, FPCGPoint& OutPoint, UPCGMetadata* OutMetadata) const override;
+
 protected:
 	virtual UPCGSpatialData* CopyInternal() const override;
 	// ~End UPCGSpatialData interface
diff --git a/Engine/Plugins/Experimental/PCG/Source/PCG/Public/PCGData.h b/Engine/Plugins/Experimental/PCG/Source/PCG/Public/PCGData.h
--- a/Engine/Plugins/Experimental/PCG/Source/PCG/Public/PCGData.h
+++ b/Engine/Plugins/Experimental/PCG/Source/PCG/Public/PCGData.h
@@ -30,7 +30,7 @@ public:
 	virtual EPCGDataType GetDataType() const { return EPCGDataType::None; }
 
 	/** Computes a Crc for this data. */
-	FPCGCrc ComputeCrc() const;
+	FPCGCrc GetOrComputeCrc() const;
 
 	/** Contributes to Crc. Fallback implementation that writes object instance UID. */
 	virtual void AddToCrc(FArchiveCrc32& Ar) const;
@@ -39,6 +39,9 @@ public:
 	UPROPERTY(Transient)
 	uint64 UID = 0;
 
+	/** CRC for this object instance. */
+	mutable FPCGCrc Crc;
+
 private:
 	/** Serves unique ID values to instances of this object. */
 	static inline std::atomic<uint64> UIDCounter{ 1 };
