[PCG] Improving behavior for subgraph and pins

* CollapseToSubgraph will know track the changes in the subgraph correctly
* Connecting advanced pins in input/output node in a subgraph will make them not advanced on a subgraph node.
* Params pins are now advanced by default.

diff --git a/Engine/Plugins/Experimental/PCG/Source/PCG/Private/PCGSettings.cpp b/Engine/Plugins/Experimental/PCG/Source/PCG/Private/PCGSettings.cpp
--- a/Engine/Plugins/Experimental/PCG/Source/PCG/Private/PCGSettings.cpp
+++ b/Engine/Plugins/Experimental/PCG/Source/PCG/Private/PCGSettings.cpp
@@ -245,7 +245,7 @@ void UPCGSettings::FillOverridableParamsPins(TArray<FPCGPinProperties>& OutPins)
 	else
 	{
 		FPCGPinProperties& ParamPin = OutPins.Emplace_GetRef(PCGPinConstants::DefaultParamsLabel, EPCGDataType::Param, /*bInAllowMultipleConnections=*/ true, /*bAllowMultipleData=*/ true);
-		ParamPin.bAdvancedPin = false;
+		ParamPin.bAdvancedPin = true;
 
 #if WITH_EDITOR
 		ParamPin.Tooltip = LOCTEXT("GlobalParamPinTooltip", "Can bundle multiple param data to override multiple parameters at the same time. Names need to match perfectly.");
diff --git a/Engine/Plugins/Experimental/PCG/Source/PCG/Private/PCGSubgraph.cpp b/Engine/Plugins/Experimental/PCG/Source/PCG/Private/PCGSubgraph.cpp
--- a/Engine/Plugins/Experimental/PCG/Source/PCG/Private/PCGSubgraph.cpp
+++ b/Engine/Plugins/Experimental/PCG/Source/PCG/Private/PCGSubgraph.cpp
@@ -8,6 +8,24 @@
 
 #include UE_INLINE_GENERATED_CPP_BY_NAME(PCGSubgraph)
 
+namespace PCGSubgraphSettings
+{
+	void RemoveAdvancedModeOnConnectedPins(const UPCGGraph* Subgraph, TArray<FPCGPinProperties>& InOutPinProperties, const bool bIsInput)
+	{
+		const UPCGNode* SubgraphNode = bIsInput ? Subgraph->GetInputNode() : Subgraph->GetOutputNode();
+		check(SubgraphNode);
+
+		for (FPCGPinProperties& PinProperties : InOutPinProperties)
+		{
+			const UPCGPin* Pin = bIsInput ? SubgraphNode->GetOutputPin(PinProperties.Label) : SubgraphNode->GetInputPin(PinProperties.Label);
+			if (ensure(Pin) && Pin->IsConnected())
+			{
+				PinProperties.bAdvancedPin = false;
+			}
+		}
+	}
+}
+
 void UPCGBaseSubgraphSettings::PostLoad()
 {
 	Super::PostLoad();
@@ -20,6 +38,25 @@ void UPCGBaseSubgraphSettings::PostLoad()
 #endif
 }
 
+void UPCGBaseSubgraphSettings::SetSubgraph(UPCGGraph* InGraph)
+{
+#if WITH_EDITOR
+	if (UPCGGraph* Subgraph = GetSubgraph())
+	{
+		Subgraph->OnGraphChangedDelegate.RemoveAll(this);
+	}
+#endif // WITH_EDITOR
+
+	SetSubgraphInternal(InGraph);
+
+#if WITH_EDITOR
+	if (UPCGGraph* Subgraph = GetSubgraph())
+	{
+		Subgraph->OnGraphChangedDelegate.AddUObject(this, &UPCGSubgraphSettings::OnSubgraphChanged);
+	}
+#endif // WITH_EDITOR
+}
+
 void UPCGBaseSubgraphSettings::BeginDestroy()
 {
 #if WITH_EDITOR
@@ -85,7 +122,9 @@ TArray<FPCGPinProperties> UPCGBaseSubgraphSettings::InputPinProperties() const
 {
 	if (UPCGGraph* Subgraph = GetSubgraph())
 	{
-		return Subgraph->GetInputNode()->InputPinProperties();
+		TArray<FPCGPinProperties> InputPins = Subgraph->GetInputNode()->InputPinProperties();
+		PCGSubgraphSettings::RemoveAdvancedModeOnConnectedPins(Subgraph, InputPins, /*bIsInput=*/true);
+		return InputPins;
 	}
 	else
 	{
@@ -97,7 +136,9 @@ TArray<FPCGPinProperties> UPCGBaseSubgraphSettings::OutputPinProperties() const
 {
 	if (UPCGGraph* Subgraph = GetSubgraph())
 	{
-		return Subgraph->GetOutputNode()->OutputPinProperties();
+		TArray<FPCGPinProperties> OutputPins = Subgraph->GetOutputNode()->OutputPinProperties();
+		PCGSubgraphSettings::RemoveAdvancedModeOnConnectedPins(Subgraph, OutputPins, /*bIsInput=*/false);
+		return OutputPins;
 	}
 	else
 	{
diff --git a/Engine/Plugins/Experimental/PCG/Source/PCG/Public/PCGSubgraph.h b/Engine/Plugins/Experimental/PCG/Source/PCG/Public/PCGSubgraph.h
--- a/Engine/Plugins/Experimental/PCG/Source/PCG/Public/PCGSubgraph.h
+++ b/Engine/Plugins/Experimental/PCG/Source/PCG/Public/PCGSubgraph.h
@@ -16,6 +16,9 @@ class PCG_API UPCGBaseSubgraphSettings : public UPCGSettings
 public:
 	virtual UPCGGraph* GetSubgraph() const { return nullptr; }
 
+	// Use this method from the outside to set the subgraph, as it will connect editor callbacks
+	virtual void SetSubgraph(UPCGGraph* InGraph);
+
 protected:
 	//~Begin UObject interface implementation
 	virtual void PostLoad() override;
@@ -34,6 +37,8 @@ protected:
 	TArray<FPCGPinProperties> OutputPinProperties() const override;
 	//~End UPCGSettings interface
 
+	virtual void SetSubgraphInternal(UPCGGraph* InGraph) {}
+
 #if WITH_EDITOR
 	void OnSubgraphChanged(UPCGGraph* InGraph, EPCGChangeType ChangeType);
 #endif
@@ -64,6 +69,7 @@ protected:
 public:
 	virtual UPCGGraph* GetSubgraph() const override { return Subgraph; }
 protected:
+	virtual void SetSubgraphInternal(UPCGGraph* InGraph) override { Subgraph = InGraph; }
 #if WITH_EDITOR
 	virtual bool IsStructuralProperty(const FName& InPropertyName) const override;
 #endif
diff --git a/Engine/Plugins/Experimental/PCG/Source/PCGEditor/Private/PCGEditorGraphSchemaActions.cpp b/Engine/Plugins/Experimental/PCG/Source/PCGEditor/Private/PCGEditorGraphSchemaActions.cpp
--- a/Engine/Plugins/Experimental/PCG/Source/PCGEditor/Private/PCGEditorGraphSchemaActions.cpp
+++ b/Engine/Plugins/Experimental/PCG/Source/PCGEditor/Private/PCGEditorGraphSchemaActions.cpp
@@ -305,7 +305,7 @@ UEdGraphNode* FPCGEditorGraphSchemaAction_NewSubgraphElement::PerformAction(UEdG
 	UPCGSettings* DefaultNodeSettings = nullptr;
 	UPCGNode* NewPCGNode = PCGGraph->AddNodeOfType(UPCGSubgraphSettings::StaticClass(), DefaultNodeSettings);
 	UPCGSubgraphSettings* DefaultSubgraphSettings = CastChecked<UPCGSubgraphSettings>(DefaultNodeSettings);
-	DefaultSubgraphSettings->Subgraph = Subgraph;
+	DefaultSubgraphSettings->SetSubgraph(Subgraph);
 
 	NewPCGNode->UpdateAfterSettingsChangeDuringCreation();
 
