Small improvement/tentative fix for copy/paste crash when in PIE

#rnx

diff --git a/Engine/Plugins/Experimental/PCG/Source/PCG/Private/PCGManagedResource.cpp b/Engine/Plugins/Experimental/PCG/Source/PCG/Private/PCGManagedResource.cpp
--- a/Engine/Plugins/Experimental/PCG/Source/PCG/Private/PCGManagedResource.cpp
+++ b/Engine/Plugins/Experimental/PCG/Source/PCG/Private/PCGManagedResource.cpp
@@ -130,6 +130,7 @@ void UPCGManagedComponent::PostEditImport()
 			{
 				GeneratedComponent = Component;
 				bFoundMatch = true;
+				break;
 			}
 		}
 
