Added next delta angle and curvature attributes to spline sampler

diff --git a/Engine/Plugins/Experimental/PCG/Source/PCG/Private/Data/PCGLandscapeSplineData.cpp b/Engine/Plugins/Experimental/PCG/Source/PCG/Private/Data/PCGLandscapeSplineData.cpp
--- a/Engine/Plugins/Experimental/PCG/Source/PCG/Private/Data/PCGLandscapeSplineData.cpp
+++ b/Engine/Plugins/Experimental/PCG/Source/PCG/Private/Data/PCGLandscapeSplineData.cpp
@@ -121,6 +121,77 @@ FTransform UPCGLandscapeSplineData::GetTransformAtDistance(int SegmentIndex, FVe
 	return FTransform();
 }
 
+FVector::FReal UPCGLandscapeSplineData::GetCurvatureAtDistance(int SegmentIndex, FVector::FReal Distance) const
+{
+	check(Spline.IsValid());
+	check(SegmentIndex >= 0 && SegmentIndex < Spline->GetSegments().Num());
+
+	const ULandscapeSplineSegment* Segment = Spline->GetSegments()[SegmentIndex];
+	check(Segment);
+
+	const FLandscapeSplineSegmentConnection& Start = Segment->Connections[0];
+	const FLandscapeSplineSegmentConnection& End = Segment->Connections[1];
+
+	const TArray<FLandscapeSplineInterpPoint>& InterpPoints = Segment->GetPoints();
+
+	// Need at least three points to compute the curvature
+	if (InterpPoints.Num() < 3)
+	{
+		return 0;
+	}
+
+	FVector::FReal Length = FMath::Max(0, Distance);
+
+	for (int PointIndex = 1; PointIndex < InterpPoints.Num(); ++PointIndex)
+	{
+		const FLandscapeSplineInterpPoint& PreviousPoint = InterpPoints[PointIndex - 1];
+		const FLandscapeSplineInterpPoint& CurrentPoint = InterpPoints[PointIndex];
+
+		const FVector::FReal SegmentLength = (CurrentPoint.Center - PreviousPoint.Center).Length();
+		if (SegmentLength > Length || PointIndex == InterpPoints.Num() - 1)
+		{
+			FVector FirstDerivative = FVector::ZeroVector;
+			FVector SecondDerivative = FVector::ZeroVector;
+
+			// Compute curvature using finite differences - here h is 1 because that's the only base unit we have.
+			// Warning: precision will be poor
+			// if last point -> use backward 2nd derivative
+			if (PointIndex == InterpPoints.Num() - 1)
+			{
+				const FLandscapeSplineInterpPoint& PreviousPreviousPoint = InterpPoints[PointIndex - 2];
+
+				// f'(x) = (f(x) - f(x-h)) / h
+				FirstDerivative = (CurrentPoint.Center - PreviousPoint.Center);
+				// f''(x) = (f(x) - 2f(x - h) + f(x - 2h)) / h2
+				SecondDerivative = (CurrentPoint.Center - 2 * PreviousPoint.Center + PreviousPreviousPoint.Center);
+			}
+			// otherwise -> use central 2nd derivative
+			else
+			{
+				const FLandscapeSplineInterpPoint& NextPoint = InterpPoints[PointIndex + 1];
+
+				// f'(x) ~= (f(x+h) - f(x-h)) / 2h
+				FirstDerivative = (NextPoint.Center - PreviousPoint.Center) / 2.0;
+				// f''(x) = (f(x+h) - 2f(x) + f(x-h)) / h2
+				SecondDerivative = NextPoint.Center - 2 * CurrentPoint.Center + PreviousPoint.Center;
+			}
+
+			const FVector::FReal FirstDerivativeLength = FMath::Max(FirstDerivative.Length(), UE_DOUBLE_SMALL_NUMBER);
+			const FVector ForwardVector = FirstDerivative / FirstDerivativeLength;
+			const FVector CurvatureVector = SecondDerivative - (SecondDerivative | ForwardVector) * ForwardVector;
+			const FVector::FReal Curvature = CurvatureVector.Length() / FirstDerivativeLength;
+
+			return FMath::Sign(CurvatureVector | (CurrentPoint.Right - CurrentPoint.Center)) * Curvature;
+		}
+		else
+		{
+			Length -= SegmentLength;
+		}
+	}
+
+	return 0;
+}
+
 const UPCGPointData* UPCGLandscapeSplineData::CreatePointData(FPCGContext* Context) const
 {
 	check(Spline.IsValid());
diff --git a/Engine/Plugins/Experimental/PCG/Source/PCG/Private/Data/PCGSplineData.cpp b/Engine/Plugins/Experimental/PCG/Source/PCG/Private/Data/PCGSplineData.cpp
--- a/Engine/Plugins/Experimental/PCG/Source/PCG/Private/Data/PCGSplineData.cpp
+++ b/Engine/Plugins/Experimental/PCG/Source/PCG/Private/Data/PCGSplineData.cpp
@@ -58,6 +58,27 @@ FTransform UPCGSplineData::GetTransformAtDistance(int SegmentIndex, FVector::FRe
 	return Spline->GetTransformAtDistanceAlongSpline(Spline->GetDistanceAlongSplineAtSplinePoint(SegmentIndex) + Distance, ESplineCoordinateSpace::World, /*bUseScale=*/true);
 }
 
+FVector::FReal UPCGSplineData::GetCurvatureAtDistance(int SegmentIndex, FVector::FReal Distance) const
+{
+	const float FullDistance = Spline->GetDistanceAlongSplineAtSplinePoint(SegmentIndex) + Distance;
+	const float Param = Spline->SplineCurves.ReparamTable.Eval(FullDistance, 0.0f);
+
+	// Since we need the first derivative (e.g. very similar to direction) to have its norm, we'll get the value directly
+	const FVector FirstDerivative = Spline->SplineCurves.Position.EvalDerivative(Param, FVector::ZeroVector);
+	const FVector::FReal FirstDerivativeLength = FMath::Max(FirstDerivative.Length(), UE_DOUBLE_SMALL_NUMBER);
+	const FVector ForwardVector = FirstDerivative / FirstDerivativeLength;
+	const FVector SecondDerivative = Spline->SplineCurves.Position.EvalSecondDerivative(Param, FVector::ZeroVector);
+	// Orthogonalize the second derivative and obtain the curvature vector
+	const FVector CurvatureVector = SecondDerivative - (SecondDerivative | ForwardVector) * ForwardVector;
+	
+	// Finally, the curvature is the ratio of the norms of the curvature vector over the first derivative norm
+	const FVector::FReal Curvature = CurvatureVector.Length() / FirstDerivativeLength;
+
+	// Compute sign based on sign of curvature vs. right axis
+	const FVector RightVector = Spline->GetRightVectorAtSplineInputKey(Param, ESplineCoordinateSpace::Local);
+	return FMath::Sign(RightVector | CurvatureVector) * Curvature;
+}
+
 const UPCGPointData* UPCGSplineData::CreatePointData(FPCGContext* Context) const
 {
 	TRACE_CPUPROFILER_EVENT_SCOPE(UPCGSplineData::CreatePointData);
@@ -180,7 +201,7 @@ bool UPCGSplineProjectionData::SamplePoint(const FTransform& InTransform, const
 
 			if (OutMetadata)
 			{
-				if (SplinePoint.MetadataEntry != PCGInvalidEntryKey && SurfacePoint.MetadataEntry)
+				if (SplinePoint.MetadataEntry != PCGInvalidEntryKey && SurfacePoint.MetadataEntry != PCGInvalidEntryKey)
 				{
 					OutMetadata->MergePointAttributesSubset(SplinePoint, OutMetadata, GetSpline()->Metadata, SurfacePoint, OutMetadata, GetSurface()->Metadata, OutPoint, ProjectionParams.AttributeMergeOperation);
 				}
diff --git a/Engine/Plugins/Experimental/PCG/Source/PCG/Private/Elements/PCGSplineSampler.cpp b/Engine/Plugins/Experimental/PCG/Source/PCG/Private/Elements/PCGSplineSampler.cpp
--- a/Engine/Plugins/Experimental/PCG/Source/PCG/Private/Elements/PCGSplineSampler.cpp
+++ b/Engine/Plugins/Experimental/PCG/Source/PCG/Private/Elements/PCGSplineSampler.cpp
@@ -309,30 +309,40 @@ namespace PCGSplineSamplerHelpers
 
 namespace PCGSplineSampler
 {
+	struct FSamplerResult
+	{
+		FTransform Transform;
+		FBox Box = FBox::BuildAABB(FVector::ZeroVector, FVector::OneVector);
+		FVector::FReal Curvature = 0;
+		FVector::FReal PreviousDeltaAngle = 0;
+		FVector::FReal NextDeltaAngle = 0;
+	};
+
 	struct FStepSampler
 	{
-		FStepSampler(const UPCGPolyLineData* InLineData)
-			: LineData(InLineData)
+		FStepSampler(const UPCGPolyLineData* InLineData, const FPCGSplineSamplerParams& Params)
+			: LineData(InLineData), bComputeCurvature(Params.bComputeCurvature)
 		{
 			check(LineData);
 			CurrentSegmentIndex = 0;
 		}
 
-		virtual void Step(FTransform& OutTransform, FBox& OutBox) = 0;
+		virtual void Step(FSamplerResult& OutSamplerResult) = 0;
 
 		bool IsDone() const
 		{
 			return CurrentSegmentIndex >= LineData->GetNumSegments();
 		}
 
-		const UPCGPolyLineData* LineData;
-		int CurrentSegmentIndex;
+		const UPCGPolyLineData* LineData = nullptr;
+		int CurrentSegmentIndex = 0;
+		bool bComputeCurvature = false;
 	};
 
 	struct FSubdivisionStepSampler : public FStepSampler
 	{
 		FSubdivisionStepSampler(const UPCGPolyLineData* InLineData, const FPCGSplineSamplerParams& Params)
-			: FStepSampler(InLineData)
+			: FStepSampler(InLineData, Params)
 		{
 			NumSegments = LineData->GetNumSegments();
 			SubdivisionsPerSegment = Params.SubdivisionsPerSegment;
@@ -341,13 +351,20 @@ namespace PCGSplineSampler
 			SubpointIndex = 0;
 		}
 
-		virtual void Step(FTransform& OutTransform, FBox& OutBox) override
+		virtual void Step(FSamplerResult& OutResult) override
 		{
 			const FVector::FReal SegmentLength = LineData->GetSegmentLength(CurrentSegmentIndex);
 			const FVector::FReal SegmentStep = SegmentLength / (SubdivisionsPerSegment + 1);
 
+			FBox& OutBox = OutResult.Box;
+			FTransform& OutTransform = OutResult.Transform;
 			OutTransform = LineData->GetTransformAtDistance(CurrentSegmentIndex, SubpointIndex * SegmentStep, &OutBox);
 
+			if (bComputeCurvature)
+			{
+				OutResult.Curvature = LineData->GetCurvatureAtDistance(CurrentSegmentIndex, SubpointIndex * SegmentStep);
+			}
+
 			if (SubpointIndex == 0)
 			{
 				const int PreviousSegmentIndex = (CurrentSegmentIndex > 0 ? CurrentSegmentIndex : NumSegments) - 1;
@@ -386,20 +403,27 @@ namespace PCGSplineSampler
 	struct FDistanceStepSampler : public FStepSampler
 	{
 		FDistanceStepSampler(const UPCGPolyLineData* InLineData, const FPCGSplineSamplerParams& Params)
-			: FStepSampler(InLineData)
+			: FStepSampler(InLineData, Params)
 		{
 			DistanceIncrement = Params.DistanceIncrement;
 			CurrentDistance = 0;
 		}
 
-		virtual void Step(FTransform& OutTransform, FBox& OutBox) override
+		virtual void Step(FSamplerResult& OutResult) override
 		{
 			FVector::FReal CurrentSegmentLength = LineData->GetSegmentLength(CurrentSegmentIndex);
+			FTransform& OutTransform = OutResult.Transform;
+			FBox& OutBox = OutResult.Box;
 			OutTransform = LineData->GetTransformAtDistance(CurrentSegmentIndex, CurrentDistance, &OutBox);
 
 			OutBox.Min.X *= DistanceIncrement / OutTransform.GetScale3D().X;
 			OutBox.Max.X *= DistanceIncrement / OutTransform.GetScale3D().X;
 
+			if (bComputeCurvature)
+			{
+				OutResult.Curvature = LineData->GetCurvatureAtDistance(CurrentSegmentIndex, CurrentDistance);
+			}
+
 			CurrentDistance += DistanceIncrement;
 			while(CurrentDistance > CurrentSegmentLength)
 			{
@@ -422,42 +446,100 @@ namespace PCGSplineSampler
 
 	struct FDimensionSampler
 	{
-		FDimensionSampler(const UPCGPolyLineData* InLineData, const UPCGSpatialData* InSpatialData, const UPCGSpatialData* InBoundingShapeData)
+		FDimensionSampler(const UPCGPolyLineData* InLineData, const UPCGSpatialData* InSpatialData, const UPCGSpatialData* InBoundingShapeData, const FPCGSplineSamplerParams& InParams, UPCGPointData* OutPointData)
+			: Params(InParams)
 		{
 			check(InLineData);
 			LineData = InLineData;
 			SpatialData = InSpatialData;
 			BoundingShapeData = InBoundingShapeData;
+
+			// Initialize metadata accessors if needed
+			if (OutPointData && OutPointData->Metadata && (Params.bComputeDirectionDelta || Params.bComputeCurvature))
+			{
+				auto FindOrCreateAttribute = [](const FName& AttributeName, UPCGMetadata* Metadata)
+				{
+					FPCGMetadataAttribute<double>* Attribute = Metadata->FindOrCreateAttribute(AttributeName, 0.0, true, true);
+					if (Attribute->GetTypeId() != PCG::Private::MetadataTypes<double>::Id)
+					{
+						Metadata->DeleteAttribute(AttributeName);
+						Attribute = Metadata->CreateAttribute<double>(AttributeName, 0.0, true, true);
+					}
+
+					return Attribute;
+				};
+
+				if (Params.bComputeDirectionDelta)
+				{
+					NextDirectionDeltaAttribute = FindOrCreateAttribute(Params.NextDirectionDeltaAttribute, OutPointData->Metadata);
+					bSetMetadata |= (NextDirectionDeltaAttribute != nullptr);
+				}
+
+				if (Params.bComputeCurvature)
+				{
+					CurvatureAttribute = FindOrCreateAttribute(Params.CurvatureAttribute, OutPointData->Metadata);
+					bSetMetadata |= (CurvatureAttribute != nullptr);
+				}
+			}
+		}
+
+		virtual ~FDimensionSampler() = default;
+
+		virtual void SetMetadata(const FSamplerResult& InResult, FPCGPoint& OutPoint, UPCGMetadata* OutMetadata)
+		{
+			if (bSetMetadata)
+			{
+				OutMetadata->InitializeOnSet(OutPoint.MetadataEntry);
+			}
+
+			if (NextDirectionDeltaAttribute)
+			{
+				NextDirectionDeltaAttribute->SetValue(OutPoint.MetadataEntry, InResult.NextDeltaAngle);
+			}
+
+			if (CurvatureAttribute)
+			{
+				CurvatureAttribute->SetValue(OutPoint.MetadataEntry, InResult.Curvature);
+			}
 		}
 
-		virtual void Sample(const FTransform& InTransform, const FBox& InBox, UPCGPointData* OutPointData)
+		virtual void Sample(const FSamplerResult& InResult, UPCGPointData* OutPointData)
 		{
 			FPCGPoint TrivialPoint, BoundsTestPoint;
-			if(SpatialData->SamplePoint(InTransform, InBox, TrivialPoint, OutPointData->Metadata)
-				&& (!BoundingShapeData || BoundingShapeData->SamplePoint(InTransform, InBox, BoundsTestPoint, nullptr)))
+			if(SpatialData->SamplePoint(InResult.Transform, InResult.Box, TrivialPoint, OutPointData->Metadata)
+				&& (!BoundingShapeData || BoundingShapeData->SamplePoint(InResult.Transform, InResult.Box, BoundsTestPoint, nullptr)))
 			{
+				SetMetadata(InResult, TrivialPoint, OutPointData->Metadata);
 				OutPointData->GetMutablePoints().Add(TrivialPoint);
 			}
 		}
 
-		const UPCGPolyLineData* LineData;
-		const UPCGSpatialData* SpatialData;
-		const UPCGSpatialData* BoundingShapeData;
+		const FPCGSplineSamplerParams Params;
+		const UPCGPolyLineData* LineData = nullptr;
+		const UPCGSpatialData* SpatialData = nullptr;
+		const UPCGSpatialData* BoundingShapeData = nullptr;
+
+		bool bSetMetadata = false;
+		FPCGMetadataAttribute<double>* NextDirectionDeltaAttribute = nullptr;
+		FPCGMetadataAttribute<double>* CurvatureAttribute = nullptr;
 	};
 
 	/** Samples in a volume surrounding the poly line. */
 	struct FVolumeSampler : public FDimensionSampler
 	{
-		FVolumeSampler(const UPCGPolyLineData* InLineData, const UPCGSpatialData* InSpatialData, const UPCGSpatialData* InBoundingShapeData, const FPCGSplineSamplerParams& Params)
-			: FDimensionSampler(InLineData, InSpatialData, InBoundingShapeData)
+		FVolumeSampler(const UPCGPolyLineData* InLineData, const UPCGSpatialData* InSpatialData, const UPCGSpatialData* InBoundingShapeData, const FPCGSplineSamplerParams& Params, UPCGPointData* OutPointData)
+			: FDimensionSampler(InLineData, InSpatialData, InBoundingShapeData, Params, OutPointData)
 		{
 			Fill = Params.Fill;
 			NumPlanarSteps = 1 + ((Params.Dimension == EPCGSplineSamplingDimension::OnVertical) ? 0 : Params.NumPlanarSubdivisions);
 			NumHeightSteps = 1 + ((Params.Dimension == EPCGSplineSamplingDimension::OnHorizontal) ? 0 : Params.NumHeightSubdivisions);
 		}
 
-		virtual void Sample(const FTransform& InTransform, const FBox& InBox, UPCGPointData* OutPointData) override
+		virtual void Sample(const FSamplerResult& InResult, UPCGPointData* OutPointData) override
 		{
+			const FTransform& InTransform = InResult.Transform;
+			const FBox& InBox = InResult.Box;
+
 			// We're assuming that we can scale against the origin in this method so this should always be true.
 			// We will also assume that we can separate the curve into 4 ellipse sections for radius checks
 			check(InBox.Max.Y > 0 && InBox.Min.Y < 0 && InBox.Max.Z > 0 && InBox.Min.Z < 0);
@@ -520,6 +602,7 @@ namespace PCGSplineSampler
 							if (SpatialData->SamplePoint(TentativeTransform, SubBox, OutPoint, OutPointData->Metadata)
 								&& (!BoundingShapeData || BoundingShapeData->SamplePoint(TentativeTransform, SubBox, BoundsTestPoint, nullptr)))
 							{
+								SetMetadata(InResult, OutPoint, OutPointData->Metadata);
 								OutPointData->GetMutablePoints().Add(OutPoint);
 							}
 						}
@@ -546,20 +629,45 @@ namespace PCGSplineSampler
 
 		FStepSampler* Sampler = ((Params.Mode == EPCGSplineSamplingMode::Subdivision) ? static_cast<FStepSampler*>(&SubdivisionSampler) : static_cast<FStepSampler*>(&DistanceSampler));
 
-		FDimensionSampler TrivialDimensionSampler(LineData, SpatialData, InBoundingShapeData);
-		FVolumeSampler VolumeSampler(LineData, SpatialData, InBoundingShapeData, Params);
+		FDimensionSampler TrivialDimensionSampler(LineData, SpatialData, InBoundingShapeData, Params, OutPointData);
+		FVolumeSampler VolumeSampler(LineData, SpatialData, InBoundingShapeData, Params, OutPointData);
 
 		FDimensionSampler* ExtentsSampler = ((Params.Dimension == EPCGSplineSamplingDimension::OnSpline) ? &TrivialDimensionSampler : static_cast<FDimensionSampler*>(&VolumeSampler));
 
-		FTransform SeedTransform;
+		bool bHasPreviousPoint = false;
+		FSamplerResult Results[2];
+		FSamplerResult* PreviousResult = &Results[0];
+		FSamplerResult* CurrentResult = &Results[1];
 
-		while (!Sampler->IsDone())
+		if (!Sampler->IsDone())
+		{
+			Sampler->Step(*PreviousResult);
+			bHasPreviousPoint = true;
+		}
+		
+		while(!Sampler->IsDone())
+		{
+			// Sample point on spline proper
+			Sampler->Step(*CurrentResult);
+
+			// Get unsigned angle difference between the two points
+			const FVector::FReal DeltaSinAngle = ((CurrentResult->Transform.GetUnitAxis(EAxis::X) ^ PreviousResult->Transform.GetUnitAxis(EAxis::X)) | PreviousResult->Transform.GetUnitAxis(EAxis::Z));
+			// Normalize value to be between -1 and 1
+			FVector::FReal DeltaAngle = FMath::Asin(DeltaSinAngle) / UE_HALF_PI;
+
+			PreviousResult->NextDeltaAngle = DeltaAngle;
+			CurrentResult->PreviousDeltaAngle = -DeltaAngle;
+			// Perform samples "around" the spline depending on the settings
+			ExtentsSampler->Sample(*PreviousResult, OutPointData);
+
+			// Prepare for next iteration
+			Swap(PreviousResult, CurrentResult);
+			*CurrentResult = FSamplerResult();
+		}
+		
+		if (bHasPreviousPoint)
 		{
-			FBox SeedBox = FBox::BuildAABB(FVector::ZeroVector, FVector::OneVector);
-			// Get seed transform/box
-			Sampler->Step(SeedTransform, SeedBox);
-			// From seed point, sample in other dimensions as needed
-			ExtentsSampler->Sample(SeedTransform, SeedBox, OutPointData);
+			ExtentsSampler->Sample(*PreviousResult, OutPointData);
 		}
 
 		// Finally, set seed on points based on position
diff --git a/Engine/Plugins/Experimental/PCG/Source/PCG/Public/Data/PCGLandscapeSplineData.h b/Engine/Plugins/Experimental/PCG/Source/PCG/Public/Data/PCGLandscapeSplineData.h
--- a/Engine/Plugins/Experimental/PCG/Source/PCG/Public/Data/PCGLandscapeSplineData.h
+++ b/Engine/Plugins/Experimental/PCG/Source/PCG/Public/Data/PCGLandscapeSplineData.h
@@ -25,6 +25,7 @@ public:
 	virtual int GetNumSegments() const override;
 	virtual FVector::FReal GetSegmentLength(int SegmentIndex) const override;
 	virtual FTransform GetTransformAtDistance(int SegmentIndex, FVector::FReal Distance, FBox* OutBounds = nullptr) const override;
+	virtual FVector::FReal GetCurvatureAtDistance(int SegmentIndex, FVector::FReal Distance) const override;
 	//~End UPCGPolyLineData interface
 
 	//~Begin UPCGSpatialDataWithPointCache interface
diff --git a/Engine/Plugins/Experimental/PCG/Source/PCG/Public/Data/PCGPolyLineData.h b/Engine/Plugins/Experimental/PCG/Source/PCG/Public/Data/PCGPolyLineData.h
--- a/Engine/Plugins/Experimental/PCG/Source/PCG/Public/Data/PCGPolyLineData.h
+++ b/Engine/Plugins/Experimental/PCG/Source/PCG/Public/Data/PCGPolyLineData.h
@@ -22,8 +22,9 @@ public:
 	//~End UPCGSpatialData interface
 
 	virtual int GetNumSegments() const PURE_VIRTUAL(UPCGPolyLineData::GetNumSegments, return 0;);
-	virtual FVector::FReal GetSegmentLength(int SegmentIndex) const PURE_VIRTUAL(UPCGPolyLineData::GetSegmentLength, return 0.0f;);
+	virtual FVector::FReal GetSegmentLength(int SegmentIndex) const PURE_VIRTUAL(UPCGPolyLineData::GetSegmentLength, return 0;);
 	virtual FTransform GetTransformAtDistance(int SegmentIndex, FVector::FReal Distance, FBox* OutBounds = nullptr) const PURE_VIRTUAL(UPCGPolyLine::GetTransformAtDistance, return FTransform(););
+	virtual FVector::FReal GetCurvatureAtDistance(int SegmentIndex, FVector::FReal Distance) const { return 0; }
 
 	virtual FVector GetLocationAtDistance(int SegmentIndex, FVector::FReal Distance) const { return GetTransformAtDistance(SegmentIndex, Distance).GetLocation(); }
 };
diff --git a/Engine/Plugins/Experimental/PCG/Source/PCG/Public/Data/PCGSplineData.h b/Engine/Plugins/Experimental/PCG/Source/PCG/Public/Data/PCGSplineData.h
--- a/Engine/Plugins/Experimental/PCG/Source/PCG/Public/Data/PCGSplineData.h
+++ b/Engine/Plugins/Experimental/PCG/Source/PCG/Public/Data/PCGSplineData.h
@@ -31,6 +31,7 @@ public:
 	virtual FVector::FReal GetSegmentLength(int SegmentIndex) const override;
 	virtual FVector GetLocationAtDistance(int SegmentIndex, FVector::FReal Distance) const override;
 	virtual FTransform GetTransformAtDistance(int SegmentIndex, FVector::FReal Distance, FBox* OutBounds = nullptr) const override;
+	virtual FVector::FReal GetCurvatureAtDistance(int SegmentIndex, FVector::FReal Distance) const override;
 	//~End UPCGPolyLineData interface
 
 	//~Begin UPCGSpatialDataWithPointCache interface
diff --git a/Engine/Plugins/Experimental/PCG/Source/PCG/Public/Elements/PCGSplineSampler.h b/Engine/Plugins/Experimental/PCG/Source/PCG/Public/Elements/PCGSplineSampler.h
--- a/Engine/Plugins/Experimental/PCG/Source/PCG/Public/Elements/PCGSplineSampler.h
+++ b/Engine/Plugins/Experimental/PCG/Source/PCG/Public/Elements/PCGSplineSampler.h
@@ -69,6 +69,20 @@ struct PCG_API FPCGSplineSamplerParams
 	UPROPERTY(BlueprintReadWrite, EditAnywhere, Category = Settings, meta = (ClampMin = "0", EditCondition = "Dimension==EPCGSplineSamplingDimension::OnVertical||Dimension==EPCGSplineSamplingDimension::OnVolume"))
 	int32 NumHeightSubdivisions = 8;
 
+	UPROPERTY(BlueprintReadWrite, EditAnywhere, Category = Settings, meta = (InlineEditConditionToggle))
+	bool bComputeDirectionDelta = false;
+
+	/** Attribute that wil contain the delta angle to the next point on the spline w.r.t to the current's point Up vector. */
+	UPROPERTY(BlueprintReadWrite, EditAnywhere, Category = Settings, meta = (EditCondition = "bComputeDirectionDelta"))
+	FName NextDirectionDeltaAttribute = "NextDirectionDelta";
+
+	UPROPERTY(BlueprintReadWrite, EditAnywhere, Category = Settings, meta = (InlineEditConditionToggle))
+	bool bComputeCurvature = false;
+
+	/** Attribute that will contain the curvature. Note that the radius of curvature is defined as 1/Curvature, and might need you to scale to world units */
+	UPROPERTY(BlueprintReadWrite, EditAnywhere, Category = Settings, meta = (EditCondition = "bComputeCurvature"))
+	FName CurvatureAttribute = "Curvature";
+
 	/** The space between each sample point */
 	UPROPERTY(BlueprintReadWrite, EditAnywhere, Category = Settings, meta = (ClampMin = "0.1", EditCondition = "Dimension==EPCGSplineSamplingDimension::OnInterior"))
 	float InteriorSampleSpacing = 100.0f;
