More thorough outer cleanup out from BP execution nodes

#rnx

diff --git a/Engine/Plugins/Experimental/PCG/Source/PCG/Private/Elements/PCGExecuteBlueprint.cpp b/Engine/Plugins/Experimental/PCG/Source/PCG/Private/Elements/PCGExecuteBlueprint.cpp
--- a/Engine/Plugins/Experimental/PCG/Source/PCG/Private/Elements/PCGExecuteBlueprint.cpp
+++ b/Engine/Plugins/Experimental/PCG/Source/PCG/Private/Elements/PCGExecuteBlueprint.cpp
@@ -529,19 +529,44 @@ bool FPCGExecuteBlueprintElement::ExecuteInternal(FPCGContext* InContext) const
 			// Note that we will recurse up the outer tree to make sure we catch every case.
 			if(Output.Data)
 			{
+				auto ReOuterToTransientPackageIfCreatedFromThis = [Context](UObject* InObject)
+				{
+					bool bHasInstanceAsOuter = false;
+					UObject* CurrentObject = InObject;
+					while (CurrentObject && !bHasInstanceAsOuter)
+					{
+						bHasInstanceAsOuter = (CurrentObject->GetOuter() == Context->BlueprintElementInstance);
+						CurrentObject = CurrentObject->GetOuter();
+					}
+
+					if (bHasInstanceAsOuter)
+					{
+						InObject->Rename(nullptr, GetTransientPackage(), REN_ForceNoResetLoaders | REN_DoNotDirty | REN_DontCreateRedirectors | REN_NonTransactional);
+					}
+				};
+
 				UObject* ThisData = const_cast<UPCGData*>(Output.Data.Get());
+				ReOuterToTransientPackageIfCreatedFromThis(ThisData);
 
-				bool bHasInstanceAsOuter = false;
-				UObject* CurrentObject = ThisData;
-				while (CurrentObject && !bHasInstanceAsOuter)
+				// Similarly, if the metadata on the data inherits from a non-transient data created by this BP instance, it should be reoutered.
+				const UPCGMetadata* Metadata = nullptr;
+				if (UPCGSpatialData* SpatialData = Cast<UPCGSpatialData>(ThisData))
 				{
-					bHasInstanceAsOuter = (CurrentObject->GetOuter() == Context->BlueprintElementInstance);
-					CurrentObject = CurrentObject->GetOuter();
+					Metadata = SpatialData->Metadata;
 				}
-
-				if (bHasInstanceAsOuter)
+				else if (UPCGParamData* ParamData = Cast<UPCGParamData>(ThisData))
+				{
+					Metadata = ParamData->Metadata;
+				}
+				
+				if (Metadata)
 				{
-					ThisData->Rename(nullptr, GetTransientPackage(), REN_ForceNoResetLoaders | REN_DoNotDirty | REN_DontCreateRedirectors | REN_NonTransactional);
+					while (Metadata->GetParent())
+					{
+						UObject* OuterObject = Metadata->GetParent()->GetOuter();
+						ReOuterToTransientPackageIfCreatedFromThis(OuterObject);
+						Metadata = Metadata->GetParent();
+					}
 				}
 			}
 		}
