Data from actor now returns collection instead of union
Added tagging based on actors & components
Fixed self case (though it will still need a no filter option)
Changed the way we parse components so we get everything (incl. primtiives even if not previously parsed)

diff --git a/Engine/Plugins/Experimental/PCG/Source/PCG/Private/Elements/PCGDataFromActor.cpp b/Engine/Plugins/Experimental/PCG/Source/PCG/Private/Elements/PCGDataFromActor.cpp
--- a/Engine/Plugins/Experimental/PCG/Source/PCG/Private/Elements/PCGDataFromActor.cpp
+++ b/Engine/Plugins/Experimental/PCG/Source/PCG/Private/Elements/PCGDataFromActor.cpp
@@ -66,7 +66,7 @@ bool FPCGDataFromActorElement::ExecuteInternal(FPCGContext* InContext) const
 
 	if (!Context->bPerformedQuery)
 	{
-		Context->FoundActors = PCGActorSelector::FindActors(Settings->ActorSelector, World, nullptr);
+		Context->FoundActors = PCGActorSelector::FindActors(Settings->ActorSelector, World, Context->SourceComponent.IsValid() ? Context->SourceComponent->GetOwner() : nullptr);
 		Context->bPerformedQuery = true;
 
 		if (Context->FoundActors.IsEmpty())
@@ -228,8 +228,9 @@ void FPCGDataFromActorElement::ProcessActor(FPCGContext* Context, const UPCGData
 	}
 	else
 	{
-		FPCGTaggedData& Output = Outputs.Emplace_GetRef();
-		Output.Data = UPCGComponent::CreateActorPCGData(FoundActor, Context->SourceComponent.Get(), /*bParseActor=*/Settings->Mode != EPCGGetDataFromActorMode::GetSinglePoint);
+		const bool bParseActor = (Settings->Mode != EPCGGetDataFromActorMode::GetSinglePoint);
+		FPCGDataCollection Collection = UPCGComponent::CreateActorPCGDataCollection(FoundActor, Context->SourceComponent.Get(), bParseActor);
+		Outputs += Collection.TaggedData;
 	}
 
 	if (Context->SourceComponent.IsValid())
diff --git a/Engine/Plugins/Experimental/PCG/Source/PCG/Private/PCGComponent.cpp b/Engine/Plugins/Experimental/PCG/Source/PCG/Private/PCGComponent.cpp
--- a/Engine/Plugins/Experimental/PCG/Source/PCG/Private/PCGComponent.cpp
+++ b/Engine/Plugins/Experimental/PCG/Source/PCG/Private/PCGComponent.cpp
@@ -25,6 +25,8 @@
 #include "Grid/PCGPartitionActor.h"
 #include "Helpers/PCGActorHelpers.h"
 
+#include "Algo/Transform.h"
+#include "Components/BillboardComponent.h"
 #include "Components/InstancedStaticMeshComponent.h"
 #include "Components/SplineComponent.h"
 #include "Components/ShapeComponent.h"
@@ -1861,16 +1863,46 @@ UPCGData* UPCGComponent::CreateActorPCGData(AActor* Actor, bool bParseActor)
 }
 
 UPCGData* UPCGComponent::CreateActorPCGData(AActor* Actor, const UPCGComponent* Component, bool bParseActor)
+{
+	FPCGDataCollection Collection = CreateActorPCGDataCollection(Actor, Component, bParseActor);
+	if (Collection.TaggedData.Num() > 1)
+	{
+		UPCGUnionData* Union = NewObject<UPCGUnionData>();
+		for (const FPCGTaggedData& TaggedData : Collection.TaggedData)
+		{
+			Union->AddData(CastChecked<const UPCGSpatialData>(TaggedData.Data));
+		}
+
+		return Union;
+	}
+	else if(Collection.TaggedData.Num() == 1)
+	{
+		return Cast<UPCGData>(Collection.TaggedData[0].Data);
+	}
+	else
+	{
+		return nullptr;
+	}
+}
+
+FPCGDataCollection UPCGComponent::CreateActorPCGDataCollection(AActor* Actor, const UPCGComponent* Component, bool bParseActor)
 {
 	TRACE_CPUPROFILER_EVENT_SCOPE(UPCGComponent::CreateActorPCGData);
+	FPCGDataCollection Collection;
 
 	if (!Actor)
 	{
-		return nullptr;
+		return Collection;
 	}
 
-	// In this case, we'll build the data type that's closest to known actor types
+	auto NameTagsToStringTags = [](const FName& InName) { return InName.ToString(); };
+	TSet<FString> ActorTags;
+	Algo::Transform(Actor->Tags, ActorTags, NameTagsToStringTags);
+
+	// Fill in collection based on the data on the given actor.
+	// Some actor types we will forego full parsing to build strictly on the actor existence, such as partition actors, volumes and landscape
 	// TODO: add factory for extensibility
+	// TODO: review the !bParseActor cases - it might make sense to have just a point for a partition actor, even if we preintersect it.
 	if (APCGPartitionActor* PartitionActor = Cast<APCGPartitionActor>(Actor))
 	{
 		check(!Component || Component->GetOwner() == Actor); // Invalid processing otherwise because of the this usage
@@ -1879,15 +1911,12 @@ UPCGData* UPCGComponent::CreateActorPCGData(AActor* Actor, const UPCGComponent*
 		Data->Initialize(PartitionActor->GetFixedBounds(), PartitionActor);
 
 		UPCGComponent* OriginalComponent = Component ? PartitionActor->GetOriginalComponent(Component) : nullptr;
-		if (OriginalComponent)
-		{
-			if (const UPCGSpatialData* OriginalComponentSpatialData = Cast<const UPCGSpatialData>(OriginalComponent->GetActorPCGData()))
-			{
-				return Data->IntersectWith(OriginalComponentSpatialData);
-			}
-		}
+		// Important note: we do NOT call the collection version here, as we want to have a union if that's the case
+		const UPCGSpatialData* OriginalComponentSpatialData = OriginalComponent ? Cast<const UPCGSpatialData>(OriginalComponent->GetActorPCGData()) : nullptr;
 
-		return Data;
+		FPCGTaggedData& TaggedData = Collection.TaggedData.Emplace_GetRef();
+		TaggedData.Data = OriginalComponentSpatialData ? Cast<UPCGData>(Data->IntersectWith(OriginalComponentSpatialData)) : Cast<UPCGData>(Data);
+		// No need to keep partition actor tags, though we might want to push PCG grid GUID at some point
 	}
 	else if (ALandscapeProxy* Landscape = Cast<ALandscapeProxy>(Actor))
 	{
@@ -1896,58 +1925,78 @@ UPCGData* UPCGComponent::CreateActorPCGData(AActor* Actor, const UPCGComponent*
 
 		Data->Initialize({ Landscape }, PCGHelpers::GetGridBounds(Actor, Component), /*bHeightOnly=*/false, bUseLandscapeMetadata);
 
-		return Data;
+		FPCGTaggedData& TaggedData = Collection.TaggedData.Emplace_GetRef();
+		TaggedData.Data = Data;
+		TaggedData.Tags = ActorTags;
 	}
 	else if (!bParseActor)
 	{
 		UPCGPointData* Data = NewObject<UPCGPointData>();
 		Data->InitializeFromActor(Actor);
-		return Data;
+
+		FPCGTaggedData& TaggedData = Collection.TaggedData.Emplace_GetRef();
+		TaggedData.Data = Data;
+		TaggedData.Tags = ActorTags;
 	}
 	else if (AVolume* Volume = Cast<AVolume>(Actor))
 	{
 		UPCGVolumeData* Data = NewObject<UPCGVolumeData>();
 		Data->Initialize(Volume);
 
-		return Data;
+		FPCGTaggedData& TaggedData = Collection.TaggedData.Emplace_GetRef();
+		TaggedData.Data = Data;
+		TaggedData.Tags = ActorTags;
 	}
 	else // Prepare data on a component basis
 	{
-		TInlineComponentArray<ULandscapeSplinesComponent*, 1> LandscapeSplines;
+		TInlineComponentArray<UPrimitiveComponent*, 4> Primitives;
+
+		auto RemoveDuplicatesFromPrimitives = [&Primitives](const auto& InComponents)
+		{
+			Primitives.RemoveAll([&InComponents](UPrimitiveComponent* Component)
+			{
+				return InComponents.Contains(Component);
+			});
+		};
+
+		auto RemovePCGGeneratedEntries = [](auto& InComponents)
+		{
+			for (int32 Index = InComponents.Num() - 1; Index >= 0; --Index)
+			{
+				if (InComponents[Index]->ComponentTags.Contains(PCGHelpers::DefaultPCGTag))
+				{
+					InComponents.RemoveAtSwap(Index);
+				}
+			}
+		};
+
+		Actor->GetComponents(Primitives);
+		RemovePCGGeneratedEntries(Primitives);
+
+		TInlineComponentArray<ULandscapeSplinesComponent*, 4> LandscapeSplines;
 		Actor->GetComponents(LandscapeSplines);
+		RemovePCGGeneratedEntries(LandscapeSplines);
+		RemoveDuplicatesFromPrimitives(LandscapeSplines);
 
-		TInlineComponentArray<USplineComponent*, 1> Splines;
+		TInlineComponentArray<USplineComponent*, 4> Splines;
 		Actor->GetComponents(Splines);
+		RemovePCGGeneratedEntries(Splines);
+		RemoveDuplicatesFromPrimitives(Splines);
 
-		TInlineComponentArray<UShapeComponent*, 1> Shapes;
+		TInlineComponentArray<UShapeComponent*, 4> Shapes;
 		Actor->GetComponents(Shapes);
-
-		// Don't get generic primitives unless it's the only thing we can find.
-		TInlineComponentArray<UPrimitiveComponent*, 1> OtherPrimitives;
-		if (LandscapeSplines.Num() == 0 && Splines.Num() == 0 && Shapes.Num() == 0)
-		{
-			Actor->GetComponents(OtherPrimitives);
-		}
-
-		UPCGUnionData* Union = nullptr;
-		if (LandscapeSplines.Num() + Splines.Num() + Shapes.Num() + OtherPrimitives.Num() > 1)
-		{
-			Union = NewObject<UPCGUnionData>();
-		}
+		RemovePCGGeneratedEntries(Shapes);
+		RemoveDuplicatesFromPrimitives(Shapes);
 
 		for (ULandscapeSplinesComponent* SplineComponent : LandscapeSplines)
 		{
 			UPCGLandscapeSplineData* SplineData = NewObject<UPCGLandscapeSplineData>();
 			SplineData->Initialize(SplineComponent);
 
-			if (Union)
-			{
-				Union->AddData(SplineData);
-			}
-			else
-			{
-				return SplineData;
-			}
+			FPCGTaggedData& TaggedData = Collection.TaggedData.Emplace_GetRef();
+			TaggedData.Data = SplineData;
+			Algo::Transform(SplineComponent->ComponentTags, TaggedData.Tags, NameTagsToStringTags);
+			TaggedData.Tags.Append(ActorTags);
 		}
 
 		for (USplineComponent* SplineComponent : Splines)
@@ -1955,56 +2004,53 @@ UPCGData* UPCGComponent::CreateActorPCGData(AActor* Actor, const UPCGComponent*
 			UPCGSplineData* SplineData = NewObject<UPCGSplineData>();
 			SplineData->Initialize(SplineComponent);
 
-			if (Union)
-			{
-				Union->AddData(SplineData);
-			}
-			else
-			{
-				return SplineData;
-			}
+			FPCGTaggedData& TaggedData = Collection.TaggedData.Emplace_GetRef();
+			TaggedData.Data = SplineData;
+			Algo::Transform(SplineComponent->ComponentTags, TaggedData.Tags, NameTagsToStringTags);
+			TaggedData.Tags.Append(ActorTags);
 		}
 
 		for (UShapeComponent* ShapeComponent : Shapes)
 		{
 			UPCGPrimitiveData* ShapeData = NewObject<UPCGPrimitiveData>();
 			ShapeData->Initialize(ShapeComponent);
-			
-			if (Union)
-			{
-				Union->AddData(ShapeData);
-			}
-			else
-			{
-				return ShapeData;
-			}
+
+			FPCGTaggedData& TaggedData = Collection.TaggedData.Emplace_GetRef();
+			TaggedData.Data = ShapeData;
+			Algo::Transform(ShapeComponent->ComponentTags, TaggedData.Tags, NameTagsToStringTags);
+			TaggedData.Tags.Append(ActorTags);
 		}
 
-		for (UPrimitiveComponent* PrimitiveComponent : OtherPrimitives)
+		for (UPrimitiveComponent* PrimitiveComponent : Primitives)
 		{
-			UPCGPrimitiveData* PrimitiveData = NewObject<UPCGPrimitiveData>();
-			PrimitiveData->Initialize(PrimitiveComponent);
-
-			if (Union)
+			// Exception: skip the billboard component
+			if (Cast<UBillboardComponent>(PrimitiveComponent))
 			{
-				Union->AddData(PrimitiveData);
+				continue;
 			}
-			else
-			{
-				return PrimitiveData;
-			}
-		}
 
-		if (Union)
-		{
-			return Union;
+			UPCGPrimitiveData* PrimitiveData = NewObject<UPCGPrimitiveData>();
+			PrimitiveData->Initialize(PrimitiveComponent);
+
+			FPCGTaggedData& TaggedData = Collection.TaggedData.Emplace_GetRef();
+			TaggedData.Data = PrimitiveData;
+			Algo::Transform(PrimitiveComponent->ComponentTags, TaggedData.Tags, NameTagsToStringTags);
+			TaggedData.Tags.Append(ActorTags);
 		}
 	}
 
 	// Finally, if it's not a special actor and there are not parsed components, then return a single point at the actor position
-	UPCGPointData* Data = NewObject<UPCGPointData>();
-	Data->InitializeFromActor(Actor);
-	return Data;
+	if (Collection.TaggedData.IsEmpty())
+	{
+		UPCGPointData* Data = NewObject<UPCGPointData>();
+		Data->InitializeFromActor(Actor);
+
+		FPCGTaggedData& TaggedData = Collection.TaggedData.Emplace_GetRef();
+		TaggedData.Data = Data;
+		TaggedData.Tags = ActorTags;
+	}
+
+	return Collection;
 }
 
 UPCGData* UPCGComponent::CreatePCGData()
diff --git a/Engine/Plugins/Experimental/PCG/Source/PCG/Public/PCGComponent.h b/Engine/Plugins/Experimental/PCG/Source/PCG/Public/PCGComponent.h
--- a/Engine/Plugins/Experimental/PCG/Source/PCG/Public/PCGComponent.h
+++ b/Engine/Plugins/Experimental/PCG/Source/PCG/Public/PCGComponent.h
@@ -238,6 +238,9 @@ public:
 
 	FBox GetGridBounds() const;
 
+	/** Builds the PCG data from a given actor and its PCG component, and places it in a data collection with appropriate tags */
+	static FPCGDataCollection CreateActorPCGDataCollection(AActor* Actor, const UPCGComponent* Component, bool bParseActor = true);
+
 	/** Builds the canonical PCG data from a given actor and its PCG component if any. */
 	static UPCGData* CreateActorPCGData(AActor* Actor, const UPCGComponent* Component, bool bParseActor = true);
 
