Added filter by tag element

diff --git a/Engine/Plugins/Experimental/PCG/Source/PCG/Private/Elements/PCGAttributeFilterElement.cpp b/Engine/Plugins/Experimental/PCG/Source/PCG/Private/Elements/PCGAttributeFilterElement.cpp
--- a/Engine/Plugins/Experimental/PCG/Source/PCG/Private/Elements/PCGAttributeFilterElement.cpp
+++ b/Engine/Plugins/Experimental/PCG/Source/PCG/Private/Elements/PCGAttributeFilterElement.cpp
@@ -3,8 +3,9 @@
 #include "Elements/PCGAttributeFilterElement.h"
 
 #include "Data/PCGSpatialData.h"
-#include "PCGParamData.h"
 #include "PCGContext.h"
+#include "PCGHelpers.h"
+#include "PCGParamData.h"
 #include "PCGPin.h"
 
 #include UE_INLINE_GENERATED_CPP_BY_NAME(PCGAttributeFilterElement)
@@ -14,16 +15,6 @@ namespace PCGAttributeFilterConstants
 	const FName NodeName = TEXT("FilterAttribute");
 }
 
-namespace PCGAttributeFilterSettings
-{
-	TArray<FString> GenerateNameArray(const FString& InString)
-	{
-		TArray<FString> Result;
-		InString.ParseIntoArrayWS(Result, TEXT(","));
-		return Result;
-	}
-}
-
 void UPCGAttributeFilterSettings::PostLoad()
 {
 	Super::PostLoad();
@@ -58,7 +49,7 @@ FName UPCGAttributeFilterSettings::GetDefaultNodeName() const
 
 FName UPCGAttributeFilterSettings::AdditionalTaskName() const
 {
-	TArray<FString> AttributesToKeep = PCGAttributeFilterSettings::GenerateNameArray(SelectedAttributes);
+	TArray<FString> AttributesToKeep = PCGHelpers::GetStringArrayFromCommaSeparatedString(SelectedAttributes);
 
 	FString NodeName = PCGAttributeFilterConstants::NodeName.ToString();
 
@@ -145,7 +136,7 @@ bool FPCGAttributeFilterElement::ExecuteInternal(FPCGContext* Context) const
 			continue;
 		}
 
-		TArray<FString> AttributesToKeep = PCGAttributeFilterSettings::GenerateNameArray(Settings->SelectedAttributes);
+		TArray<FString> AttributesToKeep = PCGHelpers::GetStringArrayFromCommaSeparatedString(Settings->SelectedAttributes);
 
 		// Then add/remove each attribute in the list explicitly
 		for (const FString& AttributeName : AttributesToKeep)
diff --git a/Engine/Plugins/Experimental/PCG/Source/PCG/Private/Elements/PCGFilterByTag.cpp b/Engine/Plugins/Experimental/PCG/Source/PCG/Private/Elements/PCGFilterByTag.cpp
new file mode 100644
--- /dev/null
+++ b/Engine/Plugins/Experimental/PCG/Source/PCG/Private/Elements/PCGFilterByTag.cpp
@@ -0,0 +1,94 @@
+// Copyright Epic Games, Inc. All Rights Reserved.
+
+#include "Elements/PCGFilterByTag.h"
+#include "PCGContext.h"
+#include "PCGHelpers.h"
+#include "PCGPin.h"
+
+#include UE_INLINE_GENERATED_CPP_BY_NAME(PCGFilterByTag)
+
+#define LOCTEXT_NAMESPACE "PCGFilterByTag"
+
+namespace PCGFilterByTagConstants
+{
+	const FName NodeName = FName(TEXT("Filter By Tag"));
+}
+
+#if WITH_EDITOR
+FName UPCGFilterByTagSettings::GetDefaultNodeName() const
+{
+	return PCGFilterByTagConstants::NodeName;
+}
+
+FText UPCGFilterByTagSettings::GetNodeTooltipText() const
+{
+	return LOCTEXT("FilterByTagNodeTooltip", "Filters data in the collection according to whether they have, or don't have, some tags");
+}
+#endif
+
+FName UPCGFilterByTagSettings::AdditionalTaskName() const
+{
+	const TArray<FString> Tags = PCGHelpers::GetStringArrayFromCommaSeparatedString(SelectedTags);
+
+	FString NodeName = PCGFilterByTagConstants::NodeName.ToString();
+	NodeName += (Operation == EPCGFilterByTagOperation::KeepTagged ? TEXT(" (Keep)") : TEXT(" (Remove)"));
+
+	if (Tags.Num() == 1)
+	{
+		return FName(FString::Printf(TEXT("%s: %s"), *NodeName, *Tags[0]));
+	}
+	else
+	{
+		return FName(NodeName);
+	}
+}
+
+TArray<FPCGPinProperties> UPCGFilterByTagSettings::OutputPinProperties() const
+{
+	TArray<FPCGPinProperties> PinProperties;
+	PinProperties.Emplace(PCGPinConstants::DefaultOutputLabel, EPCGDataType::Any);
+
+	return PinProperties;
+}
+
+FPCGElementPtr UPCGFilterByTagSettings::CreateElement() const
+{
+	return MakeShared<FPCGFilterByTagElement>();
+}
+
+bool FPCGFilterByTagElement::ExecuteInternal(FPCGContext* Context) const
+{
+	TRACE_CPUPROFILER_EVENT_SCOPE(FPCGFilterByTagElement::Execute);
+	check(Context);
+
+	const UPCGFilterByTagSettings* Settings = Context->GetInputSettings<UPCGFilterByTagSettings>();
+
+	const bool bKeepIfTag = (Settings->Operation == EPCGFilterByTagOperation::KeepTagged);
+	const TArray<FString> Tags = PCGHelpers::GetStringArrayFromCommaSeparatedString(Settings->SelectedTags);
+
+	TArray<FPCGTaggedData> Inputs = Context->InputData.GetInputs();
+	TArray<FPCGTaggedData>& Outputs = Context->OutputData.TaggedData;
+
+	for (const FPCGTaggedData& Input : Inputs)
+	{
+		bool bHasCommonTags = false;
+
+		for (const FString& Tag : Tags)
+		{
+			if (Input.Tags.Contains(Tag))
+			{
+				bHasCommonTags = true;
+				break;
+			}
+		}
+
+		if (bKeepIfTag == bHasCommonTags)
+		{
+			Outputs.Add(Input);
+		}
+	}
+
+	return true;
+}
+
+#undef LOCTEXT_NAMESPACE
\ No newline at end of file
diff --git a/Engine/Plugins/Experimental/PCG/Source/PCG/Private/PCGHelpers.cpp b/Engine/Plugins/Experimental/PCG/Source/PCG/Private/PCGHelpers.cpp
--- a/Engine/Plugins/Experimental/PCG/Source/PCG/Private/PCGHelpers.cpp
+++ b/Engine/Plugins/Experimental/PCG/Source/PCG/Private/PCGHelpers.cpp
@@ -260,6 +260,13 @@ namespace PCGHelpers
 		return (InWorld && InWorld->GetSubsystem<UPCGSubsystem>()) ? InWorld->GetSubsystem<UPCGSubsystem>()->GetPCGWorldActor() : nullptr;
 	}
 
+	TArray<FString> GetStringArrayFromCommaSeparatedString(const FString& InCommaSeparatedString)
+	{
+		TArray<FString> Result;
+		InCommaSeparatedString.ParseIntoArrayWS(Result, TEXT(","));
+		return Result;
+	}
+
 #if WITH_EDITOR
 	bool CanBeExpanded(UClass* ObjectClass)
 	{
diff --git a/Engine/Plugins/Experimental/PCG/Source/PCG/Private/PCGHelpers.h b/Engine/Plugins/Experimental/PCG/Source/PCG/Private/PCGHelpers.h
--- a/Engine/Plugins/Experimental/PCG/Source/PCG/Private/PCGHelpers.h
+++ b/Engine/Plugins/Experimental/PCG/Source/PCG/Private/PCGHelpers.h
@@ -42,6 +42,8 @@ namespace PCGHelpers
 
 	APCGWorldActor* GetPCGWorldActor(UWorld* InWorld);
 
+	TArray<FString> GetStringArrayFromCommaSeparatedString(const FString& InCommaSeparatedString);
+
 #if WITH_EDITOR
 	void GatherDependencies(UObject* Object, TSet<TObjectPtr<UObject>>& OutDependencies, int32 MaxDepth = -1);
 	void GatherDependencies(FProperty* Property, const void* InContainer, TSet<TObjectPtr<UObject>>& OutDependencies, int32 MaxDepth);
diff --git a/Engine/Plugins/Experimental/PCG/Source/PCG/Public/Elements/PCGFilterByTag.h b/Engine/Plugins/Experimental/PCG/Source/PCG/Public/Elements/PCGFilterByTag.h
new file mode 100644
--- /dev/null
+++ b/Engine/Plugins/Experimental/PCG/Source/PCG/Public/Elements/PCGFilterByTag.h
@@ -0,0 +1,50 @@
+// Copyright Epic Games, Inc. All Rights Reserved.
+
+#pragma once
+
+#include "PCGSettings.h"
+
+#include "PCGFilterByTag.generated.h"
+
+UENUM()
+enum class EPCGFilterByTagOperation
+{
+	KeepTagged,
+	RemoveTagged
+};
+
+/** Filters a data collection based on some tag criterion */
+UCLASS(BlueprintType, ClassGroup = (Procedural))
+class PCG_API UPCGFilterByTagSettings : public UPCGSettings
+{
+	GENERATED_BODY()
+
+public:
+	//~Begin UPCGSettings interface
+#if WITH_EDITOR
+	virtual FName GetDefaultNodeName() const override;
+	virtual FText GetNodeTooltipText() const override;
+	virtual EPCGSettingsType GetType() const override { return EPCGSettingsType::Filter; }
+#endif
+
+	virtual FName AdditionalTaskName() const override;
+
+protected:
+	virtual TArray<FPCGPinProperties> OutputPinProperties() const override;
+	virtual FPCGElementPtr CreateElement() const override;
+	//~End UPCGSettings interface
+
+public:
+	UPROPERTY(BlueprintReadWrite, EditAnywhere, Category = Settings, meta = (PCG_Overridable))
+	EPCGFilterByTagOperation Operation = EPCGFilterByTagOperation::KeepTagged;
+
+	/** Comma-separated list of tags */
+	UPROPERTY(BlueprintReadWrite, EditAnywhere, Category = Settings, meta = (PCG_Overridable))
+	FString SelectedTags;
+};
+
+class FPCGFilterByTagElement : public FSimplePCGElement
+{
+protected:
+	virtual bool ExecuteInternal(FPCGContext* Context) const override;
+};
\ No newline at end of file
