[PCG] ManagedResources logging to give trace of actions related to managed resources.

Add console var to force hard cleanup of resources on generate

diff --git a/Engine/Plugins/Experimental/PCG/Source/PCG/Private/PCGComponent.cpp b/Engine/Plugins/Experimental/PCG/Source/PCG/Private/PCGComponent.cpp
--- a/Engine/Plugins/Experimental/PCG/Source/PCG/Private/PCGComponent.cpp
+++ b/Engine/Plugins/Experimental/PCG/Source/PCG/Private/PCGComponent.cpp
@@ -24,6 +24,7 @@
 #include "Data/PCGVolumeData.h"
 #include "Grid/PCGPartitionActor.h"
 #include "Helpers/PCGActorHelpers.h"
+#include "Utils/PCGGeneratedResourcesLogging.h"
 
 #include "Algo/Transform.h"
 #include "Components/BillboardComponent.h"
@@ -124,6 +125,8 @@ void UPCGComponent::SetGraphLocal(UPCGGraph* InGraph)
 
 void UPCGComponent::AddToManagedResources(UPCGManagedResource* InResource)
 {
+	PCGGeneratedResourcesLogging::LogAddToManagedResources(InResource);
+
 	if (InResource)
 	{
 		FScopeLock ResourcesLock(&GeneratedResourcesLock);
@@ -338,6 +341,8 @@ bool UPCGComponent::GetActorsFromTags(const TSet<FName>& InTags, TSet<TWeakObjec
 
 void UPCGComponent::PostProcessGraph(const FBox& InNewBounds, bool bInGenerated, FPCGContext* Context)
 {
+	PCGGeneratedResourcesLogging::LogPostProcessGraph();
+
 	LastGeneratedBounds = InNewBounds;
 
 	const bool bHadGeneratedOutputBefore = GeneratedGraphOutput.TaggedData.Num() > 0;
@@ -536,6 +541,8 @@ FPCGTaskId UPCGComponent::CleanupInternal(bool bRemoveComponents, bool bSave, co
 		return InvalidPCGTaskId;
 	}
 
+	PCGGeneratedResourcesLogging::LogCleanupInternal(bRemoveComponents);
+
 	Modify();
 
 #if WITH_EDITOR
@@ -664,8 +671,15 @@ bool UPCGComponent::MoveResourcesToNewActor(AActor* InNewActor, bool bCreateChil
 
 void UPCGComponent::CleanupLocalImmediate(bool bRemoveComponents)
 {
+	PCGGeneratedResourcesLogging::LogCleanupLocalImmediate(bRemoveComponents, GeneratedResources);
+
 	TSet<TSoftObjectPtr<AActor>> ActorsToDelete;
 
+	if (!bRemoveComponents && UPCGManagedResource::DebugForcePurgeAllResourcesOnGenerate())
+	{
+		bRemoveComponents = true;
+	}
+
 	{
 		FScopeLock ResourcesLock(&GeneratedResourcesLock);
 		check(!GeneratedResourcesInaccessible);
@@ -673,6 +687,9 @@ void UPCGComponent::CleanupLocalImmediate(bool bRemoveComponents)
 		{
 			// Note: resources can be null here in some loading + bp object cases
 			UPCGManagedResource* Resource = GeneratedResources[ResourceIndex];
+
+			PCGGeneratedResourcesLogging::LogCleanupLocalImmediateResource(Resource);
+
 			if (!Resource || Resource->Release(bRemoveComponents, ActorsToDelete))
 			{
 				GeneratedResources.RemoveAtSwap(ResourceIndex);
@@ -688,6 +705,8 @@ void UPCGComponent::CleanupLocalImmediate(bool bRemoveComponents)
 	{
 		PostCleanupGraph();
 	}
+
+	PCGGeneratedResourcesLogging::LogCleanupLocalImmediateFinished(GeneratedResources);
 }
 
 FPCGTaskId UPCGComponent::CreateCleanupTask(bool bRemoveComponents, const TArray<FPCGTaskId>& Dependencies)
@@ -697,6 +716,8 @@ FPCGTaskId UPCGComponent::CreateCleanupTask(bool bRemoveComponents, const TArray
 		return InvalidPCGTaskId;
 	}
 
+	PCGGeneratedResourcesLogging::LogCreateCleanupTask(bRemoveComponents);
+
 	// Keep track of all the dependencies
 	TArray<FPCGTaskId> AdditionalDependencies;
 	const TArray<FPCGTaskId>* AllDependencies = &Dependencies;
@@ -750,6 +771,8 @@ FPCGTaskId UPCGComponent::CreateCleanupTask(bool bRemoveComponents, const TArray
 					UE_LOG(LogPCG, Error, TEXT("[UPCGComponent::CreateCleanupTask] Null generated resource encountered on actor \"%s\"."), *ThisComponent->GetOwner()->GetFName().ToString());
 				}
 
+				PCGGeneratedResourcesLogging::LogCreateCleanupTaskResource(Resource);
+
 				if (!Resource || Resource->Release(bRemoveComponents, Context->ActorsToDelete))
 				{
 					ThisComponent->GeneratedResources.RemoveAtSwap(Context->ResourceIndex);
@@ -771,6 +794,11 @@ FPCGTaskId UPCGComponent::CreateCleanupTask(bool bRemoveComponents, const TArray
 			UPCGActorHelpers::DeleteActors(World, Context->ActorsToDelete.Array());
 		}
 
+		if (ThisComponentWeakPtr.Get())
+		{
+			PCGGeneratedResourcesLogging::LogCreateCleanupTaskFinished(ThisComponentWeakPtr->GeneratedResources);
+		}
+
 		return true;
 	};
 
@@ -779,6 +807,8 @@ FPCGTaskId UPCGComponent::CreateCleanupTask(bool bRemoveComponents, const TArray
 
 void UPCGComponent::CleanupUnusedManagedResources()
 {
+	PCGGeneratedResourcesLogging::LogCleanupUnusedManagedResources(GeneratedResources);
+
 	TSet<TSoftObjectPtr<AActor>> ActorsToDelete;
 
 	{
@@ -788,6 +818,8 @@ void UPCGComponent::CleanupUnusedManagedResources()
 		{
 			UPCGManagedResource* Resource = GeneratedResources[ResourceIndex];
 
+			PCGGeneratedResourcesLogging::LogCleanupUnusedManagedResourcesResource(Resource);
+
 			if (!Resource && GetOwner())
 			{
 				UE_LOG(LogPCG, Error, TEXT("[UPCGComponent::CleanupUnusedManagedResources] Null generated resource encountered on actor \"%s\"."), *GetOwner()->GetFName().ToString());
@@ -801,6 +833,8 @@ void UPCGComponent::CleanupUnusedManagedResources()
 	}
 
 	UPCGActorHelpers::DeleteActors(GetWorld(), ActorsToDelete.Array());
+
+	PCGGeneratedResourcesLogging::LogCleanupUnusedManagedResourcesFinished(GeneratedResources);
 }
 
 void UPCGComponent::BeginPlay()
diff --git a/Engine/Plugins/Experimental/PCG/Source/PCG/Private/PCGManagedResource.cpp b/Engine/Plugins/Experimental/PCG/Source/PCG/Private/PCGManagedResource.cpp
--- a/Engine/Plugins/Experimental/PCG/Source/PCG/Private/PCGManagedResource.cpp
+++ b/Engine/Plugins/Experimental/PCG/Source/PCG/Private/PCGManagedResource.cpp
@@ -1,12 +1,19 @@
 // Copyright Epic Games, Inc. All Rights Reserved.
 
 #include "PCGManagedResource.h"
-#include "PCGComponent.h"
 
+#include "PCGComponent.h"
+#include "PCGModule.h"
 #include "Components/InstancedStaticMeshComponent.h"
+#include "Utils/PCGGeneratedResourcesLogging.h"
 
 #include UE_INLINE_GENERATED_CPP_BY_NAME(PCGManagedResource)
 
+static TAutoConsoleVariable<bool> CVarForceReleaseResourcesOnGenerate(
+	TEXT("pcg.ForceReleaseResourcesOnGenerate"),
+	false,
+	TEXT("Purges all tracked generated resources on generate"));
+
 void UPCGManagedResource::PostApplyToComponent()
 {
 	PostEditImport();
@@ -36,6 +43,11 @@ bool UPCGManagedResource::ReleaseIfUnused(TSet<TSoftObjectPtr<AActor>>& OutActor
 	return false;
 }
 
+bool UPCGManagedResource::DebugForcePurgeAllResourcesOnGenerate()
+{
+	return CVarForceReleaseResourcesOnGenerate.GetValueOnAnyThread();
+}
+
 void UPCGManagedActors::PostEditImport()
 {
 	// In this case, the managed actors won't be copied along the actor/component,
@@ -55,11 +67,15 @@ bool UPCGManagedActors::Release(bool bHardRelease, TSet<TSoftObjectPtr<AActor>>&
 
 	if (!Super::Release(bHardRelease, OutActorsToDelete))
 	{
+		PCGGeneratedResourcesLogging::LogManagedActorsSoftRelease(GeneratedActors);
+
 		return false;
 	}
 
 	OutActorsToDelete.Append(GeneratedActors);
 
+	PCGGeneratedResourcesLogging::LogManagedActorsHardRelease(GeneratedActors);
+
 	// Cleanup recursively
 	TInlineComponentArray<UPCGComponent*, 1> ComponentsToCleanup;
 
diff --git a/Engine/Plugins/Experimental/PCG/Source/PCG/Private/Utils/PCGGeneratedResourcesLogging.cpp b/Engine/Plugins/Experimental/PCG/Source/PCG/Private/Utils/PCGGeneratedResourcesLogging.cpp
new file mode 100644
--- /dev/null
+++ b/Engine/Plugins/Experimental/PCG/Source/PCG/Private/Utils/PCGGeneratedResourcesLogging.cpp
@@ -0,0 +1,401 @@
+#pragma once
+
+#include "PCGGeneratedResourcesLogging.h"
+
+#include "PCGComponent.h"
+#include "PCGManagedResource.h"
+#include "PCGModule.h"
+
+#include "GameFramework/Actor.h"
+#include "HAL/IConsoleManager.h"
+
+#if WITH_EDITOR
+static TAutoConsoleVariable<bool> CVarGenResourcesLoggingEnable(
+	TEXT("pcg.GeneratedResourcesLogging.Enable"),
+	false,
+	TEXT("Enables fine grained log of generated resources management"));
+
+static TAutoConsoleVariable<int32> CVarGenResourcesLoggingMaxPrintCount(
+	TEXT("pcg.GeneratedResourcesLogging.MaxElementPrintCount"),
+	3,
+	TEXT("Enables fine grained log of generated resources management"));
+#endif
+
+namespace PCGGeneratedResourcesLogging
+{
+	bool LogEnabled()
+	{
+#if WITH_EDITOR
+		return CVarGenResourcesLoggingEnable.GetValueOnAnyThread();
+#else
+		return false;
+#endif
+	}
+
+	void LogAddToManagedResources(const UPCGManagedResource* Resource)
+	{
+#if WITH_EDITOR
+		if (!LogEnabled())
+		{
+			return;
+		}
+
+		if (const UPCGManagedActors* ManagedActors = Cast<UPCGManagedActors>(Resource))
+		{
+			UE_LOG(LogPCG, Log, TEXT("[PCGMANAGEDRESOURCES] UPCGComponent::AddToManagedResources, actor count %d"), ManagedActors->GeneratedActors.Num());
+
+			const int32 MaxPrint = CVarGenResourcesLoggingMaxPrintCount.GetValueOnAnyThread();
+			int32 Count = 0;
+			for (const TSoftObjectPtr<AActor>& Actor : ManagedActors->GeneratedActors)
+			{
+				if (Count++ >= MaxPrint)
+				{
+					break;
+				}
+
+				if (Actor.Get())
+				{
+					UE_LOG(LogPCG, Log, TEXT("[PCGMANAGEDRESOURCES]         |- '%s' (%d/%d)"), *Actor->GetFName().ToString(), Count, ManagedActors->GeneratedActors.Num());
+				}
+				else
+				{
+					UE_LOG(LogPCG, Log, TEXT("[PCGMANAGEDRESOURCES]         |- NULL (%d/%d)"), Count, ManagedActors->GeneratedActors.Num());
+				}
+			}
+		}
+		else if (const UPCGManagedComponent* ManagedComponent = Cast<UPCGManagedComponent>(Resource))
+		{
+			if (ManagedComponent->GeneratedComponent)
+			{
+				UE_LOG(LogPCG, Log, TEXT("[PCGMANAGEDRESOURCES] UPCGComponent::AddToManagedResources, managed component on actor '%s'"), *ManagedComponent->GeneratedComponent->GetReadableName());
+			}
+			else
+			{
+				UE_LOG(LogPCG, Warning, TEXT("[PCGMANAGEDRESOURCES] UPCGComponent::AddToManagedResources, managed component null or no owner"));
+			}
+		}
+		else
+		{
+			if (Resource)
+			{
+				UE_LOG(LogPCG, Warning, TEXT("[PCGMANAGEDRESOURCES] UPCGComponent::AddToManagedResources, unidentified managed resource index"));
+			}
+			else
+			{
+				UE_LOG(LogPCG, Warning, TEXT("[PCGMANAGEDRESOURCES] UPCGComponent::AddToManagedResources, encountered NULL managed resource index"));
+			}
+		}
+#endif
+	}
+
+	void LogCleanupInternal(bool bRemoveComponents)
+	{
+#if WITH_EDITOR
+		if (!LogEnabled())
+		{
+			return;
+		}
+
+		UE_LOG(LogPCG, Log, TEXT("[PCGMANAGEDRESOURCES]")); // Blank line
+		UE_LOG(LogPCG, Log, TEXT("[PCGMANAGEDRESOURCES] UPCGComponent::CleanupInternal, bRemoveComponents: %d"), bRemoveComponents);
+#endif
+	}
+
+	void LogCleanupLocalImmediate(bool bHardRelease, const TArray<UPCGManagedResource*>& GeneratedResources)
+	{
+#if WITH_EDITOR
+		if (!LogEnabled())
+		{
+			return;
+		}
+
+		UE_LOG(LogPCG, Log, TEXT("[PCGMANAGEDRESOURCES]")); // Blank line
+		UE_LOG(LogPCG, Log, TEXT("[PCGMANAGEDRESOURCES] UPCGComponent::CleanupLocalImmediate BEGIN, bHardRelease: %d, GeneratedResources.Num() = %d"), bHardRelease, GeneratedResources.Num());
+#endif
+	}
+
+	void LogCleanupLocalImmediateResource(const UPCGManagedResource* Resource)
+	{
+#if WITH_EDITOR
+		if (!LogEnabled())
+		{
+			return;
+		}
+
+		if (const UPCGManagedActors* ManagedActors = Cast<UPCGManagedActors>(Resource))
+		{
+			UE_LOG(LogPCG, Log, TEXT("[PCGMANAGEDRESOURCES] UPCGComponent::CleanupLocalImmediate, actor count %d"), ManagedActors->GeneratedActors.Num());
+
+			const int32 MaxPrint = CVarGenResourcesLoggingMaxPrintCount.GetValueOnAnyThread();
+			int32 Count = 0;
+			for (const TSoftObjectPtr<AActor>& Actor : ManagedActors->GeneratedActors)
+			{
+				if (Count++ >= MaxPrint)
+				{
+					break;
+				}
+
+				if (Actor.Get())
+				{
+					UE_LOG(LogPCG, Log, TEXT("[PCGMANAGEDRESOURCES]         |- '%s' (%d/%d)"), *Actor->GetFName().ToString(), Count, ManagedActors->GeneratedActors.Num());
+				}
+				else
+				{
+					UE_LOG(LogPCG, Log, TEXT("[PCGMANAGEDRESOURCES]         |- NULL (%d/%d)"), Count, ManagedActors->GeneratedActors.Num());
+				}
+			}
+		}
+		else if (const UPCGManagedComponent* ManagedComponent = Cast<UPCGManagedComponent>(Resource))
+		{
+			if (ManagedComponent->GeneratedComponent)
+			{
+				UE_LOG(LogPCG, Log, TEXT("[PCGMANAGEDRESOURCES] UPCGComponent::CleanupLocalImmediate, managed component on actor '%s'"), *ManagedComponent->GeneratedComponent->GetReadableName());
+			}
+			else
+			{
+				UE_LOG(LogPCG, Warning, TEXT("[PCGMANAGEDRESOURCES] UPCGComponent::CleanupLocalImmediate, managed component null or no owner"));
+			}
+		}
+		else
+		{
+			if (Resource)
+			{
+				UE_LOG(LogPCG, Warning, TEXT("[PCGMANAGEDRESOURCES] UPCGComponent::CleanupLocalImmediate, unidentified managed resource"));
+			}
+			else
+			{
+				UE_LOG(LogPCG, Warning, TEXT("[PCGMANAGEDRESOURCES] UPCGComponent::CleanupLocalImmediate, encountered NULL managed resource"));
+			}
+		}
+#endif
+	}
+
+	void LogCleanupLocalImmediateFinished(const TArray<UPCGManagedResource*>& GeneratedResources)
+	{
+#if WITH_EDITOR
+		if (!LogEnabled())
+		{
+			return;
+		}
+
+		UE_LOG(LogPCG, Log, TEXT("[PCGMANAGEDRESOURCES] UPCGComponent::CleanupLocalImmediate FINISHED, GeneratedResources.Num() = %d"), GeneratedResources.Num());
+#endif
+	}
+
+	void LogPostProcessGraph()
+	{
+#if WITH_EDITOR
+		if (!LogEnabled())
+		{
+			return;
+		}
+
+		UE_LOG(LogPCG, Log, TEXT("[PCGMANAGEDRESOURCES] UPCGComponent::PostProcessGraph"));
+#endif
+	}
+
+	void LogCreateCleanupTask(bool bRemoveComponents)
+	{
+#if WITH_EDITOR
+		if (!LogEnabled())
+		{
+			return;
+		}
+
+		UE_LOG(LogPCG, Log, TEXT("[PCGMANAGEDRESOURCES] UPCGComponent::CreateCleanupTask, bRemoveComponents: %d"), bRemoveComponents);
+#endif
+	}
+
+	void LogCreateCleanupTaskResource(const UPCGManagedResource* Resource/*, TArray<UPCGManagedResource*>& GeneratedResources*/)
+	{
+#if WITH_EDITOR
+		if (Resource)
+		{
+			if (const UPCGManagedActors* ManagedActors = Cast<UPCGManagedActors>(Resource))
+			{
+				const int32 MaxPrint = CVarGenResourcesLoggingMaxPrintCount.GetValueOnAnyThread();
+				int32 Count = 0;
+				for (const TSoftObjectPtr<AActor>& Actor : ManagedActors->GeneratedActors)
+				{
+					if (Count++ >= MaxPrint)
+					{
+						break;
+					}
+
+					if (Actor.Get())
+					{
+						UE_LOG(LogPCG, Log, TEXT("[PCGMANAGEDRESOURCES]         |- '%s' (%d/%d)"), *Actor->GetFName().ToString(), Count, ManagedActors->GeneratedActors.Num());
+					}
+					else
+					{
+						UE_LOG(LogPCG, Log, TEXT("[PCGMANAGEDRESOURCES]         |- NULL (%d/%d)"), Count, ManagedActors->GeneratedActors.Num());
+					}
+				}
+			}
+			else if (const UPCGManagedComponent* ManagedComponent = Cast<UPCGManagedComponent>(Resource))
+			{
+				if (ManagedComponent->GeneratedComponent)
+				{
+					UE_LOG(LogPCG, Log, TEXT("[PCGMANAGEDRESOURCES] UPCGComponent::CreateCleanupTask::CleanupTask, managed component on actor '%s'"), *ManagedComponent->GeneratedComponent->GetReadableName());
+				}
+				else
+				{
+					UE_LOG(LogPCG, Warning, TEXT("[PCGMANAGEDRESOURCES] UPCGComponent::CreateCleanupTask::CleanupTask, managed component null or no owner"));
+				}
+			}
+			else
+			{
+				UE_LOG(LogPCG, Log, TEXT("[PCGMANAGEDRESOURCES] UPCGComponent::CreateCleanupTask::CleanupTask, unidentified resource or null"));
+			}
+		}
+#endif
+	}
+
+	void LogCreateCleanupTaskFinished(const TArray<UPCGManagedResource*>& GeneratedResources)
+	{
+#if WITH_EDITOR
+		if (!LogEnabled())
+		{
+			return;
+		}
+
+		UE_LOG(LogPCG, Log, TEXT("[PCGMANAGEDRESOURCES] UPCGComponent::CreateCleanupTask::CleanupTask FINISHED, GeneratedResources.Num() = %d"), GeneratedResources.Num());
+#endif
+	}
+
+	void LogCleanupUnusedManagedResources(const TArray<UPCGManagedResource*>& GeneratedResources)
+	{
+#if WITH_EDITOR
+		if (!LogEnabled())
+		{
+			return;
+		}
+
+		UE_LOG(LogPCG, Log, TEXT("[PCGMANAGEDRESOURCES] UPCGComponent::CleanupUnusedManagedResources BEGIN, GeneratedResources.Num() = %d"), GeneratedResources.Num());
+#endif
+	}
+
+	void LogCleanupUnusedManagedResourcesResource(const UPCGManagedResource* Resource)
+	{
+#if WITH_EDITOR
+		if (!LogEnabled())
+		{
+			return;
+		}
+
+		if (Resource)
+		{
+			if (const UPCGManagedActors* ManagedActors = Cast<UPCGManagedActors>(Resource))
+			{
+				const int32 MaxPrint = CVarGenResourcesLoggingMaxPrintCount.GetValueOnAnyThread();
+				int32 Count = 0;
+				for (const TSoftObjectPtr<AActor>& Actor : ManagedActors->GeneratedActors)
+				{
+					if (Count++ >= MaxPrint)
+					{
+						break;
+					}
+
+					if (Actor.Get())
+					{
+						UE_LOG(LogPCG, Log, TEXT("[PCGMANAGEDRESOURCES]         |- '%s' (%d/%d)"), *Actor->GetFName().ToString(), Count, ManagedActors->GeneratedActors.Num());
+					}
+					else
+					{
+						UE_LOG(LogPCG, Log, TEXT("[PCGMANAGEDRESOURCES]         |- NULL (%d/%d)"), Count, ManagedActors->GeneratedActors.Num());
+					}
+				}
+			}
+			else if (const UPCGManagedComponent* ManagedComponent = Cast<UPCGManagedComponent>(Resource))
+			{
+				if (ManagedComponent->GeneratedComponent)
+				{
+					UE_LOG(LogPCG, Log, TEXT("[PCGMANAGEDRESOURCES] UPCGComponent::CleanupUnusedManagedResources, managed component on actor '%s'"), *ManagedComponent->GeneratedComponent->GetReadableName());
+				}
+				else
+				{
+					UE_LOG(LogPCG, Warning, TEXT("[PCGMANAGEDRESOURCES] UPCGComponent::CleanupUnusedManagedResources, managed component null or no owner"));
+				}
+			}
+			else
+			{
+				UE_LOG(LogPCG, Log, TEXT("[PCGMANAGEDRESOURCES] UPCGComponent::CleanupUnusedManagedResources, unidentified resource or null"));
+			}
+		}
+#endif
+	}
+
+	void LogCleanupUnusedManagedResourcesFinished(const TArray<UPCGManagedResource*>& GeneratedResources)
+	{
+#if WITH_EDITOR
+		if (!LogEnabled())
+		{
+			return;
+		}
+
+		UE_LOG(LogPCG, Log, TEXT("[PCGMANAGEDRESOURCES] UPCGComponent::CleanupUnusedManagedResources, FINISHED, GeneratedResources.Num() = %d"), GeneratedResources.Num());
+#endif
+	}
+
+	void LogManagedActorsSoftRelease(const TSet<TSoftObjectPtr<AActor>>& GeneratedActors)
+	{
+#if WITH_EDITOR
+		if (!LogEnabled())
+		{
+			return;
+		}
+
+		UE_LOG(LogPCG, Log, TEXT("[PCGMANAGEDRESOURCES] UPCGManagedActors::Release, soft release, actor count %d"), GeneratedActors.Num());
+
+		const int32 MaxPrint = CVarGenResourcesLoggingMaxPrintCount.GetValueOnAnyThread();
+		int32 Count = 0;
+		for (const TSoftObjectPtr<AActor>& Actor : GeneratedActors)
+		{
+			if (Count++ >= MaxPrint)
+			{
+				break;
+			}
+
+			if (!Actor.Get())
+			{
+				UE_LOG(LogPCG, Log, TEXT("[PCGMANAGEDRESOURCES]         |- NULL (%d/%d)"), Count, GeneratedActors.Num());
+			}
+			else
+			{
+				UE_LOG(LogPCG, Log, TEXT("[PCGMANAGEDRESOURCES]         |- '%s' (%d/%d)"), *Actor->GetFName().ToString(), Count, GeneratedActors.Num());
+			}
+		}
+#endif
+	}
+
+	void LogManagedActorsHardRelease(const TSet<TSoftObjectPtr<AActor>>& ActorsToDelete)
+	{
+#if WITH_EDITOR
+		if (!LogEnabled())
+		{
+			return;
+		}
+
+		UE_LOG(LogPCG, Log, TEXT("[PCGMANAGEDRESOURCES] UPCGManagedActors::Release, hard release, actor count %d scheduled for delete"), ActorsToDelete.Num());
+
+		const uint32 MaxPrint = 3;
+		uint32 Count = 0;
+		for (const TSoftObjectPtr<AActor>& ActorToDelete : ActorsToDelete)
+		{
+			if (Count++ >= MaxPrint)
+			{
+				break;
+			}
+
+			if (ActorToDelete)
+			{
+				UE_LOG(LogPCG, Log, TEXT("[PCGMANAGEDRESOURCES]         |- '%s' (%d/%d)"), *ActorToDelete->GetFName().ToString(), Count, ActorsToDelete.Num());
+			}
+			else
+			{
+				UE_LOG(LogPCG, Log, TEXT("[PCGMANAGEDRESOURCES]         |- NULL (%d/%d)"), Count, ActorsToDelete.Num());
+			}
+		}
+#endif
+	}
+}
diff --git a/Engine/Plugins/Experimental/PCG/Source/PCG/Private/Utils/PCGGeneratedResourcesLogging.h b/Engine/Plugins/Experimental/PCG/Source/PCG/Private/Utils/PCGGeneratedResourcesLogging.h
new file mode 100644
--- /dev/null
+++ b/Engine/Plugins/Experimental/PCG/Source/PCG/Private/Utils/PCGGeneratedResourcesLogging.h
@@ -0,0 +1,33 @@
+#pragma once
+
+#include "Containers/Array.h"
+#include "Containers/Set.h"
+#include "UObject/SoftObjectPtr.h"
+
+class UPCGComponent;
+class UPCGManagedResource;
+class AActor;
+
+namespace PCGGeneratedResourcesLogging
+{
+	void LogAddToManagedResources(const UPCGManagedResource* Resource);
+	
+	void LogCleanupInternal(bool bRemoveComponents);
+	
+	void LogCleanupLocalImmediate(bool bHardRelease, const TArray<UPCGManagedResource*>& GeneratedResources);
+	void LogCleanupLocalImmediateResource(const UPCGManagedResource* Resource);
+	void LogCleanupLocalImmediateFinished(const TArray<UPCGManagedResource*>& GeneratedResources);
+	
+	void LogPostProcessGraph();
+
+	void LogCreateCleanupTask(bool bRemoveComponents);
+	void LogCreateCleanupTaskResource(const UPCGManagedResource* Resource);
+	void LogCreateCleanupTaskFinished(const TArray<UPCGManagedResource*>& GeneratedResources);
+
+	void LogCleanupUnusedManagedResources(const TArray<UPCGManagedResource*>& GeneratedResources);
+	void LogCleanupUnusedManagedResourcesResource(const UPCGManagedResource* Resource);
+	void LogCleanupUnusedManagedResourcesFinished(const TArray<UPCGManagedResource*>& GeneratedResources);
+
+	void LogManagedActorsSoftRelease(const TSet<TSoftObjectPtr<AActor>>& GeneratedActors);
+	void LogManagedActorsHardRelease(const TSet<TSoftObjectPtr<AActor>>& ActorsToDelete);
+}
diff --git a/Engine/Plugins/Experimental/PCG/Source/PCG/Public/PCGManagedResource.h b/Engine/Plugins/Experimental/PCG/Source/PCG/Public/PCGManagedResource.h
--- a/Engine/Plugins/Experimental/PCG/Source/PCG/Public/PCGManagedResource.h
+++ b/Engine/Plugins/Experimental/PCG/Source/PCG/Public/PCGManagedResource.h
@@ -34,6 +34,8 @@ public:
 	/** Move the given resource to a new actor. Return true if it has succeeded */
 	virtual bool MoveResourceToNewActor(AActor* NewActor) { return false; };
 
+	static bool DebugForcePurgeAllResourcesOnGenerate();
+
 protected:
 	UPROPERTY(Transient, VisibleAnywhere, Category = GeneratedData)
 	bool bIsMarkedUnused = false;
