[PCG] Add copyright header to fix validation

diff --git a/Engine/Plugins/Experimental/PCG/Source/PCG/Private/Utils/PCGGeneratedResourcesLogging.cpp b/Engine/Plugins/Experimental/PCG/Source/PCG/Private/Utils/PCGGeneratedResourcesLogging.cpp
--- a/Engine/Plugins/Experimental/PCG/Source/PCG/Private/Utils/PCGGeneratedResourcesLogging.cpp
+++ b/Engine/Plugins/Experimental/PCG/Source/PCG/Private/Utils/PCGGeneratedResourcesLogging.cpp
@@ -1,4 +1,4 @@
-#pragma once
+// Copyright Epic Games, Inc. All Rights Reserved.
 
 #include "PCGGeneratedResourcesLogging.h"
 
diff --git a/Engine/Plugins/Experimental/PCG/Source/PCG/Private/Utils/PCGGeneratedResourcesLogging.h b/Engine/Plugins/Experimental/PCG/Source/PCG/Private/Utils/PCGGeneratedResourcesLogging.h
--- a/Engine/Plugins/Experimental/PCG/Source/PCG/Private/Utils/PCGGeneratedResourcesLogging.h
+++ b/Engine/Plugins/Experimental/PCG/Source/PCG/Private/Utils/PCGGeneratedResourcesLogging.h
@@ -1,3 +1,5 @@
+// Copyright Epic Games, Inc. All Rights Reserved.
+
 #pragma once
 
 #include "Containers/Array.h"
