Make sure to cancel generation on component destruction, esp. for rerun construction script situations

diff --git a/Engine/Plugins/Experimental/PCG/Source/PCG/Private/PCGComponent.cpp b/Engine/Plugins/Experimental/PCG/Source/PCG/Private/PCGComponent.cpp
--- a/Engine/Plugins/Experimental/PCG/Source/PCG/Private/PCGComponent.cpp
+++ b/Engine/Plugins/Experimental/PCG/Source/PCG/Private/PCGComponent.cpp
@@ -872,9 +872,10 @@ void UPCGComponent::EndPlay(const EEndPlayReason::Type EndPlayReason)
 {
 	// Always try to unregister itself, if it doesn't exist, it will early out. 
 	// Just making sure that we don't left some resources registered while dead.
-	if (GetSubsystem())
+	if (UPCGSubsystem* Subsystem = GetSubsystem())
 	{
-		GetSubsystem()->UnregisterPCGComponent(this);
+		Subsystem->CancelGeneration(this);
+		Subsystem->UnregisterPCGComponent(this);
 	}
 
 	Super::EndPlay(EndPlayReason);
@@ -898,6 +899,7 @@ void UPCGComponent::OnComponentDestroyed(bool bDestroyingHierarchy)
 	{
 		if (!PCGHelpers::IsRuntimeOrPIE())
 		{
+			Subsystem->CancelGeneration(this);
 			Subsystem->UnregisterPCGComponent(this);
 		}
 	}
diff --git a/Engine/Plugins/Experimental/PCG/Source/PCG/Private/PCGSubsystem.cpp b/Engine/Plugins/Experimental/PCG/Source/PCG/Private/PCGSubsystem.cpp
--- a/Engine/Plugins/Experimental/PCG/Source/PCG/Private/PCGSubsystem.cpp
+++ b/Engine/Plugins/Experimental/PCG/Source/PCG/Private/PCGSubsystem.cpp
@@ -406,9 +406,8 @@ FPCGTaskId UPCGSubsystem::ScheduleGenericWithContext(TFunction<bool(FPCGContext*
 
 void UPCGSubsystem::CancelGeneration(UPCGComponent* Component)
 {
-	check(GraphExecutor);
-
-	if (!Component)
+	check(GraphExecutor && IsInGameThread());
+	if (!Component || !Component->IsGenerating())
 	{
 		return;
 	}
