[PCG] Added PCGNode name to search string to make it easier to find node errors.

diff --git a/Engine/Plugins/Experimental/PCG/Source/PCGEditor/Private/SPCGEditorGraphFind.cpp b/Engine/Plugins/Experimental/PCG/Source/PCGEditor/Private/SPCGEditorGraphFind.cpp
--- a/Engine/Plugins/Experimental/PCG/Source/PCGEditor/Private/SPCGEditorGraphFind.cpp
+++ b/Engine/Plugins/Experimental/PCG/Source/PCGEditor/Private/SPCGEditorGraphFind.cpp
@@ -2,13 +2,15 @@
 
 #include "SPCGEditorGraphFind.h"
 
-#include "Framework/Application/SlateApplication.h"
-#include "Framework/Views/TableViewMetadata.h"
-#include "Layout/WidgetPath.h"
 #include "PCGEditor.h"
 #include "PCGEditorGraph.h"
+#include "PCGEditorGraphNodeBase.h"
+#include "PCGNode.h"
 
 #include "EdGraph/EdGraphSchema.h"
+#include "Framework/Application/SlateApplication.h"
+#include "Framework/Views/TableViewMetadata.h"
+#include "Layout/WidgetPath.h"
 #include "Widgets/Input/SSearchBox.h"
 #include "Widgets/Views/STreeView.h"
 
@@ -276,6 +278,14 @@ void SPCGEditorGraphFind::MatchTokens(const TArray<FString>& InTokens)
 		const FString NodeType = Node->GetNodeTitle(ENodeTitleType::ListView).ToString();
 
 		FString NodeSearchString = NodeName + NodeType + Node->NodeComment;
+
+		if (const UPCGEditorGraphNodeBase* PCGEditorGraphNodeBase = Cast<UPCGEditorGraphNodeBase>(Node))
+        {
+        	if (const UPCGNode* PCGNode = PCGEditorGraphNodeBase->GetPCGNode())
+        	{
+        		NodeSearchString.Append(PCGNode->GetName());
+        	}
+        }
 		NodeSearchString = NodeSearchString.Replace(TEXT(" "), TEXT(""));
 
 		FPCGEditorGraphFindResultPtr NodeResult;
@@ -283,7 +293,7 @@ void SPCGEditorGraphFind::MatchTokens(const TArray<FString>& InTokens)
 		{
 			if (!NodeResult.IsValid())
 			{
-				NodeResult = MakeShared<FPCGEditorGraphFindResult>(NodeName == NodeType ? NodeName : NodeName + " - " + NodeType, RootFindResult, Node);
+				NodeResult = MakeShared<FPCGEditorGraphFindResult>((NodeName == NodeType) ? NodeName : NodeName + " - " + NodeType, RootFindResult, Node);
 			}
 			return NodeResult;
 		};
