[PCG] Improve Data From Actor UI and functionality

Requirements:

We need the Self to work
- There's generally no need to filter anything in Self except if Include Children is tagged, in which case, we should be able to filter (see point below)

We should have a No Filter option (or a checkbox to control these)

We need Original
- This is to have something akin to the original actor vs actor pins

diff --git a/Engine/Plugins/Experimental/PCG/Source/PCG/Private/Elements/PCGActorSelector.cpp b/Engine/Plugins/Experimental/PCG/Source/PCG/Private/Elements/PCGActorSelector.cpp
--- a/Engine/Plugins/Experimental/PCG/Source/PCG/Private/Elements/PCGActorSelector.cpp
+++ b/Engine/Plugins/Experimental/PCG/Source/PCG/Private/Elements/PCGActorSelector.cpp
@@ -1,9 +1,12 @@
 // Copyright Epic Games, Inc. All Rights Reserved.
 
 #include "Elements/PCGActorSelector.h"
-#include "GameFramework/Actor.h"
+
+#include "PCGComponent.h"
+#include "Grid/PCGPartitionActor.h"
 #include "Helpers/PCGActorHelpers.h"
 
+#include "GameFramework/Actor.h"
 
 #include UE_INLINE_GENERATED_CPP_BY_NAME(PCGActorSelector)
 
@@ -12,6 +15,17 @@ namespace PCGActorSelector
 	// Need to pass a pointer of pointer to the found actor. The lambda will capture this pointer and modify its value when an actor is found.
 	TFunction<bool(AActor*)> GetFilteringFunction(const FPCGActorSelectorSettings& InSettings, TArray<AActor*>& InFoundActors)
 	{
+		// If we're not filtering all world actors, and if the Include Children is enabled, then we allow
+		// user to disable the filter which ignores all filtering options, which is convenient for workflow simplification.
+		if (InSettings.ActorFilter != EPCGActorFilter::AllWorldActors && InSettings.bIncludeChildren && InSettings.bDisableFilter)
+		{
+			return [&InFoundActors](AActor* Actor) -> bool
+			{
+				InFoundActors.Add(Actor);
+				return true;
+			};
+		}
+
 		const bool bMultiSelect = InSettings.bSelectMultiple;
 
 		switch (InSettings.ActorSelection)
@@ -59,10 +73,13 @@ namespace PCGActorSelector
 		return [](AActor* Actor) -> bool { return false; };
 	}
 
-	TArray<AActor*> FindActors(const FPCGActorSelectorSettings& Settings, UWorld* World, AActor* Self)
+	TArray<AActor*> FindActors(const FPCGActorSelectorSettings& Settings, const UPCGComponent* InComponent)
 	{
 		TRACE_CPUPROFILER_EVENT_SCOPE(PCGActorSelector::FindActor);
 
+		UWorld* World = InComponent ? InComponent->GetWorld() : nullptr;
+		AActor* Self = InComponent ? InComponent->GetOwner() : nullptr;
+
 		TArray<AActor*> FoundActors;
 
 		if (!World)
@@ -134,9 +151,20 @@ namespace PCGActorSelector
 			break;
 		}
 
-		//case EPCGActorFilter::TrackedActors:
-			//	//TODO
-			//	break;
+		case EPCGActorFilter::Original:
+		{
+			APCGPartitionActor* PartitionActor = Cast<APCGPartitionActor>(Self);
+			UPCGComponent* OriginalComponent = (PartitionActor && InComponent) ? PartitionActor->GetOriginalComponent(InComponent) : nullptr;
+			AActor* OriginalActor = OriginalComponent ? OriginalComponent->GetOwner() : nullptr;
+			if (OriginalActor)
+			{
+				ActorsToCheck.Add(OriginalActor);
+			}
+			else if (Self)
+			{
+				ActorsToCheck.Add(Self);
+			}
+		}
 
 		default:
 			break;
@@ -163,13 +191,13 @@ namespace PCGActorSelector
 		return FoundActors;
 	}
 
-	AActor* FindActor(const FPCGActorSelectorSettings& InSettings, UWorld* World, AActor* Self)
+	AActor* FindActor(const FPCGActorSelectorSettings& InSettings, UPCGComponent* InComponent)
 	{
 		// In order to make sure we don't try to select multiple, we'll do a copy of the settings here.
 		FPCGActorSelectorSettings Settings = InSettings;
 		Settings.bSelectMultiple = false;
 
-		TArray<AActor*> Actors = FindActors(Settings, World, Self);
+		TArray<AActor*> Actors = FindActors(Settings, InComponent);
 		return Actors.IsEmpty() ? nullptr : Actors[0];
 	}
 }
diff --git a/Engine/Plugins/Experimental/PCG/Source/PCG/Private/Elements/PCGDataFromActor.cpp b/Engine/Plugins/Experimental/PCG/Source/PCG/Private/Elements/PCGDataFromActor.cpp
--- a/Engine/Plugins/Experimental/PCG/Source/PCG/Private/Elements/PCGDataFromActor.cpp
+++ b/Engine/Plugins/Experimental/PCG/Source/PCG/Private/Elements/PCGDataFromActor.cpp
@@ -66,7 +66,7 @@ bool FPCGDataFromActorElement::ExecuteInternal(FPCGContext* InContext) const
 
 	if (!Context->bPerformedQuery)
 	{
-		Context->FoundActors = PCGActorSelector::FindActors(Settings->ActorSelector, World, Context->SourceComponent.IsValid() ? Context->SourceComponent->GetOwner() : nullptr);
+		Context->FoundActors = PCGActorSelector::FindActors(Settings->ActorSelector, Context->SourceComponent.Get());
 		Context->bPerformedQuery = true;
 
 		if (Context->FoundActors.IsEmpty())
diff --git a/Engine/Plugins/Experimental/PCG/Source/PCG/Private/Elements/PCGPropertyToParamData.cpp b/Engine/Plugins/Experimental/PCG/Source/PCG/Private/Elements/PCGPropertyToParamData.cpp
--- a/Engine/Plugins/Experimental/PCG/Source/PCG/Private/Elements/PCGPropertyToParamData.cpp
+++ b/Engine/Plugins/Experimental/PCG/Source/PCG/Private/Elements/PCGPropertyToParamData.cpp
@@ -80,7 +80,7 @@ bool FPCGPropertyToParamDataElement::ExecuteInternal(FPCGContext* Context) const
 
 	// First find the actor depending on the selection
 	UPCGComponent* OriginalComponent = UPCGBlueprintHelpers::GetOriginalComponent(*Context);
-	AActor* FoundActor = PCGActorSelector::FindActor(Settings->ActorSelector, OriginalComponent ? OriginalComponent->GetWorld() : nullptr, OriginalComponent ? OriginalComponent->GetOwner() : nullptr);
+	AActor* FoundActor = PCGActorSelector::FindActor(Settings->ActorSelector, OriginalComponent);
 
 	if (!FoundActor)
 	{
diff --git a/Engine/Plugins/Experimental/PCG/Source/PCG/Public/Elements/PCGActorSelector.h b/Engine/Plugins/Experimental/PCG/Source/PCG/Public/Elements/PCGActorSelector.h
--- a/Engine/Plugins/Experimental/PCG/Source/PCG/Public/Elements/PCGActorSelector.h
+++ b/Engine/Plugins/Experimental/PCG/Source/PCG/Public/Elements/PCGActorSelector.h
@@ -7,6 +7,7 @@
 #include "PCGActorSelector.generated.h"
 
 class AActor;
+class UPCGComponent;
 class UWorld;
 
 UENUM()
@@ -20,12 +21,16 @@ enum class EPCGActorSelection : uint8
 UENUM()
 enum class EPCGActorFilter : uint8
 {
+	/** This actor (either the original PCG actor or the partition actor if partitioning is enabled). */
 	Self,
+	/** The parent of this actor in the hierarchy. */
 	Parent,
+	/** The top most parent of this actor in the hierarchy. */
 	Root,
-	AllWorldActors
-	// TODO
-	// TrackedActors
+	/** All actors in world. */
+	AllWorldActors,
+	/** The source PCG actor (rather than the generated partition actor). */
+	Original,
 };
 
 USTRUCT(BlueprintType)
@@ -33,30 +38,37 @@ struct FPCGActorSelectorSettings
 {
 	GENERATED_BODY()
 
+	/** Which actors to consider. */
 	UPROPERTY(BlueprintReadWrite, EditAnywhere, Category = Settings)
+	EPCGActorFilter ActorFilter = EPCGActorFilter::Self;
+
+	/** Whether to consider child actors. */
+	UPROPERTY(BlueprintReadWrite, EditAnywhere, Category = Settings, meta = (EditCondition = "ActorFilter!=EPCGActorFilter::AllWorldActors", EditConditionHides))
+	bool bIncludeChildren = false;
+
+	/** Enables/disables fine-grained actor filtering options. */
+	UPROPERTY(BlueprintReadWrite, EditAnywhere, Category = Settings, meta = (EditCondition = "ActorFilter!=EPCGActorFilter::AllWorldActors && bIncludeChildren", EditConditionHides))
+	bool bDisableFilter = false;
+
+	UPROPERTY(BlueprintReadWrite, EditAnywhere, Category = Settings, meta = (EditCondition = "(ActorFilter==EPCGActorFilter::AllWorldActors || (bIncludeChildren && !bDisableFilter))", EditConditionHides))
 	EPCGActorSelection ActorSelection = EPCGActorSelection::ByTag;
 
-	UPROPERTY(BlueprintReadWrite, EditAnywhere, Category = Settings, meta = (EditCondition = "ActorSelection==EPCGActorSelection::ByTag", EditConditionHides))
+	UPROPERTY(BlueprintReadWrite, EditAnywhere, Category = Settings, meta = (EditCondition = "(ActorFilter==EPCGActorFilter::AllWorldActors || (bIncludeChildren && !bDisableFilter)) && ActorSelection==EPCGActorSelection::ByTag", EditConditionHides))
 	FName ActorSelectionTag;
 
-	UPROPERTY(BlueprintReadWrite, EditAnywhere, Category = Settings, meta = (EditCondition = "ActorSelection==EPCGActorSelection::ByName", EditConditionHides))
+	UPROPERTY(BlueprintReadWrite, EditAnywhere, Category = Settings, meta = (EditCondition = "(ActorFilter==EPCGActorFilter::AllWorldActors || (bIncludeChildren && !bDisableFilter)) && ActorSelection==EPCGActorSelection::ByName", EditConditionHides))
 	FName ActorSelectionName;
 
-	UPROPERTY(BlueprintReadWrite, EditAnywhere, Category = Settings, meta = (EditCondition = "ActorSelection==EPCGActorSelection::ByClass", EditConditionHides))
+	UPROPERTY(BlueprintReadWrite, EditAnywhere, Category = Settings, meta = (EditCondition = "(ActorFilter==EPCGActorFilter::AllWorldActors || (bIncludeChildren && !bDisableFilter)) && ActorSelection==EPCGActorSelection::ByClass", EditConditionHides))
 	TSubclassOf<AActor> ActorSelectionClass;
 
-	UPROPERTY(BlueprintReadWrite, EditAnywhere, Category = Settings)
-	EPCGActorFilter ActorFilter = EPCGActorFilter::Self;
-
-	UPROPERTY(BlueprintReadWrite, EditAnywhere, Category = Settings, meta = (EditCondition = "ActorFilter!=EPCGActorFilter::AllWorldActors", EditConditionHides))
-	bool bIncludeChildren = false;
-
-	UPROPERTY(BlueprintReadWrite, EditAnywhere, Category = Settings, meta = (EditCondition = "ActorSelection!=EPCGActorSelection::ByName"))
+	/** If true processes all matching actors, otherwise returns data from first match. */
+	UPROPERTY(BlueprintReadWrite, EditAnywhere, Category = Settings, meta = (EditCondition = "ActorFilter==EPCGActorFilter::AllWorldActors && ActorSelection!=EPCGActorSelection::ByName", EditConditionHides))
 	bool bSelectMultiple = false;
 };
 
 namespace PCGActorSelector
 {
-	TArray<AActor*> FindActors(const FPCGActorSelectorSettings& Settings, UWorld* World, AActor* Self);
-	AActor* FindActor(const FPCGActorSelectorSettings& Settings, UWorld* World, AActor* Self);
-}
\ No newline at end of file
+	TArray<AActor*> FindActors(const FPCGActorSelectorSettings& Settings, const UPCGComponent* InComponent);
+	AActor* FindActor(const FPCGActorSelectorSettings& InSettings, UPCGComponent* InComponent);
+}
