[PCG] Fix Data From Actor - 'Self' and similar modes started failing when the filter was not set up.

Only require the filter being set up if the filter will be used.

diff --git a/Engine/Plugins/Experimental/PCG/Source/PCG/Private/Elements/PCGActorSelector.cpp b/Engine/Plugins/Experimental/PCG/Source/PCG/Private/Elements/PCGActorSelector.cpp
--- a/Engine/Plugins/Experimental/PCG/Source/PCG/Private/Elements/PCGActorSelector.cpp
+++ b/Engine/Plugins/Experimental/PCG/Source/PCG/Private/Elements/PCGActorSelector.cpp
@@ -87,10 +87,13 @@ namespace PCGActorSelector
 			return FoundActors;
 		}
 
-		// Early out if we have not the information necessary
-		if ((Settings.ActorSelection == EPCGActorSelection::ByTag && Settings.ActorSelectionTag == NAME_None) ||
+		// Early out if we have not the information necessary. A filter must be set up if we are taking all world actors,
+		// or including children and not disabling the filter.
+		const bool bFilterRequired = (Settings.ActorFilter == EPCGActorFilter::AllWorldActors) || (Settings.bIncludeChildren && !Settings.bDisableFilter);
+		if (bFilterRequired &&
+			((Settings.ActorSelection == EPCGActorSelection::ByTag && Settings.ActorSelectionTag == NAME_None) ||
 			(Settings.ActorSelection == EPCGActorSelection::ByName && Settings.ActorSelectionName == NAME_None) ||
-			(Settings.ActorSelection == EPCGActorSelection::ByClass && !Settings.ActorSelectionClass))
+			(Settings.ActorSelection == EPCGActorSelection::ByClass && !Settings.ActorSelectionClass)))
 		{
 			return FoundActors;
 		}
