Exposed a method to trigger appropriately changes driven from BP at runtime

diff --git a/Engine/Plugins/Experimental/PCG/Source/PCG/Private/PCGComponent.cpp b/Engine/Plugins/Experimental/PCG/Source/PCG/Private/PCGComponent.cpp
--- a/Engine/Plugins/Experimental/PCG/Source/PCG/Private/PCGComponent.cpp
+++ b/Engine/Plugins/Experimental/PCG/Source/PCG/Private/PCGComponent.cpp
@@ -562,6 +562,12 @@ void UPCGComponent::CancelGeneration()
 	}
 }
 
+void UPCGComponent::NotifyPropertiesChangedFromBlueprint()
+{
+	DirtyGenerated(EPCGComponentDirtyFlag::Actor);
+	Refresh();
+}
+
 AActor* UPCGComponent::ClearPCGLink(UClass* TemplateActor)
 {
 	if (!bGenerated || !GetOwner() || !GetWorld())
diff --git a/Engine/Plugins/Experimental/PCG/Source/PCG/Public/PCGComponent.h b/Engine/Plugins/Experimental/PCG/Source/PCG/Public/PCGComponent.h
--- a/Engine/Plugins/Experimental/PCG/Source/PCG/Public/PCGComponent.h
+++ b/Engine/Plugins/Experimental/PCG/Source/PCG/Public/PCGComponent.h
@@ -137,6 +137,10 @@ public:
 	/** Cancels in-progress generation */
 	void CancelGeneration();
 
+	/** Notify properties changed, used in runtime cases, will dirty & trigger a regeneration if needed */
+	UFUNCTION(BlueprintCallable, Category = PCG)
+	void NotifyPropertiesChangedFromBlueprint();
+
 	/** Retrieves generated data */
 	UFUNCTION(BlueprintCallable, Category = PCG)
 	const FPCGDataCollection& GetGeneratedGraphOutput() const { return GeneratedGraphOutput; }
