Fixed compilation - was missing if editor tags

#rnx

diff --git a/Engine/Plugins/Experimental/PCG/Source/PCG/Private/PCGComponent.cpp b/Engine/Plugins/Experimental/PCG/Source/PCG/Private/PCGComponent.cpp
--- a/Engine/Plugins/Experimental/PCG/Source/PCG/Private/PCGComponent.cpp
+++ b/Engine/Plugins/Experimental/PCG/Source/PCG/Private/PCGComponent.cpp
@@ -564,8 +564,10 @@ void UPCGComponent::CancelGeneration()
 
 void UPCGComponent::NotifyPropertiesChangedFromBlueprint()
 {
+#if WITH_EDITOR
 	DirtyGenerated(EPCGComponentDirtyFlag::Actor);
 	Refresh();
+#endif
 }
 
 AActor* UPCGComponent::ClearPCGLink(UClass* TemplateActor)
