Added WPO support - temporary version until we have a better support like ISM component descriptor

diff --git a/Engine/Plugins/Experimental/PCG/Source/PCG/Private/Elements/PCGStaticMeshSpawner.cpp b/Engine/Plugins/Experimental/PCG/Source/PCG/Private/Elements/PCGStaticMeshSpawner.cpp
--- a/Engine/Plugins/Experimental/PCG/Source/PCG/Private/Elements/PCGStaticMeshSpawner.cpp
+++ b/Engine/Plugins/Experimental/PCG/Source/PCG/Private/Elements/PCGStaticMeshSpawner.cpp
@@ -207,6 +207,7 @@ void FPCGStaticMeshSpawnerElement::SpawnStaticMeshInstances(FPCGContext* Context
 	Params.NumCustomDataFloats = PackedCustomData.NumCustomDataFloats;
 	Params.CullStartDistance = InstanceList.CullStartDistance;
 	Params.CullEndDistance = InstanceList.CullEndDistance;
+	Params.WorldPositionOffsetDisableDistance = InstanceList.WorldPositionOffsetDisableDistance;
 	Params.bIsLocalToWorldDeterminantNegative = InstanceList.bIsLocalToWorldDeterminantNegative;
 
 	UInstancedStaticMeshComponent* ISMC = UPCGActorHelpers::GetOrCreateISMC(TargetActor, Context->SourceComponent.Get(), Params);
diff --git a/Engine/Plugins/Experimental/PCG/Source/PCG/Private/Helpers/PCGActorHelpers.cpp b/Engine/Plugins/Experimental/PCG/Source/PCG/Private/Helpers/PCGActorHelpers.cpp
--- a/Engine/Plugins/Experimental/PCG/Source/PCG/Private/Helpers/PCGActorHelpers.cpp
+++ b/Engine/Plugins/Experimental/PCG/Source/PCG/Private/Helpers/PCGActorHelpers.cpp
@@ -124,6 +124,7 @@ UInstancedStaticMeshComponent* UPCGActorHelpers::GetOrCreateISMC(AActor* InTarge
 				(!ISMC->bUseDefaultCollision && ISMC->GetCollisionProfileName() != InParams.CollisionProfile) ||
 				ISMC->InstanceStartCullDistance != InParams.CullStartDistance ||
 				ISMC->InstanceEndCullDistance != InParams.CullEndDistance ||
+				ISMC->WorldPositionOffsetDisableDistance != InParams.WorldPositionOffsetDisableDistance ||
 				ISMC->bReverseCulling != InParams.bIsLocalToWorldDeterminantNegative)
 			{
 				continue;
@@ -203,6 +204,7 @@ UInstancedStaticMeshComponent* UPCGActorHelpers::GetOrCreateISMC(AActor* InTarge
 
 	ISMC->InstanceStartCullDistance = InParams.CullStartDistance;
 	ISMC->InstanceEndCullDistance = InParams.CullEndDistance;
+	ISMC->WorldPositionOffsetDisableDistance = InParams.WorldPositionOffsetDisableDistance;
 	ISMC->bReverseCulling = (InParams.bIsLocalToWorldDeterminantNegative ? 1 : 0);
 
 	// Implementation note: we use the relative for position only to prevent issues at the rendering level, but world for scale & rotation otherwise we run into issues when we have non-uniform scales esp. in volumes
diff --git a/Engine/Plugins/Experimental/PCG/Source/PCG/Private/MeshSelectors/PCGMeshSelectorBase.cpp b/Engine/Plugins/Experimental/PCG/Source/PCG/Private/MeshSelectors/PCGMeshSelectorBase.cpp
--- a/Engine/Plugins/Experimental/PCG/Source/PCG/Private/MeshSelectors/PCGMeshSelectorBase.cpp
+++ b/Engine/Plugins/Experimental/PCG/Source/PCG/Private/MeshSelectors/PCGMeshSelectorBase.cpp
@@ -13,7 +13,8 @@ int32 UPCGMeshSelectorBase::FindOrAddInstanceList(
 	bool bOverrideMaterials,
 	const TArray<TSoftObjectPtr<UMaterialInterface>>& MaterialOverrides,
 	const float InCullStartDistance,
-	const float InCullEndDistance, 
+	const float InCullEndDistance,
+	const int32 InWorldPositionOffsetDisableDistance,
 	const bool bInIsLocalToWorldDeterminantNegative)
 {
 	for (int Index = 0; Index < OutInstanceLists.Num(); ++Index)
@@ -23,7 +24,9 @@ int32 UPCGMeshSelectorBase::FindOrAddInstanceList(
 			continue;
 		}
 
-		if (OutInstanceLists[Index].CullStartDistance != InCullStartDistance || OutInstanceLists[Index].CullEndDistance != InCullEndDistance)
+		if (OutInstanceLists[Index].CullStartDistance != InCullStartDistance || 
+			OutInstanceLists[Index].CullEndDistance != InCullEndDistance ||
+			OutInstanceLists[Index].WorldPositionOffsetDisableDistance != InWorldPositionOffsetDisableDistance)
 		{
 			continue;
 		}
@@ -46,7 +49,7 @@ int32 UPCGMeshSelectorBase::FindOrAddInstanceList(
 		return Index;
 	}
 
-	return OutInstanceLists.Emplace(Mesh, bOverrideCollisionProfile, CollisionProfile, bOverrideMaterials, MaterialOverrides, InCullStartDistance, InCullEndDistance, bInIsLocalToWorldDeterminantNegative);
+	return OutInstanceLists.Emplace(Mesh, bOverrideCollisionProfile, CollisionProfile, bOverrideMaterials, MaterialOverrides, InCullStartDistance, InCullEndDistance, InWorldPositionOffsetDisableDistance, bInIsLocalToWorldDeterminantNegative);
 }
 
 FPCGMeshMaterialOverrideHelper::FPCGMeshMaterialOverrideHelper(
diff --git a/Engine/Plugins/Experimental/PCG/Source/PCG/Private/MeshSelectors/PCGMeshSelectorByAttribute.cpp b/Engine/Plugins/Experimental/PCG/Source/PCG/Private/MeshSelectors/PCGMeshSelectorByAttribute.cpp
--- a/Engine/Plugins/Experimental/PCG/Source/PCG/Private/MeshSelectors/PCGMeshSelectorByAttribute.cpp
+++ b/Engine/Plugins/Experimental/PCG/Source/PCG/Private/MeshSelectors/PCGMeshSelectorByAttribute.cpp
@@ -111,7 +111,7 @@ void UPCGMeshSelectorByAttribute::SelectInstances_Implementation(
 		}
 
 		const bool bReverseTransform = (Point.Transform.GetDeterminant() < 0);
-		const int32 Index = FindOrAddInstanceList(OutMeshInstances, Mesh, bOverrideCollisionProfile, CollisionProfile, MaterialOverrideHelper.OverridesMaterials(), MaterialOverrideHelper.GetMaterialOverrides(Point.MetadataEntry), CullStartDistance, CullEndDistance, bReverseTransform);
+		const int32 Index = FindOrAddInstanceList(OutMeshInstances, Mesh, bOverrideCollisionProfile, CollisionProfile, MaterialOverrideHelper.OverridesMaterials(), MaterialOverrideHelper.GetMaterialOverrides(Point.MetadataEntry), CullStartDistance, CullEndDistance, WorldPositionOffsetDisableDistance, bReverseTransform);
 		check(Index != INDEX_NONE);
 		
 		OutMeshInstances[Index].Instances.Emplace(Point);
diff --git a/Engine/Plugins/Experimental/PCG/Source/PCG/Private/MeshSelectors/PCGMeshSelectorWeighted.cpp b/Engine/Plugins/Experimental/PCG/Source/PCG/Private/MeshSelectors/PCGMeshSelectorWeighted.cpp
--- a/Engine/Plugins/Experimental/PCG/Source/PCG/Private/MeshSelectors/PCGMeshSelectorWeighted.cpp
+++ b/Engine/Plugins/Experimental/PCG/Source/PCG/Private/MeshSelectors/PCGMeshSelectorWeighted.cpp
@@ -57,6 +57,7 @@ namespace PCGMeshSelectorWeighted
 			NewInstanceList.CollisionProfile = SourceInstanceList.CollisionProfile;
 			NewInstanceList.CullStartDistance = SourceInstanceList.CullStartDistance;
 			NewInstanceList.CullEndDistance = SourceInstanceList.CullEndDistance;
+			NewInstanceList.WorldPositionOffsetDisableDistance = SourceInstanceList.WorldPositionOffsetDisableDistance;
 
 			NewInstanceList.bOverrideMaterials = bOverrideMaterials;
 
@@ -91,7 +92,7 @@ void UPCGMeshSelectorWeighted::SelectInstances_Implementation(
 		}
 
 		TArray<FPCGMeshInstanceList>& PickEntry = MeshInstances.Emplace_GetRef();
-		FindOrAddInstanceList(PickEntry, Entry.Mesh, Entry.bOverrideCollisionProfile, Entry.CollisionProfile, Entry.bOverrideMaterials, Entry.MaterialOverrides, Entry.CullStartDistance, Entry.CullEndDistance, /*bReverseCulling=*/false);
+		FindOrAddInstanceList(PickEntry, Entry.Mesh, Entry.bOverrideCollisionProfile, Entry.CollisionProfile, Entry.bOverrideMaterials, Entry.MaterialOverrides, Entry.CullStartDistance, Entry.CullEndDistance, Entry.WorldPositionOffsetDisableDistance, /*bReverseCulling=*/false);
 
 		TotalWeight += Entry.Weight;
 		CumulativeWeights.Add(TotalWeight);
diff --git a/Engine/Plugins/Experimental/PCG/Source/PCG/Private/MeshSelectors/PCGMeshSelectorWeightedByCategory.cpp b/Engine/Plugins/Experimental/PCG/Source/PCG/Private/MeshSelectors/PCGMeshSelectorWeightedByCategory.cpp
--- a/Engine/Plugins/Experimental/PCG/Source/PCG/Private/MeshSelectors/PCGMeshSelectorWeightedByCategory.cpp
+++ b/Engine/Plugins/Experimental/PCG/Source/PCG/Private/MeshSelectors/PCGMeshSelectorWeightedByCategory.cpp
@@ -111,7 +111,7 @@ void UPCGMeshSelectorWeightedByCategory::SelectInstances_Implementation(
 			}
 			
 			TArray<FPCGMeshInstanceList>& PickEntry = InstancesAndWeights->MeshInstances.Emplace_GetRef();
-			FindOrAddInstanceList(PickEntry, WeightedEntry.Mesh, WeightedEntry.bOverrideCollisionProfile, WeightedEntry.CollisionProfile, WeightedEntry.bOverrideMaterials, WeightedEntry.MaterialOverrides, WeightedEntry.CullStartDistance, WeightedEntry.CullEndDistance, /*bReverseCulling=*/false);
+			FindOrAddInstanceList(PickEntry, WeightedEntry.Mesh, WeightedEntry.bOverrideCollisionProfile, WeightedEntry.CollisionProfile, WeightedEntry.bOverrideMaterials, WeightedEntry.MaterialOverrides, WeightedEntry.CullStartDistance, WeightedEntry.CullEndDistance, WeightedEntry.WorldPositionOffsetDisableDistance, /*bReverseCulling=*/false);
 
 			// precompute the weights
 			TotalWeight += WeightedEntry.Weight;
diff --git a/Engine/Plugins/Experimental/PCG/Source/PCG/Public/Helpers/PCGActorHelpers.h b/Engine/Plugins/Experimental/PCG/Source/PCG/Public/Helpers/PCGActorHelpers.h
--- a/Engine/Plugins/Experimental/PCG/Source/PCG/Public/Helpers/PCGActorHelpers.h
+++ b/Engine/Plugins/Experimental/PCG/Source/PCG/Public/Helpers/PCGActorHelpers.h
@@ -28,6 +28,7 @@ struct FPCGISMCBuilderParameters
 	int32 NumCustomDataFloats = 0;
 	float CullStartDistance = 0;
 	float CullEndDistance = 0;
+	int32 WorldPositionOffsetDisableDistance = 0;
 	bool bIsLocalToWorldDeterminantNegative = false;
 };
 
diff --git a/Engine/Plugins/Experimental/PCG/Source/PCG/Public/MeshSelectors/PCGMeshSelectorBase.h b/Engine/Plugins/Experimental/PCG/Source/PCG/Public/MeshSelectors/PCGMeshSelectorBase.h
--- a/Engine/Plugins/Experimental/PCG/Source/PCG/Public/MeshSelectors/PCGMeshSelectorBase.h
+++ b/Engine/Plugins/Experimental/PCG/Source/PCG/Public/MeshSelectors/PCGMeshSelectorBase.h
@@ -24,8 +24,8 @@ struct FPCGMeshInstanceList
 
 	FPCGMeshInstanceList() = default;
 
-	FPCGMeshInstanceList(const TSoftObjectPtr<UStaticMesh>& InMesh, bool bInOverrideCollisionProfile, const FCollisionProfileName& InCollisionProfile, bool bInOverrideMaterials, const TArray<TSoftObjectPtr<UMaterialInterface>>& InMaterialOverrides, const float InCullStartDistance, const float InCullEndDistance, const bool bInIsLocalToWorldDeterminantNegative)
-		: Mesh(InMesh), bOverrideCollisionProfile(bInOverrideCollisionProfile), CollisionProfile(InCollisionProfile), bOverrideMaterials(bInOverrideMaterials), MaterialOverrides(InMaterialOverrides), CullStartDistance(InCullStartDistance), CullEndDistance(InCullEndDistance), bIsLocalToWorldDeterminantNegative(bInIsLocalToWorldDeterminantNegative)
+	FPCGMeshInstanceList(const TSoftObjectPtr<UStaticMesh>& InMesh, bool bInOverrideCollisionProfile, const FCollisionProfileName& InCollisionProfile, bool bInOverrideMaterials, const TArray<TSoftObjectPtr<UMaterialInterface>>& InMaterialOverrides, const float InCullStartDistance, const float InCullEndDistance, const int32 InWorldPositionOffsetDisableDistance, const bool bInIsLocalToWorldDeterminantNegative)
+		: Mesh(InMesh), bOverrideCollisionProfile(bInOverrideCollisionProfile), CollisionProfile(InCollisionProfile), bOverrideMaterials(bInOverrideMaterials), MaterialOverrides(InMaterialOverrides), CullStartDistance(InCullStartDistance), CullEndDistance(InCullEndDistance), WorldPositionOffsetDisableDistance(InWorldPositionOffsetDisableDistance), bIsLocalToWorldDeterminantNegative(bInIsLocalToWorldDeterminantNegative)
 	{}
 
 	UPROPERTY(BlueprintReadWrite, EditAnywhere, Category = Settings)
@@ -54,6 +54,9 @@ struct FPCGMeshInstanceList
 	UPROPERTY(BlueprintReadWrite, EditAnywhere, Category = Settings)
 	float CullEndDistance = 0;
 
+	UPROPERTY(BlueprintReadWrite, EditAnywhere, Category = Settings)
+	int32 WorldPositionOffsetDisableDistance = 0;
+
 	/** Whether the culling should be reversed or not (needed to support negative scales) */
 	UPROPERTY(BlueprintReadWrite, EditAnywhere, Category = Settings)
 	bool bIsLocalToWorldDeterminantNegative = false;
@@ -138,7 +141,8 @@ public:
 		bool bOverrideMaterials,
 		const TArray<TSoftObjectPtr<UMaterialInterface>>& MaterialOverrides,
 		const float InCullStartDistance,
-		const float InCullEndDistance, 
+		const float InCullEndDistance,
+		const int32 InWorldPositionOffsetDisableDistance,
 		const bool bInIsLocalToWorldDeterminantNegative);
 };
 
diff --git a/Engine/Plugins/Experimental/PCG/Source/PCG/Public/MeshSelectors/PCGMeshSelectorByAttribute.h b/Engine/Plugins/Experimental/PCG/Source/PCG/Public/MeshSelectors/PCGMeshSelectorByAttribute.h
--- a/Engine/Plugins/Experimental/PCG/Source/PCG/Public/MeshSelectors/PCGMeshSelectorByAttribute.h
+++ b/Engine/Plugins/Experimental/PCG/Source/PCG/Public/MeshSelectors/PCGMeshSelectorByAttribute.h
@@ -55,4 +55,7 @@ public:
 	/** Distance at which instances are culled. Use 0 to disable. */
 	UPROPERTY(BlueprintReadWrite, EditAnywhere, Category = Settings)
 	float CullEndDistance = 0;
+
+	UPROPERTY(BlueprintReadWrite, EditAnywhere, Category = Settings)
+	int32 WorldPositionOffsetDisableDistance = 0;
 };
diff --git a/Engine/Plugins/Experimental/PCG/Source/PCG/Public/MeshSelectors/PCGMeshSelectorWeighted.h b/Engine/Plugins/Experimental/PCG/Source/PCG/Public/MeshSelectors/PCGMeshSelectorWeighted.h
--- a/Engine/Plugins/Experimental/PCG/Source/PCG/Public/MeshSelectors/PCGMeshSelectorWeighted.h
+++ b/Engine/Plugins/Experimental/PCG/Source/PCG/Public/MeshSelectors/PCGMeshSelectorWeighted.h
@@ -49,6 +49,9 @@ struct PCG_API FPCGMeshSelectorWeightedEntry
 	UPROPERTY(BlueprintReadWrite, EditAnywhere, Category = Settings)
 	float CullEndDistance = 0;
 
+	UPROPERTY(BlueprintReadWrite, EditAnywhere, Category = Settings)
+	int32 WorldPositionOffsetDisableDistance = 0;
+
 	UPROPERTY(BlueprintReadWrite, EditAnywhere, Category = Settings, meta = (ClampMin = "0"))
 	int Weight = 1;
 };
