Added some functionality to get local space information on splines
Added some flexibility on seeding in spline sampler - world space or local space, 3d or 2d (xy)

diff --git a/Engine/Plugins/Experimental/PCG/Source/PCG/Private/Data/PCGLandscapeSplineData.cpp b/Engine/Plugins/Experimental/PCG/Source/PCG/Private/Data/PCGLandscapeSplineData.cpp
--- a/Engine/Plugins/Experimental/PCG/Source/PCG/Private/Data/PCGLandscapeSplineData.cpp
+++ b/Engine/Plugins/Experimental/PCG/Source/PCG/Private/Data/PCGLandscapeSplineData.cpp
@@ -44,6 +44,11 @@ void UPCGLandscapeSplineData::Initialize(ULandscapeSplinesComponent* InSplineCom
 	TargetActor = InSplineComponent->GetOwner();
 }
 
+FTransform UPCGLandscapeSplineData::GetTransform() const
+{
+	return Spline.IsValid() ? Spline->GetComponentTransform() : FTransform::Identity;
+}
+
 int UPCGLandscapeSplineData::GetNumSegments() const
 {
 	check(Spline.IsValid());
@@ -67,7 +72,7 @@ FVector::FReal UPCGLandscapeSplineData::GetSegmentLength(int SegmentIndex) const
 	return Length;
 }
 
-FTransform UPCGLandscapeSplineData::GetTransformAtDistance(int SegmentIndex, FVector::FReal Distance, FBox* OutBounds) const
+FTransform UPCGLandscapeSplineData::GetTransformAtDistance(int SegmentIndex, FVector::FReal Distance, bool bWorldSpace, FBox* OutBounds) const
 {
 	check(Spline.IsValid());
 	check(SegmentIndex >= 0 && SegmentIndex < Spline->GetSegments().Num());
@@ -109,7 +114,14 @@ FTransform UPCGLandscapeSplineData::GetTransformAtDistance(int SegmentIndex, FVe
 				OutBounds->Max.Y *= (CurrentPoint.FalloffRight - CurrentPoint.Center).Length() / (CurrentPoint.Right - CurrentPoint.Center).Length();
 			}
 
-			return PreviousTransform * Spline->GetComponentTransform();
+			if (bWorldSpace)
+			{
+				return PreviousTransform * Spline->GetComponentTransform();
+			}
+			else
+			{
+				return PreviousTransform;
+			}
 		}
 		else
 		{
diff --git a/Engine/Plugins/Experimental/PCG/Source/PCG/Private/Data/PCGSplineData.cpp b/Engine/Plugins/Experimental/PCG/Source/PCG/Private/Data/PCGSplineData.cpp
--- a/Engine/Plugins/Experimental/PCG/Source/PCG/Private/Data/PCGSplineData.cpp
+++ b/Engine/Plugins/Experimental/PCG/Source/PCG/Private/Data/PCGSplineData.cpp
@@ -33,6 +33,11 @@ void UPCGSplineData::Initialize(USplineComponent* InSpline)
 	CachedBounds = CachedBounds.ExpandBy(SplinePointsRadius, SplinePointsRadius);
 }
 
+FTransform UPCGSplineData::GetTransform() const
+{
+	return Spline ? Spline->GetComponentTransform() : FTransform::Identity;
+}
+
 int UPCGSplineData::GetNumSegments() const
 {
 	return Spline ? Spline->GetNumberOfSplineSegments() : 0;
@@ -43,19 +48,19 @@ FVector::FReal UPCGSplineData::GetSegmentLength(int SegmentIndex) const
 	return Spline->GetDistanceAlongSplineAtSplinePoint(SegmentIndex + 1) - Spline->GetDistanceAlongSplineAtSplinePoint(SegmentIndex);
 }
 
-FVector UPCGSplineData::GetLocationAtDistance(int SegmentIndex, FVector::FReal Distance) const
+FVector UPCGSplineData::GetLocationAtDistance(int SegmentIndex, FVector::FReal Distance, bool bWorldSpace) const
 {
-	return Spline->GetLocationAtDistanceAlongSpline(Spline->GetDistanceAlongSplineAtSplinePoint(SegmentIndex) + Distance, ESplineCoordinateSpace::World);
+	return Spline->GetLocationAtDistanceAlongSpline(Spline->GetDistanceAlongSplineAtSplinePoint(SegmentIndex) + Distance, bWorldSpace ? ESplineCoordinateSpace::World : ESplineCoordinateSpace::Local);
 }
 
-FTransform UPCGSplineData::GetTransformAtDistance(int SegmentIndex, FVector::FReal Distance, FBox* OutBounds) const
+FTransform UPCGSplineData::GetTransformAtDistance(int SegmentIndex, FVector::FReal Distance, bool bWorldSpace, FBox* OutBounds) const
 {
 	if (OutBounds)
 	{
 		*OutBounds = FBox::BuildAABB(FVector::ZeroVector, FVector::OneVector);
 	}
 
-	return Spline->GetTransformAtDistanceAlongSpline(Spline->GetDistanceAlongSplineAtSplinePoint(SegmentIndex) + Distance, ESplineCoordinateSpace::World, /*bUseScale=*/true);
+	return Spline->GetTransformAtDistanceAlongSpline(Spline->GetDistanceAlongSplineAtSplinePoint(SegmentIndex) + Distance, bWorldSpace ? ESplineCoordinateSpace::World : ESplineCoordinateSpace::Local, /*bUseScale=*/true);
 }
 
 FVector::FReal UPCGSplineData::GetCurvatureAtDistance(int SegmentIndex, FVector::FReal Distance) const
diff --git a/Engine/Plugins/Experimental/PCG/Source/PCG/Private/Elements/PCGSplineSampler.cpp b/Engine/Plugins/Experimental/PCG/Source/PCG/Private/Elements/PCGSplineSampler.cpp
--- a/Engine/Plugins/Experimental/PCG/Source/PCG/Private/Elements/PCGSplineSampler.cpp
+++ b/Engine/Plugins/Experimental/PCG/Source/PCG/Private/Elements/PCGSplineSampler.cpp
@@ -4,6 +4,7 @@
 
 #include "PCGContext.h"
 #include "PCGCustomVersion.h"
+#include "PCGHelpers.h"
 #include "Data/PCGIntersectionData.h"
 #include "Data/PCGLandscapeSplineData.h"
 #include "Data/PCGPointData.h"
@@ -311,13 +312,40 @@ namespace PCGSplineSampler
 {
 	struct FSamplerResult
 	{
-		FTransform Transform;
+		FTransform LocalTransform;
 		FBox Box = FBox::BuildAABB(FVector::ZeroVector, FVector::OneVector);
 		FVector::FReal Curvature = 0;
 		FVector::FReal PreviousDeltaAngle = 0;
 		FVector::FReal NextDeltaAngle = 0;
 	};
 
+	void SetSeed(FPCGPoint& Point, const FVector& LSPosition, const FPCGSplineSamplerParams& Params)
+	{
+		if (!Params.bSeedFromLocalPosition)
+		{
+			if (!Params.bSeedFrom2DPosition)
+			{
+				Point.Seed = UPCGBlueprintHelpers::ComputeSeedFromPosition(Point.Transform.GetLocation());
+			}
+			else
+			{
+				const FVector WSPosition = Point.Transform.GetLocation();
+				Point.Seed = PCGHelpers::ComputeSeed((int)WSPosition.X, (int)WSPosition.Y);
+			}
+		}
+		else // Use provided local position
+		{
+			if (!Params.bSeedFrom2DPosition)
+			{
+				Point.Seed = UPCGBlueprintHelpers::ComputeSeedFromPosition(LSPosition);
+			}
+			else
+			{
+				Point.Seed = PCGHelpers::ComputeSeed((int)LSPosition.X, (int)LSPosition.Y);
+			}
+		}
+	}
+
 	struct FStepSampler
 	{
 		FStepSampler(const UPCGPolyLineData* InLineData, const FPCGSplineSamplerParams& Params)
@@ -357,8 +385,8 @@ namespace PCGSplineSampler
 			const FVector::FReal SegmentStep = SegmentLength / (SubdivisionsPerSegment + 1);
 
 			FBox& OutBox = OutResult.Box;
-			FTransform& OutTransform = OutResult.Transform;
-			OutTransform = LineData->GetTransformAtDistance(CurrentSegmentIndex, SubpointIndex * SegmentStep, &OutBox);
+			FTransform& OutTransform = OutResult.LocalTransform;
+			OutTransform = LineData->GetTransformAtDistance(CurrentSegmentIndex, SubpointIndex * SegmentStep, /*bWorldSpace=*/false, &OutBox);
 
 			if (bComputeCurvature)
 			{
@@ -369,7 +397,7 @@ namespace PCGSplineSampler
 			{
 				const int PreviousSegmentIndex = (CurrentSegmentIndex > 0 ? CurrentSegmentIndex : NumSegments) - 1;
 				const FVector::FReal PreviousSegmentLength = LineData->GetSegmentLength(PreviousSegmentIndex);
-				FTransform PreviousSegmentEndTransform = LineData->GetTransformAtDistance(PreviousSegmentIndex, PreviousSegmentLength);
+				FTransform PreviousSegmentEndTransform = LineData->GetTransformAtDistance(PreviousSegmentIndex, PreviousSegmentLength, /*bWorldSpace=*/false);
 
 				if ((PreviousSegmentEndTransform.GetLocation() - OutTransform.GetLocation()).Length() <= KINDA_SMALL_NUMBER)
 				{
@@ -412,9 +440,9 @@ namespace PCGSplineSampler
 		virtual void Step(FSamplerResult& OutResult) override
 		{
 			FVector::FReal CurrentSegmentLength = LineData->GetSegmentLength(CurrentSegmentIndex);
-			FTransform& OutTransform = OutResult.Transform;
+			FTransform& OutTransform = OutResult.LocalTransform;
 			FBox& OutBox = OutResult.Box;
-			OutTransform = LineData->GetTransformAtDistance(CurrentSegmentIndex, CurrentDistance, &OutBox);
+			OutTransform = LineData->GetTransformAtDistance(CurrentSegmentIndex, CurrentDistance, /*bWorldSpace=*/false, &OutBox);
 
 			OutBox.Min.X *= DistanceIncrement / OutTransform.GetScale3D().X;
 			OutBox.Max.X *= DistanceIncrement / OutTransform.GetScale3D().X;
@@ -505,10 +533,13 @@ namespace PCGSplineSampler
 
 		virtual void Sample(const FSamplerResult& InResult, UPCGPointData* OutPointData)
 		{
+			const FTransform Transform = InResult.LocalTransform * LineData->GetTransform();
+
 			FPCGPoint TrivialPoint, BoundsTestPoint;
-			if(SpatialData->SamplePoint(InResult.Transform, InResult.Box, TrivialPoint, OutPointData->Metadata)
-				&& (!BoundingShapeData || BoundingShapeData->SamplePoint(InResult.Transform, InResult.Box, BoundsTestPoint, nullptr)))
+			if(SpatialData->SamplePoint(Transform, InResult.Box, TrivialPoint, OutPointData->Metadata)
+				&& (!BoundingShapeData || BoundingShapeData->SamplePoint(Transform, InResult.Box, BoundsTestPoint, nullptr)))
 			{
+				SetSeed(TrivialPoint, InResult.LocalTransform.GetLocation(), Params);
 				SetMetadata(InResult, TrivialPoint, OutPointData->Metadata);
 				OutPointData->GetMutablePoints().Add(TrivialPoint);
 			}
@@ -537,7 +568,8 @@ namespace PCGSplineSampler
 
 		virtual void Sample(const FSamplerResult& InResult, UPCGPointData* OutPointData) override
 		{
-			const FTransform& InTransform = InResult.Transform;
+			const FTransform& TransformLS = InResult.LocalTransform;
+			const FTransform TransformWS = TransformLS * LineData->GetTransform();
 			const FBox& InBox = InResult.Box;
 
 			// We're assuming that we can scale against the origin in this method so this should always be true.
@@ -552,11 +584,11 @@ namespace PCGSplineSampler
 			SubBox.Max /= FVector(1.0, NumPlanarSteps, NumHeightSteps);
 
 			FPCGPoint SeedPoint;
-			if (SpatialData->SamplePoint(InTransform, InBox, SeedPoint, nullptr))
+			if (SpatialData->SamplePoint(TransformWS, InBox, SeedPoint, nullptr))
 			{
 				// Assuming the the normal to the curve is on the Y axis
-				const FVector YAxis = SeedPoint.Transform.GetScaledAxis(EAxis::Y);
-				const FVector ZAxis = SeedPoint.Transform.GetScaledAxis(EAxis::Z);
+				const FVector YAxis = TransformLS.GetScaledAxis(EAxis::Y);
+				const FVector ZAxis = TransformLS.GetScaledAxis(EAxis::Z);
 
 				// TODO: we can optimize this if we are in the "edges only case" to pick only values that
 				FVector::FReal CurrentZ = (InBox.Min.Z + ZHalfStep);
@@ -592,16 +624,17 @@ namespace PCGSplineSampler
 
 						if (bTestPoint)
 						{
-							FVector TentativeLocation = InTransform.GetLocation() + YAxis * CurrentY + ZAxis * CurrentZ;
+							FVector TentativeLocationLS = TransformLS.GetLocation() + YAxis * CurrentY + ZAxis * CurrentZ;
 
-							FTransform TentativeTransform = InTransform;
-							TentativeTransform.SetLocation(TentativeLocation);
+							FTransform TentativeTransform = TransformWS;
+							TentativeTransform.SetLocation(TransformWS.TransformPosition(TentativeLocationLS));
 
 							// Sample point and if successful check against bounding shape (if provided).
 							FPCGPoint OutPoint, BoundsTestPoint;
 							if (SpatialData->SamplePoint(TentativeTransform, SubBox, OutPoint, OutPointData->Metadata)
 								&& (!BoundingShapeData || BoundingShapeData->SamplePoint(TentativeTransform, SubBox, BoundsTestPoint, nullptr)))
 							{
+								SetSeed(OutPoint, TentativeLocationLS, Params);
 								SetMetadata(InResult, OutPoint, OutPointData->Metadata);
 								OutPointData->GetMutablePoints().Add(OutPoint);
 							}
@@ -651,7 +684,7 @@ namespace PCGSplineSampler
 			Sampler->Step(*CurrentResult);
 
 			// Get unsigned angle difference between the two points
-			const FVector::FReal DeltaSinAngle = ((CurrentResult->Transform.GetUnitAxis(EAxis::X) ^ PreviousResult->Transform.GetUnitAxis(EAxis::X)) | PreviousResult->Transform.GetUnitAxis(EAxis::Z));
+			const FVector::FReal DeltaSinAngle = ((CurrentResult->LocalTransform.GetUnitAxis(EAxis::X) ^ PreviousResult->LocalTransform.GetUnitAxis(EAxis::X)) | PreviousResult->LocalTransform.GetUnitAxis(EAxis::Z));
 			// Normalize value to be between -1 and 1
 			FVector::FReal DeltaAngle = FMath::Asin(DeltaSinAngle) / UE_HALF_PI;
 
@@ -669,12 +702,6 @@ namespace PCGSplineSampler
 		{
 			ExtentsSampler->Sample(*PreviousResult, OutPointData);
 		}
-
-		// Finally, set seed on points based on position
-		for (FPCGPoint& Point : OutPointData->GetMutablePoints())
-		{
-			Point.Seed = UPCGBlueprintHelpers::ComputeSeedFromPosition(Point.Transform.GetLocation());
-		}
 	}
  
 	void SampleInteriorData(const UPCGPolyLineData* LineData, const UPCGSpatialData* SpatialData, const UPCGSpatialData* InBoundingShape, const FPCGSplineSamplerParams& Params, UPCGPointData* OutPointData)
@@ -707,8 +734,9 @@ namespace PCGSplineSampler
 			return;
 		}
 
-		const FVector MinPoint = Spline->Bounds.Origin - Spline->Bounds.BoxExtent;
-		const FVector MaxPoint = Spline->Bounds.Origin + Spline->Bounds.BoxExtent;
+		const FBoxSphereBounds SplineLocalBounds = Spline->CalcLocalBounds();
+		const FVector MinPoint = SplineLocalBounds.Origin - SplineLocalBounds.BoxExtent;
+		const FVector MaxPoint = SplineLocalBounds.Origin + SplineLocalBounds.BoxExtent;
 
 		const FVector::FReal MaxDimension = FMath::Max(Spline->Bounds.BoxExtent.X, Spline->Bounds.BoxExtent.Y) * 2.f;
 		const FVector::FReal MaxDimensionSquared = MaxDimension * MaxDimension;
@@ -732,7 +760,7 @@ namespace PCGSplineSampler
 			const int NumSegments = LineData->GetNumSegments();
 			for (int32 SegmentIndex = 0; SegmentIndex < NumSegments; ++SegmentIndex)
 			{
-				SplineSamplePoints.Add(LineData->GetLocationAtDistance(SegmentIndex, 0));
+				SplineSamplePoints.Add(LineData->GetLocationAtDistance(SegmentIndex, 0, /*bWorldSpace=*/false));
 			}
 		}
 		else
@@ -740,7 +768,7 @@ namespace PCGSplineSampler
 			// Get sample points along the spline that are higher resolution than our PolyLine
 			for (FVector::FReal Length = 0.f; Length < SplineLength; Length += Params.InteriorBorderSampleSpacing)
 			{
-				SplineSamplePoints.Add(Spline->GetLocationAtDistanceAlongSpline(Length, ESplineCoordinateSpace::World));
+				SplineSamplePoints.Add(Spline->GetLocationAtDistanceAlongSpline(Length, ESplineCoordinateSpace::Local));
 			}
 		}
 
@@ -897,20 +925,17 @@ namespace PCGSplineSampler
 					}
 
 					// if bTreatAsPolyline, then we shouldnt use this, we should use nearest point on the polygon line segments
-					const float NearestSplineKey = bFindNearestSplineKey ? Spline->FindInputKeyClosestToWorldLocation(SurfaceLocation) : 0.f;
+					float Dummy;
+					const float NearestSplineKey = bFindNearestSplineKey ? Spline->SplineCurves.Position.InaccurateFindNearest(SurfaceLocation, Dummy) : 0.f;
 
-					if (Params.bProjectOntoSurface)
-					{
-						Point.Transform.SetLocation(SurfaceLocation);
-					}
-					else
-					{
-						Point.Transform.SetLocation(FVector(SampleLocation, MinPoint.Z));
-					}
+					const FVector PointLocationLS = (Params.bProjectOntoSurface ? FVector(SurfaceLocation) : FVector(SampleLocation, MinPoint.Z));
+
+					FTransform TransformLS = FTransform::Identity;
+					TransformLS.SetLocation(PointLocationLS);
 
 					if (Params.InteriorOrientation == EPCGSplineSamplingInteriorOrientation::FollowCurvature)
 					{
-						Point.Transform.SetRotation(Spline->GetQuaternionAtSplineInputKey(NearestSplineKey, ESplineCoordinateSpace::World));
+						TransformLS.SetRotation(Spline->GetQuaternionAtSplineInputKey(NearestSplineKey, ESplineCoordinateSpace::Local));
 					}
 
 					// Calculate density fall off
@@ -966,7 +991,8 @@ namespace PCGSplineSampler
 
 					Point.BoundsMin = BoundsMin;
 					Point.BoundsMax = BoundsMax;
-					Point.Seed = UPCGBlueprintHelpers::ComputeSeedFromPosition(Point.Transform.GetLocation());
+					Point.Transform = TransformLS * LineData->GetTransform();
+					SetSeed(Point, PointLocationLS, Params);
 				}
 			}
 		}
diff --git a/Engine/Plugins/Experimental/PCG/Source/PCG/Public/Data/PCGLandscapeSplineData.h b/Engine/Plugins/Experimental/PCG/Source/PCG/Public/Data/PCGLandscapeSplineData.h
--- a/Engine/Plugins/Experimental/PCG/Source/PCG/Public/Data/PCGLandscapeSplineData.h
+++ b/Engine/Plugins/Experimental/PCG/Source/PCG/Public/Data/PCGLandscapeSplineData.h
@@ -22,9 +22,10 @@ public:
 	// ~End UPCGData interface
 
 	//~Begin UPCGPolyLineData interface
+	virtual FTransform GetTransform() const override;
 	virtual int GetNumSegments() const override;
 	virtual FVector::FReal GetSegmentLength(int SegmentIndex) const override;
-	virtual FTransform GetTransformAtDistance(int SegmentIndex, FVector::FReal Distance, FBox* OutBounds = nullptr) const override;
+	virtual FTransform GetTransformAtDistance(int SegmentIndex, FVector::FReal Distance, bool bWorldSpace = true, FBox* OutBounds = nullptr) const override;
 	virtual FVector::FReal GetCurvatureAtDistance(int SegmentIndex, FVector::FReal Distance) const override;
 	//~End UPCGPolyLineData interface
 
diff --git a/Engine/Plugins/Experimental/PCG/Source/PCG/Public/Data/PCGPolyLineData.h b/Engine/Plugins/Experimental/PCG/Source/PCG/Public/Data/PCGPolyLineData.h
--- a/Engine/Plugins/Experimental/PCG/Source/PCG/Public/Data/PCGPolyLineData.h
+++ b/Engine/Plugins/Experimental/PCG/Source/PCG/Public/Data/PCGPolyLineData.h
@@ -21,12 +21,13 @@ public:
 	virtual FBox GetBounds() const override;
 	//~End UPCGSpatialData interface
 
+	virtual FTransform GetTransform() const { return FTransform::Identity; }
 	virtual int GetNumSegments() const PURE_VIRTUAL(UPCGPolyLineData::GetNumSegments, return 0;);
 	virtual FVector::FReal GetSegmentLength(int SegmentIndex) const PURE_VIRTUAL(UPCGPolyLineData::GetSegmentLength, return 0;);
-	virtual FTransform GetTransformAtDistance(int SegmentIndex, FVector::FReal Distance, FBox* OutBounds = nullptr) const PURE_VIRTUAL(UPCGPolyLine::GetTransformAtDistance, return FTransform(););
+	virtual FTransform GetTransformAtDistance(int SegmentIndex, FVector::FReal Distance, bool bWorldSpace = true, FBox* OutBounds = nullptr) const PURE_VIRTUAL(UPCGPolyLine::GetTransformAtDistance, return FTransform(););
 	virtual FVector::FReal GetCurvatureAtDistance(int SegmentIndex, FVector::FReal Distance) const { return 0; }
 
-	virtual FVector GetLocationAtDistance(int SegmentIndex, FVector::FReal Distance) const { return GetTransformAtDistance(SegmentIndex, Distance).GetLocation(); }
+	virtual FVector GetLocationAtDistance(int SegmentIndex, FVector::FReal Distance, bool bWorldSpace = true) const { return GetTransformAtDistance(SegmentIndex, Distance, bWorldSpace).GetLocation(); }
 };
 
 #if UE_ENABLE_INCLUDE_ORDER_DEPRECATED_IN_5_2
diff --git a/Engine/Plugins/Experimental/PCG/Source/PCG/Public/Data/PCGSplineData.h b/Engine/Plugins/Experimental/PCG/Source/PCG/Public/Data/PCGSplineData.h
--- a/Engine/Plugins/Experimental/PCG/Source/PCG/Public/Data/PCGSplineData.h
+++ b/Engine/Plugins/Experimental/PCG/Source/PCG/Public/Data/PCGSplineData.h
@@ -27,10 +27,11 @@ public:
 	// ~End UPCGData interface
 
 	//~Begin UPCGPolyLineData interface
+	virtual FTransform GetTransform() const override;
 	virtual int GetNumSegments() const override;
 	virtual FVector::FReal GetSegmentLength(int SegmentIndex) const override;
-	virtual FVector GetLocationAtDistance(int SegmentIndex, FVector::FReal Distance) const override;
-	virtual FTransform GetTransformAtDistance(int SegmentIndex, FVector::FReal Distance, FBox* OutBounds = nullptr) const override;
+	virtual FVector GetLocationAtDistance(int SegmentIndex, FVector::FReal Distance, bool bWorldSpace = true) const override;
+	virtual FTransform GetTransformAtDistance(int SegmentIndex, FVector::FReal Distance, bool bWorldSpace = true, FBox* OutBounds = nullptr) const override;
 	virtual FVector::FReal GetCurvatureAtDistance(int SegmentIndex, FVector::FReal Distance) const override;
 	//~End UPCGPolyLineData interface
 
diff --git a/Engine/Plugins/Experimental/PCG/Source/PCG/Public/Elements/PCGSplineSampler.h b/Engine/Plugins/Experimental/PCG/Source/PCG/Public/Elements/PCGSplineSampler.h
--- a/Engine/Plugins/Experimental/PCG/Source/PCG/Public/Elements/PCGSplineSampler.h
+++ b/Engine/Plugins/Experimental/PCG/Source/PCG/Public/Elements/PCGSplineSampler.h
@@ -107,6 +107,14 @@ struct PCG_API FPCGSplineSamplerParams
 	/** Defines the density for each sample based on its distance from the spline */
 	UPROPERTY(EditAnywhere, Category = Settings, meta = (EditCondition = "Dimension==EPCGSplineSamplingDimension::OnInterior"))
 	FRuntimeFloatCurve InteriorDensityFalloffCurve;
+
+	/** Controls whether we will seed the sampled points using the final world position or the local position */
+	UPROPERTY(BlueprintReadWrite, EditAnywhere, Category = "Settings|Seeding")
+	bool bSeedFromLocalPosition = false;
+
+	/** Controls whether we will seed the sampled points using the 3D position or the 2D (XY) position */
+	UPROPERTY(BlueprintReadWrite, EditAnywhere, Category = "Settings|Seeding")
+	bool bSeedFrom2DPosition = false;
 };
 
 namespace PCGSplineSamplerHelpers
