[PCG] Fixed an issue were undoing a connection wouldnt regenerate the graph.

diff --git a/Engine/Plugins/Experimental/PCG/Source/PCGEditor/Private/PCGEditor.cpp b/Engine/Plugins/Experimental/PCG/Source/PCGEditor/Private/PCGEditor.cpp
--- a/Engine/Plugins/Experimental/PCG/Source/PCGEditor/Private/PCGEditor.cpp
+++ b/Engine/Plugins/Experimental/PCG/Source/PCGEditor/Private/PCGEditor.cpp
@@ -37,26 +37,27 @@
 
 #include "AssetToolsModule.h"
 #include "EdGraphUtilities.h"
-#include "Editor/UnrealEdEngine.h"
 #include "EditorAssetLibrary.h"
-#include "Fonts/FontMeasure.h"
-#include "Framework/Application/SlateApplication.h"
-#include "Framework/Commands/GenericCommands.h"
-#include "Framework/Notifications/NotificationManager.h"
 #include "GraphEditorActions.h"
-#include "HAL/PlatformApplicationMisc.h"
 #include "IDetailsView.h"
-#include "Preferences/UnrealEdOptions.h"
 #include "PropertyEditorModule.h"
 #include "SNodePanel.h"
-#include "Misc/ITransaction.h"
 #include "ScopedTransaction.h"
 #include "SourceCodeNavigation.h"
 #include "ToolMenu.h"
-#include "ToolMenus.h"
 #include "ToolMenuEntry.h"
-#include "UnrealEdGlobals.h"
 #include "ToolMenuSection.h"
+#include "ToolMenus.h"
+#include "UnrealEdGlobals.h"
+#include "Editor/UnrealEdEngine.h"
+#include "Fonts/FontMeasure.h"
+#include "Framework/Application/SlateApplication.h"
+#include "Framework/Commands/GenericCommands.h"
+#include "Framework/Notifications/NotificationManager.h"
+#include "HAL/PlatformApplicationMisc.h"
+#include "Misc/ITransaction.h"
+#include "Misc/TransactionObjectEvent.h"
+#include "Preferences/UnrealEdOptions.h"
 #include "Widgets/Docking/SDockTab.h"
 #include "Widgets/Notifications/SNotificationList.h"
 
@@ -241,7 +242,26 @@ void FPCGEditor::AddReferencedObjects(FReferenceCollector& Collector)
 
 bool FPCGEditor::MatchesContext(const FTransactionContext& InContext, const TArray<TPair<UObject*, FTransactionObjectEvent>>& TransactionObjectContexts) const
 {
-	return InContext.Context == FPCGEditorCommon::ContextIdentifier;
+	if (InContext.Context == FPCGEditorCommon::ContextIdentifier)
+	{
+		return true;
+	}
+
+	// This is done to catch transaction blocks made outside PCG editor code were we need to trigger PostUndo for our context, i.e. UPCGEditorGraphSchema::TryCreateConnection
+	for (const TPair<UObject*, FTransactionObjectEvent>& TransactionObjectContext : TransactionObjectContexts)
+	{
+		const UObject* Object = TransactionObjectContext.Key;
+		while (Object != nullptr)
+		{
+			if (Object == PCGGraphBeingEdited)
+			{
+				return true;
+			}
+			Object = Object->GetOuter();
+		}
+	}
+
+	return false;
 }
 
 void FPCGEditor::PostUndo(bool bSuccess)
