[PCG] Added support for reroute nodes in the PCG editor.

diff --git a/Engine/Plugins/Experimental/PCG/Source/PCG/Private/Elements/PCGExecuteBlueprint.cpp b/Engine/Plugins/Experimental/PCG/Source/PCG/Private/Elements/PCGExecuteBlueprint.cpp
--- a/Engine/Plugins/Experimental/PCG/Source/PCG/Private/Elements/PCGExecuteBlueprint.cpp
+++ b/Engine/Plugins/Experimental/PCG/Source/PCG/Private/Elements/PCGExecuteBlueprint.cpp
@@ -177,6 +177,10 @@ FRandomStream UPCGBlueprintElement::GetRandomStream(FPCGContext& InContext) cons
 UPCGBlueprintSettings::UPCGBlueprintSettings()
 {
 	bUseSeed = true;
+	
+#if WITH_EDITORONLY_DATA
+	bExposeToLibrary = false;
+#endif
 }
 
 void UPCGBlueprintSettings::SetupBlueprintEvent()
diff --git a/Engine/Plugins/Experimental/PCG/Source/PCG/Private/PCGGraph.cpp b/Engine/Plugins/Experimental/PCG/Source/PCG/Private/PCGGraph.cpp
--- a/Engine/Plugins/Experimental/PCG/Source/PCG/Private/PCGGraph.cpp
+++ b/Engine/Plugins/Experimental/PCG/Source/PCG/Private/PCGGraph.cpp
@@ -1,6 +1,7 @@
 // Copyright Epic Games, Inc. All Rights Reserved.
 
 #include "PCGGraph.h"
+
 #include "PCGInputOutputSettings.h"
 #include "PCGModule.h"
 #include "PCGSubsystem.h"
@@ -333,6 +334,9 @@ bool UPCGGraph::AddLabeledEdge(UPCGNode* From, const FName& FromPinLabel, UPCGNo
 		bToPinBrokeOtherEdges = ToPin->BreakAllIncompatibleEdges();
 	}
 	
+	From->UpdateDynamicPins();
+	To->UpdateDynamicPins();
+
 #if WITH_EDITOR
 	NotifyGraphChanged(EPCGChangeType::Structural);
 #endif
@@ -411,12 +415,15 @@ bool UPCGGraph::RemoveEdge(UPCGNode* From, const FName& FromLabel, UPCGNode* To,
 
 	const bool bChanged = OutPin && OutPin->BreakEdgeTo(InPin);
 
-#if WITH_EDITOR
 	if (bChanged)
 	{
+		From->UpdateDynamicPins();
+		To->UpdateDynamicPins();
+
+#if WITH_EDITOR
 		NotifyGraphChanged(EPCGChangeType::Structural);
-	}
 #endif
+	}
 
 	return bChanged;
 }
@@ -445,6 +452,7 @@ bool UPCGGraph::RemoveAllInboundEdges(UPCGNode* InNode)
 #if WITH_EDITOR
 	if (bChanged)
 	{
+		InNode->UpdateDynamicPins();
 		NotifyGraphChanged(EPCGChangeType::Structural);
 	}
 #endif
@@ -464,6 +472,7 @@ bool UPCGGraph::RemoveAllOutboundEdges(UPCGNode* InNode)
 #if WITH_EDITOR
 	if (bChanged)
 	{
+		InNode->UpdateDynamicPins();
 		NotifyGraphChanged(EPCGChangeType::Structural);
 	}
 #endif
@@ -484,6 +493,7 @@ bool UPCGGraph::RemoveInboundEdges(UPCGNode* InNode, const FName& InboundLabel)
 #if WITH_EDITOR
 	if (bChanged)
 	{
+		InNode->UpdateDynamicPins();
 		NotifyGraphChanged(EPCGChangeType::Structural);
 	}
 #endif
@@ -504,6 +514,7 @@ bool UPCGGraph::RemoveOutboundEdges(UPCGNode* InNode, const FName& OutboundLabel
 #if WITH_EDITOR
 	if (bChanged)
 	{
+		InNode->UpdateDynamicPins();
 		NotifyGraphChanged(EPCGChangeType::Structural);
 	}
 #endif
diff --git a/Engine/Plugins/Experimental/PCG/Source/PCG/Private/PCGNode.cpp b/Engine/Plugins/Experimental/PCG/Source/PCG/Private/PCGNode.cpp
--- a/Engine/Plugins/Experimental/PCG/Source/PCG/Private/PCGNode.cpp
+++ b/Engine/Plugins/Experimental/PCG/Source/PCG/Private/PCGNode.cpp
@@ -597,9 +597,11 @@ void UPCGNode::CreateDefaultPins(TFunctionRef<UPCGPin* (UPCGNode*)> PinAllocator
 	CreatePins(OutputPins, Settings->DefaultOutputPinProperties());
 }
 
-EPCGChangeType UPCGNode::UpdatePins(TFunctionRef<UPCGPin*(UPCGNode*)> PinAllocator)
+EPCGChangeType UPCGNode::UpdatePins(TFunctionRef<UPCGPin*(UPCGNode*)> PinAllocator, const UPCGNode* FromNode /*= nullptr*/)
 {
-	if (!GetSettings())
+	const UPCGSettings* Settings = GetSettings();
+	
+	if (!Settings)
 	{
 		bool bChanged = !InputPins.IsEmpty() || !OutputPins.IsEmpty();
 
@@ -629,10 +631,9 @@ EPCGChangeType UPCGNode::UpdatePins(TFunctionRef<UPCGPin*(UPCGNode*)> PinAllocat
 		OutputPins.Reset();
 		return EPCGChangeType::Edge | EPCGChangeType::Node;
 	}
-
-	UPCGSettings* Settings = GetSettings();
-	TArray<FPCGPinProperties> InboundPinProperties = Settings->AllInputPinProperties();
-	TArray<FPCGPinProperties> OutboundPinProperties = Settings->AllOutputPinProperties();
+	
+	const TArray<FPCGPinProperties> InboundPinProperties = Settings->AllInputPinProperties();
+	const TArray<FPCGPinProperties> OutboundPinProperties = Settings->AllOutputPinProperties();
 
 	auto UpdatePins = [this, &PinAllocator](TArray<UPCGPin*>& Pins, const TArray<FPCGPinProperties>& PinProperties)
 	{
@@ -744,6 +745,48 @@ EPCGChangeType UPCGNode::UpdatePins(TFunctionRef<UPCGPin*(UPCGNode*)> PinAllocat
 	ChangeType |= UpdatePins(InputPins, InboundPinProperties);
 	ChangeType |= UpdatePins(OutputPins, OutboundPinProperties);
 
+#if WITH_EDITOR
+	OnNodeChangedDelegate.Broadcast(this, ChangeType);
+#endif // WITH_EDITOR
+		
+	return ChangeType;
+}
+
+EPCGChangeType UPCGNode::UpdateDynamicPins(const UPCGNode* FromNode /*= nullptr*/)
+{
+	EPCGChangeType ChangeType = EPCGChangeType::None;
+	const UPCGSettings* Settings = GetSettings();
+	if (!Settings || !Settings->HasDynamicPins())
+	{
+		return ChangeType;
+	}
+
+	ChangeType |= UpdatePins();
+	
+	for (UPCGPin* OutputPin : OutputPins)
+	{
+		for (UPCGEdge* Edge : OutputPin->Edges)
+		{
+			const UPCGPin* OtherPin = Edge->GetOtherPin(OutputPin);
+			if (!OtherPin)
+			{
+				continue;
+			}
+			
+			UPCGNode* OtherNode = OtherPin->Node;
+			if (!OtherNode || OtherNode == FromNode)
+			{
+				continue;
+			}
+			
+			const UPCGSettings* OtherSettings = OtherNode->GetSettings();
+			if (Settings && OtherSettings->HasDynamicPins())
+			{
+				ChangeType |= OtherNode->UpdateDynamicPins(this);
+			}
+		}
+	}
+	
 	return ChangeType;
 }
 
diff --git a/Engine/Plugins/Experimental/PCG/Source/PCG/Private/PCGSettings.cpp b/Engine/Plugins/Experimental/PCG/Source/PCG/Private/PCGSettings.cpp
--- a/Engine/Plugins/Experimental/PCG/Source/PCG/Private/PCGSettings.cpp
+++ b/Engine/Plugins/Experimental/PCG/Source/PCG/Private/PCGSettings.cpp
@@ -599,6 +599,13 @@ void UPCGSettingsInstance::OnSettingsChanged(UPCGSettings* InSettings, EPCGChang
 }
 #endif
 
+UPCGTrivialSettings::UPCGTrivialSettings()
+{
+#if WITH_EDITORONLY_DATA
+	bExposeToLibrary = false;
+#endif
+}
+
 FPCGElementPtr UPCGTrivialSettings::CreateElement() const
 {
 	return MakeShared<FPCGTrivialElement>();
diff --git a/Engine/Plugins/Experimental/PCG/Source/PCG/Public/Elements/PCGReroute.cpp b/Engine/Plugins/Experimental/PCG/Source/PCG/Public/Elements/PCGReroute.cpp
new file mode 100644
--- /dev/null
+++ b/Engine/Plugins/Experimental/PCG/Source/PCG/Public/Elements/PCGReroute.cpp
@@ -0,0 +1,80 @@
+// Copyright Epic Games, Inc. All Rights Reserved.
+
+#include "Elements/PCGReroute.h"
+
+#include "PCGContext.h"
+#include "PCGEdge.h"
+#include "PCGNode.h"
+#include "PCGPin.h"
+
+UPCGRerouteSettings::UPCGRerouteSettings()
+{
+#if WITH_EDITORONLY_DATA
+	bExposeToLibrary = false;
+#endif
+}
+
+TArray<FPCGPinProperties> UPCGRerouteSettings::InputPinProperties() const
+{
+	FPCGPinProperties PinProperties;
+	PinProperties.Label = PCGPinConstants::DefaultInputLabel;
+	PinProperties.bAllowMultipleConnections = false;
+	PinProperties.AllowedTypes = EPCGDataType::Any;
+	
+	if (UPCGNode* PCGNode = Cast<UPCGNode>(GetOuter()))
+	{
+		if (UPCGPin* InputPin = PCGNode->GetInputPin(PCGPinConstants::DefaultInputLabel))
+		{
+			if (InputPin->EdgeCount() > 0)
+			{
+				if (UPCGEdge* Edge = InputPin->Edges[0])
+				{
+					if (const UPCGPin* OtherOutputPin = Edge->GetOtherPin(InputPin))
+					{
+						PinProperties.AllowedTypes = OtherOutputPin->Properties.AllowedTypes;
+					}
+				}
+			}
+		}
+	}
+
+	return {PinProperties };
+}
+
+TArray<FPCGPinProperties> UPCGRerouteSettings::OutputPinProperties() const
+{
+	FPCGPinProperties PinProperties;
+	PinProperties.Label = PCGPinConstants::DefaultOutputLabel;
+	PinProperties.bAllowMultipleConnections = true;
+	PinProperties.AllowedTypes = EPCGDataType::Any;
+	
+	if (UPCGNode* PCGNode = Cast<UPCGNode>(GetOuter()))
+	{
+		if (UPCGPin* InputPin = PCGNode->GetInputPin(PCGPinConstants::DefaultInputLabel))
+		{
+			if (InputPin->EdgeCount() > 0)
+			{
+				if (UPCGEdge* Edge = InputPin->Edges[0])
+				{
+					if (const UPCGPin* OtherOutputPin = Edge->GetOtherPin(InputPin))
+					{
+						PinProperties.AllowedTypes = OtherOutputPin->Properties.AllowedTypes;
+					}
+				}
+			}
+		}
+	}
+
+	return { PinProperties };
+}
+
+FPCGElementPtr UPCGRerouteSettings::CreateElement() const
+{
+	return MakeShared<FPCGRerouteElement>();
+}
+
+bool FPCGRerouteElement::ExecuteInternal(FPCGContext* Context) const
+{
+	Context->OutputData = Context->InputData;
+	return true;
+}
diff --git a/Engine/Plugins/Experimental/PCG/Source/PCG/Public/Elements/PCGReroute.h b/Engine/Plugins/Experimental/PCG/Source/PCG/Public/Elements/PCGReroute.h
new file mode 100644
--- /dev/null
+++ b/Engine/Plugins/Experimental/PCG/Source/PCG/Public/Elements/PCGReroute.h
@@ -0,0 +1,39 @@
+// Copyright Epic Games, Inc. All Rights Reserved.
+
+#pragma once
+
+#include "CoreMinimal.h"
+
+#include "PCGElement.h"
+#include "PCGSettings.h"
+
+#include "PCGReroute.generated.h"
+
+UCLASS(BlueprintType, ClassGroup = (Procedural))
+class PCG_API UPCGRerouteSettings : public UPCGSettings
+{
+	GENERATED_BODY()
+
+public:
+	UPCGRerouteSettings();
+	
+	virtual bool HasDynamicPins() const override { return true; }
+	virtual TArray<FPCGPinProperties> InputPinProperties() const override;
+	virtual TArray<FPCGPinProperties> OutputPinProperties() const override;
+	
+protected:
+	virtual FPCGElementPtr CreateElement() const override;
+	
+#if WITH_EDITOR
+	virtual FName GetDefaultNodeName() const override { return FName("Reroute"); }
+#endif
+};
+
+class PCG_API FPCGRerouteElement : public FSimplePCGElement
+{
+public:
+	virtual bool IsCacheable(const UPCGSettings* InSettings) const override { return true; }
+
+protected:
+	virtual bool ExecuteInternal(FPCGContext* Context) const override;
+};
diff --git a/Engine/Plugins/Experimental/PCG/Source/PCG/Public/PCGNode.h b/Engine/Plugins/Experimental/PCG/Source/PCG/Public/PCGNode.h
--- a/Engine/Plugins/Experimental/PCG/Source/PCG/Public/PCGNode.h
+++ b/Engine/Plugins/Experimental/PCG/Source/PCG/Public/PCGNode.h
@@ -163,7 +163,8 @@ public:
 
 protected:
 	EPCGChangeType UpdatePins();
-	EPCGChangeType UpdatePins(TFunctionRef<UPCGPin* (UPCGNode*)> PinAllocator);
+	EPCGChangeType UpdatePins(TFunctionRef<UPCGPin* (UPCGNode*)> PinAllocator, const UPCGNode* FromNode = nullptr);
+	EPCGChangeType UpdateDynamicPins(const UPCGNode* FromNode = nullptr);
 
 	// When we create a new graph, we initialize the input/output nodes as default, with default pins.
 	// Those default pins are not serialized, therefore if we change the default pins, combined with the use
diff --git a/Engine/Plugins/Experimental/PCG/Source/PCG/Public/PCGSettings.h b/Engine/Plugins/Experimental/PCG/Source/PCG/Public/PCGSettings.h
--- a/Engine/Plugins/Experimental/PCG/Source/PCG/Public/PCGSettings.h
+++ b/Engine/Plugins/Experimental/PCG/Source/PCG/Public/PCGSettings.h
@@ -135,6 +135,9 @@ public:
 	/** For symmetry reason, do the same with output pins. For now forward just the call to OutputPinProperties */
 	TArray<FPCGPinProperties> AllOutputPinProperties() const;
 
+	/** If the node has any dynamic pins that can change based on input or settings */
+	virtual bool HasDynamicPins() const { return false; }
+
 	// Internal functions, should not be used by any user.
 	// Return a different subset for for input/output pin properties, in case of a default object.
 	virtual TArray<FPCGPinProperties> DefaultInputPinProperties() const;
@@ -200,7 +203,7 @@ public:
 
 #if WITH_EDITORONLY_DATA
 	UPROPERTY(BlueprintReadWrite, EditAnywhere, Category = AssetInfo, AssetRegistrySearchable)
-	bool bExposeToLibrary = false;
+	bool bExposeToLibrary = true;
 
 	UPROPERTY(BlueprintReadWrite, EditAnywhere, Category = AssetInfo, AssetRegistrySearchable)
 	FText Category;
@@ -310,6 +313,9 @@ UCLASS(BlueprintType, ClassGroup = (Procedural))
 class PCG_API UPCGTrivialSettings : public UPCGSettings
 {
 	GENERATED_BODY()
+	
+public:
+	UPCGTrivialSettings();
 
 protected:
 	//~UPCGSettings implementation
diff --git a/Engine/Plugins/Experimental/PCG/Source/PCGEditor/Private/PCGEditorGraph.cpp b/Engine/Plugins/Experimental/PCG/Source/PCGEditor/Private/PCGEditorGraph.cpp
--- a/Engine/Plugins/Experimental/PCG/Source/PCGEditor/Private/PCGEditorGraph.cpp
+++ b/Engine/Plugins/Experimental/PCG/Source/PCGEditor/Private/PCGEditorGraph.cpp
@@ -2,15 +2,17 @@
 
 #include "PCGEditorGraph.h"
 
-#include "PCGGraph.h"
 #include "PCGEdge.h"
 #include "PCGEditorGraphNode.h"
 #include "PCGEditorGraphNodeInput.h"
 #include "PCGEditorGraphNodeOutput.h"
+#include "PCGEditorGraphNodeReroute.h"
 #include "PCGEditorModule.h"
+#include "PCGGraph.h"
+#include "PCGPin.h"
+#include "Elements/PCGReroute.h"
 
 #include "EdGraph/EdGraphPin.h"
-#include "PCGPin.h"
 
 void UPCGEditorGraph::InitFromNodeGraph(UPCGGraph* InPCGGraph)
 {
@@ -36,11 +38,27 @@ void UPCGEditorGraph::InitFromNodeGraph(UPCGGraph* InPCGGraph)
 
 	for (UPCGNode* PCGNode : PCGGraph->GetNodes())
 	{
-		FGraphNodeCreator<UPCGEditorGraphNode> NodeCreator(*this);
-		UPCGEditorGraphNode* GraphNode = NodeCreator.CreateNode(bSelectNewNode);
-		GraphNode->Construct(PCGNode);
-		NodeCreator.Finalize();
-		NodeLookup.Add(PCGNode, GraphNode);
+		if (!IsValid(PCGNode))
+		{
+			continue;	
+		}
+		
+		if (Cast<UPCGRerouteSettings>(PCGNode->GetSettings()))
+		{
+			FGraphNodeCreator<UPCGEditorGraphNodeReroute> NodeCreator(*this);
+			UPCGEditorGraphNodeReroute* RerouteGraphNode = NodeCreator.CreateNode(bSelectNewNode);
+			RerouteGraphNode->Construct(PCGNode);
+			NodeCreator.Finalize();
+			NodeLookup.Add(PCGNode, RerouteGraphNode);
+		}
+		else
+		{
+			FGraphNodeCreator<UPCGEditorGraphNode> NodeCreator(*this);
+			UPCGEditorGraphNode* GraphNode = NodeCreator.CreateNode(bSelectNewNode);
+			GraphNode->Construct(PCGNode);
+			NodeCreator.Finalize();
+			NodeLookup.Add(PCGNode, GraphNode);
+		}
 	}
 
 	for (const auto& NodeLookupIt : NodeLookup)
diff --git a/Engine/Plugins/Experimental/PCG/Source/PCGEditor/Private/PCGEditorGraphNode.h b/Engine/Plugins/Experimental/PCG/Source/PCGEditor/Private/PCGEditorGraphNode.h
--- a/Engine/Plugins/Experimental/PCG/Source/PCGEditor/Private/PCGEditorGraphNode.h
+++ b/Engine/Plugins/Experimental/PCG/Source/PCGEditor/Private/PCGEditorGraphNode.h
@@ -22,6 +22,6 @@ public:
 	virtual FText GetNodeTitle(ENodeTitleType::Type TitleType) const override;
 	virtual void GetNodeContextMenuActions(UToolMenu* Menu, class UGraphNodeContextMenuContext* Context) const override;
 	virtual void AllocateDefaultPins() override;
-	virtual void OnRenameNode(const FString& NewName);
+	virtual void OnRenameNode(const FString& NewName) override;
 	// ~End UEdGraphNode interface
 };
diff --git a/Engine/Plugins/Experimental/PCG/Source/PCGEditor/Private/PCGEditorGraphNodeBase.cpp b/Engine/Plugins/Experimental/PCG/Source/PCGEditor/Private/PCGEditorGraphNodeBase.cpp
--- a/Engine/Plugins/Experimental/PCG/Source/PCGEditor/Private/PCGEditorGraphNodeBase.cpp
+++ b/Engine/Plugins/Experimental/PCG/Source/PCGEditor/Private/PCGEditorGraphNodeBase.cpp
@@ -8,14 +8,13 @@
 #include "PCGEditorGraphSchema.h"
 #include "PCGEditorSettings.h"
 #include "PCGGraph.h"
+#include "PCGPin.h"
 
-#include "Misc/TransactionObjectEvent.h"
 #include "GraphEditorActions.h"
-#include "PCGPin.h"
 #include "ToolMenu.h"
 #include "ToolMenuSection.h"
+#include "Misc/TransactionObjectEvent.h"
 #include "Widgets/Colors/SColorPicker.h"
-#include "PCGEditorCommands.h"
 
 #define LOCTEXT_NAMESPACE "PCGEditorGraphNodeBase"
 
@@ -196,11 +195,37 @@ void UPCGEditorGraphNodeBase::PostPaste()
 		PCGNode->OnNodeChangedDelegate.AddUObject(this, &UPCGEditorGraphNodeBase::OnNodeChanged);
 		PCGNode->PositionX = NodePosX;
 		PCGNode->PositionY = NodePosY;
+
+		if (const UPCGSettings* Settings = PCGNode->GetSettings())
+		{
+			if (Settings->HasDynamicPins())
+			{
+				PCGNode->UpdateAfterSettingsChangeDuringCreation();
+			}
+		}
 	}
 
 	bDisableReconstructFromNode = false;
 }
 
+void UPCGEditorGraphNodeBase::EnableDeferredReconstruct()
+{
+	ensure(DeferredReconstructCounter >= 0);
+	++DeferredReconstructCounter;
+}
+
+void UPCGEditorGraphNodeBase::DisableDeferredReconstruct()
+{
+	ensure(DeferredReconstructCounter > 0);
+	--DeferredReconstructCounter;
+
+	if (DeferredReconstructCounter == 0 && bDeferredReconstruct)
+	{
+		ReconstructNode();
+		bDeferredReconstruct = false;
+	}
+}
+
 void UPCGEditorGraphNodeBase::RebuildEdgesFromPins()
 {
 	check(PCGNode);
@@ -249,7 +274,7 @@ void UPCGEditorGraphNodeBase::OnNodeChanged(UPCGNode* InNode, EPCGChangeType Cha
 				}
 			}
 		}
-		
+
 		if (!!(ChangeType & (EPCGChangeType::Structural | EPCGChangeType::Node | EPCGChangeType::Edge | EPCGChangeType::Cosmetic)))
 		{
 			ReconstructNode();
@@ -285,6 +310,12 @@ void UPCGEditorGraphNodeBase::ReconstructNode()
 		return;
 	}
 
+	if (DeferredReconstructCounter > 0)
+	{
+		bDeferredReconstruct = true;
+		return;
+	}
+
 	// Store copy of old pins
 	TArray<UEdGraphPin*> OldPins = MoveTemp(Pins);
 	Pins.Reset();
@@ -308,6 +339,20 @@ void UPCGEditorGraphNodeBase::ReconstructNode()
 		RemovePin(OldPin);
 	}
 
+	// Generate new links
+	// TODO: we should either keep a map in the PCGEditorGraph or do this elsewhere
+	// TODO: this will not work if we have non-PCG nodes in the graph
+	if (PCGNode)
+	{
+		for (UEdGraphPin* Pin : Pins)
+		{
+			Pin->BreakAllPinLinks();
+		}
+		
+		UPCGEditorGraph* PCGEditorGraph = CastChecked<UPCGEditorGraph>(GetGraph());
+		PCGEditorGraph->CreateLinks(this, /*bCreateInbound=*/true, /*bCreateOutbound=*/true);
+	}
+
 	// Notify editor
 	OnNodeChangedDelegate.ExecuteIfBound();
 }
diff --git a/Engine/Plugins/Experimental/PCG/Source/PCGEditor/Private/PCGEditorGraphNodeBase.h b/Engine/Plugins/Experimental/PCG/Source/PCGEditor/Private/PCGEditorGraphNodeBase.h
--- a/Engine/Plugins/Experimental/PCG/Source/PCGEditor/Private/PCGEditorGraphNodeBase.h
+++ b/Engine/Plugins/Experimental/PCG/Source/PCGEditor/Private/PCGEditorGraphNodeBase.h
@@ -38,7 +38,7 @@ public:
 	virtual FText GetTooltipText() const override;
 	virtual void GetPinHoverText(const UEdGraphPin& Pin, FString& HoverTextOut) const override;
 	virtual UObject* GetJumpTargetForDoubleClick() const override;
-	virtual void OnUpdateCommentText(const FString& NewComment);
+	virtual void OnUpdateCommentText(const FString& NewComment) override;
 	virtual void OnCommentBubbleToggled(bool bInCommentBubbleVisible) override;
 	// ~End UEdGraphNode interface
 
@@ -48,6 +48,11 @@ public:
 	void PostPaste();
 	void SetInspected(bool InIsInspecting) { bIsInspected = InIsInspecting; }
 	bool GetInspected() const { return bIsInspected; }
+	
+	/** Increase deferred reconstruct counter, calls to ReconstructNode will flag reconstruct to happen when count hits zero */
+	void EnableDeferredReconstruct();
+	/** Decrease deferred reconstruct counter, ReconstructNode will be called if counter hits zero and the node is flagged for reconstruction  */
+	void DisableDeferredReconstruct();
 
 	DECLARE_DELEGATE(FOnPCGEditorGraphNodeChanged);
 	FOnPCGEditorGraphNodeChanged OnNodeChangedDelegate;
@@ -72,6 +77,8 @@ protected:
 	UPROPERTY()
 	TObjectPtr<UPCGNode> PCGNode = nullptr;
 
+	int32 DeferredReconstructCounter = 0;
+	bool bDeferredReconstruct = false;
 	bool bDisableReconstructFromNode = false;
 	bool bIsInspected = false;
 };
diff --git a/Engine/Plugins/Experimental/PCG/Source/PCGEditor/Private/PCGEditorGraphNodeFactory.cpp b/Engine/Plugins/Experimental/PCG/Source/PCGEditor/Private/PCGEditorGraphNodeFactory.cpp
--- a/Engine/Plugins/Experimental/PCG/Source/PCGEditor/Private/PCGEditorGraphNodeFactory.cpp
+++ b/Engine/Plugins/Experimental/PCG/Source/PCGEditor/Private/PCGEditorGraphNodeFactory.cpp
@@ -3,17 +3,56 @@
 #include "PCGEditorGraphNodeFactory.h"
 
 #include "PCGEditorGraphNodeBase.h"
+#include "PCGEditorGraphNodeReroute.h"
+#include "PCGNode.h"
 #include "SPCGEditorGraphNode.h"
+#include "Elements/PCGReroute.h"
+
+#include "SGraphNodeKnot.h"
+
+class SPCGEditorGraphNodeKnot : public SGraphNodeKnot
+{
+public:
+	SLATE_BEGIN_ARGS(SPCGEditorGraphNodeKnot) {}
+	SLATE_END_ARGS();
+
+	void Construct(const FArguments& InArgs, UPCGEditorGraphNodeBase* InPCGGraphNode)
+	{
+		SGraphNodeKnot::Construct(SGraphNodeKnot::FArguments(), InPCGGraphNode);
+		InPCGGraphNode->OnNodeChangedDelegate.BindSP(this, &SPCGEditorGraphNodeKnot::OnNodeChanged);
+	}
+
+private:
+	void OnNodeChanged()
+	{
+		UpdateGraphNode();
+	}
+};
+
 
 TSharedPtr<SGraphNode> FPCGEditorGraphNodeFactory::CreateNode(UEdGraphNode* InNode) const
 {
 	if (UPCGEditorGraphNodeBase* GraphNode = Cast<UPCGEditorGraphNodeBase>(InNode))
 	{
-		TSharedRef<SGraphNode> VisualNode =
-			SNew(SPCGEditorGraphNode, GraphNode);
-
+		TSharedPtr<SGraphNode> VisualNode;
+		
+		const UPCGNode* PCGNode = GraphNode->GetPCGNode();
+		if (PCGNode && Cast<UPCGRerouteSettings>(PCGNode->GetSettings()))
+		{
+			SAssignNew(VisualNode, SPCGEditorGraphNodeKnot, GraphNode);
+		}
+		else
+		{
+			SAssignNew(VisualNode, SPCGEditorGraphNode, GraphNode);
+		}
+		
 		VisualNode->SlatePrepass();
 
+		return VisualNode.ToSharedRef();
+	}
+	else if (UPCGEditorGraphNodeReroute* RerouteNode = Cast<UPCGEditorGraphNodeReroute>(InNode))
+	{
+		TSharedRef<SGraphNode> VisualNode = SNew(SGraphNodeKnot, RerouteNode);
 		return VisualNode;
 	}
 
diff --git a/Engine/Plugins/Experimental/PCG/Source/PCGEditor/Private/PCGEditorGraphNodeReroute.cpp b/Engine/Plugins/Experimental/PCG/Source/PCGEditor/Private/PCGEditorGraphNodeReroute.cpp
new file mode 100644
--- /dev/null
+++ b/Engine/Plugins/Experimental/PCG/Source/PCGEditor/Private/PCGEditorGraphNodeReroute.cpp
@@ -0,0 +1,75 @@
+// Copyright Epic Games, Inc. All Rights Reserved.
+
+#include "PCGEditorGraphNodeReroute.h"
+
+#include "EdGraph/EdGraphPin.h"
+
+FText UPCGEditorGraphNodeReroute::GetNodeTitle(ENodeTitleType::Type TitleType) const
+{
+	return FText::FromName(TEXT("Reroute"));
+}
+
+bool UPCGEditorGraphNodeReroute::ShouldOverridePinNames() const
+{
+	return true;
+}
+
+FText UPCGEditorGraphNodeReroute::GetPinNameOverride(const UEdGraphPin& Pin) const
+{
+	return FText::GetEmpty();
+}
+
+bool UPCGEditorGraphNodeReroute::CanSplitPin(const UEdGraphPin* Pin) const
+{
+	return false;
+}
+
+bool UPCGEditorGraphNodeReroute::ShouldDrawNodeAsControlPointOnly(int32& OutInputPinIndex, int32& OutOutputPinIndex) const
+{
+	OutInputPinIndex = 0;
+	OutOutputPinIndex = 1;
+	return true;
+}
+
+FText UPCGEditorGraphNodeReroute::GetTooltipText() const
+{
+	return FText::GetEmpty();
+}
+
+UEdGraphPin* UPCGEditorGraphNodeReroute::GetPassThroughPin(const UEdGraphPin* FromPin) const
+{
+	if (FromPin == GetInputPin())
+	{
+		return GetOutputPin();
+	}
+	else
+	{
+		return GetInputPin();
+	}
+}
+
+UEdGraphPin* UPCGEditorGraphNodeReroute::GetInputPin() const
+{
+	for (UEdGraphPin* Pin : Pins)
+	{
+		if (Pin->Direction == EGPD_Input)
+		{
+			return Pin;
+		}
+	}
+
+	return nullptr;
+}
+
+UEdGraphPin* UPCGEditorGraphNodeReroute::GetOutputPin() const
+{
+	for (UEdGraphPin* Pin : Pins)
+	{
+		if (Pin->Direction == EGPD_Output)
+		{
+			return Pin;
+		}
+	}
+
+	return nullptr;
+}
diff --git a/Engine/Plugins/Experimental/PCG/Source/PCGEditor/Private/PCGEditorGraphNodeReroute.h b/Engine/Plugins/Experimental/PCG/Source/PCGEditor/Private/PCGEditorGraphNodeReroute.h
new file mode 100644
--- /dev/null
+++ b/Engine/Plugins/Experimental/PCG/Source/PCGEditor/Private/PCGEditorGraphNodeReroute.h
@@ -0,0 +1,27 @@
+// Copyright Epic Games, Inc. All Rights Reserved.
+
+#pragma once
+
+#include "PCGEditorGraphNode.h"
+
+#include "PCGEditorGraphNodeReroute.generated.h"
+
+UCLASS()
+class UPCGEditorGraphNodeReroute : public UPCGEditorGraphNode
+{
+	GENERATED_BODY()
+
+public:
+	// ~Begin UEdGraphNode interface
+	virtual FText GetNodeTitle(ENodeTitleType::Type TitleType) const override;
+	virtual bool ShouldOverridePinNames() const override;
+	virtual FText GetPinNameOverride(const UEdGraphPin& Pin) const override;
+	virtual bool CanSplitPin(const UEdGraphPin* Pin) const override;
+	virtual bool ShouldDrawNodeAsControlPointOnly(int32& OutInputPinIndex, int32& OutOutputPinIndex) const override;
+	virtual FText GetTooltipText() const override;
+	virtual UEdGraphPin* GetPassThroughPin(const UEdGraphPin* FromPin) const override;
+	// ~End UEdGraphNode interface
+
+	UEdGraphPin* GetInputPin() const;
+	UEdGraphPin* GetOutputPin() const;
+};
diff --git a/Engine/Plugins/Experimental/PCG/Source/PCGEditor/Private/PCGEditorGraphSchema.cpp b/Engine/Plugins/Experimental/PCG/Source/PCGEditor/Private/PCGEditorGraphSchema.cpp
--- a/Engine/Plugins/Experimental/PCG/Source/PCGEditor/Private/PCGEditorGraphSchema.cpp
+++ b/Engine/Plugins/Experimental/PCG/Source/PCGEditor/Private/PCGEditorGraphSchema.cpp
@@ -10,6 +10,7 @@
 #include "PCGEditorCommon.h"
 #include "PCGEditorGraph.h"
 #include "PCGEditorGraphNodeBase.h"
+#include "PCGEditorGraphNodeReroute.h"
 #include "PCGEditorGraphSchemaActions.h"
 #include "PCGEditorSettings.h"
 #include "PCGEditorUtils.h"
@@ -113,6 +114,11 @@ const FPinConnectionResponse UPCGEditorGraphSchema::CanCreateConnection(const UE
 		return FPinConnectionResponse(CONNECT_RESPONSE_DISALLOW, LOCTEXT("ConnectionTypesIncompatible", "Pins are incompatible"));
 	}
 
+	if (!InputPin->AllowMultipleConnections() && InputPin->EdgeCount() > 0)
+	{
+		return FPinConnectionResponse((A->Direction == EGPD_Output) ? CONNECT_RESPONSE_BREAK_OTHERS_B : CONNECT_RESPONSE_BREAK_OTHERS_A, LOCTEXT("ConnectionBreakExisting", "Break existing connection?"));
+	}
+
 	return FPinConnectionResponse();
 }
 
@@ -142,13 +148,7 @@ bool UPCGEditorGraphSchema::TryCreateConnection(UEdGraphPin* InA, UEdGraphPin* I
 		UPCGGraph* PCGGraph = PCGNodeA->GetGraph();
 		check(PCGGraph);
 
-		const bool bReconstructNodeB = PCGGraph->AddLabeledEdge(PCGNodeA, A->PinName, PCGNodeB, B->PinName);
-
-		// If AddLabeledEdge return true, then it means that the ToPin triggered a change that requires a full node reconstruct
-		if (bReconstructNodeB)
-		{
-			PCGGraphNodeB->ReconstructNode();
-		}
+		PCGGraph->AddLabeledEdge(PCGNodeA, A->PinName, PCGNodeB, B->PinName);
 	}
 
 	return bModified;
@@ -218,13 +218,16 @@ void UPCGEditorGraphSchema::GetNativeElementActions(FGraphActionMenuBuilder& Act
 	{
 		if (const UPCGSettings* PCGSettings = SettingsClass->GetDefaultObject<UPCGSettings>())
 		{
-			const FText MenuDesc = FText::FromName(PCGSettings->GetDefaultNodeName());
-			const FText Category = StaticEnum<EPCGSettingsType>()->GetDisplayNameTextByValue(static_cast<__underlying_type(EPCGSettingsType)>(PCGSettings->GetType()));
-			const FText Description = PCGSettings->GetNodeTooltipText();
+			if (PCGSettings->bExposeToLibrary)
+			{
+				const FText MenuDesc = FText::FromName(PCGSettings->GetDefaultNodeName());
+				const FText Category = StaticEnum<EPCGSettingsType>()->GetDisplayNameTextByValue(static_cast<__underlying_type(EPCGSettingsType)>(PCGSettings->GetType()));
+				const FText Description = PCGSettings->GetNodeTooltipText();
 
-			TSharedPtr<FPCGEditorGraphSchemaAction_NewNativeElement> NewAction(new FPCGEditorGraphSchemaAction_NewNativeElement(Category, MenuDesc, Description, 0));
-			NewAction->SettingsClass = SettingsClass;
-			ActionMenuBuilder.AddAction(NewAction);
+				TSharedPtr<FPCGEditorGraphSchemaAction_NewNativeElement> NewAction(new FPCGEditorGraphSchemaAction_NewNativeElement(Category, MenuDesc, Description, 0));
+				NewAction->SettingsClass = SettingsClass;
+				ActionMenuBuilder.AddAction(NewAction);
+			}
 		}
 	}
 }
@@ -311,12 +314,20 @@ void UPCGEditorGraphSchema::GetSubgraphElementActions(FGraphActionMenuBuilder& A
 void UPCGEditorGraphSchema::GetExtraElementActions(FGraphActionMenuBuilder& ActionMenuBuilder) const
 {
 	// Comment action
-	const FText MenuDesc = LOCTEXT("PCGAddComment", "Add Comment...");
-	const FText Category;
-	const FText Description = LOCTEXT("PCGAddCommentTooltip", "Create a resizable comment box.");
+	const FText CommentMenuDesc = LOCTEXT("PCGAddComment", "Add Comment...");
+	const FText CommentCategory;
+	const FText CommentDescription = LOCTEXT("PCGAddCommentTooltip", "Create a resizable comment box.");
 
-	const TSharedPtr<FPCGEditorGraphSchemaAction_NewComment> NewCommentAction(new FPCGEditorGraphSchemaAction_NewComment(Category, MenuDesc, Description, 0));
+	const TSharedPtr<FPCGEditorGraphSchemaAction_NewComment> NewCommentAction(new FPCGEditorGraphSchemaAction_NewComment(CommentCategory, CommentMenuDesc, CommentDescription, 0));
 	ActionMenuBuilder.AddAction(NewCommentAction);
+
+	// Reroute action
+	const FText RerouteMenuDesc = LOCTEXT("PCGAddRerouteNode", "Add Reroute Node");
+	const FText RerouteCategory;
+	const FText RerouteDescription = LOCTEXT("PCGAddRerouteNodeTooltip", "Add a reroute node, aka knot.");
+
+	const TSharedPtr<FPCGEditorGraphSchemaAction_NewReroute> NewRerouteAction(new FPCGEditorGraphSchemaAction_NewReroute(RerouteCategory, RerouteMenuDesc, RerouteDescription, 0));
+	ActionMenuBuilder.AddAction(NewRerouteAction);
 }
 
 void UPCGEditorGraphSchema::DroppedAssetsOnGraph(const TArray<FAssetData>& Assets, const FVector2D& GraphPosition, UEdGraph* Graph) const
@@ -395,6 +406,38 @@ void UPCGEditorGraphSchema::GetAssetsGraphHoverMessage(const TArray<FAssetData>&
 	OutOkIcon = false;
 }
 
+void UPCGEditorGraphSchema::OnPinConnectionDoubleCicked(UEdGraphPin* PinA, UEdGraphPin* PinB, const FVector2D& GraphPosition) const
+{
+	const FScopedTransaction Transaction(*FPCGEditorCommon::ContextIdentifier, LOCTEXT("PCGCreateRerouteNodeOnWire", "Create Reroute Node"), nullptr);
+
+	const FVector2D NodeSpacerSize(42.0f, 24.0f);
+	const FVector2D KnotTopLeft = GraphPosition - (NodeSpacerSize * 0.5f);
+
+	UEdGraph* EditorGraph = PinA->GetOwningNode()->GetGraph();
+	EditorGraph->Modify();
+
+	FPCGEditorGraphSchemaAction_NewReroute Action;
+
+	if (UPCGEditorGraphNodeReroute* RerouteNode = Cast<UPCGEditorGraphNodeReroute>(Action.PerformAction(EditorGraph, nullptr, KnotTopLeft, /*bSelectNewNode=*/true)))
+	{
+		UEdGraphNode* SourceGraphNode = PinA->GetOwningNode();
+		UEdGraphNode* TargetGraphNode = PinB->GetOwningNode();
+
+		UPCGEditorGraphNodeBase* SourcePCGGraphNode = CastChecked<UPCGEditorGraphNodeBase>(SourceGraphNode);
+		UPCGEditorGraphNodeBase* TargetPCGGraphNode = CastChecked<UPCGEditorGraphNodeBase>(TargetGraphNode);
+
+		// We need to disable full node reconstruction to make sure the pins are valid when creating the connections.
+		SourcePCGGraphNode->EnableDeferredReconstruct();
+		TargetPCGGraphNode->EnableDeferredReconstruct();
+		
+		BreakSinglePinLink(PinA, PinB);
+		TryCreateConnection(PinA, (PinA->Direction == EGPD_Output) ? RerouteNode->GetInputPin() : RerouteNode->GetOutputPin());
+		TryCreateConnection(PinB, (PinB->Direction == EGPD_Output) ? RerouteNode->GetInputPin() : RerouteNode->GetOutputPin());
+
+		SourcePCGGraphNode->DisableDeferredReconstruct();
+		TargetPCGGraphNode->DisableDeferredReconstruct();
+	}
+}
 
 FPCGEditorConnectionDrawingPolicy::FPCGEditorConnectionDrawingPolicy(int32 InBackLayerID, int32 InFrontLayerID, float InZoomFactor, const FSlateRect& InClippingRect, FSlateWindowElementList& InDrawElements, UEdGraph* InGraph)
 	: FConnectionDrawingPolicy(InBackLayerID, InFrontLayerID, InZoomFactor, InClippingRect, InDrawElements)
diff --git a/Engine/Plugins/Experimental/PCG/Source/PCGEditor/Private/PCGEditorGraphSchema.h b/Engine/Plugins/Experimental/PCG/Source/PCGEditor/Private/PCGEditorGraphSchema.h
--- a/Engine/Plugins/Experimental/PCG/Source/PCGEditor/Private/PCGEditorGraphSchema.h
+++ b/Engine/Plugins/Experimental/PCG/Source/PCGEditor/Private/PCGEditorGraphSchema.h
@@ -20,7 +20,7 @@ public:
 	void GetPaletteActions(FGraphActionMenuBuilder& ActionMenuBuilder, const EPCGElementType InPCGElementTypeFilter) const;
 
 	//~ Begin EdGraphSchema Interface
-	virtual void GetGraphContextActions(FGraphContextMenuBuilder& ContextMenuBuilder) const;
+	virtual void GetGraphContextActions(FGraphContextMenuBuilder& ContextMenuBuilder) const override;
 	virtual FLinearColor GetPinTypeColor(const FEdGraphPinType& PinType) const override;
 	virtual const FPinConnectionResponse CanCreateConnection(const UEdGraphPin* A, const UEdGraphPin* B) const override;
 	virtual bool TryCreateConnection(UEdGraphPin* A, UEdGraphPin* B) const override;
@@ -29,6 +29,7 @@ public:
 	virtual FConnectionDrawingPolicy* CreateConnectionDrawingPolicy(int32 InBackLayerID, int32 InFrontLayerID, float InZoomFactor, const FSlateRect& InClippingRect, class FSlateWindowElementList& InDrawElements, class UEdGraph* InGraphObj) const override;
 	virtual void DroppedAssetsOnGraph(const TArray<FAssetData>& Assets, const FVector2D& GraphPosition, UEdGraph* Graph) const override;
 	virtual void GetAssetsGraphHoverMessage(const TArray<FAssetData>& Assets, const UEdGraph* HoverGraph, FString& OutTooltipText, bool& OutOkIcon) const override;
+	virtual void OnPinConnectionDoubleCicked(UEdGraphPin* PinA, UEdGraphPin* PinB, const FVector2D& GraphPosition) const override;
 	//~ End EdGraphSchema Interface
 
 private:
diff --git a/Engine/Plugins/Experimental/PCG/Source/PCGEditor/Private/PCGEditorGraphSchemaActions.cpp b/Engine/Plugins/Experimental/PCG/Source/PCGEditor/Private/PCGEditorGraphSchemaActions.cpp
--- a/Engine/Plugins/Experimental/PCG/Source/PCGEditor/Private/PCGEditorGraphSchemaActions.cpp
+++ b/Engine/Plugins/Experimental/PCG/Source/PCGEditor/Private/PCGEditorGraphSchemaActions.cpp
@@ -3,9 +3,11 @@
 #include "PCGEditorGraphSchemaActions.h"
 
 #include "Elements/PCGExecuteBlueprint.h"
+#include "Elements/PCGReroute.h"
 #include "PCGEditorCommon.h"
 #include "PCGEditorGraph.h"
 #include "PCGEditorGraphNode.h"
+#include "PCGEditorGraphNodeReroute.h"
 #include "PCGEditorModule.h"
 #include "PCGGraph.h"
 #include "PCGSubgraph.h"
@@ -357,4 +359,44 @@ UEdGraphNode* FPCGEditorGraphSchemaAction_NewComment::PerformAction(class UEdGra
 	return NewNode;
 }
 
+UEdGraphNode* FPCGEditorGraphSchemaAction_NewReroute::PerformAction(class UEdGraph* ParentGraph, UEdGraphPin* FromPin, const FVector2D Location, bool bSelectNewNode)
+{
+	UPCGEditorGraph* EditorGraph = Cast<UPCGEditorGraph>(ParentGraph);
+	if (!EditorGraph)
+	{
+		UE_LOG(LogPCGEditor, Error, TEXT("Invalid EditorGraph"));
+		return nullptr;
+	}
+
+	UPCGGraph* PCGGraph = EditorGraph->GetPCGGraph();
+	if (!PCGGraph)
+	{
+		UE_LOG(LogPCGEditor, Error, TEXT("Invalid PCGGraph"));
+		return nullptr;
+	}
+
+	const FScopedTransaction Transaction(*FPCGEditorCommon::ContextIdentifier, LOCTEXT("PCGEditorNewReroute", "PCG Editor: New Reorute Node"), nullptr);
+	EditorGraph->Modify();
+
+	UPCGSettings* DefaultNodeSettings = nullptr;
+	UPCGNode* NewPCGNode = PCGGraph->AddNodeOfType(UPCGRerouteSettings::StaticClass(), DefaultNodeSettings);
+
+	FGraphNodeCreator<UPCGEditorGraphNodeReroute> NodeCreator(*EditorGraph);
+	UPCGEditorGraphNodeReroute* NewNode = NodeCreator.CreateUserInvokedNode(bSelectNewNode);
+	NewNode->Construct(NewPCGNode);
+	NewNode->NodePosX = Location.X;
+	NewNode->NodePosY = Location.Y;
+	NodeCreator.Finalize();
+
+	NewPCGNode->PositionX = Location.X;
+	NewPCGNode->PositionY = Location.Y;
+
+	if (FromPin)
+	{
+		NewNode->AutowireNewNode(FromPin);
+	}
+
+	return NewNode;
+}
+
 #undef LOCTEXT_NAMESPACE
diff --git a/Engine/Plugins/Experimental/PCG/Source/PCGEditor/Private/PCGEditorGraphSchemaActions.h b/Engine/Plugins/Experimental/PCG/Source/PCGEditor/Private/PCGEditorGraphSchemaActions.h
--- a/Engine/Plugins/Experimental/PCG/Source/PCGEditor/Private/PCGEditorGraphSchemaActions.h
+++ b/Engine/Plugins/Experimental/PCG/Source/PCGEditor/Private/PCGEditorGraphSchemaActions.h
@@ -145,3 +145,24 @@ struct FPCGEditorGraphSchemaAction_NewComment : public FEdGraphSchemaAction
 	virtual UEdGraphNode* PerformAction(class UEdGraph* ParentGraph, UEdGraphPin* FromPin, const FVector2D Location, bool bSelectNewNode = true) override;
 	// End of FEdGraphSchemaAction interface
 };
+
+USTRUCT()
+struct FPCGEditorGraphSchemaAction_NewReroute : public FEdGraphSchemaAction
+{
+	GENERATED_BODY()
+
+	// Inherit the base class's constructors
+	using FEdGraphSchemaAction::FEdGraphSchemaAction;
+
+	// Simple type info
+	static FName StaticGetTypeId()
+	{
+		static FName Type("FPCGEditorGraphSchemaAction_NewReroute");
+		return Type;
+	}
+
+	// FEdGraphSchemaAction interface
+	virtual FName GetTypeId() const override { return StaticGetTypeId(); }
+	virtual UEdGraphNode* PerformAction(class UEdGraph* ParentGraph, UEdGraphPin* FromPin, const FVector2D Location, bool bSelectNewNode = true) override;
+	// End of FEdGraphSchemaAction interface
+};
