[PCG] Making sure that we do not have an immediate recursive subgraph
Will not catch all recusions, but at least prevent when we set the current graph to a subgraph node.

diff --git a/Engine/Plugins/Experimental/PCG/Source/PCG/Private/Graph/PCGGraphCompiler.cpp b/Engine/Plugins/Experimental/PCG/Source/PCG/Private/Graph/PCGGraphCompiler.cpp
--- a/Engine/Plugins/Experimental/PCG/Source/PCG/Private/Graph/PCGGraphCompiler.cpp
+++ b/Engine/Plugins/Experimental/PCG/Source/PCG/Private/Graph/PCGGraphCompiler.cpp
@@ -40,6 +40,19 @@ TArray<FPCGGraphTask> FPCGGraphCompiler::CompileGraph(UPCGGraph* InGraph, FPCGTa
 		const UPCGBaseSubgraphNode* SubgraphNode = Cast<const UPCGBaseSubgraphNode>(Node);
 		UPCGGraph* Subgraph = SubgraphNode ? SubgraphNode->GetSubgraph() : nullptr;
 
+		// Only catch immediate infinite recursion.
+		// TODO: Add a better mecanism for detecting more complex cyclic recursions
+		// Like Graph A has subgraph node with Graph B, and Graph B has a subgraph node with Graph A (A -> B -> A)
+		// or A -> B -> C -> A, etc...
+		// NOTE (for the person that would work on it), keeping a stack of all the subgraphs is not enough, as we could already have compiled graph
+		// A, with an inclusion to graph B, but B has not yet a subgraph to A. When adding subgraph node to A in B, we will only recompile B, since A is already
+		// compiled and cached... We need to store a subgraph dependency chain to detect those more complex cases.
+		if (Subgraph == InGraph)
+		{
+			UE_LOG(LogPCG, Error, TEXT("[FPCGGraphCompiler::CompileGraph] %s cannot include itself as a subgraph, subgraph will not be executed."), *InGraph->GetName());
+			return TArray<FPCGGraphTask>();
+		}
+
 		const bool bIsNonDynamicAndNonDisabledSubgraphNode = (SubgraphNode != nullptr && !SubgraphNode->bDynamicGraph && Subgraph && SubgraphNode->GetSettings() && SubgraphNode->GetSettings()->bEnabled);
 		if (bIsNonDynamicAndNonDisabledSubgraphNode)
 		{
