[PCG] Make sure to disable spatial query if the filter data is param

diff --git a/Engine/Plugins/Experimental/PCG/Source/PCG/Private/Elements/PCGPointFilter.cpp b/Engine/Plugins/Experimental/PCG/Source/PCG/Private/Elements/PCGPointFilter.cpp
--- a/Engine/Plugins/Experimental/PCG/Source/PCG/Private/Elements/PCGPointFilter.cpp
+++ b/Engine/Plugins/Experimental/PCG/Source/PCG/Private/Elements/PCGPointFilter.cpp
@@ -2,14 +2,15 @@
 
 #include "Elements/PCGPointFilter.h"
 
-#include "Metadata/Accessors/IPCGAttributeAccessorTpl.h"
+#include "PCGContext.h"
 #include "PCGCustomVersion.h"
-#include "Data/PCGSpatialData.h"
+#include "PCGParamData.h"
+#include "PCGPin.h"
 #include "Data/PCGPointData.h"
-#include "Metadata/Accessors/PCGCustomAccessor.h"
+#include "Data/PCGSpatialData.h"
+#include "Metadata/Accessors/IPCGAttributeAccessorTpl.h"
 #include "Metadata/Accessors/PCGAttributeAccessorHelpers.h"
-#include "PCGContext.h"
-#include "PCGPin.h"
+#include "Metadata/Accessors/PCGCustomAccessor.h"
 
 #include UE_INLINE_GENERATED_CPP_BY_NAME(PCGPointFilter)
 
@@ -201,8 +202,9 @@ bool FPCGPointFilterElement::ExecuteInternal(FPCGContext* Context) const
 	Outputs.Append(Context->InputData.GetAllSettings());
 
 	const EPCGPointFilterOperator Operator = Settings->Operator;
-	const bool bUseSpatialQuery = Settings->bUseSpatialQuery;
 	const bool bUseConstantThreshold = Settings->bUseConstantThreshold;
+	// If the filter is a param data, disable bUseSpatialQuery. Therefore make it non-const, to be modified below.
+	bool bUseSpatialQuery = Settings->bUseSpatialQuery;
 	// TODO: allow selector overrides for ThresholdAttribute and TargetAttribute
 
 	// If there is no input, do nothing
@@ -259,6 +261,11 @@ bool FPCGPointFilterElement::ExecuteInternal(FPCGContext* Context) const
 					ThresholdData = ThresholdSpatialData;
 				}
 			}
+			else if (ThresholdData->IsA<UPCGParamData>())
+			{
+				// Disable spatial query as it doesn't make sense
+				bUseSpatialQuery = false;
+			}
 
 			if (ThresholdData)
 			{
@@ -321,7 +328,7 @@ bool FPCGPointFilterElement::ExecuteInternal(FPCGContext* Context) const
 			ThresholdAccessor = PCGAttributeAccessorHelpers::CreateConstAccessor(OriginalData, Settings->ThresholdAttribute);
 			ThresholdKeys = PCGAttributeAccessorHelpers::CreateConstKeys(OriginalData, Settings->ThresholdAttribute);
 		}
-		else if (ThresholdSpatialData != nullptr && Settings->bUseSpatialQuery)
+		else if (ThresholdSpatialData != nullptr && bUseSpatialQuery)
 		{
 			// Reset the point data and reserving some points
 			// No need to reserve the full number of points, since we'll go by chunk
