Fixed copy-paste error where the seeding flag would be ignored
Improved relative seeding mode to use seeds instead of recomputing from final position

#rnx

diff --git a/Engine/Plugins/Experimental/PCG/Source/PCG/Private/Elements/PCGCopyPoints.cpp b/Engine/Plugins/Experimental/PCG/Source/PCG/Private/Elements/PCGCopyPoints.cpp
--- a/Engine/Plugins/Experimental/PCG/Source/PCG/Private/Elements/PCGCopyPoints.cpp
+++ b/Engine/Plugins/Experimental/PCG/Source/PCG/Private/Elements/PCGCopyPoints.cpp
@@ -176,12 +176,11 @@ bool FPCGCopyPointsElement::ExecuteInternal(FPCGContext* Context) const
 		OutPoint.Transform.SetLocation(Location);
 
 		// Set seed based on inheritance mode
-		if (ColorInheritance == EPCGCopyPointsInheritanceMode::Relative)
+		if (SeedInheritance == EPCGCopyPointsInheritanceMode::Relative)
 		{
-			const FVector& Position = OutPoint.Transform.GetLocation();
-			OutPoint.Seed = PCGHelpers::ComputeSeed((int)Position.X, (int)Position.Y, (int)Position.Z);
+			OutPoint.Seed = PCGHelpers::ComputeSeed(SourcePoint.Seed, TargetPoint.Seed);
 		}
-		else if (ColorInheritance == EPCGCopyPointsInheritanceMode::Target)
+		else if (SeedInheritance == EPCGCopyPointsInheritanceMode::Target)
 		{
 			OutPoint.Seed = TargetPoint.Seed;
 		}
