Fixed error in self-actor parsing condition use in data from actor

#rnx

diff --git a/Engine/Plugins/Experimental/PCG/Source/PCG/Private/Elements/PCGActorSelector.cpp b/Engine/Plugins/Experimental/PCG/Source/PCG/Private/Elements/PCGActorSelector.cpp
--- a/Engine/Plugins/Experimental/PCG/Source/PCG/Private/Elements/PCGActorSelector.cpp
+++ b/Engine/Plugins/Experimental/PCG/Source/PCG/Private/Elements/PCGActorSelector.cpp
@@ -17,7 +17,7 @@ namespace PCGActorSelector
 	{
 		// If we're not filtering all world actors, and if the Include Children is enabled, then we allow
 		// user to disable the filter which ignores all filtering options, which is convenient for workflow simplification.
-		if (InSettings.ActorFilter != EPCGActorFilter::AllWorldActors && InSettings.bIncludeChildren && InSettings.bDisableFilter)
+		if (InSettings.ActorFilter != EPCGActorFilter::AllWorldActors && (!InSettings.bIncludeChildren || InSettings.bDisableFilter))
 		{
 			return [&InFoundActors](AActor* Actor) -> bool
 			{
