Improved dependency crawling to be a bit more robust - skipping transient properties, etc.

#rnx

diff --git a/Engine/Plugins/Experimental/PCG/Source/PCG/Private/PCGHelpers.cpp b/Engine/Plugins/Experimental/PCG/Source/PCG/Private/PCGHelpers.cpp
--- a/Engine/Plugins/Experimental/PCG/Source/PCG/Private/PCGHelpers.cpp
+++ b/Engine/Plugins/Experimental/PCG/Source/PCG/Private/PCGHelpers.cpp
@@ -293,6 +293,12 @@ namespace PCGHelpers
 
 		for (FProperty* Property = ObjectClass->PropertyLink; Property != nullptr; Property = Property->PropertyLinkNext)
 		{
+			// Skip any kind of internal property and the ones that are susceptible to be instable
+			if (!Property || Property->HasAnyPropertyFlags(CPF_Transient | CPF_DuplicateTransient | CPF_Deprecated))
+			{
+				continue;
+			}
+
 			GatherDependencies(Property, Object, OutDependencies, MaxDepth);
 		}
 	}
