[PCG] Add blueprint overridable params

diff --git a/Engine/Plugins/Experimental/PCG/Source/PCG/Private/Elements/PCGExecuteBlueprint.cpp b/Engine/Plugins/Experimental/PCG/Source/PCG/Private/Elements/PCGExecuteBlueprint.cpp
--- a/Engine/Plugins/Experimental/PCG/Source/PCG/Private/Elements/PCGExecuteBlueprint.cpp
+++ b/Engine/Plugins/Experimental/PCG/Source/PCG/Private/Elements/PCGExecuteBlueprint.cpp
@@ -313,6 +313,9 @@ void UPCGBlueprintSettings::OnBlueprintChanged(UBlueprint* InBlueprint)
 	TeardownBlueprintElementEvent();
 	SetupBlueprintElementEvent();
 
+	// Also, reconstruct overrides
+	InitializeCachedOverridableParams(/*bReset=*/true);
+
 	OnSettingsChangedDelegate.Broadcast(this, EPCGChangeType::Settings);
 }
 
@@ -358,7 +361,10 @@ void UPCGBlueprintSettings::RefreshBlueprintElement()
 	else
 	{
 		BlueprintElementInstance = nullptr;
-	}	
+	}
+
+	// Also, reconstruct overrides
+	InitializeCachedOverridableParams(/*bReset=*/true);
 }
 
 #if WITH_EDITOR
@@ -467,6 +473,26 @@ FPCGElementPtr UPCGBlueprintSettings::CreateElement() const
 	return MakeShared<FPCGExecuteBlueprintElement>();
 }
 
+#if WITH_EDITOR
+TArray<FPCGSettingsOverridableParam> UPCGBlueprintSettings::GatherOverridableParams() const
+{
+	TArray<FPCGSettingsOverridableParam> OverridableParams = Super::GatherOverridableParams();
+
+	if (BlueprintElementInstance)
+	{
+		if (UClass* BPClass = BlueprintElementInstance->GetClass())
+		{
+			PCGSettingsHelpers::FPCGGetAllOverridableParamsConfig Config;
+			Config.bExcludeSuperProperties = true;
+			Config.ExcludePropertyFlags = CPF_DisableEditOnInstance;
+			OverridableParams.Append(PCGSettingsHelpers::GetAllOverridableParams(BPClass, Config));
+		}
+	}
+
+	return OverridableParams;
+}
+#endif // WITH_EDITOR
+
 bool FPCGExecuteBlueprintElement::ExecuteInternal(FPCGContext* InContext) const
 {
 	FPCGBlueprintExecutionContext* Context = static_cast<FPCGBlueprintExecutionContext*>(InContext);
@@ -489,22 +515,6 @@ bool FPCGExecuteBlueprintElement::ExecuteInternal(FPCGContext* InContext) const
 		}
 #endif
 
-		/** Apply params overrides to variables if any */
-		if (UPCGParamData* Params = Context->InputData.GetParamsWithDeprecation(Context->Node))
-		{
-			for (TFieldIterator<FProperty> PropertyIt(BPClass); PropertyIt; ++PropertyIt)
-			{
-				FProperty* Property = *PropertyIt;
-				if (Property->IsNative())
-				{
-					continue;
-				}
-
-				// Apply params if any
-				PCGSettingsHelpers::SetValue(Params, Context->BlueprintElementInstance, Property);
-			}
-		}
-
 		// Log info on inputs
 		for (int32 InputIndex = 0; InputIndex < Context->InputData.TaggedData.Num(); ++InputIndex)
 		{
diff --git a/Engine/Plugins/Experimental/PCG/Source/PCG/Private/Helpers/PCGSettingsHelpers.cpp b/Engine/Plugins/Experimental/PCG/Source/PCG/Private/Helpers/PCGSettingsHelpers.cpp
--- a/Engine/Plugins/Experimental/PCG/Source/PCG/Private/Helpers/PCGSettingsHelpers.cpp
+++ b/Engine/Plugins/Experimental/PCG/Source/PCG/Private/Helpers/PCGSettingsHelpers.cpp
@@ -271,4 +271,125 @@ namespace PCGSettingsHelpers
 			}
 		}
 	}
+
+#if WITH_EDITOR
+	template <typename ClassType>
+	TArray<FPCGSettingsOverridableParam> GetAllOverridableParamsImpl(const ClassType* InClass, const FPCGGetAllOverridableParamsConfig& InConfig)
+	{
+		// TODO: Was not a concern until now, and we didn't have a solution, but this function
+		// only worked if we don't have names clashes in overriable parameters.
+		// The previous override solution was flattening structs, and only override use the struct member name,
+		// not prefixed by the struct name or anything else.
+		// We cannot prefix it now, because it will break existing node that were assuming the flattening.
+		// We'll keep this behavior for now, as it might be solved by passing structs instead of param data,
+		// but we'll still at least raise a warning if there is a clash.
+		TSet<FName> LabelCache;
+
+		TArray<FPCGSettingsOverridableParam> Res;
+
+		const bool bCheckMetadata = !InConfig.MetadataValues.IsEmpty();
+		const bool bCheckExcludePropertyFlags = (InConfig.ExcludePropertyFlags != 0);
+		const EFieldIteratorFlags::SuperClassFlags SuperFlag = InConfig.bExcludeSuperProperties ? EFieldIteratorFlags::ExcludeSuper : EFieldIteratorFlags::IncludeSuper;
+
+		check(InClass);
+
+		for (TFieldIterator<FProperty> InputIt(InClass, SuperFlag, EFieldIteratorFlags::ExcludeDeprecated); InputIt; ++InputIt)
+		{
+			const FProperty* Property = *InputIt;
+			if (!Property)
+			{
+				continue;
+			}
+
+			bool bValid = true;
+
+			if (bCheckMetadata)
+			{
+				bool bFoundAny = false;
+				for (const FName& Metadata : InConfig.MetadataValues)
+				{
+					if (Property->HasMetaData(Metadata))
+					{
+						bFoundAny = true;
+						break;
+					}
+				}
+
+				bValid &= bFoundAny;
+			}
+
+			if (bCheckExcludePropertyFlags)
+			{
+				bValid &= !Property->HasAnyPropertyFlags(InConfig.ExcludePropertyFlags);
+			}
+
+			// Don't allow to override the seed if the settings doesn't use the seed.
+			if (!InConfig.bUseSeed)
+			{
+				bValid &= (Property->GetFName() != GET_MEMBER_NAME_CHECKED(UPCGSettings, Seed));
+			}
+
+			if (!bValid)
+			{
+				continue;
+			}
+
+			// Validate that the property can be overriden by params
+			if (PCGAttributeAccessorHelpers::IsPropertyAccessorSupported(Property))
+			{
+				FName Label = *Property->GetDisplayNameText().ToString();
+				if (LabelCache.Contains(Label))
+				{
+					UE_LOG(LogPCG, Warning, TEXT("%s property clashes with another property already found. It is a limitation at the moment and this property will be ignored (ie. will not be overridable)"), *Label.ToString());
+					continue;
+				}
+
+				LabelCache.Add(Label);
+
+				FPCGSettingsOverridableParam& Param = Res.Emplace_GetRef();
+				Param.Label = Label;
+				Param.PropertiesNames.Add(Property->GetFName());
+				Param.PropertyClass = InClass;
+			}
+			else if (const FStructProperty* StructProperty = CastField<FStructProperty>(Property))
+			{
+				FString PropertyName = Property->GetDisplayNameText().ToString();
+				// Use the seed, and don't check metadata.
+				FPCGGetAllOverridableParamsConfig RecurseConfig = InConfig;
+				RecurseConfig.bUseSeed = true;
+				RecurseConfig.MetadataValues.Empty();
+
+				for (const FPCGSettingsOverridableParam& ChildParam : GetAllOverridableParams(StructProperty->Struct, RecurseConfig))
+				{
+					FName Label = ChildParam.Label;
+					if (LabelCache.Contains(Label))
+					{
+						UE_LOG(LogPCG, Warning, TEXT("%s property clashes with another property already found. It is a limitation at the moment and this property will be ignored (ie. will not be overridable)"), *Label.ToString());
+						continue;
+					}
+
+					LabelCache.Add(Label);
+
+					FPCGSettingsOverridableParam& Param = Res.Emplace_GetRef();
+					Param.Label = Label;
+					Param.PropertiesNames.Add(Property->GetFName());
+					Param.PropertiesNames.Append(ChildParam.PropertiesNames);
+					Param.PropertyClass = InClass;
+				}
+			}
+		}
+
+		return Res;
+	}
+
+	TArray<FPCGSettingsOverridableParam> GetAllOverridableParams(const UClass* InClass, const FPCGGetAllOverridableParamsConfig& InConfig)
+	{
+		return GetAllOverridableParamsImpl(InClass, InConfig);
+	}
+
+	TArray<FPCGSettingsOverridableParam> GetAllOverridableParams(const UScriptStruct* InStruct, const FPCGGetAllOverridableParamsConfig& InConfig)
+	{
+		return GetAllOverridableParamsImpl(InStruct, InConfig);
+	}
+#endif // WITH_EDITOR
 }
diff --git a/Engine/Plugins/Experimental/PCG/Source/PCG/Private/PCGContext.cpp b/Engine/Plugins/Experimental/PCG/Source/PCG/Private/PCGContext.cpp
--- a/Engine/Plugins/Experimental/PCG/Source/PCG/Private/PCGContext.cpp
+++ b/Engine/Plugins/Experimental/PCG/Source/PCG/Private/PCGContext.cpp
@@ -138,12 +138,36 @@ void FPCGContext::OverrideSettings()
 	{
 		check(!Param.Properties.IsEmpty());
 
-		// TODO: Support override on subobjects
-		if (!SettingsWithOverride->IsA(Param.Properties[0]->GetOwnerClass()))
+		// Verification that container is valid and we have the right class.
+		void* Container = nullptr;
+
+		if (!Param.PropertyClass)
 		{
 			continue;
 		}
 
+		if (!SettingsWithOverride->GetClass()->IsChildOf(Param.PropertyClass))
+		{
+			UObject* ObjectPtr = GetExternalContainerForOverridableParam(Param);
+			if (!ObjectPtr || !ObjectPtr->IsA(Param.Properties[0]->GetOwnerClass()))
+			{
+				continue;
+			}
+
+			Container = ObjectPtr;
+		}
+		else
+		{
+			if (!SettingsWithOverride->IsA(Param.Properties[0]->GetOwnerClass()))
+			{
+				continue;
+			}
+
+			Container = SettingsWithOverride.Get();
+		}
+
+		check(Container);
+
 		FName AttributeName = NAME_None;
 		TUniquePtr<const IPCGAttributeAccessor> AttributeAccessor = PCGAttributeAccessorHelpers::CreateConstAccessorForOverrideParam(InputData, Param, &AttributeName);
 
@@ -156,7 +180,7 @@ void FPCGContext::OverrideSettings()
 		TUniquePtr<IPCGAttributeAccessor> PropertyAccessor = PCGAttributeAccessorHelpers::CreatePropertyAccessor(Param.Properties.Last());
 		check(PropertyAccessor.IsValid());
 
-		PCGMetadataAttribute::CallbackWithRightType(PropertyAccessor->GetUnderlyingType(), [this, &AttributeAccessor, &PropertyAccessor, &Param, &AttributeName](auto Dummy) -> bool
+		PCGMetadataAttribute::CallbackWithRightType(PropertyAccessor->GetUnderlyingType(), [this, &AttributeAccessor, &PropertyAccessor, &Param, &AttributeName, Container](auto Dummy) -> bool
 		{
 			using PropertyType = decltype(Dummy);
 
@@ -170,7 +194,7 @@ void FPCGContext::OverrideSettings()
 				return false;
 			}
 			// TODO: Perhaps factorise this code in another property accessor.
-			void* PropertyObjectPtr = SettingsWithOverride.Get();
+			void* PropertyObjectPtr = Container;
 			for (int32 i = 0; i < Param.Properties.Num() - 1; ++i)
 			{
 				PropertyObjectPtr = Param.Properties[i]->ContainerPtrToValuePtr<void>(PropertyObjectPtr);
diff --git a/Engine/Plugins/Experimental/PCG/Source/PCG/Private/PCGSettings.cpp b/Engine/Plugins/Experimental/PCG/Source/PCG/Private/PCGSettings.cpp
--- a/Engine/Plugins/Experimental/PCG/Source/PCG/Private/PCGSettings.cpp
+++ b/Engine/Plugins/Experimental/PCG/Source/PCG/Private/PCGSettings.cpp
@@ -106,7 +106,13 @@ void UPCGSettings::PostLoad()
 {
 	Super::PostLoad();
 
-	InitializeCachedOverridableParams();
+	// In editor, we always want to rebuild our cache, to make sure we are up to date.
+	// At runtime, we can't rebuild our cache (that is why we cache it in the first place).
+#if WITH_EDITOR
+	InitializeCachedOverridableParams(/*bReset=*/true);
+#else
+	InitializeCachedOverridableParams(/*bReset=*/false);
+#endif // WITH_EDITOR
 
 #if WITH_EDITOR
 	if (ExecutionMode_DEPRECATED != EPCGSettingsExecutionMode::Enabled)
@@ -129,6 +135,12 @@ void UPCGSettings::PostInitProperties()
 
 void UPCGSettings::Serialize(FArchive& Ar)
 {
+	// Don't serialize overridable params in non-cooked builds
+	if (Ar.IsSaving() && !Ar.IsCooking())
+	{
+		CachedOverridableParams.Empty();
+	}
+
 	Super::Serialize(Ar);
 
 	Ar.UsingCustomVersion(FPCGCustomVersion::GUID);
@@ -160,6 +172,12 @@ void UPCGSettings::Serialize(FArchive& Ar)
 		}
 #endif // WITH_EDITOR
 	}
+
+	// Reconstruct it at the end
+	if (Ar.IsSaving() && !Ar.IsCooking())
+	{
+		InitializeCachedOverridableParams(/*bReset=*/true);
+	}
 }
 
 void UPCGSettings::PostSaveRoot(FObjectPostSaveRootContext ObjectSaveContext)
@@ -343,7 +361,7 @@ void UPCGSettings::PostEditChangeProperty(FPropertyChangedEvent& PropertyChanged
 	if (PropertyChangedEvent.Property && (PropertyChangedEvent.Property->GetFName() == GET_MEMBER_NAME_CHECKED(UPCGSettings, CachedOverridableParams)))
 	{
 		// Need to rebuild properties, if it ever changes.
-		InitializeCachedOverridableParams();
+		InitializeCachedOverridableParams(/*bReset=*/true);
 	}
 
 	if (PropertyChangedEvent.GetPropertyName() != GET_MEMBER_NAME_CHECKED(UPCGSettings, DeterminismSettings))
@@ -409,79 +427,18 @@ void UPCGSettings::ApplyDeprecationBeforeUpdatePins(UPCGNode* InOutNode, TArray<
 }
 #endif // WITH_EDITOR
 
-namespace PCGSettings
-{
 #if WITH_EDITOR
-	template <typename ClassType>
-	TArray<FPCGSettingsOverridableParam> GetAllParams(const ClassType* Class, bool bCheckMetadata, bool bUseSeed)
-	{
-		// TODO: Was not a concern until now, and we didn't have a solution, but this function
-		// only worked if we don't have names clashes in overriable parameters.
-		// The previous override solution was flatenning structs, and only override use the struct member name,
-		// not prefixed by the struct name or anything else.
-		// We cannot prefix it now, because it will break existing node that were assuming the flatenning.
-		// We'll keep this behavior for now, as it might be solved by passing structs instead of param data,
-		// but we'll still at least raise a warning if there is a clash.
-		TSet<FName> LabelCache;
-
-		TArray<FPCGSettingsOverridableParam> Res;
-
-		for (TFieldIterator<FProperty> InputIt(Class, EFieldIteratorFlags::IncludeSuper, EFieldIteratorFlags::ExcludeDeprecated); InputIt; ++InputIt)
-		{
-			const FProperty* Property = *InputIt;
-			if (Property && (!bCheckMetadata || Property->HasMetaData(PCGObjectMetadata::Overridable)))
-			{
-				// Don't allow to override the seed if the settings doesn't use the seed.
-				if (Property->GetFName() == GET_MEMBER_NAME_CHECKED(UPCGSettings, Seed) && !bUseSeed)
-				{
-					continue;
-				}
-
-				// Validating that the property can be overriden by params
-				if (PCGAttributeAccessorHelpers::IsPropertyAccessorSupported(Property))
-				{
-					FName Label = *Property->GetDisplayNameText().ToString();
-					if (LabelCache.Contains(Label))
-					{
-						UE_LOG(LogPCG, Warning, TEXT("%s property clashes with another property already found. It is a limitation at the moment and this property will be ignored (ie. will not be overridable)"), *Label.ToString());
-						continue;
-					}
-
-					LabelCache.Add(Label);
-
-					FPCGSettingsOverridableParam& Param = Res.Emplace_GetRef();
-					Param.Label = Label;
-					Param.PropertiesNames.Add(Property->GetFName());
-				}
-				else if (const FStructProperty* StructProperty = CastField<FStructProperty>(Property))
-				{
-					FString PropertyName = Property->GetDisplayNameText().ToString();
-					for (const FPCGSettingsOverridableParam& ChildParam : GetAllParams(StructProperty->Struct, /*bCheckMetadata=*/false, /*bUseSeed=*/true))
-					{
-						FName Label = ChildParam.Label;
-						if (LabelCache.Contains(Label))
-						{
-							UE_LOG(LogPCG, Warning, TEXT("%s property clashes with another property already found. It is a limitation at the moment and this property will be ignored (ie. will not be overridable)"), *Label.ToString());
-							continue;
-						}
-
-						LabelCache.Add(Label);
-
-						FPCGSettingsOverridableParam& Param = Res.Emplace_GetRef();
-						Param.Label = Label;
-						Param.PropertiesNames.Add(Property->GetFName());
-						Param.PropertiesNames.Append(ChildParam.PropertiesNames);
-					}
-				}
-			}
-		}
+TArray<FPCGSettingsOverridableParam> UPCGSettings::GatherOverridableParams() const
+{
+	PCGSettingsHelpers::FPCGGetAllOverridableParamsConfig Config;
+	Config.bUseSeed = bUseSeed;
+	Config.MetadataValues.Add(PCGObjectMetadata::Overridable);
 
-		return Res;
-	}
-#endif // WITH_EDITOR
+	return PCGSettingsHelpers::GetAllOverridableParams(GetClass(), Config);
 }
+#endif // WITH_EDITOR
 
-void UPCGSettings::InitializeCachedOverridableParams()
+void UPCGSettings::InitializeCachedOverridableParams(bool bReset)
 {
 	// Don't do it for default object
 	if (HasAnyFlags(RF_ClassDefaultObject))
@@ -492,31 +449,32 @@ void UPCGSettings::InitializeCachedOverridableParams()
 	TRACE_CPUPROFILER_EVENT_SCOPE(UPCGSettings::InitializeCachedOverridableParams);
 
 #if WITH_EDITOR
-	if (CachedOverridableParams.IsEmpty())
+	if (CachedOverridableParams.IsEmpty() || bReset)
 	{
-		CachedOverridableParams = PCGSettings::GetAllParams(GetClass(), true, bUseSeed);
+		CachedOverridableParams = GatherOverridableParams();
 	}
 #endif // WITH_EDITOR
 
 	for (FPCGSettingsOverridableParam& Param : CachedOverridableParams)
 	{
-		check(!Param.PropertiesNames.IsEmpty());
+		check(!Param.PropertiesNames.IsEmpty() && Param.PropertyClass);
 		Param.Properties.Reset(Param.PropertiesNames.Num());
 
 		// Some properties might not be available at runtime. Ignore them.
-		if (const FProperty* Property = GetClass()->FindPropertyByName(Param.PropertiesNames[0]))
+		const FProperty* CurrentProperty = Param.PropertyClass->FindPropertyByName(Param.PropertiesNames[0]);
+		if (ensure(CurrentProperty))
 		{
-			Param.Properties.Add(Property);
+			Param.Properties.Add(CurrentProperty);
 
 			for (int32 i = 1; i < Param.PropertiesNames.Num(); ++i)
 			{
 				// If we have multiple depth properties, it should be Struct properties by construction
-				const FStructProperty* StructProperty = CastField<FStructProperty>(Property);
+				const FStructProperty* StructProperty = CastField<FStructProperty>(CurrentProperty);
 				if (ensure(StructProperty))
 				{
-					Property = StructProperty->Struct->FindPropertyByName(Param.PropertiesNames[i]);
-					check(Property);
-					Param.Properties.Add(Property);
+					CurrentProperty = StructProperty->Struct->FindPropertyByName(Param.PropertiesNames[i]);
+					check(CurrentProperty);
+					Param.Properties.Add(CurrentProperty);
 				}
 			}
 		}
diff --git a/Engine/Plugins/Experimental/PCG/Source/PCG/Public/Elements/PCGExecuteBlueprint.h b/Engine/Plugins/Experimental/PCG/Source/PCG/Public/Elements/PCGExecuteBlueprint.h
--- a/Engine/Plugins/Experimental/PCG/Source/PCG/Public/Elements/PCGExecuteBlueprint.h
+++ b/Engine/Plugins/Experimental/PCG/Source/PCG/Public/Elements/PCGExecuteBlueprint.h
@@ -193,6 +193,9 @@ public:
 	virtual bool HasOverridableParams() const override { return true; }
 protected:
 	virtual FPCGElementPtr CreateElement() const override;
+#if WITH_EDITOR
+	virtual TArray<FPCGSettingsOverridableParam> GatherOverridableParams() const override;
+#endif // WITH_EDITOR
 	// ~End UPCGSettings interface
 
 public:
@@ -256,6 +259,9 @@ struct FPCGBlueprintExecutionContext : public FPCGContext
 	virtual ~FPCGBlueprintExecutionContext();
 
 	UPCGBlueprintElement* BlueprintElementInstance = nullptr;
+
+protected:
+	virtual UObject* GetExternalContainerForOverridableParam(const FPCGSettingsOverridableParam& InParam) override { return BlueprintElementInstance; }
 };
 
 class FPCGExecuteBlueprintElement : public IPCGElement
diff --git a/Engine/Plugins/Experimental/PCG/Source/PCG/Public/Helpers/PCGSettingsHelpers.h b/Engine/Plugins/Experimental/PCG/Source/PCG/Public/Helpers/PCGSettingsHelpers.h
--- a/Engine/Plugins/Experimental/PCG/Source/PCG/Public/Helpers/PCGSettingsHelpers.h
+++ b/Engine/Plugins/Experimental/PCG/Source/PCG/Public/Helpers/PCGSettingsHelpers.h
@@ -201,7 +201,27 @@ PRAGMA_ENABLE_DEPRECATION_WARNINGS
 			return true;
 		});
 	}
+
+#if WITH_EDITOR
+	struct FPCGGetAllOverridableParamsConfig
+	{
+		// If we don't use the seed, don't add it as override.
+		bool bUseSeed = false;
+
+		// Don't look for properties from parents
+		bool bExcludeSuperProperties = false;
+
+		// List of metadata values to find in property metadata
+		TArray<FName> MetadataValues{};
+
+		// Flags to exclude in property flags
+		uint64 ExcludePropertyFlags = 0;
+	};
+
+	TArray<FPCGSettingsOverridableParam> GetAllOverridableParams(const UClass* InClass, const FPCGGetAllOverridableParamsConfig& InConfig);
+	TArray<FPCGSettingsOverridableParam> GetAllOverridableParams(const UScriptStruct* InStruct, const FPCGGetAllOverridableParamsConfig& InConfig);
+#endif
 }
 
 // Deprecated macro, not necessary anymore. Cf. GetValue
-#define PCG_GET_OVERRIDEN_VALUE(Settings, Variable, Params) PCGSettingsHelpers::GetValue(GET_MEMBER_NAME_CHECKED(TRemovePointer<std::remove_const_t<decltype(Settings)>>::Type, Variable), (Settings)->Variable, Params)
\ No newline at end of file
+#define PCG_GET_OVERRIDEN_VALUE(Settings, Variable, Params) PCGSettingsHelpers::GetValue(GET_MEMBER_NAME_CHECKED(std::remove_pointer_t<std::remove_const_t<decltype(Settings)>>, Variable), (Settings)->Variable, Params)
diff --git a/Engine/Plugins/Experimental/PCG/Source/PCG/Public/PCGContext.h b/Engine/Plugins/Experimental/PCG/Source/PCG/Public/PCGContext.h
--- a/Engine/Plugins/Experimental/PCG/Source/PCG/Public/PCGContext.h
+++ b/Engine/Plugins/Experimental/PCG/Source/PCG/Public/PCGContext.h
@@ -10,6 +10,7 @@
 class UPCGComponent;
 class UPCGSettingsInterface;
 struct FPCGGraphCache;
+struct FPCGSettingsOverridableParam;
 
 namespace PCGContextHelpers
 {
@@ -91,6 +92,9 @@ struct PCG_API FPCGContext
 	FString GetComponentName() const;
 	bool ShouldStop() const;
 
+protected:
+	virtual UObject* GetExternalContainerForOverridableParam(const FPCGSettingsOverridableParam& InParam) { return nullptr; }
+
 private:
 	template<typename SettingsType>
 	const SettingsType* GetOriginalSettings() const 
diff --git a/Engine/Plugins/Experimental/PCG/Source/PCG/Public/PCGSettings.h b/Engine/Plugins/Experimental/PCG/Source/PCG/Public/PCGSettings.h
--- a/Engine/Plugins/Experimental/PCG/Source/PCG/Public/PCGSettings.h
+++ b/Engine/Plugins/Experimental/PCG/Source/PCG/Public/PCGSettings.h
@@ -61,6 +61,9 @@ struct FPCGSettingsOverridableParam
 	UPROPERTY()
 	TArray<FName> PropertiesNames;
 
+	UPROPERTY()
+	TObjectPtr<const UStruct> PropertyClass;
+
 	// Transient
 	TArray<const FProperty*> Properties;
 };
@@ -261,12 +264,18 @@ public:
 protected:
 	/** Iterate over OverridableParams to automatically add param pins to the list. */
 	void FillOverridableParamsPins(TArray<FPCGPinProperties>& OutPins) const;
-private:
+
+#if WITH_EDITOR
+	/** Can be overridden to add more custom params (like in BP). */
+	virtual TArray<FPCGSettingsOverridableParam> GatherOverridableParams() const;
+#endif // WITH_EDITOR
+
 	/** Called after intialization to construct the list of all overridable params. 
 	* In Editor, it will first gather all the overridable properties names and labels, based on their metadata.
+	* Can request a "reset" if something changed in the settings.
 	* And then, in Editor and Runtime, will gather the FProperty*.
 	*/
-	void InitializeCachedOverridableParams();
+	void InitializeCachedOverridableParams(bool bReset = false);
 
 	/** Needs to be serialized since property metadata (used to populate this array) is not available at runtime. */
 	UPROPERTY()
