[PCG] SpawnActor pin types changed to Point.

Any element that calls ToPointData on each bit of data coming in through a pin should declare that pin of type Point.

This makes the 'collapse' to be visible to the user, and allows them to prepare data appropriately. And this will eventually prevent them from accidentally passing in a surface or volume and triggering an huge collapse.

We don't need to worry about narrowing the input pin type because:

* Currently we allow edges from Any and Spatial to Concrete, so it won't break existing edges. This will be dealt with later (and this change prepares the pins for that)
* Thanks to [REDACTED_PERSON] recent upgrades, params are handled separately

diff --git a/Engine/Plugins/Experimental/PCG/Source/PCG/Private/PCGSettings.cpp b/Engine/Plugins/Experimental/PCG/Source/PCG/Private/PCGSettings.cpp
--- a/Engine/Plugins/Experimental/PCG/Source/PCG/Private/PCGSettings.cpp
+++ b/Engine/Plugins/Experimental/PCG/Source/PCG/Private/PCGSettings.cpp
@@ -481,6 +481,13 @@ void UPCGSettings::InitializeCachedOverridableParams(bool bReset)
 	}
 }
 
+TArray<FPCGPinProperties> UPCGSettings::DefaultPointInputPinProperties() const
+{
+	TArray<FPCGPinProperties> Properties;
+	Properties.Emplace(PCGPinConstants::DefaultInputLabel, EPCGDataType::Point);
+	return Properties;
+}
+
 TArray<FPCGPinProperties> UPCGSettings::DefaultPointOutputPinProperties() const
 {
 	TArray<FPCGPinProperties> Properties;
diff --git a/Engine/Plugins/Experimental/PCG/Source/PCG/Public/Elements/PCGSpawnActor.h b/Engine/Plugins/Experimental/PCG/Source/PCG/Public/Elements/PCGSpawnActor.h
--- a/Engine/Plugins/Experimental/PCG/Source/PCG/Public/Elements/PCGSpawnActor.h
+++ b/Engine/Plugins/Experimental/PCG/Source/PCG/Public/Elements/PCGSpawnActor.h
@@ -2,10 +2,10 @@
 
 #pragma once
 
+#include "PCGPin.h"
 #include "PCGSubgraph.h"
 #include "Metadata/PCGAttributePropertySelector.h"
 
-
 #include "PCGSpawnActor.generated.h"
 
 class AActor;
@@ -89,10 +89,14 @@ public:
 	virtual UPCGNode* CreateNode() const override;
 
 #if WITH_EDITOR	
-	virtual FName GetDefaultNodeName() const override { return FName(TEXT("SpawnActor")); }
+	virtual FName GetDefaultNodeName() const override { return FName(TEXT("Spawn Actor")); }
 	virtual EPCGSettingsType GetType() const override;
 #endif
+
 protected:
+	virtual TArray<FPCGPinProperties> InputPinProperties() const override { return DefaultPointInputPinProperties(); }
+	virtual TArray<FPCGPinProperties> OutputPinProperties() const override { return DefaultPointOutputPinProperties(); }
+
 	virtual FPCGElementPtr CreateElement() const override;
 	//~End UPCGSettings interface
 
diff --git a/Engine/Plugins/Experimental/PCG/Source/PCG/Public/PCGSettings.h b/Engine/Plugins/Experimental/PCG/Source/PCG/Public/PCGSettings.h
--- a/Engine/Plugins/Experimental/PCG/Source/PCG/Public/PCGSettings.h
+++ b/Engine/Plugins/Experimental/PCG/Source/PCG/Public/PCGSettings.h
@@ -240,6 +240,7 @@ protected:
 	bool bUseSeed = false;
 
 	/** Methods to remove boilerplate code across settings */
+	TArray<FPCGPinProperties> DefaultPointInputPinProperties() const;
 	TArray<FPCGPinProperties> DefaultPointOutputPinProperties() const;
 
 #if WITH_EDITOR
