[PCG] Purify UPCGLandscapeData::SamplePoint - on success point should not be moved from the query location.

Also update tooltip for SamplePoint, although there are still other instances of SamplePoint needing purification. However I searched two projects and only found SamplePoint operating on Landscapes so thought I would get this out of the way now so that users start getting the info now. But can defer if preferred.

diff --git a/Engine/Plugins/Experimental/PCG/Source/PCG/Private/Data/PCGIntersectionData.cpp b/Engine/Plugins/Experimental/PCG/Source/PCG/Private/Data/PCGIntersectionData.cpp
--- a/Engine/Plugins/Experimental/PCG/Source/PCG/Private/Data/PCGIntersectionData.cpp
+++ b/Engine/Plugins/Experimental/PCG/Source/PCG/Private/Data/PCGIntersectionData.cpp
@@ -155,6 +155,18 @@ const UPCGSpatialData* UPCGIntersectionData::FindShapeFromNetwork(const int InDi
 	return nullptr;
 }
 
+const UPCGSpatialData* UPCGIntersectionData::FindFirstConcreteShapeFromNetwork() const
+{
+	check(A && B);
+
+	if (const UPCGSpatialData* CandidateA = A->FindFirstConcreteShapeFromNetwork())
+	{
+		return CandidateA;
+	}
+
+	return B->FindFirstConcreteShapeFromNetwork();
+}
+
 const UPCGPointData* UPCGIntersectionData::CreatePointData(FPCGContext* Context) const
 {
 	check(A && B);
diff --git a/Engine/Plugins/Experimental/PCG/Source/PCG/Private/Data/PCGLandscapeData.cpp b/Engine/Plugins/Experimental/PCG/Source/PCG/Private/Data/PCGLandscapeData.cpp
--- a/Engine/Plugins/Experimental/PCG/Source/PCG/Private/Data/PCGLandscapeData.cpp
+++ b/Engine/Plugins/Experimental/PCG/Source/PCG/Private/Data/PCGLandscapeData.cpp
@@ -108,10 +108,21 @@ FBox UPCGLandscapeData::GetStrictBounds() const
 
 bool UPCGLandscapeData::SamplePoint(const FTransform& InTransform, const FBox& InBounds, FPCGPoint& OutPoint, UPCGMetadata* OutMetadata) const
 {
-	// SamplePoint actually does a projection for this data, so lets make it explicit and delegate to a project function.
-	// TODO fixup samplers and other things that rely on this behavior and then fix this code.
-	FPCGProjectionParams Params{};
-	return ProjectPoint(InTransform, InBounds, Params, OutPoint, OutMetadata);
+	// The point is in/on the shape if it coincides with its projection. I.e. projecting on landscape does not move the point. Implementing
+	// this way shares the sampling code.
+	if (ProjectPoint(InTransform, InBounds, {}, OutPoint, OutMetadata))
+	{
+		if (InBounds.IsValid)
+		{
+			return FMath::PointBoxIntersection(OutPoint.Transform.GetLocation(), InBounds.TransformBy(InTransform));
+		}
+		else
+		{
+			return (InTransform.GetLocation() - OutPoint.Transform.GetLocation()).SquaredLength() < UE_SMALL_NUMBER;
+		}
+	}
+
+	return false;
 }
 
 bool UPCGLandscapeData::ProjectPoint(const FTransform& InTransform, const FBox& InBounds, const FPCGProjectionParams& InParams, FPCGPoint& OutPoint, UPCGMetadata* OutMetadata) const
diff --git a/Engine/Plugins/Experimental/PCG/Source/PCG/Private/Data/PCGSpatialData.cpp b/Engine/Plugins/Experimental/PCG/Source/PCG/Private/Data/PCGSpatialData.cpp
--- a/Engine/Plugins/Experimental/PCG/Source/PCG/Private/Data/PCGSpatialData.cpp
+++ b/Engine/Plugins/Experimental/PCG/Source/PCG/Private/Data/PCGSpatialData.cpp
@@ -136,8 +136,15 @@ UPCGSpatialData* UPCGSpatialData::ProjectOn(const UPCGSpatialData* InOther, cons
 		return DuplicateData();
 	}
 
+	const UPCGSpatialData* ConcreteTarget = InOther->FindFirstConcreteShapeFromNetwork();
+	if (!ConcreteTarget)
+	{
+		UE_LOG(LogPCG, Error, TEXT("Could not find a concrete shape in the target data to project onto."));
+		return DuplicateData();
+	}
+
 	UPCGProjectionData* ProjectionData = NewObject<UPCGProjectionData>();
-	ProjectionData->Initialize(this, InOther, InParams);
+	ProjectionData->Initialize(this, ConcreteTarget, InParams);
 
 	return ProjectionData;
 }
diff --git a/Engine/Plugins/Experimental/PCG/Source/PCG/Private/Data/PCGUnionData.cpp b/Engine/Plugins/Experimental/PCG/Source/PCG/Private/Data/PCGUnionData.cpp
--- a/Engine/Plugins/Experimental/PCG/Source/PCG/Private/Data/PCGUnionData.cpp
+++ b/Engine/Plugins/Experimental/PCG/Source/PCG/Private/Data/PCGUnionData.cpp
@@ -199,6 +199,21 @@ const UPCGSpatialData* UPCGUnionData::FindShapeFromNetwork(const int InDimension
 	return nullptr;
 }
 
+const UPCGSpatialData* UPCGUnionData::FindFirstConcreteShapeFromNetwork() const
+{
+	// Return first concrete candidate data.
+	for (TObjectPtr<const UPCGSpatialData> Datum : Data)
+	{
+		const UPCGSpatialData* Candidate = Datum ? Datum->FindFirstConcreteShapeFromNetwork() : nullptr;
+		if (Candidate)
+		{
+			return Candidate;
+		}
+	}
+
+	return nullptr;
+}
+
 const UPCGPointData* UPCGUnionData::CreatePointData(FPCGContext* Context) const
 {
 	TRACE_CPUPROFILER_EVENT_SCOPE(UPCGUnionData::CreatePointData);
diff --git a/Engine/Plugins/Experimental/PCG/Source/PCG/Private/Elements/PCGSurfaceSampler.cpp b/Engine/Plugins/Experimental/PCG/Source/PCG/Private/Elements/PCGSurfaceSampler.cpp
--- a/Engine/Plugins/Experimental/PCG/Source/PCG/Private/Elements/PCGSurfaceSampler.cpp
+++ b/Engine/Plugins/Experimental/PCG/Source/PCG/Private/Elements/PCGSurfaceSampler.cpp
@@ -322,10 +322,22 @@ bool FPCGSurfaceSamplerElement::ExecuteInternal(FPCGContext* Context) const
 	TArray<const UPCGSpatialData*, TInlineAllocator<16>> GeneratingShapes;
 	for (FPCGTaggedData& TaggedData : SurfaceInputs)
 	{
-		if (UPCGSpatialData* SpatialData = Cast<UPCGSpatialData>(TaggedData.Data))
+		if (const UPCGSpatialData* SpatialData = Cast<UPCGSpatialData>(TaggedData.Data))
 		{
-			GeneratingShapes.Add(SpatialData);
-			Outputs.Add(TaggedData);
+			// Find a concrete shape for sampling. Prefer a 2D surface if we can find one.
+			if (const UPCGSpatialData* SurfaceData = SpatialData->FindShapeFromNetwork(/*InDimension=*/2))
+			{
+				GeneratingShapes.Add(SurfaceData);
+				Outputs.Add(TaggedData);
+			}
+			else if (const UPCGSpatialData* ConcreteData = SpatialData->FindFirstConcreteShapeFromNetwork())
+			{
+				// Alternatively surface-sample any concrete data - can be used to sprinkle samples down onto shapes like volumes.
+				// Searching like this allows the user to plonk in any composite network and it will often find the shape of interest.
+				// A potential extension would be to find all (unique?) concrete shapes and use all of them rather than just the first.
+				GeneratingShapes.Add(ConcreteData);
+				Outputs.Add(TaggedData);
+			}
 		}
 	}
 
diff --git a/Engine/Plugins/Experimental/PCG/Source/PCG/Public/Data/PCGDifferenceData.h b/Engine/Plugins/Experimental/PCG/Source/PCG/Public/Data/PCGDifferenceData.h
--- a/Engine/Plugins/Experimental/PCG/Source/PCG/Public/Data/PCGDifferenceData.h
+++ b/Engine/Plugins/Experimental/PCG/Source/PCG/Public/Data/PCGDifferenceData.h
@@ -58,6 +58,7 @@ public:
 	virtual bool SamplePoint(const FTransform& Transform, const FBox& Bounds, FPCGPoint& OutPoint, UPCGMetadata* OutMetadata) const override;
 	virtual bool HasNonTrivialTransform() const override;
 	virtual const UPCGSpatialData* FindShapeFromNetwork(const int InDimension) const override { return Source ? Source->FindShapeFromNetwork(InDimension) : nullptr; }
+	virtual const UPCGSpatialData* FindFirstConcreteShapeFromNetwork() const override { return Source ? Source->FindFirstConcreteShapeFromNetwork() : nullptr; }
 protected:
 	virtual UPCGSpatialData* CopyInternal() const override;
 	//~End UPCGSpatialData interface
diff --git a/Engine/Plugins/Experimental/PCG/Source/PCG/Public/Data/PCGIntersectionData.h b/Engine/Plugins/Experimental/PCG/Source/PCG/Public/Data/PCGIntersectionData.h
--- a/Engine/Plugins/Experimental/PCG/Source/PCG/Public/Data/PCGIntersectionData.h
+++ b/Engine/Plugins/Experimental/PCG/Source/PCG/Public/Data/PCGIntersectionData.h
@@ -35,6 +35,7 @@ public:
 	virtual bool SamplePoint(const FTransform& Transform, const FBox& Bounds, FPCGPoint& OutPoint, UPCGMetadata* OutMetadata) const override;
 	virtual bool HasNonTrivialTransform() const override;
 	virtual const UPCGSpatialData* FindShapeFromNetwork(const int InDimension) const override;
+	virtual const UPCGSpatialData* FindFirstConcreteShapeFromNetwork() const override;
 protected:
 	virtual UPCGSpatialData* CopyInternal() const override;
 	//~End UPCGSpatialData interface
diff --git a/Engine/Plugins/Experimental/PCG/Source/PCG/Public/Data/PCGSpatialData.h b/Engine/Plugins/Experimental/PCG/Source/PCG/Public/Data/PCGSpatialData.h
--- a/Engine/Plugins/Experimental/PCG/Source/PCG/Public/Data/PCGSpatialData.h
+++ b/Engine/Plugins/Experimental/PCG/Source/PCG/Public/Data/PCGSpatialData.h
@@ -80,7 +80,7 @@ public:
 	/** Transform a world-space position to a world-space position in relation to the current data. (ex: projection on surface) */
 	FVector TransformPosition(const FVector& InPosition) const;
 
-	/** Sample information from this data at the query. Generally OutPoint should be located at the query position, however some data such as Landscape currently also projects the point. Returns true if Transform location and Bounds overlaps this data. */
+	/** Sample rotation, scale and other attributes from this data at the query position. Returns true if Transform location and Bounds overlaps this data. */
 	UFUNCTION(BlueprintCallable, Category = SpatialData)
 	virtual bool SamplePoint(const FTransform& Transform, const FBox& Bounds, FPCGPoint& OutPoint, UPCGMetadata* OutMetadata) const PURE_VIRTUAL(UPCGSpatialData::SamplePoint, return false;);
 
@@ -125,6 +125,9 @@ public:
 	/** A call that is made recursively up through the graph to find the best candidate shape for point generation. If InDimension is -1, finds lowest dimensional shape. */
 	virtual const UPCGSpatialData* FindShapeFromNetwork(const int InDimension) const { return (InDimension == -1 || GetDimension() == InDimension) ? this : nullptr; }
 
+	/** Find the first concrete (non-composite) shape in the network. Depth first search. */
+	virtual const UPCGSpatialData* FindFirstConcreteShapeFromNetwork() const { return !(GetDataType() & EPCGDataType::Composite) ? this : nullptr; }
+
 	UPROPERTY(Transient, BlueprintReadWrite, EditAnywhere, Category = Data)
 	TWeakObjectPtr<AActor> TargetActor = nullptr;
 
diff --git a/Engine/Plugins/Experimental/PCG/Source/PCG/Public/Data/PCGUnionData.h b/Engine/Plugins/Experimental/PCG/Source/PCG/Public/Data/PCGUnionData.h
--- a/Engine/Plugins/Experimental/PCG/Source/PCG/Public/Data/PCGUnionData.h
+++ b/Engine/Plugins/Experimental/PCG/Source/PCG/Public/Data/PCGUnionData.h
@@ -47,6 +47,7 @@ public:
 	virtual bool SamplePoint(const FTransform& Transform, const FBox& Bounds, FPCGPoint& OutPoint, UPCGMetadata* OutMetadata) const override;
 	virtual bool HasNonTrivialTransform() const override;
 	virtual const UPCGSpatialData* FindShapeFromNetwork(const int InDimension) const override;
+	virtual const UPCGSpatialData* FindFirstConcreteShapeFromNetwork() const override;
 protected:
 	virtual UPCGSpatialData* CopyInternal() const override;
 	//~End UPCGSpatialData interface
