[PCG] Spline From Actor element - driven by new virtual method that filters data.

It hides the Mode UI completely, as the mode should always be parse components.

diff --git a/Engine/Plugins/Experimental/PCG/Source/PCG/Private/Elements/PCGDataFromActor.cpp b/Engine/Plugins/Experimental/PCG/Source/PCG/Private/Elements/PCGDataFromActor.cpp
--- a/Engine/Plugins/Experimental/PCG/Source/PCG/Private/Elements/PCGDataFromActor.cpp
+++ b/Engine/Plugins/Experimental/PCG/Source/PCG/Private/Elements/PCGDataFromActor.cpp
@@ -11,6 +11,8 @@
 
 #include UE_INLINE_GENERATED_CPP_BY_NAME(PCGDataFromActor)
 
+#define LOCTEXT_NAMESPACE "PCGDataFromActorElement"
+
 #if WITH_EDITOR
 void UPCGDataFromActorSettings::GetTrackedActorTags(FPCGTagToSettingsMap& OutTagToSettings, TArray<TObjectPtr<const UPCGGraph>>& OutVisitedGraphs) const
 {
@@ -20,6 +22,11 @@ void UPCGDataFromActorSettings::GetTrackedActorTags(FPCGTagToSettingsMap& OutTag
 		OutTagToSettings.FindOrAdd(ActorSelector.ActorSelectionTag).Add(this);
 	}
 }
+
+FText UPCGDataFromActorSettings::GetNodeTooltipText() const
+{
+	return LOCTEXT("DataFromActorTooltip", "Builds a collection of PCG-compatible data from the selected actors.");
+}
 #endif
 
 FPCGElementPtr UPCGDataFromActorSettings::CreateElement() const
@@ -229,7 +236,8 @@ void FPCGDataFromActorElement::ProcessActor(FPCGContext* Context, const UPCGData
 	else
 	{
 		const bool bParseActor = (Settings->Mode != EPCGGetDataFromActorMode::GetSinglePoint);
-		FPCGDataCollection Collection = UPCGComponent::CreateActorPCGDataCollection(FoundActor, Context->SourceComponent.Get(), bParseActor);
+		auto DataFilter = [Settings](EPCGDataType InDataType) { return Settings->DataFilter(InDataType); };
+		FPCGDataCollection Collection = UPCGComponent::CreateActorPCGDataCollection(FoundActor, Context->SourceComponent.Get(), DataFilter, bParseActor);
 		Outputs += Collection.TaggedData;
 	}
 
@@ -244,3 +252,5 @@ void FPCGDataFromActorElement::ProcessActor(FPCGContext* Context, const UPCGData
 		}
 	}
 }
+
+#undef LOCTEXT_NAMESPACE
diff --git a/Engine/Plugins/Experimental/PCG/Source/PCG/Private/Elements/PCGSplineFromActor.cpp b/Engine/Plugins/Experimental/PCG/Source/PCG/Private/Elements/PCGSplineFromActor.cpp
new file mode 100644
--- /dev/null
+++ b/Engine/Plugins/Experimental/PCG/Source/PCG/Private/Elements/PCGSplineFromActor.cpp
@@ -0,0 +1,35 @@
+// Copyright Epic Games, Inc. All Rights Reserved.
+
+#include "Elements/PCGSplineFromActor.h"
+
+#include UE_INLINE_GENERATED_CPP_BY_NAME(PCGSplineFromActor)
+
+#define LOCTEXT_NAMESPACE "PCGSplineFromActorElement"
+
+UPCGSplineFromActorSettings::UPCGSplineFromActorSettings()
+{
+	bDisplayModeSettings = false;
+	Mode = EPCGGetDataFromActorMode::ParseActorComponents;
+}
+
+TArray<FPCGPinProperties> UPCGSplineFromActorSettings::OutputPinProperties() const
+{
+	TArray<FPCGPinProperties> PinProperties;
+	PinProperties.Emplace(PCGPinConstants::DefaultOutputLabel, EPCGDataType::Spline);
+
+	return PinProperties;
+}
+
+#if WITH_EDITOR
+FText UPCGSplineFromActorSettings::GetNodeTooltipText() const
+{
+	return LOCTEXT("SplineFromActorTooltip", "Builds a collection of splines from the selected actors.");
+}
+#endif
+
+FPCGElementPtr UPCGSplineFromActorSettings::CreateElement() const
+{
+	return MakeShared<FPCGDataFromActorElement>();
+}
+
+#undef LOCTEXT_NAMESPACE
diff --git a/Engine/Plugins/Experimental/PCG/Source/PCG/Private/PCGComponent.cpp b/Engine/Plugins/Experimental/PCG/Source/PCG/Private/PCGComponent.cpp
--- a/Engine/Plugins/Experimental/PCG/Source/PCG/Private/PCGComponent.cpp
+++ b/Engine/Plugins/Experimental/PCG/Source/PCG/Private/PCGComponent.cpp
@@ -1908,7 +1908,8 @@ UPCGData* UPCGComponent::CreateActorPCGData(AActor* Actor, bool bParseActor)
 
 UPCGData* UPCGComponent::CreateActorPCGData(AActor* Actor, const UPCGComponent* Component, bool bParseActor)
 {
-	FPCGDataCollection Collection = CreateActorPCGDataCollection(Actor, Component, bParseActor);
+	auto AcceptAllData = [](EPCGDataType InDataType) { return true; };
+	FPCGDataCollection Collection = CreateActorPCGDataCollection(Actor, Component, AcceptAllData, bParseActor);
 	if (Collection.TaggedData.Num() > 1)
 	{
 		UPCGUnionData* Union = NewObject<UPCGUnionData>();
@@ -1929,7 +1930,7 @@ UPCGData* UPCGComponent::CreateActorPCGData(AActor* Actor, const UPCGComponent*
 	}
 }
 
-FPCGDataCollection UPCGComponent::CreateActorPCGDataCollection(AActor* Actor, const UPCGComponent* Component, bool bParseActor)
+FPCGDataCollection UPCGComponent::CreateActorPCGDataCollection(AActor* Actor, const UPCGComponent* Component, const TFunction<bool(EPCGDataType)>& InDataFilter, bool bParseActor)
 {
 	TRACE_CPUPROFILER_EVENT_SCOPE(UPCGComponent::CreateActorPCGData);
 	FPCGDataCollection Collection;
@@ -1947,7 +1948,10 @@ FPCGDataCollection UPCGComponent::CreateActorPCGDataCollection(AActor* Actor, co
 	// Some actor types we will forego full parsing to build strictly on the actor existence, such as partition actors, volumes and landscape
 	// TODO: add factory for extensibility
 	// TODO: review the !bParseActor cases - it might make sense to have just a point for a partition actor, even if we preintersect it.
-	if (APCGPartitionActor* PartitionActor = Cast<APCGPartitionActor>(Actor))
+	APCGPartitionActor* PartitionActor = InDataFilter(EPCGDataType::Composite) ? Cast<APCGPartitionActor>(Actor) : nullptr;
+	ALandscapeProxy* LandscapeActor = InDataFilter(EPCGDataType::Landscape) ? Cast<ALandscapeProxy>(Actor) : nullptr;
+	AVolume* VolumeActor = InDataFilter(EPCGDataType::Volume) ? Cast<AVolume>(Actor) : nullptr;
+	if (PartitionActor)
 	{
 		check(!Component || Component->GetOwner() == Actor); // Invalid processing otherwise because of the this usage
 
@@ -1962,18 +1966,18 @@ FPCGDataCollection UPCGComponent::CreateActorPCGDataCollection(AActor* Actor, co
 		TaggedData.Data = OriginalComponentSpatialData ? Cast<UPCGData>(Data->IntersectWith(OriginalComponentSpatialData)) : Cast<UPCGData>(Data);
 		// No need to keep partition actor tags, though we might want to push PCG grid GUID at some point
 	}
-	else if (ALandscapeProxy* Landscape = Cast<ALandscapeProxy>(Actor))
+	else if (LandscapeActor)
 	{
 		UPCGLandscapeData* Data = NewObject<UPCGLandscapeData>();
 		const bool bUseLandscapeMetadata = (!Component || !Component->Graph || (Component->Graph && Component->Graph->bLandscapeUsesMetadata));
 
-		Data->Initialize({ Landscape }, PCGHelpers::GetGridBounds(Actor, Component), /*bHeightOnly=*/false, bUseLandscapeMetadata);
+		Data->Initialize({ LandscapeActor }, PCGHelpers::GetGridBounds(Actor, Component), /*bHeightOnly=*/false, bUseLandscapeMetadata);
 
 		FPCGTaggedData& TaggedData = Collection.TaggedData.Emplace_GetRef();
 		TaggedData.Data = Data;
 		TaggedData.Tags = ActorTags;
 	}
-	else if (!bParseActor)
+	else if (!bParseActor && InDataFilter(EPCGDataType::Point))
 	{
 		UPCGPointData* Data = NewObject<UPCGPointData>();
 		Data->InitializeFromActor(Actor);
@@ -1982,10 +1986,10 @@ FPCGDataCollection UPCGComponent::CreateActorPCGDataCollection(AActor* Actor, co
 		TaggedData.Data = Data;
 		TaggedData.Tags = ActorTags;
 	}
-	else if (AVolume* Volume = Cast<AVolume>(Actor))
+	else if (VolumeActor)
 	{
 		UPCGVolumeData* Data = NewObject<UPCGVolumeData>();
-		Data->Initialize(Volume);
+		Data->Initialize(VolumeActor);
 
 		FPCGTaggedData& TaggedData = Collection.TaggedData.Emplace_GetRef();
 		TaggedData.Data = Data;
@@ -2032,59 +2036,65 @@ FPCGDataCollection UPCGComponent::CreateActorPCGDataCollection(AActor* Actor, co
 		RemovePCGGeneratedEntries(Shapes);
 		RemoveDuplicatesFromPrimitives(Shapes);
 
-		for (ULandscapeSplinesComponent* SplineComponent : LandscapeSplines)
+		if (InDataFilter(EPCGDataType::Spline))
 		{
-			UPCGLandscapeSplineData* SplineData = NewObject<UPCGLandscapeSplineData>();
-			SplineData->Initialize(SplineComponent);
+			for (ULandscapeSplinesComponent* SplineComponent : LandscapeSplines)
+			{
+				UPCGLandscapeSplineData* SplineData = NewObject<UPCGLandscapeSplineData>();
+				SplineData->Initialize(SplineComponent);
 
-			FPCGTaggedData& TaggedData = Collection.TaggedData.Emplace_GetRef();
-			TaggedData.Data = SplineData;
-			Algo::Transform(SplineComponent->ComponentTags, TaggedData.Tags, NameTagsToStringTags);
-			TaggedData.Tags.Append(ActorTags);
-		}
+				FPCGTaggedData& TaggedData = Collection.TaggedData.Emplace_GetRef();
+				TaggedData.Data = SplineData;
+				Algo::Transform(SplineComponent->ComponentTags, TaggedData.Tags, NameTagsToStringTags);
+				TaggedData.Tags.Append(ActorTags);
+			}
 
-		for (USplineComponent* SplineComponent : Splines)
-		{
-			UPCGSplineData* SplineData = NewObject<UPCGSplineData>();
-			SplineData->Initialize(SplineComponent);
+			for (USplineComponent* SplineComponent : Splines)
+			{
+				UPCGSplineData* SplineData = NewObject<UPCGSplineData>();
+				SplineData->Initialize(SplineComponent);
 
-			FPCGTaggedData& TaggedData = Collection.TaggedData.Emplace_GetRef();
-			TaggedData.Data = SplineData;
-			Algo::Transform(SplineComponent->ComponentTags, TaggedData.Tags, NameTagsToStringTags);
-			TaggedData.Tags.Append(ActorTags);
+				FPCGTaggedData& TaggedData = Collection.TaggedData.Emplace_GetRef();
+				TaggedData.Data = SplineData;
+				Algo::Transform(SplineComponent->ComponentTags, TaggedData.Tags, NameTagsToStringTags);
+				TaggedData.Tags.Append(ActorTags);
+			}
 		}
 
-		for (UShapeComponent* ShapeComponent : Shapes)
+		if (InDataFilter(EPCGDataType::Primitive))
 		{
-			UPCGPrimitiveData* ShapeData = NewObject<UPCGPrimitiveData>();
-			ShapeData->Initialize(ShapeComponent);
+			for (UShapeComponent* ShapeComponent : Shapes)
+			{
+				UPCGPrimitiveData* ShapeData = NewObject<UPCGPrimitiveData>();
+				ShapeData->Initialize(ShapeComponent);
 
-			FPCGTaggedData& TaggedData = Collection.TaggedData.Emplace_GetRef();
-			TaggedData.Data = ShapeData;
-			Algo::Transform(ShapeComponent->ComponentTags, TaggedData.Tags, NameTagsToStringTags);
-			TaggedData.Tags.Append(ActorTags);
-		}
+				FPCGTaggedData& TaggedData = Collection.TaggedData.Emplace_GetRef();
+				TaggedData.Data = ShapeData;
+				Algo::Transform(ShapeComponent->ComponentTags, TaggedData.Tags, NameTagsToStringTags);
+				TaggedData.Tags.Append(ActorTags);
+			}
 
-		for (UPrimitiveComponent* PrimitiveComponent : Primitives)
-		{
-			// Exception: skip the billboard component
-			if (Cast<UBillboardComponent>(PrimitiveComponent))
+			for (UPrimitiveComponent* PrimitiveComponent : Primitives)
 			{
-				continue;
-			}
+				// Exception: skip the billboard component
+				if (Cast<UBillboardComponent>(PrimitiveComponent))
+				{
+					continue;
+				}
 
-			UPCGPrimitiveData* PrimitiveData = NewObject<UPCGPrimitiveData>();
-			PrimitiveData->Initialize(PrimitiveComponent);
+				UPCGPrimitiveData* PrimitiveData = NewObject<UPCGPrimitiveData>();
+				PrimitiveData->Initialize(PrimitiveComponent);
 
-			FPCGTaggedData& TaggedData = Collection.TaggedData.Emplace_GetRef();
-			TaggedData.Data = PrimitiveData;
-			Algo::Transform(PrimitiveComponent->ComponentTags, TaggedData.Tags, NameTagsToStringTags);
-			TaggedData.Tags.Append(ActorTags);
+				FPCGTaggedData& TaggedData = Collection.TaggedData.Emplace_GetRef();
+				TaggedData.Data = PrimitiveData;
+				Algo::Transform(PrimitiveComponent->ComponentTags, TaggedData.Tags, NameTagsToStringTags);
+				TaggedData.Tags.Append(ActorTags);
+			}
 		}
 	}
 
 	// Finally, if it's not a special actor and there are not parsed components, then return a single point at the actor position
-	if (Collection.TaggedData.IsEmpty())
+	if (Collection.TaggedData.IsEmpty() && InDataFilter(EPCGDataType::Point))
 	{
 		UPCGPointData* Data = NewObject<UPCGPointData>();
 		Data->InitializeFromActor(Actor);
diff --git a/Engine/Plugins/Experimental/PCG/Source/PCG/Public/Elements/PCGDataFromActor.h b/Engine/Plugins/Experimental/PCG/Source/PCG/Public/Elements/PCGDataFromActor.h
--- a/Engine/Plugins/Experimental/PCG/Source/PCG/Public/Elements/PCGDataFromActor.h
+++ b/Engine/Plugins/Experimental/PCG/Source/PCG/Public/Elements/PCGDataFromActor.h
@@ -19,6 +19,7 @@ enum class EPCGGetDataFromActorMode : uint8
 	GetDataFromPCGComponentOrParseComponents
 };
 
+/** Builds a collection of PCG-compatible data from the selected actors. */
 UCLASS(BlueprintType, ClassGroup = (Procedural))
 class PCG_API UPCGDataFromActorSettings : public UPCGSettings
 {
@@ -27,7 +28,8 @@ class PCG_API UPCGDataFromActorSettings : public UPCGSettings
 public:
 	//~Begin UPCGSettings interface
 #if WITH_EDITOR
-	virtual FName GetDefaultNodeName() const override { return FName(TEXT("DataFromActor")); }
+	virtual FName GetDefaultNodeName() const override { return FName(TEXT("Data From Actor")); }
+	virtual FText GetNodeTooltipText() const override;
 	virtual EPCGSettingsType GetType() const override { return EPCGSettingsType::Spatial; }
 	virtual void GetTrackedActorTags(FPCGTagToSettingsMap& OutTagToSettings, TArray<TObjectPtr<const UPCGGraph>>& OutVisitedGraphs) const override;
 #endif
@@ -40,16 +42,23 @@ protected:
 	//~End UPCGSettings
 
 public:
+	/** Override this to filter what kinds of data should be retrieved from the actor(s). */
+	virtual bool DataFilter(EPCGDataType InDataType) const { return true; }
+
 	UPROPERTY(BlueprintReadWrite, EditAnywhere, Category = Settings, meta = (ShowOnlyInnerProperties))
 	FPCGActorSelectorSettings ActorSelector;
 
-	UPROPERTY(BlueprintReadWrite, EditAnywhere, Category = Settings)
+	UPROPERTY(BlueprintReadWrite, EditAnywhere, Category = Settings, meta = (EditCondition = bDisplayModeSettings, EditConditionHides, HideEditConditionToggle))
 	EPCGGetDataFromActorMode Mode = EPCGGetDataFromActorMode::ParseActorComponents;
 
-	UPROPERTY(BlueprintReadWrite, EditAnywhere, Category = Settings, meta = (EditCondition = "Mode == EPCGGetDataFromActorMode::GetDataFromPCGComponent || Mode == EPCGGetDataFromActorMode::GetDataFromPCGComponentOrParseComponents"))
+	// This can be set false by inheriting nodes to hide the 'Mode' property.
+	UPROPERTY(Transient, meta = (EditCondition = false, EditConditionHides))
+	bool bDisplayModeSettings = true;
+
+	UPROPERTY(BlueprintReadWrite, EditAnywhere, Category = Settings, meta = (EditCondition = "Mode == EPCGGetDataFromActorMode::GetDataFromPCGComponent || Mode == EPCGGetDataFromActorMode::GetDataFromPCGComponentOrParseComponents", EditConditionHides))
 	TArray<FName> ExpectedPins;
 
-	UPROPERTY(BlueprintReadWrite, EditAnywhere, Category = Settings, meta = (EditCondition = "Mode == EPCGGetDataFromActorMode::GetDataFromProperty"))
+	UPROPERTY(BlueprintReadWrite, EditAnywhere, Category = Settings, meta = (EditCondition = "Mode == EPCGGetDataFromActorMode::GetDataFromProperty", EditConditionHides))
 	FName PropertyName = NAME_None;
 };
 
diff --git a/Engine/Plugins/Experimental/PCG/Source/PCG/Public/Elements/PCGSplineFromActor.h b/Engine/Plugins/Experimental/PCG/Source/PCG/Public/Elements/PCGSplineFromActor.h
new file mode 100644
--- /dev/null
+++ b/Engine/Plugins/Experimental/PCG/Source/PCG/Public/Elements/PCGSplineFromActor.h
@@ -0,0 +1,35 @@
+// Copyright Epic Games, Inc. All Rights Reserved.
+
+#pragma once
+
+#include "Elements/PCGDataFromActor.h"
+
+#include "PCGSplineFromActor.generated.h"
+
+/** Builds a collection of spline data from the selected actors. */
+UCLASS(BlueprintType, ClassGroup = (Procedural))
+class PCG_API UPCGSplineFromActorSettings : public UPCGDataFromActorSettings
+{
+	GENERATED_BODY()
+
+public:
+	UPCGSplineFromActorSettings();
+
+	//~Begin UPCGSettings interface
+#if WITH_EDITOR
+	virtual FName GetDefaultNodeName() const override { return FName(TEXT("Spline From Actor")); }
+	virtual FText GetNodeTooltipText() const override;
+#endif
+
+protected:
+	virtual TArray<FPCGPinProperties> OutputPinProperties() const override;
+
+	/** Override this to filter what kinds of data should be retrieved from the actor(s). */
+	virtual FPCGElementPtr CreateElement() const override;
+	//~End UPCGSettings
+
+public:
+	//~Begin UPCGDataFromActorSettings interface
+	virtual bool DataFilter(EPCGDataType InDataType) const { return !!(InDataType & EPCGDataType::PolyLine); }
+	//~End UPCGDataFromActorSettings
+};
diff --git a/Engine/Plugins/Experimental/PCG/Source/PCG/Public/PCGComponent.h b/Engine/Plugins/Experimental/PCG/Source/PCG/Public/PCGComponent.h
--- a/Engine/Plugins/Experimental/PCG/Source/PCG/Public/PCGComponent.h
+++ b/Engine/Plugins/Experimental/PCG/Source/PCG/Public/PCGComponent.h
@@ -243,7 +243,7 @@ public:
 	FBox GetGridBounds() const;
 
 	/** Builds the PCG data from a given actor and its PCG component, and places it in a data collection with appropriate tags */
-	static FPCGDataCollection CreateActorPCGDataCollection(AActor* Actor, const UPCGComponent* Component, bool bParseActor = true);
+	static FPCGDataCollection CreateActorPCGDataCollection(AActor* Actor, const UPCGComponent* Component, const TFunction<bool(EPCGDataType)>& InDataFilter, bool bParseActor = true);
 
 	/** Builds the canonical PCG data from a given actor and its PCG component if any. */
 	static UPCGData* CreateActorPCGData(AActor* Actor, const UPCGComponent* Component, bool bParseActor = true);
