[PCG] Nodes that will convert their input to points should declare the input type as Point - make the collapse visible to user.

Follows similar change in SpawnActor.

Change went into the mesh spawner a little deeper, to bubble collapse upwards as far as possible.

diff --git a/Engine/Plugins/Experimental/PCG/Source/PCG/Private/Elements/PCGDensityFilter.cpp b/Engine/Plugins/Experimental/PCG/Source/PCG/Private/Elements/PCGDensityFilter.cpp
--- a/Engine/Plugins/Experimental/PCG/Source/PCG/Private/Elements/PCGDensityFilter.cpp
+++ b/Engine/Plugins/Experimental/PCG/Source/PCG/Private/Elements/PCGDensityFilter.cpp
@@ -10,15 +10,6 @@
 
 #include UE_INLINE_GENERATED_CPP_BY_NAME(PCGDensityFilter)
 
-TArray<FPCGPinProperties> UPCGDensityFilterSettings::InputPinProperties() const
-{
-	TArray<FPCGPinProperties> PinProperties;
-	// TODO in the future type checking of edges will be stricter and a conversion node will be added to convert from other types
-	PinProperties.Emplace(PCGPinConstants::DefaultInputLabel, EPCGDataType::Point);
-
-	return PinProperties;
-}
-
 FPCGElementPtr UPCGDensityFilterSettings::CreateElement() const
 {
 	return MakeShared<FPCGDensityFilterElement>();
diff --git a/Engine/Plugins/Experimental/PCG/Source/PCG/Private/Elements/PCGDensityRemap.cpp b/Engine/Plugins/Experimental/PCG/Source/PCG/Private/Elements/PCGDensityRemap.cpp
--- a/Engine/Plugins/Experimental/PCG/Source/PCG/Private/Elements/PCGDensityRemap.cpp
+++ b/Engine/Plugins/Experimental/PCG/Source/PCG/Private/Elements/PCGDensityRemap.cpp
@@ -13,14 +13,6 @@
 
 #include UE_INLINE_GENERATED_CPP_BY_NAME(PCGDensityRemap)
 
-TArray<FPCGPinProperties> UPCGLinearDensityRemapSettings::InputPinProperties() const
-{
-	TArray<FPCGPinProperties> PinProperties;
-	PinProperties.Emplace(PCGPinConstants::DefaultInputLabel, EPCGDataType::Spatial, /*bInAllowMultipleConnections=*/ true);
-
-	return PinProperties;
-}
-
 UPCGLinearDensityRemapSettings::UPCGLinearDensityRemapSettings()
 {
 	bUseSeed = true;
diff --git a/Engine/Plugins/Experimental/PCG/Source/PCG/Private/Elements/PCGDistance.cpp b/Engine/Plugins/Experimental/PCG/Source/PCG/Private/Elements/PCGDistance.cpp
--- a/Engine/Plugins/Experimental/PCG/Source/PCG/Private/Elements/PCGDistance.cpp
+++ b/Engine/Plugins/Experimental/PCG/Source/PCG/Private/Elements/PCGDistance.cpp
@@ -43,9 +43,8 @@ namespace PCGDistance
 TArray<FPCGPinProperties> UPCGDistanceSettings::InputPinProperties() const
 {
 	TArray<FPCGPinProperties> PinProperties;
-	PinProperties.Emplace(PCGDistance::SourceLabel, EPCGDataType::Spatial);
-	PinProperties.Emplace(PCGDistance::TargetLabel, EPCGDataType::Spatial);
-	PinProperties.Emplace(PCGPinConstants::DefaultParamsLabel, EPCGDataType::Param, /*bInAllowMultipleConnections*/ false);
+	PinProperties.Emplace(PCGDistance::SourceLabel, EPCGDataType::Point);
+	PinProperties.Emplace(PCGDistance::TargetLabel, EPCGDataType::Point);
 
 	return PinProperties;
 }
@@ -88,7 +87,6 @@ bool FPCGDistanceElement::ExecuteInternal(FPCGContext* Context) const
 
 	for (const FPCGTaggedData& Target : Targets)
 	{
-
 		const UPCGSpatialData* TargetData = Cast<UPCGSpatialData>(Target.Data);
 
 		if (!TargetData)
diff --git a/Engine/Plugins/Experimental/PCG/Source/PCG/Private/Elements/PCGSelfPruning.cpp b/Engine/Plugins/Experimental/PCG/Source/PCG/Private/Elements/PCGSelfPruning.cpp
--- a/Engine/Plugins/Experimental/PCG/Source/PCG/Private/Elements/PCGSelfPruning.cpp
+++ b/Engine/Plugins/Experimental/PCG/Source/PCG/Private/Elements/PCGSelfPruning.cpp
@@ -243,15 +243,6 @@ namespace PCGSelfPruningElement
 	}
 }
 
-TArray<FPCGPinProperties> UPCGSelfPruningSettings::InputPinProperties() const
-{
-	TArray<FPCGPinProperties> PinProperties;
-	// TODO in the future type checking of edges will be stricter and a conversion node will be added to convert from other types
-	PinProperties.Emplace(PCGPinConstants::DefaultInputLabel, EPCGDataType::Point);
-
-	return PinProperties;
-}
-
 FPCGElementPtr UPCGSelfPruningSettings::CreateElement() const
 {
 	return MakeShared<FPCGSelfPruningElement>();
diff --git a/Engine/Plugins/Experimental/PCG/Source/PCG/Private/Elements/PCGStaticMeshSpawner.cpp b/Engine/Plugins/Experimental/PCG/Source/PCG/Private/Elements/PCGStaticMeshSpawner.cpp
--- a/Engine/Plugins/Experimental/PCG/Source/PCG/Private/Elements/PCGStaticMeshSpawner.cpp
+++ b/Engine/Plugins/Experimental/PCG/Source/PCG/Private/Elements/PCGStaticMeshSpawner.cpp
@@ -66,15 +66,20 @@ bool FPCGStaticMeshSpawnerElement::PrepareDataInternal(FPCGContext* InContext) c
 	for (const FPCGTaggedData& Input : Inputs)
 	{
 		const UPCGSpatialData* SpatialData = Cast<UPCGSpatialData>(Input.Data);
-
 		if (!SpatialData)
 		{
 			PCGE_LOG(Error, "Invalid input data");
 			continue;
 		}
 
-		AActor* TargetActor = SpatialData->TargetActor.Get();
+		const UPCGPointData* PointData = SpatialData->ToPointData(Context);
+		if (!PointData)
+		{
+			PCGE_LOG(Error, "Unable to get point data from input");
+			continue;
+		}
 
+		AActor* TargetActor = PointData->TargetActor.Get();
 		if (!TargetActor)
 		{
 			PCGE_LOG(Error, "Invalid target actor");
@@ -88,7 +93,7 @@ bool FPCGStaticMeshSpawnerElement::PrepareDataInternal(FPCGContext* InContext) c
 			FPCGTaggedData& Output = Outputs.Add_GetRef(Input);
 
 			OutputPointData = NewObject<UPCGPointData>();
-			OutputPointData->InitializeFromData(SpatialData);
+			OutputPointData->InitializeFromData(PointData);
 
 			if (OutputPointData->Metadata->HasAttribute(Settings->OutAttributeName))
 			{
@@ -102,7 +107,7 @@ bool FPCGStaticMeshSpawnerElement::PrepareDataInternal(FPCGContext* InContext) c
 		}
 
 		TArray<FPCGMeshInstanceList> MeshInstances;
-		Settings->MeshSelectorInstance->SelectInstances(*Context, Settings, SpatialData, MeshInstances, OutputPointData);
+		Settings->MeshSelectorInstance->SelectInstances(*Context, Settings, PointData, MeshInstances, OutputPointData);
 
 		TArray<FPCGPackedCustomData> PackedCustomData;
 		PackedCustomData.SetNum(MeshInstances.Num());
@@ -110,12 +115,12 @@ bool FPCGStaticMeshSpawnerElement::PrepareDataInternal(FPCGContext* InContext) c
 		{
 			for(int32 InstanceListIndex = 0; InstanceListIndex < MeshInstances.Num(); ++InstanceListIndex)
 			{
-				Settings->InstancePackerInstance->PackInstances(*Context, SpatialData, MeshInstances[InstanceListIndex], PackedCustomData[InstanceListIndex]);
+				Settings->InstancePackerInstance->PackInstances(*Context, PointData, MeshInstances[InstanceListIndex], PackedCustomData[InstanceListIndex]);
 			}
 		}
 
 		FPCGStaticMeshSpawnerContext::FPackedInstanceListData& InstanceListData = Context->MeshInstancesData.Emplace_GetRef();
-		InstanceListData.SpatialData = SpatialData;
+		InstanceListData.SpatialData = PointData;
 		InstanceListData.MeshInstances = MoveTemp(MeshInstances);
 		InstanceListData.PackedCustomData = MoveTemp(PackedCustomData);
 	}
diff --git a/Engine/Plugins/Experimental/PCG/Source/PCG/Private/Elements/PCGTransformPoints.cpp b/Engine/Plugins/Experimental/PCG/Source/PCG/Private/Elements/PCGTransformPoints.cpp
--- a/Engine/Plugins/Experimental/PCG/Source/PCG/Private/Elements/PCGTransformPoints.cpp
+++ b/Engine/Plugins/Experimental/PCG/Source/PCG/Private/Elements/PCGTransformPoints.cpp
@@ -17,23 +17,6 @@ UPCGTransformPointsSettings::UPCGTransformPointsSettings()
 	bUseSeed = true;
 }
 
-TArray<FPCGPinProperties> UPCGTransformPointsSettings::InputPinProperties() const
-{
-	TArray<FPCGPinProperties> PinProperties;
-	// TODO in the future type checking of edges will be stricter and a conversion node will be added to convert from other types
-	PinProperties.Emplace(PCGPinConstants::DefaultInputLabel, EPCGDataType::Point);
-
-	return PinProperties;
-}
-
-TArray<FPCGPinProperties> UPCGTransformPointsSettings::OutputPinProperties() const
-{
-	TArray<FPCGPinProperties> PinProperties;
-	PinProperties.Emplace(PCGPinConstants::DefaultOutputLabel, EPCGDataType::Point);
-
-	return PinProperties;
-}
-
 FPCGElementPtr UPCGTransformPointsSettings::CreateElement() const
 {
 	return MakeShared<FPCGTransformPointsElement>();
diff --git a/Engine/Plugins/Experimental/PCG/Source/PCG/Private/MeshSelectors/PCGMeshSelectorByAttribute.cpp b/Engine/Plugins/Experimental/PCG/Source/PCG/Private/MeshSelectors/PCGMeshSelectorByAttribute.cpp
--- a/Engine/Plugins/Experimental/PCG/Source/PCG/Private/MeshSelectors/PCGMeshSelectorByAttribute.cpp
+++ b/Engine/Plugins/Experimental/PCG/Source/PCG/Private/MeshSelectors/PCGMeshSelectorByAttribute.cpp
@@ -25,31 +25,29 @@ void UPCGMeshSelectorByAttribute::PostLoad()
 void UPCGMeshSelectorByAttribute::SelectInstances_Implementation(
 	FPCGContext& Context, 
 	const UPCGStaticMeshSpawnerSettings* Settings, 
-	const UPCGSpatialData* InSpatialData, 
+	const UPCGPointData* InPointData,
 	TArray<FPCGMeshInstanceList>& OutMeshInstances,
 	UPCGPointData* OutPointData) const
 {
-	const UPCGPointData* PointData = InSpatialData->ToPointData(&Context);
-
-	if (!PointData)
+	if (!InPointData)
 	{
-		PCGE_LOG_C(Error, &Context, "Unable to get point data from input");
+		PCGE_LOG_C(Error, &Context, "Missing input data");
 		return;
 	}
 
-	if (!PointData->Metadata)
+	if (!InPointData->Metadata)
 	{
 		PCGE_LOG_C(Error, &Context, "Unable to get metadata from input");
 		return;
 	}
 
-	if (!PointData->Metadata->HasAttribute(AttributeName)) 
+	if (!InPointData->Metadata->HasAttribute(AttributeName))
 	{
 		PCGE_LOG_C(Error, &Context, "Attribute %s is not in the metadata", *AttributeName.ToString());
 		return;
 	}
 
-	const FPCGMetadataAttributeBase* AttributeBase = PointData->Metadata->GetConstAttribute(AttributeName);
+	const FPCGMetadataAttributeBase* AttributeBase = InPointData->Metadata->GetConstAttribute(AttributeName);
 	check(AttributeBase);
 
 	if (AttributeBase->GetTypeId() != PCG::Private::MetadataTypes<FString>::Id)
@@ -60,7 +58,7 @@ void UPCGMeshSelectorByAttribute::SelectInstances_Implementation(
 
 	const FPCGMetadataAttribute<FString>* Attribute = static_cast<const FPCGMetadataAttribute<FString>*>(AttributeBase);
 
-	FPCGMeshMaterialOverrideHelper MaterialOverrideHelper(Context, MaterialOverrideMode, MaterialOverrides, MaterialOverrideAttributes, PointData->Metadata);
+	FPCGMeshMaterialOverrideHelper MaterialOverrideHelper(Context, MaterialOverrideMode, MaterialOverrides, MaterialOverrideAttributes, InPointData->Metadata);
 
 	if (!MaterialOverrideHelper.IsValid())
 	{
@@ -72,15 +70,15 @@ void UPCGMeshSelectorByAttribute::SelectInstances_Implementation(
 	// ByAttribute takes in SoftObjectPaths per point in the metadata, so we can pass those directly into the outgoing pin if it exists
 	if (OutPointData)
 	{
-		OutPointData->SetPoints(PointData->GetPoints()); 
+		OutPointData->SetPoints(InPointData->GetPoints());
 		OutPointData->Metadata->DeleteAttribute(Settings->OutAttributeName);
-		OutPointData->Metadata->CopyAttribute(PointData->Metadata, AttributeName, Settings->OutAttributeName);
+		OutPointData->Metadata->CopyAttribute(InPointData->Metadata, AttributeName, Settings->OutAttributeName);
 	}
 
 	TRACE_CPUPROFILER_EVENT_SCOPE(FPCGStaticMeshSpawnerElement::Execute::SelectEntries);
 
 	// Assign points to entries
-	for (const FPCGPoint& Point : PointData->GetPoints()) 
+	for (const FPCGPoint& Point : InPointData->GetPoints())
 	{
 		if (Point.Density <= 0.0f)
 		{
@@ -117,4 +115,3 @@ void UPCGMeshSelectorByAttribute::SelectInstances_Implementation(
 		OutMeshInstances[Index].Instances.Emplace(Point);
 	}
 }
-
diff --git a/Engine/Plugins/Experimental/PCG/Source/PCG/Private/MeshSelectors/PCGMeshSelectorWeighted.cpp b/Engine/Plugins/Experimental/PCG/Source/PCG/Private/MeshSelectors/PCGMeshSelectorWeighted.cpp
--- a/Engine/Plugins/Experimental/PCG/Source/PCG/Private/MeshSelectors/PCGMeshSelectorWeighted.cpp
+++ b/Engine/Plugins/Experimental/PCG/Source/PCG/Private/MeshSelectors/PCGMeshSelectorWeighted.cpp
@@ -73,7 +73,7 @@ namespace PCGMeshSelectorWeighted
 void UPCGMeshSelectorWeighted::SelectInstances_Implementation(
 	FPCGContext& Context, 
 	const UPCGStaticMeshSpawnerSettings* Settings, 
-	const UPCGSpatialData* InSpatialData,
+	const UPCGPointData* InPointData,
 	TArray<FPCGMeshInstanceList>& OutMeshInstances,
 	UPCGPointData* OutPointData) const
 {
@@ -103,11 +103,9 @@ void UPCGMeshSelectorWeighted::SelectInstances_Implementation(
 		return;
 	}
 
-	const UPCGPointData* PointData = InSpatialData->ToPointData(&Context);
-
-	if (!PointData)
+	if (!InPointData)
 	{
-		PCGE_LOG_C(Error, &Context, "Unable to get point data from input");
+		PCGE_LOG_C(Error, &Context, "Missing input data");
 		return;
 	}
 
@@ -115,7 +113,7 @@ void UPCGMeshSelectorWeighted::SelectInstances_Implementation(
 	FPCGMetadataAttribute<FString>* OutAttribute = nullptr;
 	TMap<TSoftObjectPtr<UStaticMesh>, PCGMetadataValueKey> MeshToValueKey;
 
-	FPCGMeshMaterialOverrideHelper MaterialOverrideHelper(Context, bUseAttributeMaterialOverrides, MaterialOverrideAttributes, PointData->Metadata);
+	FPCGMeshMaterialOverrideHelper MaterialOverrideHelper(Context, bUseAttributeMaterialOverrides, MaterialOverrideAttributes, InPointData->Metadata);
 
 	if (!MaterialOverrideHelper.IsValid())
 	{
@@ -151,7 +149,7 @@ void UPCGMeshSelectorWeighted::SelectInstances_Implementation(
 	{
 		TRACE_CPUPROFILER_EVENT_SCOPE(FPCGStaticMeshSpawnerElement::Execute::SelectEntries);
 
-		for (const FPCGPoint& Point : PointData->GetPoints()) 
+		for (const FPCGPoint& Point : InPointData->GetPoints())
 		{
 			if (Point.Density <= 0.0f)
 			{
diff --git a/Engine/Plugins/Experimental/PCG/Source/PCG/Private/MeshSelectors/PCGMeshSelectorWeightedByCategory.cpp b/Engine/Plugins/Experimental/PCG/Source/PCG/Private/MeshSelectors/PCGMeshSelectorWeightedByCategory.cpp
--- a/Engine/Plugins/Experimental/PCG/Source/PCG/Private/MeshSelectors/PCGMeshSelectorWeightedByCategory.cpp
+++ b/Engine/Plugins/Experimental/PCG/Source/PCG/Private/MeshSelectors/PCGMeshSelectorWeightedByCategory.cpp
@@ -21,31 +21,29 @@ struct FPCGInstancesAndWeights
 void UPCGMeshSelectorWeightedByCategory::SelectInstances_Implementation(
 	FPCGContext& Context,
 	const UPCGStaticMeshSpawnerSettings* Settings,
-	const UPCGSpatialData* InSpatialData,
+	const UPCGPointData* InPointData,
 	TArray<FPCGMeshInstanceList>& OutMeshInstances,
 	UPCGPointData* OutPointData) const
 {
-	const UPCGPointData* PointData = InSpatialData->ToPointData(&Context);
-
-	if (!PointData)
+	if (!InPointData)
 	{
-		PCGE_LOG_C(Error, &Context, "Unable to get point data from input");
+		PCGE_LOG_C(Error, &Context, "Missing input data");
 		return;
 	}
 
-	if (!PointData->Metadata)
+	if (!InPointData->Metadata)
 	{
 		PCGE_LOG_C(Error, &Context, "Unable to get metadata from input");
 		return;
 	}
 
-	if (!PointData->Metadata->HasAttribute(CategoryAttribute))
+	if (!InPointData->Metadata->HasAttribute(CategoryAttribute))
 	{
 		PCGE_LOG_C(Error, &Context, "Attribute %s is not in the metadata", *CategoryAttribute.ToString());
 		return;
 	}
 
-	const FPCGMetadataAttributeBase* AttributeBase = PointData->Metadata->GetConstAttribute(CategoryAttribute);
+	const FPCGMetadataAttributeBase* AttributeBase = InPointData->Metadata->GetConstAttribute(CategoryAttribute);
 	check(AttributeBase);
 
 	// TODO: support enum type as well
@@ -129,7 +127,7 @@ void UPCGMeshSelectorWeightedByCategory::SelectInstances_Implementation(
 	FPCGMetadataAttribute<FString>* OutAttribute = nullptr;
 	TMap<TSoftObjectPtr<UStaticMesh>, PCGMetadataValueKey> MeshToValueKey;
 
-	FPCGMeshMaterialOverrideHelper MaterialOverrideHelper(Context, bUseAttributeMaterialOverrides, MaterialOverrideAttributes, PointData->Metadata);
+	FPCGMeshMaterialOverrideHelper MaterialOverrideHelper(Context, bUseAttributeMaterialOverrides, MaterialOverrideAttributes, InPointData->Metadata);
 
 	if (!MaterialOverrideHelper.IsValid())
 	{
@@ -164,7 +162,7 @@ void UPCGMeshSelectorWeightedByCategory::SelectInstances_Implementation(
 	TRACE_CPUPROFILER_EVENT_SCOPE(FPCGStaticMeshSpawnerElement::Execute::SelectEntries);
 
 	// Assign points to entries
-	for (const FPCGPoint& Point : PointData->GetPoints())
+	for (const FPCGPoint& Point : InPointData->GetPoints())
 	{
 		if (Point.Density <= 0.0f)
 		{
diff --git a/Engine/Plugins/Experimental/PCG/Source/PCG/Public/Elements/PCGDebugElement.h b/Engine/Plugins/Experimental/PCG/Source/PCG/Public/Elements/PCGDebugElement.h
--- a/Engine/Plugins/Experimental/PCG/Source/PCG/Public/Elements/PCGDebugElement.h
+++ b/Engine/Plugins/Experimental/PCG/Source/PCG/Public/Elements/PCGDebugElement.h
@@ -23,9 +23,10 @@ public:
 	virtual EPCGSettingsType GetType() const override { return EPCGSettingsType::Debug; }
 #endif
 
+protected:
+	virtual TArray<FPCGPinProperties> InputPinProperties() const override { return Super::DefaultPointInputPinProperties(); }
 	virtual TArray<FPCGPinProperties> OutputPinProperties() const override;
 
-protected:
 	virtual FPCGElementPtr CreateElement() const override;
 	// ~End UPCGSettings interface
 };
diff --git a/Engine/Plugins/Experimental/PCG/Source/PCG/Public/Elements/PCGDensityFilter.h b/Engine/Plugins/Experimental/PCG/Source/PCG/Public/Elements/PCGDensityFilter.h
--- a/Engine/Plugins/Experimental/PCG/Source/PCG/Public/Elements/PCGDensityFilter.h
+++ b/Engine/Plugins/Experimental/PCG/Source/PCG/Public/Elements/PCGDensityFilter.h
@@ -16,14 +16,14 @@ class PCG_API UPCGDensityFilterSettings : public UPCGSettings
 public:
 	//~Begin UPCGSettings interface
 #if WITH_EDITOR
-	virtual FName GetDefaultNodeName() const override { return FName(TEXT("DensityFilter")); }
+	virtual FName GetDefaultNodeName() const override { return FName(TEXT("Density Filter")); }
 	virtual EPCGSettingsType GetType() const override { return EPCGSettingsType::Filter; }
 #endif
 
-	virtual TArray<FPCGPinProperties> InputPinProperties() const override;
+protected:
+	virtual TArray<FPCGPinProperties> InputPinProperties() const override { return Super::DefaultPointInputPinProperties(); }
 	virtual TArray<FPCGPinProperties> OutputPinProperties() const override { return Super::DefaultPointOutputPinProperties(); }
 
-protected:
 	virtual FPCGElementPtr CreateElement() const override;
 	//~End UPCGSettings interface
 
diff --git a/Engine/Plugins/Experimental/PCG/Source/PCG/Public/Elements/PCGDensityRemap.h b/Engine/Plugins/Experimental/PCG/Source/PCG/Public/Elements/PCGDensityRemap.h
--- a/Engine/Plugins/Experimental/PCG/Source/PCG/Public/Elements/PCGDensityRemap.h
+++ b/Engine/Plugins/Experimental/PCG/Source/PCG/Public/Elements/PCGDensityRemap.h
@@ -17,14 +17,14 @@ public:
 
 	// ~Begin UPCGSettings interface
 #if WITH_EDITOR
-	virtual FName GetDefaultNodeName() const override { return FName(TEXT("LinearDensityRemap")); }
+	virtual FName GetDefaultNodeName() const override { return FName(TEXT("Linear Density Remap")); }
 	virtual EPCGSettingsType GetType() const override { return EPCGSettingsType::Density; }
 #endif
 
-	virtual TArray<FPCGPinProperties> InputPinProperties() const override;
+protected:
+	virtual TArray<FPCGPinProperties> InputPinProperties() const override { return Super::DefaultPointInputPinProperties(); }
 	virtual TArray<FPCGPinProperties> OutputPinProperties() const override { return Super::DefaultPointOutputPinProperties(); }
 
-protected:
 	virtual FPCGElementPtr CreateElement() const override;
 	//~End UPCGSettings interface
 
diff --git a/Engine/Plugins/Experimental/PCG/Source/PCG/Public/Elements/PCGSelfPruning.h b/Engine/Plugins/Experimental/PCG/Source/PCG/Public/Elements/PCGSelfPruning.h
--- a/Engine/Plugins/Experimental/PCG/Source/PCG/Public/Elements/PCGSelfPruning.h
+++ b/Engine/Plugins/Experimental/PCG/Source/PCG/Public/Elements/PCGSelfPruning.h
@@ -39,14 +39,14 @@ public:
 
 #if WITH_EDITOR
 	//~Begin UPCGSettings interface
-	virtual FName GetDefaultNodeName() const override { return FName(TEXT("SelfPruning")); }
+	virtual FName GetDefaultNodeName() const override { return FName(TEXT("Self Pruning")); }
 	virtual EPCGSettingsType GetType() const override { return EPCGSettingsType::Filter; }
 #endif
 
-	virtual TArray<FPCGPinProperties> InputPinProperties() const override;
+protected:
+	virtual TArray<FPCGPinProperties> InputPinProperties() const override { return Super::DefaultPointInputPinProperties(); }
 	virtual TArray<FPCGPinProperties> OutputPinProperties() const override { return Super::DefaultPointOutputPinProperties(); }
 
-protected:
 	virtual FPCGElementPtr CreateElement() const override;
 	// ~End UPCGSettings interface
 };
diff --git a/Engine/Plugins/Experimental/PCG/Source/PCG/Public/Elements/PCGStaticMeshSpawner.h b/Engine/Plugins/Experimental/PCG/Source/PCG/Public/Elements/PCGStaticMeshSpawner.h
--- a/Engine/Plugins/Experimental/PCG/Source/PCG/Public/Elements/PCGStaticMeshSpawner.h
+++ b/Engine/Plugins/Experimental/PCG/Source/PCG/Public/Elements/PCGStaticMeshSpawner.h
@@ -50,9 +50,10 @@ public:
 	virtual EPCGSettingsType GetType() const override { return EPCGSettingsType::Spawner; }
 #endif
 
+protected:
+	virtual TArray<FPCGPinProperties> InputPinProperties() const override { return Super::DefaultPointInputPinProperties(); }
 	virtual TArray<FPCGPinProperties> OutputPinProperties() const override { return Super::DefaultPointOutputPinProperties(); }
 
-protected:
 	virtual FPCGElementPtr CreateElement() const override;
 	// ~End UPCGSettings interface
 
diff --git a/Engine/Plugins/Experimental/PCG/Source/PCG/Public/Elements/PCGTransformPoints.h b/Engine/Plugins/Experimental/PCG/Source/PCG/Public/Elements/PCGTransformPoints.h
--- a/Engine/Plugins/Experimental/PCG/Source/PCG/Public/Elements/PCGTransformPoints.h
+++ b/Engine/Plugins/Experimental/PCG/Source/PCG/Public/Elements/PCGTransformPoints.h
@@ -17,14 +17,14 @@ public:
 
 	//~Begin UPCGSettings interface
 #if WITH_EDITOR
-	virtual FName GetDefaultNodeName() const override { return FName(TEXT("TransformPoints")); }
+	virtual FName GetDefaultNodeName() const override { return FName(TEXT("Transform Points")); }
 	virtual EPCGSettingsType GetType() const override { return EPCGSettingsType::Spatial; }
 #endif
 
-	virtual TArray<FPCGPinProperties> InputPinProperties() const override;
-	virtual TArray<FPCGPinProperties> OutputPinProperties() const override;
-
 protected:
+	virtual TArray<FPCGPinProperties> InputPinProperties() const override { return Super::DefaultPointInputPinProperties(); }
+	virtual TArray<FPCGPinProperties> OutputPinProperties() const override { return Super::DefaultPointOutputPinProperties(); }
+
 	virtual FPCGElementPtr CreateElement() const override;
 	//~End UPCGSettings interface
 
diff --git a/Engine/Plugins/Experimental/PCG/Source/PCG/Public/MeshSelectors/PCGMeshSelectorBase.h b/Engine/Plugins/Experimental/PCG/Source/PCG/Public/MeshSelectors/PCGMeshSelectorBase.h
--- a/Engine/Plugins/Experimental/PCG/Source/PCG/Public/MeshSelectors/PCGMeshSelectorBase.h
+++ b/Engine/Plugins/Experimental/PCG/Source/PCG/Public/MeshSelectors/PCGMeshSelectorBase.h
@@ -120,14 +120,14 @@ public:
 	void SelectInstances(
 		FPCGContext& Context,
 		const UPCGStaticMeshSpawnerSettings* Settings,
-		const UPCGSpatialData* InSpatialData,
+		const UPCGPointData* InPointData,
 		TArray<FPCGMeshInstanceList>& OutMeshInstances,
 		UPCGPointData* OutPointData) const;
 
 	virtual void SelectInstances_Implementation(
 		FPCGContext& Context,
 		const UPCGStaticMeshSpawnerSettings* Settings,
-		const UPCGSpatialData* InSpatialData,
+		const UPCGPointData* InPointData,
 		TArray<FPCGMeshInstanceList>& OutMeshInstances,
 		UPCGPointData* OutPointData) const PURE_VIRTUAL(UPCGMeshSelectorBase::SelectInstances_Implementation);
 
diff --git a/Engine/Plugins/Experimental/PCG/Source/PCG/Public/MeshSelectors/PCGMeshSelectorByAttribute.h b/Engine/Plugins/Experimental/PCG/Source/PCG/Public/MeshSelectors/PCGMeshSelectorByAttribute.h
--- a/Engine/Plugins/Experimental/PCG/Source/PCG/Public/MeshSelectors/PCGMeshSelectorByAttribute.h
+++ b/Engine/Plugins/Experimental/PCG/Source/PCG/Public/MeshSelectors/PCGMeshSelectorByAttribute.h
@@ -20,7 +20,7 @@ public:
 	virtual void SelectInstances_Implementation(
 		UPARAM(ref) FPCGContext& Context,
 		const UPCGStaticMeshSpawnerSettings* Settings,
-		const UPCGSpatialData* InSpatialData,
+		const UPCGPointData* InPointData,
 		TArray<FPCGMeshInstanceList>& OutMeshInstances,
 		UPCGPointData* OutPointData) const override;
 
diff --git a/Engine/Plugins/Experimental/PCG/Source/PCG/Public/MeshSelectors/PCGMeshSelectorWeighted.h b/Engine/Plugins/Experimental/PCG/Source/PCG/Public/MeshSelectors/PCGMeshSelectorWeighted.h
--- a/Engine/Plugins/Experimental/PCG/Source/PCG/Public/MeshSelectors/PCGMeshSelectorWeighted.h
+++ b/Engine/Plugins/Experimental/PCG/Source/PCG/Public/MeshSelectors/PCGMeshSelectorWeighted.h
@@ -65,7 +65,7 @@ public:
 	virtual void SelectInstances_Implementation(
 		UPARAM(ref) FPCGContext& Context, 
 		const UPCGStaticMeshSpawnerSettings* Settings, 
-		const UPCGSpatialData* InSpatialData,
+		const UPCGPointData* InPointData,
 		TArray<FPCGMeshInstanceList>& OutMeshInstances,
 		UPCGPointData* OutPointData) const override;
 
diff --git a/Engine/Plugins/Experimental/PCG/Source/PCG/Public/MeshSelectors/PCGMeshSelectorWeightedByCategory.h b/Engine/Plugins/Experimental/PCG/Source/PCG/Public/MeshSelectors/PCGMeshSelectorWeightedByCategory.h
--- a/Engine/Plugins/Experimental/PCG/Source/PCG/Public/MeshSelectors/PCGMeshSelectorWeightedByCategory.h
+++ b/Engine/Plugins/Experimental/PCG/Source/PCG/Public/MeshSelectors/PCGMeshSelectorWeightedByCategory.h
@@ -37,7 +37,7 @@ public:
 	virtual void SelectInstances_Implementation(
 		UPARAM(ref) FPCGContext& Context,
 		const UPCGStaticMeshSpawnerSettings* Settings,
-		const UPCGSpatialData* InSpatialData,
+		const UPCGPointData* InPointData,
 		TArray<FPCGMeshInstanceList>& OutMeshInstances,
 		UPCGPointData* OutPointData) const override;
 
