[PCG] Surface sampler - fix case where generating shape came from actor, not from input pin.

It was hitting the check() in this case.

diff --git a/Engine/Plugins/Experimental/PCG/Source/PCG/Private/Elements/PCGSurfaceSampler.cpp b/Engine/Plugins/Experimental/PCG/Source/PCG/Private/Elements/PCGSurfaceSampler.cpp
--- a/Engine/Plugins/Experimental/PCG/Source/PCG/Private/Elements/PCGSurfaceSampler.cpp
+++ b/Engine/Plugins/Experimental/PCG/Source/PCG/Private/Elements/PCGSurfaceSampler.cpp
@@ -348,8 +348,15 @@ bool FPCGSurfaceSamplerElement::ExecuteInternal(FPCGContext* Context) const
 		{
 			GeneratingShapes.Add(GeneratorFromBoundingShapeInput);
 
-			check(BoundingShapeInputs.Num() > 0);
-			Outputs.Add(BoundingShapeInputs[0]);
+			// If there was a bounding shape input, use it as the starting point to get the tags
+			if (BoundingShapeInputs.Num() > 0)
+			{
+				Outputs.Add(BoundingShapeInputs[0]);
+			}
+			else
+			{
+				Outputs.Emplace();
+			}
 		}
 	}
 
