[PCG] Fix non-unity builds

diff --git a/Engine/Plugins/Experimental/PCG/Source/PCG/Public/Elements/PCGTransformPoints.h b/Engine/Plugins/Experimental/PCG/Source/PCG/Public/Elements/PCGTransformPoints.h
--- a/Engine/Plugins/Experimental/PCG/Source/PCG/Public/Elements/PCGTransformPoints.h
+++ b/Engine/Plugins/Experimental/PCG/Source/PCG/Public/Elements/PCGTransformPoints.h
@@ -2,9 +2,9 @@
 
 #pragma once
 
+#include "PCGPin.h"
 #include "PCGSettings.h"
 
-
 #include "PCGTransformPoints.generated.h"
 
 UCLASS(BlueprintType, ClassGroup = (Procedural))
