Mirrored mobility to ISM based on target actor

#rnx

diff --git a/Engine/Plugins/Experimental/PCG/Source/PCG/Private/Elements/PCGStaticMeshSpawner.cpp b/Engine/Plugins/Experimental/PCG/Source/PCG/Private/Elements/PCGStaticMeshSpawner.cpp
--- a/Engine/Plugins/Experimental/PCG/Source/PCG/Private/Elements/PCGStaticMeshSpawner.cpp
+++ b/Engine/Plugins/Experimental/PCG/Source/PCG/Private/Elements/PCGStaticMeshSpawner.cpp
@@ -215,6 +215,12 @@ void FPCGStaticMeshSpawnerElement::SpawnStaticMeshInstances(FPCGContext* Context
 	Params.WorldPositionOffsetDisableDistance = InstanceList.WorldPositionOffsetDisableDistance;
 	Params.bIsLocalToWorldDeterminantNegative = InstanceList.bIsLocalToWorldDeterminantNegative;
 
+	// If the root actor we're binding to is movable, then the ISMC should be movable by default
+	if (USceneComponent* SceneComponent = TargetActor->GetRootComponent())
+	{
+		Params.Mobility = SceneComponent->Mobility;
+	}
+
 	UInstancedStaticMeshComponent* ISMC = UPCGActorHelpers::GetOrCreateISMC(TargetActor, Context->SourceComponent.Get(), Params);
 
 	const int32 PreExistingInstanceCount = ISMC->GetInstanceCount();
