Prevent data from actor from getting its own output data

diff --git a/Engine/Plugins/Experimental/PCG/Source/PCG/Private/Elements/PCGDataFromActor.cpp b/Engine/Plugins/Experimental/PCG/Source/PCG/Private/Elements/PCGDataFromActor.cpp
--- a/Engine/Plugins/Experimental/PCG/Source/PCG/Private/Elements/PCGDataFromActor.cpp
+++ b/Engine/Plugins/Experimental/PCG/Source/PCG/Private/Elements/PCGDataFromActor.cpp
@@ -88,7 +88,7 @@ bool FPCGDataFromActorElement::ExecuteInternal(FPCGContext* InContext) const
 			TArray<FPCGTaskId> WaitOnTaskIds;
 			for (AActor* Actor : Context->FoundActors)
 			{
-				GatherWaitTasks(Actor, WaitOnTaskIds);
+				GatherWaitTasks(Actor, Context, WaitOnTaskIds);
 			}
 
 			if (!WaitOnTaskIds.IsEmpty())
@@ -127,19 +127,22 @@ bool FPCGDataFromActorElement::ExecuteInternal(FPCGContext* InContext) const
 	return true;
 }
 
-void FPCGDataFromActorElement::GatherWaitTasks(AActor* FoundActor, TArray<FPCGTaskId>& OutWaitTasks) const
+void FPCGDataFromActorElement::GatherWaitTasks(AActor* FoundActor, FPCGContext* Context, TArray<FPCGTaskId>& OutWaitTasks) const
 {
 	if (!FoundActor)
 	{
 		return;
 	}
 
+	// We will prevent gathering the current execution - this task cannot wait on itself
+	AActor* ThisOwner = ((Context && Context->SourceComponent.IsValid()) ? Context->SourceComponent->GetOwner() : nullptr);
+
 	TInlineComponentArray<UPCGComponent*, 1> PCGComponents;
 	FoundActor->GetComponents(PCGComponents);
 
 	for (UPCGComponent* Component : PCGComponents)
 	{
-		if (Component->IsGenerating())
+		if (Component->IsGenerating() && Component->GetOwner() != ThisOwner)
 		{
 			OutWaitTasks.Add(Component->GetGenerationTaskId());
 		}
@@ -156,11 +159,14 @@ void FPCGDataFromActorElement::ProcessActor(FPCGContext* Context, const UPCGData
 		return;
 	}
 
+	AActor* ThisOwner = ((Context && Context->SourceComponent.Get()) ? Context->SourceComponent->GetOwner() : nullptr);
 	TInlineComponentArray<UPCGComponent*, 1> PCGComponents;
 	bool bHasGeneratedPCGData = false;
 	FProperty* FoundProperty = nullptr;
 
-	if (Settings->Mode == EPCGGetDataFromActorMode::GetDataFromPCGComponent || Settings->Mode == EPCGGetDataFromActorMode::GetDataFromPCGComponentOrParseComponents)
+	const bool bCanGetDataFromComponent = (FoundActor != ThisOwner);
+
+	if (bCanGetDataFromComponent && (Settings->Mode == EPCGGetDataFromActorMode::GetDataFromPCGComponent || Settings->Mode == EPCGGetDataFromActorMode::GetDataFromPCGComponentOrParseComponents))
 	{
 		FoundActor->GetComponents(PCGComponents);
 
@@ -180,7 +186,15 @@ void FPCGDataFromActorElement::ProcessActor(FPCGContext* Context, const UPCGData
 	// Some additional validation
 	if (Settings->Mode == EPCGGetDataFromActorMode::GetDataFromPCGComponent && !bHasGeneratedPCGData)
 	{
-		PCGE_LOG(Warning, "Actor (%s) does not have any previously generated data.", *FoundActor->GetFName().ToString());
+		if (bCanGetDataFromComponent)
+		{
+			PCGE_LOG(Warning, "Actor (%s) does not have any previously generated data.", *FoundActor->GetFName().ToString());
+		}
+		else
+		{
+			PCGE_LOG(Error, "Actor (%s) cannot get its own generated data during generation", *FoundActor->GetFName().ToString());
+		}
+		
 		return;
 	}
 	else if (Settings->Mode == EPCGGetDataFromActorMode::GetDataFromProperty && !FoundProperty)
diff --git a/Engine/Plugins/Experimental/PCG/Source/PCG/Public/Elements/PCGDataFromActor.h b/Engine/Plugins/Experimental/PCG/Source/PCG/Public/Elements/PCGDataFromActor.h
--- a/Engine/Plugins/Experimental/PCG/Source/PCG/Public/Elements/PCGDataFromActor.h
+++ b/Engine/Plugins/Experimental/PCG/Source/PCG/Public/Elements/PCGDataFromActor.h
@@ -77,6 +77,6 @@ public:
 	virtual bool IsCacheable(const UPCGSettings* InSettings) const override { return false; }
 protected:
 	virtual bool ExecuteInternal(FPCGContext* Context) const;
-	void GatherWaitTasks(AActor* FoundActor, TArray<FPCGTaskId>& OutWaitTasks) const;
+	void GatherWaitTasks(AActor* FoundActor, FPCGContext* Context, TArray<FPCGTaskId>& OutWaitTasks) const;
 	void ProcessActor(FPCGContext* Context, const UPCGDataFromActorSettings* Settings, AActor* FoundActor) const;
 };
