[PCG] Tentative fix for crash when loading UPCGBlueprintSettings

diff --git a/Engine/Plugins/Experimental/PCG/Source/PCG/Private/Elements/PCGExecuteBlueprint.cpp b/Engine/Plugins/Experimental/PCG/Source/PCG/Private/Elements/PCGExecuteBlueprint.cpp
--- a/Engine/Plugins/Experimental/PCG/Source/PCG/Private/Elements/PCGExecuteBlueprint.cpp
+++ b/Engine/Plugins/Experimental/PCG/Source/PCG/Private/Elements/PCGExecuteBlueprint.cpp
@@ -493,6 +493,26 @@ TArray<FPCGSettingsOverridableParam> UPCGBlueprintSettings::GatherOverridablePar
 }
 #endif // WITH_EDITOR
 
+void UPCGBlueprintSettings::FixingOverridableParamPropertyClass(FPCGSettingsOverridableParam& Param) const
+{
+	bool bFound = false;
+
+	if (BlueprintElementInstance && !Param.PropertiesNames.IsEmpty())
+	{
+		UClass* BPClass = BlueprintElementInstance->GetClass();
+		if (BPClass && BPClass->FindPropertyByName(Param.PropertiesNames[0]))
+		{
+			Param.PropertyClass = BPClass;
+			bFound = true;
+		}
+	}
+
+	if(!bFound)
+	{
+		Super::FixingOverridableParamPropertyClass(Param);
+	}
+}
+
 bool FPCGExecuteBlueprintElement::ExecuteInternal(FPCGContext* InContext) const
 {
 	FPCGBlueprintExecutionContext* Context = static_cast<FPCGBlueprintExecutionContext*>(InContext);
diff --git a/Engine/Plugins/Experimental/PCG/Source/PCG/Private/PCGSettings.cpp b/Engine/Plugins/Experimental/PCG/Source/PCG/Private/PCGSettings.cpp
--- a/Engine/Plugins/Experimental/PCG/Source/PCG/Private/PCGSettings.cpp
+++ b/Engine/Plugins/Experimental/PCG/Source/PCG/Private/PCGSettings.cpp
@@ -455,24 +455,45 @@ void UPCGSettings::InitializeCachedOverridableParams(bool bReset)
 	}
 #endif // WITH_EDITOR
 
-	for (FPCGSettingsOverridableParam& Param : CachedOverridableParams)
+	for (int32 i = 0; i < CachedOverridableParams.Num(); ++i)
 	{
-		check(!Param.PropertiesNames.IsEmpty() && Param.PropertyClass);
+		FPCGSettingsOverridableParam& Param = CachedOverridableParams[i];
+
+		if (!Param.PropertiesNames.IsEmpty())
+		{
+			UE_LOG(LogPCG, Error, TEXT("[InitializeCachedOverridableParams] %s one param is missing its property names, we cannot use this one as overridable param"), *GetName());
+			CachedOverridableParams.RemoveAt(i--);
+			continue;
+		}
+
+		// In some cases, the property class is not serialized correctly. We can still try to recover from this.
+		if (!Param.PropertyClass)
+		{
+			FixingOverridableParamPropertyClass(Param);
+			
+			if (!Param.PropertyClass)
+			{
+				UE_LOG(LogPCG, Error, TEXT("[InitializeCachedOverridableParams] %s one param is missing its property class, we cannot use this one as overridable param"), *GetName());
+				CachedOverridableParams.RemoveAt(i--);
+				continue;
+			}
+		}
+
 		Param.Properties.Reset(Param.PropertiesNames.Num());
 
 		// Some properties might not be available at runtime. Ignore them.
 		const FProperty* CurrentProperty = Param.PropertyClass->FindPropertyByName(Param.PropertiesNames[0]);
-		if (ensure(CurrentProperty))
+		if (CurrentProperty)
 		{
 			Param.Properties.Add(CurrentProperty);
 
-			for (int32 i = 1; i < Param.PropertiesNames.Num(); ++i)
+			for (int32 j = 1; j < Param.PropertiesNames.Num(); ++j)
 			{
 				// If we have multiple depth properties, it should be Struct properties by construction
 				const FStructProperty* StructProperty = CastField<FStructProperty>(CurrentProperty);
 				if (ensure(StructProperty))
 				{
-					CurrentProperty = StructProperty->Struct->FindPropertyByName(Param.PropertiesNames[i]);
+					CurrentProperty = StructProperty->Struct->FindPropertyByName(Param.PropertiesNames[j]);
 					check(CurrentProperty);
 					Param.Properties.Add(CurrentProperty);
 				}
@@ -481,6 +502,19 @@ void UPCGSettings::InitializeCachedOverridableParams(bool bReset)
 	}
 }
 
+void UPCGSettings::FixingOverridableParamPropertyClass(FPCGSettingsOverridableParam& Param) const
+{
+	if (!GetClass() || Param.PropertiesNames.IsEmpty())
+	{
+		return;
+	}
+
+	if (GetClass()->FindPropertyByName(Param.PropertiesNames[0]) != nullptr)
+	{
+		Param.PropertyClass = GetClass();
+	}
+}
+
 TArray<FPCGPinProperties> UPCGSettings::DefaultPointInputPinProperties() const
 {
 	TArray<FPCGPinProperties> Properties;
diff --git a/Engine/Plugins/Experimental/PCG/Source/PCG/Public/Elements/PCGExecuteBlueprint.h b/Engine/Plugins/Experimental/PCG/Source/PCG/Public/Elements/PCGExecuteBlueprint.h
--- a/Engine/Plugins/Experimental/PCG/Source/PCG/Public/Elements/PCGExecuteBlueprint.h
+++ b/Engine/Plugins/Experimental/PCG/Source/PCG/Public/Elements/PCGExecuteBlueprint.h
@@ -200,6 +200,7 @@ protected:
 #if WITH_EDITOR
 	virtual TArray<FPCGSettingsOverridableParam> GatherOverridableParams() const override;
 #endif // WITH_EDITOR
+	virtual void FixingOverridableParamPropertyClass(FPCGSettingsOverridableParam& Param) const override;
 	// ~End UPCGSettings interface
 
 public:
diff --git a/Engine/Plugins/Experimental/PCG/Source/PCG/Public/PCGSettings.h b/Engine/Plugins/Experimental/PCG/Source/PCG/Public/PCGSettings.h
--- a/Engine/Plugins/Experimental/PCG/Source/PCG/Public/PCGSettings.h
+++ b/Engine/Plugins/Experimental/PCG/Source/PCG/Public/PCGSettings.h
@@ -278,6 +278,12 @@ protected:
 	*/
 	void InitializeCachedOverridableParams(bool bReset = false);
 
+	/**
+	* There is a weird issue where the BP class is not set correctly in some Server cases.
+	* We can try to recover if the PropertyClass is null.
+	*/
+	virtual void FixingOverridableParamPropertyClass(FPCGSettingsOverridableParam& Param) const;
+
 	/** Needs to be serialized since property metadata (used to populate this array) is not available at runtime. */
 	UPROPERTY()
 	TArray<FPCGSettingsOverridableParam> CachedOverridableParams;
