[PCG] Demote Composite to hidden type, and apply its pin icons to Spatial.

Pick White as default Spatial colour, to hopefully hint that it accepts connections to/from concrete types.

This type is getting hidden because there are not known cases where we need to distinguish between Spatial and Composite:

* If you need Spatial but not Composite - that is Concrete and already exists (and is used)
* Composite but not Spatial - this does not make any sense because Composite IS Spatial for all intents and purposes.

diff --git a/Engine/Plugins/Experimental/PCG/Source/PCG/Private/PCGComponent.cpp b/Engine/Plugins/Experimental/PCG/Source/PCG/Private/PCGComponent.cpp
--- a/Engine/Plugins/Experimental/PCG/Source/PCG/Private/PCGComponent.cpp
+++ b/Engine/Plugins/Experimental/PCG/Source/PCG/Private/PCGComponent.cpp
@@ -1948,7 +1948,7 @@ FPCGDataCollection UPCGComponent::CreateActorPCGDataCollection(AActor* Actor, co
 	// Some actor types we will forego full parsing to build strictly on the actor existence, such as partition actors, volumes and landscape
 	// TODO: add factory for extensibility
 	// TODO: review the !bParseActor cases - it might make sense to have just a point for a partition actor, even if we preintersect it.
-	APCGPartitionActor* PartitionActor = InDataFilter(EPCGDataType::Composite) ? Cast<APCGPartitionActor>(Actor) : nullptr;
+	APCGPartitionActor* PartitionActor = InDataFilter(EPCGDataType::Spatial) ? Cast<APCGPartitionActor>(Actor) : nullptr;
 	ALandscapeProxy* LandscapeActor = InDataFilter(EPCGDataType::Landscape) ? Cast<ALandscapeProxy>(Actor) : nullptr;
 	AVolume* VolumeActor = InDataFilter(EPCGDataType::Volume) ? Cast<AVolume>(Actor) : nullptr;
 	if (PartitionActor)
diff --git a/Engine/Plugins/Experimental/PCG/Source/PCG/Private/PCGPin.cpp b/Engine/Plugins/Experimental/PCG/Source/PCG/Private/PCGPin.cpp
--- a/Engine/Plugins/Experimental/PCG/Source/PCG/Private/PCGPin.cpp
+++ b/Engine/Plugins/Experimental/PCG/Source/PCG/Private/PCGPin.cpp
@@ -225,16 +225,16 @@ bool UPCGPin::IsCompatible(const UPCGPin* OtherPin) const
 
 	// Concrete can always be used as a composite - allow connections from concrete to composite
 	const bool bUpstreamConcrete = !!(UpstreamPin->Properties.AllowedTypes & EPCGDataType::Concrete);
-	const bool bDownstreamComposite = !!(DownstreamPin->Properties.AllowedTypes & EPCGDataType::Composite);
-	if (bUpstreamConcrete && bDownstreamComposite)
+	const bool bDownstreamSpatial = !!(DownstreamPin->Properties.AllowedTypes & EPCGDataType::Spatial);
+	if (bUpstreamConcrete && bDownstreamSpatial)
 	{
 		return true;
 	}
 
 	// Catch case when a composite type is being connected to a concrete type - that can require collapse
-	const bool bUpstreamComposite = !!(UpstreamPin->Properties.AllowedTypes & EPCGDataType::Composite);
+	const bool bUpstreamSpatial = !!(UpstreamPin->Properties.AllowedTypes & EPCGDataType::Spatial);
 	const bool bDownstreamConcrete = !!(DownstreamPin->Properties.AllowedTypes & EPCGDataType::Concrete);
-	if (bUpstreamComposite && bDownstreamConcrete)
+	if (bUpstreamSpatial && bDownstreamConcrete)
 	{
 		// This will trigger a collapse, but let it slide for now.
 		// TODO in the future we should inject a conversion node.
diff --git a/Engine/Plugins/Experimental/PCG/Source/PCG/Public/Data/PCGDifferenceData.h b/Engine/Plugins/Experimental/PCG/Source/PCG/Public/Data/PCGDifferenceData.h
--- a/Engine/Plugins/Experimental/PCG/Source/PCG/Public/Data/PCGDifferenceData.h
+++ b/Engine/Plugins/Experimental/PCG/Source/PCG/Public/Data/PCGDifferenceData.h
@@ -48,7 +48,7 @@ public:
 #endif
 
 	// ~Begin UPCGData interface
-	virtual EPCGDataType GetDataType() const override { return EPCGDataType::Composite; }
+	virtual EPCGDataType GetDataType() const override { return EPCGDataType::Spatial; }
 	// ~End UPCGData interface
 
 	//~Begin UPCGSpatialData interface
diff --git a/Engine/Plugins/Experimental/PCG/Source/PCG/Public/Data/PCGIntersectionData.h b/Engine/Plugins/Experimental/PCG/Source/PCG/Public/Data/PCGIntersectionData.h
--- a/Engine/Plugins/Experimental/PCG/Source/PCG/Public/Data/PCGIntersectionData.h
+++ b/Engine/Plugins/Experimental/PCG/Source/PCG/Public/Data/PCGIntersectionData.h
@@ -25,7 +25,7 @@ public:
 	void Initialize(const UPCGSpatialData* InA, const UPCGSpatialData* InB);
 
 	// ~Begin UPCGData interface
-	virtual EPCGDataType GetDataType() const override { return EPCGDataType::Composite; }
+	virtual EPCGDataType GetDataType() const override { return EPCGDataType::Spatial; }
 	// ~End UPCGData interface
 
 	//~Begin UPCGSpatialData interface
diff --git a/Engine/Plugins/Experimental/PCG/Source/PCG/Public/Data/PCGSpatialData.h b/Engine/Plugins/Experimental/PCG/Source/PCG/Public/Data/PCGSpatialData.h
--- a/Engine/Plugins/Experimental/PCG/Source/PCG/Public/Data/PCGSpatialData.h
+++ b/Engine/Plugins/Experimental/PCG/Source/PCG/Public/Data/PCGSpatialData.h
@@ -126,7 +126,7 @@ public:
 	virtual const UPCGSpatialData* FindShapeFromNetwork(const int InDimension) const { return (InDimension == -1 || GetDimension() == InDimension) ? this : nullptr; }
 
 	/** Find the first concrete (non-composite) shape in the network. Depth first search. */
-	virtual const UPCGSpatialData* FindFirstConcreteShapeFromNetwork() const { return !(GetDataType() & EPCGDataType::Composite) ? this : nullptr; }
+	virtual const UPCGSpatialData* FindFirstConcreteShapeFromNetwork() const { return !!(GetDataType() & EPCGDataType::Concrete) ? this : nullptr; }
 
 	UPROPERTY(Transient, BlueprintReadWrite, EditAnywhere, Category = Data)
 	TWeakObjectPtr<AActor> TargetActor = nullptr;
diff --git a/Engine/Plugins/Experimental/PCG/Source/PCG/Public/Data/PCGUnionData.h b/Engine/Plugins/Experimental/PCG/Source/PCG/Public/Data/PCGUnionData.h
--- a/Engine/Plugins/Experimental/PCG/Source/PCG/Public/Data/PCGUnionData.h
+++ b/Engine/Plugins/Experimental/PCG/Source/PCG/Public/Data/PCGUnionData.h
@@ -37,7 +37,7 @@ public:
 	void SetDensityFunction(EPCGUnionDensityFunction InDensityFunction) { DensityFunction = InDensityFunction; }
 
 	// ~Begin UPCGData interface
-	virtual EPCGDataType GetDataType() const override { return EPCGDataType::Composite; }
+	virtual EPCGDataType GetDataType() const override { return EPCGDataType::Spatial; }
 	// ~End UPCGData interface
 
 	//~Begin UPCGSpatialData interface
diff --git a/Engine/Plugins/Experimental/PCG/Source/PCG/Public/PCGCommon.h b/Engine/Plugins/Experimental/PCG/Source/PCG/Public/PCGCommon.h
--- a/Engine/Plugins/Experimental/PCG/Source/PCG/Public/PCGCommon.h
+++ b/Engine/Plugins/Experimental/PCG/Source/PCG/Public/PCGCommon.h
@@ -43,8 +43,9 @@ enum class EPCGDataType : uint32
 	Concrete = Point | PolyLine | Surface | Volume | Primitive,
 
 	/** Boolean operations like union, difference, intersection. */
-	Composite = 1 << 9,
+	Composite = 1 << 9 UMETA(Hidden),
 
+	/** Combinations of concrete data and/or boolean operations. */
 	Spatial = Composite | Concrete,
 
 	Param = 1 << 27,
diff --git a/Engine/Plugins/Experimental/PCG/Source/PCGEditor/Private/PCGEditorCommon.h b/Engine/Plugins/Experimental/PCG/Source/PCGEditor/Private/PCGEditorCommon.h
--- a/Engine/Plugins/Experimental/PCG/Source/PCGEditor/Private/PCGEditorCommon.h
+++ b/Engine/Plugins/Experimental/PCG/Source/PCGEditor/Private/PCGEditorCommon.h
@@ -32,8 +32,6 @@ namespace FPCGEditorCommon
 	const FName SettingsDataType = FName(TEXT("Settings Data"));
 	const FName OtherDataType = FName(TEXT("Other Data"));
 
-	const FName CompositeDataType = FName(TEXT("Composite Data"));
-	
 	const FName ConcreteDataType = FName(TEXT("Concrete Data"));
 	const FName PointDataType = FName(TEXT("Point Data"));
 	const FName PolyLineDataType = FName(TEXT("Poly Line Data"));
diff --git a/Engine/Plugins/Experimental/PCG/Source/PCGEditor/Private/PCGEditorGraphNodeBase.cpp b/Engine/Plugins/Experimental/PCG/Source/PCGEditor/Private/PCGEditorGraphNodeBase.cpp
--- a/Engine/Plugins/Experimental/PCG/Source/PCGEditor/Private/PCGEditorGraphNodeBase.cpp
+++ b/Engine/Plugins/Experimental/PCG/Source/PCGEditor/Private/PCGEditorGraphNodeBase.cpp
@@ -455,10 +455,6 @@ FEdGraphPinType UPCGEditorGraphNodeBase::GetPinType(const UPCGPin* InPin)
 			EdPinType.PinSubCategory = FPCGEditorCommon::PrimitiveDataType;
 		}
 	}
-	else if (CheckType(EPCGDataType::Composite))
-	{
-		EdPinType.PinCategory = FPCGEditorCommon::CompositeDataType;
-	}
 	else if (CheckType(EPCGDataType::Spatial))
 	{
 		EdPinType.PinCategory = FPCGEditorCommon::SpatialDataType;
diff --git a/Engine/Plugins/Experimental/PCG/Source/PCGEditor/Private/PCGEditorSettings.cpp b/Engine/Plugins/Experimental/PCG/Source/PCGEditor/Private/PCGEditorSettings.cpp
--- a/Engine/Plugins/Experimental/PCG/Source/PCGEditor/Private/PCGEditorSettings.cpp
+++ b/Engine/Plugins/Experimental/PCG/Source/PCGEditor/Private/PCGEditorSettings.cpp
@@ -25,8 +25,7 @@ UPCGEditorSettings::UPCGEditorSettings(const FObjectInitializer& ObjectInitializ
 
 	DefaultPinColor = FLinearColor(1.0f, 1.0f, 1.0f);
 
-	SpatialDataPinColor = FLinearColor(0.2f, 0.2f, 1.0f);
-	CompositeDataPinColor = FLinearColor(0.2f, 0.2f, 1.0f);
+	SpatialDataPinColor = FLinearColor(1.0f, 1.0f, 1.0f);
 	ConcreteDataPinColor = FLinearColor(0.2f, 0.2f, 1.0f);
 	PointDataPinColor = FLinearColor(0.2f, 0.2f, 1.0f);
 	PolyLineDataPinColor = FLinearColor(0.2f, 0.2f, 1.0f);
@@ -148,10 +147,6 @@ FLinearColor UPCGEditorSettings::GetPinColor(const FEdGraphPinType& PinType) con
 			return ConcreteDataPinColor;
 		}
 	}
-	else if (PinType.PinCategory == FPCGEditorCommon::CompositeDataType)
-	{
-		return CompositeDataPinColor;
-	}
 	else if (PinType.PinCategory == FPCGEditorCommon::SpatialDataType)
 	{
 		return SpatialDataPinColor;
diff --git a/Engine/Plugins/Experimental/PCG/Source/PCGEditor/Private/SPCGEditorGraphNode.cpp b/Engine/Plugins/Experimental/PCG/Source/PCGEditor/Private/SPCGEditorGraphNode.cpp
--- a/Engine/Plugins/Experimental/PCG/Source/PCGEditor/Private/SPCGEditorGraphNode.cpp
+++ b/Engine/Plugins/Experimental/PCG/Source/PCGEditor/Private/SPCGEditorGraphNode.cpp
@@ -314,7 +314,7 @@ void SPCGEditorGraphNode::AddPin(const TSharedRef<SGraphPin>& PinToAdd)
 
 				PinToAdd->SetCustomPinIcon(ConnectedBrush, DisconnectedBrush);
 			}
-			else if (Pin->Properties.AllowedTypes == EPCGDataType::Composite)
+			else if (Pin->Properties.AllowedTypes == EPCGDataType::Spatial)
 			{
 				const FSlateBrush* ConnectedBrush = FPCGEditorStyle::Get().GetBrush(bIsInPin ? PCGEditorStyleConstants::Pin_Composite_IN_C : PCGEditorStyleConstants::Pin_Composite_OUT_C);
 				const FSlateBrush* DisconnectedBrush = FPCGEditorStyle::Get().GetBrush(bIsInPin ? PCGEditorStyleConstants::Pin_Composite_IN_DC : PCGEditorStyleConstants::Pin_Composite_OUT_DC);
diff --git a/Engine/Plugins/Experimental/PCG/Source/PCGEditor/Public/PCGEditorSettings.h b/Engine/Plugins/Experimental/PCG/Source/PCGEditor/Public/PCGEditorSettings.h
--- a/Engine/Plugins/Experimental/PCG/Source/PCGEditor/Public/PCGEditorSettings.h
+++ b/Engine/Plugins/Experimental/PCG/Source/PCGEditor/Public/PCGEditorSettings.h
@@ -80,10 +80,6 @@ public:
 	/** Color used for spatial data pins */
 	UPROPERTY(EditAnywhere, config, Category = Node, meta = (HideAlphaChannel))
 	FLinearColor SpatialDataPinColor;
-
-	/** Color used for composite shape spatial data pins */
-	UPROPERTY(EditAnywhere, config, Category = Node, meta = (HideAlphaChannel))
-	FLinearColor CompositeDataPinColor;	
 	
 	/** Color used for concrete/simple spatial data pins */
 	UPROPERTY(EditAnywhere, config, Category = Node, meta = (HideAlphaChannel))
