[PCG] Changed node inspection to toggle command.

diff --git a/Engine/Plugins/Experimental/PCG/Source/PCGEditor/Private/PCGEditor.cpp b/Engine/Plugins/Experimental/PCG/Source/PCGEditor/Private/PCGEditor.cpp
--- a/Engine/Plugins/Experimental/PCG/Source/PCGEditor/Private/PCGEditor.cpp
+++ b/Engine/Plugins/Experimental/PCG/Source/PCGEditor/Private/PCGEditor.cpp
@@ -415,12 +415,10 @@ void FPCGEditor::BindCommands()
 		FCanExecuteAction::CreateSP(this, &FPCGEditor::CanConvertToStandaloneNodes));
 
 	GraphEditorCommands->MapAction(
-		PCGEditorCommands.StartInspectNode,
-		FExecuteAction::CreateSP(this, &FPCGEditor::OnStartInspectNode));
-
-	GraphEditorCommands->MapAction(
-		PCGEditorCommands.StopInspectNode,
-		FExecuteAction::CreateSP(this, &FPCGEditor::OnStopInspectNode));
+		PCGEditorCommands.ToggleInspect,
+		FExecuteAction::CreateSP(this, &FPCGEditor::OnToggleInspected),
+		FCanExecuteAction::CreateSP(this, &FPCGEditor::CanToggleInspected),
+		FGetActionCheckState::CreateSP(this, &FPCGEditor::GetInspectedCheckState));
 
 	GraphEditorCommands->MapAction(
 		PCGEditorCommands.RunDeterminismNodeTest,
@@ -1198,50 +1196,88 @@ bool FPCGEditor::CanConvertToStandaloneNodes() const
 	return false;
 }
 
-void FPCGEditor::OnStartInspectNode()
+void FPCGEditor::OnToggleInspected()
 {
 	if (!GraphEditorWidget.IsValid())
 	{
 		return;
 	}
 
-	UEdGraphNode* GraphNode = GraphEditorWidget->GetSingleSelectedNode();
-	if (!GraphNode)
+	if (PCGGraphNodeBeingInspected)
 	{
-		return;
+		PCGGraphNodeBeingInspected->SetInspected(false);
 	}
-
+	
+	UEdGraphNode* GraphNode = GraphEditorWidget->GetSingleSelectedNode();
 	UPCGEditorGraphNodeBase* PCGGraphNodeBase = Cast<UPCGEditorGraphNodeBase>(GraphNode);
-	if (!PCGGraphNodeBase)
+	if (PCGGraphNodeBase && PCGGraphNodeBase != PCGGraphNodeBeingInspected)
 	{
-		return;
+		PCGGraphNodeBeingInspected = PCGGraphNodeBase;
+		PCGGraphNodeBeingInspected->SetInspected(true);
 	}
-	
-	if (PCGGraphNodeBase == PCGGraphNodeBeingInspected)
+	else
 	{
-		return;
+		PCGGraphNodeBeingInspected = nullptr;
 	}
 	
+	OnInspectedNodeChangedDelegate.Broadcast(PCGGraphNodeBeingInspected);
+	GetTabManager()->TryInvokeTab(FPCGEditor_private::AttributesID);
+}
+
+bool FPCGEditor::CanToggleInspected() const
+{
 	if (PCGGraphNodeBeingInspected)
 	{
-		PCGGraphNodeBeingInspected->SetInspected(false);
+		return true;
 	}
 
-	PCGGraphNodeBeingInspected = PCGGraphNodeBase;
-	PCGGraphNodeBeingInspected->SetInspected(true);
-	OnInspectedNodeChangedDelegate.Broadcast(PCGGraphNodeBeingInspected);
+	if (!GraphEditorWidget.IsValid())
+	{
+		return false;
+	}
 
-	GetTabManager()->TryInvokeTab(FPCGEditor_private::AttributesID);
+	const UEdGraphNode* GraphNode = GraphEditorWidget->GetSingleSelectedNode();
+
+	return GraphNode && GraphNode->IsA<UPCGEditorGraphNodeBase>();
 }
 
-void FPCGEditor::OnStopInspectNode()
+ECheckBoxState FPCGEditor::GetInspectedCheckState() const
 {
-	if (PCGGraphNodeBeingInspected)
+	if (GraphEditorWidget.IsValid())
 	{
-		PCGGraphNodeBeingInspected->SetInspected(false);
-		PCGGraphNodeBeingInspected = nullptr;
-		OnInspectedNodeChangedDelegate.Broadcast(nullptr);
+		const FGraphPanelSelectionSet& SelectedNodes = GraphEditorWidget->GetSelectedNodes();
+
+		if (SelectedNodes.IsEmpty())
+		{
+			return ECheckBoxState::Unchecked;
+		}
+	
+		bool bAllEnabled = true;
+		bool bAnyEnabled = false;
+		
+		for (UObject* Object : SelectedNodes)
+		{
+			const UPCGEditorGraphNodeBase* PCGEditorGraphNode = Cast<UPCGEditorGraphNodeBase>(Object);
+			if (!PCGEditorGraphNode)
+			{
+				continue;
+			}
+			
+			bAllEnabled &= PCGEditorGraphNode->GetInspected();
+			bAnyEnabled |= PCGEditorGraphNode->GetInspected();
+		}
+
+		if (bAllEnabled)
+		{
+			return ECheckBoxState::Checked;
+		}
+		else if (bAnyEnabled)
+		{
+			return ECheckBoxState::Undetermined; 
+		}
 	}
+	
+	return ECheckBoxState::Unchecked;	
 }
 
 void FPCGEditor::OnToggleEnabled()
diff --git a/Engine/Plugins/Experimental/PCG/Source/PCGEditor/Private/PCGEditor.h b/Engine/Plugins/Experimental/PCG/Source/PCGEditor/Private/PCGEditor.h
--- a/Engine/Plugins/Experimental/PCG/Source/PCGEditor/Private/PCGEditor.h
+++ b/Engine/Plugins/Experimental/PCG/Source/PCGEditor/Private/PCGEditor.h
@@ -107,11 +107,13 @@ private:
 	/** Force a regeneration by invoking the graph notifications  */
 	void OnForceGraphRegeneration_Clicked();
 
-	/** Start inspecting the current selected node */
-	void OnStartInspectNode();
-	/** Stop inspecting the current inspected node */
-	void OnStopInspectNode();
-
+	/** Toggle node inspection state for selected nodes */
+	void OnToggleInspected();
+	/** Whether we can toggle inspection of selected nodes */
+	bool CanToggleInspected() const;
+	/** Whether selected nodes are inspected or not */
+	ECheckBoxState GetInspectedCheckState() const;
+	
 	/** Toggle node enabled state for selected nodes */
 	void OnToggleEnabled();
 	/** Whether selected nodes are enabled or not */
diff --git a/Engine/Plugins/Experimental/PCG/Source/PCGEditor/Private/PCGEditorCommands.cpp b/Engine/Plugins/Experimental/PCG/Source/PCGEditor/Private/PCGEditorCommands.cpp
--- a/Engine/Plugins/Experimental/PCG/Source/PCGEditor/Private/PCGEditorCommands.cpp
+++ b/Engine/Plugins/Experimental/PCG/Source/PCGEditor/Private/PCGEditorCommands.cpp
@@ -26,11 +26,10 @@ void FPCGEditorCommands::RegisterCommands()
 	UI_COMMAND(RunDeterminismNodeTest, "Run Determinism Test on Node", "Evaluate the current node for determinism.", EUserInterfaceActionType::Button, FInputChord(EModifierKey::Alt, EKeys::D));
 	UI_COMMAND(RunDeterminismGraphTest, "Graph Determinism Test", "Evaluate the current graph for determinism.", EUserInterfaceActionType::Button, FInputChord());
 	UI_COMMAND(EditGraphSettings, "Graph Settings", "Edit the graph settings.", EUserInterfaceActionType::ToggleButton, FInputChord());
-	UI_COMMAND(StartInspectNode, "Start Inspecting Node", "Start Inspecting Node", EUserInterfaceActionType::Button, FInputChord(EKeys::I));
-	UI_COMMAND(StopInspectNode, "Stop Inspecting Node", "Stop Inspecting Node", EUserInterfaceActionType::Button, FInputChord(EModifierKey::Shift, EKeys::I));
 	UI_COMMAND(CancelExecution, "Cancel Execution", "Cancels the execution of the current graph", EUserInterfaceActionType::Button, FInputChord(EModifierKey::Control, EKeys::Escape));
 	UI_COMMAND(ToggleEnabled, "Toggle Enabled", "Toggle node enabled state for selected nodes.", EUserInterfaceActionType::ToggleButton, FInputChord(EKeys::E));
 	UI_COMMAND(ToggleDebug, "Toggle Debug", "Toggle node debug state for selected nodes", EUserInterfaceActionType::ToggleButton, FInputChord(EKeys::D));
+	UI_COMMAND(ToggleInspect, "Toggle Inspection", "Toggle node inspection for selected node", EUserInterfaceActionType::ToggleButton, FInputChord(EKeys::I));
 }
 
 #undef LOCTEXT_NAMESPACE
diff --git a/Engine/Plugins/Experimental/PCG/Source/PCGEditor/Private/PCGEditorCommands.h b/Engine/Plugins/Experimental/PCG/Source/PCGEditor/Private/PCGEditorCommands.h
--- a/Engine/Plugins/Experimental/PCG/Source/PCGEditor/Private/PCGEditorCommands.h
+++ b/Engine/Plugins/Experimental/PCG/Source/PCGEditor/Private/PCGEditorCommands.h
@@ -22,9 +22,8 @@ public:
 	TSharedPtr<FUICommandInfo> RunDeterminismNodeTest;
 	TSharedPtr<FUICommandInfo> RunDeterminismGraphTest;
 	TSharedPtr<FUICommandInfo> EditGraphSettings;
-	TSharedPtr<FUICommandInfo> StartInspectNode;
-	TSharedPtr<FUICommandInfo> StopInspectNode;
 	TSharedPtr<FUICommandInfo> CancelExecution;
 	TSharedPtr<FUICommandInfo> ToggleEnabled;
 	TSharedPtr<FUICommandInfo> ToggleDebug;
+	TSharedPtr<FUICommandInfo> ToggleInspect;
 };
diff --git a/Engine/Plugins/Experimental/PCG/Source/PCGEditor/Private/PCGEditorGraphNodeBase.cpp b/Engine/Plugins/Experimental/PCG/Source/PCGEditor/Private/PCGEditorGraphNodeBase.cpp
--- a/Engine/Plugins/Experimental/PCG/Source/PCGEditor/Private/PCGEditorGraphNodeBase.cpp
+++ b/Engine/Plugins/Experimental/PCG/Source/PCGEditor/Private/PCGEditorGraphNodeBase.cpp
@@ -77,7 +77,7 @@ void UPCGEditorGraphNodeBase::GetNodeContextMenuActions(UToolMenu* Menu, class U
 		FToolMenuSection& Section = Menu->AddSection("EdGraphSchemaNodeActions", LOCTEXT("NodeActionsHeader", "Node Actions"));
 		Section.AddMenuEntry(FPCGEditorCommands::Get().ToggleEnabled, LOCTEXT("ToggleEnabledLabel", "Enable"));
 		Section.AddMenuEntry(FPCGEditorCommands::Get().ToggleDebug, LOCTEXT("ToggleDebugLabel", "Debug"));
-		Section.AddMenuEntry(bIsInspected ? FPCGEditorCommands::Get().StopInspectNode : FPCGEditorCommands::Get().StartInspectNode);
+		Section.AddMenuEntry(FPCGEditorCommands::Get().ToggleInspect, LOCTEXT("ToggleinspectionLabel", "Inspect"));
 		Section.AddMenuEntry(FGraphEditorCommands::Get().BreakNodeLinks);
 		Section.AddMenuEntry(FPCGEditorCommands::Get().CollapseNodes);
 		Section.AddMenuEntry(FPCGEditorCommands::Get().ExportNodes);
