Fixed spline sampler local vs. worldspace bug

#rnx

diff --git a/Engine/Plugins/Experimental/PCG/Source/PCG/Private/Elements/PCGSplineSampler.cpp b/Engine/Plugins/Experimental/PCG/Source/PCG/Private/Elements/PCGSplineSampler.cpp
--- a/Engine/Plugins/Experimental/PCG/Source/PCG/Private/Elements/PCGSplineSampler.cpp
+++ b/Engine/Plugins/Experimental/PCG/Source/PCG/Private/Elements/PCGSplineSampler.cpp
@@ -799,7 +799,7 @@ namespace PCGSplineSampler
 				const int NumSegments = LineData->GetNumSegments();
 				for (int32 SegmentIndex = 0; SegmentIndex < NumSegments; ++SegmentIndex)
 				{
-					FVector& PolygonPoint = PolygonPoints.Add_GetRef(LineData->GetLocationAtDistance(SegmentIndex, 0));
+					FVector& PolygonPoint = PolygonPoints.Add_GetRef(LineData->GetLocationAtDistance(SegmentIndex, 0, /*bWorldSpace=*/false));
 					PolygonPoint.Z = MinPoint.Z;
 				}
 			}
@@ -975,7 +975,7 @@ namespace PCGSplineSampler
 						}
 						else
 						{
-							const FVector NearestSplineLocation = Spline->GetLocationAtSplineInputKey(NearestSplineKey, ESplineCoordinateSpace::World);
+							const FVector NearestSplineLocation = Spline->GetLocationAtSplineInputKey(NearestSplineKey, ESplineCoordinateSpace::Local);
 							PointToSplineDist = FVector2D::Distance(SampleLocation, FVector2D(NearestSplineLocation.X, NearestSplineLocation.Y));
 						}
 
