[PCG] Console var for disabling per-point-data CRC. Falls back to data UID.

CVar is enabled by default here because I did not want to mess with this stream until I retest. I will disable it by default in a downstream project where the point CRC was dominating cost.

diff --git a/Engine/Plugins/Experimental/PCG/Source/PCG/Private/Data/PCGPointData.cpp b/Engine/Plugins/Experimental/PCG/Source/PCG/Private/Data/PCGPointData.cpp
--- a/Engine/Plugins/Experimental/PCG/Source/PCG/Private/Data/PCGPointData.cpp
+++ b/Engine/Plugins/Experimental/PCG/Source/PCG/Private/Data/PCGPointData.cpp
@@ -8,10 +8,16 @@
 #include "Metadata/Accessors/PCGAttributeAccessorKeys.h"
 
 #include "GameFramework/Actor.h"
+#include "HAL/IConsoleManager.h"
 #include "Serialization/ArchiveCrc32.h"
 
 #include UE_INLINE_GENERATED_CPP_BY_NAME(PCGPointData)
 
+static TAutoConsoleVariable<bool> CVarCacheFullPointDataCrc(
+	TEXT("pcg.Cache.FullPointDataCrc"),
+	true,
+	TEXT("Enable fine-grained CRC of point data for change tracking, rather than using data UID."));
+
 namespace PCGPointHelpers
 {
 	bool GetDistanceRatios(const FPCGPoint& InPoint, const FVector& InPosition, FVector& OutRatios)
@@ -254,6 +260,14 @@ const UPCGPointData::PointOctree& UPCGPointData::GetOctree() const
 
 void UPCGPointData::AddToCrc(FArchiveCrc32& Ar) const
 {
+	// The code below has non-trivial cost, and can be disabled from console.
+	if (!CVarCacheFullPointDataCrc.GetValueOnAnyThread())
+	{
+		// Fallback to UID
+		Super::AddToCrc(Ar);
+		return;
+	}
+
 	uint32 ThisType = static_cast<uint32>(GetDataType());
 	Ar << ThisType;
 
