[PCG] Made the native blueprint node exposed in the node palette.

diff --git a/Engine/Plugins/Experimental/PCG/Source/PCG/Private/Elements/PCGExecuteBlueprint.cpp b/Engine/Plugins/Experimental/PCG/Source/PCG/Private/Elements/PCGExecuteBlueprint.cpp
--- a/Engine/Plugins/Experimental/PCG/Source/PCG/Private/Elements/PCGExecuteBlueprint.cpp
+++ b/Engine/Plugins/Experimental/PCG/Source/PCG/Private/Elements/PCGExecuteBlueprint.cpp
@@ -179,7 +179,7 @@ UPCGBlueprintSettings::UPCGBlueprintSettings()
 	bUseSeed = true;
 	
 #if WITH_EDITORONLY_DATA
-	bExposeToLibrary = false;
+	bExposeToLibrary = HasAnyFlags(RF_ClassDefaultObject);
 #endif
 }
 
