[PCG] Changed default pin colour (pins of type Any) - darken to communicate 'untyped' to user.

diff --git a/Engine/Plugins/Experimental/PCG/Source/PCGEditor/Private/PCGEditorSettings.cpp b/Engine/Plugins/Experimental/PCG/Source/PCGEditor/Private/PCGEditorSettings.cpp
--- a/Engine/Plugins/Experimental/PCG/Source/PCGEditor/Private/PCGEditorSettings.cpp
+++ b/Engine/Plugins/Experimental/PCG/Source/PCGEditor/Private/PCGEditorSettings.cpp
@@ -23,7 +23,7 @@ UPCGEditorSettings::UPCGEditorSettings(const FObjectInitializer& ObjectInitializ
 	ParamDataNodeColor = FLinearColor(1.0f, 0.6f, 0.0f);
 	DebugNodeColor = FLinearColor(1.0f, 0.0f, 1.0f);
 
-	DefaultPinColor = FLinearColor(1.0f, 1.0f, 1.0f);
+	DefaultPinColor = FLinearColor(0.16f, 0.16f, 0.16f);
 
 	SpatialDataPinColor = FLinearColor(1.0f, 1.0f, 1.0f);
 	ConcreteDataPinColor = FLinearColor(0.2f, 0.2f, 1.0f);
