[PCG] Typed getters - use "Get" naming convention for getters, and add Get Landscape and Get Volume variants.

Copies pattern established by [REDACTED_PERSON].

diff --git a/Engine/Plugins/Experimental/PCG/Source/PCG/Private/Elements/PCGSplineFromActor.cpp b/Engine/Plugins/Experimental/PCG/Source/PCG/Private/Elements/PCGSplineFromActor.cpp
deleted file mode 100644
diff --git a/Engine/Plugins/Experimental/PCG/Source/PCG/Private/Elements/PCGTypedGetter.cpp b/Engine/Plugins/Experimental/PCG/Source/PCG/Private/Elements/PCGTypedGetter.cpp
new file mode 100644
--- /dev/null
+++ b/Engine/Plugins/Experimental/PCG/Source/PCG/Private/Elements/PCGTypedGetter.cpp
@@ -0,0 +1,72 @@
+// Copyright Epic Games, Inc. All Rights Reserved.
+
+#include "Elements/PCGTypedGetter.h"
+
+#include UE_INLINE_GENERATED_CPP_BY_NAME(PCGTypedGetter)
+
+#define LOCTEXT_NAMESPACE "PCGTypedGetterElements"
+
+UPCGGetLandscapeSettings::UPCGGetLandscapeSettings()
+{
+	bDisplayModeSettings = false;
+	Mode = EPCGGetDataFromActorMode::ParseActorComponents;
+}
+
+TArray<FPCGPinProperties> UPCGGetLandscapeSettings::OutputPinProperties() const
+{
+	TArray<FPCGPinProperties> PinProperties;
+	PinProperties.Emplace(PCGPinConstants::DefaultOutputLabel, EPCGDataType::Landscape);
+
+	return PinProperties;
+}
+
+#if WITH_EDITOR
+FText UPCGGetLandscapeSettings::GetNodeTooltipText() const
+{
+	return LOCTEXT("GetLandscapeTooltip", "Builds a collection of landscapes from the selected actors.");
+}
+#endif
+
+UPCGGetSplineSettings::UPCGGetSplineSettings()
+{
+	bDisplayModeSettings = false;
+	Mode = EPCGGetDataFromActorMode::ParseActorComponents;
+}
+
+TArray<FPCGPinProperties> UPCGGetSplineSettings::OutputPinProperties() const
+{
+	TArray<FPCGPinProperties> PinProperties;
+	PinProperties.Emplace(PCGPinConstants::DefaultOutputLabel, EPCGDataType::Spline);
+
+	return PinProperties;
+}
+
+#if WITH_EDITOR
+FText UPCGGetSplineSettings::GetNodeTooltipText() const
+{
+	return LOCTEXT("GetSplineTooltip", "Builds a collection of splines from the selected actors.");
+}
+#endif
+
+UPCGGetVolumeSettings::UPCGGetVolumeSettings()
+{
+	bDisplayModeSettings = false;
+	Mode = EPCGGetDataFromActorMode::ParseActorComponents;
+}
+
+TArray<FPCGPinProperties> UPCGGetVolumeSettings::OutputPinProperties() const
+{
+	TArray<FPCGPinProperties> PinProperties;
+	PinProperties.Emplace(PCGPinConstants::DefaultOutputLabel, EPCGDataType::Volume);
+
+	return PinProperties;
+}
+
+#if WITH_EDITOR
+FText UPCGGetVolumeSettings::GetNodeTooltipText() const
+{
+	return LOCTEXT("GetVolumeTooltip", "Builds a collection of volumes from the selected actors.");
+}
+#endif
+
+#undef LOCTEXT_NAMESPACE
diff --git a/Engine/Plugins/Experimental/PCG/Source/PCG/Public/Elements/PCGDataFromActor.h b/Engine/Plugins/Experimental/PCG/Source/PCG/Public/Elements/PCGDataFromActor.h
--- a/Engine/Plugins/Experimental/PCG/Source/PCG/Public/Elements/PCGDataFromActor.h
+++ b/Engine/Plugins/Experimental/PCG/Source/PCG/Public/Elements/PCGDataFromActor.h
@@ -28,7 +28,7 @@ class PCG_API UPCGDataFromActorSettings : public UPCGSettings
 public:
 	//~Begin UPCGSettings interface
 #if WITH_EDITOR
-	virtual FName GetDefaultNodeName() const override { return FName(TEXT("Data From Actor")); }
+	virtual FName GetDefaultNodeName() const override { return FName(TEXT("Get Actor Data")); }
 	virtual FText GetNodeTooltipText() const override;
 	virtual EPCGSettingsType GetType() const override { return EPCGSettingsType::Spatial; }
 	virtual void GetTrackedActorTags(FPCGTagToSettingsMap& OutTagToSettings, TArray<TObjectPtr<const UPCGGraph>>& OutVisitedGraphs) const override;
diff --git a/Engine/Plugins/Experimental/PCG/Source/PCG/Public/Elements/PCGSplineFromActor.h b/Engine/Plugins/Experimental/PCG/Source/PCG/Public/Elements/PCGSplineFromActor.h
deleted file mode 100644
diff --git a/Engine/Plugins/Experimental/PCG/Source/PCG/Public/Elements/PCGTypedGetter.h b/Engine/Plugins/Experimental/PCG/Source/PCG/Public/Elements/PCGTypedGetter.h
new file mode 100644
--- /dev/null
+++ b/Engine/Plugins/Experimental/PCG/Source/PCG/Public/Elements/PCGTypedGetter.h
@@ -0,0 +1,82 @@
+// Copyright Epic Games, Inc. All Rights Reserved.
+
+#pragma once
+
+#include "Elements/PCGDataFromActor.h"
+
+#include "PCGTypedGetter.generated.h"
+
+/** Builds a collection of landscape data from the selected actors. */
+UCLASS(BlueprintType, ClassGroup = (Procedural))
+class PCG_API UPCGGetLandscapeSettings : public UPCGDataFromActorSettings
+{
+	GENERATED_BODY()
+
+public:
+	UPCGGetLandscapeSettings();
+
+	//~Begin UPCGSettings interface
+#if WITH_EDITOR
+	virtual FName GetDefaultNodeName() const override { return FName(TEXT("Get Landscape Data")); }
+	virtual FText GetNodeTooltipText() const override;
+#endif
+
+protected:
+	virtual TArray<FPCGPinProperties> OutputPinProperties() const override;
+	//~End UPCGSettings
+
+public:
+	//~Begin UPCGDataFromActorSettings interface
+	virtual bool DataFilter(EPCGDataType InDataType) const { return !!(InDataType & EPCGDataType::Landscape); }
+	//~End UPCGDataFromActorSettings
+};
+
+/** Builds a collection of spline data from the selected actors. */
+UCLASS(BlueprintType, ClassGroup = (Procedural))
+class PCG_API UPCGGetSplineSettings : public UPCGDataFromActorSettings
+{
+	GENERATED_BODY()
+
+public:
+	UPCGGetSplineSettings();
+
+	//~Begin UPCGSettings interface
+#if WITH_EDITOR
+	virtual FName GetDefaultNodeName() const override { return FName(TEXT("Get Spline Data")); }
+	virtual FText GetNodeTooltipText() const override;
+#endif
+
+protected:
+	virtual TArray<FPCGPinProperties> OutputPinProperties() const override;
+	//~End UPCGSettings
+
+public:
+	//~Begin UPCGDataFromActorSettings interface
+	virtual bool DataFilter(EPCGDataType InDataType) const { return !!(InDataType & EPCGDataType::PolyLine); }
+	//~End UPCGDataFromActorSettings
+};
+
+/** Builds a collection of volume data from the selected actors. */
+UCLASS(BlueprintType, ClassGroup = (Procedural))
+class PCG_API UPCGGetVolumeSettings : public UPCGDataFromActorSettings
+{
+	GENERATED_BODY()
+
+public:
+	UPCGGetVolumeSettings();
+
+	//~Begin UPCGSettings interface
+#if WITH_EDITOR
+	virtual FName GetDefaultNodeName() const override { return FName(TEXT("Get Volume Data")); }
+	virtual FText GetNodeTooltipText() const override;
+#endif
+
+protected:
+	virtual TArray<FPCGPinProperties> OutputPinProperties() const override;
+	//~End UPCGSettings
+
+public:
+	//~Begin UPCGDataFromActorSettings interface
+	virtual bool DataFilter(EPCGDataType InDataType) const { return !!(InDataType & EPCGDataType::Volume); }
+	//~End UPCGDataFromActorSettings
+};
