[PCG] First pass on graph instances

* Add PCGGraphInterface and PCGGraphInstance.
* All PCGComponent/PCGSubgraph now own a PCGGraphInterface, and will automatically create a GraphInstance from a Graph (and take ownership in this case).
* Should be transparent to the user.

diff --git a/Engine/Plugins/Experimental/PCG/Source/PCG/Private/Elements/PCGSpawnActor.cpp b/Engine/Plugins/Experimental/PCG/Source/PCG/Private/Elements/PCGSpawnActor.cpp
--- a/Engine/Plugins/Experimental/PCG/Source/PCG/Private/Elements/PCGSpawnActor.cpp
+++ b/Engine/Plugins/Experimental/PCG/Source/PCG/Private/Elements/PCGSpawnActor.cpp
@@ -2,12 +2,13 @@
 
 #include "Elements/PCGSpawnActor.h"
 
-#include "Data/PCGSpatialData.h"
 #include "PCGComponent.h"
+#include "PCGGraph.h"
 #include "PCGHelpers.h"
 #include "PCGManagedResource.h"
 
 #include "Data/PCGPointData.h"
+#include "Data/PCGSpatialData.h"
 #include "Grid/PCGPartitionActor.h"
 #include "Helpers/PCGActorHelpers.h"
 #include "Metadata/Accessors/IPCGAttributeAccessor.h"
@@ -131,7 +132,7 @@ FPCGElementPtr UPCGSpawnActorSettings::CreateElement() const
 	return MakeShared<FPCGSpawnActorElement>();
 }
 
-UPCGGraph* UPCGSpawnActorSettings::GetSubgraph() const
+UPCGGraphInterface* UPCGSpawnActorSettings::GetSubgraphInterface() const
 {
 	if(!TemplateActorClass || TemplateActorClass->HasAnyClassFlags(CLASS_Abstract))
 	{
@@ -150,9 +151,10 @@ UPCGGraph* UPCGSpawnActorSettings::GetSubgraph() const
 	{
 		if (UPCGComponent* PCGComponent = Cast<UPCGComponent>(Component))
 		{
+			// If there is no graph, there is no graph instance
 			if (PCGComponent->GetGraph() && PCGComponent->bActivated)
 			{
-				return PCGComponent->GetGraph();
+				return PCGComponent->GetGraphInstance();
 			}
 		}
 	}
@@ -221,10 +223,10 @@ bool UPCGSpawnActorSettings::IsStructuralProperty(const FName& InPropertyName) c
 }
 #endif // WITH_EDITOR
 
-TObjectPtr<UPCGGraph> UPCGSpawnActorNode::GetSubgraph() const
+TObjectPtr<UPCGGraphInterface> UPCGSpawnActorNode::GetSubgraphInterface() const
 {
 	TObjectPtr<UPCGSpawnActorSettings> Settings = Cast<UPCGSpawnActorSettings>(GetSettings());
-	return (Settings && Settings->Option != EPCGSpawnActorOption::NoMerging) ? Settings->GetSubgraph() : nullptr;
+	return (Settings && Settings->Option != EPCGSpawnActorOption::NoMerging) ? Settings->GetSubgraphInterface() : nullptr;
 }
 
 #if WITH_EDITOR
diff --git a/Engine/Plugins/Experimental/PCG/Source/PCG/Private/Grid/PCGPartitionActor.cpp b/Engine/Plugins/Experimental/PCG/Source/PCG/Private/Grid/PCGPartitionActor.cpp
--- a/Engine/Plugins/Experimental/PCG/Source/PCG/Private/Grid/PCGPartitionActor.cpp
+++ b/Engine/Plugins/Experimental/PCG/Source/PCG/Private/Grid/PCGPartitionActor.cpp
@@ -67,6 +67,7 @@ void APCGPartitionActor::PostLoad()
 	for (UPCGComponent* LocalComponent : GetAllLocalPCGComponents())
 	{
 		LocalComponent->MarkAsLocalComponent();
+		LocalComponent->ConditionalPostLoad();
 	}
 
 	bWasPostCreatedLoaded = true;
diff --git a/Engine/Plugins/Experimental/PCG/Source/PCG/Private/PCGComponent.cpp b/Engine/Plugins/Experimental/PCG/Source/PCG/Private/PCGComponent.cpp
--- a/Engine/Plugins/Experimental/PCG/Source/PCG/Private/PCGComponent.cpp
+++ b/Engine/Plugins/Experimental/PCG/Source/PCG/Private/PCGComponent.cpp
@@ -94,33 +94,45 @@ void UPCGComponent::SetIsPartitioned(bool bIsNowPartitioned)
 
 void UPCGComponent::SetGraph_Implementation(UPCGGraph* InGraph)
 {
-	SetGraphLocal(InGraph);
+	SetGraphInterfaceLocal(InGraph);
+}
+
+UPCGGraph* UPCGComponent::GetGraph() const
+{
+	return (GraphInstance ? GraphInstance->GetGraph() : nullptr);
 }
 
 void UPCGComponent::SetGraphLocal(UPCGGraph* InGraph)
 {
-	if(Graph == InGraph)
-	{
-		return;
-	}
+	SetGraphInterfaceLocal(InGraph);
+}
 
-#if WITH_EDITOR
-	if (Graph)
+void UPCGComponent::SetGraphInterfaceLocal(UPCGGraphInterface* InGraphInterface)
+{
+	if (ensure(GraphInstance))
 	{
-		Graph->OnGraphChangedDelegate.RemoveAll(this);
+		GraphInstance->SetGraph(InGraphInterface);
+		RefreshAfterGraphChanged(GraphInstance, /*bIsStructural=*/true, /*bDirtyInputs=*/true);
 	}
-#endif
+}
 
-	Graph = InGraph;
+void UPCGComponent::SetGraphInstanceFromParent(UPCGGraphInstance* InParentGraphInstance)
+{
+	// Since it must be used on local components only, add a check here
+	check(bIsComponentLocal);
 
 #if WITH_EDITOR
-	if (InGraph)
+	if (GraphInstance)
 	{
-		Graph->OnGraphChangedDelegate.AddUObject(this, &UPCGComponent::OnGraphChanged);
+		GraphInstance->OnGraphChangedDelegate.RemoveAll(this);
 	}
-#endif
+#endif // WITH_EDITOR
+
+	GraphInstance = InParentGraphInstance;
 
-	RefreshAfterGraphChanged(Graph, /*bIsStructural=*/true, /*bDirtyInputs=*/true);
+#if WITH_EDITOR
+	SetupCallbacksOnCreation();
+#endif // WITH_EDITOR
 }
 
 void UPCGComponent::AddToManagedResources(UPCGManagedResource* InResource)
@@ -150,7 +162,7 @@ void UPCGComponent::ForEachManagedResource(TFunctionRef<void(UPCGManagedResource
 
 bool UPCGComponent::ShouldGenerate(bool bForce, EPCGComponentGenerationTrigger RequestedGenerationTrigger) const
 {
-	if (!bActivated || !Graph || !GetSubsystem())
+	if (!bActivated || !GetGraph() || !GetSubsystem())
 	{
 		return false;
 	}
@@ -193,7 +205,7 @@ void UPCGComponent::SetPropertiesFromOriginal(const UPCGComponent* Original)
 #if WITH_EDITOR
 	const bool bHasDirtyInput = InputType != NewInputType;
 	const bool bHasDirtyExclusions = !(ExcludedTags.Num() == Original->ExcludedTags.Num() && ExcludedTags.Includes(Original->ExcludedTags));
-	const bool bIsDirty = bHasDirtyInput || bHasDirtyExclusions || Graph != Original->Graph;
+	const bool bIsDirty = bHasDirtyInput || bHasDirtyExclusions || GraphInstance != Original->GraphInstance;
 
 	if (bHasDirtyExclusions)
 	{
@@ -208,7 +220,7 @@ void UPCGComponent::SetPropertiesFromOriginal(const UPCGComponent* Original)
 
 	InputType = NewInputType;
 	Seed = Original->Seed;
-	SetGraphLocal(Original->Graph);
+	GraphInstance = Original->GraphInstance;
 
 	GenerationTrigger = Original->GenerationTrigger;
 
@@ -951,9 +963,18 @@ void UPCGComponent::PostLoad()
 			GeneratedActors_DEPRECATED.Reset();
 		}
 	}
-#endif
 
-#if WITH_EDITOR
+	if (Graph_DEPRECATED)
+	{
+		if (!GraphInstance)
+		{
+			GraphInstance = NewObject<UPCGGraphInstance>(this, TEXT("PCGGraphInstance"), RF_Transactional | RF_Public);
+		}
+
+		GraphInstance->SetGraph(Graph_DEPRECATED);
+		Graph_DEPRECATED = nullptr;
+	}
+
 	SetupCallbacksOnCreation();
 #endif
 }
@@ -973,20 +994,28 @@ void UPCGComponent::SetupCallbacksOnCreation()
 		UpdateTrackedLandscape(/*bBoundsCheck=*/false);
 	}
 
-	if (Graph)
+	if (GraphInstance)
 	{
-		Graph->OnGraphChangedDelegate.AddUObject(this, &UPCGComponent::OnGraphChanged);
+		// We might have already connected in PostInitProperties
+		// To be sure, remove it and re-add it.
+		GraphInstance->OnGraphChangedDelegate.RemoveAll(this);
+		GraphInstance->OnGraphChangedDelegate.AddUObject(this, &UPCGComponent::OnGraphChanged);
 	}
 }
+
+bool UPCGComponent::CanEditChange(const FProperty* InProperty) const
+{
+	// Can't change anything if the component is local
+	return !bIsComponentLocal;
+}
 #endif
 
 void UPCGComponent::BeginDestroy()
 {
 #if WITH_EDITOR
-	if (Graph)
+	if (GraphInstance)
 	{
-		Graph->OnGraphChangedDelegate.RemoveAll(this);
-		Graph = nullptr;
+		GraphInstance->OnGraphChangedDelegate.RemoveAll(this);
 	}
 
 	if (!IsEngineExitRequested())
@@ -1000,6 +1029,21 @@ void UPCGComponent::BeginDestroy()
 	Super::BeginDestroy();
 }
 
+void UPCGComponent::PostInitProperties()
+{
+	// Don't initialize the graph instance for default objects and not for local components (since it will be overriden when generated).
+	if (!GraphInstance && !bIsComponentLocal && !HasAnyFlags(RF_ClassDefaultObject | RF_ArchetypeObject))
+	{
+		GraphInstance = NewObject<UPCGGraphInstance>(this, TEXT("PCGGraphInstance"), RF_Transactional | RF_Public);
+
+#if WITH_EDITOR
+		GraphInstance->OnGraphChangedDelegate.AddUObject(this, &UPCGComponent::OnGraphChanged);
+#endif // WITH_EDITOR
+	}
+
+	Super::PostInitProperties();
+}
+
 void UPCGComponent::OnRegister()
 {
 	Super::OnRegister();
@@ -1028,7 +1072,7 @@ TStructOnScope<FActorComponentInstanceData> UPCGComponent::GetComponentInstanceD
 	return InstanceData;
 }
 
-void UPCGComponent::OnGraphChanged(UPCGGraph* InGraph, EPCGChangeType ChangeType)
+void UPCGComponent::OnGraphChanged(UPCGGraphInterface* InGraph, EPCGChangeType ChangeType)
 {
 	const bool bIsStructural = ((ChangeType & (EPCGChangeType::Edge | EPCGChangeType::Structural)) != EPCGChangeType::None);
 	const bool bDirtyInputs = bIsStructural || ((ChangeType & EPCGChangeType::Input) != EPCGChangeType::None);
@@ -1036,13 +1080,15 @@ void UPCGComponent::OnGraphChanged(UPCGGraph* InGraph, EPCGChangeType ChangeType
 	RefreshAfterGraphChanged(InGraph, bIsStructural, bDirtyInputs);
 }
 
-void UPCGComponent::RefreshAfterGraphChanged(UPCGGraph* InGraph, bool bIsStructural, bool bDirtyInputs)
+void UPCGComponent::RefreshAfterGraphChanged(UPCGGraphInterface* InGraph, bool bIsStructural, bool bDirtyInputs)
 {
-	if (InGraph != Graph)
+	if (InGraph != GraphInstance)
 	{
 		return;
 	}
 
+	const bool bHasGraph = (InGraph && InGraph->GetGraph());
+
 #if WITH_EDITOR
 	// In editor, since we've changed the graph, we might have changed the tracked actor tags as well
 	if (!PCGHelpers::IsRuntimeOrPIE())
@@ -1058,11 +1104,11 @@ void UPCGComponent::RefreshAfterGraphChanged(UPCGGraph* InGraph, bool bIsStructu
 		}
 
 		DirtyGenerated(bDirtyInputs ? (EPCGComponentDirtyFlag::Actor | EPCGComponentDirtyFlag::Landscape) : EPCGComponentDirtyFlag::None);
-		if (InGraph)
+		if (bHasGraph)
 		{
 			Refresh();
 		}
-		else if (!InGraph)
+		else
 		{
 			// With no graph, we clean up
 			CleanupLocal(/*bRemoveComponents=*/true, /*bSave=*/ false);
@@ -1074,11 +1120,11 @@ void UPCGComponent::RefreshAfterGraphChanged(UPCGGraph* InGraph, bool bIsStructu
 #endif
 
 	// Otherwise, if we are in PIE or runtime, force generate if we have a graph (and were generated). Or cleanup if we have no graph
-	if (InGraph && bGenerated)
+	if (bHasGraph && bGenerated)
 	{
 		GenerateLocal(/*bForce=*/true);
 	}
-	else if (!InGraph)
+	else if (!bHasGraph)
 	{
 		CleanupLocal(/*bRemoveComponents=*/true, /*bSave=*/ false);
 	}
@@ -1091,11 +1137,7 @@ void UPCGComponent::PreEditChange(FProperty* PropertyAboutToChange)
 	{
 		const FName PropName = PropertyAboutToChange->GetFName();
 
-		if (PropName == GET_MEMBER_NAME_CHECKED(UPCGComponent, Graph) && Graph)
-		{
-			Graph->OnGraphChangedDelegate.RemoveAll(this);
-		}
-		else if (PropName == GET_MEMBER_NAME_CHECKED(UPCGComponent, ExcludedTags))
+		if (PropName == GET_MEMBER_NAME_CHECKED(UPCGComponent, ExcludedTags))
 		{
 			TeardownTrackingCallbacks();
 		}
@@ -1138,14 +1180,9 @@ void UPCGComponent::PostEditChangeProperty(FPropertyChangedEvent& PropertyChange
 			}
 		}
 	}
-	else if (PropName == GET_MEMBER_NAME_CHECKED(UPCGComponent, Graph))
+	else if (PropName == GET_MEMBER_NAME_CHECKED(UPCGComponent, GraphInstance))
 	{
-		if (Graph)
-		{
-			Graph->OnGraphChangedDelegate.AddUObject(this, &UPCGComponent::OnGraphChanged);
-		}
-
-		RefreshAfterGraphChanged(Graph, /*bIsStructural=*/true, /*bDirtyInputs=*/true);
+		OnGraphChanged(GraphInstance, EPCGChangeType::Structural);
 	}
 	else if (PropName == GET_MEMBER_NAME_CHECKED(UPCGComponent, InputType))
 	{
@@ -1192,12 +1229,6 @@ void UPCGComponent::PreEditUndo()
 	// so we can have a consistent state
 	LastGeneratedBoundsPriorToUndo = LastGeneratedBounds;
 
-	// We don't know what is changing so remove all callbacks
-	if (Graph)
-	{
-		Graph->OnGraphChangedDelegate.RemoveAll(this);
-	}
-
 	if (bGenerated)
 	{
 		// Cleanup so managed resources are cleaned in all cases
@@ -1213,11 +1244,6 @@ void UPCGComponent::PostEditUndo()
 {
 	LastGeneratedBounds = LastGeneratedBoundsPriorToUndo;
 
-	if (Graph)
-	{
-		Graph->OnGraphChangedDelegate.AddUObject(this, &UPCGComponent::OnGraphChanged);
-	}
-
 	SetupTrackingCallbacks();
 	RefreshTrackingData();
 	UpdateTrackedLandscape();
@@ -1245,9 +1271,9 @@ void UPCGComponent::TeardownActorCallbacks()
 void UPCGComponent::SetupTrackingCallbacks()
 {
 	CachedTrackedTagsToSettings.Reset();
-	if (Graph)
+	if (UPCGGraph* PCGGraph = GetGraph())
 	{
-		CachedTrackedTagsToSettings = Graph->GetTrackedTagsToSettings();
+		CachedTrackedTagsToSettings = PCGGraph->GetTrackedTagsToSettings();
 	}
 
 	if(!ExcludedTags.IsEmpty() || !CachedTrackedTagsToSettings.IsEmpty())
@@ -1321,7 +1347,7 @@ bool UPCGComponent::UpdateExcludedActor(AActor* InActor)
 
 bool UPCGComponent::ActorIsTracked(AActor* InActor) const
 {
-	if (!InActor || !Graph)
+	if (!InActor || !GetGraph())
 	{
 		return false;
 	}
@@ -1346,7 +1372,7 @@ void UPCGComponent::OnActorAdded(AActor* InActor)
 
 	if (bIsExcluded || bIsTracked)
 	{
-		DirtyGenerated(bIsExcluded ? EPCGComponentDirtyFlag::Exclusions : EPCGComponentDirtyFlag::None);
+		DirtyGenerated(bIsExcluded ? EPCGComponentDirtyFlag::Exclusions : EPCGComponentDirtyFlag::None, /*bDispatchToLocalComponents=*/ false);
 		Refresh();
 	}
 }
@@ -1358,7 +1384,7 @@ void UPCGComponent::OnActorDeleted(AActor* InActor)
 
 	if (bWasExcluded || bWasTracked)
 	{
-		DirtyGenerated(bWasExcluded ? EPCGComponentDirtyFlag::Exclusions : EPCGComponentDirtyFlag::None);
+		DirtyGenerated(bWasExcluded ? EPCGComponentDirtyFlag::Exclusions : EPCGComponentDirtyFlag::None, /*bDispatchToLocalComponents=*/ false);
 		Refresh();
 	}
 }
@@ -1396,7 +1422,7 @@ void UPCGComponent::OnActorMoved(AActor* InActor)
 
 		if (bDirtyAndRefresh)
 		{
-			DirtyGenerated(bDirtyExclusions ? EPCGComponentDirtyFlag::Exclusions : EPCGComponentDirtyFlag::None);
+			DirtyGenerated(bDirtyExclusions ? EPCGComponentDirtyFlag::Exclusions : EPCGComponentDirtyFlag::None, /*bDispatchToLocalComponents=*/true);
 			Refresh();
 		}
 	}
@@ -1492,7 +1518,7 @@ void UPCGComponent::OnLandscapeChanged(ALandscapeProxy* Landscape, const FLandsc
 
 		if (DirtyFlag != EPCGComponentDirtyFlag::None)
 		{
-			DirtyGenerated(DirtyFlag);
+			DirtyGenerated(DirtyFlag, /*bDispatchToLocalComponents=*/ false);
 			Refresh();
 		}
 	}
@@ -1570,13 +1596,13 @@ void UPCGComponent::OnActorChanged(AActor* Actor, UObject* InObject, bool bActor
 
 		if (bDirtyAndRefresh)
 		{
-			DirtyGenerated();
+			DirtyGenerated(EPCGComponentDirtyFlag::None, /*bDispatchToLocalComponents=*/ false);
 			Refresh();
 		}
 	}
 }
 
-void UPCGComponent::DirtyGenerated(EPCGComponentDirtyFlag DirtyFlag)
+void UPCGComponent::DirtyGenerated(EPCGComponentDirtyFlag DirtyFlag, const bool bDispatchToLocalComponents)
 {
 	if (GetSubsystem() && GetSubsystem()->IsGraphCacheDebuggingEnabled())
 	{
@@ -1628,11 +1654,9 @@ void UPCGComponent::DirtyGenerated(EPCGComponentDirtyFlag DirtyFlag)
 		CachedPCGData = nullptr;
 	}
 
-	// For partitioned graph, we must forward the call to the partition actor
-	// Note that we do not need to forward "normal" dirty as these will be picked up by the local PCG components
-	// However, input changes / moves of the partitioned object will not be caught
-	// It would be possible for partitioned actors to add callbacks to their original component, but that inverses the processing flow
-	if (DirtyFlag != EPCGComponentDirtyFlag::None && bActivated && IsPartitioned())
+	// For partitioned graph, we must forward the call to the partition actor, if we need to
+	// TODO: Don't forward for None for now, as it could break some stuff
+	if (DirtyFlag != EPCGComponentDirtyFlag::None && bActivated && IsPartitioned() && bDispatchToLocalComponents)
 	{
 		if (GetSubsystem())
 		{
@@ -1969,7 +1993,8 @@ FPCGDataCollection UPCGComponent::CreateActorPCGDataCollection(AActor* Actor, co
 	else if (LandscapeActor)
 	{
 		UPCGLandscapeData* Data = NewObject<UPCGLandscapeData>();
-		const bool bUseLandscapeMetadata = (!Component || !Component->Graph || (Component->Graph && Component->Graph->bLandscapeUsesMetadata));
+		const UPCGGraph* PCGGraph = Component ? Component->GetGraph() : nullptr;
+		const bool bUseLandscapeMetadata = (!PCGGraph || PCGGraph->bLandscapeUsesMetadata);
 
 		Data->Initialize({ LandscapeActor }, PCGHelpers::GetGridBounds(Actor, Component), /*bHeightOnly=*/false, bUseLandscapeMetadata);
 
@@ -2193,7 +2218,8 @@ UPCGData* UPCGComponent::CreateLandscapePCGData(bool bHeightOnly)
 
 	// TODO: we're creating separate landscape data instances here so we can do some tweaks on it (such as storing the right target actor) but this probably should change
 	UPCGLandscapeData* LandscapeData = NewObject<UPCGLandscapeData>();
-	LandscapeData->Initialize(Landscapes, LandscapeBounds, bHeightOnly, /*bUseMetadata=*/Graph && Graph->bLandscapeUsesMetadata);
+	const UPCGGraph* PCGGraph = GetGraph();
+	LandscapeData->Initialize(Landscapes, LandscapeBounds, bHeightOnly, /*bUseMetadata=*/PCGGraph && PCGGraph->bLandscapeUsesMetadata);
 	// Need to override target actor for this one, not the landscape
 	LandscapeData->TargetActor = Actor;
 
@@ -2262,7 +2288,8 @@ bool UPCGComponent::IsLandscapeCachedDataDirty(const UPCGData* Data) const
 
 	if (const UPCGLandscapeData* CachedData = Cast<UPCGLandscapeData>(Data))
 	{
-		IsCacheDirty = Graph && (CachedData->IsUsingMetadata() != Graph->bLandscapeUsesMetadata);
+		const UPCGGraph* PCGGraph = GetGraph();
+		IsCacheDirty = PCGGraph && (CachedData->IsUsingMetadata() != PCGGraph->bLandscapeUsesMetadata);
 	}
 
 	return IsCacheDirty;
@@ -2460,9 +2487,10 @@ void UPCGComponent::DirtyCacheForAllTrackedTags()
 
 bool UPCGComponent::GraphUsesLandscapePin() const
 {
-	return Graph &&
-		(Graph->GetInputNode()->IsOutputPinConnected(PCGInputOutputConstants::DefaultLandscapeLabel) ||
-		Graph->GetInputNode()->IsOutputPinConnected(PCGInputOutputConstants::DefaultLandscapeHeightLabel));
+	const UPCGGraph* PCGGraph = GetGraph();
+	return PCGGraph &&
+		(PCGGraph->GetInputNode()->IsOutputPinConnected(PCGInputOutputConstants::DefaultLandscapeLabel) ||
+		PCGGraph->GetInputNode()->IsOutputPinConnected(PCGInputOutputConstants::DefaultLandscapeHeightLabel));
 }
 
 #endif // WITH_EDITOR
diff --git a/Engine/Plugins/Experimental/PCG/Source/PCG/Private/PCGGraph.cpp b/Engine/Plugins/Experimental/PCG/Source/PCG/Private/PCGGraph.cpp
--- a/Engine/Plugins/Experimental/PCG/Source/PCG/Private/PCGGraph.cpp
+++ b/Engine/Plugins/Experimental/PCG/Source/PCG/Private/PCGGraph.cpp
@@ -2,10 +2,12 @@
 
 #include "PCGGraph.h"
 
+#include "PCGComponent.h"
 #include "PCGInputOutputSettings.h"
 #include "PCGModule.h"
-#include "PCGSubsystem.h"
 #include "PCGPin.h"
+#include "PCGSubsystem.h"
+#include "PCGSubgraph.h"
 
 #include UE_INLINE_GENERATED_CPP_BY_NAME(PCGGraph)
 
@@ -15,6 +17,19 @@
 #include "UObject/Package.h"
 #endif
 
+/****************************
+* UPCGGraphInterface
+****************************/
+
+bool UPCGGraphInterface::IsInstance() const
+{
+	return this != GetGraph();
+}
+
+/****************************
+* UPCGGraph
+****************************/
+
 UPCGGraph::UPCGGraph(const FObjectInitializer& ObjectInitializer)
 	: Super(ObjectInitializer)
 {
@@ -675,3 +690,141 @@ void UPCGGraph::PostEditChangeProperty(FPropertyChangedEvent& PropertyChangedEve
 }
 #endif // WITH_EDITOR
 
+/****************************
+* UPCGGraphInstance
+****************************/
+
+void UPCGGraphInstance::PostLoad()
+{
+	Super::PostLoad();
+
+#if WITH_EDITOR
+	if (Graph)
+	{
+		Graph->ConditionalPostLoad();
+		Graph->OnGraphChangedDelegate.AddUObject(this, &UPCGGraphInstance::OnGraphChanged);
+	}
+#endif // WITH_EDITOR
+}
+
+void UPCGGraphInstance::BeginDestroy()
+{
+#if WITH_EDITOR
+	if (Graph)
+	{
+		Graph->OnGraphChangedDelegate.RemoveAll(this);
+	}
+#endif // WITH_EDITOR
+
+	Super::BeginDestroy();
+}
+
+#if WITH_EDITOR
+void UPCGGraphInstance::PreEditChange(FProperty* InProperty)
+{
+	Super::PreEditChange(InProperty);
+
+	if (InProperty && InProperty->GetFName() == GET_MEMBER_NAME_CHECKED(UPCGGraphInstance, Graph) && Graph)
+	{
+		Graph->OnGraphChangedDelegate.RemoveAll(this);
+	}
+}
+
+void UPCGGraphInstance::PostEditChangeProperty(FPropertyChangedEvent& PropertyChangedEvent)
+{
+	Super::PostEditChangeProperty(PropertyChangedEvent);
+
+	if (PropertyChangedEvent.Property && PropertyChangedEvent.Property->GetFName() == GET_MEMBER_NAME_CHECKED(UPCGGraphInstance, Graph) && Graph)
+	{
+		Graph->OnGraphChangedDelegate.AddUObject(this, &UPCGGraphInstance::OnGraphChanged);
+	}
+}
+
+void UPCGGraphInstance::PreEditUndo()
+{
+	Super::PreEditUndo();
+
+	if (Graph)
+	{
+		Graph->OnGraphChangedDelegate.RemoveAll(this);
+	}
+
+	UndoRedoGraphCache = Graph;
+}
+
+void UPCGGraphInstance::PostEditUndo()
+{
+	Super::PostEditUndo();
+
+	if (Graph)
+	{
+		Graph->OnGraphChangedDelegate.AddUObject(this, &UPCGGraphInstance::OnGraphChanged);
+	}
+}
+
+void UPCGGraphInstance::OnGraphChanged(UPCGGraphInterface* InGraph, EPCGChangeType ChangeType)
+{
+	if (InGraph == Graph)
+	{
+		OnGraphChangedDelegate.Broadcast(this, ChangeType);
+	}
+}
+
+bool UPCGGraphInstance::CanEditChange(const FProperty* InProperty) const
+{
+	// Graph can only be changed if it is in a PCGComponent (not local) or a PCGSubgraphSettings
+	if (InProperty && InProperty->GetFName() == GET_MEMBER_NAME_CHECKED(UPCGGraphInstance, Graph))
+	{
+		UObject* Outer = this->GetOuter();
+
+		if (UPCGComponent* Component = Cast<UPCGComponent>(Outer))
+		{
+			return !Component->IsLocalComponent();
+		}
+		else
+		{
+			return Outer && Outer->IsA<UPCGSubgraphSettings>();
+		}
+	}
+
+	return true;
+}
+#endif
+
+void UPCGGraphInstance::SetGraph(UPCGGraphInterface* InGraph)
+{
+	if (InGraph == this)
+	{
+		UE_LOG(LogPCG, Error, TEXT("Try to set the graph of a graph instance to itself, would cause infinite recursion."));
+		return;
+	}
+
+#if WITH_EDITOR
+	if (Graph)
+	{
+		Graph->OnGraphChangedDelegate.RemoveAll(this);
+	}
+#endif // WITH_EDITOR
+
+	Graph = InGraph;
+
+#if WITH_EDITOR
+	if (Graph)
+	{
+		Graph->OnGraphChangedDelegate.AddUObject(this, &UPCGGraphInstance::OnGraphChanged);
+	}
+#endif // WITH_EDITOR
+}
+
+TObjectPtr<UPCGGraphInterface> UPCGGraphInstance::CreateInstance(UObject* InOwner, UPCGGraphInterface* InGraph)
+{
+	if (!InOwner || !InGraph)
+	{
+		return nullptr;
+	}
+
+	TObjectPtr<UPCGGraphInstance> GraphInstance = NewObject<UPCGGraphInstance>(InOwner, MakeUniqueObjectName(InOwner, UPCGGraphInstance::StaticClass(), InGraph->GetFName()), RF_Transactional | RF_Public);
+	GraphInstance->SetGraph(InGraph);
+
+	return GraphInstance;
+}
diff --git a/Engine/Plugins/Experimental/PCG/Source/PCG/Private/PCGSubgraph.cpp b/Engine/Plugins/Experimental/PCG/Source/PCG/Private/PCGSubgraph.cpp
--- a/Engine/Plugins/Experimental/PCG/Source/PCG/Private/PCGSubgraph.cpp
+++ b/Engine/Plugins/Experimental/PCG/Source/PCG/Private/PCGSubgraph.cpp
@@ -1,7 +1,7 @@
 // Copyright Epic Games, Inc. All Rights Reserved.
 
 #include "PCGSubgraph.h"
-#include "PCGGraph.h"
+
 #include "PCGComponent.h"
 #include "PCGPin.h"
 #include "PCGSubsystem.h"
@@ -26,22 +26,43 @@ namespace PCGSubgraphSettings
 	}
 }
 
+UPCGGraph* UPCGBaseSubgraphSettings::GetSubgraph() const
+{
+	UPCGGraphInterface* SubgraphInterface = GetSubgraphInterface();
+	return SubgraphInterface ? SubgraphInterface->GetGraph() : nullptr;
+}
+
+void UPCGBaseSubgraphSettings::PostInitProperties()
+{
+	Super::PostInitProperties();
+
+#if WITH_EDITOR
+	if (UPCGGraphInterface* Subgraph = GetSubgraphInterface())
+	{
+		Subgraph->OnGraphChangedDelegate.AddUObject(this, &UPCGBaseSubgraphSettings::OnSubgraphChanged);
+	}
+#endif
+}
+
 void UPCGBaseSubgraphSettings::PostLoad()
 {
 	Super::PostLoad();
 
 #if WITH_EDITOR
-	if (UPCGGraph* Subgraph = GetSubgraph())
+	if (UPCGGraphInterface* Subgraph = GetSubgraphInterface())
 	{
-		Subgraph->OnGraphChangedDelegate.AddUObject(this, &UPCGSubgraphSettings::OnSubgraphChanged);
+		// We might have already connected in PostInitProperties
+		// To be sure, remove it and re-add it.
+		Subgraph->OnGraphChangedDelegate.RemoveAll(this);
+		Subgraph->OnGraphChangedDelegate.AddUObject(this, &UPCGBaseSubgraphSettings::OnSubgraphChanged);
 	}
 #endif
 }
 
-void UPCGBaseSubgraphSettings::SetSubgraph(UPCGGraph* InGraph)
+void UPCGBaseSubgraphSettings::SetSubgraph(UPCGGraphInterface* InGraph)
 {
 #if WITH_EDITOR
-	if (UPCGGraph* Subgraph = GetSubgraph())
+	if (UPCGGraphInterface* Subgraph = GetSubgraphInterface())
 	{
 		Subgraph->OnGraphChangedDelegate.RemoveAll(this);
 	}
@@ -50,9 +71,9 @@ void UPCGBaseSubgraphSettings::SetSubgraph(UPCGGraph* InGraph)
 	SetSubgraphInternal(InGraph);
 
 #if WITH_EDITOR
-	if (UPCGGraph* Subgraph = GetSubgraph())
+	if (UPCGGraphInterface* Subgraph = GetSubgraphInterface())
 	{
-		Subgraph->OnGraphChangedDelegate.AddUObject(this, &UPCGSubgraphSettings::OnSubgraphChanged);
+		Subgraph->OnGraphChangedDelegate.AddUObject(this, &UPCGBaseSubgraphSettings::OnSubgraphChanged);
 	}
 #endif // WITH_EDITOR
 }
@@ -60,7 +81,7 @@ void UPCGBaseSubgraphSettings::SetSubgraph(UPCGGraph* InGraph)
 void UPCGBaseSubgraphSettings::BeginDestroy()
 {
 #if WITH_EDITOR
-	if (UPCGGraph* Subgraph = GetSubgraph())
+	if (UPCGGraphInterface* Subgraph = GetSubgraphInterface())
 	{
 		Subgraph->OnGraphChangedDelegate.RemoveAll(this);
 	}
@@ -74,7 +95,7 @@ void UPCGBaseSubgraphSettings::PreEditChange(FProperty* PropertyAboutToChange)
 {
 	if (PropertyAboutToChange && IsStructuralProperty(PropertyAboutToChange->GetFName()))
 	{
-		if (UPCGGraph* Subgraph = GetSubgraph())
+		if (UPCGGraphInterface* Subgraph = GetSubgraphInterface())
 		{
 			Subgraph->OnGraphChangedDelegate.RemoveAll(this);
 		}
@@ -89,9 +110,9 @@ void UPCGBaseSubgraphSettings::PostEditChangeProperty(struct FPropertyChangedEve
 
 	if (PropertyChangedEvent.Property && IsStructuralProperty(PropertyChangedEvent.Property->GetFName()))
 	{
-		if (UPCGGraph* Subgraph = GetSubgraph())
+		if (UPCGGraphInterface* Subgraph = GetSubgraphInterface())
 		{
-			Subgraph->OnGraphChangedDelegate.AddUObject(this, &UPCGSubgraphSettings::OnSubgraphChanged);
+			Subgraph->OnGraphChangedDelegate.AddUObject(this, &UPCGBaseSubgraphSettings::OnSubgraphChanged);
 		}
 	}
 }
@@ -109,9 +130,9 @@ bool UPCGBaseSubgraphSettings::IsStructuralProperty(const FName& InPropertyName)
 	return InPropertyName == GET_MEMBER_NAME_CHECKED(UPCGSettingsInterface, bEnabled) || Super::IsStructuralProperty(InPropertyName);
 }
 
-void UPCGBaseSubgraphSettings::OnSubgraphChanged(UPCGGraph* InGraph, EPCGChangeType ChangeType)
+void UPCGBaseSubgraphSettings::OnSubgraphChanged(UPCGGraphInterface* InGraph, EPCGChangeType ChangeType)
 {
-	if (InGraph == GetSubgraph())
+	if (InGraph == GetSubgraphInterface())
 	{
 		OnSettingsChangedDelegate.Broadcast(this, (ChangeType | EPCGChangeType::Settings));
 	}
@@ -146,6 +167,12 @@ TArray<FPCGPinProperties> UPCGBaseSubgraphSettings::OutputPinProperties() const
 	}
 }
 
+UPCGSubgraphSettings::UPCGSubgraphSettings(const FObjectInitializer& InObjectInitializer)
+	: Super(InObjectInitializer)
+{
+	SubgraphInstance = InObjectInitializer.CreateDefaultSubobject<UPCGGraphInstance>(this, TEXT("PCGSubgraphInstance"));
+}
+
 UPCGNode* UPCGSubgraphSettings::CreateNode() const
 {
 	return NewObject<UPCGSubgraphNode>();
@@ -163,17 +190,36 @@ FName UPCGSubgraphSettings::AdditionalTaskName() const
 	}
 }
 
+void UPCGSubgraphSettings::SetSubgraphInternal(UPCGGraphInterface* InGraph)
+{
+	SubgraphInstance->SetGraph(InGraph);
+}
+
+void UPCGSubgraphSettings::PostLoad()
+{
+#if WITH_EDITOR
+	if (Subgraph_DEPRECATED)
+	{
+		SubgraphInstance->SetGraph(Subgraph_DEPRECATED);
+		Subgraph_DEPRECATED = nullptr;
+	}
+#endif // WITH_EDITOR
+
+	Super::PostLoad();
+}
+
 #if WITH_EDITOR
 UObject* UPCGSubgraphSettings::GetJumpTargetForDoubleClick() const
 {
 	// Note that there is a const_cast done behind the scenes in Cast.
 	// And this behavior is already used in similar part of UE.
-	return Cast<UObject>(Subgraph);
+	return Cast<UObject>(GetSubgraph());
 }
 
 bool UPCGSubgraphSettings::IsStructuralProperty(const FName& InPropertyName) const
 {
-	return (InPropertyName == GET_MEMBER_NAME_CHECKED(UPCGSubgraphSettings, Subgraph)) || Super::IsStructuralProperty(InPropertyName);
+	// Force structural if name is none. We are probably in a undo/redo situation
+	return (InPropertyName == NAME_None) || (InPropertyName == GET_MEMBER_NAME_CHECKED(UPCGSubgraphSettings, SubgraphInstance)) || Super::IsStructuralProperty(InPropertyName);
 }
 #endif
 
@@ -182,10 +228,16 @@ FPCGElementPtr UPCGSubgraphSettings::CreateElement() const
 	return MakeShared<FPCGSubgraphElement>();
 }
 
-TObjectPtr<UPCGGraph> UPCGSubgraphNode::GetSubgraph() const
+TObjectPtr<UPCGGraph> UPCGBaseSubgraphNode::GetSubgraph() const
+{
+	UPCGGraphInterface* SubgraphInterface = GetSubgraphInterface();
+	return SubgraphInterface ? SubgraphInterface->GetGraph() : nullptr;
+}
+
+TObjectPtr<UPCGGraphInterface> UPCGSubgraphNode::GetSubgraphInterface() const
 {
 	TObjectPtr<UPCGSubgraphSettings> Settings = Cast<UPCGSubgraphSettings>(GetSettings());
-	return Settings ? Settings->Subgraph : nullptr;
+	return Settings ? Settings->GetSubgraphInterface() : nullptr;
 }
 
 FPCGContext* FPCGSubgraphElement::Initialize(const FPCGDataCollection& InputData, TWeakObjectPtr<UPCGComponent> SourceComponent, const UPCGNode* Node)
@@ -211,7 +263,7 @@ bool FPCGSubgraphElement::ExecuteInternal(FPCGContext* InContext) const
 		{
 			const UPCGSubgraphSettings* Settings = Context->GetInputSettings<UPCGSubgraphSettings>();
 			check(Settings);
-			UPCGGraph* Subgraph = Settings->Subgraph;
+			UPCGGraph* Subgraph = Settings->GetSubgraph();
 
 			UPCGSubsystem* Subsystem = Context->SourceComponent.IsValid() ? Context->SourceComponent->GetSubsystem() : nullptr;
 
diff --git a/Engine/Plugins/Experimental/PCG/Source/PCG/Public/Elements/PCGSpawnActor.h b/Engine/Plugins/Experimental/PCG/Source/PCG/Public/Elements/PCGSpawnActor.h
--- a/Engine/Plugins/Experimental/PCG/Source/PCG/Public/Elements/PCGSpawnActor.h
+++ b/Engine/Plugins/Experimental/PCG/Source/PCG/Public/Elements/PCGSpawnActor.h
@@ -9,6 +9,7 @@
 #include "PCGSpawnActor.generated.h"
 
 class AActor;
+class UPCGGraphInterface;
 
 UENUM()
 enum class EPCGSpawnActorOption : uint8
@@ -102,7 +103,7 @@ protected:
 
 public:
 	//~Begin UPCGBaseSubgraphSettings interface
-	virtual UPCGGraph* GetSubgraph() const override;
+	virtual UPCGGraphInterface* GetSubgraphInterface() const override;
 
 protected:
 #if WITH_EDITOR
@@ -123,7 +124,7 @@ class PCG_API UPCGSpawnActorNode : public UPCGBaseSubgraphNode
 	GENERATED_BODY()
 public:
 	/** ~Begin UPCGBaseSubgraphNode interface */
-	virtual TObjectPtr<UPCGGraph> GetSubgraph() const override;
+	virtual TObjectPtr<UPCGGraphInterface> GetSubgraphInterface() const override;
 	/** ~End UPCGBaseSubgraphNode interface */
 };
 
diff --git a/Engine/Plugins/Experimental/PCG/Source/PCG/Public/PCGComponent.h b/Engine/Plugins/Experimental/PCG/Source/PCG/Public/PCGComponent.h
--- a/Engine/Plugins/Experimental/PCG/Source/PCG/Public/PCGComponent.h
+++ b/Engine/Plugins/Experimental/PCG/Source/PCG/Public/PCGComponent.h
@@ -16,6 +16,8 @@ class APCGPartitionActor;
 struct FPCGContext;
 class UPCGComponent;
 class UPCGGraph;
+class UPCGGraphInterface;
+class UPCGGraphInstance;
 class UPCGManagedResource;
 class UPCGData;
 class UPCGSubsystem;
@@ -69,9 +71,11 @@ class PCG_API UPCGComponent : public UActorComponent
 public:
 	/** ~Begin UObject interface */
 	virtual void PostLoad() override;
+	virtual void PostInitProperties() override;
 	virtual void BeginDestroy() override;
 
 #if WITH_EDITOR
+	virtual bool CanEditChange(const FProperty* InProperty) const override;
 	virtual void PostEditImport() override;
 #endif
 	/** ~End UObject interface */
@@ -99,7 +103,9 @@ public:
 
 	bool CanPartition() const;
 
-	UPCGGraph* GetGraph() const { return Graph; }
+	UPCGGraph* GetGraph() const;
+	UPCGGraphInstance* GetGraphInstance() const { return GraphInstance; }
+
 	void SetGraphLocal(UPCGGraph* InGraph);
 
 	UFUNCTION(BlueprintCallable, NetMulticast, Reliable, Category = PCG)
@@ -210,7 +216,10 @@ public:
 	void Refresh();
 	void OnRefresh();
 
-	void DirtyGenerated(EPCGComponentDirtyFlag DataToDirtyFlag = EPCGComponentDirtyFlag::None);
+	/** Dirty generated data depending on the flag. By default the call is forwarded to the local components.
+	    We don't forward if the local component has callbacks that would dirty them too.
+		For example: When a tracked actor move, we only want to dirty the impacted local components.*/
+	void DirtyGenerated(EPCGComponentDirtyFlag DataToDirtyFlag = EPCGComponentDirtyFlag::None, const bool bDispatchToLocalComponents = true);
 
 	/** Reset last generated bounds to force PCGPartitionActor creation on next refresh */
 	void ResetLastGeneratedBounds();
@@ -249,8 +258,13 @@ public:
 	static UPCGData* CreateActorPCGData(AActor* Actor, const UPCGComponent* Component, bool bParseActor = true);
 
 protected:
-	UPROPERTY(BlueprintReadWrite, EditAnywhere, Category = Properties)
-	TObjectPtr<UPCGGraph> Graph;
+	UPROPERTY(BlueprintReadOnly, VisibleAnywhere, Category = Properties, Instanced, meta = (NoResetToDefault, ShowOnlyInnerProperties))
+	TObjectPtr<UPCGGraphInstance> GraphInstance;
+
+#if WITH_EDITORONLY_DATA
+	UPROPERTY()
+	TObjectPtr<UPCGGraph> Graph_DEPRECATED;
+#endif // WITH_EDITORONLY_DATA
 
 private:
 	UPCGData* CreatePCGData();
@@ -286,8 +300,8 @@ private:
 
 	bool GetActorsFromTags(const TSet<FName>& InTags, TSet<TWeakObjectPtr<AActor>>& OutActors, bool bCullAgainstLocalBounds);
 
-	void RefreshAfterGraphChanged(UPCGGraph* InGraph, bool bIsStructural, bool bDirtyInputs);
-	void OnGraphChanged(UPCGGraph* InGraph, EPCGChangeType ChangeType);
+	void RefreshAfterGraphChanged(UPCGGraphInterface* InGraph, bool bIsStructural, bool bDirtyInputs);
+	void OnGraphChanged(UPCGGraphInterface* InGraph, EPCGChangeType ChangeType);
 
 #if WITH_EDITOR
 	virtual void PreEditChange(FProperty* PropertyAboutToChange) override;
@@ -408,6 +422,14 @@ private:
 
 	mutable FCriticalSection GeneratedResourcesLock;
 
+	// Graph instance
+private:
+	/** Will set the given graph interface into our owned graph instance. Must not be used on local components.*/
+	void SetGraphInterfaceLocal(UPCGGraphInterface* InGraphInterface);
+
+	/** Replace the graph instance by this one. Must only be used with local components.*/
+	void SetGraphInstanceFromParent(UPCGGraphInstance* InParentGraphInstance);
+
 #if WITH_EDITOR
 public:
 	mutable PCGUtils::FExtraCapture ExtraCapture;
diff --git a/Engine/Plugins/Experimental/PCG/Source/PCG/Public/PCGGraph.h b/Engine/Plugins/Experimental/PCG/Source/PCG/Public/PCGGraph.h
--- a/Engine/Plugins/Experimental/PCG/Source/PCG/Public/PCGGraph.h
+++ b/Engine/Plugins/Experimental/PCG/Source/PCG/Public/PCGGraph.h
@@ -5,14 +5,32 @@
 #include "PCGNode.h"
 #include "PCGSettings.h"
 
+#include "UObject/ObjectPtr.h"
+
 #include "PCGGraph.generated.h"
 
 #if WITH_EDITOR
-	DECLARE_MULTICAST_DELEGATE_TwoParams(FOnPCGGraphChanged, UPCGGraph* /*Graph*/, EPCGChangeType /*ChangeType*/);
+	DECLARE_MULTICAST_DELEGATE_TwoParams(FOnPCGGraphChanged, UPCGGraphInterface* /*Graph*/, EPCGChangeType /*ChangeType*/);
 #endif // WITH_EDITOR
 
+UCLASS(Abstract)
+class PCG_API UPCGGraphInterface : public UObject
+{
+	GENERATED_BODY()
+
+public:
+	virtual UPCGGraph* GetGraph() PURE_VIRTUAL(UPCGGraphInterface::GetGraph, return nullptr;)
+	virtual const UPCGGraph* GetGraph() const PURE_VIRTUAL(UPCGGraphInterface::GetGraph, return nullptr;)
+
+	bool IsInstance() const;
+
+#if WITH_EDITOR
+	FOnPCGGraphChanged OnGraphChangedDelegate;
+#endif // WITH_EDITOR
+};
+
 UCLASS(BlueprintType, ClassGroup = (Procedural), hidecategories=(Object))
-class PCG_API UPCGGraph : public UObject
+class PCG_API UPCGGraph : public UPCGGraphInterface
 {
 #if WITH_EDITOR
 	friend class FPCGEditor;
@@ -35,6 +53,11 @@ public:
 #endif
 	/** ~End UObject interface */
 
+	/** ~Begin UPCGGraphInterface interface */
+	virtual UPCGGraph* GetGraph() override { return this; }
+	virtual const UPCGGraph* GetGraph() const override { return this; }
+	/** ~End UPCGGraphInterface interface */
+
 #if WITH_EDITORONLY_DATA
 	UPROPERTY(BlueprintReadWrite, EditAnywhere, Category = AssetInfo, AssetRegistrySearchable)
 	bool bExposeToLibrary = false;
@@ -123,10 +146,6 @@ public:
 	void GetTrackedTagsToSettings(FPCGTagToSettingsMap& OutTagsToSettings, TArray<TObjectPtr<const UPCGGraph>>& OutVisitedGraphs) const;
 #endif
 
-#if WITH_EDITOR
-	FOnPCGGraphChanged OnGraphChangedDelegate;
-#endif // WITH_EDITOR
-
 protected:
 	void OnNodeAdded(UPCGNode* InNode);
 	void OnNodeRemoved(UPCGNode* InNode);
@@ -160,6 +179,49 @@ private:
 #endif // WITH_EDITOR
 };
 
+UCLASS(BlueprintType, ClassGroup = (Procedural), hidecategories = (Object))
+class PCG_API UPCGGraphInstance : public UPCGGraphInterface
+{
+	GENERATED_BODY()
+public:
+	/** ~Begin UPCGGraphInterface interface */
+	virtual UPCGGraph* GetGraph() override { return Graph ? Graph->GetGraph() : nullptr; }
+	virtual const UPCGGraph* GetGraph() const override { return Graph ? Graph->GetGraph() : nullptr; }
+	/** ~End UPCGGraphInterface interface */
+
+	// ~Begin UObject interface
+	virtual void PostLoad() override;
+	virtual void BeginDestroy() override;
+
+#if WITH_EDITOR
+	virtual void PreEditChange(FProperty* InProperty) override;
+	virtual void PostEditChangeProperty(FPropertyChangedEvent& PropertyChangedEvent) override;
+	virtual bool CanEditChange(const FProperty* InProperty) const override;
+	virtual void PreEditUndo() override;
+	virtual void PostEditUndo() override;
+#endif
+	// ~End UObject interface
+
+protected:
+#if WITH_EDITOR
+	void OnGraphChanged(UPCGGraphInterface* InGraph, EPCGChangeType ChangeType);
+#endif
+
+public:
+	static TObjectPtr<UPCGGraphInterface> CreateInstance(UObject* InOwner, UPCGGraphInterface* InGraph);
+
+	void SetGraph(UPCGGraphInterface* InGraph);
+
+	UPROPERTY(EditAnywhere, BlueprintReadOnly, Category = Instance)
+	TObjectPtr<UPCGGraphInterface> Graph;
+
+#if WITH_EDITORONLY_DATA
+private:
+	// Transient, to keep track the undo/redo changed the graph.
+	UPCGGraphInterface* UndoRedoGraphCache = nullptr;
+#endif // WITH_EDITORONLY_DATA
+};
+
 #if UE_ENABLE_INCLUDE_ORDER_DEPRECATED_IN_5_2
 #include "CoreMinimal.h"
 #endif
diff --git a/Engine/Plugins/Experimental/PCG/Source/PCG/Public/PCGSubgraph.h b/Engine/Plugins/Experimental/PCG/Source/PCG/Public/PCGSubgraph.h
--- a/Engine/Plugins/Experimental/PCG/Source/PCG/Public/PCGSubgraph.h
+++ b/Engine/Plugins/Experimental/PCG/Source/PCG/Public/PCGSubgraph.h
@@ -3,9 +3,10 @@
 #pragma once
 
 #include "PCGContext.h"
+#include "PCGGraph.h"
 #include "PCGSettings.h"
 
-class UPCGGraph;
+#include "UObject/ObjectPtr.h"
 
 #include "PCGSubgraph.generated.h"
 
@@ -14,15 +15,17 @@ class PCG_API UPCGBaseSubgraphSettings : public UPCGSettings
 {
 	GENERATED_BODY()
 public:
-	virtual UPCGGraph* GetSubgraph() const { return nullptr; }
+	UPCGGraph* GetSubgraph() const;
+	virtual UPCGGraphInterface* GetSubgraphInterface() const { return nullptr; }
 
 	// Use this method from the outside to set the subgraph, as it will connect editor callbacks
-	virtual void SetSubgraph(UPCGGraph* InGraph);
+	virtual void SetSubgraph(UPCGGraphInterface* InGraph);
 
 protected:
 	//~Begin UObject interface implementation
 	virtual void PostLoad() override;
 	virtual void BeginDestroy() override;
+	virtual void PostInitProperties() override;
 #if WITH_EDITOR
 	virtual void PreEditChange(FProperty* PropertyAboutToChange) override;
 	virtual void PostEditChangeProperty(struct FPropertyChangedEvent& PropertyChangedEvent) override;
@@ -37,10 +40,10 @@ protected:
 	TArray<FPCGPinProperties> OutputPinProperties() const override;
 	//~End UPCGSettings interface
 
-	virtual void SetSubgraphInternal(UPCGGraph* InGraph) {}
+	virtual void SetSubgraphInternal(UPCGGraphInterface* InGraph) {}
 
 #if WITH_EDITOR
-	void OnSubgraphChanged(UPCGGraph* InGraph, EPCGChangeType ChangeType);
+	void OnSubgraphChanged(UPCGGraphInterface* InGraph, EPCGChangeType ChangeType);
 #endif
 };
 
@@ -49,6 +52,14 @@ class PCG_API UPCGSubgraphSettings : public UPCGBaseSubgraphSettings
 {
 	GENERATED_BODY()
 
+public:
+	UPCGSubgraphSettings(const FObjectInitializer& InObjectInitializer);
+
+protected:
+	//~Begin UObject interface implementation
+	virtual void PostLoad() override;
+	//~End UObject interface implementation
+
 public:
 	//~UPCGSettings interface implementation
 	virtual UPCGNode* CreateNode() const override;
@@ -67,17 +78,22 @@ protected:
 
 	//~Begin UPCGBaseSubgraphSettings interface
 public:
-	virtual UPCGGraph* GetSubgraph() const override { return Subgraph; }
+	virtual UPCGGraphInterface* GetSubgraphInterface() const override { return SubgraphInstance.Get(); }
 protected:
-	virtual void SetSubgraphInternal(UPCGGraph* InGraph) override { Subgraph = InGraph; }
+	virtual void SetSubgraphInternal(UPCGGraphInterface* InGraph) override;
 #if WITH_EDITOR
 	virtual bool IsStructuralProperty(const FName& InPropertyName) const override;
 #endif
 	//~End UPCGBaseSubgraphSettings interface
 
 public:
-	UPROPERTY(BlueprintReadWrite, EditAnywhere, Category = Settings)
-	TObjectPtr<UPCGGraph> Subgraph;
+	UPROPERTY(BlueprintReadOnly, VisibleAnywhere, Category = Properties, Instanced, meta = (NoResetToDefault, ShowOnlyInnerProperties))
+	TObjectPtr<UPCGGraphInstance> SubgraphInstance;
+
+#if WITH_EDITORONLY_DATA
+	UPROPERTY()
+	TObjectPtr<UPCGGraph> Subgraph_DEPRECATED;
+#endif // WITH_EDITORONLY_DATA
 };
 
 UCLASS(Abstract)
@@ -89,7 +105,8 @@ public:
 	UPROPERTY(BlueprintReadWrite, EditAnywhere, Category = Settings)
 	bool bDynamicGraph = false;
 
-	virtual TObjectPtr<UPCGGraph> GetSubgraph() const { return nullptr; }
+	TObjectPtr<UPCGGraph> GetSubgraph() const;
+	virtual TObjectPtr<UPCGGraphInterface> GetSubgraphInterface() const { return nullptr; }
 };
 
 UCLASS(ClassGroup = (Procedural))
@@ -99,7 +116,7 @@ class PCG_API UPCGSubgraphNode : public UPCGBaseSubgraphNode
 
 public:
 	/** ~Begin UPCGBaseSubgraphNode interface */
-	virtual TObjectPtr<UPCGGraph> GetSubgraph() const override;
+	virtual TObjectPtr<UPCGGraphInterface> GetSubgraphInterface() const override;
 	/** ~End UPCGBaseSubgraphNode interface */
 };
 
diff --git a/Engine/Plugins/Experimental/PCG/Source/PCGEditor/Private/AssetTypeActions/PCGGraphAssetTypeActions.cpp b/Engine/Plugins/Experimental/PCG/Source/PCGEditor/Private/AssetTypeActions/PCGGraphAssetTypeActions.cpp
--- a/Engine/Plugins/Experimental/PCG/Source/PCGEditor/Private/AssetTypeActions/PCGGraphAssetTypeActions.cpp
+++ b/Engine/Plugins/Experimental/PCG/Source/PCGEditor/Private/AssetTypeActions/PCGGraphAssetTypeActions.cpp
@@ -25,3 +25,23 @@ void FPCGGraphAssetTypeActions::OpenAssetEditor(const TArray<UObject*>& InObject
 		}
 	}
 }
+
+FText FPCGGraphInstanceAssetTypeActions::GetName() const
+{
+	return NSLOCTEXT("AssetTypeActions", "PCGGraphInstanceAssetTypeActions", "PCG Graph Instance");
+}
+
+UClass* FPCGGraphInstanceAssetTypeActions::GetSupportedClass() const
+{
+	return UPCGGraphInstance::StaticClass();
+}
+
+FText FPCGGraphInterfaceAssetTypeActions::GetName() const
+{
+	return NSLOCTEXT("AssetTypeActions", "PCGGraphInterfaceAssetTypeActions", "PCG Graph Interface");
+}
+
+UClass* FPCGGraphInterfaceAssetTypeActions::GetSupportedClass() const
+{
+	return UPCGGraphInterface::StaticClass();
+}
\ No newline at end of file
diff --git a/Engine/Plugins/Experimental/PCG/Source/PCGEditor/Private/AssetTypeActions/PCGGraphAssetTypeActions.h b/Engine/Plugins/Experimental/PCG/Source/PCGEditor/Private/AssetTypeActions/PCGGraphAssetTypeActions.h
--- a/Engine/Plugins/Experimental/PCG/Source/PCGEditor/Private/AssetTypeActions/PCGGraphAssetTypeActions.h
+++ b/Engine/Plugins/Experimental/PCG/Source/PCGEditor/Private/AssetTypeActions/PCGGraphAssetTypeActions.h
@@ -11,4 +11,20 @@ public:
 	virtual FText GetName() const override;
 	virtual UClass* GetSupportedClass() const override;
 	virtual void OpenAssetEditor(const TArray<UObject*>& InObjects, TSharedPtr<class IToolkitHost> EditWithinLevelEditor = TSharedPtr<IToolkitHost>()) override;
+};
+
+class FPCGGraphInstanceAssetTypeActions : public FPCGCommonAssetTypeActions
+{
+public:
+	// IAssetTypeActions Implementation
+	virtual FText GetName() const override;
+	virtual UClass* GetSupportedClass() const override;
+};
+
+class FPCGGraphInterfaceAssetTypeActions : public FPCGCommonAssetTypeActions
+{
+public:
+	// IAssetTypeActions Implementation
+	virtual FText GetName() const override;
+	virtual UClass* GetSupportedClass() const override;
 };
\ No newline at end of file
diff --git a/Engine/Plugins/Experimental/PCG/Source/PCGEditor/Private/Details/PCGGraphInstanceDetails.cpp b/Engine/Plugins/Experimental/PCG/Source/PCGEditor/Private/Details/PCGGraphInstanceDetails.cpp
new file mode 100644
--- /dev/null
+++ b/Engine/Plugins/Experimental/PCG/Source/PCGEditor/Private/Details/PCGGraphInstanceDetails.cpp
@@ -0,0 +1,46 @@
+// Copyright Epic Games, Inc. All Rights Reserved.
+
+#include "Details/PCGGraphInstanceDetails.h"
+#include "PCGGraph.h"
+
+#include "PropertyHandle.h"
+#include "DetailLayoutBuilder.h"
+#include "DetailCategoryBuilder.h"
+
+#define LOCTEXT_NAMESPACE "PCGGraphDetails"
+
+TSharedRef<IDetailCustomization> FPCGGraphInstanceDetails::MakeInstance()
+{
+	return MakeShareable(new FPCGGraphInstanceDetails());
+}
+
+void FPCGGraphInstanceDetails::CustomizeDetails(IDetailLayoutBuilder& DetailBuilder)
+{
+	const FName PCGCategoryName("PCG");
+	IDetailCategoryBuilder& PCGCategory = DetailBuilder.EditCategory(PCGCategoryName);
+
+	TArray<TWeakObjectPtr<UObject>> ObjectsBeingCustomized;
+	DetailBuilder.GetObjectsBeingCustomized(ObjectsBeingCustomized);
+
+	for (TWeakObjectPtr<UObject>& Object : ObjectsBeingCustomized)
+	{
+		UPCGGraphInstance* GraphInstance = Cast<UPCGGraphInstance>(Object.Get());
+		if (ensure(GraphInstance))
+		{
+			SelectedGraphInstances.Add(GraphInstance);
+		}
+	}
+
+	TArray<TSharedRef<IPropertyHandle>> AllProperties;
+	bool bSimpleProperties = true;
+	bool bAdvancedProperties = false;
+	// Add all properties in the category in order
+	PCGCategory.GetDefaultProperties(AllProperties, bSimpleProperties, bAdvancedProperties);
+
+	for (auto& Property : AllProperties)
+	{
+		PCGCategory.AddProperty(Property);
+	}
+}
+
+#undef LOCTEXT_NAMESPACE
diff --git a/Engine/Plugins/Experimental/PCG/Source/PCGEditor/Private/PCGEditorModule.cpp b/Engine/Plugins/Experimental/PCG/Source/PCGEditor/Private/PCGEditorModule.cpp
--- a/Engine/Plugins/Experimental/PCG/Source/PCGEditor/Private/PCGEditorModule.cpp
+++ b/Engine/Plugins/Experimental/PCG/Source/PCGEditor/Private/PCGEditorModule.cpp
@@ -15,6 +15,7 @@
 
 #include "Details/PCGBlueprintSettingsDetails.h"
 #include "Details/PCGGraphDetails.h"
+#include "Details/PCGGraphInstanceDetails.h"
 #include "Details/PCGVolumeDetails.h"
 #include "Details/PCGAttributePropertySelectorDetails.h"
 
@@ -75,6 +76,7 @@ void FPCGEditorModule::RegisterDetailsCustomizations()
 	PropertyEditor.RegisterCustomClassLayout("PCGBlueprintSettings", FOnGetDetailCustomizationInstance::CreateStatic(&FPCGBlueprintSettingsDetails::MakeInstance));
 	PropertyEditor.RegisterCustomClassLayout("PCGComponent", FOnGetDetailCustomizationInstance::CreateStatic(&FPCGComponentDetails::MakeInstance));
 	PropertyEditor.RegisterCustomClassLayout("PCGGraph", FOnGetDetailCustomizationInstance::CreateStatic(&FPCGGraphDetails::MakeInstance));
+	PropertyEditor.RegisterCustomClassLayout("PCGGraphInstance", FOnGetDetailCustomizationInstance::CreateStatic(&FPCGGraphInstanceDetails::MakeInstance));
 	PropertyEditor.RegisterCustomClassLayout("PCGVolume", FOnGetDetailCustomizationInstance::CreateStatic(&FPCGVolumeDetails::MakeInstance));
 
 	PropertyEditor.RegisterCustomPropertyTypeLayout("PCGAttributePropertySelector", FOnGetPropertyTypeCustomizationInstance::CreateStatic(&FPCGAttributePropertySelectorDetails::MakeInstance));
@@ -89,6 +91,7 @@ void FPCGEditorModule::UnregisterDetailsCustomizations()
 		PropertyModule->UnregisterCustomClassLayout("PCGBlueprintSettings");
 		PropertyModule->UnregisterCustomClassLayout("PCGComponent");
 		PropertyModule->UnregisterCustomClassLayout("PCGGraph");
+		PropertyModule->UnregisterCustomClassLayout("PCGGraphInstance");
 		PropertyModule->UnregisterCustomClassLayout("PCGVolume");
 
 		PropertyModule->UnregisterCustomPropertyTypeLayout("PCGAttributePropertySelector");
@@ -103,6 +106,8 @@ void FPCGEditorModule::RegisterAssetTypeActions()
 	PCGAssetCategory = AssetTools.RegisterAdvancedAssetCategory(FName(TEXT("PCG")), LOCTEXT("PCGAssetCategory", "PCG"));
 
 	RegisteredAssetTypeActions.Emplace(MakeShareable(new FPCGGraphAssetTypeActions()));
+	RegisteredAssetTypeActions.Emplace(MakeShareable(new FPCGGraphInstanceAssetTypeActions()));
+	RegisteredAssetTypeActions.Emplace(MakeShareable(new FPCGGraphInterfaceAssetTypeActions()));
 	RegisteredAssetTypeActions.Emplace(MakeShareable(new FPCGSettingsAssetTypeActions()));
 
 	for (auto Action : RegisteredAssetTypeActions)
diff --git a/Engine/Plugins/Experimental/PCG/Source/PCGEditor/Public/Details/PCGGraphInstanceDetails.h b/Engine/Plugins/Experimental/PCG/Source/PCGEditor/Public/Details/PCGGraphInstanceDetails.h
new file mode 100644
--- /dev/null
+++ b/Engine/Plugins/Experimental/PCG/Source/PCGEditor/Public/Details/PCGGraphInstanceDetails.h
@@ -0,0 +1,24 @@
+// Copyright Epic Games, Inc. All Rights Reserved.
+
+#pragma once
+
+#include "IDetailCustomization.h"
+#include "UObject/WeakObjectPtrTemplates.h"
+
+class UPCGGraphInstance;
+
+class FPCGGraphInstanceDetails : public IDetailCustomization
+{
+public:
+	static TSharedRef<IDetailCustomization> MakeInstance();
+
+	virtual void CustomizeDetails(IDetailLayoutBuilder& DetailBuilder) override;
+
+private:
+	TArray<TWeakObjectPtr<UPCGGraphInstance>> SelectedGraphInstances;
+};
+
+#if UE_ENABLE_INCLUDE_ORDER_DEPRECATED_IN_5_2
+#include "Input/Reply.h"
+#include "UObject/WeakObjectPtr.h"
+#endif
