[PCG] Change bIsPartitioned default from true to false

diff --git a/Engine/Plugins/Experimental/PCG/Source/PCG/Private/PCGComponent.cpp b/Engine/Plugins/Experimental/PCG/Source/PCG/Private/PCGComponent.cpp
--- a/Engine/Plugins/Experimental/PCG/Source/PCG/Private/PCGComponent.cpp
+++ b/Engine/Plugins/Experimental/PCG/Source/PCG/Private/PCGComponent.cpp
@@ -939,6 +939,12 @@ void UPCGComponent::PostLoad()
 		ExclusionTags_DEPRECATED.Reset();
 	}
 
+	// Make sure that bPartitioned is false when it can't be true.
+	if (bIsPartitioned && !CanPartition())
+	{
+		bIsPartitioned = false;
+	}
+
 	/** Deprecation code, should be removed once generated data has been updated */
 	if (bGenerated && GeneratedResources.Num() == 0)
 	{
@@ -1041,6 +1047,24 @@ void UPCGComponent::PostInitProperties()
 #endif // WITH_EDITOR
 	}
 
+	// Force bIsPartitioned at False for new objects
+	// We detect new object if they are not a default object/archetype and/or they do not need load.
+	// In some cases, were the component is a default sub object (like APCGVolume), it has no loading flags
+	// even if it is loading, we need to recurse until we find an owner that is not a default sub object.
+	bool bIsNewObject = false;
+	UObject* CurrentInspectedObject = this;
+	while (CurrentInspectedObject && CurrentInspectedObject->HasAnyFlags(RF_DefaultSubObject))
+	{
+		CurrentInspectedObject = CurrentInspectedObject->GetOuter();
+	}
+
+	bIsNewObject = CurrentInspectedObject && !CurrentInspectedObject->HasAnyFlags(RF_ClassDefaultObject | RF_ArchetypeObject | RF_NeedLoad | RF_NeedPostLoad);
+
+	if (bIsNewObject)
+	{
+		bIsPartitioned = false;
+	}
+
 	Super::PostInitProperties();
 }
 
diff --git a/Engine/Plugins/Experimental/PCG/Source/PCG/Public/PCGComponent.h b/Engine/Plugins/Experimental/PCG/Source/PCG/Public/PCGComponent.h
--- a/Engine/Plugins/Experimental/PCG/Source/PCG/Public/PCGComponent.h
+++ b/Engine/Plugins/Experimental/PCG/Source/PCG/Public/PCGComponent.h
@@ -175,7 +175,8 @@ public:
 	UPROPERTY(BlueprintReadWrite, EditAnywhere, Category = Properties)
 	bool bActivated = true;
 
-	UPROPERTY(BlueprintReadWrite, EditAnywhere, AdvancedDisplay, Category = Properties, meta = (EditCondition = "!bIsComponentLocal"))
+	// Note that default value is now "False" for new objects. True is left for backcompatibility.
+	UPROPERTY(BlueprintReadWrite, EditAnywhere, AdvancedDisplay, Category = Properties, meta = (EditCondition = "!bIsComponentLocal", Tooltip = "In WorldPartiion map, will partition the component in a grid, according to PCGWorldActor settings, dispatching the generation to multiple local components."))
 	bool bIsPartitioned = true;
 
 	UPROPERTY(BlueprintReadOnly, EditAnywhere, Category = Properties, AdvancedDisplay, meta = (EditCondition = "!bIsComponentLocal", EditConditionHides))
