[PCG] Make FindOrCreateAttribute more robust with type mismatch.

diff --git a/Engine/Plugins/Experimental/PCG/Source/PCG/Private/Elements/PCGSplineSampler.cpp b/Engine/Plugins/Experimental/PCG/Source/PCG/Private/Elements/PCGSplineSampler.cpp
--- a/Engine/Plugins/Experimental/PCG/Source/PCG/Private/Elements/PCGSplineSampler.cpp
+++ b/Engine/Plugins/Experimental/PCG/Source/PCG/Private/Elements/PCGSplineSampler.cpp
@@ -485,27 +485,16 @@ namespace PCGSplineSampler
 			// Initialize metadata accessors if needed
 			if (OutPointData && OutPointData->Metadata && (Params.bComputeDirectionDelta || Params.bComputeCurvature))
 			{
-				auto FindOrCreateAttribute = [](const FName& AttributeName, UPCGMetadata* Metadata)
-				{
-					FPCGMetadataAttribute<double>* Attribute = Metadata->FindOrCreateAttribute(AttributeName, 0.0, true, true);
-					if (Attribute->GetTypeId() != PCG::Private::MetadataTypes<double>::Id)
-					{
-						Metadata->DeleteAttribute(AttributeName);
-						Attribute = Metadata->CreateAttribute<double>(AttributeName, 0.0, true, true);
-					}
-
-					return Attribute;
-				};
-
+				constexpr double DefaultValue = 0.0;
 				if (Params.bComputeDirectionDelta)
 				{
-					NextDirectionDeltaAttribute = FindOrCreateAttribute(Params.NextDirectionDeltaAttribute, OutPointData->Metadata);
+					NextDirectionDeltaAttribute = OutPointData->Metadata->FindOrCreateAttribute(Params.NextDirectionDeltaAttribute, DefaultValue);
 					bSetMetadata |= (NextDirectionDeltaAttribute != nullptr);
 				}
 
 				if (Params.bComputeCurvature)
 				{
-					CurvatureAttribute = FindOrCreateAttribute(Params.CurvatureAttribute, OutPointData->Metadata);
+					CurvatureAttribute = OutPointData->Metadata->FindOrCreateAttribute(Params.CurvatureAttribute, DefaultValue);
 					bSetMetadata |= (CurvatureAttribute != nullptr);
 				}
 			}
diff --git a/Engine/Plugins/Experimental/PCG/Source/PCG/Public/Metadata/PCGMetadata.h b/Engine/Plugins/Experimental/PCG/Source/PCG/Public/Metadata/PCGMetadata.h
--- a/Engine/Plugins/Experimental/PCG/Source/PCG/Public/Metadata/PCGMetadata.h
+++ b/Engine/Plugins/Experimental/PCG/Source/PCG/Public/Metadata/PCGMetadata.h
@@ -248,11 +248,22 @@ public:
 	int64 GetItemKeyCountForParent() const;
 	int64 GetItemCountForChild() const;
 
+	/**
+	* Create a new attribute. If the attribute already exists, it will raise a warning (use FindOrCreateAttribute if this usecase can arise)
+	* If the attribute already exists but is of the wrong type, it will fail and return nullptr. Same if the name is invalid.
+	* Return a typed attribute pointer, of the requested type T.
+	*/
 	template<typename T>
 	FPCGMetadataAttribute<T>* CreateAttribute(FName AttributeName, const T& DefaultValue, bool bAllowsInterpolation, bool bOverrideParent);
 
+	/**
+	* Find or create an attribute. Follows CreateAttribute signature.
+	* Extra boolean bOverwriteIfTypeMismatch allows to overwrite an existing attribute if the type mismatch.
+	* Same as CreateAttribute, it will return nullptr if the attribute name is invalid.
+	* Return a typed attribute pointer, of the requested type T.
+	*/
 	template<typename T>
-	FPCGMetadataAttribute<T>* FindOrCreateAttribute(FName AttributeName, const T& DefaultValue = T{}, bool bAllowsInterpolation = true, bool bOverrideParent = true);
+	FPCGMetadataAttribute<T>* FindOrCreateAttribute(FName AttributeName, const T& DefaultValue = T{}, bool bAllowsInterpolation = true, bool bOverrideParent = true, bool bOverwriteIfTypeMismatch = true);
 
 protected:
 	FPCGMetadataAttributeBase* CopyAttribute(FName AttributeToCopy, FName NewAttributeName, bool bKeepParent, bool bCopyEntries, bool bCopyValues);
@@ -332,9 +343,17 @@ FPCGMetadataAttribute<T>* UPCGMetadata::CreateAttribute(FName AttributeName, con
 }
 
 template<typename T>
-FPCGMetadataAttribute<T>* UPCGMetadata::FindOrCreateAttribute(FName AttributeName, const T& DefaultValue, bool bAllowsInterpolation, bool bOverrideParent)
+FPCGMetadataAttribute<T>* UPCGMetadata::FindOrCreateAttribute(FName AttributeName, const T& DefaultValue, bool bAllowsInterpolation, bool bOverrideParent, bool bOverwriteIfTypeMismatch)
 {
 	FPCGMetadataAttribute<T>* Attribute = GetMutableTypedAttribute<T>(AttributeName);
+
+	// If Attribute is null, but we have an attribute with this name, we have a type mismatch.
+	// Will be overwrite if flag bOverwriteIfTypeMismatch is at true.
+	if (!Attribute && HasAttribute(AttributeName) && bOverwriteIfTypeMismatch)
+	{
+		DeleteAttribute(AttributeName);
+	}
+
 	if (!Attribute)
 	{
 		Attribute = CreateAttribute<T>(AttributeName, DefaultValue, bAllowsInterpolation, bOverrideParent);
