Demoted error to warning in trivial cases of invalid mesh paths in the static mesh spawner.
Improved logging so we don't log multiple times the same issue.

#rnx

diff --git a/Engine/Plugins/Experimental/PCG/Source/PCG/Private/MeshSelectors/PCGMeshSelectorByAttribute.cpp b/Engine/Plugins/Experimental/PCG/Source/PCG/Private/MeshSelectors/PCGMeshSelectorByAttribute.cpp
--- a/Engine/Plugins/Experimental/PCG/Source/PCG/Private/MeshSelectors/PCGMeshSelectorByAttribute.cpp
+++ b/Engine/Plugins/Experimental/PCG/Source/PCG/Private/MeshSelectors/PCGMeshSelectorByAttribute.cpp
@@ -92,8 +92,22 @@ void UPCGMeshSelectorByAttribute::SelectInstances_Implementation(
 		// if this ValueKey has not been seen before, let's cache it for the future
 		if (!NewMesh)
 		{
-			FSoftObjectPath MeshPath(Attribute->GetValue(ValueKey));
-			Mesh = TSoftObjectPtr<UStaticMesh>(MeshPath);
+			FString MeshSoftObjectPath = Attribute->GetValue(ValueKey);
+
+			if (!MeshSoftObjectPath.IsEmpty() && MeshSoftObjectPath != TEXT("None"))
+			{
+				FSoftObjectPath MeshPath(MeshSoftObjectPath);
+				Mesh = TSoftObjectPtr<UStaticMesh>(MeshPath);
+
+				if (Mesh.IsNull())
+				{
+					PCGE_LOG_C(Error, &Context, "Invalid mesh path: %s.", *MeshSoftObjectPath);
+				}
+			}
+			else
+			{
+				PCGE_LOG_C(Warning, &Context, "Trivially invalid mesh path used: %s", *MeshSoftObjectPath);
+			}
 
 			ValueKeyToMesh.Add(ValueKey, Mesh);
 		}
@@ -104,7 +118,6 @@ void UPCGMeshSelectorByAttribute::SelectInstances_Implementation(
 
 		if (Mesh.IsNull())
 		{
-			PCGE_LOG_C(Error, &Context, "Invalid object path.");
 			continue;
 		}
 
