[PCG] Fix incorrect default for full-point-crc console var

diff --git a/Engine/Plugins/Experimental/PCG/Source/PCG/Private/Data/PCGPointData.cpp b/Engine/Plugins/Experimental/PCG/Source/PCG/Private/Data/PCGPointData.cpp
--- a/Engine/Plugins/Experimental/PCG/Source/PCG/Private/Data/PCGPointData.cpp
+++ b/Engine/Plugins/Experimental/PCG/Source/PCG/Private/Data/PCGPointData.cpp
@@ -15,7 +15,7 @@
 
 static TAutoConsoleVariable<bool> CVarCacheFullPointDataCrc(
 	TEXT("pcg.Cache.FullPointDataCrc"),
-	true,
+	false,
 	TEXT("Enable fine-grained CRC of point data for change tracking, rather than using data UID."));
 
 namespace PCGPointHelpers
