Added some robustness to some PCG data classes esp. when some components are deleted
Fixed an issue with dependencies across scheduled-but-not-processed tasks in PCG and cancellation - improved cancellation on active tasks too.
Protected some code paths against generation on temporary actors
Deactivated full point CRC until we fix the pending kill actor dependency

#rnx

diff --git a/Engine/Plugins/Experimental/PCG/Source/PCG/Private/Data/PCGPrimitiveData.cpp b/Engine/Plugins/Experimental/PCG/Source/PCG/Private/Data/PCGPrimitiveData.cpp
--- a/Engine/Plugins/Experimental/PCG/Source/PCG/Private/Data/PCGPrimitiveData.cpp
+++ b/Engine/Plugins/Experimental/PCG/Source/PCG/Private/Data/PCGPrimitiveData.cpp
@@ -28,7 +28,7 @@ bool UPCGPrimitiveData::SamplePoint(const FTransform& InTransform, const FBox& I
 
 	const FVector BoxCenter = InTransform.TransformPosition(InBounds.GetCenter());
 
-	if (Primitive->OverlapComponent(BoxCenter, InTransform.GetRotation(), CollisionShape))
+	if (ensure(IsValid(Primitive.Get())) && Primitive->OverlapComponent(BoxCenter, InTransform.GetRotation(), CollisionShape))
 	{
 		OutPoint.Transform = InTransform;
 		OutPoint.SetLocalBounds(InBounds);
@@ -50,7 +50,7 @@ const UPCGPointData* UPCGPrimitiveData::CreatePointData(FPCGContext* Context) co
 
 	UPCGPointData* Data = PCGVolumeSampler::SampleVolume(Context, this, SamplerSettings);
 
-	if (Data)
+	if (Data && ensure(Primitive.Get()))
 	{
 		UE_LOG(LogPCG, Verbose, TEXT("Primitive %s extracted %d points"), *Primitive->GetFName().ToString(), Data->GetPoints().Num());
 	}
diff --git a/Engine/Plugins/Experimental/PCG/Source/PCG/Private/Data/PCGSplineData.cpp b/Engine/Plugins/Experimental/PCG/Source/PCG/Private/Data/PCGSplineData.cpp
--- a/Engine/Plugins/Experimental/PCG/Source/PCG/Private/Data/PCGSplineData.cpp
+++ b/Engine/Plugins/Experimental/PCG/Source/PCG/Private/Data/PCGSplineData.cpp
@@ -45,12 +45,12 @@ int UPCGSplineData::GetNumSegments() const
 
 FVector::FReal UPCGSplineData::GetSegmentLength(int SegmentIndex) const
 {
-	return Spline->GetDistanceAlongSplineAtSplinePoint(SegmentIndex + 1) - Spline->GetDistanceAlongSplineAtSplinePoint(SegmentIndex);
+	return Spline ? Spline->GetDistanceAlongSplineAtSplinePoint(SegmentIndex + 1) - Spline->GetDistanceAlongSplineAtSplinePoint(SegmentIndex) : 0;
 }
 
 FVector UPCGSplineData::GetLocationAtDistance(int SegmentIndex, FVector::FReal Distance, bool bWorldSpace) const
 {
-	return Spline->GetLocationAtDistanceAlongSpline(Spline->GetDistanceAlongSplineAtSplinePoint(SegmentIndex) + Distance, bWorldSpace ? ESplineCoordinateSpace::World : ESplineCoordinateSpace::Local);
+	return Spline ? Spline->GetLocationAtDistanceAlongSpline(Spline->GetDistanceAlongSplineAtSplinePoint(SegmentIndex) + Distance, bWorldSpace ? ESplineCoordinateSpace::World : ESplineCoordinateSpace::Local) : FVector::ZeroVector;
 }
 
 FTransform UPCGSplineData::GetTransformAtDistance(int SegmentIndex, FVector::FReal Distance, bool bWorldSpace, FBox* OutBounds) const
@@ -60,11 +60,16 @@ FTransform UPCGSplineData::GetTransformAtDistance(int SegmentIndex, FVector::FRe
 		*OutBounds = FBox::BuildAABB(FVector::ZeroVector, FVector::OneVector);
 	}
 
-	return Spline->GetTransformAtDistanceAlongSpline(Spline->GetDistanceAlongSplineAtSplinePoint(SegmentIndex) + Distance, bWorldSpace ? ESplineCoordinateSpace::World : ESplineCoordinateSpace::Local, /*bUseScale=*/true);
+	return Spline ? Spline->GetTransformAtDistanceAlongSpline(Spline->GetDistanceAlongSplineAtSplinePoint(SegmentIndex) + Distance, bWorldSpace ? ESplineCoordinateSpace::World : ESplineCoordinateSpace::Local, /*bUseScale=*/true) : FTransform::Identity;
 }
 
 FVector::FReal UPCGSplineData::GetCurvatureAtDistance(int SegmentIndex, FVector::FReal Distance) const
 {
+	if (!Spline)
+	{
+		return 0;
+	}
+
 	const float FullDistance = Spline->GetDistanceAlongSplineAtSplinePoint(SegmentIndex) + Distance;
 	const float Param = Spline->SplineCurves.ReparamTable.Eval(FullDistance, 0.0f);
 
@@ -95,7 +100,10 @@ const UPCGPointData* UPCGSplineData::CreatePointData(FPCGContext* Context) const
 
 	PCGSplineSampler::SampleLineData(this, this, nullptr, SamplerParams, Data);
 
-	UE_LOG(LogPCG, Verbose, TEXT("Spline %s generated %d points"), *Spline->GetFName().ToString(), Data->GetPoints().Num());
+	if (Spline)
+	{
+		UE_LOG(LogPCG, Verbose, TEXT("Spline %s generated %d points"), *Spline->GetFName().ToString(), Data->GetPoints().Num());
+	}
 
 	return Data;
 }
@@ -107,6 +115,11 @@ FBox UPCGSplineData::GetBounds() const
 
 bool UPCGSplineData::SamplePoint(const FTransform& InTransform, const FBox& InBounds, FPCGPoint& OutPoint, UPCGMetadata* OutMetadata) const
 {
+	if (!Spline)
+	{
+		return false;
+	}
+
 	// TODO: support metadata
 	// TODO: support proper bounds
 
@@ -172,11 +185,18 @@ bool UPCGSplineProjectionData::SamplePoint(const FTransform& InTransform, const
 		return Super::SamplePoint(InTransform, InBounds, OutPoint, OutMetadata);
 	}
 
+	check(GetSpline());
+	if (!GetSpline()->Spline)
+	{
+		return false;
+	}
+
 	// Find nearest point on projected spline by lifting point along projection direction to closest position on spline. This way
 	// when we sample the spline we get a similar result to if the spline had been projected onto the surface.
 
 	const FVector InPosition = InTransform.GetLocation();
 	const USplineComponent* Spline = GetSpline()->Spline.Get();
+	check(GetSurface());
 	const FVector& SurfaceNormal = GetSurface()->GetNormal();
 
 	// Project to 2D space
@@ -225,6 +245,7 @@ bool UPCGSplineProjectionData::SamplePoint(const FTransform& InTransform, const
 
 FVector2D UPCGSplineProjectionData::Project(const FVector& InVector) const
 {
+	check(GetSurface());
 	const FVector& SurfaceNormal = GetSurface()->GetNormal();
 	FVector Projection = InVector - InVector.ProjectOnToNormal(SurfaceNormal);
 
@@ -260,7 +281,10 @@ void UPCGSplineProjectionData::Initialize(const UPCGSplineData* InSourceSpline,
 {
 	Super::Initialize(InSourceSpline, InTargetSurface, InParams);
 
+	check(GetSpline());
 	const USplineComponent* Spline = GetSpline()->Spline.Get();
+
+	check(GetSurface());
 	const FVector& SurfaceNormal = GetSurface()->GetNormal();
 
 	if (Spline)
diff --git a/Engine/Plugins/Experimental/PCG/Source/PCG/Private/Graph/PCGGraphExecutor.cpp b/Engine/Plugins/Experimental/PCG/Source/PCG/Private/Graph/PCGGraphExecutor.cpp
--- a/Engine/Plugins/Experimental/PCG/Source/PCG/Private/Graph/PCGGraphExecutor.cpp
+++ b/Engine/Plugins/Experimental/PCG/Source/PCG/Private/Graph/PCGGraphExecutor.cpp
@@ -15,6 +15,7 @@
 #include "Graph/PCGGraphCache.h"
 #include "Metadata/PCGMetadata.h"
 
+#include "Algo/AnyOf.h"
 #include "Async/Async.h"
 #include "GameFramework/Actor.h"
 
@@ -131,113 +132,179 @@ FPCGTaskId FPCGGraphExecutor::Schedule(UPCGGraph* Graph, UPCGComponent* SourceCo
 			ScheduledTask.Tasks[ScheduledTask.Tasks.Num() - 2].Inputs.Emplace(ExternalDependency, nullptr, nullptr, /*bConsumeData=*/false);
 		}
 
+		ScheduledTask.FirstTaskIndex = ScheduledTask.Tasks.Num() - 2;
+		ScheduledTask.LastTaskIndex = ScheduledTask.Tasks.Num() - 1;
+
 		ScheduleLock.Unlock();
 	}
 
 	return ScheduledId;
 }
 
-void FPCGGraphExecutor::Cancel(UPCGComponent* InComponent)
+TArray<UPCGComponent*> FPCGGraphExecutor::Cancel(UPCGComponent* InComponent)
 {
 	auto CancelIfSameComponent = [InComponent](TWeakObjectPtr<UPCGComponent> Component)
 	{
 		return Component.IsValid() && InComponent == Component;
 	};
 
-	return Cancel(CancelIfSameComponent);
+	return Cancel(CancelIfSameComponent).Array();
 }
 
 TArray<UPCGComponent*> FPCGGraphExecutor::Cancel(UPCGGraph* InGraph)
 {
-	TSet<UPCGComponent*> CancelledComponents;
-
-	auto CancelIfComponentUsesGraph = [InGraph, &CancelledComponents](TWeakObjectPtr<UPCGComponent> Component)
+	auto CancelIfComponentUsesGraph = [InGraph](TWeakObjectPtr<UPCGComponent> Component)
 	{
-		if (Component.IsValid() && Component->GetGraph() == InGraph)
-		{
-			CancelledComponents.Add(Component.Get());
-			return true;
-		}
-		else
-		{
-			return false;
-		}
+		return (Component.IsValid() && Component->GetGraph() == InGraph);
 	};
 
-	Cancel(CancelIfComponentUsesGraph);
-	return CancelledComponents.Array();
+	return Cancel(CancelIfComponentUsesGraph).Array();
 }
 
 TArray<UPCGComponent*> FPCGGraphExecutor::CancelAll()
 {
-	TSet<UPCGComponent*> CancelledComponents;
-
-	auto CancelAllGeneration = [&CancelledComponents](TWeakObjectPtr<UPCGComponent> Component)
+	auto CancelAllGeneration = [](TWeakObjectPtr<UPCGComponent> Component)
 	{
-		if (Component.IsValid())
-		{
-			CancelledComponents.Add(Component.Get());
-			return true;
-		}
-		else
-		{
-			return false;
-		}
+		return Component.IsValid();
 	};
 
-	Cancel(CancelAllGeneration);
-	return CancelledComponents.Array();
+	return Cancel(CancelAllGeneration).Array();
 }
 
-void FPCGGraphExecutor::Cancel(TFunctionRef<bool(TWeakObjectPtr<UPCGComponent>)> CancelFilter)
+TSet<UPCGComponent*> FPCGGraphExecutor::Cancel(TFunctionRef<bool(TWeakObjectPtr<UPCGComponent>)> CancelFilter)
 {
-	// Remove from scheduled tasks
+	TSet<UPCGComponent*> CancelledComponents;
+
+	// Visit scheduled tasks
 	ScheduleLock.Lock();
-	for (int32 ScheduledTaskIndex = ScheduledTasks.Num() - 1; ScheduledTaskIndex >= 0; --ScheduledTaskIndex)
+	for (const FPCGGraphScheduleTask& ScheduledTask : ScheduledTasks)
 	{
-		if(CancelFilter(ScheduledTasks[ScheduledTaskIndex].SourceComponent))
+		if (CancelFilter(ScheduledTask.SourceComponent))
 		{
-			ScheduledTasks.RemoveAtSwap(ScheduledTaskIndex);
+			CancelledComponents.Add(ScheduledTask.SourceComponent.Get());
 		}
 	}
 	ScheduleLock.Unlock();
 
-	// Remove from ready tasks
-	for (int32 ReadyTaskIndex = ReadyTasks.Num() - 1; ReadyTaskIndex >= 0; --ReadyTaskIndex)
+	// Visit ready tasks
+	for (const FPCGGraphTask& Task : ReadyTasks)
 	{
-		FPCGGraphTask& Task = ReadyTasks[ReadyTaskIndex];
+		if (CancelFilter(Task.SourceComponent))
+		{
+			CancelledComponents.Add(Task.SourceComponent.Get());
+		}
+	}
 
-		if(CancelFilter(Task.SourceComponent))
+	// Visit active tasks
+	for (const FPCGGraphActiveTask& Task : ActiveTasks)
+	{
+		if (Task.Context && CancelFilter(Task.Context->SourceComponent))
 		{
-			FPCGTaskId CancelledTaskId = Task.NodeId;
-			delete Task.Context;
-			ReadyTasks.RemoveAtSwap(ReadyTaskIndex);
-			CancelNextTasks(CancelledTaskId);
+			CancelledComponents.Add(Task.Context->SourceComponent.Get());
 		}
 	}
 
-	// Remove from active tasks
-	for (int32 ActiveTaskIndex = ActiveTasks.Num() - 1; ActiveTaskIndex >= 0; --ActiveTaskIndex)
+	// Visit sleeping tasks
+	for (const FPCGGraphActiveTask& Task : SleepingTasks)
 	{
-		FPCGGraphActiveTask& Task = ActiveTasks[ActiveTaskIndex];
-		if(Task.Context && CancelFilter(Task.Context->SourceComponent))
+		if (Task.Context && CancelFilter(Task.Context->SourceComponent))
 		{
-			CurrentlyUsedThreads -= Task.Context->NumAvailableTasks;
-			FPCGTaskId CancelledTaskId = Task.NodeId;
-			ActiveTasks.RemoveAtSwap(ActiveTaskIndex);
-			CancelNextTasks(CancelledTaskId);
+			CancelledComponents.Add(Task.Context->SourceComponent.Get());
 		}
 	}
 
-	// Remove from sleeping tasks
-	for (int32 SleepingTaskIndex = SleepingTasks.Num() - 1; SleepingTaskIndex >= 0; --SleepingTaskIndex)
+	// Early out - nothing to cancel
+	if (CancelledComponents.IsEmpty())
+	{
+		return CancelledComponents;
+	}
+
+	TArray<FPCGTaskId> CancelledScheduledTasks;
+
+	bool bStableCancellationSet = false;
+	while (!bStableCancellationSet)
 	{
-		FPCGGraphActiveTask& Task = SleepingTasks[SleepingTaskIndex];
-		if(Task.Context && CancelFilter(Task.Context->SourceComponent))
+		bStableCancellationSet = true;
+
+		// Remove from scheduled tasks
+		ScheduleLock.Lock();
+		for (int32 ScheduledTaskIndex = ScheduledTasks.Num() - 1; ScheduledTaskIndex >= 0; --ScheduledTaskIndex)
 		{
-			FPCGTaskId CancelledTaskId = Task.NodeId;
-			SleepingTasks.RemoveAtSwap(CancelledTaskId);
-			CancelNextTasks(CancelledTaskId);
+			FPCGGraphScheduleTask& ScheduledTask = ScheduledTasks[ScheduledTaskIndex];
+			if (CancelledComponents.Contains(ScheduledTask.SourceComponent.Get()))
+			{
+				CancelledScheduledTasks.Add(ScheduledTask.Tasks[ScheduledTask.LastTaskIndex].NodeId);
+				ScheduledTasks.RemoveAtSwap(ScheduledTaskIndex);
+			}
+		}
+
+		auto RemoveScheduledTasks = [this, &bStableCancellationSet, &CancelledComponents, &CancelledScheduledTasks](FPCGTaskId EndTaskId)
+		{
+
+		};
+
+		// WARNING: variable upper bound
+		for (int32 CancelledTaskIdIndex = 0; CancelledTaskIdIndex < CancelledScheduledTasks.Num(); ++CancelledTaskIdIndex)
+		{
+			const FPCGTaskId EndTaskId = CancelledScheduledTasks[CancelledTaskIdIndex];
+			for (int32 ScheduledTaskIndex = ScheduledTasks.Num() - 1; ScheduledTaskIndex >= 0; --ScheduledTaskIndex)
+			{
+				FPCGGraphScheduleTask& ScheduledTask = ScheduledTasks[ScheduledTaskIndex];
+				const bool bContainsDependency = Algo::AnyOf(ScheduledTask.Tasks[ScheduledTask.FirstTaskIndex].Inputs, [EndTaskId](const FPCGGraphTaskInput& Input) { return Input.TaskId == EndTaskId; });
+
+				if (bContainsDependency)
+				{
+					if (!CancelledComponents.Contains(ScheduledTask.SourceComponent.Get()))
+					{
+						CancelledComponents.Add(ScheduledTask.SourceComponent.Get());
+						bStableCancellationSet = false;
+					}
+
+					CancelledScheduledTasks.Add(ScheduledTask.Tasks[ScheduledTask.LastTaskIndex].NodeId);
+					ScheduledTasks.RemoveAtSwap(ScheduledTaskIndex);
+				}
+			}
+		}
+
+		CancelledScheduledTasks.Reset();
+		ScheduleLock.Unlock();
+
+		// Remove from ready tasks
+		for (int32 ReadyTaskIndex = ReadyTasks.Num() - 1; ReadyTaskIndex >= 0; --ReadyTaskIndex)
+		{
+			FPCGGraphTask& Task = ReadyTasks[ReadyTaskIndex];
+
+			if (CancelledComponents.Contains(Task.SourceComponent.Get()))
+			{
+				FPCGTaskId CancelledTaskId = Task.NodeId;
+				delete Task.Context;
+				ReadyTasks.RemoveAtSwap(ReadyTaskIndex);
+				bStableCancellationSet &= !CancelNextTasks(CancelledTaskId, CancelledComponents);
+			}
+		}
+
+		// Mark as cancelled in the active tasks - needed to make sure we're not breaking the current execution (if any)
+		for (int32 ActiveTaskIndex = ActiveTasks.Num() - 1; ActiveTaskIndex >= 0; --ActiveTaskIndex)
+		{
+			FPCGGraphActiveTask& Task = ActiveTasks[ActiveTaskIndex];
+			if(Task.Context && CancelledComponents.Contains(Task.Context->SourceComponent.Get()))
+			{
+				FPCGTaskId CancelledTaskId = Task.NodeId;
+				Task.bWasCancelled = true;
+				bStableCancellationSet &= !CancelNextTasks(CancelledTaskId, CancelledComponents);
+			}
+		}
+
+		// Remove from sleeping tasks
+		for (int32 SleepingTaskIndex = SleepingTasks.Num() - 1; SleepingTaskIndex >= 0; --SleepingTaskIndex)
+		{
+			FPCGGraphActiveTask& Task = SleepingTasks[SleepingTaskIndex];
+			if(Task.Context && CancelledComponents.Contains(Task.Context->SourceComponent.Get()))
+			{
+				FPCGTaskId CancelledTaskId = Task.NodeId;
+				SleepingTasks.RemoveAtSwap(CancelledTaskId);
+				bStableCancellationSet &= !CancelNextTasks(CancelledTaskId, CancelledComponents);
+			}
 		}
 	}
 
@@ -245,6 +312,8 @@ void FPCGGraphExecutor::Cancel(TFunctionRef<bool(TWeakObjectPtr<UPCGComponent>)>
 #if WITH_EDITOR
 	UpdateGenerationNotification();
 #endif
+
+	return CancelledComponents;
 }
 
 bool FPCGGraphExecutor::IsGraphCurrentlyExecuting(UPCGGraph* InGraph)
@@ -330,6 +399,7 @@ void FPCGGraphExecutor::Execute()
 		{
 			FPCGTaskId TaskId = Task.NodeId;
 
+			// TODO: review if it's actually possible for a task with inputs to be ready at this point - it seems very unlikely
 			bool bPushToReady = true;
 			for (const FPCGGraphTaskInput& Input : Task.Inputs)
 			{
@@ -397,18 +467,18 @@ void FPCGGraphExecutor::Execute()
 		const bool bMainThreadWasAvailable = bMainThreadAvailable;
 		const int32 TasksAlreadyLaunchedCount = ActiveTasks.Num();
 
-		if (bMainThreadAvailable || NumAvailableThreads > 0)
+		bool bHasDispatchedTasks = !ActiveTasks.IsEmpty();
+
+		auto CannotDispatchMoreTasks = [&bMainThreadAvailable, &NumAvailableThreads, &bHasDispatchedTasks, bAllowMultiDispatch]
+		{
+			return ((!bAllowMultiDispatch && bHasDispatchedTasks) || (!bMainThreadAvailable && NumAvailableThreads == 0));
+		};
+
+		if(!CannotDispatchMoreTasks())
 		{
 			TRACE_CPUPROFILER_EVENT_SCOPE(FPCGGraphExecutor::Execute::PrepareTasks);
 			// Sort tasks by priority (highest priority should be at the end)
 			// TODO
-			bool bHasDispatchedTasks = false;
-
-			auto CannotDispatchMoreTasks = [&bMainThreadAvailable, &NumAvailableThreads, &bHasDispatchedTasks, bAllowMultiDispatch]
-			{
-				return ((!bAllowMultiDispatch && bHasDispatchedTasks) || (!bMainThreadAvailable && NumAvailableThreads == 0));
-			};
-
 			auto TaskDispatchBookKeeping = [&bMainThreadAvailable, &NumAvailableThreads, &bHasDispatchedTasks](bool bIsMainThreadTask)
 			{
 				bHasDispatchedTasks = true;
@@ -639,7 +709,9 @@ void FPCGGraphExecutor::Execute()
 			FPCGGraphActiveTask& ActiveTask = ActiveTasks[TaskIndex];
 
 #if WITH_EDITOR
-			if (!ActiveTask.bIsBypassed)
+			if (!ActiveTask.bWasCancelled && !ActiveTask.bIsBypassed)
+#else
+			if (!ActiveTask.bWasCancelled)
 #endif
 			{
 				if (bGraphCacheDebuggingEnabled && ActiveTask.Context->SourceComponent.Get() && ActiveTask.Context->Node)
@@ -704,9 +776,9 @@ void FPCGGraphExecutor::Execute()
 				MainThreadTask.Context->OutputData.RemoveFromRootSet(DataRootSet);
 
 #if WITH_EDITOR
-				if(MainThreadTask.bIsBypassed || MainThreadTask.Element->Execute(MainThreadTask.Context.Get()))
+				if(MainThreadTask.bIsBypassed || MainThreadTask.bWasCancelled || MainThreadTask.Element->Execute(MainThreadTask.Context.Get()))
 #else
-				if(MainThreadTask.Element->Execute(MainThreadTask.Context.Get()))
+				if(MainThreadTask.bWasCancelled || MainThreadTask.Element->Execute(MainThreadTask.Context.Get()))
 #endif
 				{
 					bMainTaskDone = true;
@@ -727,7 +799,7 @@ void FPCGGraphExecutor::Execute()
 						bTaskDone = Future->Get();
 					}
 
-					if (bTaskDone)
+					if (bTaskDone || ActiveTasks[ExecutionIndex].bWasCancelled)
 					{
 						PostTaskExecute(ExecutionIndex);
 					}
@@ -857,18 +929,32 @@ void FPCGGraphExecutor::QueueNextTasks(FPCGTaskId FinishedTask)
 	}
 }
 
-void FPCGGraphExecutor::CancelNextTasks(FPCGTaskId CancelledTask)
+bool FPCGGraphExecutor::CancelNextTasks(FPCGTaskId CancelledTask, TSet<UPCGComponent*>& OutCancelledComponents)
 {
+	bool bAddedComponents = false;
+
 	if (TSet<FPCGTaskId>* Successors = TaskSuccessors.Find(CancelledTask))
 	{
 		for (FPCGTaskId Successor : *Successors)
 		{
-			Tasks.Remove(Successor);
-			CancelNextTasks(Successor);
+			if (FPCGGraphTask* Task = Tasks.Find(Successor))
+			{
+				if(!OutCancelledComponents.Contains(Task->SourceComponent.Get()))
+				{
+					OutCancelledComponents.Add(Task->SourceComponent.Get());
+					bAddedComponents = true;
+				}
+
+				Tasks.Remove(Successor);
+			}
+
+			bAddedComponents |= CancelNextTasks(Successor, OutCancelledComponents);
 		}
 
 		TaskSuccessors.Remove(CancelledTask);
 	}
+
+	return bAddedComponents;
 }
 
 void FPCGGraphExecutor::BuildTaskInput(const FPCGGraphTask& Task, FPCGDataCollection& TaskInput)
diff --git a/Engine/Plugins/Experimental/PCG/Source/PCG/Private/Graph/PCGGraphExecutor.h b/Engine/Plugins/Experimental/PCG/Source/PCG/Private/Graph/PCGGraphExecutor.h
--- a/Engine/Plugins/Experimental/PCG/Source/PCG/Private/Graph/PCGGraphExecutor.h
+++ b/Engine/Plugins/Experimental/PCG/Source/PCG/Private/Graph/PCGGraphExecutor.h
@@ -51,6 +51,8 @@ struct FPCGGraphScheduleTask
 {
 	TArray<FPCGGraphTask> Tasks;
 	TWeakObjectPtr<UPCGComponent> SourceComponent = nullptr;
+	int32 FirstTaskIndex = 0;
+	int32 LastTaskIndex = 0;
 };
 
 struct FPCGGraphActiveTask
@@ -58,6 +60,7 @@ struct FPCGGraphActiveTask
 	FPCGElementPtr Element;
 	TUniquePtr<FPCGContext> Context;
 	FPCGTaskId NodeId = InvalidPCGTaskId;
+	bool bWasCancelled = false;
 #if WITH_EDITOR
 	bool bIsBypassed = false;
 #endif
@@ -77,7 +80,7 @@ public:
 	FPCGTaskId Schedule(UPCGGraph* Graph, UPCGComponent* InSourceComponent, FPCGElementPtr InputElement, const TArray<FPCGTaskId>& TaskDependency);
 
 	/** Cancels all tasks originating from the given component */
-	void Cancel(UPCGComponent* InComponent);
+	TArray<UPCGComponent*> Cancel(UPCGComponent* InComponent);
 	
 	/** Cancels all tasks running a given graph */
 	TArray<UPCGComponent*> Cancel(UPCGGraph* InGraph);
@@ -128,10 +131,10 @@ public:
 	bool IsGraphCacheDebuggingEnabled() const { return GraphCache.IsDebuggingEnabled(); }
 
 private:
-	void Cancel(TFunctionRef<bool(TWeakObjectPtr<UPCGComponent>)> CancelFilter);
+	TSet<UPCGComponent*> Cancel(TFunctionRef<bool(TWeakObjectPtr<UPCGComponent>)> CancelFilter);
 	void ClearAllTasks();
 	void QueueNextTasks(FPCGTaskId FinishedTask);
-	void CancelNextTasks(FPCGTaskId CancelledTask);
+	bool CancelNextTasks(FPCGTaskId CancelledTask, TSet<UPCGComponent*>& OutCancelledComponents);
 	void BuildTaskInput(const FPCGGraphTask& Task, FPCGDataCollection& TaskInput);
 	/** Combine all param data into one on the Params pin, if any.*/
 	void CombineParams(FPCGTaskId InTaskId, FPCGDataCollection& InTaskInput);
diff --git a/Engine/Plugins/Experimental/PCG/Source/PCG/Private/Grid/PCGPartitionActor.cpp b/Engine/Plugins/Experimental/PCG/Source/PCG/Private/Grid/PCGPartitionActor.cpp
--- a/Engine/Plugins/Experimental/PCG/Source/PCG/Private/Grid/PCGPartitionActor.cpp
+++ b/Engine/Plugins/Experimental/PCG/Source/PCG/Private/Grid/PCGPartitionActor.cpp
@@ -114,7 +114,7 @@ void APCGPartitionActor::PostRegisterAllComponents()
 	OriginalToLocal.Reserve(LocalToOriginal.Num());
 	for (const auto& It : LocalToOriginal)
 	{
-		OriginalToLocal.Add(It.Value, It.Key);
+		OriginalToLocal.Add(It.Value.Get(), It.Key);
 	}
 
 	// Make the Partition actor register itself to the PCG Subsystem
@@ -230,7 +230,7 @@ UPCGComponent* APCGPartitionActor::GetOriginalComponent(const UPCGComponent* Loc
 void APCGPartitionActor::CleanupDeadGraphInstances(bool bRemoveNonNullOnly)
 {
 	// First find if we have any local dead instance (= nullptr) hooked to an original component.
-	TSet<TSoftObjectPtr<UPCGComponent>> DeadOriginalInstances;
+	TSet<TObjectPtr<UPCGComponent>> DeadOriginalInstances;
 	for (const auto& OriginalToLocalItem : OriginalToLocal)
 	{
 		if (!OriginalToLocalItem.Value)
@@ -243,7 +243,7 @@ void APCGPartitionActor::CleanupDeadGraphInstances(bool bRemoveNonNullOnly)
 	{
 		Modify();
 
-		for (const TSoftObjectPtr<UPCGComponent>& DeadInstance : DeadOriginalInstances)
+		for (const TObjectPtr<UPCGComponent>& DeadInstance : DeadOriginalInstances)
 		{
 			OriginalToLocal.Remove(DeadInstance);
 		}
@@ -333,6 +333,11 @@ void APCGPartitionActor::RemapGraphInstance(const UPCGComponent* OldOriginalComp
 	LocalComponent->SetPropertiesFromOriginal(NewOriginalComponent);
 	OriginalToLocal.Emplace(NewOriginalComponent, LocalComponent);
 	LocalToOriginal.Emplace(LocalComponent, NewOriginalComponent);
+
+#if WITH_EDITOR
+	// When changing original data, it means that the data we have might point to newly stale data, hence we need to force dirty here
+	LocalComponent->DirtyGenerated(EPCGComponentDirtyFlag::Actor);
+#endif
 }
 
 bool APCGPartitionActor::RemoveGraphInstance(UPCGComponent* OriginalComponent)
@@ -433,7 +438,7 @@ TSet<TObjectPtr<UPCGComponent>> APCGPartitionActor::GetAllLocalPCGComponents() c
 TSet<TObjectPtr<UPCGComponent>> APCGPartitionActor::GetAllOriginalPCGComponents() const
 {
 	TSet<TObjectPtr<UPCGComponent>> ResultComponents;
-	for(const TPair<TSoftObjectPtr<UPCGComponent>, TObjectPtr<UPCGComponent>>& OriginalPair : OriginalToLocal)
+	for(const TPair<TObjectPtr<UPCGComponent>, TObjectPtr<UPCGComponent>>& OriginalPair : OriginalToLocal)
 	{
 		ResultComponents.Add(OriginalPair.Key.Get());
 	}
diff --git a/Engine/Plugins/Experimental/PCG/Source/PCG/Private/Grid/PCGPartitionActor.h b/Engine/Plugins/Experimental/PCG/Source/PCG/Private/Grid/PCGPartitionActor.h
--- a/Engine/Plugins/Experimental/PCG/Source/PCG/Private/Grid/PCGPartitionActor.h
+++ b/Engine/Plugins/Experimental/PCG/Source/PCG/Private/Grid/PCGPartitionActor.h
@@ -88,7 +88,7 @@ private:
 	UPCGSubsystem* GetSubsystem() const;
 
 	// Note: this map is not a property and not serialized since we will rebuild it from the LocalToOriginal
-	TMap<TSoftObjectPtr<UPCGComponent>, TObjectPtr<UPCGComponent>> OriginalToLocal;
+	TMap<TObjectPtr<UPCGComponent>, TObjectPtr<UPCGComponent>> OriginalToLocal;
 
 	UPROPERTY()
 	TMap<TObjectPtr<UPCGComponent>, TSoftObjectPtr<UPCGComponent>> LocalToOriginal;
diff --git a/Engine/Plugins/Experimental/PCG/Source/PCG/Private/PCGComponent.cpp b/Engine/Plugins/Experimental/PCG/Source/PCG/Private/PCGComponent.cpp
--- a/Engine/Plugins/Experimental/PCG/Source/PCG/Private/PCGComponent.cpp
+++ b/Engine/Plugins/Experimental/PCG/Source/PCG/Private/PCGComponent.cpp
@@ -518,6 +518,7 @@ void UPCGComponent::OnProcessGraphAborted(bool bQuiet)
 	CurrentCleanupTask = InvalidPCGTaskId; // this is needed to support cancellation
 
 #if WITH_EDITOR
+	CurrentRefreshTask = InvalidPCGTaskId;
 	bDirtyGenerated = false;
 #endif
 }
@@ -861,6 +862,14 @@ void UPCGComponent::BeginPlay()
 {
 	Super::BeginPlay();
 
+#if WITH_EDITOR
+	// Disable auto-refreshing on preview actors until we have something more robust on the execution side.
+	if (GetOwner() && GetOwner()->bIsEditorPreviewActor)
+	{
+		return;
+	}
+#endif
+
 	// First if it is partitioned, register itself to the PCGSubsystem, to map the component to all its corresponding PartitionActors
 	if (IsPartitioned() && GetSubsystem())
 	{
@@ -1073,6 +1082,12 @@ void UPCGComponent::OnRegister()
 	Super::OnRegister();
 
 #if WITH_EDITOR
+	// Disable auto-refreshing on preview actors until we have something more robust on the execution side.
+	if (GetOwner() && GetOwner()->bIsEditorPreviewActor)
+	{
+		return;
+	}
+
 	// We can't register to the subsystem in OnRegister if we are at runtime because
 	// the landscape can be not loaded yet.
 	// It will be done in BeginPlay at runtime.
@@ -1391,6 +1406,11 @@ bool UPCGComponent::ActorIsTracked(AActor* InActor) const
 
 void UPCGComponent::OnActorAdded(AActor* InActor)
 {
+	if (!InActor || InActor->bIsEditorPreviewActor)
+	{
+		return;
+	}
+
 	const bool bIsExcluded = UpdateExcludedActor(InActor);
 	const bool bIsTracked = AddTrackedActor(InActor);
 
@@ -1403,6 +1423,11 @@ void UPCGComponent::OnActorAdded(AActor* InActor)
 
 void UPCGComponent::OnActorDeleted(AActor* InActor)
 {
+	if (!InActor || InActor->bIsEditorPreviewActor)
+	{
+		return;
+	}
+
 	const bool bWasExcluded = UpdateExcludedActor(InActor);
 	const bool bWasTracked = RemoveTrackedActor(InActor);
 
@@ -1415,6 +1440,11 @@ void UPCGComponent::OnActorDeleted(AActor* InActor)
 
 void UPCGComponent::OnActorMoved(AActor* InActor)
 {
+	if (!InActor || InActor->bIsEditorPreviewActor)
+	{
+		return;
+	}
+
 	const bool bOwnerMoved = (InActor == GetOwner());
 	const bool bLandscapeMoved = (InActor && TrackedLandscapes.Contains(InActor));
 
@@ -1680,7 +1710,7 @@ void UPCGComponent::DirtyGenerated(EPCGComponentDirtyFlag DirtyFlag, const bool
 
 	// For partitioned graph, we must forward the call to the partition actor, if we need to
 	// TODO: Don't forward for None for now, as it could break some stuff
-	if (DirtyFlag != EPCGComponentDirtyFlag::None && bActivated && IsPartitioned() && bDispatchToLocalComponents)
+	if (bActivated && IsPartitioned() && bDispatchToLocalComponents)
 	{
 		if (GetSubsystem())
 		{
diff --git a/Engine/Plugins/Experimental/PCG/Source/PCG/Private/PCGSubsystem.cpp b/Engine/Plugins/Experimental/PCG/Source/PCG/Private/PCGSubsystem.cpp
--- a/Engine/Plugins/Experimental/PCG/Source/PCG/Private/PCGSubsystem.cpp
+++ b/Engine/Plugins/Experimental/PCG/Source/PCG/Private/PCGSubsystem.cpp
@@ -427,8 +427,11 @@ void UPCGSubsystem::CancelGeneration(UPCGComponent* Component)
 		DispatchToRegisteredLocalComponents(Component, LocalCancel);
 	}
 
-	GraphExecutor->Cancel(Component);
-	Component->OnProcessGraphAborted(/*bQuiet=*/true);
+	TArray<UPCGComponent*> CancelledComponents = GraphExecutor->Cancel(Component);
+	for (UPCGComponent* CancelledComponent : CancelledComponents)
+	{
+		CancelledComponent->OnProcessGraphAborted(/*bQuiet=*/true);
+	}	
 }
 
 void UPCGSubsystem::CancelGeneration(UPCGGraph* Graph)
@@ -441,9 +444,9 @@ void UPCGSubsystem::CancelGeneration(UPCGGraph* Graph)
 	}
 
 	TArray<UPCGComponent*> CancelledComponents = GraphExecutor->Cancel(Graph);
-	for (UPCGComponent* Component : CancelledComponents)
+	for (UPCGComponent* CancelledComponent : CancelledComponents)
 	{
-		Component->OnProcessGraphAborted(/*bQuiet=*/true);
+		CancelledComponent->OnProcessGraphAborted(/*bQuiet=*/true);
 	}
 }
 
@@ -1143,15 +1146,18 @@ namespace PCGSubsystem
 					SetPreviousTaskIfValid(SaveTaskId);
 				}
 
-				// Unload actors from cell (or the pcg actor refered here)
-				auto UnloadActorsTask = [GraphExecutor, ActorReferences]() {
-					GraphExecutor->AddToUnusedActors(*ActorReferences);
-					return true;
-				};
+				if (bLoadCell)
+				{
+					// Unload actors from cell (or the pcg actor refered here)
+					auto UnloadActorsTask = [GraphExecutor, ActorReferences]() {
+						GraphExecutor->AddToUnusedActors(*ActorReferences);
+						return true;
+					};
 
-				// Schedule after the save (if valid), then the execute so we can queue this after the load.
-				FPCGTaskId UnloadTaskId = GraphExecutor->ScheduleGeneric(UnloadActorsTask, nullptr, PreviousTasks);
-				SetPreviousTaskIfValid(UnloadTaskId);
+					// Schedule after the save (if valid), then the execute so we can queue this after the load.
+					FPCGTaskId UnloadTaskId = GraphExecutor->ScheduleGeneric(UnloadActorsTask, nullptr, PreviousTasks);
+					SetPreviousTaskIfValid(UnloadTaskId);
+				}
 
 				// Finally, mark "last" valid task in the cell tasks.
 				CellTasks.Append(PreviousTasks);
diff --git a/Engine/Plugins/Experimental/PCG/Source/PCGEditor/Private/PCGVolumeFactory.cpp b/Engine/Plugins/Experimental/PCG/Source/PCGEditor/Private/PCGVolumeFactory.cpp
--- a/Engine/Plugins/Experimental/PCG/Source/PCGEditor/Private/PCGVolumeFactory.cpp
+++ b/Engine/Plugins/Experimental/PCG/Source/PCGEditor/Private/PCGVolumeFactory.cpp
@@ -37,6 +37,12 @@ void UPCGVolumeFactory::PostSpawnActor(UObject* Asset, AActor* NewActor)
 {
 	Super::PostSpawnActor(Asset, NewActor);
 
+	// Disable auto-refreshing on preview actors until we have something more robust on the execution side.
+	if (NewActor && NewActor->bIsEditorPreviewActor)
+	{
+		return;
+	}
+
 	if (UPCGGraph* PCGGraph = Cast<UPCGGraph>(Asset))
 	{
 		const UPCGEditorSettings* PCGEditorSettings = GetDefault<UPCGEditorSettings>();
