Mirrored mobility to Debug ISM based on target actor
#rnx

diff --git a/Engine/Plugins/Experimental/PCG/Source/PCG/Private/Elements/PCGDebugElement.cpp b/Engine/Plugins/Experimental/PCG/Source/PCG/Private/Elements/PCGDebugElement.cpp
--- a/Engine/Plugins/Experimental/PCG/Source/PCG/Private/Elements/PCGDebugElement.cpp
+++ b/Engine/Plugins/Experimental/PCG/Source/PCG/Private/Elements/PCGDebugElement.cpp
@@ -157,6 +157,12 @@ namespace PCGDebugElement
 			// Note: In the future we may consider enabling culling for performance reasons, but for now culling disabled.
 			Params.CullStartDistance = Params.CullEndDistance = 0;
 
+			// If the root actor we're binding to is movable, then the ISMC should be movable by default
+			if (USceneComponent* SceneComponent = TargetActor->GetRootComponent())
+			{
+				Params.Mobility = SceneComponent->Mobility;
+			}
+
 			for (int32 Direction = 0; Direction < 2; ++Direction)
 			{
 				TArray<FTransform>& Instances = (Direction == 0 ? ForwardInstances : ReverseInstances);
