[PCG] Add "Debug only selected nodes" and "Disable debug on all nodes" options in the PCG Editor. Also remapping Determinism shortcut to Alt + T .

diff --git a/Engine/Plugins/Experimental/PCG/Source/PCGEditor/Private/PCGEditor.cpp b/Engine/Plugins/Experimental/PCG/Source/PCGEditor/Private/PCGEditor.cpp
--- a/Engine/Plugins/Experimental/PCG/Source/PCGEditor/Private/PCGEditor.cpp
+++ b/Engine/Plugins/Experimental/PCG/Source/PCGEditor/Private/PCGEditor.cpp
@@ -436,6 +436,14 @@ void FPCGEditor::BindCommands()
 		FExecuteAction::CreateSP(this, &FPCGEditor::OnToggleDebug),
 		FCanExecuteAction(),
 		FGetActionCheckState::CreateSP(this, &FPCGEditor::GetDebugCheckState));
+
+	GraphEditorCommands->MapAction(
+		PCGEditorCommands.DebugOnlySelected,
+		FExecuteAction::CreateSP(this, &FPCGEditor::OnDebugOnlySelected));
+
+	GraphEditorCommands->MapAction(
+		PCGEditorCommands.DisableDebugOnAllNodes,
+		FExecuteAction::CreateSP(this, &FPCGEditor::OnDisableDebugOnAllNodes));
 }
 
 void FPCGEditor::OnFind()
@@ -1407,6 +1415,79 @@ void FPCGEditor::OnToggleDebug()
 	}
 }
 
+void FPCGEditor::OnDebugOnlySelected()
+{
+	if (GraphEditorWidget.IsValid() && PCGEditorGraph)
+	{
+		bool bChanged = false;
+
+		const FGraphPanelSelectionSet& SelectedNodes = GraphEditorWidget->GetSelectedNodes();
+
+		FScopedTransaction Transaction(*FPCGEditorCommon::ContextIdentifier, LOCTEXT("PCGEditorDebugOnlySelectedTransactionMessage", "PCG Editor: Debug only selected nodes"), nullptr);
+
+		for (UEdGraphNode* Node : PCGEditorGraph->Nodes)
+		{
+			UPCGEditorGraphNodeBase* PCGEditorGraphNode = Cast<UPCGEditorGraphNodeBase>(Node);
+			UPCGNode* PCGNode = PCGEditorGraphNode ? PCGEditorGraphNode->GetPCGNode() : nullptr;
+			UPCGSettingsInterface* PCGSettingsInterface = PCGNode ? PCGNode->GetSettingsInterface() : nullptr;
+
+			if (!PCGEditorGraphNode || !PCGNode || !PCGSettingsInterface)
+			{
+				continue;
+			}
+
+			const bool bShouldBeDebug = SelectedNodes.Contains(PCGEditorGraphNode);
+
+			if (PCGSettingsInterface->bDebug != bShouldBeDebug)
+			{
+				PCGSettingsInterface->Modify();
+				PCGSettingsInterface->bDebug = bShouldBeDebug;
+				PCGNode->OnNodeChangedDelegate.Broadcast(PCGNode, EPCGChangeType::Settings);
+				bChanged = true;
+			}
+		}
+
+		if (!bChanged)
+		{
+			Transaction.Cancel();
+		}
+	}
+}
+
+void FPCGEditor::OnDisableDebugOnAllNodes()
+{
+	if (GraphEditorWidget.IsValid() && PCGEditorGraph)
+	{
+		bool bChanged = false;
+		FScopedTransaction Transaction(*FPCGEditorCommon::ContextIdentifier, LOCTEXT("PCGEditorDisableDebugAllNodesTransactionMessage", "PCG Editor: Disable debug on all nodes"), nullptr);
+
+		for (UEdGraphNode* Node : PCGEditorGraph->Nodes)
+		{
+			UPCGEditorGraphNodeBase* PCGEditorGraphNode = Cast<UPCGEditorGraphNodeBase>(Node);
+			UPCGNode* PCGNode = PCGEditorGraphNode ? PCGEditorGraphNode->GetPCGNode() : nullptr;
+			UPCGSettingsInterface* PCGSettingsInterface = PCGNode ? PCGNode->GetSettingsInterface() : nullptr;
+
+			if (!PCGEditorGraphNode || !PCGNode || !PCGSettingsInterface)
+			{
+				continue;
+			}
+
+			if (PCGSettingsInterface->bDebug)
+			{
+				PCGSettingsInterface->Modify();
+				PCGSettingsInterface->bDebug = false;
+				PCGNode->OnNodeChangedDelegate.Broadcast(PCGNode, EPCGChangeType::Settings);
+				bChanged = true;
+			}
+		}
+
+		if (!bChanged)
+		{
+			Transaction.Cancel();
+		}
+	}
+}
+
 ECheckBoxState FPCGEditor::GetDebugCheckState() const
 {
 	if (GraphEditorWidget.IsValid())
diff --git a/Engine/Plugins/Experimental/PCG/Source/PCGEditor/Private/PCGEditor.h b/Engine/Plugins/Experimental/PCG/Source/PCGEditor/Private/PCGEditor.h
--- a/Engine/Plugins/Experimental/PCG/Source/PCGEditor/Private/PCGEditor.h
+++ b/Engine/Plugins/Experimental/PCG/Source/PCGEditor/Private/PCGEditor.h
@@ -124,6 +124,12 @@ private:
 	/** Whether selected nodes are being debugged or not */
 	ECheckBoxState GetDebugCheckState() const;
 
+	/** Enable node debug state for selected nodes and disable for others */
+	void OnDebugOnlySelected();
+
+	/** Disable node debug state for all nodes */
+	void OnDisableDebugOnAllNodes();
+
 	/** Cancels the current execution of the selected graph */
 	void OnCancelExecution_Clicked();
 
diff --git a/Engine/Plugins/Experimental/PCG/Source/PCGEditor/Private/PCGEditorCommands.cpp b/Engine/Plugins/Experimental/PCG/Source/PCGEditor/Private/PCGEditorCommands.cpp
--- a/Engine/Plugins/Experimental/PCG/Source/PCGEditor/Private/PCGEditorCommands.cpp
+++ b/Engine/Plugins/Experimental/PCG/Source/PCGEditor/Private/PCGEditorCommands.cpp
@@ -23,12 +23,14 @@ void FPCGEditorCommands::RegisterCommands()
 	UI_COMMAND(Find, "Find", "Finds PCG nodes and comments in the current graph", EUserInterfaceActionType::Button, FInputChord(EModifierKey::Control, EKeys::F));
 	UI_COMMAND(PauseAutoRegeneration, "Pause Regen", "Pause automatic regeneration of the current graph.", EUserInterfaceActionType::ToggleButton, FInputChord(EModifierKey::Alt, EKeys::R));
 	UI_COMMAND(ForceGraphRegeneration, "Force Regen", "Manually force a regeneration of the current graph.\nCtrl-click will also perform a flush cache before the regeneration.", EUserInterfaceActionType::Button, FInputChord());
-	UI_COMMAND(RunDeterminismNodeTest, "Run Determinism Test on Node", "Evaluate the current node for determinism.", EUserInterfaceActionType::Button, FInputChord(EModifierKey::Alt, EKeys::D));
+	UI_COMMAND(RunDeterminismNodeTest, "Run Determinism Test on Node", "Evaluate the current node for determinism.", EUserInterfaceActionType::Button, FInputChord(EModifierKey::Alt, EKeys::T));
 	UI_COMMAND(RunDeterminismGraphTest, "Graph Determinism Test", "Evaluate the current graph for determinism.", EUserInterfaceActionType::Button, FInputChord());
 	UI_COMMAND(EditGraphSettings, "Graph Settings", "Edit the graph settings.", EUserInterfaceActionType::ToggleButton, FInputChord());
 	UI_COMMAND(CancelExecution, "Cancel Execution", "Cancels the execution of the current graph", EUserInterfaceActionType::Button, FInputChord(EModifierKey::Control, EKeys::Escape));
 	UI_COMMAND(ToggleEnabled, "Toggle Enabled", "Toggle node enabled state for selected nodes.", EUserInterfaceActionType::ToggleButton, FInputChord(EKeys::E));
 	UI_COMMAND(ToggleDebug, "Toggle Debug", "Toggle node debug state for selected nodes", EUserInterfaceActionType::ToggleButton, FInputChord(EKeys::D));
+	UI_COMMAND(DebugOnlySelected, "Debug Only Selected", "Enable node debug state for selected nodes and disable debug state for the others", EUserInterfaceActionType::Button, FInputChord(EModifierKey::Control | EModifierKey::Alt, EKeys::D));
+	UI_COMMAND(DisableDebugOnAllNodes, "Disable Debug on all nodes", "Disable debug state for all nodes", EUserInterfaceActionType::Button, FInputChord(EModifierKey::Alt, EKeys::D));
 	UI_COMMAND(ToggleInspect, "Toggle Inspection", "Toggle node inspection for selected node", EUserInterfaceActionType::ToggleButton, FInputChord(EKeys::I));
 }
 
diff --git a/Engine/Plugins/Experimental/PCG/Source/PCGEditor/Private/PCGEditorCommands.h b/Engine/Plugins/Experimental/PCG/Source/PCGEditor/Private/PCGEditorCommands.h
--- a/Engine/Plugins/Experimental/PCG/Source/PCGEditor/Private/PCGEditorCommands.h
+++ b/Engine/Plugins/Experimental/PCG/Source/PCGEditor/Private/PCGEditorCommands.h
@@ -25,5 +25,7 @@ public:
 	TSharedPtr<FUICommandInfo> CancelExecution;
 	TSharedPtr<FUICommandInfo> ToggleEnabled;
 	TSharedPtr<FUICommandInfo> ToggleDebug;
+	TSharedPtr<FUICommandInfo> DebugOnlySelected;
+	TSharedPtr<FUICommandInfo> DisableDebugOnAllNodes;
 	TSharedPtr<FUICommandInfo> ToggleInspect;
 };
diff --git a/Engine/Plugins/Experimental/PCG/Source/PCGEditor/Private/PCGEditorGraphNodeBase.cpp b/Engine/Plugins/Experimental/PCG/Source/PCGEditor/Private/PCGEditorGraphNodeBase.cpp
--- a/Engine/Plugins/Experimental/PCG/Source/PCGEditor/Private/PCGEditorGraphNodeBase.cpp
+++ b/Engine/Plugins/Experimental/PCG/Source/PCGEditor/Private/PCGEditorGraphNodeBase.cpp
@@ -77,6 +77,8 @@ void UPCGEditorGraphNodeBase::GetNodeContextMenuActions(UToolMenu* Menu, class U
 		FToolMenuSection& Section = Menu->AddSection("EdGraphSchemaNodeActions", LOCTEXT("NodeActionsHeader", "Node Actions"));
 		Section.AddMenuEntry(FPCGEditorCommands::Get().ToggleEnabled, LOCTEXT("ToggleEnabledLabel", "Enable"));
 		Section.AddMenuEntry(FPCGEditorCommands::Get().ToggleDebug, LOCTEXT("ToggleDebugLabel", "Debug"));
+		Section.AddMenuEntry(FPCGEditorCommands::Get().DebugOnlySelected);
+		Section.AddMenuEntry(FPCGEditorCommands::Get().DisableDebugOnAllNodes);
 		Section.AddMenuEntry(FPCGEditorCommands::Get().ToggleInspect, LOCTEXT("ToggleinspectionLabel", "Inspect"));
 		Section.AddMenuEntry(FGraphEditorCommands::Get().BreakNodeLinks);
 		Section.AddMenuEntry(FPCGEditorCommands::Get().CollapseNodes);
