[PCG] Stop passing through settings through nodes.

Settings should be explicitly connected rather than chained together. Dependencies should be clear to user.

diff --git a/Engine/Plugins/Experimental/PCG/Source/PCG/Private/Elements/Metadata/PCGMetadataElement.cpp b/Engine/Plugins/Experimental/PCG/Source/PCG/Private/Elements/Metadata/PCGMetadataElement.cpp
--- a/Engine/Plugins/Experimental/PCG/Source/PCG/Private/Elements/Metadata/PCGMetadataElement.cpp
+++ b/Engine/Plugins/Experimental/PCG/Source/PCG/Private/Elements/Metadata/PCGMetadataElement.cpp
@@ -74,9 +74,6 @@ bool FPCGMetadataOperationElement::ExecuteInternal(FPCGContext* Context) const
 
 	TArray<FPCGTaggedData>& Outputs = Context->OutputData.TaggedData;
 
-	// Forward any non-input data
-	Outputs.Append(Context->InputData.GetAllSettings());
-
 	for (const FPCGTaggedData& Input : Inputs)
 	{
 		FPCGTaggedData& Output = Outputs.Add_GetRef(Input);
diff --git a/Engine/Plugins/Experimental/PCG/Source/PCG/Private/Elements/Metadata/PCGMetadataPartition.cpp b/Engine/Plugins/Experimental/PCG/Source/PCG/Private/Elements/Metadata/PCGMetadataPartition.cpp
--- a/Engine/Plugins/Experimental/PCG/Source/PCG/Private/Elements/Metadata/PCGMetadataPartition.cpp
+++ b/Engine/Plugins/Experimental/PCG/Source/PCG/Private/Elements/Metadata/PCGMetadataPartition.cpp
@@ -25,9 +25,6 @@ bool FPCGMetadataPartitionElement::ExecuteInternal(FPCGContext* Context) const
 	TArray<FPCGTaggedData> Inputs = Context->InputData.GetInputs();
 	TArray<FPCGTaggedData>& Outputs = Context->OutputData.TaggedData;
 
-	// Forward any non-input data
-	Outputs.Append(Context->InputData.GetAllSettings());
-
 	const FName PartitionAttribute = Settings->PartitionAttribute;
 
 	for (const FPCGTaggedData& Input : Inputs)
diff --git a/Engine/Plugins/Experimental/PCG/Source/PCG/Private/Elements/Metadata/PCGMetadataRenameElement.cpp b/Engine/Plugins/Experimental/PCG/Source/PCG/Private/Elements/Metadata/PCGMetadataRenameElement.cpp
--- a/Engine/Plugins/Experimental/PCG/Source/PCG/Private/Elements/Metadata/PCGMetadataRenameElement.cpp
+++ b/Engine/Plugins/Experimental/PCG/Source/PCG/Private/Elements/Metadata/PCGMetadataRenameElement.cpp
@@ -85,8 +85,5 @@ bool FPCGMetadataRenameElement::ExecuteInternal(FPCGContext* Context) const
 		}
 	}
 
-	// Pass-through settings
-	Outputs.Append(Context->InputData.GetAllSettings());
-
 	return true;
 }
diff --git a/Engine/Plugins/Experimental/PCG/Source/PCG/Private/Elements/PCGBoundsModifier.cpp b/Engine/Plugins/Experimental/PCG/Source/PCG/Private/Elements/PCGBoundsModifier.cpp
--- a/Engine/Plugins/Experimental/PCG/Source/PCG/Private/Elements/PCGBoundsModifier.cpp
+++ b/Engine/Plugins/Experimental/PCG/Source/PCG/Private/Elements/PCGBoundsModifier.cpp
@@ -121,9 +121,7 @@ bool FPCGBoundsModifier::ExecuteInternal(FPCGContext* Context) const
 		break;
 	}
 
-	Outputs.Append(Context->InputData.GetAllSettings());
-	
 	return true;
 }
 
-#undef LOCTEXT_NAMESPACE
\ No newline at end of file
+#undef LOCTEXT_NAMESPACE
diff --git a/Engine/Plugins/Experimental/PCG/Source/PCG/Private/Elements/PCGCopyPoints.cpp b/Engine/Plugins/Experimental/PCG/Source/PCG/Private/Elements/PCGCopyPoints.cpp
--- a/Engine/Plugins/Experimental/PCG/Source/PCG/Private/Elements/PCGCopyPoints.cpp
+++ b/Engine/Plugins/Experimental/PCG/Source/PCG/Private/Elements/PCGCopyPoints.cpp
@@ -215,9 +215,5 @@ bool FPCGCopyPointsElement::ExecuteInternal(FPCGContext* Context) const
 
 	OutPointData->Metadata->AddDelayedEntries(AllMetadataEntries);
 
-	// Forward any non-input data
-	Outputs.Append(Context->InputData.GetAllSettings());
-
 	return true;
 }
-
diff --git a/Engine/Plugins/Experimental/PCG/Source/PCG/Private/Elements/PCGDensityFilter.cpp b/Engine/Plugins/Experimental/PCG/Source/PCG/Private/Elements/PCGDensityFilter.cpp
--- a/Engine/Plugins/Experimental/PCG/Source/PCG/Private/Elements/PCGDensityFilter.cpp
+++ b/Engine/Plugins/Experimental/PCG/Source/PCG/Private/Elements/PCGDensityFilter.cpp
@@ -25,9 +25,6 @@ bool FPCGDensityFilterElement::ExecuteInternal(FPCGContext* Context) const
 	TArray<FPCGTaggedData> Inputs = Context->InputData.GetInputs();
 	TArray<FPCGTaggedData>& Outputs = Context->OutputData.TaggedData;
 
-	// Forward any non-input data
-	Outputs.Append(Context->InputData.GetAllSettings());
-
 	const bool bInvertFilter = Settings->bInvertFilter;
 	const float LowerBound = Settings->LowerBound;
 	const float UpperBound = Settings->UpperBound;
diff --git a/Engine/Plugins/Experimental/PCG/Source/PCG/Private/Elements/PCGDensityNoise.cpp b/Engine/Plugins/Experimental/PCG/Source/PCG/Private/Elements/PCGDensityNoise.cpp
--- a/Engine/Plugins/Experimental/PCG/Source/PCG/Private/Elements/PCGDensityNoise.cpp
+++ b/Engine/Plugins/Experimental/PCG/Source/PCG/Private/Elements/PCGDensityNoise.cpp
@@ -86,8 +86,5 @@ bool FPCGDensityNoiseElement::ExecuteInternal(FPCGContext* Context) const
 		return true;
 	});
 
-	// Forward any non-input data
-	Outputs.Append(Context->InputData.GetAllSettings());
-
 	return true;
-}
\ No newline at end of file
+}
diff --git a/Engine/Plugins/Experimental/PCG/Source/PCG/Private/Elements/PCGDensityRemap.cpp b/Engine/Plugins/Experimental/PCG/Source/PCG/Private/Elements/PCGDensityRemap.cpp
--- a/Engine/Plugins/Experimental/PCG/Source/PCG/Private/Elements/PCGDensityRemap.cpp
+++ b/Engine/Plugins/Experimental/PCG/Source/PCG/Private/Elements/PCGDensityRemap.cpp
@@ -33,9 +33,6 @@ bool FPCGLinearDensityRemapElement::ExecuteInternal(FPCGContext* Context) const
 	TArray<FPCGTaggedData> Inputs = Context->InputData.GetInputs();
 	TArray<FPCGTaggedData>& Outputs = Context->OutputData.TaggedData;
 
-	// Forward any non-input data
-	Outputs.Append(Context->InputData.GetAllSettings());
-
 	const float SettingsRemapMin = Settings->RemapMin;
 	const float SettingsRemapMax = Settings->RemapMax;
 	const bool bMultiplyDensity = Settings->bMultiplyDensity;
diff --git a/Engine/Plugins/Experimental/PCG/Source/PCG/Private/Elements/PCGDensityRemapElement.cpp b/Engine/Plugins/Experimental/PCG/Source/PCG/Private/Elements/PCGDensityRemapElement.cpp
--- a/Engine/Plugins/Experimental/PCG/Source/PCG/Private/Elements/PCGDensityRemapElement.cpp
+++ b/Engine/Plugins/Experimental/PCG/Source/PCG/Private/Elements/PCGDensityRemapElement.cpp
@@ -75,8 +75,5 @@ bool FPCGDensityRemapElement::ExecuteInternal(FPCGContext* Context) const
 		return true;
 	});
 
-	// Forward any non-input data
-	Outputs.Append(Context->InputData.GetAllSettings());
-
 	return true;
 }
diff --git a/Engine/Plugins/Experimental/PCG/Source/PCG/Private/Elements/PCGDifferenceElement.cpp b/Engine/Plugins/Experimental/PCG/Source/PCG/Private/Elements/PCGDifferenceElement.cpp
--- a/Engine/Plugins/Experimental/PCG/Source/PCG/Private/Elements/PCGDifferenceElement.cpp
+++ b/Engine/Plugins/Experimental/PCG/Source/PCG/Private/Elements/PCGDifferenceElement.cpp
@@ -205,7 +205,4 @@ void FPCGDifferenceElement::LabellessProcessing(FPCGContext* Context) const
 
 		AddToDifference(SpatialData);
 	}
-
-	// Finally, pass-through settings
-	Outputs.Append(Context->InputData.GetAllSettings());
 }
diff --git a/Engine/Plugins/Experimental/PCG/Source/PCG/Private/Elements/PCGDistance.cpp b/Engine/Plugins/Experimental/PCG/Source/PCG/Private/Elements/PCGDistance.cpp
--- a/Engine/Plugins/Experimental/PCG/Source/PCG/Private/Elements/PCGDistance.cpp
+++ b/Engine/Plugins/Experimental/PCG/Source/PCG/Private/Elements/PCGDistance.cpp
@@ -189,8 +189,5 @@ bool FPCGDistanceElement::ExecuteInternal(FPCGContext* Context) const
 		);		
 	}
 
-	// Forward any non-input data
-	Outputs.Append(Context->InputData.GetAllSettings());
 	return true;
-
 }
diff --git a/Engine/Plugins/Experimental/PCG/Source/PCG/Private/Elements/PCGIntersectionElement.cpp b/Engine/Plugins/Experimental/PCG/Source/PCG/Private/Elements/PCGIntersectionElement.cpp
--- a/Engine/Plugins/Experimental/PCG/Source/PCG/Private/Elements/PCGIntersectionElement.cpp
+++ b/Engine/Plugins/Experimental/PCG/Source/PCG/Private/Elements/PCGIntersectionElement.cpp
@@ -83,8 +83,5 @@ bool FPCGIntersectionElement::ExecuteInternal(FPCGContext* Context) const
 		IntersectionTaggedData.Tags.Append(Input.Tags);
 	}
 
-	// Pass-through settings & exclusions
-	Outputs.Append(Context->InputData.GetAllSettings());
-
 	return true;
 }
diff --git a/Engine/Plugins/Experimental/PCG/Source/PCG/Private/Elements/PCGPointExtentsModifier.cpp b/Engine/Plugins/Experimental/PCG/Source/PCG/Private/Elements/PCGPointExtentsModifier.cpp
--- a/Engine/Plugins/Experimental/PCG/Source/PCG/Private/Elements/PCGPointExtentsModifier.cpp
+++ b/Engine/Plugins/Experimental/PCG/Source/PCG/Private/Elements/PCGPointExtentsModifier.cpp
@@ -74,8 +74,6 @@ bool FPCGPointExtentsModifier::ExecuteInternal(FPCGContext* Context) const
 		});
 		break;
 	}
-
-	Outputs.Append(Context->InputData.GetAllSettings());
 	
 	return true;
 }
diff --git a/Engine/Plugins/Experimental/PCG/Source/PCG/Private/Elements/PCGPointFilter.cpp b/Engine/Plugins/Experimental/PCG/Source/PCG/Private/Elements/PCGPointFilter.cpp
--- a/Engine/Plugins/Experimental/PCG/Source/PCG/Private/Elements/PCGPointFilter.cpp
+++ b/Engine/Plugins/Experimental/PCG/Source/PCG/Private/Elements/PCGPointFilter.cpp
@@ -198,9 +198,6 @@ bool FPCGPointFilterElement::ExecuteInternal(FPCGContext* Context) const
 	TArray<FPCGTaggedData>& Outputs = Context->OutputData.TaggedData;
 	TArray<FPCGTaggedData> FilterData = Context->InputData.GetInputsByPin(PCGPointFilterConstants::FilterLabel);
 
-	// Forward any non-input data
-	Outputs.Append(Context->InputData.GetAllSettings());
-
 	const EPCGPointFilterOperator Operator = Settings->Operator;
 	const bool bUseConstantThreshold = Settings->bUseConstantThreshold;
 	// If the filter is a param data, disable bUseSpatialQuery. Therefore make it non-const, to be modified below.
diff --git a/Engine/Plugins/Experimental/PCG/Source/PCG/Private/Elements/PCGPointMatchAndSet.cpp b/Engine/Plugins/Experimental/PCG/Source/PCG/Private/Elements/PCGPointMatchAndSet.cpp
--- a/Engine/Plugins/Experimental/PCG/Source/PCG/Private/Elements/PCGPointMatchAndSet.cpp
+++ b/Engine/Plugins/Experimental/PCG/Source/PCG/Private/Elements/PCGPointMatchAndSet.cpp
@@ -155,9 +155,6 @@ bool FPCGPointMatchAndSetElement::ExecuteInternal(FPCGContext* Context) const
 	TArray<FPCGTaggedData> Inputs = Context->InputData.GetInputs();
 	TArray<FPCGTaggedData>& Outputs = Context->OutputData.TaggedData;
 
-	// Forward any non-input data
-	Outputs.Append(Context->InputData.GetAllSettings());
-
 	for (const FPCGTaggedData& Input : Inputs)
 	{
 		FPCGTaggedData& Output = Outputs.Add_GetRef(Input);
diff --git a/Engine/Plugins/Experimental/PCG/Source/PCG/Private/Elements/PCGPointSampler.cpp b/Engine/Plugins/Experimental/PCG/Source/PCG/Private/Elements/PCGPointSampler.cpp
--- a/Engine/Plugins/Experimental/PCG/Source/PCG/Private/Elements/PCGPointSampler.cpp
+++ b/Engine/Plugins/Experimental/PCG/Source/PCG/Private/Elements/PCGPointSampler.cpp
@@ -31,9 +31,6 @@ bool FPCGPointSamplerElement::ExecuteInternal(FPCGContext* Context) const
 
 	TArray<FPCGTaggedData>& Outputs = Context->OutputData.TaggedData;
 
-	// Forward any non-input data, excluding params
-	Outputs.Append(Context->InputData.GetAllSettings());
-
 	const float Ratio = Settings->Ratio;
 #if WITH_EDITOR
 	const bool bKeepZeroDensityPoints = Settings->bKeepZeroDensityPoints;
diff --git a/Engine/Plugins/Experimental/PCG/Source/PCG/Private/Elements/PCGProjectionElement.cpp b/Engine/Plugins/Experimental/PCG/Source/PCG/Private/Elements/PCGProjectionElement.cpp
--- a/Engine/Plugins/Experimental/PCG/Source/PCG/Private/Elements/PCGProjectionElement.cpp
+++ b/Engine/Plugins/Experimental/PCG/Source/PCG/Private/Elements/PCGProjectionElement.cpp
@@ -95,9 +95,6 @@ bool FPCGProjectionElement::ExecuteInternal(FPCGContext* Context) const
 		ProjectionTaggedData.Tags.Append(Targets[0].Tags);
 	}
 
-	// Pass-through exclusions/settings
-	Outputs.Append(Context->InputData.GetAllSettings());
-
 	return true;
 }
 
diff --git a/Engine/Plugins/Experimental/PCG/Source/PCG/Private/Elements/PCGSelfPruning.cpp b/Engine/Plugins/Experimental/PCG/Source/PCG/Private/Elements/PCGSelfPruning.cpp
--- a/Engine/Plugins/Experimental/PCG/Source/PCG/Private/Elements/PCGSelfPruning.cpp
+++ b/Engine/Plugins/Experimental/PCG/Source/PCG/Private/Elements/PCGSelfPruning.cpp
@@ -237,9 +237,6 @@ namespace PCGSelfPruningElement
 				PCGE_LOG_C(Verbose, Context, "Generated %d points from %d source points", OutputPoints.Num(), Points.Num());
 			}
 		}
-
-		// Finally, forward any settings
-		Outputs.Append(Context->InputData.GetAllSettings());
 	}
 }
 
diff --git a/Engine/Plugins/Experimental/PCG/Source/PCG/Private/Elements/PCGStaticMeshSpawner.cpp b/Engine/Plugins/Experimental/PCG/Source/PCG/Private/Elements/PCGStaticMeshSpawner.cpp
--- a/Engine/Plugins/Experimental/PCG/Source/PCG/Private/Elements/PCGStaticMeshSpawner.cpp
+++ b/Engine/Plugins/Experimental/PCG/Source/PCG/Private/Elements/PCGStaticMeshSpawner.cpp
@@ -53,9 +53,6 @@ bool FPCGStaticMeshSpawnerElement::PrepareDataInternal(FPCGContext* InContext) c
 	TArray<FPCGTaggedData> Inputs = Context->InputData.GetInputs();
 	TArray<FPCGTaggedData>& Outputs = Context->OutputData.TaggedData;
 
-	// Forward any non-input data
-	Outputs.Append(Context->InputData.GetAllSettings());
-
 #if WITH_EDITOR
 	// In editor, we always want to generate this data for inspection & to prevent caching issues
 	const bool bGenerateOutput = true;
diff --git a/Engine/Plugins/Experimental/PCG/Source/PCG/Private/Elements/PCGSurfaceSampler.cpp b/Engine/Plugins/Experimental/PCG/Source/PCG/Private/Elements/PCGSurfaceSampler.cpp
--- a/Engine/Plugins/Experimental/PCG/Source/PCG/Private/Elements/PCGSurfaceSampler.cpp
+++ b/Engine/Plugins/Experimental/PCG/Source/PCG/Private/Elements/PCGSurfaceSampler.cpp
@@ -407,9 +407,6 @@ bool FPCGSurfaceSamplerElement::ExecuteInternal(FPCGContext* Context) const
 		Outputs[GenerationIndex].Data = PCGSurfaceSampler::SampleSurface(Context, GeneratingShape, BoundingShapeSpatialInput, LoopData);
 	}
 
-	// Finally, forward any exclusions/settings
-	Outputs.Append(Context->InputData.GetAllSettings());
-
 	return true;
 }
 
diff --git a/Engine/Plugins/Experimental/PCG/Source/PCG/Private/Elements/PCGTransformPoints.cpp b/Engine/Plugins/Experimental/PCG/Source/PCG/Private/Elements/PCGTransformPoints.cpp
--- a/Engine/Plugins/Experimental/PCG/Source/PCG/Private/Elements/PCGTransformPoints.cpp
+++ b/Engine/Plugins/Experimental/PCG/Source/PCG/Private/Elements/PCGTransformPoints.cpp
@@ -209,9 +209,6 @@ bool FPCGTransformPointsElement::ExecuteInternal(FPCGContext* Context) const
 		}
 	}
 
-	// Forward any non-input data
-	Outputs.Append(Context->InputData.GetAllSettings());
-
 	return true;
 }
 
diff --git a/Engine/Plugins/Experimental/PCG/Source/PCG/Private/Elements/PCGUnionElement.cpp b/Engine/Plugins/Experimental/PCG/Source/PCG/Private/Elements/PCGUnionElement.cpp
--- a/Engine/Plugins/Experimental/PCG/Source/PCG/Private/Elements/PCGUnionElement.cpp
+++ b/Engine/Plugins/Experimental/PCG/Source/PCG/Private/Elements/PCGUnionElement.cpp
@@ -85,8 +85,5 @@ bool FPCGUnionElement::ExecuteInternal(FPCGContext* Context) const
 		UnionTaggedData.Data = UnionData;
 	}
 
-	// Finally, pass-through settings
-	Outputs.Append(Context->InputData.GetAllSettings());
-
 	return true;
 }
diff --git a/Engine/Plugins/Experimental/PCG/Source/PCG/Private/Elements/PCGVolumeSampler.cpp b/Engine/Plugins/Experimental/PCG/Source/PCG/Private/Elements/PCGVolumeSampler.cpp
--- a/Engine/Plugins/Experimental/PCG/Source/PCG/Private/Elements/PCGVolumeSampler.cpp
+++ b/Engine/Plugins/Experimental/PCG/Source/PCG/Private/Elements/PCGVolumeSampler.cpp
@@ -228,9 +228,6 @@ bool FPCGVolumeSamplerElement::ExecuteInternal(FPCGContext* Context) const
 		}
 	}
 
-	// Finally, forward any settings
-	Outputs.Append(Context->InputData.GetAllSettings());
-
 	return true;
 }
 
diff --git a/Engine/Plugins/Experimental/PCG/Source/PCG/Private/Elements/PCGWorldQuery.cpp b/Engine/Plugins/Experimental/PCG/Source/PCG/Private/Elements/PCGWorldQuery.cpp
--- a/Engine/Plugins/Experimental/PCG/Source/PCG/Private/Elements/PCGWorldQuery.cpp
+++ b/Engine/Plugins/Experimental/PCG/Source/PCG/Private/Elements/PCGWorldQuery.cpp
@@ -40,9 +40,6 @@ bool FPCGWorldVolumetricQueryElement::ExecuteInternal(FPCGContext* Context) cons
 	FPCGTaggedData& Output = Context->OutputData.TaggedData.Emplace_GetRef();
 	Output.Data = Data;
 
-	// Pass-through settings & exclusions
-	Context->OutputData.TaggedData.Append(Context->InputData.GetAllSettings());
-
 	return true;
 }
 
@@ -107,8 +104,5 @@ bool FPCGWorldRayHitQueryElement::ExecuteInternal(FPCGContext* Context) const
 	FPCGTaggedData& Output = Context->OutputData.TaggedData.Emplace_GetRef();
 	Output.Data = Data;
 
-	// Pass-through settings & exclusions
-	Context->OutputData.TaggedData.Append(Context->InputData.GetAllSettings());
-
 	return true;
 }
