[PCG] Change default colour scheme for PCG graph to new defaults

diff --git a/Engine/Plugins/Experimental/PCG/Source/PCGEditor/Private/PCGEditorSettings.cpp b/Engine/Plugins/Experimental/PCG/Source/PCGEditor/Private/PCGEditorSettings.cpp
--- a/Engine/Plugins/Experimental/PCG/Source/PCGEditor/Private/PCGEditorSettings.cpp
+++ b/Engine/Plugins/Experimental/PCG/Source/PCGEditor/Private/PCGEditorSettings.cpp
@@ -20,23 +20,23 @@ UPCGEditorSettings::UPCGEditorSettings(const FObjectInitializer& ObjectInitializ
 	SamplerNodeColor = FLinearColor(0.8f, 1.0f, 0.4f);
 	SpawnerNodeColor = FLinearColor(1.0f, 0.6f, 0.4f);
 	SubgraphNodeColor = FLinearColor(1.0f, 0.1f, 0.1f);
-	ParamDataNodeColor = FLinearColor(1.0f, 0.6f, 0.0f);
+	ParamDataNodeColor = FLinearColor(1.0f, 0.47f, 0.0f);
 	DebugNodeColor = FLinearColor(1.0f, 0.0f, 1.0f);
 
 	DefaultPinColor = FLinearColor(0.16f, 0.16f, 0.16f);
 
 	SpatialDataPinColor = FLinearColor(1.0f, 1.0f, 1.0f);
-	ConcreteDataPinColor = FLinearColor(0.2f, 0.2f, 1.0f);
-	PointDataPinColor = FLinearColor(0.2f, 0.2f, 1.0f);
-	PolyLineDataPinColor = FLinearColor(0.2f, 0.2f, 1.0f);
-	LandscapeDataPinColor = FLinearColor(0.2f, 0.2f, 1.0f);
-	TextureDataPinColor = FLinearColor(0.2f, 0.2f, 1.0f);
-	RenderTargetDataPinColor = FLinearColor(1.0f, 0.3f, 0.f);
-	SurfaceDataPinColor = FLinearColor(0.2f, 0.2f, 1.0f);
-	VolumeDataPinColor = FLinearColor(0.2f, 0.2f, 1.0f);
-	PrimitiveDataPinColor = FLinearColor(0.2f, 0.2f, 1.0f);
+	ConcreteDataPinColor = FLinearColor(0.45f, 0.38f, 0.96f);
+	PointDataPinColor = FLinearColor(0.05f, 0.25f, 1.0f);
+	PolyLineDataPinColor = FLinearColor(0.05f, 0.75f, 0.82f);
+	LandscapeDataPinColor = FLinearColor(0.66f, 0.66f, 0.07f);
+	TextureDataPinColor = FLinearColor(0.79f, 0.08f, 0.01f);
+	RenderTargetDataPinColor = FLinearColor(0.8f, 0.18f, 0.12f);
+	SurfaceDataPinColor = FLinearColor(0.06f, 0.55f, 0.21f);
+	VolumeDataPinColor = FLinearColor(0.79f, 0.06f, 0.5f);
+	PrimitiveDataPinColor = FLinearColor(0.22f, 0.05f, 1.0f);
 
-	ParamDataPinColor = FLinearColor(1.0f, 0.6f, 0.0f);
+	ParamDataPinColor = FLinearColor(1.0f, 0.38f, 0.02f);
 	UnknownDataPinColor = FLinearColor(0.3f, 0.3f, 0.3f);
 
 	bEnableNavigateToNativeNodes = true;
