Added missing resource logging guard

#rnx

diff --git a/Engine/Plugins/Experimental/PCG/Source/PCG/Private/Utils/PCGGeneratedResourcesLogging.cpp b/Engine/Plugins/Experimental/PCG/Source/PCG/Private/Utils/PCGGeneratedResourcesLogging.cpp
--- a/Engine/Plugins/Experimental/PCG/Source/PCG/Private/Utils/PCGGeneratedResourcesLogging.cpp
+++ b/Engine/Plugins/Experimental/PCG/Source/PCG/Private/Utils/PCGGeneratedResourcesLogging.cpp
@@ -209,6 +209,11 @@ namespace PCGGeneratedResourcesLogging
 	void LogCreateCleanupTaskResource(const UPCGManagedResource* Resource/*, TArray<UPCGManagedResource*>& GeneratedResources*/)
 	{
 #if WITH_EDITOR
+		if (!LogEnabled())
+		{
+			return;
+		}
+
 		if (Resource)
 		{
 			if (const UPCGManagedActors* ManagedActors = Cast<UPCGManagedActors>(Resource))
