[PCG] Fix newly placed nodes using default settings upon copy paste

diff --git a/Engine/Plugins/Experimental/PCG/Source/PCG/Private/PCGNode.cpp b/Engine/Plugins/Experimental/PCG/Source/PCG/Private/PCGNode.cpp
--- a/Engine/Plugins/Experimental/PCG/Source/PCG/Private/PCGNode.cpp
+++ b/Engine/Plugins/Experimental/PCG/Source/PCG/Private/PCGNode.cpp
@@ -5,11 +5,12 @@
 #include "PCGCustomVersion.h"
 #include "PCGEdge.h"
 #include "PCGGraph.h"
-
-#include "Algo/Find.h"
 #include "PCGModule.h"
 #include "PCGPin.h"
 
+#include "Algo/Find.h"
+#include "UObject/Package.h"
+
 #include UE_INLINE_GENERATED_CPP_BY_NAME(PCGNode)
 
 UPCGNode::UPCGNode(const FObjectInitializer& ObjectInitializer)
@@ -501,14 +502,22 @@ bool UPCGNode::IsEdgeUsedByNodeExecution(const UPCGEdge* InEdge) const
 
 void UPCGNode::SetSettingsInterface(UPCGSettingsInterface* InSettingsInterface, bool bUpdatePins)
 {
-#if WITH_EDITOR
 	const bool bDifferentInterface = (SettingsInterface.Get() != InSettingsInterface);
 	if (bDifferentInterface && SettingsInterface)
 	{
+#if WITH_EDITOR
 		SettingsInterface->OnSettingsChangedDelegate.RemoveAll(this);
-	}
 #endif
 
+		// Un-outer the current settings to disassociate old settings from node. Without this one can copy paste
+		// a node and get both settings objects in the clipboard text, and the wrong settings can be used upon paste.
+		if (ensure(SettingsInterface->GetOuter() == this))
+		{
+			SettingsInterface->Rename(nullptr, GetTransientPackage(), REN_ForceNoResetLoaders | REN_DoNotDirty | REN_DontCreateRedirectors | REN_NonTransactional);
+			SettingsInterface->MarkAsGarbage();
+		}
+	}
+
 	SettingsInterface = InSettingsInterface;
 
 #if WITH_EDITOR
