[PCG] Moved PCGTools from Window to Tools menu.

diff --git a/Engine/Plugins/Experimental/PCG/Source/PCGEditor/Private/PCGEditorModule.cpp b/Engine/Plugins/Experimental/PCG/Source/PCGEditor/Private/PCGEditorModule.cpp
--- a/Engine/Plugins/Experimental/PCG/Source/PCGEditor/Private/PCGEditorModule.cpp
+++ b/Engine/Plugins/Experimental/PCG/Source/PCGEditor/Private/PCGEditorModule.cpp
@@ -2,9 +2,6 @@
 
 #include "PCGEditorModule.h"
 
-#include "AssetTypeActions/PCGCommonAssetTypeActions.h"
-#include "AssetTypeActions/PCGGraphAssetTypeActions.h"
-#include "AssetTypeActions/PCGSettingsAssetTypeActions.h"
 #include "PCGEditorCommands.h"
 #include "PCGEditorGraphNodeFactory.h"
 #include "PCGEditorSettings.h"
@@ -12,17 +9,20 @@
 #include "PCGEditorUtils.h"
 #include "PCGSubsystem.h"
 #include "PCGVolumeFactory.h"
+#include "AssetTypeActions/PCGCommonAssetTypeActions.h"
+#include "AssetTypeActions/PCGGraphAssetTypeActions.h"
+#include "AssetTypeActions/PCGSettingsAssetTypeActions.h"
 
+#include "ISettingsModule.h"
+#include "LevelEditor.h"
+#include "PropertyEditorModule.h"
+#include "ToolMenus.h"
+#include "Details/PCGAttributePropertySelectorDetails.h"
 #include "Details/PCGBlueprintSettingsDetails.h"
 #include "Details/PCGGraphDetails.h"
 #include "Details/PCGGraphInstanceDetails.h"
 #include "Details/PCGVolumeDetails.h"
-#include "Details/PCGAttributePropertySelectorDetails.h"
-
 #include "Framework/MultiBox/MultiBoxBuilder.h"
-#include "ISettingsModule.h"
-#include "LevelEditor.h"
-#include "PropertyEditorModule.h"
 
 #define LOCTEXT_NAMESPACE "FPCGEditorModule"
 
@@ -131,38 +131,20 @@ void FPCGEditorModule::UnregisterAssetTypeActions()
 
 void FPCGEditorModule::RegisterMenuExtensions()
 {
-	MenuExtensibilityManager = MakeShareable(new FExtensibilityManager);
-	ToolBarExtensibilityManager = MakeShareable(new FExtensibilityManager);
-
-	FLevelEditorModule& LevelEditorModule = FModuleManager::LoadModuleChecked<FLevelEditorModule>("LevelEditor");
-
-	{
-		TSharedPtr<FExtender> NewMenuExtender = MakeShareable(new FExtender);
-		NewMenuExtender->AddMenuExtension("LevelEditor",
-			EExtensionHook::After,
-			nullptr,
-			FMenuExtensionDelegate::CreateRaw(this, &FPCGEditorModule::AddMenuEntry));
-
-		LevelEditorModule.GetMenuExtensibilityManager()->AddExtender(NewMenuExtender);
-	}
-}
-
-void FPCGEditorModule::UnregisterMenuExtensions()
-{
-	MenuExtensibilityManager.Reset();
-	ToolBarExtensibilityManager.Reset();
-}
-
-void FPCGEditorModule::AddMenuEntry(FMenuBuilder& MenuBuilder)
-{
-	MenuBuilder.BeginSection("PCGMenu", TAttribute<FText>(FText::FromString("PCG Tools")));
-
-	MenuBuilder.AddSubMenu(
+	FToolMenuOwnerScoped OwnerScoped(this);
+	UToolMenu* Menu = UToolMenus::Get()->ExtendMenu("LevelEditor.MainMenu.Tools");
+	FToolMenuSection& Section = Menu->AddSection("PCGToolsSection", LOCTEXT("PCGToolsSection", "Procedural Generation Tools"));
+	
+	Section.AddSubMenu(
+		"PCGToolsSubMenu",
 		LOCTEXT("PCGSubMenu", "PCG Framework"),
 		LOCTEXT("PCGSubMenu_Tooltip", "PCG Framework related functionality"),
 		FNewMenuDelegate::CreateRaw(this, &FPCGEditorModule::PopulateMenuActions));
-	
-	MenuBuilder.EndSection();
+}
+
+void FPCGEditorModule::UnregisterMenuExtensions()
+{	
+	UToolMenus::UnregisterOwner(this);
 }
 
 void FPCGEditorModule::PopulateMenuActions(FMenuBuilder& MenuBuilder)
diff --git a/Engine/Plugins/Experimental/PCG/Source/PCGEditor/Private/PCGEditorModule.h b/Engine/Plugins/Experimental/PCG/Source/PCGEditor/Private/PCGEditorModule.h
--- a/Engine/Plugins/Experimental/PCG/Source/PCGEditor/Private/PCGEditorModule.h
+++ b/Engine/Plugins/Experimental/PCG/Source/PCGEditor/Private/PCGEditorModule.h
@@ -14,7 +14,7 @@ DECLARE_LOG_CATEGORY_EXTERN(LogPCGEditor, Log, All);
 class FMenuBuilder;
 class FPCGEditorGraphNodeFactory;
 
-class FPCGEditorModule : public IModuleInterface, public IHasMenuExtensibility, public IHasToolBarExtensibility
+class FPCGEditorModule : public IModuleInterface
 {
 public:
 	// ~IModuleInterface implementation
@@ -23,14 +23,6 @@ public:
 	virtual bool SupportsDynamicReloading() override;
 	// ~End IModuleInterface implementation
 
-	//~ IHasMenuExtensibility interface
-	virtual TSharedPtr<FExtensibilityManager> GetMenuExtensibilityManager() override { return MenuExtensibilityManager; }
-	//~ End IHasMenuExtensibility interface
-
-	//~ IHasToolBarExtensibility interface
-	virtual TSharedPtr<FExtensibilityManager> GetToolBarExtensibilityManager() override { return ToolBarExtensibilityManager; }
-	//~ End IHasToolBarExtensibility interface
-
 	static EAssetTypeCategories::Type GetAssetCategory() { return PCGAssetCategory; }
 
 protected:
@@ -40,7 +32,6 @@ protected:
 	void UnregisterAssetTypeActions();
 	void RegisterMenuExtensions();
 	void UnregisterMenuExtensions();
-	void AddMenuEntry(FMenuBuilder& MenuBuilder);
 	void PopulateMenuActions(FMenuBuilder& MenuBuilder);
 	void RegisterSettings();
 	void UnregisterSettings();
