[PCG] Improved naming of PCG created mesh components.

diff --git a/Engine/Plugins/Experimental/PCG/Source/PCG/Private/Helpers/PCGActorHelpers.cpp b/Engine/Plugins/Experimental/PCG/Source/PCG/Private/Helpers/PCGActorHelpers.cpp
--- a/Engine/Plugins/Experimental/PCG/Source/PCG/Private/Helpers/PCGActorHelpers.cpp
+++ b/Engine/Plugins/Experimental/PCG/Source/PCG/Private/Helpers/PCGActorHelpers.cpp
@@ -66,9 +66,25 @@ UInstancedStaticMeshComponent* UPCGActorHelpers::GetOrCreateISMC(AActor* InTarge
 	const bool bMeshHasNaniteData = StaticMesh->GetRenderData()->NaniteResources.PageStreamingStates.Num() > 0;
 #endif
 
-	const TSubclassOf<UInstancedStaticMeshComponent> ComponentClass = bMeshHasNaniteData ? UInstancedStaticMeshComponent::StaticClass() : UHierarchicalInstancedStaticMeshComponent::StaticClass();
+	FString ComponentName;
+	TSubclassOf<UInstancedStaticMeshComponent> ComponentClass;
+	if (bMeshHasNaniteData)
+	{
+		ComponentClass = UInstancedStaticMeshComponent::StaticClass();
+		ComponentName += TEXT("ISM");
+	}
+	else
+	{
+		ComponentClass = UHierarchicalInstancedStaticMeshComponent::StaticClass();
+		ComponentName += TEXT("HISM");
+	}
+
+	if (InISMCDescriptor.StaticMesh)
+	{
+		ComponentName += TEXT("_") + InISMCDescriptor.StaticMesh->GetName();
+	}
 	
-	UInstancedStaticMeshComponent* ISMC = NewObject<UInstancedStaticMeshComponent>(InTargetActor, ComponentClass);
+	UInstancedStaticMeshComponent* ISMC = NewObject<UInstancedStaticMeshComponent>(InTargetActor, ComponentClass, MakeUniqueObjectName(InTargetActor, ComponentClass, FName(ComponentName)));
 	InISMCDescriptor.InitComponent(ISMC);
 		
 	ISMC->RegisterComponent();
@@ -161,9 +177,6 @@ UInstancedStaticMeshComponent* UPCGActorHelpers::GetOrCreateISMC(AActor* InTarge
 	InTargetActor->Modify();
 
 	// Otherwise, create a new component
-	// TODO: use static mesh component if there's only one instance
-	// TODO: add hism/ism switch or better yet, use a template component
-	UInstancedStaticMeshComponent* ISMC = nullptr;
 
 	// Done as in InstancedStaticMesh.cpp
 #if WITH_EDITOR
@@ -172,15 +185,25 @@ UInstancedStaticMeshComponent* UPCGActorHelpers::GetOrCreateISMC(AActor* InTarge
 	const bool bMeshHasNaniteData = InMesh->GetRenderData()->NaniteResources.PageStreamingStates.Num() > 0;
 #endif
 
+	FString ComponentName;
+	TSubclassOf<UInstancedStaticMeshComponent> ComponentClass;
 	if (bMeshHasNaniteData)
 	{
-		ISMC = NewObject<UInstancedStaticMeshComponent>(InTargetActor);
+		ComponentClass = UInstancedStaticMeshComponent::StaticClass();
+		ComponentName += TEXT("ISM");
 	}
 	else
 	{
-		ISMC = NewObject<UHierarchicalInstancedStaticMeshComponent>(InTargetActor);
+		ComponentClass = UHierarchicalInstancedStaticMeshComponent::StaticClass();
+		ComponentName += TEXT("HISM");
 	}
 
+	ComponentName += TEXT("_") + InMesh->GetName();
+	
+	// TODO: use static mesh component if there's only one instance
+	// TODO: add hism/ism switch or better yet, use a template component
+	UInstancedStaticMeshComponent* ISMC = NewObject<UInstancedStaticMeshComponent>(InTargetActor, ComponentClass, MakeUniqueObjectName(InTargetActor, ComponentClass, FName(ComponentName)));
+	
 	ISMC->SetStaticMesh(InMesh);
 
 	// TODO: improve material override mechanisms
