Finer grained time-slicing in the static mesh spawner
Improved the IsRuntimeOrPIE check to include the -game case
Added some profiling tags to help find bottlenecks

diff --git a/Engine/Plugins/Experimental/PCG/Source/PCG/Private/Data/PCGPointData.cpp b/Engine/Plugins/Experimental/PCG/Source/PCG/Private/Data/PCGPointData.cpp
--- a/Engine/Plugins/Experimental/PCG/Source/PCG/Private/Data/PCGPointData.cpp
+++ b/Engine/Plugins/Experimental/PCG/Source/PCG/Private/Data/PCGPointData.cpp
@@ -268,6 +268,8 @@ void UPCGPointData::AddToCrc(FArchiveCrc32& Ar) const
 		return;
 	}
 
+	TRACE_CPUPROFILER_EVENT_SCOPE(UPCGPointData::AddToCrc);
+
 	uint32 ThisType = static_cast<uint32>(GetDataType());
 	Ar << ThisType;
 
diff --git a/Engine/Plugins/Experimental/PCG/Source/PCG/Private/Elements/PCGStaticMeshSpawner.cpp b/Engine/Plugins/Experimental/PCG/Source/PCG/Private/Elements/PCGStaticMeshSpawner.cpp
--- a/Engine/Plugins/Experimental/PCG/Source/PCG/Private/Elements/PCGStaticMeshSpawner.cpp
+++ b/Engine/Plugins/Experimental/PCG/Source/PCG/Private/Elements/PCGStaticMeshSpawner.cpp
@@ -137,15 +137,27 @@ bool FPCGStaticMeshSpawnerElement::ExecuteInternal(FPCGContext* InContext) const
 		const FPCGStaticMeshSpawnerContext::FPackedInstanceListData& InstanceList = Context->MeshInstancesData.Last();
 		check(InstanceList.MeshInstances.Num() == InstanceList.PackedCustomData.Num());
 
-		if (InstanceList.SpatialData->TargetActor.IsValid())
+		const bool bTargetActorValid = InstanceList.SpatialData->TargetActor.IsValid();
+
+		if (bTargetActorValid)
 		{
-			for (int32 DataIndex = 0; DataIndex < InstanceList.MeshInstances.Num(); ++DataIndex)
+			while (Context->CurrentDataIndex < InstanceList.MeshInstances.Num())
 			{
-				SpawnStaticMeshInstances(Context, InstanceList.MeshInstances[DataIndex], InstanceList.SpatialData->TargetActor.Get(), InstanceList.PackedCustomData[DataIndex]);
+				SpawnStaticMeshInstances(Context, InstanceList.MeshInstances[Context->CurrentDataIndex], InstanceList.SpatialData->TargetActor.Get(), InstanceList.PackedCustomData[Context->CurrentDataIndex]);
+				++Context->CurrentDataIndex;
+
+				if (Context->ShouldStop())
+				{
+					break;
+				}
 			}
 		}
 
-		Context->MeshInstancesData.RemoveAtSwap(Context->MeshInstancesData.Num() - 1);
+		if (!bTargetActorValid || Context->CurrentDataIndex == InstanceList.MeshInstances.Num())
+		{
+			Context->MeshInstancesData.RemoveAtSwap(Context->MeshInstancesData.Num() - 1);
+			Context->CurrentDataIndex = 0;
+		}
 
 		if (Context->ShouldStop())
 		{
diff --git a/Engine/Plugins/Experimental/PCG/Source/PCG/Private/PCGElement.cpp b/Engine/Plugins/Experimental/PCG/Source/PCG/Private/PCGElement.cpp
--- a/Engine/Plugins/Experimental/PCG/Source/PCG/Private/PCGElement.cpp
+++ b/Engine/Plugins/Experimental/PCG/Source/PCG/Private/PCGElement.cpp
@@ -418,6 +418,7 @@ bool IPCGElement::IsCacheableInstance(const UPCGSettingsInterface* InSettingsInt
 
 void IPCGElement::GetDependenciesCrc(const FPCGDataCollection& InInput, const UPCGSettings* InSettings, UPCGComponent* InComponent, FPCGCrc& OutCrc) const
 {
+	TRACE_CPUPROFILER_EVENT_SCOPE_TEXT(*FString::Printf(TEXT("IPCGElement::GetDependenciesCrc (%s)"), InSettings ? *InSettings->GetFName().ToString() : TEXT("")));
 	FPCGCrc Crc = InInput.Crc;
 
 	if (InSettings)
diff --git a/Engine/Plugins/Experimental/PCG/Source/PCG/Private/PCGHelpers.cpp b/Engine/Plugins/Experimental/PCG/Source/PCG/Private/PCGHelpers.cpp
--- a/Engine/Plugins/Experimental/PCG/Source/PCG/Private/PCGHelpers.cpp
+++ b/Engine/Plugins/Experimental/PCG/Source/PCG/Private/PCGHelpers.cpp
@@ -156,7 +156,7 @@ namespace PCGHelpers
 	bool IsRuntimeOrPIE()
 	{
 #if WITH_EDITOR
-		return (GEditor && GEditor->PlayWorld) || GIsPlayInEditorWorld;
+		return (GEditor && GEditor->PlayWorld) || GIsPlayInEditorWorld || IsRunningGame();
 #else
 		return true;
 #endif // WITH_EDITOR
diff --git a/Engine/Plugins/Experimental/PCG/Source/PCG/Public/Elements/PCGStaticMeshSpawner.h b/Engine/Plugins/Experimental/PCG/Source/PCG/Public/Elements/PCGStaticMeshSpawner.h
--- a/Engine/Plugins/Experimental/PCG/Source/PCG/Public/Elements/PCGStaticMeshSpawner.h
+++ b/Engine/Plugins/Experimental/PCG/Source/PCG/Public/Elements/PCGStaticMeshSpawner.h
@@ -113,6 +113,8 @@ struct FPCGStaticMeshSpawnerContext : public FPCGContext
 	};
 
 	TArray<FPackedInstanceListData> MeshInstancesData;
+	// Index of MeshInstances/PackedCustomData in last MeshInstancesData element
+	int32 CurrentDataIndex = 0;
 };
 
 class FPCGStaticMeshSpawnerElement : public IPCGElement
