[PCG] Disable serializing the data from a pcg data asset into the CRC because it is very slow.

A data UID is serialized which should be adequate to uniquely specify the data.

diff --git a/Engine/Plugins/Experimental/PCG/Source/PCG/Private/PCGSettings.cpp b/Engine/Plugins/Experimental/PCG/Source/PCG/Private/PCGSettings.cpp
--- a/Engine/Plugins/Experimental/PCG/Source/PCG/Private/PCGSettings.cpp
+++ b/Engine/Plugins/Experimental/PCG/Source/PCG/Private/PCGSettings.cpp
@@ -32,6 +32,18 @@ public:
 #if WITH_EDITOR
 	virtual bool ShouldSkipProperty(const FProperty* InProperty) const override
 	{
+		// Currently we rely on the 'UID' property getting included in the Crc. An example of this are asset settings
+		// for which we avoid doing a full data CRC and instead rely on including hte UID. This property is transient
+		// and will only serialize if IsPersistent() is false (see tests in FProperty::ShouldSerializeValue()).
+		ensure(!IsPersistent());
+
+		// Omit CRC'ing data collections here as it is very slow. Rely instead on UID.
+		const FStructProperty* StructProperty = CastField<FStructProperty>(InProperty);
+		if (StructProperty && StructProperty->Struct && StructProperty->Struct->IsChildOf(FPCGDataCollection::StaticStruct()))
+		{
+			return true;
+		}
+
 		const bool bSkip = InProperty && (
 			InProperty->GetFName() == GET_MEMBER_NAME_CHECKED(UPCGSettings, DebugSettings)
 			|| InProperty->GetFName() == GET_MEMBER_NAME_CHECKED(UPCGSettings, DeterminismSettings)
@@ -39,7 +51,9 @@ public:
 			|| InProperty->GetFName() == GET_MEMBER_NAME_CHECKED(UPCGSettings, Category)
 			|| InProperty->GetFName() == GET_MEMBER_NAME_CHECKED(UPCGSettings, Description)
 			|| InProperty->GetFName() == GET_MEMBER_NAME_CHECKED(UPCGSettings, bExposeToLibrary)
+			|| InProperty->GetFName() == GET_MEMBER_NAME_CHECKED(UPCGSettings, CachedOverridableParams)
 			);
+
 		return bSkip;
 	}
 #endif // WITH_EDITOR
diff --git a/Engine/Plugins/Experimental/PCG/Source/PCG/Public/PCGSettings.h b/Engine/Plugins/Experimental/PCG/Source/PCG/Public/PCGSettings.h
--- a/Engine/Plugins/Experimental/PCG/Source/PCG/Public/PCGSettings.h
+++ b/Engine/Plugins/Experimental/PCG/Source/PCG/Public/PCGSettings.h
@@ -106,6 +106,7 @@ class PCG_API UPCGSettings : public UPCGSettingsInterface
 {
 	GENERATED_BODY()
 
+	friend class FPCGSettingsObjectCrc32;
 	friend class UPCGSettingsInterface;
 
 public:
