[PCG] Re-visit bIsPartitioned and GraphInstance to make it work correctly with BPs.

diff --git a/Engine/Plugins/Experimental/PCG/Source/PCG/Private/PCGComponent.cpp b/Engine/Plugins/Experimental/PCG/Source/PCG/Private/PCGComponent.cpp
--- a/Engine/Plugins/Experimental/PCG/Source/PCG/Private/PCGComponent.cpp
+++ b/Engine/Plugins/Experimental/PCG/Source/PCG/Private/PCGComponent.cpp
@@ -48,6 +48,12 @@ namespace PCGComponent
 	const bool bSaveOnCleanupAndGenerate = false;
 }
 
+UPCGComponent::UPCGComponent(const FObjectInitializer& InObjectInitializer)
+	: Super(InObjectInitializer)
+{
+	GraphInstance = InObjectInitializer.CreateDefaultSubobject<UPCGGraphInstance>(this, TEXT("PCGGraphInstance"));
+}
+
 bool UPCGComponent::CanPartition() const
 {
 	// Support/Force partitioning on non-PCG partition actors in WP worlds.
@@ -56,12 +62,12 @@ bool UPCGComponent::CanPartition() const
 
 bool UPCGComponent::IsPartitioned() const
 {
-	return bIsPartitioned && CanPartition();
+	return bIsComponentPartitioned && CanPartition();
 }
 
 void UPCGComponent::SetIsPartitioned(bool bIsNowPartitioned)
 {
-	if (bIsNowPartitioned == bIsPartitioned)
+	if (bIsNowPartitioned == bIsComponentPartitioned)
 	{
 		return;
 	}
@@ -77,18 +83,18 @@ void UPCGComponent::SetIsPartitioned(bool bIsNowPartitioned)
 
 		if (bIsNowPartitioned)
 		{
-			bIsPartitioned = bIsNowPartitioned;
+			bIsComponentPartitioned = bIsNowPartitioned;
 			Subsystem->RegisterOrUpdatePCGComponent(this, bDoActorMapping);
 		}
 		else
 		{
 			Subsystem->UnregisterPCGComponent(this);
-			bIsPartitioned = bIsNowPartitioned;
+			bIsComponentPartitioned = bIsNowPartitioned;
 		}
 	}
 	else
 	{
-		bIsPartitioned = false;
+		bIsComponentPartitioned = false;
 	}
 }
 
@@ -116,25 +122,6 @@ void UPCGComponent::SetGraphInterfaceLocal(UPCGGraphInterface* InGraphInterface)
 	}
 }
 
-void UPCGComponent::SetGraphInstanceFromParent(UPCGGraphInstance* InParentGraphInstance)
-{
-	// Since it must be used on local components only, add a check here
-	check(bIsComponentLocal);
-
-#if WITH_EDITOR
-	if (GraphInstance)
-	{
-		GraphInstance->OnGraphChangedDelegate.RemoveAll(this);
-	}
-#endif // WITH_EDITOR
-
-	GraphInstance = InParentGraphInstance;
-
-#if WITH_EDITOR
-	SetupCallbacksOnCreation();
-#endif // WITH_EDITOR
-}
-
 void UPCGComponent::AddToManagedResources(UPCGManagedResource* InResource)
 {
 	PCGGeneratedResourcesLogging::LogAddToManagedResources(InResource);
@@ -205,7 +192,7 @@ void UPCGComponent::SetPropertiesFromOriginal(const UPCGComponent* Original)
 #if WITH_EDITOR
 	const bool bHasDirtyInput = InputType != NewInputType;
 	const bool bHasDirtyExclusions = !(ExcludedTags.Num() == Original->ExcludedTags.Num() && ExcludedTags.Includes(Original->ExcludedTags));
-	const bool bIsDirty = bHasDirtyInput || bHasDirtyExclusions || GraphInstance != Original->GraphInstance;
+	const bool bIsDirty = bHasDirtyInput || bHasDirtyExclusions || GraphInstance->Graph != Original->GraphInstance;
 
 	if (bHasDirtyExclusions)
 	{
@@ -220,7 +207,7 @@ void UPCGComponent::SetPropertiesFromOriginal(const UPCGComponent* Original)
 
 	InputType = NewInputType;
 	Seed = Original->Seed;
-	GraphInstance = Original->GraphInstance;
+	GraphInstance->SetGraph(Original->GraphInstance);
 
 	GenerationTrigger = Original->GenerationTrigger;
 
@@ -948,10 +935,13 @@ void UPCGComponent::PostLoad()
 		ExclusionTags_DEPRECATED.Reset();
 	}
 
-	// Make sure that bPartitioned is false when it can't be true.
-	if (bIsPartitioned && !CanPartition())
+	// If we have both default value (bIsComponentPartitioned = false and bIsPartitioned_DEPRECATED = true)
+	// we will follow the value of bIsPartitioned_DEPRECATED.
+	// bIsPartitioned_DEPRECATED will be set to false to new objects
+	if (!bIsComponentPartitioned && bIsPartitioned_DEPRECATED)
 	{
-		bIsPartitioned = false;
+		bIsComponentPartitioned = bIsPartitioned_DEPRECATED;
+		bIsPartitioned_DEPRECATED = false;
 	}
 
 	/** Deprecation code, should be removed once generated data has been updated */
@@ -981,11 +971,6 @@ void UPCGComponent::PostLoad()
 
 	if (Graph_DEPRECATED)
 	{
-		if (!GraphInstance)
-		{
-			GraphInstance = NewObject<UPCGGraphInstance>(this, TEXT("PCGGraphInstance"), RF_Transactional | RF_Public);
-		}
-
 		GraphInstance->SetGraph(Graph_DEPRECATED);
 		Graph_DEPRECATED = nullptr;
 	}
@@ -1046,33 +1031,29 @@ void UPCGComponent::BeginDestroy()
 
 void UPCGComponent::PostInitProperties()
 {
-	// Don't initialize the graph instance for default objects and not for local components (since it will be overriden when generated).
-	if (!GraphInstance && !bIsComponentLocal && !HasAnyFlags(RF_ClassDefaultObject | RF_ArchetypeObject))
-	{
-		GraphInstance = NewObject<UPCGGraphInstance>(this, TEXT("PCGGraphInstance"), RF_Transactional | RF_Public);
-
 #if WITH_EDITOR
-		GraphInstance->OnGraphChangedDelegate.AddUObject(this, &UPCGComponent::OnGraphChanged);
+	GraphInstance->OnGraphChangedDelegate.AddUObject(this, &UPCGComponent::OnGraphChanged);
 #endif // WITH_EDITOR
-	}
 
-	// Force bIsPartitioned at False for new objects
-	// We detect new object if they are not a default object/archetype and/or they do not need load.
-	// In some cases, were the component is a default sub object (like APCGVolume), it has no loading flags
-	// even if it is loading, we need to recurse until we find an owner that is not a default sub object.
-	bool bIsNewObject = false;
+	// Note that if the component is a default sub object, we find the first outer that is not a sub object.
 	UObject* CurrentInspectedObject = this;
 	while (CurrentInspectedObject && CurrentInspectedObject->HasAnyFlags(RF_DefaultSubObject))
 	{
 		CurrentInspectedObject = CurrentInspectedObject->GetOuter();
 	}
 
-	bIsNewObject = CurrentInspectedObject && !CurrentInspectedObject->HasAnyFlags(RF_ClassDefaultObject | RF_ArchetypeObject | RF_NeedLoad | RF_NeedPostLoad);
+	// We detect new object if they are not a default object/archetype and/or they do not need load.
+	// In some cases, were the component is a default sub object (like APCGVolume), it has no loading flags
+	// even if it is loading, so we use the outer found above.
+	const bool bIsNewObject = CurrentInspectedObject && !CurrentInspectedObject->HasAnyFlags(RF_ClassDefaultObject | RF_NeedLoad | RF_NeedPostLoad);
 
+#if WITH_EDITOR
+	// Force bIsPartitioned_DEPRECATED at False for new objects
 	if (bIsNewObject)
 	{
-		bIsPartitioned = false;
+		bIsPartitioned_DEPRECATED = false;
 	}
+#endif // WITH_EDITOR
 
 	Super::PostInitProperties();
 }
@@ -1198,15 +1179,15 @@ void UPCGComponent::PostEditChangeProperty(FPropertyChangedEvent& PropertyChange
 
 	// Important note: all property changes already go through the OnObjectPropertyChanged, and will be dirtied here.
 	// So where only a Refresh is needed, it goes through the "capture all" else case.
-	if (PropName == GET_MEMBER_NAME_CHECKED(UPCGComponent, bIsPartitioned))
+	if (PropName == GET_MEMBER_NAME_CHECKED(UPCGComponent, bIsComponentPartitioned))
 	{
 		if (CanPartition())
 		{
-			// At this point, bIsPartitioned is already set with the new value.
+			// At this point, bIsComponentPartitioned is already set with the new value.
 			// But we need to do some cleanup before
 			// So keep this new value, and take its negation for the cleanup.
-			bool bIsNowPartitioned = bIsPartitioned;
-			bIsPartitioned = !bIsPartitioned;
+			bool bIsNowPartitioned = bIsComponentPartitioned;
+			bIsComponentPartitioned = !bIsComponentPartitioned;
 
 			// SetIsPartioned cleans up before, so keep track if we were generated or not.
 			bool bWasGenerated = bGenerated;
@@ -2611,6 +2592,9 @@ void FPCGComponentInstanceData::ApplyToComponent(UActorComponent* Component, con
 			PCGComponent->SetManagedResources(DuplicatedResources);
 		}
 
+		// Re-set the graph to reconnect callbacks correctly
+		PCGComponent->GraphInstance->SetGraph(PCGComponent->GraphInstance->Graph);
+
 		// Also remap if we are partitioned
 		UPCGSubsystem* Subsystem = PCGComponent->GetSubsystem();
 		if (Subsystem && PCGComponent->IsPartitioned() && SourceComponent)
diff --git a/Engine/Plugins/Experimental/PCG/Source/PCG/Private/Tests/PCGTestsCommon.cpp b/Engine/Plugins/Experimental/PCG/Source/PCG/Private/Tests/PCGTestsCommon.cpp
--- a/Engine/Plugins/Experimental/PCG/Source/PCG/Private/Tests/PCGTestsCommon.cpp
+++ b/Engine/Plugins/Experimental/PCG/Source/PCG/Private/Tests/PCGTestsCommon.cpp
@@ -40,7 +40,7 @@ namespace PCGTestsCommon
 		TestPCGComponent = NewObject<UPCGComponent>(TestActor, FName(TEXT("Test PCG Component")), RF_Transient);
 		check(TestPCGComponent);
 		// By default PCG components for tests will be non-partitioned
-		TestPCGComponent->bIsPartitioned = false;
+		TestPCGComponent->SetIsPartitioned(false);
 		TestActor->AddInstanceComponent(TestPCGComponent);
 		TestPCGComponent->RegisterComponent();
 
diff --git a/Engine/Plugins/Experimental/PCG/Source/PCG/Public/PCGComponent.h b/Engine/Plugins/Experimental/PCG/Source/PCG/Public/PCGComponent.h
--- a/Engine/Plugins/Experimental/PCG/Source/PCG/Public/PCGComponent.h
+++ b/Engine/Plugins/Experimental/PCG/Source/PCG/Public/PCGComponent.h
@@ -63,6 +63,8 @@ ENUM_CLASS_FLAGS(EPCGComponentDirtyFlag);
 UCLASS(BlueprintType, ClassGroup = (Procedural), meta = (BlueprintSpawnableComponent))
 class PCG_API UPCGComponent : public UActorComponent
 {
+	UPCGComponent(const FObjectInitializer& InObjectInitializer);
+
 	GENERATED_BODY()
 
 	friend class UPCGSubsystem;
@@ -170,14 +172,18 @@ public:
 #if WITH_EDITORONLY_DATA
 	UPROPERTY()
 	TArray<FName> ExclusionTags_DEPRECATED;
+
+	// Note for upgrade: can be safely replaced by bIsComponentPartitioned. Needed a new variable to change the default value.
+	UPROPERTY()
+	bool bIsPartitioned_DEPRECATED = true;
 #endif
 
 	UPROPERTY(BlueprintReadWrite, EditAnywhere, Category = Properties)
 	bool bActivated = true;
 
-	// Note that default value is now "False" for new objects. True is left for backcompatibility.
-	UPROPERTY(BlueprintReadWrite, EditAnywhere, AdvancedDisplay, Category = Properties, meta = (EditCondition = "!bIsComponentLocal", Tooltip = "In WorldPartiion map, will partition the component in a grid, according to PCGWorldActor settings, dispatching the generation to multiple local components."))
-	bool bIsPartitioned = true;
+	/* In World Partition map, will partition the component in a grid, according to PCGWorldActor settings, dispatching the generation to multiple local components.*/
+	UPROPERTY(BlueprintReadWrite, EditAnywhere, AdvancedDisplay, Category = Properties, meta = (EditCondition = "!bIsComponentLocal", DisplayName = "Is Partitioned"))
+	bool bIsComponentPartitioned = false;
 
 	UPROPERTY(BlueprintReadOnly, EditAnywhere, Category = Properties, AdvancedDisplay, meta = (EditCondition = "!bIsComponentLocal", EditConditionHides))
 	EPCGComponentGenerationTrigger GenerationTrigger = EPCGComponentGenerationTrigger::GenerateOnLoad;
@@ -233,7 +239,7 @@ public:
 	const FPCGDataCollection* GetInspectionData(const UPCGNode* InNode) const;
 #endif
 
-	/** Utility function (mostly for tests) to properly set the value of bIsPartitioned.
+	/** Utility function (mostly for tests) to properly set the value of bIsComponentPartitioned.
 	*   Will do an immediate cleanup first and then register/unregister the component to the subsystem.
 	*   It's your responsibility after to regenerate the graph if you want to.
 	*/
@@ -428,9 +434,6 @@ private:
 	/** Will set the given graph interface into our owned graph instance. Must not be used on local components.*/
 	void SetGraphInterfaceLocal(UPCGGraphInterface* InGraphInterface);
 
-	/** Replace the graph instance by this one. Must only be used with local components.*/
-	void SetGraphInstanceFromParent(UPCGGraphInstance* InParentGraphInstance);
-
 #if WITH_EDITOR
 public:
 	mutable PCGUtils::FExtraCapture ExtraCapture;
