[PCG] Fix PropertyToParamData attribute names for user defined structs.

diff --git a/Engine/Plugins/Experimental/PCG/Source/PCG/Private/Elements/PCGPropertyToParamData.cpp b/Engine/Plugins/Experimental/PCG/Source/PCG/Private/Elements/PCGPropertyToParamData.cpp
--- a/Engine/Plugins/Experimental/PCG/Source/PCG/Private/Elements/PCGPropertyToParamData.cpp
+++ b/Engine/Plugins/Experimental/PCG/Source/PCG/Private/Elements/PCGPropertyToParamData.cpp
@@ -149,7 +149,10 @@ bool FPCGPropertyToParamDataElement::ExecuteInternal(FPCGContext* Context) const
 				const FName ChildPropertyName = Param.PropertiesNames[0];
 				if (const FProperty* ChildProperty = (UnderlyingStruct ? UnderlyingStruct->FindPropertyByName(ChildPropertyName) : UnderlyingClass->FindPropertyByName(ChildPropertyName)))
 				{
-					ExtractableProperties.Emplace(ChildPropertyName, ObjectAddress, ChildProperty);
+					// We use authored name as attribute name to avoid issue with noisy property names, like in UUserDefinedStructs, where some random number is appended to the property name.
+					// By default, it will just return the property name anyway.
+					const FString AuthoredName = UnderlyingStruct ? UnderlyingStruct->GetAuthoredNameForField(ChildProperty) : UnderlyingClass->GetAuthoredNameForField(ChildProperty);
+					ExtractableProperties.Emplace(FName(AuthoredName), ObjectAddress, ChildProperty);
 				}
 			}
 		}
