Put in place some mechanisms to support resource reuse
Implemented for the static mesh spawner, will be done in other cases (debug, spawn actor) in subsequent CLs
Fixed an issue with no-merge spawning of actors using partitions

diff --git a/Engine/Plugins/Experimental/PCG/Source/PCG/Private/Elements/PCGDebugElement.cpp b/Engine/Plugins/Experimental/PCG/Source/PCG/Private/Elements/PCGDebugElement.cpp
--- a/Engine/Plugins/Experimental/PCG/Source/PCG/Private/Elements/PCGDebugElement.cpp
+++ b/Engine/Plugins/Experimental/PCG/Source/PCG/Private/Elements/PCGDebugElement.cpp
@@ -89,7 +89,7 @@ namespace PCGDebugElement
 				continue;
 			}
 
-			AActor* TargetActor = SpatialData->TargetActor.Get();
+			AActor* TargetActor = Context->GetTargetActor(SpatialData);
 
 			if (!TargetActor)
 			{
diff --git a/Engine/Plugins/Experimental/PCG/Source/PCG/Private/Elements/PCGSpawnActor.cpp b/Engine/Plugins/Experimental/PCG/Source/PCG/Private/Elements/PCGSpawnActor.cpp
--- a/Engine/Plugins/Experimental/PCG/Source/PCG/Private/Elements/PCGSpawnActor.cpp
+++ b/Engine/Plugins/Experimental/PCG/Source/PCG/Private/Elements/PCGSpawnActor.cpp
@@ -6,7 +6,7 @@
 #include "PCGGraph.h"
 #include "PCGHelpers.h"
 #include "PCGManagedResource.h"
-
+#include "PCGSubsystem.h"
 #include "Data/PCGPointData.h"
 #include "Data/PCGSpatialData.h"
 #include "Grid/PCGPartitionActor.h"
@@ -360,7 +360,7 @@ bool FPCGSpawnActorElement::ExecuteInternal(FPCGContext* Context) const
 			continue;
 		}
 
-		AActor* TargetActor = SpatialData->TargetActor.Get();
+		AActor* TargetActor = Context->GetTargetActor(SpatialData);
 
 		if (!TargetActor)
 		{
@@ -521,6 +521,8 @@ bool FPCGSpawnActorElement::ExecuteInternal(FPCGContext* Context) const
 						TArray<UPCGComponent*> PCGComponents;
 						GeneratedActor->GetComponents(PCGComponents);
 
+						UPCGSubsystem* Subsystem = Context->SourceComponent.Get() ? Context->SourceComponent->GetSubsystem() : nullptr;
+
 						for (UPCGComponent* PCGComponent : PCGComponents)
 						{
 							if (Settings->Option == EPCGSpawnActorOption::NoMerging)
@@ -532,6 +534,11 @@ bool FPCGSpawnActorElement::ExecuteInternal(FPCGContext* Context) const
 
 								if(bForceCallGenerate || (bOnLoadCallGenerate && PCGComponent->GenerationTrigger == EPCGComponentGenerationTrigger::GenerateOnLoad))
 								{
+									if (Subsystem && PCGComponent->IsPartitioned())
+									{
+										Subsystem->RegisterOrUpdatePCGComponent(PCGComponent);
+									}
+									
 									PCGComponent->Generate();
 								}
 							}
diff --git a/Engine/Plugins/Experimental/PCG/Source/PCG/Private/Elements/PCGStaticMeshSpawner.cpp b/Engine/Plugins/Experimental/PCG/Source/PCG/Private/Elements/PCGStaticMeshSpawner.cpp
--- a/Engine/Plugins/Experimental/PCG/Source/PCG/Private/Elements/PCGStaticMeshSpawner.cpp
+++ b/Engine/Plugins/Experimental/PCG/Source/PCG/Private/Elements/PCGStaticMeshSpawner.cpp
@@ -2,6 +2,8 @@
 
 #include "Elements/PCGStaticMeshSpawner.h"
 
+#include "PCGComponent.h"
+#include "PCGManagedResource.h"
 #include "Data/PCGPointData.h"
 #include "Data/PCGSpatialData.h"
 #include "Helpers/PCGActorHelpers.h"
@@ -16,6 +18,11 @@
 
 #include UE_INLINE_GENERATED_CPP_BY_NAME(PCGStaticMeshSpawner)
 
+static TAutoConsoleVariable<bool> CVarAllowISMReuse(
+	TEXT("pcg.ISM.AllowReuse"),
+	false,
+	TEXT("Controls whether ISMs can be reused and skipped when re-executing"));
+
 UPCGStaticMeshSpawnerSettings::UPCGStaticMeshSpawnerSettings(const FObjectInitializer &ObjectInitializer)
 {
 	bUseSeed = true;
@@ -49,9 +56,10 @@ bool FPCGStaticMeshSpawnerElement::PrepareDataInternal(FPCGContext* InContext) c
 		return true;
 	}
 
-	// perform mesh selection
-	TArray<FPCGTaggedData> Inputs = Context->InputData.GetInputs();
-	TArray<FPCGTaggedData>& Outputs = Context->OutputData.TaggedData;
+	if (!Context->SourceComponent.Get())
+	{
+		return true;
+	}
 
 #if WITH_EDITOR
 	// In editor, we always want to generate this data for inspection & to prevent caching issues
@@ -60,6 +68,53 @@ bool FPCGStaticMeshSpawnerElement::PrepareDataInternal(FPCGContext* InContext) c
 	const bool bGenerateOutput = Context->Node && Context->Node->IsOutputPinConnected(PCGPinConstants::DefaultOutputLabel);
 #endif
 
+	// Check if we can reuse existing resources
+	bool bSkippedDueToReuse = false;
+
+	if (CVarAllowISMReuse.GetValueOnAnyThread())
+	{
+		// Compute CRC if it has not been computed (it likely isn't, but this is to futureproof this)
+		if (!Context->DependenciesCrc.IsValid())
+		{
+			GetDependenciesCrc(Context->InputData, Settings, Context->SourceComponent.Get(), Context->DependenciesCrc);
+		}
+		
+		if (Context->DependenciesCrc.IsValid())
+		{
+			TArray<UPCGManagedISMComponent*> MISMCs;
+			Context->SourceComponent->ForEachManagedResource([&MISMCs, &Context](UPCGManagedResource* InResource)
+			{
+				if (UPCGManagedISMComponent* Resource = Cast<UPCGManagedISMComponent>(InResource))
+				{
+					if (Resource->GetCrc().IsValid() && Resource->GetCrc() == Context->DependenciesCrc)
+					{
+						MISMCs.Add(Resource);
+					}
+				}
+			});
+
+			for (UPCGManagedISMComponent* MISMC : MISMCs)
+			{
+				MISMC->MarkAsReused();
+			}
+
+			if (!MISMCs.IsEmpty())
+			{
+				bSkippedDueToReuse = true;
+			}
+		}
+	}
+
+	// Early out - if we've established we could reuse resources and there is no need to generate an output, quit now
+	if (!bGenerateOutput && bSkippedDueToReuse)
+	{
+		return true;
+	}
+
+	// perform mesh selection
+	TArray<FPCGTaggedData> Inputs = Context->InputData.GetInputs();
+	TArray<FPCGTaggedData>& Outputs = Context->OutputData.TaggedData;
+
 	for (const FPCGTaggedData& Input : Inputs)
 	{
 		const UPCGSpatialData* SpatialData = Cast<UPCGSpatialData>(Input.Data);
@@ -76,7 +131,7 @@ bool FPCGStaticMeshSpawnerElement::PrepareDataInternal(FPCGContext* InContext) c
 			continue;
 		}
 
-		AActor* TargetActor = PointData->TargetActor.Get();
+		AActor* TargetActor = Context->GetTargetActor(PointData);
 		if (!TargetActor)
 		{
 			PCGE_LOG(Error, "Invalid target actor");
@@ -106,6 +161,12 @@ bool FPCGStaticMeshSpawnerElement::PrepareDataInternal(FPCGContext* InContext) c
 		TArray<FPCGMeshInstanceList> MeshInstances;
 		Settings->MeshSelectorInstance->SelectInstances(*Context, Settings, PointData, MeshInstances, OutputPointData);
 
+		// If we need the output but would otherwise skip the resource creation, just don't push them to the MeshInstancesData array
+		if (bSkippedDueToReuse)
+		{
+			continue;
+		}
+
 		TArray<FPCGPackedCustomData> PackedCustomData;
 		PackedCustomData.SetNum(MeshInstances.Num());
 		if (Settings->InstancePackerInstance)
@@ -117,6 +178,7 @@ bool FPCGStaticMeshSpawnerElement::PrepareDataInternal(FPCGContext* InContext) c
 		}
 
 		FPCGStaticMeshSpawnerContext::FPackedInstanceListData& InstanceListData = Context->MeshInstancesData.Emplace_GetRef();
+		InstanceListData.TargetActor = TargetActor;
 		InstanceListData.SpatialData = PointData;
 		InstanceListData.MeshInstances = MoveTemp(MeshInstances);
 		InstanceListData.PackedCustomData = MoveTemp(PackedCustomData);
@@ -137,13 +199,13 @@ bool FPCGStaticMeshSpawnerElement::ExecuteInternal(FPCGContext* InContext) const
 		const FPCGStaticMeshSpawnerContext::FPackedInstanceListData& InstanceList = Context->MeshInstancesData.Last();
 		check(InstanceList.MeshInstances.Num() == InstanceList.PackedCustomData.Num());
 
-		const bool bTargetActorValid = InstanceList.SpatialData->TargetActor.IsValid();
+		const bool bTargetActorValid = (InstanceList.TargetActor && IsValid(InstanceList.TargetActor));
 
 		if (bTargetActorValid)
 		{
 			while (Context->CurrentDataIndex < InstanceList.MeshInstances.Num())
 			{
-				SpawnStaticMeshInstances(Context, InstanceList.MeshInstances[Context->CurrentDataIndex], InstanceList.SpatialData->TargetActor.Get(), InstanceList.PackedCustomData[Context->CurrentDataIndex]);
+				SpawnStaticMeshInstances(Context, InstanceList.MeshInstances[Context->CurrentDataIndex], InstanceList.TargetActor, InstanceList.PackedCustomData[Context->CurrentDataIndex]);
 				++Context->CurrentDataIndex;
 
 				if (Context->ShouldStop())
@@ -230,7 +292,13 @@ void FPCGStaticMeshSpawnerElement::SpawnStaticMeshInstances(FPCGContext* Context
 		Params.Mobility = SceneComponent->Mobility;
 	}
 
-	UInstancedStaticMeshComponent* ISMC = UPCGActorHelpers::GetOrCreateISMC(TargetActor, Context->SourceComponent.Get(), Params);
+	UPCGManagedISMComponent* MISMC = UPCGActorHelpers::GetOrCreateManagedISMC(TargetActor, Context->SourceComponent.Get(), Params);
+	check(MISMC);
+
+	MISMC->SetCrc(Context->DependenciesCrc);
+
+	UInstancedStaticMeshComponent* ISMC = MISMC->GetComponent();
+	check(ISMC);
 
 	const int32 PreExistingInstanceCount = ISMC->GetInstanceCount();
 	const int32 NewInstanceCount = InstanceList.Instances.Num();
diff --git a/Engine/Plugins/Experimental/PCG/Source/PCG/Private/Helpers/PCGActorHelpers.cpp b/Engine/Plugins/Experimental/PCG/Source/PCG/Private/Helpers/PCGActorHelpers.cpp
--- a/Engine/Plugins/Experimental/PCG/Source/PCG/Private/Helpers/PCGActorHelpers.cpp
+++ b/Engine/Plugins/Experimental/PCG/Source/PCG/Private/Helpers/PCGActorHelpers.cpp
@@ -104,6 +104,19 @@ UInstancedStaticMeshComponent* UPCGActorHelpers::GetOrCreateISMC(AActor* InTarge
 }
 
 UInstancedStaticMeshComponent* UPCGActorHelpers::GetOrCreateISMC(AActor* InTargetActor, UPCGComponent* InSourceComponent, const FPCGISMCBuilderParameters& InParams)
+{
+	UPCGManagedISMComponent* MISMC = GetOrCreateManagedISMC(InTargetActor, InSourceComponent, InParams);
+	if (MISMC)
+	{
+		return MISMC->GetComponent();
+	}
+	else
+	{
+		return nullptr;
+	}
+}
+
+UPCGManagedISMComponent* UPCGActorHelpers::GetOrCreateManagedISMC(AActor* InTargetActor, UPCGComponent* InSourceComponent, const FPCGISMCBuilderParameters& InParams)
 {
 	UStaticMesh* InMesh = InParams.Mesh;
 	const TArray<UMaterialInterface*>& InMaterials = InParams.MaterialOverrides;
@@ -167,9 +180,8 @@ UInstancedStaticMeshComponent* UPCGActorHelpers::GetOrCreateISMC(AActor* InTarge
 
 			if (bMaterialsMatched)
 			{
-				// In this case, we make sure to mark the managed resource as used
 				MISMC->MarkAsUsed();
-				return ISMC;
+				return MISMC;
 			}
 		}
 	}
@@ -240,7 +252,7 @@ UInstancedStaticMeshComponent* UPCGActorHelpers::GetOrCreateISMC(AActor* InTarge
 	Resource->GeneratedComponent = ISMC;
 	InSourceComponent->AddToManagedResources(Resource);
 
-	return ISMC;
+	return Resource;
 }
 
 bool UPCGActorHelpers::DeleteActors(UWorld* World, const TArray<TSoftObjectPtr<AActor>>& ActorsToDelete)
diff --git a/Engine/Plugins/Experimental/PCG/Source/PCG/Private/PCGContext.cpp b/Engine/Plugins/Experimental/PCG/Source/PCG/Private/PCGContext.cpp
--- a/Engine/Plugins/Experimental/PCG/Source/PCG/Private/PCGContext.cpp
+++ b/Engine/Plugins/Experimental/PCG/Source/PCG/Private/PCGContext.cpp
@@ -5,6 +5,7 @@
 #include "PCGHelpers.h"
 #include "PCGParamData.h"
 #include "PCGPin.h"
+#include "Data/PCGSpatialData.h"
 #include "Metadata/PCGAttributePropertySelector.h"
 #include "Metadata/PCGMetadata.h"
 #include "Metadata/PCGMetadataAttribute.h"
@@ -78,6 +79,18 @@ bool FPCGContext::ShouldStop() const
 	return FPlatformTime::Seconds() > EndTime;
 }
 
+AActor* FPCGContext::GetTargetActor(const UPCGSpatialData* InSpatialData) const
+{
+	if (SourceComponent.IsValid() && SourceComponent->GetOwner())
+	{
+		return SourceComponent->GetOwner();
+	}
+	else
+	{
+		return InSpatialData ? InSpatialData->TargetActor.Get() : nullptr;
+	}
+}
+
 const UPCGSettingsInterface* FPCGContext::GetInputSettingsInterface() const
 {
 	if (Node)
diff --git a/Engine/Plugins/Experimental/PCG/Source/PCG/Public/Elements/PCGStaticMeshSpawner.h b/Engine/Plugins/Experimental/PCG/Source/PCG/Public/Elements/PCGStaticMeshSpawner.h
--- a/Engine/Plugins/Experimental/PCG/Source/PCG/Public/Elements/PCGStaticMeshSpawner.h
+++ b/Engine/Plugins/Experimental/PCG/Source/PCG/Public/Elements/PCGStaticMeshSpawner.h
@@ -107,7 +107,8 @@ struct FPCGStaticMeshSpawnerContext : public FPCGContext
 	{
 		FPackedInstanceListData();
 		~FPackedInstanceListData();
-		const UPCGSpatialData* SpatialData;
+		AActor* TargetActor = nullptr;
+		const UPCGSpatialData* SpatialData = nullptr;
 		TArray<FPCGMeshInstanceList> MeshInstances;
 		TArray<FPCGPackedCustomData> PackedCustomData;
 	};
diff --git a/Engine/Plugins/Experimental/PCG/Source/PCG/Public/Helpers/PCGActorHelpers.h b/Engine/Plugins/Experimental/PCG/Source/PCG/Public/Helpers/PCGActorHelpers.h
--- a/Engine/Plugins/Experimental/PCG/Source/PCG/Public/Helpers/PCGActorHelpers.h
+++ b/Engine/Plugins/Experimental/PCG/Source/PCG/Public/Helpers/PCGActorHelpers.h
@@ -15,6 +15,7 @@ class UInstancedStaticMeshComponent;
 class ULevel;
 class UMaterialInterface;
 class UPCGComponent;
+class UPCGManagedISMComponent;
 class UStaticMesh;
 class UWorld;
 struct FISMComponentDescriptor;
@@ -40,6 +41,7 @@ class PCG_API UPCGActorHelpers : public UBlueprintFunctionLibrary
 public:
 	static UInstancedStaticMeshComponent* GetOrCreateISMC(AActor* InTargetActor, UPCGComponent* InSourceComponent, const FISMComponentDescriptor& InISMCDescriptor);
 	static UInstancedStaticMeshComponent* GetOrCreateISMC(AActor* InTargetActor, UPCGComponent* SourceComponent, const FPCGISMCBuilderParameters& Params);
+	static UPCGManagedISMComponent* GetOrCreateManagedISMC(AActor* InTargetActor, UPCGComponent* SourceComponent, const FPCGISMCBuilderParameters& Params);
 	static bool DeleteActors(UWorld* World, const TArray<TSoftObjectPtr<AActor>>& ActorsToDelete);
 
 	/**
diff --git a/Engine/Plugins/Experimental/PCG/Source/PCG/Public/PCGContext.h b/Engine/Plugins/Experimental/PCG/Source/PCG/Public/PCGContext.h
--- a/Engine/Plugins/Experimental/PCG/Source/PCG/Public/PCGContext.h
+++ b/Engine/Plugins/Experimental/PCG/Source/PCG/Public/PCGContext.h
@@ -9,6 +9,7 @@
 
 class UPCGComponent;
 class UPCGSettingsInterface;
+class UPCGSpatialData;
 struct FPCGGraphCache;
 struct FPCGSettingsOverridableParam;
 
@@ -92,6 +93,8 @@ struct PCG_API FPCGContext
 	FString GetComponentName() const;
 	bool ShouldStop() const;
 
+	AActor* GetTargetActor(const UPCGSpatialData* InSpatialData) const;
+
 protected:
 	virtual UObject* GetExternalContainerForOverridableParam(const FPCGSettingsOverridableParam& InParam) { return nullptr; }
 
diff --git a/Engine/Plugins/Experimental/PCG/Source/PCG/Public/PCGCrc.h b/Engine/Plugins/Experimental/PCG/Source/PCG/Public/PCGCrc.h
--- a/Engine/Plugins/Experimental/PCG/Source/PCG/Public/PCGCrc.h
+++ b/Engine/Plugins/Experimental/PCG/Source/PCG/Public/PCGCrc.h
@@ -6,9 +6,14 @@
 #include "Templates/TypeHash.h"
 #include "UObject/NameTypes.h"
 
+#include "PCGCrc.generated.h"
+
 /** Crc with valid flag and helper functionality. */
+USTRUCT()
 struct PCG_API FPCGCrc
 {
+	GENERATED_BODY()
+
 	/** Creates an invalid Crc. */
 	FPCGCrc() = default;
 
@@ -67,7 +72,9 @@ struct PCG_API FPCGCrc
 
 private:
 	/** Crc32 value. */
+	UPROPERTY(VisibleAnywhere, Category = Crc)
 	uint32 Value = 0;
 
+	UPROPERTY(VisibleAnywhere, Category = Crc)
 	bool bValid = false;
 };
diff --git a/Engine/Plugins/Experimental/PCG/Source/PCG/Public/PCGManagedResource.h b/Engine/Plugins/Experimental/PCG/Source/PCG/Public/PCGManagedResource.h
--- a/Engine/Plugins/Experimental/PCG/Source/PCG/Public/PCGManagedResource.h
+++ b/Engine/Plugins/Experimental/PCG/Source/PCG/Public/PCGManagedResource.h
@@ -2,6 +2,7 @@
 
 #pragma once
 
+#include "PCGCrc.h"
 
 #include "PCGManagedResource.generated.h"
 
@@ -29,6 +30,7 @@ public:
 	virtual bool ReleaseIfUnused(TSet<TSoftObjectPtr<AActor>>& OutActorsToDelete);
 
 	virtual void MarkAsUsed() { bIsMarkedUnused = false; }
+	virtual void MarkAsReused() { ensure(bIsMarkedUnused); bIsMarkedUnused = false; }
 	bool IsMarkedUnused() const { return bIsMarkedUnused; }
 
 	/** Move the given resource to a new actor. Return true if it has succeeded */
@@ -36,7 +38,13 @@ public:
 
 	static bool DebugForcePurgeAllResourcesOnGenerate();
 
+	const FPCGCrc& GetCrc() const { return Crc; }
+	void SetCrc(const FPCGCrc& InCrc) { Crc = InCrc; }
+
 protected:
+	UPROPERTY(VisibleAnywhere, Category = GeneratedData)
+	FPCGCrc Crc;
+
 	UPROPERTY(Transient, VisibleAnywhere, Category = GeneratedData)
 	bool bIsMarkedUnused = false;
 };
