Added option to get distance vector instead of distance scalar only

diff --git a/Engine/Plugins/Experimental/PCG/Source/PCG/Private/Elements/PCGDistance.cpp b/Engine/Plugins/Experimental/PCG/Source/PCG/Private/Elements/PCGDistance.cpp
--- a/Engine/Plugins/Experimental/PCG/Source/PCG/Private/Elements/PCGDistance.cpp
+++ b/Engine/Plugins/Experimental/PCG/Source/PCG/Private/Elements/PCGDistance.cpp
@@ -19,16 +19,16 @@ namespace PCGDistance
 			FVector Dir = TargetCenter - SourceCenter;
 			Dir.Normalize();
 
-			return SourceCenter + Dir * ((SourcePoint.BoundsMax-SourcePoint.BoundsMin) * SourcePoint.Transform.GetScale3D()).Length() * 0.5;
+			return SourceCenter + Dir * SourcePoint.GetScaledExtents().Length();
 		}
 		else if (Shape == PCGDistanceShape::BoxBounds)
 		{
-			const FVector LocalTargetCenter = SourcePoint.Transform.Inverse().TransformPosition(TargetCenter);
+			const FVector LocalTargetCenter = SourcePoint.Transform.InverseTransformPosition(TargetCenter);
 
 			const double DistanceSquared = ComputeSquaredDistanceFromBoxToPoint(SourcePoint.BoundsMin, SourcePoint.BoundsMax, LocalTargetCenter);
 
 			FVector Dir = -LocalTargetCenter;
-			Dir.Normalize();			
+			Dir.Normalize();
 
 			const FVector LocalClosestPoint = LocalTargetCenter + Dir * FMath::Sqrt(DistanceSquared);
 
@@ -71,6 +71,7 @@ bool FPCGDistanceElement::ExecuteInternal(FPCGContext* Context) const
 
 	const FName AttributeName = Settings->AttributeName;
 	const bool bSetDensity = Settings->bSetDensity;
+	const bool bOutputDistanceVector = Settings->bOutputDistanceVector;
 	const double MaximumDistance = Settings->MaximumDistance;
 	const PCGDistanceShape SourceShape = Settings->SourceShape;
 	const PCGDistanceShape TargetShape = Settings->TargetShape;
@@ -129,10 +130,21 @@ bool FPCGDistanceElement::ExecuteInternal(FPCGContext* Context) const
 		OutputData->InitializeFromData(SourcePointData);
 		Outputs.Add_GetRef(Source).Data = OutputData;
 
-		FPCGMetadataAttribute<float>* Attribute = (AttributeName != NAME_None) ? OutputData->Metadata->FindOrCreateAttribute<float>(AttributeName, 0.0f) : nullptr;
+		FPCGMetadataAttribute<float>* ScalarAttribute = nullptr;
+		FPCGMetadataAttribute<FVector>* VectorAttribute = nullptr;
+
+		if (AttributeName != NAME_None && !bOutputDistanceVector)
+		{
+			ScalarAttribute = OutputData->Metadata->FindOrCreateAttribute<float>(AttributeName, 0.0f);
+		}
+
+		if (AttributeName != NAME_None && bOutputDistanceVector)
+		{
+			VectorAttribute = OutputData->Metadata->FindOrCreateAttribute<FVector>(AttributeName, FVector::ZeroVector);
+		}
 
 		FPCGAsync::AsyncPointProcessing(Context, SourcePointData->GetPoints(), OutputData->GetMutablePoints(),
-			[OutputData, SourceShape, TargetShape, &TargetPointDatas, MaximumDistance, Attribute, bSetDensity](const FPCGPoint& SourcePoint, FPCGPoint& OutPoint) {
+			[OutputData, SourceShape, TargetShape, &TargetPointDatas, MaximumDistance, ScalarAttribute, VectorAttribute, bSetDensity](const FPCGPoint& SourcePoint, FPCGPoint& OutPoint) {
 
 				OutPoint = SourcePoint;
 
@@ -140,13 +152,14 @@ bool FPCGDistanceElement::ExecuteInternal(FPCGContext* Context) const
 
 				const FVector SourceCenter = SourcePoint.Transform.TransformPosition(SourcePoint.GetLocalCenter());
 
-				double DistanceSquared = MaximumDistance*MaximumDistance;
+				double MinDistanceSquared = MaximumDistance*MaximumDistance;
+				FVector MinDistanceVector = FVector::ZeroVector;
 
-				auto CalculateSDF = [&DistanceSquared, &SourcePoint, SourceCenter, SourceShape, TargetShape](const FPCGPointRef& TargetPointRef)
+				auto CalculateSDF = [&MinDistanceSquared, &MinDistanceVector, &SourcePoint, SourceCenter, SourceShape, TargetShape](const FPCGPointRef& TargetPointRef)
 				{
 					const FPCGPoint& TargetPoint = *TargetPointRef.Point;
 
-					const FVector TargetCenter = TargetPoint.Transform.TransformPosition(TargetPoint.GetLocalCenter());
+					const FVector& TargetCenter = TargetPointRef.Bounds.Origin;
 
 					const FVector SourceShapePos = PCGDistance::CalcPosition(SourceShape, SourcePoint, TargetPoint, SourceCenter, TargetCenter);
 					const FVector TargetShapePos = PCGDistance::CalcPosition(TargetShape, TargetPoint, SourcePoint, TargetCenter, SourceCenter);
@@ -154,10 +167,14 @@ bool FPCGDistanceElement::ExecuteInternal(FPCGContext* Context) const
 					const FVector ToTargetShapeDir = TargetShapePos - SourceShapePos;
 					const FVector ToTargetCenterDir = TargetCenter - SourceCenter;
 
-					const double Sign = ToTargetShapeDir.Dot(ToTargetCenterDir) > 0 ? 1.0 : -1.0;
+					const double Sign = FMath::Sign(ToTargetShapeDir.Dot(ToTargetCenterDir));
 					const double ThisDistanceSquared = ToTargetShapeDir.SquaredLength() * Sign;
 
-					DistanceSquared = FMath::Min(DistanceSquared, ThisDistanceSquared);
+					if (ThisDistanceSquared < MinDistanceSquared)
+					{
+						MinDistanceSquared = ThisDistanceSquared;
+						MinDistanceVector = ToTargetShapeDir;
+					}
 				};
 
 				for (const UPCGPointData* TargetPointData : TargetPointDatas)
@@ -170,12 +187,21 @@ bool FPCGDistanceElement::ExecuteInternal(FPCGContext* Context) const
 					);
 				}
 
-				const float Distance = FMath::Sqrt(FMath::Abs(DistanceSquared)) * (DistanceSquared < 0 ? -1.0 : 1.0);
+				const float Distance = FMath::Sign(MinDistanceSquared) * FMath::Sqrt(FMath::Abs(MinDistanceSquared));
 
-				if (Attribute)
+				if (ScalarAttribute || VectorAttribute)
 				{
 					OutputData->Metadata->InitializeOnSet(OutPoint.MetadataEntry);
-					Attribute->SetValue(OutPoint.MetadataEntry, Distance);
+				}
+
+				if (ScalarAttribute)
+				{
+					ScalarAttribute->SetValue(OutPoint.MetadataEntry, Distance);
+				}
+
+				if (VectorAttribute)
+				{
+					VectorAttribute->SetValue(OutPoint.MetadataEntry, MinDistanceVector);
 				}
 				
 				if (bSetDensity)
@@ -186,7 +212,7 @@ bool FPCGDistanceElement::ExecuteInternal(FPCGContext* Context) const
 
 				return true;
 			}
-		);		
+		);
 	}
 
 	return true;
diff --git a/Engine/Plugins/Experimental/PCG/Source/PCG/Public/Elements/PCGDistance.h b/Engine/Plugins/Experimental/PCG/Source/PCG/Public/Elements/PCGDistance.h
--- a/Engine/Plugins/Experimental/PCG/Source/PCG/Public/Elements/PCGDistance.h
+++ b/Engine/Plugins/Experimental/PCG/Source/PCG/Public/Elements/PCGDistance.h
@@ -44,23 +44,27 @@ protected:
 
 public:
 
-	// The name of the attribute to store on the point. Use 'None' to disable
+	/** The name of the attribute to store on the point.Use 'None' to disable */
 	UPROPERTY(BlueprintReadWrite, EditAnywhere, Category = Settings, meta = (PCG_Overridable))
 	FName AttributeName = TEXT("Distance");
 
-	// If true, will also set the density to be 0-1 based on MaximumDistance
+	/** Controls whether the attribute will be a scalar or a vector */
+	UPROPERTY(BlueprintReadWrite, EditAnywhere, Category = Settings, meta = (PCG_Overridable))
+	bool bOutputDistanceVector = false;
+
+	/** If true, will also set the density to be 0 - 1 based on MaximumDistance */
 	UPROPERTY(BlueprintReadWrite, EditAnywhere, Category = Settings, meta = (PCG_Overridable))
 	bool bSetDensity = false;
 
-	// A maximum distance to search, which is used as an optimization
+	/** A maximum distance to search, which is used as an optimization */
 	UPROPERTY(BlueprintReadWrite, EditAnywhere, Category = Settings, meta = (ClampMin = "1", PCG_Overridable))
 	double MaximumDistance = 20000.0;
 
-	// What shape is used on the 'source' points
+	/** What shape is used on the 'source' points */
 	UPROPERTY(BlueprintReadWrite, EditAnywhere, Category = Settings, meta = (PCG_Overridable))
 	PCGDistanceShape SourceShape = PCGDistanceShape::SphereBounds;
 
-	// What shape is used on the 'target' points
+	/** What shape is used on the 'target' points */
 	UPROPERTY(BlueprintReadWrite, EditAnywhere, Category = Settings, meta = (PCG_Overridable))
 	PCGDistanceShape TargetShape = PCGDistanceShape::SphereBounds;
 };
diff --git a/Engine/Plugins/Experimental/PCG/Source/PCG/Public/PCGPoint.h b/Engine/Plugins/Experimental/PCG/Source/PCG/Public/PCGPoint.h
--- a/Engine/Plugins/Experimental/PCG/Source/PCG/Public/PCGPoint.h
+++ b/Engine/Plugins/Experimental/PCG/Source/PCG/Public/PCGPoint.h
@@ -69,6 +69,8 @@ public:
 		BoundsMax = Center + InExtents;
 	}
 
+	FVector GetScaledExtents() const { return GetExtents() * Transform.GetScale3D(); }
+
 	FVector GetLocalCenter() const { return (BoundsMax + BoundsMin) / 2.0; }
 	void SetLocalCenter(const FVector& InCenter)
 	{
