[PCG] Fixed is partitioned issue, since deprecating was not re-saving it back.

#rnx

diff --git a/Engine/Plugins/Experimental/PCG/Source/PCG/Private/PCGComponent.cpp b/Engine/Plugins/Experimental/PCG/Source/PCG/Private/PCGComponent.cpp
--- a/Engine/Plugins/Experimental/PCG/Source/PCG/Private/PCGComponent.cpp
+++ b/Engine/Plugins/Experimental/PCG/Source/PCG/Private/PCGComponent.cpp
@@ -935,13 +935,13 @@ void UPCGComponent::PostLoad()
 		ExclusionTags_DEPRECATED.Reset();
 	}
 
-	// If we have both default value (bIsComponentPartitioned = false and bIsPartitioned_DEPRECATED = true)
-	// we will follow the value of bIsPartitioned_DEPRECATED.
-	// bIsPartitioned_DEPRECATED will be set to false to new objects
-	if (!bIsComponentPartitioned && bIsPartitioned_DEPRECATED)
+	// If we have both default value (bIsComponentPartitioned = false and bIsPartitioned = true)
+	// we will follow the value of bIsPartitioned.
+	// bIsPartitioned will be set to false to new objects
+	if (!bIsComponentPartitioned && bIsPartitioned)
 	{
-		bIsComponentPartitioned = bIsPartitioned_DEPRECATED;
-		bIsPartitioned_DEPRECATED = false;
+		bIsComponentPartitioned = bIsPartitioned;
+		bIsPartitioned = false;
 	}
 
 	/** Deprecation code, should be removed once generated data has been updated */
@@ -1048,10 +1048,10 @@ void UPCGComponent::PostInitProperties()
 	const bool bIsNewObject = CurrentInspectedObject && !CurrentInspectedObject->HasAnyFlags(RF_ClassDefaultObject | RF_NeedLoad | RF_NeedPostLoad);
 
 #if WITH_EDITOR
-	// Force bIsPartitioned_DEPRECATED at False for new objects
+	// Force bIsPartitioned at false for new objects
 	if (bIsNewObject)
 	{
-		bIsPartitioned_DEPRECATED = false;
+		bIsPartitioned = false;
 	}
 #endif // WITH_EDITOR
 
diff --git a/Engine/Plugins/Experimental/PCG/Source/PCG/Public/PCGComponent.h b/Engine/Plugins/Experimental/PCG/Source/PCG/Public/PCGComponent.h
--- a/Engine/Plugins/Experimental/PCG/Source/PCG/Public/PCGComponent.h
+++ b/Engine/Plugins/Experimental/PCG/Source/PCG/Public/PCGComponent.h
@@ -172,10 +172,6 @@ public:
 #if WITH_EDITORONLY_DATA
 	UPROPERTY()
 	TArray<FName> ExclusionTags_DEPRECATED;
-
-	// Note for upgrade: can be safely replaced by bIsComponentPartitioned. Needed a new variable to change the default value.
-	UPROPERTY()
-	bool bIsPartitioned_DEPRECATED = true;
 #endif
 
 	UPROPERTY(BlueprintReadWrite, EditAnywhere, Category = Properties)
@@ -271,6 +267,11 @@ protected:
 #if WITH_EDITORONLY_DATA
 	UPROPERTY()
 	TObjectPtr<UPCGGraph> Graph_DEPRECATED;
+
+private:
+	// Note for upgrade: can be safely replaced by bIsComponentPartitioned. Needed a new variable to change the default value. Kept to allow proper value change, cannot be deprecated. Do not use.
+	UPROPERTY()
+	bool bIsPartitioned = true;
 #endif // WITH_EDITORONLY_DATA
 
 private:
