Improved trivial metadata cases in the CopyPoints implementation
Fixed another temporary actor causing PCG refresh that should have been ignored

diff --git a/Engine/Plugins/Experimental/PCG/Source/PCG/Private/Elements/PCGCopyPoints.cpp b/Engine/Plugins/Experimental/PCG/Source/PCG/Private/Elements/PCGCopyPoints.cpp
--- a/Engine/Plugins/Experimental/PCG/Source/PCG/Private/Elements/PCGCopyPoints.cpp
+++ b/Engine/Plugins/Experimental/PCG/Source/PCG/Private/Elements/PCGCopyPoints.cpp
@@ -84,15 +84,22 @@ bool FPCGCopyPointsElement::ExecuteInternal(FPCGContext* Context) const
 	Output.Data = OutPointData;
 
 	// RootMetadata will be parent to the ouptut metadata, while NonRootMetadata will carry attributes from the input not selected for inheritance
+	// Note that this is a preference, as we can and should pick more efficiently in the trivial cases
 	const UPCGMetadata* RootMetadata = nullptr;
 	const UPCGMetadata* NonRootMetadata = nullptr;
-	if (Settings->AttributeInheritance == EPCGCopyPointsMetadataInheritanceMode::Source)
+
+	const bool bSourceHasMetadata = (SourcePointData->Metadata->GetAttributeCount() > 0 && SourcePointData->Metadata->GetItemCountForChild() > 0);
+	const bool bTargetHasMetadata = (TargetPointData->Metadata->GetAttributeCount() > 0 && TargetPointData->Metadata->GetItemCountForChild() > 0);
+	const bool bInheritMetadataFromSource = (!bTargetHasMetadata || (AttributeInheritance == EPCGCopyPointsMetadataInheritanceMode::Source && bSourceHasMetadata));
+	const bool bProcessMetadata = (bSourceHasMetadata || bTargetHasMetadata);
+
+	if (bInheritMetadataFromSource)
 	{
 		OutPointData->InitializeFromData(SourcePointData);
 		RootMetadata = SourcePointData->Metadata;
 		NonRootMetadata = TargetPointData->Metadata;
 	}
-	else // if (Settings->AttributeInheritance == EPCGCopyPointsMetadataInheritanceMode::Target)
+	else
 	{
 		OutPointData->InitializeFromData(TargetPointData);
 		RootMetadata = TargetPointData->Metadata;
@@ -104,25 +111,31 @@ bool FPCGCopyPointsElement::ExecuteInternal(FPCGContext* Context) const
 
 	check(OutPointData->Metadata && NonRootMetadata);
 
-	TArray<FName> AttributeNames;
-	TArray<EPCGMetadataTypes> AttributeTypes;
-	NonRootMetadata->GetAttributes(AttributeNames, AttributeTypes);
-
 	TArray<FPCGMetadataAttributeBase*> NonRootAttributes;
-	for (const FName& AttributeName : AttributeNames)
+	TArray<TTuple<int64, int64>> AllMetadataEntries;
+
+	if (bProcessMetadata)
 	{
-		if (!OutPointData->Metadata->HasAttribute(AttributeName))
+		// Prepare the attributes from the non-root that we'll need to use to copy values over
+		TArray<FName> AttributeNames;
+		TArray<EPCGMetadataTypes> AttributeTypes;
+		NonRootMetadata->GetAttributes(AttributeNames, AttributeTypes);
+	
+		for (const FName& AttributeName : AttributeNames)
 		{
-			const FPCGMetadataAttributeBase* Attribute = NonRootMetadata->GetConstAttribute(AttributeName);
-			if (FPCGMetadataAttributeBase* NewAttribute = OutPointData->Metadata->CopyAttribute(Attribute, AttributeName, /*bKeepRoot=*/false, /*bCopyEntries=*/false, /*bCopyValues=*/true))
+			if (!OutPointData->Metadata->HasAttribute(AttributeName))
 			{
-				NonRootAttributes.Add(NewAttribute);
+				const FPCGMetadataAttributeBase* Attribute = NonRootMetadata->GetConstAttribute(AttributeName);
+				if (FPCGMetadataAttributeBase* NewAttribute = OutPointData->Metadata->CopyAttribute(Attribute, AttributeName, /*bKeepRoot=*/false, /*bCopyEntries=*/false, /*bCopyValues=*/true))
+				{
+					NonRootAttributes.Add(NewAttribute);
+				}
 			}
 		}
-	}
 
-	TArray<TTuple<int64, int64>> AllMetadataEntries;
-	AllMetadataEntries.SetNum(SourcePoints.Num() * TargetPoints.Num());
+		// Preallocate the metadata entries array if we're going to use it
+		AllMetadataEntries.SetNumUninitialized(SourcePoints.Num() * TargetPoints.Num());
+	}
 
 	// Use implicit capture, since we capture a lot
 	FPCGAsync::AsyncPointProcessing(Context, SourcePoints.Num() * TargetPoints.Num(), OutPoints, [&](int32 Index, FPCGPoint& OutPoint)
@@ -185,35 +198,34 @@ bool FPCGCopyPointsElement::ExecuteInternal(FPCGContext* Context) const
 			OutPoint.Seed = TargetPoint.Seed;
 		}
 
-		const FPCGPoint* RootPoint = nullptr;
-		const FPCGPoint* NonRootPoint = nullptr;
-		if (Settings->AttributeInheritance == EPCGCopyPointsMetadataInheritanceMode::Source)
+		if (bProcessMetadata)
 		{
-			RootPoint = &SourcePoint;
-			NonRootPoint = &TargetPoint;
-		}
-		else // if (Settings->AttributeInheritance == EPCGCopyPointsMetadataInheritanceMode::Target)
-		{
-			RootPoint = &TargetPoint;
-			NonRootPoint = &SourcePoint;
-		}
+			const FPCGPoint* RootPoint = (bInheritMetadataFromSource ? &SourcePoint : &TargetPoint);
+			const FPCGPoint* NonRootPoint = (bInheritMetadataFromSource ? &TargetPoint : &SourcePoint);
 
-		OutPoint.MetadataEntry = OutPointData->Metadata->AddEntryPlaceholder();
-		AllMetadataEntries[Index] = TTuple<int64, int64>(OutPoint.MetadataEntry, RootPoint->MetadataEntry);
+			OutPoint.MetadataEntry = OutPointData->Metadata->AddEntryPlaceholder();
+			AllMetadataEntries[Index] = TTuple<int64, int64>(OutPoint.MetadataEntry, RootPoint->MetadataEntry);
 
-		// Copy EntryToValue key mappings from NonRootAttributes
-		for (FPCGMetadataAttributeBase* NonRootAttribute : NonRootAttributes)
-		{
-			const FPCGMetadataAttributeBase* Attribute = NonRootMetadata->GetConstAttribute(NonRootAttribute->Name);
-			check(Attribute);
-			const PCGMetadataValueKey ValueKey = Attribute->GetValueKey(NonRootPoint->MetadataEntry);
-			NonRootAttribute->SetValueFromValueKey(OutPoint.MetadataEntry, ValueKey);
+			// Copy EntryToValue key mappings from NonRootAttributes - no need to do it if the non-root uses the default values
+			if (NonRootPoint->MetadataEntry != PCGInvalidEntryKey)
+			{
+				for (FPCGMetadataAttributeBase* NonRootAttribute : NonRootAttributes)
+				{
+					const FPCGMetadataAttributeBase* Attribute = NonRootMetadata->GetConstAttribute(NonRootAttribute->Name);
+					check(Attribute);
+					const PCGMetadataValueKey ValueKey = Attribute->GetValueKey(NonRootPoint->MetadataEntry);
+					NonRootAttribute->SetValueFromValueKey(OutPoint.MetadataEntry, ValueKey);
+				}
+			}
 		}
 
 		return true;
 	});
 
-	OutPointData->Metadata->AddDelayedEntries(AllMetadataEntries);
+	if (bProcessMetadata)
+	{
+		OutPointData->Metadata->AddDelayedEntries(AllMetadataEntries);
+	}
 
 	return true;
 }
diff --git a/Engine/Plugins/Experimental/PCG/Source/PCG/Private/PCGComponent.cpp b/Engine/Plugins/Experimental/PCG/Source/PCG/Private/PCGComponent.cpp
--- a/Engine/Plugins/Experimental/PCG/Source/PCG/Private/PCGComponent.cpp
+++ b/Engine/Plugins/Experimental/PCG/Source/PCG/Private/PCGComponent.cpp
@@ -1615,7 +1615,7 @@ void UPCGComponent::OnActorChanged(AActor* Actor, UObject* InObject, bool bActor
 		DirtyGenerated(EPCGComponentDirtyFlag::Actor);
 		Refresh();
 	}
-	else if(Actor)
+	else if(Actor && !Actor->bIsEditorPreviewActor)
 	{
 		bool bDirtyAndRefresh = false;
 
