Tentative fix for non-unity builds

#rnx

diff --git a/Engine/Plugins/Experimental/PCG/Source/PCG/Public/Elements/PCGNormalToDensity.h b/Engine/Plugins/Experimental/PCG/Source/PCG/Public/Elements/PCGNormalToDensity.h
--- a/Engine/Plugins/Experimental/PCG/Source/PCG/Public/Elements/PCGNormalToDensity.h
+++ b/Engine/Plugins/Experimental/PCG/Source/PCG/Public/Elements/PCGNormalToDensity.h
@@ -1,6 +1,7 @@
 // Copyright Epic Games, Inc. All Rights Reserved.
 #pragma once
 
+#include "PCGPin.h"
 #include "PCGSettings.h"
 #include "Elements/PCGPointProcessingElementBase.h"
 
