[PCG] Proper overflow handling for surface/volume sampler point count.

These overflows have happened in practice.

We aim to have safety limits on computation in the future but those limits will also likely need a robust iteration count from this change.

Also while testing the repro I found the early out on invalid point extents does not happen earlier enough and the node can pass through the input surface. Moved this up.

diff --git a/Engine/Plugins/Experimental/PCG/Source/PCG/Private/Elements/PCGSurfaceSampler.cpp b/Engine/Plugins/Experimental/PCG/Source/PCG/Private/Elements/PCGSurfaceSampler.cpp
--- a/Engine/Plugins/Experimental/PCG/Source/PCG/Private/Elements/PCGSurfaceSampler.cpp
+++ b/Engine/Plugins/Experimental/PCG/Source/PCG/Private/Elements/PCGSurfaceSampler.cpp
@@ -64,17 +64,33 @@ namespace PCGSurfaceSampler
 		CellMinY = FMath::CeilToInt((InputBounds.Min.Y) / CellSize.Y);
 		CellMaxY = FMath::FloorToInt((InputBounds.Max.Y) / CellSize.Y);
 
-		if (CellMinX > CellMaxX || CellMinY > CellMaxY)
 		{
-			if (Context)
+			const int64 CellCountX = 1 + CellMaxX - CellMinX;
+			const int64 CellCountY = 1 + CellMaxY - CellMinY;
+			if (CellCountX <= 0 || CellCountY <= 0)
 			{
-				PCGE_LOG_C(Verbose, Context, "Skipped - invalid cell bounds");
+				if (Context)
+				{
+					PCGE_LOG_C(Verbose, Context, "Skipped - invalid cell bounds (%lld x %lld)", CellCountX, CellCountY);
+				}
+
+				return false;
 			}
-			
-			return false;
+
+			const int64 CellCount64 = CellCountX * CellCountY;
+			if (CellCount64 <= 0 || CellCount64 >= MAX_int32)
+			{
+				if (Context)
+				{
+					PCGE_LOG_C(Error, Context, "Skipped - tried to generate too many points (%lld)", CellCount64);
+				}
+
+				return false;
+			}
+
+			CellCount = static_cast<int32>(CellCount64);
 		}
 
-		CellCount = (1 + CellMaxX - CellMinX) * (1 + CellMaxY - CellMinY);
 		check(CellCount > 0);
 
 		const FVector::FReal InvSquaredMeterUnits = 1.0 / (100.0 * 100.0);
@@ -287,6 +303,15 @@ bool FPCGSurfaceSamplerElement::ExecuteInternal(FPCGContext* Context) const
 	const UPCGSurfaceSamplerSettings* Settings = Context->GetInputSettings<UPCGSurfaceSamplerSettings>();
 	check(Settings);
 	
+	// Early out on invalid settings
+	// TODO: we could compute an approximate radius based on the points per squared meters if that's useful
+	const FVector& PointExtents = Settings->PointExtents;
+	if (PointExtents.X <= 0 || PointExtents.Y <= 0)
+	{
+		PCGE_LOG(Warning, "Skipped - Invalid point extents");
+		return true;
+	}
+
 	TArray<FPCGTaggedData>& Outputs = Context->OutputData.TaggedData;
 
 	// Grab the Bounding Shape input if there is one.
@@ -366,15 +391,6 @@ bool FPCGSurfaceSamplerElement::ExecuteInternal(FPCGContext* Context) const
 		PCGE_LOG(Warning, "No Surface input was provided, and no surface could be found in the Bounding Shape input for sampling. Connect the surface to be sampled to the Surface input.");
 	}
 
-	// Early out on invalid settings
-	// TODO: we could compute an approximate radius based on the points per squared meters if that's useful
-	const FVector& PointExtents = Settings->PointExtents;
-	if(PointExtents.X <= 0 || PointExtents.Y <= 0)
-	{
-		PCGE_LOG(Warning, "Skipped - Invalid point extents");
-		return true;
-	}
-	
 	// TODO: embarassingly parallel loop
 	for (int GenerationIndex = 0; GenerationIndex < GeneratingShapes.Num(); ++GenerationIndex)
 	{
diff --git a/Engine/Plugins/Experimental/PCG/Source/PCG/Private/Elements/PCGVolumeSampler.cpp b/Engine/Plugins/Experimental/PCG/Source/PCG/Private/Elements/PCGVolumeSampler.cpp
--- a/Engine/Plugins/Experimental/PCG/Source/PCG/Private/Elements/PCGVolumeSampler.cpp
+++ b/Engine/Plugins/Experimental/PCG/Source/PCG/Private/Elements/PCGVolumeSampler.cpp
@@ -57,17 +57,40 @@ namespace PCGVolumeSampler
 		const int32 MinZ = FMath::CeilToInt(InBounds.Min.Z / VoxelSize.Z);
 		const int32 MaxZ = FMath::FloorToInt(InBounds.Max.Z / VoxelSize.Z);
 
-		if (MinX > MaxX || MinY > MaxY || MinZ > MaxZ)
+		// Set uninitialized, then carefully initialize step by step with overflow checks
+		int32 NumIterations = -1;
+
 		{
-			if (InContext)
+			const int64 NumX = MaxX - MinX;
+			const int64 NumY = MaxY - MinY;
+			const int64 NumZ = MaxZ - MinZ;
+			const int64 NumIterationsXY64 = NumX * NumY;
+			const int64 NumIterations64 = NumIterationsXY64 * NumZ;
+
+			if (NumX <= 0 || NumY <= 0 || NumZ <= 0)
 			{
-				PCGE_LOG_C(Verbose, InContext, "Skipped - invalid cell bounds");
+				if (InContext)
+				{
+					PCGE_LOG_C(Verbose, InContext, "Skipped - invalid cell bounds (%lld x %lld x %lld)", NumX, NumY, NumZ);
+				}
+
+				return;
 			}
 
-			return;
-		}
+			if (NumIterationsXY64 > 0 && NumIterationsXY64 < MAX_int32 && NumIterations64 > 0 && NumIterations64 < MAX_int32)
+			{
+				NumIterations = static_cast<int32>(NumIterations64);
+			}
+			else
+			{
+				if (InContext)
+				{
+					PCGE_LOG_C(Error, InContext, "Skipped - tried to generate too many points (%lld x %lld x %lld)", NumX, NumY, NumZ);
+				}
 
-		const int32 NumIterations = (MaxX - MinX) * (MaxY - MinY) * (MaxZ - MinZ);
+				return;
+			}
+		}
 
 		FPCGAsync::AsyncPointProcessing(InContext, NumIterations, Points, [InVolume, InBoundingShape, VoxelSize, MinX, MaxX, MinY, MaxY, MinZ, MaxZ](int32 Index, FPCGPoint& OutPoint)
 		{
