[PCG] Fix unitialized property in test struct

#rnx

diff --git a/Engine/Plugins/Experimental/PCG/Source/PCG/Public/Tests/Accessor/PCGAttributeExtractorTest.h b/Engine/Plugins/Experimental/PCG/Source/PCG/Public/Tests/Accessor/PCGAttributeExtractorTest.h
--- a/Engine/Plugins/Experimental/PCG/Source/PCG/Public/Tests/Accessor/PCGAttributeExtractorTest.h
+++ b/Engine/Plugins/Experimental/PCG/Source/PCG/Public/Tests/Accessor/PCGAttributeExtractorTest.h
@@ -34,7 +34,7 @@ struct FPCGAttributeExtractorTestStructDepth1
 	FPCGAttributeExtractorTestStructDepth2 Depth2Struct;
 
 	UPROPERTY()
-	float FloatValue;
+	float FloatValue = 0.0f;
 };
 
 USTRUCT(meta = (Hidden))
