[PCG] Move Graph property in PCG Component details view

#rnx

diff --git a/Engine/Plugins/Experimental/PCG/Source/PCG/Public/PCGComponent.h b/Engine/Plugins/Experimental/PCG/Source/PCG/Public/PCGComponent.h
--- a/Engine/Plugins/Experimental/PCG/Source/PCG/Public/PCGComponent.h
+++ b/Engine/Plugins/Experimental/PCG/Source/PCG/Public/PCGComponent.h
@@ -60,7 +60,7 @@ enum class EPCGComponentDirtyFlag : uint8
 };
 ENUM_CLASS_FLAGS(EPCGComponentDirtyFlag);
 
-UCLASS(BlueprintType, ClassGroup = (Procedural), meta = (BlueprintSpawnableComponent))
+UCLASS(BlueprintType, ClassGroup = (Procedural), meta = (BlueprintSpawnableComponent, PrioritizeCategories = "PCG"))
 class PCG_API UPCGComponent : public UActorComponent
 {
 	UPCGComponent(const FObjectInitializer& InObjectInitializer);
@@ -261,7 +261,7 @@ public:
 	static UPCGData* CreateActorPCGData(AActor* Actor, const UPCGComponent* Component, bool bParseActor = true);
 
 protected:
-	UPROPERTY(BlueprintReadOnly, VisibleAnywhere, Category = Properties, Instanced, meta = (NoResetToDefault, ShowOnlyInnerProperties))
+	UPROPERTY(BlueprintReadOnly, VisibleAnywhere, Category = PCG, Instanced, meta = (NoResetToDefault, ShowOnlyInnerProperties))
 	TObjectPtr<UPCGGraphInstance> GraphInstance;
 
 #if WITH_EDITORONLY_DATA
