[PCG] Use transformed bound to compute overlap ratio in VolumeOverlap

#rnx

diff --git a/Engine/Plugins/Experimental/PCG/Source/PCG/Private/Data/PCGPointData.cpp b/Engine/Plugins/Experimental/PCG/Source/PCG/Private/Data/PCGPointData.cpp
--- a/Engine/Plugins/Experimental/PCG/Source/PCG/Private/Data/PCGPointData.cpp
+++ b/Engine/Plugins/Experimental/PCG/Source/PCG/Private/Data/PCGPointData.cpp
@@ -98,7 +98,9 @@ namespace PCGPointHelpers
 		const FBox PointBounds = InPoint.GetLocalDensityBounds();
 		const FTransform& PointTransform = InPoint.Transform;
 
-		const FBox FirstOverlap = PointBounds.Overlap(InBounds.TransformBy(InTransform.GetRelativeTransform(PointTransform)));
+		const FBox InBoundsTransformed = InBounds.TransformBy(InTransform.GetRelativeTransform(PointTransform));
+
+		const FBox FirstOverlap = PointBounds.Overlap(InBoundsTransformed);
 		if (!FirstOverlap.IsValid)
 		{
 			return 0;
@@ -110,7 +112,7 @@ namespace PCGPointHelpers
 			return 0;
 		}
 		
-		return FMath::Min(ComputeOverlapRatio(FirstOverlap, InBounds), ComputeOverlapRatio(SecondOverlap, InBounds));
+		return FMath::Min(ComputeOverlapRatio(FirstOverlap, InBoundsTransformed), ComputeOverlapRatio(SecondOverlap, InBounds));
 	}
 
 	/** Helper function for additive blending of quaternions (copied from ControlRig) */
