Allocate point array with no initialization as all code paths do it in some form

diff --git a/Engine/Plugins/Experimental/PCG/Source/PCG/Private/Elements/PCGSurfaceSampler.cpp b/Engine/Plugins/Experimental/PCG/Source/PCG/Private/Elements/PCGSurfaceSampler.cpp
--- a/Engine/Plugins/Experimental/PCG/Source/PCG/Private/Elements/PCGSurfaceSampler.cpp
+++ b/Engine/Plugins/Experimental/PCG/Source/PCG/Private/Elements/PCGSurfaceSampler.cpp
@@ -178,6 +178,9 @@ namespace PCGSurfaceSampler
 			const FVector TentativeLocation = FVector(CurrentX + RandX * InnerCellSize.X, CurrentY + RandY * InnerCellSize.Y, SampleZ);
 			const FBox LocalBound(-LoopData.PointExtents, LoopData.PointExtents);
 
+			// The output at this point is not initialized
+			OutPoint = FPCGPoint();
+
 			// Firstly project onto elected generating shape to move to final position.
 			if (!InSurface->ProjectPoint(FTransform(TentativeLocation), LocalBound, ProjectionParams, OutPoint, SampledData->Metadata))
 			{
diff --git a/Engine/Plugins/Experimental/PCG/Source/PCG/Private/Elements/PCGVolumeSampler.cpp b/Engine/Plugins/Experimental/PCG/Source/PCG/Private/Elements/PCGVolumeSampler.cpp
--- a/Engine/Plugins/Experimental/PCG/Source/PCG/Private/Elements/PCGVolumeSampler.cpp
+++ b/Engine/Plugins/Experimental/PCG/Source/PCG/Private/Elements/PCGVolumeSampler.cpp
@@ -102,6 +102,10 @@ namespace PCGVolumeSampler
 			const FBox VoxelBox(VoxelSize * -0.5, VoxelSize * 0.5);
 
 			const FTransform SampleTransform = FTransform(SampleLocation);
+
+			// The OutPoint has not been initialized, so do it now
+			OutPoint = FPCGPoint();
+
 			if (InVolume->SamplePoint(SampleTransform, VoxelBox, OutPoint, nullptr))
 			{
 				if (InBoundingShape)
diff --git a/Engine/Plugins/Experimental/PCG/Source/PCG/Private/Helpers/PCGAsync.cpp b/Engine/Plugins/Experimental/PCG/Source/PCG/Private/Helpers/PCGAsync.cpp
--- a/Engine/Plugins/Experimental/PCG/Source/PCG/Private/Helpers/PCGAsync.cpp
+++ b/Engine/Plugins/Experimental/PCG/Source/PCG/Private/Helpers/PCGAsync.cpp
@@ -26,7 +26,7 @@ namespace FPCGAsync
 		{
 			TRACE_CPUPROFILER_EVENT_SCOPE(IPCGElement::AsyncPointProcessing::AllocatingArray);
 			// Pre-reserve the out points array
-			OutPoints.SetNum(NumIterations);
+			OutPoints.SetNumUninitialized(NumIterations);
 		}
 
 		// Setup [current, last, nb points] data per dispatch
@@ -186,8 +186,8 @@ namespace FPCGAsync
 		{
 			TRACE_CPUPROFILER_EVENT_SCOPE(IPCGElement::AsyncPointProcessing::AllocatingArray);
 			// Pre-reserve the out points array
-			InFilterPoints.SetNum(NumIterations);
-			OutFilterPoints.SetNum(NumIterations);
+			InFilterPoints.SetNumUninitialized(NumIterations);
+			OutFilterPoints.SetNumUninitialized(NumIterations);
 		}
 
 		auto IterationInnerLoop = [&PointFunc, &InFilterPoints, &OutFilterPoints](int32 StartIndex, int32 EndIndex) -> TPair<int32, int32>
