Added more flexibility in how metadata is processed in the CopyPoints native node, which is a good lever for optimizations

diff --git a/Engine/Plugins/Experimental/PCG/Source/PCG/Private/Elements/PCGCopyPoints.cpp b/Engine/Plugins/Experimental/PCG/Source/PCG/Private/Elements/PCGCopyPoints.cpp
--- a/Engine/Plugins/Experimental/PCG/Source/PCG/Private/Elements/PCGCopyPoints.cpp
+++ b/Engine/Plugins/Experimental/PCG/Source/PCG/Private/Elements/PCGCopyPoints.cpp
@@ -8,8 +8,19 @@
 #include "Data/PCGSpatialData.h"
 #include "Helpers/PCGAsync.h"
 
+#include "Async/ParallelFor.h"
+
 #include UE_INLINE_GENERATED_CPP_BY_NAME(PCGCopyPoints)
 
+#define LOCTEXT_NAMESPACE "PCGCopyPointsElement"
+
+#if WITH_EDITOR
+FText UPCGCopyPointsSettings::GetNodeTooltipText() const
+{
+	return LOCTEXT("NodeTooltip", "For each point pair from the source and the target, create a copy, inheriting properties & attributes depending on the node settings.");
+}
+#endif
+
 TArray<FPCGPinProperties> UPCGCopyPointsSettings::InputPinProperties() const
 {
 	TArray<FPCGPinProperties> PinProperties;
@@ -90,47 +101,95 @@ bool FPCGCopyPointsElement::ExecuteInternal(FPCGContext* Context) const
 
 	const bool bSourceHasMetadata = (SourcePointData->Metadata->GetAttributeCount() > 0 && SourcePointData->Metadata->GetItemCountForChild() > 0);
 	const bool bTargetHasMetadata = (TargetPointData->Metadata->GetAttributeCount() > 0 && TargetPointData->Metadata->GetItemCountForChild() > 0);
-	const bool bInheritMetadataFromSource = (!bTargetHasMetadata || (AttributeInheritance == EPCGCopyPointsMetadataInheritanceMode::Source && bSourceHasMetadata));
-	const bool bProcessMetadata = (bSourceHasMetadata || bTargetHasMetadata);
 
-	if (bInheritMetadataFromSource)
+	bool bInheritMetadataFromSource = true;
+	bool bProcessMetadata = (bSourceHasMetadata || bTargetHasMetadata);
+
+	if (AttributeInheritance == EPCGCopyPointsMetadataInheritanceMode::SourceOnly)
+	{
+		bInheritMetadataFromSource = true;
+		bProcessMetadata = bSourceHasMetadata;
+
+		OutPointData->InitializeFromData(SourcePointData);
+		RootMetadata = SourcePointData->Metadata;
+		NonRootMetadata = nullptr;
+	}
+	else if (AttributeInheritance == EPCGCopyPointsMetadataInheritanceMode::TargetOnly)
+	{
+		bInheritMetadataFromSource = false;
+		bProcessMetadata = bTargetHasMetadata;
+
+		OutPointData->InitializeFromData(TargetPointData);
+
+		RootMetadata = TargetPointData->Metadata;
+		NonRootMetadata = nullptr;
+	}
+	else if (AttributeInheritance == EPCGCopyPointsMetadataInheritanceMode::SourceFirst)
 	{
+		bInheritMetadataFromSource = bSourceHasMetadata || !bTargetHasMetadata;
+
 		OutPointData->InitializeFromData(SourcePointData);
 		RootMetadata = SourcePointData->Metadata;
 		NonRootMetadata = TargetPointData->Metadata;
 	}
-	else
+	else if (AttributeInheritance == EPCGCopyPointsMetadataInheritanceMode::TargetFirst)
 	{
+		bInheritMetadataFromSource = !bTargetHasMetadata;
+
 		OutPointData->InitializeFromData(TargetPointData);
+
 		RootMetadata = TargetPointData->Metadata;
 		NonRootMetadata = SourcePointData->Metadata;
 	}
+	else // None
+	{
+		OutPointData->InitializeFromData(SourcePointData, nullptr, /*bInheritMetadata=*/false, /*bInheritAttributes=*/false);
+
+		bProcessMetadata = false;
+		RootMetadata = NonRootMetadata = nullptr;
+	}
 
 	// Priorize use the target actor from the target, irrespective of the source
 	OutPointData->TargetActor = TargetPointData->TargetActor.IsValid() ? TargetPointData->TargetActor : SourcePointData->TargetActor;
 
-	check(OutPointData->Metadata && NonRootMetadata);
+	check(OutPointData->Metadata);
 
-	TArray<FPCGMetadataAttributeBase*> NonRootAttributes;
+	TArray<FPCGMetadataAttributeBase*> AttributesToSet;
+	TArray<const FPCGMetadataAttributeBase*> NonRootAttributes;
 	TArray<TTuple<int64, int64>> AllMetadataEntries;
+	TArray<TArray<TTuple<PCGMetadataEntryKey, PCGMetadataValueKey>>> AttributeValuesToSet;
 
 	if (bProcessMetadata)
 	{
-		// Prepare the attributes from the non-root that we'll need to use to copy values over
-		TArray<FName> AttributeNames;
-		TArray<EPCGMetadataTypes> AttributeTypes;
-		NonRootMetadata->GetAttributes(AttributeNames, AttributeTypes);
-	
-		for (const FName& AttributeName : AttributeNames)
+		TRACE_CPUPROFILER_EVENT_SCOPE(FPCGCopyPointsElement::Execute::SetupMetadata);
+		if (NonRootMetadata)
 		{
-			if (!OutPointData->Metadata->HasAttribute(AttributeName))
+			// Prepare the attributes from the non-root that we'll need to use to copy values over
+			TArray<FName> AttributeNames;
+			TArray<EPCGMetadataTypes> AttributeTypes;
+			NonRootMetadata->GetAttributes(AttributeNames, AttributeTypes);
+
+			for (const FName& AttributeName : AttributeNames)
 			{
-				const FPCGMetadataAttributeBase* Attribute = NonRootMetadata->GetConstAttribute(AttributeName);
-				if (FPCGMetadataAttributeBase* NewAttribute = OutPointData->Metadata->CopyAttribute(Attribute, AttributeName, /*bKeepRoot=*/false, /*bCopyEntries=*/false, /*bCopyValues=*/true))
+				if (!OutPointData->Metadata->HasAttribute(AttributeName))
 				{
-					NonRootAttributes.Add(NewAttribute);
+					const FPCGMetadataAttributeBase* Attribute = NonRootMetadata->GetConstAttribute(AttributeName);
+					if (FPCGMetadataAttributeBase* NewAttribute = OutPointData->Metadata->CopyAttribute(Attribute, AttributeName, /*bKeepRoot=*/false, /*bCopyEntries=*/false, /*bCopyValues=*/true))
+					{
+						AttributesToSet.Add(NewAttribute);
+						NonRootAttributes.Add(Attribute);
+					}
 				}
 			}
+
+			// Considering writing to the attribute value requires a lock, we'll gather the value keys to write
+			// and do it on a 1-thread-per-attribute basis at the end
+			AttributeValuesToSet.SetNum(AttributesToSet.Num());
+
+			for (TArray<TTuple<PCGMetadataEntryKey, PCGMetadataValueKey>>& AttributeValues : AttributeValuesToSet)
+			{
+				AttributeValues.SetNumUninitialized(SourcePoints.Num() * TargetPoints.Num());
+			}
 		}
 
 		// Preallocate the metadata entries array if we're going to use it
@@ -201,20 +260,28 @@ bool FPCGCopyPointsElement::ExecuteInternal(FPCGContext* Context) const
 		if (bProcessMetadata)
 		{
 			const FPCGPoint* RootPoint = (bInheritMetadataFromSource ? &SourcePoint : &TargetPoint);
-			const FPCGPoint* NonRootPoint = (bInheritMetadataFromSource ? &TargetPoint : &SourcePoint);
 
 			OutPoint.MetadataEntry = OutPointData->Metadata->AddEntryPlaceholder();
 			AllMetadataEntries[Index] = TTuple<int64, int64>(OutPoint.MetadataEntry, RootPoint->MetadataEntry);
 
-			// Copy EntryToValue key mappings from NonRootAttributes - no need to do it if the non-root uses the default values
-			if (NonRootPoint->MetadataEntry != PCGInvalidEntryKey)
+			if (NonRootMetadata)
 			{
-				for (FPCGMetadataAttributeBase* NonRootAttribute : NonRootAttributes)
+				const FPCGPoint* NonRootPoint = (bInheritMetadataFromSource ? &TargetPoint : &SourcePoint);
+
+				// Copy EntryToValue key mappings from NonRootAttributes - no need to do it if the non-root uses the default values
+				if (NonRootPoint->MetadataEntry != PCGInvalidEntryKey)
 				{
-					const FPCGMetadataAttributeBase* Attribute = NonRootMetadata->GetConstAttribute(NonRootAttribute->Name);
-					check(Attribute);
-					const PCGMetadataValueKey ValueKey = Attribute->GetValueKey(NonRootPoint->MetadataEntry);
-					NonRootAttribute->SetValueFromValueKey(OutPoint.MetadataEntry, ValueKey);
+					for (int32 AttributeIndex = 0; AttributeIndex < NonRootAttributes.Num(); ++AttributeIndex)
+					{
+						AttributeValuesToSet[AttributeIndex][Index] = TTuple<PCGMetadataEntryKey, PCGMetadataValueKey>(OutPoint.MetadataEntry, NonRootAttributes[AttributeIndex]->GetValueKey(NonRootPoint->MetadataEntry));
+					}
+				}
+				else
+				{
+					for (int32 AttributeIndex = 0; AttributeIndex < NonRootAttributes.Num(); ++AttributeIndex)
+					{
+						AttributeValuesToSet[AttributeIndex][Index] = TTuple<PCGMetadataEntryKey, PCGMetadataValueKey>(OutPoint.MetadataEntry, PCGDefaultValueKey);
+					}
 				}
 			}
 		}
@@ -224,8 +291,32 @@ bool FPCGCopyPointsElement::ExecuteInternal(FPCGContext* Context) const
 
 	if (bProcessMetadata)
 	{
+		TRACE_CPUPROFILER_EVENT_SCOPE(FPCGCopyPointsElement::Execute::SetMetadata);
+		check(AttributesToSet.Num() == AttributeValuesToSet.Num());
+		if (AttributesToSet.Num() > 0)
+		{
+			int32 AttributeOffset = 0;
+			const int32 AttributePerDispatch = FMath::Max(1, Context->NumAvailableTasks);
+
+			while (AttributeOffset < AttributesToSet.Num())
+			{
+				const int32 AttributeCountInCurrentDispatch = FMath::Min(AttributePerDispatch, AttributesToSet.Num() - AttributeOffset);
+				ParallelFor(AttributeCountInCurrentDispatch, [AttributeOffset, &AttributesToSet, &AttributeValuesToSet](int32 WorkerIndex)
+				{
+					FPCGMetadataAttributeBase* Attribute = AttributesToSet[AttributeOffset + WorkerIndex];
+					const TArray<TTuple<PCGMetadataEntryKey, PCGMetadataValueKey>>& Values = AttributeValuesToSet[AttributeOffset + WorkerIndex];
+					check(Attribute);
+					Attribute->SetValuesFromValueKeys(Values, /*bResetValueOnDefaultValueKey*/false); // no need for the reset here, our points will not have any prior value for these attributes
+				});
+
+				AttributeOffset += AttributeCountInCurrentDispatch;
+			}
+		}
+
 		OutPointData->Metadata->AddDelayedEntries(AllMetadataEntries);
 	}
 
 	return true;
 }
+
+#undef LOCTEXT_NAMESPACE
diff --git a/Engine/Plugins/Experimental/PCG/Source/PCG/Private/Metadata/PCGMetadataAttribute.cpp b/Engine/Plugins/Experimental/PCG/Source/PCG/Private/Metadata/PCGMetadataAttribute.cpp
--- a/Engine/Plugins/Experimental/PCG/Source/PCG/Private/Metadata/PCGMetadataAttribute.cpp
+++ b/Engine/Plugins/Experimental/PCG/Source/PCG/Private/Metadata/PCGMetadataAttribute.cpp
@@ -41,6 +41,31 @@ void FPCGMetadataAttributeBase::SetValueFromValueKey(PCGMetadataEntryKey EntryKe
 	EntryToValueKeyMap.FindOrAdd(EntryKey) = ValueKey;
 }
 
+void FPCGMetadataAttributeBase::SetValuesFromValueKeys(const TArray<TTuple<PCGMetadataEntryKey, PCGMetadataValueKey>>& EntryValuePairs, bool bResetValueOnDefaultValueKey)
+{
+	if (EntryValuePairs.IsEmpty())
+	{
+		return;
+	}
+
+	FWriteScopeLock ScopeLock(EntryMapLock);
+	for (const TTuple<PCGMetadataEntryKey, PCGMetadataValueKey>& EntryValuePair : EntryValuePairs)
+	{
+		check(EntryValuePair.Key != PCGInvalidEntryKey);
+		if (EntryValuePair.Value == PCGDefaultValueKey)
+		{
+			if (bResetValueOnDefaultValueKey)
+			{
+				EntryToValueKeyMap.Remove(EntryValuePair.Key);
+			}
+		}
+		else
+		{
+			EntryToValueKeyMap.FindOrAdd(EntryValuePair.Key, EntryValuePair.Value);
+		}
+	}
+}
+
 PCGMetadataValueKey FPCGMetadataAttributeBase::GetValueKey(PCGMetadataEntryKey EntryKey) const
 {
 	if (EntryKey == PCGInvalidEntryKey)
diff --git a/Engine/Plugins/Experimental/PCG/Source/PCG/Private/Tests/Elements/PCGCopyPointsTest.cpp b/Engine/Plugins/Experimental/PCG/Source/PCG/Private/Tests/Elements/PCGCopyPointsTest.cpp
--- a/Engine/Plugins/Experimental/PCG/Source/PCG/Private/Tests/Elements/PCGCopyPointsTest.cpp
+++ b/Engine/Plugins/Experimental/PCG/Source/PCG/Private/Tests/Elements/PCGCopyPointsTest.cpp
@@ -134,19 +134,19 @@ bool FPCGCopyPointsTest::RunTest(const FString& Parameters)
 		EPCGCopyPointsMetadataInheritanceMode RootMode;
 		EPCGCopyPointsMetadataInheritanceMode NonRootMode;
 
-		if (Settings->AttributeInheritance == EPCGCopyPointsMetadataInheritanceMode::Source)
+		if (Settings->AttributeInheritance == EPCGCopyPointsMetadataInheritanceMode::SourceFirst)
 		{
 			RootMetadata = SourcePointData->Metadata;
 			NonRootMetadata = TargetPointData->Metadata;
-			RootMode = EPCGCopyPointsMetadataInheritanceMode::Source;
-			NonRootMode = EPCGCopyPointsMetadataInheritanceMode::Target;
+			RootMode = EPCGCopyPointsMetadataInheritanceMode::SourceFirst;
+			NonRootMode = EPCGCopyPointsMetadataInheritanceMode::TargetFirst;
 		}
-		else // if (Settings->AttributeInheritance == EPCGCopyPointsMetadataInheritanceMode::Target)
+		else // if (Settings->AttributeInheritance == EPCGCopyPointsMetadataInheritanceMode::TargetFirst)
 		{
 			RootMetadata = TargetPointData->Metadata;
 			NonRootMetadata = SourcePointData->Metadata;
-			RootMode = EPCGCopyPointsMetadataInheritanceMode::Target;
-			NonRootMode = EPCGCopyPointsMetadataInheritanceMode::Source;
+			RootMode = EPCGCopyPointsMetadataInheritanceMode::TargetFirst;
+			NonRootMode = EPCGCopyPointsMetadataInheritanceMode::SourceFirst;
 		}
 
 		if (!TestTrue("Valid input metadata", RootMetadata && NonRootMetadata))
@@ -295,7 +295,7 @@ bool FPCGCopyPointsTest::RunTest(const FString& Parameters)
 				const FPCGMetadataAttribute<float>* OutAttribute = OutAttributes[AttributeIndex];
 				check(InheritedAttribute && OutAttribute);
 
-				const PCGMetadataEntryKey EntryKey = (InheritedAttributeModes[AttributeIndex] == EPCGCopyPointsMetadataInheritanceMode::Source) ? SourcePoint.MetadataEntry : TargetPoint.MetadataEntry;
+				const PCGMetadataEntryKey EntryKey = (InheritedAttributeModes[AttributeIndex] == EPCGCopyPointsMetadataInheritanceMode::SourceFirst) ? SourcePoint.MetadataEntry : TargetPoint.MetadataEntry;
 				bTestPassed &= TestEqual("Valid metadata value", InheritedAttribute->GetValueKey(EntryKey), OutAttribute->GetValueKey(OutPoint.MetadataEntry));
 			}
 		}
@@ -317,7 +317,7 @@ bool FPCGCopyPointsTest::RunTest(const FString& Parameters)
 	Settings->ScaleInheritance = EPCGCopyPointsInheritanceMode::Source;
 	Settings->ColorInheritance = EPCGCopyPointsInheritanceMode::Source;
 	Settings->SeedInheritance = EPCGCopyPointsInheritanceMode::Source;
-	Settings->AttributeInheritance = EPCGCopyPointsMetadataInheritanceMode::Source;
+	Settings->AttributeInheritance = EPCGCopyPointsMetadataInheritanceMode::SourceFirst;
 	bTestPassed &= ValidateCopyPointsPoints();
 
 	// Test 3 - inherit from target points
@@ -325,7 +325,7 @@ bool FPCGCopyPointsTest::RunTest(const FString& Parameters)
 	Settings->ScaleInheritance = EPCGCopyPointsInheritanceMode::Target;
 	Settings->ColorInheritance = EPCGCopyPointsInheritanceMode::Target;
 	Settings->SeedInheritance = EPCGCopyPointsInheritanceMode::Target;
-	Settings->AttributeInheritance = EPCGCopyPointsMetadataInheritanceMode::Target;
+	Settings->AttributeInheritance = EPCGCopyPointsMetadataInheritanceMode::TargetFirst;
 	bTestPassed &= ValidateCopyPointsPoints();
 
 	// Test 4 - empty source data
diff --git a/Engine/Plugins/Experimental/PCG/Source/PCG/Public/Elements/PCGCopyPoints.h b/Engine/Plugins/Experimental/PCG/Source/PCG/Public/Elements/PCGCopyPoints.h
--- a/Engine/Plugins/Experimental/PCG/Source/PCG/Public/Elements/PCGCopyPoints.h
+++ b/Engine/Plugins/Experimental/PCG/Source/PCG/Public/Elements/PCGCopyPoints.h
@@ -24,8 +24,11 @@ enum class EPCGCopyPointsInheritanceMode : uint8
 UENUM()
 enum class EPCGCopyPointsMetadataInheritanceMode : uint8
 {
-	Source,
-	Target
+	SourceFirst UMETA(Tooltip = "Points will inherit from source metadata and apply only unique attributes from target."),
+	TargetFirst UMETA(Tooltip = "Points will inherit from target metadata and apply only unique attributes from source."),
+	SourceOnly  UMETA(Tooltip = "Points will inherit metadata only from the source."),
+	TargetOnly  UMETA(Tooltip = "Points will inherit metadata only from the target."),
+	None        UMETA(Tooltip = "Points will have no metadata.")
 };
 
 UCLASS(BlueprintType, ClassGroup = (Procedural))
@@ -37,6 +40,7 @@ public:
 	//~Begin UPCGSettings interface
 #if WITH_EDITOR
 	virtual FName GetDefaultNodeName() const override { return FName(TEXT("Copy Points")); }
+	virtual FText GetNodeTooltipText() const override;
 	virtual EPCGSettingsType GetType() const override { return EPCGSettingsType::Sampler; }
 #endif
 
@@ -66,7 +70,7 @@ public:
 
 	/** The method used to determine output data attributes */
 	UPROPERTY(BlueprintReadWrite, EditAnywhere, Category = Settings, meta = (PCG_Overridable))
-	EPCGCopyPointsMetadataInheritanceMode AttributeInheritance = EPCGCopyPointsMetadataInheritanceMode::Source;
+	EPCGCopyPointsMetadataInheritanceMode AttributeInheritance = EPCGCopyPointsMetadataInheritanceMode::SourceFirst;
 };
 
 class FPCGCopyPointsElement : public FSimplePCGElement
diff --git a/Engine/Plugins/Experimental/PCG/Source/PCG/Public/Metadata/PCGMetadataAttribute.h b/Engine/Plugins/Experimental/PCG/Source/PCG/Public/Metadata/PCGMetadataAttribute.h
--- a/Engine/Plugins/Experimental/PCG/Source/PCG/Public/Metadata/PCGMetadataAttribute.h
+++ b/Engine/Plugins/Experimental/PCG/Source/PCG/Public/Metadata/PCGMetadataAttribute.h
@@ -37,6 +37,7 @@ public:
 	virtual bool AreValuesEqual(PCGMetadataValueKey ValueKey1, PCGMetadataValueKey ValueKey2) const = 0;
 
 	void SetValueFromValueKey(PCGMetadataEntryKey EntryKey, PCGMetadataValueKey ValueKey);
+	void SetValuesFromValueKeys(const TArray<TTuple<PCGMetadataEntryKey, PCGMetadataValueKey>>& EntryValuePairs, bool bResetValueOnDefaultValueKey = true);
 	PCGMetadataValueKey GetValueKey(PCGMetadataEntryKey EntryKey) const;
 	bool HasNonDefaultValue(PCGMetadataEntryKey EntryKey) const;
 	void ClearEntries();
