Turned on ISM reuse by default
Rejected interactive changes in settings to prevent issues

diff --git a/Engine/Plugins/Experimental/PCG/Source/PCG/Private/Elements/PCGStaticMeshSpawner.cpp b/Engine/Plugins/Experimental/PCG/Source/PCG/Private/Elements/PCGStaticMeshSpawner.cpp
--- a/Engine/Plugins/Experimental/PCG/Source/PCG/Private/Elements/PCGStaticMeshSpawner.cpp
+++ b/Engine/Plugins/Experimental/PCG/Source/PCG/Private/Elements/PCGStaticMeshSpawner.cpp
@@ -20,7 +20,7 @@
 
 static TAutoConsoleVariable<bool> CVarAllowISMReuse(
 	TEXT("pcg.ISM.AllowReuse"),
-	false,
+	true,
 	TEXT("Controls whether ISMs can be reused and skipped when re-executing"));
 
 UPCGStaticMeshSpawnerSettings::UPCGStaticMeshSpawnerSettings(const FObjectInitializer &ObjectInitializer)
diff --git a/Engine/Plugins/Experimental/PCG/Source/PCG/Private/PCGSettings.cpp b/Engine/Plugins/Experimental/PCG/Source/PCG/Private/PCGSettings.cpp
--- a/Engine/Plugins/Experimental/PCG/Source/PCG/Private/PCGSettings.cpp
+++ b/Engine/Plugins/Experimental/PCG/Source/PCG/Private/PCGSettings.cpp
@@ -372,6 +372,11 @@ void UPCGSettings::PostEditChangeProperty(FPropertyChangedEvent& PropertyChanged
 {
 	Super::PostEditChangeProperty(PropertyChangedEvent);
 
+	if (PropertyChangedEvent.ChangeType == EPropertyChangeType::Interactive)
+	{
+		return;
+	}
+
 	if (PropertyChangedEvent.Property && (PropertyChangedEvent.Property->GetFName() == GET_MEMBER_NAME_CHECKED(UPCGSettings, CachedOverridableParams)))
 	{
 		// Need to rebuild properties, if it ever changes.
