Undo //UE5/Release-5.2/Engine/Plugins/Experimental/PCG/Source/PCG/Private/Data/... changelist 24268456

Causing issues with Opal, reverting and will investigate further

#rnx

diff --git a/Engine/Plugins/Experimental/PCG/Source/PCG/Private/Data/PCGPointData.cpp b/Engine/Plugins/Experimental/PCG/Source/PCG/Private/Data/PCGPointData.cpp
--- a/Engine/Plugins/Experimental/PCG/Source/PCG/Private/Data/PCGPointData.cpp
+++ b/Engine/Plugins/Experimental/PCG/Source/PCG/Private/Data/PCGPointData.cpp
@@ -98,9 +98,7 @@ namespace PCGPointHelpers
 		const FBox PointBounds = InPoint.GetLocalDensityBounds();
 		const FTransform& PointTransform = InPoint.Transform;
 
-		const FBox InBoundsTransformed = InBounds.TransformBy(InTransform.GetRelativeTransform(PointTransform));
-
-		const FBox FirstOverlap = PointBounds.Overlap(InBoundsTransformed);
+		const FBox FirstOverlap = PointBounds.Overlap(InBounds.TransformBy(InTransform.GetRelativeTransform(PointTransform)));
 		if (!FirstOverlap.IsValid)
 		{
 			return 0;
@@ -112,7 +110,7 @@ namespace PCGPointHelpers
 			return 0;
 		}
 		
-		return FMath::Min(ComputeOverlapRatio(FirstOverlap, InBoundsTransformed), ComputeOverlapRatio(SecondOverlap, InBounds));
+		return FMath::Min(ComputeOverlapRatio(FirstOverlap, InBounds), ComputeOverlapRatio(SecondOverlap, InBounds));
 	}
 
 	/** Helper function for additive blending of quaternions (copied from ControlRig) */
