[PCG] Fixed a crash when adding node comments

diff --git a/Engine/Plugins/Experimental/PCG/Source/PCGEditor/Private/PCGEditorGraphNodeBase.cpp b/Engine/Plugins/Experimental/PCG/Source/PCGEditor/Private/PCGEditorGraphNodeBase.cpp
--- a/Engine/Plugins/Experimental/PCG/Source/PCGEditor/Private/PCGEditorGraphNodeBase.cpp
+++ b/Engine/Plugins/Experimental/PCG/Source/PCGEditor/Private/PCGEditorGraphNodeBase.cpp
@@ -62,8 +62,6 @@ void UPCGEditorGraphNodeBase::PostTransacted(const FTransactionObjectEvent& Tran
 	{
 		UpdatePosition();
 	}
-
-	ReconstructNode();
 }
 
 void UPCGEditorGraphNodeBase::GetNodeContextMenuActions(UToolMenu* Menu, class UGraphNodeContextMenuContext* Context) const
