[PCG] Moved PCGreroute.cpp to private

diff --git a/Engine/Plugins/Experimental/PCG/Source/PCG/Private/Elements/PCGReroute.cpp b/Engine/Plugins/Experimental/PCG/Source/PCG/Private/Elements/PCGReroute.cpp
new file mode 100644
--- /dev/null
+++ b/Engine/Plugins/Experimental/PCG/Source/PCG/Private/Elements/PCGReroute.cpp
@@ -0,0 +1,80 @@
+// Copyright Epic Games, Inc. All Rights Reserved.
+
+#include "Elements/PCGReroute.h"
+
+#include "PCGContext.h"
+#include "PCGEdge.h"
+#include "PCGNode.h"
+#include "PCGPin.h"
+
+UPCGRerouteSettings::UPCGRerouteSettings()
+{
+#if WITH_EDITORONLY_DATA
+	bExposeToLibrary = false;
+#endif
+}
+
+TArray<FPCGPinProperties> UPCGRerouteSettings::InputPinProperties() const
+{
+	FPCGPinProperties PinProperties;
+	PinProperties.Label = PCGPinConstants::DefaultInputLabel;
+	PinProperties.bAllowMultipleConnections = false;
+	PinProperties.AllowedTypes = EPCGDataType::Any;
+	
+	if (UPCGNode* PCGNode = Cast<UPCGNode>(GetOuter()))
+	{
+		if (UPCGPin* InputPin = PCGNode->GetInputPin(PCGPinConstants::DefaultInputLabel))
+		{
+			if (InputPin->EdgeCount() > 0)
+			{
+				if (UPCGEdge* Edge = InputPin->Edges[0])
+				{
+					if (const UPCGPin* OtherOutputPin = Edge->GetOtherPin(InputPin))
+					{
+						PinProperties.AllowedTypes = OtherOutputPin->Properties.AllowedTypes;
+					}
+				}
+			}
+		}
+	}
+
+	return {PinProperties };
+}
+
+TArray<FPCGPinProperties> UPCGRerouteSettings::OutputPinProperties() const
+{
+	FPCGPinProperties PinProperties;
+	PinProperties.Label = PCGPinConstants::DefaultOutputLabel;
+	PinProperties.bAllowMultipleConnections = true;
+	PinProperties.AllowedTypes = EPCGDataType::Any;
+	
+	if (UPCGNode* PCGNode = Cast<UPCGNode>(GetOuter()))
+	{
+		if (UPCGPin* InputPin = PCGNode->GetInputPin(PCGPinConstants::DefaultInputLabel))
+		{
+			if (InputPin->EdgeCount() > 0)
+			{
+				if (UPCGEdge* Edge = InputPin->Edges[0])
+				{
+					if (const UPCGPin* OtherOutputPin = Edge->GetOtherPin(InputPin))
+					{
+						PinProperties.AllowedTypes = OtherOutputPin->Properties.AllowedTypes;
+					}
+				}
+			}
+		}
+	}
+
+	return { PinProperties };
+}
+
+FPCGElementPtr UPCGRerouteSettings::CreateElement() const
+{
+	return MakeShared<FPCGRerouteElement>();
+}
+
+bool FPCGRerouteElement::ExecuteInternal(FPCGContext* Context) const
+{
+	Context->OutputData = Context->InputData;
+	return true;
+}
diff --git a/Engine/Plugins/Experimental/PCG/Source/PCG/Public/Elements/PCGReroute.cpp b/Engine/Plugins/Experimental/PCG/Source/PCG/Public/Elements/PCGReroute.cpp
deleted file mode 100644
