Fixed world ray hit/ world volumetric data type
Added tooltip

diff --git a/Engine/Plugins/Experimental/PCG/Source/PCG/Private/Elements/PCGWorldQuery.cpp b/Engine/Plugins/Experimental/PCG/Source/PCG/Private/Elements/PCGWorldQuery.cpp
--- a/Engine/Plugins/Experimental/PCG/Source/PCG/Private/Elements/PCGWorldQuery.cpp
+++ b/Engine/Plugins/Experimental/PCG/Source/PCG/Private/Elements/PCGWorldQuery.cpp
@@ -10,6 +10,23 @@
 
 #include UE_INLINE_GENERATED_CPP_BY_NAME(PCGWorldQuery)
 
+#define LOCTEXT_NAMESPACE "PCGWorldQuery"
+
+#if WITH_EDITOR
+FText UPCGWorldQuerySettings::GetNodeTooltipText() const
+{
+	return LOCTEXT("WorldQueryTooltip", "Allows generic access (based on overlaps) to collisions in the world that behaves like a volume.");
+}
+#endif
+
+TArray<FPCGPinProperties> UPCGWorldQuerySettings::OutputPinProperties() const
+{
+	TArray<FPCGPinProperties> PinProperties;
+	PinProperties.Emplace(PCGPinConstants::DefaultOutputLabel, EPCGDataType::Volume);
+
+	return PinProperties;
+}
+
 FPCGElementPtr UPCGWorldQuerySettings::CreateElement() const
 {
 	return MakeShared<FPCGWorldVolumetricQueryElement>();
@@ -43,6 +60,21 @@ bool FPCGWorldVolumetricQueryElement::ExecuteInternal(FPCGContext* Context) cons
 	return true;
 }
 
+#if WITH_EDITOR
+FText UPCGWorldRayHitSettings::GetNodeTooltipText() const
+{
+	return LOCTEXT("WorldRayHitTooltip", "Allows generic access (based on raycasts) to collisions in the world that behaves like a surface.");
+}
+#endif
+
+TArray<FPCGPinProperties> UPCGWorldRayHitSettings::OutputPinProperties() const
+{
+	TArray<FPCGPinProperties> PinProperties;
+	PinProperties.Emplace(PCGPinConstants::DefaultOutputLabel, EPCGDataType::Surface);
+
+	return PinProperties;
+}
+
 FPCGElementPtr UPCGWorldRayHitSettings::CreateElement() const
 {
 	return MakeShared<FPCGWorldRayHitQueryElement>();
@@ -106,3 +138,5 @@ bool FPCGWorldRayHitQueryElement::ExecuteInternal(FPCGContext* Context) const
 
 	return true;
 }
+
+#undef LOCTEXT_NAMESPACE
\ No newline at end of file
diff --git a/Engine/Plugins/Experimental/PCG/Source/PCG/Public/Elements/PCGWorldQuery.h b/Engine/Plugins/Experimental/PCG/Source/PCG/Public/Elements/PCGWorldQuery.h
--- a/Engine/Plugins/Experimental/PCG/Source/PCG/Public/Elements/PCGWorldQuery.h
+++ b/Engine/Plugins/Experimental/PCG/Source/PCG/Public/Elements/PCGWorldQuery.h
@@ -16,13 +16,15 @@ class PCG_API UPCGWorldQuerySettings : public UPCGSettings
 public:
 	//~Begin UPCGSettings interface
 #if WITH_EDITOR
-	virtual FName GetDefaultNodeName() const override { return FName(TEXT("WorldVolumetric")); }
+	virtual FName GetDefaultNodeName() const override { return FName(TEXT("World Volumetric Query")); }
+	virtual FText GetNodeTooltipText() const override;
 	virtual EPCGSettingsType GetType() const override { return EPCGSettingsType::Spatial; }
 #endif
 
-	TArray<FPCGPinProperties> InputPinProperties() const { return TArray<FPCGPinProperties>(); }
-
 protected:
+	virtual TArray<FPCGPinProperties> InputPinProperties() const override { return TArray<FPCGPinProperties>(); }
+	virtual TArray<FPCGPinProperties> OutputPinProperties() const override;
+
 	virtual FPCGElementPtr CreateElement() const override;
 	//~End UPCGSettings interface
 
@@ -47,13 +49,15 @@ class PCG_API UPCGWorldRayHitSettings : public UPCGSettings
 public:
 	//~Begin UPCGSettings interface
 #if WITH_EDITOR
-	virtual FName GetDefaultNodeName() const override { return FName(TEXT("WorldRayHit")); }
+	virtual FName GetDefaultNodeName() const override { return FName(TEXT("World Ray Hit Query")); }
+	virtual FText GetNodeTooltipText() const override;
 	virtual EPCGSettingsType GetType() const override { return EPCGSettingsType::Spatial; }
 #endif
 
-	TArray<FPCGPinProperties> InputPinProperties() const { return TArray<FPCGPinProperties>(); }
-
 protected:
+	virtual TArray<FPCGPinProperties> InputPinProperties() const override { return TArray<FPCGPinProperties>(); }
+	virtual TArray<FPCGPinProperties> OutputPinProperties() const override;
+
 	virtual FPCGElementPtr CreateElement() const override;
 	//~End UPCGSettings interface
 
