Fixed small issue in the filter attribute node - it did not behave like the CreateAttribute so it was not possible to inspect it.

#rnx

diff --git a/Engine/Plugins/Experimental/PCG/Source/PCG/Private/Elements/PCGAttributeFilterElement.cpp b/Engine/Plugins/Experimental/PCG/Source/PCG/Private/Elements/PCGAttributeFilterElement.cpp
--- a/Engine/Plugins/Experimental/PCG/Source/PCG/Private/Elements/PCGAttributeFilterElement.cpp
+++ b/Engine/Plugins/Experimental/PCG/Source/PCG/Private/Elements/PCGAttributeFilterElement.cpp
@@ -96,6 +96,14 @@ bool FPCGAttributeFilterElement::ExecuteInternal(FPCGContext* Context) const
 	const UPCGAttributeFilterSettings* Settings = Context->GetInputSettings<UPCGAttributeFilterSettings>();
 
 	const bool bAddAttributesFromParent = (Settings->Operation == EPCGAttributeFilterOperation::DeleteSelectedAttributes);
+	const EPCGMetadataFilterMode FilterMode = bAddAttributesFromParent ? EPCGMetadataFilterMode::ExcludeAttributes : EPCGMetadataFilterMode::IncludeAttributes;
+
+	TSet<FName> AttributesToFilter; 
+	const TArray<FString> FilterAttributes = PCGHelpers::GetStringArrayFromCommaSeparatedString(Settings->SelectedAttributes);
+	for (const FString& FilterAttribute : FilterAttributes)
+	{
+		AttributesToFilter.Add(FName(*FilterAttribute));
+	}
 
 	TArray<FPCGTaggedData> Inputs = Context->InputData.GetInputsByPin(PCGPinConstants::DefaultInputLabel);
 
@@ -113,7 +121,7 @@ bool FPCGAttributeFilterElement::ExecuteInternal(FPCGContext* Context) const
 
 			UPCGSpatialData* NewSpatialData = InputSpatialData->DuplicateData(/*bInitializeFromThisData=*/false);
 			Metadata = NewSpatialData->Metadata;
-			NewSpatialData->Metadata->Initialize(ParentMetadata, bAddAttributesFromParent);
+			NewSpatialData->Metadata->InitializeWithAttributeFilter(ParentMetadata, AttributesToFilter, FilterMode);
 
 			// No need to inherit metadata since we already initialized it.
 			NewSpatialData->InitializeFromData(InputSpatialData, /*InMetadataParentOverride=*/ nullptr, /*bInheritMetadata=*/ false);
@@ -126,8 +134,8 @@ bool FPCGAttributeFilterElement::ExecuteInternal(FPCGContext* Context) const
 
 			UPCGParamData* NewParamData = NewObject<UPCGParamData>();
 			Metadata = NewParamData->Metadata;
+			Metadata->InitializeAsCopyWithAttributeFilter(ParentMetadata, AttributesToFilter, FilterMode);
 
-			Metadata->Initialize(InputParamData->Metadata, bAddAttributesFromParent);
 			OutputData = NewParamData;
 		}
 		else
@@ -136,22 +144,6 @@ bool FPCGAttributeFilterElement::ExecuteInternal(FPCGContext* Context) const
 			continue;
 		}
 
-		TArray<FString> AttributesToKeep = PCGHelpers::GetStringArrayFromCommaSeparatedString(Settings->SelectedAttributes);
-
-		// Then add/remove each attribute in the list explicitly
-		for (const FString& AttributeName : AttributesToKeep)
-		{
-			switch (Settings->Operation)
-			{
-			case EPCGAttributeFilterOperation::KeepSelectedAttributes:
-				Metadata->AddAttribute(ParentMetadata, FName(AttributeName));
-				break;
-			case EPCGAttributeFilterOperation::DeleteSelectedAttributes:
-				Metadata->DeleteAttribute(FName(AttributeName));
-				break;
-			}
-		}
-
 		TArray<FPCGTaggedData>& Outputs = Context->OutputData.TaggedData;
 		FPCGTaggedData& Output = Outputs.Emplace_GetRef();
 		Output.Data = OutputData;
diff --git a/Engine/Plugins/Experimental/PCG/Source/PCG/Private/Metadata/PCGMetadata.cpp b/Engine/Plugins/Experimental/PCG/Source/PCG/Private/Metadata/PCGMetadata.cpp
--- a/Engine/Plugins/Experimental/PCG/Source/PCG/Private/Metadata/PCGMetadata.cpp
+++ b/Engine/Plugins/Experimental/PCG/Source/PCG/Private/Metadata/PCGMetadata.cpp
@@ -104,6 +104,11 @@ void UPCGMetadata::InitializeWithAttributeFilter(const UPCGMetadata* InParent, c
 }
 
 void UPCGMetadata::InitializeAsCopy(const UPCGMetadata* InMetadataToCopy)
+{
+	InitializeAsCopyWithAttributeFilter(InMetadataToCopy, TSet<FName>(), EPCGMetadataFilterMode::ExcludeAttributes);
+}
+
+void UPCGMetadata::InitializeAsCopyWithAttributeFilter(const UPCGMetadata* InMetadataToCopy, const TSet<FName>& InFilteredAttributes, EPCGMetadataFilterMode InFilterMode)
 {
 	if (!InMetadataToCopy)
 	{
@@ -122,10 +127,18 @@ void UPCGMetadata::InitializeAsCopy(const UPCGMetadata* InMetadataToCopy)
 	ParentKeys = InMetadataToCopy->ParentKeys;
 	ItemKeyOffset = InMetadataToCopy->ItemKeyOffset;
 
+	const bool bSkipAttributesInFilterList = (InFilterMode == EPCGMetadataFilterMode::ExcludeAttributes);
+
 	// Copy attributes
 	for (const TPair<FName, FPCGMetadataAttributeBase*>& OtherAttribute : InMetadataToCopy->Attributes)
 	{
-		CopyAttribute(OtherAttribute.Value, OtherAttribute.Key, /*bKeepParent=*/false, /*bCopyEntries=*/true, /*bCopyValues=*/true);
+		const bool bAttributeInFilterList = InFilteredAttributes.Contains(OtherAttribute.Key);
+		const bool bSkipThisAttribute = (bSkipAttributesInFilterList == bAttributeInFilterList);
+
+		if (!bSkipThisAttribute)
+		{
+			CopyAttribute(OtherAttribute.Value, OtherAttribute.Key, /*bKeepParent=*/false, /*bCopyEntries=*/true, /*bCopyValues=*/true);
+		}
 	}
 }
 
diff --git a/Engine/Plugins/Experimental/PCG/Source/PCG/Public/Metadata/PCGMetadata.h b/Engine/Plugins/Experimental/PCG/Source/PCG/Public/Metadata/PCGMetadata.h
--- a/Engine/Plugins/Experimental/PCG/Source/PCG/Public/Metadata/PCGMetadata.h
+++ b/Engine/Plugins/Experimental/PCG/Source/PCG/Public/Metadata/PCGMetadata.h
@@ -43,6 +43,10 @@ public:
 	UFUNCTION(BlueprintCallable, Category = "PCG|Metadata")
 	void InitializeAsCopy(const UPCGMetadata* InMetadataToCopy);
 
+	/** Initializes the metadata from a parent metadata by copy filtered attributes only to it */
+	UFUNCTION(BlueprintCallable, Category = "PCG|Metadata")
+	void InitializeAsCopyWithAttributeFilter(const UPCGMetadata* InMetadataToCopy, const TSet<FName>& InFilteredAttributes, EPCGMetadataFilterMode InFilterMode = EPCGMetadataFilterMode::ExcludeAttributes);
+
 	/** Creates missing attributes from another metadata if they are not currently present - note that this does not copy values */
 	UFUNCTION(BlueprintCallable, Category = "PCG|Metadata")
 	void AddAttributes(const UPCGMetadata* InOther) { AddAttributesFiltered(InOther, TSet<FName>()); }
