Threaded landscape cache creation
Added console command to build landscape cache
Added cvar to trigger landscape cache creation on PIE, turned on by default
Added some profiling tags

diff --git a/Engine/Plugins/Experimental/PCG/Source/PCG/Private/Grid/PCGLandscapeCache.cpp b/Engine/Plugins/Experimental/PCG/Source/PCG/Private/Grid/PCGLandscapeCache.cpp
--- a/Engine/Plugins/Experimental/PCG/Source/PCG/Private/Grid/PCGLandscapeCache.cpp
+++ b/Engine/Plugins/Experimental/PCG/Source/PCG/Private/Grid/PCGLandscapeCache.cpp
@@ -2,17 +2,17 @@
 
 #include "Grid/PCGLandscapeCache.h"
 
-#include "Engine/World.h"
+#include "PCGPoint.h"
 #include "Helpers/PCGBlueprintHelpers.h"
 #include "Metadata/PCGMetadata.h"
 
-
+#include "Async/ParallelFor.h"
+#include "Engine/World.h"
+#include "Kismet/GameplayStatics.h"
 #include "LandscapeComponent.h"
 #include "LandscapeInfoMap.h"
 #include "LandscapeProxy.h"
 #include "LandscapeDataAccess.h"
-#include "Kismet/GameplayStatics.h"
-#include "PCGPoint.h"
 #include "Serialization/BufferWriter.h"
 
 #include UE_INLINE_GENERATED_CPP_BY_NAME(PCGLandscapeCache)
@@ -548,13 +548,16 @@ void UPCGLandscapeCache::PrimeCache()
 		return;
 	}
 
+	// First, gather all potential cached data to create, emplace a nullptr at these locations
+	TArray<TTuple<ULandscapeInfo*, ULandscapeComponent*, TPair<FGuid, FIntPoint>>> CacheEntriesToBuild;
+
 	for (auto It = ULandscapeInfoMap::GetLandscapeInfoMap(World).Map.CreateIterator(); It; ++It)
 	{
 		ULandscapeInfo* LandscapeInfo = It.Value();
 		if (IsValid(LandscapeInfo))
 		{
 			// Build per-component information
-			LandscapeInfo->ForAllLandscapeProxies([this, LandscapeInfo](const ALandscapeProxy* LandscapeProxy)
+			LandscapeInfo->ForAllLandscapeProxies([this, LandscapeInfo, &CacheEntriesToBuild](const ALandscapeProxy* LandscapeProxy)
 			{
 				check(LandscapeProxy);
 				const FGuid LandscapeGuid = LandscapeProxy->GetLandscapeGuid();
@@ -571,16 +574,32 @@ void UPCGLandscapeCache::PrimeCache()
 
 					if (!CachedData.Contains(ComponentKey))
 					{
-						if (FPCGLandscapeCacheEntry* NewEntry = FPCGLandscapeCacheEntry::CreateCacheEntry(LandscapeInfo, LandscapeComponent))
-						{
-							CachedData.Add(ComponentKey, NewEntry);
-						}
+						CacheEntriesToBuild.Emplace(LandscapeInfo, LandscapeComponent, ComponentKey);
 					}
 				}
 			});
 		}
 	}
 
+	// Then create the landscape cache entries
+	TArray<FPCGLandscapeCacheEntry*> NewEntries;
+	NewEntries.SetNum(CacheEntriesToBuild.Num());
+
+	ParallelFor(CacheEntriesToBuild.Num(), [&CacheEntriesToBuild, &NewEntries](int32 EntryIndex)
+	{
+		const TTuple<ULandscapeInfo*, ULandscapeComponent*, TPair<FGuid, FIntPoint>>& CacheEntryInfo = CacheEntriesToBuild[EntryIndex];
+		NewEntries[EntryIndex] = FPCGLandscapeCacheEntry::CreateCacheEntry(CacheEntryInfo.Get<0>(), CacheEntryInfo.Get<1>());
+	});
+
+	// Finally, write them back to the cache
+	for (int32 EntryIndex = 0; EntryIndex < CacheEntriesToBuild.Num(); ++EntryIndex)
+	{
+		if (NewEntries[EntryIndex])
+		{
+			CachedData.Add(CacheEntriesToBuild[EntryIndex].Get<2>(), NewEntries[EntryIndex]);
+		}
+	}
+
 	CacheLayerNames();
 #endif
 }
@@ -616,7 +635,7 @@ const FPCGLandscapeCacheEntry* UPCGLandscapeCache::GetCacheEntry(ULandscapeCompo
 			if (FPCGLandscapeCacheEntry* NewEntry = FPCGLandscapeCacheEntry::CreateCacheEntry(LandscapeComponent->GetLandscapeInfo(), LandscapeComponent))
 			{
 				CacheEntry = NewEntry;
-				CachedData.Add(ComponentKey, NewEntry);				
+				CachedData.Add(ComponentKey, NewEntry);
 			}
 		}
 
diff --git a/Engine/Plugins/Experimental/PCG/Source/PCG/Private/PCGSubsystem.cpp b/Engine/Plugins/Experimental/PCG/Source/PCG/Private/PCGSubsystem.cpp
--- a/Engine/Plugins/Experimental/PCG/Source/PCG/Private/PCGSubsystem.cpp
+++ b/Engine/Plugins/Experimental/PCG/Source/PCG/Private/PCGSubsystem.cpp
@@ -61,6 +61,25 @@ namespace PCGSubsystemConsole
 					World->GetSubsystem<UPCGSubsystem>()->FlushCache();
 				}
 			}));
+
+#if WITH_EDITOR
+	static FAutoConsoleCommand CommandBuildLandscapeCache(
+		TEXT("pcg.BuildLandscapeCache"),
+		TEXT("Builds the landscape cache in the current world."),
+		FConsoleCommandDelegate::CreateLambda([]()
+			{
+				UWorld* World = (GEditor ? (GEditor->PlayWorld ? GEditor->PlayWorld.Get() : GEditor->GetEditorWorldContext().World()) : (GEngine ? GEngine->GetCurrentPlayWorld() : nullptr));
+				if(World && World->GetSubsystem<UPCGSubsystem>())
+				{
+					World->GetSubsystem<UPCGSubsystem>()->BuildLandscapeCache();
+				}
+			}));
+
+	static TAutoConsoleVariable<bool> CVarRebuildLandscapeOnPIE(
+		TEXT("pcg.PIE.RebuildLandscapeOnPIE"),
+		true,
+		TEXT("Controls whether the landscape cache will be rebuilt on PIE"));
+#endif
 }
 
 #if WITH_EDITOR
@@ -87,6 +106,7 @@ void UPCGSubsystem::Deinitialize()
 	GraphExecutor = nullptr;
 
 	PCGWorldActor = nullptr;
+	bHasTickedOnce = false;
 
 	Super::Deinitialize();
 }
@@ -117,6 +137,18 @@ void UPCGSubsystem::Tick(float DeltaSeconds)
 {
 	Super::Tick(DeltaSeconds);
 
+	if (!bHasTickedOnce)
+	{
+#if WITH_EDITOR
+		if (PCGSubsystemConsole::CVarRebuildLandscapeOnPIE.GetValueOnAnyThread() && PCGHelpers::IsRuntimeOrPIE())
+		{
+			BuildLandscapeCache();
+		}
+#endif
+
+		bHasTickedOnce = true;
+	}
+
 	// If we have any tasks to execute, schedule some
 	GraphExecutor->Execute();
 
@@ -1609,6 +1641,7 @@ void UPCGSubsystem::CleanFromCache(const IPCGElement* InElement, const UPCGSetti
 
 void UPCGSubsystem::BuildLandscapeCache()
 {
+	TRACE_CPUPROFILER_EVENT_SCOPE(UPCGSubsystem::BuildLandscapeCache);
 	if (UPCGLandscapeCache* LandscapeCache = GetLandscapeCache())
 	{
 		PCGWorldActor->Modify();
diff --git a/Engine/Plugins/Experimental/PCG/Source/PCG/Public/PCGSubsystem.h b/Engine/Plugins/Experimental/PCG/Source/PCG/Public/PCGSubsystem.h
--- a/Engine/Plugins/Experimental/PCG/Source/PCG/Public/PCGSubsystem.h
+++ b/Engine/Plugins/Experimental/PCG/Source/PCG/Public/PCGSubsystem.h
@@ -204,6 +204,7 @@ private:
 private:
 	APCGWorldActor* PCGWorldActor = nullptr;
 	FPCGGraphExecutor* GraphExecutor = nullptr;
+	bool bHasTickedOnce = false;
 
 #if WITH_EDITOR
 	FCriticalSection PCGWorldActorLock;
