[PCG] Add culling for tracking actors, depending on the setting

diff --git a/Engine/Plugins/Experimental/PCG/Source/PCG/Private/Elements/PCGDataFromActor.cpp b/Engine/Plugins/Experimental/PCG/Source/PCG/Private/Elements/PCGDataFromActor.cpp
--- a/Engine/Plugins/Experimental/PCG/Source/PCG/Private/Elements/PCGDataFromActor.cpp
+++ b/Engine/Plugins/Experimental/PCG/Source/PCG/Private/Elements/PCGDataFromActor.cpp
@@ -19,7 +19,7 @@ void UPCGDataFromActorSettings::GetTrackedActorTags(FPCGTagToSettingsMap& OutTag
 	if (ActorSelector.ActorSelection == EPCGActorSelection::ByTag &&
 		ActorSelector.ActorFilter == EPCGActorFilter::AllWorldActors)
 	{
-		OutTagToSettings.FindOrAdd(ActorSelector.ActorSelectionTag).Add(this);
+		OutTagToSettings.FindOrAdd(ActorSelector.ActorSelectionTag).Emplace({ this, bTrackActorsOnlyWithinBounds });
 	}
 }
 
diff --git a/Engine/Plugins/Experimental/PCG/Source/PCG/Private/Elements/PCGExecuteBlueprint.cpp b/Engine/Plugins/Experimental/PCG/Source/PCG/Private/Elements/PCGExecuteBlueprint.cpp
--- a/Engine/Plugins/Experimental/PCG/Source/PCG/Private/Elements/PCGExecuteBlueprint.cpp
+++ b/Engine/Plugins/Experimental/PCG/Source/PCG/Private/Elements/PCGExecuteBlueprint.cpp
@@ -394,12 +394,10 @@ EPCGSettingsType UPCGBlueprintSettings::GetType() const
 
 void UPCGBlueprintSettings::GetTrackedActorTags(FPCGTagToSettingsMap& OutTagToSettings, TArray<TObjectPtr<const UPCGGraph>>& OutVisitedGraphs) const
 {
-#if WITH_EDITOR
 	for (const FName& Tag : TrackedActorTags)
 	{
-		OutTagToSettings.FindOrAdd(Tag).Add(this);
+		OutTagToSettings.FindOrAdd(Tag).Emplace({ this, bTrackActorsOnlyWithinBounds });
 	}
-#endif // WITH_EDITOR
 }
 
 UObject* UPCGBlueprintSettings::GetJumpTargetForDoubleClick() const
diff --git a/Engine/Plugins/Experimental/PCG/Source/PCG/Private/Elements/PCGPropertyToParamData.cpp b/Engine/Plugins/Experimental/PCG/Source/PCG/Private/Elements/PCGPropertyToParamData.cpp
--- a/Engine/Plugins/Experimental/PCG/Source/PCG/Private/Elements/PCGPropertyToParamData.cpp
+++ b/Engine/Plugins/Experimental/PCG/Source/PCG/Private/Elements/PCGPropertyToParamData.cpp
@@ -13,6 +13,18 @@
 
 #include UE_INLINE_GENERATED_CPP_BY_NAME(PCGPropertyToParamData)
 
+#if WITH_EDITOR
+void UPCGPropertyToParamDataSettings::GetTrackedActorTags(FPCGTagToSettingsMap& OutTagToSettings, TArray<TObjectPtr<const UPCGGraph>>& OutVisitedGraphs) const
+{
+	if (ActorSelector.ActorSelection == EPCGActorSelection::ByTag &&
+		ActorSelector.ActorFilter == EPCGActorFilter::AllWorldActors)
+	{
+		OutTagToSettings.FindOrAdd(ActorSelector.ActorSelectionTag).Emplace({ this, bTrackActorsOnlyWithinBounds });
+	}
+
+}
+#endif // WITH_EDITOR
+
 void UPCGPropertyToParamDataSettings::PostLoad()
 {
 	Super::PostLoad();
diff --git a/Engine/Plugins/Experimental/PCG/Source/PCG/Private/PCGComponent.cpp b/Engine/Plugins/Experimental/PCG/Source/PCG/Private/PCGComponent.cpp
--- a/Engine/Plugins/Experimental/PCG/Source/PCG/Private/PCGComponent.cpp
+++ b/Engine/Plugins/Experimental/PCG/Source/PCG/Private/PCGComponent.cpp
@@ -192,32 +192,39 @@ void UPCGComponent::SetPropertiesFromOriginal(const UPCGComponent* Original)
 #if WITH_EDITOR
 	const bool bHasDirtyInput = InputType != NewInputType;
 	const bool bHasDirtyExclusions = !(ExcludedTags.Num() == Original->ExcludedTags.Num() && ExcludedTags.Includes(Original->ExcludedTags));
-	const bool bIsDirty = bHasDirtyInput || bHasDirtyExclusions || GraphInstance->Graph != Original->GraphInstance;
 
-	if (bHasDirtyExclusions)
-	{
-		TeardownTrackingCallbacks();
-		ExcludedTags = Original->ExcludedTags;
-		SetupTrackingCallbacks();
-		RefreshTrackingData();
-	}
-#else
-	ExcludedTags = Original->ExcludedTags;
-#endif
+	TSet<FName> TrackedTags;
+	TSet<FName> OriginalTrackedTags;
+	CachedTrackedTagsToSettings.GetKeys(TrackedTags);
+	Original->CachedTrackedTagsToSettings.GetKeys(OriginalTrackedTags);
+	const bool bHasDirtyTracking = !(TrackedTags.Num() == OriginalTrackedTags.Num() && TrackedTags.Includes(OriginalTrackedTags));
+
+	const bool bGraphInstanceIsDifferent = GraphInstance->Graph != Original->GraphInstance;
+	const bool bIsDirty = bHasDirtyInput || bHasDirtyExclusions || bHasDirtyTracking || bGraphInstanceIsDifferent;
+#endif // WITH_EDITOR
 
 	InputType = NewInputType;
 	Seed = Original->Seed;
 	GraphInstance->SetGraph(Original->GraphInstance);
-
 	GenerationTrigger = Original->GenerationTrigger;
 
 #if WITH_EDITOR
+	if (bHasDirtyExclusions || bHasDirtyTracking)
+	{
+		TeardownTrackingCallbacks();
+		ExcludedTags = Original->ExcludedTags;
+		SetupTrackingCallbacks();
+		RefreshTrackingData();
+	}
+
 	// Note that while we dirty here, we won't trigger a refresh since we don't have the required context
 	if (bIsDirty)
 	{
 		Modify();
 		DirtyGenerated((bHasDirtyInput ? EPCGComponentDirtyFlag::Input : EPCGComponentDirtyFlag::None) | (bHasDirtyExclusions ? EPCGComponentDirtyFlag::Exclusions : EPCGComponentDirtyFlag::None));
 	}
+#else
+	ExcludedTags = Original->ExcludedTags;
 #endif
 }
 
@@ -302,6 +309,7 @@ FPCGTaskId UPCGComponent::CreateGenerateTask(bool bForce, const TArray<FPCGTaskI
 
 bool UPCGComponent::GetActorsFromTags(const TSet<FName>& InTags, TSet<TWeakObjectPtr<AActor>>& OutActors, bool bCullAgainstLocalBounds)
 {
+	TRACE_CPUPROFILER_EVENT_SCOPE(UPCGComponent::GetActorsFromTags::Excluded);
 	UWorld* World = GetWorld();
 
 	if (!World)
@@ -338,6 +346,48 @@ bool UPCGComponent::GetActorsFromTags(const TSet<FName>& InTags, TSet<TWeakObjec
 	return bHasValidTag;
 }
 
+bool UPCGComponent::GetActorsFromTags(const TMap<FName, bool>& InTagsAndCulling, TSet<TWeakObjectPtr<AActor>>& OutActors)
+{
+	TRACE_CPUPROFILER_EVENT_SCOPE(UPCGComponent::GetActorsFromTags::Tracked);
+	UWorld* World = GetWorld();
+
+	if (!World)
+	{
+		return false;
+	}
+
+	const FBox LocalBounds = GetGridBounds();
+
+	TArray<AActor*> PerTagActors;
+
+	OutActors.Reset();
+
+	bool bHasValidTag = false;
+	for (const TPair<FName, bool>& TagAndCulling : InTagsAndCulling)
+	{
+		const FName& Tag = TagAndCulling.Key;
+		const bool bCullAgainstLocalBounds = TagAndCulling.Value;
+
+		if (Tag != NAME_None)
+		{
+			bHasValidTag = true;
+			UGameplayStatics::GetAllActorsWithTag(World, Tag, PerTagActors);
+
+			for (AActor* Actor : PerTagActors)
+			{
+				if (!bCullAgainstLocalBounds || LocalBounds.Intersect(GetGridBounds(Actor)))
+				{
+					OutActors.Emplace(Actor);
+				}
+			}
+
+			PerTagActors.Reset();
+		}
+	}
+
+	return bHasValidTag;
+}
+
 void UPCGComponent::PostProcessGraph(const FBox& InNewBounds, bool bInGenerated, FPCGContext* Context)
 {
 	PCGGeneratedResourcesLogging::LogPostProcessGraph();
@@ -1291,9 +1341,31 @@ void UPCGComponent::TeardownActorCallbacks()
 void UPCGComponent::SetupTrackingCallbacks()
 {
 	CachedTrackedTagsToSettings.Reset();
+	CachedTrackedTagsToCulling.Reset();
 	if (UPCGGraph* PCGGraph = GetGraph())
 	{
 		CachedTrackedTagsToSettings = PCGGraph->GetTrackedTagsToSettings();
+
+		// A tag should be culled, if only all the settings that track this tag should cull.
+		// Note that is only impact the fact that we track (or not) this tag.
+		// If a setting is marked as "should cull", it will only be dirtied (at least by default), if the actor with the
+		// given tag intersect with the component.
+		for (const TPair<FName, TSet<FPCGSettingsAndCulling>>& It : CachedTrackedTagsToSettings)
+		{
+			const FName& Tag = It.Key;
+
+			bool bShouldCull = true;
+			for (const FPCGSettingsAndCulling& SettingsAndCulling : It.Value)
+			{
+				if (!SettingsAndCulling.Value)
+				{
+					bShouldCull = false;
+					break;
+				}
+			}
+
+			CachedTrackedTagsToCulling.Emplace(Tag, bShouldCull);
+		}
 	}
 
 	if(!ExcludedTags.IsEmpty() || !CachedTrackedTagsToSettings.IsEmpty())
@@ -1305,11 +1377,9 @@ void UPCGComponent::SetupTrackingCallbacks()
 
 void UPCGComponent::RefreshTrackingData()
 {
+	TRACE_CPUPROFILER_EVENT_SCOPE(UPCGComponent::RefreshTrackingData);
 	GetActorsFromTags(ExcludedTags, CachedExcludedActors, /*bCullAgainstLocalBounds=*/true);
-
-	TSet<FName> TrackedTags;
-	CachedTrackedTagsToSettings.GetKeys(TrackedTags);
-	GetActorsFromTags(TrackedTags, CachedTrackedActors, /*bCullAgainstLocalBounds=*/false);
+	GetActorsFromTags(CachedTrackedTagsToCulling, CachedTrackedActors);
 	PopulateTrackedActorToTagsMap(/*bForce=*/true);
 }
 
@@ -1342,6 +1412,7 @@ bool UPCGComponent::ActorHasExcludedTag(AActor* InActor) const
 
 bool UPCGComponent::UpdateExcludedActor(AActor* InActor)
 {
+	TRACE_CPUPROFILER_EVENT_SCOPE(UPCGComponent::UpdateExcludedActor);
 	// Dirty data in all cases - the tag or positional changes will be picked up in the test later
 	if (CachedExcludedActors.Contains(InActor))
 	{
@@ -1387,6 +1458,7 @@ bool UPCGComponent::ActorIsTracked(AActor* InActor) const
 
 void UPCGComponent::OnActorAdded(AActor* InActor)
 {
+	TRACE_CPUPROFILER_EVENT_SCOPE(UPCGComponent::OnActorAdded);
 	if (!InActor || InActor->bIsEditorPreviewActor)
 	{
 		return;
@@ -1404,6 +1476,7 @@ void UPCGComponent::OnActorAdded(AActor* InActor)
 
 void UPCGComponent::OnActorDeleted(AActor* InActor)
 {
+	TRACE_CPUPROFILER_EVENT_SCOPE(UPCGComponent::OnActorDeleted);
 	if (!InActor || InActor->bIsEditorPreviewActor)
 	{
 		return;
@@ -1421,6 +1494,7 @@ void UPCGComponent::OnActorDeleted(AActor* InActor)
 
 void UPCGComponent::OnActorMoved(AActor* InActor)
 {
+	TRACE_CPUPROFILER_EVENT_SCOPE(UPCGComponent::OnActorMoved);
 	if (!InActor || InActor->bIsEditorPreviewActor)
 	{
 		return;
@@ -1457,7 +1531,7 @@ void UPCGComponent::OnActorMoved(AActor* InActor)
 
 		if (bDirtyAndRefresh)
 		{
-			DirtyGenerated(bDirtyExclusions ? EPCGComponentDirtyFlag::Exclusions : EPCGComponentDirtyFlag::None, /*bDispatchToLocalComponents=*/true);
+			DirtyGenerated(bDirtyExclusions ? EPCGComponentDirtyFlag::Exclusions : EPCGComponentDirtyFlag::None, /*bDispatchToLocalComponents=*/false);
 			Refresh();
 		}
 	}
@@ -1561,6 +1635,7 @@ void UPCGComponent::OnLandscapeChanged(ALandscapeProxy* Landscape, const FLandsc
 
 void UPCGComponent::OnObjectPropertyChanged(UObject* InObject, FPropertyChangedEvent& InEvent)
 {
+	TRACE_CPUPROFILER_EVENT_SCOPE(UPCGComponent::OnObjectPropertyChanged);
 	bool bValueNotInteractive = (InEvent.ChangeType != EPropertyChangeType::Interactive);
 	// Special exception for actor tags, as we can't track otherwise an actor "losing" a tag
 	bool bActorTagChange = (InEvent.Property && InEvent.Property->GetFName() == GET_MEMBER_NAME_CHECKED(AActor, Tags));
@@ -1607,6 +1682,7 @@ void UPCGComponent::OnObjectPropertyChanged(UObject* InObject, FPropertyChangedE
 
 void UPCGComponent::OnActorChanged(AActor* Actor, UObject* InObject, bool bActorTagChange)
 {
+	TRACE_CPUPROFILER_EVENT_SCOPE(UPCGComponent::OnActorChanged);
 	if (Actor == GetOwner())
 	{
 		// Something has changed on the owner (including properties of this component)
@@ -1902,6 +1978,7 @@ TArray<UPCGData*> UPCGComponent::GetPCGExclusionData()
 
 void UPCGComponent::UpdatePCGExclusionData()
 {
+	TRACE_CPUPROFILER_EVENT_SCOPE(UPCGComponent::UpdatePCGExclusionData);
 	const UPCGData* InputData = GetInputPCGData();
 	const UPCGSpatialData* InputSpatialData = Cast<const UPCGSpatialData>(InputData);
 
@@ -2335,7 +2412,7 @@ FBox UPCGComponent::GetGridBounds() const
 	return PCGHelpers::GetGridBounds(GetOwner(), this);
 }
 
-FBox UPCGComponent::GetGridBounds(AActor* Actor) const
+FBox UPCGComponent::GetGridBounds(const AActor* Actor) const
 {
 	return PCGHelpers::GetGridBounds(Actor, this);
 }
@@ -2348,6 +2425,7 @@ UPCGSubsystem* UPCGComponent::GetSubsystem() const
 #if WITH_EDITOR
 bool UPCGComponent::PopulateTrackedActorToTagsMap(bool bForce)
 {
+	TRACE_CPUPROFILER_EVENT_SCOPE(UPCGComponent::PopulateTrackedActorToTagsMap);
 	if (bActorToTagsMapPopulated && !bForce)
 	{
 		return false;
@@ -2369,6 +2447,7 @@ bool UPCGComponent::PopulateTrackedActorToTagsMap(bool bForce)
 
 bool UPCGComponent::AddTrackedActor(AActor* InActor, bool bForce)
 {
+	TRACE_CPUPROFILER_EVENT_SCOPE(UPCGComponent::AddTrackedActor);
 	if (!bForce)
 	{
 		PopulateTrackedActorToTagsMap();
@@ -2377,20 +2456,40 @@ bool UPCGComponent::AddTrackedActor(AActor* InActor, bool bForce)
 	check(InActor);
 	bool bAppliedChange = false;
 
+	bool bIntersectionComputed = false;
+	bool bIntersect = false;
+
 	for (const FName& Tag : InActor->Tags)
 	{
-		if (!CachedTrackedTagsToSettings.Contains(Tag))
+		const bool* ShouldCullPtr = CachedTrackedTagsToCulling.Find(Tag);
+
+		if (!ShouldCullPtr)
 		{
 			continue;
 		}
 
+		if (*ShouldCullPtr)
+		{
+			if (!bIntersectionComputed)
+			{
+				bIntersect = GetGridBounds().Intersect(GetGridBounds(InActor));
+				bIntersectionComputed = true;
+			}
+
+			if (!bIntersect)
+			{
+				continue;
+			}
+		}
+
 		bAppliedChange = true;
 		CachedTrackedActorToTags.FindOrAdd(InActor).Add(Tag);
 		PCGHelpers::GatherDependencies(InActor, CachedTrackedActorToDependencies.FindOrAdd(InActor), 1);
 
 		if (!bForce)
 		{
-			DirtyCacheFromTag(Tag);
+			// Since we arrived here, we already verified culling, so we can ignore it.
+			DirtyCacheFromTag(Tag, InActor, /*bIgnoreCull=*/true);
 		}
 	}
 
@@ -2399,6 +2498,7 @@ bool UPCGComponent::AddTrackedActor(AActor* InActor, bool bForce)
 
 bool UPCGComponent::RemoveTrackedActor(AActor* InActor)
 {
+	TRACE_CPUPROFILER_EVENT_SCOPE(UPCGComponent::RemoveTrackedActor);
 	PopulateTrackedActorToTagsMap();
 
 	check(InActor);
@@ -2408,12 +2508,11 @@ bool UPCGComponent::RemoveTrackedActor(AActor* InActor)
 	{
 		for (const FName& Tag : CachedTrackedActorToTags[InActor])
 		{
-			DirtyCacheFromTag(Tag);
+			bAppliedChange |= DirtyCacheFromTag(Tag, InActor);
 		}
 
 		CachedTrackedActorToTags.Remove(InActor);
 		CachedTrackedActorToDependencies.Remove(InActor);
-		bAppliedChange = true;
 	}
 
 	return bAppliedChange;
@@ -2421,6 +2520,7 @@ bool UPCGComponent::RemoveTrackedActor(AActor* InActor)
 
 bool UPCGComponent::UpdateTrackedActor(AActor* InActor)
 {
+	TRACE_CPUPROFILER_EVENT_SCOPE(UPCGComponent::UpdateTrackedActor);
 	check(InActor);
 	// If the tracked data wasn't initialized before, then it is not possible to know if we need to update or not - take no chances
 	bool bAppliedChange = PopulateTrackedActorToTagsMap();
@@ -2434,9 +2534,8 @@ bool UPCGComponent::UpdateTrackedActor(AActor* InActor)
 		{
 			if (!InActor->Tags.Contains(CachedTag))
 			{
+				bAppliedChange |= DirtyCacheFromTag(CachedTag, InActor);
 				CachedTrackedActorToTags[InActor].Remove(CachedTag);
-				DirtyCacheFromTag(CachedTag);
-				bAppliedChange = true;
 			}
 		}
 	}
@@ -2453,8 +2552,7 @@ bool UPCGComponent::UpdateTrackedActor(AActor* InActor)
 		{
 			CachedTrackedActorToTags[InActor].Add(Tag);
 			PCGHelpers::GatherDependencies(InActor, CachedTrackedActorToDependencies.FindOrAdd(InActor), 1);
-			DirtyCacheFromTag(Tag);
-			bAppliedChange = true;
+			bAppliedChange |= DirtyCacheFromTag(Tag, InActor);
 		}
 	}
 
@@ -2470,48 +2568,115 @@ bool UPCGComponent::UpdateTrackedActor(AActor* InActor)
 
 bool UPCGComponent::DirtyTrackedActor(AActor* InActor)
 {
+	TRACE_CPUPROFILER_EVENT_SCOPE(UPCGComponent::DirtyTrackedActor);
 	PopulateTrackedActorToTagsMap();
 
 	check(InActor);
 	bool bAppliedChange = false;
 
+	bool bIntersectionComputed = false;
+	bool bIntersect = false;
+
 	if (CachedTrackedActorToTags.Contains(InActor))
 	{
-		for (const FName& Tag : CachedTrackedActorToTags[InActor])
+		TSet<FName> CachedTags = CachedTrackedActorToTags[InActor];
+		for (const FName& CachedTag : CachedTags)
 		{
-			DirtyCacheFromTag(Tag);
-		}
+			const bool* ShouldCullPtr = CachedTrackedTagsToCulling.Find(CachedTag);
 
-		bAppliedChange = true;
+			if (!ShouldCullPtr)
+			{
+				continue;
+			}
+
+			if (*ShouldCullPtr)
+			{
+				if (!bIntersectionComputed)
+				{
+					bIntersect = GetGridBounds().Intersect(GetGridBounds(InActor));
+					bIntersectionComputed = true;
+				}
+
+				if (!bIntersect)
+				{
+					// We were tracking the tag, but it is now culled. We'll need to dirty the settings
+					// even if it is culled (that's why we ignore cull here).
+					bAppliedChange |= DirtyCacheFromTag(CachedTag, InActor, /*bIgnoreCull=*/ true);
+					CachedTrackedActorToTags[InActor].Remove(CachedTag);
+					continue;
+				}
+			}
+
+			bAppliedChange |= DirtyCacheFromTag(CachedTag, InActor);
+		}
 	}
 	else if (AddTrackedActor(InActor))
 	{
 		bAppliedChange = true;
 	}
 
+	// Since we might have remove some tags, do some cleaning there
+	if (CachedTrackedActorToTags.Contains(InActor) && CachedTrackedActorToTags[InActor].IsEmpty())
+	{
+		CachedTrackedActorToTags.Remove(InActor);
+		CachedTrackedActorToDependencies.Remove(InActor);
+	}
+
 	return bAppliedChange;
 }
 
-void UPCGComponent::DirtyCacheFromTag(const FName& InTag)
+bool UPCGComponent::DirtyCacheFromTag(const FName& InTag, const AActor* InActor, bool bIgnoreCull)
 {
-	if (CachedTrackedTagsToSettings.Contains(InTag))
+	TRACE_CPUPROFILER_EVENT_SCOPE(UPCGComponent::DirtyCacheFromTag);
+	bool bWasDirtied = false;
+
+	const TSet<FPCGSettingsAndCulling>* SettingsAndCullingPtr = CachedTrackedTagsToSettings.Find(InTag);
+
+	if (SettingsAndCullingPtr)
 	{
-		for (TWeakObjectPtr<const UPCGSettings> Settings : CachedTrackedTagsToSettings[InTag])
+		// Don't compute the bounds if we never cull.
+		bool bIntersectionComputed = false;
+		bool bIntersect = false;
+
+		for (const FPCGSettingsAndCulling& SettingsAndCulling : *SettingsAndCullingPtr)
 		{
+			const TWeakObjectPtr<const UPCGSettings>& Settings = SettingsAndCulling.Key;
+			const bool bShouldCull = SettingsAndCulling.Value && !bIgnoreCull;
+
 			if (Settings.IsValid() && GetSubsystem())
 			{
+				// If we cull and no intersection, continue
+				if (bShouldCull)
+				{
+					if (!bIntersectionComputed)
+					{
+						bIntersect = GetGridBounds().Intersect(GetGridBounds(InActor));
+						bIntersectionComputed = true;
+					}
+
+					if (!bIntersect)
+					{
+						continue;
+					}
+				}
+
 				GetSubsystem()->CleanFromCache(Settings->GetElement().Get(), Settings.Get());
+				bWasDirtied = true;
 			}
 		}
 	}
+
+	return bWasDirtied;
 }
 
 void UPCGComponent::DirtyCacheForAllTrackedTags()
 {
+	TRACE_CPUPROFILER_EVENT_SCOPE(UPCGComponent::DirtyCacheForAllTrackedTags);
 	for (const auto& TagToSettings : CachedTrackedTagsToSettings)
 	{
-		for (TWeakObjectPtr<const UPCGSettings> Settings : TagToSettings.Value)
+		for (const FPCGSettingsAndCulling& SettingsAndCulling : TagToSettings.Value)
 		{
+			const TWeakObjectPtr<const UPCGSettings>& Settings = SettingsAndCulling.Key;
 			if (Settings.IsValid() && GetSubsystem())
 			{
 				GetSubsystem()->CleanFromCache(Settings->GetElement().Get(), Settings.Get());
diff --git a/Engine/Plugins/Experimental/PCG/Source/PCG/Private/PCGHelpers.cpp b/Engine/Plugins/Experimental/PCG/Source/PCG/Private/PCGHelpers.cpp
--- a/Engine/Plugins/Experimental/PCG/Source/PCG/Private/PCGHelpers.cpp
+++ b/Engine/Plugins/Experimental/PCG/Source/PCG/Private/PCGHelpers.cpp
@@ -58,13 +58,13 @@ namespace PCGHelpers
 		}
 	}
 
-	FBox GetGridBounds(AActor* Actor, const UPCGComponent* Component)
+	FBox GetGridBounds(const AActor* Actor, const UPCGComponent* Component)
 	{
 		check(Actor);
 
 		FBox Bounds(EForceInit::ForceInit);
 
-		if (APCGPartitionActor* PartitionActor = Cast<APCGPartitionActor>(Actor))
+		if (const APCGPartitionActor* PartitionActor = Cast<const APCGPartitionActor>(Actor))
 		{
 			// First, get the bounds from the partition actor
 			Bounds = PartitionActor->GetFixedBounds();
@@ -79,7 +79,7 @@ namespace PCGHelpers
 			}
 		}
 		// TODO: verify this works as expected in non-editor builds
-		else if (ALandscapeProxy* LandscapeActor = Cast<ALandscape>(Actor))
+		else if (const ALandscapeProxy* LandscapeActor = Cast<const ALandscape>(Actor))
 		{
 			Bounds = GetLandscapeBounds(LandscapeActor);
 		}
@@ -91,7 +91,7 @@ namespace PCGHelpers
 		return Bounds;
 	}
 
-	FBox GetActorBounds(AActor* InActor, bool bIgnorePCGCreatedComponents)
+	FBox GetActorBounds(const AActor* InActor, bool bIgnorePCGCreatedComponents)
 	{
 		// Specialized version of GetComponentsBoundingBox that skips over PCG generated components
 		// This is to ensure stable bounds and no timing issues (cleared ISMs, etc.)
@@ -121,7 +121,7 @@ namespace PCGHelpers
 		return Box;
 	}
 
-	FBox GetActorLocalBounds(AActor* InActor, bool bIgnorePCGCreatedComponents)
+	FBox GetActorLocalBounds(const AActor* InActor, bool bIgnorePCGCreatedComponents)
 	{
 		// Specialized version of CalculateComponentsBoundingBoxInLocalScape that skips over PCG generated components
 		// This is to ensure stable bounds and no timing issues (cleared ISMs, etc.)
@@ -162,11 +162,11 @@ namespace PCGHelpers
 #endif // WITH_EDITOR
 	}
 
-	FBox GetLandscapeBounds(ALandscapeProxy* InLandscape)
+	FBox GetLandscapeBounds(const ALandscapeProxy* InLandscape)
 	{
 		check(InLandscape);
 
-		if (ALandscape* Landscape = Cast<ALandscape>(InLandscape))
+		if (const ALandscape* Landscape = Cast<const ALandscape>(InLandscape))
 		{
 			// If the landscape isn't done being loaded, we're very unlikely to want to interact with it
 			if (Landscape->GetLandscapeInfo() == nullptr)
diff --git a/Engine/Plugins/Experimental/PCG/Source/PCG/Private/PCGHelpers.h b/Engine/Plugins/Experimental/PCG/Source/PCG/Private/PCGHelpers.h
--- a/Engine/Plugins/Experimental/PCG/Source/PCG/Private/PCGHelpers.h
+++ b/Engine/Plugins/Experimental/PCG/Source/PCG/Private/PCGHelpers.h
@@ -28,11 +28,11 @@ namespace PCGHelpers
 	FBox OverlapBounds(const FBox& InBoxA, const FBox& InBoxB);
 
 	/** Returns the bounds of InActor, intersected with the component if InActor is a partition actor */
-	FBox GetGridBounds(AActor* InActor, const UPCGComponent* InComponent);
+	FBox GetGridBounds(const AActor* InActor, const UPCGComponent* InComponent);
 
-	FBox GetActorBounds(AActor* InActor, bool bIgnorePCGCreatedComponents = true);
-	FBox GetActorLocalBounds(AActor* InActor, bool bIgnorePCGCreatedComponents = true);
-	FBox GetLandscapeBounds(ALandscapeProxy* InLandscape);
+	FBox GetActorBounds(const AActor* InActor, bool bIgnorePCGCreatedComponents = true);
+	FBox GetActorLocalBounds(const AActor* InActor, bool bIgnorePCGCreatedComponents = true);
+	FBox GetLandscapeBounds(const ALandscapeProxy* InLandscape);
 
 	ALandscape* GetLandscape(UWorld* InWorld, const FBox& InActorBounds);
 	TArray<TWeakObjectPtr<ALandscapeProxy>> GetLandscapeProxies(UWorld* InWorld, const FBox& InActorBounds);
diff --git a/Engine/Plugins/Experimental/PCG/Source/PCG/Public/Elements/PCGDataFromActor.h b/Engine/Plugins/Experimental/PCG/Source/PCG/Public/Elements/PCGDataFromActor.h
--- a/Engine/Plugins/Experimental/PCG/Source/PCG/Public/Elements/PCGDataFromActor.h
+++ b/Engine/Plugins/Experimental/PCG/Source/PCG/Public/Elements/PCGDataFromActor.h
@@ -60,6 +60,12 @@ public:
 
 	UPROPERTY(BlueprintReadWrite, EditAnywhere, Category = Settings, meta = (EditCondition = "Mode == EPCGGetDataFromActorMode::GetDataFromProperty", EditConditionHides))
 	FName PropertyName = NAME_None;
+
+#if WITH_EDITORONLY_DATA
+	/** If this is checked, found actors that are outside component bounds will not trigger a refresh. Only works for tags for now in editor. */
+	UPROPERTY(BlueprintReadOnly, EditAnywhere, Category = Settings)
+	bool bTrackActorsOnlyWithinBounds = true;
+#endif // WITH_EDITORONLY_DATA
 };
 
 class FPCGDataFromActorContext : public FPCGContext
diff --git a/Engine/Plugins/Experimental/PCG/Source/PCG/Public/Elements/PCGExecuteBlueprint.h b/Engine/Plugins/Experimental/PCG/Source/PCG/Public/Elements/PCGExecuteBlueprint.h
--- a/Engine/Plugins/Experimental/PCG/Source/PCG/Public/Elements/PCGExecuteBlueprint.h
+++ b/Engine/Plugins/Experimental/PCG/Source/PCG/Public/Elements/PCGExecuteBlueprint.h
@@ -239,6 +239,10 @@ protected:
 	UPROPERTY(BlueprintReadWrite, EditAnywhere, Category = Settings)
 	TArray<FName> TrackedActorTags;
 
+	/** If this is checked, found actors that are outside component bounds will not trigger a refresh. Only works for tags for now in editor. */
+	UPROPERTY(BlueprintReadOnly, EditAnywhere, Category = Settings)
+	bool bTrackActorsOnlyWithinBounds = false;
+
 	UPROPERTY()
 	bool bCreatesArtifacts_DEPRECATED = false;
 
diff --git a/Engine/Plugins/Experimental/PCG/Source/PCG/Public/Elements/PCGPropertyToParamData.h b/Engine/Plugins/Experimental/PCG/Source/PCG/Public/Elements/PCGPropertyToParamData.h
--- a/Engine/Plugins/Experimental/PCG/Source/PCG/Public/Elements/PCGPropertyToParamData.h
+++ b/Engine/Plugins/Experimental/PCG/Source/PCG/Public/Elements/PCGPropertyToParamData.h
@@ -27,6 +27,7 @@ public:
 #if WITH_EDITOR
 	virtual FName GetDefaultNodeName() const override { return FName(TEXT("PropertyToParamData")); }
 	virtual EPCGSettingsType GetType() const override { return EPCGSettingsType::Param; }
+	virtual void GetTrackedActorTags(FPCGTagToSettingsMap& OutTagToSettings, TArray<TObjectPtr<const UPCGGraph>>& OutVisitedGraphs) const override;
 #endif
 
 	virtual TArray<FPCGPinProperties> OutputPinProperties() const override;
@@ -64,6 +65,12 @@ public:
 	UPROPERTY(BlueprintReadWrite, EditAnywhere, Category = Settings)
 	bool bAlwaysRequeryActors = true;
 
+#if WITH_EDITORONLY_DATA
+	/** If this is checked, found actors that are outside component bounds will not trigger a refresh. Only works for tags for now in editor. */
+	UPROPERTY(BlueprintReadOnly, EditAnywhere, Category = Settings)
+	bool bTrackActorsOnlyWithinBounds = false;
+#endif // WITH_EDITORONLY_DATA
+
 private:
 	UPROPERTY()
 	EPCGActorSelection ActorSelection_DEPRECATED;
diff --git a/Engine/Plugins/Experimental/PCG/Source/PCG/Public/PCGComponent.h b/Engine/Plugins/Experimental/PCG/Source/PCG/Public/PCGComponent.h
--- a/Engine/Plugins/Experimental/PCG/Source/PCG/Public/PCGComponent.h
+++ b/Engine/Plugins/Experimental/PCG/Source/PCG/Public/PCGComponent.h
@@ -307,6 +307,7 @@ private:
 	void SetManagedResources(const TArray<TObjectPtr<UPCGManagedResource>>& Resources);
 
 	bool GetActorsFromTags(const TSet<FName>& InTags, TSet<TWeakObjectPtr<AActor>>& OutActors, bool bCullAgainstLocalBounds);
+	bool GetActorsFromTags(const TMap<FName, bool>& InTagsAndCulling, TSet<TWeakObjectPtr<AActor>>& OutActors);
 
 	void RefreshAfterGraphChanged(UPCGGraphInterface* InGraph, bool bIsStructural, bool bDirtyInputs);
 	void OnGraphChanged(UPCGGraphInterface* InGraph, EPCGChangeType ChangeType);
@@ -342,13 +343,13 @@ private:
 	bool RemoveTrackedActor(AActor* InActor);
 	bool UpdateTrackedActor(AActor* InActor);
 	bool DirtyTrackedActor(AActor* InActor);
-	void DirtyCacheFromTag(const FName& InTag);
+	bool DirtyCacheFromTag(const FName& InTag, const AActor* InActor, bool bIgnoreCull = false);
 	void DirtyCacheForAllTrackedTags();
 
 	bool GraphUsesLandscapePin() const;
 #endif
 
-	FBox GetGridBounds(AActor* InActor) const;
+	FBox GetGridBounds(const AActor* InActor) const;
 
 	UPROPERTY(Transient, NonPIEDuplicateTransient)
 	TObjectPtr<UPCGData> CachedPCGData = nullptr;
@@ -404,6 +405,8 @@ private:
 	FBox LastGeneratedBoundsPriorToUndo = FBox(EForceInit::ForceInit);
 	FPCGTagToSettingsMap CachedTrackedTagsToSettings;
 
+	TMap<FName, bool> CachedTrackedTagsToCulling;
+
 	void SetupLandscapeTracking();
 	void TeardownLandscapeTracking();
 	void UpdateTrackedLandscape(bool bBoundsCheck = true);
diff --git a/Engine/Plugins/Experimental/PCG/Source/PCG/Public/PCGSettings.h b/Engine/Plugins/Experimental/PCG/Source/PCG/Public/PCGSettings.h
--- a/Engine/Plugins/Experimental/PCG/Source/PCG/Public/PCGSettings.h
+++ b/Engine/Plugins/Experimental/PCG/Source/PCG/Public/PCGSettings.h
@@ -18,7 +18,8 @@ class UPCGGraph;
 class UPCGNode;
 class UPCGSettings;
 
-typedef TMap<FName, TSet<TWeakObjectPtr<const UPCGSettings>>> FPCGTagToSettingsMap;
+using FPCGSettingsAndCulling = TPair<TWeakObjectPtr<const UPCGSettings>, bool>;
+using FPCGTagToSettingsMap = TMap<FName, TSet<FPCGSettingsAndCulling>>;
 
 UENUM()
 enum class EPCGSettingsExecutionMode : uint8
