[PCG] Don't track actors if they have no owner (like if they are BP templates)

#rnx

diff --git a/Engine/Plugins/Experimental/PCG/Source/PCG/Private/PCGComponent.cpp b/Engine/Plugins/Experimental/PCG/Source/PCG/Private/PCGComponent.cpp
--- a/Engine/Plugins/Experimental/PCG/Source/PCG/Private/PCGComponent.cpp
+++ b/Engine/Plugins/Experimental/PCG/Source/PCG/Private/PCGComponent.cpp
@@ -1328,6 +1328,12 @@ void UPCGComponent::PostEditUndo()
 
 void UPCGComponent::SetupActorCallbacks()
 {
+	// Without an owner, it probably means we are in a BP template, so no need to setup callbacks
+	if (!GetOwner())
+	{
+		return;
+	}
+
 	GEngine->OnActorMoved().AddUObject(this, &UPCGComponent::OnActorMoved);
 	FCoreUObjectDelegates::OnObjectPropertyChanged.AddUObject(this, &UPCGComponent::OnObjectPropertyChanged);
 }
@@ -1340,6 +1346,12 @@ void UPCGComponent::TeardownActorCallbacks()
 
 void UPCGComponent::SetupTrackingCallbacks()
 {
+	// Without an owner, it probably means we are in a BP template, so no need to setup callbacks
+	if (!GetOwner())
+	{
+		return;
+	}
+
 	CachedTrackedTagsToSettings.Reset();
 	CachedTrackedTagsToCulling.Reset();
 	if (UPCGGraph* PCGGraph = GetGraph())
@@ -1378,6 +1390,13 @@ void UPCGComponent::SetupTrackingCallbacks()
 void UPCGComponent::RefreshTrackingData()
 {
 	TRACE_CPUPROFILER_EVENT_SCOPE(UPCGComponent::RefreshTrackingData);
+
+	// Without an owner, it probably means we are in a BP template, so no need to setup callbacks
+	if (!GetOwner())
+	{
+		return;
+	}
+
 	GetActorsFromTags(ExcludedTags, CachedExcludedActors, /*bCullAgainstLocalBounds=*/true);
 	GetActorsFromTags(CachedTrackedTagsToCulling, CachedTrackedActors);
 	PopulateTrackedActorToTagsMap(/*bForce=*/true);
@@ -1542,6 +1561,12 @@ void UPCGComponent::UpdateTrackedLandscape(bool bBoundsCheck)
 	TeardownLandscapeTracking();
 	TrackedLandscapes.Reset();
 
+	// Without an owner, it probably means we are in a BP template, so no need to setup callbacks
+	if (!GetOwner())
+	{
+		return;
+	}
+
 	if (ALandscapeProxy* Landscape = Cast<ALandscapeProxy>(GetOwner()))
 	{
 		TrackedLandscapes.Add(Landscape);
@@ -2451,7 +2476,14 @@ bool UPCGComponent::AddTrackedActor(AActor* InActor, bool bForce)
 	if (!bForce)
 	{
 		PopulateTrackedActorToTagsMap();
-	}	
+	}
+
+	if (!GetOwner())
+	{
+		// Without an owner, we can't get our bounds, so we won't be able to track.
+		// If we don't have an owner, we are probably in a BP, not instanciated, so no need to track anything.
+		return false;
+	}
 
 	check(InActor);
 	bool bAppliedChange = false;
