[PCG] Workflows - make Debug Only Selected toggle if only selected nodes are currently debugged and all selected nodes are already being debugged.

diff --git a/Engine/Plugins/Experimental/PCG/Source/PCGEditor/Private/PCGEditor.cpp b/Engine/Plugins/Experimental/PCG/Source/PCGEditor/Private/PCGEditor.cpp
--- a/Engine/Plugins/Experimental/PCG/Source/PCGEditor/Private/PCGEditor.cpp
+++ b/Engine/Plugins/Experimental/PCG/Source/PCGEditor/Private/PCGEditor.cpp
@@ -1425,6 +1425,35 @@ void FPCGEditor::OnDebugOnlySelected()
 
 		FScopedTransaction Transaction(*FPCGEditorCommon::ContextIdentifier, LOCTEXT("PCGEditorDebugOnlySelectedTransactionMessage", "PCG Editor: Debug only selected nodes"), nullptr);
 
+		bool bAnyNonSelectedNodesDebugged = false;
+		bool bAllSelectedNodesDebugged = true;
+
+		// Initial pass - inspect state of selected and non-selected nodes.
+		for (UEdGraphNode* Node : PCGEditorGraph->Nodes)
+		{
+			UPCGEditorGraphNodeBase* PCGEditorGraphNode = Cast<UPCGEditorGraphNodeBase>(Node);
+			UPCGNode* PCGNode = PCGEditorGraphNode ? PCGEditorGraphNode->GetPCGNode() : nullptr;
+			UPCGSettingsInterface* PCGSettingsInterface = PCGNode ? PCGNode->GetSettingsInterface() : nullptr;
+
+			if (!PCGEditorGraphNode || !PCGNode || !PCGSettingsInterface)
+			{
+				continue;
+			}
+
+			if (SelectedNodes.Contains(PCGEditorGraphNode))
+			{
+				bAllSelectedNodesDebugged &= PCGSettingsInterface->bDebug;
+			}
+			else
+			{
+				bAnyNonSelectedNodesDebugged |= PCGSettingsInterface->bDebug;
+			}
+		}
+
+		// The selected nodes should be debugged if any non-selected nodes are being debugged, or if the selected
+		// nodes are partially being debugged.
+		const bool bTargetDebugState = bAnyNonSelectedNodesDebugged || !bAllSelectedNodesDebugged;
+
 		for (UEdGraphNode* Node : PCGEditorGraph->Nodes)
 		{
 			UPCGEditorGraphNodeBase* PCGEditorGraphNode = Cast<UPCGEditorGraphNodeBase>(Node);
@@ -1436,7 +1465,8 @@ void FPCGEditor::OnDebugOnlySelected()
 				continue;
 			}
 
-			const bool bShouldBeDebug = SelectedNodes.Contains(PCGEditorGraphNode);
+			// Selected set to target state, non-selected should not be debugged.
+			const bool bShouldBeDebug = SelectedNodes.Contains(PCGEditorGraphNode) ? bTargetDebugState : false;
 
 			if (PCGSettingsInterface->bDebug != bShouldBeDebug)
 			{
