Fixed issue in prepare data reentrance
Added time-slicing capability in prepare data phase for static mesh spawner
Added threading in spline interior sampling

diff --git a/Engine/Plugins/Experimental/PCG/Source/PCG/Private/Elements/PCGSplineSampler.cpp b/Engine/Plugins/Experimental/PCG/Source/PCG/Private/Elements/PCGSplineSampler.cpp
--- a/Engine/Plugins/Experimental/PCG/Source/PCG/Private/Elements/PCGSplineSampler.cpp
+++ b/Engine/Plugins/Experimental/PCG/Source/PCG/Private/Elements/PCGSplineSampler.cpp
@@ -13,6 +13,7 @@
 #include "Data/PCGSpatialData.h"
 #include "Helpers/PCGBlueprintHelpers.h"
 
+#include "Async/ParallelFor.h"
 #include "Components/SplineComponent.h"
 #include "Voronoi/Voronoi.h"
 
@@ -693,9 +694,9 @@ namespace PCGSplineSampler
 		}
 	}
  
-	void SampleInteriorData(const UPCGPolyLineData* LineData, const UPCGSpatialData* SpatialData, const UPCGSpatialData* InBoundingShape, const FPCGSplineSamplerParams& Params, UPCGPointData* OutPointData)
+	void SampleInteriorData(FPCGContext* Context, const UPCGPolyLineData* LineData, const UPCGSpatialData* SpatialData, const UPCGSpatialData* InBoundingShape, const FPCGSplineSamplerParams& Params, UPCGPointData* OutPointData)
 	{
-		check(LineData && OutPointData);
+		check(Context && LineData && OutPointData);
 
 		TSoftObjectPtr<USplineComponent> Spline;
 
@@ -747,13 +748,17 @@ namespace PCGSplineSampler
 		{
 			// Treat spline interface points as vertices of a polyline
 			const int NumSegments = LineData->GetNumSegments();
+			SplineSamplePoints.Reserve(NumSegments);
+
 			for (int32 SegmentIndex = 0; SegmentIndex < NumSegments; ++SegmentIndex)
 			{
 				SplineSamplePoints.Add(LineData->GetLocationAtDistance(SegmentIndex, 0, /*bWorldSpace=*/false));
 			}
 		}
-		else
+		else if(Params.InteriorBorderSampleSpacing > 0)
 		{
+			SplineSamplePoints.Reserve(1 + SplineLength / Params.InteriorBorderSampleSpacing);
+
 			// Get sample points along the spline that are higher resolution than our PolyLine
 			for (FVector::FReal Length = 0.f; Length < SplineLength; Length += Params.InteriorBorderSampleSpacing)
 			{
@@ -763,6 +768,8 @@ namespace PCGSplineSampler
 
 		// Flat polygon representation of our spline points
 		TArray<FVector2D> SplineSamplePoints2D;
+		SplineSamplePoints2D.Reserve(SplineSamplePoints.Num());
+
 		for (const FVector& Point : SplineSamplePoints)
 		{
 			SplineSamplePoints2D.Add(FVector2D(Point));
@@ -777,6 +784,7 @@ namespace PCGSplineSampler
 			if (Params.bTreatSplineAsPolyline)
 			{
 				// If we already computed the polygon, use a copy instead of generating it again
+				PolygonPoints.Reserve(SplineSamplePoints.Num());
 				for (const FVector& SplinePoint : SplineSamplePoints)
 				{
 					FVector& PolygonPoint = PolygonPoints.Add_GetRef(SplinePoint);
@@ -786,6 +794,8 @@ namespace PCGSplineSampler
 			else
 			{
 				const int NumSegments = LineData->GetNumSegments();
+				PolygonPoints.Reserve(NumSegments);
+
 				for (int32 SegmentIndex = 0; SegmentIndex < NumSegments; ++SegmentIndex)
 				{
 					FVector& PolygonPoint = PolygonPoints.Add_GetRef(LineData->GetLocationAtDistance(SegmentIndex, 0, /*bWorldSpace=*/false));
@@ -794,6 +804,8 @@ namespace PCGSplineSampler
 			}
 
 			TArray<FVector2D> PolygonPoints2D;
+			PolygonPoints2D.Reserve(PolygonPoints.Num());
+
 			for (const FVector& Point : PolygonPoints)
 			{
 				PolygonPoints2D.Add(FVector2D(Point));
@@ -802,7 +814,10 @@ namespace PCGSplineSampler
 			TArray<TTuple<FVector, FVector>> VoronoiEdges;
 			TArray<int32> CellMember;
 
-			GetVoronoiEdges(PolygonPoints, FBox(MinPoint, FVector(MaxPoint.X, MaxPoint.Y, MinPoint.Z)), VoronoiEdges, CellMember);
+			{
+				TRACE_CPUPROFILER_EVENT_SCOPE(FPCGSplineSamplerElement::Execute::GetVoronoiEdges);
+				GetVoronoiEdges(PolygonPoints, FBox(MinPoint, FVector(MaxPoint.X, MaxPoint.Y, MinPoint.Z)), VoronoiEdges, CellMember);
+			}
 
 			// Find the subset of the Voronoi Diagram which composes the Medial Axis
 			for (const TTuple<FVector, FVector>& Edge : VoronoiEdges)
@@ -833,157 +848,190 @@ namespace PCGSplineSampler
 		}
 
 		TArray<FPCGPoint>& OutPoints = OutPointData->GetMutablePoints();
-		TArray<FVector::FReal> DistancesSquaredOnXY;
 
 		// RayPadding should be large enough to avoid potential misclassifications in SegmentPolygonIntersection2D
 		const FVector::FReal RayPadding = 100.f;
 		const FVector::FReal MinY = FMath::CeilToDouble(MinPoint.Y / Params.InteriorSampleSpacing) * Params.InteriorSampleSpacing;
 		const FVector::FReal MaxY = FMath::FloorToDouble(MaxPoint.Y / Params.InteriorSampleSpacing) * Params.InteriorSampleSpacing;
 
-		// Point sampling
-		for (FVector::FReal Y = MinY; Y < MaxY + UE_KINDA_SMALL_NUMBER; Y += Params.InteriorSampleSpacing)
-		{
-			const FVector2D RayMin(MinPoint.X - RayPadding, Y);
-			const FVector2D RayMax(MaxPoint.X + RayPadding, Y);
+		constexpr int32 MinIterationPerDispatch = 4;
+		const int32 NumIterations = (MaxY + UE_KINDA_SMALL_NUMBER - MinY) / Params.InteriorSampleSpacing;
+		const int32 NumDispatch = FMath::Max(1, FMath::Min(Context->NumAvailableTasks, NumIterations / MinIterationPerDispatch));
+		const int32 NumIterationsPerDispatch = NumIterations / NumDispatch;
 
-			// Get the intersections along this ray, sorted by Y value
-			TArray<FVector2D> Intersections;
-			PCGSplineSamplerHelpers::SegmentPolygonIntersection2D(RayMin, RayMax, SplineSamplePoints2D, Intersections);
+		TArray<TArray<TTuple<FTransform, FVector, float>>> InteriorSplinePointData;
+		InteriorSplinePointData.SetNum(NumDispatch);
 
-			if (Intersections.Num() % 2 != 0)
-			{
-				UE_LOG(LogPCG, Error, TEXT("Intersection test failed, skipping samples for this row"));
-				continue;
-			}
+		ParallelFor(NumDispatch, [&](int32 DispatchIndex)
+		{
+			const bool bIsLastIteration = (DispatchIndex == (NumDispatch - 1));
+			const int32 StartIterationIndex = DispatchIndex * NumIterationsPerDispatch;
+			const int32 EndIterationIndex = bIsLastIteration ? NumIterations : (StartIterationIndex + NumIterationsPerDispatch);
+
+			const FVector::FReal LocalMinY = MinY + StartIterationIndex * Params.InteriorSampleSpacing;
+			const FVector::FReal LocalMaxY = bIsLastIteration ? (MaxY + UE_KINDA_SMALL_NUMBER) : (MinY + EndIterationIndex * Params.InteriorSampleSpacing);
 
-			Intersections.Sort([](const FVector2D& LHS, const FVector2D& RHS) { return LHS.X < RHS.X; });
+			TArray<FVector::FReal> DistancesSquaredOnXY;
 
-			// TODO: async processing
-			// Each pair of intersections defines a range in which point samples may lie
-			for (int32 RangeIndex = 0; RangeIndex < Intersections.Num(); RangeIndex += 2)
+			// Point sampling
+			for(FVector::FReal Y = LocalMinY; Y < LocalMaxY; Y += Params.InteriorSampleSpacing)
 			{
-				const FVector::FReal MinX = FMath::CeilToDouble(Intersections[RangeIndex].X / Params.InteriorSampleSpacing) * Params.InteriorSampleSpacing;
-				const FVector::FReal MaxX = FMath::FloorToDouble(Intersections[RangeIndex + 1].X / Params.InteriorSampleSpacing) * Params.InteriorSampleSpacing;
+				const FVector2D RayMin(MinPoint.X - RayPadding, Y);
+				const FVector2D RayMax(MaxPoint.X + RayPadding, Y);
 
-				if (MaxX - MinX < Params.InteriorSampleSpacing)
+				// Get the intersections along this ray, sorted by Y value
+				TArray<FVector2D> Intersections;
+				PCGSplineSamplerHelpers::SegmentPolygonIntersection2D(RayMin, RayMax, SplineSamplePoints2D, Intersections);
+
+				if (Intersections.Num() % 2 != 0)
 				{
+					UE_LOG(LogPCG, Error, TEXT("Intersection test failed, skipping samples for this row"));
 					continue;
 				}
 
-				for (FVector::FReal X = MinX; X < MaxX + UE_KINDA_SMALL_NUMBER; X += Params.InteriorSampleSpacing)
+				Intersections.Sort([](const FVector2D& LHS, const FVector2D& RHS) { return LHS.X < RHS.X; });
+
+				// TODO: async processing
+				// Each pair of intersections defines a range in which point samples may lie
+				for (int32 RangeIndex = 0; RangeIndex < Intersections.Num(); RangeIndex += 2)
 				{
-					FPCGPoint& Point = OutPoints.Emplace_GetRef();
+					const FVector::FReal MinX = FMath::CeilToDouble(Intersections[RangeIndex].X / Params.InteriorSampleSpacing) * Params.InteriorSampleSpacing;
+					const FVector::FReal MaxX = FMath::FloorToDouble(Intersections[RangeIndex + 1].X / Params.InteriorSampleSpacing) * Params.InteriorSampleSpacing;
 
-					const FVector2D SampleLocation = FVector2D(X, Y);
-					FVector SurfaceLocation = FVector(SampleLocation, MinPoint.Z);
+					if (MaxX - MinX < Params.InteriorSampleSpacing)
+					{
+						continue;
+					}
 
-					if (bProjectOntoSurface)
+					for (FVector::FReal X = MinX; X < MaxX + UE_KINDA_SMALL_NUMBER; X += Params.InteriorSampleSpacing)
 					{
-						// Precompute 2D distance to every spline point
-						DistancesSquaredOnXY.Reset(SplineSamplePoints.Num());
-						for (const FVector& SplinePoint : SplineSamplePoints)
-						{
-							FVector::FReal DistanceSquared = FVector::DistSquaredXY(SurfaceLocation, SplinePoint);
-							DistancesSquaredOnXY.Add(DistanceSquared);
-						}
+						const FVector2D SampleLocation = FVector2D(X, Y);
+						FVector SurfaceLocation = FVector(SampleLocation, MinPoint.Z);
 
-						// Compute average Z value weighted by 1 / Distance^2
-						FVector::FReal SumZ = 0.f;
-						FVector::FReal SumWeights = 0.f;
-						for (int32 PointIndex = 0; PointIndex < SplineSamplePoints.Num(); ++PointIndex)
+						if (bProjectOntoSurface)
 						{
-							const FVector::FReal DistanceSquared = DistancesSquaredOnXY[PointIndex];
-
-							// If sample point overlaps exactly with a border point, then that must be the height
-							if (FMath::IsNearlyZero(DistanceSquared))
+							// Precompute 2D distance to every spline point
+							DistancesSquaredOnXY.Reset(SplineSamplePoints.Num());
+							for (const FVector& SplinePoint : SplineSamplePoints)
 							{
-								SumZ = SplineSamplePoints[PointIndex].Z;
-								SumWeights = 1.f;
-								break;
+								FVector::FReal DistanceSquared = FVector::DistSquaredXY(SurfaceLocation, SplinePoint);
+								DistancesSquaredOnXY.Add(DistanceSquared);
 							}
 
-							// TODO: it would be more accurate to use distance to the polyline instead of distance to the polyline points,
-							// however it would also be much more expensive. Perhaps worth investigating when Params.bTreatSplineAsPolyline is true
-							const FVector::FReal Weight = 1.f / DistanceSquared;
+							// Compute average Z value weighted by 1 / Distance^2
+							FVector::FReal SumZ = 0.f;
+							FVector::FReal SumWeights = 0.f;
+							for (int32 PointIndex = 0; PointIndex < SplineSamplePoints.Num(); ++PointIndex)
+							{
+								const FVector::FReal DistanceSquared = DistancesSquaredOnXY[PointIndex];
 
-							SumWeights += Weight;
-							SumZ += SplineSamplePoints[PointIndex].Z * Weight;
-						}
+								// If sample point overlaps exactly with a border point, then that must be the height
+								if (FMath::IsNearlyZero(DistanceSquared))
+								{
+									SumZ = SplineSamplePoints[PointIndex].Z;
+									SumWeights = 1.f;
+									break;
+								}
 
-						SurfaceLocation.Z = SumZ / SumWeights;
-					}
+								// TODO: it would be more accurate to use distance to the polyline instead of distance to the polyline points,
+								// however it would also be much more expensive. Perhaps worth investigating when Params.bTreatSplineAsPolyline is true
+								const FVector::FReal Weight = 1.f / DistanceSquared;
 
-					// if bTreatAsPolyline, then we shouldnt use this, we should use nearest point on the polygon line segments
-					float Dummy;
-					const float NearestSplineKey = bFindNearestSplineKey ? Spline->SplineCurves.Position.InaccurateFindNearest(SurfaceLocation, Dummy) : 0.f;
+								SumWeights += Weight;
+								SumZ += SplineSamplePoints[PointIndex].Z * Weight;
+							}
 
-					const FVector PointLocationLS = (Params.bProjectOntoSurface ? FVector(SurfaceLocation) : FVector(SampleLocation, MinPoint.Z));
+							SurfaceLocation.Z = SumZ / SumWeights;
+						}
 
-					FTransform TransformLS = FTransform::Identity;
-					TransformLS.SetLocation(PointLocationLS);
+						// if bTreatAsPolyline, then we shouldnt use this, we should use nearest point on the polygon line segments
+						float Dummy;
+						const float NearestSplineKey = bFindNearestSplineKey ? Spline->SplineCurves.Position.InaccurateFindNearest(SurfaceLocation, Dummy) : 0.f;
 
-					if (Params.InteriorOrientation == EPCGSplineSamplingInteriorOrientation::FollowCurvature)
-					{
-						TransformLS.SetRotation(Spline->GetQuaternionAtSplineInputKey(NearestSplineKey, ESplineCoordinateSpace::Local));
-					}
+						const FVector PointLocationLS = (Params.bProjectOntoSurface ? FVector(SurfaceLocation) : FVector(SampleLocation, MinPoint.Z));
 
-					// Calculate density fall off
-					if (bGenerateMedialAxis && MedialAxisEdges.Num())
-					{
-						FVector::FReal SmallestDistSquared = MaxDimensionSquared;
+						FTransform TransformLS = FTransform::Identity;
+						TransformLS.SetLocation(PointLocationLS);
 
-						// Find distance from SampleLocation to MedialAxis
-						for (const TTuple<FVector2D, FVector2D>& Edge : MedialAxisEdges)
+						if (Params.InteriorOrientation == EPCGSplineSamplingInteriorOrientation::FollowCurvature)
 						{
-							const FVector::FReal DistSquared = FMath::PointDistToSegmentSquared(FVector(SampleLocation, 0), FVector(Edge.Get<0>(), 0), FVector(Edge.Get<1>(), 0));
-
-							if (DistSquared < SmallestDistSquared)
-							{
-								SmallestDistSquared = DistSquared;
-							}
+							TransformLS.SetRotation(Spline->GetQuaternionAtSplineInputKey(NearestSplineKey, ESplineCoordinateSpace::Local));
 						}
 
-						const FVector::FReal SmallestDist = FMath::Sqrt(SmallestDistSquared);
-
-						FVector::FReal PointToSplineDist = 0.f;
+						// Calculate density fall off
+						float Density = 1.0f;
 
-						if (Params.bTreatSplineAsPolyline)
+						if (bGenerateMedialAxis && MedialAxisEdges.Num())
 						{
-							SmallestDistSquared = MaxDimensionSquared;
-							for (int32 PointIndex = 0; PointIndex < SplineSamplePoints.Num(); ++PointIndex)
+							FVector::FReal SmallestDistSquared = MaxDimensionSquared;
+
+							// Find distance from SampleLocation to MedialAxis
+							for (const TTuple<FVector2D, FVector2D>& Edge : MedialAxisEdges)
 							{
-								const FVector::FReal DistSquared = FMath::PointDistToSegmentSquared(SurfaceLocation, SplineSamplePoints[PointIndex], SplineSamplePoints[(PointIndex + 1) % SplineSamplePoints.Num()]);
+								const FVector::FReal DistSquared = FMath::PointDistToSegmentSquared(FVector(SampleLocation, 0), FVector(Edge.Get<0>(), 0), FVector(Edge.Get<1>(), 0));
 
 								if (DistSquared < SmallestDistSquared)
 								{
 									SmallestDistSquared = DistSquared;
 								}
 							}
-							
-							PointToSplineDist = FMath::Sqrt(SmallestDistSquared);
-						}
-						else
-						{
-							const FVector NearestSplineLocation = Spline->GetLocationAtSplineInputKey(NearestSplineKey, ESplineCoordinateSpace::Local);
-							PointToSplineDist = FVector2D::Distance(SampleLocation, FVector2D(NearestSplineLocation.X, NearestSplineLocation.Y));
-						}
 
-						// Linear fall off in the range [0, 1]
-						const FVector::FReal T = SmallestDist / (SmallestDist + PointToSplineDist + UE_KINDA_SMALL_NUMBER);
+							const FVector::FReal SmallestDist = FMath::Sqrt(SmallestDistSquared);
 
-						Point.Density = DensityFalloffCurve->Eval(T);
-					}
-					else
-					{
-						Point.Density = 1.f;
-					}
+							FVector::FReal PointToSplineDist = 0.f;
+
+							if (Params.bTreatSplineAsPolyline)
+							{
+								SmallestDistSquared = MaxDimensionSquared;
+								for (int32 PointIndex = 0; PointIndex < SplineSamplePoints.Num(); ++PointIndex)
+								{
+									const FVector::FReal DistSquared = FMath::PointDistToSegmentSquared(SurfaceLocation, SplineSamplePoints[PointIndex], SplineSamplePoints[(PointIndex + 1) % SplineSamplePoints.Num()]);
 
-					Point.BoundsMin = BoundsMin;
-					Point.BoundsMax = BoundsMax;
-					Point.Transform = TransformLS * LineData->GetTransform();
-					SetSeed(Point, PointLocationLS, Params);
+									if (DistSquared < SmallestDistSquared)
+									{
+										SmallestDistSquared = DistSquared;
+									}
+								}
+
+								PointToSplineDist = FMath::Sqrt(SmallestDistSquared);
+							}
+							else
+							{
+								const FVector NearestSplineLocation = Spline->GetLocationAtSplineInputKey(NearestSplineKey, ESplineCoordinateSpace::Local);
+								PointToSplineDist = FVector2D::Distance(SampleLocation, FVector2D(NearestSplineLocation.X, NearestSplineLocation.Y));
+							}
+
+							// Linear fall off in the range [0, 1]
+							const FVector::FReal T = SmallestDist / (SmallestDist + PointToSplineDist + UE_KINDA_SMALL_NUMBER);
+							Density = DensityFalloffCurve->Eval(T);
+						}
+
+						InteriorSplinePointData[DispatchIndex].Emplace(TransformLS * LineData->GetTransform(), PointLocationLS, Density);
+					}
 				}
 			}
+		});
+
+		// Finally, gather the data and push to the points
+		int32 PointCount = 0;
+		for (const TArray<TTuple<FTransform, FVector, float>>& InteriorData : InteriorSplinePointData)
+		{
+			PointCount += InteriorData.Num();
+		}
+
+		// TODO: should we parallel for this too?
+		OutPoints.Reserve(PointCount);
+		for (const TArray<TTuple<FTransform, FVector, float>>& InteriorData : InteriorSplinePointData)
+		{
+			for (const TTuple<FTransform, FVector, float>& InteriorPoint : InteriorData)
+			{
+				FPCGPoint& Point = OutPoints.Emplace_GetRef();
+				Point.Transform = InteriorPoint.Get<0>();
+				SetSeed(Point, InteriorPoint.Get<1>(), Params);
+				Point.Density = InteriorPoint.Get<2>();
+				Point.BoundsMin = BoundsMin;
+				Point.BoundsMax = BoundsMax;
+			}
 		}
 	}
 
@@ -1095,7 +1143,7 @@ bool FPCGSplineSamplerElement::ExecuteInternal(FPCGContext* Context) const
 
 		if (SamplerParams.Dimension == EPCGSplineSamplingDimension::OnInterior)
 		{
-			PCGSplineSampler::SampleInteriorData(LineData, SpatialData, BoundingShapeSpatialInput, SamplerParams, SampledPointData);
+			PCGSplineSampler::SampleInteriorData(Context, LineData, SpatialData, BoundingShapeSpatialInput, SamplerParams, SampledPointData);
 		}
 		else
 		{
diff --git a/Engine/Plugins/Experimental/PCG/Source/PCG/Private/Elements/PCGStaticMeshSpawner.cpp b/Engine/Plugins/Experimental/PCG/Source/PCG/Private/Elements/PCGStaticMeshSpawner.cpp
--- a/Engine/Plugins/Experimental/PCG/Source/PCG/Private/Elements/PCGStaticMeshSpawner.cpp
+++ b/Engine/Plugins/Experimental/PCG/Source/PCG/Private/Elements/PCGStaticMeshSpawner.cpp
@@ -71,7 +71,7 @@ bool FPCGStaticMeshSpawnerElement::PrepareDataInternal(FPCGContext* InContext) c
 	// Check if we can reuse existing resources
 	bool bSkippedDueToReuse = false;
 
-	if (CVarAllowISMReuse.GetValueOnAnyThread())
+	if (Context->CurrentInputIndex == 0 && CVarAllowISMReuse.GetValueOnAnyThread())
 	{
 		// Compute CRC if it has not been computed (it likely isn't, but this is to futureproof this)
 		if (!Context->DependenciesCrc.IsValid())
@@ -115,9 +115,14 @@ bool FPCGStaticMeshSpawnerElement::PrepareDataInternal(FPCGContext* InContext) c
 	TArray<FPCGTaggedData> Inputs = Context->InputData.GetInputs();
 	TArray<FPCGTaggedData>& Outputs = Context->OutputData.TaggedData;
 
-	for (const FPCGTaggedData& Input : Inputs)
+	while(Context->CurrentInputIndex < Inputs.Num())
 	{
+		const FPCGTaggedData& Input = Inputs[Context->CurrentInputIndex];
 		const UPCGSpatialData* SpatialData = Cast<UPCGSpatialData>(Input.Data);
+
+		// Preincrement so we can have all paths continue properly - note that this works because none of this is internally time-sliced
+		++Context->CurrentInputIndex;
+
 		if (!SpatialData)
 		{
 			PCGE_LOG(Error, "Invalid input data");
@@ -182,6 +187,12 @@ bool FPCGStaticMeshSpawnerElement::PrepareDataInternal(FPCGContext* InContext) c
 		InstanceListData.SpatialData = PointData;
 		InstanceListData.MeshInstances = MoveTemp(MeshInstances);
 		InstanceListData.PackedCustomData = MoveTemp(PackedCustomData);
+
+		// Continue on to next iteration if there is time left, otherwise, exit here
+		if (Context->ShouldStop() && Context->CurrentInputIndex < Inputs.Num())
+		{
+			return false;
+		}
 	}
 
 	return true;
diff --git a/Engine/Plugins/Experimental/PCG/Source/PCG/Private/PCGElement.cpp b/Engine/Plugins/Experimental/PCG/Source/PCG/Private/PCGElement.cpp
--- a/Engine/Plugins/Experimental/PCG/Source/PCG/Private/PCGElement.cpp
+++ b/Engine/Plugins/Experimental/PCG/Source/PCG/Private/PCGElement.cpp
@@ -34,18 +34,20 @@ bool IPCGElement::Execute(FPCGContext* Context) const
 
 		switch (Context->CurrentPhase)
 		{
-			case EPCGExecutionPhase::NotExecuted: // Fall-through
+			case EPCGExecutionPhase::NotExecuted:
 			{
 				TRACE_CPUPROFILER_EVENT_SCOPE(EPCGExecutionPhase::NotExecuted);
 				PreExecute(Context);
+
+				// Will override the settings if there is any override.
+				Context->OverrideSettings();
+
 				break;
 			}
 
 			case EPCGExecutionPhase::PrepareData:
 			{
 				TRACE_CPUPROFILER_EVENT_SCOPE(EPCGExecutionPhase::PrepareData);
-				// Will override the settings if there is any override.
-				Context->OverrideSettings();
 
 				if (PrepareDataInternal(Context))
 				{
diff --git a/Engine/Plugins/Experimental/PCG/Source/PCG/Public/Elements/PCGStaticMeshSpawner.h b/Engine/Plugins/Experimental/PCG/Source/PCG/Public/Elements/PCGStaticMeshSpawner.h
--- a/Engine/Plugins/Experimental/PCG/Source/PCG/Public/Elements/PCGStaticMeshSpawner.h
+++ b/Engine/Plugins/Experimental/PCG/Source/PCG/Public/Elements/PCGStaticMeshSpawner.h
@@ -114,6 +114,8 @@ struct FPCGStaticMeshSpawnerContext : public FPCGContext
 	};
 
 	TArray<FPackedInstanceListData> MeshInstancesData;
+	// Index of input in the context (for selection)
+	int32 CurrentInputIndex = 0;
 	// Index of MeshInstances/PackedCustomData in last MeshInstancesData element
 	int32 CurrentDataIndex = 0;
 };
