[PCG] Cache memory accounting first pass

Measures memory used by *unique* data objects that are stored in the cache. Uniqueness handled by keeping records of what has been added.

GetResourceSizeEx is overridden to calculate memory used by a single object. Used in conjunction with VisitDataNetwork to measure data usage of all objects in expression tree.

It is holding memory usage steady when regenerating graphs.

Does not account for metadata memory usage yet, so actual memory will be larger than budget by some proportion.

diff --git a/Engine/Plugins/Experimental/PCG/Source/PCG/Private/Data/PCGDifferenceData.cpp b/Engine/Plugins/Experimental/PCG/Source/PCG/Private/Data/PCGDifferenceData.cpp
--- a/Engine/Plugins/Experimental/PCG/Source/PCG/Private/Data/PCGDifferenceData.cpp
+++ b/Engine/Plugins/Experimental/PCG/Source/PCG/Private/Data/PCGDifferenceData.cpp
@@ -89,6 +89,17 @@ void UPCGDifferenceData::PostEditChangeProperty(FPropertyChangedEvent& PropertyC
 }
 #endif
 
+void UPCGDifferenceData::VisitDataNetwork(TFunctionRef<void(const UPCGData*)> Action) const
+{
+	check(Source);
+	Source->VisitDataNetwork(Action);
+
+	if (Difference)
+	{
+		Difference->VisitDataNetwork(Action);
+	}
+}
+
 int UPCGDifferenceData::GetDimension() const
 {
 	return Source->GetDimension();
diff --git a/Engine/Plugins/Experimental/PCG/Source/PCG/Private/Data/PCGIntersectionData.cpp b/Engine/Plugins/Experimental/PCG/Source/PCG/Private/Data/PCGIntersectionData.cpp
--- a/Engine/Plugins/Experimental/PCG/Source/PCG/Private/Data/PCGIntersectionData.cpp
+++ b/Engine/Plugins/Experimental/PCG/Source/PCG/Private/Data/PCGIntersectionData.cpp
@@ -47,6 +47,13 @@ void UPCGIntersectionData::Initialize(const UPCGSpatialData* InA, const UPCGSpat
 	}
 }
 
+void UPCGIntersectionData::VisitDataNetwork(TFunctionRef<void(const UPCGData*)> Action) const
+{
+	check(A && B);
+	A->VisitDataNetwork(Action);
+	B->VisitDataNetwork(Action);
+}
+
 int UPCGIntersectionData::GetDimension() const
 {
 	check(A && B);
diff --git a/Engine/Plugins/Experimental/PCG/Source/PCG/Private/Data/PCGPointData.cpp b/Engine/Plugins/Experimental/PCG/Source/PCG/Private/Data/PCGPointData.cpp
--- a/Engine/Plugins/Experimental/PCG/Source/PCG/Private/Data/PCGPointData.cpp
+++ b/Engine/Plugins/Experimental/PCG/Source/PCG/Private/Data/PCGPointData.cpp
@@ -241,6 +241,13 @@ FPCGPointRef::FPCGPointRef(const FPCGPointRef& InPointRef)
 	Bounds = InPointRef.Bounds;
 }
 
+void UPCGPointData::GetResourceSizeEx(FResourceSizeEx& CumulativeResourceSize)
+{
+	Super::GetResourceSizeEx(CumulativeResourceSize);
+
+	CumulativeResourceSize.AddDedicatedSystemMemoryBytes(Points.GetAllocatedSize() + Octree.GetSizeBytes() + sizeof(Bounds));
+}
+
 TArray<FPCGPoint>& UPCGPointData::GetMutablePoints()
 {
 	bOctreeIsDirty = true;
diff --git a/Engine/Plugins/Experimental/PCG/Source/PCG/Private/Data/PCGSpatialData.cpp b/Engine/Plugins/Experimental/PCG/Source/PCG/Private/Data/PCGSpatialData.cpp
--- a/Engine/Plugins/Experimental/PCG/Source/PCG/Private/Data/PCGSpatialData.cpp
+++ b/Engine/Plugins/Experimental/PCG/Source/PCG/Private/Data/PCGSpatialData.cpp
@@ -15,6 +15,26 @@ UPCGSpatialData::UPCGSpatialData(const FObjectInitializer& ObjectInitializer)
 	Metadata = ObjectInitializer.CreateDefaultSubobject<UPCGMetadata>(this, TEXT("Metadata"));
 }
 
+void UPCGSpatialDataWithPointCache::GetResourceSizeEx(FResourceSizeEx& CumulativeResourceSize)
+{
+	Super::GetResourceSizeEx(CumulativeResourceSize);
+
+	if (CachedPointData)
+	{
+		const_cast<UPCGPointData*>(CachedPointData.Get())->GetResourceSizeEx(CumulativeResourceSize);
+	}
+
+	CumulativeResourceSize.AddDedicatedSystemMemoryBytes(CachedBoundedPointDataBoxes.GetAllocatedSize() + CachedBoundedPointData.GetAllocatedSize());
+
+	for (const UPCGPointData* Data : CachedBoundedPointData)
+	{
+		if (Data)
+		{
+			const_cast<UPCGPointData*>(Data)->GetResourceSizeEx(CumulativeResourceSize);
+		}
+	}
+}
+
 const UPCGPointData* UPCGSpatialDataWithPointCache::ToPointData(FPCGContext* Context, const FBox& InBounds) const
 {
 	if (InBounds.IsValid && SupportsBoundedPointData())
@@ -59,6 +79,16 @@ const UPCGPointData* UPCGSpatialDataWithPointCache::ToPointData(FPCGContext* Con
 	}
 }
 
+void UPCGSpatialData::GetResourceSizeEx(FResourceSizeEx& CumulativeResourceSize)
+{
+	Super::GetResourceSizeEx(CumulativeResourceSize);
+
+	if (Metadata)
+	{
+		Metadata->GetResourceSizeEx(CumulativeResourceSize);
+	}
+}
+
 float UPCGSpatialData::GetDensityAtPosition(const FVector& InPosition) const
 {
 	FPCGPoint TemporaryPoint;
diff --git a/Engine/Plugins/Experimental/PCG/Source/PCG/Private/Data/PCGUnionData.cpp b/Engine/Plugins/Experimental/PCG/Source/PCG/Private/Data/PCGUnionData.cpp
--- a/Engine/Plugins/Experimental/PCG/Source/PCG/Private/Data/PCGUnionData.cpp
+++ b/Engine/Plugins/Experimental/PCG/Source/PCG/Private/Data/PCGUnionData.cpp
@@ -70,6 +70,17 @@ void UPCGUnionData::AddData(const UPCGSpatialData* InData)
 	}
 }
 
+void UPCGUnionData::VisitDataNetwork(TFunctionRef<void(const UPCGData*)> Action) const
+{
+	for (TObjectPtr<const UPCGSpatialData> Datum : Data)
+	{
+		if (Datum)
+		{
+			Datum->VisitDataNetwork(Action);
+		}
+	}
+}
+
 int UPCGUnionData::GetDimension() const
 {
 	return CachedDimension;
diff --git a/Engine/Plugins/Experimental/PCG/Source/PCG/Private/Graph/PCGGraphCache.cpp b/Engine/Plugins/Experimental/PCG/Source/PCG/Private/Graph/PCGGraphCache.cpp
--- a/Engine/Plugins/Experimental/PCG/Source/PCG/Private/Graph/PCGGraphCache.cpp
+++ b/Engine/Plugins/Experimental/PCG/Source/PCG/Private/Graph/PCGGraphCache.cpp
@@ -14,15 +14,28 @@ static TAutoConsoleVariable<bool> CVarCacheDebugging(
 	false,
 	TEXT("Enable various features for debugging the graph cache system."));
 
-FPCGGraphCacheEntry::FPCGGraphCacheEntry(const FPCGCrc& InDependenciesCrc, const FPCGDataCollection& InOutput, FPCGRootSet& OutRootSet)
-	: Output(InOutput)
-	, DependenciesCrc(InDependenciesCrc)
-{
-	Output.AddToRootSet(OutRootSet);
-}
+static TAutoConsoleVariable<int32> CVarCacheMemoryBudgetMB(
+	TEXT("pcg.Cache.MemoryBudgetMB"),
+	6144,
+	TEXT("Memory budget for data in cache (MB)."));
+
+static TAutoConsoleVariable<bool> CVarCacheMemoryBudgetEnabled(
+	TEXT("pcg.Cache.EnableMemoryBudget"),
+	true,
+	TEXT("Whether memory budget is enforced (items purged from cache to respect pcg.Cache.MemoryBudgetMB."));
+
+static int32 GPCGGraphCacheMaxElements = 16384;
+static FAutoConsoleVariableRef CVarNiagaraGraphDataCacheSize(
+	TEXT("pcg.Cache.GraphCacheMaxElements"),
+	GPCGGraphCacheMaxElements,
+	TEXT("Maximum number of elements to store within the graph cache."),
+	ECVF_ReadOnly
+);
 
 FPCGGraphCache::FPCGGraphCache(TWeakObjectPtr<UObject> InOwner, FPCGRootSet* InRootSet)
-	: Owner(InOwner), RootSet(InRootSet)
+	: CacheData(GPCGGraphCacheMaxElements)
+	, Owner(InOwner)
+	, RootSet(InRootSet)
 {
 	check(InOwner.Get() && InRootSet);
 }
@@ -47,41 +60,33 @@ bool FPCGGraphCache::GetFromCache(const UPCGNode* InNode, const IPCGElement* InE
 
 	const bool bDebuggingEnabled = IsDebuggingEnabled() && InComponent && InComponent->GetOwner() && InNode;
 
-	TRACE_CPUPROFILER_EVENT_SCOPE(FPCGGraphCache::GetFromCache);
-	FReadScopeLock ScopedReadLock(CacheLock);
-
-	if (const FPCGGraphCacheEntries* Entries = CacheData.Find(InElement))
 	{
-		for (const FPCGGraphCacheEntry& Entry : *Entries)
+		TRACE_CPUPROFILER_EVENT_SCOPE(FPCGGraphCache::GetFromCache);
+		FReadScopeLock ScopedReadLock(CacheLock);
+
+		FPCGCacheEntryKey CacheKey(InElement, InDependenciesCrc);
+		if (const FPCGDataCollection* Value = const_cast<FPCGGraphCache*>(this)->CacheData.FindAndTouch(CacheKey))
 		{
-			if (Entry.DependenciesCrc == InDependenciesCrc)
+			if (bDebuggingEnabled)
 			{
-				OutOutput = Entry.Output;
+				// Leading spaces to align log content with warnings below - helps readability a lot.
+				UE_LOG(LogPCG, Log, TEXT("         [%s] %s\t\tCACHE HIT %u"), *InComponent->GetOwner()->GetName(), *InNode->GetNodeTitle().ToString(), InDependenciesCrc.GetValue());
+			}
 
-				if (bDebuggingEnabled)
-				{
-					// Leading spaces to align log content with warnings below - helps readability a lot.
-					UE_LOG(LogPCG, Log, TEXT("         [%s] %s\t\tCACHE HIT %u"), *InComponent->GetOwner()->GetName(), *InNode->GetNodeTitle().ToString(), InDependenciesCrc.GetValue());
-				}
+			OutOutput = *Value;
 
-				return true;
-			}
+			return true;
 		}
-
-		if (bDebuggingEnabled)
+		else
 		{
-			UE_LOG(LogPCG, Warning, TEXT("[%s] %s\t\tCACHE MISS %u"), *InComponent->GetOwner()->GetName() , *InNode->GetNodeTitle().ToString(), InDependenciesCrc.GetValue());
-		}
-
-		return false;
-	}
+			if (bDebuggingEnabled)
+			{
+				UE_LOG(LogPCG, Warning, TEXT("[%s] %s\t\tCACHE MISS %u"), *InComponent->GetOwner()->GetName(), *InNode->GetNodeTitle().ToString(), InDependenciesCrc.GetValue());
+			}
 
-	if (bDebuggingEnabled)
-	{
-		UE_LOG(LogPCG, Warning, TEXT("[%s] %s\t\tCACHE MISS NOELEMENT"), *InComponent->GetOwner()->GetName() , *InNode->GetNodeTitle().ToString());
+			return false;
+		}
 	}
-
-	return false;
 }
 
 void FPCGGraphCache::StoreInCache(const IPCGElement* InElement, const FPCGCrc& InDependenciesCrc, const FPCGDataCollection& InInput, const UPCGSettings* InSettings, const UPCGComponent* InComponent, const FPCGDataCollection& InOutput)
@@ -95,14 +100,12 @@ void FPCGGraphCache::StoreInCache(const IPCGElement* InElement, const FPCGCrc& I
 		TRACE_CPUPROFILER_EVENT_SCOPE(FPCGGraphCache::StoreInCache);
 		FWriteScopeLock ScopedWriteLock(CacheLock);
 
-		FPCGGraphCacheEntries* Entries = CacheData.Find(InElement);
-		if (!Entries)
-		{
-			Entries = &(CacheData.Add(InElement));
-		}
-
+		FPCGCacheEntryKey CacheKey(InElement, InDependenciesCrc);
+		CacheData.Add(CacheKey, InOutput);
+		
+		InOutput.AddToRootSet(*RootSet);
 
-		Entries->Emplace(InDependenciesCrc, InOutput, *RootSet);
+		AddDataToAccountedMemory(InOutput);
 	}
 }
 
@@ -111,16 +114,48 @@ void FPCGGraphCache::ClearCache()
 	FWriteScopeLock ScopedWriteLock(CacheLock);
 
 	// Unroot all previously rooted data
-	for (TPair<const IPCGElement*, FPCGGraphCacheEntries>& CacheEntry : CacheData)
+	for (FPCGDataCollection CacheEntry : CacheData)
 	{
-		for (FPCGGraphCacheEntry& Entry : CacheEntry.Value)
-		{
-			Entry.Output.RemoveFromRootSet(*RootSet);
-		}
+		CacheEntry.RemoveFromRootSet(*RootSet);
 	}
 
+	MemoryRecords.Empty();
+	TotalMemoryUsed = 0;
+
 	// Remove all entries
-	CacheData.Reset();
+	CacheData.Empty(GPCGGraphCacheMaxElements);
+}
+
+void FPCGGraphCache::EnforceMemoryBudget()
+{
+	TRACE_CPUPROFILER_EVENT_SCOPE(FPCGGraphCache::FPCGGraphCache::EnforceMemoryBudget);
+
+	if (!CVarCacheMemoryBudgetEnabled.GetValueOnAnyThread())
+	{
+		return;
+	}
+
+	if (CacheData.Num() == GPCGGraphCacheMaxElements)
+	{
+		UE_LOG(LogPCG, Warning, TEXT("Graph cache full (%d entries). Consider increasing GPCGGraphCacheMaxElements in code."), GPCGGraphCacheMaxElements);
+	}
+
+	const uint64 MemoryBudget = static_cast<uint64>(CVarCacheMemoryBudgetMB.GetValueOnAnyThread()) * 1024 * 1024;
+	if (TotalMemoryUsed < MemoryBudget)
+	{
+		return;
+	}
+
+	{
+		FWriteScopeLock ScopeWriteLock(CacheLock);
+
+		while (TotalMemoryUsed > MemoryBudget && CacheData.Num() > 0)
+		{
+			FPCGDataCollection RemovedData = CacheData.RemoveLeastRecent();
+			RemovedData.RemoveFromRootSet(*RootSet);
+			RemoveFromMemoryTotal(RemovedData);
+		}
+	}
 }
 
 #if WITH_EDITOR
@@ -131,23 +166,33 @@ void FPCGGraphCache::CleanFromCache(const IPCGElement* InElement, const UPCGSett
 		return;
 	}
 
-	if (IsDebuggingEnabled() && InSettings)
+	if (IsDebuggingEnabled())
 	{
-		UE_LOG(LogPCG, Warning, TEXT("CACHE: PURGED [%s]"), *InSettings->GetDefaultNodeName().ToString());
+		UE_LOG(LogPCG, Warning, TEXT("[] \t\tCACHE: PURGED [%s]"), InSettings ? *InSettings->GetDefaultNodeName().ToString() : TEXT("AnonymousElement"));
 	}
 
-	FWriteScopeLock ScopeWriteLock(CacheLock);
-	FPCGGraphCacheEntries* Entries = CacheData.Find(InElement);
-	if (Entries)
 	{
-		for (FPCGGraphCacheEntry& Entry : *Entries)
+		FWriteScopeLock ScopeWriteLock(CacheLock);
+
+		TArray<FPCGCacheEntryKey> Keys;
+		CacheData.GetKeys(Keys);
+
+		for (const FPCGCacheEntryKey& Key : Keys)
 		{
-			Entry.Output.RemoveFromRootSet(*RootSet);
+			if (Key.GetElement() != InElement)
+			{
+				continue;
+			}
+
+			if (const FPCGDataCollection* Data = CacheData.Find(Key))
+			{
+				RemoveFromMemoryTotal(*Data);
+				Data->RemoveFromRootSet(*RootSet);
+				
+				CacheData.Remove(Key);
+			}
 		}
 	}
-
-	// Finally, remove all entries matching that element
-	CacheData.Remove(InElement);
 }
 
 uint32 FPCGGraphCache::GetGraphCacheEntryCount(IPCGElement* InElement) const
@@ -155,15 +200,93 @@ uint32 FPCGGraphCache::GetGraphCacheEntryCount(IPCGElement* InElement) const
 	TRACE_CPUPROFILER_EVENT_SCOPE(FPCGGraphCache::GetFromCache);
 	FReadScopeLock ScopedReadLock(CacheLock);
 
-	if (const FPCGGraphCacheEntries* Entries = CacheData.Find(InElement))
+	uint32 Count = 0;
+
+	TArray<FPCGCacheEntryKey> Keys;
+	CacheData.GetKeys(Keys);
+
+	for (const FPCGCacheEntryKey& Key : Keys)
 	{
-		return Entries->Num();
+		if (Key.GetElement() == InElement)
+		{
+			Count++;
+		}
 	}
 
-	return 0;
+	return Count;
 }
 #endif // WITH_EDITOR
 
+void FPCGGraphCache::AddDataToAccountedMemory(const FPCGDataCollection& InCollection)
+{
+	for (const FPCGTaggedData& Data : InCollection.TaggedData)
+	{
+		if (Data.Data)
+		{
+			Data.Data->VisitDataNetwork([this](const UPCGData* Data)
+			{
+				if (Data)
+				{
+					FResourceSizeEx ResSize = FResourceSizeEx(EResourceSizeMode::Exclusive);
+					// Calculate data size. Function is non-const but is const-like, especially when
+					// resource mode is Exclusive. The other mode calls a function to find all outer'd
+					// objects which is non-const.
+					const_cast<UPCGData*>(Data)->GetResourceSizeEx(ResSize);
+					const SIZE_T DataSize = ResSize.GetDedicatedSystemMemoryBytes();
+
+					// Find or add record
+					FCachedMemoryRecord& NewRecord = MemoryRecords.FindOrAdd(Data->UID);
+
+					// Cache the size
+					NewRecord.MemoryPerInstance = DataSize;
+
+					if (NewRecord.InstanceCount == 0)
+					{
+						// Account for memory if first instance
+						TotalMemoryUsed += DataSize;
+					}
+
+					// Count instances
+					NewRecord.InstanceCount++;
+				}
+			});
+		}
+	}
+}
+
+void FPCGGraphCache::RemoveFromMemoryTotal(const FPCGDataCollection& InCollection)
+{
+	for (const FPCGTaggedData& Data : InCollection.TaggedData)
+	{
+		if (Data.Data)
+		{
+			Data.Data->VisitDataNetwork([this](const UPCGData* Data)
+			{
+				FCachedMemoryRecord* Record = Data ? MemoryRecords.Find(Data->UID) : nullptr;
+				if (ensure(Record))
+				{
+					// Update instance count
+					if (ensure(Record->InstanceCount > 0))
+					{
+						--Record->InstanceCount;
+					}
+
+					if (Record->InstanceCount == 0)
+					{
+						// Last instance removed, update accordingly
+						if (ensure(TotalMemoryUsed > Record->MemoryPerInstance))
+						{
+							TotalMemoryUsed -= Record->MemoryPerInstance;
+						}
+
+						MemoryRecords.Remove(Data->UID);
+					}
+				}
+			});
+		}
+	}
+}
+
 bool FPCGGraphCache::IsDebuggingEnabled() const
 {
 	return CVarCacheDebugging.GetValueOnAnyThread();
diff --git a/Engine/Plugins/Experimental/PCG/Source/PCG/Private/Graph/PCGGraphCache.h b/Engine/Plugins/Experimental/PCG/Source/PCG/Private/Graph/PCGGraphCache.h
--- a/Engine/Plugins/Experimental/PCG/Source/PCG/Private/Graph/PCGGraphCache.h
+++ b/Engine/Plugins/Experimental/PCG/Source/PCG/Private/Graph/PCGGraphCache.h
@@ -5,26 +5,58 @@
 #include "PCGCrc.h"
 #include "PCGData.h"
 
+#include "Containers/Array.h"
+#include "Containers/RingBuffer.h"
+#include "Containers/LruCache.h"
+
 class IPCGElement;
 class UPCGComponent;
 class UPCGNode;
 class UPCGSettings;
-class AActor;
 
-struct FPCGGraphCacheEntry
+/** A record kept each data object instance to count how many instances are tracked. */
+struct FCachedMemoryRecord
 {
-	FPCGGraphCacheEntry() = default;
-	FPCGGraphCacheEntry(const FPCGCrc& InCrc, const FPCGDataCollection& InOutput, FPCGRootSet& OutRootSet);
+	FCachedMemoryRecord() = default;
 
-	FPCGDataCollection Output;
+	/** Number of instances that are contributing to used memory. */
+	uint32 InstanceCount = 0;
 
-	/** A Crc value that encapsulates all information that can affect the result of an element. */
-	FPCGCrc DependenciesCrc;
+	/** Memory usage in bytes of one instance. Cached here to save recomputing later. */
+	SIZE_T MemoryPerInstance = 0;
 };
 
-// TODO: investigate if we need a more evolved data structure here
-// since we could want to have a lock per entries structure
-typedef TArray<FPCGGraphCacheEntry> FPCGGraphCacheEntries;
+struct FPCGCacheEntryKey
+{
+	FPCGCacheEntryKey(const IPCGElement* InElement, const FPCGCrc& InCrc)
+		: Element(InElement)
+		, Crc(InCrc)
+	{
+		check(Element);
+		ensure(Crc.IsValid());
+
+		Hash = GetTypeHash(Element);
+		Hash = HashCombine(Hash, Crc.GetValue());
+	}
+
+	bool operator==(const FPCGCacheEntryKey& Other) const
+	{
+		return Element == Other.Element && Crc == Other.Crc;
+	}
+
+	friend uint32 GetTypeHash(const FPCGCacheEntryKey& Key)
+	{
+		return Key.Hash;
+	}
+
+	const IPCGElement* GetElement() const { return Element; }
+
+private:
+	const IPCGElement* Element = nullptr;
+	FPCGCrc Crc;
+
+	uint32 Hash = 0;
+};
 
 /** Core idea is to store cache entries per node, but that will be less efficient
 * In cases where we have some subgraph reuse. Under that premise, we can then
@@ -44,6 +76,9 @@ struct FPCGGraphCache
 	/** Removes all entries from the cache, unroots data, etc. */
 	void ClearCache();
 
+	/** While memory usage is more than budget, remove cache entries for elements, LRU policy. Returns true if something removed. */
+	void EnforceMemoryBudget();
+
 	/** True if debugging features enabled. Exposes a CVar so it can be called from editor module. */
 	bool IsDebuggingEnabled() const;
 
@@ -56,13 +91,25 @@ struct FPCGGraphCache
 #endif
 
 private:
+	/** For each data element in collection, keep records of memory usage and instance count. */
+	void AddDataToAccountedMemory(const FPCGDataCollection& InCollection);
+
+	/** For each data element in collection, update records to reflect removal and update total memory figure if all instances removed. */
+	void RemoveFromMemoryTotal(const FPCGDataCollection& InCollection);
+
 	// Note: we are not going to serialize this as-is, since the pointers will change
 	// We will have to serialize on a node id basis most likely
-	TMap<const IPCGElement*, FPCGGraphCacheEntries> CacheData;
+	TLruCache<FPCGCacheEntryKey, FPCGDataCollection> CacheData;
 	TWeakObjectPtr<UObject> Owner = nullptr;
 
 	/** To prevent garbage collection on data in the cache, we'll need to root some data */
 	FPCGRootSet* RootSet = nullptr;
 
+	/** Map from data UIDs to records. Provides ref counting and caches memory size. */
+	TMap<uint64, FCachedMemoryRecord> MemoryRecords;
+
+	/** Total memory usage by all data objects in cache. */
+	uint64 TotalMemoryUsed = 0;
+
 	mutable FRWLock CacheLock;
 };
diff --git a/Engine/Plugins/Experimental/PCG/Source/PCG/Private/Graph/PCGGraphExecutor.cpp b/Engine/Plugins/Experimental/PCG/Source/PCG/Private/Graph/PCGGraphExecutor.cpp
--- a/Engine/Plugins/Experimental/PCG/Source/PCG/Private/Graph/PCGGraphExecutor.cpp
+++ b/Engine/Plugins/Experimental/PCG/Source/PCG/Private/Graph/PCGGraphExecutor.cpp
@@ -879,6 +879,9 @@ void FPCGGraphExecutor::Execute()
 			}
 		}
 
+		// Purge things from cache if memory usage is too high
+		GraphCache.EnforceMemoryBudget();
+
 #if WITH_EDITOR
 		// Save & release resources when running in-editor
 		SaveDirtyActors();
diff --git a/Engine/Plugins/Experimental/PCG/Source/PCG/Private/PCGData.cpp b/Engine/Plugins/Experimental/PCG/Source/PCG/Private/PCGData.cpp
--- a/Engine/Plugins/Experimental/PCG/Source/PCG/Private/PCGData.cpp
+++ b/Engine/Plugins/Experimental/PCG/Source/PCG/Private/PCGData.cpp
@@ -104,6 +104,11 @@ void UPCGData::AddToCrc(FArchiveCrc32& Ar) const
 	Ar << UIDValue;
 }
 
+void UPCGData::VisitDataNetwork(TFunctionRef<void(const UPCGData*)> Action) const
+{
+	Action(this);
+}
+
 bool FPCGTaggedData::operator==(const FPCGTaggedData& Other) const
 {
 	return Data == Other.Data &&
@@ -220,6 +225,22 @@ UPCGParamData* FPCGDataCollection::GetFirstParamsOnParamsPin() const
 	return (ParamsOnDefaultPin.IsEmpty() ? nullptr : Cast<UPCGParamData>(ParamsOnDefaultPin[0].Data));
 }
 
+void FPCGDataCollection::GetResourceSizeEx(FResourceSizeEx& CumulativeResourceSize) const
+{
+	for (const FPCGTaggedData& Data : TaggedData)
+	{
+		if (Data.Data)
+		{
+			Data.Data->VisitDataNetwork([&CumulativeResourceSize](const UPCGData* Data) {
+				// Cast away const-ness. The extended mode of GetResourceSizeEx accounts memory for all objects outer'd
+				// to this object, and that calls GetObjectsWithOuter which is non-const. We don't use this extended mode
+				// and we need to be able to operate on const objects.
+				const_cast<UPCGData*>(Data)->GetResourceSizeEx(CumulativeResourceSize);
+			});
+		}
+	}
+}
+
 const UPCGSettings* FPCGDataCollection::GetSettings(const UPCGSettings* InDefaultSettings) const
 {
 	if (!InDefaultSettings)
diff --git a/Engine/Plugins/Experimental/PCG/Source/PCG/Public/Data/PCGDifferenceData.h b/Engine/Plugins/Experimental/PCG/Source/PCG/Public/Data/PCGDifferenceData.h
--- a/Engine/Plugins/Experimental/PCG/Source/PCG/Public/Data/PCGDifferenceData.h
+++ b/Engine/Plugins/Experimental/PCG/Source/PCG/Public/Data/PCGDifferenceData.h
@@ -49,6 +49,7 @@ public:
 
 	// ~Begin UPCGData interface
 	virtual EPCGDataType GetDataType() const override { return EPCGDataType::Spatial; }
+	virtual void VisitDataNetwork(TFunctionRef<void(const UPCGData*)> Action) const override;
 	// ~End UPCGData interface
 
 	//~Begin UPCGSpatialData interface
diff --git a/Engine/Plugins/Experimental/PCG/Source/PCG/Public/Data/PCGIntersectionData.h b/Engine/Plugins/Experimental/PCG/Source/PCG/Public/Data/PCGIntersectionData.h
--- a/Engine/Plugins/Experimental/PCG/Source/PCG/Public/Data/PCGIntersectionData.h
+++ b/Engine/Plugins/Experimental/PCG/Source/PCG/Public/Data/PCGIntersectionData.h
@@ -26,6 +26,7 @@ public:
 
 	// ~Begin UPCGData interface
 	virtual EPCGDataType GetDataType() const override { return EPCGDataType::Spatial; }
+	virtual void VisitDataNetwork(TFunctionRef<void(const UPCGData*)> Action) const override;
 	// ~End UPCGData interface
 
 	//~Begin UPCGSpatialData interface
@@ -36,6 +37,7 @@ public:
 	virtual bool HasNonTrivialTransform() const override;
 	virtual const UPCGSpatialData* FindShapeFromNetwork(const int InDimension) const override;
 	virtual const UPCGSpatialData* FindFirstConcreteShapeFromNetwork() const override;
+
 protected:
 	virtual UPCGSpatialData* CopyInternal() const override;
 	//~End UPCGSpatialData interface
diff --git a/Engine/Plugins/Experimental/PCG/Source/PCG/Public/Data/PCGPointData.h b/Engine/Plugins/Experimental/PCG/Source/PCG/Public/Data/PCGPointData.h
--- a/Engine/Plugins/Experimental/PCG/Source/PCG/Public/Data/PCGPointData.h
+++ b/Engine/Plugins/Experimental/PCG/Source/PCG/Public/Data/PCGPointData.h
@@ -67,6 +67,10 @@ class PCG_API UPCGPointData : public UPCGSpatialData
 public:
 	typedef TOctree2<FPCGPointRef, FPCGPointRefSemantics> PointOctree;
 
+	//~Begin UObject Interface
+	virtual void GetResourceSizeEx(FResourceSizeEx& CumulativeResourceSize) override;
+	//~End UObject Interface
+
 	// ~Begin UPCGData interface
 	virtual EPCGDataType GetDataType() const override { return EPCGDataType::Point; }
 	virtual void AddToCrc(FArchiveCrc32& Ar) const override;
diff --git a/Engine/Plugins/Experimental/PCG/Source/PCG/Public/Data/PCGSpatialData.h b/Engine/Plugins/Experimental/PCG/Source/PCG/Public/Data/PCGSpatialData.h
--- a/Engine/Plugins/Experimental/PCG/Source/PCG/Public/Data/PCGSpatialData.h
+++ b/Engine/Plugins/Experimental/PCG/Source/PCG/Public/Data/PCGSpatialData.h
@@ -34,6 +34,10 @@ class PCG_API UPCGSpatialData : public UPCGData
 public:
 	UPCGSpatialData(const FObjectInitializer& ObjectInitializer);
 
+	//~Begin UObject Interface
+	virtual void GetResourceSizeEx(FResourceSizeEx& CumulativeResourceSize) override;
+	//~End UObject Interface
+
 	// ~Begin UPCGData interface
 	virtual EPCGDataType GetDataType() const override { return EPCGDataType::Spatial; }
 	// ~End UPCGData interface
@@ -150,6 +154,10 @@ class PCG_API UPCGSpatialDataWithPointCache : public UPCGSpatialData
 	GENERATED_BODY()
 
 public:
+	//~Begin UObject Interface
+	virtual void GetResourceSizeEx(FResourceSizeEx& CumulativeResourceSize) override;
+	//~End UObject Interface
+
 	// ~UPCGSpatialData implementation
 	virtual const UPCGPointData* ToPointData(FPCGContext* Context, const FBox& InBounds = FBox(EForceInit::ForceInit)) const override;
 	// ~End UPCGSpatialData implementation
diff --git a/Engine/Plugins/Experimental/PCG/Source/PCG/Public/Data/PCGUnionData.h b/Engine/Plugins/Experimental/PCG/Source/PCG/Public/Data/PCGUnionData.h
--- a/Engine/Plugins/Experimental/PCG/Source/PCG/Public/Data/PCGUnionData.h
+++ b/Engine/Plugins/Experimental/PCG/Source/PCG/Public/Data/PCGUnionData.h
@@ -38,6 +38,7 @@ public:
 
 	// ~Begin UPCGData interface
 	virtual EPCGDataType GetDataType() const override { return EPCGDataType::Spatial; }
+	virtual void VisitDataNetwork(TFunctionRef<void(const UPCGData*)> Action) const override;
 	// ~End UPCGData interface
 
 	//~Begin UPCGSpatialData interface
diff --git a/Engine/Plugins/Experimental/PCG/Source/PCG/Public/PCGData.h b/Engine/Plugins/Experimental/PCG/Source/PCG/Public/PCGData.h
--- a/Engine/Plugins/Experimental/PCG/Source/PCG/Public/PCGData.h
+++ b/Engine/Plugins/Experimental/PCG/Source/PCG/Public/PCGData.h
@@ -35,6 +35,9 @@ public:
 	/** Contributes to Crc. Fallback implementation that writes object instance UID. */
 	virtual void AddToCrc(FArchiveCrc32& Ar) const;
 
+	/** Executes a lambda over all connected data objects. */
+	virtual void VisitDataNetwork(TFunctionRef<void(const UPCGData*)> Action) const;
+
 	/** Unique ID for this object instance. */
 	UPROPERTY(Transient)
 	uint64 UID = 0;
@@ -112,6 +115,9 @@ struct PCG_API FPCGDataCollection
 	const UPCGSettingsInterface* GetSettingsInterface() const;
 	const UPCGSettingsInterface* GetSettingsInterface(const UPCGSettingsInterface* InDefaultSettingsInterface) const;
 
+	/** Memory size calculation. Forward call to the data objects in the collection. */
+	void GetResourceSizeEx(FResourceSizeEx& CumulativeResourceSize) const;
+
 	template<typename SettingsType>
 	const SettingsType* GetSettings() const;
 
