[PCG] Apply new PCG editor settings.

Also rearranged initialisation order to match export/member order to make applying defaults easier in future.

diff --git a/Engine/Plugins/Experimental/PCG/Source/PCGEditor/Private/PCGEditorSettings.cpp b/Engine/Plugins/Experimental/PCG/Source/PCGEditor/Private/PCGEditorSettings.cpp
--- a/Engine/Plugins/Experimental/PCG/Source/PCGEditor/Private/PCGEditorSettings.cpp
+++ b/Engine/Plugins/Experimental/PCG/Source/PCGEditor/Private/PCGEditorSettings.cpp
@@ -12,15 +12,15 @@ UPCGEditorSettings::UPCGEditorSettings(const FObjectInitializer& ObjectInitializ
 	DefaultNodeColor = FLinearColor(0.4f, 0.62f, 1.0f);
 	InstancedNodeBodyTintColor = FLinearColor(0.5f, 0.5f, 0.5f);
 	InputOutputNodeColor = FLinearColor(1.0f, 0.0f, 0.0f);
-	SetOperationNodeColor = FLinearColor(0.8f, 0.2f, 0.8f);
-	DensityOperationNodeColor = FLinearColor(0.6f, 1.0f, 0.6f);
-	BlueprintNodeColor = FLinearColor(0.0f, 0.6f, 1.0f);
-	MetadataNodeColor = FLinearColor(0.4f, 0.4f, 0.8f);
-	FilterNodeColor = FLinearColor(0.4f, 0.8f, 0.4f);
-	SamplerNodeColor = FLinearColor(0.8f, 1.0f, 0.4f);
-	SpawnerNodeColor = FLinearColor(1.0f, 0.6f, 0.4f);
-	SubgraphNodeColor = FLinearColor(1.0f, 0.1f, 0.1f);
-	ParamDataNodeColor = FLinearColor(1.0f, 0.47f, 0.0f);
+	SetOperationNodeColor = FLinearColor(1.0f, 0.2f, 1.0f);
+	DensityOperationNodeColor = FLinearColor(0.2f, 0.59f, 0.99f);
+	BlueprintNodeColor = FLinearColor(0.02f, 0.18f, 1.0f);
+	MetadataNodeColor = FLinearColor(1.0f, 0.99f, 0.99f);
+	FilterNodeColor = FLinearColor(0.24f, 0.09f, 0.85f);
+	SamplerNodeColor = FLinearColor(0.0f, 0.0f, 0.0f);
+	SpawnerNodeColor = FLinearColor(0.0f, 1.0f, 0.69f);
+	SubgraphNodeColor = FLinearColor(1.0f, 0.05f, 0.05f);
+	ParamDataNodeColor = FLinearColor(1.0f, 0.38f, 0.02f);
 	DebugNodeColor = FLinearColor(1.0f, 0.0f, 1.0f);
 
 	DefaultPinColor = FLinearColor(0.16f, 0.16f, 0.16f);
@@ -29,10 +29,10 @@ UPCGEditorSettings::UPCGEditorSettings(const FObjectInitializer& ObjectInitializ
 	ConcreteDataPinColor = FLinearColor(0.45f, 0.38f, 0.96f);
 	PointDataPinColor = FLinearColor(0.05f, 0.25f, 1.0f);
 	PolyLineDataPinColor = FLinearColor(0.05f, 0.75f, 0.82f);
+	SurfaceDataPinColor = FLinearColor(0.06f, 0.55f, 0.21f);
 	LandscapeDataPinColor = FLinearColor(0.66f, 0.66f, 0.07f);
 	TextureDataPinColor = FLinearColor(0.79f, 0.08f, 0.01f);
 	RenderTargetDataPinColor = FLinearColor(0.8f, 0.18f, 0.12f);
-	SurfaceDataPinColor = FLinearColor(0.06f, 0.55f, 0.21f);
 	VolumeDataPinColor = FLinearColor(0.79f, 0.06f, 0.5f);
 	PrimitiveDataPinColor = FLinearColor(0.22f, 0.05f, 1.0f);
 
