Added reuse of spawned actor resources
Fixed issue the instanced template actor not being outered properly which made for lost data on reload

diff --git a/Engine/Plugins/Experimental/PCG/Source/PCG/Private/Elements/PCGSpawnActor.cpp b/Engine/Plugins/Experimental/PCG/Source/PCG/Private/Elements/PCGSpawnActor.cpp
--- a/Engine/Plugins/Experimental/PCG/Source/PCG/Private/Elements/PCGSpawnActor.cpp
+++ b/Engine/Plugins/Experimental/PCG/Source/PCG/Private/Elements/PCGSpawnActor.cpp
@@ -24,6 +24,11 @@
 
 #include UE_INLINE_GENERATED_CPP_BY_NAME(PCGSpawnActor)
 
+static TAutoConsoleVariable<bool> CVarAllowActorReuse(
+	TEXT("pcg.Actor.AllowReuse"),
+	true,
+	TEXT("Controls whether PCG spawned actors can be reused and skipped when re-executing"));
+
 namespace PCGSpawnActorHelpers
 {
 	struct FActorSingleOverride
@@ -281,7 +286,7 @@ void UPCGSpawnActorSettings::RefreshTemplateActor()
 {
 	if (TemplateActorClass)
 	{
-		AActor* NewTemplateActor = NewObject<AActor>(GetTransientPackage(), TemplateActorClass, NAME_None, RF_ArchetypeObject | RF_Transactional | RF_Public);
+		AActor* NewTemplateActor = NewObject<AActor>(this, TemplateActorClass, NAME_None, RF_ArchetypeObject | RF_Transactional | RF_Public);
 
 		if (TemplateActor)
 		{
@@ -319,6 +324,47 @@ bool FPCGSpawnActorElement::ExecuteInternal(FPCGContext* Context) const
 		return true;
 	}
 
+	// Check if we can reuse existing resources - note that this is done on a per-settings basis when collapsed,
+	// Otherwise we'll check against merged crc 
+	bool bFullySkippedDueToReuse = false;
+
+	if (CVarAllowActorReuse.GetValueOnAnyThread())
+	{
+		// Compute CRC if it has not been computed (it likely isn't, but this is to futureproof this)
+		if (!Context->DependenciesCrc.IsValid())
+		{
+			GetDependenciesCrc(Context->InputData, Settings, Context->SourceComponent.Get(), Context->DependenciesCrc);
+		}
+
+		if (Context->DependenciesCrc.IsValid())
+		{
+			if (Settings->Option == EPCGSpawnActorOption::CollapseActors)
+			{
+				TArray<UPCGManagedISMComponent*> MISMCs;
+				Context->SourceComponent->ForEachManagedResource([&MISMCs, &Context](UPCGManagedResource* InResource)
+				{
+					if (UPCGManagedISMComponent* Resource = Cast<UPCGManagedISMComponent>(InResource))
+					{
+						if (Resource->GetCrc().IsValid() && Resource->GetCrc() == Context->DependenciesCrc)
+						{
+							MISMCs.Add(Resource);
+						}
+					}
+				});
+
+				for (UPCGManagedISMComponent* MISMC : MISMCs)
+				{
+					MISMC->MarkAsReused();
+				}
+
+				if (!MISMCs.IsEmpty())
+				{
+					bFullySkippedDueToReuse = true;
+				}
+			}
+		}
+	}
+
 	TArray<UFunction*> PostSpawnFunctions;
 	for (FName PostSpawnFunctionName : Settings->PostSpawnFunctionNames)
 	{
@@ -389,169 +435,224 @@ bool FPCGSpawnActorElement::ExecuteInternal(FPCGContext* Context) const
 		}
 
 		// Spawn actors/populate ISM
+		if (!bFullySkippedDueToReuse && Settings->Option == EPCGSpawnActorOption::CollapseActors)
 		{
-			if (Settings->Option == EPCGSpawnActorOption::CollapseActors)
-			{
-				TArray<UActorComponent*> Components;
-				UPCGActorHelpers::GetActorClassDefaultComponents(Settings->TemplateActorClass, Components, UStaticMeshComponent::StaticClass());
-				
-				TMap<FISMComponentDescriptor, TArray<FTransform>> MeshDescriptorTransforms;
+			TRACE_CPUPROFILER_EVENT_SCOPE(FPCGSpawnActorElement::ExecuteInternal::CollapseActors);
+			TArray<UActorComponent*> Components;
+			UPCGActorHelpers::GetActorClassDefaultComponents(Settings->TemplateActorClass, Components, UStaticMeshComponent::StaticClass());
+			
+			TMap<FISMComponentDescriptor, TArray<FTransform>> MeshDescriptorTransforms;
 
-				for (UActorComponent* Component : Components)
+			for (UActorComponent* Component : Components)
+			{
+				if (UStaticMeshComponent* StaticMeshComponent = Cast<UStaticMeshComponent>(Component))
 				{
-					if (UStaticMeshComponent* StaticMeshComponent = Cast<UStaticMeshComponent>(Component))
-					{
-						FISMComponentDescriptor ISMDescriptor;
-						ISMDescriptor.InitFrom(StaticMeshComponent);
-						ISMDescriptor.ComputeHash();
+					FISMComponentDescriptor ISMDescriptor;
+					ISMDescriptor.InitFrom(StaticMeshComponent);
+					ISMDescriptor.ComputeHash();
 
-						TArray<FTransform>& Transforms = MeshDescriptorTransforms.FindOrAdd(ISMDescriptor);
+					TArray<FTransform>& Transforms = MeshDescriptorTransforms.FindOrAdd(ISMDescriptor);
 
-						if (UInstancedStaticMeshComponent* InstancedStaticMeshComponent = Cast<UInstancedStaticMeshComponent>(StaticMeshComponent))
+					if (UInstancedStaticMeshComponent* InstancedStaticMeshComponent = Cast<UInstancedStaticMeshComponent>(StaticMeshComponent))
+					{
+						const int32 NumInstances = InstancedStaticMeshComponent->GetInstanceCount();
+						Transforms.Reserve(Transforms.Num() + NumInstances);
+						
+						for (int32 InstanceIndex = 0; InstanceIndex < NumInstances; InstanceIndex++)
 						{
-							const int32 NumInstances = InstancedStaticMeshComponent->GetInstanceCount();
-							Transforms.Reserve(Transforms.Num() + NumInstances);
-							
-							for (int32 InstanceIndex = 0; InstanceIndex < NumInstances; InstanceIndex++)
+							FTransform InstanceTransform;
+							if (InstancedStaticMeshComponent->GetInstanceTransform(InstanceIndex, InstanceTransform))
 							{
-								FTransform InstanceTransform;
-								if (InstancedStaticMeshComponent->GetInstanceTransform(InstanceIndex, InstanceTransform))
-								{
-									Transforms.Add(InstanceTransform);
-								}
+								Transforms.Add(InstanceTransform);
 							}
 						}
-						else
-						{
-							Transforms.Add(StaticMeshComponent->GetRelativeTransform());
-						}
+					}
+					else
+					{
+						Transforms.Add(StaticMeshComponent->GetRelativeTransform());
 					}
 				}
+			}
 				
-				for (const TPair<FISMComponentDescriptor, TArray<FTransform>>& ISMCDescriptorTransforms : MeshDescriptorTransforms)
+			for (const TPair<FISMComponentDescriptor, TArray<FTransform>>& ISMCDescriptorTransforms : MeshDescriptorTransforms)
+			{
+				const FISMComponentDescriptor& ISMCDescriptor = ISMCDescriptorTransforms.Key;
+
+				UPCGManagedISMComponent* MISMC = UPCGActorHelpers::GetOrCreateManagedISMC(TargetActor, Context->SourceComponent.Get(), ISMCDescriptor);
+				if (!MISMC)
 				{
-					const FISMComponentDescriptor& ISMCDescriptor = ISMCDescriptorTransforms.Key;
-					
-					UInstancedStaticMeshComponent* ISMC = UPCGActorHelpers::GetOrCreateISMC(TargetActor, Context->SourceComponent.Get(), ISMCDescriptor);
-					if (!ISMC)
+					continue;
+				}
+
+				MISMC->SetCrc(Context->DependenciesCrc);
+
+				UInstancedStaticMeshComponent* ISMC = MISMC->GetComponent();
+				check(ISMC);
+				
+				const TArray<FTransform>& ISMCTransforms = ISMCDescriptorTransforms.Value;
+				
+				TArray<FTransform> Transforms;
+				Transforms.Reserve(Points.Num() * ISMCTransforms.Num());
+				for (int32 PointIndex = 0; PointIndex < Points.Num(); PointIndex++)
+				{
+					const FPCGPoint& Point = Points[PointIndex];
+					for (int32 TransformIndex = 0; TransformIndex < ISMCTransforms.Num(); TransformIndex++)
 					{
-						continue;
+						const FTransform& Transform = ISMCTransforms[TransformIndex];
+						Transforms.Add(Transform * Point.Transform);
 					}
-										
-					const TArray<FTransform>& ISMCTransforms = ISMCDescriptorTransforms.Value;
-					
-					TArray<FTransform> Transforms;
-					Transforms.Reserve(Points.Num() * ISMCTransforms.Num());
-					for (int32 PointIndex = 0; PointIndex < Points.Num(); PointIndex++)
+				}
+				
+				ISMC->NumCustomDataFloats = 0;
+				ISMC->AddInstances(Transforms, false, true);
+				ISMC->UpdateBounds();
+				
+				PCGE_LOG(Verbose, "Added %d instances of %s to %s in %s", Transforms.Num(), *ISMC->GetStaticMesh().GetName(), *ISMC->GetName(), *TargetActor->GetActorNameOrLabel());
+			}
+		}
+		else if (!bFullySkippedDueToReuse && Settings->Option != EPCGSpawnActorOption::CollapseActors && (bHasAuthority || !bSpawnedActorsRequireAuthority))
+		{
+			TRACE_CPUPROFILER_EVENT_SCOPE(FPCGSpawnActorElement::ExecuteInternal::SpawnActors);
+			AActor* TemplateActor = Settings->ActorOverrides.IsEmpty() ? Settings->TemplateActor.Get() : DuplicateObject(Settings->TemplateActor, GetTransientPackage());
+
+			PCGSpawnActorHelpers::FActorOverrides ActorOverrides(Settings->ActorOverrides, TemplateActor, PointData);
+
+			FActorSpawnParameters SpawnParams;
+			SpawnParams.Owner = TargetActor;
+			SpawnParams.Template = TemplateActor;
+			SpawnParams.SpawnCollisionHandlingOverride = ESpawnActorCollisionHandlingMethod::AlwaysSpawn;
+
+			if (PCGHelpers::IsRuntimeOrPIE())
+			{
+				SpawnParams.ObjectFlags |= RF_Transient;
+			}
+
+			const bool bForceCallGenerate = Settings->bGenerationTrigger == EPCGSpawnActorGenerationTrigger::ForceGenerate;
+#if WITH_EDITOR
+			const bool bOnLoadCallGenerate = Settings->bGenerationTrigger == EPCGSpawnActorGenerationTrigger::Default;
+#else
+			const bool bOnLoadCallGenerate = (Settings->bGenerationTrigger == EPCGSpawnActorGenerationTrigger::Default ||
+				Settings->bGenerationTrigger == EPCGSpawnActorGenerationTrigger::DoNotGenerateInEditor);
+#endif
+			UPCGSubsystem* Subsystem = (Context->SourceComponent.Get() ? Context->SourceComponent->GetSubsystem() : nullptr);
+
+			// Try to reuse actors if the are preexisting
+			TArray<UPCGManagedActors*> ReusedManagedActorsResources;
+			if (CVarAllowActorReuse.GetValueOnAnyThread())
+			{
+				if (Context->DependenciesCrc.IsValid())
+				{
+					Context->SourceComponent->ForEachManagedResource([&ReusedManagedActorsResources, &Context](UPCGManagedResource* InResource)
 					{
-						const FPCGPoint& Point = Points[PointIndex];
-						for (int32 TransformIndex = 0; TransformIndex < ISMCTransforms.Num(); TransformIndex++)
+						if (UPCGManagedActors* Resource = Cast<UPCGManagedActors>(InResource))
 						{
-							const FTransform& Transform = ISMCTransforms[TransformIndex];
-							Transforms.Add(Transform * Point.Transform);
+							if (Resource->GetCrc().IsValid() && Resource->GetCrc() == Context->DependenciesCrc)
+							{
+								ReusedManagedActorsResources.Add(Resource);
+							}
 						}
-					}
-					
-					ISMC->NumCustomDataFloats = 0;
-					ISMC->AddInstances(Transforms, false, true);
-					ISMC->UpdateBounds();
-					
-					PCGE_LOG(Verbose, "Added %d instances of %s to %s in %s", Transforms.Num(), *ISMC->GetStaticMesh().GetName(), *ISMC->GetName(), *TargetActor->GetActorNameOrLabel());
+					});
 				}
 			}
 
+			TArray<AActor*> ProcessedActors;
+			const bool bActorsHavePCGComponents = (Settings->GetSubgraphInterface() != nullptr);
+
+			if (!ReusedManagedActorsResources.IsEmpty())
 			{
-				TRACE_CPUPROFILER_EVENT_SCOPE(FPCGSpawnActorElement::ExecuteInternal::SpawnActors);
-				if (Settings->Option != EPCGSpawnActorOption::CollapseActors && (bHasAuthority || !bSpawnedActorsRequireAuthority))
+				// If the actors are fully independent, we might need to make sure to call Generate if the underlying graph has changed - e.g. if the actor is dirty
+				for (UPCGManagedActors* ManagedActors : ReusedManagedActorsResources)
 				{
-					AActor* TemplateActor = Settings->ActorOverrides.IsEmpty() ? Settings->TemplateActor.Get() : DuplicateObject(Settings->TemplateActor, GetTransientPackage());
-
-					PCGSpawnActorHelpers::FActorOverrides ActorOverrides(Settings->ActorOverrides, TemplateActor, PointData);
-
-					FActorSpawnParameters SpawnParams;
-					SpawnParams.Owner = TargetActor;
-					SpawnParams.Template = TemplateActor;
-					SpawnParams.SpawnCollisionHandlingOverride = ESpawnActorCollisionHandlingMethod::AlwaysSpawn;
+					check(ManagedActors);
+					ManagedActors->MarkAsReused();
 
-					if (PCGHelpers::IsRuntimeOrPIE())
+					// There's no setup to be done, just generation if we're in the no-merge case, so keep track of these actors only in this case
+					if (bActorsHavePCGComponents && Settings->Option == EPCGSpawnActorOption::NoMerging)
 					{
-						SpawnParams.ObjectFlags |= RF_Transient;
+						for (TSoftObjectPtr<AActor>& ManagedActorPtr : ManagedActors->GeneratedActors)
+						{
+							if (AActor* ManagedActor = ManagedActorPtr.Get())
+							{
+								ProcessedActors.Add(ManagedActor);
+							}
+						}
 					}
+				}
+			}
+			else
+			{
+				TArray<FName> NewActorTags = GetNewActorTags(Context, TargetActor, Settings->bInheritActorTags, Settings->TagsToAddOnActors);
 
-					TArray<FName> NewActorTags = GetNewActorTags(Context, TargetActor, Settings->bInheritActorTags, Settings->TagsToAddOnActors);
+				// Create managed resource for actor tracking
+				UPCGManagedActors* ManagedActors = NewObject<UPCGManagedActors>(Context->SourceComponent.Get());
+				ManagedActors->SetCrc(Context->DependenciesCrc);
 
-					// Create managed resource for actor tracking
-					UPCGManagedActors* ManagedActors = NewObject<UPCGManagedActors>(Context->SourceComponent.Get());
+				for (int32 i = 0; i < Points.Num(); ++i)
+				{
+					const FPCGPoint& Point = Points[i];
 
-					const bool bForceCallGenerate = Settings->bGenerationTrigger == EPCGSpawnActorGenerationTrigger::ForceGenerate;
-#if WITH_EDITOR
-					const bool bOnLoadCallGenerate = Settings->bGenerationTrigger == EPCGSpawnActorGenerationTrigger::Default;
-#else
-					const bool bOnLoadCallGenerate = (Settings->bGenerationTrigger == EPCGSpawnActorGenerationTrigger::Default ||
-						Settings->bGenerationTrigger == EPCGSpawnActorGenerationTrigger::DoNotGenerateInEditor);
-#endif
+					ActorOverrides.Apply(i);
+
+					AActor* GeneratedActor = TargetActor->GetWorld()->SpawnActor(Settings->TemplateActorClass, &Point.Transform, SpawnParams);
 
-					for (int32 i = 0; i < Points.Num(); ++i)
+					if (!GeneratedActor)
 					{
-						const FPCGPoint& Point = Points[i];
+						PCGE_LOG(Error, "Failed to spawn actor on point with index %d", i);
+						continue;
+					}
 
-						ActorOverrides.Apply(i);
+					// HACK: until UE-62747 is fixed, we have to force set the scale after spawning the actor
+					GeneratedActor->SetActorRelativeScale3D(Point.Transform.GetScale3D());
+					GeneratedActor->Tags.Append(NewActorTags);
+					GeneratedActor->AttachToActor(TargetActor, FAttachmentTransformRules::KeepWorldTransform);
 
-						AActor* GeneratedActor = TargetActor->GetWorld()->SpawnActor(Settings->TemplateActorClass, &Point.Transform, SpawnParams);
+					for (UFunction* PostSpawnFunction : PostSpawnFunctions)
+					{
+						GeneratedActor->ProcessEvent(PostSpawnFunction, nullptr);
+					}
 
-						if (!GeneratedActor)
-						{
-							PCGE_LOG(Error, "Failed to spawn actor on point with index %d", i);
-							continue;
-						}
+					ManagedActors->GeneratedActors.Add(GeneratedActor);
 
-						// HACK: until UE-62747 is fixed, we have to force set the scale after spawning the actor
-						GeneratedActor->SetActorRelativeScale3D(Point.Transform.GetScale3D());
-						GeneratedActor->Tags.Append(NewActorTags);
-						GeneratedActor->AttachToActor(TargetActor, FAttachmentTransformRules::KeepWorldTransform);
+					if (bActorsHavePCGComponents)
+					{
+						ProcessedActors.Add(GeneratedActor);
+					}
+				}
 
-						for (UFunction* PostSpawnFunction : PostSpawnFunctions)
-						{
-							GeneratedActor->ProcessEvent(PostSpawnFunction, nullptr);
-						}
+				Context->SourceComponent->AddToManagedResources(ManagedActors);
 
-						ManagedActors->GeneratedActors.Add(GeneratedActor);
+				PCGE_LOG(Verbose, "Generated %d actors", Points.Num());
+			}
 
-						// If the actor spawned has a PCG component, either generate it or mark it as generated if we pass through its inputs
-						TArray<UPCGComponent*> PCGComponents;
-						GeneratedActor->GetComponents(PCGComponents);
+			// Setup & Generate on PCG components if needed
+			for (AActor* Actor : ProcessedActors)
+			{
+				TInlineComponentArray<UPCGComponent*, 1> PCGComponents;
+				Actor->GetComponents(PCGComponents);
 
-						UPCGSubsystem* Subsystem = Context->SourceComponent.Get() ? Context->SourceComponent->GetSubsystem() : nullptr;
+				for (UPCGComponent* PCGComponent : PCGComponents)
+				{
+					if(Settings->Option == EPCGSpawnActorOption::NoMerging)
+					{
+						if (bForceDisableActorParsing)
+						{
+							PCGComponent->bParseActorComponents = false;
+						}
 
-						for (UPCGComponent* PCGComponent : PCGComponents)
+						if (bForceCallGenerate || (bOnLoadCallGenerate && PCGComponent->GenerationTrigger == EPCGComponentGenerationTrigger::GenerateOnLoad))
 						{
-							if (Settings->Option == EPCGSpawnActorOption::NoMerging)
+							if (Subsystem && PCGComponent->IsPartitioned())
 							{
-								if (bForceDisableActorParsing)
-								{
-									PCGComponent->bParseActorComponents = false;
-								}
-
-								if(bForceCallGenerate || (bOnLoadCallGenerate && PCGComponent->GenerationTrigger == EPCGComponentGenerationTrigger::GenerateOnLoad))
-								{
-									if (Subsystem && PCGComponent->IsPartitioned())
-									{
-										Subsystem->RegisterOrUpdatePCGComponent(PCGComponent);
-									}
-									
-									PCGComponent->Generate();
-								}
-							}
-							else
-							{
-								PCGComponent->bActivated = false;
+								Subsystem->RegisterOrUpdatePCGComponent(PCGComponent);
 							}
+
+							PCGComponent->Generate();
 						}
 					}
-
-					Context->SourceComponent->AddToManagedResources(ManagedActors);
-
-					PCGE_LOG(Verbose, "Generated %d actors", Points.Num());
+					else
+					{
+						PCGComponent->bActivated = false;
+					}
 				}
 			}
 		}
@@ -560,9 +661,7 @@ bool FPCGSpawnActorElement::ExecuteInternal(FPCGContext* Context) const
 		// - if it's not merged, will be the input points directly
 		// - if it's merged but there is no subgraph, will be the input points directly
 		// - if it's merged and there is a subgraph, we'd need to pass the data for it to be given to the subgraph
-		{
-			Outputs.Add(Input);
-		}
+		Outputs.Add(Input);
 	}
 
 	return true;
diff --git a/Engine/Plugins/Experimental/PCG/Source/PCG/Private/Helpers/PCGActorHelpers.cpp b/Engine/Plugins/Experimental/PCG/Source/PCG/Private/Helpers/PCGActorHelpers.cpp
--- a/Engine/Plugins/Experimental/PCG/Source/PCG/Private/Helpers/PCGActorHelpers.cpp
+++ b/Engine/Plugins/Experimental/PCG/Source/PCG/Private/Helpers/PCGActorHelpers.cpp
@@ -20,7 +20,7 @@
 #if WITH_EDITOR
 #endif
 
-UInstancedStaticMeshComponent* UPCGActorHelpers::GetOrCreateISMC(AActor* InTargetActor, UPCGComponent* InSourceComponent, const FISMComponentDescriptor& InISMCDescriptor)
+UPCGManagedISMComponent* UPCGActorHelpers::GetOrCreateManagedISMC(AActor* InTargetActor, UPCGComponent* InSourceComponent, const FISMComponentDescriptor& InISMCDescriptor)
 {
 	check(InTargetActor);
 	check(InSourceComponent);
@@ -52,7 +52,7 @@ UInstancedStaticMeshComponent* UPCGActorHelpers::GetOrCreateISMC(AActor* InTarge
 			if (ManagedDescriptor == InISMCDescriptor)
 			{
 				MISMC->MarkAsUsed();
-				return ISMC;
+				return MISMC;
 			}
 		}
 	}
@@ -100,7 +100,7 @@ UInstancedStaticMeshComponent* UPCGActorHelpers::GetOrCreateISMC(AActor* InTarge
 	Resource->GeneratedComponent = ISMC;
 	InSourceComponent->AddToManagedResources(Resource);
 
-	return ISMC;
+	return Resource;
 }
 
 UInstancedStaticMeshComponent* UPCGActorHelpers::GetOrCreateISMC(AActor* InTargetActor, UPCGComponent* InSourceComponent, const FPCGISMCBuilderParameters& InParams)
diff --git a/Engine/Plugins/Experimental/PCG/Source/PCG/Public/Helpers/PCGActorHelpers.h b/Engine/Plugins/Experimental/PCG/Source/PCG/Public/Helpers/PCGActorHelpers.h
--- a/Engine/Plugins/Experimental/PCG/Source/PCG/Public/Helpers/PCGActorHelpers.h
+++ b/Engine/Plugins/Experimental/PCG/Source/PCG/Public/Helpers/PCGActorHelpers.h
@@ -39,7 +39,7 @@ class PCG_API UPCGActorHelpers : public UBlueprintFunctionLibrary
 	GENERATED_BODY()
 
 public:
-	static UInstancedStaticMeshComponent* GetOrCreateISMC(AActor* InTargetActor, UPCGComponent* InSourceComponent, const FISMComponentDescriptor& InISMCDescriptor);
+	static UPCGManagedISMComponent* GetOrCreateManagedISMC(AActor* InTargetActor, UPCGComponent* InSourceComponent, const FISMComponentDescriptor& InISMCDescriptor);
 	static UInstancedStaticMeshComponent* GetOrCreateISMC(AActor* InTargetActor, UPCGComponent* SourceComponent, const FPCGISMCBuilderParameters& Params);
 	static UPCGManagedISMComponent* GetOrCreateManagedISMC(AActor* InTargetActor, UPCGComponent* SourceComponent, const FPCGISMCBuilderParameters& Params);
 	static bool DeleteActors(UWorld* World, const TArray<TSoftObjectPtr<AActor>>& ActorsToDelete);
