[PCG] Tackling dirty on load with regeneration

* Dirty flag is now transient and force to false on load.
* Add `IsEquivalent` method in graph interface to avoid dirtying if the graph instance is the same (but not with the same ptr)
* Don't dirty PartitionActor on remap with BP on load

diff --git a/Engine/Plugins/Experimental/PCG/Source/PCG/Private/Grid/PCGPartitionActor.cpp b/Engine/Plugins/Experimental/PCG/Source/PCG/Private/Grid/PCGPartitionActor.cpp
--- a/Engine/Plugins/Experimental/PCG/Source/PCG/Private/Grid/PCGPartitionActor.cpp
+++ b/Engine/Plugins/Experimental/PCG/Source/PCG/Private/Grid/PCGPartitionActor.cpp
@@ -318,6 +318,8 @@ void APCGPartitionActor::AddGraphInstance(UPCGComponent* OriginalComponent)
 
 void APCGPartitionActor::RemapGraphInstance(const UPCGComponent* OldOriginalComponent, UPCGComponent* NewOriginalComponent)
 {
+	check(OldOriginalComponent && NewOriginalComponent);
+
 	UPCGComponent* LocalComponent = GetLocalComponent(OldOriginalComponent);
 
 	if (!LocalComponent)
@@ -325,7 +327,14 @@ void APCGPartitionActor::RemapGraphInstance(const UPCGComponent* OldOriginalComp
 		return;
 	}
 
-	Modify();
+	// If the old original component was loaded, we can assume we are in a loading phase.
+	// In this case, we don't want to dirty or register a transaction
+	const bool bIsLoading = OldOriginalComponent->HasAnyFlags(RF_WasLoaded);
+
+	if (!bIsLoading)
+	{
+		Modify();
+	}
 
 	OriginalToLocal.Remove(OldOriginalComponent);
 	LocalToOriginal.Remove(LocalComponent);
@@ -336,7 +345,10 @@ void APCGPartitionActor::RemapGraphInstance(const UPCGComponent* OldOriginalComp
 
 #if WITH_EDITOR
 	// When changing original data, it means that the data we have might point to newly stale data, hence we need to force dirty here
-	LocalComponent->DirtyGenerated(EPCGComponentDirtyFlag::Actor);
+	if (!bIsLoading)
+	{
+		LocalComponent->DirtyGenerated(EPCGComponentDirtyFlag::Actor);
+	}
 #endif
 }
 
diff --git a/Engine/Plugins/Experimental/PCG/Source/PCG/Private/PCGComponent.cpp b/Engine/Plugins/Experimental/PCG/Source/PCG/Private/PCGComponent.cpp
--- a/Engine/Plugins/Experimental/PCG/Source/PCG/Private/PCGComponent.cpp
+++ b/Engine/Plugins/Experimental/PCG/Source/PCG/Private/PCGComponent.cpp
@@ -199,7 +199,7 @@ void UPCGComponent::SetPropertiesFromOriginal(const UPCGComponent* Original)
 	Original->CachedTrackedTagsToSettings.GetKeys(OriginalTrackedTags);
 	const bool bHasDirtyTracking = !(TrackedTags.Num() == OriginalTrackedTags.Num() && TrackedTags.Includes(OriginalTrackedTags));
 
-	const bool bGraphInstanceIsDifferent = GraphInstance->Graph != Original->GraphInstance;
+	const bool bGraphInstanceIsDifferent = !GraphInstance->IsEquivalent(Original->GraphInstance);
 	const bool bIsDirty = bHasDirtyInput || bHasDirtyExclusions || bHasDirtyTracking || bGraphInstanceIsDifferent;
 #endif // WITH_EDITOR
 
@@ -979,6 +979,9 @@ void UPCGComponent::PostLoad()
 	Super::PostLoad();
 
 #if WITH_EDITOR
+	// Force dirty to be false on load. We should never refresh on load.
+	bDirtyGenerated = false;
+
 	if (!ExclusionTags_DEPRECATED.IsEmpty() && ExcludedTags.IsEmpty())
 	{
 		ExcludedTags.Append(ExclusionTags_DEPRECATED);
diff --git a/Engine/Plugins/Experimental/PCG/Source/PCG/Private/PCGGraph.cpp b/Engine/Plugins/Experimental/PCG/Source/PCG/Private/PCGGraph.cpp
--- a/Engine/Plugins/Experimental/PCG/Source/PCG/Private/PCGGraph.cpp
+++ b/Engine/Plugins/Experimental/PCG/Source/PCG/Private/PCGGraph.cpp
@@ -26,6 +26,18 @@ bool UPCGGraphInterface::IsInstance() const
 	return this != GetGraph();
 }
 
+bool UPCGGraphInterface::IsEquivalent(const UPCGGraphInterface* Other) const
+{
+	if (this == Other)
+	{
+		return true;
+	}
+
+	const UPCGGraph* OtherGraph = Other ? Other->GetGraph() : nullptr;
+
+	return GetGraph() == OtherGraph;
+}
+
 /****************************
 * UPCGGraph
 ****************************/
diff --git a/Engine/Plugins/Experimental/PCG/Source/PCG/Public/PCGComponent.h b/Engine/Plugins/Experimental/PCG/Source/PCG/Public/PCGComponent.h
--- a/Engine/Plugins/Experimental/PCG/Source/PCG/Public/PCGComponent.h
+++ b/Engine/Plugins/Experimental/PCG/Source/PCG/Public/PCGComponent.h
@@ -195,7 +195,7 @@ public:
 	UPROPERTY(BlueprintReadWrite, EditAnywhere, AdvancedDisplay, Category = Properties, meta = (DisplayName = "Regenerate PCG volume in editor"))
 	bool bRegenerateInEditor = true;
 
-	UPROPERTY(BlueprintReadOnly, VisibleAnywhere, AdvancedDisplay, Category = Properties)
+	UPROPERTY(BlueprintReadOnly, VisibleAnywhere, Transient, AdvancedDisplay, Category = Properties)
 	bool bDirtyGenerated = false;
 
 	FOnPCGGraphGenerated OnPCGGraphGeneratedDelegate;
diff --git a/Engine/Plugins/Experimental/PCG/Source/PCG/Public/PCGGraph.h b/Engine/Plugins/Experimental/PCG/Source/PCG/Public/PCGGraph.h
--- a/Engine/Plugins/Experimental/PCG/Source/PCG/Public/PCGGraph.h
+++ b/Engine/Plugins/Experimental/PCG/Source/PCG/Public/PCGGraph.h
@@ -24,6 +24,9 @@ public:
 
 	bool IsInstance() const;
 
+	/** A graph interface is equivalent to another graph interface if they are the same (same ptr), or if they have the same graph. Will be overriden when graph instance supports overrides. */
+	virtual bool IsEquivalent(const UPCGGraphInterface* Other) const;
+
 #if WITH_EDITOR
 	FOnPCGGraphChanged OnGraphChangedDelegate;
 #endif // WITH_EDITOR
