Fixed raycast dependency on having cleaned data by marking/unmarking components as needed

diff --git a/Engine/Plugins/Experimental/PCG/Source/PCG/Private/Data/PCGWorldData.cpp b/Engine/Plugins/Experimental/PCG/Source/PCG/Private/Data/PCGWorldData.cpp
--- a/Engine/Plugins/Experimental/PCG/Source/PCG/Private/Data/PCGWorldData.cpp
+++ b/Engine/Plugins/Experimental/PCG/Source/PCG/Private/Data/PCGWorldData.cpp
@@ -52,8 +52,14 @@ bool UPCGWorldVolumetricData::SamplePoint(const FTransform& InTransform, const F
 			continue;
 		}
 
-		// Skip PCG-created objects optionally
-		if (QueryParams.bIgnorePCGHits && OverlappedComponent->ComponentTags.Contains(PCGHelpers::DefaultPCGTag))
+		// Skip to-be-cleaned-up PCG-created objects
+		if (OverlappedComponent->ComponentHasTag(PCGHelpers::MarkedForCleanupPCGTag) || (OverlappedComponent->GetOwner() && OverlappedComponent->GetOwner()->ActorHasTag(PCGHelpers::MarkedForCleanupPCGTag)))
+		{
+			continue;
+		}
+
+		// Optionally skip all PCG created objects
+		if (QueryParams.bIgnorePCGHits && (OverlappedComponent->ComponentHasTag(PCGHelpers::DefaultPCGTag) || (OverlappedComponent->GetOwner() && OverlappedComponent->GetOwner()->ActorHasTag(PCGHelpers::DefaultPCGActorTag))))
 		{
 			continue;
 		}
@@ -187,8 +193,14 @@ bool UPCGWorldRayHitData::SamplePoint(const FTransform& InTransform, const FBox&
 			continue;
 		}
 		
-		// Skip PCG-created objects optionally
-		if (QueryParams.bIgnorePCGHits && HitComponent->ComponentTags.Contains(PCGHelpers::DefaultPCGTag))
+		// Skip to-be-cleaned-up PCG-created objects
+		if (HitComponent->ComponentHasTag(PCGHelpers::MarkedForCleanupPCGTag) || (HitComponent->GetOwner() && HitComponent->GetOwner()->ActorHasTag(PCGHelpers::MarkedForCleanupPCGTag)))
+		{
+			continue;
+		}
+
+		// Optionally skip all PCG created objects
+		if (QueryParams.bIgnorePCGHits && (HitComponent->ComponentHasTag(PCGHelpers::DefaultPCGTag) || (HitComponent->GetOwner() && HitComponent->GetOwner()->ActorHasTag(PCGHelpers::DefaultPCGActorTag))))
 		{
 			continue;
 		}
diff --git a/Engine/Plugins/Experimental/PCG/Source/PCG/Private/PCGHelpers.h b/Engine/Plugins/Experimental/PCG/Source/PCG/Private/PCGHelpers.h
--- a/Engine/Plugins/Experimental/PCG/Source/PCG/Private/PCGHelpers.h
+++ b/Engine/Plugins/Experimental/PCG/Source/PCG/Private/PCGHelpers.h
@@ -17,6 +17,7 @@ namespace PCGHelpers
 	const FName DefaultPCGTag = TEXT("PCG Generated Component");
 	const FName DefaultPCGDebugTag = TEXT("PCG Generated Debug Component");
 	const FName DefaultPCGActorTag = TEXT("PCG Generated Actor");
+	const FName MarkedForCleanupPCGTag = TEXT("PCG Marked For Cleanup");
 
 	int ComputeSeed(int A);
 	int ComputeSeed(int A, int B);
diff --git a/Engine/Plugins/Experimental/PCG/Source/PCG/Private/PCGManagedResource.cpp b/Engine/Plugins/Experimental/PCG/Source/PCG/Private/PCGManagedResource.cpp
--- a/Engine/Plugins/Experimental/PCG/Source/PCG/Private/PCGManagedResource.cpp
+++ b/Engine/Plugins/Experimental/PCG/Source/PCG/Private/PCGManagedResource.cpp
@@ -3,6 +3,7 @@
 #include "PCGManagedResource.h"
 
 #include "PCGComponent.h"
+#include "PCGHelpers.h"
 #include "PCGModule.h"
 #include "Components/InstancedStaticMeshComponent.h"
 #include "Utils/PCGGeneratedResourcesLogging.h"
@@ -69,6 +70,15 @@ bool UPCGManagedActors::Release(bool bHardRelease, TSet<TSoftObjectPtr<AActor>>&
 	{
 		PCGGeneratedResourcesLogging::LogManagedActorsSoftRelease(GeneratedActors);
 
+		// Mark actors as potentially-to-be-cleaned-up
+		for (TSoftObjectPtr<AActor> GeneratedActor : GeneratedActors)
+		{
+			if (GeneratedActor.IsValid())
+			{
+				GeneratedActor->Tags.Add(PCGHelpers::MarkedForCleanupPCGTag);
+			}
+		}
+
 		return false;
 	}
 
@@ -125,6 +135,26 @@ bool UPCGManagedActors::MoveResourceToNewActor(AActor* NewActor)
 	return true;
 }
 
+void UPCGManagedActors::MarkAsUsed()
+{
+	Super::MarkAsUsed();
+	// Technically we don't ever have to "use" a preexisting managed actor resource, but this is to be consistent with the other implementations
+	ensure(0);
+}
+
+void UPCGManagedActors::MarkAsReused()
+{
+	Super::MarkAsReused();
+
+	for (TSoftObjectPtr<AActor> GeneratedActor : GeneratedActors)
+	{
+		if (GeneratedActor.IsValid())
+		{
+			GeneratedActor->Tags.Remove(PCGHelpers::MarkedForCleanupPCGTag);
+		}
+	}
+}
+
 void UPCGManagedComponent::PostEditImport()
 {
 	Super::PostEditImport();
@@ -181,6 +211,7 @@ bool UPCGManagedComponent::Release(bool bHardRelease, TSet<TSoftObjectPtr<AActor
 		{
 			// We can only mark it unused if we can reset the component.
 			bIsMarkedUnused = true;
+			GeneratedComponent->ComponentTags.Add(PCGHelpers::MarkedForCleanupPCGTag);
 		}
 	}
 
@@ -241,11 +272,27 @@ void UPCGManagedComponent::MarkAsUsed()
 		return;
 	}
 
+	Super::MarkAsUsed();
+
 	// Can't reuse a resource if we can't reset it. Make sure we never take this path in this case.
 	check(SupportsComponentReset());
 
 	ResetComponent();
-	bIsMarkedUnused = false;
+
+	if (GeneratedComponent.Get())
+	{
+		GeneratedComponent->ComponentTags.Remove(PCGHelpers::MarkedForCleanupPCGTag);
+	}
+}
+
+void UPCGManagedComponent::MarkAsReused()
+{
+	Super::MarkAsReused();
+
+	if (GeneratedComponent.Get())
+	{
+		GeneratedComponent->ComponentTags.Remove(PCGHelpers::MarkedForCleanupPCGTag);
+	}
 }
 
 bool UPCGManagedISMComponent::ReleaseIfUnused(TSet<TSoftObjectPtr<AActor>>& OutActorsToDelete)
diff --git a/Engine/Plugins/Experimental/PCG/Source/PCG/Public/PCGManagedResource.h b/Engine/Plugins/Experimental/PCG/Source/PCG/Public/PCGManagedResource.h
--- a/Engine/Plugins/Experimental/PCG/Source/PCG/Public/PCGManagedResource.h
+++ b/Engine/Plugins/Experimental/PCG/Source/PCG/Public/PCGManagedResource.h
@@ -64,6 +64,8 @@ public:
 	virtual bool Release(bool bHardRelease, TSet<TSoftObjectPtr<AActor>>& OutActorsToDelete) override;
 	virtual bool ReleaseIfUnused(TSet<TSoftObjectPtr<AActor>>& OutActorsToDelete) override;
 	virtual bool MoveResourceToNewActor(AActor* NewActor) override;
+	virtual void MarkAsUsed() override;
+	virtual void MarkAsReused() override;
 	//~End UPCGManagedResource interface
 
 	UPROPERTY(BlueprintReadWrite, EditAnywhere, Category = GeneratedData)
@@ -89,6 +91,7 @@ public:
 	virtual void ResetComponent() { check(0); }
 	virtual bool SupportsComponentReset() const { return false; }
 	virtual void MarkAsUsed() override;
+	virtual void MarkAsReused() override;
 
 	UPROPERTY(BlueprintReadWrite, EditAnywhere, Category = GeneratedData)
 	TSoftObjectPtr<UActorComponent> GeneratedComponent;
