[PCG] Don't refresh BP on load

diff --git a/Engine/Plugins/Experimental/PCG/Source/PCG/Private/PCGComponent.cpp b/Engine/Plugins/Experimental/PCG/Source/PCG/Private/PCGComponent.cpp
--- a/Engine/Plugins/Experimental/PCG/Source/PCG/Private/PCGComponent.cpp
+++ b/Engine/Plugins/Experimental/PCG/Source/PCG/Private/PCGComponent.cpp
@@ -2806,7 +2806,11 @@ void FPCGComponentInstanceData::ApplyToComponent(UActorComponent* Component, con
 		// Finally, start a delayed refresh task (if there is not one already), in editor only
 		// It is important to be delayed, because we cannot spawn Partition Actors within this scope,
 		// because we are in a construction script.
-		PCGComponent->Refresh();
+		// Note that we only do this if we are not currently loading
+		if (!SourceComponent || !SourceComponent->HasAllFlags(RF_WasLoaded))
+		{
+			PCGComponent->Refresh();
+		}
 #endif // WITH_EDITOR
 	}
 }
