[PCG] Be more robust if PCG Components are deleted by the user on the PA

* Local components now notify the PA if they are destroyed, to remove any dangling references
* Safeguard to remove any dangling reference in the post load.

diff --git a/Engine/Plugins/Experimental/PCG/Source/PCG/Private/Grid/PCGPartitionActor.cpp b/Engine/Plugins/Experimental/PCG/Source/PCG/Private/Grid/PCGPartitionActor.cpp
--- a/Engine/Plugins/Experimental/PCG/Source/PCG/Private/Grid/PCGPartitionActor.cpp
+++ b/Engine/Plugins/Experimental/PCG/Source/PCG/Private/Grid/PCGPartitionActor.cpp
@@ -59,6 +59,13 @@ void APCGPartitionActor::PostLoad()
 	}
 #endif
 
+	// Safe guard if we ever load a local that was deleted but not removed (like if the user deleted themselves the component)
+	// We can have multiple nullptr entries (if multiple components were removed), hence the while.
+	while (LocalToOriginal.Contains(nullptr))
+	{
+		LocalToOriginal.Remove(nullptr);
+	}
+
 	// Make sure that we don't track objects that do not exist anymore
 	CleanupDeadGraphInstances(/*bRemoveNullOnly=*/true);
 
@@ -376,6 +383,19 @@ bool APCGPartitionActor::RemoveGraphInstance(UPCGComponent* OriginalComponent)
 	return OriginalToLocal.IsEmpty();
 }
 
+void APCGPartitionActor::RemoveLocalComponent(UPCGComponent* LocalComponent)
+{
+	if (!LocalComponent)
+	{
+		return;
+	}
+
+	UPCGComponent* OriginalComponent = GetOriginalComponent(LocalComponent);
+
+	LocalToOriginal.Remove(LocalComponent);
+	OriginalToLocal.Remove(OriginalComponent);
+}
+
 #if WITH_EDITOR
 FBox APCGPartitionActor::GetStreamingBounds() const
 {
diff --git a/Engine/Plugins/Experimental/PCG/Source/PCG/Private/Grid/PCGPartitionActor.h b/Engine/Plugins/Experimental/PCG/Source/PCG/Private/Grid/PCGPartitionActor.h
--- a/Engine/Plugins/Experimental/PCG/Source/PCG/Private/Grid/PCGPartitionActor.h
+++ b/Engine/Plugins/Experimental/PCG/Source/PCG/Private/Grid/PCGPartitionActor.h
@@ -59,6 +59,9 @@ public:
 	bool RemoveGraphInstance(UPCGComponent* OriginalComponent);
 	void CleanupDeadGraphInstances(bool bRemoveNonNullOnly = false);
 
+	// When a local component is destroyed. It calls this function. We make sure we don't keep mappings that are dead.
+	void RemoveLocalComponent(UPCGComponent* LocalComponent);
+
 #if WITH_EDITOR
 	/** To be called after the creation of a new actor to copy the GridSize property (Editor only) into the PCGGridSize property */
 	void PostCreation();
diff --git a/Engine/Plugins/Experimental/PCG/Source/PCG/Private/PCGComponent.cpp b/Engine/Plugins/Experimental/PCG/Source/PCG/Private/PCGComponent.cpp
--- a/Engine/Plugins/Experimental/PCG/Source/PCG/Private/PCGComponent.cpp
+++ b/Engine/Plugins/Experimental/PCG/Source/PCG/Private/PCGComponent.cpp
@@ -971,6 +971,13 @@ void UPCGComponent::OnComponentDestroyed(bool bDestroyingHierarchy)
 	}
 #endif
 
+	// Bookkeeping local components that might be deleted by the user.
+	// Making sure that the corresponding partition actor doesn't keep a dangling references
+	if (APCGPartitionActor* PAOwner = Cast<APCGPartitionActor>(GetOwner()))
+	{
+		PAOwner->RemoveLocalComponent(this);
+	}
+
 	Super::OnComponentDestroyed(bDestroyingHierarchy);
 }
 
