[PCG] First pass for revamping the tracking system

* Tracking will be done in a specialized class owned by the subsystem
* This new class will also handle mapping for partition actors, making it clear that it's its responsibility to register and track component mappings.
* Therefore all the logic for registering components was moved in this new class
* Avoid to add tracked actor for each component, which can be costly with partitioned components
* Each source component (not local) will register itself when it gets created, to know if it tracks actors
* Each component will update its tracking requirements when the component change
* Actor update will be captured by the subsystem, and will dispatch information to the right components

diff --git a/Engine/Plugins/Experimental/PCG/Source/PCG/Private/Elements/PCGSpawnActor.cpp b/Engine/Plugins/Experimental/PCG/Source/PCG/Private/Elements/PCGSpawnActor.cpp
--- a/Engine/Plugins/Experimental/PCG/Source/PCG/Private/Elements/PCGSpawnActor.cpp
+++ b/Engine/Plugins/Experimental/PCG/Source/PCG/Private/Elements/PCGSpawnActor.cpp
@@ -641,7 +641,7 @@ bool FPCGSpawnActorElement::ExecuteInternal(FPCGContext* Context) const
 
 						if (bForceCallGenerate || (bOnLoadCallGenerate && PCGComponent->GenerationTrigger == EPCGComponentGenerationTrigger::GenerateOnLoad))
 						{
-							if (Subsystem && PCGComponent->IsPartitioned())
+							if (Subsystem)
 							{
 								Subsystem->RegisterOrUpdatePCGComponent(PCGComponent);
 							}
diff --git a/Engine/Plugins/Experimental/PCG/Source/PCG/Private/Grid/PCGComponentOctree.cpp b/Engine/Plugins/Experimental/PCG/Source/PCG/Private/Grid/PCGComponentOctree.cpp
--- a/Engine/Plugins/Experimental/PCG/Source/PCG/Private/Grid/PCGComponentOctree.cpp
+++ b/Engine/Plugins/Experimental/PCG/Source/PCG/Private/Grid/PCGComponentOctree.cpp
@@ -19,3 +19,160 @@ void FPCGComponentRef::UpdateBounds()
 
 	Bounds = Component->GetGridBounds();
 }
+
+FPCGComponentOctreeAndMap::FPCGComponentOctreeAndMap(const FVector& InOrigin, FVector::FReal InExtent)
+	: Octree(InOrigin, InExtent)
+{
+
+}
+
+void FPCGComponentOctreeAndMap::Reset(const FVector& InOrigin, FVector::FReal InExtent)
+{
+	FWriteScopeLock WriteLock(Lock);
+	Octree = FPCGComponentOctree(InOrigin, InExtent);
+	ComponentToIdMap.Empty();
+}
+
+bool FPCGComponentOctreeAndMap::Contains(const UPCGComponent* InComponent) const
+{
+	FReadScopeLock ReadLock(Lock);
+
+	return ComponentToIdMap.Find(InComponent) != nullptr;
+}
+
+FBox FPCGComponentOctreeAndMap::GetBounds(const UPCGComponent* InComponent) const
+{
+	FBox Bounds(EForceInit::ForceInit);
+
+	FReadScopeLock ReadLock(Lock);
+	if (const FPCGComponentOctreeIDSharedRef* ElementIdPtr = ComponentToIdMap.Find(InComponent))
+	{
+		const FPCGComponentRef& ComponentRef = Octree.GetElementById((*ElementIdPtr)->Id);
+		Bounds = ComponentRef.Bounds.GetBox();
+	}
+
+	return Bounds;
+}
+
+void FPCGComponentOctreeAndMap::AddOrUpdateComponent(UPCGComponent* InComponent, FBox& OutBounds, bool& bOutComponentHasChanged, bool& bOutComponentWasAdded)
+{
+	FWriteScopeLock WriteLock(Lock);
+
+	FPCGComponentOctreeIDSharedRef* ElementIdPtr = ComponentToIdMap.Find(InComponent);
+
+	if (!ElementIdPtr)
+	{
+		// Does not exist yet, add it.
+		FPCGComponentOctreeIDSharedRef IdShared = MakeShared<FPCGComponentOctreeID>();
+		FPCGComponentRef ComponentRef(InComponent, IdShared);
+		OutBounds = ComponentRef.Bounds.GetBox();
+		check(OutBounds.IsValid);
+
+		// If the Component is already generated, it probably mean we are in loading. The component bounds and last
+		// generated bounds should be the same.
+		// If the bounds depends on other components on the owner however, it might not be the same, because of the registration order.
+		// In this case, override the bounds by the last generated ones.
+		FBox LastGeneratedBounds = InComponent->GetLastGeneratedBounds();
+		if (InComponent->bGenerated && !LastGeneratedBounds.Equals(OutBounds))
+		{
+			OutBounds = LastGeneratedBounds;
+			ComponentRef.Bounds = OutBounds;
+		}
+
+		Octree.AddElement(ComponentRef);
+
+		bOutComponentHasChanged = true;
+		bOutComponentWasAdded = true;
+
+		// Store the shared ptr, because if we add/remove components in the octree, the id might change.
+		// We need to make sure that we always have the latest id for the given component.
+		ComponentToIdMap.Add(InComponent, MoveTemp(IdShared));
+	}
+	else
+	{
+		// It already exists, update it if the bounds changed.
+
+		// Do a copy here.
+		FPCGComponentRef ComponentRef = Octree.GetElementById((*ElementIdPtr)->Id);
+		FBox PreviousBounds = ComponentRef.Bounds.GetBox();
+		ComponentRef.UpdateBounds();
+		OutBounds = ComponentRef.Bounds.GetBox();
+		check(OutBounds.IsValid);
+
+		// If bounds changed, remove and re-add to the octree
+		if (!PreviousBounds.Equals(OutBounds))
+		{
+			Octree.RemoveElement((*ElementIdPtr)->Id);
+			Octree.AddElement(ComponentRef);
+			bOutComponentHasChanged = true;
+		}
+	}
+}
+
+bool FPCGComponentOctreeAndMap::RemapComponent(const UPCGComponent* InOldComponent, UPCGComponent* InNewComponent, bool& bOutBoundsHasChanged)
+{
+	// First verification we have the old component registered
+	if (!Contains(InOldComponent))
+	{
+		return false;
+	}
+
+	// If so, lock again in write and recheck if it has not been remapped already.
+	{
+		FWriteScopeLock WriteLock(Lock);
+
+		FPCGComponentOctreeIDSharedRef* ElementIdPtr = ComponentToIdMap.Find(InOldComponent);
+
+		if (!ElementIdPtr)
+		{
+			// Nothing done
+			return false;
+		}
+
+		FPCGComponentOctreeIDSharedRef ElementId = *ElementIdPtr;
+
+		FPCGComponentRef ComponentRef = Octree.GetElementById(ElementId->Id);
+		FBox PreviousBounds = ComponentRef.Bounds.GetBox();
+		ComponentRef.Component = InNewComponent;
+		ComponentRef.UpdateBounds();
+		FBox Bounds = ComponentRef.Bounds.GetBox();
+		check(Bounds.IsValid);
+
+		// If bounds changed, we need to update the mapping
+		bOutBoundsHasChanged = !Bounds.Equals(PreviousBounds);
+
+		Octree.RemoveElement((*ElementIdPtr)->Id);
+		Octree.AddElement(ComponentRef);
+		ComponentToIdMap.Remove(InOldComponent);
+		ComponentToIdMap.Add(InNewComponent, MoveTemp(ElementId));
+	}
+
+	return true;
+}
+
+bool FPCGComponentOctreeAndMap::RemoveComponent(UPCGComponent* InComponent)
+{
+	FWriteScopeLock WriteLock(Lock);
+
+	FPCGComponentOctreeIDSharedRef* ElementIdPtr = ComponentToIdMap.Find(InComponent);
+
+	if (!ElementIdPtr)
+	{
+		return false;
+	}
+
+	Octree.RemoveElement((*ElementIdPtr)->Id);
+	ComponentToIdMap.Remove(InComponent);
+
+	return true;
+}
+
+TSet<TObjectPtr<UPCGComponent>> FPCGComponentOctreeAndMap::GetAllComponents() const
+{
+	FReadScopeLock ReadLock(Lock);
+
+	TSet<TObjectPtr<UPCGComponent>> Components;
+	ComponentToIdMap.GetKeys(Components);
+
+	return Components;
+}
diff --git a/Engine/Plugins/Experimental/PCG/Source/PCG/Private/PCGActorAndComponentMapping.cpp b/Engine/Plugins/Experimental/PCG/Source/PCG/Private/PCGActorAndComponentMapping.cpp
new file mode 100644
--- /dev/null
+++ b/Engine/Plugins/Experimental/PCG/Source/PCG/Private/PCGActorAndComponentMapping.cpp
@@ -0,0 +1,578 @@
+// Copyright Epic Games, Inc. All Rights Reserved.
+
+#include "PCGActorAndComponentMapping.h"
+
+#include "PCGComponent.h"
+#include "PCGHelpers.h"
+#include "PCGModule.h"
+#include "PCGSubsystem.h"
+#include "PCGWorldActor.h"
+#include "Grid/PCGPartitionActor.h"
+#include "Helpers/PCGActorHelpers.h"
+
+#include "Engine/Engine.h"
+
+UPCGActorAndComponentMapping::UPCGActorAndComponentMapping(UPCGSubsystem* InPCGSubsystem)
+	: PCGSubsystem(InPCGSubsystem)
+{
+	check(PCGSubsystem);
+
+	// TODO: For now we set our octree to be 2km wide, but it would be perhaps better to
+	// scale it to the size of our world.
+	constexpr FVector::FReal OctreeExtent = 200000; // 2km
+	PartitionedOctree.Reset(FVector::ZeroVector, OctreeExtent);
+	NonPartitionedOctree.Reset(FVector::ZeroVector, OctreeExtent);
+}
+
+void UPCGActorAndComponentMapping::Tick()
+{
+	TSet<TObjectPtr<UPCGComponent>> ComponentToUnregister;
+	{
+		FScopeLock Lock(&DelayedComponentToUnregisterLock);
+		ComponentToUnregister = MoveTemp(DelayedComponentToUnregister);
+	}
+
+	for (UPCGComponent* Component : ComponentToUnregister)
+	{
+		UnregisterPCGComponent(Component, /*bForce=*/true);
+	}
+}
+
+TArray<FPCGTaskId> UPCGActorAndComponentMapping::DispatchToRegisteredLocalComponents(UPCGComponent* OriginalComponent, const TFunction<FPCGTaskId(UPCGComponent*)>& InFunc) const
+{
+	TRACE_CPUPROFILER_EVENT_SCOPE(UPCGActorAndComponentMapping::DispatchToRegisteredLocalComponents);
+
+	// TODO: Might be more interesting to copy the set and release the lock.
+	FReadScopeLock ReadLock(ComponentToPartitionActorsMapLock);
+	const TSet<TObjectPtr<APCGPartitionActor>>* PartitionActorsPtr = ComponentToPartitionActorsMap.Find(OriginalComponent);
+
+	if (!PartitionActorsPtr)
+	{
+		return TArray<FPCGTaskId>();
+	}
+
+	return DispatchToLocalComponents(OriginalComponent, *PartitionActorsPtr, InFunc);
+}
+
+TArray<FPCGTaskId> UPCGActorAndComponentMapping::DispatchToLocalComponents(UPCGComponent* OriginalComponent, const TSet<TObjectPtr<APCGPartitionActor>>& PartitionActors, const TFunction<FPCGTaskId(UPCGComponent*)>& InFunc) const
+{
+	TArray<FPCGTaskId> TaskIds;
+	for (APCGPartitionActor* PartitionActor : PartitionActors)
+	{
+		if (PartitionActor)
+		{
+			if (UPCGComponent* LocalComponent = PartitionActor->GetLocalComponent(OriginalComponent))
+			{
+				// Add check to avoid infinite loop
+				if (ensure(!LocalComponent->IsPartitioned()))
+				{
+					FPCGTaskId LocalTask = InFunc(LocalComponent);
+
+					if (LocalTask != InvalidPCGTaskId)
+					{
+						TaskIds.Add(LocalTask);
+					}
+				}
+			}
+		}
+	}
+
+	return TaskIds;
+}
+
+bool UPCGActorAndComponentMapping::RegisterOrUpdatePCGComponent(UPCGComponent* InComponent, bool bDoActorMapping)
+{
+	check(InComponent);
+
+	// Discard BP templates, local components and invalid component
+	if (!InComponent->GetOwner() || InComponent->GetOwner()->IsA<APCGPartitionActor>() || !IsValid(InComponent))
+	{
+		return false;
+	}
+
+	// Check also that the bounds are valid. If not early out.
+	if (!InComponent->GetGridBounds().IsValid)
+	{
+		UE_LOG(LogPCG, Error, TEXT("[RegisterOrUpdatePCGComponent] Component has invalid bounds, not registered nor updated."));
+		return false;
+	}
+
+	// First check if the component has changed its partitioned flag.
+	const bool bIsPartitioned = InComponent->IsPartitioned();
+	if (bIsPartitioned && NonPartitionedOctree.Contains(InComponent))
+	{
+		UnregisterNonPartitionedPCGComponent(InComponent);
+	}
+	else if (!bIsPartitioned && PartitionedOctree.Contains(InComponent))
+	{
+		UnregisterPartitionedPCGComponent(InComponent);
+	}
+
+	// Then register/update accordingly
+	bool bHasChanged = false;
+	if (bIsPartitioned)
+	{
+		bHasChanged = RegisterOrUpdatePartitionedPCGComponent(InComponent, bDoActorMapping);
+	}
+	else
+	{
+		bHasChanged = RegisterOrUpdateNonPartitionedPCGComponent(InComponent);
+	}
+
+	// And finally handle the tracking
+	RegisterOrUpdateTracking(InComponent);
+
+	return bHasChanged;
+}
+
+bool UPCGActorAndComponentMapping::RegisterOrUpdatePartitionedPCGComponent(UPCGComponent* InComponent, bool bDoActorMapping)
+{
+	FBox Bounds(EForceInit::ForceInit);
+	bool bComponentHasChanged = false;
+	bool bComponentWasAdded = false;
+
+	PartitionedOctree.AddOrUpdateComponent(InComponent, Bounds, bComponentHasChanged, bComponentWasAdded);
+
+#if WITH_EDITOR
+	// In Editor only, we will create new partition actors depending on the new bounds
+	// TODO: For now it will always create the PA. But if we want to create them only when we generate, we need to make
+	// sure to update the runtime flow, for them to also create PA if they need to.
+	if (bComponentHasChanged || bComponentWasAdded)
+	{
+		PCGSubsystem->CreatePartitionActorsWithinBounds(Bounds);
+	}
+#endif // WITH_EDITOR
+
+	// After adding/updating, try to do the mapping (if we asked for it and the component changed)
+	if (bDoActorMapping)
+	{
+		if (bComponentHasChanged)
+		{
+			UpdateMappingPCGComponentPartitionActor(InComponent);
+		}
+	}
+	else
+	{
+		if (!bComponentWasAdded)
+		{
+			// If we do not want a mapping, delete the existing one
+			DeleteMappingPCGComponentPartitionActor(InComponent);
+		}
+	}
+
+	return bComponentHasChanged;
+}
+
+bool UPCGActorAndComponentMapping::RegisterOrUpdateNonPartitionedPCGComponent(UPCGComponent* InComponent)
+{
+	// Tracking is only done in Editor for now
+#if WITH_EDITOR
+	const bool bIsTracking = InComponent->ShouldTrackActors();
+
+	if (!bIsTracking && NonPartitionedOctree.Contains(InComponent))
+	{
+		// Was tracking, but no more, remove it
+		UnregisterNonPartitionedPCGComponent(InComponent);
+		return true;
+	}
+
+	if (!bIsTracking)
+	{
+		// Nothing to do
+		return false;
+	}
+
+	FBox Bounds(EForceInit::ForceInit);
+	bool bComponentHasChanged = false;
+	bool bComponentWasAdded = false;
+
+	NonPartitionedOctree.AddOrUpdateComponent(InComponent, Bounds, bComponentHasChanged, bComponentWasAdded);
+
+	return bComponentHasChanged;
+#else
+	return false;
+#endif // WITH_EDITOR
+}
+
+bool UPCGActorAndComponentMapping::RemapPCGComponent(const UPCGComponent* OldComponent, UPCGComponent* NewComponent, bool bDoActorMapping)
+{
+	check(OldComponent && NewComponent);
+
+	bool bBoundsChanged = false;
+
+	if (OldComponent->IsPartitioned())
+	{
+		if (!PartitionedOctree.RemapComponent(OldComponent, NewComponent, bBoundsChanged))
+		{
+			return false;
+		}
+	}
+	else
+	{
+		if (!NonPartitionedOctree.RemapComponent(OldComponent, NewComponent, bBoundsChanged))
+		{
+			return false;
+		}
+	}
+
+	// Remove it from the delayed
+	{
+		FScopeLock Lock(&DelayedComponentToUnregisterLock);
+		DelayedComponentToUnregister.Remove(OldComponent);
+	}
+
+	// Remap all previous instances
+	{
+		FWriteScopeLock WriteLock(ComponentToPartitionActorsMapLock);
+		TSet<TObjectPtr<APCGPartitionActor>>* PartitionActorsPtr = ComponentToPartitionActorsMap.Find(OldComponent);
+
+		if (PartitionActorsPtr)
+		{
+			TSet<TObjectPtr<APCGPartitionActor>> PartitionActorsToRemap = MoveTemp(*PartitionActorsPtr);
+			ComponentToPartitionActorsMap.Remove(OldComponent);
+
+			for (APCGPartitionActor* Actor : PartitionActorsToRemap)
+			{
+				Actor->RemapGraphInstance(OldComponent, NewComponent);
+			}
+
+			ComponentToPartitionActorsMap.Add(NewComponent, MoveTemp(PartitionActorsToRemap));
+		}
+	}
+
+	// And update the mapping if bounds changed and we want to do actor mapping
+	if (bBoundsChanged && NewComponent->IsPartitioned() && bDoActorMapping)
+	{
+		UpdateMappingPCGComponentPartitionActor(NewComponent);
+	}
+
+	RemapTracking(OldComponent, NewComponent);
+
+	return true;
+}
+
+void UPCGActorAndComponentMapping::UnregisterPCGComponent(UPCGComponent* InComponent, bool bForce)
+{
+	if (!InComponent)
+	{
+		return;
+	}
+
+	if ((PartitionedOctree.Contains(InComponent) || NonPartitionedOctree.Contains(InComponent)))
+	{
+		// We also need to check that our current PCG Component is not deleted while being reconstructed by a construction script.
+		// If so, it will be "re-created" at some point with the same properties.
+		// In this particular case, we don't remove the PCG component from the octree and we won't delete the mapping, but mark it to be removed
+		// at next Subsystem tick. If we call "RemapPCGComponent" before, we will re-connect everything correctly.
+		// Ignore this if we force (aka when we actually unregister the delayed one)
+		if (InComponent->IsCreatedByConstructionScript() && !bForce)
+		{
+			FScopeLock Lock(&DelayedComponentToUnregisterLock);
+			DelayedComponentToUnregister.Add(InComponent);
+			return;
+		}
+	}
+
+	UnregisterPartitionedPCGComponent(InComponent);
+	UnregisterNonPartitionedPCGComponent(InComponent);
+
+	UnregisterTracking(InComponent);
+}
+
+void UPCGActorAndComponentMapping::UnregisterPartitionedPCGComponent(UPCGComponent* InComponent)
+{
+	if (!PartitionedOctree.RemoveComponent(InComponent))
+	{
+		return;
+	}
+
+	// Because of recursive component deletes actors that has components, we cannot do RemoveGraphInstance
+	// inside a lock. So copy the actors to clean up and release the lock before doing RemoveGraphInstance.
+	TSet<TObjectPtr<APCGPartitionActor>> PartitionActorsToCleanUp;
+	{
+		FWriteScopeLock WriteLock(ComponentToPartitionActorsMapLock);
+		TSet<TObjectPtr<APCGPartitionActor>>* PartitionActorsPtr = ComponentToPartitionActorsMap.Find(InComponent);
+
+		if (PartitionActorsPtr)
+		{
+			PartitionActorsToCleanUp = MoveTemp(*PartitionActorsPtr);
+			ComponentToPartitionActorsMap.Remove(InComponent);
+		}
+	}
+
+	for (APCGPartitionActor* Actor : PartitionActorsToCleanUp)
+	{
+		Actor->RemoveGraphInstance(InComponent);
+	}
+}
+
+void UPCGActorAndComponentMapping::UnregisterNonPartitionedPCGComponent(UPCGComponent* InComponent)
+{
+	NonPartitionedOctree.RemoveComponent(InComponent);
+}
+
+void UPCGActorAndComponentMapping::ForAllIntersectingComponents(const FBoxCenterAndExtent& InBounds, TFunction<void(UPCGComponent*)> InFunc) const
+{
+	PartitionedOctree.FindElementsWithBoundsTest(InBounds, [&InFunc](const FPCGComponentRef& ComponentRef)
+	{
+		InFunc(ComponentRef.Component);
+	});
+}
+
+void UPCGActorAndComponentMapping::RegisterPartitionActor(APCGPartitionActor* Actor, bool bDoComponentMapping)
+{
+	check(Actor);
+
+	FIntVector GridCoord = Actor->GetGridCoord();
+	{
+		FWriteScopeLock WriteLock(PartitionActorsMapLock);
+
+		if (PartitionActorsMap.Contains(GridCoord))
+		{
+			return;
+		}
+
+		PartitionActorsMap.Add(GridCoord, Actor);
+	}
+
+	// For deprecration: bUse2DGrid is now true by default. But if we already have Partition Actors that were created when the flag was false by default,
+	// we keep this flag
+	if (APCGWorldActor* WorldActor = PCGSubsystem->GetPCGWorldActor())
+	{
+		if (WorldActor->bUse2DGrid != Actor->IsUsing2DGrid())
+		{
+			WorldActor->bUse2DGrid = Actor->IsUsing2DGrid();
+		}
+	}
+
+	// And then register itself to all the components that intersect with it
+	{
+		FWriteScopeLock WriteLock(ComponentToPartitionActorsMapLock);
+		ForAllIntersectingComponents(FBoxCenterAndExtent(Actor->GetFixedBounds()), [this, Actor, bDoComponentMapping](UPCGComponent* Component)
+		{
+			// For each component, do the mapping if we ask it explicitly, or if the component is generated
+			if (bDoComponentMapping || Component->bGenerated)
+			{
+				TSet<TObjectPtr<APCGPartitionActor>>* PartitionActorsPtr = ComponentToPartitionActorsMap.Find(Component);
+				// In editor we might load/create partition actors while the component is registering. Because of that,
+				// the mapping might not already exists, even if the component is marked generated.
+				if (PartitionActorsPtr)
+				{
+					Actor->AddGraphInstance(Component);
+					PartitionActorsPtr->Add(Actor);
+				}
+			}
+		});
+	}
+}
+
+void UPCGActorAndComponentMapping::UnregisterPartitionActor(APCGPartitionActor* Actor)
+{
+	check(Actor);
+
+	FIntVector GridCoord = Actor->GetGridCoord();
+
+	{
+		FWriteScopeLock WriteLock(PartitionActorsMapLock);
+		PartitionActorsMap.Remove(GridCoord);
+	}
+
+	// And then unregister itself to all the components that intersect with it
+	{
+		FWriteScopeLock WriteLock(ComponentToPartitionActorsMapLock);
+		ForAllIntersectingComponents(FBoxCenterAndExtent(Actor->GetFixedBounds()), [this, Actor](UPCGComponent* Component)
+		{
+			TSet<TObjectPtr<APCGPartitionActor>>* PartitionActorsPtr = ComponentToPartitionActorsMap.Find(Component);
+		if (PartitionActorsPtr)
+		{
+			PartitionActorsPtr->Remove(Actor);
+		}
+		});
+	}
+}
+
+void UPCGActorAndComponentMapping::ForAllIntersectingPartitionActors(const FBox& InBounds, TFunction<void(APCGPartitionActor*)> InFunc) const
+{
+	// No PCGWorldActor just early out. Same for invalid bounds.
+	APCGWorldActor* PCGWorldActor = PCGSubsystem->GetPCGWorldActor();
+
+	if (!PCGWorldActor || !InBounds.IsValid)
+	{
+		return;
+	}
+
+	const uint32 GridSize = PCGWorldActor->PartitionGridSize;
+	const bool bUse2DGrid = PCGWorldActor->bUse2DGrid;
+	FIntVector MinCellCoords = UPCGActorHelpers::GetCellCoord(InBounds.Min, GridSize, bUse2DGrid);
+	FIntVector MaxCellCoords = UPCGActorHelpers::GetCellCoord(InBounds.Max, GridSize, bUse2DGrid);
+
+	{
+		FReadScopeLock ReadLock(PartitionActorsMapLock);
+
+		if (PartitionActorsMap.IsEmpty())
+		{
+			return;
+		}
+
+		for (int32 z = MinCellCoords.Z; z <= MaxCellCoords.Z; z++)
+		{
+			for (int32 y = MinCellCoords.Y; y <= MaxCellCoords.Y; y++)
+			{
+				for (int32 x = MinCellCoords.X; x <= MaxCellCoords.X; x++)
+				{
+					FIntVector CellCoords(x, y, z);
+					if (const TObjectPtr<APCGPartitionActor>* ActorPtr = PartitionActorsMap.Find(CellCoords))
+					{
+						if (APCGPartitionActor* Actor = ActorPtr->Get())
+						{
+							InFunc(Actor);
+						}
+					}
+				}
+			}
+		}
+	}
+}
+
+void UPCGActorAndComponentMapping::UpdateMappingPCGComponentPartitionActor(UPCGComponent* InComponent)
+{
+	check(InComponent);
+
+	// Get the bounds
+	FBox Bounds = PartitionedOctree.GetBounds(InComponent);
+
+	if (!Bounds.IsValid)
+	{
+		return;
+	}
+
+	TSet<TObjectPtr<APCGPartitionActor>> RemovedActors;
+
+	{
+		FWriteScopeLock WriteLock(ComponentToPartitionActorsMapLock);
+
+		TSet<TObjectPtr<APCGPartitionActor>>* PartitionActorsPtr = ComponentToPartitionActorsMap.Find(InComponent);
+
+		if (!PartitionActorsPtr)
+		{
+			// Does not yet exists, add it
+			PartitionActorsPtr = &ComponentToPartitionActorsMap.Emplace(InComponent);
+			check(PartitionActorsPtr);
+		}
+
+		TSet<TObjectPtr<APCGPartitionActor>> NewMapping;
+		ForAllIntersectingPartitionActors(Bounds, [&NewMapping, InComponent](APCGPartitionActor* Actor)
+		{
+			Actor->AddGraphInstance(InComponent);
+			NewMapping.Add(Actor);
+		});
+
+		// Find the ones that were removed
+		RemovedActors = PartitionActorsPtr->Difference(NewMapping);
+
+		*PartitionActorsPtr = MoveTemp(NewMapping);
+	}
+
+	// No need to be locked to do this.
+	for (APCGPartitionActor* RemovedActor : RemovedActors)
+	{
+		RemovedActor->RemoveGraphInstance(InComponent);
+	}
+}
+
+void UPCGActorAndComponentMapping::DeleteMappingPCGComponentPartitionActor(UPCGComponent* InComponent)
+{
+	check(InComponent);
+
+	if (!InComponent->IsPartitioned())
+	{
+		return;
+	}
+
+	FWriteScopeLock WriteLock(ComponentToPartitionActorsMapLock);
+
+	TSet<TObjectPtr<APCGPartitionActor>>* PartitionActorsPtr = ComponentToPartitionActorsMap.Find(InComponent);
+	if (PartitionActorsPtr)
+	{
+		for (APCGPartitionActor* Actor : *PartitionActorsPtr)
+		{
+			Actor->RemoveGraphInstance(InComponent);
+		}
+
+		PartitionActorsPtr->Empty();
+	}
+}
+
+TSet<TObjectPtr<UPCGComponent>> UPCGActorAndComponentMapping::GetAllRegisteredPartitionedComponents() const
+{
+	return PartitionedOctree.GetAllComponents();
+}
+
+bool UPCGActorAndComponentMapping::IsComponentTracked(const UPCGComponent* InComponent) const
+{
+	FReadScopeLock ReadLock(TrackedComponentsLock);
+	return TrackedComponentsToActorsMap.Contains(InComponent);
+}
+
+// TODO: In a following pass
+void UPCGActorAndComponentMapping::RegisterOrUpdateTracking(UPCGComponent* InComponent)
+{
+
+}
+
+void UPCGActorAndComponentMapping::RemapTracking(const UPCGComponent* InOldComponent, UPCGComponent* InNewComponent)
+{
+
+}
+
+void UPCGActorAndComponentMapping::UnregisterTracking(UPCGComponent* InComponent)
+{
+
+}
+
+#if WITH_EDITOR
+void UPCGActorAndComponentMapping::ResetPartitionActorsMap()
+{
+	PartitionActorsMapLock.WriteLock();
+	PartitionActorsMap.Empty();
+	PartitionActorsMapLock.WriteUnlock();
+}
+
+void UPCGActorAndComponentMapping::RegisterTrackingCallbacks()
+{
+	GEngine->OnActorMoved().AddRaw(this, &UPCGActorAndComponentMapping::OnActorMoved);
+	GEngine->OnLevelActorAdded().AddRaw(this, &UPCGActorAndComponentMapping::OnActorAdded);
+	GEngine->OnLevelActorDeleted().AddRaw(this, &UPCGActorAndComponentMapping::OnActorDeleted);
+	FCoreUObjectDelegates::OnObjectPropertyChanged.AddRaw(this, &UPCGActorAndComponentMapping::OnObjectPropertyChanged);
+}
+
+void UPCGActorAndComponentMapping::TeardownTrackingCallbacks()
+{
+	GEngine->OnActorMoved().RemoveAll(this);
+	GEngine->OnLevelActorAdded().RemoveAll(this);
+	GEngine->OnLevelActorDeleted().RemoveAll(this);
+	FCoreUObjectDelegates::OnObjectPropertyChanged.RemoveAll(this);
+}
+
+// TODO: In a following pass
+void UPCGActorAndComponentMapping::OnActorAdded(AActor* InActor)
+{
+
+}
+
+void UPCGActorAndComponentMapping::OnActorDeleted(AActor* InActor)
+{
+
+}
+
+void UPCGActorAndComponentMapping::OnActorMoved(AActor* InActor)
+{
+
+}
+
+void UPCGActorAndComponentMapping::OnObjectPropertyChanged(UObject* InObject, FPropertyChangedEvent& InEvent)
+{
+
+}
+
+#endif // WITH_EDITOR
\ No newline at end of file
diff --git a/Engine/Plugins/Experimental/PCG/Source/PCG/Private/PCGComponent.cpp b/Engine/Plugins/Experimental/PCG/Source/PCG/Private/PCGComponent.cpp
--- a/Engine/Plugins/Experimental/PCG/Source/PCG/Private/PCGComponent.cpp
+++ b/Engine/Plugins/Experimental/PCG/Source/PCG/Private/PCGComponent.cpp
@@ -81,16 +81,9 @@ void UPCGComponent::SetIsPartitioned(bool bIsNowPartitioned)
 			CleanupLocalImmediate(/*bRemoveComponents=*/true);
 		}
 
-		if (bIsNowPartitioned)
-		{
-			bIsComponentPartitioned = bIsNowPartitioned;
-			Subsystem->RegisterOrUpdatePCGComponent(this, bDoActorMapping);
-		}
-		else
-		{
-			Subsystem->UnregisterPCGComponent(this);
-			bIsComponentPartitioned = bIsNowPartitioned;
-		}
+		// Update the component on the subsystem
+		bIsComponentPartitioned = bIsNowPartitioned;
+		Subsystem->RegisterOrUpdatePCGComponent(this, bDoActorMapping);
 	}
 	else
 	{
@@ -907,8 +900,8 @@ void UPCGComponent::BeginPlay()
 	}
 #endif
 
-	// First if it is partitioned, register itself to the PCGSubsystem, to map the component to all its corresponding PartitionActors
-	if (IsPartitioned() && GetSubsystem())
+	// Register itself to the PCGSubsystem, to map the component to all its corresponding PartitionActors if it is partition among other things.
+	if (GetSubsystem())
 	{
 		GetSubsystem()->RegisterOrUpdatePCGComponent(this);
 	}
@@ -1132,7 +1125,7 @@ void UPCGComponent::OnRegister()
 	// We can't register to the subsystem in OnRegister if we are at runtime because
 	// the landscape can be not loaded yet.
 	// It will be done in BeginPlay at runtime.
-	if (!PCGHelpers::IsRuntimeOrPIE() && IsPartitioned() && GetSubsystem())
+	if (!PCGHelpers::IsRuntimeOrPIE() && GetSubsystem())
 	{
 		if (UWorld* World = GetWorld())
 		{
@@ -1880,14 +1873,7 @@ void UPCGComponent::OnRefresh()
 
 	if (Subsystem)
 	{
-		if (IsPartitioned())
-		{
-			Subsystem->RegisterOrUpdatePCGComponent(this, /*bDoActorMapping=*/ bWasGenerated);
-		}
-		else
-		{
-			Subsystem->UnregisterPCGComponent(this);
-		}
+		Subsystem->RegisterOrUpdatePCGComponent(this, /*bDoActorMapping=*/ bWasGenerated);
 	}
 
 	// Following a change in some properties or in some spatial information related to this component,
@@ -2802,11 +2788,13 @@ void FPCGComponentInstanceData::ApplyToComponent(UActorComponent* Component, con
 		// Re-set the graph to reconnect callbacks correctly
 		PCGComponent->GraphInstance->SetGraph(PCGComponent->GraphInstance->Graph);
 
-		// Also remap if we are partitioned
+		bool bDoActorMapping = PCGComponent->bGenerated || PCGHelpers::IsRuntimeOrPIE();
+
+		// Also remap
 		UPCGSubsystem* Subsystem = PCGComponent->GetSubsystem();
-		if (Subsystem && PCGComponent->IsPartitioned() && SourceComponent)
+		if (Subsystem && SourceComponent)
 		{
-			Subsystem->RemapPCGComponent(SourceComponent, PCGComponent);
+			Subsystem->RemapPCGComponent(SourceComponent, PCGComponent, bDoActorMapping);
 		}
 
 #if WITH_EDITOR
diff --git a/Engine/Plugins/Experimental/PCG/Source/PCG/Private/PCGSubsystem.cpp b/Engine/Plugins/Experimental/PCG/Source/PCG/Private/PCGSubsystem.cpp
--- a/Engine/Plugins/Experimental/PCG/Source/PCG/Private/PCGSubsystem.cpp
+++ b/Engine/Plugins/Experimental/PCG/Source/PCG/Private/PCGSubsystem.cpp
@@ -91,11 +91,15 @@ namespace PCGSubsystem
 
 UPCGSubsystem::UPCGSubsystem()
 	: Super()
+	, ActorAndComponentMapping(this)
 {
-	// TODO: For now we set our octree to be 2km wide, but it would be perhaps better to
-	// scale it to the size of our world.
-	constexpr FVector::FReal OctreeExtent = 200000; // 2km
-	PCGComponentOctree = FPCGComponentOctree(FVector::ZeroVector, OctreeExtent);
+#if WITH_EDITOR
+	if (!HasAnyFlags(RF_ClassDefaultObject))
+	{
+		// TODO: Uncomment when tracking is functionnal
+		// ActorAndComponentMapping.RegisterTrackingCallbacks();
+	}
+#endif // WITH_EDITOR
 }
 
 void UPCGSubsystem::Deinitialize()
@@ -108,6 +112,11 @@ void UPCGSubsystem::Deinitialize()
 	PCGWorldActor = nullptr;
 	bHasTickedOnce = false;
 
+#if WITH_EDITOR
+	// TODO: Uncomment when tracking is functionnal
+	// ActorAndComponentMapping.TeardownTrackingCallbacks();
+#endif // WITH_EDITOR
+
 	Super::Deinitialize();
 }
 
@@ -158,16 +167,7 @@ void UPCGSubsystem::Tick(float DeltaSeconds)
 		GetLandscapeCache()->Tick(DeltaSeconds);
 	}
 
-	TSet<TObjectPtr<UPCGComponent>> ComponentToUnregister;
-	{
-		FScopeLock Lock(&DelayedComponentToUnregisterLock);
-		ComponentToUnregister = MoveTemp(DelayedComponentToUnregister);
-	}
-
-	for (UPCGComponent* Component : ComponentToUnregister)
-	{
-		UnregisterPCGComponent(Component, /*bForce=*/true);
-	}
+	ActorAndComponentMapping.Tick();
 }
 
 APCGWorldActor* UPCGSubsystem::GetPCGWorldActor()
@@ -241,7 +241,7 @@ FPCGTaskId UPCGSubsystem::ScheduleComponent(UPCGComponent* PCGComponent, bool bS
 
 		PCGSubsystem::ForEachIntersectingCell(GraphExecutor, PCGComponent->GetWorld(), PCGComponent->GetGridBounds(), /*bCreateActor=*/true, /*bLoadCell=*/false, /*bSave=*/false, ScheduleTask);
 
-		UpdateMappingPCGComponentPartitionActor(PCGComponent);
+		ActorAndComponentMapping.UpdateMappingPCGComponentPartitionActor(PCGComponent);
 	}
 #endif // WITH_EDITOR
 
@@ -272,7 +272,7 @@ FPCGTaskId UPCGSubsystem::ScheduleComponent(UPCGComponent* PCGComponent, bool bS
 			return LocalComponent->GenerateInternal(/*bForce=*/ false, EPCGComponentGenerationTrigger::GenerateOnDemand, Dependencies);
 		};
 
-		AllTasks = DispatchToRegisteredLocalComponents(PCGComponent, LocalGenerateTask);
+		AllTasks = ActorAndComponentMapping.DispatchToRegisteredLocalComponents(PCGComponent, LocalGenerateTask);
 	}
 	else
 	{
@@ -338,7 +338,7 @@ FPCGTaskId UPCGSubsystem::ScheduleCleanup(UPCGComponent* PCGComponent, bool bRem
 			return LocalComponent->CleanupInternal(bRemoveComponents, /*bSave=*/ false, Dependencies);
 		};
 
-		AllTasks = DispatchToRegisteredLocalComponents(PCGComponent, LocalCleanupTask);
+		AllTasks = ActorAndComponentMapping.DispatchToRegisteredLocalComponents(PCGComponent, LocalCleanupTask);
 	}
 	else
 	{
@@ -366,7 +366,7 @@ FPCGTaskId UPCGSubsystem::ScheduleCleanup(UPCGComponent* PCGComponent, bool bRem
 			// If we are in Editor and partitioned, delete the mapping.
 			if (Component->IsPartitioned() && !PCGHelpers::IsRuntimeOrPIE())
 			{
-				DeleteMappingPCGComponentPartitionActor(Component);
+				ActorAndComponentMapping.DeleteMappingPCGComponentPartitionActor(Component);
 			}
 #endif // WITH_EDITOR
 		}
@@ -456,7 +456,7 @@ void UPCGSubsystem::CancelGeneration(UPCGComponent* Component)
 			return InvalidPCGTaskId;
 		};
 
-		DispatchToRegisteredLocalComponents(Component, LocalCancel);
+		ActorAndComponentMapping.DispatchToRegisteredLocalComponents(Component, LocalCancel);
 	}
 
 	TArray<UPCGComponent*> CancelledComponents = GraphExecutor->Cancel(Component);
@@ -513,7 +513,7 @@ void UPCGSubsystem::ForAllRegisteredLocalComponents(UPCGComponent* OriginalCompo
 		return InvalidPCGTaskId;
 	};
 
-	DispatchToRegisteredLocalComponents(OriginalComponent, WrapperFunc);
+	ActorAndComponentMapping.DispatchToRegisteredLocalComponents(OriginalComponent, WrapperFunc);
 }
 
 bool UPCGSubsystem::IsGraphCacheDebuggingEnabled() const
@@ -521,526 +521,12 @@ bool UPCGSubsystem::IsGraphCacheDebuggingEnabled() const
 	return GraphExecutor && GraphExecutor->IsGraphCacheDebuggingEnabled();
 }
 
-TArray<FPCGTaskId> UPCGSubsystem::DispatchToRegisteredLocalComponents(UPCGComponent* OriginalComponent, const TFunction<FPCGTaskId(UPCGComponent*)>& InFunc) const
-{
-	TRACE_CPUPROFILER_EVENT_SCOPE(UPCGSubsystem::DispatchToRegisteredLocalComponents);
-
-	// TODO: Might be more interesting to copy the set and release the lock.
-	FReadScopeLock ReadLock(ComponentToPartitionActorsMapLock);
-	const TSet<TObjectPtr<APCGPartitionActor>>* PartitionActorsPtr = ComponentToPartitionActorsMap.Find(OriginalComponent);
-
-	if (!PartitionActorsPtr)
-	{
-		return TArray<FPCGTaskId>();
-	}
-
-	return DispatchToLocalComponents(OriginalComponent, *PartitionActorsPtr, InFunc);
-}
-
-TArray<FPCGTaskId> UPCGSubsystem::DispatchToLocalComponents(UPCGComponent* OriginalComponent, const TSet<TObjectPtr<APCGPartitionActor>>& PartitionActors, const TFunction<FPCGTaskId(UPCGComponent*)>& InFunc) const
-{
-	TArray<FPCGTaskId> TaskIds;
-	for (APCGPartitionActor* PartitionActor : PartitionActors)
-	{
-		if (PartitionActor)
-		{
-			if (UPCGComponent* LocalComponent = PartitionActor->GetLocalComponent(OriginalComponent))
-			{
-				// Add check to avoid infinite loop
-				if (ensure(!LocalComponent->IsPartitioned()))
-				{
-					FPCGTaskId LocalTask = InFunc(LocalComponent);
-
-					if (LocalTask != InvalidPCGTaskId)
-					{
-						TaskIds.Add(LocalTask);
-					}
-				}
-			}
-		}
-	}
-
-	return TaskIds;
-}
-
 bool UPCGSubsystem::GetOutputData(FPCGTaskId TaskId, FPCGDataCollection& OutData)
 {
 	check(GraphExecutor);
 	return GraphExecutor->GetOutputData(TaskId, OutData);
 }
 
-bool UPCGSubsystem::RegisterOrUpdatePCGComponent(UPCGComponent* InComponent, bool bDoActorMapping)
-{
-	// Just make sure that we don't register components that are from the PCG Partition Actor
-	check(InComponent && InComponent->GetOwner() && !InComponent->GetOwner()->IsA<APCGPartitionActor>());
-
-	// If the component is not partitioned nor valid, don't register/update it.
-	if (!InComponent->IsPartitioned() || !IsValid(InComponent))
-	{
-		return false;
-	}
-
-	// Check also that the bounds are valid. If not early out.
-	if (!InComponent->GetGridBounds().IsValid)
-	{
-		UE_LOG(LogPCG, Error, TEXT("[RegisterOrUpdatePCGComponent] Component has invalid bounds, not registered nor updated."));
-		return false;
-	}
-
-	FBox Bounds(EForceInit::ForceInit);
-	bool bComponentHasChanged = false;
-	bool bComponentWasAdded = false;
-
-	{
-		FWriteScopeLock WriteLock(PCGVolumeOctreeLock);
-
-		FPCGComponentOctreeIDSharedRef* ElementIdPtr = ComponentToIdMap.Find(InComponent);
-
-		if (!ElementIdPtr)
-		{
-			// Does not exist yet, add it.
-			FPCGComponentOctreeIDSharedRef IdShared = MakeShared<FPCGComponentOctreeID>();
-			FPCGComponentRef ComponentRef(InComponent, IdShared);
-			Bounds = ComponentRef.Bounds.GetBox();
-			check(Bounds.IsValid);
-
-			// If the Component is already generated, it probably mean we are in loading. The component bounds and last
-			// generated bounds should be the same.
-			// If the bounds depends on other components on the owner however, it might not be the same, because of the registration order.
-			// In this case, override the bounds by the last generated ones.
-			if (InComponent->bGenerated && !InComponent->LastGeneratedBounds.Equals(Bounds))
-			{
-				Bounds = InComponent->LastGeneratedBounds;
-				ComponentRef.Bounds = Bounds;
-			}
-
-			PCGComponentOctree.AddElement(ComponentRef);
-
-			bComponentHasChanged = true;
-			bComponentWasAdded = true;
-
-			// Store the shared ptr, because if we add/remove components in the octree, the id might change.
-			// We need to make sure that we always have the latest id for the given component.
-			ComponentToIdMap.Add(InComponent, MoveTemp(IdShared));
-		}
-		else
-		{
-			// It already exists, update it if the bounds changed.
-
-			// Do a copy here.
-			FPCGComponentRef ComponentRef = PCGComponentOctree.GetElementById((*ElementIdPtr)->Id);
-			FBox PreviousBounds = ComponentRef.Bounds.GetBox();
-			ComponentRef.UpdateBounds();
-			Bounds = ComponentRef.Bounds.GetBox();
-			check(Bounds.IsValid);
-
-			// If bounds changed, remove and re-add to the octree
-			if (!PreviousBounds.Equals(Bounds))
-			{
-				PCGComponentOctree.RemoveElement((*ElementIdPtr)->Id);
-				PCGComponentOctree.AddElement(ComponentRef);
-				bComponentHasChanged = true;
-			}
-		}
-	}
-
-#if WITH_EDITOR
-	// In Editor only, we will create new partition actors depending on the new bounds
-	// TODO: For now it will always create the PA. But if we want to create them only when we generate, we need to make
-	// sure to update the runtime flow, for them to also create PA if they need to.
-	if (!PCGHelpers::IsRuntimeOrPIE())
-	{
-		auto ScheduleTask = [](APCGPartitionActor* PCGActor, const FBox& InIntersectedBounds, const TArray<FPCGTaskId>& TaskDependencies) { return InvalidPCGTaskId; };
-
-		// We can't spawn actors if we are running constructions scripts, asserting when we try to get the actor with the WP API.
-		// We should never enter this if we are in a construction script. If the ensure is hit, we need to fix it.
-		UWorld* World = InComponent->GetWorld();
-		if (ensure(World && !World->bIsRunningConstructionScript))
-		{
-			PCGSubsystem::ForEachIntersectingCell(GraphExecutor, World, Bounds, /*bCreateActor=*/true, /*bLoadCell=*/false, /*bSave=*/false, ScheduleTask);
-		}
-	}
-#endif // WITH_EDITOR
-
-	// After adding/updating, try to do the mapping (if we asked for it and the component changed)
-	if (bDoActorMapping)
-	{
-		if (bComponentHasChanged)
-		{
-			UpdateMappingPCGComponentPartitionActor(InComponent);
-		}
-	}
-	else
-	{
-		if (!bComponentWasAdded)
-		{
-			// If we do not want a mapping, delete the existing one
-			DeleteMappingPCGComponentPartitionActor(InComponent);
-		}
-	}
-
-	return bComponentHasChanged;
-}
-
-bool UPCGSubsystem::RemapPCGComponent(const UPCGComponent* OldComponent, UPCGComponent* NewComponent)
-{
-	check(OldComponent && NewComponent);
-
-	// First verification we have the old component registered
-	{
-		FReadScopeLock ReadLock(PCGVolumeOctreeLock);
-
-		FPCGComponentOctreeIDSharedRef* ElementIdPtr = ComponentToIdMap.Find(OldComponent);
-
-		if (!ElementIdPtr)
-		{
-			return false;
-		}
-	}
-
-	// If so, lock again in write and recheck if it has not been remapped already.
-	bool bBoundsChanged = false;
-	{
-		FWriteScopeLock WriteLock(PCGVolumeOctreeLock);
-
-		FPCGComponentOctreeIDSharedRef* ElementIdPtr = ComponentToIdMap.Find(OldComponent);
-
-		if (!ElementIdPtr)
-		{
-			// It was remapped, so return true there.
-			return true;
-		}
-
-		FPCGComponentOctreeIDSharedRef ElementId = *ElementIdPtr;
-
-		FPCGComponentRef ComponentRef = PCGComponentOctree.GetElementById(ElementId->Id);
-		FBox PreviousBounds = ComponentRef.Bounds.GetBox();
-		ComponentRef.Component = NewComponent;
-		ComponentRef.UpdateBounds();
-		FBox Bounds = ComponentRef.Bounds.GetBox();
-		check(Bounds.IsValid);
-
-		// If bounds changed, we need to update the mapping
-		bBoundsChanged = !Bounds.Equals(PreviousBounds);
-
-		PCGComponentOctree.RemoveElement((*ElementIdPtr)->Id);
-		PCGComponentOctree.AddElement(ComponentRef);
-		ComponentToIdMap.Remove(OldComponent);
-		ComponentToIdMap.Add(NewComponent, MoveTemp(ElementId));
-	}
-
-	// Remove it from the delayed
-	{
-		FScopeLock Lock(&DelayedComponentToUnregisterLock);
-		DelayedComponentToUnregister.Remove(OldComponent);
-	}
-
-	// Remap all previous instances
-	{
-		FWriteScopeLock WriteLock(ComponentToPartitionActorsMapLock);
-		TSet<TObjectPtr<APCGPartitionActor>>* PartitionActorsPtr = ComponentToPartitionActorsMap.Find(OldComponent);
-
-		if (PartitionActorsPtr)
-		{
-			TSet<TObjectPtr<APCGPartitionActor>> PartitionActorsToRemap = MoveTemp(*PartitionActorsPtr);
-			ComponentToPartitionActorsMap.Remove(OldComponent);
-
-			for (APCGPartitionActor* Actor : PartitionActorsToRemap)
-			{
-				Actor->RemapGraphInstance(OldComponent, NewComponent);
-			}
-
-			ComponentToPartitionActorsMap.Add(NewComponent, MoveTemp(PartitionActorsToRemap));
-		}
-	}
-
-	// And update the mapping if bounds changed
-	if (bBoundsChanged)
-	{
-		UpdateMappingPCGComponentPartitionActor(NewComponent);
-	}
-
-	return true;
-}
-
-void UPCGSubsystem::UnregisterPCGComponent(UPCGComponent* InComponent, bool bForce)
-{
-	check(InComponent);
-
-	// First verification we have this component registered
-	{
-		FReadScopeLock ReadLock(PCGVolumeOctreeLock);
-
-		FPCGComponentOctreeIDSharedRef* ElementIdPtr = ComponentToIdMap.Find(InComponent);
-
-		if (!ElementIdPtr)
-		{
-			return;
-		}
-	}
-
-	// We also need to check that our current PCG Component is not deleted while being reconstructed by a construction script.
-	// If so, it will be "re-created" at some point with the same properties.
-	// In this particular case, we don't remove the PCG component from the octree and we won't delete the mapping, but mark it to be removed
-	// at next Subsystem tick. If we call "RemapPCGComponent" before, we will re-connect everything correctly.
-	// Ignore this if we force (aka when we actually unregister the delayed one)
-	if (InComponent->IsCreatedByConstructionScript() && !bForce)
-	{
-		FScopeLock Lock(&DelayedComponentToUnregisterLock);
-		DelayedComponentToUnregister.Add(InComponent);
-		return;
-	}
-
-	// If so, lock in write and retry to find it (to avoid removing it twice)
-	{
-		FWriteScopeLock WriteLock(PCGVolumeOctreeLock);
-
-		FPCGComponentOctreeIDSharedRef* ElementIdPtr = ComponentToIdMap.Find(InComponent);
-
-		if (!ElementIdPtr)
-		{
-			return;
-		}
-
-		PCGComponentOctree.RemoveElement((*ElementIdPtr)->Id);
-		ComponentToIdMap.Remove(InComponent);
-	}
-
-	// Because of recursive component deletes actors that has components, we cannot do RemoveGraphInstance
-	// inside a lock. So copy the actors to clean up and release the lock before doing RemoveGraphInstance.
-	TSet<TObjectPtr<APCGPartitionActor>> PartitionActorsToCleanUp;
-	{
-		FWriteScopeLock WriteLock(ComponentToPartitionActorsMapLock);
-		TSet<TObjectPtr<APCGPartitionActor>>* PartitionActorsPtr = ComponentToPartitionActorsMap.Find(InComponent);
-		
-		if (PartitionActorsPtr)
-		{
-			PartitionActorsToCleanUp = MoveTemp(*PartitionActorsPtr);
-			ComponentToPartitionActorsMap.Remove(InComponent);
-		}
-	}
-
-	for (APCGPartitionActor* Actor : PartitionActorsToCleanUp)
-	{
-		Actor->RemoveGraphInstance(InComponent);
-	}
-}
-
-void UPCGSubsystem::ForAllIntersectingComponents(const FBoxCenterAndExtent& InBounds, TFunction<void(UPCGComponent*)> InFunc) const
-{
-	FReadScopeLock ReadLock(PCGVolumeOctreeLock);
-
-	PCGComponentOctree.FindElementsWithBoundsTest(InBounds, [&InFunc](const FPCGComponentRef& ComponentRef)
-		{
-			InFunc(ComponentRef.Component);
-		});
-}
-
-void UPCGSubsystem::RegisterPartitionActor(APCGPartitionActor* Actor, bool bDoComponentMapping)
-{
-	check(Actor);
-
-	FIntVector GridCoord = Actor->GetGridCoord();
-	{
-		FWriteScopeLock WriteLock(PartitionActorsMapLock);
-
-		if (PartitionActorsMap.Contains(GridCoord))
-		{
-			return;
-		}
-
-		PartitionActorsMap.Add(GridCoord, Actor);
-	}
-
-	// For deprecration: bUse2DGrid is now true by default. But if we already have Partition Actors that were created when the flag was false by default,
-	// we keep this flag
-	if (APCGWorldActor* WorldActor = GetPCGWorldActor())
-	{
-		if (WorldActor->bUse2DGrid != Actor->IsUsing2DGrid())
-		{
-			WorldActor->bUse2DGrid = Actor->IsUsing2DGrid();
-		}
-	}
-
-	// And then register itself to all the components that intersect with it
-	{
-		FWriteScopeLock WriteLock(ComponentToPartitionActorsMapLock);
-		ForAllIntersectingComponents(FBoxCenterAndExtent(Actor->GetFixedBounds()), [this, Actor, bDoComponentMapping](UPCGComponent* Component)
-			{
-				// For each component, do the mapping if we ask it explicitly, or if the component is generated
-				if (bDoComponentMapping || Component->bGenerated)
-				{
-					TSet<TObjectPtr<APCGPartitionActor>>* PartitionActorsPtr = ComponentToPartitionActorsMap.Find(Component);
-					// In editor we might load/create partition actors while the component is registering. Because of that,
-					// the mapping might not already exists, even if the component is marked generated.
-					if (PartitionActorsPtr)
-					{
-						Actor->AddGraphInstance(Component);
-						PartitionActorsPtr->Add(Actor);
-					}
-				}
-			});
-	}
-}
-
-void UPCGSubsystem::UnregisterPartitionActor(APCGPartitionActor* Actor)
-{
-	check(Actor);
-
-	FIntVector GridCoord = Actor->GetGridCoord();
-
-	{
-		FWriteScopeLock WriteLock(PartitionActorsMapLock);
-		PartitionActorsMap.Remove(GridCoord);
-	}
-
-	// And then unregister itself to all the components that intersect with it
-	{
-		FWriteScopeLock WriteLock(ComponentToPartitionActorsMapLock);
-		ForAllIntersectingComponents(FBoxCenterAndExtent(Actor->GetFixedBounds()), [this, Actor](UPCGComponent* Component)
-			{
-				TSet<TObjectPtr<APCGPartitionActor>>* PartitionActorsPtr = ComponentToPartitionActorsMap.Find(Component);
-				if (PartitionActorsPtr)
-				{
-					PartitionActorsPtr->Remove(Actor);
-				}
-			});
-	}
-}
-
-void UPCGSubsystem::ForAllIntersectingPartitionActors(const FBox& InBounds, TFunction<void(APCGPartitionActor*)> InFunc) const
-{
-	// No PCGWorldActor just early out. Same for invalid bounds.
-	if (!PCGWorldActor || !InBounds.IsValid)
-	{
-		return;
-	}
-
-	const uint32 GridSize = PCGWorldActor->PartitionGridSize;
-	const bool bUse2DGrid = PCGWorldActor->bUse2DGrid;
-	FIntVector MinCellCoords = UPCGActorHelpers::GetCellCoord(InBounds.Min, GridSize, bUse2DGrid);
-	FIntVector MaxCellCoords = UPCGActorHelpers::GetCellCoord(InBounds.Max, GridSize, bUse2DGrid);
-
-	{
-		FReadScopeLock ReadLock(PartitionActorsMapLock);
-
-		if (PartitionActorsMap.IsEmpty())
-		{
-			return;
-		}
-
-		for (int32 z = MinCellCoords.Z; z <= MaxCellCoords.Z; z++)
-		{
-			for (int32 y = MinCellCoords.Y; y <= MaxCellCoords.Y; y++)
-			{
-				for (int32 x = MinCellCoords.X; x <= MaxCellCoords.X; x++)
-				{
-					FIntVector CellCoords(x, y, z);
-					if (const TObjectPtr<APCGPartitionActor>* ActorPtr = PartitionActorsMap.Find(CellCoords))
-					{
-						if (APCGPartitionActor* Actor = ActorPtr->Get())
-						{
-							InFunc(Actor);
-						}
-					}
-				}
-			}
-		}
-	}
-}
-
-void UPCGSubsystem::UpdateMappingPCGComponentPartitionActor(UPCGComponent* InComponent)
-{
-	check(InComponent);
-
-	// Get the bounds
-	FBox Bounds(EForceInit::ForceInit);
-	{
-		FReadScopeLock OctreeReadLock(PCGVolumeOctreeLock);
-
-		FPCGComponentOctreeIDSharedRef* ElementIdPtr = ComponentToIdMap.Find(InComponent);
-
-		if (ensure(ElementIdPtr))
-		{
-			const FPCGComponentRef& ComponentRef = PCGComponentOctree.GetElementById((*ElementIdPtr)->Id);
-			Bounds = ComponentRef.Bounds.GetBox();
-		}
-		else
-		{
-			return;
-		}
-	}
-
-	if (!Bounds.IsValid)
-	{
-		return;
-	}
-
-	TSet<TObjectPtr<APCGPartitionActor>> RemovedActors;
-
-	{
-		FWriteScopeLock WriteLock(ComponentToPartitionActorsMapLock);
-
-		TSet<TObjectPtr<APCGPartitionActor>>* PartitionActorsPtr = ComponentToPartitionActorsMap.Find(InComponent);
-
-		if (!PartitionActorsPtr)
-		{
-			// Does not yet exists, add it
-			PartitionActorsPtr = &ComponentToPartitionActorsMap.Emplace(InComponent);
-			check(PartitionActorsPtr);
-		}
-
-		TSet<TObjectPtr<APCGPartitionActor>> NewMapping;
-		ForAllIntersectingPartitionActors(Bounds, [&NewMapping, InComponent](APCGPartitionActor* Actor)
-		{
-			Actor->AddGraphInstance(InComponent);
-			NewMapping.Add(Actor);
-		});
-
-		// Find the ones that were removed
-		RemovedActors = PartitionActorsPtr->Difference(NewMapping);
-
-		*PartitionActorsPtr = MoveTemp(NewMapping);
-	}
-
-	// No need to be locked to do this.
-	for (APCGPartitionActor* RemovedActor : RemovedActors)
-	{
-		RemovedActor->RemoveGraphInstance(InComponent);
-	}
-}
-
-void UPCGSubsystem::DeleteMappingPCGComponentPartitionActor(UPCGComponent* InComponent)
-{
-	check(InComponent);
-
-	if (!InComponent->IsPartitioned())
-	{
-		return;
-	}
-
-	FWriteScopeLock WriteLock(ComponentToPartitionActorsMapLock);
-
-	TSet<TObjectPtr<APCGPartitionActor>>* PartitionActorsPtr = ComponentToPartitionActorsMap.Find(InComponent);
-	if (PartitionActorsPtr)
-	{
-		for (APCGPartitionActor* Actor : *PartitionActorsPtr)
-		{
-			Actor->RemoveGraphInstance(InComponent);
-		}
-
-		PartitionActorsPtr->Empty();
-	}
-}
-
-TSet<TObjectPtr<UPCGComponent>> UPCGSubsystem::GetAllRegisteredComponents() const
-{
-	TSet<TObjectPtr<UPCGComponent>> Components;
-	ComponentToIdMap.GetKeys(Components);
-
-	return Components;
-}
-
 #if WITH_EDITOR
 
 namespace PCGSubsystem
@@ -1215,6 +701,22 @@ namespace PCGSubsystem
 
 } // end namepsace PCGSubsystem
 
+void UPCGSubsystem::CreatePartitionActorsWithinBounds(const FBox& InBounds)
+{
+	if (!PCGHelpers::IsRuntimeOrPIE())
+	{
+		auto ScheduleTask = [](APCGPartitionActor* PCGActor, const FBox& InIntersectedBounds, const TArray<FPCGTaskId>& TaskDependencies) { return InvalidPCGTaskId; };
+
+		// We can't spawn actors if we are running constructions scripts, asserting when we try to get the actor with the WP API.
+		// We should never enter this if we are in a construction script. If the ensure is hit, we need to fix it.
+		UWorld* World = GetWorld();
+		if (ensure(World && !World->bIsRunningConstructionScript))
+		{
+			PCGSubsystem::ForEachIntersectingCell(GraphExecutor, World, InBounds, /*bCreateActor=*/true, /*bLoadCell=*/false, /*bSave=*/false, ScheduleTask);
+		}
+	}
+}
+
 FPCGTaskId UPCGSubsystem::ProcessGraph(UPCGComponent* Component, const FBox& InPreviousBounds, const FBox& InNewBounds, EOperation InOperation, bool bSave)
 {
 	check(Component);
@@ -1410,7 +912,7 @@ void UPCGSubsystem::DirtyGraph(UPCGComponent* Component, const FBox& InBounds, E
 		return InvalidPCGTaskId;
 	};
 
-	DispatchToRegisteredLocalComponents(Component, DirtyTask);
+	ActorAndComponentMapping.DispatchToRegisteredLocalComponents(Component, DirtyTask);
 }
 
 void UPCGSubsystem::CleanupPartitionActors(const FBox& InBounds)
@@ -1569,8 +1071,8 @@ void UPCGSubsystem::DeletePartitionActors(bool bOnlyDeleteUnused)
 		// Need to copy to avoid deadlock
 		TMap<FIntVector, TObjectPtr<APCGPartitionActor>> PartitionActorsMapCopy;
 		{
-			FReadScopeLock ReadLock(PartitionActorsMapLock);
-			PartitionActorsMapCopy = PartitionActorsMap;
+			FReadScopeLock ReadLock(ActorAndComponentMapping.PartitionActorsMapLock);
+			PartitionActorsMapCopy = ActorAndComponentMapping.PartitionActorsMap;
 		}
 
 		for (const TPair<FIntVector, TObjectPtr<APCGPartitionActor>>& Item : PartitionActorsMapCopy)
@@ -1581,7 +1083,7 @@ void UPCGSubsystem::DeletePartitionActors(bool bOnlyDeleteUnused)
 
 				auto FoundComponents = [&bIntersectWithOneComponent](UPCGComponent*) { bIntersectWithOneComponent = true; };
 
-				ForAllIntersectingComponents(PartitionActor->GetFixedBounds(), FoundComponents);
+				ActorAndComponentMapping.ForAllIntersectingComponents(PartitionActor->GetFixedBounds(), FoundComponents);
 
 				if (!bIntersectWithOneComponent)
 				{
@@ -1676,13 +1178,6 @@ uint32 UPCGSubsystem::GetGraphCacheEntryCount(IPCGElement* InElement) const
 	return GraphExecutor ? GraphExecutor->GetGraphCacheEntryCount(InElement) : 0;
 }
 
-void UPCGSubsystem::ResetPartitionActorsMap()
-{
-	PartitionActorsMapLock.WriteLock();
-	PartitionActorsMap.Empty();
-	PartitionActorsMapLock.WriteUnlock();
-}
-
 #endif // WITH_EDITOR
 
 void UPCGSubsystem::FlushCache()
diff --git a/Engine/Plugins/Experimental/PCG/Source/PCG/Private/PCGWorldActor.cpp b/Engine/Plugins/Experimental/PCG/Source/PCG/Private/PCGWorldActor.cpp
--- a/Engine/Plugins/Experimental/PCG/Source/PCG/Private/PCGWorldActor.cpp
+++ b/Engine/Plugins/Experimental/PCG/Source/PCG/Private/PCGWorldActor.cpp
@@ -137,7 +137,7 @@ void APCGWorldActor::OnPartitionGridSizeChanged()
 
 	// And finally, refresh all components that are partitioned (registered to the PCGSubsystem)
 	// to let them recreate the needed PCG Partition Actors.
-	for (UPCGComponent* PCGComponent : PCGSubsystem->GetAllRegisteredComponents())
+	for (UPCGComponent* PCGComponent : PCGSubsystem->GetAllRegisteredPartitionedComponents())
 	{
 		check(PCGComponent);
 		PCGComponent->DirtyGenerated();
diff --git a/Engine/Plugins/Experimental/PCG/Source/PCG/Public/Grid/PCGComponentOctree.h b/Engine/Plugins/Experimental/PCG/Source/PCG/Public/Grid/PCGComponentOctree.h
--- a/Engine/Plugins/Experimental/PCG/Source/PCG/Public/Grid/PCGComponentOctree.h
+++ b/Engine/Plugins/Experimental/PCG/Source/PCG/Public/Grid/PCGComponentOctree.h
@@ -3,8 +3,10 @@
 #pragma once
 
 #include "Math/BoxSphereBounds.h"
-#include "Math/GenericOctreePublic.h"
+#include "Math/GenericOctree.h"
 #include "UObject/ObjectPtr.h"
+#include "Templates/SharedPointer.h"
+#include "Templates/SharedPointerFwd.h"
 
 template <typename ElementType, typename OctreeSemantics> class TOctree2;
 
@@ -15,7 +17,7 @@ struct PCG_API FPCGComponentOctreeID : public TSharedFromThis<FPCGComponentOctre
 	FOctreeElementId2 Id;
 };
 
-typedef TSharedRef<struct FPCGComponentOctreeID, ESPMode::ThreadSafe> FPCGComponentOctreeIDSharedRef;
+using FPCGComponentOctreeIDSharedRef = TSharedRef<struct FPCGComponentOctreeID, ESPMode::ThreadSafe>;
 
 struct PCG_API FPCGComponentRef
 {
@@ -57,7 +59,38 @@ struct PCG_API FPCGComponentRefSemantics
 	}
 };
 
-typedef TOctree2<FPCGComponentRef, FPCGComponentRefSemantics> FPCGComponentOctree;
+using FPCGComponentOctree = TOctree2<FPCGComponentRef, FPCGComponentRefSemantics> ;
+using FPCGComponentToIdMap = TMap<TObjectPtr<UPCGComponent>, FPCGComponentOctreeIDSharedRef>;
+
+class FPCGComponentOctreeAndMap
+{
+public:
+	FPCGComponentOctreeAndMap() = default;
+	FPCGComponentOctreeAndMap(const FVector& InOrigin, FVector::FReal InExtent);
+
+	void Reset(const FVector& InOrigin, FVector::FReal InExtent);
+
+	TSet<TObjectPtr<UPCGComponent>> GetAllComponents() const;
+
+	template<typename IterateBoundsFunc>
+	inline void FindElementsWithBoundsTest(const FBoxCenterAndExtent& BoxBounds, const IterateBoundsFunc& Func) const
+	{
+		FReadScopeLock ReadLock(Lock);
+		return Octree.FindElementsWithBoundsTest(BoxBounds, Func);
+	}
+
+	bool Contains(const UPCGComponent* InComponent) const;
+	FBox GetBounds(const UPCGComponent* InComponent) const;
+
+	void AddOrUpdateComponent(UPCGComponent* InComponent, FBox& OutBounds, bool& bOutComponentHasChanged, bool& bOutComponentWasAdded);
+	bool RemapComponent(const UPCGComponent* InOldComponent, UPCGComponent* InNewComponent, bool& bOutBoundsHasChanged);
+	bool RemoveComponent(UPCGComponent* InComponent);
+
+private:
+	FPCGComponentOctree Octree;
+	FPCGComponentToIdMap ComponentToIdMap;
+	mutable FRWLock Lock;
+};
 
 #if UE_ENABLE_INCLUDE_ORDER_DEPRECATED_IN_5_2
 #include "Math/GenericOctree.h"
diff --git a/Engine/Plugins/Experimental/PCG/Source/PCG/Public/PCGActorAndComponentMapping.h b/Engine/Plugins/Experimental/PCG/Source/PCG/Public/PCGActorAndComponentMapping.h
new file mode 100644
--- /dev/null
+++ b/Engine/Plugins/Experimental/PCG/Source/PCG/Public/PCGActorAndComponentMapping.h
@@ -0,0 +1,146 @@
+// Copyright Epic Games, Inc. All Rights Reserved.
+
+#pragma once
+
+#include "PCGCommon.h"
+#include "Grid/PCGComponentOctree.h"
+
+#include "Containers/Array.h"
+#include "Containers/Map.h"
+#include "Containers/Set.h"
+#include "Templates/Function.h"
+#include "UObject/ObjectPtr.h"
+#include "UObject/WeakObjectPtrTemplates.h"
+
+class AActor;
+class APCGPartitionActor;
+class UObject;
+class UPCGComponent;
+class UPCGSubsystem;
+
+/**
+* This class handle any necessary mapping between actors and pcg components.
+* Its meant to be part of the PCG Subsystem and owned by it. We offload some logic to this class to avoid to clutter the subsystem.
+* For now it is used for:
+* - Mapping between PCG Components and Partition actors: Dispatch tasks from original components to local components
+* - Tracking non partitioned PCG Components that has tracking actors needs.
+* - TODO: Tracking actors: Be able to react to actors changes, and potentially dirty/refresh affected components
+*/
+class UPCGActorAndComponentMapping
+{
+public:
+	friend UPCGSubsystem;
+
+	~UPCGActorAndComponentMapping() = default;
+	
+	/** Should be called by the subsystem to handle delayed operations. */
+	void Tick();
+
+#if WITH_EDITOR
+	/** If the partition grid size change, call this to empty the Partition actors map */
+	void ResetPartitionActorsMap();
+
+	void RegisterTrackingCallbacks();
+	void TeardownTrackingCallbacks();
+
+	void OnActorAdded(AActor* InActor);
+	void OnActorDeleted(AActor* InActor);
+	void OnActorMoved(AActor* InActor);
+	void OnObjectPropertyChanged(UObject* InObject, FPropertyChangedEvent& InEvent);
+#endif // WITH_EDITOR
+
+	/** Register a new PCG Component or update it. Returns true if it was added/updated. Thread safe */
+	bool RegisterOrUpdatePCGComponent(UPCGComponent* InComponent, bool bDoActorMapping = true);
+
+	/** In case of BP Actors, we need to remap the old component destroyed by the construction script to the new one. Returns true if a re-mapping was done. Thread safe */
+	bool RemapPCGComponent(const UPCGComponent* OldComponent, UPCGComponent* NewComponent, bool bDoActorMapping);
+
+	/** Unregister a PCG Component. Can force it, if we have a delayed unregister. Thread safe */
+	void UnregisterPCGComponent(UPCGComponent* InComponent, bool bForce = false);
+
+	/** Register a new Partition actor, will be added to a map and will query all intersecting volume to bind to them if asked. Thread safe */
+	void RegisterPartitionActor(APCGPartitionActor* InActor, bool bDoComponentMapping = true);
+
+	/** Unregister a Partition actor, will be removed from the map and remove itself to all intersecting volumes. Thread safe */
+	void UnregisterPartitionActor(APCGPartitionActor* InActor);
+
+	/** Return a copy of all the registered partitioned components. Thread safe */
+	TSet<TObjectPtr<UPCGComponent>> GetAllRegisteredPartitionedComponents() const;
+
+	/** Check if a component is currently tracked. Thread safe */
+	bool IsComponentTracked(const UPCGComponent* InComponent) const;
+
+private:
+	// This class is only meant to be used as part of the PCG Subsytem and owned by it.
+	// So we put constructors private.
+	// We also need this class to be default constructible, since the PCGSubsytem needs to be default constructible.
+	UPCGActorAndComponentMapping() = default;
+	explicit UPCGActorAndComponentMapping(UPCGSubsystem* PCGSubsystem);
+
+	bool RegisterOrUpdatePartitionedPCGComponent(UPCGComponent* InComponent, bool bDoActorMapping = true);
+	bool RegisterOrUpdateNonPartitionedPCGComponent(UPCGComponent* InComponent);
+
+	void UnregisterPartitionedPCGComponent(UPCGComponent* InComponent);
+	void UnregisterNonPartitionedPCGComponent(UPCGComponent* InComponent);
+
+	/* Call the InFunc function to all local component registered to the original component. Return the list of all the tasks scheduled. Thread safe*/
+	TArray<FPCGTaskId> DispatchToRegisteredLocalComponents(UPCGComponent* OriginalComponent, const TFunction<FPCGTaskId(UPCGComponent*)>& InFunc) const;
+
+	/* Call the InFunc function to all local component from the set of partition actors. Return the list of all the tasks scheduled. */
+	TArray<FPCGTaskId> DispatchToLocalComponents(UPCGComponent* OriginalComponent, const TSet<TObjectPtr<APCGPartitionActor>>& PartitionActors, const TFunction<FPCGTaskId(UPCGComponent*)>& InFunc) const;
+
+	/** Iterate other all the components which bounds intersect the box in param and call a callback. Thread safe */
+	void ForAllIntersectingComponents(const FBoxCenterAndExtent& InBounds, TFunction<void(UPCGComponent*)> InFunc) const;
+
+	/** Iterate other all the int coordinates given a box and call a callback. Thread safe */
+	void ForAllIntersectingPartitionActors(const FBox& InBounds, TFunction<void(APCGPartitionActor*)> InFunc) const;
+
+	/** Update the current mapping between a PCG component and its PCG Partition actors */
+	void UpdateMappingPCGComponentPartitionActor(UPCGComponent* InComponent);
+
+	/** Delete the current mapping between a PCG component and its PCG Partition actors */
+	void DeleteMappingPCGComponentPartitionActor(UPCGComponent* InComponent);
+
+	/** Will register/update tracking if a component was registered/updated. */
+	void RegisterOrUpdateTracking(UPCGComponent* InComponent);
+
+	/** Remap the tracking in case of BP components. */
+	void RemapTracking(const UPCGComponent* InOldComponent, UPCGComponent* InNewComponent);
+
+	/** Unregister tracking when a component is removed. */
+	void UnregisterTracking(UPCGComponent* InComponent);
+
+private:
+	// Cached subsystem
+	UPCGSubsystem* PCGSubsystem = nullptr;
+
+	// Mapping Component <-> PartitionActor
+	/** Octree tracking all partitioned pcg components */
+	FPCGComponentOctreeAndMap PartitionedOctree;
+
+	/** Grid mapping for partition actors. We can only have 1 partition actor per grid cell. */
+	TMap<FIntVector, TObjectPtr<APCGPartitionActor>> PartitionActorsMap;
+	mutable FRWLock PartitionActorsMapLock;
+
+	/** Mapping between original components and its overlapping partition actors. */
+	TMap<TObjectPtr<const UPCGComponent>, TSet<TObjectPtr<APCGPartitionActor>>> ComponentToPartitionActorsMap;
+	mutable FRWLock ComponentToPartitionActorsMapLock;
+
+	/** Components to be unregister at the next frame. cf. UnregisterComponent for a better understanding on why it is needed. */
+	TSet<TObjectPtr<UPCGComponent>> DelayedComponentToUnregister;
+	mutable FCriticalSection DelayedComponentToUnregisterLock;
+
+	// Tracking actors
+	/** Will hold all the components that are not partitioned (and not local) and are tracking something. Will be use to dispatch actor tracking updates. */
+	FPCGComponentOctreeAndMap NonPartitionedOctree;
+
+	/** Keep a mapping between tracked actors and the components that track them. */
+	TMap<TWeakObjectPtr<AActor>, TSet<TObjectPtr<UPCGComponent>>> TrackedActorsToComponentsMap;
+
+	/** Also keep the reverse mapping. */
+	TMap<TObjectPtr<UPCGComponent>, TSet<TWeakObjectPtr<AActor>>> TrackedComponentsToActorsMap;
+	mutable FRWLock TrackedComponentsLock;
+
+	/** Keep a mapping between tracked actors and their dependencies. */
+	TMap<TWeakObjectPtr<AActor>, TSet<TObjectPtr<UObject>>> TrackedActorsToDependenciesMap;
+};
\ No newline at end of file
diff --git a/Engine/Plugins/Experimental/PCG/Source/PCG/Public/PCGComponent.h b/Engine/Plugins/Experimental/PCG/Source/PCG/Public/PCGComponent.h
--- a/Engine/Plugins/Experimental/PCG/Source/PCG/Public/PCGComponent.h
+++ b/Engine/Plugins/Experimental/PCG/Source/PCG/Public/PCGComponent.h
@@ -67,8 +67,8 @@ class PCG_API UPCGComponent : public UActorComponent
 
 	GENERATED_BODY()
 
-	friend class UPCGSubsystem;
 	friend class UPCGManagedActors;
+	friend class UPCGSubsystem;
 
 public:
 	/** ~Begin UObject interface */
@@ -233,6 +233,9 @@ public:
 	void DisableInspection();
 	void StoreInspectionData(const UPCGNode* InNode, const FPCGDataCollection& InInspectionData);
 	const FPCGDataCollection* GetInspectionData(const UPCGNode* InNode) const;
+
+	/** Used by the tracking system to know if the component need to track actors. Not enabled for now, tracking is still done on the component.*/
+	bool ShouldTrackActors() const { return false; }
 #endif
 
 	/** Utility function (mostly for tests) to properly set the value of bIsComponentPartitioned.
@@ -253,6 +256,7 @@ public:
 	UPCGSubsystem* GetSubsystem() const;
 
 	FBox GetGridBounds() const;
+	FBox GetLastGeneratedBounds() const { return LastGeneratedBounds; }
 
 	/** Builds the PCG data from a given actor and its PCG component, and places it in a data collection with appropriate tags */
 	static FPCGDataCollection CreateActorPCGDataCollection(AActor* Actor, const UPCGComponent* Component, const TFunction<bool(EPCGDataType)>& InDataFilter, bool bParseActor = true);
diff --git a/Engine/Plugins/Experimental/PCG/Source/PCG/Public/PCGSubsystem.h b/Engine/Plugins/Experimental/PCG/Source/PCG/Public/PCGSubsystem.h
--- a/Engine/Plugins/Experimental/PCG/Source/PCG/Public/PCGSubsystem.h
+++ b/Engine/Plugins/Experimental/PCG/Source/PCG/Public/PCGSubsystem.h
@@ -2,10 +2,10 @@
 
 #pragma once
 
-#include "Math/GenericOctree.h"
 #include "Subsystems/WorldSubsystem.h"
 
 #include "PCGCommon.h"
+#include "PCGActorAndComponentMapping.h"
 #include "Grid/PCGComponentOctree.h"
 
 #include "PCGSubsystem.generated.h"
@@ -37,6 +37,8 @@ class PCG_API UPCGSubsystem : public UTickableWorldSubsystem
 	GENERATED_BODY()
 
 public:
+	friend UPCGActorAndComponentMapping;
+
 	UPCGSubsystem();
 
 	//~ Begin USubsystem Interface.
@@ -107,21 +109,21 @@ public:
 	bool GetOutputData(FPCGTaskId InTaskId, FPCGDataCollection& OutData);
 
 	/** Register a new PCG Component or update it, will be added to the octree if it doesn't exists yet. Returns true if it was added/updated. Thread safe */
-	bool RegisterOrUpdatePCGComponent(UPCGComponent* InComponent, bool bDoActorMapping = true);
+	bool RegisterOrUpdatePCGComponent(UPCGComponent* InComponent, bool bDoActorMapping = true) { return ActorAndComponentMapping.RegisterOrUpdatePCGComponent(InComponent, bDoActorMapping); }
 
 	/** In case of BP Actors, we need to remap the old component destroyed by the construction script to the new one. Returns true if re-mapping succeeded. */
-	bool RemapPCGComponent(const UPCGComponent* OldComponent, UPCGComponent* NewComponent);
+	bool RemapPCGComponent(const UPCGComponent* OldComponent, UPCGComponent* NewComponent, bool bDoActorMapping) { return ActorAndComponentMapping.RemapPCGComponent(OldComponent, NewComponent, bDoActorMapping); }
 
 	/** Unregister a PCG Component, will be removed from the octree. Can force it, if we have a delayed unregister. Thread safe */
-	void UnregisterPCGComponent(UPCGComponent* InComponent, bool bForce = false);
+	void UnregisterPCGComponent(UPCGComponent* InComponent, bool bForce = false) { ActorAndComponentMapping.UnregisterPCGComponent(InComponent, bForce); }
 
 	/** Register a new Partition actor, will be added to a map and will query all intersecting volume to bind to them if asked. Thread safe */
-	void RegisterPartitionActor(APCGPartitionActor* InActor, bool bDoComponentMapping = true);
+	void RegisterPartitionActor(APCGPartitionActor* InActor, bool bDoComponentMapping = true) { ActorAndComponentMapping.RegisterPartitionActor(InActor, bDoComponentMapping); }
 
 	/** Unregister a Partition actor, will be removed from the map and remove itself to all intersecting volumes. Thread safe */
-	void UnregisterPartitionActor(APCGPartitionActor* InActor);
+	void UnregisterPartitionActor(APCGPartitionActor* InActor) { ActorAndComponentMapping.UnregisterPartitionActor(InActor); }
 
-	TSet<TObjectPtr<UPCGComponent>> GetAllRegisteredComponents() const;
+	TSet<TObjectPtr<UPCGComponent>> GetAllRegisteredPartitionedComponents() const { return ActorAndComponentMapping.GetAllRegisteredPartitionedComponents(); }
 
 	/** Flushes the graph cache completely, use only for debugging */
 	void FlushCache();
@@ -157,7 +159,7 @@ public:
 	void ClearPCGLink(UPCGComponent* InComponent, const FBox& InBounds, AActor* InNewActor);
 
 	/** If the partition grid size change, call this to empty the Partition actors map */
-	void ResetPartitionActorsMap();
+	void ResetPartitionActorsMap() { ActorAndComponentMapping.ResetPartitionActorsMap(); }
 
 	/** Builds the landscape data cache */
 	void BuildLandscapeCache(bool bQuiet = false);
@@ -180,48 +182,18 @@ private:
 	};
 
 	FPCGTaskId ProcessGraph(UPCGComponent* Component, const FBox& InPreviousBounds, const FBox& InNewBounds, EOperation InOperation, bool bSave);
+	void CreatePartitionActorsWithinBounds(const FBox& InBounds);
 #endif // WITH_EDITOR
 	
-private:
-	/* Call the InFunc function to all local component registered to the original component. Return the list of all the tasks scheduled. Thread safe*/
-	TArray<FPCGTaskId> DispatchToRegisteredLocalComponents(UPCGComponent* OriginalComponent, const TFunction<FPCGTaskId(UPCGComponent*)>& InFunc) const;
-
-	/* Call the InFunc function to all local component from the set of partition actors. Return the list of all the tasks scheduled. */
-	TArray<FPCGTaskId> DispatchToLocalComponents(UPCGComponent* OriginalComponent, const TSet<TObjectPtr<APCGPartitionActor>>& PartitionActors, const TFunction<FPCGTaskId(UPCGComponent*)>& InFunc) const;
-
-	/** Iterate other all the components which bounds intersect the box in param and call a callback. Thread safe */
-	void ForAllIntersectingComponents(const FBoxCenterAndExtent& InBounds, TFunction<void(UPCGComponent*)> InFunc) const;
-
-	/** Iterate other all the int coordinates given a box and call a callback. Thread safe */
-	void ForAllIntersectingPartitionActors(const FBox& InBounds, TFunction<void(APCGPartitionActor*)> InFunc) const;
-
-	/** Update the current mapping between a PCG component and its PCG Partition actors */
-	void UpdateMappingPCGComponentPartitionActor(UPCGComponent* InComponent);
-
-	/** Delete the current mapping between a PCG component and its PCG Partition actors */
-	void DeleteMappingPCGComponentPartitionActor(UPCGComponent* InComponent);
-	
 private:
 	APCGWorldActor* PCGWorldActor = nullptr;
 	FPCGGraphExecutor* GraphExecutor = nullptr;
 	bool bHasTickedOnce = false;
+	UPCGActorAndComponentMapping ActorAndComponentMapping;
 
 #if WITH_EDITOR
 	FCriticalSection PCGWorldActorLock;
 #endif
-
-	FPCGComponentOctree PCGComponentOctree;
-	TMap<TObjectPtr<UPCGComponent>, FPCGComponentOctreeIDSharedRef> ComponentToIdMap;
-	mutable FRWLock PCGVolumeOctreeLock;
-
-	TMap<FIntVector, TObjectPtr<APCGPartitionActor>> PartitionActorsMap;
-	mutable FRWLock PartitionActorsMapLock;
-
-	TMap<TObjectPtr<const UPCGComponent>, TSet<TObjectPtr<APCGPartitionActor>>> ComponentToPartitionActorsMap;
-	mutable FRWLock ComponentToPartitionActorsMapLock;
-
-	TSet<TObjectPtr<UPCGComponent>> DelayedComponentToUnregister;
-	mutable FCriticalSection DelayedComponentToUnregisterLock;
 };
 
 #if UE_ENABLE_INCLUDE_ORDER_DEPRECATED_IN_5_2
