[PCG] Improved resize behavior for attribute columns, they will now default to column header text size (string attributes will be defaulted to a minimum of 200px).
Added cell textbox overflow elipsis and margins for better readability.

diff --git a/Engine/Plugins/Experimental/PCG/Source/PCGEditor/Private/SPCGEditorGraphAttributeListView.cpp b/Engine/Plugins/Experimental/PCG/Source/PCGEditor/Private/SPCGEditorGraphAttributeListView.cpp
--- a/Engine/Plugins/Experimental/PCG/Source/PCGEditor/Private/SPCGEditorGraphAttributeListView.cpp
+++ b/Engine/Plugins/Experimental/PCG/Source/PCGEditor/Private/SPCGEditorGraphAttributeListView.cpp
@@ -2,7 +2,6 @@
 
 #include "SPCGEditorGraphAttributeListView.h"
 
-#include "Framework/Views/TableViewMetadata.h"
 #include "PCGComponent.h"
 #include "PCGEditor.h"
 #include "PCGEditorGraphNodeBase.h"
@@ -11,6 +10,7 @@
 #include "Data/PCGPointData.h"
 #include "Data/PCGSpatialData.h"
 
+#include "Fonts/FontMeasure.h"
 #include "Framework/MultiBox/MultiBoxBuilder.h"
 #include "Widgets/Images/SLayeredImage.h"
 #include "Widgets/Input/SComboBox.h"
@@ -126,7 +126,10 @@ TSharedRef<SWidget> SPCGListViewItemRow::GenerateWidgetForColumn(const FName& Co
 		}
 	}
 
-	return SNew(STextBlock).Text(ColumnData);
+	return SNew(STextBlock)
+		.Text(ColumnData)
+		.OverflowPolicy(ETextOverflowPolicy::Ellipsis)
+		.Margin(FMargin(2.0f, 0.0f));
 }
 
 FText SPCGListViewItemRow::ConvertPointDataToText(const FPCGPoint* PCGPoint, const FName& ColumnId) const
@@ -843,15 +846,16 @@ void SPCGEditorGraphAttributeListView::GenerateColumnsFromMetadata(const UPCGMet
 	for (int32 I = 0; I < AttributeNames.Num(); I++)
 	{
 		const FName& AttributeName = AttributeNames[I];
+		const EPCGMetadataTypes AttributeType = AttributeTypes[I];
 		FName ColumnName = AttributeName;
 
 		if (ColumnName == NAME_None)
 		{
-			const FString TypeString = UEnum::GetDisplayValueAsText(AttributeTypes[I]).ToString();
+			const FString TypeString = UEnum::GetDisplayValueAsText(AttributeType).ToString();
 			ColumnName = *TypeString;
 		}
 
-		switch (AttributeTypes[I])
+		switch (AttributeType)
 		{
 		case EPCGMetadataTypes::Float:
 		case EPCGMetadataTypes::Double:
@@ -861,42 +865,42 @@ void SPCGEditorGraphAttributeListView::GenerateColumnsFromMetadata(const UPCGMet
 		case EPCGMetadataTypes::String:
 		case EPCGMetadataTypes::Name:
 			{
-				AddMetadataColumn(ColumnName, AttributeName);
+				AddMetadataColumn(ColumnName, AttributeName, AttributeType);
 				break;
 			}
 		case EPCGMetadataTypes::Vector2:
 		{
-			AddMetadataColumn(ColumnName, AttributeName, 0, TEXT("_X"));
-			AddMetadataColumn(ColumnName, AttributeName, 1, TEXT("_Y"));
+			AddMetadataColumn(ColumnName, AttributeName, AttributeType, 0, TEXT("_X"));
+			AddMetadataColumn(ColumnName, AttributeName, AttributeType, 1, TEXT("_Y"));
 			break;
 		}
 		case EPCGMetadataTypes::Vector:
 			{
-				AddMetadataColumn(ColumnName, AttributeName, 0, TEXT("_X"));
-				AddMetadataColumn(ColumnName, AttributeName, 1, TEXT("_Y"));
-				AddMetadataColumn(ColumnName, AttributeName, 2, TEXT("_Z"));
+				AddMetadataColumn(ColumnName, AttributeName, AttributeType, 0, TEXT("_X"));
+				AddMetadataColumn(ColumnName, AttributeName, AttributeType, 1, TEXT("_Y"));
+				AddMetadataColumn(ColumnName, AttributeName, AttributeType, 2, TEXT("_Z"));
 				break;
 			}
 		case EPCGMetadataTypes::Vector4:
 		case EPCGMetadataTypes::Quaternion:
 			{
-				AddMetadataColumn(ColumnName, AttributeName, 0, TEXT("_X"));
-				AddMetadataColumn(ColumnName, AttributeName, 1, TEXT("_Y"));
-				AddMetadataColumn(ColumnName, AttributeName, 2, TEXT("_Z"));
-				AddMetadataColumn(ColumnName, AttributeName, 3, TEXT("_W"));
+				AddMetadataColumn(ColumnName, AttributeName, AttributeType, 0, TEXT("_X"));
+				AddMetadataColumn(ColumnName, AttributeName, AttributeType, 1, TEXT("_Y"));
+				AddMetadataColumn(ColumnName, AttributeName, AttributeType, 2, TEXT("_Z"));
+				AddMetadataColumn(ColumnName, AttributeName, AttributeType, 3, TEXT("_W"));
 				break;
 			}
 		case EPCGMetadataTypes::Transform:
 			{
-				AddMetadataColumn(ColumnName, AttributeName, 0, TEXT("_tX"));
-				AddMetadataColumn(ColumnName, AttributeName, 1, TEXT("_tY"));
-				AddMetadataColumn(ColumnName, AttributeName, 2, TEXT("_tZ"));
-				AddMetadataColumn(ColumnName, AttributeName, 3, TEXT("_rX"));
-				AddMetadataColumn(ColumnName, AttributeName, 4, TEXT("_rY"));
-				AddMetadataColumn(ColumnName, AttributeName, 5, TEXT("_rZ"));
-				AddMetadataColumn(ColumnName, AttributeName, 6, TEXT("_sX"));
-				AddMetadataColumn(ColumnName, AttributeName, 7, TEXT("_sY"));
-				AddMetadataColumn(ColumnName, AttributeName, 8, TEXT("_sZ"));
+				AddMetadataColumn(ColumnName, AttributeName, AttributeType, 0, TEXT("_tX"));
+				AddMetadataColumn(ColumnName, AttributeName, AttributeType, 1, TEXT("_tY"));
+				AddMetadataColumn(ColumnName, AttributeName, AttributeType, 2, TEXT("_tZ"));
+				AddMetadataColumn(ColumnName, AttributeName, AttributeType, 3, TEXT("_rX"));
+				AddMetadataColumn(ColumnName, AttributeName, AttributeType, 4, TEXT("_rY"));
+				AddMetadataColumn(ColumnName, AttributeName, AttributeType, 5, TEXT("_rZ"));
+				AddMetadataColumn(ColumnName, AttributeName, AttributeType, 6, TEXT("_sX"));
+				AddMetadataColumn(ColumnName, AttributeName, AttributeType, 7, TEXT("_sY"));
+				AddMetadataColumn(ColumnName, AttributeName, AttributeType, 8, TEXT("_sZ"));
 				break;
 			}
 		default:
@@ -1068,7 +1072,7 @@ void SPCGEditorGraphAttributeListView::RemovePointDataColumns()
 	RemoveColumn(PCGEditorGraphAttributeListView::NAME_PointMetadataEntryParent);
 }
 
-void SPCGEditorGraphAttributeListView::AddMetadataColumn(const FName& InColumnId, const FName& InMetadataId, const int8 InValueIndex, const TCHAR* PostFix)
+void SPCGEditorGraphAttributeListView::AddMetadataColumn(const FName& InColumnId, const FName& InMetadataId, EPCGMetadataTypes InMetadataType, const int8 InValueIndex, const TCHAR* PostFix)
 {
 	FString ColumnIdString = InColumnId.ToString();
 
@@ -1084,12 +1088,27 @@ void SPCGEditorGraphAttributeListView::AddMetadataColumn(const FName& InColumnId
 	MetadataInfo.Index = InValueIndex;
 	MetadataInfos.Add(ColumnId, MetadataInfo);
 
+	check(FSlateApplication::Get().GetRenderer());
+	const TSharedRef<FSlateFontMeasure> FontMeasure = FSlateApplication::Get().GetRenderer()->GetFontMeasureService();
+	const FSlateFontInfo FontInfo = FAppStyle::GetFontStyle(TEXT("NormalText"));
+
+	const FText DefaultLabel = FText::FromName(ColumnId);
+	const float DefaultLabelWidth = FontMeasure->Measure(DefaultLabel, FontInfo).X;
+
+	constexpr float ColumnPadding = 12.0f; // Todo: Grab padding from header style
+	float ColumnWidth = DefaultLabelWidth + ColumnPadding;
+
+	if (InMetadataType == EPCGMetadataTypes::String)
+	{
+		ColumnWidth = FMath::Max(ColumnWidth, 200.0f);
+	}
+	
 	SHeaderRow::FColumn::FArguments ColumnArguments;
 	ColumnArguments.ColumnId(ColumnId);
-	ColumnArguments.DefaultLabel(FText::FromName(ColumnId));
+	ColumnArguments.DefaultLabel(DefaultLabel);
 	ColumnArguments.HAlignHeader(EHorizontalAlignment::HAlign_Center);
 	ColumnArguments.HAlignCell(EHorizontalAlignment::HAlign_Right);
-	ColumnArguments.FillWidth(1.0f);
+	ColumnArguments.ManualWidth(ColumnWidth);
 
 	SHeaderRow::FColumn* NewColumn = new SHeaderRow::FColumn(ColumnArguments);
 	NewColumn->bIsVisible = !HiddenAttributes.Contains(InColumnId);
diff --git a/Engine/Plugins/Experimental/PCG/Source/PCGEditor/Private/SPCGEditorGraphAttributeListView.h b/Engine/Plugins/Experimental/PCG/Source/PCGEditor/Private/SPCGEditorGraphAttributeListView.h
--- a/Engine/Plugins/Experimental/PCG/Source/PCGEditor/Private/SPCGEditorGraphAttributeListView.h
+++ b/Engine/Plugins/Experimental/PCG/Source/PCGEditor/Private/SPCGEditorGraphAttributeListView.h
@@ -16,6 +16,7 @@ class UPCGEditorGraphNodeBase;
 class UPCGMetadata;
 class UPCGParamData;
 struct FPCGPoint;
+enum class EPCGMetadataTypes : uint8;
 
 class SComboButton;
 class SHeaderRow;
@@ -102,7 +103,7 @@ private:
 	void AddPointDataColumns();
 	void RemovePointDataColumns();
 
-	void AddMetadataColumn(const FName& InColumnId, const FName& InMetadataId, const int8 InValueIndex = INDEX_NONE, const TCHAR* PostFix = nullptr);
+	void AddMetadataColumn(const FName& InColumnId, const FName& InMetadataId, EPCGMetadataTypes InMetadataType, const int8 InValueIndex = INDEX_NONE, const TCHAR* PostFix = nullptr);
 	void RemoveMetadataColumns();
 
 	/** Pointer back to the PCG editor that owns us */
