Removed useless metadata initializations in trivial cases where parent entry reuse was sufficient

#rnx

diff --git a/Engine/Plugins/Experimental/PCG/Source/PCG/Private/Data/PCGDifferenceData.cpp b/Engine/Plugins/Experimental/PCG/Source/PCG/Private/Data/PCGDifferenceData.cpp
--- a/Engine/Plugins/Experimental/PCG/Source/PCG/Private/Data/PCGDifferenceData.cpp
+++ b/Engine/Plugins/Experimental/PCG/Source/PCG/Private/Data/PCGDifferenceData.cpp
@@ -198,7 +198,6 @@ const UPCGPointData* UPCGDifferenceData::CreatePointData(FPCGContext* Context) c
 			const bool bBinaryDensity = (DensityFunction == EPCGDifferenceDensityFunction::Binary);
 
 			OutPoint = Point;
-			//UPCGMetadataAccessorHelpers::InitializeMetadata(OutPoint, Data->Metadata, Point);
 			OutPoint.Density = bBinaryDensity ? 0 : FMath::Max(0, Point.Density - PointFromDiff.Density);
 
 			if (TempDiffMetadata && OutPoint.Density > 0 && PointFromDiff.MetadataEntry != PCGInvalidEntryKey)
@@ -215,7 +214,6 @@ const UPCGPointData* UPCGDifferenceData::CreatePointData(FPCGContext* Context) c
 		else
 		{
 			OutPoint = Point;
-			//UPCGMetadataAccessorHelpers::InitializeMetadata(OutPoint, Data->Metadata, Point);
 			return true;
 		}
 	});
diff --git a/Engine/Plugins/Experimental/PCG/Source/PCG/Private/Data/PCGIntersectionData.cpp b/Engine/Plugins/Experimental/PCG/Source/PCG/Private/Data/PCGIntersectionData.cpp
--- a/Engine/Plugins/Experimental/PCG/Source/PCG/Private/Data/PCGIntersectionData.cpp
+++ b/Engine/Plugins/Experimental/PCG/Source/PCG/Private/Data/PCGIntersectionData.cpp
@@ -235,7 +235,6 @@ UPCGPointData* UPCGIntersectionData::CreateAndFilterPointData(FPCGContext* Conte
 		}
 
 		OutPoint = Point;
-		//UPCGMetadataAccessorHelpers::InitializeMetadata(OutPoint, Data->Metadata, Point);
 		OutPoint.Density = PCGIntersectionDataMaths::ComputeDensity(Point.Density, PointFromY.Density, DensityFunction);
 		OutPoint.Color = Point.Color * PointFromY.Color;
 
diff --git a/Engine/Plugins/Experimental/PCG/Source/PCG/Private/Data/PCGProjectionData.cpp b/Engine/Plugins/Experimental/PCG/Source/PCG/Private/Data/PCGProjectionData.cpp
--- a/Engine/Plugins/Experimental/PCG/Source/PCG/Private/Data/PCGProjectionData.cpp
+++ b/Engine/Plugins/Experimental/PCG/Source/PCG/Private/Data/PCGProjectionData.cpp
@@ -223,8 +223,6 @@ const UPCGPointData* UPCGProjectionData::CreatePointData(FPCGContext* Context) c
 
 		// Merge points into a single point
 		OutPoint = SourcePoint;
-		// Assign metadata key and copy values
-		UPCGMetadataAccessorHelpers::InitializeMetadata(OutPoint, PointData->Metadata, SourcePoint);
 
 		// Apply projection result. Some of the params are already used inside ProjectPoint, so we are just applying the remaining bits that ProjectPoint() did not have access to.
 		// TODO this would be cleaner if there was a ProjectPoint that took an FPCGPoint
diff --git a/Engine/Plugins/Experimental/PCG/Source/PCG/Private/Data/PCGUnionData.cpp b/Engine/Plugins/Experimental/PCG/Source/PCG/Private/Data/PCGUnionData.cpp
--- a/Engine/Plugins/Experimental/PCG/Source/PCG/Private/Data/PCGUnionData.cpp
+++ b/Engine/Plugins/Experimental/PCG/Source/PCG/Private/Data/PCGUnionData.cpp
@@ -351,11 +351,7 @@ void UPCGUnionData::CreateSequentialPointData(FPCGContext* Context, UPCGPointDat
 			check(PointData && PointData->Metadata);
 
 			OutPoint = Point;
-			if (PointData->Metadata->GetParent() == Data[DataIndex]->Metadata)
-			{
-				UPCGMetadataAccessorHelpers::InitializeMetadata(OutPoint, PointData->Metadata, Point);
-			}
-			else
+			if (PointData->Metadata->GetParent() != Data[DataIndex]->Metadata)
 			{
 				UPCGMetadataAccessorHelpers::InitializeMetadata(OutPoint, PointData->Metadata);
 				// Since we can't inherit from the parent point, we'll set the values directly here
