[PCG] Don't reset loaders when changing the package.

This seems to cause freezes in some situations in uobject linking (EndLoad would not terminate).

diff --git a/Engine/Plugins/Experimental/PCG/Source/PCG/Private/Elements/PCGSpawnActor.cpp b/Engine/Plugins/Experimental/PCG/Source/PCG/Private/Elements/PCGSpawnActor.cpp
--- a/Engine/Plugins/Experimental/PCG/Source/PCG/Private/Elements/PCGSpawnActor.cpp
+++ b/Engine/Plugins/Experimental/PCG/Source/PCG/Private/Elements/PCGSpawnActor.cpp
@@ -294,7 +294,7 @@ void UPCGSpawnActorSettings::RefreshTemplateActor()
 			Options.bNotifyObjectReplacement = true;
 			UEngine::CopyPropertiesForUnrelatedObjects(TemplateActor, NewTemplateActor, Options);
 
-			TemplateActor->Rename(nullptr, GetTransientPackage(), REN_DontCreateRedirectors);
+			TemplateActor->Rename(nullptr, GetTransientPackage(), REN_DontCreateRedirectors | REN_ForceNoResetLoaders);
 		}
 
 		TemplateActor = NewTemplateActor;
