Fixed outstanding bug with BP actors containing PCG components - some improtant properties would not be serialized properly

#rnx

diff --git a/Engine/Plugins/Experimental/PCG/Source/PCG/Private/PCGComponent.cpp b/Engine/Plugins/Experimental/PCG/Source/PCG/Private/PCGComponent.cpp
--- a/Engine/Plugins/Experimental/PCG/Source/PCG/Private/PCGComponent.cpp
+++ b/Engine/Plugins/Experimental/PCG/Source/PCG/Private/PCGComponent.cpp
@@ -903,6 +903,7 @@ void UPCGComponent::BeginPlay()
 	// Register itself to the PCGSubsystem, to map the component to all its corresponding PartitionActors if it is partition among other things.
 	if (GetSubsystem())
 	{
+		ensure(!bGenerated || LastGeneratedBounds.IsValid);
 		GetSubsystem()->RegisterOrUpdatePCGComponent(this);
 	}
 
@@ -1132,6 +1133,7 @@ void UPCGComponent::OnRegister()
 			// We won't be able to spawn any actors if we are currently running a construction script.
 			if (!World->bIsRunningConstructionScript)
 			{
+				ensure(!bGenerated || LastGeneratedBounds.IsValid);
 				GetSubsystem()->RegisterOrUpdatePCGComponent(this, bGenerated);
 			}
 		}
@@ -2768,6 +2770,28 @@ void FPCGComponentInstanceData::ApplyToComponent(UActorComponent* Component, con
 	{
 		UPCGComponent* PCGComponent = CastChecked<UPCGComponent>(Component);
 
+		// IMPORTANT NOTE:
+		// Any non-visible (i.e. UPROPERTY() with no specifiers) are NOT copied over when re-running the construction script
+		// This means that some properties need to be reapplied here manually unless we make them visible
+		if (SourceComponent)
+		{
+			// Critical: LastGeneratedBounds & GeneratedGraphOutput
+			PCGComponent->LastGeneratedBounds = SourceComponent->LastGeneratedBounds;
+
+			PCGComponent->GeneratedGraphOutput = SourceComponent->GeneratedGraphOutput;
+			// Re-outer any data moved here
+			for (FPCGTaggedData& TaggedData : PCGComponent->GeneratedGraphOutput.TaggedData)
+			{
+				if (TaggedData.Data)
+				{
+					const_cast<UPCGData*>(TaggedData.Data.Get())->Rename(nullptr, PCGComponent, REN_ForceNoResetLoaders | REN_DoNotDirty | REN_DontCreateRedirectors | REN_NonTransactional);
+				}
+			}
+
+			// Non-critical but should be done: transient data, excluded & tracked actors cache, landscape tracking
+			// TODO Validate usefulness + move accordingly
+		}
+		
 		// Duplicate generated resources + retarget them
 		TArray<TObjectPtr<UPCGManagedResource>> DuplicatedResources;
 		for (TObjectPtr<UPCGManagedResource> Resource : GeneratedResources)
