[PCG] Expose GetGraph to blueprint for PCG Graph Interface

diff --git a/Engine/Plugins/Experimental/PCG/Source/PCG/Public/PCGGraph.h b/Engine/Plugins/Experimental/PCG/Source/PCG/Public/PCGGraph.h
--- a/Engine/Plugins/Experimental/PCG/Source/PCG/Public/PCGGraph.h
+++ b/Engine/Plugins/Experimental/PCG/Source/PCG/Public/PCGGraph.h
@@ -19,6 +19,14 @@ class PCG_API UPCGGraphInterface : public UObject
 	GENERATED_BODY()
 
 public:
+	/** Return the underlying PCG Graph for this interface. */
+	UFUNCTION(BlueprintCallable, Category = Graph)
+	UPCGGraph* GetMutablePCGGraph() { return GetGraph(); }
+
+	/** Return the underlying PCG Graph for this interface. */
+	UFUNCTION(BlueprintCallable, Category = Graph)
+	const UPCGGraph* GetConstPCGGraph() const { return GetGraph(); }
+
 	virtual UPCGGraph* GetGraph() PURE_VIRTUAL(UPCGGraphInterface::GetGraph, return nullptr;)
 	virtual const UPCGGraph* GetGraph() const PURE_VIRTUAL(UPCGGraphInterface::GetGraph, return nullptr;)
 
