[PCG] Fixed a crash when redoing node connections after gc

diff --git a/Engine/Plugins/Experimental/PCG/Source/PCG/Private/PCGEdge.cpp b/Engine/Plugins/Experimental/PCG/Source/PCG/Private/PCGEdge.cpp
--- a/Engine/Plugins/Experimental/PCG/Source/PCG/Private/PCGEdge.cpp
+++ b/Engine/Plugins/Experimental/PCG/Source/PCG/Private/PCGEdge.cpp
@@ -5,6 +5,18 @@
 
 #include UE_INLINE_GENERATED_CPP_BY_NAME(PCGEdge)
 
+UPCGEdge::UPCGEdge(const FObjectInitializer& ObjectInitializer)
+	: Super(ObjectInitializer)
+{
+	SetFlags(RF_Transactional);
+}
+
+void UPCGEdge::PostLoad()
+{
+	Super::PostLoad();
+	SetFlags(RF_Transactional);
+}
+
 bool UPCGEdge::IsValid() const
 {
 	return InputPin.Get() && OutputPin.Get();
diff --git a/Engine/Plugins/Experimental/PCG/Source/PCG/Private/PCGPin.cpp b/Engine/Plugins/Experimental/PCG/Source/PCG/Private/PCGPin.cpp
--- a/Engine/Plugins/Experimental/PCG/Source/PCG/Private/PCGPin.cpp
+++ b/Engine/Plugins/Experimental/PCG/Source/PCG/Private/PCGPin.cpp
@@ -85,7 +85,8 @@ bool UPCGPin::AddEdgeTo(UPCGPin* OtherPin)
 
 	UPCGEdge* NewEdge = Edges.Add_GetRef(NewObject<UPCGEdge>(this));
 	OtherPin->Edges.Add(NewEdge);
-
+	
+	NewEdge->Modify();
 	NewEdge->InputPin = bThisPinIsUpstream ? this : OtherPin;
 	NewEdge->OutputPin = bThisPinIsUpstream ? OtherPin : this;
 
diff --git a/Engine/Plugins/Experimental/PCG/Source/PCG/Public/PCGEdge.h b/Engine/Plugins/Experimental/PCG/Source/PCG/Public/PCGEdge.h
--- a/Engine/Plugins/Experimental/PCG/Source/PCG/Public/PCGEdge.h
+++ b/Engine/Plugins/Experimental/PCG/Source/PCG/Public/PCGEdge.h
@@ -12,8 +12,13 @@ UCLASS(ClassGroup = (Procedural))
 class PCG_API UPCGEdge : public UObject
 {
 	GENERATED_BODY()
-
 public:
+	UPCGEdge(const FObjectInitializer& ObjectInitializer);
+
+	// ~Begin UObject interface
+	virtual void PostLoad() override;
+	// ~End UObject interface
+
 	UPROPERTY()
 	FName InboundLabel_DEPRECATED = NAME_None;
 
