[PCG] Fix PCG Crc being broken for FNames.

Use fast direct hash of FName, and remove helper which was not contributing much.

diff --git a/Engine/Plugins/Experimental/PCG/Source/PCG/Private/Graph/PCGGraphExecutor.cpp b/Engine/Plugins/Experimental/PCG/Source/PCG/Private/Graph/PCGGraphExecutor.cpp
--- a/Engine/Plugins/Experimental/PCG/Source/PCG/Private/Graph/PCGGraphExecutor.cpp
+++ b/Engine/Plugins/Experimental/PCG/Source/PCG/Private/Graph/PCGGraphExecutor.cpp
@@ -988,7 +988,7 @@ void FPCGGraphExecutor::BuildTaskInput(const FPCGGraphTask& Task, FPCGDataCollec
 			TaskInput.TaggedData.Append(InputCollection.GetInputsByPin(Input.InPin->Properties.Label));
 
 			// Write pin name Crc to uniquely identify inputs per-pin.
-			Crc.Combine(Input.InPin->Properties.Label);
+			Crc.Combine(GetTypeHash(Input.InPin->Properties.Label));
 		}
 		else
 		{
diff --git a/Engine/Plugins/Experimental/PCG/Source/PCG/Public/PCGCrc.h b/Engine/Plugins/Experimental/PCG/Source/PCG/Public/PCGCrc.h
--- a/Engine/Plugins/Experimental/PCG/Source/PCG/Public/PCGCrc.h
+++ b/Engine/Plugins/Experimental/PCG/Source/PCG/Public/PCGCrc.h
@@ -48,18 +48,6 @@ struct PCG_API FPCGCrc
 		Value = HashCombineFast(Value, InOtherCrcValue);
 	}
 
-	/** Computes Crc of given FName and combines it into this Crc. */
-	void Combine(FName InName)
-	{
-#if WITH_EDITOR
-		ensure(IsValid());
-#endif
-
-		uint32 Crc = 0;
-		FCrc::StrCrc32(*InName.ToString(), Crc);
-		Combine(Crc);
-	}
-
 	/** Compares Crc. This and other Crc must be valid. */
 	bool operator==(const FPCGCrc& InOtherCrc) const
 	{
