[PCG] Initial SplineSampler optimisations

Takes typical costs from 10s of ms to a few ms.

diff --git a/Engine/Plugins/Experimental/PCG/Source/PCG/Private/Elements/PCGSplineSampler.cpp b/Engine/Plugins/Experimental/PCG/Source/PCG/Private/Elements/PCGSplineSampler.cpp
--- a/Engine/Plugins/Experimental/PCG/Source/PCG/Private/Elements/PCGSplineSampler.cpp
+++ b/Engine/Plugins/Experimental/PCG/Source/PCG/Private/Elements/PCGSplineSampler.cpp
@@ -816,7 +816,7 @@ namespace PCGSplineSampler
 
 			{
 				TRACE_CPUPROFILER_EVENT_SCOPE(FPCGSplineSamplerElement::Execute::GetVoronoiEdges);
-			GetVoronoiEdges(PolygonPoints, FBox(MinPoint, FVector(MaxPoint.X, MaxPoint.Y, MinPoint.Z)), VoronoiEdges, CellMember);
+				GetVoronoiEdges(PolygonPoints, FBox(MinPoint, FVector(MaxPoint.X, MaxPoint.Y, MinPoint.Z)), VoronoiEdges, CellMember);
 			}
 
 			// Find the subset of the Voronoi Diagram which composes the Medial Axis
@@ -862,6 +862,8 @@ namespace PCGSplineSampler
 		TArray<TArray<TTuple<FTransform, FVector, float>>> InteriorSplinePointData;
 		InteriorSplinePointData.SetNum(NumDispatch);
 
+		const FTransform LineDataTransform = LineData->GetTransform();
+
 		ParallelFor(NumDispatch, [&](int32 DispatchIndex)
 		{
 			const bool bIsLastIteration = (DispatchIndex == (NumDispatch - 1));
@@ -873,119 +875,104 @@ namespace PCGSplineSampler
 
 			TArray<FVector::FReal> DistancesSquaredOnXY;
 
-		// Point sampling
+			// Point sampling
 			for(FVector::FReal Y = LocalMinY; Y < LocalMaxY; Y += Params.InteriorSampleSpacing)
-		{
-			const FVector2D RayMin(MinPoint.X - RayPadding, Y);
-			const FVector2D RayMax(MaxPoint.X + RayPadding, Y);
-
-			// Get the intersections along this ray, sorted by Y value
-			TArray<FVector2D> Intersections;
-			PCGSplineSamplerHelpers::SegmentPolygonIntersection2D(RayMin, RayMax, SplineSamplePoints2D, Intersections);
-
-			if (Intersections.Num() % 2 != 0)
 			{
-				UE_LOG(LogPCG, Error, TEXT("Intersection test failed, skipping samples for this row"));
-				continue;
-			}
-
-			Intersections.Sort([](const FVector2D& LHS, const FVector2D& RHS) { return LHS.X < RHS.X; });
+				const FVector2D RayMin(MinPoint.X - RayPadding, Y);
+				const FVector2D RayMax(MaxPoint.X + RayPadding, Y);
 
-			// TODO: async processing
-			// Each pair of intersections defines a range in which point samples may lie
-			for (int32 RangeIndex = 0; RangeIndex < Intersections.Num(); RangeIndex += 2)
-			{
-				const FVector::FReal MinX = FMath::CeilToDouble(Intersections[RangeIndex].X / Params.InteriorSampleSpacing) * Params.InteriorSampleSpacing;
-				const FVector::FReal MaxX = FMath::FloorToDouble(Intersections[RangeIndex + 1].X / Params.InteriorSampleSpacing) * Params.InteriorSampleSpacing;
+				// Get the intersections along this ray, sorted by Y value
+				TArray<FVector2D> Intersections;
+				PCGSplineSamplerHelpers::SegmentPolygonIntersection2D(RayMin, RayMax, SplineSamplePoints2D, Intersections);
 
-				if (MaxX - MinX < Params.InteriorSampleSpacing)
+				if (Intersections.Num() % 2 != 0)
 				{
+					UE_LOG(LogPCG, Error, TEXT("Intersection test failed, skipping samples for this row"));
 					continue;
 				}
 
-				for (FVector::FReal X = MinX; X < MaxX + UE_KINDA_SMALL_NUMBER; X += Params.InteriorSampleSpacing)
+				Intersections.Sort([](const FVector2D& LHS, const FVector2D& RHS) { return LHS.X < RHS.X; });
+
+				// TODO: async processing
+				// Each pair of intersections defines a range in which point samples may lie
+				for (int32 RangeIndex = 0; RangeIndex < Intersections.Num(); RangeIndex += 2)
 				{
-					const FVector2D SampleLocation = FVector2D(X, Y);
-					FVector SurfaceLocation = FVector(SampleLocation, MinPoint.Z);
+					const FVector::FReal MinX = FMath::CeilToDouble(Intersections[RangeIndex].X / Params.InteriorSampleSpacing) * Params.InteriorSampleSpacing;
+					const FVector::FReal MaxX = FMath::FloorToDouble(Intersections[RangeIndex + 1].X / Params.InteriorSampleSpacing) * Params.InteriorSampleSpacing;
 
-					if (bProjectOntoSurface)
+					if (MaxX - MinX < Params.InteriorSampleSpacing)
 					{
-						// Precompute 2D distance to every spline point
-						DistancesSquaredOnXY.Reset(SplineSamplePoints.Num());
-						for (const FVector& SplinePoint : SplineSamplePoints)
-						{
-							FVector::FReal DistanceSquared = FVector::DistSquaredXY(SurfaceLocation, SplinePoint);
-							DistancesSquaredOnXY.Add(DistanceSquared);
-						}
+						continue;
+					}
 
-						// Compute average Z value weighted by 1 / Distance^2
-						FVector::FReal SumZ = 0.f;
-						FVector::FReal SumWeights = 0.f;
-						for (int32 PointIndex = 0; PointIndex < SplineSamplePoints.Num(); ++PointIndex)
-						{
-							const FVector::FReal DistanceSquared = DistancesSquaredOnXY[PointIndex];
+					InteriorSplinePointData[DispatchIndex].Reserve(2 + FMath::Max(MaxX - MinX, 0) / Params.InteriorSampleSpacing);
+
+					for (FVector::FReal X = MinX; X < MaxX + UE_KINDA_SMALL_NUMBER; X += Params.InteriorSampleSpacing)
+					{
+						const FVector2D SampleLocation = FVector2D(X, Y);
+						FVector SurfaceLocation = FVector(SampleLocation, MinPoint.Z);
 
-							// If sample point overlaps exactly with a border point, then that must be the height
-							if (FMath::IsNearlyZero(DistanceSquared))
+						if (bProjectOntoSurface)
+						{
+							// Precompute 2D distance to every spline point
+							DistancesSquaredOnXY.Reset(SplineSamplePoints.Num());
+							for (const FVector& SplinePoint : SplineSamplePoints)
 							{
-								SumZ = SplineSamplePoints[PointIndex].Z;
-								SumWeights = 1.f;
-								break;
+								FVector::FReal DistanceSquared = FVector::DistSquaredXY(SurfaceLocation, SplinePoint);
+								DistancesSquaredOnXY.Add(DistanceSquared);
 							}
 
-							// TODO: it would be more accurate to use distance to the polyline instead of distance to the polyline points,
-							// however it would also be much more expensive. Perhaps worth investigating when Params.bTreatSplineAsPolyline is true
-							const FVector::FReal Weight = 1.f / DistanceSquared;
+							// Compute average Z value weighted by 1 / Distance^2
+							FVector::FReal SumZ = 0.f;
+							FVector::FReal SumWeights = 0.f;
+							for (int32 PointIndex = 0; PointIndex < SplineSamplePoints.Num(); ++PointIndex)
+							{
+								const FVector::FReal DistanceSquared = DistancesSquaredOnXY[PointIndex];
 
-							SumWeights += Weight;
-							SumZ += SplineSamplePoints[PointIndex].Z * Weight;
-						}
+								// If sample point overlaps exactly with a border point, then that must be the height
+								if (FMath::IsNearlyZero(DistanceSquared))
+								{
+									SumZ = SplineSamplePoints[PointIndex].Z;
+									SumWeights = 1.f;
+									break;
+								}
 
-						SurfaceLocation.Z = SumZ / SumWeights;
-					}
+								// TODO: it would be more accurate to use distance to the polyline instead of distance to the polyline points,
+								// however it would also be much more expensive. Perhaps worth investigating when Params.bTreatSplineAsPolyline is true
+								const FVector::FReal Weight = 1.f / DistanceSquared;
 
-					// if bTreatAsPolyline, then we shouldnt use this, we should use nearest point on the polygon line segments
-					float Dummy;
-					const float NearestSplineKey = bFindNearestSplineKey ? Spline->SplineCurves.Position.FindNearest(SurfaceLocation, Dummy) : 0.f;
+								SumWeights += Weight;
+								SumZ += SplineSamplePoints[PointIndex].Z * Weight;
+							}
 
-					const FVector PointLocationLS = (Params.bProjectOntoSurface ? FVector(SurfaceLocation) : FVector(SampleLocation, MinPoint.Z));
+							SurfaceLocation.Z = SumZ / SumWeights;
+						}
 
-					FTransform TransformLS = FTransform::Identity;
-					TransformLS.SetLocation(PointLocationLS);
+						// if bTreatAsPolyline, then we shouldnt use this, we should use nearest point on the polygon line segments
+						float Dummy;
+						const float NearestSplineKey = bFindNearestSplineKey ? Spline->SplineCurves.Position.FindNearest(SurfaceLocation, Dummy) : 0.f;
 
-					if (Params.InteriorOrientation == EPCGSplineSamplingInteriorOrientation::FollowCurvature)
-					{
-						TransformLS.SetRotation(Spline->GetQuaternionAtSplineInputKey(NearestSplineKey, ESplineCoordinateSpace::Local));
-					}
+						const FVector PointLocationLS = (Params.bProjectOntoSurface ? FVector(SurfaceLocation) : FVector(SampleLocation, MinPoint.Z));
 
-					// Calculate density fall off
-						float Density = 1.0f;
+						FTransform TransformLS = FTransform::Identity;
+						TransformLS.SetLocation(PointLocationLS);
 
-					if (bGenerateMedialAxis && MedialAxisEdges.Num())
-					{
-						FVector::FReal SmallestDistSquared = MaxDimensionSquared;
-
-						// Find distance from SampleLocation to MedialAxis
-						for (const TTuple<FVector2D, FVector2D>& Edge : MedialAxisEdges)
+						if (Params.InteriorOrientation == EPCGSplineSamplingInteriorOrientation::FollowCurvature)
 						{
-							const FVector::FReal DistSquared = FMath::PointDistToSegmentSquared(FVector(SampleLocation, 0), FVector(Edge.Get<0>(), 0), FVector(Edge.Get<1>(), 0));
-
-							if (DistSquared < SmallestDistSquared)
-							{
-								SmallestDistSquared = DistSquared;
-							}
+							TransformLS.SetRotation(Spline->GetQuaternionAtSplineInputKey(NearestSplineKey, ESplineCoordinateSpace::Local));
 						}
 
-						const FVector::FReal SmallestDist = FMath::Sqrt(SmallestDistSquared);
-
-						FVector::FReal PointToSplineDist = 0.f;
+						// Calculate density fall off
+						float Density = 1.0f;
 
-						if (Params.bTreatSplineAsPolyline)
+						if (bGenerateMedialAxis && MedialAxisEdges.Num())
 						{
-							SmallestDistSquared = MaxDimensionSquared;
-							for (int32 PointIndex = 0; PointIndex < SplineSamplePoints.Num(); ++PointIndex)
+							FVector::FReal SmallestDistSquared = MaxDimensionSquared;
+
+							// Find distance from SampleLocation to MedialAxis
+							for (const TTuple<FVector2D, FVector2D>& Edge : MedialAxisEdges)
 							{
-								const FVector::FReal DistSquared = FMath::PointDistToSegmentSquared(SurfaceLocation, SplineSamplePoints[PointIndex], SplineSamplePoints[(PointIndex + 1) % SplineSamplePoints.Num()]);
+								const FVector::FReal DistSquared = FMath::PointDistToSegmentSquared(FVector(SampleLocation, 0), FVector(Edge.Get<0>(), 0), FVector(Edge.Get<1>(), 0));
 
 								if (DistSquared < SmallestDistSquared)
 								{
@@ -993,31 +980,48 @@ namespace PCGSplineSampler
 								}
 							}
 
-							PointToSplineDist = FMath::Sqrt(SmallestDistSquared);
-						}
-						else
-						{
+							const FVector::FReal SmallestDist = FMath::Sqrt(SmallestDistSquared);
+
+							FVector::FReal PointToSplineDist = 0.f;
+
+							if (Params.bTreatSplineAsPolyline)
+							{
+								SmallestDistSquared = MaxDimensionSquared;
+								for (int32 PointIndex = 0; PointIndex < SplineSamplePoints.Num(); ++PointIndex)
+								{
+									const FVector::FReal DistSquared = FMath::PointDistToSegmentSquared(SurfaceLocation, SplineSamplePoints[PointIndex], SplineSamplePoints[(PointIndex + 1) % SplineSamplePoints.Num()]);
+
+									if (DistSquared < SmallestDistSquared)
+									{
+										SmallestDistSquared = DistSquared;
+									}
+								}
+
+								PointToSplineDist = FMath::Sqrt(SmallestDistSquared);
+							}
+							else
+							{
 								const FVector NearestSplineLocation = Spline->GetLocationAtSplineInputKey(NearestSplineKey, ESplineCoordinateSpace::Local);
-							PointToSplineDist = FVector2D::Distance(SampleLocation, FVector2D(NearestSplineLocation.X, NearestSplineLocation.Y));
-						}
+								PointToSplineDist = FVector2D::Distance(SampleLocation, FVector2D(NearestSplineLocation.X, NearestSplineLocation.Y));
+							}
 
-						// Linear fall off in the range [0, 1]
-						const FVector::FReal T = SmallestDist / (SmallestDist + PointToSplineDist + UE_KINDA_SMALL_NUMBER);
+							// Linear fall off in the range [0, 1]
+							const FVector::FReal T = SmallestDist / (SmallestDist + PointToSplineDist + UE_KINDA_SMALL_NUMBER);
 							Density = DensityFalloffCurve->Eval(T);
 						}
 
-						InteriorSplinePointData[DispatchIndex].Emplace(TransformLS * LineData->GetTransform(), PointLocationLS, Density);
-					}
+						InteriorSplinePointData[DispatchIndex].Emplace(TransformLS * LineDataTransform, PointLocationLS, Density);
 					}
+				}
 			}
 		});
 
 		// Finally, gather the data and push to the points
 		int32 PointCount = 0;
 		for (const TArray<TTuple<FTransform, FVector, float>>& InteriorData : InteriorSplinePointData)
-					{
+		{
 			PointCount += InteriorData.Num();
-					}
+		}
 
 		// TODO: should we parallel for this too?
 		OutPoints.Reserve(PointCount);
@@ -1029,8 +1033,8 @@ namespace PCGSplineSampler
 				Point.Transform = InteriorPoint.Get<0>();
 				SetSeed(Point, InteriorPoint.Get<1>(), Params);
 				Point.Density = InteriorPoint.Get<2>();
-					Point.BoundsMin = BoundsMin;
-					Point.BoundsMax = BoundsMax;
+				Point.BoundsMin = BoundsMin;
+				Point.BoundsMax = BoundsMax;
 			}
 		}
 	}
