[PCG] Sort out most expensive cases of FObjectPtr overheads

Results in 45% saving of surface sampler cost in my test.

diff --git a/Engine/Plugins/Experimental/PCG/Source/PCG/Private/Data/PCGIntersectionData.cpp b/Engine/Plugins/Experimental/PCG/Source/PCG/Private/Data/PCGIntersectionData.cpp
--- a/Engine/Plugins/Experimental/PCG/Source/PCG/Private/Data/PCGIntersectionData.cpp
+++ b/Engine/Plugins/Experimental/PCG/Source/PCG/Private/Data/PCGIntersectionData.cpp
@@ -30,54 +30,69 @@ void UPCGIntersectionData::Initialize(const UPCGSpatialData* InA, const UPCGSpat
 	B = InB;
 	TargetActor = A->TargetActor;
 
-	CachedBounds = PCGHelpers::OverlapBounds(A->GetBounds(), B->GetBounds());
-	CachedStrictBounds = PCGHelpers::OverlapBounds(A->GetStrictBounds(), B->GetStrictBounds());
+#if WITH_EDITOR
+	RawPointerA = A;
+	RawPointerB = B;
+#endif
+
+	CachedBounds = PCGHelpers::OverlapBounds(GetA()->GetBounds(), GetB()->GetBounds());
+	CachedStrictBounds = PCGHelpers::OverlapBounds(GetA()->GetStrictBounds(), GetB()->GetStrictBounds());
 
 	check(Metadata);
 	// Note: this should behave the same way as the ToPointData
-	if (A->GetDimension() <= B->GetDimension())
+	if (GetA()->GetDimension() <= GetB()->GetDimension())
 	{
-		Metadata->Initialize(A->Metadata);
-		Metadata->AddAttributes(B->Metadata);
+		Metadata->Initialize(GetA()->Metadata);
+		Metadata->AddAttributes(GetB()->Metadata);
 	}
 	else
 	{
-		Metadata->Initialize(B->Metadata);
-		Metadata->AddAttributes(A->Metadata);
+		Metadata->Initialize(GetB()->Metadata);
+		Metadata->AddAttributes(GetA()->Metadata);
 	}
 }
 
+#if WITH_EDITOR
+void UPCGIntersectionData::PostLoad()
+{
+	Super::PostLoad();
+
+	RawPointerA = A;
+	RawPointerB = B;
+}
+#endif
+
 void UPCGIntersectionData::VisitDataNetwork(TFunctionRef<void(const UPCGData*)> Action) const
 {
-	check(A && B);
-	A->VisitDataNetwork(Action);
-	B->VisitDataNetwork(Action);
+	check(GetA() && GetB());
+	GetA()->VisitDataNetwork(Action);
+	GetB()->VisitDataNetwork(Action);
 }
 
 int UPCGIntersectionData::GetDimension() const
 {
-	check(A && B);
-	return FMath::Min(A->GetDimension(), B->GetDimension());
+	check(GetA() && GetB());
+	return FMath::Min(GetA()->GetDimension(), GetB()->GetDimension());
 }
 
 FBox UPCGIntersectionData::GetBounds() const
 {
-	check(A && B);
+	check(GetA() && GetB());
 	return CachedBounds;
 }
 
 FBox UPCGIntersectionData::GetStrictBounds() const
 {
-	check(A && B);
+	check(GetA() && GetB());
 	return CachedStrictBounds;
 }
 
 bool UPCGIntersectionData::SamplePoint(const FTransform& InTransform, const FBox& InBounds, FPCGPoint& OutPoint, UPCGMetadata* OutMetadata) const
 {
 	//TRACE_CPUPROFILER_EVENT_SCOPE(UPCGIntersectionData::SamplePoint);
-	check(A && B);
-	const UPCGSpatialData* X = (A->HasNonTrivialTransform() || !B->HasNonTrivialTransform()) ? A : B;
-	const UPCGSpatialData* Y = (X == A) ? B : A;
+	check(GetA() && GetB());
+	const UPCGSpatialData* X = (GetA()->HasNonTrivialTransform() || !GetB()->HasNonTrivialTransform()) ? GetA() : GetB();
+	const UPCGSpatialData* Y = (X == GetA()) ? GetB() : GetA();
 
 	FPCGPoint PointFromX;
 	if(!X->SamplePoint(InTransform, InBounds, PointFromX, OutMetadata))
@@ -117,23 +132,23 @@ bool UPCGIntersectionData::SamplePoint(const FTransform& InTransform, const FBox
 
 bool UPCGIntersectionData::HasNonTrivialTransform() const
 {
-	check(A && B);
-	return A->HasNonTrivialTransform() || B->HasNonTrivialTransform();
+	check(GetA() && GetB());
+	return GetA()->HasNonTrivialTransform() || GetB()->HasNonTrivialTransform();
 }
 
 const UPCGSpatialData* UPCGIntersectionData::FindShapeFromNetwork(const int InDimension) const
 {
-	check(A && B);
+	check(GetA() && GetB());
 
 	if (InDimension >= 0)
 	{
 		// Return first candidate that matches Dimension
-		if (const UPCGSpatialData* CandidateA = A->FindShapeFromNetwork(InDimension))
+		if (const UPCGSpatialData* CandidateA = GetA()->FindShapeFromNetwork(InDimension))
 		{
 			return CandidateA;
 		}
 
-		if (const UPCGSpatialData* CandidateB = B->FindShapeFromNetwork(InDimension))
+		if (const UPCGSpatialData* CandidateB = GetB()->FindShapeFromNetwork(InDimension))
 		{
 			return CandidateB;
 		}
@@ -146,13 +161,13 @@ const UPCGSpatialData* UPCGIntersectionData::FindShapeFromNetwork(const int InDi
 		int LowestDimension = MAX_uint32;
 		const UPCGSpatialData* LowestDimensionalShape = nullptr;
 
-		if (const UPCGSpatialData* CandidateA = A->FindShapeFromNetwork(InDimension))
+		if (const UPCGSpatialData* CandidateA = GetA()->FindShapeFromNetwork(InDimension))
 		{
 			Result = CandidateA;
 			LowestDimension = CandidateA->GetDimension();
 		}
 
-		const UPCGSpatialData* CandidateB = B->FindShapeFromNetwork(InDimension);
+		const UPCGSpatialData* CandidateB = GetB()->FindShapeFromNetwork(InDimension);
 		if (CandidateB && CandidateB->GetDimension() < LowestDimension)
 		{
 			Result = CandidateB;
@@ -164,29 +179,29 @@ const UPCGSpatialData* UPCGIntersectionData::FindShapeFromNetwork(const int InDi
 
 const UPCGSpatialData* UPCGIntersectionData::FindFirstConcreteShapeFromNetwork() const
 {
-	check(A && B);
+	check(GetA() && GetB());
 
-	if (const UPCGSpatialData* CandidateA = A->FindFirstConcreteShapeFromNetwork())
+	if (const UPCGSpatialData* CandidateA = GetA()->FindFirstConcreteShapeFromNetwork())
 	{
 		return CandidateA;
 	}
 
-	return B->FindFirstConcreteShapeFromNetwork();
+	return GetB()->FindFirstConcreteShapeFromNetwork();
 }
 
 const UPCGPointData* UPCGIntersectionData::CreatePointData(FPCGContext* Context) const
 {
-	check(A && B);
+	check(GetA() && GetB());
 	// TODO: this is a placeholder;
 	// Here we will get the point data from the lower-dimensionality data
 	// and then cull out any of the points that are outside the bounds of the other
-	if (A->GetDimension() <= B->GetDimension())
+	if (GetA()->GetDimension() <= GetB()->GetDimension())
 	{
-		return CreateAndFilterPointData(Context, A, B);
+		return CreateAndFilterPointData(Context, GetA(), GetB());
 	}
 	else
 	{
-		return CreateAndFilterPointData(Context, B, A);
+		return CreateAndFilterPointData(Context, GetB(), GetA());
 	}
 }
 
@@ -268,5 +283,10 @@ UPCGSpatialData* UPCGIntersectionData::CopyInternal() const
 	NewIntersectionData->CachedBounds = CachedBounds;
 	NewIntersectionData->CachedStrictBounds = CachedStrictBounds;
 
+#if WITH_EDITOR
+	NewIntersectionData->RawPointerA = RawPointerA;
+	NewIntersectionData->RawPointerB = RawPointerB;
+#endif
+
 	return NewIntersectionData;
 }
diff --git a/Engine/Plugins/Experimental/PCG/Source/PCG/Private/Elements/PCGSurfaceSampler.cpp b/Engine/Plugins/Experimental/PCG/Source/PCG/Private/Elements/PCGSurfaceSampler.cpp
--- a/Engine/Plugins/Experimental/PCG/Source/PCG/Private/Elements/PCGSurfaceSampler.cpp
+++ b/Engine/Plugins/Experimental/PCG/Source/PCG/Private/Elements/PCGSurfaceSampler.cpp
@@ -154,7 +154,11 @@ namespace PCGSurfaceSampler
 		// Make sure we're still in bounds though!
 		SampleZ = FMath::Max(SampleZ, LoopData.InputBoundsMinZ);
 
-		FPCGAsync::AsyncPointProcessing(Context, LoopData.CellCount, SampledPoints, [&LoopData, SampledData, InBoundingShape, InSurface, &ProjectionParams, SampleZ](int32 Index, FPCGPoint& OutPoint)
+		// Cache pointer ahead of time to avoid dereferencing object pointer which does access tracking and supports lazy loading, and can come with substantial
+		// overhead (add trace marker to FObjectPtr::Get to see).
+		UPCGMetadata* OutMetadata = SampledData->Metadata.Get();
+
+		FPCGAsync::AsyncPointProcessing(Context, LoopData.CellCount, SampledPoints, [&LoopData, SampledData, InBoundingShape, InSurface, &ProjectionParams, SampleZ, OutMetadata](int32 Index, FPCGPoint& OutPoint)
 		{
 			const FIntVector2 Indices = LoopData.ComputeCellIndices(Index);
 
@@ -182,7 +186,7 @@ namespace PCGSurfaceSampler
 			OutPoint = FPCGPoint();
 
 			// Firstly project onto elected generating shape to move to final position.
-			if (!InSurface->ProjectPoint(FTransform(TentativeLocation), LocalBound, ProjectionParams, OutPoint, SampledData->Metadata))
+			if (!InSurface->ProjectPoint(FTransform(TentativeLocation), LocalBound, ProjectionParams, OutPoint, OutMetadata))
 			{
 				return false;
 			}
diff --git a/Engine/Plugins/Experimental/PCG/Source/PCG/Public/Data/PCGIntersectionData.h b/Engine/Plugins/Experimental/PCG/Source/PCG/Public/Data/PCGIntersectionData.h
--- a/Engine/Plugins/Experimental/PCG/Source/PCG/Public/Data/PCGIntersectionData.h
+++ b/Engine/Plugins/Experimental/PCG/Source/PCG/Public/Data/PCGIntersectionData.h
@@ -24,6 +24,12 @@ public:
 	UFUNCTION(BlueprintCallable, Category = SpatialData)
 	void Initialize(const UPCGSpatialData* InA, const UPCGSpatialData* InB);
 
+	// ~Begin UObject interface
+#if WITH_EDITOR
+	virtual void PostLoad();
+#endif
+	// ~End UObject interface
+
 	// ~Begin UPCGData interface
 	virtual EPCGDataType GetDataType() const override { return EPCGDataType::Spatial; }
 	virtual void VisitDataNetwork(TFunctionRef<void(const UPCGData*)> Action) const override;
@@ -65,4 +71,20 @@ protected:
 
 	UPROPERTY()
 	FBox CachedStrictBounds = FBox(EForceInit::ForceInit);
+
+#if WITH_EDITOR
+	inline const UPCGSpatialData* GetA() const { return RawPointerA; }
+	inline const UPCGSpatialData* GetB() const { return RawPointerB; }
+#else
+	inline const UPCGSpatialData* GetA() const { return A.Get(); }
+	inline const UPCGSpatialData* GetB() const { return B.Get(); }
+#endif
+
+#if WITH_EDITOR
+private:
+	// Cached pointers to avoid dereferencing object pointer which does access tracking and supports lazy loading, and can come with substantial
+	// overhead (add trace marker to FObjectPtr::Get to see).
+	const UPCGSpatialData* RawPointerA = nullptr;
+	const UPCGSpatialData* RawPointerB = nullptr;
+#endif
 };
