Replaced previous static mesh spawner data by a full ISM descriptor
Replaced full point storage in the mesh instance lists by transform only, saves some transformations
Optimized soft object pointer access to be faster than previously

diff --git a/Engine/Plugins/Experimental/PCG/Source/PCG/Private/Elements/PCGDebugElement.cpp b/Engine/Plugins/Experimental/PCG/Source/PCG/Private/Elements/PCGDebugElement.cpp
--- a/Engine/Plugins/Experimental/PCG/Source/PCG/Private/Elements/PCGDebugElement.cpp
+++ b/Engine/Plugins/Experimental/PCG/Source/PCG/Private/Elements/PCGDebugElement.cpp
@@ -2,18 +2,19 @@
 
 #include "Elements/PCGDebugElement.h"
 
-#include "Data/PCGSpatialData.h"
 #include "PCGComponent.h"
-#include "Engine/CollisionProfile.h"
+#include "PCGContext.h"
 #include "PCGHelpers.h"
+#include "PCGPin.h"
+#include "Data/PCGSpatialData.h"
 #include "Data/PCGPointData.h"
 #include "Helpers/PCGActorHelpers.h"
 
 #include "Components/InstancedStaticMeshComponent.h"
+#include "Engine/CollisionProfile.h"
 #include "Engine/StaticMesh.h"
+#include "ISMPartition/ISMComponentDescriptor.h"
 #include "Materials/MaterialInterface.h"
-#include "PCGContext.h"
-#include "PCGPin.h"
 
 #include UE_INLINE_GENERATED_CPP_BY_NAME(PCGDebugElement)
 
@@ -150,19 +151,23 @@ namespace PCGDebugElement
 				}
 			}
 
-			FPCGISMCBuilderParameters Params;
-			Params.Mesh = Mesh;
-			Params.MaterialOverrides = Materials;
-			Params.CollisionProfile = UCollisionProfile::NoCollision_ProfileName;
+			FPCGISMCBuilderParameters Params[2];
+			Params[0].NumCustomDataFloats = NumCustomData;
+			Params[0].Descriptor.StaticMesh = Mesh;
+			Params[0].Descriptor.OverrideMaterials = Materials;
+			Params[0].Descriptor.BodyInstance.SetCollisionProfileName(UCollisionProfile::NoCollision_ProfileName);
 			// Note: In the future we may consider enabling culling for performance reasons, but for now culling disabled.
-			Params.CullStartDistance = Params.CullEndDistance = 0;
-
+			Params[0].Descriptor.InstanceStartCullDistance = Params[0].Descriptor.InstanceEndCullDistance = 0;
+			
 			// If the root actor we're binding to is movable, then the ISMC should be movable by default
 			if (USceneComponent* SceneComponent = TargetActor->GetRootComponent())
 			{
-				Params.Mobility = SceneComponent->Mobility;
+				Params[0].Descriptor.Mobility = SceneComponent->Mobility;
 			}
 
+			Params[1] = Params[0];
+			Params[1].Descriptor.bReverseCulling = true;
+
 			for (int32 Direction = 0; Direction < 2; ++Direction)
 			{
 				TArray<FTransform>& Instances = (Direction == 0 ? ForwardInstances : ReverseInstances);
@@ -172,12 +177,10 @@ namespace PCGDebugElement
 					continue;
 				}
 
-				Params.bIsLocalToWorldDeterminantNegative = (Direction != 0);
-				UInstancedStaticMeshComponent* ISMC = UPCGActorHelpers::GetOrCreateISMC(TargetActor, Context->SourceComponent.Get(), Params);
-				check(ISMC);
+				UInstancedStaticMeshComponent* ISMC = UPCGActorHelpers::GetOrCreateISMC(TargetActor, Context->SourceComponent.Get(), Params[Direction]);
+				check(ISMC && ISMC->NumCustomDataFloats == NumCustomData);
 
 				ISMC->ComponentTags.AddUnique(PCGHelpers::DefaultPCGDebugTag);
-				ISMC->NumCustomDataFloats = NumCustomData;
 				const int32 PreExistingInstanceCount = ISMC->GetInstanceCount();
 				ISMC->AddInstances(Instances, /*bShouldReturnIndices=*/false, /*bWorldSpace=*/true);
 
diff --git a/Engine/Plugins/Experimental/PCG/Source/PCG/Private/Elements/PCGSpawnActor.cpp b/Engine/Plugins/Experimental/PCG/Source/PCG/Private/Elements/PCGSpawnActor.cpp
--- a/Engine/Plugins/Experimental/PCG/Source/PCG/Private/Elements/PCGSpawnActor.cpp
+++ b/Engine/Plugins/Experimental/PCG/Source/PCG/Private/Elements/PCGSpawnActor.cpp
@@ -441,17 +441,17 @@ bool FPCGSpawnActorElement::ExecuteInternal(FPCGContext* Context) const
 			TArray<UActorComponent*> Components;
 			UPCGActorHelpers::GetActorClassDefaultComponents(Settings->TemplateActorClass, Components, UStaticMeshComponent::StaticClass());
 			
-			TMap<FISMComponentDescriptor, TArray<FTransform>> MeshDescriptorTransforms;
+			TMap<FPCGISMCBuilderParameters, TArray<FTransform>> MeshDescriptorTransforms;
 
 			for (UActorComponent* Component : Components)
 			{
 				if (UStaticMeshComponent* StaticMeshComponent = Cast<UStaticMeshComponent>(Component))
 				{
-					FISMComponentDescriptor ISMDescriptor;
-					ISMDescriptor.InitFrom(StaticMeshComponent);
-					ISMDescriptor.ComputeHash();
+					FPCGISMCBuilderParameters Params;
+					Params.Descriptor.InitFrom(StaticMeshComponent);
+					// TODO: No custom data float support?
 
-					TArray<FTransform>& Transforms = MeshDescriptorTransforms.FindOrAdd(ISMDescriptor);
+					TArray<FTransform>& Transforms = MeshDescriptorTransforms.FindOrAdd(Params);
 
 					if (UInstancedStaticMeshComponent* InstancedStaticMeshComponent = Cast<UInstancedStaticMeshComponent>(StaticMeshComponent))
 					{
@@ -474,11 +474,11 @@ bool FPCGSpawnActorElement::ExecuteInternal(FPCGContext* Context) const
 				}
 			}
 				
-			for (const TPair<FISMComponentDescriptor, TArray<FTransform>>& ISMCDescriptorTransforms : MeshDescriptorTransforms)
+			for(const TPair<FPCGISMCBuilderParameters, TArray<FTransform>>& ISMCBuilderTransforms : MeshDescriptorTransforms)
 			{
-				const FISMComponentDescriptor& ISMCDescriptor = ISMCDescriptorTransforms.Key;
+				const FPCGISMCBuilderParameters& ISMCParams = ISMCBuilderTransforms.Key;
 
-				UPCGManagedISMComponent* MISMC = UPCGActorHelpers::GetOrCreateManagedISMC(TargetActor, Context->SourceComponent.Get(), ISMCDescriptor);
+				UPCGManagedISMComponent* MISMC = UPCGActorHelpers::GetOrCreateManagedISMC(TargetActor, Context->SourceComponent.Get(), ISMCParams);
 				if (!MISMC)
 				{
 					continue;
@@ -489,7 +489,7 @@ bool FPCGSpawnActorElement::ExecuteInternal(FPCGContext* Context) const
 				UInstancedStaticMeshComponent* ISMC = MISMC->GetComponent();
 				check(ISMC);
 				
-				const TArray<FTransform>& ISMCTransforms = ISMCDescriptorTransforms.Value;
+				const TArray<FTransform>& ISMCTransforms = ISMCBuilderTransforms.Value;
 				
 				TArray<FTransform> Transforms;
 				Transforms.Reserve(Points.Num() * ISMCTransforms.Num());
@@ -503,7 +503,7 @@ bool FPCGSpawnActorElement::ExecuteInternal(FPCGContext* Context) const
 					}
 				}
 				
-				ISMC->NumCustomDataFloats = 0;
+				// Fill in custom data (?)
 				ISMC->AddInstances(Transforms, false, true);
 				ISMC->UpdateBounds();
 				
diff --git a/Engine/Plugins/Experimental/PCG/Source/PCG/Private/Elements/PCGStaticMeshSpawner.cpp b/Engine/Plugins/Experimental/PCG/Source/PCG/Private/Elements/PCGStaticMeshSpawner.cpp
--- a/Engine/Plugins/Experimental/PCG/Source/PCG/Private/Elements/PCGStaticMeshSpawner.cpp
+++ b/Engine/Plugins/Experimental/PCG/Source/PCG/Private/Elements/PCGStaticMeshSpawner.cpp
@@ -266,46 +266,43 @@ void FPCGStaticMeshSpawnerElement::SpawnStaticMeshInstances(FPCGContext* Context
 		return;
 	}
 
-	// Todo: we could likely pre-load these meshes asynchronously in the settings
-	UStaticMesh* LoadedMesh = InstanceList.Mesh.LoadSynchronous();
+	// TODO: we could likely pre-load these meshes asynchronously in the settings
+	UStaticMesh* LoadedMesh = InstanceList.Descriptor.StaticMesh.LoadSynchronous();
 
 	if (!LoadedMesh)
 	{
-		return;
-	}
-
-	FPCGISMCBuilderParameters Params;
-	Params.Mesh = LoadedMesh;
+		// Either we have no mesh (so nothing to do) or the mesh couldn't be loaded
+		if (InstanceList.Descriptor.StaticMesh.IsValid())
+		{
+			PCGE_LOG(Error, "Unable to load mesh %s", *InstanceList.Descriptor.StaticMesh.ToString());
+		}
 
-	if (InstanceList.bOverrideCollisionProfile)
-	{
-		Params.CollisionProfile = InstanceList.CollisionProfile.Name;
+		return;
 	}
 
-	if (InstanceList.bOverrideMaterials)
+	// TODO: we could likely pre-load these materials asynchronously in the settings
+	for (TSoftObjectPtr<UMaterialInterface> OverrideMaterial : InstanceList.Descriptor.OverrideMaterials)
 	{
-		Params.MaterialOverrides.Reserve(InstanceList.MaterialOverrides.Num());
-		for (TSoftObjectPtr<UMaterialInterface> MaterialOverride : InstanceList.MaterialOverrides)
+		if (OverrideMaterial.IsValid() && !OverrideMaterial.LoadSynchronous())
 		{
-			Params.MaterialOverrides.Add(MaterialOverride.LoadSynchronous());
+			PCGE_LOG(Error, "Unable to load override material %s", *OverrideMaterial.ToString());
+			return;
 		}
 	}
 
+	FPCGISMCBuilderParameters Params;
+	Params.Descriptor = FISMComponentDescriptor(InstanceList.Descriptor);
 	Params.NumCustomDataFloats = PackedCustomData.NumCustomDataFloats;
-	Params.CullStartDistance = InstanceList.CullStartDistance;
-	Params.CullEndDistance = InstanceList.CullEndDistance;
-	Params.WorldPositionOffsetDisableDistance = InstanceList.WorldPositionOffsetDisableDistance;
-	Params.bIsLocalToWorldDeterminantNegative = InstanceList.bIsLocalToWorldDeterminantNegative;
 
 	// If the root actor we're binding to is movable, then the ISMC should be movable by default
 	if (USceneComponent* SceneComponent = TargetActor->GetRootComponent())
 	{
-		Params.Mobility = SceneComponent->Mobility;
+		Params.Descriptor.Mobility = SceneComponent->Mobility;
 	}
 
 	UPCGManagedISMComponent* MISMC = UPCGActorHelpers::GetOrCreateManagedISMC(TargetActor, Context->SourceComponent.Get(), Params);
+	
 	check(MISMC);
-
 	MISMC->SetCrc(Context->DependenciesCrc);
 
 	UInstancedStaticMeshComponent* ISMC = MISMC->GetComponent();
@@ -322,14 +319,7 @@ void FPCGStaticMeshSpawnerElement::SpawnStaticMeshInstances(FPCGContext* Context
 	const int32 PreviousCustomDataOffset = PreExistingInstanceCount * NumCustomDataFloats;
 
 	// Populate the ISM instances
-	TArray<FTransform> Instances;
-	Instances.Reserve(NewInstanceCount);
-	for (int32 InstanceIndex = 0; InstanceIndex < NewInstanceCount; ++InstanceIndex)
-	{
-		Instances.Emplace(InstanceList.Instances[InstanceIndex].Transform);
-	}
-
-	ISMC->AddInstances(Instances, /*bShouldReturnIndices=*/false, /*bWorldSpace=*/true);
+	ISMC->AddInstances(InstanceList.Instances, /*bShouldReturnIndices=*/false, /*bWorldSpace=*/true);
 
 	// Copy new CustomData into the ISMC PerInstanceSMCustomData
 	if (NumCustomDataFloats > 0)
@@ -343,7 +333,7 @@ void FPCGStaticMeshSpawnerElement::SpawnStaticMeshInstances(FPCGContext* Context
 
 	ISMC->UpdateBounds();
 
-	PCGE_LOG(Verbose, "Added %d instances of %s on actor %s", InstanceList.Instances.Num(), *InstanceList.Mesh->GetFName().ToString(), *TargetActor->GetFName().ToString());
+	PCGE_LOG(Verbose, "Added %d instances of %s on actor %s", InstanceList.Instances.Num(), *InstanceList.Descriptor.StaticMesh->GetFName().ToString(), *TargetActor->GetFName().ToString());
 }
 
 void UPCGStaticMeshSpawnerSettings::PostLoad()
@@ -360,8 +350,10 @@ void UPCGStaticMeshSpawnerSettings::PostLoad()
 		for (const FPCGStaticMeshSpawnerEntry& Entry : Meshes_DEPRECATED)
 		{
 			FPCGMeshSelectorWeightedEntry& NewEntry = MeshSelector->MeshEntries.Emplace_GetRef(Entry.Mesh, Entry.Weight);
-			NewEntry.CollisionProfile = Entry.CollisionProfile;
-			NewEntry.bOverrideCollisionProfile = Entry.bOverrideCollisionProfile;
+			if (Entry.bOverrideCollisionProfile)
+			{
+				NewEntry.Descriptor.BodyInstance.SetCollisionProfileName(Entry.CollisionProfile.Name);
+			}
 		}
 
 		Meshes_DEPRECATED.Reset();
diff --git a/Engine/Plugins/Experimental/PCG/Source/PCG/Private/Helpers/PCGActorHelpers.cpp b/Engine/Plugins/Experimental/PCG/Source/PCG/Private/Helpers/PCGActorHelpers.cpp
--- a/Engine/Plugins/Experimental/PCG/Source/PCG/Private/Helpers/PCGActorHelpers.cpp
+++ b/Engine/Plugins/Experimental/PCG/Source/PCG/Private/Helpers/PCGActorHelpers.cpp
@@ -12,7 +12,6 @@
 #include "Engine/InheritableComponentHandler.h"
 #include "Engine/SCS_Node.h"
 #include "Engine/StaticMesh.h"
-#include "ISMPartition/ISMComponentDescriptor.h"
 #include "PCGModule.h"
 
 #include UE_INLINE_GENERATED_CPP_BY_NAME(PCGActorHelpers)
@@ -20,89 +19,6 @@
 #if WITH_EDITOR
 #endif
 
-UPCGManagedISMComponent* UPCGActorHelpers::GetOrCreateManagedISMC(AActor* InTargetActor, UPCGComponent* InSourceComponent, const FISMComponentDescriptor& InISMCDescriptor)
-{
-	check(InTargetActor);
-	check(InSourceComponent);
-
-	const UStaticMesh* StaticMesh = InISMCDescriptor.StaticMesh;
-	if (!StaticMesh)
-	{
-		return nullptr;
-	}
-	
-	TArray<UPCGManagedISMComponent*> MISMCs;
-	InSourceComponent->ForEachManagedResource([&MISMCs](UPCGManagedResource* InResource)
-	{
-		if (UPCGManagedISMComponent* Resource = Cast<UPCGManagedISMComponent>(InResource))
-		{
-			MISMCs.Add(Resource);
-		}
-	});
-
-	for (UPCGManagedISMComponent* MISMC : MISMCs)
-	{
-		UInstancedStaticMeshComponent* ISMC = MISMC->GetComponent();
-		if (ISMC)
-		{
-			FISMComponentDescriptor ManagedDescriptor;
-			ManagedDescriptor.InitFrom(MISMC->GetComponent());
-			ManagedDescriptor.ComputeHash();
-			
-			if (ManagedDescriptor == InISMCDescriptor)
-			{
-				MISMC->MarkAsUsed();
-				return MISMC;
-			}
-		}
-	}
-
-	InTargetActor->Modify();
-	
-	// Done as in InstancedStaticMesh.cpp
-#if WITH_EDITOR
-	const bool bMeshHasNaniteData = StaticMesh->IsNaniteEnabled();
-#else
-	const bool bMeshHasNaniteData = StaticMesh->GetRenderData()->NaniteResources.PageStreamingStates.Num() > 0;
-#endif
-
-	FString ComponentName;
-	TSubclassOf<UInstancedStaticMeshComponent> ComponentClass;
-	if (bMeshHasNaniteData)
-	{
-		ComponentClass = UInstancedStaticMeshComponent::StaticClass();
-		ComponentName += TEXT("ISM");
-	}
-	else
-	{
-		ComponentClass = UHierarchicalInstancedStaticMeshComponent::StaticClass();
-		ComponentName += TEXT("HISM");
-	}
-
-	if (InISMCDescriptor.StaticMesh)
-	{
-		ComponentName += TEXT("_") + InISMCDescriptor.StaticMesh->GetName();
-	}
-	
-	UInstancedStaticMeshComponent* ISMC = NewObject<UInstancedStaticMeshComponent>(InTargetActor, ComponentClass, MakeUniqueObjectName(InTargetActor, ComponentClass, FName(ComponentName)));
-	InISMCDescriptor.InitComponent(ISMC);
-		
-	ISMC->RegisterComponent();
-
-	InTargetActor->AddInstanceComponent(ISMC);
-	
-	ISMC->AttachToComponent(InTargetActor->GetRootComponent(), FAttachmentTransformRules(EAttachmentRule::KeepRelative, EAttachmentRule::KeepWorld, EAttachmentRule::KeepWorld, false));
-	ISMC->ComponentTags.Add(InSourceComponent->GetFName());
-	ISMC->ComponentTags.Add(PCGHelpers::DefaultPCGTag);
-
-	// Create managed resource on source component
-	UPCGManagedISMComponent* Resource = NewObject<UPCGManagedISMComponent>(InSourceComponent);
-	Resource->GeneratedComponent = ISMC;
-	InSourceComponent->AddToManagedResources(Resource);
-
-	return Resource;
-}
-
 UInstancedStaticMeshComponent* UPCGActorHelpers::GetOrCreateISMC(AActor* InTargetActor, UPCGComponent* InSourceComponent, const FPCGISMCBuilderParameters& InParams)
 {
 	UPCGManagedISMComponent* MISMC = GetOrCreateManagedISMC(InTargetActor, InSourceComponent, InParams);
@@ -118,83 +34,56 @@ UInstancedStaticMeshComponent* UPCGActorHelpers::GetOrCreateISMC(AActor* InTarge
 
 UPCGManagedISMComponent* UPCGActorHelpers::GetOrCreateManagedISMC(AActor* InTargetActor, UPCGComponent* InSourceComponent, const FPCGISMCBuilderParameters& InParams)
 {
-	UStaticMesh* InMesh = InParams.Mesh;
-	const TArray<UMaterialInterface*>& InMaterials = InParams.MaterialOverrides;
-	const int32 InNumCustomDataFloats = InParams.NumCustomDataFloats;
+	check(InTargetActor && InSourceComponent);
 
-	check(InTargetActor != nullptr && InMesh != nullptr);
-	check(InSourceComponent);
-
-	TArray<UPCGManagedISMComponent*> MISMCs;
-	InSourceComponent->ForEachManagedResource([&MISMCs](UPCGManagedResource* InResource)
+	const UStaticMesh* StaticMesh = InParams.Descriptor.StaticMesh;
+	if (!StaticMesh)
 	{
-		if (UPCGManagedISMComponent* Resource = Cast<UPCGManagedISMComponent>(InResource))
-		{
-			MISMCs.Add(Resource);
-		}
-	});
+		return nullptr;
+	}
 
-	for (UPCGManagedISMComponent* MISMC : MISMCs)
-	{
-		UInstancedStaticMeshComponent* ISMC = MISMC->GetComponent();
+	TRACE_CPUPROFILER_EVENT_SCOPE(UPCGActorHelpers::GetOrCreateManagedISMC);
 
-		if (ISMC && 
-			ISMC->GetStaticMesh() == InMesh &&
-			ISMC->ComponentTags.Contains(InSourceComponent->GetFName()) &&
-			ISMC->NumCustomDataFloats == InNumCustomDataFloats)
+	{
+		TRACE_CPUPROFILER_EVENT_SCOPE(UPCGActorHelpers::GetOrCreateManagedISMC::FindMatchingMISMC);
+		UPCGManagedISMComponent* MatchingResource = nullptr;
+		InSourceComponent->ForEachManagedResource([&MatchingResource, &InParams](UPCGManagedResource* InResource)
 		{
-			// If materials are provided, we'll make sure they match to the already set materials.
-			// If not provided, we'll make sure that the current materials aren't overriden
-			bool bMaterialsMatched = true;
-
-			// Check basic parameters
-			if (ISMC->Mobility != InParams.Mobility ||
-				(ISMC->bUseDefaultCollision && InParams.CollisionProfile != TEXT("Default")) ||
-				(!ISMC->bUseDefaultCollision && ISMC->GetCollisionProfileName() != InParams.CollisionProfile) ||
-				ISMC->InstanceStartCullDistance != InParams.CullStartDistance ||
-				ISMC->InstanceEndCullDistance != InParams.CullEndDistance ||
-				ISMC->WorldPositionOffsetDisableDistance != InParams.WorldPositionOffsetDisableDistance ||
-				ISMC->bReverseCulling != InParams.bIsLocalToWorldDeterminantNegative)
+			// Early out if already found a match
+			if (MatchingResource)
 			{
-				continue;
+				return;
 			}
 
-			for (int32 MaterialIndex = 0; MaterialIndex < ISMC->GetNumMaterials() && bMaterialsMatched; ++MaterialIndex)
+			if (UPCGManagedISMComponent* Resource = Cast<UPCGManagedISMComponent>(InResource))
 			{
-				if (MaterialIndex < InMaterials.Num() && InMaterials[MaterialIndex])
+				if (UInstancedStaticMeshComponent* ISMC = Resource->GetComponent())
 				{
-					if (InMaterials[MaterialIndex] != ISMC->GetMaterial(MaterialIndex))
+					if (ISMC->NumCustomDataFloats == InParams.NumCustomDataFloats &&
+						Resource->GetDescriptor() == InParams.Descriptor)
 					{
-						bMaterialsMatched = false;
-					}
-				}
-				else
-				{
-					// Material is currently overriden
-					if (ISMC->OverrideMaterials.IsValidIndex(MaterialIndex) && ISMC->OverrideMaterials[MaterialIndex])
-					{
-						bMaterialsMatched = false;
+						MatchingResource = Resource;
 					}
 				}
 			}
+		});
 
-			if (bMaterialsMatched)
-			{
-				MISMC->MarkAsUsed();
-				return MISMC;
-			}
+		if (MatchingResource)
+		{
+			TRACE_CPUPROFILER_EVENT_SCOPE(UPCGActorHelpers::GetOrCreateManagedISMC::MarkAsUsed);
+			MatchingResource->MarkAsUsed();
+			return MatchingResource;
 		}
 	}
 
+	// No matching ISM component found, let's create a new one
 	InTargetActor->Modify();
 
-	// Otherwise, create a new component
-
 	// Done as in InstancedStaticMesh.cpp
 #if WITH_EDITOR
-	const bool bMeshHasNaniteData = InMesh->IsNaniteEnabled();
+	const bool bMeshHasNaniteData = StaticMesh->IsNaniteEnabled();
 #else
-	const bool bMeshHasNaniteData = InMesh->GetRenderData()->NaniteResources.PageStreamingStates.Num() > 0;
+	const bool bMeshHasNaniteData = StaticMesh->GetRenderData()->NaniteResources.PageStreamingStates.Num() > 0;
 #endif
 
 	FString ComponentName;
@@ -210,46 +99,23 @@ UPCGManagedISMComponent* UPCGActorHelpers::GetOrCreateManagedISMC(AActor* InTarg
 		ComponentName += TEXT("HISM");
 	}
 
-	ComponentName += TEXT("_") + InMesh->GetName();
-	
-	// TODO: use static mesh component if there's only one instance
-	// TODO: add hism/ism switch or better yet, use a template component
-	UInstancedStaticMeshComponent* ISMC = NewObject<UInstancedStaticMeshComponent>(InTargetActor, ComponentClass, MakeUniqueObjectName(InTargetActor, ComponentClass, FName(ComponentName)));
-	
-	ISMC->SetStaticMesh(InMesh);
+	ComponentName += TEXT("_") + StaticMesh->GetName();
 
-	// TODO: improve material override mechanisms
-	const int32 NumMaterials = ISMC->GetNumMaterials();
-	for (int32 MaterialIndex = 0; MaterialIndex < NumMaterials; ++MaterialIndex)
-	{
-		ISMC->SetMaterial(MaterialIndex, (MaterialIndex < InMaterials.Num()) ? InMaterials[MaterialIndex] : nullptr);
-	}
+	UInstancedStaticMeshComponent* ISMC = NewObject<UInstancedStaticMeshComponent>(InTargetActor, ComponentClass, MakeUniqueObjectName(InTargetActor, ComponentClass, FName(ComponentName)));
+	InParams.Descriptor.InitComponent(ISMC);
+	ISMC->NumCustomDataFloats = InParams.NumCustomDataFloats;
 
 	ISMC->RegisterComponent();
 	InTargetActor->AddInstanceComponent(ISMC);
-	ISMC->SetMobility(InParams.Mobility);
-	if (InParams.CollisionProfile != TEXT("Default"))
-	{
-		ISMC->SetCollisionProfileName(InParams.CollisionProfile, /*bOverlaps=*/false);
-	}
-	else
-	{
-		ISMC->bUseDefaultCollision = true;
-	}
-
-	ISMC->InstanceStartCullDistance = InParams.CullStartDistance;
-	ISMC->InstanceEndCullDistance = InParams.CullEndDistance;
-	ISMC->WorldPositionOffsetDisableDistance = InParams.WorldPositionOffsetDisableDistance;
-	ISMC->bReverseCulling = (InParams.bIsLocalToWorldDeterminantNegative ? 1 : 0);
 
-	// Implementation note: we use the relative for position only to prevent issues at the rendering level, but world for scale & rotation otherwise we run into issues when we have non-uniform scales esp. in volumes
 	ISMC->AttachToComponent(InTargetActor->GetRootComponent(), FAttachmentTransformRules(EAttachmentRule::KeepRelative, EAttachmentRule::KeepWorld, EAttachmentRule::KeepWorld, false));
 	ISMC->ComponentTags.Add(InSourceComponent->GetFName());
 	ISMC->ComponentTags.Add(PCGHelpers::DefaultPCGTag);
 
 	// Create managed resource on source component
 	UPCGManagedISMComponent* Resource = NewObject<UPCGManagedISMComponent>(InSourceComponent);
-	Resource->GeneratedComponent = ISMC;
+	Resource->SetComponent(ISMC);
+	Resource->SetDescriptor(InParams.Descriptor);
 	InSourceComponent->AddToManagedResources(Resource);
 
 	return Resource;
diff --git a/Engine/Plugins/Experimental/PCG/Source/PCG/Private/InstancePackers/PCGInstancePackerBase.cpp b/Engine/Plugins/Experimental/PCG/Source/PCG/Private/InstancePackers/PCGInstancePackerBase.cpp
--- a/Engine/Plugins/Experimental/PCG/Source/PCG/Private/InstancePackers/PCGInstancePackerBase.cpp
+++ b/Engine/Plugins/Experimental/PCG/Source/PCG/Private/InstancePackers/PCGInstancePackerBase.cpp
@@ -64,7 +64,7 @@ void UPCGInstancePackerBase::PackCustomDataFromAttributes(const FPCGMeshInstance
 
 void UPCGInstancePackerBase::PackCustomDataFromAttributes(const FPCGMeshInstanceList& InstanceList, const TArray<const FPCGMetadataAttributeBase*>& Attributes, FPCGPackedCustomData& OutPackedCustomData) const
 {
-	for (const FPCGPoint& Point : InstanceList.Instances)
+	for(uint64 PointMetadataEntry : InstanceList.InstancesMetadataEntry)
 	{
 		for (const FPCGMetadataAttributeBase* AttributeBase : Attributes)
 		{
@@ -75,7 +75,7 @@ void UPCGInstancePackerBase::PackCustomDataFromAttributes(const FPCGMeshInstance
 				const FPCGMetadataAttribute<float>* Attribute = static_cast<const FPCGMetadataAttribute<float>*>(AttributeBase);
 				check(Attribute);
 
-				const float Value = Attribute->GetValueFromItemKey(Point.MetadataEntry);
+				const float Value = Attribute->GetValueFromItemKey(PointMetadataEntry);
 				OutPackedCustomData.CustomData.Add(Value);
 			}
 			else if (AttributeBase->GetTypeId() == PCG::Private::MetadataTypes<double>::Id)
@@ -83,7 +83,7 @@ void UPCGInstancePackerBase::PackCustomDataFromAttributes(const FPCGMeshInstance
 				const FPCGMetadataAttribute<double>* Attribute = static_cast<const FPCGMetadataAttribute<double>*>(AttributeBase);
 				check(Attribute);
 
-				const double Value = Attribute->GetValueFromItemKey(Point.MetadataEntry);
+				const double Value = Attribute->GetValueFromItemKey(PointMetadataEntry);
 				OutPackedCustomData.CustomData.Add(Value);
 			}
 			else if (AttributeBase->GetTypeId() == PCG::Private::MetadataTypes<int32>::Id)
@@ -91,7 +91,7 @@ void UPCGInstancePackerBase::PackCustomDataFromAttributes(const FPCGMeshInstance
 				const FPCGMetadataAttribute<int32>* Attribute = static_cast<const FPCGMetadataAttribute<int32>*>(AttributeBase);
 				check(Attribute);
 
-				const float Value = static_cast<float>(Attribute->GetValueFromItemKey(Point.MetadataEntry));
+				const float Value = static_cast<float>(Attribute->GetValueFromItemKey(PointMetadataEntry));
 				OutPackedCustomData.CustomData.Add(Value);
 			}
 			else if (AttributeBase->GetTypeId() == PCG::Private::MetadataTypes<int64>::Id)
@@ -99,7 +99,7 @@ void UPCGInstancePackerBase::PackCustomDataFromAttributes(const FPCGMeshInstance
 				const FPCGMetadataAttribute<int64>* Attribute = static_cast<const FPCGMetadataAttribute<int64>*>(AttributeBase);
 				check(Attribute);
 
-				const float Value = static_cast<float>(Attribute->GetValueFromItemKey(Point.MetadataEntry));
+				const float Value = static_cast<float>(Attribute->GetValueFromItemKey(PointMetadataEntry));
 				OutPackedCustomData.CustomData.Add(Value);
 			}
 			else if (AttributeBase->GetTypeId() == PCG::Private::MetadataTypes<FRotator>::Id)
@@ -107,7 +107,7 @@ void UPCGInstancePackerBase::PackCustomDataFromAttributes(const FPCGMeshInstance
 				const FPCGMetadataAttribute<FRotator>* Attribute = static_cast<const FPCGMetadataAttribute<FRotator>*>(AttributeBase);
 				check(Attribute);
 
-				const FRotator Value = Attribute->GetValueFromItemKey(Point.MetadataEntry);
+				const FRotator Value = Attribute->GetValueFromItemKey(PointMetadataEntry);
 				OutPackedCustomData.CustomData.Add(Value.Roll);
 				OutPackedCustomData.CustomData.Add(Value.Pitch);
 				OutPackedCustomData.CustomData.Add(Value.Yaw);
@@ -117,7 +117,7 @@ void UPCGInstancePackerBase::PackCustomDataFromAttributes(const FPCGMeshInstance
 				const FPCGMetadataAttribute<FVector>* Attribute = static_cast<const FPCGMetadataAttribute<FVector>*>(AttributeBase);
 				check(Attribute);
 
-				const FVector Value = Attribute->GetValueFromItemKey(Point.MetadataEntry);
+				const FVector Value = Attribute->GetValueFromItemKey(PointMetadataEntry);
 				OutPackedCustomData.CustomData.Add(Value.X);
 				OutPackedCustomData.CustomData.Add(Value.Y);
 				OutPackedCustomData.CustomData.Add(Value.Z);
@@ -127,7 +127,7 @@ void UPCGInstancePackerBase::PackCustomDataFromAttributes(const FPCGMeshInstance
 				const FPCGMetadataAttribute<FVector4>* Attribute = static_cast<const FPCGMetadataAttribute<FVector4>*>(AttributeBase);
 				check(Attribute);
 
-				const FVector4 Value = Attribute->GetValueFromItemKey(Point.MetadataEntry);
+				const FVector4 Value = Attribute->GetValueFromItemKey(PointMetadataEntry);
 				OutPackedCustomData.CustomData.Add(Value.X);
 				OutPackedCustomData.CustomData.Add(Value.Y);
 				OutPackedCustomData.CustomData.Add(Value.Z);
diff --git a/Engine/Plugins/Experimental/PCG/Source/PCG/Private/MeshSelectors/PCGMeshSelectorBase.cpp b/Engine/Plugins/Experimental/PCG/Source/PCG/Private/MeshSelectors/PCGMeshSelectorBase.cpp
--- a/Engine/Plugins/Experimental/PCG/Source/PCG/Private/MeshSelectors/PCGMeshSelectorBase.cpp
+++ b/Engine/Plugins/Experimental/PCG/Source/PCG/Private/MeshSelectors/PCGMeshSelectorBase.cpp
@@ -5,61 +5,14 @@
 
 #include UE_INLINE_GENERATED_CPP_BY_NAME(PCGMeshSelectorBase)
 
-int32 UPCGMeshSelectorBase::FindOrAddInstanceList(
-	TArray<FPCGMeshInstanceList>& OutInstanceLists,
-	const TSoftObjectPtr<UStaticMesh>& Mesh,
-	bool bOverrideCollisionProfile,
-	const FCollisionProfileName& CollisionProfile,
-	bool bOverrideMaterials,
-	const TArray<TSoftObjectPtr<UMaterialInterface>>& MaterialOverrides,
-	const float InCullStartDistance,
-	const float InCullEndDistance,
-	const int32 InWorldPositionOffsetDisableDistance,
-	const bool bInIsLocalToWorldDeterminantNegative)
-{
-	for (int Index = 0; Index < OutInstanceLists.Num(); ++Index)
-	{
-		if (OutInstanceLists[Index].Mesh != Mesh || OutInstanceLists[Index].bOverrideCollisionProfile != bOverrideCollisionProfile || OutInstanceLists[Index].bOverrideMaterials != bOverrideMaterials)
-		{
-			continue;
-		}
-
-		if (OutInstanceLists[Index].CullStartDistance != InCullStartDistance || 
-			OutInstanceLists[Index].CullEndDistance != InCullEndDistance ||
-			OutInstanceLists[Index].WorldPositionOffsetDisableDistance != InWorldPositionOffsetDisableDistance)
-		{
-			continue;
-		}
-
-		if (OutInstanceLists[Index].bIsLocalToWorldDeterminantNegative != bInIsLocalToWorldDeterminantNegative)
-		{
-			continue;
-		}
-
-		if (OutInstanceLists[Index].bOverrideCollisionProfile && !OutInstanceLists[Index].CollisionProfile.Name.IsEqual(CollisionProfile.Name))
-		{
-			continue;
-		}
-
-		if (OutInstanceLists[Index].bOverrideMaterials && ToRawPtrTArrayUnsafe(OutInstanceLists[Index].MaterialOverrides) != MaterialOverrides)
-		{
-			continue;
-		}
-
-		return Index;
-	}
-
-	return OutInstanceLists.Emplace(Mesh, bOverrideCollisionProfile, CollisionProfile, bOverrideMaterials, MaterialOverrides, InCullStartDistance, InCullEndDistance, InWorldPositionOffsetDisableDistance, bInIsLocalToWorldDeterminantNegative);
-}
-
 FPCGMeshMaterialOverrideHelper::FPCGMeshMaterialOverrideHelper(
 	FPCGContext& InContext,
-	EPCGMeshSelectorMaterialOverrideMode InMaterialOverrideMode,
+	bool bInUseMaterialOverrideAttributes,
 	const TArray<TSoftObjectPtr<UMaterialInterface>>& InStaticMaterialOverrides,
 	const TArray<FName>& InMaterialOverrideAttributeNames,
 	const UPCGMetadata* InMetadata
 )
-	: MaterialOverrideMode(InMaterialOverrideMode)
+	: bUseMaterialOverrideAttributes(bInUseMaterialOverrideAttributes)
 	, StaticMaterialOverrides(InStaticMaterialOverrides)
 	, MaterialOverrideAttributeNames(InMaterialOverrideAttributeNames)
 	, Metadata(InMetadata)
@@ -73,7 +26,7 @@ FPCGMeshMaterialOverrideHelper::FPCGMeshMaterialOverrideHelper(
 	const TArray<FName>& InMaterialOverrideAttributeNames,
 	const UPCGMetadata* InMetadata
 )
-	: MaterialOverrideMode(bInByAttributeOverride ? EPCGMeshSelectorMaterialOverrideMode::ByAttributeOverride : EPCGMeshSelectorMaterialOverrideMode::NoOverride)
+	: bUseMaterialOverrideAttributes(bInByAttributeOverride)
 	, StaticMaterialOverrides(EmptyArray)
 	, MaterialOverrideAttributeNames(InMaterialOverrideAttributeNames)
 	, Metadata(InMetadata)
@@ -84,7 +37,7 @@ FPCGMeshMaterialOverrideHelper::FPCGMeshMaterialOverrideHelper(
 void FPCGMeshMaterialOverrideHelper::Initialize(FPCGContext& InContext)
 {
 	// Perform data setup & validation up-front
-	if (MaterialOverrideMode == EPCGMeshSelectorMaterialOverrideMode::ByAttributeOverride)
+	if (bUseMaterialOverrideAttributes)
 	{
 		if (!Metadata)
 		{
@@ -122,7 +75,7 @@ void FPCGMeshMaterialOverrideHelper::Initialize(FPCGContext& InContext)
 const TArray<TSoftObjectPtr<UMaterialInterface>>& FPCGMeshMaterialOverrideHelper::GetMaterialOverrides(PCGMetadataEntryKey EntryKey)
 {
 	check(bIsValid);
-	if (MaterialOverrideMode == EPCGMeshSelectorMaterialOverrideMode::ByAttributeOverride)
+	if (bUseMaterialOverrideAttributes)
 	{
 		WorkingMaterialOverrides.Reset();
 
@@ -149,12 +102,8 @@ const TArray<TSoftObjectPtr<UMaterialInterface>>& FPCGMeshMaterialOverrideHelper
 
 		return WorkingMaterialOverrides;
 	}
-	else if(MaterialOverrideMode == EPCGMeshSelectorMaterialOverrideMode::StaticOverride)
+	else
 	{
 		return StaticMaterialOverrides;
 	}
-	else // no override
-	{
-		return EmptyArray;
-	}
 }
diff --git a/Engine/Plugins/Experimental/PCG/Source/PCG/Private/MeshSelectors/PCGMeshSelectorByAttribute.cpp b/Engine/Plugins/Experimental/PCG/Source/PCG/Private/MeshSelectors/PCGMeshSelectorByAttribute.cpp
--- a/Engine/Plugins/Experimental/PCG/Source/PCG/Private/MeshSelectors/PCGMeshSelectorByAttribute.cpp
+++ b/Engine/Plugins/Experimental/PCG/Source/PCG/Private/MeshSelectors/PCGMeshSelectorByAttribute.cpp
@@ -9,6 +9,35 @@
 
 #include UE_INLINE_GENERATED_CPP_BY_NAME(PCGMeshSelectorByAttribute)
 
+namespace PCGMeshSelectorAttribute
+{
+	// Returns variation based on mesh, material overrides and reverse culling
+	FPCGMeshInstanceList& GetInstanceList(
+		TArray<FPCGMeshInstanceList>& InstanceLists,
+		const FSoftISMComponentDescriptor& TemplateDescriptor,
+		TSoftObjectPtr<UStaticMesh> Mesh,
+		const TArray<TSoftObjectPtr<UMaterialInterface>>& MaterialOverrides,
+		bool bReverseCulling)
+	{
+		for (FPCGMeshInstanceList& InstanceList : InstanceLists)
+		{
+			if (InstanceList.Descriptor.StaticMesh == Mesh &&
+				InstanceList.Descriptor.bReverseCulling == bReverseCulling &&
+				InstanceList.Descriptor.OverrideMaterials == MaterialOverrides)
+			{
+				return InstanceList;
+			}
+		}
+
+		FPCGMeshInstanceList& NewInstanceList = InstanceLists.Emplace_GetRef(TemplateDescriptor);
+		NewInstanceList.Descriptor.StaticMesh = Mesh;
+		NewInstanceList.Descriptor.OverrideMaterials = MaterialOverrides;
+		NewInstanceList.Descriptor.bReverseCulling = bReverseCulling;
+
+		return NewInstanceList;
+	}
+}
+
 void UPCGMeshSelectorByAttribute::PostLoad()
 {
 	Super::PostLoad();
@@ -16,9 +45,47 @@ void UPCGMeshSelectorByAttribute::PostLoad()
 #if WITH_EDITOR
 	if (bOverrideMaterials_DEPRECATED)
 	{
-		MaterialOverrideMode = EPCGMeshSelectorMaterialOverrideMode::StaticOverride;
+		MaterialOverrideMode_DEPRECATED = EPCGMeshSelectorMaterialOverrideMode::StaticOverride;
 		bOverrideMaterials_DEPRECATED = false;
 	}
+
+	if (bOverrideCollisionProfile_DEPRECATED ||
+		CollisionProfile_DEPRECATED.Name != UCollisionProfile::NoCollision_ProfileName ||
+		MaterialOverrides_DEPRECATED.Num() > 0 ||
+		MaterialOverrideMode_DEPRECATED != EPCGMeshSelectorMaterialOverrideMode::NoOverride ||
+		CullStartDistance_DEPRECATED != 0 ||
+		CullEndDistance_DEPRECATED != 0 ||
+		WorldPositionOffsetDisableDistance_DEPRECATED != 0)
+	{
+		if (bOverrideCollisionProfile_DEPRECATED)
+		{
+			TemplateDescriptor.bUseDefaultCollision = false;
+			TemplateDescriptor.BodyInstance.SetCollisionProfileName(CollisionProfile_DEPRECATED.Name);
+		}
+		else
+		{
+			TemplateDescriptor.bUseDefaultCollision = true;
+		}
+
+		if (MaterialOverrideMode_DEPRECATED != EPCGMeshSelectorMaterialOverrideMode::NoOverride)
+		{
+			TemplateDescriptor.OverrideMaterials = MaterialOverrides_DEPRECATED;
+		}
+		
+		TemplateDescriptor.InstanceStartCullDistance = CullStartDistance_DEPRECATED;
+		TemplateDescriptor.InstanceEndCullDistance = CullEndDistance_DEPRECATED;
+		TemplateDescriptor.WorldPositionOffsetDisableDistance = WorldPositionOffsetDisableDistance_DEPRECATED;
+
+		bUseAttributeMaterialOverrides = (MaterialOverrideMode_DEPRECATED == EPCGMeshSelectorMaterialOverrideMode::ByAttributeOverride);
+
+		bOverrideCollisionProfile_DEPRECATED = false;
+		CollisionProfile_DEPRECATED = UCollisionProfile::NoCollision_ProfileName;
+		MaterialOverrideMode_DEPRECATED = EPCGMeshSelectorMaterialOverrideMode::NoOverride;
+		MaterialOverrides_DEPRECATED.Reset();
+		CullStartDistance_DEPRECATED = 0;
+		CullEndDistance_DEPRECATED = 0;
+		WorldPositionOffsetDisableDistance_DEPRECATED = 0;
+	}
 #endif
 }
 
@@ -58,7 +125,7 @@ void UPCGMeshSelectorByAttribute::SelectInstances_Implementation(
 
 	const FPCGMetadataAttribute<FString>* Attribute = static_cast<const FPCGMetadataAttribute<FString>*>(AttributeBase);
 
-	FPCGMeshMaterialOverrideHelper MaterialOverrideHelper(Context, MaterialOverrideMode, MaterialOverrides, MaterialOverrideAttributes, InPointData->Metadata);
+	FPCGMeshMaterialOverrideHelper MaterialOverrideHelper(Context, bUseAttributeMaterialOverrides, TemplateDescriptor.OverrideMaterials, MaterialOverrideAttributes, InPointData->Metadata);
 
 	if (!MaterialOverrideHelper.IsValid())
 	{
@@ -122,9 +189,8 @@ void UPCGMeshSelectorByAttribute::SelectInstances_Implementation(
 		}
 
 		const bool bReverseTransform = (Point.Transform.GetDeterminant() < 0);
-		const int32 Index = FindOrAddInstanceList(OutMeshInstances, Mesh, bOverrideCollisionProfile, CollisionProfile, MaterialOverrideHelper.OverridesMaterials(), MaterialOverrideHelper.GetMaterialOverrides(Point.MetadataEntry), CullStartDistance, CullEndDistance, WorldPositionOffsetDisableDistance, bReverseTransform);
-		check(Index != INDEX_NONE);
-		
-		OutMeshInstances[Index].Instances.Emplace(Point);
+		FPCGMeshInstanceList& InstanceList = PCGMeshSelectorAttribute::GetInstanceList(OutMeshInstances, TemplateDescriptor, Mesh, MaterialOverrideHelper.GetMaterialOverrides(Point.MetadataEntry), bReverseTransform);
+		InstanceList.Instances.Emplace(Point.Transform);
+		InstanceList.InstancesMetadataEntry.Emplace(Point.MetadataEntry);
 	}
 }
diff --git a/Engine/Plugins/Experimental/PCG/Source/PCG/Private/MeshSelectors/PCGMeshSelectorWeighted.cpp b/Engine/Plugins/Experimental/PCG/Source/PCG/Private/MeshSelectors/PCGMeshSelectorWeighted.cpp
--- a/Engine/Plugins/Experimental/PCG/Source/PCG/Private/MeshSelectors/PCGMeshSelectorWeighted.cpp
+++ b/Engine/Plugins/Experimental/PCG/Source/PCG/Private/MeshSelectors/PCGMeshSelectorWeighted.cpp
@@ -16,58 +16,99 @@ namespace PCGMeshSelectorWeighted
 {
 	FPCGMeshInstanceList& GetInstanceList(
 		TArray<FPCGMeshInstanceList>& InstanceLists,
-		bool bUseAttributeMaterialOverrides,
+		bool bUseMaterialOverrides,
 		const TArray<TSoftObjectPtr<UMaterialInterface>>& InMaterialOverrides,
-		bool bInIsLocalToWorldDeterminantNegative)
+		bool bInReverseCulling)
 	{
 		check(InstanceLists.Num() > 0);
-		const bool bOverrideMaterials = (bUseAttributeMaterialOverrides || InstanceLists[0].bOverrideMaterials);
 
 		// First look through previously existing values - note that we scope this to prevent issues with the 0 index access which might become invalid below
 		{
-			const TArray<TSoftObjectPtr<UMaterialInterface>>& MaterialOverrides = (bUseAttributeMaterialOverrides ? InMaterialOverrides : InstanceLists[0].MaterialOverrides);
-
-			for (int Index = 0; Index < InstanceLists.Num(); ++Index)
+			const TArray<TSoftObjectPtr<UMaterialInterface>>& MaterialOverrides = (bUseMaterialOverrides ? InMaterialOverrides : InstanceLists[0].Descriptor.OverrideMaterials);
+			for (FPCGMeshInstanceList& InstanceList : InstanceLists)
 			{
-				if (InstanceLists[Index].bOverrideMaterials != bOverrideMaterials)
+				if (InstanceList.Descriptor.bReverseCulling == bInReverseCulling &&
+					InstanceList.Descriptor.OverrideMaterials == MaterialOverrides)
 				{
-					continue;
+					return InstanceList;
 				}
+			}
+		}
 
-				if (InstanceLists[Index].bIsLocalToWorldDeterminantNegative != bInIsLocalToWorldDeterminantNegative)
-				{
-					continue;
-				}
+		FPCGMeshInstanceList& NewInstanceList = InstanceLists.Emplace_GetRef();
+		NewInstanceList.Descriptor = InstanceLists[0].Descriptor;
+		NewInstanceList.Descriptor.bReverseCulling = bInReverseCulling;
 
-				if (bOverrideMaterials && InstanceLists[Index].MaterialOverrides != MaterialOverrides)
-				{
-					continue;
-				}
+		if (bUseMaterialOverrides)
+		{
+			NewInstanceList.Descriptor.OverrideMaterials = InMaterialOverrides;
+		}
 
-				return InstanceLists[Index];
-			}
+		return NewInstanceList;
+	}
+}
+
+FPCGMeshSelectorWeightedEntry::FPCGMeshSelectorWeightedEntry(TSoftObjectPtr<UStaticMesh> InMesh, int InWeight)
+	: Weight(InWeight)
+{
+	Descriptor.StaticMesh = InMesh;
+}
+
+#if WITH_EDITOR
+void FPCGMeshSelectorWeightedEntry::ApplyDeprecation()
+{
+	if (Mesh_DEPRECATED.IsValid() ||
+		bOverrideCollisionProfile_DEPRECATED ||
+		CollisionProfile_DEPRECATED.Name != UCollisionProfile::NoCollision_ProfileName ||
+		bOverrideMaterials_DEPRECATED ||
+		MaterialOverrides_DEPRECATED.Num() > 0 ||
+		CullStartDistance_DEPRECATED != 0 ||
+		CullEndDistance_DEPRECATED != 0 ||
+		WorldPositionOffsetDisableDistance_DEPRECATED != 0)
+	{
+		Descriptor.StaticMesh = Mesh_DEPRECATED;
+		
+		if (bOverrideCollisionProfile_DEPRECATED)
+		{
+			Descriptor.bUseDefaultCollision = false;
+			Descriptor.BodyInstance.SetCollisionProfileName(CollisionProfile_DEPRECATED.Name);
 		}
+		else
+		{
+			Descriptor.bUseDefaultCollision = true;
+		}
+
+		Descriptor.InstanceStartCullDistance = CullStartDistance_DEPRECATED;
+		Descriptor.InstanceEndCullDistance = CullEndDistance_DEPRECATED;
+		Descriptor.WorldPositionOffsetDisableDistance = WorldPositionOffsetDisableDistance_DEPRECATED;
 
-		// If not found, then copy first entry which is our "clean" version and apply the new values
+		if (bOverrideMaterials_DEPRECATED)
 		{
-			FPCGMeshInstanceList& NewInstanceList = InstanceLists.Emplace_GetRef();
-			const FPCGMeshInstanceList& SourceInstanceList = InstanceLists[0];
-			NewInstanceList.Mesh = SourceInstanceList.Mesh;
-			NewInstanceList.bOverrideCollisionProfile = SourceInstanceList.bOverrideCollisionProfile;
-			NewInstanceList.CollisionProfile = SourceInstanceList.CollisionProfile;
-			NewInstanceList.CullStartDistance = SourceInstanceList.CullStartDistance;
-			NewInstanceList.CullEndDistance = SourceInstanceList.CullEndDistance;
-			NewInstanceList.WorldPositionOffsetDisableDistance = SourceInstanceList.WorldPositionOffsetDisableDistance;
-
-			NewInstanceList.bOverrideMaterials = bOverrideMaterials;
-
-			const TArray<TSoftObjectPtr<UMaterialInterface>>& MaterialOverrides = (bUseAttributeMaterialOverrides ? InMaterialOverrides : SourceInstanceList.MaterialOverrides);
-			NewInstanceList.MaterialOverrides = MaterialOverrides;
-			NewInstanceList.bIsLocalToWorldDeterminantNegative = bInIsLocalToWorldDeterminantNegative;
-
-			return NewInstanceList;
+			Descriptor.OverrideMaterials = MaterialOverrides_DEPRECATED;
 		}
+
+		Mesh_DEPRECATED.Reset();
+		bOverrideCollisionProfile_DEPRECATED = false;
+		CollisionProfile_DEPRECATED = UCollisionProfile::NoCollision_ProfileName;
+		bOverrideMaterials_DEPRECATED = false;
+		MaterialOverrides_DEPRECATED.Reset();
+		CullStartDistance_DEPRECATED = 0;
+		CullEndDistance_DEPRECATED = 0;		
+		WorldPositionOffsetDisableDistance_DEPRECATED = 0;
+	}
+}
+#endif
+
+void UPCGMeshSelectorWeighted::PostLoad()
+{
+	Super::PostLoad();
+
+#if WITH_EDITOR
+	for (FPCGMeshSelectorWeightedEntry& Entry : MeshEntries)
+	{
+		Entry.ApplyDeprecation();
 	}
+#endif
 }
 
 void UPCGMeshSelectorWeighted::SelectInstances_Implementation(
@@ -92,7 +133,7 @@ void UPCGMeshSelectorWeighted::SelectInstances_Implementation(
 		}
 
 		TArray<FPCGMeshInstanceList>& PickEntry = MeshInstances.Emplace_GetRef();
-		FindOrAddInstanceList(PickEntry, Entry.Mesh, Entry.bOverrideCollisionProfile, Entry.CollisionProfile, Entry.bOverrideMaterials, Entry.MaterialOverrides, Entry.CullStartDistance, Entry.CullEndDistance, Entry.WorldPositionOffsetDisableDistance, /*bReverseCulling=*/false);
+		PickEntry.Emplace_GetRef(Entry.Descriptor);
 
 		TotalWeight += Entry.Weight;
 		CumulativeWeights.Add(TotalWeight);
@@ -169,9 +210,10 @@ void UPCGMeshSelectorWeighted::SelectInstances_Implementation(
 			{
 				const bool bNeedsReverseCulling = (Point.Transform.GetDeterminant() < 0);
 				FPCGMeshInstanceList& InstanceList = PCGMeshSelectorWeighted::GetInstanceList(MeshInstances[RandomPick], bUseAttributeMaterialOverrides, MaterialOverrideHelper.GetMaterialOverrides(Point.MetadataEntry), bNeedsReverseCulling);
-				InstanceList.Instances.Emplace(Point);
+				InstanceList.Instances.Emplace(Point.Transform);
+				InstanceList.InstancesMetadataEntry.Emplace(Point.MetadataEntry);
 
-				const TSoftObjectPtr<UStaticMesh>& Mesh = InstanceList.Mesh;
+				const TSoftObjectPtr<UStaticMesh>& Mesh = InstanceList.Descriptor.StaticMesh;
 
 				if (OutPointData && OutAttribute)
 				{
diff --git a/Engine/Plugins/Experimental/PCG/Source/PCG/Private/MeshSelectors/PCGMeshSelectorWeightedByCategory.cpp b/Engine/Plugins/Experimental/PCG/Source/PCG/Private/MeshSelectors/PCGMeshSelectorWeightedByCategory.cpp
--- a/Engine/Plugins/Experimental/PCG/Source/PCG/Private/MeshSelectors/PCGMeshSelectorWeightedByCategory.cpp
+++ b/Engine/Plugins/Experimental/PCG/Source/PCG/Private/MeshSelectors/PCGMeshSelectorWeightedByCategory.cpp
@@ -18,6 +18,28 @@ struct FPCGInstancesAndWeights
 	TArray<int> CumulativeWeights;
 };
 
+#if WITH_EDITOR
+void FPCGWeightedByCategoryEntryList::ApplyDeprecation()
+{
+	for (FPCGMeshSelectorWeightedEntry& Entry : WeightedMeshEntries)
+	{
+		Entry.ApplyDeprecation();
+	}
+}
+#endif
+
+void UPCGMeshSelectorWeightedByCategory::PostLoad()
+{
+	Super::PostLoad();
+
+#if WITH_EDITOR
+	for (FPCGWeightedByCategoryEntryList& Entry : Entries)
+	{
+		Entry.ApplyDeprecation();
+	}
+#endif
+}
+
 void UPCGMeshSelectorWeightedByCategory::SelectInstances_Implementation(
 	FPCGContext& Context,
 	const UPCGStaticMeshSpawnerSettings* Settings,
@@ -109,7 +131,7 @@ void UPCGMeshSelectorWeightedByCategory::SelectInstances_Implementation(
 			}
 			
 			TArray<FPCGMeshInstanceList>& PickEntry = InstancesAndWeights->MeshInstances.Emplace_GetRef();
-			FindOrAddInstanceList(PickEntry, WeightedEntry.Mesh, WeightedEntry.bOverrideCollisionProfile, WeightedEntry.CollisionProfile, WeightedEntry.bOverrideMaterials, WeightedEntry.MaterialOverrides, WeightedEntry.CullStartDistance, WeightedEntry.CullEndDistance, WeightedEntry.WorldPositionOffsetDisableDistance, /*bReverseCulling=*/false);
+			PickEntry.Emplace_GetRef(WeightedEntry.Descriptor);
 
 			// precompute the weights
 			TotalWeight += WeightedEntry.Weight;
@@ -201,9 +223,10 @@ void UPCGMeshSelectorWeightedByCategory::SelectInstances_Implementation(
 		{
 			const bool bNeedsReverseCulling = (Point.Transform.GetDeterminant() < 0);
 			FPCGMeshInstanceList& InstanceList = PCGMeshSelectorWeighted::GetInstanceList(InstancesAndWeights->MeshInstances[RandomPick], bUseAttributeMaterialOverrides, MaterialOverrideHelper.GetMaterialOverrides(Point.MetadataEntry), bNeedsReverseCulling);
-			InstanceList.Instances.Emplace(Point);
+			InstanceList.Instances.Emplace(Point.Transform);
+			InstanceList.InstancesMetadataEntry.Emplace(Point.MetadataEntry);
 
-			const TSoftObjectPtr<UStaticMesh>& Mesh = InstanceList.Mesh;
+			const TSoftObjectPtr<UStaticMesh>& Mesh = InstanceList.Descriptor.StaticMesh;
 
 			if (OutPointData && OutAttribute)
 			{
diff --git a/Engine/Plugins/Experimental/PCG/Source/PCG/Private/PCGManagedResource.cpp b/Engine/Plugins/Experimental/PCG/Source/PCG/Private/PCGManagedResource.cpp
--- a/Engine/Plugins/Experimental/PCG/Source/PCG/Private/PCGManagedResource.cpp
+++ b/Engine/Plugins/Experimental/PCG/Source/PCG/Private/PCGManagedResource.cpp
@@ -184,13 +184,13 @@ void UPCGManagedComponent::PostEditImport()
 		{
 			// Not quite clear what to do when we have a component that cannot be remapped.
 			// Maybe we should check against guids instead?
-			GeneratedComponent.Reset();
+			ForgetComponent();
 		}
 	}
 	else
 	{
 		// Somewhat irrelevant case, if we don't have an actor or a component, there's not a lot we can do.
-		GeneratedComponent.Reset();
+		ForgetComponent();
 	}
 }
 
@@ -260,7 +260,7 @@ bool UPCGManagedComponent::MoveResourceToNewActor(AActor* NewActor)
 		NewActor->AddInstanceComponent(GeneratedComponent.Get());
 	}
 
-	GeneratedComponent.Reset();
+	ForgetComponent();
 
 	return true;
 }
@@ -295,6 +295,37 @@ void UPCGManagedComponent::MarkAsReused()
 	}
 }
 
+void UPCGManagedISMComponent::PostLoad()
+{
+	Super::PostLoad();
+
+	if (!bHasDescriptor)
+	{
+		if (UInstancedStaticMeshComponent* ISMC = GetComponent())
+		{
+			FISMComponentDescriptor NewDescriptor;
+			NewDescriptor.InitFrom(ISMC);
+
+			SetDescriptor(NewDescriptor);
+		}
+	}
+
+	// Cache raw ptr
+	GetComponent();
+}
+
+void UPCGManagedISMComponent::ForgetComponent()
+{
+	Super::ForgetComponent();
+	CachedRawComponentPtr = nullptr;
+}
+
+void UPCGManagedISMComponent::SetDescriptor(const FISMComponentDescriptor& InDescriptor)
+{
+	bHasDescriptor = true;
+	Descriptor = InDescriptor;
+}
+
 bool UPCGManagedISMComponent::ReleaseIfUnused(TSet<TSoftObjectPtr<AActor>>& OutActorsToDelete)
 {
 	if (Super::ReleaseIfUnused(OutActorsToDelete) || !GetComponent())
@@ -304,6 +335,7 @@ bool UPCGManagedISMComponent::ReleaseIfUnused(TSet<TSoftObjectPtr<AActor>>& OutA
 	else if (GetComponent()->GetInstanceCount() == 0)
 	{
 		GeneratedComponent->DestroyComponent();
+		ForgetComponent();
 		return true;
 	}
 	else
@@ -323,5 +355,16 @@ void UPCGManagedISMComponent::ResetComponent()
 
 UInstancedStaticMeshComponent* UPCGManagedISMComponent::GetComponent() const
 {
-	return Cast<UInstancedStaticMeshComponent>(GeneratedComponent.Get());
+	if (!CachedRawComponentPtr)
+	{
+		CachedRawComponentPtr = Cast<UInstancedStaticMeshComponent>(GeneratedComponent.Get());
+	}
+
+	return CachedRawComponentPtr;
 }
+
+void UPCGManagedISMComponent::SetComponent(UInstancedStaticMeshComponent* InComponent)
+{
+	GeneratedComponent = InComponent;
+	CachedRawComponentPtr = InComponent;
+}
\ No newline at end of file
diff --git a/Engine/Plugins/Experimental/PCG/Source/PCG/Private/Tests/Elements/PCGStaticMeshSpawnerTest.cpp b/Engine/Plugins/Experimental/PCG/Source/PCG/Private/Tests/Elements/PCGStaticMeshSpawnerTest.cpp
--- a/Engine/Plugins/Experimental/PCG/Source/PCG/Private/Tests/Elements/PCGStaticMeshSpawnerTest.cpp
+++ b/Engine/Plugins/Experimental/PCG/Source/PCG/Private/Tests/Elements/PCGStaticMeshSpawnerTest.cpp
@@ -304,7 +304,7 @@ void TestMeshSelectorWeighted(
 		TObjectPtr<UStaticMesh> StaticMesh = ISMC->GetStaticMesh();
 
 		FPCGMeshSelectorWeightedEntry* Entry = MeshSelector->MeshEntries.FindByPredicate([StaticMesh](const FPCGMeshSelectorWeightedEntry& Entry) {
-			return Entry.Mesh == StaticMesh;
+			return Entry.Descriptor.StaticMesh == StaticMesh;
 			});
 
 		if (Test->TestNotNull("Validate instanced mesh exists in MeshEntries", Entry))
diff --git a/Engine/Plugins/Experimental/PCG/Source/PCG/Public/Helpers/PCGActorHelpers.h b/Engine/Plugins/Experimental/PCG/Source/PCG/Public/Helpers/PCGActorHelpers.h
--- a/Engine/Plugins/Experimental/PCG/Source/PCG/Public/Helpers/PCGActorHelpers.h
+++ b/Engine/Plugins/Experimental/PCG/Source/PCG/Public/Helpers/PCGActorHelpers.h
@@ -2,10 +2,10 @@
 
 #pragma once
 
+#include "Engine/EngineTypes.h"
+#include "ISMPartition/ISMComponentDescriptor.h"
 #include "Kismet/BlueprintFunctionLibrary.h"
 #include "Templates/SubclassOf.h"
-#include "Engine/EngineTypes.h"
-
 
 #include "PCGActorHelpers.generated.h"
 
@@ -18,19 +18,18 @@ class UPCGComponent;
 class UPCGManagedISMComponent;
 class UStaticMesh;
 class UWorld;
-struct FISMComponentDescriptor;
 
 struct FPCGISMCBuilderParameters
 {
-	UStaticMesh* Mesh = nullptr;
-	TArray<UMaterialInterface*> MaterialOverrides;
-	EComponentMobility::Type Mobility = EComponentMobility::Static;
-	FName CollisionProfile = TEXT("Default");
+	FISMComponentDescriptor Descriptor;
 	int32 NumCustomDataFloats = 0;
-	float CullStartDistance = 0;
-	float CullEndDistance = 0;
-	int32 WorldPositionOffsetDisableDistance = 0;
-	bool bIsLocalToWorldDeterminantNegative = false;
+
+	friend inline uint32 GetTypeHash(const FPCGISMCBuilderParameters& Key)
+	{
+		return HashCombine(GetTypeHash(Key.Descriptor), 1 + Key.NumCustomDataFloats);
+	}
+
+	inline bool operator==(const FPCGISMCBuilderParameters& Other) const { return Descriptor == Other.Descriptor && NumCustomDataFloats == Other.NumCustomDataFloats; }
 };
 
 UCLASS(BlueprintType)
@@ -39,7 +38,6 @@ class PCG_API UPCGActorHelpers : public UBlueprintFunctionLibrary
 	GENERATED_BODY()
 
 public:
-	static UPCGManagedISMComponent* GetOrCreateManagedISMC(AActor* InTargetActor, UPCGComponent* InSourceComponent, const FISMComponentDescriptor& InISMCDescriptor);
 	static UInstancedStaticMeshComponent* GetOrCreateISMC(AActor* InTargetActor, UPCGComponent* SourceComponent, const FPCGISMCBuilderParameters& Params);
 	static UPCGManagedISMComponent* GetOrCreateManagedISMC(AActor* InTargetActor, UPCGComponent* SourceComponent, const FPCGISMCBuilderParameters& Params);
 	static bool DeleteActors(UWorld* World, const TArray<TSoftObjectPtr<AActor>>& ActorsToDelete);
diff --git a/Engine/Plugins/Experimental/PCG/Source/PCG/Public/MeshSelectors/PCGMeshSelectorBase.h b/Engine/Plugins/Experimental/PCG/Source/PCG/Public/MeshSelectors/PCGMeshSelectorBase.h
--- a/Engine/Plugins/Experimental/PCG/Source/PCG/Public/MeshSelectors/PCGMeshSelectorBase.h
+++ b/Engine/Plugins/Experimental/PCG/Source/PCG/Public/MeshSelectors/PCGMeshSelectorBase.h
@@ -6,6 +6,7 @@
 #include "Metadata/PCGMetadata.h"
 
 #include "Engine/CollisionProfile.h"
+#include "ISMPartition/ISMComponentDescriptor.h"
 
 #include "PCGMeshSelectorBase.generated.h"
 
@@ -24,42 +25,18 @@ struct FPCGMeshInstanceList
 
 	FPCGMeshInstanceList() = default;
 
-	FPCGMeshInstanceList(const TSoftObjectPtr<UStaticMesh>& InMesh, bool bInOverrideCollisionProfile, const FCollisionProfileName& InCollisionProfile, bool bInOverrideMaterials, const TArray<TSoftObjectPtr<UMaterialInterface>>& InMaterialOverrides, const float InCullStartDistance, const float InCullEndDistance, const int32 InWorldPositionOffsetDisableDistance, const bool bInIsLocalToWorldDeterminantNegative)
-		: Mesh(InMesh), bOverrideCollisionProfile(bInOverrideCollisionProfile), CollisionProfile(InCollisionProfile), bOverrideMaterials(bInOverrideMaterials), MaterialOverrides(InMaterialOverrides), CullStartDistance(InCullStartDistance), CullEndDistance(InCullEndDistance), WorldPositionOffsetDisableDistance(InWorldPositionOffsetDisableDistance), bIsLocalToWorldDeterminantNegative(bInIsLocalToWorldDeterminantNegative)
+	explicit FPCGMeshInstanceList(const FSoftISMComponentDescriptor& InDescriptor)
+		: Descriptor(InDescriptor)
 	{}
 
-	UPROPERTY(BlueprintReadWrite, EditAnywhere, Category = Settings)
-	TSoftObjectPtr<UStaticMesh> Mesh;
-
-	UPROPERTY(BlueprintReadWrite, EditAnywhere, Category = Settings)
-	bool bOverrideCollisionProfile = false;
-
-	UPROPERTY(BlueprintReadWrite, EditAnywhere, Category = Settings)
-	FCollisionProfileName CollisionProfile = UCollisionProfile::NoCollision_ProfileName;
-
-	UPROPERTY(BlueprintReadWrite, EditAnywhere, Category = Settings)
-	bool bOverrideMaterials = false;
+	UPROPERTY(EditAnywhere, Category = Settings)
+	FSoftISMComponentDescriptor Descriptor;
 
 	UPROPERTY(BlueprintReadWrite, EditAnywhere, Category = Settings)
-	TArray<TSoftObjectPtr<UMaterialInterface>> MaterialOverrides;
+	TArray<FTransform> Instances;
 
 	UPROPERTY(BlueprintReadWrite, EditAnywhere, Category = Settings)
-    TArray<FPCGPoint> Instances;
-
-	/** Distance at which instances begin to fade. */
-	UPROPERTY(BlueprintReadWrite, EditAnywhere, Category = Settings)
-	float CullStartDistance = 0;
-	
-	/** Distance at which instances are culled. Use 0 to disable. */
-	UPROPERTY(BlueprintReadWrite, EditAnywhere, Category = Settings)
-	float CullEndDistance = 0;
-
-	UPROPERTY(BlueprintReadWrite, EditAnywhere, Category = Settings)
-	int32 WorldPositionOffsetDisableDistance = 0;
-
-	/** Whether the culling should be reversed or not (needed to support negative scales) */
-	UPROPERTY(BlueprintReadWrite, EditAnywhere, Category = Settings)
-	bool bIsLocalToWorldDeterminantNegative = false;
+	TArray<int64> InstancesMetadataEntry;
 };
 
 UENUM()
@@ -76,7 +53,7 @@ struct FPCGMeshMaterialOverrideHelper
 	// Use this constructor when you have a 1:1 mapping between attributes or static overrides
 	FPCGMeshMaterialOverrideHelper(
 		FPCGContext& InContext,
-		EPCGMeshSelectorMaterialOverrideMode InMaterialOverrideMode,
+		bool bUseMaterialOverrideAttributes,
 		const TArray<TSoftObjectPtr<UMaterialInterface>>& InStaticMaterialOverrides,
 		const TArray<FName>& InMaterialOverrideAttributeNames,
 		const UPCGMetadata* InMetadata);
@@ -89,7 +66,7 @@ struct FPCGMeshMaterialOverrideHelper
 		const UPCGMetadata* InMetadata);
 
 	bool IsValid() const { return bIsValid; }
-	bool OverridesMaterials() const { return MaterialOverrideMode != EPCGMeshSelectorMaterialOverrideMode::NoOverride; }
+	bool OverridesMaterials() const { return bUseMaterialOverrideAttributes; }
 	const TArray<TSoftObjectPtr<UMaterialInterface>>& GetMaterialOverrides(PCGMetadataEntryKey EntryKey);
 
 	// Cached data
@@ -100,7 +77,7 @@ struct FPCGMeshMaterialOverrideHelper
 
 	// Data needed to perform operations
 	bool bIsValid = false;
-	EPCGMeshSelectorMaterialOverrideMode MaterialOverrideMode = EPCGMeshSelectorMaterialOverrideMode::NoOverride;
+	bool bUseMaterialOverrideAttributes = false;
 
 	const TArray<TSoftObjectPtr<UMaterialInterface>>& StaticMaterialOverrides;
 	const TArray<FName>& MaterialOverrideAttributeNames;
@@ -130,20 +107,6 @@ public:
 		const UPCGPointData* InPointData,
 		TArray<FPCGMeshInstanceList>& OutMeshInstances,
 		UPCGPointData* OutPointData) const PURE_VIRTUAL(UPCGMeshSelectorBase::SelectInstances_Implementation);
-
-	/** Searches OutInstanceLists for an InstanceList matching the given parameters. If nothing is found, creates a new InstanceList and adds to OutInstanceLists. Returns index of the matching instance list. */
-	UFUNCTION(BlueprintCallable, Category = MeshSelection)
-	static int32 FindOrAddInstanceList(
-		TArray<FPCGMeshInstanceList>& OutInstanceLists,
-		const TSoftObjectPtr<UStaticMesh>& Mesh,
-		bool bOverrideCollisionProfile,
-		const FCollisionProfileName& CollisionProfile,
-		bool bOverrideMaterials,
-		const TArray<TSoftObjectPtr<UMaterialInterface>>& MaterialOverrides,
-		const float InCullStartDistance,
-		const float InCullEndDistance,
-		const int32 InWorldPositionOffsetDisableDistance,
-		const bool bInIsLocalToWorldDeterminantNegative);
 };
 
 #if UE_ENABLE_INCLUDE_ORDER_DEPRECATED_IN_5_2
diff --git a/Engine/Plugins/Experimental/PCG/Source/PCG/Public/MeshSelectors/PCGMeshSelectorByAttribute.h b/Engine/Plugins/Experimental/PCG/Source/PCG/Public/MeshSelectors/PCGMeshSelectorByAttribute.h
--- a/Engine/Plugins/Experimental/PCG/Source/PCG/Public/MeshSelectors/PCGMeshSelectorByAttribute.h
+++ b/Engine/Plugins/Experimental/PCG/Source/PCG/Public/MeshSelectors/PCGMeshSelectorByAttribute.h
@@ -26,36 +26,42 @@ public:
 
 public:
 	UPROPERTY(BlueprintReadWrite, EditAnywhere, Category = Settings)
-	FName AttributeName; 
+	FName AttributeName;
 
-	UPROPERTY(BlueprintReadWrite, EditAnywhere, Category = Settings)
-	bool bOverrideCollisionProfile = false;
+	UPROPERTY(EditAnywhere, Category = Settings)
+	FSoftISMComponentDescriptor TemplateDescriptor;
+
+	UPROPERTY(BlueprintReadWrite, EditAnywhere, Category = Settings, meta = (InlineEditConditionToggle))
+	bool bUseAttributeMaterialOverrides = false;
 
-	UPROPERTY(BlueprintReadWrite, EditAnywhere, Category = Settings, meta = (EditCondition = "bOverrideCollisionProfile", EditConditionHides))
-	FCollisionProfileName CollisionProfile = UCollisionProfile::NoCollision_ProfileName;
+	UPROPERTY(BlueprintReadWrite, EditAnywhere, DisplayName = "By Attribute Material Overrides", Category = Settings, meta = (EditCondition = "bUseAttributeMaterialOverrides"))
+	TArray<FName> MaterialOverrideAttributes;
 
 #if WITH_EDITORONLY_DATA
 	UPROPERTY()
-	bool bOverrideMaterials_DEPRECATED = false;
-#endif
+	bool bOverrideCollisionProfile_DEPRECATED = false;
 
-	UPROPERTY(BlueprintReadWrite, EditAnywhere, Category = Settings)
-	EPCGMeshSelectorMaterialOverrideMode MaterialOverrideMode = EPCGMeshSelectorMaterialOverrideMode::NoOverride;
+	UPROPERTY()
+	FCollisionProfileName CollisionProfile_DEPRECATED = UCollisionProfile::NoCollision_ProfileName;
 
-	UPROPERTY(BlueprintReadWrite, EditAnywhere, DisplayName = "Static Material Overrides", Category = Settings, meta = (EditCondition = "MaterialOverrideMode==EPCGMeshSelectorMaterialOverrideMode::StaticOverride", EditConditionHides))
-	TArray<TSoftObjectPtr<UMaterialInterface>> MaterialOverrides;
+	UPROPERTY()
+	EPCGMeshSelectorMaterialOverrideMode MaterialOverrideMode_DEPRECATED = EPCGMeshSelectorMaterialOverrideMode::NoOverride;
 
-	UPROPERTY(BlueprintReadWrite, EditAnywhere, DisplayName = "By Attribute Material Overrides", Category = Settings, meta = (EditCondition = "MaterialOverrideMode==EPCGMeshSelectorMaterialOverrideMode::ByAttributeOverride", EditConditionHides))
-	TArray<FName> MaterialOverrideAttributes;
+	UPROPERTY()
+	bool bOverrideMaterials_DEPRECATED = false;
+
+	UPROPERTY()
+	TArray<TSoftObjectPtr<UMaterialInterface>> MaterialOverrides_DEPRECATED;
 
 	/** Distance at which instances begin to fade. */
-	UPROPERTY(BlueprintReadWrite, EditAnywhere, Category = Settings)
-	float CullStartDistance = 0;
+	UPROPERTY()
+	float CullStartDistance_DEPRECATED = 0;
 	
 	/** Distance at which instances are culled. Use 0 to disable. */
-	UPROPERTY(BlueprintReadWrite, EditAnywhere, Category = Settings)
-	float CullEndDistance = 0;
+	UPROPERTY()
+	float CullEndDistance_DEPRECATED = 0;
 
-	UPROPERTY(BlueprintReadWrite, EditAnywhere, Category = Settings)
-	int32 WorldPositionOffsetDisableDistance = 0;
+	UPROPERTY()
+	int32 WorldPositionOffsetDisableDistance_DEPRECATED = 0;
+#endif // WITH_EDITORONLY_DATA
 };
diff --git a/Engine/Plugins/Experimental/PCG/Source/PCG/Public/MeshSelectors/PCGMeshSelectorWeighted.h b/Engine/Plugins/Experimental/PCG/Source/PCG/Public/MeshSelectors/PCGMeshSelectorWeighted.h
--- a/Engine/Plugins/Experimental/PCG/Source/PCG/Public/MeshSelectors/PCGMeshSelectorWeighted.h
+++ b/Engine/Plugins/Experimental/PCG/Source/PCG/Public/MeshSelectors/PCGMeshSelectorWeighted.h
@@ -8,9 +8,10 @@
 
 namespace PCGMeshSelectorWeighted
 {
+	// Returns variation matching the overrides & the reverse culling flag
 	FPCGMeshInstanceList& GetInstanceList(
 		TArray<FPCGMeshInstanceList>& InstanceLists,
-		bool bUseAttributeMaterialOverrides,
+		bool bUseMaterialOverrides,
 		const TArray<TSoftObjectPtr<UMaterialInterface>>& InMaterialOverrides,
 		bool bInIsLocalToWorldDeterminantNegative);
 }
@@ -21,39 +22,45 @@ struct PCG_API FPCGMeshSelectorWeightedEntry
 	GENERATED_BODY()
 
 	FPCGMeshSelectorWeightedEntry() = default;
+	FPCGMeshSelectorWeightedEntry(TSoftObjectPtr<UStaticMesh> InMesh, int InWeight);
 
-	FPCGMeshSelectorWeightedEntry(TSoftObjectPtr<UStaticMesh> InMesh, int InWeight)
-		: Mesh(InMesh), Weight(InWeight)
-	{}
+#if WITH_EDITOR
+	void ApplyDeprecation();
+#endif
 
-	UPROPERTY(BlueprintReadWrite, EditAnywhere, Category = Settings)
-	TSoftObjectPtr<UStaticMesh> Mesh;
+	UPROPERTY(EditAnywhere, Category = Settings)
+	FSoftISMComponentDescriptor Descriptor;
 
-	UPROPERTY(BlueprintReadWrite, EditAnywhere, Category = Settings)
-	bool bOverrideCollisionProfile = false;
+	UPROPERTY(BlueprintReadWrite, EditAnywhere, Category = Settings, meta = (ClampMin = "0"))
+	int Weight = 1;
 
-	UPROPERTY(BlueprintReadWrite, EditAnywhere, Category = Settings, meta = (EditCondition = "bOverrideCollisionProfile", EditConditionHides))
-	FCollisionProfileName CollisionProfile = UCollisionProfile::NoCollision_ProfileName;
+#if WITH_EDITORONLY_DATA
+	UPROPERTY()
+	TSoftObjectPtr<UStaticMesh> Mesh_DEPRECATED;
 
-	UPROPERTY(BlueprintReadWrite, EditAnywhere, Category = Settings, meta = (InlineEditConditionToggle))
-	bool bOverrideMaterials = false;
+	UPROPERTY()
+	bool bOverrideCollisionProfile_DEPRECATED = false;
+
+	UPROPERTY()
+	FCollisionProfileName CollisionProfile_DEPRECATED = UCollisionProfile::NoCollision_ProfileName;
+
+	UPROPERTY()
+	bool bOverrideMaterials_DEPRECATED = false;
 
-	UPROPERTY(BlueprintReadWrite, EditAnywhere, Category = Settings, DisplayName = "Static Material Overrides", meta = (EditCondition = "bOverrideMaterials"))
-	TArray<TSoftObjectPtr<UMaterialInterface>> MaterialOverrides;
+	UPROPERTY()
+	TArray<TSoftObjectPtr<UMaterialInterface>> MaterialOverrides_DEPRECATED;
 
 	/** Distance at which instances begin to fade. */
-	UPROPERTY(BlueprintReadWrite, EditAnywhere, Category = Settings)
-	float CullStartDistance = 0;
+	UPROPERTY()
+	float CullStartDistance_DEPRECATED = 0;
 	
 	/** Distance at which instances are culled. Use 0 to disable. */
-	UPROPERTY(BlueprintReadWrite, EditAnywhere, Category = Settings)
-	float CullEndDistance = 0;
-
-	UPROPERTY(BlueprintReadWrite, EditAnywhere, Category = Settings)
-	int32 WorldPositionOffsetDisableDistance = 0;
+	UPROPERTY()
+	float CullEndDistance_DEPRECATED = 0;
 
-	UPROPERTY(BlueprintReadWrite, EditAnywhere, Category = Settings, meta = (ClampMin = "0"))
-	int Weight = 1;
+	UPROPERTY()
+	int32 WorldPositionOffsetDisableDistance_DEPRECATED = 0;
+#endif
 };
 
 UCLASS(BlueprintType, ClassGroup = (Procedural))
@@ -69,6 +76,8 @@ public:
 		TArray<FPCGMeshInstanceList>& OutMeshInstances,
 		UPCGPointData* OutPointData) const override;
 
+	void PostLoad();
+
 public:
 	UPROPERTY(BlueprintReadWrite, EditAnywhere, Category = Settings)
 	TArray<FPCGMeshSelectorWeightedEntry> MeshEntries;
diff --git a/Engine/Plugins/Experimental/PCG/Source/PCG/Public/MeshSelectors/PCGMeshSelectorWeightedByCategory.h b/Engine/Plugins/Experimental/PCG/Source/PCG/Public/MeshSelectors/PCGMeshSelectorWeightedByCategory.h
--- a/Engine/Plugins/Experimental/PCG/Source/PCG/Public/MeshSelectors/PCGMeshSelectorWeightedByCategory.h
+++ b/Engine/Plugins/Experimental/PCG/Source/PCG/Public/MeshSelectors/PCGMeshSelectorWeightedByCategory.h
@@ -18,6 +18,10 @@ struct PCG_API FPCGWeightedByCategoryEntryList
 		: CategoryEntry(InCategoryEntry), WeightedMeshEntries(InWeightedMeshEntries)
 	{}
 
+#if WITH_EDITOR
+	void ApplyDeprecation();
+#endif
+
 	UPROPERTY(BlueprintReadWrite, EditAnywhere, Category = Settings)
 	FString CategoryEntry;
 
@@ -41,6 +45,8 @@ public:
 		TArray<FPCGMeshInstanceList>& OutMeshInstances,
 		UPCGPointData* OutPointData) const override;
 
+	void PostLoad();
+
 public:
 	UPROPERTY(BlueprintReadWrite, EditAnywhere, Category = Settings)
 	FName CategoryAttribute;
diff --git a/Engine/Plugins/Experimental/PCG/Source/PCG/Public/PCGManagedResource.h b/Engine/Plugins/Experimental/PCG/Source/PCG/Public/PCGManagedResource.h
--- a/Engine/Plugins/Experimental/PCG/Source/PCG/Public/PCGManagedResource.h
+++ b/Engine/Plugins/Experimental/PCG/Source/PCG/Public/PCGManagedResource.h
@@ -4,6 +4,8 @@
 
 #include "PCGCrc.h"
 
+#include "ISMPartition/ISMComponentDescriptor.h"
+
 #include "PCGManagedResource.generated.h"
 
 class UActorComponent;
@@ -92,7 +94,9 @@ public:
 	virtual bool SupportsComponentReset() const { return false; }
 	virtual void MarkAsUsed() override;
 	virtual void MarkAsReused() override;
+	virtual void ForgetComponent() { GeneratedComponent.Reset(); }
 
+public:
 	UPROPERTY(BlueprintReadWrite, EditAnywhere, Category = GeneratedData)
 	TSoftObjectPtr<UActorComponent> GeneratedComponent;
 };
@@ -103,16 +107,35 @@ class PCG_API UPCGManagedISMComponent : public UPCGManagedComponent
 	GENERATED_BODY()
 
 public:
+	//~Begin UObject interface
+	virtual void PostLoad() override;
+	//~End UObject interface
+
 	//~Begin UPCGManagedResource interface
 	virtual bool ReleaseIfUnused(TSet<TSoftObjectPtr<AActor>>& OutActorsToDelete) override;
 	//~End UPCGManagedResource interface
 
 	//~Begin UPCGManagedComponents interface
 	virtual void ResetComponent() override;
-	virtual bool SupportsComponentReset() const { return true; }
+	virtual bool SupportsComponentReset() const override{ return true; }
+	virtual void ForgetComponent() override;
 	//~End UPCGManagedComponents interface
 
 	UInstancedStaticMeshComponent* GetComponent() const;
+	void SetComponent(UInstancedStaticMeshComponent* InComponent);
+
+	void SetDescriptor(const FISMComponentDescriptor& InDescriptor);
+	const FISMComponentDescriptor& GetDescriptor() const { return Descriptor; }
+
+protected:
+	UPROPERTY()
+	bool bHasDescriptor = false;
+
+	UPROPERTY()
+	FISMComponentDescriptor Descriptor;
+
+	// Cached raw pointer to ISM component
+	mutable UInstancedStaticMeshComponent* CachedRawComponentPtr = nullptr;
 };
 
 #if UE_ENABLE_INCLUDE_ORDER_DEPRECATED_IN_5_2
