Fixed deprecation in static mesh selectors

#rnx

diff --git a/Engine/Plugins/Experimental/PCG/Source/PCG/Private/MeshSelectors/PCGMeshSelectorWeighted.cpp b/Engine/Plugins/Experimental/PCG/Source/PCG/Private/MeshSelectors/PCGMeshSelectorWeighted.cpp
--- a/Engine/Plugins/Experimental/PCG/Source/PCG/Private/MeshSelectors/PCGMeshSelectorWeighted.cpp
+++ b/Engine/Plugins/Experimental/PCG/Source/PCG/Private/MeshSelectors/PCGMeshSelectorWeighted.cpp
@@ -57,7 +57,7 @@ FPCGMeshSelectorWeightedEntry::FPCGMeshSelectorWeightedEntry(TSoftObjectPtr<USta
 #if WITH_EDITOR
 void FPCGMeshSelectorWeightedEntry::ApplyDeprecation()
 {
-	if (Mesh_DEPRECATED.IsValid() ||
+	if (!Mesh_DEPRECATED.IsNull() ||
 		bOverrideCollisionProfile_DEPRECATED ||
 		CollisionProfile_DEPRECATED.Name != UCollisionProfile::NoCollision_ProfileName ||
 		bOverrideMaterials_DEPRECATED ||
