[PCG]Various fixes for the attribute list view:
*Strings left aligned
*Index column always visible
*Point property column widths based on labels (clamped at 200)

diff --git a/Engine/Plugins/Experimental/PCG/Source/PCGEditor/Private/SPCGEditorGraphAttributeListView.cpp b/Engine/Plugins/Experimental/PCG/Source/PCGEditor/Private/SPCGEditorGraphAttributeListView.cpp
--- a/Engine/Plugins/Experimental/PCG/Source/PCGEditor/Private/SPCGEditorGraphAttributeListView.cpp
+++ b/Engine/Plugins/Experimental/PCG/Source/PCGEditor/Private/SPCGEditorGraphAttributeListView.cpp
@@ -83,12 +83,26 @@ namespace PCGEditorGraphAttributeListView
 	const FText TEXT_PointMetadataEntryLabel = LOCTEXT("PointMetadataEntryLabel", "Entry Key");
 	const FText TEXT_PointMetadataEntryParentLabel = LOCTEXT("PointMetadataEntryParentLabel", "Parent Key");
 
+	constexpr float MaxColumnWidth = 200.0f;
+
 	bool IsGraphCacheDebuggingEnabled()
 	{
 		UWorld* World = GEditor ? (GEditor->PlayWorld ? GEditor->PlayWorld.Get() : GEditor->GetEditorWorldContext().World()) : nullptr;
 		UPCGSubsystem* Subsystem = World ? World->GetSubsystem<UPCGSubsystem>() : nullptr;
 		return Subsystem && Subsystem->IsGraphCacheDebuggingEnabled();
 	}
+
+	float CalculateColumnWidth(const FText& InText)
+	{
+		check(FSlateApplication::Get().GetRenderer());
+		const TSharedRef<FSlateFontMeasure> FontMeasure = FSlateApplication::Get().GetRenderer()->GetFontMeasureService();
+		const FSlateFontInfo FontInfo = FAppStyle::GetFontStyle(TEXT("NormalText"));
+		
+		const float TextWidth = FontMeasure->Measure(InText, FontInfo).X;
+		constexpr float ColumnPadding = 12.0f; // TODO: Grab padding from header style
+		const float ColumnWidth = TextWidth + ColumnPadding;
+		return FMath::Min(ColumnWidth, MaxColumnWidth);
+	}
 }
 
 void SPCGListViewItemRow::Construct(const FArguments& InArgs, const TSharedRef<STableViewBase>& InOwnerTableView, const PCGListviewItemPtr& Item)
@@ -467,6 +481,7 @@ void SPCGEditorGraphAttributeListView::Construct(const FArguments& InArgs, TShar
 		SNew(SVerticalBox)
 		+SVerticalBox::Slot()
 		.AutoHeight()
+		.Padding(1.0f)
 		[
 			SNew(SHorizontalBox)
 			+SHorizontalBox::Slot()
@@ -786,7 +801,12 @@ TSharedRef<SWidget> SPCGEditorGraphAttributeListView::OnGenerateFilterMenu()
 	TArray<FName> HiddenColumns = ListViewHeader->GetHiddenColumnIds();
 
 	for (const SHeaderRow::FColumn& Column : Columns)
-	{	
+	{
+		if (Column.ColumnId == PCGEditorGraphAttributeListView::NAME_IndexColumn)
+		{
+			continue;
+		}
+
 		MenuBuilder.AddMenuEntry(
 			Column.DefaultText,
 			Column.DefaultTooltip,
@@ -923,8 +943,11 @@ void SPCGEditorGraphAttributeListView::ToggleAllAttributes()
 	{
 		const TIndirectArray<SHeaderRow::FColumn>& Columns = ListViewHeader->GetColumns();
 		for (const SHeaderRow::FColumn& Column : Columns)
-		{	
-			ListViewHeader->SetShowGeneratedColumn(Column.ColumnId, /*InShow=*/false);
+		{
+			if (Column.ColumnId != PCGEditorGraphAttributeListView::NAME_IndexColumn)
+			{
+				ListViewHeader->SetShowGeneratedColumn(Column.ColumnId, /*InShow=*/false);
+			}
 		}
 	}
 }
@@ -941,6 +964,11 @@ ECheckBoxState SPCGEditorGraphAttributeListView::GetAnyAttributeEnabledState() c
 	
 	for (const SHeaderRow::FColumn& Column : ListViewHeader->GetColumns())
 	{
+		if (Column.ColumnId == PCGEditorGraphAttributeListView::NAME_IndexColumn)
+		{
+			continue;
+		}
+
 		bAllEnabled &= Column.bIsVisible;
 		bAnyEnabled |= Column.bIsVisible;
 	}
@@ -980,8 +1008,10 @@ void SPCGEditorGraphAttributeListView::OnItemDoubleClicked(PCGListviewItemPtr It
 	}
 }
 
-void SPCGEditorGraphAttributeListView::AddColumn(const FName& InColumnId, const FText& ColumnLabel, float ColumnWidth, EHorizontalAlignment HeaderHAlign, EHorizontalAlignment CellHAlign)
+void SPCGEditorGraphAttributeListView::AddColumn(const FName& InColumnId, const FText& ColumnLabel, EHorizontalAlignment HeaderHAlign, EHorizontalAlignment CellHAlign)
 {
+	const float ColumnWidth = PCGEditorGraphAttributeListView::CalculateColumnWidth(ColumnLabel);
+
 	SHeaderRow::FColumn::FArguments Arguments;
 	Arguments.ColumnId(InColumnId);
 	Arguments.DefaultLabel(ColumnLabel);
@@ -1001,7 +1031,7 @@ void SPCGEditorGraphAttributeListView::RemoveColumn(const FName& InColumnId)
 
 void SPCGEditorGraphAttributeListView::AddIndexColumn()
 {
-	AddColumn(PCGEditorGraphAttributeListView::NAME_IndexColumn, PCGEditorGraphAttributeListView::TEXT_IndexLabel, 44);
+	AddColumn(PCGEditorGraphAttributeListView::NAME_IndexColumn, PCGEditorGraphAttributeListView::TEXT_IndexLabel);
 }
 
 void SPCGEditorGraphAttributeListView::RemoveIndexColumn()
@@ -1012,33 +1042,33 @@ void SPCGEditorGraphAttributeListView::RemoveIndexColumn()
 void SPCGEditorGraphAttributeListView::AddPointDataColumns()
 {
 	AddIndexColumn();
-	AddColumn(PCGEditorGraphAttributeListView::NAME_PointPositionX, PCGEditorGraphAttributeListView::TEXT_PointPositionLabelX, 94);
-	AddColumn(PCGEditorGraphAttributeListView::NAME_PointPositionY, PCGEditorGraphAttributeListView::TEXT_PointPositionLabelY, 94);
-	AddColumn(PCGEditorGraphAttributeListView::NAME_PointPositionZ, PCGEditorGraphAttributeListView::TEXT_PointPositionLabelZ, 94);
-	AddColumn(PCGEditorGraphAttributeListView::NAME_PointRotationX, PCGEditorGraphAttributeListView::TEXT_PointRotationLabelX, 68);
-	AddColumn(PCGEditorGraphAttributeListView::NAME_PointRotationY, PCGEditorGraphAttributeListView::TEXT_PointRotationLabelY, 68);
-	AddColumn(PCGEditorGraphAttributeListView::NAME_PointRotationZ, PCGEditorGraphAttributeListView::TEXT_PointRotationLabelZ, 68);
-	AddColumn(PCGEditorGraphAttributeListView::NAME_PointScaleX, PCGEditorGraphAttributeListView::TEXT_PointScaleLabelX, 50);
-	AddColumn(PCGEditorGraphAttributeListView::NAME_PointScaleY, PCGEditorGraphAttributeListView::TEXT_PointScaleLabelY, 50);
-	AddColumn(PCGEditorGraphAttributeListView::NAME_PointScaleZ, PCGEditorGraphAttributeListView::TEXT_PointScaleLabelZ, 50);
-	AddColumn(PCGEditorGraphAttributeListView::NAME_PointBoundsMinX, PCGEditorGraphAttributeListView::TEXT_PointBoundsLabelMinX, 80);
-	AddColumn(PCGEditorGraphAttributeListView::NAME_PointBoundsMinY, PCGEditorGraphAttributeListView::TEXT_PointBoundsLabelMinY, 80);
-	AddColumn(PCGEditorGraphAttributeListView::NAME_PointBoundsMinZ, PCGEditorGraphAttributeListView::TEXT_PointBoundsLabelMinZ, 80);
-	AddColumn(PCGEditorGraphAttributeListView::NAME_PointBoundsMaxX, PCGEditorGraphAttributeListView::TEXT_PointBoundsLabelMaxX, 88);
-	AddColumn(PCGEditorGraphAttributeListView::NAME_PointBoundsMaxY, PCGEditorGraphAttributeListView::TEXT_PointBoundsLabelMaxY, 88);
-	AddColumn(PCGEditorGraphAttributeListView::NAME_PointBoundsMaxZ, PCGEditorGraphAttributeListView::TEXT_PointBoundsLabelMaxZ, 88);
-	AddColumn(PCGEditorGraphAttributeListView::NAME_PointColorR, PCGEditorGraphAttributeListView::TEXT_PointColorLabelR, 50);
-	AddColumn(PCGEditorGraphAttributeListView::NAME_PointColorG, PCGEditorGraphAttributeListView::TEXT_PointColorLabelG, 50);
-	AddColumn(PCGEditorGraphAttributeListView::NAME_PointColorB, PCGEditorGraphAttributeListView::TEXT_PointColorLabelB, 50);
-	AddColumn(PCGEditorGraphAttributeListView::NAME_PointColorA, PCGEditorGraphAttributeListView::TEXT_PointColorLabelA, 50);
-	AddColumn(PCGEditorGraphAttributeListView::NAME_PointDensity, PCGEditorGraphAttributeListView::TEXT_PointDensityLabel, 54);
-	AddColumn(PCGEditorGraphAttributeListView::NAME_PointSteepness, PCGEditorGraphAttributeListView::TEXT_PointSteepnessLabel, 73);
-	AddColumn(PCGEditorGraphAttributeListView::NAME_PointSeed, PCGEditorGraphAttributeListView::TEXT_PointSeedLabel, 88);
+	AddColumn(PCGEditorGraphAttributeListView::NAME_PointPositionX, PCGEditorGraphAttributeListView::TEXT_PointPositionLabelX);
+	AddColumn(PCGEditorGraphAttributeListView::NAME_PointPositionY, PCGEditorGraphAttributeListView::TEXT_PointPositionLabelY);
+	AddColumn(PCGEditorGraphAttributeListView::NAME_PointPositionZ, PCGEditorGraphAttributeListView::TEXT_PointPositionLabelZ);
+	AddColumn(PCGEditorGraphAttributeListView::NAME_PointRotationX, PCGEditorGraphAttributeListView::TEXT_PointRotationLabelX);
+	AddColumn(PCGEditorGraphAttributeListView::NAME_PointRotationY, PCGEditorGraphAttributeListView::TEXT_PointRotationLabelY);
+	AddColumn(PCGEditorGraphAttributeListView::NAME_PointRotationZ, PCGEditorGraphAttributeListView::TEXT_PointRotationLabelZ);
+	AddColumn(PCGEditorGraphAttributeListView::NAME_PointScaleX, PCGEditorGraphAttributeListView::TEXT_PointScaleLabelX);
+	AddColumn(PCGEditorGraphAttributeListView::NAME_PointScaleY, PCGEditorGraphAttributeListView::TEXT_PointScaleLabelY);
+	AddColumn(PCGEditorGraphAttributeListView::NAME_PointScaleZ, PCGEditorGraphAttributeListView::TEXT_PointScaleLabelZ);
+	AddColumn(PCGEditorGraphAttributeListView::NAME_PointBoundsMinX, PCGEditorGraphAttributeListView::TEXT_PointBoundsLabelMinX);
+	AddColumn(PCGEditorGraphAttributeListView::NAME_PointBoundsMinY, PCGEditorGraphAttributeListView::TEXT_PointBoundsLabelMinY);
+	AddColumn(PCGEditorGraphAttributeListView::NAME_PointBoundsMinZ, PCGEditorGraphAttributeListView::TEXT_PointBoundsLabelMinZ);
+	AddColumn(PCGEditorGraphAttributeListView::NAME_PointBoundsMaxX, PCGEditorGraphAttributeListView::TEXT_PointBoundsLabelMaxX);
+	AddColumn(PCGEditorGraphAttributeListView::NAME_PointBoundsMaxY, PCGEditorGraphAttributeListView::TEXT_PointBoundsLabelMaxY);
+	AddColumn(PCGEditorGraphAttributeListView::NAME_PointBoundsMaxZ, PCGEditorGraphAttributeListView::TEXT_PointBoundsLabelMaxZ);
+	AddColumn(PCGEditorGraphAttributeListView::NAME_PointColorR, PCGEditorGraphAttributeListView::TEXT_PointColorLabelR);
+	AddColumn(PCGEditorGraphAttributeListView::NAME_PointColorG, PCGEditorGraphAttributeListView::TEXT_PointColorLabelG);
+	AddColumn(PCGEditorGraphAttributeListView::NAME_PointColorB, PCGEditorGraphAttributeListView::TEXT_PointColorLabelB);
+	AddColumn(PCGEditorGraphAttributeListView::NAME_PointColorA, PCGEditorGraphAttributeListView::TEXT_PointColorLabelA);
+	AddColumn(PCGEditorGraphAttributeListView::NAME_PointDensity, PCGEditorGraphAttributeListView::TEXT_PointDensityLabel);
+	AddColumn(PCGEditorGraphAttributeListView::NAME_PointSteepness, PCGEditorGraphAttributeListView::TEXT_PointSteepnessLabel);
+	AddColumn(PCGEditorGraphAttributeListView::NAME_PointSeed, PCGEditorGraphAttributeListView::TEXT_PointSeedLabel);
 
 	if (CVarShowAdvancedAttributesFields.GetValueOnAnyThread())
 	{
-		AddColumn(PCGEditorGraphAttributeListView::NAME_PointMetadataEntry, PCGEditorGraphAttributeListView::TEXT_PointMetadataEntryLabel, 50);
-		AddColumn(PCGEditorGraphAttributeListView::NAME_PointMetadataEntryParent, PCGEditorGraphAttributeListView::TEXT_PointMetadataEntryParentLabel, 50);
+		AddColumn(PCGEditorGraphAttributeListView::NAME_PointMetadataEntry, PCGEditorGraphAttributeListView::TEXT_PointMetadataEntryLabel);
+		AddColumn(PCGEditorGraphAttributeListView::NAME_PointMetadataEntryParent, PCGEditorGraphAttributeListView::TEXT_PointMetadataEntryParentLabel);
 	}	
 }
 
@@ -1088,26 +1118,26 @@ void SPCGEditorGraphAttributeListView::AddMetadataColumn(const FName& InColumnId
 	MetadataInfo.Index = InValueIndex;
 	MetadataInfos.Add(ColumnId, MetadataInfo);
 
-	check(FSlateApplication::Get().GetRenderer());
-	const TSharedRef<FSlateFontMeasure> FontMeasure = FSlateApplication::Get().GetRenderer()->GetFontMeasureService();
-	const FSlateFontInfo FontInfo = FAppStyle::GetFontStyle(TEXT("NormalText"));
-
-	const FText DefaultLabel = FText::FromName(ColumnId);
-	const float DefaultLabelWidth = FontMeasure->Measure(DefaultLabel, FontInfo).X;
+	const FText ColumnLabel = FText::FromName(ColumnId);
+	float ColumnWidth = 0.0f;
 
-	constexpr float ColumnPadding = 12.0f; // Todo: Grab padding from header style
-	float ColumnWidth = DefaultLabelWidth + ColumnPadding;
+	EHorizontalAlignment CellAlignment = EHorizontalAlignment::HAlign_Right;
 
 	if (InMetadataType == EPCGMetadataTypes::String)
 	{
-		ColumnWidth = FMath::Max(ColumnWidth, 200.0f);
+		ColumnWidth = PCGEditorGraphAttributeListView::MaxColumnWidth;
+		CellAlignment = EHorizontalAlignment::HAlign_Left;
 	}
-	
+	else
+	{
+		ColumnWidth = PCGEditorGraphAttributeListView::CalculateColumnWidth(ColumnLabel);
+	}
+
 	SHeaderRow::FColumn::FArguments ColumnArguments;
 	ColumnArguments.ColumnId(ColumnId);
-	ColumnArguments.DefaultLabel(DefaultLabel);
+	ColumnArguments.DefaultLabel(ColumnLabel);
 	ColumnArguments.HAlignHeader(EHorizontalAlignment::HAlign_Center);
-	ColumnArguments.HAlignCell(EHorizontalAlignment::HAlign_Right);
+	ColumnArguments.HAlignCell(CellAlignment);
 	ColumnArguments.ManualWidth(ColumnWidth);
 
 	SHeaderRow::FColumn* NewColumn = new SHeaderRow::FColumn(ColumnArguments);
diff --git a/Engine/Plugins/Experimental/PCG/Source/PCGEditor/Private/SPCGEditorGraphAttributeListView.h b/Engine/Plugins/Experimental/PCG/Source/PCGEditor/Private/SPCGEditorGraphAttributeListView.h
--- a/Engine/Plugins/Experimental/PCG/Source/PCGEditor/Private/SPCGEditorGraphAttributeListView.h
+++ b/Engine/Plugins/Experimental/PCG/Source/PCGEditor/Private/SPCGEditorGraphAttributeListView.h
@@ -94,7 +94,7 @@ private:
 	TSharedRef<ITableRow> OnGenerateRow(PCGListviewItemPtr Item, const TSharedRef<STableViewBase>& OwnerTable) const;
 	void OnItemDoubleClicked(PCGListviewItemPtr Item) const;
 
-	void AddColumn(const FName& InColumnId, const FText& ColumnLabel, float ColumnWidth, EHorizontalAlignment HeaderHAlign = HAlign_Center, EHorizontalAlignment CellHAlign = HAlign_Right);
+	void AddColumn(const FName& InColumnId, const FText& ColumnLabel, EHorizontalAlignment HeaderHAlign = HAlign_Center, EHorizontalAlignment CellHAlign = HAlign_Right);
 	void RemoveColumn(const FName& InColumnId);
 
 	void AddIndexColumn();
