[PCG] Lighten up Any colour a bit as it gets lost against default-white comment boxes.

diff --git a/Engine/Plugins/Experimental/PCG/Source/PCGEditor/Private/PCGEditorSettings.cpp b/Engine/Plugins/Experimental/PCG/Source/PCGEditor/Private/PCGEditorSettings.cpp
--- a/Engine/Plugins/Experimental/PCG/Source/PCGEditor/Private/PCGEditorSettings.cpp
+++ b/Engine/Plugins/Experimental/PCG/Source/PCGEditor/Private/PCGEditorSettings.cpp
@@ -23,7 +23,7 @@ UPCGEditorSettings::UPCGEditorSettings(const FObjectInitializer& ObjectInitializ
 	ParamDataNodeColor = FLinearColor(1.0f, 0.38f, 0.02f);
 	DebugNodeColor = FLinearColor(1.0f, 0.0f, 1.0f);
 
-	DefaultPinColor = FLinearColor(0.16f, 0.16f, 0.16f);
+	DefaultPinColor = FLinearColor(0.29f, 0.29f, 0.29f);
 
 	SpatialDataPinColor = FLinearColor(1.0f, 1.0f, 1.0f);
 	ConcreteDataPinColor = FLinearColor(0.45f, 0.38f, 0.96f);
