Fixed missing seed usage setup in match & set post load
On creation will now have a default weighted object
Cleanup in ownership when changing types (also in static mesh spawner)

#rnx

diff --git a/Engine/Plugins/Experimental/PCG/Source/PCG/Private/Elements/PCGPointMatchAndSet.cpp b/Engine/Plugins/Experimental/PCG/Source/PCG/Private/Elements/PCGPointMatchAndSet.cpp
--- a/Engine/Plugins/Experimental/PCG/Source/PCG/Private/Elements/PCGPointMatchAndSet.cpp
+++ b/Engine/Plugins/Experimental/PCG/Source/PCG/Private/Elements/PCGPointMatchAndSet.cpp
@@ -17,6 +17,12 @@
 UPCGPointMatchAndSetSettings::UPCGPointMatchAndSetSettings(const FObjectInitializer& ObjectInitializer)
 {
 	MatchAndSetType = UPCGMatchAndSetWeighted::StaticClass();
+
+	if (!this->HasAnyFlags(RF_ClassDefaultObject))
+	{
+		MatchAndSetInstance = ObjectInitializer.CreateDefaultSubobject<UPCGMatchAndSetWeighted>(this, TEXT("DefaultMatchAndSet"));
+	}
+
 	bUseSeed = MatchAndSetInstance && MatchAndSetInstance->UsesRandomProcess();
 }
 
@@ -62,6 +68,7 @@ void UPCGPointMatchAndSetSettings::PostLoad()
 	{
 		const EObjectFlags Flags = GetMaskedFlags(RF_PropagateToSubObjects) | RF_Transactional;
 		MatchAndSetInstance->SetFlags(Flags);
+		bUseSeed = MatchAndSetInstance->UsesRandomProcess();
 	}
 }
 
@@ -124,6 +131,14 @@ void UPCGPointMatchAndSetSettings::RefreshMatchAndSet()
 {
 	if (MatchAndSetType)
 	{
+		// Forget previous instance
+		if (MatchAndSetInstance)
+		{
+			MatchAndSetInstance->Rename(nullptr, GetTransientPackage(), REN_DontCreateRedirectors | REN_ForceNoResetLoaders);
+			MatchAndSetInstance->MarkAsGarbage();
+			MatchAndSetInstance = nullptr;
+		}
+
 		const EObjectFlags Flags = GetMaskedFlags(RF_PropagateToSubObjects);
 		MatchAndSetInstance = NewObject<UPCGMatchAndSetBase>(this, MatchAndSetType, NAME_None, Flags);
 		check(MatchAndSetInstance);
diff --git a/Engine/Plugins/Experimental/PCG/Source/PCG/Private/Elements/PCGStaticMeshSpawner.cpp b/Engine/Plugins/Experimental/PCG/Source/PCG/Private/Elements/PCGStaticMeshSpawner.cpp
--- a/Engine/Plugins/Experimental/PCG/Source/PCG/Private/Elements/PCGStaticMeshSpawner.cpp
+++ b/Engine/Plugins/Experimental/PCG/Source/PCG/Private/Elements/PCGStaticMeshSpawner.cpp
@@ -432,6 +432,13 @@ void UPCGStaticMeshSpawnerSettings::RefreshMeshSelector()
 {
 	if (MeshSelectorType)
 	{
+		if (MeshSelectorInstance)
+		{
+			MeshSelectorInstance->Rename(nullptr, GetTransientPackage(), REN_DontCreateRedirectors | REN_ForceNoResetLoaders);
+			MeshSelectorInstance->MarkAsGarbage();
+			MeshSelectorInstance = nullptr;
+		}
+
 		const EObjectFlags Flags = GetMaskedFlags(RF_PropagateToSubObjects);
 		MeshSelectorInstance = NewObject<UPCGMeshSelectorBase>(this, MeshSelectorType, NAME_None, Flags);
 	}
@@ -445,6 +452,13 @@ void UPCGStaticMeshSpawnerSettings::RefreshInstancePacker()
 {
 	if (InstancePackerType)
 	{
+		if (InstancePackerInstance)
+		{
+			InstancePackerInstance->Rename(nullptr, GetTransientPackage(), REN_DontCreateRedirectors | REN_ForceNoResetLoaders);
+			InstancePackerInstance->MarkAsGarbage();
+			InstancePackerInstance = nullptr;
+		}
+
 		const EObjectFlags Flags = GetMaskedFlags(RF_PropagateToSubObjects);
 		InstancePackerInstance = NewObject<UPCGInstancePackerBase>(this, InstancePackerType, NAME_None, Flags);
 	}
