Fixed axis retrieval bug in the normal to density native node - should not have used scale.

#rnx

diff --git a/Engine/Plugins/Experimental/PCG/Source/PCG/Private/Elements/PCGNormalToDensity.cpp b/Engine/Plugins/Experimental/PCG/Source/PCG/Private/Elements/PCGNormalToDensity.cpp
--- a/Engine/Plugins/Experimental/PCG/Source/PCG/Private/Elements/PCGNormalToDensity.cpp
+++ b/Engine/Plugins/Experimental/PCG/Source/PCG/Private/Elements/PCGNormalToDensity.cpp
@@ -36,7 +36,7 @@ bool FPCGNormalToDensityElement::ExecuteInternal(FPCGContext* Context) const
 	const UPCGNormalToDensitySettings* Settings = Context->GetInputSettings<UPCGNormalToDensitySettings>();
 	check(Settings);
 
-	const FVector Normal = Settings->Normal;
+	const FVector Normal = Settings->Normal.GetSafeNormal();
 	const double Offset = Settings->Offset;
 	const double Strength = Settings->Strength;
 	const PCGNormalToDensityMode DensityMode = Settings->DensityMode;
@@ -48,7 +48,7 @@ bool FPCGNormalToDensityElement::ExecuteInternal(FPCGContext* Context) const
 
 	const auto CalcValue = [Normal, Offset, InvStrength](const FPCGPoint& InPoint)
 	{
-		const FVector Up = InPoint.Transform.GetScaledAxis(EAxis::Z);
+		const FVector Up = InPoint.Transform.GetUnitAxis(EAxis::Z);
 		return FMath::Pow(FMath::Clamp(Up.Dot(Normal) + Offset, 0.0, 1.0), InvStrength);
 	};
 
