Normalized files across branches

#rnx

diff --git a/Engine/Plugins/Experimental/PCG/Source/PCG/Private/Elements/PCGNormalToDensity.cpp b/Engine/Plugins/Experimental/PCG/Source/PCG/Private/Elements/PCGNormalToDensity.cpp
--- a/Engine/Plugins/Experimental/PCG/Source/PCG/Private/Elements/PCGNormalToDensity.cpp
+++ b/Engine/Plugins/Experimental/PCG/Source/PCG/Private/Elements/PCGNormalToDensity.cpp
@@ -7,23 +7,6 @@
 #include "PCGContext.h"
 #include "PCGPin.h"
 
-TArray<FPCGPinProperties> UPCGNormalToDensitySettings::InputPinProperties() const
-{
-	TArray<FPCGPinProperties> PinProperties;
-	PinProperties.Emplace(PCGPinConstants::DefaultInputLabel, EPCGDataType::Spatial);
-	PinProperties.Emplace(PCGPinConstants::DefaultParamsLabel, EPCGDataType::Param, /*bInAllowMultipleConnections*/ false);
-
-	return PinProperties;
-}
-
-TArray<FPCGPinProperties> UPCGNormalToDensitySettings::OutputPinProperties() const
-{
-	TArray<FPCGPinProperties> PinProperties;
-	PinProperties.Emplace(PCGPinConstants::DefaultOutputLabel, EPCGDataType::Spatial);
-
-	return PinProperties;
-}
-
 FPCGElementPtr UPCGNormalToDensitySettings::CreateElement() const
 {
 	return MakeShared<FPCGNormalToDensityElement>();
diff --git a/Engine/Plugins/Experimental/PCG/Source/PCG/Public/Elements/PCGNormalToDensity.h b/Engine/Plugins/Experimental/PCG/Source/PCG/Public/Elements/PCGNormalToDensity.h
--- a/Engine/Plugins/Experimental/PCG/Source/PCG/Public/Elements/PCGNormalToDensity.h
+++ b/Engine/Plugins/Experimental/PCG/Source/PCG/Public/Elements/PCGNormalToDensity.h
@@ -34,10 +34,10 @@ public:
 	virtual EPCGSettingsType GetType() const override { return EPCGSettingsType::Spatial; }
 #endif
 
-	virtual TArray<FPCGPinProperties> InputPinProperties() const override;
-	virtual TArray<FPCGPinProperties> OutputPinProperties() const override;
-
 protected:
+	virtual TArray<FPCGPinProperties> InputPinProperties() const override { return Super::DefaultPointInputPinProperties(); }
+	virtual TArray<FPCGPinProperties> OutputPinProperties() const override { return Super::DefaultPointOutputPinProperties(); }
+
 	virtual FPCGElementPtr CreateElement() const override;
 	//~End UPCGSettings interface
 
