Improved API for world ray hit/volumetric
Added precise check option.

diff --git a/Engine/Plugins/Experimental/PCG/Source/PCG/Private/Data/PCGWorldData.cpp b/Engine/Plugins/Experimental/PCG/Source/PCG/Private/Data/PCGWorldData.cpp
--- a/Engine/Plugins/Experimental/PCG/Source/PCG/Private/Data/PCGWorldData.cpp
+++ b/Engine/Plugins/Experimental/PCG/Source/PCG/Private/Data/PCGWorldData.cpp
@@ -33,6 +33,7 @@ bool UPCGWorldVolumetricData::SamplePoint(const FTransform& InTransform, const F
 	FCollisionObjectQueryParams ObjectQueryParams(QueryParams.CollisionChannel);
 	FCollisionShape CollisionShape = FCollisionShape::MakeBox(InBounds.GetExtent() * InTransform.GetScale3D());
 	FCollisionQueryParams Params; // TODO: apply properties from the settings when/if they exist
+	Params.bTraceComplex = QueryParams.bTraceComplex;
 
 	TArray<FOverlapResult> Overlaps;
 	/*bool bOverlaps =*/ World->OverlapMultiByObjectType(Overlaps, InTransform.TransformPosition(InBounds.GetCenter()), InTransform.GetRotation(), ObjectQueryParams, CollisionShape, Params);
@@ -169,6 +170,7 @@ bool UPCGWorldRayHitData::SamplePoint(const FTransform& InTransform, const FBox&
 	// Todo: consider prebuilding this
 	FCollisionObjectQueryParams ObjectQueryParams(QueryParams.CollisionChannel);
 	FCollisionQueryParams Params; // TODO: apply properties from the settings when/if they exist
+	Params.bTraceComplex = QueryParams.bTraceComplex;
 
 	// Project the InTransform location on the ray origin plane
 	const FVector PointLocation = InTransform.GetLocation();
@@ -257,7 +259,7 @@ const UPCGPointData* UPCGWorldRayHitData::CreatePointData(FPCGContext* Context,
 	{
 		if (!Bounds.IsValid && !InBounds.IsValid)
 		{
-			UE_LOG(LogPCG, Error, TEXT("PCG World Volumetric Data cannot generate if there are no framing bounds"));
+			UE_LOG(LogPCG, Error, TEXT("PCG World Ray Hit Data cannot generate if there are no framing bounds"));
 		}
 
 		return Data;
diff --git a/Engine/Plugins/Experimental/PCG/Source/PCG/Private/Elements/PCGSurfaceSampler.cpp b/Engine/Plugins/Experimental/PCG/Source/PCG/Private/Elements/PCGSurfaceSampler.cpp
--- a/Engine/Plugins/Experimental/PCG/Source/PCG/Private/Elements/PCGSurfaceSampler.cpp
+++ b/Engine/Plugins/Experimental/PCG/Source/PCG/Private/Elements/PCGSurfaceSampler.cpp
@@ -406,10 +406,20 @@ bool FPCGSurfaceSamplerElement::ExecuteInternal(FPCGContext* Context) const
 		check(GeneratingShape);
 
 		// Calculate the intersection of bounds of the provided inputs
-		FBox InputBounds = GeneratingShape->GetBounds();
-		if (BoundingShapeBounds.IsValid)
+		FBox InputBounds = FBox(EForceInit::ForceInit);
+
+		if (GeneratingShape->IsBounded())
+		{
+			InputBounds = GeneratingShape->GetBounds();
+
+			if (BoundingShapeBounds.IsValid)
+			{
+				InputBounds = PCGHelpers::OverlapBounds(InputBounds, BoundingShapeBounds);
+			}
+		}
+		else
 		{
-			InputBounds = PCGHelpers::OverlapBounds(InputBounds, BoundingShapeBounds);
+			InputBounds = BoundingShapeBounds;
 		}
 
 		PCGSurfaceSampler::FSurfaceSamplerSettings LoopData;
diff --git a/Engine/Plugins/Experimental/PCG/Source/PCG/Private/Elements/PCGVolumeSampler.cpp b/Engine/Plugins/Experimental/PCG/Source/PCG/Private/Elements/PCGVolumeSampler.cpp
--- a/Engine/Plugins/Experimental/PCG/Source/PCG/Private/Elements/PCGVolumeSampler.cpp
+++ b/Engine/Plugins/Experimental/PCG/Source/PCG/Private/Elements/PCGVolumeSampler.cpp
@@ -229,10 +229,19 @@ bool FPCGVolumeSamplerElement::ExecuteInternal(FPCGContext* Context) const
 		check(GeneratingShape);
 
 		// Calculate the intersection of bounds of the provided inputs
-		FBox InputBounds = GeneratingShape->GetBounds();
-		if (BoundingShapeBounds.IsValid)
+		FBox InputBounds = FBox(EForceInit::ForceInit);
+		if (GeneratingShape->IsBounded())
 		{
-			InputBounds = PCGHelpers::OverlapBounds(InputBounds, BoundingShapeBounds);
+			InputBounds = GeneratingShape->GetBounds();
+
+			if (BoundingShapeBounds.IsValid)
+			{
+				InputBounds = PCGHelpers::OverlapBounds(InputBounds, BoundingShapeBounds);
+			}
+		}
+		else
+		{
+			InputBounds = BoundingShapeBounds;
 		}
 
 		if (!InputBounds.IsValid)
diff --git a/Engine/Plugins/Experimental/PCG/Source/PCG/Private/Elements/PCGWorldQuery.cpp b/Engine/Plugins/Experimental/PCG/Source/PCG/Private/Elements/PCGWorldQuery.cpp
--- a/Engine/Plugins/Experimental/PCG/Source/PCG/Private/Elements/PCGWorldQuery.cpp
+++ b/Engine/Plugins/Experimental/PCG/Source/PCG/Private/Elements/PCGWorldQuery.cpp
@@ -45,11 +45,9 @@ bool FPCGWorldVolumetricQueryElement::ExecuteInternal(FPCGContext* Context) cons
 
 	check(Context->SourceComponent.IsValid());
 	UWorld* World = Context->SourceComponent->GetWorld();
-	// TODO: Not strictly required but will be useful until we change the usage of this
-	FBox SourceBounds = PCGHelpers::GetActorBounds(Context->SourceComponent->GetOwner());
 
 	UPCGWorldVolumetricData* Data = NewObject<UPCGWorldVolumetricData>();
-	Data->Initialize(World, SourceBounds);
+	Data->Initialize(World);
 	Data->QueryParams = QueryParams;
 	Data->OriginatingComponent = Context->SourceComponent;
 	Data->TargetActor = Context->SourceComponent->GetOwner();
@@ -124,11 +122,9 @@ bool FPCGWorldRayHitQueryElement::ExecuteInternal(FPCGContext* Context) const
 
 	check(Context->SourceComponent.IsValid());
 	UWorld* World = Context->SourceComponent->GetWorld();
-	// TODO: Not strictly required but will be useful until we change the usage of this
-	FBox SourceBounds = PCGHelpers::GetActorBounds(Context->SourceComponent->GetOwner());
 
 	UPCGWorldRayHitData* Data = NewObject<UPCGWorldRayHitData>();
-	Data->Initialize(World, SourceBounds);
+	Data->Initialize(World);
 	Data->QueryParams = QueryParams;
 	Data->OriginatingComponent = Context->SourceComponent;
 	Data->TargetActor = Context->SourceComponent->GetOwner();
diff --git a/Engine/Plugins/Experimental/PCG/Source/PCG/Public/Data/PCGSpatialData.h b/Engine/Plugins/Experimental/PCG/Source/PCG/Public/Data/PCGSpatialData.h
--- a/Engine/Plugins/Experimental/PCG/Source/PCG/Public/Data/PCGSpatialData.h
+++ b/Engine/Plugins/Experimental/PCG/Source/PCG/Public/Data/PCGSpatialData.h
@@ -60,6 +60,9 @@ public:
 	UFUNCTION(BlueprintCallable, Category = SpatialData)
 	virtual FBox GetBounds() const PURE_VIRTUAL(UPCGSpatialData::GetBounds, return FBox(EForceInit::ForceInit););
 
+	/** Returns whether a given spatial data is bounded as some data types do not require bounds by themselves */
+	virtual bool IsBounded() const { return true; }
+
 	/** Returns the bounds in which the density is always 1 */
 	UFUNCTION(BlueprintCallable, Category = SpatialData)
 	virtual FBox GetStrictBounds() const { return FBox(EForceInit::ForceInit); }
diff --git a/Engine/Plugins/Experimental/PCG/Source/PCG/Public/Data/PCGWorldData.h b/Engine/Plugins/Experimental/PCG/Source/PCG/Public/Data/PCGWorldData.h
--- a/Engine/Plugins/Experimental/PCG/Source/PCG/Public/Data/PCGWorldData.h
+++ b/Engine/Plugins/Experimental/PCG/Source/PCG/Public/Data/PCGWorldData.h
@@ -31,6 +31,10 @@ struct FPCGWorldVolumetricQueryParams
 
 	UPROPERTY(BlueprintReadWrite, EditAnywhere, Category = "Data|Advanced")
 	TEnumAsByte<ECollisionChannel> CollisionChannel = ECC_WorldStatic;
+
+	/** Queries against complex collision if enabled, performance warning */
+	UPROPERTY(BlueprintReadWrite, EditAnywhere, Category = "Data|Advanced")
+	bool bTraceComplex = false;
 };
 
 /** Queries volume for presence of world collision or not. Can be used to voxelize environment. */
@@ -43,6 +47,7 @@ public:
 	void Initialize(UWorld* InWorld, const FBox& InBounds = FBox(EForceInit::ForceInit));
 
 	//~Begin UPCGSpatialData interface
+	virtual bool IsBounded() const override { return !!Bounds.IsValid; }
 	virtual bool SamplePoint(const FTransform& Transform, const FBox& Bounds, FPCGPoint& OutPoint, UPCGMetadata* OutMetadata) const override;
 	// TODO not sure what this would mean. Without a direction, this means perhaps finding closest point on any collision surface? Should we implement this disabled?
 	//virtual bool ProjectPoint(const FTransform& InTransform, const FBox& InBounds, const FPCGProjectionParams& InParams, FPCGPoint& OutPoint, UPCGMetadata* OutMetadata) const;
@@ -93,8 +98,12 @@ struct FPCGWorldRayHitQueryParams
 	UPROPERTY(BlueprintReadWrite, EditAnywhere, Category = "Data|Advanced")
 	TEnumAsByte<ECollisionChannel> CollisionChannel = ECC_WorldStatic;
 
+	/** Queries against complex collision if enabled, performance warning */
+	UPROPERTY(BlueprintReadWrite, EditAnywhere, Category = "Data|Advanced")
+	bool bTraceComplex = false;
+
 	// TODO: see in FCollisionQueryParams if there are some flags we want to expose
-	// examples: bTraceComplex, bReturnFaceIndex, bReturnPhysicalMaterial, some ignore patterns
+	// examples: bReturnFaceIndex, bReturnPhysicalMaterial, some ignore patterns
 
 	//TODO UPROPERTY()
 	//bool bUseMetadataFromLandscape = true;
@@ -112,6 +121,7 @@ public:
 	//~Begin UPCGSpatialData interface
 	virtual FBox GetBounds() const override { return Bounds; }
 	virtual FBox GetStrictBounds() const override { return Bounds; }
+	virtual bool IsBounded() const override { return !!Bounds.IsValid; }
 	virtual bool SamplePoint(const FTransform& Transform, const FBox& Bounds, FPCGPoint& OutPoint, UPCGMetadata* OutMetadata) const override;
 	virtual bool HasNonTrivialTransform() const override { return true; }
 protected:
