[PCG] Dynamic pin types for CreateAttribute node

Behaviour modes:

* No input edges - input pin type is Any inviting user to connect anything. Output type is param.
* Input edge of type param - input pin type is param, output pin type is param.
* Input edge of non-param type(s) - type is Union of all incident types.

Does not break existing connections due to current lenient type checking.

diff --git a/Engine/Plugins/Experimental/PCG/Source/PCG/Private/Elements/PCGCreateAttribute.cpp b/Engine/Plugins/Experimental/PCG/Source/PCG/Private/Elements/PCGCreateAttribute.cpp
--- a/Engine/Plugins/Experimental/PCG/Source/PCG/Private/Elements/PCGCreateAttribute.cpp
+++ b/Engine/Plugins/Experimental/PCG/Source/PCG/Private/Elements/PCGCreateAttribute.cpp
@@ -76,11 +76,19 @@ FName UPCGCreateAttributeSettings::GetDefaultNodeName() const
 TArray<FPCGPinProperties> UPCGCreateAttributeSettings::InputPinProperties() const
 {
 	TArray<FPCGPinProperties> PinProperties;
-	PinProperties.Emplace(PCGPinConstants::DefaultInputLabel, EPCGDataType::Any, /*bInAllowMultipleConnections=*/ true);
+
+	EPCGDataType PrimaryInputType = GetTypeUnionOfIncidentEdges(PCGPinConstants::DefaultInputLabel);
+	if (PrimaryInputType == EPCGDataType::None)
+	{
+		// Fall back to Any if no edges, or no overlap in types (which should not normally occur).
+		PrimaryInputType = EPCGDataType::Any;
+	}
+
+	PinProperties.Emplace(PCGPinConstants::DefaultInputLabel, PrimaryInputType);
 
 	if (bFromSourceParam)
 	{
-		PinProperties.Emplace(PCGCreateAttributeConstants::SourceLabel, EPCGDataType::Param, /*bInAllowMultipleConnections=*/ false);
+		PinProperties.Emplace(PCGCreateAttributeConstants::SourceLabel, EPCGDataType::Param);
 	}
 
 	return PinProperties;
@@ -88,10 +96,17 @@ TArray<FPCGPinProperties> UPCGCreateAttributeSettings::InputPinProperties() cons
 
 TArray<FPCGPinProperties> UPCGCreateAttributeSettings::OutputPinProperties() const
 {
-	TArray<FPCGPinProperties> PinProperties;
-	PinProperties.Emplace(PCGPinConstants::DefaultOutputLabel, EPCGDataType::Any);
+	FPCGPinProperties PinProperties;
+	PinProperties.Label = PCGPinConstants::DefaultOutputLabel;
 
-	return PinProperties;
+	PinProperties.AllowedTypes = GetTypeUnionOfIncidentEdges(PCGPinConstants::DefaultInputLabel);
+	if (PinProperties.AllowedTypes == EPCGDataType::None)
+	{
+		// Nothing connected or incompatible types - output Param.
+		PinProperties.AllowedTypes = EPCGDataType::Param;
+	}
+
+	return { PinProperties };
 }
 
 FPCGElementPtr UPCGCreateAttributeSettings::CreateElement() const
@@ -139,11 +154,17 @@ bool FPCGCreateAttributeElement::ExecuteInternal(FPCGContext* Context) const
 
 	TArray<FPCGTaggedData> Inputs = Context->InputData.GetInputsByPin(PCGPinConstants::DefaultInputLabel);
 
+	// Only if we have no input connections we'll switch from "adding attribute" to "creating param data"
+	const UPCGPin* InputPin = Context->Node ? Context->Node->GetInputPin(PCGPinConstants::DefaultInputLabel) : nullptr;
+	const bool bHasInputConnections = InputPin && InputPin->EdgeCount() > 0;
+
 	// If the input is empty, we will create a new ParamData.
 	// We can re-use this newly object as the output
 	bool bCanReuseInputData = false;
-	if (Inputs.IsEmpty())
+	if (!bHasInputConnections)
 	{
+		ensure(Inputs.IsEmpty());
+
 		FPCGTaggedData& NewData = Inputs.Emplace_GetRef();
 		NewData.Data = NewObject<UPCGParamData>();
 		NewData.Pin = PCGPinConstants::DefaultInputLabel;
diff --git a/Engine/Plugins/Experimental/PCG/Source/PCG/Private/PCGSettings.cpp b/Engine/Plugins/Experimental/PCG/Source/PCG/Private/PCGSettings.cpp
--- a/Engine/Plugins/Experimental/PCG/Source/PCG/Private/PCGSettings.cpp
+++ b/Engine/Plugins/Experimental/PCG/Source/PCG/Private/PCGSettings.cpp
@@ -5,6 +5,7 @@
 #include "PCGComponent.h"
 #include "PCGContext.h"
 #include "PCGCustomVersion.h"
+#include "PCGEdge.h"
 #include "PCGGraph.h"
 #include "PCGHelpers.h"
 #include "PCGModule.h"
@@ -325,6 +326,27 @@ TArray<FPCGPinProperties> UPCGSettings::AllInputPinProperties() const
 	return InputPins;
 }
 
+EPCGDataType UPCGSettings::GetTypeUnionOfIncidentEdges(const FName& PinLabel) const
+{
+	EPCGDataType Result = EPCGDataType::None;
+
+	const UPCGNode* Node = Cast<UPCGNode>(GetOuter());
+	const UPCGPin* Pin = Node ? Node->GetInputPin(PinLabel) : nullptr;
+	if (Pin && Pin->EdgeCount() > 0)
+	{
+		for (const UPCGEdge* Edge : Pin->Edges)
+		{
+			const UPCGPin* OtherOutputPin = Edge ? Edge->GetOtherPin(Pin) : nullptr;
+			if (OtherOutputPin)
+			{
+				Result |= OtherOutputPin->Properties.AllowedTypes;
+			}
+		}
+	}
+
+	return Result;
+}
+
 TArray<FPCGPinProperties> UPCGSettings::AllOutputPinProperties() const
 {
 	return OutputPinProperties();
diff --git a/Engine/Plugins/Experimental/PCG/Source/PCG/Public/Elements/PCGCreateAttribute.h b/Engine/Plugins/Experimental/PCG/Source/PCG/Public/Elements/PCGCreateAttribute.h
--- a/Engine/Plugins/Experimental/PCG/Source/PCG/Public/Elements/PCGCreateAttribute.h
+++ b/Engine/Plugins/Experimental/PCG/Source/PCG/Public/Elements/PCGCreateAttribute.h
@@ -32,6 +32,7 @@ public:
 #if WITH_EDITOR
 	virtual FName GetDefaultNodeName() const override;
 	virtual EPCGSettingsType GetType() const override { return EPCGSettingsType::Metadata; }
+	virtual bool HasDynamicPins() const override { return true; }
 #endif // WITH_EDITOR
 
 	virtual TArray<FPCGPinProperties> InputPinProperties() const override;
diff --git a/Engine/Plugins/Experimental/PCG/Source/PCG/Public/PCGSettings.h b/Engine/Plugins/Experimental/PCG/Source/PCG/Public/PCGSettings.h
--- a/Engine/Plugins/Experimental/PCG/Source/PCG/Public/PCGSettings.h
+++ b/Engine/Plugins/Experimental/PCG/Source/PCG/Public/PCGSettings.h
@@ -143,6 +143,9 @@ public:
 	/** If the node has any dynamic pins that can change based on input or settings */
 	virtual bool HasDynamicPins() const { return false; }
 
+	/** Bitwise union of the allowed types of each incident edge on pin. Returns None if no common bits, or no edges. */
+	EPCGDataType GetTypeUnionOfIncidentEdges(const FName& PinLabel) const;
+
 	// Internal functions, should not be used by any user.
 	// Return a different subset for for input/output pin properties, in case of a default object.
 	virtual TArray<FPCGPinProperties> DefaultInputPinProperties() const;
