[PCG] Introduce PCGSubsystem::GetInstance to get the Subsystem. It will also validate that the subsystem was initialized.

* Should fix issue when we require a Refresh in PostLoad, when the subsystem is not yet ready.

diff --git a/Engine/Plugins/Experimental/PCG/Source/PCG/Private/Data/PCGLandscapeData.cpp b/Engine/Plugins/Experimental/PCG/Source/PCG/Private/Data/PCGLandscapeData.cpp
--- a/Engine/Plugins/Experimental/PCG/Source/PCG/Private/Data/PCGLandscapeData.cpp
+++ b/Engine/Plugins/Experimental/PCG/Source/PCG/Private/Data/PCGLandscapeData.cpp
@@ -39,7 +39,8 @@ void UPCGLandscapeData::Initialize(const TArray<TWeakObjectPtr<ALandscapeProxy>>
 	Transform = FirstLandscape->GetActorTransform();
 
 	// Store cache pointer for easier access
-	LandscapeCache = (FirstLandscape->GetWorld() && FirstLandscape->GetWorld()->GetSubsystem<UPCGSubsystem>()) ? FirstLandscape->GetWorld()->GetSubsystem<UPCGSubsystem>()->GetLandscapeCache() : nullptr;
+	UPCGSubsystem* PCGSubsystem = UPCGSubsystem::GetInstance(FirstLandscape->GetWorld());
+	LandscapeCache = PCGSubsystem ? PCGSubsystem->GetLandscapeCache() : nullptr;
 
 	// TODO: find a better way to do this - maybe there should be a prototype metadata in the landscape cache
 	if (LandscapeCache)
@@ -91,7 +92,8 @@ void UPCGLandscapeData::PostLoad()
 	}
 
 	TargetActor = FirstLandscape;
-	LandscapeCache = (FirstLandscape && FirstLandscape->GetWorld() && FirstLandscape->GetWorld()->GetSubsystem<UPCGSubsystem>()) ? FirstLandscape->GetWorld()->GetSubsystem<UPCGSubsystem>()->GetLandscapeCache() : nullptr;
+	UPCGSubsystem* PCGSubsystem = UPCGSubsystem::GetInstance(FirstLandscape->GetWorld());
+	LandscapeCache = PCGSubsystem ? PCGSubsystem->GetLandscapeCache() : nullptr;
 }
 
 FBox UPCGLandscapeData::GetBounds() const
diff --git a/Engine/Plugins/Experimental/PCG/Source/PCG/Private/Grid/PCGPartitionActor.cpp b/Engine/Plugins/Experimental/PCG/Source/PCG/Private/Grid/PCGPartitionActor.cpp
--- a/Engine/Plugins/Experimental/PCG/Source/PCG/Private/Grid/PCGPartitionActor.cpp
+++ b/Engine/Plugins/Experimental/PCG/Source/PCG/Private/Grid/PCGPartitionActor.cpp
@@ -34,7 +34,7 @@ APCGPartitionActor::APCGPartitionActor(const FObjectInitializer& ObjectInitializ
 
 UPCGSubsystem* APCGPartitionActor::GetSubsystem() const 
 {
-	return GetWorld() ? GetWorld()->GetSubsystem<UPCGSubsystem>() : nullptr;
+	return UPCGSubsystem::GetInstance(GetWorld());
 }
 
 void APCGPartitionActor::PostLoad()
diff --git a/Engine/Plugins/Experimental/PCG/Source/PCG/Private/Helpers/PCGBlueprintHelpers.cpp b/Engine/Plugins/Experimental/PCG/Source/PCG/Private/Helpers/PCGBlueprintHelpers.cpp
--- a/Engine/Plugins/Experimental/PCG/Source/PCG/Private/Helpers/PCGBlueprintHelpers.cpp
+++ b/Engine/Plugins/Experimental/PCG/Source/PCG/Private/Helpers/PCGBlueprintHelpers.cpp
@@ -132,7 +132,7 @@ TArray<FPCGLandscapeLayerWeight> UPCGBlueprintHelpers::GetInterpolatedPCGLandsca
 		return {};
 	}
 
-	UPCGSubsystem* PCGSubSystem = UWorld::GetSubsystem<UPCGSubsystem>(World);
+	UPCGSubsystem* PCGSubSystem = UPCGSubsystem::GetInstance(World);
 	if (!PCGSubSystem)
 	{
 		return {};
diff --git a/Engine/Plugins/Experimental/PCG/Source/PCG/Private/PCGComponent.cpp b/Engine/Plugins/Experimental/PCG/Source/PCG/Private/PCGComponent.cpp
--- a/Engine/Plugins/Experimental/PCG/Source/PCG/Private/PCGComponent.cpp
+++ b/Engine/Plugins/Experimental/PCG/Source/PCG/Private/PCGComponent.cpp
@@ -2442,7 +2442,7 @@ FBox UPCGComponent::GetGridBounds(const AActor* Actor) const
 
 UPCGSubsystem* UPCGComponent::GetSubsystem() const
 {
-	return (GetOwner() && GetOwner()->GetWorld()) ? GetOwner()->GetWorld()->GetSubsystem<UPCGSubsystem>() : nullptr;
+	return GetOwner() ? UPCGSubsystem::GetInstance(GetOwner()->GetWorld()) : nullptr;
 }
 
 #if WITH_EDITOR
diff --git a/Engine/Plugins/Experimental/PCG/Source/PCG/Private/PCGGraph.cpp b/Engine/Plugins/Experimental/PCG/Source/PCG/Private/PCGGraph.cpp
--- a/Engine/Plugins/Experimental/PCG/Source/PCG/Private/PCGGraph.cpp
+++ b/Engine/Plugins/Experimental/PCG/Source/PCG/Private/PCGGraph.cpp
@@ -191,9 +191,7 @@ void UPCGGraph::BeginDestroy()
 	// Notify the compiler to remove this graph from its cache
 	if (GEditor)
 	{
-		UWorld* World = GEditor->GetEditorWorldContext().World();
-		UPCGSubsystem* PCGSubsystem = World ? World->GetSubsystem<UPCGSubsystem>() : nullptr;
-		if (PCGSubsystem)
+		if (UPCGSubsystem* PCGSubsystem = UPCGSubsystem::GetInstance(GEditor->GetEditorWorldContext().World()))
 		{
 			PCGSubsystem->NotifyGraphChanged(this);
 		}
@@ -670,9 +668,7 @@ void UPCGGraph::NotifyGraphChanged(EPCGChangeType ChangeType)
 	const bool bNotifySubsystem = ((ChangeType & (EPCGChangeType::Structural | EPCGChangeType::Edge)) != EPCGChangeType::None);
 	if (bNotifySubsystem && GEditor)
 	{
-		UWorld* World = GEditor->GetEditorWorldContext().World();
-		UPCGSubsystem* PCGSubsystem = World ? World->GetSubsystem<UPCGSubsystem>() : nullptr;
-		if (PCGSubsystem)
+		if (UPCGSubsystem* PCGSubsystem = UPCGSubsystem::GetInstance(GEditor->GetEditorWorldContext().World()))
 		{
 			PCGSubsystem->NotifyGraphChanged(this);
 		}
diff --git a/Engine/Plugins/Experimental/PCG/Source/PCG/Private/PCGSettings.cpp b/Engine/Plugins/Experimental/PCG/Source/PCG/Private/PCGSettings.cpp
--- a/Engine/Plugins/Experimental/PCG/Source/PCG/Private/PCGSettings.cpp
+++ b/Engine/Plugins/Experimental/PCG/Source/PCG/Private/PCGSettings.cpp
@@ -415,9 +415,7 @@ void UPCGSettings::DirtyCache()
 {
 	if (GEditor)
 	{
-		UWorld* World = GEditor->GetEditorWorldContext().World();
-		UPCGSubsystem* PCGSubsystem = World ? World->GetSubsystem<UPCGSubsystem>() : nullptr;
-		if (PCGSubsystem)
+		if (UPCGSubsystem* PCGSubsystem = UPCGSubsystem::GetInstance(GEditor->GetEditorWorldContext().World()))
 		{
 			PCGSubsystem->CleanFromCache(GetElement().Get(), this);
 		}
diff --git a/Engine/Plugins/Experimental/PCG/Source/PCG/Private/PCGSubsystem.cpp b/Engine/Plugins/Experimental/PCG/Source/PCG/Private/PCGSubsystem.cpp
--- a/Engine/Plugins/Experimental/PCG/Source/PCG/Private/PCGSubsystem.cpp
+++ b/Engine/Plugins/Experimental/PCG/Source/PCG/Private/PCGSubsystem.cpp
@@ -56,9 +56,9 @@ namespace PCGSubsystemConsole
 					World = GEngine->GetCurrentPlayWorld();
 				}
 
-				if (World && World->GetSubsystem<UPCGSubsystem>())
+				if (UPCGSubsystem* PCGSubsystem = UPCGSubsystem::GetInstance(World))
 				{
-					World->GetSubsystem<UPCGSubsystem>()->FlushCache();
+					PCGSubsystem->FlushCache();
 				}
 			}));
 
@@ -69,9 +69,9 @@ namespace PCGSubsystemConsole
 		FConsoleCommandDelegate::CreateLambda([]()
 			{
 				UWorld* World = (GEditor ? (GEditor->PlayWorld ? GEditor->PlayWorld.Get() : GEditor->GetEditorWorldContext().World()) : (GEngine ? GEngine->GetCurrentPlayWorld() : nullptr));
-				if(World && World->GetSubsystem<UPCGSubsystem>())
+				if (UPCGSubsystem* PCGSubsystem = UPCGSubsystem::GetInstance(World))
 				{
-					World->GetSubsystem<UPCGSubsystem>()->BuildLandscapeCache();
+					PCGSubsystem->FlushCache();
 				}
 			}));
 
@@ -142,6 +142,19 @@ void UPCGSubsystem::PostInitialize()
 	}
 }
 
+UPCGSubsystem* UPCGSubsystem::GetInstance(UWorld* World)
+{
+	if (World)
+	{
+		UPCGSubsystem* Subsystem = World->GetSubsystem<UPCGSubsystem>();
+		return (Subsystem && Subsystem->IsInitialized()) ? Subsystem : nullptr;
+	}
+	else
+	{
+		return nullptr;
+	}
+}
+
 void UPCGSubsystem::Tick(float DeltaSeconds)
 {
 	Super::Tick(DeltaSeconds);
diff --git a/Engine/Plugins/Experimental/PCG/Source/PCG/Private/PCGWorldActor.cpp b/Engine/Plugins/Experimental/PCG/Source/PCG/Private/PCGWorldActor.cpp
--- a/Engine/Plugins/Experimental/PCG/Source/PCG/Private/PCGWorldActor.cpp
+++ b/Engine/Plugins/Experimental/PCG/Source/PCG/Private/PCGWorldActor.cpp
@@ -73,8 +73,7 @@ APCGWorldActor* APCGWorldActor::CreatePCGWorldActor(UWorld* InWorld)
 
 void APCGWorldActor::RegisterToSubsystem()
 {
-	UPCGSubsystem* PCGSubsystem = GetWorld() ? GetWorld()->GetSubsystem<UPCGSubsystem>() : nullptr;
-	if (PCGSubsystem)
+	if (UPCGSubsystem* PCGSubsystem = UPCGSubsystem::GetInstance(GetWorld()))
 	{
 		PCGSubsystem->RegisterPCGWorldActor(this);
 	}
@@ -82,8 +81,7 @@ void APCGWorldActor::RegisterToSubsystem()
 
 void APCGWorldActor::UnregisterFromSubsystem()
 {
-	UPCGSubsystem* PCGSubsystem = GetWorld() ? GetWorld()->GetSubsystem<UPCGSubsystem>() : nullptr;
-	if (PCGSubsystem)
+	if (UPCGSubsystem* PCGSubsystem = UPCGSubsystem::GetInstance(GetWorld()))
 	{
 		PCGSubsystem->UnregisterPCGWorldActor(this);
 	}
@@ -100,7 +98,7 @@ void APCGWorldActor::OnPartitionGridSizeChanged()
 		return;
 	}
 
-	UPCGSubsystem* PCGSubsystem = World->GetSubsystem<UPCGSubsystem>();
+	UPCGSubsystem* PCGSubsystem = UPCGSubsystem::GetInstance(World);
 	ULevel* Level = World->GetCurrentLevel();
 	if (!PCGSubsystem || !Level)
 	{
diff --git a/Engine/Plugins/Experimental/PCG/Source/PCG/Public/PCGSubsystem.h b/Engine/Plugins/Experimental/PCG/Source/PCG/Public/PCGSubsystem.h
--- a/Engine/Plugins/Experimental/PCG/Source/PCG/Public/PCGSubsystem.h
+++ b/Engine/Plugins/Experimental/PCG/Source/PCG/Public/PCGSubsystem.h
@@ -28,6 +28,8 @@ class UPCGSettings;
 class IPCGElement;
 typedef TSharedPtr<IPCGElement, ESPMode::ThreadSafe> FPCGElementPtr;
 
+class UWorld;
+
 /**
 * UPCGSubsystem
 */
@@ -57,6 +59,12 @@ public:
 	virtual TStatId GetStatId() const override;
 	//~ End FTickableGameObject
 
+	/** Will return the subsystem from the World if it exists and if it is initialized */
+	static UPCGSubsystem* GetInstance(UWorld* World);
+
+	/** Subsytem must not be used without this condition being true. */
+	bool IsInitialized() const { return GraphExecutor != nullptr; }
+
 	APCGWorldActor* GetPCGWorldActor();
 #if WITH_EDITOR
 	void DestroyPCGWorldActor();
diff --git a/Engine/Plugins/Experimental/PCG/Source/PCGEditor/Private/PCGEditor.cpp b/Engine/Plugins/Experimental/PCG/Source/PCGEditor/Private/PCGEditor.cpp
--- a/Engine/Plugins/Experimental/PCG/Source/PCGEditor/Private/PCGEditor.cpp
+++ b/Engine/Plugins/Experimental/PCG/Source/PCGEditor/Private/PCGEditor.cpp
@@ -2215,7 +2215,7 @@ bool FPCGEditor::IsVisibleProperty(const FPropertyAndParent& InPropertyAndParent
 UPCGSubsystem* FPCGEditor::GetSubsystem()
 {
 	UWorld* World = (GEditor ? (GEditor->PlayWorld ? GEditor->PlayWorld.Get() : GEditor->GetEditorWorldContext().World()) : nullptr);
-	return World ? World->GetSubsystem<UPCGSubsystem>() : nullptr;
+	return UPCGSubsystem::GetInstance(World);
 }
 
 TSharedRef<SDockTab> FPCGEditor::SpawnTab_GraphEditor(const FSpawnTabArgs& Args)
diff --git a/Engine/Plugins/Experimental/PCG/Source/PCGEditor/Private/PCGEditorGraphNode.cpp b/Engine/Plugins/Experimental/PCG/Source/PCGEditor/Private/PCGEditorGraphNode.cpp
--- a/Engine/Plugins/Experimental/PCG/Source/PCGEditor/Private/PCGEditorGraphNode.cpp
+++ b/Engine/Plugins/Experimental/PCG/Source/PCGEditor/Private/PCGEditorGraphNode.cpp
@@ -20,7 +20,7 @@ namespace UPCGEditorGraphNodeHelpers
 	void GetGraphCacheDebugInfo(const UPCGNode* InNode, bool& bOutDebuggingEnabled, uint32& OutNumCacheEntries)
 	{
 		UWorld* World = GEditor ? (GEditor->PlayWorld ? GEditor->PlayWorld.Get() : GEditor->GetEditorWorldContext().World()) : nullptr;
-		UPCGSubsystem* Subsystem = World ? World->GetSubsystem<UPCGSubsystem>() : nullptr;
+		UPCGSubsystem* Subsystem = UPCGSubsystem::GetInstance(World);
 		bOutDebuggingEnabled = Subsystem && Subsystem->IsGraphCacheDebuggingEnabled();
 
 		if (bOutDebuggingEnabled)
diff --git a/Engine/Plugins/Experimental/PCG/Source/PCGEditor/Private/PCGEditorModule.cpp b/Engine/Plugins/Experimental/PCG/Source/PCGEditor/Private/PCGEditorModule.cpp
--- a/Engine/Plugins/Experimental/PCG/Source/PCGEditor/Private/PCGEditorModule.cpp
+++ b/Engine/Plugins/Experimental/PCG/Source/PCGEditor/Private/PCGEditorModule.cpp
@@ -157,9 +157,9 @@ void FPCGEditorModule::PopulateMenuActions(FMenuBuilder& MenuBuilder)
 			FExecuteAction::CreateLambda([]() {
 				if (GEditor)
 				{
-					if (UWorld* World = GEditor->GetEditorWorldContext().World())
+					if (UPCGSubsystem* PCGSubsystem = UPCGSubsystem::GetInstance(GEditor->GetEditorWorldContext().World()))
 					{
-						World->GetSubsystem<UPCGSubsystem>()->DeletePartitionActors(/*bOnlyDeleteUnused=*/false);
+						PCGSubsystem->DeletePartitionActors(/*bOnlyDeleteUnused=*/false);
 					}
 				}
 			})),
@@ -173,9 +173,9 @@ void FPCGEditorModule::PopulateMenuActions(FMenuBuilder& MenuBuilder)
 			FExecuteAction::CreateLambda([]() {
 				if (GEditor)
 				{
-					if (UWorld* World = GEditor->GetEditorWorldContext().World())
+					if (UPCGSubsystem* PCGSubsystem = UPCGSubsystem::GetInstance(GEditor->GetEditorWorldContext().World()))
 					{
-						World->GetSubsystem<UPCGSubsystem>()->DeletePartitionActors(/*bOnlyDeleteUnused=*/true);
+						PCGSubsystem->DeletePartitionActors(/*bOnlyDeleteUnused=*/true);
 					}
 				}
 			})),
@@ -189,11 +189,11 @@ void FPCGEditorModule::PopulateMenuActions(FMenuBuilder& MenuBuilder)
 			FExecuteAction::CreateLambda([]() {
 				if (GEditor)
 				{
-					if (UWorld* World = GEditor->GetEditorWorldContext().World())
+					if (UPCGSubsystem* PCGSubsystem = UPCGSubsystem::GetInstance(GEditor->GetEditorWorldContext().World()))
 					{
-						if (APCGWorldActor* PCGWorldActor = World->GetSubsystem<UPCGSubsystem>()->GetPCGWorldActor())
+						if (APCGWorldActor* PCGWorldActor = PCGSubsystem->GetPCGWorldActor())
 						{
-							World->GetSubsystem<UPCGSubsystem>()->DestroyPCGWorldActor();
+							PCGSubsystem->DestroyPCGWorldActor();
 						}
 					}
 				}
@@ -208,12 +208,9 @@ void FPCGEditorModule::PopulateMenuActions(FMenuBuilder& MenuBuilder)
 			FExecuteAction::CreateLambda([]() {
 				if (GEditor)
 				{
-					if (UWorld* World = GEditor->GetEditorWorldContext().World())
+					if (UPCGSubsystem* PCGSubsystem = UPCGSubsystem::GetInstance(GEditor->GetEditorWorldContext().World()))
 					{
-						if (UPCGSubsystem* Subsystem = World->GetSubsystem<UPCGSubsystem>())
-						{
-							Subsystem->BuildLandscapeCache();
-						}
+						PCGSubsystem->BuildLandscapeCache();
 					}
 				}
 			})),
@@ -227,12 +224,9 @@ void FPCGEditorModule::PopulateMenuActions(FMenuBuilder& MenuBuilder)
 			FExecuteAction::CreateLambda([]() {
 				if (GEditor)
 				{
-					if (UWorld* World = GEditor->GetEditorWorldContext().World())
+					if (UPCGSubsystem* PCGSubsystem = UPCGSubsystem::GetInstance(GEditor->GetEditorWorldContext().World()))
 					{
-						if (UPCGSubsystem* Subsystem = World->GetSubsystem<UPCGSubsystem>())
-						{
-							Subsystem->ClearLandscapeCache();
-						}
+						PCGSubsystem->ClearLandscapeCache();
 					}
 				}
 				})),
@@ -246,12 +240,9 @@ void FPCGEditorModule::PopulateMenuActions(FMenuBuilder& MenuBuilder)
 			FExecuteAction::CreateLambda([]() {
 				if (GEditor)
 				{
-					if (UWorld* World = GEditor->GetEditorWorldContext().World())
+					if (UPCGSubsystem* PCGSubsystem = UPCGSubsystem::GetInstance(GEditor->GetEditorWorldContext().World()))
 					{
-						if (UPCGSubsystem* Subsystem = World->GetSubsystem<UPCGSubsystem>())
-						{
-							Subsystem->CancelAllGeneration();
-						}
+						PCGSubsystem->CancelAllGeneration();
 					}
 				}
 				})),
diff --git a/Engine/Plugins/Experimental/PCG/Source/PCGEditor/Private/SPCGEditorGraphAttributeListView.cpp b/Engine/Plugins/Experimental/PCG/Source/PCGEditor/Private/SPCGEditorGraphAttributeListView.cpp
--- a/Engine/Plugins/Experimental/PCG/Source/PCGEditor/Private/SPCGEditorGraphAttributeListView.cpp
+++ b/Engine/Plugins/Experimental/PCG/Source/PCGEditor/Private/SPCGEditorGraphAttributeListView.cpp
@@ -88,7 +88,7 @@ namespace PCGEditorGraphAttributeListView
 	bool IsGraphCacheDebuggingEnabled()
 	{
 		UWorld* World = GEditor ? (GEditor->PlayWorld ? GEditor->PlayWorld.Get() : GEditor->GetEditorWorldContext().World()) : nullptr;
-		UPCGSubsystem* Subsystem = World ? World->GetSubsystem<UPCGSubsystem>() : nullptr;
+		UPCGSubsystem* Subsystem = UPCGSubsystem::GetInstance(World);
 		return Subsystem && Subsystem->IsGraphCacheDebuggingEnabled();
 	}
 
