[PCG] Static analysis fix

#rnx

diff --git a/Engine/Plugins/Experimental/PCG/Source/PCG/Private/Data/PCGLandscapeData.cpp b/Engine/Plugins/Experimental/PCG/Source/PCG/Private/Data/PCGLandscapeData.cpp
--- a/Engine/Plugins/Experimental/PCG/Source/PCG/Private/Data/PCGLandscapeData.cpp
+++ b/Engine/Plugins/Experimental/PCG/Source/PCG/Private/Data/PCGLandscapeData.cpp
@@ -30,6 +30,7 @@ void UPCGLandscapeData::Initialize(const TArray<TWeakObjectPtr<ALandscapeProxy>>
 	check(!Landscapes.IsEmpty());
 
 	ALandscapeProxy* FirstLandscape = Landscapes[0].Get();
+	check(FirstLandscape);
 
 	TargetActor = FirstLandscape;
 	Bounds = InBounds;
@@ -92,7 +93,7 @@ void UPCGLandscapeData::PostLoad()
 	}
 
 	TargetActor = FirstLandscape;
-	UPCGSubsystem* PCGSubsystem = UPCGSubsystem::GetInstance(FirstLandscape->GetWorld());
+	UPCGSubsystem* PCGSubsystem = UPCGSubsystem::GetInstance(FirstLandscape ? FirstLandscape->GetWorld() : nullptr);
 	LandscapeCache = PCGSubsystem ? PCGSubsystem->GetLandscapeCache() : nullptr;
 }
 
