[PCG] Fix ensure hit when trying to debug an element

* Mimic how we set custom float data in PCGStaticMeshSpawnerElement.

diff --git a/Engine/Plugins/Experimental/PCG/Source/PCG/Private/Elements/PCGDebugElement.cpp b/Engine/Plugins/Experimental/PCG/Source/PCG/Private/Elements/PCGDebugElement.cpp
--- a/Engine/Plugins/Experimental/PCG/Source/PCG/Private/Elements/PCGDebugElement.cpp
+++ b/Engine/Plugins/Experimental/PCG/Source/PCG/Private/Elements/PCGDebugElement.cpp
@@ -198,7 +198,16 @@ namespace PCGDebugElement
 					InstanceCustomData.Add(Point.Color[2]);
 					InstanceCustomData.Add(Point.Color[3]);
 
-					ISMC->SetCustomData(PointIndex + PreExistingInstanceCount, InstanceCustomData);
+					// Doing the same thing than in PCGStaticMeshSpawnerElement
+					FMemory::Memcpy(&ISMC->PerInstanceSMCustomData[(PreExistingInstanceCount + PointIndex) * NumCustomData], InstanceCustomData.GetData(), NumCustomData * sizeof(float));
+
+					// In PCGStaticMeshSpawnerElement, this is incremented only once, because it copies the custom data in one go.
+					// Here we also increment this value only once, when the whole copy is done.
+					if (PointIndex == Points.Num() - 1)
+					{
+						// Force recreation of the render data when proxy is created
+						ISMC->InstanceUpdateCmdBuffer.NumEdits++;
+					}
 
 					InstanceCustomData.Reset();
 				}
