[PCG] Graph cache memory tweaks. Trigger GC each 100 tidy ups, and increase max entry count to 64k

diff --git a/Engine/Plugins/Experimental/PCG/Source/PCG/Private/Graph/PCGGraphCache.cpp b/Engine/Plugins/Experimental/PCG/Source/PCG/Private/Graph/PCGGraphCache.cpp
--- a/Engine/Plugins/Experimental/PCG/Source/PCG/Private/Graph/PCGGraphCache.cpp
+++ b/Engine/Plugins/Experimental/PCG/Source/PCG/Private/Graph/PCGGraphCache.cpp
@@ -24,7 +24,7 @@ static TAutoConsoleVariable<bool> CVarCacheMemoryBudgetEnabled(
 	true,
 	TEXT("Whether memory budget is enforced (items purged from cache to respect pcg.Cache.MemoryBudgetMB."));
 
-static int32 GPCGGraphCacheMaxElements = 16384;
+static int32 GPCGGraphCacheMaxElements = 65536;
 static FAutoConsoleVariableRef CVarNiagaraGraphDataCacheSize(
 	TEXT("pcg.Cache.GraphCacheMaxElements"),
 	GPCGGraphCacheMaxElements,
@@ -126,13 +126,13 @@ void FPCGGraphCache::ClearCache()
 	CacheData.Empty(GPCGGraphCacheMaxElements);
 }
 
-void FPCGGraphCache::EnforceMemoryBudget()
+bool FPCGGraphCache::EnforceMemoryBudget()
 {
 	TRACE_CPUPROFILER_EVENT_SCOPE(FPCGGraphCache::FPCGGraphCache::EnforceMemoryBudget);
 
 	if (!CVarCacheMemoryBudgetEnabled.GetValueOnAnyThread())
 	{
-		return;
+		return false;
 	}
 
 	if (CacheData.Num() == GPCGGraphCacheMaxElements)
@@ -141,9 +141,9 @@ void FPCGGraphCache::EnforceMemoryBudget()
 	}
 
 	const uint64 MemoryBudget = static_cast<uint64>(CVarCacheMemoryBudgetMB.GetValueOnAnyThread()) * 1024 * 1024;
-	if (TotalMemoryUsed < MemoryBudget)
+	if (TotalMemoryUsed <= MemoryBudget)
 	{
-		return;
+		return false;
 	}
 
 	{
@@ -156,6 +156,8 @@ void FPCGGraphCache::EnforceMemoryBudget()
 			RemoveFromMemoryTotal(RemovedData);
 		}
 	}
+
+	return true;
 }
 
 #if WITH_EDITOR
diff --git a/Engine/Plugins/Experimental/PCG/Source/PCG/Private/Graph/PCGGraphCache.h b/Engine/Plugins/Experimental/PCG/Source/PCG/Private/Graph/PCGGraphCache.h
--- a/Engine/Plugins/Experimental/PCG/Source/PCG/Private/Graph/PCGGraphCache.h
+++ b/Engine/Plugins/Experimental/PCG/Source/PCG/Private/Graph/PCGGraphCache.h
@@ -75,7 +75,7 @@ struct FPCGGraphCache
 	void ClearCache();
 
 	/** While memory usage is more than budget, remove cache entries for elements, LRU policy. Returns true if something removed. */
-	void EnforceMemoryBudget();
+	bool EnforceMemoryBudget();
 
 	/** True if debugging features enabled. Exposes a CVar so it can be called from editor module. */
 	bool IsDebuggingEnabled() const;
diff --git a/Engine/Plugins/Experimental/PCG/Source/PCG/Private/Graph/PCGGraphExecutor.cpp b/Engine/Plugins/Experimental/PCG/Source/PCG/Private/Graph/PCGGraphExecutor.cpp
--- a/Engine/Plugins/Experimental/PCG/Source/PCG/Private/Graph/PCGGraphExecutor.cpp
+++ b/Engine/Plugins/Experimental/PCG/Source/PCG/Private/Graph/PCGGraphExecutor.cpp
@@ -880,7 +880,19 @@ void FPCGGraphExecutor::Execute()
 		}
 
 		// Purge things from cache if memory usage is too high
-		GraphCache.EnforceMemoryBudget();
+		const bool bSomethingTidied = GraphCache.EnforceMemoryBudget();
+
+#if WITH_EDITOR
+		if (bSomethingTidied && !PCGHelpers::IsRuntimeOrPIE())
+		{
+			--TidyCacheCountUntilGC;
+			if (TidyCacheCountUntilGC <= 0)
+			{
+				TidyCacheCountUntilGC = 100;
+				CollectGarbage(RF_NoFlags, true);
+			}
+		}
+#endif
 
 #if WITH_EDITOR
 		// Save & release resources when running in-editor
@@ -1189,10 +1201,10 @@ void FPCGGraphExecutor::ReleaseUnusedActors()
 #if WITH_EDITOR
 	if (bRunGC && !PCGHelpers::IsRuntimeOrPIE())
 	{
-		--CountUntilGC;
-		if (CountUntilGC <= 0)
+		--ReleaseActorsCountUntilGC;
+		if (ReleaseActorsCountUntilGC <= 0)
 		{
-			CountUntilGC = 30;
+			ReleaseActorsCountUntilGC = 30;
 			CollectGarbage(RF_NoFlags, true);
 		}
 	}
diff --git a/Engine/Plugins/Experimental/PCG/Source/PCG/Private/Graph/PCGGraphExecutor.h b/Engine/Plugins/Experimental/PCG/Source/PCG/Private/Graph/PCGGraphExecutor.h
--- a/Engine/Plugins/Experimental/PCG/Source/PCG/Private/Graph/PCGGraphExecutor.h
+++ b/Engine/Plugins/Experimental/PCG/Source/PCG/Private/Graph/PCGGraphExecutor.h
@@ -187,8 +187,10 @@ private:
 	TSet<AActor*> ActorsToSave;
 	TSet<FWorldPartitionReference> ActorsToRelease;
 
-	int32 CountUntilGC = 30;
+	int32 ReleaseActorsCountUntilGC = 30;
 	FAsyncCompilationNotification GenerationProgressNotification;
+
+	int32 TidyCacheCountUntilGC = 100;
 #endif
 };
 
