[PCG] Added visibility check to attribute list to avoid showing empty widget when list is not populated.

diff --git a/Engine/Plugins/Experimental/PCG/Source/PCGEditor/Private/SPCGEditorGraphAttributeListView.cpp b/Engine/Plugins/Experimental/PCG/Source/PCGEditor/Private/SPCGEditorGraphAttributeListView.cpp
--- a/Engine/Plugins/Experimental/PCG/Source/PCGEditor/Private/SPCGEditorGraphAttributeListView.cpp
+++ b/Engine/Plugins/Experimental/PCG/Source/PCGEditor/Private/SPCGEditorGraphAttributeListView.cpp
@@ -442,6 +442,11 @@ void SPCGEditorGraphAttributeListView::Construct(const FArguments& InArgs, TShar
 		.Orientation(Orient_Vertical)
 		.Thickness(FVector2D(12.0f, 12.0f));
 
+	auto VisibilityTest = [this]()
+	{
+		return (ListViewItems.IsEmpty() && ListViewHeader->GetColumns().IsEmpty()) ? EVisibility::Hidden : EVisibility::Visible;
+	};
+
 	SAssignNew(ListView, SListView<PCGListviewItemPtr>)
 		.ListItemsSource(&ListViewItems)
 		.HeaderRow(ListViewHeader)
@@ -449,6 +454,7 @@ void SPCGEditorGraphAttributeListView::Construct(const FArguments& InArgs, TShar
 		.OnMouseButtonDoubleClick(this, &SPCGEditorGraphAttributeListView::OnItemDoubleClicked)
 		.AllowOverscroll(EAllowOverscroll::No)
 		.ExternalScrollbar(VerticalScrollBar)
+		.Visibility_Lambda(VisibilityTest)
 		.ConsumeMouseWheel(EConsumeMouseWheel::Always);
 
 	SAssignNew(DataComboBox, SComboBox<TSharedPtr<FName>>)
@@ -485,10 +491,10 @@ void SPCGEditorGraphAttributeListView::Construct(const FArguments& InArgs, TShar
 		[
 			SNew(SHorizontalBox)
 			+SHorizontalBox::Slot()
-            .AutoWidth()
-            [
-            	FilterButton->AsShared()
-            ]
+			.AutoWidth()
+			[
+				FilterButton->AsShared()
+			]
 			+SHorizontalBox::Slot()
 			.AutoWidth()
 			[
