[PCG] Performance fixes - switching to raw pointers locally in a few elements.

I went through all CreatePointData implementations.

Gave a really good bump to CopyPoints in my test, and others should help a bit overall.

diff --git a/Engine/Plugins/Experimental/PCG/Source/PCG/Private/Data/PCGDifferenceData.cpp b/Engine/Plugins/Experimental/PCG/Source/PCG/Private/Data/PCGDifferenceData.cpp
--- a/Engine/Plugins/Experimental/PCG/Source/PCG/Private/Data/PCGDifferenceData.cpp
+++ b/Engine/Plugins/Experimental/PCG/Source/PCG/Private/Data/PCGDifferenceData.cpp
@@ -175,20 +175,24 @@ const UPCGPointData* UPCGDifferenceData::CreatePointData(FPCGContext* Context) c
 		return SourcePointData;
 	}
 
+	const UPCGMetadata* SourceMetadata = SourcePointData->Metadata;
+
 	UPCGPointData* Data = NewObject<UPCGPointData>();
-	Data->InitializeFromData(this, SourcePointData->Metadata);
+	Data->InitializeFromData(this, SourceMetadata);
+	
+	UPCGMetadata* OutMetadata = Data->Metadata;
 
 	const TArray<FPCGPoint>& SourcePoints = SourcePointData->GetPoints();
 	TArray<FPCGPoint>& TargetPoints = Data->GetMutablePoints();
 
 	UPCGMetadata* TempDiffMetadata = nullptr;
-	if (bDiffMetadata && Data->Metadata && Difference->Metadata)
+	if (bDiffMetadata && OutMetadata && Difference->Metadata)
 	{
 		TempDiffMetadata = NewObject<UPCGMetadata>();
 		TempDiffMetadata->Initialize(Difference->Metadata);
 	}
 
-	FPCGAsync::AsyncPointProcessing(Context, SourcePoints.Num(), TargetPoints, [this, Data, SourcePointData, TempDiffMetadata, &SourcePoints](int32 Index, FPCGPoint& OutPoint)
+	FPCGAsync::AsyncPointProcessing(Context, SourcePoints.Num(), TargetPoints, [this, Data, OutMetadata, SourcePointData, SourceMetadata, TempDiffMetadata, &SourcePoints](int32 Index, FPCGPoint& OutPoint)
 	{
 		const FPCGPoint& Point = SourcePoints[Index];
 
@@ -202,7 +206,7 @@ const UPCGPointData* UPCGDifferenceData::CreatePointData(FPCGContext* Context) c
 
 			if (TempDiffMetadata && OutPoint.Density > 0 && PointFromDiff.MetadataEntry != PCGInvalidEntryKey)
 			{
-				Data->Metadata->MergePointAttributesSubset(Point, SourcePointData->Metadata, SourcePointData->Metadata, PointFromDiff, TempDiffMetadata, TempDiffMetadata, OutPoint, EPCGMetadataOp::Sub);
+				OutMetadata->MergePointAttributesSubset(Point, SourceMetadata, SourceMetadata, PointFromDiff, TempDiffMetadata, TempDiffMetadata, OutPoint, EPCGMetadataOp::Sub);
 			}
 
 #if WITH_EDITOR
diff --git a/Engine/Plugins/Experimental/PCG/Source/PCG/Private/Data/PCGLandscapeData.cpp b/Engine/Plugins/Experimental/PCG/Source/PCG/Private/Data/PCGLandscapeData.cpp
--- a/Engine/Plugins/Experimental/PCG/Source/PCG/Private/Data/PCGLandscapeData.cpp
+++ b/Engine/Plugins/Experimental/PCG/Source/PCG/Private/Data/PCGLandscapeData.cpp
@@ -130,7 +130,6 @@ bool UPCGLandscapeData::SamplePoint(const FTransform& InTransform, const FBox& I
 
 bool UPCGLandscapeData::ProjectPoint(const FTransform& InTransform, const FBox& InBounds, const FPCGProjectionParams& InParams, FPCGPoint& OutPoint, UPCGMetadata* OutMetadata) const
 {
-	//TRACE_CPUPROFILER_EVENT_SCOPE(UPCGLandscapeData::SamplePoint);
 	if (!LandscapeCache)
 	{
 		return false;
@@ -216,6 +215,8 @@ const UPCGPointData* UPCGLandscapeData::CreatePointData(FPCGContext* Context, co
 		return Data;
 	}
 
+	UPCGMetadata* OutMetadata = bUseMetadata ? Data->Metadata : nullptr;
+
 	for (const TPair<FBox, ULandscapeInfo*>& LandscapeInfoPair : LandscapeInfos)
 	{
 		ULandscapeInfo* LandscapeInfo = LandscapeInfoPair.Value;
@@ -295,7 +296,7 @@ const UPCGPointData* UPCGLandscapeData::CreatePointData(FPCGContext* Context, co
 						}
 						else
 						{
-							LandscapeCacheEntry->GetPoint(PointIndex, Point, bUseMetadata ? Data->Metadata : nullptr);
+							LandscapeCacheEntry->GetPoint(PointIndex, Point, OutMetadata);
 						}
 					}
 				}
diff --git a/Engine/Plugins/Experimental/PCG/Source/PCG/Private/Data/PCGProjectionData.cpp b/Engine/Plugins/Experimental/PCG/Source/PCG/Private/Data/PCGProjectionData.cpp
--- a/Engine/Plugins/Experimental/PCG/Source/PCG/Private/Data/PCGProjectionData.cpp
+++ b/Engine/Plugins/Experimental/PCG/Source/PCG/Private/Data/PCGProjectionData.cpp
@@ -173,13 +173,15 @@ const UPCGPointData* UPCGProjectionData::CreatePointData(FPCGContext* Context) c
 	// TODO: add mechanism in the ToPointData so we can pass in a transform
 	// so we can forego creating the points twice if they're not used.
 	const UPCGPointData* SourcePointData = Source->ToPointData(Context);
+	const UPCGMetadata* SourceMetadata = SourcePointData->Metadata;
 	const TArray<FPCGPoint>& SourcePoints = SourcePointData->GetPoints();
 
 	UPCGPointData* PointData = NewObject<UPCGPointData>();
 
 	// Copy metadata attributes from source point including values
-	PointData->InitializeFromData(this, SourcePointData->Metadata);
-	check(PointData->Metadata);
+	PointData->InitializeFromData(this, SourceMetadata);
+	UPCGMetadata* OutMetadata = PointData->Metadata;
+	check(OutMetadata);
 
 	// The projection operation will write into this temporary metadata
 	UPCGMetadata* TempTargetMetadata = nullptr;
@@ -203,11 +205,11 @@ const UPCGPointData* UPCGProjectionData::CreatePointData(FPCGContext* Context) c
 	}
 
 	// Add any attributes from filtered metadata to produce final list
-	PointData->Metadata->AddAttributes(TempTargetMetadata);
+	OutMetadata->AddAttributes(TempTargetMetadata);
 
 	TArray<FPCGPoint>& Points = PointData->GetMutablePoints();
 
-	FPCGAsync::AsyncPointProcessing(Context, SourcePoints.Num(), Points, [this, SourcePointData, PointData, TempTargetMetadata, &SourcePoints](int32 Index, FPCGPoint& OutPoint)
+	FPCGAsync::AsyncPointProcessing(Context, SourcePoints.Num(), Points, [this, SourcePointData, SourceMetadata, PointData, OutMetadata, TempTargetMetadata, &SourcePoints](int32 Index, FPCGPoint& OutPoint)
 	{
 		const FPCGPoint& SourcePoint = SourcePoints[Index];
 
@@ -230,10 +232,10 @@ const UPCGPointData* UPCGProjectionData::CreatePointData(FPCGContext* Context) c
 		OutPoint.Color = ProjectionParams.bProjectColors ? (PointFromTarget.Color * SourcePoint.Color) : SourcePoint.Color;
 		OutPoint.Density *= PointFromTarget.Density;
 		
-		if (PointData->Metadata && TempTargetMetadata && PointFromTarget.MetadataEntry != PCGInvalidEntryKey)
+		if (OutMetadata && TempTargetMetadata && PointFromTarget.MetadataEntry != PCGInvalidEntryKey)
 		{
 			// Merge metadata to produce final attribute values
-			PointData->Metadata->MergePointAttributesSubset(SourcePoint, SourcePointData->Metadata, SourcePointData->Metadata, PointFromTarget, TempTargetMetadata, TempTargetMetadata, OutPoint, ProjectionParams.AttributeMergeOperation);
+			OutMetadata->MergePointAttributesSubset(SourcePoint, SourceMetadata, SourceMetadata, PointFromTarget, TempTargetMetadata, TempTargetMetadata, OutPoint, ProjectionParams.AttributeMergeOperation);
 		}
 
 		return true;
diff --git a/Engine/Plugins/Experimental/PCG/Source/PCG/Private/Data/PCGUnionData.cpp b/Engine/Plugins/Experimental/PCG/Source/PCG/Private/Data/PCGUnionData.cpp
--- a/Engine/Plugins/Experimental/PCG/Source/PCG/Private/Data/PCGUnionData.cpp
+++ b/Engine/Plugins/Experimental/PCG/Source/PCG/Private/Data/PCGUnionData.cpp
@@ -116,17 +116,19 @@ bool UPCGUnionData::SamplePoint(const FTransform& InTransform, const FBox& InBou
 		}
 	}
 
+	TArray<const UPCGSpatialData*> DataRawPtr = Data;
+
 	const bool bSkipLoop = (bHasSetPoint && !OutMetadata && OutPoint.Density >= 1.0f);
-	const int32 DataCount = Data.Num();
+	const int32 DataCount = DataRawPtr.Num();
 	for (int32 DataIndex = 0; DataIndex < DataCount && !bSkipLoop; ++DataIndex)
 	{
-		if (Data[DataIndex] == FirstNonTrivialTransformData)
+		if (DataRawPtr[DataIndex] == FirstNonTrivialTransformData)
 		{
 			continue;
 		}
 
 		FPCGPoint PointInData;
-		if(Data[DataIndex]->SamplePoint(PointTransform, InBounds, PointInData, OutMetadata))
+		if(DataRawPtr[DataIndex]->SamplePoint(PointTransform, InBounds, PointInData, OutMetadata))
 		{
 			if (!bHasSetPoint)
 			{
@@ -149,7 +151,7 @@ bool UPCGUnionData::SamplePoint(const FTransform& InTransform, const FBox& InBou
 				{
 					if (OutPoint.MetadataEntry != PCGInvalidEntryKey && PointInData.MetadataEntry != PCGInvalidEntryKey)
 					{
-						OutMetadata->MergePointAttributesSubset(OutPoint, OutMetadata, OutMetadata, PointInData, OutMetadata, Data[DataIndex]->Metadata, OutPoint, EPCGMetadataOp::Max);
+						OutMetadata->MergePointAttributesSubset(OutPoint, OutMetadata, OutMetadata, PointInData, OutMetadata, DataRawPtr[DataIndex]->Metadata, OutPoint, EPCGMetadataOp::Max);
 					}
 					else if (PointInData.MetadataEntry != PCGInvalidEntryKey)
 					{
@@ -243,35 +245,46 @@ const UPCGPointData* UPCGUnionData::CreatePointData(FPCGContext* Context) const
 		return Data[0]->ToPointData(Context);
 	}
 
+	// Cache raw pointers for metadata as these are much faster to use
+	TArray<const UPCGSpatialData*> DataRawPtr = Data;
+	TArray<const UPCGMetadata*> InputMetadatas;
+	InputMetadatas.SetNumUninitialized(Data.Num());
+	for (int32 i = 0; i < DataRawPtr.Num(); i++)
+	{
+		InputMetadatas[i] = DataRawPtr[i]->ToPointData(Context)->Metadata;
+	}
+
 	UPCGPointData* PointData = NewObject<UPCGPointData>();
-	PointData->InitializeFromData(this, Data[0]->Metadata);
-	check(PointData->Metadata);
+	PointData->InitializeFromData(this, InputMetadatas[0]);
+
+	UPCGMetadata* OutMetadata = PointData->Metadata;
+	check(OutMetadata);
 
 	// Initialize metadata
-	for (TObjectPtr<const UPCGSpatialData> Datum : Data)
+	for (const UPCGMetadata* InputMetadata : InputMetadatas)
 	{
-		PointData->Metadata->AddAttributes(Datum->Metadata);
+		OutMetadata->AddAttributes(InputMetadata);
 	}
 
 	switch (UnionType)
 	{
 	case EPCGUnionType::LeftToRightPriority:
 	default:
-		CreateSequentialPointData(Context, PointData, /*bLeftToRight=*/true);
+		CreateSequentialPointData(Context, DataRawPtr, InputMetadatas, PointData, OutMetadata, /*bLeftToRight=*/true);
 		break;
 
 	case EPCGUnionType::RightToLeftPriority:
-		CreateSequentialPointData(Context, PointData, /*bLeftToRight=*/false);
+		CreateSequentialPointData(Context, DataRawPtr, InputMetadatas, PointData, OutMetadata, /*bLeftToRight=*/false);
 		break;
 
 	case EPCGUnionType::KeepAll:
 		{
 			TArray<FPCGPoint>& TargetPoints = PointData->GetMutablePoints();
-			for(int32 DataIndex = 0; DataIndex < Data.Num(); ++DataIndex)
+			for(int32 DataIndex = 0; DataIndex < DataRawPtr.Num(); ++DataIndex)
 			{
-				TObjectPtr<const UPCGSpatialData> Datum = Data[DataIndex];
-
+				const UPCGSpatialData* Datum = DataRawPtr[DataIndex];
 				const UPCGPointData* DatumPointData = Datum->ToPointData(Context);
+				const UPCGMetadata* DatumPointMetadata = DatumPointData->Metadata;
 
 				int32 TargetPointOffset = TargetPoints.Num();
 				TargetPoints.Append(DatumPointData->GetPoints());
@@ -285,9 +298,9 @@ const UPCGPointData* UPCGUnionData::CreatePointData(FPCGContext* Context) const
 						Point.MetadataEntry = PCGInvalidEntryKey;
 					}
 
-					if (PointData->Metadata && DatumPointData->Metadata && DatumPointData->Metadata->GetAttributeCount() > 0)
+					if (OutMetadata && DatumPointMetadata && DatumPointMetadata->GetAttributeCount() > 0)
 					{
-						PointData->Metadata->SetPointAttributes(MakeArrayView(DatumPointData->GetPoints()), DatumPointData->Metadata, TargetPointsSubset);
+						OutMetadata->SetPointAttributes(MakeArrayView(DatumPointData->GetPoints()), DatumPointMetadata, TargetPointsSubset);
 					}
 				}
 			}
@@ -309,15 +322,15 @@ const UPCGPointData* UPCGUnionData::CreatePointData(FPCGContext* Context) const
 	return PointData;
 }
 
-void UPCGUnionData::CreateSequentialPointData(FPCGContext* Context, UPCGPointData* PointData, bool bLeftToRight) const
+void UPCGUnionData::CreateSequentialPointData(FPCGContext* Context, TArray<const UPCGSpatialData*>& InputDatas, TArray<const UPCGMetadata*>& InputMetadatas, UPCGPointData* PointData, UPCGMetadata* OutMetadata, bool bLeftToRight) const
 {
 	check(PointData);
 
 	TArray<FPCGPoint>& TargetPoints = PointData->GetMutablePoints();
 	TArray<FPCGPoint> SelectedDataPoints;
 
-	int32 FirstDataIndex = (bLeftToRight ? 0 : Data.Num() - 1);
-	int32 LastDataIndex = (bLeftToRight ? Data.Num() : -1);
+	int32 FirstDataIndex = (bLeftToRight ? 0 : InputDatas.Num() - 1);
+	int32 LastDataIndex = (bLeftToRight ? InputDatas.Num() : -1);
 	int32 DataIndexIncrement = (bLeftToRight ? 1 : -1);
 
 	// Note: this is a O(N^2) implementation. 
@@ -326,9 +339,10 @@ void UPCGUnionData::CreateSequentialPointData(FPCGContext* Context, UPCGPointDat
 	{
 		// For each point, if it is not already "processed" by previous data,
 		// add it & compute its final density
-		const TArray<FPCGPoint>& Points = Data[DataIndex]->ToPointData(Context)->GetPoints();
+		const UPCGPointData* CurrentPointData = InputDatas[DataIndex]->ToPointData(Context);
+		const TArray<FPCGPoint>& Points = CurrentPointData->GetPoints();
 
-		FPCGAsync::AsyncPointProcessing(Context, Points.Num(), SelectedDataPoints, [this, PointData, &Points, DataIndex, FirstDataIndex, LastDataIndex, DataIndexIncrement](int32 Index, FPCGPoint& OutPoint)
+		FPCGAsync::AsyncPointProcessing(Context, Points.Num(), SelectedDataPoints, [this, PointData, OutMetadata, &Points, &InputDatas, &InputMetadatas, DataIndex, FirstDataIndex, LastDataIndex, DataIndexIncrement](int32 Index, FPCGPoint& OutPoint)
 		{
 			const FPCGPoint& Point = Points[Index];
 
@@ -336,7 +350,7 @@ void UPCGUnionData::CreateSequentialPointData(FPCGContext* Context, UPCGPointDat
 			bool bPointToExclude = false;
 			for (int32 PreviousDataIndex = FirstDataIndex; PreviousDataIndex != DataIndex; PreviousDataIndex += DataIndexIncrement)
 			{
-				if (Data[PreviousDataIndex]->GetDensityAtPosition(Point.Transform.GetLocation()) != 0)
+				if (InputDatas[PreviousDataIndex]->GetDensityAtPosition(Point.Transform.GetLocation()) != 0)
 				{
 					bPointToExclude = true;
 					break;
@@ -348,14 +362,14 @@ void UPCGUnionData::CreateSequentialPointData(FPCGContext* Context, UPCGPointDat
 				return false;
 			}
 
-			check(PointData && PointData->Metadata);
+			check(PointData && OutMetadata);
 
 			OutPoint = Point;
-			if (PointData->Metadata->GetParent() != Data[DataIndex]->Metadata)
+			if (OutMetadata->GetParent() != InputMetadatas[DataIndex])
 			{
-				UPCGMetadataAccessorHelpers::InitializeMetadata(OutPoint, PointData->Metadata);
+				UPCGMetadataAccessorHelpers::InitializeMetadata(OutPoint, OutMetadata);
 				// Since we can't inherit from the parent point, we'll set the values directly here
-				PointData->Metadata->SetPointAttributes(Point, Data[DataIndex]->Metadata, OutPoint);
+				OutMetadata->SetPointAttributes(Point, InputMetadatas[DataIndex], OutPoint);
 			}
 
 			if (DensityFunction == EPCGUnionDensityFunction::Binary && OutPoint.Density > 0)
@@ -364,11 +378,11 @@ void UPCGUnionData::CreateSequentialPointData(FPCGContext* Context, UPCGPointDat
 			}
 
 			// Update density & metadata based on current & following data
-			const bool bSkipLoop = (OutPoint.Density >= 1.0f && !PointData->Metadata);
+			const bool bSkipLoop = (OutPoint.Density >= 1.0f && !OutMetadata);
 			for (int32 FollowingDataIndex = DataIndex + DataIndexIncrement; FollowingDataIndex != LastDataIndex && !bSkipLoop; FollowingDataIndex += DataIndexIncrement)
 			{
 				FPCGPoint PointInData;
-				if(Data[FollowingDataIndex]->SamplePoint(OutPoint.Transform, OutPoint.GetLocalBounds(), PointInData, PointData->Metadata))
+				if(InputDatas[FollowingDataIndex]->SamplePoint(OutPoint.Transform, OutPoint.GetLocalBounds(), PointInData, OutMetadata))
 				{
 					// Update density
 					PCGUnionDataMaths::UpdateDensity(OutPoint.Density, PointInData.Density, DensityFunction);
@@ -380,11 +394,11 @@ void UPCGUnionData::CreateSequentialPointData(FPCGContext* Context, UPCGPointDat
 						FMath::Max(OutPoint.Color.Z, PointInData.Color.Z),
 						FMath::Max(OutPoint.Color.W, PointInData.Color.W));
 
-					if (PointData->Metadata)
+					if (OutMetadata)
 					{
 						if (OutPoint.MetadataEntry != PCGInvalidEntryKey && PointInData.MetadataEntry != PCGInvalidEntryKey)
 						{
-							PointData->Metadata->MergePointAttributesSubset(OutPoint, PointData->Metadata, PointData->Metadata, PointInData, PointData->Metadata, Data[FollowingDataIndex]->Metadata, OutPoint, EPCGMetadataOp::Max);
+							OutMetadata->MergePointAttributesSubset(OutPoint, OutMetadata, OutMetadata, PointInData, OutMetadata, InputMetadatas[FollowingDataIndex], OutPoint, EPCGMetadataOp::Max);
 						}
 						else if (PointInData.MetadataEntry != PCGInvalidEntryKey)
 						{
@@ -421,4 +435,3 @@ UPCGSpatialData* UPCGUnionData::CopyInternal() const
 
 	return NewUnionData;
 }
-
diff --git a/Engine/Plugins/Experimental/PCG/Source/PCG/Private/Elements/PCGCopyPoints.cpp b/Engine/Plugins/Experimental/PCG/Source/PCG/Private/Elements/PCGCopyPoints.cpp
--- a/Engine/Plugins/Experimental/PCG/Source/PCG/Private/Elements/PCGCopyPoints.cpp
+++ b/Engine/Plugins/Experimental/PCG/Source/PCG/Private/Elements/PCGCopyPoints.cpp
@@ -87,6 +87,9 @@ bool FPCGCopyPointsElement::ExecuteInternal(FPCGContext* Context) const
 		return true;
 	}
 
+	const UPCGMetadata* SourcePointMetadata = SourcePointData->Metadata;
+	const UPCGMetadata* TargetPointMetadata = TargetPointData->Metadata;
+
 	const TArray<FPCGPoint>& SourcePoints = SourcePointData->GetPoints();
 	const TArray<FPCGPoint>& TargetPoints = TargetPointData->GetPoints();
 
@@ -99,8 +102,8 @@ bool FPCGCopyPointsElement::ExecuteInternal(FPCGContext* Context) const
 	const UPCGMetadata* RootMetadata = nullptr;
 	const UPCGMetadata* NonRootMetadata = nullptr;
 
-	const bool bSourceHasMetadata = (SourcePointData->Metadata->GetAttributeCount() > 0 && SourcePointData->Metadata->GetItemCountForChild() > 0);
-	const bool bTargetHasMetadata = (TargetPointData->Metadata->GetAttributeCount() > 0 && TargetPointData->Metadata->GetItemCountForChild() > 0);
+	const bool bSourceHasMetadata = (SourcePointMetadata->GetAttributeCount() > 0 && SourcePointMetadata->GetItemCountForChild() > 0);
+	const bool bTargetHasMetadata = (TargetPointMetadata->GetAttributeCount() > 0 && TargetPointMetadata->GetItemCountForChild() > 0);
 
 	bool bInheritMetadataFromSource = true;
 	bool bProcessMetadata = (bSourceHasMetadata || bTargetHasMetadata);
@@ -111,7 +114,7 @@ bool FPCGCopyPointsElement::ExecuteInternal(FPCGContext* Context) const
 		bProcessMetadata = bSourceHasMetadata;
 
 		OutPointData->InitializeFromData(SourcePointData);
-		RootMetadata = SourcePointData->Metadata;
+		RootMetadata = SourcePointMetadata;
 		NonRootMetadata = nullptr;
 	}
 	else if (AttributeInheritance == EPCGCopyPointsMetadataInheritanceMode::TargetOnly)
@@ -121,7 +124,7 @@ bool FPCGCopyPointsElement::ExecuteInternal(FPCGContext* Context) const
 
 		OutPointData->InitializeFromData(TargetPointData);
 
-		RootMetadata = TargetPointData->Metadata;
+		RootMetadata = TargetPointMetadata;
 		NonRootMetadata = nullptr;
 	}
 	else if (AttributeInheritance == EPCGCopyPointsMetadataInheritanceMode::SourceFirst)
@@ -129,8 +132,8 @@ bool FPCGCopyPointsElement::ExecuteInternal(FPCGContext* Context) const
 		bInheritMetadataFromSource = bSourceHasMetadata || !bTargetHasMetadata;
 
 		OutPointData->InitializeFromData(SourcePointData);
-		RootMetadata = SourcePointData->Metadata;
-		NonRootMetadata = TargetPointData->Metadata;
+		RootMetadata = SourcePointMetadata;
+		NonRootMetadata = TargetPointMetadata;
 	}
 	else if (AttributeInheritance == EPCGCopyPointsMetadataInheritanceMode::TargetFirst)
 	{
@@ -138,8 +141,8 @@ bool FPCGCopyPointsElement::ExecuteInternal(FPCGContext* Context) const
 
 		OutPointData->InitializeFromData(TargetPointData);
 
-		RootMetadata = TargetPointData->Metadata;
-		NonRootMetadata = SourcePointData->Metadata;
+		RootMetadata = TargetPointMetadata;
+		NonRootMetadata = SourcePointMetadata;
 	}
 	else // None
 	{
@@ -152,7 +155,8 @@ bool FPCGCopyPointsElement::ExecuteInternal(FPCGContext* Context) const
 	// Priorize use the target actor from the target, irrespective of the source
 	OutPointData->TargetActor = TargetPointData->TargetActor.IsValid() ? TargetPointData->TargetActor : SourcePointData->TargetActor;
 
-	check(OutPointData->Metadata);
+	UPCGMetadata* OutMetadata = OutPointData->Metadata;
+	check(OutMetadata);
 
 	TArray<FPCGMetadataAttributeBase*> AttributesToSet;
 	TArray<const FPCGMetadataAttributeBase*> NonRootAttributes;
@@ -171,10 +175,10 @@ bool FPCGCopyPointsElement::ExecuteInternal(FPCGContext* Context) const
 
 			for (const FName& AttributeName : AttributeNames)
 			{
-				if (!OutPointData->Metadata->HasAttribute(AttributeName))
+				if (!OutMetadata->HasAttribute(AttributeName))
 				{
 					const FPCGMetadataAttributeBase* Attribute = NonRootMetadata->GetConstAttribute(AttributeName);
-					if (FPCGMetadataAttributeBase* NewAttribute = OutPointData->Metadata->CopyAttribute(Attribute, AttributeName, /*bKeepRoot=*/false, /*bCopyEntries=*/false, /*bCopyValues=*/true))
+					if (FPCGMetadataAttributeBase* NewAttribute = OutMetadata->CopyAttribute(Attribute, AttributeName, /*bKeepRoot=*/false, /*bCopyEntries=*/false, /*bCopyValues=*/true))
 					{
 						AttributesToSet.Add(NewAttribute);
 						NonRootAttributes.Add(Attribute);
@@ -261,7 +265,7 @@ bool FPCGCopyPointsElement::ExecuteInternal(FPCGContext* Context) const
 		{
 			const FPCGPoint* RootPoint = (bInheritMetadataFromSource ? &SourcePoint : &TargetPoint);
 
-			OutPoint.MetadataEntry = OutPointData->Metadata->AddEntryPlaceholder();
+			OutPoint.MetadataEntry = OutMetadata->AddEntryPlaceholder();
 			AllMetadataEntries[Index] = TTuple<int64, int64>(OutPoint.MetadataEntry, RootPoint->MetadataEntry);
 
 			if (NonRootMetadata)
@@ -313,7 +317,7 @@ bool FPCGCopyPointsElement::ExecuteInternal(FPCGContext* Context) const
 			}
 		}
 
-		OutPointData->Metadata->AddDelayedEntries(AllMetadataEntries);
+		OutMetadata->AddDelayedEntries(AllMetadataEntries);
 	}
 
 	return true;
diff --git a/Engine/Plugins/Experimental/PCG/Source/PCG/Public/Data/PCGUnionData.h b/Engine/Plugins/Experimental/PCG/Source/PCG/Public/Data/PCGUnionData.h
--- a/Engine/Plugins/Experimental/PCG/Source/PCG/Public/Data/PCGUnionData.h
+++ b/Engine/Plugins/Experimental/PCG/Source/PCG/Public/Data/PCGUnionData.h
@@ -81,7 +81,7 @@ protected:
 	int CachedDimension = 0;
 
 private:
-	void CreateSequentialPointData(FPCGContext* Context, UPCGPointData* PointData, bool bLeftToRight) const;
+	void CreateSequentialPointData(FPCGContext* Context, TArray<const UPCGSpatialData*>& DataRawPtr, TArray<const UPCGMetadata*>& InputMetadatas, UPCGPointData* PointData, UPCGMetadata* OutMetadata, bool bLeftToRight) const;
 };
 
 #if UE_ENABLE_INCLUDE_ORDER_DEPRECATED_IN_5_2
