[PCG] Expand CRC out to give caching a chance to work, giving big perf gains from warm cache scenarios where reusable work is in the cache.

Does not include anything with the Landscape, which rains on the parade currently because we pessimistically assume it might have changed, which dirties lots of downstream things and negates much of the benefits. I'm having another look into that to see what we can do. Posting this anyway as an early review/illustration.

If i force the landscape CRC to freeze, and use these changes, i get perfect cache behaviour, like regenerating gets 100% cache hits besides uncacheable elements). Regenerating forest in opal was 5x faster. And this is all with the per-point crc disabled.

diff --git a/Engine/Plugins/Experimental/PCG/Source/PCG/Private/Data/PCGDifferenceData.cpp b/Engine/Plugins/Experimental/PCG/Source/PCG/Private/Data/PCGDifferenceData.cpp
--- a/Engine/Plugins/Experimental/PCG/Source/PCG/Private/Data/PCGDifferenceData.cpp
+++ b/Engine/Plugins/Experimental/PCG/Source/PCG/Private/Data/PCGDifferenceData.cpp
@@ -6,6 +6,8 @@
 #include "Data/PCGUnionData.h"
 #include "Helpers/PCGAsync.h"
 
+#include "Serialization/ArchiveCrc32.h"
+
 #include UE_INLINE_GENERATED_CPP_BY_NAME(PCGDifferenceData)
 
 namespace PCGDifferenceDataUtils
@@ -100,6 +102,36 @@ void UPCGDifferenceData::VisitDataNetwork(TFunctionRef<void(const UPCGData*)> Ac
 	}
 }
 
+FPCGCrc UPCGDifferenceData::ComputeCrc() const
+{
+	FArchiveCrc32 Ar;
+	AddToCrc(Ar);
+
+	// Chain together CRCs of operands
+	check(Source);
+	FPCGCrc SourceCrc = Source->GetOrComputeCrc();
+	Ar << SourceCrc;
+
+	if (Difference)
+	{
+		FPCGCrc DifferenceCrc = Difference->GetOrComputeCrc();
+		Ar << DifferenceCrc;
+	}
+
+	Crc = FPCGCrc(Ar.GetCrc());
+
+	return Crc;
+}
+
+void UPCGDifferenceData::AddToCrc(FArchiveCrc32& Ar) const
+{
+	uint32 UniqueTypeID = StaticClass()->GetDefaultObject()->GetUniqueID();
+	Ar << UniqueTypeID;
+
+	uint32 DensityFunctionValue = static_cast<uint32>(DensityFunction);
+	Ar << DensityFunctionValue;
+}
+
 int UPCGDifferenceData::GetDimension() const
 {
 	return Source->GetDimension();
diff --git a/Engine/Plugins/Experimental/PCG/Source/PCG/Private/Data/PCGIntersectionData.cpp b/Engine/Plugins/Experimental/PCG/Source/PCG/Private/Data/PCGIntersectionData.cpp
--- a/Engine/Plugins/Experimental/PCG/Source/PCG/Private/Data/PCGIntersectionData.cpp
+++ b/Engine/Plugins/Experimental/PCG/Source/PCG/Private/Data/PCGIntersectionData.cpp
@@ -6,6 +6,8 @@
 #include "Helpers/PCGAsync.h"
 #include "PCGHelpers.h"
 
+#include "Serialization/ArchiveCrc32.h"
+
 #include UE_INLINE_GENERATED_CPP_BY_NAME(PCGIntersectionData)
 
 namespace PCGIntersectionDataMaths
@@ -69,6 +71,32 @@ void UPCGIntersectionData::VisitDataNetwork(TFunctionRef<void(const UPCGData*)>
 	GetB()->VisitDataNetwork(Action);
 }
 
+FPCGCrc UPCGIntersectionData::ComputeCrc() const
+{
+	FArchiveCrc32 Ar;
+	AddToCrc(Ar);
+
+	// Chain together CRCs of operands
+	check(A && B);
+	uint32 CrcA = A->GetOrComputeCrc().GetValue();
+	uint32 CrcB = B->GetOrComputeCrc().GetValue();
+	Ar << CrcA;
+	Ar << CrcB;
+
+	Crc = FPCGCrc(Ar.GetCrc());
+
+	return Crc;
+}
+
+void UPCGIntersectionData::AddToCrc(FArchiveCrc32& Ar) const
+{
+	uint32 UniqueTypeID = StaticClass()->GetDefaultObject()->GetUniqueID();
+	Ar << UniqueTypeID;
+
+	uint32 DensityFunctionValue = static_cast<uint32>(DensityFunction);
+	Ar << DensityFunctionValue;
+}
+
 int UPCGIntersectionData::GetDimension() const
 {
 	check(GetA() && GetB());
diff --git a/Engine/Plugins/Experimental/PCG/Source/PCG/Private/Data/PCGPrimitiveData.cpp b/Engine/Plugins/Experimental/PCG/Source/PCG/Private/Data/PCGPrimitiveData.cpp
--- a/Engine/Plugins/Experimental/PCG/Source/PCG/Private/Data/PCGPrimitiveData.cpp
+++ b/Engine/Plugins/Experimental/PCG/Source/PCG/Private/Data/PCGPrimitiveData.cpp
@@ -6,7 +6,10 @@
 #include "Data/PCGPointData.h"
 #include "Data/PCGSpatialData.h"
 #include "Elements/PCGVolumeSampler.h"
+
+#include "Components/BoxComponent.h"
 #include "Components/PrimitiveComponent.h"
+#include "Serialization/ArchiveCrc32.h"
 
 #include UE_INLINE_GENERATED_CPP_BY_NAME(PCGPrimitiveData)
 
@@ -19,6 +22,26 @@ void UPCGPrimitiveData::Initialize(UPrimitiveComponent* InPrimitive)
 	// Not obvious to find strict bounds, leave at the default value
 }
 
+void UPCGPrimitiveData::AddToCrc(FArchiveCrc32& Ar) const
+{
+	check(Primitive.Get());
+
+	// Can implement specific CRC's for different primitive types.
+	if (UBoxComponent* Box = Cast<UBoxComponent>(Primitive.Get()))
+	{
+		FVector Extent = Box->GetUnscaledBoxExtent();
+		Ar << Extent;
+
+		FTransform Transform = Box->GetComponentTransform();
+		Ar << Transform;
+	}
+	else
+	{
+		// Fallback is the UID which changes when data recreated.
+		Super::AddToCrc(Ar);
+	}
+}
+
 bool UPCGPrimitiveData::SamplePoint(const FTransform& InTransform, const FBox& InBounds, FPCGPoint& OutPoint, UPCGMetadata* OutMetadata) const
 {
 	// Pure implementation
diff --git a/Engine/Plugins/Experimental/PCG/Source/PCG/Private/Data/PCGSplineData.cpp b/Engine/Plugins/Experimental/PCG/Source/PCG/Private/Data/PCGSplineData.cpp
--- a/Engine/Plugins/Experimental/PCG/Source/PCG/Private/Data/PCGSplineData.cpp
+++ b/Engine/Plugins/Experimental/PCG/Source/PCG/Private/Data/PCGSplineData.cpp
@@ -10,6 +10,7 @@
 
 #include "Components/SplineComponent.h"
 #include "Data/PCGSpatialData.h"
+#include "Serialization/ArchiveCrc32.h"
 
 #include UE_INLINE_GENERATED_CPP_BY_NAME(PCGSplineData)
 
@@ -33,6 +34,22 @@ void UPCGSplineData::Initialize(USplineComponent* InSpline)
 	CachedBounds = CachedBounds.ExpandBy(SplinePointsRadius, SplinePointsRadius);
 }
 
+void UPCGSplineData::AddToCrc(FArchiveCrc32& Ar) const
+{
+	check(Spline);
+
+	Ar << const_cast<FBox&>(CachedBounds);
+
+	Ar << const_cast<FInterpCurveVector&>(Spline->GetSplinePointsPosition());
+	Ar << const_cast<FInterpCurveQuat&>(Spline->GetSplinePointsRotation());
+	Ar << const_cast<FInterpCurveVector&>(Spline->GetSplinePointsScale());
+
+	Ar << const_cast<FTransform&>(Spline->GetComponentTransform());
+
+	bool bIsClosed = Spline->IsClosedLoop();
+	Ar << bIsClosed;
+}
+
 FTransform UPCGSplineData::GetTransform() const
 {
 	return Spline ? Spline->GetComponentTransform() : FTransform::Identity;
diff --git a/Engine/Plugins/Experimental/PCG/Source/PCG/Private/Data/PCGUnionData.cpp b/Engine/Plugins/Experimental/PCG/Source/PCG/Private/Data/PCGUnionData.cpp
--- a/Engine/Plugins/Experimental/PCG/Source/PCG/Private/Data/PCGUnionData.cpp
+++ b/Engine/Plugins/Experimental/PCG/Source/PCG/Private/Data/PCGUnionData.cpp
@@ -7,6 +7,8 @@
 #include "Metadata/PCGMetadataAccessor.h"
 #include "PCGHelpers.h"
 
+#include "Serialization/ArchiveCrc32.h"
+
 #include UE_INLINE_GENERATED_CPP_BY_NAME(PCGUnionData)
 
 namespace PCGUnionDataMaths
@@ -81,6 +83,41 @@ void UPCGUnionData::VisitDataNetwork(TFunctionRef<void(const UPCGData*)> Action)
 	}
 }
 
+FPCGCrc UPCGUnionData::ComputeCrc() const
+{
+	FArchiveCrc32 Ar;
+	AddToCrc(Ar);
+
+	// Chain together CRCs of operands
+	int32 NumOperands = Data.Num();
+	Ar << NumOperands;
+
+	for (TObjectPtr<const UPCGSpatialData> Datum : Data)
+	{
+		if (Datum)
+		{
+			uint32 DatumCrc = Datum->GetOrComputeCrc().GetValue();
+			Ar << DatumCrc;
+		}
+	}
+
+	Crc = FPCGCrc(Ar.GetCrc());
+
+	return Crc;
+}
+
+void UPCGUnionData::AddToCrc(FArchiveCrc32& Ar) const
+{
+	uint32 UniqueTypeID = StaticClass()->GetDefaultObject()->GetUniqueID();
+	Ar << UniqueTypeID;
+
+	uint32 UnionTypeValue = static_cast<uint32>(UnionType);
+	Ar << UnionTypeValue;
+
+	uint32 DensityFunctionValue = static_cast<uint32>(DensityFunction);
+	Ar << DensityFunctionValue;
+}
+
 int UPCGUnionData::GetDimension() const
 {
 	return CachedDimension;
diff --git a/Engine/Plugins/Experimental/PCG/Source/PCG/Private/Data/PCGVolumeData.cpp b/Engine/Plugins/Experimental/PCG/Source/PCG/Private/Data/PCGVolumeData.cpp
--- a/Engine/Plugins/Experimental/PCG/Source/PCG/Private/Data/PCGVolumeData.cpp
+++ b/Engine/Plugins/Experimental/PCG/Source/PCG/Private/Data/PCGVolumeData.cpp
@@ -6,10 +6,23 @@
 #include "PCGHelpers.h"
 #include "Elements/PCGVolumeSampler.h"
 
+#include "Components/BrushComponent.h"
 #include "GameFramework/Volume.h"
+#include "Serialization/ArchiveCrc32.h"
+#include "Serialization/ArchiveObjectCrc32.h"
 
 #include UE_INLINE_GENERATED_CPP_BY_NAME(PCGVolumeData)
 
+/** Does CRC on Brush data within BrushComponent. */
+class FPCGBrushCrc32 : public FArchiveObjectCrc32
+{
+public:
+	virtual bool ShouldSkipProperty(const FProperty* InProperty) const override
+	{
+		return InProperty && (InProperty->GetFName() != FName("Brush"));
+	}
+};
+
 void UPCGVolumeData::Initialize(AVolume* InVolume, AActor* InTargetActor)
 {
 	check(InVolume);
@@ -30,6 +43,25 @@ void UPCGVolumeData::Initialize(const FBox& InBounds, AActor* InTargetActor)
 	TargetActor = InTargetActor;
 }
 
+void UPCGVolumeData::AddToCrc(FArchiveCrc32& Ar) const
+{
+	Ar << const_cast<FBox&>(Bounds);
+
+	Ar << const_cast<FBox&>(StrictBounds);
+
+	if (Volume.IsValid())
+	{
+		if (UBrushComponent* Brush = Volume->GetBrushComponent())
+		{
+			Ar << const_cast<FTransform&>(Brush->GetComponentTransform());
+
+			FPCGBrushCrc32 VolumeCrc;
+			uint32 BrushCrc = VolumeCrc.Crc32(Brush);
+			Ar << BrushCrc;
+		}
+	}
+}
+
 FBox UPCGVolumeData::GetBounds() const
 {
 	return Bounds;
diff --git a/Engine/Plugins/Experimental/PCG/Source/PCG/Private/PCGData.cpp b/Engine/Plugins/Experimental/PCG/Source/PCG/Private/PCGData.cpp
--- a/Engine/Plugins/Experimental/PCG/Source/PCG/Private/PCGData.cpp
+++ b/Engine/Plugins/Experimental/PCG/Source/PCG/Private/PCGData.cpp
@@ -89,6 +89,13 @@ FPCGCrc UPCGData::GetOrComputeCrc() const
 		return Crc;
 	}
 
+	Crc = ComputeCrc();
+
+	return Crc;
+}
+
+FPCGCrc UPCGData::ComputeCrc() const
+{
 	FArchiveCrc32 Ar;
 	AddToCrc(Ar);
 
diff --git a/Engine/Plugins/Experimental/PCG/Source/PCG/Public/Data/PCGDifferenceData.h b/Engine/Plugins/Experimental/PCG/Source/PCG/Public/Data/PCGDifferenceData.h
--- a/Engine/Plugins/Experimental/PCG/Source/PCG/Public/Data/PCGDifferenceData.h
+++ b/Engine/Plugins/Experimental/PCG/Source/PCG/Public/Data/PCGDifferenceData.h
@@ -50,8 +50,13 @@ public:
 	// ~Begin UPCGData interface
 	virtual EPCGDataType GetDataType() const override { return EPCGDataType::Spatial; }
 	virtual void VisitDataNetwork(TFunctionRef<void(const UPCGData*)> Action) const override;
+
+protected:
+	virtual FPCGCrc ComputeCrc() const override;
+	virtual void AddToCrc(FArchiveCrc32& Ar) const override;
 	// ~End UPCGData interface
 
+public:
 	//~Begin UPCGSpatialData interface
 	virtual int GetDimension() const override;
 	virtual FBox GetBounds() const override;
diff --git a/Engine/Plugins/Experimental/PCG/Source/PCG/Public/Data/PCGIntersectionData.h b/Engine/Plugins/Experimental/PCG/Source/PCG/Public/Data/PCGIntersectionData.h
--- a/Engine/Plugins/Experimental/PCG/Source/PCG/Public/Data/PCGIntersectionData.h
+++ b/Engine/Plugins/Experimental/PCG/Source/PCG/Public/Data/PCGIntersectionData.h
@@ -33,8 +33,13 @@ public:
 	// ~Begin UPCGData interface
 	virtual EPCGDataType GetDataType() const override { return EPCGDataType::Spatial; }
 	virtual void VisitDataNetwork(TFunctionRef<void(const UPCGData*)> Action) const override;
+
+protected:
+	virtual FPCGCrc ComputeCrc() const override;
+	virtual void AddToCrc(FArchiveCrc32& Ar) const override;
 	// ~End UPCGData interface
 
+public:
 	//~Begin UPCGSpatialData interface
 	virtual int GetDimension() const override;
 	virtual FBox GetBounds() const override;
diff --git a/Engine/Plugins/Experimental/PCG/Source/PCG/Public/Data/PCGPrimitiveData.h b/Engine/Plugins/Experimental/PCG/Source/PCG/Public/Data/PCGPrimitiveData.h
--- a/Engine/Plugins/Experimental/PCG/Source/PCG/Public/Data/PCGPrimitiveData.h
+++ b/Engine/Plugins/Experimental/PCG/Source/PCG/Public/Data/PCGPrimitiveData.h
@@ -16,6 +16,7 @@ public:
 
 	// ~Begin UPCGData interface
 	virtual EPCGDataType GetDataType() const override { return EPCGDataType::Primitive; }
+	void AddToCrc(FArchiveCrc32& Ar) const override;
 	// ~End UPCGData interface
 
 	// ~Begin UPCGSpatialData interface
diff --git a/Engine/Plugins/Experimental/PCG/Source/PCG/Public/Data/PCGSplineData.h b/Engine/Plugins/Experimental/PCG/Source/PCG/Public/Data/PCGSplineData.h
--- a/Engine/Plugins/Experimental/PCG/Source/PCG/Public/Data/PCGSplineData.h
+++ b/Engine/Plugins/Experimental/PCG/Source/PCG/Public/Data/PCGSplineData.h
@@ -24,6 +24,7 @@ public:
 
 	// ~Begin UPCGData interface
 	virtual EPCGDataType GetDataType() const override { return EPCGDataType::Spline; }
+	void AddToCrc(FArchiveCrc32& Ar) const override;
 	// ~End UPCGData interface
 
 	//~Begin UPCGPolyLineData interface
@@ -43,6 +44,7 @@ public:
 	virtual FBox GetBounds() const override;
 	virtual bool SamplePoint(const FTransform& Transform, const FBox& Bounds, FPCGPoint& OutPoint, UPCGMetadata* OutMetadata) const override;
 	virtual UPCGSpatialData* ProjectOn(const UPCGSpatialData* InOther, const FPCGProjectionParams& InParams = FPCGProjectionParams()) const override;
+
 protected:
 	virtual UPCGSpatialData* CopyInternal() const override;
 	//~End UPCGSpatialData interface
diff --git a/Engine/Plugins/Experimental/PCG/Source/PCG/Public/Data/PCGUnionData.h b/Engine/Plugins/Experimental/PCG/Source/PCG/Public/Data/PCGUnionData.h
--- a/Engine/Plugins/Experimental/PCG/Source/PCG/Public/Data/PCGUnionData.h
+++ b/Engine/Plugins/Experimental/PCG/Source/PCG/Public/Data/PCGUnionData.h
@@ -39,8 +39,13 @@ public:
 	// ~Begin UPCGData interface
 	virtual EPCGDataType GetDataType() const override { return EPCGDataType::Spatial; }
 	virtual void VisitDataNetwork(TFunctionRef<void(const UPCGData*)> Action) const override;
+
+protected:
+	virtual FPCGCrc ComputeCrc() const override;
+	virtual void AddToCrc(FArchiveCrc32& Ar) const override;
 	// ~End UPCGData interface
 
+public:
 	//~Begin UPCGSpatialData interface
 	virtual int GetDimension() const override;
 	virtual FBox GetBounds() const override;
diff --git a/Engine/Plugins/Experimental/PCG/Source/PCG/Public/Data/PCGVolumeData.h b/Engine/Plugins/Experimental/PCG/Source/PCG/Public/Data/PCGVolumeData.h
--- a/Engine/Plugins/Experimental/PCG/Source/PCG/Public/Data/PCGVolumeData.h
+++ b/Engine/Plugins/Experimental/PCG/Source/PCG/Public/Data/PCGVolumeData.h
@@ -19,6 +19,7 @@ public:
 
 	// ~Begin UPCGData interface
 	virtual EPCGDataType GetDataType() const override { return EPCGDataType::Volume; }
+	virtual void AddToCrc(FArchiveCrc32& Ar) const;
 	// ~End UPCGData interface
 
 	// ~Begin UPGCSpatialData interface
diff --git a/Engine/Plugins/Experimental/PCG/Source/PCG/Public/PCGData.h b/Engine/Plugins/Experimental/PCG/Source/PCG/Public/PCGData.h
--- a/Engine/Plugins/Experimental/PCG/Source/PCG/Public/PCGData.h
+++ b/Engine/Plugins/Experimental/PCG/Source/PCG/Public/PCGData.h
@@ -29,12 +29,9 @@ public:
 	UPCGData(const FObjectInitializer& ObjectInitializer);
 	virtual EPCGDataType GetDataType() const { return EPCGDataType::None; }
 
-	/** Computes a Crc for this data. */
+	/** Returns a Crc for this and any connected data. */
 	FPCGCrc GetOrComputeCrc() const;
 
-	/** Contributes to Crc. Fallback implementation that writes object instance UID. */
-	virtual void AddToCrc(FArchiveCrc32& Ar) const;
-
 	/** Executes a lambda over all connected data objects. */
 	virtual void VisitDataNetwork(TFunctionRef<void(const UPCGData*)> Action) const;
 
@@ -45,6 +42,13 @@ public:
 	/** CRC for this object instance. */
 	mutable FPCGCrc Crc;
 
+protected:
+	/** Computes Crc for this and any connected data. */
+	virtual FPCGCrc ComputeCrc() const;
+
+	/** Adds this data to Crc. Fallback implementation writes object instance UID. */
+	virtual void AddToCrc(FArchiveCrc32& Ar) const;
+
 private:
 	/** Serves unique ID values to instances of this object. */
 	static inline std::atomic<uint64> UIDCounter{ 1 };
