[PCG] Attempt to make pointer safety a little more explicit for static analysis

diff --git a/Engine/Plugins/Experimental/PCG/Source/PCG/Private/Data/PCGUnionData.cpp b/Engine/Plugins/Experimental/PCG/Source/PCG/Private/Data/PCGUnionData.cpp
--- a/Engine/Plugins/Experimental/PCG/Source/PCG/Private/Data/PCGUnionData.cpp
+++ b/Engine/Plugins/Experimental/PCG/Source/PCG/Private/Data/PCGUnionData.cpp
@@ -399,7 +399,8 @@ void UPCGUnionData::CreateSequentialPointData(FPCGContext* Context, TArray<const
 				return false;
 			}
 
-			check(PointData && OutMetadata);
+			check(PointData);
+			check(OutMetadata);
 
 			OutPoint = Point;
 			if (OutMetadata->GetParent() != InputMetadatas[DataIndex])
