[PCG] Remove unwanted dllexport macro in PCGAsync

#rnx

diff --git a/Engine/Plugins/Experimental/PCG/Source/PCG/Public/Helpers/PCGAsync.h b/Engine/Plugins/Experimental/PCG/Source/PCG/Public/Helpers/PCGAsync.h
--- a/Engine/Plugins/Experimental/PCG/Source/PCG/Public/Helpers/PCGAsync.h
+++ b/Engine/Plugins/Experimental/PCG/Source/PCG/Public/Helpers/PCGAsync.h
@@ -318,7 +318,7 @@ namespace FPCGAsync
 	* @returns true if the processing is done, false otherwise. Use this to know if you need to reschedule the task.
 	*/
 	template <typename OutputType, typename Func>
-	PCG_API bool AsyncProcessing(FPCGAsyncState* AsyncState, int32 NumIterations, TArray<OutputType>& OutData, Func InFunc, const bool bEnableTimeSlicing, const int32 ChunkSize = 16)
+	bool AsyncProcessing(FPCGAsyncState* AsyncState, int32 NumIterations, TArray<OutputType>& OutData, Func InFunc, const bool bEnableTimeSlicing, const int32 ChunkSize = 16)
 	{
 		auto IterationInnerLoop = [&InFunc, &OutData](int32 StartReadIndex, int32 StartWriteIndex, int32 Count) -> int32
 		{
