[PCG] Fix issue where timesliced async tasks would be run on 1 thread

#rnx

diff --git a/Engine/Plugins/Experimental/PCG/Source/PCG/Public/Helpers/PCGAsync.h b/Engine/Plugins/Experimental/PCG/Source/PCG/Public/Helpers/PCGAsync.h
--- a/Engine/Plugins/Experimental/PCG/Source/PCG/Public/Helpers/PCGAsync.h
+++ b/Engine/Plugins/Experimental/PCG/Source/PCG/Public/Helpers/PCGAsync.h
@@ -339,8 +339,9 @@ namespace FPCGAsync
 		if (AsyncState && !AsyncState->bIsRunningAsyncCall)
 		{
 			AsyncState->bIsRunningAsyncCall = true;
-			return Private::AsyncProcessing<OutputType>(*AsyncState, NumIterations, OutData, IterationInnerLoop, bEnableTimeSlicing, ChunkSize);
+			bool bIsDone = Private::AsyncProcessing<OutputType>(*AsyncState, NumIterations, OutData, IterationInnerLoop, bEnableTimeSlicing, ChunkSize);
 			AsyncState->bIsRunningAsyncCall = false;
+			return bIsDone;
 		}
 		else
 		{
@@ -348,6 +349,7 @@ namespace FPCGAsync
 			// We also force using one thread (the current one).
 			FPCGAsyncState DummyState;
 			DummyState.NumAvailableTasks = 1;
+			DummyState.bIsRunningAsyncCall = true;
 			return Private::AsyncProcessing<OutputType>(DummyState, NumIterations, OutData, IterationInnerLoop, /*bEnableTimeSlicing=*/false, ChunkSize);
 		}
 	}
