Partial rollback on crc changes because we can run into issues when the associated live data is destroyed in the editor

#rnx

diff --git a/Engine/Plugins/Experimental/PCG/Source/PCG/Private/Data/PCGDifferenceData.cpp b/Engine/Plugins/Experimental/PCG/Source/PCG/Private/Data/PCGDifferenceData.cpp
--- a/Engine/Plugins/Experimental/PCG/Source/PCG/Private/Data/PCGDifferenceData.cpp
+++ b/Engine/Plugins/Experimental/PCG/Source/PCG/Private/Data/PCGDifferenceData.cpp
@@ -118,9 +118,7 @@ FPCGCrc UPCGDifferenceData::ComputeCrc() const
 		Ar << DifferenceCrc;
 	}
 
-	Crc = FPCGCrc(Ar.GetCrc());
-
-	return Crc;
+	return FPCGCrc(Ar.GetCrc());
 }
 
 void UPCGDifferenceData::AddToCrc(FArchiveCrc32& Ar) const
diff --git a/Engine/Plugins/Experimental/PCG/Source/PCG/Private/Data/PCGPrimitiveData.cpp b/Engine/Plugins/Experimental/PCG/Source/PCG/Private/Data/PCGPrimitiveData.cpp
--- a/Engine/Plugins/Experimental/PCG/Source/PCG/Private/Data/PCGPrimitiveData.cpp
+++ b/Engine/Plugins/Experimental/PCG/Source/PCG/Private/Data/PCGPrimitiveData.cpp
@@ -24,18 +24,20 @@ void UPCGPrimitiveData::Initialize(UPrimitiveComponent* InPrimitive)
 
 void UPCGPrimitiveData::AddToCrc(FArchiveCrc32& Ar) const
 {
-	check(Primitive.Get());
-
-	// Can implement specific CRC's for different primitive types.
-	if (UBoxComponent* Box = Cast<UBoxComponent>(Primitive.Get()))
-	{
-		FVector Extent = Box->GetUnscaledBoxExtent();
-		Ar << Extent;
-
-		FTransform Transform = Box->GetComponentTransform();
-		Ar << Transform;
-	}
-	else
+	//TODO: we need a way to keep the primitive alive or some mechanism to eject this from the cache if it's deleted
+	// Until then, use the unique id.
+	//check(Primitive.Get());
+
+	//// Can implement specific CRC's for different primitive types.
+	//if (UBoxComponent* Box = Cast<UBoxComponent>(Primitive.Get()))
+	//{
+	//	FVector Extent = Box->GetUnscaledBoxExtent();
+	//	Ar << Extent;
+
+	//	FTransform Transform = Box->GetComponentTransform();
+	//	Ar << Transform;
+	//}
+	//else
 	{
 		// Fallback is the UID which changes when data recreated.
 		Super::AddToCrc(Ar);
diff --git a/Engine/Plugins/Experimental/PCG/Source/PCG/Private/Data/PCGSplineData.cpp b/Engine/Plugins/Experimental/PCG/Source/PCG/Private/Data/PCGSplineData.cpp
--- a/Engine/Plugins/Experimental/PCG/Source/PCG/Private/Data/PCGSplineData.cpp
+++ b/Engine/Plugins/Experimental/PCG/Source/PCG/Private/Data/PCGSplineData.cpp
@@ -36,18 +36,21 @@ void UPCGSplineData::Initialize(USplineComponent* InSpline)
 
 void UPCGSplineData::AddToCrc(FArchiveCrc32& Ar) const
 {
-	check(Spline);
+	// TODO: do not rely on having a live spline if we want to CRC this.
+	Super::AddToCrc(Ar);
 
-	Ar << const_cast<FBox&>(CachedBounds);
+	//check(Spline);
 
-	Ar << const_cast<FInterpCurveVector&>(Spline->GetSplinePointsPosition());
-	Ar << const_cast<FInterpCurveQuat&>(Spline->GetSplinePointsRotation());
-	Ar << const_cast<FInterpCurveVector&>(Spline->GetSplinePointsScale());
+	//Ar << const_cast<FBox&>(CachedBounds);
 
-	Ar << const_cast<FTransform&>(Spline->GetComponentTransform());
+	//Ar << const_cast<FInterpCurveVector&>(Spline->GetSplinePointsPosition());
+	//Ar << const_cast<FInterpCurveQuat&>(Spline->GetSplinePointsRotation());
+	//Ar << const_cast<FInterpCurveVector&>(Spline->GetSplinePointsScale());
 
-	bool bIsClosed = Spline->IsClosedLoop();
-	Ar << bIsClosed;
+	//Ar << const_cast<FTransform&>(Spline->GetComponentTransform());
+
+	//bool bIsClosed = Spline->IsClosedLoop();
+	//Ar << bIsClosed;
 }
 
 FTransform UPCGSplineData::GetTransform() const
diff --git a/Engine/Plugins/Experimental/PCG/Source/PCG/Private/Data/PCGUnionData.cpp b/Engine/Plugins/Experimental/PCG/Source/PCG/Private/Data/PCGUnionData.cpp
--- a/Engine/Plugins/Experimental/PCG/Source/PCG/Private/Data/PCGUnionData.cpp
+++ b/Engine/Plugins/Experimental/PCG/Source/PCG/Private/Data/PCGUnionData.cpp
@@ -101,9 +101,7 @@ FPCGCrc UPCGUnionData::ComputeCrc() const
 		}
 	}
 
-	Crc = FPCGCrc(Ar.GetCrc());
-
-	return Crc;
+	return FPCGCrc(Ar.GetCrc());
 }
 
 void UPCGUnionData::AddToCrc(FArchiveCrc32& Ar) const
diff --git a/Engine/Plugins/Experimental/PCG/Source/PCG/Private/Data/PCGVolumeData.cpp b/Engine/Plugins/Experimental/PCG/Source/PCG/Private/Data/PCGVolumeData.cpp
--- a/Engine/Plugins/Experimental/PCG/Source/PCG/Private/Data/PCGVolumeData.cpp
+++ b/Engine/Plugins/Experimental/PCG/Source/PCG/Private/Data/PCGVolumeData.cpp
@@ -45,21 +45,26 @@ void UPCGVolumeData::Initialize(const FBox& InBounds, AActor* InTargetActor)
 
 void UPCGVolumeData::AddToCrc(FArchiveCrc32& Ar) const
 {
-	Ar << const_cast<FBox&>(Bounds);
-
-	Ar << const_cast<FBox&>(StrictBounds);
-
+	// Until we can track the destruction of the volume, we must use the UID
 	if (Volume.IsValid())
 	{
-		if (UBrushComponent* Brush = Volume->GetBrushComponent())
-		{
-			Ar << const_cast<FTransform&>(Brush->GetComponentTransform());
-
-			FPCGBrushCrc32 VolumeCrc;
-			uint32 BrushCrc = VolumeCrc.Crc32(Brush);
-			Ar << BrushCrc;
-		}
+		return Super::AddToCrc(Ar);
 	}
+
+	Ar << const_cast<FBox&>(Bounds);
+	Ar << const_cast<FBox&>(StrictBounds);
+
+	//if (Volume.IsValid())
+	//{
+	//	if (UBrushComponent* Brush = Volume->GetBrushComponent())
+	//	{
+	//		Ar << const_cast<FTransform&>(Brush->GetComponentTransform());
+
+	//		FPCGBrushCrc32 VolumeCrc;
+	//		uint32 BrushCrc = VolumeCrc.Crc32(Brush);
+	//		Ar << BrushCrc;
+	//	}
+	//}
 }
 
 FBox UPCGVolumeData::GetBounds() const
diff --git a/Engine/Plugins/Experimental/PCG/Source/PCG/Private/PCGData.cpp b/Engine/Plugins/Experimental/PCG/Source/PCG/Private/PCGData.cpp
--- a/Engine/Plugins/Experimental/PCG/Source/PCG/Private/PCGData.cpp
+++ b/Engine/Plugins/Experimental/PCG/Source/PCG/Private/PCGData.cpp
@@ -84,13 +84,11 @@ UPCGData::UPCGData(const FObjectInitializer& ObjectInitializer)
 
 FPCGCrc UPCGData::GetOrComputeCrc() const
 {
-	if (Crc.IsValid())
+	if (!Crc.IsValid())
 	{
-		return Crc;
+		Crc = ComputeCrc();
 	}
 
-	Crc = ComputeCrc();
-
 	return Crc;
 }
 
@@ -99,9 +97,7 @@ FPCGCrc UPCGData::ComputeCrc() const
 	FArchiveCrc32 Ar;
 	AddToCrc(Ar);
 
-	Crc = FPCGCrc(Ar.GetCrc());
-
-	return Crc;
+	return FPCGCrc(Ar.GetCrc());
 }
 
 void UPCGData::AddToCrc(FArchiveCrc32& Ar) const
