[Backout] - CL24530585
Original CL Desc
-----------------------------------------------------------------
Partial rollback on crc changes because we can run into issues when the associated live data is destroyed in the editor

#rnx

diff --git a/Engine/Plugins/Experimental/PCG/Source/PCG/Private/Data/PCGDifferenceData.cpp b/Engine/Plugins/Experimental/PCG/Source/PCG/Private/Data/PCGDifferenceData.cpp
--- a/Engine/Plugins/Experimental/PCG/Source/PCG/Private/Data/PCGDifferenceData.cpp
+++ b/Engine/Plugins/Experimental/PCG/Source/PCG/Private/Data/PCGDifferenceData.cpp
@@ -118,7 +118,9 @@ FPCGCrc UPCGDifferenceData::ComputeCrc() const
 		Ar << DifferenceCrc;
 	}
 
-	return FPCGCrc(Ar.GetCrc());
+	Crc = FPCGCrc(Ar.GetCrc());
+
+	return Crc;
 }
 
 void UPCGDifferenceData::AddToCrc(FArchiveCrc32& Ar) const
diff --git a/Engine/Plugins/Experimental/PCG/Source/PCG/Private/Data/PCGPrimitiveData.cpp b/Engine/Plugins/Experimental/PCG/Source/PCG/Private/Data/PCGPrimitiveData.cpp
--- a/Engine/Plugins/Experimental/PCG/Source/PCG/Private/Data/PCGPrimitiveData.cpp
+++ b/Engine/Plugins/Experimental/PCG/Source/PCG/Private/Data/PCGPrimitiveData.cpp
@@ -24,20 +24,18 @@ void UPCGPrimitiveData::Initialize(UPrimitiveComponent* InPrimitive)
 
 void UPCGPrimitiveData::AddToCrc(FArchiveCrc32& Ar) const
 {
-	//TODO: we need a way to keep the primitive alive or some mechanism to eject this from the cache if it's deleted
-	// Until then, use the unique id.
-	//check(Primitive.Get());
-
-	//// Can implement specific CRC's for different primitive types.
-	//if (UBoxComponent* Box = Cast<UBoxComponent>(Primitive.Get()))
-	//{
-	//	FVector Extent = Box->GetUnscaledBoxExtent();
-	//	Ar << Extent;
-
-	//	FTransform Transform = Box->GetComponentTransform();
-	//	Ar << Transform;
-	//}
-	//else
+	check(Primitive.Get());
+
+	// Can implement specific CRC's for different primitive types.
+	if (UBoxComponent* Box = Cast<UBoxComponent>(Primitive.Get()))
+	{
+		FVector Extent = Box->GetUnscaledBoxExtent();
+		Ar << Extent;
+
+		FTransform Transform = Box->GetComponentTransform();
+		Ar << Transform;
+	}
+	else
 	{
 		// Fallback is the UID which changes when data recreated.
 		Super::AddToCrc(Ar);
diff --git a/Engine/Plugins/Experimental/PCG/Source/PCG/Private/Data/PCGSplineData.cpp b/Engine/Plugins/Experimental/PCG/Source/PCG/Private/Data/PCGSplineData.cpp
--- a/Engine/Plugins/Experimental/PCG/Source/PCG/Private/Data/PCGSplineData.cpp
+++ b/Engine/Plugins/Experimental/PCG/Source/PCG/Private/Data/PCGSplineData.cpp
@@ -36,21 +36,18 @@ void UPCGSplineData::Initialize(USplineComponent* InSpline)
 
 void UPCGSplineData::AddToCrc(FArchiveCrc32& Ar) const
 {
-	// TODO: do not rely on having a live spline if we want to CRC this.
-	Super::AddToCrc(Ar);
+	check(Spline);
 
-	//check(Spline);
+	Ar << const_cast<FBox&>(CachedBounds);
 
-	//Ar << const_cast<FBox&>(CachedBounds);
+	Ar << const_cast<FInterpCurveVector&>(Spline->GetSplinePointsPosition());
+	Ar << const_cast<FInterpCurveQuat&>(Spline->GetSplinePointsRotation());
+	Ar << const_cast<FInterpCurveVector&>(Spline->GetSplinePointsScale());
 
-	//Ar << const_cast<FInterpCurveVector&>(Spline->GetSplinePointsPosition());
-	//Ar << const_cast<FInterpCurveQuat&>(Spline->GetSplinePointsRotation());
-	//Ar << const_cast<FInterpCurveVector&>(Spline->GetSplinePointsScale());
+	Ar << const_cast<FTransform&>(Spline->GetComponentTransform());
 
-	//Ar << const_cast<FTransform&>(Spline->GetComponentTransform());
-
-	//bool bIsClosed = Spline->IsClosedLoop();
-	//Ar << bIsClosed;
+	bool bIsClosed = Spline->IsClosedLoop();
+	Ar << bIsClosed;
 }
 
 FTransform UPCGSplineData::GetTransform() const
diff --git a/Engine/Plugins/Experimental/PCG/Source/PCG/Private/Data/PCGUnionData.cpp b/Engine/Plugins/Experimental/PCG/Source/PCG/Private/Data/PCGUnionData.cpp
--- a/Engine/Plugins/Experimental/PCG/Source/PCG/Private/Data/PCGUnionData.cpp
+++ b/Engine/Plugins/Experimental/PCG/Source/PCG/Private/Data/PCGUnionData.cpp
@@ -101,7 +101,9 @@ FPCGCrc UPCGUnionData::ComputeCrc() const
 		}
 	}
 
-	return FPCGCrc(Ar.GetCrc());
+	Crc = FPCGCrc(Ar.GetCrc());
+
+	return Crc;
 }
 
 void UPCGUnionData::AddToCrc(FArchiveCrc32& Ar) const
diff --git a/Engine/Plugins/Experimental/PCG/Source/PCG/Private/Data/PCGVolumeData.cpp b/Engine/Plugins/Experimental/PCG/Source/PCG/Private/Data/PCGVolumeData.cpp
--- a/Engine/Plugins/Experimental/PCG/Source/PCG/Private/Data/PCGVolumeData.cpp
+++ b/Engine/Plugins/Experimental/PCG/Source/PCG/Private/Data/PCGVolumeData.cpp
@@ -45,26 +45,21 @@ void UPCGVolumeData::Initialize(const FBox& InBounds, AActor* InTargetActor)
 
 void UPCGVolumeData::AddToCrc(FArchiveCrc32& Ar) const
 {
-	// Until we can track the destruction of the volume, we must use the UID
-	if (Volume.IsValid())
-	{
-		return Super::AddToCrc(Ar);
-	}
-
 	Ar << const_cast<FBox&>(Bounds);
+
 	Ar << const_cast<FBox&>(StrictBounds);
 
-	//if (Volume.IsValid())
-	//{
-	//	if (UBrushComponent* Brush = Volume->GetBrushComponent())
-	//	{
-	//		Ar << const_cast<FTransform&>(Brush->GetComponentTransform());
-
-	//		FPCGBrushCrc32 VolumeCrc;
-	//		uint32 BrushCrc = VolumeCrc.Crc32(Brush);
-	//		Ar << BrushCrc;
-	//	}
-	//}
+	if (Volume.IsValid())
+	{
+		if (UBrushComponent* Brush = Volume->GetBrushComponent())
+		{
+			Ar << const_cast<FTransform&>(Brush->GetComponentTransform());
+
+			FPCGBrushCrc32 VolumeCrc;
+			uint32 BrushCrc = VolumeCrc.Crc32(Brush);
+			Ar << BrushCrc;
+		}
+	}
 }
 
 FBox UPCGVolumeData::GetBounds() const
diff --git a/Engine/Plugins/Experimental/PCG/Source/PCG/Private/PCGData.cpp b/Engine/Plugins/Experimental/PCG/Source/PCG/Private/PCGData.cpp
--- a/Engine/Plugins/Experimental/PCG/Source/PCG/Private/PCGData.cpp
+++ b/Engine/Plugins/Experimental/PCG/Source/PCG/Private/PCGData.cpp
@@ -84,11 +84,13 @@ UPCGData::UPCGData(const FObjectInitializer& ObjectInitializer)
 
 FPCGCrc UPCGData::GetOrComputeCrc() const
 {
-	if (!Crc.IsValid())
+	if (Crc.IsValid())
 	{
-		Crc = ComputeCrc();
+		return Crc;
 	}
 
+	Crc = ComputeCrc();
+
 	return Crc;
 }
 
@@ -97,7 +99,9 @@ FPCGCrc UPCGData::ComputeCrc() const
 	FArchiveCrc32 Ar;
 	AddToCrc(Ar);
 
-	return FPCGCrc(Ar.GetCrc());
+	Crc = FPCGCrc(Ar.GetCrc());
+
+	return Crc;
 }
 
 void UPCGData::AddToCrc(FArchiveCrc32& Ar) const
