[Backout] - CL24523922
Original CL Desc
-----------------------------------------------------------------
[PCG] Expand CRC out to give caching a chance to work, giving big perf gains from warm cache scenarios where reusable work is in the cache.

diff --git a/Engine/Plugins/Experimental/PCG/Source/PCG/Private/Data/PCGDifferenceData.cpp b/Engine/Plugins/Experimental/PCG/Source/PCG/Private/Data/PCGDifferenceData.cpp
--- a/Engine/Plugins/Experimental/PCG/Source/PCG/Private/Data/PCGDifferenceData.cpp
+++ b/Engine/Plugins/Experimental/PCG/Source/PCG/Private/Data/PCGDifferenceData.cpp
@@ -6,8 +6,6 @@
 #include "Data/PCGUnionData.h"
 #include "Helpers/PCGAsync.h"
 
-#include "Serialization/ArchiveCrc32.h"
-
 #include UE_INLINE_GENERATED_CPP_BY_NAME(PCGDifferenceData)
 
 namespace PCGDifferenceDataUtils
@@ -102,36 +100,6 @@ void UPCGDifferenceData::VisitDataNetwork(TFunctionRef<void(const UPCGData*)> Ac
 	}
 }
 
-FPCGCrc UPCGDifferenceData::ComputeCrc() const
-{
-	FArchiveCrc32 Ar;
-	AddToCrc(Ar);
-
-	// Chain together CRCs of operands
-	check(Source);
-	FPCGCrc SourceCrc = Source->GetOrComputeCrc();
-	Ar << SourceCrc;
-
-	if (Difference)
-	{
-		FPCGCrc DifferenceCrc = Difference->GetOrComputeCrc();
-		Ar << DifferenceCrc;
-	}
-
-	Crc = FPCGCrc(Ar.GetCrc());
-
-	return Crc;
-}
-
-void UPCGDifferenceData::AddToCrc(FArchiveCrc32& Ar) const
-{
-	uint32 UniqueTypeID = StaticClass()->GetDefaultObject()->GetUniqueID();
-	Ar << UniqueTypeID;
-
-	uint32 DensityFunctionValue = static_cast<uint32>(DensityFunction);
-	Ar << DensityFunctionValue;
-}
-
 int UPCGDifferenceData::GetDimension() const
 {
 	return Source->GetDimension();
diff --git a/Engine/Plugins/Experimental/PCG/Source/PCG/Private/Data/PCGIntersectionData.cpp b/Engine/Plugins/Experimental/PCG/Source/PCG/Private/Data/PCGIntersectionData.cpp
--- a/Engine/Plugins/Experimental/PCG/Source/PCG/Private/Data/PCGIntersectionData.cpp
+++ b/Engine/Plugins/Experimental/PCG/Source/PCG/Private/Data/PCGIntersectionData.cpp
@@ -6,8 +6,6 @@
 #include "Helpers/PCGAsync.h"
 #include "PCGHelpers.h"
 
-#include "Serialization/ArchiveCrc32.h"
-
 #include UE_INLINE_GENERATED_CPP_BY_NAME(PCGIntersectionData)
 
 namespace PCGIntersectionDataMaths
@@ -71,32 +69,6 @@ void UPCGIntersectionData::VisitDataNetwork(TFunctionRef<void(const UPCGData*)>
 	GetB()->VisitDataNetwork(Action);
 }
 
-FPCGCrc UPCGIntersectionData::ComputeCrc() const
-{
-	FArchiveCrc32 Ar;
-	AddToCrc(Ar);
-
-	// Chain together CRCs of operands
-	check(A && B);
-	uint32 CrcA = A->GetOrComputeCrc().GetValue();
-	uint32 CrcB = B->GetOrComputeCrc().GetValue();
-	Ar << CrcA;
-	Ar << CrcB;
-
-	Crc = FPCGCrc(Ar.GetCrc());
-
-	return Crc;
-}
-
-void UPCGIntersectionData::AddToCrc(FArchiveCrc32& Ar) const
-{
-	uint32 UniqueTypeID = StaticClass()->GetDefaultObject()->GetUniqueID();
-	Ar << UniqueTypeID;
-
-	uint32 DensityFunctionValue = static_cast<uint32>(DensityFunction);
-	Ar << DensityFunctionValue;
-}
-
 int UPCGIntersectionData::GetDimension() const
 {
 	check(GetA() && GetB());
diff --git a/Engine/Plugins/Experimental/PCG/Source/PCG/Private/Data/PCGPrimitiveData.cpp b/Engine/Plugins/Experimental/PCG/Source/PCG/Private/Data/PCGPrimitiveData.cpp
--- a/Engine/Plugins/Experimental/PCG/Source/PCG/Private/Data/PCGPrimitiveData.cpp
+++ b/Engine/Plugins/Experimental/PCG/Source/PCG/Private/Data/PCGPrimitiveData.cpp
@@ -6,10 +6,7 @@
 #include "Data/PCGPointData.h"
 #include "Data/PCGSpatialData.h"
 #include "Elements/PCGVolumeSampler.h"
-
-#include "Components/BoxComponent.h"
 #include "Components/PrimitiveComponent.h"
-#include "Serialization/ArchiveCrc32.h"
 
 #include UE_INLINE_GENERATED_CPP_BY_NAME(PCGPrimitiveData)
 
@@ -22,26 +19,6 @@ void UPCGPrimitiveData::Initialize(UPrimitiveComponent* InPrimitive)
 	// Not obvious to find strict bounds, leave at the default value
 }
 
-void UPCGPrimitiveData::AddToCrc(FArchiveCrc32& Ar) const
-{
-	check(Primitive.Get());
-
-	// Can implement specific CRC's for different primitive types.
-	if (UBoxComponent* Box = Cast<UBoxComponent>(Primitive.Get()))
-	{
-		FVector Extent = Box->GetUnscaledBoxExtent();
-		Ar << Extent;
-
-		FTransform Transform = Box->GetComponentTransform();
-		Ar << Transform;
-	}
-	else
-	{
-		// Fallback is the UID which changes when data recreated.
-		Super::AddToCrc(Ar);
-	}
-}
-
 bool UPCGPrimitiveData::SamplePoint(const FTransform& InTransform, const FBox& InBounds, FPCGPoint& OutPoint, UPCGMetadata* OutMetadata) const
 {
 	// Pure implementation
diff --git a/Engine/Plugins/Experimental/PCG/Source/PCG/Private/Data/PCGSplineData.cpp b/Engine/Plugins/Experimental/PCG/Source/PCG/Private/Data/PCGSplineData.cpp
--- a/Engine/Plugins/Experimental/PCG/Source/PCG/Private/Data/PCGSplineData.cpp
+++ b/Engine/Plugins/Experimental/PCG/Source/PCG/Private/Data/PCGSplineData.cpp
@@ -10,7 +10,6 @@
 
 #include "Components/SplineComponent.h"
 #include "Data/PCGSpatialData.h"
-#include "Serialization/ArchiveCrc32.h"
 
 #include UE_INLINE_GENERATED_CPP_BY_NAME(PCGSplineData)
 
@@ -34,22 +33,6 @@ void UPCGSplineData::Initialize(USplineComponent* InSpline)
 	CachedBounds = CachedBounds.ExpandBy(SplinePointsRadius, SplinePointsRadius);
 }
 
-void UPCGSplineData::AddToCrc(FArchiveCrc32& Ar) const
-{
-	check(Spline);
-
-	Ar << const_cast<FBox&>(CachedBounds);
-
-	Ar << const_cast<FInterpCurveVector&>(Spline->GetSplinePointsPosition());
-	Ar << const_cast<FInterpCurveQuat&>(Spline->GetSplinePointsRotation());
-	Ar << const_cast<FInterpCurveVector&>(Spline->GetSplinePointsScale());
-
-	Ar << const_cast<FTransform&>(Spline->GetComponentTransform());
-
-	bool bIsClosed = Spline->IsClosedLoop();
-	Ar << bIsClosed;
-}
-
 FTransform UPCGSplineData::GetTransform() const
 {
 	return Spline ? Spline->GetComponentTransform() : FTransform::Identity;
diff --git a/Engine/Plugins/Experimental/PCG/Source/PCG/Private/Data/PCGUnionData.cpp b/Engine/Plugins/Experimental/PCG/Source/PCG/Private/Data/PCGUnionData.cpp
--- a/Engine/Plugins/Experimental/PCG/Source/PCG/Private/Data/PCGUnionData.cpp
+++ b/Engine/Plugins/Experimental/PCG/Source/PCG/Private/Data/PCGUnionData.cpp
@@ -7,8 +7,6 @@
 #include "Metadata/PCGMetadataAccessor.h"
 #include "PCGHelpers.h"
 
-#include "Serialization/ArchiveCrc32.h"
-
 #include UE_INLINE_GENERATED_CPP_BY_NAME(PCGUnionData)
 
 namespace PCGUnionDataMaths
@@ -83,41 +81,6 @@ void UPCGUnionData::VisitDataNetwork(TFunctionRef<void(const UPCGData*)> Action)
 	}
 }
 
-FPCGCrc UPCGUnionData::ComputeCrc() const
-{
-	FArchiveCrc32 Ar;
-	AddToCrc(Ar);
-
-	// Chain together CRCs of operands
-	int32 NumOperands = Data.Num();
-	Ar << NumOperands;
-
-	for (TObjectPtr<const UPCGSpatialData> Datum : Data)
-	{
-		if (Datum)
-		{
-			uint32 DatumCrc = Datum->GetOrComputeCrc().GetValue();
-			Ar << DatumCrc;
-		}
-	}
-
-	Crc = FPCGCrc(Ar.GetCrc());
-
-	return Crc;
-}
-
-void UPCGUnionData::AddToCrc(FArchiveCrc32& Ar) const
-{
-	uint32 UniqueTypeID = StaticClass()->GetDefaultObject()->GetUniqueID();
-	Ar << UniqueTypeID;
-
-	uint32 UnionTypeValue = static_cast<uint32>(UnionType);
-	Ar << UnionTypeValue;
-
-	uint32 DensityFunctionValue = static_cast<uint32>(DensityFunction);
-	Ar << DensityFunctionValue;
-}
-
 int UPCGUnionData::GetDimension() const
 {
 	return CachedDimension;
diff --git a/Engine/Plugins/Experimental/PCG/Source/PCG/Private/Data/PCGVolumeData.cpp b/Engine/Plugins/Experimental/PCG/Source/PCG/Private/Data/PCGVolumeData.cpp
--- a/Engine/Plugins/Experimental/PCG/Source/PCG/Private/Data/PCGVolumeData.cpp
+++ b/Engine/Plugins/Experimental/PCG/Source/PCG/Private/Data/PCGVolumeData.cpp
@@ -6,23 +6,10 @@
 #include "PCGHelpers.h"
 #include "Elements/PCGVolumeSampler.h"
 
-#include "Components/BrushComponent.h"
 #include "GameFramework/Volume.h"
-#include "Serialization/ArchiveCrc32.h"
-#include "Serialization/ArchiveObjectCrc32.h"
 
 #include UE_INLINE_GENERATED_CPP_BY_NAME(PCGVolumeData)
 
-/** Does CRC on Brush data within BrushComponent. */
-class FPCGBrushCrc32 : public FArchiveObjectCrc32
-{
-public:
-	virtual bool ShouldSkipProperty(const FProperty* InProperty) const override
-	{
-		return InProperty && (InProperty->GetFName() != FName("Brush"));
-	}
-};
-
 void UPCGVolumeData::Initialize(AVolume* InVolume, AActor* InTargetActor)
 {
 	check(InVolume);
@@ -43,25 +30,6 @@ void UPCGVolumeData::Initialize(const FBox& InBounds, AActor* InTargetActor)
 	TargetActor = InTargetActor;
 }
 
-void UPCGVolumeData::AddToCrc(FArchiveCrc32& Ar) const
-{
-	Ar << const_cast<FBox&>(Bounds);
-
-	Ar << const_cast<FBox&>(StrictBounds);
-
-	if (Volume.IsValid())
-	{
-		if (UBrushComponent* Brush = Volume->GetBrushComponent())
-		{
-			Ar << const_cast<FTransform&>(Brush->GetComponentTransform());
-
-			FPCGBrushCrc32 VolumeCrc;
-			uint32 BrushCrc = VolumeCrc.Crc32(Brush);
-			Ar << BrushCrc;
-		}
-	}
-}
-
 FBox UPCGVolumeData::GetBounds() const
 {
 	return Bounds;
diff --git a/Engine/Plugins/Experimental/PCG/Source/PCG/Private/PCGData.cpp b/Engine/Plugins/Experimental/PCG/Source/PCG/Private/PCGData.cpp
--- a/Engine/Plugins/Experimental/PCG/Source/PCG/Private/PCGData.cpp
+++ b/Engine/Plugins/Experimental/PCG/Source/PCG/Private/PCGData.cpp
@@ -89,13 +89,6 @@ FPCGCrc UPCGData::GetOrComputeCrc() const
 		return Crc;
 	}
 
-	Crc = ComputeCrc();
-
-	return Crc;
-}
-
-FPCGCrc UPCGData::ComputeCrc() const
-{
 	FArchiveCrc32 Ar;
 	AddToCrc(Ar);
 
diff --git a/Engine/Plugins/Experimental/PCG/Source/PCG/Public/Data/PCGDifferenceData.h b/Engine/Plugins/Experimental/PCG/Source/PCG/Public/Data/PCGDifferenceData.h
--- a/Engine/Plugins/Experimental/PCG/Source/PCG/Public/Data/PCGDifferenceData.h
+++ b/Engine/Plugins/Experimental/PCG/Source/PCG/Public/Data/PCGDifferenceData.h
@@ -50,13 +50,8 @@ public:
 	// ~Begin UPCGData interface
 	virtual EPCGDataType GetDataType() const override { return EPCGDataType::Spatial; }
 	virtual void VisitDataNetwork(TFunctionRef<void(const UPCGData*)> Action) const override;
-
-protected:
-	virtual FPCGCrc ComputeCrc() const override;
-	virtual void AddToCrc(FArchiveCrc32& Ar) const override;
 	// ~End UPCGData interface
 
-public:
 	//~Begin UPCGSpatialData interface
 	virtual int GetDimension() const override;
 	virtual FBox GetBounds() const override;
diff --git a/Engine/Plugins/Experimental/PCG/Source/PCG/Public/Data/PCGIntersectionData.h b/Engine/Plugins/Experimental/PCG/Source/PCG/Public/Data/PCGIntersectionData.h
--- a/Engine/Plugins/Experimental/PCG/Source/PCG/Public/Data/PCGIntersectionData.h
+++ b/Engine/Plugins/Experimental/PCG/Source/PCG/Public/Data/PCGIntersectionData.h
@@ -33,13 +33,8 @@ public:
 	// ~Begin UPCGData interface
 	virtual EPCGDataType GetDataType() const override { return EPCGDataType::Spatial; }
 	virtual void VisitDataNetwork(TFunctionRef<void(const UPCGData*)> Action) const override;
-
-protected:
-	virtual FPCGCrc ComputeCrc() const override;
-	virtual void AddToCrc(FArchiveCrc32& Ar) const override;
 	// ~End UPCGData interface
 
-public:
 	//~Begin UPCGSpatialData interface
 	virtual int GetDimension() const override;
 	virtual FBox GetBounds() const override;
diff --git a/Engine/Plugins/Experimental/PCG/Source/PCG/Public/Data/PCGPrimitiveData.h b/Engine/Plugins/Experimental/PCG/Source/PCG/Public/Data/PCGPrimitiveData.h
--- a/Engine/Plugins/Experimental/PCG/Source/PCG/Public/Data/PCGPrimitiveData.h
+++ b/Engine/Plugins/Experimental/PCG/Source/PCG/Public/Data/PCGPrimitiveData.h
@@ -16,7 +16,6 @@ public:
 
 	// ~Begin UPCGData interface
 	virtual EPCGDataType GetDataType() const override { return EPCGDataType::Primitive; }
-	void AddToCrc(FArchiveCrc32& Ar) const override;
 	// ~End UPCGData interface
 
 	// ~Begin UPCGSpatialData interface
diff --git a/Engine/Plugins/Experimental/PCG/Source/PCG/Public/Data/PCGSplineData.h b/Engine/Plugins/Experimental/PCG/Source/PCG/Public/Data/PCGSplineData.h
--- a/Engine/Plugins/Experimental/PCG/Source/PCG/Public/Data/PCGSplineData.h
+++ b/Engine/Plugins/Experimental/PCG/Source/PCG/Public/Data/PCGSplineData.h
@@ -24,7 +24,6 @@ public:
 
 	// ~Begin UPCGData interface
 	virtual EPCGDataType GetDataType() const override { return EPCGDataType::Spline; }
-	void AddToCrc(FArchiveCrc32& Ar) const override;
 	// ~End UPCGData interface
 
 	//~Begin UPCGPolyLineData interface
@@ -44,7 +43,6 @@ public:
 	virtual FBox GetBounds() const override;
 	virtual bool SamplePoint(const FTransform& Transform, const FBox& Bounds, FPCGPoint& OutPoint, UPCGMetadata* OutMetadata) const override;
 	virtual UPCGSpatialData* ProjectOn(const UPCGSpatialData* InOther, const FPCGProjectionParams& InParams = FPCGProjectionParams()) const override;
-
 protected:
 	virtual UPCGSpatialData* CopyInternal() const override;
 	//~End UPCGSpatialData interface
diff --git a/Engine/Plugins/Experimental/PCG/Source/PCG/Public/Data/PCGUnionData.h b/Engine/Plugins/Experimental/PCG/Source/PCG/Public/Data/PCGUnionData.h
--- a/Engine/Plugins/Experimental/PCG/Source/PCG/Public/Data/PCGUnionData.h
+++ b/Engine/Plugins/Experimental/PCG/Source/PCG/Public/Data/PCGUnionData.h
@@ -39,13 +39,8 @@ public:
 	// ~Begin UPCGData interface
 	virtual EPCGDataType GetDataType() const override { return EPCGDataType::Spatial; }
 	virtual void VisitDataNetwork(TFunctionRef<void(const UPCGData*)> Action) const override;
-
-protected:
-	virtual FPCGCrc ComputeCrc() const override;
-	virtual void AddToCrc(FArchiveCrc32& Ar) const override;
 	// ~End UPCGData interface
 
-public:
 	//~Begin UPCGSpatialData interface
 	virtual int GetDimension() const override;
 	virtual FBox GetBounds() const override;
diff --git a/Engine/Plugins/Experimental/PCG/Source/PCG/Public/Data/PCGVolumeData.h b/Engine/Plugins/Experimental/PCG/Source/PCG/Public/Data/PCGVolumeData.h
--- a/Engine/Plugins/Experimental/PCG/Source/PCG/Public/Data/PCGVolumeData.h
+++ b/Engine/Plugins/Experimental/PCG/Source/PCG/Public/Data/PCGVolumeData.h
@@ -19,7 +19,6 @@ public:
 
 	// ~Begin UPCGData interface
 	virtual EPCGDataType GetDataType() const override { return EPCGDataType::Volume; }
-	virtual void AddToCrc(FArchiveCrc32& Ar) const;
 	// ~End UPCGData interface
 
 	// ~Begin UPGCSpatialData interface
diff --git a/Engine/Plugins/Experimental/PCG/Source/PCG/Public/PCGData.h b/Engine/Plugins/Experimental/PCG/Source/PCG/Public/PCGData.h
--- a/Engine/Plugins/Experimental/PCG/Source/PCG/Public/PCGData.h
+++ b/Engine/Plugins/Experimental/PCG/Source/PCG/Public/PCGData.h
@@ -29,9 +29,12 @@ public:
 	UPCGData(const FObjectInitializer& ObjectInitializer);
 	virtual EPCGDataType GetDataType() const { return EPCGDataType::None; }
 
-	/** Returns a Crc for this and any connected data. */
+	/** Computes a Crc for this data. */
 	FPCGCrc GetOrComputeCrc() const;
 
+	/** Contributes to Crc. Fallback implementation that writes object instance UID. */
+	virtual void AddToCrc(FArchiveCrc32& Ar) const;
+
 	/** Executes a lambda over all connected data objects. */
 	virtual void VisitDataNetwork(TFunctionRef<void(const UPCGData*)> Action) const;
 
@@ -42,13 +45,6 @@ public:
 	/** CRC for this object instance. */
 	mutable FPCGCrc Crc;
 
-protected:
-	/** Computes Crc for this and any connected data. */
-	virtual FPCGCrc ComputeCrc() const;
-
-	/** Adds this data to Crc. Fallback implementation writes object instance UID. */
-	virtual void AddToCrc(FArchiveCrc32& Ar) const;
-
 private:
 	/** Serves unique ID values to instances of this object. */
 	static inline std::atomic<uint64> UIDCounter{ 1 };
