[PCG] Replace reinterpret_cast by static_cast in surface sampler.

#rnx

diff --git a/Engine/Plugins/Experimental/PCG/Source/PCG/Private/Elements/PCGSurfaceSampler.cpp b/Engine/Plugins/Experimental/PCG/Source/PCG/Private/Elements/PCGSurfaceSampler.cpp
--- a/Engine/Plugins/Experimental/PCG/Source/PCG/Private/Elements/PCGSurfaceSampler.cpp
+++ b/Engine/Plugins/Experimental/PCG/Source/PCG/Private/Elements/PCGSurfaceSampler.cpp
@@ -482,7 +482,7 @@ void FPCGSurfaceSamplerElement::AllocateOutputs(FPCGSurfaceSamplerContext* InCon
 bool FPCGSurfaceSamplerElement::PrepareDataInternal(FPCGContext* InContext) const
 {
 	TRACE_CPUPROFILER_EVENT_SCOPE(FPCGSurfaceSamplerElement::PrepareData);
-	FPCGSurfaceSamplerContext* Context = reinterpret_cast<FPCGSurfaceSamplerContext*>(InContext);
+	FPCGSurfaceSamplerContext* Context = static_cast<FPCGSurfaceSamplerContext*>(InContext);
 
 	check(Context);
 
@@ -498,7 +498,7 @@ bool FPCGSurfaceSamplerElement::PrepareDataInternal(FPCGContext* InContext) cons
 bool FPCGSurfaceSamplerElement::ExecuteInternal(FPCGContext* InContext) const
 {
 	TRACE_CPUPROFILER_EVENT_SCOPE(FPCGSurfaceSamplerElement::Execute);
-	FPCGSurfaceSamplerContext* Context = reinterpret_cast<FPCGSurfaceSamplerContext*>(InContext);
+	FPCGSurfaceSamplerContext* Context = static_cast<FPCGSurfaceSamplerContext*>(InContext);
 
 	check(Context);
 
