Fixed missing step in actor cleanup when doing the break PCG link

#rnx

diff --git a/Engine/Plugins/Experimental/PCG/Source/PCG/Private/PCGComponent.cpp b/Engine/Plugins/Experimental/PCG/Source/PCG/Private/PCGComponent.cpp
--- a/Engine/Plugins/Experimental/PCG/Source/PCG/Private/PCGComponent.cpp
+++ b/Engine/Plugins/Experimental/PCG/Source/PCG/Private/PCGComponent.cpp
@@ -684,7 +684,8 @@ bool UPCGComponent::MoveResourcesToNewActor(AActor* InNewActor, bool bCreateChil
 
 	if (bCreateChild)
 	{
-		NewActor = UPCGActorHelpers::SpawnDefaultActor(GetWorld(), NewActor->GetClass(), TEXT("PCGStampChild"), Owner->GetTransform(), NewActor);
+		NewActor = UPCGActorHelpers::SpawnDefaultActor(GetWorld(), NewActor->GetClass(), TEXT("PCGStampChild"), Owner->GetTransform());
+		NewActor->AttachToActor(InNewActor, FAttachmentTransformRules::KeepWorldTransform);
 		check(NewActor);
 	}
 
diff --git a/Engine/Plugins/Experimental/PCG/Source/PCG/Private/PCGManagedResource.cpp b/Engine/Plugins/Experimental/PCG/Source/PCG/Private/PCGManagedResource.cpp
--- a/Engine/Plugins/Experimental/PCG/Source/PCG/Private/PCGManagedResource.cpp
+++ b/Engine/Plugins/Experimental/PCG/Source/PCG/Private/PCGManagedResource.cpp
@@ -127,6 +127,7 @@ bool UPCGManagedActors::MoveResourceToNewActor(AActor* NewActor)
 		}
 
 		Actor->DetachFromActor(FDetachmentTransformRules::KeepWorldTransform);
+		Actor->SetOwner(nullptr);
 		Actor->AttachToActor(NewActor, FAttachmentTransformRules::KeepWorldTransform);
 	}
 
