[PCG] Static analysis fix second attempt (Issue 272356)

Theory is that the pointer guards are negating the check() and making the analysis think the pointer could be null. Since guards are unnecessary, better to just remove them.

diff --git a/Engine/Plugins/Experimental/PCG/Source/PCG/Private/Data/PCGUnionData.cpp b/Engine/Plugins/Experimental/PCG/Source/PCG/Private/Data/PCGUnionData.cpp
--- a/Engine/Plugins/Experimental/PCG/Source/PCG/Private/Data/PCGUnionData.cpp
+++ b/Engine/Plugins/Experimental/PCG/Source/PCG/Private/Data/PCGUnionData.cpp
@@ -379,11 +379,18 @@ void UPCGUnionData::CreateSequentialPointData(FPCGContext* Context, TArray<const
 			}
 
 			// Update density & metadata based on current & following data
-			const bool bSkipLoop = (OutPoint.Density >= 1.0f && !OutMetadata);
-			for (int32 FollowingDataIndex = DataIndex + DataIndexIncrement; FollowingDataIndex != LastDataIndex && !bSkipLoop; FollowingDataIndex += DataIndexIncrement)
+			for (int32 FollowingDataIndex = DataIndex + DataIndexIncrement; FollowingDataIndex != LastDataIndex; FollowingDataIndex += DataIndexIncrement)
 			{
+				const UPCGMetadata* FollowingMetadata = InputMetadatas[FollowingDataIndex];
+
+				// If density is saturated and there are no metadata attributes then we can skip this data as it will not contribute.
+				if (OutPoint.Density >= 1.0f && (!FollowingMetadata || FollowingMetadata->GetAttributeCount() == 0))
+				{
+					continue;
+				}
+
 				FPCGPoint PointInData;
-				if(InputDatas[FollowingDataIndex]->SamplePoint(OutPoint.Transform, OutPoint.GetLocalBounds(), PointInData, OutMetadata))
+				if (InputDatas[FollowingDataIndex]->SamplePoint(OutPoint.Transform, OutPoint.GetLocalBounds(), PointInData, OutMetadata))
 				{
 					// Update density
 					PCGUnionDataMaths::UpdateDensity(OutPoint.Density, PointInData.Density, DensityFunction);
@@ -395,20 +402,13 @@ void UPCGUnionData::CreateSequentialPointData(FPCGContext* Context, TArray<const
 						FMath::Max(OutPoint.Color.Z, PointInData.Color.Z),
 						FMath::Max(OutPoint.Color.W, PointInData.Color.W));
 
-					if (OutMetadata)
+					if (OutPoint.MetadataEntry != PCGInvalidEntryKey && PointInData.MetadataEntry != PCGInvalidEntryKey)
 					{
-						if (OutPoint.MetadataEntry != PCGInvalidEntryKey && PointInData.MetadataEntry != PCGInvalidEntryKey)
-						{
-							OutMetadata->MergePointAttributesSubset(OutPoint, OutMetadata, OutMetadata, PointInData, OutMetadata, InputMetadatas[FollowingDataIndex], OutPoint, EPCGMetadataOp::Max);
-						}
-						else if (PointInData.MetadataEntry != PCGInvalidEntryKey)
-						{
-							OutPoint.MetadataEntry = PointInData.MetadataEntry;
-						}
+						OutMetadata->MergePointAttributesSubset(OutPoint, OutMetadata, OutMetadata, PointInData, OutMetadata, FollowingMetadata, OutPoint, EPCGMetadataOp::Max);
 					}
-					else if (OutPoint.Density >= 1.0f)
+					else if (PointInData.MetadataEntry != PCGInvalidEntryKey)
 					{
-						break;
+						OutPoint.MetadataEntry = PointInData.MetadataEntry;
 					}
 				}
 			}
