Fixed small robustness issue in BP call as we can't readily guarantee we'll be in a proper state by then, and a warning message.

#rnx

diff --git a/Engine/Plugins/Experimental/PCG/Source/PCG/Private/Helpers/PCGBlueprintHelpers.cpp b/Engine/Plugins/Experimental/PCG/Source/PCG/Private/Helpers/PCGBlueprintHelpers.cpp
--- a/Engine/Plugins/Experimental/PCG/Source/PCG/Private/Helpers/PCGBlueprintHelpers.cpp
+++ b/Engine/Plugins/Experimental/PCG/Source/PCG/Private/Helpers/PCGBlueprintHelpers.cpp
@@ -4,6 +4,7 @@
 #include "PCGComponent.h"
 #include "PCGContext.h"
 #include "PCGHelpers.h"
+#include "PCGModule.h"
 #include "PCGPoint.h"
 #include "PCGSubsystem.h"
 #include "Grid/PCGPartitionActor.h"
@@ -156,7 +157,11 @@ TArray<FPCGLandscapeLayerWeight> UPCGBlueprintHelpers::GetInterpolatedPCGLandsca
 	check(Landscape);
 
 	ULandscapeInfo* LandscapeInfo = Landscape->GetLandscapeInfo();
-	check(LandscapeInfo);
+	if (!LandscapeInfo)
+	{
+		UE_LOG(LogPCG, Warning, TEXT("Unable to get landscape layer weights because the landscape info is not available (landscape not registered yet?"));
+		return {};
+	}
 
 	const FVector LocalPoint = Landscape->GetTransform().InverseTransformPosition(Location);
 	const FIntPoint ComponentMapKey(FMath::FloorToInt(LocalPoint.X / LandscapeInfo->ComponentSizeQuads), FMath::FloorToInt(LocalPoint.Y / LandscapeInfo->ComponentSizeQuads));
