Fixed behavior of structural changes in the editor while in PIE.
Fixed an error in a console command

#rnx

diff --git a/Engine/Plugins/Experimental/PCG/Source/PCG/Private/PCGGraph.cpp b/Engine/Plugins/Experimental/PCG/Source/PCG/Private/PCGGraph.cpp
--- a/Engine/Plugins/Experimental/PCG/Source/PCG/Private/PCGGraph.cpp
+++ b/Engine/Plugins/Experimental/PCG/Source/PCG/Private/PCGGraph.cpp
@@ -668,6 +668,14 @@ void UPCGGraph::NotifyGraphChanged(EPCGChangeType ChangeType)
 	const bool bNotifySubsystem = ((ChangeType & (EPCGChangeType::Structural | EPCGChangeType::Edge)) != EPCGChangeType::None);
 	if (bNotifySubsystem && GEditor)
 	{
+		if (GEditor->PlayWorld)
+		{
+			if (UPCGSubsystem* PCGPIESubsystem = UPCGSubsystem::GetInstance(GEditor->PlayWorld.Get()))
+			{
+				PCGPIESubsystem->NotifyGraphChanged(this);
+			}
+		}
+
 		if (UPCGSubsystem* PCGSubsystem = UPCGSubsystem::GetInstance(GEditor->GetEditorWorldContext().World()))
 		{
 			PCGSubsystem->NotifyGraphChanged(this);
diff --git a/Engine/Plugins/Experimental/PCG/Source/PCG/Private/PCGSettings.cpp b/Engine/Plugins/Experimental/PCG/Source/PCG/Private/PCGSettings.cpp
--- a/Engine/Plugins/Experimental/PCG/Source/PCG/Private/PCGSettings.cpp
+++ b/Engine/Plugins/Experimental/PCG/Source/PCG/Private/PCGSettings.cpp
@@ -415,6 +415,14 @@ void UPCGSettings::DirtyCache()
 {
 	if (GEditor)
 	{
+		if (GEditor->PlayWorld)
+		{
+			if (UPCGSubsystem* PCGPIESubsystem = UPCGSubsystem::GetInstance(GEditor->PlayWorld.Get()))
+			{
+				PCGPIESubsystem->CleanFromCache(GetElement().Get(), this);
+			}
+		}
+
 		if (UPCGSubsystem* PCGSubsystem = UPCGSubsystem::GetInstance(GEditor->GetEditorWorldContext().World()))
 		{
 			PCGSubsystem->CleanFromCache(GetElement().Get(), this);
diff --git a/Engine/Plugins/Experimental/PCG/Source/PCG/Private/PCGSubsystem.cpp b/Engine/Plugins/Experimental/PCG/Source/PCG/Private/PCGSubsystem.cpp
--- a/Engine/Plugins/Experimental/PCG/Source/PCG/Private/PCGSubsystem.cpp
+++ b/Engine/Plugins/Experimental/PCG/Source/PCG/Private/PCGSubsystem.cpp
@@ -71,7 +71,7 @@ namespace PCGSubsystemConsole
 				UWorld* World = (GEditor ? (GEditor->PlayWorld ? GEditor->PlayWorld.Get() : GEditor->GetEditorWorldContext().World()) : (GEngine ? GEngine->GetCurrentPlayWorld() : nullptr));
 				if (UPCGSubsystem* PCGSubsystem = UPCGSubsystem::GetInstance(World))
 				{
-					PCGSubsystem->FlushCache();
+					PCGSubsystem->BuildLandscapeCache();
 				}
 			}));
 
