[PCG] Cache raw pointers for Difference data after seeing it appear on profiling radar

Copies same pattern as in intersection/union data.

diff --git a/Engine/Plugins/Experimental/PCG/Source/PCG/Private/Data/PCGDifferenceData.cpp b/Engine/Plugins/Experimental/PCG/Source/PCG/Private/Data/PCGDifferenceData.cpp
--- a/Engine/Plugins/Experimental/PCG/Source/PCG/Private/Data/PCGDifferenceData.cpp
+++ b/Engine/Plugins/Experimental/PCG/Source/PCG/Private/Data/PCGDifferenceData.cpp
@@ -33,6 +33,10 @@ void UPCGDifferenceData::Initialize(const UPCGSpatialData* InData)
 	Source = InData;
 	TargetActor = InData->TargetActor;
 
+#if WITH_EDITOR
+	RawPointerSource = Source;
+#endif
+
 	check(Metadata);
 	Metadata->Initialize(Source->Metadata);
 }
@@ -51,6 +55,10 @@ void UPCGDifferenceData::AddDifference(const UPCGSpatialData* InDifference)
 	if (!Difference)
 	{
 		Difference = InDifference;
+
+#if WITH_EDITOR
+		RawPointerDifference = InDifference;
+#endif
 	}
 	else
 	{
@@ -60,6 +68,11 @@ void UPCGDifferenceData::AddDifference(const UPCGSpatialData* InDifference)
 			DifferencesUnion->AddData(Difference);
 			DifferencesUnion->SetDensityFunction(PCGDifferenceDataUtils::ToUnionDensityFunction(DensityFunction));
 			Difference = DifferencesUnion;
+
+#if WITH_EDITOR
+			RawPointerDifference = Difference;
+			RawPointerDifferencesUnion = DifferencesUnion;
+#endif
 		}
 
 		check(Difference == DifferencesUnion);
@@ -71,9 +84,9 @@ void UPCGDifferenceData::SetDensityFunction(EPCGDifferenceDensityFunction InDens
 {
 	DensityFunction = InDensityFunction;
 
-	if (DifferencesUnion)
+	if (GetDifferencesUnion())
 	{
-		DifferencesUnion->SetDensityFunction(PCGDifferenceDataUtils::ToUnionDensityFunction(DensityFunction));
+		GetDifferencesUnion()->SetDensityFunction(PCGDifferenceDataUtils::ToUnionDensityFunction(DensityFunction));
 	}
 }
 
@@ -87,41 +100,50 @@ void UPCGDifferenceData::PostEditChangeProperty(FPropertyChangedEvent& PropertyC
 
 	Super::PostEditChangeProperty(PropertyChangedEvent);
 }
+
+void UPCGDifferenceData::PostLoad()
+{
+	Super::PostLoad();
+
+	RawPointerSource = Source;
+	RawPointerDifference = Difference;
+	RawPointerDifferencesUnion = DifferencesUnion;
+}
 #endif
 
 void UPCGDifferenceData::VisitDataNetwork(TFunctionRef<void(const UPCGData*)> Action) const
 {
-	check(Source);
-	Source->VisitDataNetwork(Action);
+	check(GetSource());
+	GetSource()->VisitDataNetwork(Action);
 
-	if (Difference)
+	if (GetDifference())
 	{
-		Difference->VisitDataNetwork(Action);
+		GetDifference()->VisitDataNetwork(Action);
 	}
 }
 
 int UPCGDifferenceData::GetDimension() const
 {
-	return Source->GetDimension();
+	return GetSource()->GetDimension();
 }
 
 FBox UPCGDifferenceData::GetBounds() const
 {
-	return Source->GetBounds();
+	return GetSource()->GetBounds();
 }
 
 FBox UPCGDifferenceData::GetStrictBounds() const
 {
-	return Difference ? FBox(EForceInit::ForceInit) : Source->GetStrictBounds();
+	return GetDifference() ? FBox(EForceInit::ForceInit) : GetSource()->GetStrictBounds();
 }
 
 bool UPCGDifferenceData::SamplePoint(const FTransform& InTransform, const FBox& InBounds, FPCGPoint& OutPoint, UPCGMetadata* OutMetadata) const
 {
 	//TRACE_CPUPROFILER_EVENT_SCOPE(UPCGDifferenceData::SamplePoint);
-	check(Source);
+	check(GetSource());
 
 	FPCGPoint PointFromSource;
-	if(!Source->SamplePoint(InTransform, InBounds, PointFromSource, OutMetadata))
+	if(!GetSource()->SamplePoint(InTransform, InBounds, PointFromSource, OutMetadata))
 	{
 		return false;
 	}
@@ -130,7 +152,7 @@ bool UPCGDifferenceData::SamplePoint(const FTransform& InTransform, const FBox&
 
 	FPCGPoint PointFromDiff;
 	// Important note: here we will not use the point we got from the source, otherwise we are introducing severe bias
-	if (Difference && Difference->SamplePoint(InTransform, InBounds, PointFromDiff, (bDiffMetadata ? OutMetadata : nullptr)))
+	if (GetDifference() && GetDifference()->SamplePoint(InTransform, InBounds, PointFromDiff, (bDiffMetadata ? OutMetadata : nullptr)))
 	{
 		const bool bBinaryDensity = (DensityFunction == EPCGDifferenceDensityFunction::Binary);
 		
@@ -139,7 +161,8 @@ bool UPCGDifferenceData::SamplePoint(const FTransform& InTransform, const FBox&
 		// Color?
 		if (bDiffMetadata && OutMetadata && OutPoint.Density > 0 && PointFromDiff.MetadataEntry != PCGInvalidEntryKey)
 		{
-			OutMetadata->MergePointAttributesSubset(PointFromSource, OutMetadata, Source->Metadata, PointFromDiff, OutMetadata, Difference->Metadata, OutPoint, EPCGMetadataOp::Sub);
+			// Safe to also cache GetSource()->Metadata ? I'm not sure it is, but if it is it could also benefit UnionData which sometimes accesses input metadata, and also intersection data
+			OutMetadata->MergePointAttributesSubset(PointFromSource, OutMetadata, GetSource()->Metadata, PointFromDiff, OutMetadata, GetDifference()->Metadata, OutPoint, EPCGMetadataOp::Sub);
 		}
 
 		return OutPoint.Density > 0;
@@ -152,8 +175,8 @@ bool UPCGDifferenceData::SamplePoint(const FTransform& InTransform, const FBox&
 
 bool UPCGDifferenceData::HasNonTrivialTransform() const
 {
-	check(Source);
-	return Source->HasNonTrivialTransform();
+	check(GetSource());
+	return GetSource()->HasNonTrivialTransform();
 }
 
 const UPCGPointData* UPCGDifferenceData::CreatePointData(FPCGContext* Context) const
@@ -161,7 +184,7 @@ const UPCGPointData* UPCGDifferenceData::CreatePointData(FPCGContext* Context) c
 	TRACE_CPUPROFILER_EVENT_SCOPE(UPCGDifferenceData::CreatePointData);
 	
 	// This is similar to what we are doing in UPCGUnionData::CreatePointData
-	const UPCGPointData* SourcePointData = Source->ToPointData(Context);
+	const UPCGPointData* SourcePointData = GetSource()->ToPointData(Context);
 
 	if (!SourcePointData)
 	{
@@ -169,7 +192,7 @@ const UPCGPointData* UPCGDifferenceData::CreatePointData(FPCGContext* Context) c
 		return SourcePointData;
 	}
 
-	if (!Difference)
+	if (!GetDifference())
 	{
 		UE_LOG(LogPCG, Verbose, TEXT("Difference is trivial"));
 		return SourcePointData;
@@ -185,11 +208,12 @@ const UPCGPointData* UPCGDifferenceData::CreatePointData(FPCGContext* Context) c
 	const TArray<FPCGPoint>& SourcePoints = SourcePointData->GetPoints();
 	TArray<FPCGPoint>& TargetPoints = Data->GetMutablePoints();
 
+	const UPCGMetadata* DifferenceMetadata = GetDifference()->Metadata;
 	UPCGMetadata* TempDiffMetadata = nullptr;
-	if (bDiffMetadata && OutMetadata && Difference->Metadata)
+	if (bDiffMetadata && OutMetadata && DifferenceMetadata)
 	{
 		TempDiffMetadata = NewObject<UPCGMetadata>();
-		TempDiffMetadata->Initialize(Difference->Metadata);
+		TempDiffMetadata->Initialize(DifferenceMetadata);
 	}
 
 	FPCGAsync::AsyncPointProcessing(Context, SourcePoints.Num(), TargetPoints, [this, Data, OutMetadata, SourcePointData, SourceMetadata, TempDiffMetadata, &SourcePoints](int32 Index, FPCGPoint& OutPoint)
@@ -197,7 +221,7 @@ const UPCGPointData* UPCGDifferenceData::CreatePointData(FPCGContext* Context) c
 		const FPCGPoint& Point = SourcePoints[Index];
 
 		FPCGPoint PointFromDiff;
-		if (Difference && Difference->SamplePoint(Point.Transform, Point.GetLocalBounds(), PointFromDiff, TempDiffMetadata))
+		if (GetDifference() && GetDifference()->SamplePoint(Point.Transform, Point.GetLocalBounds(), PointFromDiff, TempDiffMetadata))
 		{
 			const bool bBinaryDensity = (DensityFunction == EPCGDifferenceDensityFunction::Binary);
 
@@ -237,7 +261,16 @@ UPCGSpatialData* UPCGDifferenceData::CopyInternal() const
 	if (DifferencesUnion)
 	{
 		NewDifferenceData->DifferencesUnion = static_cast<UPCGUnionData*>(DifferencesUnion->DuplicateData());
+
+#if WITH_EDITOR
+		NewDifferenceData->RawPointerDifferencesUnion = NewDifferenceData->DifferencesUnion;
+#endif
 	}
 
+#if WITH_EDITOR
+	NewDifferenceData->RawPointerSource = NewDifferenceData->Source;
+	NewDifferenceData->RawPointerDifference = NewDifferenceData->Difference;
+#endif
+
 	return NewDifferenceData;
 }
diff --git a/Engine/Plugins/Experimental/PCG/Source/PCG/Public/Data/PCGDifferenceData.h b/Engine/Plugins/Experimental/PCG/Source/PCG/Public/Data/PCGDifferenceData.h
--- a/Engine/Plugins/Experimental/PCG/Source/PCG/Public/Data/PCGDifferenceData.h
+++ b/Engine/Plugins/Experimental/PCG/Source/PCG/Public/Data/PCGDifferenceData.h
@@ -47,6 +47,12 @@ public:
 	virtual void PostEditChangeProperty(FPropertyChangedEvent& PropertyChangedEvent) override;
 #endif
 
+	// ~Begin UObject interface
+#if WITH_EDITOR
+	virtual void PostLoad();
+#endif
+	// ~End UObject interface
+
 	// ~Begin UPCGData interface
 	virtual EPCGDataType GetDataType() const override { return EPCGDataType::Spatial; }
 	virtual void VisitDataNetwork(TFunctionRef<void(const UPCGData*)> Action) const override;
@@ -58,8 +64,8 @@ public:
 	virtual FBox GetStrictBounds() const override;
 	virtual bool SamplePoint(const FTransform& Transform, const FBox& Bounds, FPCGPoint& OutPoint, UPCGMetadata* OutMetadata) const override;
 	virtual bool HasNonTrivialTransform() const override;
-	virtual const UPCGSpatialData* FindShapeFromNetwork(const int InDimension) const override { return Source ? Source->FindShapeFromNetwork(InDimension) : nullptr; }
-	virtual const UPCGSpatialData* FindFirstConcreteShapeFromNetwork() const override { return Source ? Source->FindFirstConcreteShapeFromNetwork() : nullptr; }
+	virtual const UPCGSpatialData* FindShapeFromNetwork(const int InDimension) const override { return GetSource() ? GetSource()->FindShapeFromNetwork(InDimension) : nullptr; }
+	virtual const UPCGSpatialData* FindFirstConcreteShapeFromNetwork() const override { return GetSource() ? GetSource()->FindFirstConcreteShapeFromNetwork() : nullptr; }
 protected:
 	virtual UPCGSpatialData* CopyInternal() const override;
 	//~End UPCGSpatialData interface
@@ -80,6 +86,25 @@ protected:
 
 	UPROPERTY(BlueprintSetter = SetDensityFunction, EditAnywhere, Category = Settings)
 	EPCGDifferenceDensityFunction DensityFunction = EPCGDifferenceDensityFunction::Minimum;
+
+#if WITH_EDITOR
+	inline const UPCGSpatialData* GetSource() const { return RawPointerSource; }
+	inline const UPCGSpatialData* GetDifference() const { return RawPointerDifference; }
+	inline UPCGUnionData* GetDifferencesUnion() const { return RawPointerDifferencesUnion; }
+#else
+	inline const UPCGSpatialData* GetSource() const { return Source.Get(); }
+	inline const UPCGSpatialData* GetDifference() const { return Difference.Get(); }
+	inline UPCGUnionData* GetDifferencesUnion() const { return DifferencesUnion.Get(); }
+#endif
+
+#if WITH_EDITOR
+private:
+	// Cached pointers to avoid dereferencing object pointer which does access tracking and supports lazy loading, and can come with substantial
+	// overhead (add trace marker to FObjectPtr::Get to see).
+	const UPCGSpatialData* RawPointerSource = nullptr;
+	const UPCGSpatialData* RawPointerDifference = nullptr;
+	UPCGUnionData* RawPointerDifferencesUnion = nullptr;
+#endif
 };
 
 #if UE_ENABLE_INCLUDE_ORDER_DEPRECATED_IN_5_2
