[PCG] Add option to generate BP on drop with the addition of a new PCG Engine Settings

* Moved "VolumeScale" and "GenerateOnDrop" from PCGEditorSettings to PCGEngineSettings
* BP will now be generated on drop in editor if the setting is enabled and they are flagged GenerateOnLoad.
* PCG Graphs are still behaving the same on drag and drop.

diff --git a/Engine/Plugins/Experimental/PCG/Source/PCG/PCG.Build.cs b/Engine/Plugins/Experimental/PCG/Source/PCG/PCG.Build.cs
--- a/Engine/Plugins/Experimental/PCG/Source/PCG/PCG.Build.cs
+++ b/Engine/Plugins/Experimental/PCG/Source/PCG/PCG.Build.cs
@@ -22,7 +22,8 @@ namespace UnrealBuildTool.Rules
 
 			PrivateDependencyModuleNames.AddRange(
 				new string[]
-				{					
+				{
+					"Settings",
 					"Voronoi",
 				}
 			);
diff --git a/Engine/Plugins/Experimental/PCG/Source/PCG/Private/PCGComponent.cpp b/Engine/Plugins/Experimental/PCG/Source/PCG/Private/PCGComponent.cpp
--- a/Engine/Plugins/Experimental/PCG/Source/PCG/Private/PCGComponent.cpp
+++ b/Engine/Plugins/Experimental/PCG/Source/PCG/Private/PCGComponent.cpp
@@ -1,24 +1,23 @@
 // Copyright Epic Games, Inc. All Rights Reserved.
 
 #include "PCGComponent.h"
-#include "Data/PCGProjectionData.h"
-#include "Engine/Engine.h"
-#include "Data/PCGSpatialData.h"
+
+#include "PCGContext.h"
+#include "PCGEngineSettings.h"
 #include "PCGGraph.h"
-#include "LandscapeComponent.h"
 #include "PCGHelpers.h"
-#include "LandscapeProxy.h"
 #include "PCGInputOutputSettings.h"
-#include "PCGContext.h"
+#include "PCGManagedResource.h"
 #include "PCGParamData.h"
 #include "PCGSubsystem.h"
-#include "PCGManagedResource.h"
 #include "Data/PCGDifferenceData.h"
 #include "Data/PCGIntersectionData.h"
 #include "Data/PCGLandscapeData.h"
 #include "Data/PCGLandscapeSplineData.h"
 #include "Data/PCGPointData.h"
 #include "Data/PCGPrimitiveData.h"
+#include "Data/PCGProjectionData.h"
+#include "Data/PCGSpatialData.h"
 #include "Data/PCGSplineData.h"
 #include "Data/PCGUnionData.h"
 #include "Data/PCGVolumeData.h"
@@ -26,11 +25,14 @@
 #include "Helpers/PCGActorHelpers.h"
 #include "Utils/PCGGeneratedResourcesLogging.h"
 
+#include "LandscapeComponent.h"
+#include "LandscapeProxy.h"
 #include "Algo/Transform.h"
 #include "Components/BillboardComponent.h"
 #include "Components/InstancedStaticMeshComponent.h"
 #include "Components/SplineComponent.h"
 #include "Components/ShapeComponent.h"
+#include "Engine/Engine.h"
 #include "GameFramework/Volume.h"
 #include "Kismet/GameplayStatics.h"
 #include "LandscapeSplinesComponent.h"
@@ -52,6 +54,14 @@ UPCGComponent::UPCGComponent(const FObjectInitializer& InObjectInitializer)
 	: Super(InObjectInitializer)
 {
 	GraphInstance = InObjectInitializer.CreateDefaultSubobject<UPCGGraphInstance>(this, TEXT("PCGGraphInstance"));
+
+#if WITH_EDITOR
+	// If we are in Editor, and we are a BP template (no owner), we will mark this component to force a generate when added to world.
+	if (!PCGHelpers::IsRuntimeOrPIE() && !GetOwner() && !HasAnyFlags(RF_ClassDefaultObject))
+	{
+		bForceGenerateOnBPAddedToWorld = true;
+	}
+#endif // WITH_EDITOR
 }
 
 bool UPCGComponent::CanPartition() const
@@ -400,6 +410,9 @@ void UPCGComponent::PostProcessGraph(const FBox& InNewBounds, bool bInGenerated,
 		CurrentGenerationTask = InvalidPCGTaskId;
 
 #if WITH_EDITOR
+		// Reset this flag to avoid re-generating on further refreshes.
+		bForceGenerateOnBPAddedToWorld = false;
+
 		bDirtyGenerated = false;
 		OnPCGGraphGeneratedDelegate.Broadcast(this);
 #endif
@@ -1000,7 +1013,7 @@ void UPCGComponent::PostLoad()
 	}
 
 	/** Deprecation code, should be removed once generated data has been updated */
-	if (bGenerated && GeneratedResources.Num() == 0)
+	if (GetOwner() && bGenerated && GeneratedResources.Num() == 0)
 	{
 		TArray<UInstancedStaticMeshComponent*> ISMCs;
 		GetOwner()->GetComponents(ISMCs);
@@ -1855,6 +1868,20 @@ void UPCGComponent::Refresh()
 	}
 }
 
+bool UPCGComponent::ShouldGenerateBPPCGAddedToWorld() const
+{
+	if (PCGHelpers::IsRuntimeOrPIE())
+	{
+		return false;
+	}
+	else
+	{
+		// Generate on drop can be disabled by global settings or locally by not having "GenerateOnLoad" as a generation trigger.
+		const UPCGEngineSettings* Settings = GetDefault<UPCGEngineSettings>();
+		return Settings ? (Settings->bGenerateOnDrop && bForceGenerateOnBPAddedToWorld && (GenerationTrigger == EPCGComponentGenerationTrigger::GenerateOnLoad)) : false;
+	}
+}
+
 void UPCGComponent::OnRefresh()
 {
 	// Mark the refresh task invalid to allow re-triggering refreshes
@@ -1889,7 +1916,8 @@ void UPCGComponent::OnRefresh()
 	else
 	{
 		// If we just cleaned up resources, call back generate
-		if (bWasGenerated && (!bGenerated || bRegenerateInEditor))
+		// Also, for BPs, we ask if we should generate, to support generate on added to world.
+		if ((bWasGenerated || ShouldGenerateBPPCGAddedToWorld()) && (!bGenerated || bRegenerateInEditor))
 		{
 			GenerateLocal(/*bForce=*/false);
 		}
diff --git a/Engine/Plugins/Experimental/PCG/Source/PCG/Private/PCGModule.cpp b/Engine/Plugins/Experimental/PCG/Source/PCG/Private/PCGModule.cpp
--- a/Engine/Plugins/Experimental/PCG/Source/PCG/Private/PCGModule.cpp
+++ b/Engine/Plugins/Experimental/PCG/Source/PCG/Private/PCGModule.cpp
@@ -2,16 +2,18 @@
 
 #include "PCGModule.h"
 
-#if WITH_EDITOR
-#include "Elements/PCGDifferenceElement.h"
+#include "PCGEngineSettings.h"
 
+#include "ISettingsModule.h"
 #include "Modules/ModuleManager.h"
+
+#if WITH_EDITOR
+#include "Elements/PCGDifferenceElement.h"
 #include "Tests/Determinism/PCGDeterminismNativeTests.h"
 #include "Tests/Determinism/PCGDifferenceDeterminismTest.h"
-#else
-#include "Modules/ModuleManager.h"
 #endif
 
+#define LOCTEXT_NAMESPACE "FPCGModule"
 
 class FPCGModule final : public IModuleInterface
 {
@@ -29,6 +31,9 @@ public:
 	//~ End IModuleInterface implementation
 
 private:
+	void RegisterSettings();
+	void UnregisterSettings();
+
 #if WITH_EDITOR
 	void RegisterNativeElementDeterminismTests();
 	void DeregisterNativeElementDeterminismTests();
@@ -37,6 +42,8 @@ private:
 
 void FPCGModule::StartupModule()
 {
+	RegisterSettings();
+
 #if WITH_EDITOR
 	PCGDeterminismTests::FNativeTestRegistry::Create();
 
@@ -46,6 +53,8 @@ void FPCGModule::StartupModule()
 
 void FPCGModule::ShutdownModule()
 {
+	UnregisterSettings();
+
 #if WITH_EDITOR
 	DeregisterNativeElementDeterminismTests();
 
@@ -53,6 +62,25 @@ void FPCGModule::ShutdownModule()
 #endif
 }
 
+void FPCGModule::RegisterSettings()
+{
+	if (ISettingsModule* SettingsModule = FModuleManager::GetModulePtr<ISettingsModule>("Settings"))
+	{
+		SettingsModule->RegisterSettings("Project", "Plugins", "PCG",
+			LOCTEXT("PCGEngineSettingsName", "PCG"),
+			LOCTEXT("PCGEngineSettingsDescription", "Configure PCG."),
+			GetMutableDefault<UPCGEngineSettings>());
+	}
+}
+
+void FPCGModule::UnregisterSettings()
+{
+	if (ISettingsModule* SettingsModule = FModuleManager::GetModulePtr<ISettingsModule>("Settings"))
+	{
+		SettingsModule->UnregisterSettings("Project", "Plugins", "PCG");
+	}
+}
+
 #if WITH_EDITOR
 
 void FPCGModule::RegisterNativeElementDeterminismTests()
@@ -71,3 +99,5 @@ void FPCGModule::DeregisterNativeElementDeterminismTests()
 IMPLEMENT_MODULE(FPCGModule, PCG);
 
 PCG_API DEFINE_LOG_CATEGORY(LogPCG);
+
+#undef LOCTEXT_NAMESPACE
diff --git a/Engine/Plugins/Experimental/PCG/Source/PCG/Public/PCGComponent.h b/Engine/Plugins/Experimental/PCG/Source/PCG/Public/PCGComponent.h
--- a/Engine/Plugins/Experimental/PCG/Source/PCG/Public/PCGComponent.h
+++ b/Engine/Plugins/Experimental/PCG/Source/PCG/Public/PCGComponent.h
@@ -198,6 +198,12 @@ public:
 	UPROPERTY(BlueprintReadOnly, VisibleAnywhere, Transient, AdvancedDisplay, Category = Properties)
 	bool bDirtyGenerated = false;
 
+	// Property that will automatically be set on BP templates, to allow for "Generate on add to world" in editor.
+	// Set it as a property to automatically transfer it to its child.
+	// Don't use it directly, use ShouldGenerateBPPCGAddedToWorld, as there are other conditions checked.
+	UPROPERTY()
+	bool bForceGenerateOnBPAddedToWorld = false;
+
 	FOnPCGGraphGenerated OnPCGGraphGeneratedDelegate;
 	FOnPCGGraphCleaned OnPCGGraphCleanedDelegate;
 #endif
@@ -236,6 +242,9 @@ public:
 
 	/** Used by the tracking system to know if the component need to track actors. Not enabled for now, tracking is still done on the component.*/
 	bool ShouldTrackActors() const { return false; }
+
+	/** Know if we need to force a generation, in case of BP added to the world in editor */
+	bool ShouldGenerateBPPCGAddedToWorld() const;
 #endif
 
 	/** Utility function (mostly for tests) to properly set the value of bIsComponentPartitioned.
diff --git a/Engine/Plugins/Experimental/PCG/Source/PCG/Public/PCGEngineSettings.h b/Engine/Plugins/Experimental/PCG/Source/PCG/Public/PCGEngineSettings.h
new file mode 100644
--- /dev/null
+++ b/Engine/Plugins/Experimental/PCG/Source/PCG/Public/PCGEngineSettings.h
@@ -0,0 +1,22 @@
+// Copyright Epic Games, Inc. All Rights Reserved.
+
+#pragma once
+
+#include "UObject/Object.h"
+
+#include "PCGEngineSettings.generated.h"
+
+UCLASS(config = EditorPerProjectUserSettings)
+class PCG_API UPCGEngineSettings : public UObject
+{
+	GENERATED_BODY()
+
+public:
+	/** Specifies the scale of the volume created on PCG graph drag/drop */
+	UPROPERTY(EditAnywhere, Config, Category = Workflow)
+	FVector VolumeScale = FVector(25, 25, 10);
+
+	/** Whether we want to generate PCG graph/BP with PCG after drag/drop or not */
+	UPROPERTY(EditAnywhere, Config, Category = Workflow)
+	bool bGenerateOnDrop = true;
+};
diff --git a/Engine/Plugins/Experimental/PCG/Source/PCGEditor/Private/PCGVolumeFactory.cpp b/Engine/Plugins/Experimental/PCG/Source/PCGEditor/Private/PCGVolumeFactory.cpp
--- a/Engine/Plugins/Experimental/PCG/Source/PCGEditor/Private/PCGVolumeFactory.cpp
+++ b/Engine/Plugins/Experimental/PCG/Source/PCGEditor/Private/PCGVolumeFactory.cpp
@@ -3,7 +3,7 @@
 #include "PCGVolumeFactory.h"
 
 #include "PCGComponent.h"
-#include "PCGEditorSettings.h"
+#include "PCGEngineSettings.h"
 #include "PCGGraph.h"
 #include "PCGVolume.h"
 
@@ -45,15 +45,15 @@ void UPCGVolumeFactory::PostSpawnActor(UObject* Asset, AActor* NewActor)
 
 	if (UPCGGraph* PCGGraph = Cast<UPCGGraph>(Asset))
 	{
-		const UPCGEditorSettings* PCGEditorSettings = GetDefault<UPCGEditorSettings>();
+		const UPCGEngineSettings* Settings = GetDefault<UPCGEngineSettings>();
 		
 		APCGVolume* PCGVolume = CastChecked<APCGVolume>(NewActor);
-		PCGVolume->SetActorScale3D(PCGEditorSettings->VolumeScale);
+		PCGVolume->SetActorScale3D(Settings->VolumeScale);
 		
 		UPCGComponent* PCGComponent = CastChecked<UPCGComponent>(PCGVolume->GetComponentByClass(UPCGComponent::StaticClass()));
 		PCGComponent->SetGraph(PCGGraph);
 
-		if (PCGEditorSettings->bGenerateOnDrop)
+		if (Settings->bGenerateOnDrop)
 		{
 			PCGComponent->Generate();
 		}
diff --git a/Engine/Plugins/Experimental/PCG/Source/PCGEditor/Public/PCGEditorSettings.h b/Engine/Plugins/Experimental/PCG/Source/PCGEditor/Public/PCGEditorSettings.h
--- a/Engine/Plugins/Experimental/PCG/Source/PCGEditor/Public/PCGEditorSettings.h
+++ b/Engine/Plugins/Experimental/PCG/Source/PCGEditor/Public/PCGEditorSettings.h
@@ -132,14 +132,6 @@ public:
 	/** Specify if we want to jump to definition in case of double click on native PCG Nodes */
 	UPROPERTY(EditAnywhere, config, Category = Workflow)
 	bool bEnableNavigateToNativeNodes;
-
-	/** Specifies the scale of the volume created on PCG graph drag/drop */
-	UPROPERTY(EditAnywhere, Config, Category = Workflow)
-	FVector VolumeScale = FVector(25, 25, 10);
-
-	/** Whether we want to generate the PCG graph after drag/drop or not */
-	UPROPERTY(EditAnywhere, Config, Category = Workflow)
-	bool bGenerateOnDrop = true;
 };
 
 #if UE_ENABLE_INCLUDE_ORDER_DEPRECATED_IN_5_2
