[PCG] Set the Async chunk size to 64 (from 16) by default. Also add a CVar for overriding this value, for quick benchmarking.

diff --git a/Engine/Plugins/Experimental/PCG/Source/PCG/Private/Helpers/PCGAsync.cpp b/Engine/Plugins/Experimental/PCG/Source/PCG/Private/Helpers/PCGAsync.cpp
--- a/Engine/Plugins/Experimental/PCG/Source/PCG/Private/Helpers/PCGAsync.cpp
+++ b/Engine/Plugins/Experimental/PCG/Source/PCG/Private/Helpers/PCGAsync.cpp
@@ -9,12 +9,18 @@
 
 namespace FPCGAsync
 {
-	TAutoConsoleVariable<bool> ConsoleVar::CVarDisableAsyncTimeSlicing{
+	TAutoConsoleVariable<bool> ConsoleVar::CVarDisableAsyncTimeSlicing {
 		TEXT("pcg.DisableAsyncTimeSlicing"),
 		false,
 		TEXT("To help debug, we can disable time slicing for async tasks.")
 	};
 
+	TAutoConsoleVariable<int32> ConsoleVar::CVarAsyncOverrideChunkSize{
+		TEXT("pcg.AsyncOverrideChunkSize"),
+		-1,
+		TEXT("For quick benchmarking, we can override the value of chunk size for async processing. Any negative value is discarded.")
+	};
+
 	void AsyncPointProcessing(int32 NumAvailableTasks, int32 MinIterationsPerTask, int32 NumIterations, TArray<FPCGPoint>& OutPoints, const TFunction<int32(int32, int32)>& IterationInnerLoop)
 	{
 		TRACE_CPUPROFILER_EVENT_SCOPE(IPCGElement::AsyncPointProcessing);
diff --git a/Engine/Plugins/Experimental/PCG/Source/PCG/Public/Helpers/PCGAsync.h b/Engine/Plugins/Experimental/PCG/Source/PCG/Public/Helpers/PCGAsync.h
--- a/Engine/Plugins/Experimental/PCG/Source/PCG/Public/Helpers/PCGAsync.h
+++ b/Engine/Plugins/Experimental/PCG/Source/PCG/Public/Helpers/PCGAsync.h
@@ -20,6 +20,7 @@ namespace FPCGAsync
 	struct ConsoleVar
 	{
 		static TAutoConsoleVariable<bool> CVarDisableAsyncTimeSlicing;
+		static TAutoConsoleVariable<int32> CVarAsyncOverrideChunkSize;
 	};
 
 	/** 
@@ -93,10 +94,13 @@ namespace FPCGAsync
 	namespace Private
 	{
 	template <typename OutputType, typename Func>
-	bool AsyncProcessing(FPCGAsyncState& AsyncState, int32 NumIterations, TArray<OutputType>& OutData, Func IterationInnerLoop, const bool bInEnableTimeSlicing, const int32 ChunkSize)
+	bool AsyncProcessing(FPCGAsyncState& AsyncState, int32 NumIterations, TArray<OutputType>& OutData, Func IterationInnerLoop, const bool bInEnableTimeSlicing, const int32 InChunkSize)
 	{
 		TRACE_CPUPROFILER_EVENT_SCOPE(FPCGAsync::AsyncProcessing);
 
+		const int32 OverrideChunkSize = ConsoleVar::CVarAsyncOverrideChunkSize.GetValueOnAnyThread();
+		const int32 ChunkSize = OverrideChunkSize > 0 ? OverrideChunkSize : InChunkSize;
+
 		const int32 StartIndex = AsyncState.AsyncCurrentReadIndex;
 		const bool bEnableTimeSlicing = bInEnableTimeSlicing && !ConsoleVar::CVarDisableAsyncTimeSlicing.GetValueOnAnyThread();
 
@@ -318,7 +322,7 @@ namespace FPCGAsync
 	* @returns true if the processing is done, false otherwise. Use this to know if you need to reschedule the task.
 	*/
 	template <typename OutputType, typename Func>
-	bool AsyncProcessing(FPCGAsyncState* AsyncState, int32 NumIterations, TArray<OutputType>& OutData, Func InFunc, const bool bEnableTimeSlicing, const int32 ChunkSize = 16)
+	bool AsyncProcessing(FPCGAsyncState* AsyncState, int32 NumIterations, TArray<OutputType>& OutData, Func InFunc, const bool bEnableTimeSlicing, const int32 ChunkSize = 64)
 	{
 		auto IterationInnerLoop = [&InFunc, &OutData](int32 StartReadIndex, int32 StartWriteIndex, int32 Count) -> int32
 		{
