Fix PCG settings interface usage to be editor only
#rnx

diff --git a/Engine/Plugins/Experimental/PCG/Source/PCG/Private/PCGModule.cpp b/Engine/Plugins/Experimental/PCG/Source/PCG/Private/PCGModule.cpp
--- a/Engine/Plugins/Experimental/PCG/Source/PCG/Private/PCGModule.cpp
+++ b/Engine/Plugins/Experimental/PCG/Source/PCG/Private/PCGModule.cpp
@@ -4,11 +4,11 @@
 
 #include "PCGEngineSettings.h"
 
-#include "ISettingsModule.h"
 #include "Modules/ModuleManager.h"
 
 #if WITH_EDITOR
 #include "Elements/PCGDifferenceElement.h"
+#include "ISettingsModule.h"
 #include "Tests/Determinism/PCGDeterminismNativeTests.h"
 #include "Tests/Determinism/PCGDifferenceDeterminismTest.h"
 #endif
@@ -19,9 +19,11 @@ class FPCGModule final : public IModuleInterface
 {
 public:
 	//~ IModuleInterface implementation
-	virtual void StartupModule() override;
 
+#if WITH_EDITOR
+	virtual void StartupModule() override;
 	virtual void ShutdownModule() override;
+#endif
 
 	virtual bool SupportsDynamicReloading() override
 	{
@@ -31,35 +33,32 @@ public:
 	//~ End IModuleInterface implementation
 
 private:
+#if WITH_EDITOR
 	void RegisterSettings();
 	void UnregisterSettings();
 
-#if WITH_EDITOR
 	void RegisterNativeElementDeterminismTests();
 	void DeregisterNativeElementDeterminismTests();
 #endif
 };
 
+#if WITH_EDITOR
 void FPCGModule::StartupModule()
 {
 	RegisterSettings();
 
-#if WITH_EDITOR
 	PCGDeterminismTests::FNativeTestRegistry::Create();
 
 	RegisterNativeElementDeterminismTests();
-#endif
 }
 
 void FPCGModule::ShutdownModule()
 {
 	UnregisterSettings();
 
-#if WITH_EDITOR
 	DeregisterNativeElementDeterminismTests();
 
 	PCGDeterminismTests::FNativeTestRegistry::Destroy();
-#endif
 }
 
 void FPCGModule::RegisterSettings()
@@ -81,8 +80,6 @@ void FPCGModule::UnregisterSettings()
 	}
 }
 
-#if WITH_EDITOR
-
 void FPCGModule::RegisterNativeElementDeterminismTests()
 {
 	PCGDeterminismTests::FNativeTestRegistry::RegisterTestFunction(UPCGDifferenceSettings::StaticClass(), PCGDeterminismTests::DifferenceElement::RunTestSuite);
