Added time-slicing capabilities to the mesh selector

diff --git a/Engine/Plugins/Experimental/PCG/Source/PCG/Private/Elements/PCGStaticMeshSpawner.cpp b/Engine/Plugins/Experimental/PCG/Source/PCG/Private/Elements/PCGStaticMeshSpawner.cpp
--- a/Engine/Plugins/Experimental/PCG/Source/PCG/Private/Elements/PCGStaticMeshSpawner.cpp
+++ b/Engine/Plugins/Experimental/PCG/Source/PCG/Private/Elements/PCGStaticMeshSpawner.cpp
@@ -6,8 +6,8 @@
 #include "PCGManagedResource.h"
 #include "Data/PCGPointData.h"
 #include "Data/PCGSpatialData.h"
+#include "Elements/PCGStaticMeshSpawnerContext.h"
 #include "Helpers/PCGActorHelpers.h"
-
 #include "InstancePackers/PCGInstancePackerBase.h"
 #include "MeshSelectors/PCGMeshSelectorBase.h"
 #include "MeshSelectors/PCGMeshSelectorWeighted.h"
@@ -70,7 +70,7 @@ bool FPCGStaticMeshSpawnerElement::PrepareDataInternal(FPCGContext* InContext) c
 #endif
 
 	// Check if we can reuse existing resources
-	bool bSkippedDueToReuse = false;
+	bool& bSkippedDueToReuse = Context->bSkippedDueToReuse;
 
 	if (Context->CurrentInputIndex == 0 && CVarAllowISMReuse.GetValueOnAnyThread())
 	{
@@ -118,76 +118,103 @@ bool FPCGStaticMeshSpawnerElement::PrepareDataInternal(FPCGContext* InContext) c
 
 	while(Context->CurrentInputIndex < Inputs.Num())
 	{
-		const FPCGTaggedData& Input = Inputs[Context->CurrentInputIndex];
-		const UPCGSpatialData* SpatialData = Cast<UPCGSpatialData>(Input.Data);
+		if (!Context->bCurrentInputSetup)
+		{
+			const FPCGTaggedData& Input = Inputs[Context->CurrentInputIndex];
+			const UPCGSpatialData* SpatialData = Cast<UPCGSpatialData>(Input.Data);
 
-		// Preincrement so we can have all paths continue properly - note that this works because none of this is internally time-sliced
-		++Context->CurrentInputIndex;
+			if (!SpatialData)
+			{
+				PCGE_LOG(Error, "Invalid input data");
+				++Context->CurrentInputIndex;
+				continue;
+			}
 
-		if (!SpatialData)
-		{
-			PCGE_LOG(Error, "Invalid input data");
-			continue;
-		}
+			const UPCGPointData* PointData = SpatialData->ToPointData(Context);
+			if (!PointData)
+			{
+				PCGE_LOG(Error, "Unable to get point data from input");
+				++Context->CurrentInputIndex;
+				continue;
+			}
 
-		const UPCGPointData* PointData = SpatialData->ToPointData(Context);
-		if (!PointData)
-		{
-			PCGE_LOG(Error, "Unable to get point data from input");
-			continue;
-		}
+			AActor* TargetActor = Context->GetTargetActor(PointData);
+			if (!TargetActor)
+			{
+				PCGE_LOG(Error, "Invalid target actor");
+				++Context->CurrentInputIndex;
+				continue;
+			}
 
-		AActor* TargetActor = Context->GetTargetActor(PointData);
-		if (!TargetActor)
-		{
-			PCGE_LOG(Error, "Invalid target actor");
-			continue;
-		}
+			if (bGenerateOutput)
+			{
+				FPCGTaggedData& Output = Outputs.Add_GetRef(Input);
 
-		UPCGPointData* OutputPointData = nullptr;
+				UPCGPointData* OutputPointData = NewObject<UPCGPointData>();
+				OutputPointData->InitializeFromData(PointData);
 
-		if (bGenerateOutput)
-		{
-			FPCGTaggedData& Output = Outputs.Add_GetRef(Input);
+				if (OutputPointData->Metadata->HasAttribute(Settings->OutAttributeName))
+				{
+					OutputPointData->Metadata->DeleteAttribute(Settings->OutAttributeName);
+					PCGE_LOG(Verbose, "Metadata attribute %s is being overwritten in the output data", *Settings->OutAttributeName.ToString());
+				}
 
-			OutputPointData = NewObject<UPCGPointData>();
-			OutputPointData->InitializeFromData(PointData);
+				OutputPointData->Metadata->CreateStringAttribute(Settings->OutAttributeName, FName(NAME_None).ToString(), /*bAllowsInterpolation=*/false);
 
-			if (OutputPointData->Metadata->HasAttribute(Settings->OutAttributeName))
-			{
-				OutputPointData->Metadata->DeleteAttribute(Settings->OutAttributeName);
-				PCGE_LOG(Verbose, "Metadata attribute %s is being overwritten in the output data", *Settings->OutAttributeName.ToString());
-			}
+				Output.Data = OutputPointData;
+				check(!Context->CurrentOutputPointData);
+				Context->CurrentOutputPointData = OutputPointData;
 
-			OutputPointData->Metadata->CreateStringAttribute(Settings->OutAttributeName, FName(NAME_None).ToString(), /*bAllowsInterpolation=*/false);
+				// Create an entry in the MeshInstancesData if we're going to keep it
+				if (!bSkippedDueToReuse)
+				{
+					FPCGStaticMeshSpawnerContext::FPackedInstanceListData& InstanceListData = Context->MeshInstancesData.Emplace_GetRef();
+					InstanceListData.TargetActor = TargetActor;
+					InstanceListData.SpatialData = PointData;
+				}
+			}
 
-			Output.Data = OutputPointData;
+			Context->CurrentPointData = PointData;
+			Context->bCurrentInputSetup = true;
 		}
 
-		TArray<FPCGMeshInstanceList> MeshInstances;
-		Settings->MeshSelectorInstance->SelectInstances(*Context, Settings, PointData, MeshInstances, OutputPointData);
+		if (!Context->bSelectionDone)
+		{
+			TArray<FPCGMeshInstanceList> DummyMeshInstances;
+			TArray<FPCGMeshInstanceList>& MeshInstances = (bSkippedDueToReuse ? DummyMeshInstances : Context->MeshInstancesData.Last().MeshInstances);
+
+			check(Context->CurrentPointData);
+			Context->bSelectionDone = Settings->MeshSelectorInstance->SelectInstances(*Context, Settings, Context->CurrentPointData, MeshInstances, Context->CurrentOutputPointData);
+		}
 
-		// If we need the output but would otherwise skip the resource creation, just don't push them to the MeshInstancesData array
-		if (bSkippedDueToReuse)
+		if (!Context->bSelectionDone)
 		{
-			continue;
+			return false;
 		}
 
-		TArray<FPCGPackedCustomData> PackedCustomData;
-		PackedCustomData.SetNum(MeshInstances.Num());
-		if (Settings->InstancePackerInstance)
+		// If we need the output but would otherwise skip the resource creation, we don't need to run the instance packing part of the processing
+		if (!bSkippedDueToReuse)
 		{
-			for(int32 InstanceListIndex = 0; InstanceListIndex < MeshInstances.Num(); ++InstanceListIndex)
+			TArray<FPCGPackedCustomData>& PackedCustomData = Context->MeshInstancesData.Last().PackedCustomData;
+			const TArray<FPCGMeshInstanceList>& MeshInstances = Context->MeshInstancesData.Last().MeshInstances;
+
+			if (PackedCustomData.Num() != MeshInstances.Num())
+			{
+				PackedCustomData.SetNum(MeshInstances.Num());
+			}
+
+			if (Settings->InstancePackerInstance)
 			{
-				Settings->InstancePackerInstance->PackInstances(*Context, PointData, MeshInstances[InstanceListIndex], PackedCustomData[InstanceListIndex]);
+				for (int32 InstanceListIndex = 0; InstanceListIndex < MeshInstances.Num(); ++InstanceListIndex)
+				{
+					Settings->InstancePackerInstance->PackInstances(*Context, Context->CurrentPointData, MeshInstances[InstanceListIndex], PackedCustomData[InstanceListIndex]);
+				}
 			}
 		}
 
-		FPCGStaticMeshSpawnerContext::FPackedInstanceListData& InstanceListData = Context->MeshInstancesData.Emplace_GetRef();
-		InstanceListData.TargetActor = TargetActor;
-		InstanceListData.SpatialData = PointData;
-		InstanceListData.MeshInstances = MoveTemp(MeshInstances);
-		InstanceListData.PackedCustomData = MoveTemp(PackedCustomData);
+		// We're done - cleanup for next iteration if we still have time
+		++Context->CurrentInputIndex;
+		Context->ResetInputIterationData();
 
 		// Continue on to next iteration if there is time left, otherwise, exit here
 		if (Context->AsyncState.ShouldStop() && Context->CurrentInputIndex < Inputs.Num())
diff --git a/Engine/Plugins/Experimental/PCG/Source/PCG/Private/Elements/PCGStaticMeshSpawnerContext.cpp b/Engine/Plugins/Experimental/PCG/Source/PCG/Private/Elements/PCGStaticMeshSpawnerContext.cpp
new file mode 100644
--- /dev/null
+++ b/Engine/Plugins/Experimental/PCG/Source/PCG/Private/Elements/PCGStaticMeshSpawnerContext.cpp
@@ -0,0 +1,18 @@
+// Copyright Epic Games, Inc. All Rights Reserved.
+
+#include "Elements/PCGStaticMeshSpawnerContext.h"
+
+void FPCGStaticMeshSpawnerContext::ResetInputIterationData()
+{
+	bCurrentInputSetup = false;
+	bSelectionDone = false;
+	CurrentPointData = nullptr;
+	CurrentOutputPointData = nullptr;
+	MaterialOverrideHelper.Reset();
+	CurrentPointIndex = 0;
+	ValueKeyToMesh.Reset();
+	WeightedMeshInstances.Reset();
+	MeshToValueKey.Reset();
+	CumulativeWeights.Reset();
+	CategoryEntryToInstancesAndWeights.Reset();
+}
\ No newline at end of file
diff --git a/Engine/Plugins/Experimental/PCG/Source/PCG/Private/MeshSelectors/PCGMeshSelectorBase.cpp b/Engine/Plugins/Experimental/PCG/Source/PCG/Private/MeshSelectors/PCGMeshSelectorBase.cpp
--- a/Engine/Plugins/Experimental/PCG/Source/PCG/Private/MeshSelectors/PCGMeshSelectorBase.cpp
+++ b/Engine/Plugins/Experimental/PCG/Source/PCG/Private/MeshSelectors/PCGMeshSelectorBase.cpp
@@ -1,37 +1,48 @@
 // Copyright Epic Games, Inc. All Rights Reserved.
 
 #include "MeshSelectors/PCGMeshSelectorBase.h"
+
 #include "PCGElement.h"
+#include "Elements/PCGStaticMeshSpawnerContext.h"
 
 #include UE_INLINE_GENERATED_CPP_BY_NAME(PCGMeshSelectorBase)
 
-FPCGMeshMaterialOverrideHelper::FPCGMeshMaterialOverrideHelper(
+void FPCGMeshMaterialOverrideHelper::Initialize(
 	FPCGContext& InContext,
 	bool bInUseMaterialOverrideAttributes,
 	const TArray<TSoftObjectPtr<UMaterialInterface>>& InStaticMaterialOverrides,
 	const TArray<FName>& InMaterialOverrideAttributeNames,
 	const UPCGMetadata* InMetadata
 )
-	: bUseMaterialOverrideAttributes(bInUseMaterialOverrideAttributes)
-	, StaticMaterialOverrides(InStaticMaterialOverrides)
-	, MaterialOverrideAttributeNames(InMaterialOverrideAttributeNames)
-	, Metadata(InMetadata)
 {
+	check(!bIsInitialized);
+
+	bUseMaterialOverrideAttributes = bInUseMaterialOverrideAttributes;
+	StaticMaterialOverrides = InStaticMaterialOverrides;
+	MaterialOverrideAttributeNames = InMaterialOverrideAttributeNames;
+	Metadata = InMetadata;
+
 	Initialize(InContext);
+
+	bIsInitialized = true;
 }
 
-FPCGMeshMaterialOverrideHelper::FPCGMeshMaterialOverrideHelper(
+void FPCGMeshMaterialOverrideHelper::Initialize(
 	FPCGContext& InContext,
 	bool bInByAttributeOverride,
 	const TArray<FName>& InMaterialOverrideAttributeNames,
 	const UPCGMetadata* InMetadata
 )
-	: bUseMaterialOverrideAttributes(bInByAttributeOverride)
-	, StaticMaterialOverrides(EmptyArray)
-	, MaterialOverrideAttributeNames(InMaterialOverrideAttributeNames)
-	, Metadata(InMetadata)
 {
+	check(!bIsInitialized);
+
+	bUseMaterialOverrideAttributes = bInByAttributeOverride;
+	MaterialOverrideAttributeNames = InMaterialOverrideAttributeNames;
+	Metadata = InMetadata;
+
 	Initialize(InContext);
+
+	bIsInitialized = true;
 }
 
 void FPCGMeshMaterialOverrideHelper::Initialize(FPCGContext& InContext)
@@ -72,6 +83,19 @@ void FPCGMeshMaterialOverrideHelper::Initialize(FPCGContext& InContext)
 	bIsValid = true;
 }
 
+void FPCGMeshMaterialOverrideHelper::Reset()
+{
+	MaterialAttributes.Reset();
+	ValueKeyToOverrideMaterials.Reset();
+	WorkingMaterialOverrides.Reset();
+	bIsInitialized = false;
+	bIsValid = false;
+	bUseMaterialOverrideAttributes = false;
+	StaticMaterialOverrides.Reset();
+	MaterialOverrideAttributeNames.Reset();
+	Metadata = nullptr;
+}
+
 const TArray<TSoftObjectPtr<UMaterialInterface>>& FPCGMeshMaterialOverrideHelper::GetMaterialOverrides(PCGMetadataEntryKey EntryKey)
 {
 	check(bIsValid);
diff --git a/Engine/Plugins/Experimental/PCG/Source/PCG/Private/MeshSelectors/PCGMeshSelectorByAttribute.cpp b/Engine/Plugins/Experimental/PCG/Source/PCG/Private/MeshSelectors/PCGMeshSelectorByAttribute.cpp
--- a/Engine/Plugins/Experimental/PCG/Source/PCG/Private/MeshSelectors/PCGMeshSelectorByAttribute.cpp
+++ b/Engine/Plugins/Experimental/PCG/Source/PCG/Private/MeshSelectors/PCGMeshSelectorByAttribute.cpp
@@ -4,6 +4,7 @@
 
 #include "Data/PCGPointData.h"
 #include "Elements/PCGStaticMeshSpawner.h"
+#include "Elements/PCGStaticMeshSpawnerContext.h"
 #include "Data/PCGSpatialData.h"
 #include "MeshSelectors/PCGMeshSelectorBase.h"
 
@@ -89,29 +90,31 @@ void UPCGMeshSelectorByAttribute::PostLoad()
 #endif
 }
 
-void UPCGMeshSelectorByAttribute::SelectInstances_Implementation(
-	FPCGContext& Context, 
+bool UPCGMeshSelectorByAttribute::SelectInstances(
+	FPCGStaticMeshSpawnerContext& Context,
 	const UPCGStaticMeshSpawnerSettings* Settings, 
 	const UPCGPointData* InPointData,
 	TArray<FPCGMeshInstanceList>& OutMeshInstances,
 	UPCGPointData* OutPointData) const
 {
+	TRACE_CPUPROFILER_EVENT_SCOPE(UPCGMeshSelectorByAttribute::SelectInstances);
+
 	if (!InPointData)
 	{
 		PCGE_LOG_C(Error, &Context, "Missing input data");
-		return;
+		return true;
 	}
 
 	if (!InPointData->Metadata)
 	{
 		PCGE_LOG_C(Error, &Context, "Unable to get metadata from input");
-		return;
+		return true;
 	}
 
 	if (!InPointData->Metadata->HasAttribute(AttributeName))
 	{
 		PCGE_LOG_C(Error, &Context, "Attribute %s is not in the metadata", *AttributeName.ToString());
-		return;
+		return true;
 	}
 
 	const FPCGMetadataAttributeBase* AttributeBase = InPointData->Metadata->GetConstAttribute(AttributeName);
@@ -120,33 +123,44 @@ void UPCGMeshSelectorByAttribute::SelectInstances_Implementation(
 	if (AttributeBase->GetTypeId() != PCG::Private::MetadataTypes<FString>::Id)
 	{
 		PCGE_LOG_C(Error, &Context, "Attribute is not of valid type FString");
-		return;
+		return true;
 	}
 
 	const FPCGMetadataAttribute<FString>* Attribute = static_cast<const FPCGMetadataAttribute<FString>*>(AttributeBase);
 
-	FPCGMeshMaterialOverrideHelper MaterialOverrideHelper(Context, bUseAttributeMaterialOverrides, TemplateDescriptor.OverrideMaterials, MaterialOverrideAttributes, InPointData->Metadata);
+	FPCGMeshMaterialOverrideHelper& MaterialOverrideHelper = Context.MaterialOverrideHelper;
+	if (!MaterialOverrideHelper.IsInitialized())
+	{
+		MaterialOverrideHelper.Initialize(Context, bUseAttributeMaterialOverrides, TemplateDescriptor.OverrideMaterials, MaterialOverrideAttributes, InPointData->Metadata);
+	}
 
 	if (!MaterialOverrideHelper.IsValid())
 	{
-		return;
+		return true;
 	}
 
-	TMap<PCGMetadataValueKey, TSoftObjectPtr<UStaticMesh>> ValueKeyToMesh;
-
 	// ByAttribute takes in SoftObjectPaths per point in the metadata, so we can pass those directly into the outgoing pin if it exists
-	if (OutPointData)
+	if (Context.CurrentPointIndex == 0 && OutPointData)
 	{
+		TRACE_CPUPROFILER_EVENT_SCOPE(UPCGMeshSelectorByAttribute::SetupOutPointData);
 		OutPointData->SetPoints(InPointData->GetPoints());
 		OutPointData->Metadata->DeleteAttribute(Settings->OutAttributeName);
 		OutPointData->Metadata->CopyAttribute(InPointData->Metadata, AttributeName, Settings->OutAttributeName);
 	}
 
-	TRACE_CPUPROFILER_EVENT_SCOPE(FPCGStaticMeshSpawnerElement::Execute::SelectEntries);
+	TRACE_CPUPROFILER_EVENT_SCOPE(UPCGMeshSelectorByAttribute::SelectEntries);
 
 	// Assign points to entries
-	for (const FPCGPoint& Point : InPointData->GetPoints())
+	int32 CurrentPointIndex = Context.CurrentPointIndex;
+	int32 LastCheckpointIndex = CurrentPointIndex;
+	constexpr int32 TimeSlicingCheckFrequency = 1024;
+	TMap<PCGMetadataValueKey, TSoftObjectPtr<UStaticMesh>>& ValueKeyToMesh = Context.ValueKeyToMesh;
+
+	const TArray<FPCGPoint>& Points = InPointData->GetPoints();
+	while(CurrentPointIndex < Points.Num())
 	{
+		const FPCGPoint& Point = Points[CurrentPointIndex++];
+
 		if (Point.Density <= 0.0f)
 		{
 			continue;
@@ -192,5 +206,22 @@ void UPCGMeshSelectorByAttribute::SelectInstances_Implementation(
 		FPCGMeshInstanceList& InstanceList = PCGMeshSelectorAttribute::GetInstanceList(OutMeshInstances, TemplateDescriptor, Mesh, MaterialOverrideHelper.GetMaterialOverrides(Point.MetadataEntry), bReverseTransform);
 		InstanceList.Instances.Emplace(Point.Transform);
 		InstanceList.InstancesMetadataEntry.Emplace(Point.MetadataEntry);
+
+		// Check if we should stop here and continue in a subsequent call
+		if(CurrentPointIndex - LastCheckpointIndex >= TimeSlicingCheckFrequency)
+		{
+			if (Context.ShouldStop())
+			{
+				break;
+			}
+			else
+			{
+				LastCheckpointIndex = CurrentPointIndex;
+			}
+		}
 	}
+
+	Context.CurrentPointIndex = CurrentPointIndex;
+
+	return (CurrentPointIndex == Points.Num());
 }
diff --git a/Engine/Plugins/Experimental/PCG/Source/PCG/Private/MeshSelectors/PCGMeshSelectorWeighted.cpp b/Engine/Plugins/Experimental/PCG/Source/PCG/Private/MeshSelectors/PCGMeshSelectorWeighted.cpp
--- a/Engine/Plugins/Experimental/PCG/Source/PCG/Private/MeshSelectors/PCGMeshSelectorWeighted.cpp
+++ b/Engine/Plugins/Experimental/PCG/Source/PCG/Private/MeshSelectors/PCGMeshSelectorWeighted.cpp
@@ -4,6 +4,7 @@
 
 #include "Data/PCGPointData.h"
 #include "Elements/PCGStaticMeshSpawner.h"
+#include "Elements/PCGStaticMeshSpawnerContext.h"
 #include "Data/PCGSpatialData.h"
 #include "Helpers/PCGBlueprintHelpers.h"
 
@@ -111,55 +112,62 @@ void UPCGMeshSelectorWeighted::PostLoad()
 #endif
 }
 
-void UPCGMeshSelectorWeighted::SelectInstances_Implementation(
-	FPCGContext& Context, 
+bool UPCGMeshSelectorWeighted::SelectInstances(
+	FPCGStaticMeshSpawnerContext& Context,
 	const UPCGStaticMeshSpawnerSettings* Settings, 
 	const UPCGPointData* InPointData,
 	TArray<FPCGMeshInstanceList>& OutMeshInstances,
 	UPCGPointData* OutPointData) const
 {
-	TArray<TArray<FPCGMeshInstanceList>> MeshInstances;
-	TArray<int> CumulativeWeights;
+	if (!InPointData)
+	{
+		PCGE_LOG_C(Error, &Context, "Missing input data");
+		return true;
+	}
 
-	int TotalWeight = 0;
+	TArray<TArray<FPCGMeshInstanceList>>& MeshInstances = Context.WeightedMeshInstances;
+	TArray<int>& CumulativeWeights = Context.CumulativeWeights;
 
-	// Prepare common mesh setups which we will use as a kind of map
-	for (const FPCGMeshSelectorWeightedEntry& Entry : MeshEntries)
+	// Setup
+	if (Context.CurrentPointIndex == 0)
 	{
-		if (Entry.Weight <= 0)
+		int TotalWeight = 0;
+
+		// Prepare common mesh setups which we will use as a kind of map
+		for (const FPCGMeshSelectorWeightedEntry& Entry : MeshEntries)
 		{
-			PCGE_LOG_C(Verbose, &Context, "Entry found with weight <= 0");
-			continue;
-		}
+			if (Entry.Weight <= 0)
+			{
+				PCGE_LOG_C(Verbose, &Context, "Entry found with weight <= 0");
+				continue;
+			}
 
-		TArray<FPCGMeshInstanceList>& PickEntry = MeshInstances.Emplace_GetRef();
-		PickEntry.Emplace_GetRef(Entry.Descriptor);
+			TArray<FPCGMeshInstanceList>& PickEntry = MeshInstances.Emplace_GetRef();
+			PickEntry.Emplace_GetRef(Entry.Descriptor);
 
-		TotalWeight += Entry.Weight;
-		CumulativeWeights.Add(TotalWeight);
+			TotalWeight += Entry.Weight;
+			CumulativeWeights.Add(TotalWeight);
+		}
+
+		if (TotalWeight <= 0)
+		{
+			return true;
+		}
 	}
 
-	if (TotalWeight <= 0)
+	FPCGMeshMaterialOverrideHelper& MaterialOverrideHelper = Context.MaterialOverrideHelper;
+	if (!MaterialOverrideHelper.IsInitialized())
 	{
-		return;
+		MaterialOverrideHelper.Initialize(Context, bUseAttributeMaterialOverrides, MaterialOverrideAttributes, InPointData->Metadata);
 	}
 
-	if (!InPointData)
+	if (!MaterialOverrideHelper.IsValid())
 	{
-		PCGE_LOG_C(Error, &Context, "Missing input data");
-		return;
+		return true;
 	}
 
 	TArray<FPCGPoint>* OutPoints = nullptr;
 	FPCGMetadataAttribute<FString>* OutAttribute = nullptr;
-	TMap<TSoftObjectPtr<UStaticMesh>, PCGMetadataValueKey> MeshToValueKey;
-
-	FPCGMeshMaterialOverrideHelper MaterialOverrideHelper(Context, bUseAttributeMaterialOverrides, MaterialOverrideAttributes, InPointData->Metadata);
-
-	if (!MaterialOverrideHelper.IsValid())
-	{
-		return;
-	}
 
 	if (OutPointData)
 	{
@@ -190,8 +198,17 @@ void UPCGMeshSelectorWeighted::SelectInstances_Implementation(
 	{
 		TRACE_CPUPROFILER_EVENT_SCOPE(FPCGStaticMeshSpawnerElement::Execute::SelectEntries);
 
-		for (const FPCGPoint& Point : InPointData->GetPoints())
+		int32 CurrentPointIndex = Context.CurrentPointIndex;
+		int32 LastCheckpointIndex = CurrentPointIndex;
+		constexpr int32 TimeSlicingCheckFrequency = 1024;
+		TMap<TSoftObjectPtr<UStaticMesh>, PCGMetadataValueKey>& MeshToValueKey = Context.MeshToValueKey;
+		const int32 TotalWeight = CumulativeWeights.Last();
+
+		const TArray<FPCGPoint>& Points = InPointData->GetPoints();
+		while(CurrentPointIndex < Points.Num())
 		{
+			const FPCGPoint& Point = Points[CurrentPointIndex++];
+
 			if (Point.Density <= 0.0f)
 			{
 				continue;
@@ -231,16 +248,40 @@ void UPCGMeshSelectorWeighted::SelectInstances_Implementation(
 					OutAttribute->SetValueFromValueKey(OutPoint.MetadataEntry, *OutValueKey);
 				}
 			}
+
+			// Check if we should stop here and continue in a subsequent call
+			if (CurrentPointIndex - LastCheckpointIndex >= TimeSlicingCheckFrequency)
+			{
+				if (Context.ShouldStop())
+				{
+					break;
+				}
+				else
+				{
+					LastCheckpointIndex = CurrentPointIndex;
+				}
+			}
 		}
+
+		Context.CurrentPointIndex = CurrentPointIndex;
 	}
 
-	// Collapse to OutMeshInstances
-	for (TArray<FPCGMeshInstanceList>& PickedMeshInstances : MeshInstances)
+	if (Context.CurrentPointIndex == InPointData->GetPoints().Num())
 	{
-		for (FPCGMeshInstanceList& PickedMeshInstanceEntry : PickedMeshInstances)
+		// Finally, collapse to OutMeshInstances
+		for (TArray<FPCGMeshInstanceList>& PickedMeshInstances : MeshInstances)
 		{
-			OutMeshInstances.Emplace(MoveTemp(PickedMeshInstanceEntry));
+			for (FPCGMeshInstanceList& PickedMeshInstanceEntry : PickedMeshInstances)
+			{
+				OutMeshInstances.Emplace(MoveTemp(PickedMeshInstanceEntry));
+			}
 		}
+
+		return true;
+	}
+	else
+	{
+		return false;
 	}
 }
 
diff --git a/Engine/Plugins/Experimental/PCG/Source/PCG/Private/MeshSelectors/PCGMeshSelectorWeightedByCategory.cpp b/Engine/Plugins/Experimental/PCG/Source/PCG/Private/MeshSelectors/PCGMeshSelectorWeightedByCategory.cpp
--- a/Engine/Plugins/Experimental/PCG/Source/PCG/Private/MeshSelectors/PCGMeshSelectorWeightedByCategory.cpp
+++ b/Engine/Plugins/Experimental/PCG/Source/PCG/Private/MeshSelectors/PCGMeshSelectorWeightedByCategory.cpp
@@ -3,8 +3,9 @@
 #include "MeshSelectors/PCGMeshSelectorWeightedByCategory.h"
 
 #include "Data/PCGPointData.h"
-#include "Elements/PCGStaticMeshSpawner.h"
 #include "Data/PCGSpatialData.h"
+#include "Elements/PCGStaticMeshSpawner.h"
+#include "Elements/PCGStaticMeshSpawnerContext.h"
 #include "Helpers/PCGBlueprintHelpers.h"
 
 #include "Math/RandomStream.h"
@@ -12,12 +13,6 @@
 
 #include UE_INLINE_GENERATED_CPP_BY_NAME(PCGMeshSelectorWeightedByCategory)
 
-struct FPCGInstancesAndWeights
-{
-	TArray<TArray<FPCGMeshInstanceList>> MeshInstances;
-	TArray<int> CumulativeWeights;
-};
-
 #if WITH_EDITOR
 void FPCGWeightedByCategoryEntryList::ApplyDeprecation()
 {
@@ -40,8 +35,8 @@ void UPCGMeshSelectorWeightedByCategory::PostLoad()
 #endif
 }
 
-void UPCGMeshSelectorWeightedByCategory::SelectInstances_Implementation(
-	FPCGContext& Context,
+bool UPCGMeshSelectorWeightedByCategory::SelectInstances(
+	FPCGStaticMeshSpawnerContext& Context,
 	const UPCGStaticMeshSpawnerSettings* Settings,
 	const UPCGPointData* InPointData,
 	TArray<FPCGMeshInstanceList>& OutMeshInstances,
@@ -50,19 +45,19 @@ void UPCGMeshSelectorWeightedByCategory::SelectInstances_Implementation(
 	if (!InPointData)
 	{
 		PCGE_LOG_C(Error, &Context, "Missing input data");
-		return;
+		return true;
 	}
 
 	if (!InPointData->Metadata)
 	{
 		PCGE_LOG_C(Error, &Context, "Unable to get metadata from input");
-		return;
+		return true;
 	}
 
 	if (!InPointData->Metadata->HasAttribute(CategoryAttribute))
 	{
 		PCGE_LOG_C(Error, &Context, "Attribute %s is not in the metadata", *CategoryAttribute.ToString());
-		return;
+		return true;
 	}
 
 	const FPCGMetadataAttributeBase* AttributeBase = InPointData->Metadata->GetConstAttribute(CategoryAttribute);
@@ -72,90 +67,97 @@ void UPCGMeshSelectorWeightedByCategory::SelectInstances_Implementation(
 	if (AttributeBase->GetTypeId() != PCG::Private::MetadataTypes<FString>::Id)
 	{
 		PCGE_LOG_C(Error, &Context, "Attribute is not of valid type FString");
-		return;
+		return true;
 	}
 
 	const FPCGMetadataAttribute<FString>* Attribute = static_cast<const FPCGMetadataAttribute<FString>*>(AttributeBase);
 
 	// maps a CategoryEntry ValueKey to the meshes and precomputed weight data 
-	TMap<PCGMetadataValueKey, FPCGInstancesAndWeights> CategoryEntryToInstancesAndWeights;
+	TMap<PCGMetadataValueKey, FPCGInstancesAndWeights>& CategoryEntryToInstancesAndWeights = Context.CategoryEntryToInstancesAndWeights;
 
 	// unmarked points will fallback to the MeshEntries associated with the DefaultValueKey
 	PCGMetadataValueKey DefaultValueKey = PCGDefaultValueKey;
 
-	for (const FPCGWeightedByCategoryEntryList& Entry : Entries)
+	// Setup
+	if (Context.CurrentPointIndex == 0)
 	{
-		if (Entry.WeightedMeshEntries.Num() == 0)
+		for (const FPCGWeightedByCategoryEntryList& Entry : Entries)
 		{
-			PCGE_LOG_C(Verbose, &Context, "Empty entry found in category %s", *Entry.CategoryEntry);
-			continue;
-		}
-
-		PCGMetadataValueKey ValueKey = Attribute->FindValue<FString>(Entry.CategoryEntry);
+			if (Entry.WeightedMeshEntries.Num() == 0)
+			{
+				PCGE_LOG_C(Verbose, &Context, "Empty entry found in category %s", *Entry.CategoryEntry);
+				continue;
+			}
 
-		if (ValueKey == PCGDefaultValueKey)
-		{
-			PCGE_LOG_C(Verbose, &Context, "Invalid category %s", *Entry.CategoryEntry);
-			continue;
-		}
+			PCGMetadataValueKey ValueKey = Attribute->FindValue<FString>(Entry.CategoryEntry);
 
-		FPCGInstancesAndWeights* InstancesAndWeights = CategoryEntryToInstancesAndWeights.Find(ValueKey);
+			if (ValueKey == PCGDefaultValueKey)
+			{
+				PCGE_LOG_C(Verbose, &Context, "Invalid category %s", *Entry.CategoryEntry);
+				continue;
+			}
 
-		if (InstancesAndWeights)
-		{
-			PCGE_LOG_C(Warning, &Context, "Duplicate entry found in category %s. Subsequent entries are ignored.", *Entry.CategoryEntry);
-			continue;
-		}
+			FPCGInstancesAndWeights* InstancesAndWeights = CategoryEntryToInstancesAndWeights.Find(ValueKey);
 
-		if (Entry.IsDefault)
-		{
-			if (DefaultValueKey == PCGDefaultValueKey)
+			if (InstancesAndWeights)
 			{
-				DefaultValueKey = ValueKey;
+				PCGE_LOG_C(Warning, &Context, "Duplicate entry found in category %s. Subsequent entries are ignored.", *Entry.CategoryEntry);
+				continue;
 			}
-			else
+
+			if (Entry.IsDefault)
 			{
-				PCGE_LOG_C(Warning, &Context, "Duplicate default entry found. Subsequent default entries are ignored.");
+				if (DefaultValueKey == PCGDefaultValueKey)
+				{
+					DefaultValueKey = ValueKey;
+				}
+				else
+				{
+					PCGE_LOG_C(Warning, &Context, "Duplicate default entry found. Subsequent default entries are ignored.");
+				}
 			}
-		}
 
-		InstancesAndWeights = &CategoryEntryToInstancesAndWeights.Add(ValueKey, FPCGInstancesAndWeights());
+			InstancesAndWeights = &CategoryEntryToInstancesAndWeights.Add(ValueKey, FPCGInstancesAndWeights());
 
-		int TotalWeight = 0;
-		for (const FPCGMeshSelectorWeightedEntry& WeightedEntry : Entry.WeightedMeshEntries)
-		{
-			if (WeightedEntry.Weight <= 0)
+			int TotalWeight = 0;
+			for (const FPCGMeshSelectorWeightedEntry& WeightedEntry : Entry.WeightedMeshEntries)
 			{
-				PCGE_LOG_C(Verbose, &Context, "Entry found with weight <= 0 in category %s", *Entry.CategoryEntry);
-				continue;
-			}
-			
-			TArray<FPCGMeshInstanceList>& PickEntry = InstancesAndWeights->MeshInstances.Emplace_GetRef();
-			PickEntry.Emplace_GetRef(WeightedEntry.Descriptor);
+				if (WeightedEntry.Weight <= 0)
+				{
+					PCGE_LOG_C(Verbose, &Context, "Entry found with weight <= 0 in category %s", *Entry.CategoryEntry);
+					continue;
+				}
 
-			// precompute the weights
-			TotalWeight += WeightedEntry.Weight;
-			InstancesAndWeights->CumulativeWeights.Add(TotalWeight);
-		}
+				TArray<FPCGMeshInstanceList>& PickEntry = InstancesAndWeights->MeshInstances.Emplace_GetRef();
+				PickEntry.Emplace_GetRef(WeightedEntry.Descriptor);
 
-		// if no weighted entries were collected, discard this InstancesAndWeights
-		if (InstancesAndWeights->CumulativeWeights.Num() == 0)
-		{
-			CategoryEntryToInstancesAndWeights.Remove(ValueKey);
+				// precompute the weights
+				TotalWeight += WeightedEntry.Weight;
+				InstancesAndWeights->CumulativeWeights.Add(TotalWeight);
+			}
+
+			// if no weighted entries were collected, discard this InstancesAndWeights
+			if (InstancesAndWeights->CumulativeWeights.Num() == 0)
+			{
+				CategoryEntryToInstancesAndWeights.Remove(ValueKey);
+			}
 		}
 	}
 
-	TArray<FPCGPoint>* OutPoints = nullptr;
-	FPCGMetadataAttribute<FString>* OutAttribute = nullptr;
-	TMap<TSoftObjectPtr<UStaticMesh>, PCGMetadataValueKey> MeshToValueKey;
-
-	FPCGMeshMaterialOverrideHelper MaterialOverrideHelper(Context, bUseAttributeMaterialOverrides, MaterialOverrideAttributes, InPointData->Metadata);
+	FPCGMeshMaterialOverrideHelper& MaterialOverrideHelper = Context.MaterialOverrideHelper;
+	if (!MaterialOverrideHelper.IsInitialized())
+	{
+		MaterialOverrideHelper.Initialize(Context, bUseAttributeMaterialOverrides, MaterialOverrideAttributes, InPointData->Metadata);
+	}
 
 	if (!MaterialOverrideHelper.IsValid())
 	{
-		return;
+		return true;
 	}
 
+	TArray<FPCGPoint>* OutPoints = nullptr;
+	FPCGMetadataAttribute<FString>* OutAttribute = nullptr;
+
 	if (OutPointData)
 	{
 		check(OutPointData->Metadata);
@@ -183,9 +185,17 @@ void UPCGMeshSelectorWeightedByCategory::SelectInstances_Implementation(
 
 	TRACE_CPUPROFILER_EVENT_SCOPE(FPCGStaticMeshSpawnerElement::Execute::SelectEntries);
 
+	int32 CurrentPointIndex = Context.CurrentPointIndex;
+	int32 LastCheckpointIndex = CurrentPointIndex;
+	constexpr int32 TimeSlicingCheckFrequency = 1024;
+	TMap<TSoftObjectPtr<UStaticMesh>, PCGMetadataValueKey>& MeshToValueKey = Context.MeshToValueKey;
+
 	// Assign points to entries
-	for (const FPCGPoint& Point : InPointData->GetPoints())
+	const TArray<FPCGPoint>& Points = InPointData->GetPoints();
+	while(CurrentPointIndex < Points.Num())
 	{
+		const FPCGPoint& Point = Points[CurrentPointIndex++];
+
 		if (Point.Density <= 0.0f)
 		{
 			continue;
@@ -244,17 +254,41 @@ void UPCGMeshSelectorWeightedByCategory::SelectInstances_Implementation(
 				OutAttribute->SetValueFromValueKey(OutPoint.MetadataEntry, *OutValueKey);
 			}
 		}
+
+		// Check if we should stop here and continue in a subsequent call
+		if (CurrentPointIndex - LastCheckpointIndex >= TimeSlicingCheckFrequency)
+		{
+			if (Context.ShouldStop())
+			{
+				break;
+			}
+			else
+			{
+				LastCheckpointIndex = CurrentPointIndex;
+			}
+		}
 	}
 
-	// Collapse to OutMeshInstances
-	for (TPair<PCGMetadataValueKey, FPCGInstancesAndWeights>& Entry : CategoryEntryToInstancesAndWeights)
+	Context.CurrentPointIndex = CurrentPointIndex;
+
+	if (CurrentPointIndex == Points.Num())
 	{
-		for (TArray<FPCGMeshInstanceList>& PickedMeshInstances : Entry.Value.MeshInstances)
+		// Collapse to OutMeshInstances
+		for (TPair<PCGMetadataValueKey, FPCGInstancesAndWeights>& Entry : CategoryEntryToInstancesAndWeights)
 		{
-			for (FPCGMeshInstanceList& PickedMeshInstanceEntry : PickedMeshInstances)
+			for (TArray<FPCGMeshInstanceList>& PickedMeshInstances : Entry.Value.MeshInstances)
 			{
-				OutMeshInstances.Emplace(MoveTemp(PickedMeshInstanceEntry));
+				for (FPCGMeshInstanceList& PickedMeshInstanceEntry : PickedMeshInstances)
+				{
+					OutMeshInstances.Emplace(MoveTemp(PickedMeshInstanceEntry));
+				}
 			}
 		}
+
+		return true;
+	}
+	else
+	{
+		return false;
 	}
 }
diff --git a/Engine/Plugins/Experimental/PCG/Source/PCG/Private/Tests/Elements/PCGInstancePackerTest.cpp b/Engine/Plugins/Experimental/PCG/Source/PCG/Private/Tests/Elements/PCGInstancePackerTest.cpp
--- a/Engine/Plugins/Experimental/PCG/Source/PCG/Private/Tests/Elements/PCGInstancePackerTest.cpp
+++ b/Engine/Plugins/Experimental/PCG/Source/PCG/Private/Tests/Elements/PCGInstancePackerTest.cpp
@@ -8,6 +8,7 @@
 #include "PCGManagedResource.h"
 #include "Data/PCGPointData.h"
 #include "Elements/PCGStaticMeshSpawner.h"
+#include "Elements/PCGStaticMeshSpawnerContext.h"
 #include "InstancePackers/PCGInstancePackerBase.h"
 #include "InstancePackers/PCGInstancePackerByAttribute.h"
 #include "InstancePackers/PCGInstancePackerByRegex.h"
diff --git a/Engine/Plugins/Experimental/PCG/Source/PCG/Public/Elements/PCGStaticMeshSpawner.h b/Engine/Plugins/Experimental/PCG/Source/PCG/Public/Elements/PCGStaticMeshSpawner.h
--- a/Engine/Plugins/Experimental/PCG/Source/PCG/Public/Elements/PCGStaticMeshSpawner.h
+++ b/Engine/Plugins/Experimental/PCG/Source/PCG/Public/Elements/PCGStaticMeshSpawner.h
@@ -2,7 +2,6 @@
 
 #pragma once
 
-#include "PCGContext.h"
 #include "PCGSettings.h"
 
 #include "Engine/CollisionProfile.h"
@@ -13,6 +12,7 @@
 class UPCGInstancePackerBase;
 class UPCGMeshSelectorBase;
 class UPCGSpatialData;
+struct FPCGContext;
 struct FPCGMeshInstanceList;
 struct FPCGPackedCustomData;
 
@@ -98,28 +98,6 @@ protected:
 	void RefreshInstancePacker();
 };
 
-USTRUCT()
-struct FPCGStaticMeshSpawnerContext : public FPCGContext
-{
-	GENERATED_BODY()
-
-	struct FPackedInstanceListData
-	{
-		FPackedInstanceListData();
-		~FPackedInstanceListData();
-		AActor* TargetActor = nullptr;
-		const UPCGSpatialData* SpatialData = nullptr;
-		TArray<FPCGMeshInstanceList> MeshInstances;
-		TArray<FPCGPackedCustomData> PackedCustomData;
-	};
-
-	TArray<FPackedInstanceListData> MeshInstancesData;
-	// Index of input in the context (for selection)
-	int32 CurrentInputIndex = 0;
-	// Index of MeshInstances/PackedCustomData in last MeshInstancesData element
-	int32 CurrentDataIndex = 0;
-};
-
 class FPCGStaticMeshSpawnerElement : public IPCGElement
 {
 public:
diff --git a/Engine/Plugins/Experimental/PCG/Source/PCG/Public/Elements/PCGStaticMeshSpawnerContext.h b/Engine/Plugins/Experimental/PCG/Source/PCG/Public/Elements/PCGStaticMeshSpawnerContext.h
new file mode 100644
--- /dev/null
+++ b/Engine/Plugins/Experimental/PCG/Source/PCG/Public/Elements/PCGStaticMeshSpawnerContext.h
@@ -0,0 +1,65 @@
+// Copyright Epic Games, Inc. All Rights Reserved.
+
+#pragma once
+
+#include "PCGContext.h"
+#include "MeshSelectors/PCGMeshSelectorBase.h"
+#include "InstancePackers/PCGInstancePackerBase.h"
+
+#include "PCGStaticMeshSpawnerContext.generated.h"
+
+class AActor;
+class UPCGSpatialData;
+
+struct FPCGInstancesAndWeights
+{
+	TArray<TArray<FPCGMeshInstanceList>> MeshInstances;
+	TArray<int> CumulativeWeights;
+};
+
+USTRUCT(BlueprintType)
+struct FPCGStaticMeshSpawnerContext : public FPCGContext
+{
+	GENERATED_BODY()
+
+	struct FPackedInstanceListData
+	{
+		FPackedInstanceListData();
+		~FPackedInstanceListData();
+		AActor* TargetActor = nullptr;
+		const UPCGSpatialData* SpatialData = nullptr;
+		TArray<FPCGMeshInstanceList> MeshInstances;
+		TArray<FPCGPackedCustomData> PackedCustomData;
+	};
+
+	TArray<FPackedInstanceListData> MeshInstancesData;
+	// Index of input in the context (for selection)
+	int32 CurrentInputIndex = 0;
+	// Index of MeshInstances/PackedCustomData in last MeshInstancesData element
+	int32 CurrentDataIndex = 0;
+
+	// Whole-execution variables
+	bool bSkippedDueToReuse = false;
+
+	// Per-input context variables
+	bool bCurrentInputSetup = false;
+	bool bSelectionDone = false;
+
+	const UPCGPointData* CurrentPointData = nullptr;
+	UPCGPointData* CurrentOutputPointData = nullptr;
+	FPCGMeshMaterialOverrideHelper MaterialOverrideHelper;
+	int32 CurrentPointIndex = 0;
+
+	// Used in by-attribute selector
+	TMap<PCGMetadataValueKey, TSoftObjectPtr<UStaticMesh>> ValueKeyToMesh;
+
+	// Used in weighted selector
+	TArray<TArray<FPCGMeshInstanceList>> WeightedMeshInstances;
+	TMap<TSoftObjectPtr<UStaticMesh>, PCGMetadataValueKey> MeshToValueKey;
+	TArray<int> CumulativeWeights;
+
+	// Used in the weighted by category selector
+	TMap<PCGMetadataValueKey, FPCGInstancesAndWeights> CategoryEntryToInstancesAndWeights;
+
+	void ResetInputIterationData();
+};
\ No newline at end of file
diff --git a/Engine/Plugins/Experimental/PCG/Source/PCG/Public/MeshSelectors/PCGMeshSelectorBase.h b/Engine/Plugins/Experimental/PCG/Source/PCG/Public/MeshSelectors/PCGMeshSelectorBase.h
--- a/Engine/Plugins/Experimental/PCG/Source/PCG/Public/MeshSelectors/PCGMeshSelectorBase.h
+++ b/Engine/Plugins/Experimental/PCG/Source/PCG/Public/MeshSelectors/PCGMeshSelectorBase.h
@@ -14,6 +14,7 @@ class UPCGPointData;
 class UPCGSpatialData;
 class UStaticMesh;
 struct FPCGContext;
+struct FPCGStaticMeshSpawnerContext;
 
 class UPCGStaticMeshSpawnerSettings;
 class UMaterialInterface;
@@ -50,8 +51,10 @@ enum class EPCGMeshSelectorMaterialOverrideMode : uint8
 /** Struct used to efficiently gather overrides and cache them during instance packing */
 struct FPCGMeshMaterialOverrideHelper
 {
+	FPCGMeshMaterialOverrideHelper() = default;
+
 	// Use this constructor when you have a 1:1 mapping between attributes or static overrides
-	FPCGMeshMaterialOverrideHelper(
+	void Initialize(
 		FPCGContext& InContext,
 		bool bUseMaterialOverrideAttributes,
 		const TArray<TSoftObjectPtr<UMaterialInterface>>& InStaticMaterialOverrides,
@@ -59,54 +62,49 @@ struct FPCGMeshMaterialOverrideHelper
 		const UPCGMetadata* InMetadata);
 
 	// Use this constructor when you have common attribute usage or separate static overrides
-	FPCGMeshMaterialOverrideHelper(
+	void Initialize(
 		FPCGContext& InContext,
 		bool bInByAttributeOverride,
 		const TArray<FName>& InMaterialOverrideAttributeNames,
 		const UPCGMetadata* InMetadata);
 
+	void Reset();
+
+	bool IsInitialized() const { return bIsInitialized; }
 	bool IsValid() const { return bIsValid; }
 	bool OverridesMaterials() const { return bUseMaterialOverrideAttributes; }
 	const TArray<TSoftObjectPtr<UMaterialInterface>>& GetMaterialOverrides(PCGMetadataEntryKey EntryKey);
 
+private:
 	// Cached data
 	TArray<const FPCGMetadataAttribute<FString>*> MaterialAttributes;
 	TArray<TMap<PCGMetadataValueKey, TSoftObjectPtr<UMaterialInterface>>> ValueKeyToOverrideMaterials;
 	TArray<TSoftObjectPtr<UMaterialInterface>> WorkingMaterialOverrides;
-	TArray<TSoftObjectPtr<UMaterialInterface>> EmptyArray;
 
 	// Data needed to perform operations
+	bool bIsInitialized = false;
 	bool bIsValid = false;
 	bool bUseMaterialOverrideAttributes = false;
 
-	const TArray<TSoftObjectPtr<UMaterialInterface>>& StaticMaterialOverrides;
-	const TArray<FName>& MaterialOverrideAttributeNames;
+	TArray<TSoftObjectPtr<UMaterialInterface>> StaticMaterialOverrides;
+	TArray<FName> MaterialOverrideAttributeNames;
 	const UPCGMetadata* Metadata = nullptr;
 
-private:
 	void Initialize(FPCGContext& InContext);
 };
 
-UCLASS(Abstract, BlueprintType, Blueprintable, ClassGroup = (Procedural))
+UCLASS(Abstract, BlueprintType, ClassGroup = (Procedural))
 class PCG_API UPCGMeshSelectorBase : public UObject 
 {
 	GENERATED_BODY()
 
 public:
-	UFUNCTION(BlueprintNativeEvent, Category = MeshSelection)
-	void SelectInstances(
-		FPCGContext& Context,
-		const UPCGStaticMeshSpawnerSettings* Settings,
-		const UPCGPointData* InPointData,
-		TArray<FPCGMeshInstanceList>& OutMeshInstances,
-		UPCGPointData* OutPointData) const;
-
-	virtual void SelectInstances_Implementation(
-		FPCGContext& Context,
+	virtual bool SelectInstances(
+		FPCGStaticMeshSpawnerContext& Context,
 		const UPCGStaticMeshSpawnerSettings* Settings,
 		const UPCGPointData* InPointData,
 		TArray<FPCGMeshInstanceList>& OutMeshInstances,
-		UPCGPointData* OutPointData) const PURE_VIRTUAL(UPCGMeshSelectorBase::SelectInstances_Implementation);
+		UPCGPointData* OutPointData) const PURE_VIRTUAL(UPCGMeshSelectorBase::SelectInstances, return true;);
 };
 
 #if UE_ENABLE_INCLUDE_ORDER_DEPRECATED_IN_5_2
diff --git a/Engine/Plugins/Experimental/PCG/Source/PCG/Public/MeshSelectors/PCGMeshSelectorByAttribute.h b/Engine/Plugins/Experimental/PCG/Source/PCG/Public/MeshSelectors/PCGMeshSelectorByAttribute.h
--- a/Engine/Plugins/Experimental/PCG/Source/PCG/Public/MeshSelectors/PCGMeshSelectorByAttribute.h
+++ b/Engine/Plugins/Experimental/PCG/Source/PCG/Public/MeshSelectors/PCGMeshSelectorByAttribute.h
@@ -4,7 +4,6 @@
 
 #include "PCGMeshSelectorBase.h"
 
-
 #include "PCGMeshSelectorByAttribute.generated.h"
 
 UCLASS(BlueprintType, ClassGroup = (Procedural))
@@ -17,8 +16,8 @@ public:
 	void PostLoad() override;
 	// ~End UObject interface
 
-	virtual void SelectInstances_Implementation(
-		UPARAM(ref) FPCGContext& Context,
+	virtual bool SelectInstances(
+		FPCGStaticMeshSpawnerContext& Context,
 		const UPCGStaticMeshSpawnerSettings* Settings,
 		const UPCGPointData* InPointData,
 		TArray<FPCGMeshInstanceList>& OutMeshInstances,
diff --git a/Engine/Plugins/Experimental/PCG/Source/PCG/Public/MeshSelectors/PCGMeshSelectorWeighted.h b/Engine/Plugins/Experimental/PCG/Source/PCG/Public/MeshSelectors/PCGMeshSelectorWeighted.h
--- a/Engine/Plugins/Experimental/PCG/Source/PCG/Public/MeshSelectors/PCGMeshSelectorWeighted.h
+++ b/Engine/Plugins/Experimental/PCG/Source/PCG/Public/MeshSelectors/PCGMeshSelectorWeighted.h
@@ -69,8 +69,8 @@ class PCG_API UPCGMeshSelectorWeighted : public UPCGMeshSelectorBase
 	GENERATED_BODY()
 
 public:
-	virtual void SelectInstances_Implementation(
-		UPARAM(ref) FPCGContext& Context, 
+	virtual bool SelectInstances(
+		FPCGStaticMeshSpawnerContext& Context,
 		const UPCGStaticMeshSpawnerSettings* Settings, 
 		const UPCGPointData* InPointData,
 		TArray<FPCGMeshInstanceList>& OutMeshInstances,
diff --git a/Engine/Plugins/Experimental/PCG/Source/PCG/Public/MeshSelectors/PCGMeshSelectorWeightedByCategory.h b/Engine/Plugins/Experimental/PCG/Source/PCG/Public/MeshSelectors/PCGMeshSelectorWeightedByCategory.h
--- a/Engine/Plugins/Experimental/PCG/Source/PCG/Public/MeshSelectors/PCGMeshSelectorWeightedByCategory.h
+++ b/Engine/Plugins/Experimental/PCG/Source/PCG/Public/MeshSelectors/PCGMeshSelectorWeightedByCategory.h
@@ -38,8 +38,8 @@ class PCG_API UPCGMeshSelectorWeightedByCategory : public UPCGMeshSelectorBase
 	GENERATED_BODY()
 
 public:
-	virtual void SelectInstances_Implementation(
-		UPARAM(ref) FPCGContext& Context,
+	virtual bool SelectInstances(
+		FPCGStaticMeshSpawnerContext& Context,
 		const UPCGStaticMeshSpawnerSettings* Settings,
 		const UPCGPointData* InPointData,
 		TArray<FPCGMeshInstanceList>& OutMeshInstances,
