Added basic collision shape data implementation for better efficiency

diff --git a/Engine/Plugins/Experimental/PCG/Source/PCG/PCG.Build.cs b/Engine/Plugins/Experimental/PCG/Source/PCG/PCG.Build.cs
--- a/Engine/Plugins/Experimental/PCG/Source/PCG/PCG.Build.cs
+++ b/Engine/Plugins/Experimental/PCG/Source/PCG/PCG.Build.cs
@@ -10,11 +10,13 @@ namespace UnrealBuildTool.Rules
 
 			PublicDependencyModuleNames.AddRange(
 				new string[] {
+					"Chaos",
 					"Core",
 					"CoreUObject",
 					"Engine",
 					"Landscape",
 					"Foliage",
+					"PhysicsCore",
 					"Projects",
 					"RenderCore",
 					"RHI",
diff --git a/Engine/Plugins/Experimental/PCG/Source/PCG/Private/Data/PCGCollisionShapeData.cpp b/Engine/Plugins/Experimental/PCG/Source/PCG/Private/Data/PCGCollisionShapeData.cpp
new file mode 100644
--- /dev/null
+++ b/Engine/Plugins/Experimental/PCG/Source/PCG/Private/Data/PCGCollisionShapeData.cpp
@@ -0,0 +1,91 @@
+// Copyright Epic Games, Inc. All Rights Reserved.
+
+#include "Data/PCGCollisionShapeData.h"
+
+#include "Data/PCGPointData.h"
+#include "Data/PCGSpatialData.h"
+#include "Elements/PCGVolumeSampler.h"
+
+#include "Components/BoxComponent.h"
+#include "Components/CapsuleComponent.h"
+#include "Components/ShapeComponent.h"
+#include "Components/SphereComponent.h"
+
+#include "Chaos/GeometryQueries.h"
+
+void UPCGCollisionShapeData::Initialize(UShapeComponent* InComponent)
+{
+	check(InComponent && IsSupported(InComponent));
+	Shape = InComponent->GetCollisionShape();
+	Transform = InComponent->GetComponentTransform();
+
+	//Note: Shape is pre-scaled
+	ShapeAdapter = MakeUnique<FPhysicsShapeAdapter>(Transform.GetRotation(), Shape);
+	TargetActor = InComponent->GetOwner();
+
+	CachedBounds = InComponent->Bounds.GetBox();
+	CachedStrictBounds = CachedBounds;
+}
+
+bool UPCGCollisionShapeData::IsSupported(UShapeComponent* InComponent)
+{
+	if (!InComponent)
+	{
+		return false;
+	}
+
+	if(InComponent->IsA<UBoxComponent>() || InComponent->IsA<UCapsuleComponent>() || InComponent->IsA<USphereComponent>())
+	{
+		return true;
+	}
+
+	return false;
+}
+
+bool UPCGCollisionShapeData::SamplePoint(const FTransform& InTransform, const FBox& InBounds, FPCGPoint& OutPoint, UPCGMetadata* OutMetadata) const
+{
+	FCollisionShape CollisionShape;
+	CollisionShape.SetBox(FVector3f(InBounds.GetExtent() * InTransform.GetScale3D())); // make sure to prescale
+	FPhysicsShapeAdapter PointAdapter(InTransform.GetRotation(), CollisionShape);
+
+	if (Chaos::Utilities::CastHelper(PointAdapter.GetGeometry(), PointAdapter.GetGeomPose(InTransform.GetTranslation()), [this](const auto& Downcast, const auto& FullGeomTransform) { return Chaos::OverlapQuery(ShapeAdapter->GetGeometry(), ShapeAdapter->GetGeomPose(Transform.GetTranslation()), Downcast, FullGeomTransform, /*Thickness=*/0); }))
+	{
+		return true;
+	}
+	else
+	{
+		return false;
+	}
+}
+
+const UPCGPointData* UPCGCollisionShapeData::CreatePointData(FPCGContext* Context) const
+{
+	TRACE_CPUPROFILER_EVENT_SCOPE(UPCGShapeData::CreatePointData);
+
+	const FVector DefaultVoxelSize = FVector(100.0, 100.0, 100.0);
+
+	PCGVolumeSampler::FVolumeSamplerSettings SamplerSettings;
+	SamplerSettings.VoxelSize = DefaultVoxelSize;
+
+	const UPCGPointData* Data = PCGVolumeSampler::SampleVolume(Context, this, SamplerSettings);
+
+	if (Data)
+	{
+		UE_LOG(LogPCG, Verbose, TEXT("Shape extracted %d points"), Data->GetPoints().Num());
+	}
+
+	return Data;
+}
+
+UPCGSpatialData* UPCGCollisionShapeData::CopyInternal() const
+{
+	UPCGCollisionShapeData* NewShapeData = NewObject<UPCGCollisionShapeData>();
+
+	NewShapeData->Transform = Transform;
+	NewShapeData->Shape = Shape;
+	NewShapeData->ShapeAdapter = MakeUnique<FPhysicsShapeAdapter>(NewShapeData->Transform.GetRotation(), NewShapeData->Shape);
+	NewShapeData->CachedBounds = CachedBounds;
+	NewShapeData->CachedStrictBounds = CachedStrictBounds;
+
+	return NewShapeData;
+}
diff --git a/Engine/Plugins/Experimental/PCG/Source/PCG/Private/PCGComponent.cpp b/Engine/Plugins/Experimental/PCG/Source/PCG/Private/PCGComponent.cpp
--- a/Engine/Plugins/Experimental/PCG/Source/PCG/Private/PCGComponent.cpp
+++ b/Engine/Plugins/Experimental/PCG/Source/PCG/Private/PCGComponent.cpp
@@ -10,6 +10,7 @@
 #include "PCGManagedResource.h"
 #include "PCGParamData.h"
 #include "PCGSubsystem.h"
+#include "Data/PCGCollisionShapeData.h"
 #include "Data/PCGDifferenceData.h"
 #include "Data/PCGIntersectionData.h"
 #include "Data/PCGLandscapeData.h"
@@ -2254,11 +2255,24 @@ FPCGDataCollection UPCGComponent::CreateActorPCGDataCollection(AActor* Actor, co
 		{
 			for (UShapeComponent* ShapeComponent : Shapes)
 			{
-				UPCGPrimitiveData* ShapeData = NewObject<UPCGPrimitiveData>();
-				ShapeData->Initialize(ShapeComponent);
+				UPCGSpatialData* Data = nullptr;
+				if (UPCGCollisionShapeData::IsSupported(ShapeComponent))
+				{
+					UPCGCollisionShapeData* ShapeData = NewObject<UPCGCollisionShapeData>();
+					ShapeData->Initialize(ShapeComponent);
+
+					Data = ShapeData;
+				}
+				else
+				{
+					UPCGPrimitiveData* ShapeData = NewObject<UPCGPrimitiveData>();
+					ShapeData->Initialize(ShapeComponent);
+
+					Data = ShapeData;
+				}
 
 				FPCGTaggedData& TaggedData = Collection.TaggedData.Emplace_GetRef();
-				TaggedData.Data = ShapeData;
+				TaggedData.Data = Data;
 				Algo::Transform(ShapeComponent->ComponentTags, TaggedData.Tags, NameTagsToStringTags);
 				TaggedData.Tags.Append(ActorTags);
 			}
diff --git a/Engine/Plugins/Experimental/PCG/Source/PCG/Public/Data/PCGCollisionShapeData.h b/Engine/Plugins/Experimental/PCG/Source/PCG/Public/Data/PCGCollisionShapeData.h
new file mode 100644
--- /dev/null
+++ b/Engine/Plugins/Experimental/PCG/Source/PCG/Public/Data/PCGCollisionShapeData.h
@@ -0,0 +1,58 @@
+// Copyright Epic Games, Inc. All Rights Reserved.
+
+#pragma once
+
+#include "PCGSpatialData.h"
+
+#include "CollisionShape.h"
+#include "Physics/PhysicsInterfaceDeclares.h"
+
+#include "PCGCollisionShapeData.generated.h"
+
+class UShapeComponent;
+
+UCLASS(BlueprintType, ClassGroup = (Procedural))
+class PCG_API UPCGCollisionShapeData : public UPCGSpatialDataWithPointCache
+{
+	GENERATED_BODY()
+
+public:
+	void Initialize(UShapeComponent* InShape);
+	static bool IsSupported(UShapeComponent* InShape);
+
+	// ~Begin UPCGData interface
+	virtual EPCGDataType GetDataType() const override { return EPCGDataType::Primitive; }
+	// ~End UPCGData interface
+
+	// ~Begin UPCGSpatialData interface
+	virtual int GetDimension() const override { return 3; }
+	virtual FBox GetBounds() const override { return CachedBounds; }
+	virtual FBox GetStrictBounds() const override { return CachedStrictBounds; }
+	virtual bool SamplePoint(const FTransform& Transform, const FBox& Bounds, FPCGPoint& OutPoint, UPCGMetadata* OutMetadata) const override;
+	// TODO needs an implementation to support projection
+	//virtual bool ProjectPoint(const FTransform& InTransform, const FBox& InBounds, const FPCGProjectionParams& InParams, FPCGPoint& OutPoint, UPCGMetadata* OutMetadata) const;
+protected:
+	virtual UPCGSpatialData* CopyInternal() const override;
+	//~End UPCGSpatialData interface
+
+public:
+	// ~Begin UPCGSpatialDataWithPointCache implementation
+	virtual const UPCGPointData* CreatePointData(FPCGContext* Context) const override;
+	// ~End UPCGSpatialDataWithPointCache implementation
+
+protected:
+	UPROPERTY(BlueprintReadOnly, VisibleAnywhere, Category = SpatialData)
+	FTransform Transform;
+
+	// If we want to serialize this data, this will need to be broken up, and the shape adapter will need to be recreated as well
+	FCollisionShape Shape;
+
+	// Built from shape & transform, not serialized
+	TUniquePtr<FPhysicsShapeAdapter> ShapeAdapter;
+
+	UPROPERTY()
+	FBox CachedBounds = FBox(EForceInit::ForceInit);
+
+	UPROPERTY()
+	FBox CachedStrictBounds = FBox(EForceInit::ForceInit);
+};
\ No newline at end of file
