[PCG] Added functionality for propagating Crcs through boolean nodes (union etc), disabled for now.

This will pave the way for more aggressive cache-based optimisations to graph execution.

diff --git a/Engine/Plugins/Experimental/PCG/Source/PCG/Private/Data/PCGDifferenceData.cpp b/Engine/Plugins/Experimental/PCG/Source/PCG/Private/Data/PCGDifferenceData.cpp
--- a/Engine/Plugins/Experimental/PCG/Source/PCG/Private/Data/PCGDifferenceData.cpp
+++ b/Engine/Plugins/Experimental/PCG/Source/PCG/Private/Data/PCGDifferenceData.cpp
@@ -6,6 +6,8 @@
 #include "Data/PCGUnionData.h"
 #include "Helpers/PCGAsync.h"
 
+#include "Serialization/ArchiveCrc32.h"
+
 #include UE_INLINE_GENERATED_CPP_BY_NAME(PCGDifferenceData)
 
 namespace PCGDifferenceDataUtils
@@ -122,6 +124,42 @@ void UPCGDifferenceData::VisitDataNetwork(TFunctionRef<void(const UPCGData*)> Ac
 	}
 }
 
+FPCGCrc UPCGDifferenceData::ComputeCrc() const
+{
+	FArchiveCrc32 Ar;
+
+	if (PropagateCrcThroughBooleanData())
+	{
+		AddToCrc(Ar);
+
+		// Chain together CRCs of operands
+		check(GetSource());
+		uint32 SourceCrc = GetSource()->GetOrComputeCrc().GetValue();
+		Ar << SourceCrc;
+
+		if (GetDifference())
+		{
+			uint32 DifferenceCrc = GetDifference()->GetOrComputeCrc().GetValue();
+			Ar << DifferenceCrc;
+		}
+	}
+	else
+	{
+		UPCGData::AddToCrc(Ar);
+	}
+
+	return FPCGCrc(Ar.GetCrc());
+}
+
+void UPCGDifferenceData::AddToCrc(FArchiveCrc32& Ar) const
+{
+	uint32 UniqueTypeID = StaticClass()->GetDefaultObject()->GetUniqueID();
+	Ar << UniqueTypeID;
+
+	uint32 DensityFunctionValue = static_cast<uint32>(DensityFunction);
+	Ar << DensityFunctionValue;
+}
+
 int UPCGDifferenceData::GetDimension() const
 {
 	return GetSource()->GetDimension();
diff --git a/Engine/Plugins/Experimental/PCG/Source/PCG/Private/Data/PCGIntersectionData.cpp b/Engine/Plugins/Experimental/PCG/Source/PCG/Private/Data/PCGIntersectionData.cpp
--- a/Engine/Plugins/Experimental/PCG/Source/PCG/Private/Data/PCGIntersectionData.cpp
+++ b/Engine/Plugins/Experimental/PCG/Source/PCG/Private/Data/PCGIntersectionData.cpp
@@ -6,6 +6,8 @@
 #include "Helpers/PCGAsync.h"
 #include "PCGHelpers.h"
 
+#include "Serialization/ArchiveCrc32.h"
+
 #include UE_INLINE_GENERATED_CPP_BY_NAME(PCGIntersectionData)
 
 namespace PCGIntersectionDataMaths
@@ -69,6 +71,38 @@ void UPCGIntersectionData::VisitDataNetwork(TFunctionRef<void(const UPCGData*)>
 	GetB()->VisitDataNetwork(Action);
 }
 
+FPCGCrc UPCGIntersectionData::ComputeCrc() const
+{
+	FArchiveCrc32 Ar;
+
+	if (PropagateCrcThroughBooleanData())
+	{
+		AddToCrc(Ar);
+
+		// Chain together CRCs of operands
+		check(GetA() && GetB());
+		uint32 CrcA = GetA()->GetOrComputeCrc().GetValue();
+		uint32 CrcB = GetB()->GetOrComputeCrc().GetValue();
+		Ar << CrcA;
+		Ar << CrcB;
+	}
+	else
+	{
+		UPCGData::AddToCrc(Ar);
+	}
+
+	return FPCGCrc(Ar.GetCrc());
+}
+
+void UPCGIntersectionData::AddToCrc(FArchiveCrc32& Ar) const
+{
+	uint32 UniqueTypeID = StaticClass()->GetDefaultObject()->GetUniqueID();
+	Ar << UniqueTypeID;
+
+	uint32 DensityFunctionValue = static_cast<uint32>(DensityFunction);
+	Ar << DensityFunctionValue;
+}
+
 int UPCGIntersectionData::GetDimension() const
 {
 	check(GetA() && GetB());
diff --git a/Engine/Plugins/Experimental/PCG/Source/PCG/Private/Data/PCGUnionData.cpp b/Engine/Plugins/Experimental/PCG/Source/PCG/Private/Data/PCGUnionData.cpp
--- a/Engine/Plugins/Experimental/PCG/Source/PCG/Private/Data/PCGUnionData.cpp
+++ b/Engine/Plugins/Experimental/PCG/Source/PCG/Private/Data/PCGUnionData.cpp
@@ -7,6 +7,8 @@
 #include "Metadata/PCGMetadataAccessor.h"
 #include "PCGHelpers.h"
 
+#include "Serialization/ArchiveCrc32.h"
+
 #include UE_INLINE_GENERATED_CPP_BY_NAME(PCGUnionData)
 
 namespace PCGUnionDataMaths
@@ -81,6 +83,47 @@ void UPCGUnionData::VisitDataNetwork(TFunctionRef<void(const UPCGData*)> Action)
 	}
 }
 
+FPCGCrc UPCGUnionData::ComputeCrc() const
+{
+	FArchiveCrc32 Ar;
+
+	if (PropagateCrcThroughBooleanData())
+	{
+		AddToCrc(Ar);
+
+		// Chain together CRCs of operands
+		int32 NumOperands = Data.Num();
+		Ar << NumOperands;
+
+		for (TObjectPtr<const UPCGSpatialData> Datum : Data)
+		{
+			if (Datum)
+			{
+				uint32 DatumCrc = Datum->GetOrComputeCrc().GetValue();
+				Ar << DatumCrc;
+			}
+		}
+	}
+	else
+	{
+		UPCGData::AddToCrc(Ar);
+	}
+
+	return FPCGCrc(Ar.GetCrc());
+}
+
+void UPCGUnionData::AddToCrc(FArchiveCrc32& Ar) const
+{
+	uint32 UniqueTypeID = StaticClass()->GetDefaultObject()->GetUniqueID();
+	Ar << UniqueTypeID;
+
+	uint32 UnionTypeValue = static_cast<uint32>(UnionType);
+	Ar << UnionTypeValue;
+
+	uint32 DensityFunctionValue = static_cast<uint32>(DensityFunction);
+	Ar << DensityFunctionValue;
+}
+
 int UPCGUnionData::GetDimension() const
 {
 	return CachedDimension;
diff --git a/Engine/Plugins/Experimental/PCG/Source/PCG/Private/PCGData.cpp b/Engine/Plugins/Experimental/PCG/Source/PCG/Private/PCGData.cpp
--- a/Engine/Plugins/Experimental/PCG/Source/PCG/Private/PCGData.cpp
+++ b/Engine/Plugins/Experimental/PCG/Source/PCG/Private/PCGData.cpp
@@ -9,11 +9,17 @@
 #include "Data/PCGPointData.h"
 #include "Data/PCGSpatialData.h"
 
+#include "HAL/IConsoleManager.h"
 #include "Serialization/ArchiveCrc32.h"
 #include "UObject/Package.h"
 
 #include UE_INLINE_GENERATED_CPP_BY_NAME(PCGData)
 
+static TAutoConsoleVariable<bool> CVarCacheDebugging(
+	TEXT("pcg.Cache.PropagateCrcThroughBooleanData"),
+	false,
+	TEXT("Whether intersection, union, difference combine Crc values from operands. If false they fall back to using data UID."));
+
 void FPCGRootSet::Clear()
 {
 	for (TPair<UObject*, int32>& Entry : RootSet)
@@ -89,12 +95,17 @@ FPCGCrc UPCGData::GetOrComputeCrc() const
 		return Crc;
 	}
 
+	Crc = ComputeCrc();
+
+	return Crc;
+}
+
+FPCGCrc UPCGData::ComputeCrc() const
+{
 	FArchiveCrc32 Ar;
 	AddToCrc(Ar);
 
-	Crc = FPCGCrc(Ar.GetCrc());
-
-	return Crc;
+	return FPCGCrc(Ar.GetCrc());
 }
 
 void UPCGData::AddToCrc(FArchiveCrc32& Ar) const
@@ -104,6 +115,11 @@ void UPCGData::AddToCrc(FArchiveCrc32& Ar) const
 	Ar << UIDValue;
 }
 
+bool UPCGData::PropagateCrcThroughBooleanData() const
+{
+	return CVarCacheDebugging.GetValueOnAnyThread();
+}
+
 void UPCGData::VisitDataNetwork(TFunctionRef<void(const UPCGData*)> Action) const
 {
 	Action(this);
diff --git a/Engine/Plugins/Experimental/PCG/Source/PCG/Public/Data/PCGDifferenceData.h b/Engine/Plugins/Experimental/PCG/Source/PCG/Public/Data/PCGDifferenceData.h
--- a/Engine/Plugins/Experimental/PCG/Source/PCG/Public/Data/PCGDifferenceData.h
+++ b/Engine/Plugins/Experimental/PCG/Source/PCG/Public/Data/PCGDifferenceData.h
@@ -56,8 +56,13 @@ public:
 	// ~Begin UPCGData interface
 	virtual EPCGDataType GetDataType() const override { return EPCGDataType::Spatial; }
 	virtual void VisitDataNetwork(TFunctionRef<void(const UPCGData*)> Action) const override;
+
+protected:
+	virtual FPCGCrc ComputeCrc() const override;
+	virtual void AddToCrc(FArchiveCrc32& Ar) const override;
 	// ~End UPCGData interface
 
+public:
 	//~Begin UPCGSpatialData interface
 	virtual int GetDimension() const override;
 	virtual FBox GetBounds() const override;
diff --git a/Engine/Plugins/Experimental/PCG/Source/PCG/Public/Data/PCGIntersectionData.h b/Engine/Plugins/Experimental/PCG/Source/PCG/Public/Data/PCGIntersectionData.h
--- a/Engine/Plugins/Experimental/PCG/Source/PCG/Public/Data/PCGIntersectionData.h
+++ b/Engine/Plugins/Experimental/PCG/Source/PCG/Public/Data/PCGIntersectionData.h
@@ -33,8 +33,13 @@ public:
 	// ~Begin UPCGData interface
 	virtual EPCGDataType GetDataType() const override { return EPCGDataType::Spatial; }
 	virtual void VisitDataNetwork(TFunctionRef<void(const UPCGData*)> Action) const override;
+
+protected:
+	virtual FPCGCrc ComputeCrc() const override;
+	virtual void AddToCrc(FArchiveCrc32& Ar) const override;
 	// ~End UPCGData interface
 
+public:
 	//~Begin UPCGSpatialData interface
 	virtual int GetDimension() const override;
 	virtual FBox GetBounds() const override;
diff --git a/Engine/Plugins/Experimental/PCG/Source/PCG/Public/Data/PCGUnionData.h b/Engine/Plugins/Experimental/PCG/Source/PCG/Public/Data/PCGUnionData.h
--- a/Engine/Plugins/Experimental/PCG/Source/PCG/Public/Data/PCGUnionData.h
+++ b/Engine/Plugins/Experimental/PCG/Source/PCG/Public/Data/PCGUnionData.h
@@ -39,8 +39,13 @@ public:
 	// ~Begin UPCGData interface
 	virtual EPCGDataType GetDataType() const override { return EPCGDataType::Spatial; }
 	virtual void VisitDataNetwork(TFunctionRef<void(const UPCGData*)> Action) const override;
+
+protected:
+	virtual FPCGCrc ComputeCrc() const override;
+	virtual void AddToCrc(FArchiveCrc32& Ar) const override;
 	// ~End UPCGData interface
 
+public:
 	//~Begin UPCGSpatialData interface
 	virtual int GetDimension() const override;
 	virtual FBox GetBounds() const override;
diff --git a/Engine/Plugins/Experimental/PCG/Source/PCG/Public/PCGData.h b/Engine/Plugins/Experimental/PCG/Source/PCG/Public/PCGData.h
--- a/Engine/Plugins/Experimental/PCG/Source/PCG/Public/PCGData.h
+++ b/Engine/Plugins/Experimental/PCG/Source/PCG/Public/PCGData.h
@@ -29,12 +29,9 @@ public:
 	UPCGData(const FObjectInitializer& ObjectInitializer);
 	virtual EPCGDataType GetDataType() const { return EPCGDataType::None; }
 
-	/** Computes a Crc for this data. */
+	/** Returns a Crc for this and any connected data. */
 	FPCGCrc GetOrComputeCrc() const;
 
-	/** Contributes to Crc. Fallback implementation that writes object instance UID. */
-	virtual void AddToCrc(FArchiveCrc32& Ar) const;
-
 	/** Executes a lambda over all connected data objects. */
 	virtual void VisitDataNetwork(TFunctionRef<void(const UPCGData*)> Action) const;
 
@@ -45,6 +42,16 @@ public:
 	/** CRC for this object instance. */
 	mutable FPCGCrc Crc;
 
+protected:
+	/** Computes Crc for this and any connected data. */
+	virtual FPCGCrc ComputeCrc() const;
+
+	/** Adds this data to Crc. Fallback implementation writes object instance UID. */
+	virtual void AddToCrc(FArchiveCrc32& Ar) const;
+
+	/** Whether intersection, union, difference combine Crc values from operands. If false they fall back to using data UID. */
+	bool PropagateCrcThroughBooleanData() const;
+
 private:
 	/** Serves unique ID values to instances of this object. */
 	static inline std::atomic<uint64> UIDCounter{ 1 };
