[PCG] Add Crc for params, and enable point Crc but use it sparingly.

Due to limitations with crc'ing external data, we tend to get fresh data objects created too frequently in many cases. Surface sampler must normally assume landscape has changed even if it has not. Param data gets created everytime for PropertyToParam node (by default).

This allows some elements to opt into a finegrained data Crc which only changes when the output data changes, which can be used to halt unnecessary change propagation / halt unnecessary element execution.

Does a finegrained crc after Collapse nodes. Many cases of collapse happen without an explicit collapse node for now, so the opportunity to halt executions is missed on those. This will change in the future.

diff --git a/Engine/Plugins/Experimental/PCG/Source/PCG/Private/Data/PCGDifferenceData.cpp b/Engine/Plugins/Experimental/PCG/Source/PCG/Private/Data/PCGDifferenceData.cpp
--- a/Engine/Plugins/Experimental/PCG/Source/PCG/Private/Data/PCGDifferenceData.cpp
+++ b/Engine/Plugins/Experimental/PCG/Source/PCG/Private/Data/PCGDifferenceData.cpp
@@ -124,34 +124,34 @@ void UPCGDifferenceData::VisitDataNetwork(TFunctionRef<void(const UPCGData*)> Ac
 	}
 }
 
-FPCGCrc UPCGDifferenceData::ComputeCrc() const
+FPCGCrc UPCGDifferenceData::ComputeCrc(bool bFullDataCrc) const
 {
 	FArchiveCrc32 Ar;
 
 	if (PropagateCrcThroughBooleanData())
 	{
-		AddToCrc(Ar);
+		AddToCrc(Ar, bFullDataCrc);
 
 		// Chain together CRCs of operands
 		check(GetSource());
-		uint32 SourceCrc = GetSource()->GetOrComputeCrc().GetValue();
+		uint32 SourceCrc = GetSource()->GetOrComputeCrc(bFullDataCrc).GetValue();
 		Ar << SourceCrc;
 
 		if (GetDifference())
 		{
-			uint32 DifferenceCrc = GetDifference()->GetOrComputeCrc().GetValue();
+			uint32 DifferenceCrc = GetDifference()->GetOrComputeCrc(bFullDataCrc).GetValue();
 			Ar << DifferenceCrc;
 		}
 	}
 	else
 	{
-		UPCGData::AddToCrc(Ar);
+		UPCGData::AddToCrc(Ar, bFullDataCrc);
 	}
 
 	return FPCGCrc(Ar.GetCrc());
 }
 
-void UPCGDifferenceData::AddToCrc(FArchiveCrc32& Ar) const
+void UPCGDifferenceData::AddToCrc(FArchiveCrc32& Ar, bool bFullDataCrc) const
 {
 	uint32 UniqueTypeID = StaticClass()->GetDefaultObject()->GetUniqueID();
 	Ar << UniqueTypeID;
diff --git a/Engine/Plugins/Experimental/PCG/Source/PCG/Private/Data/PCGIntersectionData.cpp b/Engine/Plugins/Experimental/PCG/Source/PCG/Private/Data/PCGIntersectionData.cpp
--- a/Engine/Plugins/Experimental/PCG/Source/PCG/Private/Data/PCGIntersectionData.cpp
+++ b/Engine/Plugins/Experimental/PCG/Source/PCG/Private/Data/PCGIntersectionData.cpp
@@ -71,30 +71,30 @@ void UPCGIntersectionData::VisitDataNetwork(TFunctionRef<void(const UPCGData*)>
 	GetB()->VisitDataNetwork(Action);
 }
 
-FPCGCrc UPCGIntersectionData::ComputeCrc() const
+FPCGCrc UPCGIntersectionData::ComputeCrc(bool bFullDataCrc) const
 {
 	FArchiveCrc32 Ar;
 
 	if (PropagateCrcThroughBooleanData())
 	{
-		AddToCrc(Ar);
+		AddToCrc(Ar, bFullDataCrc);
 
 		// Chain together CRCs of operands
 		check(GetA() && GetB());
-		uint32 CrcA = GetA()->GetOrComputeCrc().GetValue();
-		uint32 CrcB = GetB()->GetOrComputeCrc().GetValue();
+		uint32 CrcA = GetA()->GetOrComputeCrc(bFullDataCrc).GetValue();
+		uint32 CrcB = GetB()->GetOrComputeCrc(bFullDataCrc).GetValue();
 		Ar << CrcA;
 		Ar << CrcB;
 	}
 	else
 	{
-		UPCGData::AddToCrc(Ar);
+		UPCGData::AddToCrc(Ar, bFullDataCrc);
 	}
 
 	return FPCGCrc(Ar.GetCrc());
 }
 
-void UPCGIntersectionData::AddToCrc(FArchiveCrc32& Ar) const
+void UPCGIntersectionData::AddToCrc(FArchiveCrc32& Ar, bool bFullDataCrc) const
 {
 	uint32 UniqueTypeID = StaticClass()->GetDefaultObject()->GetUniqueID();
 	Ar << UniqueTypeID;
diff --git a/Engine/Plugins/Experimental/PCG/Source/PCG/Private/Data/PCGPointData.cpp b/Engine/Plugins/Experimental/PCG/Source/PCG/Private/Data/PCGPointData.cpp
--- a/Engine/Plugins/Experimental/PCG/Source/PCG/Private/Data/PCGPointData.cpp
+++ b/Engine/Plugins/Experimental/PCG/Source/PCG/Private/Data/PCGPointData.cpp
@@ -15,8 +15,8 @@
 
 static TAutoConsoleVariable<bool> CVarCacheFullPointDataCrc(
 	TEXT("pcg.Cache.FullPointDataCrc"),
-	false,
-	TEXT("Enable fine-grained CRC of point data for change tracking, rather than using data UID."));
+	true,
+	TEXT("Enable fine-grained CRC of point data for change tracking on elements that request it, rather than using data UID."));
 
 namespace PCGPointHelpers
 {
@@ -265,13 +265,13 @@ const UPCGPointData::PointOctree& UPCGPointData::GetOctree() const
 	return Octree;
 }
 
-void UPCGPointData::AddToCrc(FArchiveCrc32& Ar) const
+void UPCGPointData::AddToCrc(FArchiveCrc32& Ar, bool bFullDataCrc) const
 {
 	// The code below has non-trivial cost, and can be disabled from console.
-	if (!CVarCacheFullPointDataCrc.GetValueOnAnyThread())
+	if (!bFullDataCrc || !CVarCacheFullPointDataCrc.GetValueOnAnyThread())
 	{
 		// Fallback to UID
-		Super::AddToCrc(Ar);
+		Super::AddToCrc(Ar, bFullDataCrc);
 		return;
 	}
 
diff --git a/Engine/Plugins/Experimental/PCG/Source/PCG/Private/Data/PCGUnionData.cpp b/Engine/Plugins/Experimental/PCG/Source/PCG/Private/Data/PCGUnionData.cpp
--- a/Engine/Plugins/Experimental/PCG/Source/PCG/Private/Data/PCGUnionData.cpp
+++ b/Engine/Plugins/Experimental/PCG/Source/PCG/Private/Data/PCGUnionData.cpp
@@ -83,13 +83,13 @@ void UPCGUnionData::VisitDataNetwork(TFunctionRef<void(const UPCGData*)> Action)
 	}
 }
 
-FPCGCrc UPCGUnionData::ComputeCrc() const
+FPCGCrc UPCGUnionData::ComputeCrc(bool bFullDataCrc) const
 {
 	FArchiveCrc32 Ar;
 
 	if (PropagateCrcThroughBooleanData())
 	{
-		AddToCrc(Ar);
+		AddToCrc(Ar, bFullDataCrc);
 
 		// Chain together CRCs of operands
 		int32 NumOperands = Data.Num();
@@ -99,20 +99,20 @@ FPCGCrc UPCGUnionData::ComputeCrc() const
 		{
 			if (Datum)
 			{
-				uint32 DatumCrc = Datum->GetOrComputeCrc().GetValue();
+				uint32 DatumCrc = Datum->GetOrComputeCrc(bFullDataCrc).GetValue();
 				Ar << DatumCrc;
 			}
 		}
 	}
 	else
 	{
-		UPCGData::AddToCrc(Ar);
+		UPCGData::AddToCrc(Ar, bFullDataCrc);
 	}
 
 	return FPCGCrc(Ar.GetCrc());
 }
 
-void UPCGUnionData::AddToCrc(FArchiveCrc32& Ar) const
+void UPCGUnionData::AddToCrc(FArchiveCrc32& Ar, bool bFullDataCrc) const
 {
 	uint32 UniqueTypeID = StaticClass()->GetDefaultObject()->GetUniqueID();
 	Ar << UniqueTypeID;
diff --git a/Engine/Plugins/Experimental/PCG/Source/PCG/Private/Elements/PCGSurfaceSampler.cpp b/Engine/Plugins/Experimental/PCG/Source/PCG/Private/Elements/PCGSurfaceSampler.cpp
--- a/Engine/Plugins/Experimental/PCG/Source/PCG/Private/Elements/PCGSurfaceSampler.cpp
+++ b/Engine/Plugins/Experimental/PCG/Source/PCG/Private/Elements/PCGSurfaceSampler.cpp
@@ -547,7 +547,7 @@ void FPCGSurfaceSamplerElement::GetDependenciesCrc(const FPCGDataCollection& InI
 		{
 			if (const UPCGData* Data = InComponent->GetActorPCGData())
 			{
-				Crc.Combine(Data->GetOrComputeCrc());
+				Crc.Combine(Data->GetOrComputeCrc(/*bFullDataCrc=*/false));
 			}
 		}
 	}
diff --git a/Engine/Plugins/Experimental/PCG/Source/PCG/Private/Elements/PCGTextureSampler.cpp b/Engine/Plugins/Experimental/PCG/Source/PCG/Private/Elements/PCGTextureSampler.cpp
--- a/Engine/Plugins/Experimental/PCG/Source/PCG/Private/Elements/PCGTextureSampler.cpp
+++ b/Engine/Plugins/Experimental/PCG/Source/PCG/Private/Elements/PCGTextureSampler.cpp
@@ -95,7 +95,7 @@ void FPCGTextureSamplerElement::GetDependenciesCrc(const FPCGDataCollection& InI
 		{
 			if (const UPCGData* Data = InComponent->GetActorPCGData())
 			{
-				Crc.Combine(Data->GetOrComputeCrc());
+				Crc.Combine(Data->GetOrComputeCrc(/*bFullDataCrc=*/false));
 			}
 		}
 	}
diff --git a/Engine/Plugins/Experimental/PCG/Source/PCG/Private/PCGData.cpp b/Engine/Plugins/Experimental/PCG/Source/PCG/Private/PCGData.cpp
--- a/Engine/Plugins/Experimental/PCG/Source/PCG/Private/PCGData.cpp
+++ b/Engine/Plugins/Experimental/PCG/Source/PCG/Private/PCGData.cpp
@@ -88,27 +88,27 @@ UPCGData::UPCGData(const FObjectInitializer& ObjectInitializer)
 	}
 }
 
-FPCGCrc UPCGData::GetOrComputeCrc() const
+FPCGCrc UPCGData::GetOrComputeCrc(bool bFullDataCrc) const
 {
 	if (Crc.IsValid())
 	{
 		return Crc;
 	}
 
-	Crc = ComputeCrc();
+	Crc = ComputeCrc(bFullDataCrc);
 
 	return Crc;
 }
 
-FPCGCrc UPCGData::ComputeCrc() const
+FPCGCrc UPCGData::ComputeCrc(bool bFullDataCrc) const
 {
 	FArchiveCrc32 Ar;
-	AddToCrc(Ar);
+	AddToCrc(Ar, bFullDataCrc);
 
 	return FPCGCrc(Ar.GetCrc());
 }
 
-void UPCGData::AddToCrc(FArchiveCrc32& Ar) const
+void UPCGData::AddToCrc(FArchiveCrc32& Ar, bool bFullDataCrc) const
 {
 	// Fallback implementation uses UID to ensure every object returns a different Crc.
 	uint64 UIDValue = UID;
@@ -351,7 +351,7 @@ void FPCGDataCollection::RemoveFromRootSet(FPCGRootSet& RootSet) const
 	}
 }
 
-FPCGCrc FPCGDataCollection::ComputeCrc()
+FPCGCrc FPCGDataCollection::ComputeCrc(bool bFullDataCrc)
 {
 	// If there is no data, will return valid Crc==0 which is fine. No such thing as an invalid FPCGDataCollection
 	FArchiveCrc32 Ar;
@@ -362,7 +362,7 @@ FPCGCrc FPCGDataCollection::ComputeCrc()
 
 		if (Data.Data)
 		{
-			const FPCGCrc DataCrc = Data.Data->GetOrComputeCrc();
+			const FPCGCrc DataCrc = Data.Data->GetOrComputeCrc(bFullDataCrc);
 			uint32 Result = DataCrc.GetValue();
 			Ar << Result;
 		}
diff --git a/Engine/Plugins/Experimental/PCG/Source/PCG/Private/PCGElement.cpp b/Engine/Plugins/Experimental/PCG/Source/PCG/Private/PCGElement.cpp
--- a/Engine/Plugins/Experimental/PCG/Source/PCG/Private/PCGElement.cpp
+++ b/Engine/Plugins/Experimental/PCG/Source/PCG/Private/PCGElement.cpp
@@ -17,11 +17,6 @@ static TAutoConsoleVariable<bool> CVarPCGValidatePointMetadata(
 	true,
 	TEXT("Controls whether we validate that the metadata entry keys on the output point data are consistent"));
 
-static TAutoConsoleVariable<bool> CVarCacheUseDependenciesCrc(
-	TEXT("pcg.Cache.UseDependenciesCrc"),
-	false,
-	TEXT("If element is cacheable, use dependencies crc instead of calculating crc from output data."));
-
 bool IPCGElement::Execute(FPCGContext* Context) const
 {
 	check(Context && Context->AsyncState.NumAvailableTasks > 0 && Context->CurrentPhase < EPCGExecutionPhase::Done);
@@ -186,23 +181,17 @@ void IPCGElement::PostExecute(FPCGContext* Context) const
 		}
 	}
 	
-	// Output data crc.
+	// Output data Crc
 	{
-		if (CVarCacheUseDependenciesCrc.GetValueOnAnyThread() && IsCacheableInstance(SettingsInterface))
-		{
-			// If cacheable, then assume that the dependencies Crc is a good Crc to use from the output data - assumes changes
-			// in inputs and settings are 1:1 with changes in output. In general this is not the case as some settings changes will
-			// have no impact on the output. And since we are unlikely to do full Crc's of landscapes, a landscape data may produce
-			// a different Crc on each regenerate, triggering unnecessary executions.
-			Context->OutputData.Crc = Context->DependenciesCrc;
-		}
-		else
-		{
-			TRACE_CPUPROFILER_EVENT_SCOPE(IPCGElement::PostExecute::CRC);
+		TRACE_CPUPROFILER_EVENT_SCOPE(IPCGElement::PostExecute::CRC);
 
-			// Compute Crc from output data
-			Context->OutputData.Crc = Context->OutputData.ComputeCrc();
-		}
+		// Some nodes benefit from computing an actual CRC from the data. This can halt the propagation of change/executions through the graph. For
+		// data like landscapes we will never have a full accurate data crc for it so we'll tend to assume changed which triggers downstream
+		// execution. Performing a detailed CRC of output data can detect real change in the data and halt the cascade of execution.
+		const bool bShouldComputeFullOutputDataCrc = ShouldComputeFullOutputDataCrc();
+
+		// Compute Crc from output data
+		Context->OutputData.Crc = Context->OutputData.ComputeCrc(bShouldComputeFullOutputDataCrc);
 	}
 
 	// Additional debug things (check for duplicates),
diff --git a/Engine/Plugins/Experimental/PCG/Source/PCG/Private/PCGParamData.cpp b/Engine/Plugins/Experimental/PCG/Source/PCG/Private/PCGParamData.cpp
--- a/Engine/Plugins/Experimental/PCG/Source/PCG/Private/PCGParamData.cpp
+++ b/Engine/Plugins/Experimental/PCG/Source/PCG/Private/PCGParamData.cpp
@@ -2,6 +2,13 @@
 
 #include "PCGParamData.h"
 
+#include "Metadata/PCGAttributePropertySelector.h"
+#include "Metadata/Accessors/IPCGAttributeAccessor.h"
+#include "Metadata/Accessors/PCGAttributeAccessorHelpers.h"
+#include "Metadata/Accessors/PCGAttributeAccessorKeys.h"
+
+#include "Serialization/ArchiveCrc32.h"
+
 #include UE_INLINE_GENERATED_CPP_BY_NAME(PCGParamData)
 
 UPCGParamData::UPCGParamData(const FObjectInitializer& ObjectInitializer)
@@ -10,6 +17,48 @@ UPCGParamData::UPCGParamData(const FObjectInitializer& ObjectInitializer)
 	Metadata = ObjectInitializer.CreateDefaultSubobject<UPCGMetadata>(this, TEXT("Metadata"));
 }
 
+void UPCGParamData::AddToCrc(FArchiveCrc32& Ar, bool bFullDataCrc) const
+{
+	if (!Metadata)
+	{
+		// Nothing to contribute
+		return;
+	}
+
+	// Get attribute names. Preserve order as attribute order matters.
+	TArray<FName> AttributeNames;
+	{
+		TArray<EPCGMetadataTypes> AttributeTypes;
+		Metadata->GetAttributes(AttributeNames, AttributeTypes);
+	}
+
+	// Add all attribute values
+	for (FName AttributeName : AttributeNames)
+	{
+		FPCGAttributePropertySelector InputSource;
+		InputSource.SetAttributeName(AttributeName);
+
+		TUniquePtr<const IPCGAttributeAccessor> InputAccessor = PCGAttributeAccessorHelpers::CreateConstAccessor(this, InputSource);
+		TUniquePtr<const IPCGAttributeAccessorKeys> InputKeys = PCGAttributeAccessorHelpers::CreateConstKeys(this, InputSource);
+
+		auto Callback = [&InputAccessor, &InputKeys, &Ar](auto && Dummy)
+		{
+			using AttributeType = std::decay_t<decltype(Dummy)>;
+			TArray<AttributeType> Values;
+			Values.SetNum(InputKeys->GetNum());
+			InputAccessor->GetRange<AttributeType>(Values, 0, *InputKeys);
+
+			for (AttributeType Value : Values)
+			{
+				// Add value to Crc
+				Ar << Value;
+			}
+		};
+
+		PCGMetadataAttribute::CallbackWithRightType(InputAccessor->GetUnderlyingType(), Callback);
+	}
+}
+
 int64 UPCGParamData::FindMetadataKey(const FName& InName) const
 {
 	if (const PCGMetadataEntryKey* FoundKey = NameMap.Find(InName))
diff --git a/Engine/Plugins/Experimental/PCG/Source/PCG/Public/Data/PCGDifferenceData.h b/Engine/Plugins/Experimental/PCG/Source/PCG/Public/Data/PCGDifferenceData.h
--- a/Engine/Plugins/Experimental/PCG/Source/PCG/Public/Data/PCGDifferenceData.h
+++ b/Engine/Plugins/Experimental/PCG/Source/PCG/Public/Data/PCGDifferenceData.h
@@ -58,8 +58,8 @@ public:
 	virtual void VisitDataNetwork(TFunctionRef<void(const UPCGData*)> Action) const override;
 
 protected:
-	virtual FPCGCrc ComputeCrc() const override;
-	virtual void AddToCrc(FArchiveCrc32& Ar) const override;
+	virtual FPCGCrc ComputeCrc(bool bFullDataCrc) const override;
+	virtual void AddToCrc(FArchiveCrc32& Ar, bool bFullDataCrc) const override;
 	// ~End UPCGData interface
 
 public:
diff --git a/Engine/Plugins/Experimental/PCG/Source/PCG/Public/Data/PCGIntersectionData.h b/Engine/Plugins/Experimental/PCG/Source/PCG/Public/Data/PCGIntersectionData.h
--- a/Engine/Plugins/Experimental/PCG/Source/PCG/Public/Data/PCGIntersectionData.h
+++ b/Engine/Plugins/Experimental/PCG/Source/PCG/Public/Data/PCGIntersectionData.h
@@ -35,8 +35,8 @@ public:
 	virtual void VisitDataNetwork(TFunctionRef<void(const UPCGData*)> Action) const override;
 
 protected:
-	virtual FPCGCrc ComputeCrc() const override;
-	virtual void AddToCrc(FArchiveCrc32& Ar) const override;
+	virtual FPCGCrc ComputeCrc(bool bFullDataCrc) const override;
+	virtual void AddToCrc(FArchiveCrc32& Ar, bool bFullDataCrc) const override;
 	// ~End UPCGData interface
 
 public:
diff --git a/Engine/Plugins/Experimental/PCG/Source/PCG/Public/Data/PCGPointData.h b/Engine/Plugins/Experimental/PCG/Source/PCG/Public/Data/PCGPointData.h
--- a/Engine/Plugins/Experimental/PCG/Source/PCG/Public/Data/PCGPointData.h
+++ b/Engine/Plugins/Experimental/PCG/Source/PCG/Public/Data/PCGPointData.h
@@ -73,7 +73,7 @@ public:
 
 	// ~Begin UPCGData interface
 	virtual EPCGDataType GetDataType() const override { return EPCGDataType::Point; }
-	virtual void AddToCrc(FArchiveCrc32& Ar) const override;
+	virtual void AddToCrc(FArchiveCrc32& Ar, bool bFullDataCrc) const override;
 	// ~End UPCGData interface
 
 	// ~Begin UPCGSpatialData interface
diff --git a/Engine/Plugins/Experimental/PCG/Source/PCG/Public/Data/PCGUnionData.h b/Engine/Plugins/Experimental/PCG/Source/PCG/Public/Data/PCGUnionData.h
--- a/Engine/Plugins/Experimental/PCG/Source/PCG/Public/Data/PCGUnionData.h
+++ b/Engine/Plugins/Experimental/PCG/Source/PCG/Public/Data/PCGUnionData.h
@@ -41,8 +41,8 @@ public:
 	virtual void VisitDataNetwork(TFunctionRef<void(const UPCGData*)> Action) const override;
 
 protected:
-	virtual FPCGCrc ComputeCrc() const override;
-	virtual void AddToCrc(FArchiveCrc32& Ar) const override;
+	virtual FPCGCrc ComputeCrc(bool bFullDataCrc) const override;
+	virtual void AddToCrc(FArchiveCrc32& Ar, bool bFullDataCrc) const override;
 	// ~End UPCGData interface
 
 public:
diff --git a/Engine/Plugins/Experimental/PCG/Source/PCG/Public/Elements/PCGCollapseElement.h b/Engine/Plugins/Experimental/PCG/Source/PCG/Public/Elements/PCGCollapseElement.h
--- a/Engine/Plugins/Experimental/PCG/Source/PCG/Public/Elements/PCGCollapseElement.h
+++ b/Engine/Plugins/Experimental/PCG/Source/PCG/Public/Elements/PCGCollapseElement.h
@@ -29,6 +29,10 @@ protected:
 
 class FPCGCollapseElement : public FSimplePCGElement
 {
+public:
+	// Might be sampling spline/landscape or other external data, worth computing a full CRC in case we can halt change propagation/re-executions
+	virtual bool ShouldComputeFullOutputDataCrc() const override { return true; }
+
 protected:
 	virtual bool ExecuteInternal(FPCGContext* Context) const;
 };
diff --git a/Engine/Plugins/Experimental/PCG/Source/PCG/Public/Elements/PCGSplineSampler.h b/Engine/Plugins/Experimental/PCG/Source/PCG/Public/Elements/PCGSplineSampler.h
--- a/Engine/Plugins/Experimental/PCG/Source/PCG/Public/Elements/PCGSplineSampler.h
+++ b/Engine/Plugins/Experimental/PCG/Source/PCG/Public/Elements/PCGSplineSampler.h
@@ -162,6 +162,10 @@ public:
 
 class FPCGSplineSamplerElement : public FSimplePCGElement
 {
+public:
+	// Worth computing a full CRC in case we can halt change propagation/re-executions
+	virtual bool ShouldComputeFullOutputDataCrc() const override { return true; }
+
 protected:
 	virtual bool ExecuteInternal(FPCGContext* Context) const override;
 };
diff --git a/Engine/Plugins/Experimental/PCG/Source/PCG/Public/Elements/PCGSurfaceSampler.h b/Engine/Plugins/Experimental/PCG/Source/PCG/Public/Elements/PCGSurfaceSampler.h
--- a/Engine/Plugins/Experimental/PCG/Source/PCG/Public/Elements/PCGSurfaceSampler.h
+++ b/Engine/Plugins/Experimental/PCG/Source/PCG/Public/Elements/PCGSurfaceSampler.h
@@ -139,6 +139,9 @@ class FPCGSurfaceSamplerElement : public IPCGElement
 public:
 	virtual void GetDependenciesCrc(const FPCGDataCollection& InInput, const UPCGSettings* InSettings, UPCGComponent* InComponent, FPCGCrc& OutCrc) const override;
 	virtual FPCGContext* Initialize(const FPCGDataCollection& InputData, TWeakObjectPtr<UPCGComponent> SourceComponent, const UPCGNode* Node) override;
+	// Might be sampling landscape or other external data, worth computing a full CRC in case we can halt change propagation/re-executions
+	virtual bool ShouldComputeFullOutputDataCrc() const override { return true; }
+
 protected:
 	virtual bool PrepareDataInternal(FPCGContext* InContext) const override;
 	virtual bool ExecuteInternal(FPCGContext* InContext) const override;
diff --git a/Engine/Plugins/Experimental/PCG/Source/PCG/Public/Elements/PCGVolumeSampler.h b/Engine/Plugins/Experimental/PCG/Source/PCG/Public/Elements/PCGVolumeSampler.h
--- a/Engine/Plugins/Experimental/PCG/Source/PCG/Public/Elements/PCGVolumeSampler.h
+++ b/Engine/Plugins/Experimental/PCG/Source/PCG/Public/Elements/PCGVolumeSampler.h
@@ -54,6 +54,10 @@ protected:
 
 class FPCGVolumeSamplerElement : public FSimplePCGElement
 {
+public:
+	// Might be sampling external data like brush, worth computing a full CRC in case we can halt change propagation/re-executions
+	virtual bool ShouldComputeFullOutputDataCrc() const override { return true; }
+
 protected:
 	bool ExecuteInternal(FPCGContext* Context) const override;
 };
diff --git a/Engine/Plugins/Experimental/PCG/Source/PCG/Public/PCGData.h b/Engine/Plugins/Experimental/PCG/Source/PCG/Public/PCGData.h
--- a/Engine/Plugins/Experimental/PCG/Source/PCG/Public/PCGData.h
+++ b/Engine/Plugins/Experimental/PCG/Source/PCG/Public/PCGData.h
@@ -30,7 +30,7 @@ public:
 	virtual EPCGDataType GetDataType() const { return EPCGDataType::None; }
 
 	/** Returns a Crc for this and any connected data. */
-	FPCGCrc GetOrComputeCrc() const;
+	FPCGCrc GetOrComputeCrc(bool bFullDataCrc) const;
 
 	/** Executes a lambda over all connected data objects. */
 	virtual void VisitDataNetwork(TFunctionRef<void(const UPCGData*)> Action) const;
@@ -44,10 +44,10 @@ public:
 
 protected:
 	/** Computes Crc for this and any connected data. */
-	virtual FPCGCrc ComputeCrc() const;
+	virtual FPCGCrc ComputeCrc(bool bFullDataCrc) const;
 
 	/** Adds this data to Crc. Fallback implementation writes object instance UID. */
-	virtual void AddToCrc(FArchiveCrc32& Ar) const;
+	virtual void AddToCrc(FArchiveCrc32& Ar, bool bFullDataCrc) const;
 
 	/** Whether intersection, union, difference combine Crc values from operands. If false they fall back to using data UID. */
 	bool PropagateCrcThroughBooleanData() const;
@@ -136,7 +136,7 @@ struct PCG_API FPCGDataCollection
 	void RemoveFromRootSet(FPCGRootSet& RootSet) const;
 
 	/** Computes Crc for this data. */
-	FPCGCrc ComputeCrc();
+	FPCGCrc ComputeCrc(bool bFullDataCrc);
 
 	/** Cleans up the collection, but does not unroot any previously rooted data. */
 	void Reset();
diff --git a/Engine/Plugins/Experimental/PCG/Source/PCG/Public/PCGElement.h b/Engine/Plugins/Experimental/PCG/Source/PCG/Public/PCGElement.h
--- a/Engine/Plugins/Experimental/PCG/Source/PCG/Public/PCGElement.h
+++ b/Engine/Plugins/Experimental/PCG/Source/PCG/Public/PCGElement.h
@@ -51,6 +51,9 @@ public:
 	/** Returns true if the node can be cached (e.g. does not create artifacts & does not depend on untracked data */
 	virtual bool IsCacheable(const UPCGSettings* InSettings) const { return true; }
 
+	/** Whether to do a 'deep' fine-grained CRC of the output data to pass to downstream nodes. Can be expensive so should be used sparingly. */
+	virtual bool ShouldComputeFullOutputDataCrc() const { return false; }
+
 	/**
 	 * Calculate a Crc that provides a receipt for the input data that can be paired with output data from the cache. If any dependency (setting, node input or
 	 * external data) changes then this value should change. For some elements it is inefficient or not possible to output a Crc here. These can return an invalid
diff --git a/Engine/Plugins/Experimental/PCG/Source/PCG/Public/PCGParamData.h b/Engine/Plugins/Experimental/PCG/Source/PCG/Public/PCGParamData.h
--- a/Engine/Plugins/Experimental/PCG/Source/PCG/Public/PCGParamData.h
+++ b/Engine/Plugins/Experimental/PCG/Source/PCG/Public/PCGParamData.h
@@ -20,6 +20,7 @@ public:
 
 	// ~Begin UPCGData interface
 	virtual EPCGDataType GetDataType() const override { return EPCGDataType::Param; }
+	virtual void AddToCrc(FArchiveCrc32& Ar, bool bFullDataCrc) const override;
 	// ~End UPCGData interface
 
 	UFUNCTION(BlueprintCallable, Category = Metadata)
