[PCG] Fix CVar with duplicated name

#rnx

diff --git a/Engine/Plugins/Experimental/PCG/Source/PCG/Private/PCGData.cpp b/Engine/Plugins/Experimental/PCG/Source/PCG/Private/PCGData.cpp
--- a/Engine/Plugins/Experimental/PCG/Source/PCG/Private/PCGData.cpp
+++ b/Engine/Plugins/Experimental/PCG/Source/PCG/Private/PCGData.cpp
@@ -15,7 +15,7 @@
 
 #include UE_INLINE_GENERATED_CPP_BY_NAME(PCGData)
 
-static TAutoConsoleVariable<bool> CVarCacheDebugging(
+static TAutoConsoleVariable<bool> CVarCachePropagateCrcThroughBooleanData(
 	TEXT("pcg.Cache.PropagateCrcThroughBooleanData"),
 	false,
 	TEXT("Whether intersection, union, difference combine Crc values from operands. If false they fall back to using data UID."));
@@ -117,7 +117,7 @@ void UPCGData::AddToCrc(FArchiveCrc32& Ar, bool bFullDataCrc) const
 
 bool UPCGData::PropagateCrcThroughBooleanData() const
 {
-	return CVarCacheDebugging.GetValueOnAnyThread();
+	return CVarCachePropagateCrcThroughBooleanData.GetValueOnAnyThread();
 }
 
 void UPCGData::VisitDataNetwork(TFunctionRef<void(const UPCGData*)> Action) const
