[PCG] Fix test and missing include

#rnx

diff --git a/Engine/Plugins/Experimental/PCG/Source/PCG/Private/Tests/Elements/PCGStaticMeshSpawnerTest.cpp b/Engine/Plugins/Experimental/PCG/Source/PCG/Private/Tests/Elements/PCGStaticMeshSpawnerTest.cpp
--- a/Engine/Plugins/Experimental/PCG/Source/PCG/Private/Tests/Elements/PCGStaticMeshSpawnerTest.cpp
+++ b/Engine/Plugins/Experimental/PCG/Source/PCG/Private/Tests/Elements/PCGStaticMeshSpawnerTest.cpp
@@ -4,6 +4,7 @@
 #include "Tests/Determinism/PCGDeterminismTestsCommon.h"
 
 #include "PCGComponent.h"
+#include "PCGContext.h"
 #include "PCGManagedResource.h"
 
 #include "Data/PCGPointData.h"
@@ -138,9 +139,9 @@ void TestMeshSelectorByAttribute(
 
 	// initialize and execute the StaticMeshSpawner
 	FPCGElementPtr StaticMeshSpawner = Settings->GetElement();
-	FPCGContext* Context = StaticMeshSpawner->Initialize(TestData.InputData, TestData.TestPCGComponent, nullptr);
+	TUniquePtr<FPCGContext> Context = TestData.InitializeTestContext();
 
-	while (!StaticMeshSpawner->Execute(Context))
+	while (!StaticMeshSpawner->Execute(Context.Get()))
 	{}
 
 	TMap<FVector, FPCGPoint> LocationToPoint;
@@ -261,9 +262,9 @@ void TestMeshSelectorWeighted(
 
 	// initialize and execute the StaticMeshSpawner
 	FPCGElementPtr StaticMeshSpawner = Settings->GetElement();
-	FPCGContext* Context = StaticMeshSpawner->Initialize(TestData.InputData, TestData.TestPCGComponent, nullptr);
+	TUniquePtr<FPCGContext> Context = TestData.InitializeTestContext();
 
-	while (!StaticMeshSpawner->Execute(Context))
+	while (!StaticMeshSpawner->Execute(Context.Get()))
 	{}
 
 	TMap<FVector, FPCGPoint> LocationToPoint;
@@ -348,9 +349,9 @@ void TestMeshSelectorWeightedByCategory(
 
 	// initialize and execute the StaticMeshSpawner
 	FPCGElementPtr StaticMeshSpawner = Settings->GetElement();
-	FPCGContext* Context = StaticMeshSpawner->Initialize(TestData.InputData, TestData.TestPCGComponent, nullptr);
+	TUniquePtr<FPCGContext> Context = TestData.InitializeTestContext();
 
-	while (!StaticMeshSpawner->Execute(Context))
+	while (!StaticMeshSpawner->Execute(Context.Get()))
 	{
 	}
 
