Tentative fix for non-unity build

#rnx

diff --git a/Engine/Plugins/Experimental/PCG/Source/PCG/Public/Data/PCGCollisionShapeData.h b/Engine/Plugins/Experimental/PCG/Source/PCG/Public/Data/PCGCollisionShapeData.h
--- a/Engine/Plugins/Experimental/PCG/Source/PCG/Public/Data/PCGCollisionShapeData.h
+++ b/Engine/Plugins/Experimental/PCG/Source/PCG/Public/Data/PCGCollisionShapeData.h
@@ -5,6 +5,7 @@
 #include "PCGSpatialData.h"
 
 #include "CollisionShape.h"
+#include "Chaos/ChaosEngineInterface.h"
 #include "Physics/PhysicsInterfaceDeclares.h"
 
 #include "PCGCollisionShapeData.generated.h"
