[PCG] Substantial speedup by pre-computing the settings Crc.

Tests: Changed settings in graph editor, changed settings of child objects like mesh selector settings, undo/redo and with resetting values, creating nodes in graphs.

diff --git a/Engine/Plugins/Experimental/PCG/Source/PCG/Private/PCGElement.cpp b/Engine/Plugins/Experimental/PCG/Source/PCG/Private/PCGElement.cpp
--- a/Engine/Plugins/Experimental/PCG/Source/PCG/Private/PCGElement.cpp
+++ b/Engine/Plugins/Experimental/PCG/Source/PCG/Private/PCGElement.cpp
@@ -414,7 +414,11 @@ void IPCGElement::GetDependenciesCrc(const FPCGDataCollection& InInput, const UP
 
 	if (InSettings)
 	{
-		Crc.Combine(InSettings->GetCrc32());
+		const FPCGCrc& SettingsCrc = InSettings->GetCachedCrc();
+		if (ensure(SettingsCrc.IsValid()))
+		{
+			Crc.Combine(SettingsCrc);
+		}
 	}
 
 	if (InComponent)
diff --git a/Engine/Plugins/Experimental/PCG/Source/PCG/Private/PCGSettings.cpp b/Engine/Plugins/Experimental/PCG/Source/PCG/Private/PCGSettings.cpp
--- a/Engine/Plugins/Experimental/PCG/Source/PCG/Private/PCGSettings.cpp
+++ b/Engine/Plugins/Experimental/PCG/Source/PCG/Private/PCGSettings.cpp
@@ -95,12 +95,6 @@ bool UPCGSettings::operator==(const UPCGSettings& Other) const
 	}
 }
 
-uint32 UPCGSettings::GetCrc32() const
-{
-	FPCGSettingsObjectCrc32 Ar;
-	return Ar.Crc32(const_cast<UPCGSettings*>(this));
-}
-
 #if WITH_EDITOR
 void UPCGSettings::ApplyDeprecation(UPCGNode* InOutNode)
 {
@@ -137,6 +131,8 @@ void UPCGSettings::PostLoad()
 		ExecutionMode_DEPRECATED = EPCGSettingsExecutionMode::Enabled;
 	}
 #endif
+
+	CacheCrc();
 }
 
 void UPCGSettings::PostInitProperties()
@@ -146,6 +142,8 @@ void UPCGSettings::PostInitProperties()
 #if WITH_EDITOR
 	InitializeCachedOverridableParams();
 #endif //WITH_EDITOR
+
+	CacheCrc();
 }
 
 void UPCGSettings::Serialize(FArchive& Ar)
@@ -409,6 +407,8 @@ void UPCGSettings::PostEditChangeProperty(FPropertyChangedEvent& PropertyChanged
 	{
 		OnSettingsChangedDelegate.Broadcast(this, IsStructuralProperty(PropertyChangedEvent.GetPropertyName()) ? EPCGChangeType::Structural : EPCGChangeType::Settings);
 	}
+
+	CacheCrc();
 }
 
 void UPCGSettings::DirtyCache()
@@ -562,6 +562,13 @@ void UPCGSettings::FixingOverridableParamPropertyClass(FPCGSettingsOverridablePa
 	}
 }
 
+void UPCGSettings::CacheCrc()
+{
+	FPCGSettingsObjectCrc32 Ar;
+	const uint32 CrcValue = Ar.Crc32(const_cast<UPCGSettings*>(this));
+	CachedCrc = FPCGCrc(CrcValue);
+}
+
 TArray<FPCGPinProperties> UPCGSettings::DefaultPointInputPinProperties() const
 {
 	TArray<FPCGPinProperties> Properties;
diff --git a/Engine/Plugins/Experimental/PCG/Source/PCG/Public/PCGSettings.h b/Engine/Plugins/Experimental/PCG/Source/PCG/Public/PCGSettings.h
--- a/Engine/Plugins/Experimental/PCG/Source/PCG/Public/PCGSettings.h
+++ b/Engine/Plugins/Experimental/PCG/Source/PCG/Public/PCGSettings.h
@@ -152,12 +152,15 @@ public:
 	virtual TArray<FPCGPinProperties> DefaultOutputPinProperties() const;
 
 	bool operator==(const UPCGSettings& Other) const;
-	uint32 GetCrc32() const;
+	
 	bool UseSeed() const { return bUseSeed; }
 
 	// Get the seed, combined with optional PCGComponent seed
 	int GetSeed(const UPCGComponent* InSourceComponent = nullptr) const;
 
+	/** Get the Crc value precomputed for these settings. */
+	const FPCGCrc& GetCachedCrc() const { return CachedCrc; }
+
 #if WITH_EDITOR
 	/** UpdatePins will kick off invalid edges, so this is useful for moving edges around in case of pin changes. */
 	virtual void ApplyDeprecationBeforeUpdatePins(UPCGNode* InOutNode, TArray<TObjectPtr<UPCGPin>>& InputPins, TArray<TObjectPtr<UPCGPin>>& OutputPins);
@@ -292,6 +295,13 @@ protected:
 	/** Needs to be serialized since property metadata (used to populate this array) is not available at runtime. */
 	UPROPERTY()
 	TArray<FPCGSettingsOverridableParam> CachedOverridableParams;
+
+private:
+	/** Calculate Crc for these settings and save it. */
+	void CacheCrc();
+
+	/** The cached Crc for these settings. */
+	FPCGCrc CachedCrc;
 };
 
 UCLASS(BlueprintType, ClassGroup = (Procedural))
