Fixed outstanding issue in the intersection CreatePointData method where we would assign metadata entries from a different metadata referential, which was wrong

#rnx

diff --git a/Engine/Plugins/Experimental/PCG/Source/PCG/Private/Data/PCGIntersectionData.cpp b/Engine/Plugins/Experimental/PCG/Source/PCG/Private/Data/PCGIntersectionData.cpp
--- a/Engine/Plugins/Experimental/PCG/Source/PCG/Private/Data/PCGIntersectionData.cpp
+++ b/Engine/Plugins/Experimental/PCG/Source/PCG/Private/Data/PCGIntersectionData.cpp
@@ -266,9 +266,11 @@ UPCGPointData* UPCGIntersectionData::CreateAndFilterPointData(FPCGContext* Conte
 		TempYMetadata->Initialize(Y->Metadata);
 	}
 
+	const bool bPointDataHasCommonAttributes = (SourcePointData->Metadata && Y->Metadata && SourcePointData->Metadata->HasCommonAttributes(Y->Metadata));
+
 	TArray<FPCGPoint>& TargetPoints = Data->GetMutablePoints();
 
-	FPCGAsync::AsyncPointProcessing(Context, SourcePoints.Num(), TargetPoints, [this, Data, SourcePointData, &SourcePoints, Y, TempYMetadata](int32 Index, FPCGPoint& OutPoint)
+	FPCGAsync::AsyncPointProcessing(Context, SourcePoints.Num(), TargetPoints, [this, Data, SourcePointData, &SourcePoints, Y, TempYMetadata, bPointDataHasCommonAttributes](int32 Index, FPCGPoint& OutPoint)
 	{
 		TRACE_CPUPROFILER_EVENT_SCOPE(UPCGIntersectionData::CreateAndFilterPointData::Iteration);
 		const FPCGPoint& Point = SourcePoints[Index];
@@ -287,16 +289,10 @@ UPCGPointData* UPCGIntersectionData::CreateAndFilterPointData(FPCGContext* Conte
 		OutPoint.Density = PCGIntersectionDataMaths::ComputeDensity(Point.Density, PointFromY.Density, DensityFunction);
 		OutPoint.Color = Point.Color * PointFromY.Color;
 
-		if (Data->Metadata)
+		// If either the point from Y has metadata or the merge would be a non-trivial value, then perform the full merge
+		if (Data->Metadata && (bPointDataHasCommonAttributes || PointFromY.MetadataEntry != PCGInvalidEntryKey))
 		{
-			if (Point.MetadataEntry != PCGInvalidEntryKey && PointFromY.MetadataEntry != PCGInvalidEntryKey)
-			{
-				Data->Metadata->MergePointAttributesSubset(Point, SourcePointData->Metadata, SourcePointData->Metadata, PointFromY, TempYMetadata, TempYMetadata, OutPoint, EPCGMetadataOp::Min);
-			}
-			else if (PointFromY.MetadataEntry != PCGInvalidEntryKey)
-			{
-				OutPoint.MetadataEntry = PointFromY.MetadataEntry;
-			}
+			Data->Metadata->MergePointAttributesSubset(Point, SourcePointData->Metadata, SourcePointData->Metadata, PointFromY, TempYMetadata, TempYMetadata, OutPoint, EPCGMetadataOp::Min);
 		}
 
 		return true;
