Fixed minor issue under time-slicing vs ism reuse, was tripping a check in a case where we should have run only once.

#rnx

diff --git a/Engine/Plugins/Experimental/PCG/Source/PCG/Private/Elements/PCGStaticMeshSpawner.cpp b/Engine/Plugins/Experimental/PCG/Source/PCG/Private/Elements/PCGStaticMeshSpawner.cpp
--- a/Engine/Plugins/Experimental/PCG/Source/PCG/Private/Elements/PCGStaticMeshSpawner.cpp
+++ b/Engine/Plugins/Experimental/PCG/Source/PCG/Private/Elements/PCGStaticMeshSpawner.cpp
@@ -72,7 +72,7 @@ bool FPCGStaticMeshSpawnerElement::PrepareDataInternal(FPCGContext* InContext) c
 	// Check if we can reuse existing resources
 	bool& bSkippedDueToReuse = Context->bSkippedDueToReuse;
 
-	if (Context->CurrentInputIndex == 0 && CVarAllowISMReuse.GetValueOnAnyThread())
+	if (!Context->bReuseCheckDone && CVarAllowISMReuse.GetValueOnAnyThread())
 	{
 		// Compute CRC if it has not been computed (it likely isn't, but this is to futureproof this)
 		if (!Context->DependenciesCrc.IsValid())
@@ -104,6 +104,8 @@ bool FPCGStaticMeshSpawnerElement::PrepareDataInternal(FPCGContext* InContext) c
 				bSkippedDueToReuse = true;
 			}
 		}
+
+		Context->bReuseCheckDone = true;
 	}
 
 	// Early out - if we've established we could reuse resources and there is no need to generate an output, quit now
diff --git a/Engine/Plugins/Experimental/PCG/Source/PCG/Public/Elements/PCGStaticMeshSpawnerContext.h b/Engine/Plugins/Experimental/PCG/Source/PCG/Public/Elements/PCGStaticMeshSpawnerContext.h
--- a/Engine/Plugins/Experimental/PCG/Source/PCG/Public/Elements/PCGStaticMeshSpawnerContext.h
+++ b/Engine/Plugins/Experimental/PCG/Source/PCG/Public/Elements/PCGStaticMeshSpawnerContext.h
@@ -39,6 +39,7 @@ struct FPCGStaticMeshSpawnerContext : public FPCGContext
 	int32 CurrentDataIndex = 0;
 
 	// Whole-execution variables
+	bool bReuseCheckDone = false;
 	bool bSkippedDueToReuse = false;
 
 	// Per-input context variables
