Internalized spline data from spline component inside the pcg spline data for easier caching.
Added create spline node, that can either just create a spline data, a component or a full-fledged actor (hidden option).
Update spline data code & samplers to take this new struct into account.

diff --git a/Engine/Plugins/Experimental/PCG/Source/PCG/Private/Data/PCGSplineData.cpp b/Engine/Plugins/Experimental/PCG/Source/PCG/Private/Data/PCGSplineData.cpp
--- a/Engine/Plugins/Experimental/PCG/Source/PCG/Private/Data/PCGSplineData.cpp
+++ b/Engine/Plugins/Experimental/PCG/Source/PCG/Private/Data/PCGSplineData.cpp
@@ -2,29 +2,30 @@
 
 #include "Data/PCGSplineData.h"
 
-#include "Data/PCGPolyLineData.h"
 #include "PCGHelpers.h"
 #include "Data/PCGPointData.h"
+#include "Data/PCGPolyLineData.h"
 #include "Data/PCGProjectionData.h"
+#include "Data/PCGSpatialData.h"
 #include "Elements/PCGSplineSampler.h"
 
 #include "Components/SplineComponent.h"
-#include "Data/PCGSpatialData.h"
 
 #include UE_INLINE_GENERATED_CPP_BY_NAME(PCGSplineData)
 
 void UPCGSplineData::Initialize(USplineComponent* InSpline)
 {
 	check(InSpline);
-	Spline = InSpline;
 	TargetActor = InSpline->GetOwner();
 
-	CachedBounds = PCGHelpers::GetActorBounds(Spline->GetOwner());
+	SplineStruct.Initialize(InSpline);
+
+	CachedBounds = PCGHelpers::GetActorBounds(InSpline->GetOwner());
 
 	// Expand bounds by the radius of points, otherwise sections of the curve that are close
 	// to the bounds will report an invalid density.
 	FVector SplinePointsRadius = FVector::ZeroVector;
-	const FInterpCurveVector& SplineScales = Spline->GetSplinePointsScale();
+	const FInterpCurveVector& SplineScales = SplineStruct.GetSplinePointsScale();
 	for (const FInterpCurvePoint<FVector>& SplineScale : SplineScales.Points)
 	{
 		SplinePointsRadius = FVector::Max(SplinePointsRadius, SplineScale.OutVal.GetAbs());
@@ -33,24 +34,51 @@ void UPCGSplineData::Initialize(USplineComponent* InSpline)
 	CachedBounds = CachedBounds.ExpandBy(SplinePointsRadius, SplinePointsRadius);
 }
 
+void UPCGSplineData::Initialize(const TArray<FSplinePoint>& InSplinePoints, bool bIsClosedLoop, AActor* InTargetActor, const FTransform& InTransform)
+{
+	check(InTargetActor);
+	TargetActor = InTargetActor;
+
+	SplineStruct.Initialize(InSplinePoints, bIsClosedLoop, InTransform);
+
+	CachedBounds = SplineStruct.GetBounds();
+
+	// Expand bounds by the radius of points, otherwise sections of the curve that are close
+	// to the bounds will report an invalid density.
+	FVector SplinePointsRadius = FVector::ZeroVector;
+	const FInterpCurveVector& SplineScales = SplineStruct.GetSplinePointsScale();
+	for (const FInterpCurvePoint<FVector>& SplineScale : SplineScales.Points)
+	{
+		SplinePointsRadius = FVector::Max(SplinePointsRadius, SplineScale.OutVal.GetAbs());
+	}
+
+	CachedBounds = CachedBounds.ExpandBy(SplinePointsRadius, SplinePointsRadius);
+	CachedBounds = CachedBounds.TransformBy(InTransform);
+}
+
+void UPCGSplineData::ApplyTo(USplineComponent* InSplineComponent)
+{
+	SplineStruct.ApplyTo(InSplineComponent);
+}
+
 FTransform UPCGSplineData::GetTransform() const
 {
-	return Spline ? Spline->GetComponentTransform() : FTransform::Identity;
+	return SplineStruct.GetTransform();
 }
 
 int UPCGSplineData::GetNumSegments() const
 {
-	return Spline ? Spline->GetNumberOfSplineSegments() : 0;
+	return SplineStruct.GetNumberOfSplineSegments();
 }
 
 FVector::FReal UPCGSplineData::GetSegmentLength(int SegmentIndex) const
 {
-	return Spline ? Spline->GetDistanceAlongSplineAtSplinePoint(SegmentIndex + 1) - Spline->GetDistanceAlongSplineAtSplinePoint(SegmentIndex) : 0;
+	return SplineStruct.GetDistanceAlongSplineAtSplinePoint(SegmentIndex + 1) - SplineStruct.GetDistanceAlongSplineAtSplinePoint(SegmentIndex);
 }
 
 FVector UPCGSplineData::GetLocationAtDistance(int SegmentIndex, FVector::FReal Distance, bool bWorldSpace) const
 {
-	return Spline ? Spline->GetLocationAtDistanceAlongSpline(Spline->GetDistanceAlongSplineAtSplinePoint(SegmentIndex) + Distance, bWorldSpace ? ESplineCoordinateSpace::World : ESplineCoordinateSpace::Local) : FVector::ZeroVector;
+	return SplineStruct.GetLocationAtDistanceAlongSpline(SplineStruct.GetDistanceAlongSplineAtSplinePoint(SegmentIndex) + Distance, bWorldSpace ? ESplineCoordinateSpace::World : ESplineCoordinateSpace::Local);
 }
 
 FTransform UPCGSplineData::GetTransformAtDistance(int SegmentIndex, FVector::FReal Distance, bool bWorldSpace, FBox* OutBounds) const
@@ -60,24 +88,19 @@ FTransform UPCGSplineData::GetTransformAtDistance(int SegmentIndex, FVector::FRe
 		*OutBounds = FBox::BuildAABB(FVector::ZeroVector, FVector::OneVector);
 	}
 
-	return Spline ? Spline->GetTransformAtDistanceAlongSpline(Spline->GetDistanceAlongSplineAtSplinePoint(SegmentIndex) + Distance, bWorldSpace ? ESplineCoordinateSpace::World : ESplineCoordinateSpace::Local, /*bUseScale=*/true) : FTransform::Identity;
+	return SplineStruct.GetTransformAtDistanceAlongSpline(SplineStruct.GetDistanceAlongSplineAtSplinePoint(SegmentIndex) + Distance, bWorldSpace ? ESplineCoordinateSpace::World : ESplineCoordinateSpace::Local, /*bUseScale=*/true);
 }
 
 FVector::FReal UPCGSplineData::GetCurvatureAtDistance(int SegmentIndex, FVector::FReal Distance) const
 {
-	if (!Spline)
-	{
-		return 0;
-	}
-
-	const float FullDistance = Spline->GetDistanceAlongSplineAtSplinePoint(SegmentIndex) + Distance;
-	const float Param = Spline->SplineCurves.ReparamTable.Eval(FullDistance, 0.0f);
+	const float FullDistance = SplineStruct.GetDistanceAlongSplineAtSplinePoint(SegmentIndex) + Distance;
+	const float Param = SplineStruct.SplineCurves.ReparamTable.Eval(FullDistance, 0.0f);
 
 	// Since we need the first derivative (e.g. very similar to direction) to have its norm, we'll get the value directly
-	const FVector FirstDerivative = Spline->SplineCurves.Position.EvalDerivative(Param, FVector::ZeroVector);
+	const FVector FirstDerivative = SplineStruct.SplineCurves.Position.EvalDerivative(Param, FVector::ZeroVector);
 	const FVector::FReal FirstDerivativeLength = FMath::Max(FirstDerivative.Length(), UE_DOUBLE_SMALL_NUMBER);
 	const FVector ForwardVector = FirstDerivative / FirstDerivativeLength;
-	const FVector SecondDerivative = Spline->SplineCurves.Position.EvalSecondDerivative(Param, FVector::ZeroVector);
+	const FVector SecondDerivative = SplineStruct.SplineCurves.Position.EvalSecondDerivative(Param, FVector::ZeroVector);
 	// Orthogonalize the second derivative and obtain the curvature vector
 	const FVector CurvatureVector = SecondDerivative - (SecondDerivative | ForwardVector) * ForwardVector;
 	
@@ -85,7 +108,7 @@ FVector::FReal UPCGSplineData::GetCurvatureAtDistance(int SegmentIndex, FVector:
 	const FVector::FReal Curvature = CurvatureVector.Length() / FirstDerivativeLength;
 
 	// Compute sign based on sign of curvature vs. right axis
-	const FVector RightVector = Spline->GetRightVectorAtSplineInputKey(Param, ESplineCoordinateSpace::Local);
+	const FVector RightVector = SplineStruct.GetRightVectorAtSplineInputKey(Param, ESplineCoordinateSpace::Local);
 	return FMath::Sign(RightVector | CurvatureVector) * Curvature;
 }
 
@@ -99,11 +122,7 @@ const UPCGPointData* UPCGSplineData::CreatePointData(FPCGContext* Context) const
 	SamplerParams.Mode = EPCGSplineSamplingMode::Distance;
 
 	PCGSplineSampler::SampleLineData(this, this, nullptr, SamplerParams, Data);
-
-	if (Spline)
-	{
-		UE_LOG(LogPCG, Verbose, TEXT("Spline %s generated %d points"), *Spline->GetFName().ToString(), Data->GetPoints().Num());
-	}
+	UE_LOG(LogPCG, Verbose, TEXT("Spline generated %d points"), Data->GetPoints().Num());
 
 	return Data;
 }
@@ -115,11 +134,6 @@ FBox UPCGSplineData::GetBounds() const
 
 bool UPCGSplineData::SamplePoint(const FTransform& InTransform, const FBox& InBounds, FPCGPoint& OutPoint, UPCGMetadata* OutMetadata) const
 {
-	if (!Spline)
-	{
-		return false;
-	}
-
 	// TODO: support metadata
 	// TODO: support proper bounds
 
@@ -127,8 +141,8 @@ bool UPCGSplineData::SamplePoint(const FTransform& InTransform, const FBox& InBo
 	
 	// Find nearest point on spline
 	const FVector InPosition = InTransform.GetLocation();
-	float NearestPointKey = Spline->FindInputKeyClosestToWorldLocation(InPosition);
-	FTransform NearestTransform = Spline->GetTransformAtSplineInputKey(NearestPointKey, ESplineCoordinateSpace::World, true);
+	float NearestPointKey = SplineStruct.FindInputKeyClosestToWorldLocation(InPosition);
+	FTransform NearestTransform = SplineStruct.GetTransformAtSplineInputKey(NearestPointKey, ESplineCoordinateSpace::World, true);
 	FVector LocalPoint = NearestTransform.InverseTransformPosition(InPosition);
 	
 	// Linear fall off based on the distance to the nearest point
@@ -167,7 +181,7 @@ UPCGSpatialData* UPCGSplineData::CopyInternal() const
 {
 	UPCGSplineData* NewSplineData = NewObject<UPCGSplineData>();
 
-	NewSplineData->Spline = Spline;
+	NewSplineData->SplineStruct = SplineStruct;
 	NewSplineData->CachedBounds = CachedBounds;
 
 	return NewSplineData;
@@ -185,22 +199,17 @@ bool UPCGSplineProjectionData::SamplePoint(const FTransform& InTransform, const
 		return Super::SamplePoint(InTransform, InBounds, OutPoint, OutMetadata);
 	}
 
-	check(GetSpline());
-	if (!GetSpline()->Spline)
-	{
-		return false;
-	}
-
 	// Find nearest point on projected spline by lifting point along projection direction to closest position on spline. This way
 	// when we sample the spline we get a similar result to if the spline had been projected onto the surface.
 
 	const FVector InPosition = InTransform.GetLocation();
-	const USplineComponent* Spline = GetSpline()->Spline.Get();
+	check(GetSpline());
+	const FPCGSplineStruct& Spline = GetSpline()->SplineStruct;
 	check(GetSurface());
 	const FVector& SurfaceNormal = GetSurface()->GetNormal();
 
 	// Project to 2D space
-	const FTransform LocalTransform = InTransform * Spline->GetComponentTransform().Inverse();
+	const FTransform LocalTransform = InTransform * Spline.GetTransform().Inverse();
 	FVector2D LocalPosition2D = Project(LocalTransform.GetLocation());
 	float Dummy;
 	// Find nearest key on 2D spline
@@ -210,7 +219,7 @@ bool UPCGSplineProjectionData::SamplePoint(const FTransform& InTransform, const
 	// we are changing the curve length. Also, to support surface orientations that are not axis aligned, the project function
 	// probably needs to construct into a coordinate space and project onto it rather than discarding an axis, otherwise project
 	// coordinates may be non-uniformly scaled.
-	const FVector NearestPointOnSpline = Spline->GetLocationAtSplineInputKey(NearestInputKey, ESplineCoordinateSpace::World);
+	const FVector NearestPointOnSpline = Spline.GetLocationAtSplineInputKey(NearestInputKey, ESplineCoordinateSpace::World);
 	const FVector PointOnLine = FMath::ClosestPointOnInfiniteLine(InPosition, InPosition + SurfaceNormal, NearestPointOnSpline);
 
 	// TODO: this is super inefficient, could be done in 2D if we duplicate the sampling code
@@ -281,15 +290,12 @@ void UPCGSplineProjectionData::Initialize(const UPCGSplineData* InSourceSpline,
 {
 	Super::Initialize(InSourceSpline, InTargetSurface, InParams);
 
-	check(GetSpline());
-	const USplineComponent* Spline = GetSpline()->Spline.Get();
-
 	check(GetSurface());
 	const FVector& SurfaceNormal = GetSurface()->GetNormal();
 
-	if (Spline)
+	if (GetSpline())
 	{
-		const FInterpCurveVector& SplinePosition = Spline->GetSplinePointsPosition();
+		const FInterpCurveVector& SplinePosition = GetSpline()->SplineStruct.GetSplinePointsPosition();
 
 		// Build projected spline data
 		ProjectedPosition.bIsLooped = SplinePosition.bIsLooped;
diff --git a/Engine/Plugins/Experimental/PCG/Source/PCG/Private/Data/PCGSplineStruct.cpp b/Engine/Plugins/Experimental/PCG/Source/PCG/Private/Data/PCGSplineStruct.cpp
new file mode 100644
--- /dev/null
+++ b/Engine/Plugins/Experimental/PCG/Source/PCG/Private/Data/PCGSplineStruct.cpp
@@ -0,0 +1,347 @@
+// Copyright Epic Games, Inc. All Rights Reserved.
+
+#include "Data/PCGSplineStruct.h"
+
+#include UE_INLINE_GENERATED_CPP_BY_NAME(PCGSplineStruct)
+
+namespace PCGSplineStruct
+{
+	static int32 UpperBound(const TArray<FInterpCurvePoint<FVector>>& SplinePoints, float Value)
+	{
+		int32 Count = SplinePoints.Num();
+		int32 First = 0;
+
+		while (Count > 0)
+		{
+			const int32 Middle = Count / 2;
+			if (Value >= SplinePoints[First + Middle].InVal)
+			{
+				First += Middle + 1;
+				Count -= Middle + 1;
+			}
+			else
+			{
+				Count = Middle;
+			}
+		}
+
+		return First;
+	}
+
+	// Note: copied verbatim from USplineComponent::CalcBounds
+	static FBoxSphereBounds CalcBounds(const FSplineCurves& SplineCurves, bool bClosedLoop, const FTransform& LocalToWorld)
+	{
+//#if SPLINE_FAST_BOUNDS_CALCULATION
+//		FBox BoundingBox(0);
+//		for (const auto& InterpPoint : SplineCurves.Position.Points)
+//		{
+//			BoundingBox += InterpPoint.OutVal;
+//		}
+//
+//		return FBoxSphereBounds(BoundingBox.TransformBy(LocalToWorld));
+//#else
+		const int32 NumPoints = SplineCurves.Position.Points.Num();
+		const int32 NumSegments = bClosedLoop ? NumPoints : NumPoints - 1;
+
+		FVector Min(WORLD_MAX);
+		FVector Max(-WORLD_MAX);
+		if (NumSegments > 0)
+		{
+			for (int32 Index = 0; Index < NumSegments; Index++)
+			{
+				const bool bLoopSegment = (Index == NumPoints - 1);
+				const int32 NextIndex = bLoopSegment ? 0 : Index + 1;
+				const FInterpCurvePoint<FVector>& ThisInterpPoint = SplineCurves.Position.Points[Index];
+				FInterpCurvePoint<FVector> NextInterpPoint = SplineCurves.Position.Points[NextIndex];
+				if (bLoopSegment)
+				{
+					NextInterpPoint.InVal = ThisInterpPoint.InVal + SplineCurves.Position.LoopKeyOffset;
+				}
+
+				CurveVectorFindIntervalBounds(ThisInterpPoint, NextInterpPoint, Min, Max);
+			}
+		}
+		else if (NumPoints == 1)
+		{
+			Min = Max = SplineCurves.Position.Points[0].OutVal;
+		}
+		else
+		{
+			Min = FVector::ZeroVector;
+			Max = FVector::ZeroVector;
+		}
+
+		return FBoxSphereBounds(FBox(Min, Max).TransformBy(LocalToWorld));
+//#endif
+	}
+}
+
+void FPCGSplineStruct::Initialize(USplineComponent* InSplineComponent)
+{
+	check(InSplineComponent);
+	SplineCurves = InSplineComponent->SplineCurves;
+	Transform = InSplineComponent->GetComponentTransform();
+	DefaultUpVector = InSplineComponent->DefaultUpVector;
+	ReparamStepsPerSegment = InSplineComponent->ReparamStepsPerSegment;
+	bClosedLoop = InSplineComponent->IsClosedLoop();
+
+	Bounds = InSplineComponent->Bounds;
+	LocalBounds = InSplineComponent->CalcLocalBounds();
+}
+
+void FPCGSplineStruct::Initialize(const TArray<FSplinePoint>& InSplinePoints, bool bIsClosedLoop, const FTransform& InTransform)
+{
+	Transform = InTransform;
+	DefaultUpVector = FVector::ZAxisVector;
+	ReparamStepsPerSegment = 10; // default value in USplineComponent
+
+	bClosedLoop = bIsClosedLoop;
+	AddPoints(InSplinePoints, true);
+
+	Bounds = PCGSplineStruct::CalcBounds(SplineCurves, bClosedLoop, InTransform);
+	LocalBounds = PCGSplineStruct::CalcBounds(SplineCurves, bClosedLoop, FTransform::Identity);
+}
+
+void FPCGSplineStruct::ApplyTo(USplineComponent* InSplineComponent)
+{
+	check(InSplineComponent);
+
+	InSplineComponent->ClearSplinePoints(false);
+	InSplineComponent->SetComponentToWorld(Transform);
+	InSplineComponent->DefaultUpVector = DefaultUpVector;
+	InSplineComponent->ReparamStepsPerSegment = ReparamStepsPerSegment;
+
+	InSplineComponent->SplineCurves = SplineCurves;
+	InSplineComponent->bStationaryEndpoints = false;
+	// TODO: metadata? might not be needed
+	InSplineComponent->SetClosedLoop(bClosedLoop);
+	InSplineComponent->UpdateSpline();
+	InSplineComponent->UpdateBounds();
+}
+
+void FPCGSplineStruct::AddPoint(const FSplinePoint& InSplinePoint, bool bUpdateSpline)
+{
+	const int32 Index = PCGSplineStruct::UpperBound(SplineCurves.Position.Points, InSplinePoint.InputKey);
+
+	SplineCurves.Position.Points.Insert(FInterpCurvePoint<FVector>(
+		InSplinePoint.InputKey,
+		InSplinePoint.Position,
+		InSplinePoint.ArriveTangent,
+		InSplinePoint.LeaveTangent,
+		ConvertSplinePointTypeToInterpCurveMode(InSplinePoint.Type)
+		),
+		Index);
+
+	SplineCurves.Rotation.Points.Insert(FInterpCurvePoint<FQuat>(
+		InSplinePoint.InputKey,
+		InSplinePoint.Rotation.Quaternion(),
+		FQuat::Identity,
+		FQuat::Identity,
+		CIM_CurveAuto
+		),
+		Index);
+
+	SplineCurves.Scale.Points.Insert(FInterpCurvePoint<FVector>(
+		InSplinePoint.InputKey,
+		InSplinePoint.Scale,
+		FVector::ZeroVector,
+		FVector::ZeroVector,
+		CIM_CurveAuto
+		),
+		Index);
+
+	// TODO spline metadata?
+	//USplineMetadata* Metadata = GetSplinePointsMetadata();
+	//if (Metadata)
+	//{
+	//	Metadata->AddPoint(InSplinePoint.InputKey);
+	//}
+
+	if (bUpdateSpline)
+	{
+		UpdateSpline();
+	}
+}
+
+void FPCGSplineStruct::AddPoints(const TArray<FSplinePoint>& InSplinePoints, bool bUpdateSpline)
+{
+	SplineCurves.Position.Points.Reserve(SplineCurves.Position.Points.Num() + InSplinePoints.Num());
+	SplineCurves.Rotation.Points.Reserve(SplineCurves.Rotation.Points.Num() + InSplinePoints.Num());
+	SplineCurves.Scale.Points.Reserve(SplineCurves.Scale.Points.Num() + InSplinePoints.Num());
+
+	for (const auto& SplinePoint : InSplinePoints)
+	{
+		AddPoint(SplinePoint, false);
+	}
+
+	if (bUpdateSpline)
+	{
+		UpdateSpline();
+	}
+}
+
+void FPCGSplineStruct::UpdateSpline()
+{
+	const bool bLoopPositionOverride = false;
+	const bool bStationaryEndpoints = false;
+	const float LoopPosition = 0.0f;
+
+	SplineCurves.UpdateSpline(bClosedLoop, bStationaryEndpoints, ReparamStepsPerSegment, bLoopPositionOverride, LoopPosition, Transform.GetScale3D());
+}
+
+int FPCGSplineStruct::GetNumberOfSplineSegments() const
+{
+	const int32 NumPoints = SplineCurves.Position.Points.Num();
+	return (bClosedLoop ? NumPoints : NumPoints - 1);
+}
+
+FVector::FReal FPCGSplineStruct::GetSplineLength() const
+{
+	return SplineCurves.GetSplineLength();
+}
+
+FBox FPCGSplineStruct::GetBounds() const
+{
+	// See USplineComponent::CalcBounds
+	const int32 NumPoints = SplineCurves.Position.Points.Num();
+	const int32 NumSegments = bClosedLoop ? NumPoints : NumPoints - 1;
+
+	FVector Min(WORLD_MAX);
+	FVector Max(-WORLD_MAX);
+	if (NumSegments > 0)
+	{
+		for (int32 Index = 0; Index < NumSegments; Index++)
+		{
+			const bool bLoopSegment = (Index == NumPoints - 1);
+			const int32 NextIndex = bLoopSegment ? 0 : Index + 1;
+			const FInterpCurvePoint<FVector>& ThisInterpPoint = SplineCurves.Position.Points[Index];
+			FInterpCurvePoint<FVector> NextInterpPoint = SplineCurves.Position.Points[NextIndex];
+			if (bLoopSegment)
+			{
+				NextInterpPoint.InVal = ThisInterpPoint.InVal /* + SplineCurves.Position.LoopKeyOffset*/;
+			}
+
+			CurveVectorFindIntervalBounds(ThisInterpPoint, NextInterpPoint, Min, Max);
+		}
+	}
+	else if (NumPoints == 1)
+	{
+		Min = Max = SplineCurves.Position.Points[0].OutVal;
+	}
+	else
+	{
+		Min = FVector::ZeroVector;
+		Max = FVector::ZeroVector;
+	}
+
+	return FBox(Min, Max);
+}
+
+const FInterpCurveVector& FPCGSplineStruct::GetSplinePointsScale() const
+{
+	return SplineCurves.Scale;
+}
+
+const FInterpCurveVector& FPCGSplineStruct::GetSplinePointsPosition() const
+{
+	return SplineCurves.Position;
+}
+
+FVector::FReal FPCGSplineStruct::GetDistanceAlongSplineAtSplinePoint(int32 PointIndex) const
+{
+	const int32 NumPoints = SplineCurves.Position.Points.Num();
+	const int32 NumSegments = bClosedLoop ? NumPoints : NumPoints - 1;
+
+	// Ensure that if the reparam table is not prepared yet we don't attempt to access it. This can happen
+	// early in the construction of the spline component object.
+	if ((PointIndex >= 0) && (PointIndex < NumSegments + 1) && ((PointIndex * ReparamStepsPerSegment) < SplineCurves.ReparamTable.Points.Num()))
+	{
+		return SplineCurves.ReparamTable.Points[PointIndex * ReparamStepsPerSegment].InVal;
+	}
+
+	return 0.0f;
+}
+
+FVector FPCGSplineStruct::GetLocationAtDistanceAlongSpline(FVector::FReal Distance, ESplineCoordinateSpace::Type CoordinateSpace) const
+{
+	const float Param = SplineCurves.ReparamTable.Eval(Distance, 0.0f);
+	return GetLocationAtSplineInputKey(Param, CoordinateSpace);
+}
+
+FTransform FPCGSplineStruct::GetTransformAtDistanceAlongSpline(FVector::FReal Distance, ESplineCoordinateSpace::Type CoordinateSpace, bool bUseScale) const
+{
+	const float Param = SplineCurves.ReparamTable.Eval(Distance, 0.0f);
+	return GetTransformAtSplineInputKey(Param, CoordinateSpace, bUseScale);
+}
+
+FVector FPCGSplineStruct::GetRightVectorAtSplineInputKey(float InKey, ESplineCoordinateSpace::Type CoordinateSpace) const
+{
+	const FQuat Quat = GetQuaternionAtSplineInputKey(InKey, ESplineCoordinateSpace::Local);
+	FVector RightVector = Quat.RotateVector(FVector::RightVector);
+
+	if (CoordinateSpace == ESplineCoordinateSpace::World)
+	{
+		RightVector = Transform.TransformVectorNoScale(RightVector);
+	}
+
+	return RightVector;
+}
+
+FTransform FPCGSplineStruct::GetTransformAtSplineInputKey(float InKey, ESplineCoordinateSpace::Type CoordinateSpace, bool bUseScale) const
+{
+	const FVector Location(GetLocationAtSplineInputKey(InKey, ESplineCoordinateSpace::Local));
+	const FQuat Rotation(GetQuaternionAtSplineInputKey(InKey, ESplineCoordinateSpace::Local));
+	const FVector Scale = bUseScale ? GetScaleAtSplineInputKey(InKey) : FVector(1.0f);
+
+	FTransform KeyTransform(Rotation, Location, Scale);
+
+	if (CoordinateSpace == ESplineCoordinateSpace::World)
+	{
+		KeyTransform = KeyTransform * Transform;
+	}
+
+	return KeyTransform;
+}
+
+float FPCGSplineStruct::FindInputKeyClosestToWorldLocation(const FVector& WorldLocation) const
+{
+	const FVector LocalLocation = Transform.InverseTransformPosition(WorldLocation);
+	float Dummy;
+	return SplineCurves.Position.InaccurateFindNearest(LocalLocation, Dummy);
+}
+
+FVector FPCGSplineStruct::GetLocationAtSplineInputKey(float InKey, ESplineCoordinateSpace::Type CoordinateSpace) const
+{
+	FVector Location = SplineCurves.Position.Eval(InKey, FVector::ZeroVector);
+
+	if (CoordinateSpace == ESplineCoordinateSpace::World)
+	{
+		Location = Transform.TransformPosition(Location);
+	}
+
+	return Location;
+}
+
+FQuat FPCGSplineStruct::GetQuaternionAtSplineInputKey(float InKey, ESplineCoordinateSpace::Type CoordinateSpace) const
+{
+	FQuat Quat = SplineCurves.Rotation.Eval(InKey, FQuat::Identity);
+	Quat.Normalize();
+
+	const FVector Direction = SplineCurves.Position.EvalDerivative(InKey, FVector::ZeroVector).GetSafeNormal();
+	const FVector UpVector = Quat.RotateVector(DefaultUpVector);
+
+	FQuat Rot = (FRotationMatrix::MakeFromXZ(Direction, UpVector)).ToQuat();
+
+	if (CoordinateSpace == ESplineCoordinateSpace::World)
+	{
+		Rot = Transform.GetRotation() * Rot;
+	}
+
+
+	return Rot;
+}
+
+FVector FPCGSplineStruct::GetScaleAtSplineInputKey(float InKey) const
+{
+	return SplineCurves.Scale.Eval(InKey, FVector(1.0f));
+}
\ No newline at end of file
diff --git a/Engine/Plugins/Experimental/PCG/Source/PCG/Private/Elements/PCGCreateSpline.cpp b/Engine/Plugins/Experimental/PCG/Source/PCG/Private/Elements/PCGCreateSpline.cpp
new file mode 100644
--- /dev/null
+++ b/Engine/Plugins/Experimental/PCG/Source/PCG/Private/Elements/PCGCreateSpline.cpp
@@ -0,0 +1,223 @@
+// Copyright Epic Games, Inc. All Rights Reserved.
+
+#include "Elements/PCGCreateSpline.h"
+
+#include "PCGCommon.h"
+#include "PCGComponent.h"
+#include "PCGContext.h"
+#include "PCGHelpers.h"
+#include "PCGManagedResource.h"
+#include "PCGPin.h"
+#include "Data/PCGPointData.h"
+#include "Data/PCGSplineData.h"
+#include "Metadata/PCGMetadata.h"
+#include "Metadata/PCGMetadataAttributeTpl.h"
+
+#include "Components/SplineComponent.h"
+#include "Engine/World.h"
+
+#include UE_INLINE_GENERATED_CPP_BY_NAME(PCGCreateSpline)
+
+#define LOCTEXT_NAMESPACE "PCGCreateSpline"
+
+#if WITH_EDITOR
+FText UPCGCreateSplineSettings::GetNodeTooltipText() const
+{
+	return LOCTEXT("CreateSplineTooltip", "Creates PCG spline data from the input PCG point data, in a sequential order.");
+}
+#endif
+
+TArray<FPCGPinProperties> UPCGCreateSplineSettings::OutputPinProperties() const
+{
+	TArray<FPCGPinProperties> PinProperties;
+	PinProperties.Emplace(PCGPinConstants::DefaultOutputLabel, EPCGDataType::Spline);
+
+	return PinProperties;
+}
+
+FPCGElementPtr UPCGCreateSplineSettings::CreateElement() const
+{
+	return MakeShared<FPCGCreateSplineElement>();
+}
+
+bool FPCGCreateSplineElement::CanExecuteOnlyOnMainThread(FPCGContext* Context) const
+{
+	const UPCGCreateSplineSettings* Settings = Context->GetInputSettings<UPCGCreateSplineSettings>();
+	return Settings && Settings->Mode == EPCGCreateSplineMode::CreateNewActor;
+}
+
+bool FPCGCreateSplineElement::IsCacheable(const UPCGSettings* InSettings) const
+{
+	const UPCGCreateSplineSettings* Settings = Cast<const UPCGCreateSplineSettings>(InSettings);
+	return !Settings || Settings->Mode == EPCGCreateSplineMode::CreateDataOnly;
+}
+
+bool FPCGCreateSplineElement::ExecuteInternal(FPCGContext* Context) const
+{
+	TRACE_CPUPROFILER_EVENT_SCOPE(FPCGCreateSplineElement::Execute);
+
+	const UPCGCreateSplineSettings* Settings = Context->GetInputSettings<UPCGCreateSplineSettings>();
+	check(Settings);
+
+	TArray<FPCGTaggedData> Inputs = Context->InputData.GetInputs();
+	TArray<FPCGTaggedData>& Outputs = Context->OutputData.TaggedData;
+
+	EPCGCreateSplineMode Mode = Settings->Mode;
+
+	if (Mode == EPCGCreateSplineMode::CreateNewActor)
+	{
+		const bool bHasAuthority = !Context->SourceComponent.IsValid() || (Context->SourceComponent->GetOwner() && Context->SourceComponent->GetOwner()->HasAuthority());
+		if (!bHasAuthority)
+		{
+			Mode = EPCGCreateSplineMode::CreateDataOnly;
+		}
+	}
+
+	for (const FPCGTaggedData& Input : Inputs)
+	{
+		const UPCGSpatialData* SpatialData = Cast<UPCGSpatialData>(Input.Data);
+
+		if (!SpatialData)
+		{
+			PCGE_LOG(Error, "Invalid input data");
+			continue;
+		}
+
+		const UPCGPointData* PointData = SpatialData->ToPointData(Context);
+
+		if (!PointData)
+		{
+			PCGE_LOG(Error, "Unable to get PointData from input");
+			continue;
+		}
+
+		AActor* TargetActor = PointData->TargetActor.Get();
+		if (!TargetActor)
+		{
+			PCGE_LOG(Error, "Invalid target actor");
+			continue;
+		}
+
+		const FPCGMetadataAttribute<FVector>* ArriveTangentAttribute = nullptr;
+		const FPCGMetadataAttribute<FVector>* LeaveTangentAttribute = nullptr;
+
+		if (Settings->bApplyCustomTangents)
+		{
+			const UPCGMetadata* PointMetadata = PointData->ConstMetadata();
+			check(PointMetadata);
+
+			FName LocalArriveTangentName = (Settings->ArriveTangentAttribute == NAME_None ? PointMetadata->GetLatestAttributeNameOrNone() : Settings->ArriveTangentAttribute);
+			FName LocalLeaveTangentName = (Settings->LeaveTangentAttribute == NAME_None ? PointMetadata->GetLatestAttributeNameOrNone() : Settings->LeaveTangentAttribute);
+
+			const FPCGMetadataAttributeBase* ArriveTangentBaseAttribute = PointMetadata->GetConstAttribute(LocalArriveTangentName);
+			if (!ArriveTangentBaseAttribute || ArriveTangentBaseAttribute->GetTypeId() != PCG::Private::MetadataTypes<FVector>::Id)
+			{
+				PCGE_LOG(Error, "Attribute %s does not exist or is not a vector", *LocalArriveTangentName.ToString());
+				continue;
+			}
+
+			const FPCGMetadataAttributeBase* LeaveTangentBaseAttribute = PointMetadata->GetConstAttribute(LocalLeaveTangentName);
+			if (!LeaveTangentBaseAttribute || LeaveTangentBaseAttribute->GetTypeId() != PCG::Private::MetadataTypes<FVector>::Id)
+			{
+				PCGE_LOG(Error, "Attribute %s does not exist or is not a vector", *LocalLeaveTangentName.ToString());
+				continue;
+			}
+
+			ArriveTangentAttribute = static_cast<const FPCGMetadataAttribute<FVector>*>(ArriveTangentBaseAttribute);
+			LeaveTangentAttribute = static_cast<const FPCGMetadataAttribute<FVector>*>(LeaveTangentBaseAttribute);
+		}
+
+		UPCGSplineData* SplineData = NewObject<UPCGSplineData>();
+		AActor* SplineActor = TargetActor;
+
+		if (Settings->Mode == EPCGCreateSplineMode::CreateDataOnly)
+		{
+			// Nothing
+		}
+		else if(Settings->Mode == EPCGCreateSplineMode::CreateNewActor)
+		{
+			FActorSpawnParameters SpawnParams;
+			SpawnParams.Owner = TargetActor;
+			if (PCGHelpers::IsRuntimeOrPIE())
+			{
+				SpawnParams.ObjectFlags |= RF_Transient;
+			}
+
+			// TODO: check ownership of target actor like in the spawn actor?
+			// TODO: allow template?
+			SplineActor = TargetActor->GetWorld()->SpawnActor(AActor::StaticClass(), &TargetActor->GetTransform(), SpawnParams);
+
+			if (!SplineActor)
+			{
+				PCGE_LOG(Error, "Failed to create actor to hold the spline");
+				continue;
+			}
+
+			UPCGManagedActors* ManagedActors = NewObject<UPCGManagedActors>(Context->SourceComponent.Get());
+
+			FTransform Transform = TargetActor->GetTransform();
+
+			SplineActor->Tags = TargetActor->Tags;
+			SplineActor->Tags.AddUnique(PCGHelpers::DefaultPCGActorTag);
+			SplineActor->AttachToActor(TargetActor, FAttachmentTransformRules::KeepWorldTransform);
+
+			SplineData->TargetActor = SplineActor;
+
+			ManagedActors->GeneratedActors.Add(SplineActor);
+			Context->SourceComponent->AddToManagedResources(ManagedActors);
+		}
+
+		const TArray<FPCGPoint>& Points = PointData->GetPoints();
+		TArray<FSplinePoint> SplinePoints;
+		SplinePoints.Reserve(Points.Num());
+
+		const FTransform SplineActorTransform = SplineActor->GetTransform();
+		ESplinePointType::Type PointType = (Settings->bLinear ? ESplinePointType::Linear : ESplinePointType::Curve);
+
+		for(int32 PointIndex = 0; PointIndex < Points.Num(); ++PointIndex)
+		{
+			const FPCGPoint& Point = Points[PointIndex];
+			const FTransform& PointTransform = Point.Transform;
+			const FVector LocalPosition = PointTransform.GetLocation() - SplineActorTransform.GetLocation();
+
+			SplinePoints.Emplace(static_cast<float>(PointIndex),
+				LocalPosition,
+				ArriveTangentAttribute ? ArriveTangentAttribute->GetValueFromItemKey(Point.MetadataEntry) : FVector::ZeroVector,
+				LeaveTangentAttribute ? LeaveTangentAttribute->GetValueFromItemKey(Point.MetadataEntry) : FVector::ZeroVector,
+				PointTransform.GetRotation().Rotator(),
+				PointTransform.GetScale3D(),
+				PointType);
+		}
+
+		SplineData->Initialize(SplinePoints, Settings->bClosedLoop, TargetActor, FTransform(SplineActorTransform.GetLocation()));
+
+		USplineComponent* SplineComponent = nullptr;
+
+		if (Settings->Mode != EPCGCreateSplineMode::CreateDataOnly)
+		{
+			SplineComponent = NewObject<USplineComponent>(SplineActor);
+			SplineComponent->ComponentTags.Add(Context->SourceComponent.Get()->GetFName());
+			SplineComponent->ComponentTags.Add(PCGHelpers::DefaultPCGTag);
+
+			SplineData->ApplyTo(SplineComponent);
+
+			SplineComponent->RegisterComponent();
+			SplineActor->AddInstanceComponent(SplineComponent);
+			SplineComponent->AttachToComponent(SplineActor->GetRootComponent(), FAttachmentTransformRules(EAttachmentRule::KeepRelative, EAttachmentRule::KeepWorld, EAttachmentRule::KeepWorld, false));
+
+			UPCGManagedComponent* ManagedComponent = NewObject<UPCGManagedComponent>(Context->SourceComponent.Get());
+			ManagedComponent->GeneratedComponent = SplineComponent;
+			Context->SourceComponent->AddToManagedResources(ManagedComponent);
+		}
+
+		FPCGTaggedData& Output = Outputs.Add_GetRef(Input);
+		Output.Data = SplineData;
+	}
+
+	// Pass-through settings & exclusions
+	Context->OutputData.TaggedData.Append(Context->InputData.GetAllSettings());
+
+	return true;
+}
+
+#undef LOCTEXT_NAMESPACE
\ No newline at end of file
diff --git a/Engine/Plugins/Experimental/PCG/Source/PCG/Private/Elements/PCGSplineSampler.cpp b/Engine/Plugins/Experimental/PCG/Source/PCG/Private/Elements/PCGSplineSampler.cpp
--- a/Engine/Plugins/Experimental/PCG/Source/PCG/Private/Elements/PCGSplineSampler.cpp
+++ b/Engine/Plugins/Experimental/PCG/Source/PCG/Private/Elements/PCGSplineSampler.cpp
@@ -445,8 +445,8 @@ namespace PCGSplineSampler
 			FBox& OutBox = OutResult.Box;
 			OutTransform = LineData->GetTransformAtDistance(CurrentSegmentIndex, CurrentDistance, /*bWorldSpace=*/false, &OutBox);
 
-			OutBox.Min.X *= DistanceIncrement / OutTransform.GetScale3D().X;
-			OutBox.Max.X *= DistanceIncrement / OutTransform.GetScale3D().X;
+			OutBox.Min.X *= 0.5 * DistanceIncrement / OutTransform.GetScale3D().X;
+			OutBox.Max.X *= 0.5 * DistanceIncrement / OutTransform.GetScale3D().X;
 
 			if (bComputeCurvature)
 			{
@@ -576,10 +576,6 @@ namespace PCGSplineSampler
 			FPCGPoint SeedPoint;
 			if (SpatialData->SamplePoint(TransformWS, InBox, SeedPoint, nullptr))
 			{
-				// Assuming the the normal to the curve is on the Y axis
-				const FVector YAxis = TransformLS.GetScaledAxis(EAxis::Y);
-				const FVector ZAxis = TransformLS.GetScaledAxis(EAxis::Z);
-
 				// TODO: we can optimize this if we are in the "edges only case" to pick only values that
 				FVector::FReal CurrentZ = (InBox.Min.Z + ZHalfStep);
 				while(CurrentZ <= InBox.Max.Z - ZHalfStep + KINDA_SMALL_NUMBER)
@@ -614,7 +610,7 @@ namespace PCGSplineSampler
 
 						if (bTestPoint)
 						{
-							FVector TentativeLocationLS = TransformLS.GetLocation() + YAxis * CurrentY + ZAxis * CurrentZ;
+							FVector TentativeLocationLS = FVector(0.0, CurrentY, CurrentZ);
 
 							FTransform TentativeTransform = TransformWS;
 							TentativeTransform.SetLocation(TransformWS.TransformPosition(TentativeLocationLS));
@@ -698,12 +694,12 @@ namespace PCGSplineSampler
 	{
 		check(Context && LineData && OutPointData);
 
-		TSoftObjectPtr<USplineComponent> Spline;
+		const FPCGSplineStruct* Spline = nullptr;
 
 		// TODO: handle projected splines
 		if (const UPCGSplineData* SplineData = Cast<UPCGSplineData>(LineData))
 		{
-			Spline = SplineData->Spline;
+			Spline = &SplineData->SplineStruct;
 		}
 		else if (const UPCGLandscapeSplineData* LandscapeSplineData = Cast<UPCGLandscapeSplineData>(LineData))
 		{
@@ -724,7 +720,7 @@ namespace PCGSplineSampler
 			return;
 		}
 
-		const FBoxSphereBounds SplineLocalBounds = Spline->CalcLocalBounds();
+		const FBoxSphereBounds SplineLocalBounds = Spline->LocalBounds;
 		const FVector MinPoint = SplineLocalBounds.Origin - SplineLocalBounds.BoxExtent;
 		const FVector MaxPoint = SplineLocalBounds.Origin + SplineLocalBounds.BoxExtent;
 
diff --git a/Engine/Plugins/Experimental/PCG/Source/PCG/Public/Data/PCGSplineData.h b/Engine/Plugins/Experimental/PCG/Source/PCG/Public/Data/PCGSplineData.h
--- a/Engine/Plugins/Experimental/PCG/Source/PCG/Public/Data/PCGSplineData.h
+++ b/Engine/Plugins/Experimental/PCG/Source/PCG/Public/Data/PCGSplineData.h
@@ -5,7 +5,8 @@
 #include "Math/InterpCurve.h"
 #include "PCGPolyLineData.h"
 #include "PCGProjectionData.h"
-
+#include "PCGSplineStruct.h"
+#include "Elements/PCGProjectionParams.h"
 
 #include "PCGSplineData.generated.h"
 
@@ -21,6 +22,8 @@ class PCG_API UPCGSplineData : public UPCGPolyLineData
 
 public:
 	void Initialize(USplineComponent* InSpline);
+	void Initialize(const TArray<FSplinePoint>& InSplinePoints, bool bInClosedLoop, AActor* InTargetActor, const FTransform& InTransform);
+	void ApplyTo(USplineComponent* InSpline);
 
 	// ~Begin UPCGData interface
 	virtual EPCGDataType GetDataType() const override { return EPCGDataType::Spline; }
@@ -48,8 +51,9 @@ protected:
 	//~End UPCGSpatialData interface
 
 public:
-	UPROPERTY(BlueprintReadWrite, EditAnywhere, Category = SourceData)
-	TSoftObjectPtr<USplineComponent> Spline;
+	// Minimal data needed to replicate the behavior from USplineComponent
+	UPROPERTY()
+	FPCGSplineStruct SplineStruct;
 
 protected:
 	UPROPERTY()
diff --git a/Engine/Plugins/Experimental/PCG/Source/PCG/Public/Data/PCGSplineStruct.h b/Engine/Plugins/Experimental/PCG/Source/PCG/Public/Data/PCGSplineStruct.h
new file mode 100644
--- /dev/null
+++ b/Engine/Plugins/Experimental/PCG/Source/PCG/Public/Data/PCGSplineStruct.h
@@ -0,0 +1,68 @@
+// Copyright Epic Games, Inc. All Rights Reserved.
+
+#pragma once
+
+#include "CoreMinimal.h"
+#include "Components/SplineComponent.h"
+
+#include "PCGSplineStruct.generated.h"
+
+/** Subset of the Spline Component API in a standalone struct */
+USTRUCT()
+struct PCG_API FPCGSplineStruct
+{
+	GENERATED_BODY()
+
+	void Initialize(USplineComponent* InSplineComponent);
+	void Initialize(const TArray<FSplinePoint>& InSplinePoints, bool bIsClosedLoop, const FTransform& InTransform);
+	void ApplyTo(USplineComponent* InSplineComponent);
+
+	FTransform GetTransform() const { return Transform; }
+
+	// Spline-related methods
+	void AddPoint(const FSplinePoint& InSplinePoint, bool bUpdateSpline); 
+	void AddPoints(const TArray<FSplinePoint>& InSplinePoints, bool bUpdateSpline);
+	void UpdateSpline();
+
+	int GetNumberOfSplineSegments() const;
+	bool IsClosedLoop() const { return bClosedLoop; }
+	FVector::FReal GetSplineLength() const;
+	FBox GetBounds() const;
+
+	const FInterpCurveVector& GetSplinePointsScale() const;
+	const FInterpCurveVector& GetSplinePointsPosition() const;
+
+	FVector GetRightVectorAtSplineInputKey(float InKey, ESplineCoordinateSpace::Type CoordinateSpace) const;
+	FTransform GetTransformAtSplineInputKey(float InKey, ESplineCoordinateSpace::Type CoordinateSpace, bool bUseScale = false) const;
+	FVector GetLocationAtSplineInputKey(float InKey, ESplineCoordinateSpace::Type CoordinateSpace) const;
+	FQuat GetQuaternionAtSplineInputKey(float InKey, ESplineCoordinateSpace::Type CoordinateSpace) const;
+	FVector GetScaleAtSplineInputKey(float InKey) const;
+
+	FVector::FReal GetDistanceAlongSplineAtSplinePoint(int32 PointIndex) const;
+	FVector GetLocationAtDistanceAlongSpline(FVector::FReal Distance, ESplineCoordinateSpace::Type CoordinateSpace) const;
+	FTransform GetTransformAtDistanceAlongSpline(FVector::FReal Distance, ESplineCoordinateSpace::Type CoordinateSpace, bool bUseScale = false) const;
+	
+	float FindInputKeyClosestToWorldLocation(const FVector& WorldLocation) const;	
+
+	UPROPERTY()
+	FSplineCurves SplineCurves;
+
+	// Replaces the component trasnform
+	UPROPERTY()
+	FTransform Transform;
+
+	UPROPERTY()
+	FVector DefaultUpVector;
+
+	UPROPERTY()
+	int32 ReparamStepsPerSegment;
+
+	UPROPERTY()
+	bool bClosedLoop;
+
+	UPROPERTY()
+	FBoxSphereBounds LocalBounds;
+
+	UPROPERTY()
+	FBoxSphereBounds Bounds;
+};
\ No newline at end of file
diff --git a/Engine/Plugins/Experimental/PCG/Source/PCG/Public/Elements/PCGCreateSpline.h b/Engine/Plugins/Experimental/PCG/Source/PCG/Public/Elements/PCGCreateSpline.h
new file mode 100644
--- /dev/null
+++ b/Engine/Plugins/Experimental/PCG/Source/PCG/Public/Elements/PCGCreateSpline.h
@@ -0,0 +1,69 @@
+// Copyright Epic Games, Inc. All Rights Reserved.
+
+#pragma once
+
+#include "PCGElement.h"
+#include "PCGSettings.h"
+#include "Data/PCGSplineData.h"
+
+#include "PCGCreateSpline.generated.h"
+
+struct FPCGContext;
+
+UENUM()
+enum class EPCGCreateSplineMode : uint8
+{
+	CreateDataOnly,
+	CreateComponent,
+	CreateNewActor UMETA(Hidden)
+};
+
+/** PCG node that creates a spline presentation from the input points data, with optional tangents */
+UCLASS(BlueprintType, ClassGroup = (Procedural))
+class PCG_API UPCGCreateSplineSettings : public UPCGSettings
+{
+	GENERATED_BODY()
+
+public:
+	//~Begin UPCGSettings interface
+#if WITH_EDITOR
+	virtual FName GetDefaultNodeName() const override { return FName(TEXT("Create Spline")); }
+	virtual FText GetNodeTooltipText() const override;
+	virtual EPCGSettingsType GetType() const override { return EPCGSettingsType::Spatial; }
+#endif
+
+	virtual TArray<FPCGPinProperties> InputPinProperties() const override { return Super::DefaultPointInputPinProperties(); }
+	virtual TArray<FPCGPinProperties> OutputPinProperties() const override;
+
+protected:
+	virtual FPCGElementPtr CreateElement() const override;
+	//~End UPCGSettings interface
+
+public:
+	UPROPERTY(BlueprintReadWrite, EditAnywhere, Category = Settings)
+	EPCGCreateSplineMode Mode = EPCGCreateSplineMode::CreateDataOnly;
+
+	UPROPERTY(BlueprintReadWrite, EditAnywhere, Category = Settings)
+	bool bClosedLoop = false;
+
+	// Controls whether the segment between control points is a curve (when false) or a straight line (when true).
+	UPROPERTY(BlueprintReadWrite, EditAnywhere, Category = Settings)
+	bool bLinear = false;
+
+	UPROPERTY(BlueprintReadWrite, EditAnywhere, Category = Settings)
+	bool bApplyCustomTangents = false;
+
+	UPROPERTY(BlueprintReadWrite, EditAnywhere, Category = Settings, meta = (EditCondition = "bApplyCustomTangents"))
+	FName ArriveTangentAttribute;
+
+	UPROPERTY(BlueprintReadWrite, EditAnywhere, Category = Settings, meta = (EditCondition = "bApplyCustomTangents"))
+	FName LeaveTangentAttribute;
+};
+
+class FPCGCreateSplineElement : public FSimplePCGElement
+{
+protected:
+	virtual bool CanExecuteOnlyOnMainThread(FPCGContext* Context) const override;
+	virtual bool IsCacheable(const UPCGSettings* InSettings) const override;
+	virtual bool ExecuteInternal(FPCGContext* Context) const override;
+};
\ No newline at end of file
