[PCG] Fine grained CRC of CollisionShapeData

Performance improvement may come if the spline data is internalized (spline points added to UPCGSplineData) (in review) and crc'd (todo), and if pcg.Cache.PropagateCrcThroughBooleanData is set true (todo). In this case the spline samplers may avoid repeated executing for unchanging splines.

diff --git a/Engine/Plugins/Experimental/PCG/Source/PCG/Private/Data/PCGCollisionShapeData.cpp b/Engine/Plugins/Experimental/PCG/Source/PCG/Private/Data/PCGCollisionShapeData.cpp
--- a/Engine/Plugins/Experimental/PCG/Source/PCG/Private/Data/PCGCollisionShapeData.cpp
+++ b/Engine/Plugins/Experimental/PCG/Source/PCG/Private/Data/PCGCollisionShapeData.cpp
@@ -6,12 +6,12 @@
 #include "Data/PCGSpatialData.h"
 #include "Elements/PCGVolumeSampler.h"
 
+#include "Chaos/GeometryQueries.h"
 #include "Components/BoxComponent.h"
 #include "Components/CapsuleComponent.h"
 #include "Components/ShapeComponent.h"
 #include "Components/SphereComponent.h"
-
-#include "Chaos/GeometryQueries.h"
+#include "Serialization/ArchiveCrc32.h"
 
 void UPCGCollisionShapeData::Initialize(UShapeComponent* InComponent)
 {
@@ -42,6 +42,35 @@ bool UPCGCollisionShapeData::IsSupported(UShapeComponent* InComponent)
 	return false;
 }
 
+void UPCGCollisionShapeData::AddToCrc(FArchiveCrc32& Ar, bool bFullDataCrc) const
+{
+	uint32 UniqueTypeID = StaticClass()->GetDefaultObject()->GetUniqueID();
+	Ar << UniqueTypeID;
+
+	Ar << const_cast<FTransform&>(Transform);
+
+	// Shape - only Crc data that is used.
+	uint32 ShapeType = static_cast<uint32>(Shape.ShapeType);
+	Ar << ShapeType;
+
+	if (Shape.IsSphere())
+	{
+		Ar << const_cast<float&>(Shape.Sphere.Radius);
+	}
+	else if (Shape.IsCapsule())
+	{
+		Ar << const_cast<float&>(Shape.Capsule.Radius);
+		Ar << const_cast<float&>(Shape.Capsule.HalfHeight);
+	}
+	else
+	{
+		// All other cases - just serialize all three floats
+		Ar << const_cast<float&>(Shape.Box.HalfExtentX);
+		Ar << const_cast<float&>(Shape.Box.HalfExtentY);
+		Ar << const_cast<float&>(Shape.Box.HalfExtentZ);
+	}
+}
+
 bool UPCGCollisionShapeData::SamplePoint(const FTransform& InTransform, const FBox& InBounds, FPCGPoint& OutPoint, UPCGMetadata* OutMetadata) const
 {
 	FCollisionShape CollisionShape;
@@ -50,6 +79,8 @@ bool UPCGCollisionShapeData::SamplePoint(const FTransform& InTransform, const FB
 
 	if (Chaos::Utilities::CastHelper(PointAdapter.GetGeometry(), PointAdapter.GetGeomPose(InTransform.GetTranslation()), [this](const auto& Downcast, const auto& FullGeomTransform) { return Chaos::OverlapQuery(ShapeAdapter->GetGeometry(), ShapeAdapter->GetGeomPose(Transform.GetTranslation()), Downcast, FullGeomTransform, /*Thickness=*/0); }))
 	{
+		new(&OutPoint) FPCGPoint(InTransform, /*Density=*/1.0f, /*Seed=*/0);
+		OutPoint.SetLocalBounds(InBounds);
 		return true;
 	}
 	else
diff --git a/Engine/Plugins/Experimental/PCG/Source/PCG/Private/Data/PCGPointData.cpp b/Engine/Plugins/Experimental/PCG/Source/PCG/Private/Data/PCGPointData.cpp
--- a/Engine/Plugins/Experimental/PCG/Source/PCG/Private/Data/PCGPointData.cpp
+++ b/Engine/Plugins/Experimental/PCG/Source/PCG/Private/Data/PCGPointData.cpp
@@ -277,8 +277,8 @@ void UPCGPointData::AddToCrc(FArchiveCrc32& Ar, bool bFullDataCrc) const
 
 	TRACE_CPUPROFILER_EVENT_SCOPE(UPCGPointData::AddToCrc);
 
-	uint32 ThisType = static_cast<uint32>(GetDataType());
-	Ar << ThisType;
+	uint32 UniqueTypeID = StaticClass()->GetDefaultObject()->GetUniqueID();
+	Ar << UniqueTypeID;
 
 	if (Points.Num() == 0)
 	{
diff --git a/Engine/Plugins/Experimental/PCG/Source/PCG/Private/PCGParamData.cpp b/Engine/Plugins/Experimental/PCG/Source/PCG/Private/PCGParamData.cpp
--- a/Engine/Plugins/Experimental/PCG/Source/PCG/Private/PCGParamData.cpp
+++ b/Engine/Plugins/Experimental/PCG/Source/PCG/Private/PCGParamData.cpp
@@ -19,6 +19,9 @@ UPCGParamData::UPCGParamData(const FObjectInitializer& ObjectInitializer)
 
 void UPCGParamData::AddToCrc(FArchiveCrc32& Ar, bool bFullDataCrc) const
 {
+	uint32 UniqueTypeID = StaticClass()->GetDefaultObject()->GetUniqueID();
+	Ar << UniqueTypeID;
+
 	if (!Metadata)
 	{
 		// Nothing to contribute
diff --git a/Engine/Plugins/Experimental/PCG/Source/PCG/Public/Data/PCGCollisionShapeData.h b/Engine/Plugins/Experimental/PCG/Source/PCG/Public/Data/PCGCollisionShapeData.h
--- a/Engine/Plugins/Experimental/PCG/Source/PCG/Public/Data/PCGCollisionShapeData.h
+++ b/Engine/Plugins/Experimental/PCG/Source/PCG/Public/Data/PCGCollisionShapeData.h
@@ -23,6 +23,7 @@ public:
 
 	// ~Begin UPCGData interface
 	virtual EPCGDataType GetDataType() const override { return EPCGDataType::Primitive; }
+	virtual void AddToCrc(FArchiveCrc32& Ar, bool bFullDataCrc) const override;
 	// ~End UPCGData interface
 
 	// ~Begin UPCGSpatialData interface
@@ -56,4 +57,4 @@ protected:
 
 	UPROPERTY()
 	FBox CachedStrictBounds = FBox(EForceInit::ForceInit);
-};
\ No newline at end of file
+};
