Changed some ordering in settings for better readability
Reformulated some tooltips to be clearer on some fields use.
Made debug state transient & not dirty the graph

diff --git a/Engine/Plugins/Experimental/PCG/Source/PCG/Private/PCGVolume.cpp b/Engine/Plugins/Experimental/PCG/Source/PCG/Private/PCGVolume.cpp
--- a/Engine/Plugins/Experimental/PCG/Source/PCG/Private/PCGVolume.cpp
+++ b/Engine/Plugins/Experimental/PCG/Source/PCG/Private/PCGVolume.cpp
@@ -4,12 +4,30 @@
 #include "PCGComponent.h"
 #include "PCGGraph.h"
 
+#include "Components/BrushComponent.h"
+
 #include UE_INLINE_GENERATED_CPP_BY_NAME(PCGVolume)
 
 APCGVolume::APCGVolume(const FObjectInitializer& ObjectInitializer)
 	: Super(ObjectInitializer)
 {
 	PCGComponent = ObjectInitializer.CreateDefaultSubobject<UPCGComponent>(this, TEXT("PCG Component"));
+
+	UObject* ThisObject = this;
+	while (ThisObject && ThisObject->HasAnyFlags(RF_DefaultSubObject))
+	{
+		ThisObject = ThisObject->GetOuter();
+	}
+
+	if (ThisObject && !ThisObject->HasAnyFlags(RF_ClassDefaultObject | RF_NeedLoad | RF_NeedPostLoad))
+	{
+		if (UBrushComponent* MyBrushComponent = GetBrushComponent())
+		{
+			MyBrushComponent->SetCollisionObjectType(ECC_WorldStatic);
+			MyBrushComponent->SetCollisionResponseToAllChannels(ECR_Ignore);
+			MyBrushComponent->SetGenerateOverlapEvents(false);
+		}
+	}
 }
 
 #if WITH_EDITOR
diff --git a/Engine/Plugins/Experimental/PCG/Source/PCG/Public/PCGSettings.h b/Engine/Plugins/Experimental/PCG/Source/PCG/Public/PCGSettings.h
--- a/Engine/Plugins/Experimental/PCG/Source/PCG/Public/PCGSettings.h
+++ b/Engine/Plugins/Experimental/PCG/Source/PCG/Public/PCGSettings.h
@@ -90,7 +90,7 @@ public:
 	UPROPERTY(EditAnywhere, BlueprintReadWrite, Category = Debug)
 	bool bEnabled = true;
 
-	UPROPERTY(EditAnywhere, BlueprintReadWrite, Category = Debug)
+	UPROPERTY(Transient, EditAnywhere, BlueprintReadWrite, Category = Debug)
 	bool bDebug = false;
 
 #if WITH_EDITORONLY_DATA
@@ -192,16 +192,16 @@ public:
 	UPROPERTY(EditAnywhere, BlueprintReadWrite, Category = Settings, meta=(EditCondition=bUseSeed, EditConditionHides, PCG_Overridable))
 	int Seed = 0xC35A9631; // random prime number
 
-	/** TODO: Remove this - Placeholder feature until we have a nodegraph */
-	UPROPERTY(EditAnywhere, BlueprintReadWrite, Category = "Settings|Tags")
+	/** Warning - this is deprecated and will be removed soon since we have a Filter By Tag node for this specific purpose */
+	UPROPERTY(EditAnywhere, BlueprintReadWrite, Category = "Tags")
 	TSet<FString> FilterOnTags;
 
-	/** TODO: Remove this - Placeholder feature until we have a nodegraph */
-	UPROPERTY(EditAnywhere, BlueprintReadWrite, Category = "Settings|Tags")
+	/** Warning - this is deprecated and will be removed soon since we have a Filter By Tag node for this specific purpose */
+	UPROPERTY(EditAnywhere, BlueprintReadWrite, Category = "Tags")
 	bool bPassThroughFilteredOutInputs = true;
 
-	/** TODO: Remove this - Placeholder feature until we have a nodegraph */
-	UPROPERTY(EditAnywhere, BlueprintReadWrite, Category = "Settings|Tags")
+	/** Applies the specified tags on the output data. Note - this might be replaced by a dedicated Tagging node */
+	UPROPERTY(EditAnywhere, BlueprintReadWrite, Category = "Tags")
 	TSet<FString> TagsAppliedOnOutput;
 
 #if WITH_EDITORONLY_DATA
diff --git a/Engine/Plugins/Experimental/PCG/Source/PCGEditor/Private/PCGEditor.cpp b/Engine/Plugins/Experimental/PCG/Source/PCGEditor/Private/PCGEditor.cpp
--- a/Engine/Plugins/Experimental/PCG/Source/PCGEditor/Private/PCGEditor.cpp
+++ b/Engine/Plugins/Experimental/PCG/Source/PCGEditor/Private/PCGEditor.cpp
@@ -1407,7 +1407,6 @@ void FPCGEditor::OnToggleDebug()
 
 			if (PCGSettingsInterface->bDebug != bNewCheckState)
 			{
-				PCGSettingsInterface->Modify();
 				PCGSettingsInterface->bDebug = bNewCheckState;
 				PCGNode->OnNodeChangedDelegate.Broadcast(PCGNode, EPCGChangeType::Settings);
 			}
