[PCG] Fix crash with LevelInstances and PCGGraphInstances

* Having a dependency on objects that are not in the same world will produce crashes and hang.
* Will probably dirty already generated partitioned volumes

Note to reviewers: It doesn't fix the use of Level instances and PCG, but at least it avoid crashing.

#rnx

diff --git a/Engine/Plugins/Experimental/PCG/Source/PCG/Private/PCGComponent.cpp b/Engine/Plugins/Experimental/PCG/Source/PCG/Private/PCGComponent.cpp
--- a/Engine/Plugins/Experimental/PCG/Source/PCG/Private/PCGComponent.cpp
+++ b/Engine/Plugins/Experimental/PCG/Source/PCG/Private/PCGComponent.cpp
@@ -209,7 +209,7 @@ void UPCGComponent::SetPropertiesFromOriginal(const UPCGComponent* Original)
 
 	InputType = NewInputType;
 	Seed = Original->Seed;
-	GraphInstance->SetGraph(Original->GraphInstance);
+	GraphInstance->SetGraph(Original->GraphInstance ? Original->GraphInstance->GetGraph() : nullptr);
 	GenerationTrigger = Original->GenerationTrigger;
 
 #if WITH_EDITOR
