Fixed static analysis warning

#rnx

diff --git a/Engine/Plugins/Experimental/PCG/Source/PCG/Private/Elements/PCGCreateSpline.cpp b/Engine/Plugins/Experimental/PCG/Source/PCG/Private/Elements/PCGCreateSpline.cpp
--- a/Engine/Plugins/Experimental/PCG/Source/PCG/Private/Elements/PCGCreateSpline.cpp
+++ b/Engine/Plugins/Experimental/PCG/Source/PCG/Private/Elements/PCGCreateSpline.cpp
@@ -106,8 +106,8 @@ bool FPCGCreateSplineElement::ExecuteInternal(FPCGContext* Context) const
 			const UPCGMetadata* PointMetadata = PointData->ConstMetadata();
 			check(PointMetadata);
 
-			FName LocalArriveTangentName = (Settings->ArriveTangentAttribute == NAME_None ? PointMetadata->GetLatestAttributeNameOrNone() : Settings->ArriveTangentAttribute);
-			FName LocalLeaveTangentName = (Settings->LeaveTangentAttribute == NAME_None ? PointMetadata->GetLatestAttributeNameOrNone() : Settings->LeaveTangentAttribute);
+			FName LocalArriveTangentName = ((Settings->ArriveTangentAttribute == NAME_None) ? PointMetadata->GetLatestAttributeNameOrNone() : Settings->ArriveTangentAttribute);
+			FName LocalLeaveTangentName = ((Settings->LeaveTangentAttribute == NAME_None) ? PointMetadata->GetLatestAttributeNameOrNone() : Settings->LeaveTangentAttribute);
 
 			const FPCGMetadataAttributeBase* ArriveTangentBaseAttribute = PointMetadata->GetConstAttribute(LocalArriveTangentName);
 			if (!ArriveTangentBaseAttribute || ArriveTangentBaseAttribute->GetTypeId() != PCG::Private::MetadataTypes<FVector>::Id)
