[PCG] Fix crash in debug render code - assign custom data correctly based on direction.

In my spline test level I get 100% crash when moving debug flag across nodes in my graph. This was also hit by a colleague too.

diff --git a/Engine/Plugins/Experimental/PCG/Source/PCG/Private/Elements/PCGDebugElement.cpp b/Engine/Plugins/Experimental/PCG/Source/PCG/Private/Elements/PCGDebugElement.cpp
--- a/Engine/Plugins/Experimental/PCG/Source/PCG/Private/Elements/PCGDebugElement.cpp
+++ b/Engine/Plugins/Experimental/PCG/Source/PCG/Private/Elements/PCGDebugElement.cpp
@@ -184,10 +184,16 @@ namespace PCGDebugElement
 				const int32 PreExistingInstanceCount = ISMC->GetInstanceCount();
 				ISMC->AddInstances(Instances, /*bShouldReturnIndices=*/false, /*bWorldSpace=*/true);
 
-				// Then get & assign custom data
-				for (int32 PointIndex = 0; PointIndex < Points.Num(); ++PointIndex)
+				// Scan all points looking for points that match current direction and add their custom data.
+				int32 PointCounter = 0;
+				for (const FPCGPoint& Point : Points)
 				{
-					const FPCGPoint& Point = Points[PointIndex];
+					const int32 PointDirection = ((bIsAbsolute || Point.Transform.GetDeterminant() >= 0) ? 0 : 1);
+					if (PointDirection != Direction)
+					{
+						continue;
+					}
+
 					InstanceCustomData.Add(Point.Density);
 					const FVector Extents = Point.GetExtents();
 					InstanceCustomData.Add(Extents[0]);
@@ -199,19 +205,17 @@ namespace PCGDebugElement
 					InstanceCustomData.Add(Point.Color[3]);
 
 					// Doing the same thing than in PCGStaticMeshSpawnerElement
-					FMemory::Memcpy(&ISMC->PerInstanceSMCustomData[(PreExistingInstanceCount + PointIndex) * NumCustomData], InstanceCustomData.GetData(), NumCustomData * sizeof(float));
-
-					// In PCGStaticMeshSpawnerElement, this is incremented only once, because it copies the custom data in one go.
-					// Here we also increment this value only once, when the whole copy is done.
-					if (PointIndex == Points.Num() - 1)
-					{
-						// Force recreation of the render data when proxy is created
-						ISMC->InstanceUpdateCmdBuffer.NumEdits++;
-					}
+					FMemory::Memcpy(&ISMC->PerInstanceSMCustomData[(PreExistingInstanceCount + PointCounter) * NumCustomData], InstanceCustomData.GetData(), NumCustomData * sizeof(float));
 
 					InstanceCustomData.Reset();
+
+					++PointCounter;
 				}
 
+				// Force recreation of the render data when proxy is created. In PCGStaticMeshSpawnerElement, this is incremented only once,
+				// because it copies the custom data in one go. Here we also increment this value only once, when the whole copy is done.
+				ISMC->InstanceUpdateCmdBuffer.NumEdits++;
+
 				ISMC->UpdateBounds();
 			}
 		}
