[PCG] Fix up spline projection (don't rely on SamplePoint performing projection).

Fixes bug where spline projection fails due to recent purification of UPCGSurface::SamplePoint.

Makes projection of spline samples explicit via projection target argument. Inputs for spline sampling are now:

* The spline to sample (line data)
* Projection target data (optional)
* Bounding shape - used for culling

Previously this was all wrapped into a SamplePoint() call but is teased out and explicit now.

When projecting onto landscape, only apply roll and pitch rotation, don't apply a rotation that spins around the ZAxis as a 2D surface has no business changing that.

Interior sampling now supports projection (did not before). Projection overrides the "sample spline surface" option. I tried to respect both but a bunch of WS<->LS conversions made the implementation unpleasant/inefficient so i settled on projection overriding.

And this fixes SplineData::SamplePoint moving the query point (non-pure).

diff --git a/Engine/Plugins/Experimental/PCG/Source/PCG/Private/Data/PCGLandscapeData.cpp b/Engine/Plugins/Experimental/PCG/Source/PCG/Private/Data/PCGLandscapeData.cpp
--- a/Engine/Plugins/Experimental/PCG/Source/PCG/Private/Data/PCGLandscapeData.cpp
+++ b/Engine/Plugins/Experimental/PCG/Source/PCG/Private/Data/PCGLandscapeData.cpp
@@ -180,6 +180,13 @@ bool UPCGLandscapeData::ProjectPoint(const FTransform& InTransform, const FBox&
 	{
 		OutPoint.Transform.SetRotation(InTransform.GetRotation());
 	}
+	else
+	{
+		// Take landscape transform, but respect initial point yaw (don't spin points around Z axis).
+		FVector RotVector = InTransform.GetRotation().ToRotationVector();
+		RotVector.X = RotVector.Y = 0;
+		OutPoint.Transform.SetRotation(OutPoint.Transform.GetRotation() * FQuat::MakeFromRotationVector(RotVector));
+	}
 
 	if (!InParams.bProjectScales)
 	{
diff --git a/Engine/Plugins/Experimental/PCG/Source/PCG/Private/Data/PCGLandscapeSplineData.cpp b/Engine/Plugins/Experimental/PCG/Source/PCG/Private/Data/PCGLandscapeSplineData.cpp
--- a/Engine/Plugins/Experimental/PCG/Source/PCG/Private/Data/PCGLandscapeSplineData.cpp
+++ b/Engine/Plugins/Experimental/PCG/Source/PCG/Private/Data/PCGLandscapeSplineData.cpp
@@ -217,7 +217,7 @@ const UPCGPointData* UPCGLandscapeSplineData::CreatePointData(FPCGContext* Conte
 	SamplerParams.Mode = EPCGSplineSamplingMode::Distance;
 	SamplerParams.Dimension = EPCGSplineSamplingDimension::OnHorizontal;
 
-	PCGSplineSampler::SampleLineData(this, this, nullptr, SamplerParams, Data);
+	PCGSplineSampler::SampleLineData(/*LineData=*/this, /*InBoundingShapeData=*/nullptr, /*InProjectionTarget=*/nullptr, /*InProjectionParams=*/{}, SamplerParams, Data);
 
 	UE_LOG(LogPCG, Verbose, TEXT("Landscape spline %s generated %d points"), *Spline->GetFName().ToString(), Points.Num());
 
diff --git a/Engine/Plugins/Experimental/PCG/Source/PCG/Private/Data/PCGSplineData.cpp b/Engine/Plugins/Experimental/PCG/Source/PCG/Private/Data/PCGSplineData.cpp
--- a/Engine/Plugins/Experimental/PCG/Source/PCG/Private/Data/PCGSplineData.cpp
+++ b/Engine/Plugins/Experimental/PCG/Source/PCG/Private/Data/PCGSplineData.cpp
@@ -121,7 +121,15 @@ const UPCGPointData* UPCGSplineData::CreatePointData(FPCGContext* Context) const
 	FPCGSplineSamplerParams SamplerParams;
 	SamplerParams.Mode = EPCGSplineSamplingMode::Distance;
 
-	PCGSplineSampler::SampleLineData(this, this, nullptr, SamplerParams, Data);
+	const UPCGSpatialData* ProjectionTarget = nullptr;
+	FPCGProjectionParams ProjectionParams;
+	if (const UPCGSplineProjectionData* SplineProjection = Cast<const UPCGSplineProjectionData>(this))
+	{
+		ProjectionTarget = SplineProjection->GetSurface();
+		ProjectionParams = SplineProjection->GetProjectionParams();
+	}
+
+	PCGSplineSampler::SampleLineData(this, nullptr, ProjectionTarget, ProjectionParams, SamplerParams, Data);
 	UE_LOG(LogPCG, Verbose, TEXT("Spline generated %d points"), Data->GetPoints().Num());
 
 	return Data;
@@ -154,8 +162,7 @@ bool UPCGSplineData::SamplePoint(const FTransform& InTransform, const FBox& InBo
 	}
 	else
 	{
-		OutPoint.Transform = NearestTransform;
-		OutPoint.Transform.SetLocation(InPosition);
+		OutPoint.Transform = InTransform;
 		OutPoint.SetLocalBounds(InBounds);
 		OutPoint.Density = 1.0f - Distance;
 
diff --git a/Engine/Plugins/Experimental/PCG/Source/PCG/Private/Elements/PCGSplineSampler.cpp b/Engine/Plugins/Experimental/PCG/Source/PCG/Private/Elements/PCGSplineSampler.cpp
--- a/Engine/Plugins/Experimental/PCG/Source/PCG/Private/Elements/PCGSplineSampler.cpp
+++ b/Engine/Plugins/Experimental/PCG/Source/PCG/Private/Elements/PCGSplineSampler.cpp
@@ -445,6 +445,7 @@ namespace PCGSplineSampler
 			FBox& OutBox = OutResult.Box;
 			OutTransform = LineData->GetTransformAtDistance(CurrentSegmentIndex, CurrentDistance, /*bWorldSpace=*/false, &OutBox);
 
+			// Set min/max to half of extent
 			OutBox.Min.X *= 0.5 * DistanceIncrement / OutTransform.GetScale3D().X;
 			OutBox.Max.X *= 0.5 * DistanceIncrement / OutTransform.GetScale3D().X;
 
@@ -475,13 +476,14 @@ namespace PCGSplineSampler
 
 	struct FDimensionSampler
 	{
-		FDimensionSampler(const UPCGPolyLineData* InLineData, const UPCGSpatialData* InSpatialData, const UPCGSpatialData* InBoundingShapeData, const FPCGSplineSamplerParams& InParams, UPCGPointData* OutPointData)
+		FDimensionSampler(const UPCGPolyLineData* InLineData, const UPCGSpatialData* InBoundingShapeData, const UPCGSpatialData* InProjectionTarget, const FPCGProjectionParams& InProjectionParams, const FPCGSplineSamplerParams& InParams, UPCGPointData* OutPointData)
 			: Params(InParams)
 		{
 			check(InLineData);
 			LineData = InLineData;
-			SpatialData = InSpatialData;
 			BoundingShapeData = InBoundingShapeData;
+			ProjectionTargetData = InProjectionTarget;
+			ProjectionParams = InProjectionParams;
 
 			// Initialize metadata accessors if needed
 			if (OutPointData && OutPointData->Metadata && (Params.bComputeDirectionDelta || Params.bComputeCurvature))
@@ -525,20 +527,37 @@ namespace PCGSplineSampler
 		{
 			const FTransform Transform = InResult.LocalTransform * LineData->GetTransform();
 
-			FPCGPoint TrivialPoint, BoundsTestPoint;
-			if(SpatialData->SamplePoint(Transform, InResult.Box, TrivialPoint, OutPointData->Metadata)
-				&& (!BoundingShapeData || BoundingShapeData->SamplePoint(Transform, InResult.Box, BoundsTestPoint, nullptr)))
+			bool bValid = true;
+
+			FPCGPoint OutPoint;
+			OutPoint.Density = 1.0f;
+			OutPoint.SetLocalBounds(InResult.Box);
+
+			if (ProjectionTargetData)
+			{
+				FPCGPoint ProjPoint;
+				bValid = ProjectionTargetData->ProjectPoint(Transform, InResult.Box, ProjectionParams, ProjPoint, nullptr);
+				OutPoint.Transform = ProjPoint.Transform;
+			}
+			else
 			{
-				SetSeed(TrivialPoint, InResult.LocalTransform.GetLocation(), Params);
-				SetMetadata(InResult, TrivialPoint, OutPointData->Metadata);
-				OutPointData->GetMutablePoints().Add(TrivialPoint);
+				OutPoint.Transform = Transform;
+			}
+
+			FPCGPoint BoundsTestPoint;
+			if (bValid && (!BoundingShapeData || BoundingShapeData->SamplePoint(Transform, InResult.Box, BoundsTestPoint, nullptr)))
+			{
+				SetSeed(OutPoint, InResult.LocalTransform.GetLocation(), Params);
+				SetMetadata(InResult, OutPoint, OutPointData->Metadata);
+				OutPointData->GetMutablePoints().Add(OutPoint);
 			}
 		}
 
 		const FPCGSplineSamplerParams Params;
 		const UPCGPolyLineData* LineData = nullptr;
-		const UPCGSpatialData* SpatialData = nullptr;
 		const UPCGSpatialData* BoundingShapeData = nullptr;
+		const UPCGSpatialData* ProjectionTargetData = nullptr;
+		FPCGProjectionParams ProjectionParams;
 
 		bool bSetMetadata = false;
 		FPCGMetadataAttribute<double>* NextDirectionDeltaAttribute = nullptr;
@@ -548,8 +567,8 @@ namespace PCGSplineSampler
 	/** Samples in a volume surrounding the poly line. */
 	struct FVolumeSampler : public FDimensionSampler
 	{
-		FVolumeSampler(const UPCGPolyLineData* InLineData, const UPCGSpatialData* InSpatialData, const UPCGSpatialData* InBoundingShapeData, const FPCGSplineSamplerParams& Params, UPCGPointData* OutPointData)
-			: FDimensionSampler(InLineData, InSpatialData, InBoundingShapeData, Params, OutPointData)
+		FVolumeSampler(const UPCGPolyLineData* InLineData, const UPCGSpatialData* InBoundingShapeData, const UPCGSpatialData* InProjectionTarget, const FPCGProjectionParams& InProjectionParams, const FPCGSplineSamplerParams& Params, UPCGPointData* OutPointData)
+			: FDimensionSampler(InLineData, InBoundingShapeData, InProjectionTarget, InProjectionParams, Params, OutPointData)
 		{
 			Fill = Params.Fill;
 			NumPlanarSteps = 1 + ((Params.Dimension == EPCGSplineSamplingDimension::OnVertical) ? 0 : Params.NumPlanarSubdivisions);
@@ -559,8 +578,8 @@ namespace PCGSplineSampler
 		virtual void Sample(const FSamplerResult& InResult, UPCGPointData* OutPointData) override
 		{
 			const FTransform& TransformLS = InResult.LocalTransform;
-			const FTransform TransformWS = TransformLS * LineData->GetTransform();
-			const FBox& InBox = InResult.Box;
+			FTransform TransformWS = TransformLS * LineData->GetTransform();
+			FBox InBox = InResult.Box;
 
 			// We're assuming that we can scale against the origin in this method so this should always be true.
 			// We will also assume that we can separate the curve into 4 ellipse sections for radius checks
@@ -573,64 +592,71 @@ namespace PCGSplineSampler
 			SubBox.Min /= FVector(1.0, NumPlanarSteps, NumHeightSteps);
 			SubBox.Max /= FVector(1.0, NumPlanarSteps, NumHeightSteps);
 
-			FPCGPoint SeedPoint;
-			if (SpatialData->SamplePoint(TransformWS, InBox, SeedPoint, nullptr))
+			// TODO: we can optimize this if we are in the "edges only case" to only pick boundary values.
+			FVector::FReal CurrentZ = (InBox.Min.Z + ZHalfStep);
+			while (CurrentZ <= InBox.Max.Z - ZHalfStep + KINDA_SMALL_NUMBER)
 			{
-				// TODO: we can optimize this if we are in the "edges only case" to pick only values that
-				FVector::FReal CurrentZ = (InBox.Min.Z + ZHalfStep);
-				while(CurrentZ <= InBox.Max.Z - ZHalfStep + KINDA_SMALL_NUMBER)
+				// Compute inner/outer distance Z contribution (squared value since we'll compare against 1)
+				const FVector::FReal InnerZ = ((NumHeightSteps > 1) ? (CurrentZ - FMath::Sign(CurrentZ) * ZHalfStep) : 0);
+				const FVector::FReal OuterZ = ((NumHeightSteps > 1) ? (CurrentZ + FMath::Sign(CurrentZ) * ZHalfStep) : 0);
+
+				// TODO: based on the current Z, we can compute the "unit" circle (as seen below) so we don't run outside of it by design, which would be a bit more efficient
+				//  care needs to be taken to make sure we are on the same "steps" though
+				for (FVector::FReal CurrentY = (InBox.Min.Y + YHalfStep); CurrentY <= (InBox.Max.Y - YHalfStep + KINDA_SMALL_NUMBER); CurrentY += 2.0 * YHalfStep)
 				{
-					// Compute inner/outer distance Z contribution (squared value since we'll compare against 1)
-					const FVector::FReal InnerZ = ((NumHeightSteps > 1) ? (CurrentZ - FMath::Sign(CurrentZ) * ZHalfStep) : 0);
-					const FVector::FReal OuterZ = ((NumHeightSteps > 1) ? (CurrentZ + FMath::Sign(CurrentZ) * ZHalfStep) : 0);
-
-					// TODO: based on the current Z, we can compute the "unit" circle (as seen below) so we don't run outside of it by design, which would be a bit more efficient
-					//  care needs to be taken to make sure we are on the same "steps" though
-					FVector::FReal CurrentY = (InBox.Min.Y + YHalfStep);
-					while(CurrentY <= InBox.Max.Y - YHalfStep + KINDA_SMALL_NUMBER)
-					{
-						// Compute inner/outer distance Y contribution
-						const FVector::FReal InnerY = ((NumPlanarSteps > 1) ? (CurrentY - FMath::Sign(CurrentY) * YHalfStep) : 0);
-						const FVector::FReal OuterY = ((NumPlanarSteps > 1) ? (CurrentY + FMath::Sign(CurrentY) * YHalfStep) : 0);
+					// Compute inner/outer distance Y contribution
+					const FVector::FReal InnerY = ((NumPlanarSteps > 1) ? (CurrentY - FMath::Sign(CurrentY) * YHalfStep) : 0);
+					const FVector::FReal OuterY = ((NumPlanarSteps > 1) ? (CurrentY + FMath::Sign(CurrentY) * YHalfStep) : 0);
 
-						bool bTestPoint = true;
-						const FVector::FReal InnerDistance = FMath::Square((CurrentZ >= 0) ? (InnerZ / InBox.Max.Z) : (InnerZ / InBox.Min.Z)) + FMath::Square((CurrentY >= 0) ? (InnerY / InBox.Max.Y) : (InnerY / InBox.Min.Y));
-						const FVector::FReal OuterDistance = FMath::Square((CurrentZ >= 0) ? (OuterZ / InBox.Max.Z) : (OuterZ / InBox.Min.Z)) + FMath::Square((CurrentY >= 0) ? (OuterY / InBox.Max.Y) : (OuterY / InBox.Min.Y));
+					const FVector::FReal InnerDistance = FMath::Square((CurrentZ >= 0) ? (InnerZ / InBox.Max.Z) : (InnerZ / InBox.Min.Z)) + FMath::Square((CurrentY >= 0) ? (InnerY / InBox.Max.Y) : (InnerY / InBox.Min.Y));
+					const FVector::FReal OuterDistance = FMath::Square((CurrentZ >= 0) ? (OuterZ / InBox.Max.Z) : (OuterZ / InBox.Min.Z)) + FMath::Square((CurrentY >= 0) ? (OuterY / InBox.Max.Y) : (OuterY / InBox.Min.Y));
 
-						// Check if we should consider this point based on the fill mode / the position in the iteration
-						// If the normalized z^2 + y^2 > 1, then there's no point in testing this point
-						if (InnerDistance >= 1.0 + KINDA_SMALL_NUMBER)
-						{
-							bTestPoint = false; // fully outside the unit circle
-						}
-						else if (Fill == EPCGSplineSamplingFill::EdgesOnly && OuterDistance < 1.0 - KINDA_SMALL_NUMBER)
-						{
-							bTestPoint = false; // Not the edge point
-						}
+					// Check if we should consider this point based on the fill mode / the position in the iteration
+					// If the normalized z^2 + y^2 > 1, then there's no point in testing this point
+					if (InnerDistance >= 1.0 + KINDA_SMALL_NUMBER)
+					{
+						continue; // fully outside the unit circle
+					}
+					else if (Fill == EPCGSplineSamplingFill::EdgesOnly && OuterDistance < 1.0 - KINDA_SMALL_NUMBER)
+					{
+						continue; // Not the edge point
+					}
 
-						if (bTestPoint)
-						{
-							FVector TentativeLocationLS = FVector(0.0, CurrentY, CurrentZ);
+					FVector TentativeLocationLS = FVector(0.0, CurrentY, CurrentZ);
+					FTransform TentativeTransform = TransformWS;
+					TentativeTransform.SetLocation(TransformWS.TransformPosition(TentativeLocationLS));
 
-							FTransform TentativeTransform = TransformWS;
-							TentativeTransform.SetLocation(TransformWS.TransformPosition(TentativeLocationLS));
+					// Sample spline to get density.
+					FPCGPoint OutPoint;
+					if (!LineData->SamplePoint(TentativeTransform, SubBox, OutPoint, nullptr))
+					{
+						continue;
+					}
 
-							// Sample point and if successful check against bounding shape (if provided).
-							FPCGPoint OutPoint, BoundsTestPoint;
-							if (SpatialData->SamplePoint(TentativeTransform, SubBox, OutPoint, OutPointData->Metadata)
-								&& (!BoundingShapeData || BoundingShapeData->SamplePoint(TentativeTransform, SubBox, BoundsTestPoint, nullptr)))
-							{
-								SetSeed(OutPoint, TentativeLocationLS, Params);
-								SetMetadata(InResult, OutPoint, OutPointData->Metadata);
-								OutPointData->GetMutablePoints().Add(OutPoint);
-							}
+					// Project point if projection target provided.
+					if (ProjectionTargetData)
+					{
+						FPCGPoint ProjPoint;
+						if (!ProjectionTargetData->ProjectPoint(OutPoint.Transform, OutPoint.GetLocalBounds(), ProjectionParams, ProjPoint, nullptr))
+						{
+							continue;
 						}
+						OutPoint.Transform = ProjPoint.Transform;
+					}
 
-						CurrentY += 2.0 * YHalfStep;
+					// Test point against bounds.
+					FPCGPoint BoundsTestPoint;
+					if (BoundingShapeData && !BoundingShapeData->SamplePoint(OutPoint.Transform, OutPoint.GetLocalBounds(), BoundsTestPoint, nullptr))
+					{
+						continue;
 					}
 
-					CurrentZ += 2.0 * ZHalfStep;
+					SetSeed(OutPoint, TentativeLocationLS, Params);
+					SetMetadata(InResult, OutPoint, OutPointData->Metadata);
+					OutPointData->GetMutablePoints().Add(OutPoint);
 				}
+
+				CurrentZ += 2.0 * ZHalfStep;
 			}
 		}
 
@@ -639,7 +665,8 @@ namespace PCGSplineSampler
 		int NumHeightSteps;
 	};
 
-	void SampleLineData(const UPCGPolyLineData* LineData, const UPCGSpatialData* SpatialData, const UPCGSpatialData* InBoundingShapeData, const FPCGSplineSamplerParams& Params, UPCGPointData* OutPointData)
+	/** Samples on spline or within volume around it. */
+	void SampleLineData(const UPCGPolyLineData* LineData, const UPCGSpatialData* InBoundingShapeData, const UPCGSpatialData* InProjectionTarget, const FPCGProjectionParams& InProjectionParams, const FPCGSplineSamplerParams& Params, UPCGPointData* OutPointData)
 	{
 		check(LineData && OutPointData);
 
@@ -648,8 +675,8 @@ namespace PCGSplineSampler
 
 		FStepSampler* Sampler = ((Params.Mode == EPCGSplineSamplingMode::Subdivision) ? static_cast<FStepSampler*>(&SubdivisionSampler) : static_cast<FStepSampler*>(&DistanceSampler));
 
-		FDimensionSampler TrivialDimensionSampler(LineData, SpatialData, InBoundingShapeData, Params, OutPointData);
-		FVolumeSampler VolumeSampler(LineData, SpatialData, InBoundingShapeData, Params, OutPointData);
+		FDimensionSampler TrivialDimensionSampler(LineData, InBoundingShapeData, InProjectionTarget, InProjectionParams, Params, OutPointData);
+		FVolumeSampler VolumeSampler(LineData, InBoundingShapeData, InProjectionTarget, InProjectionParams, Params, OutPointData);
 
 		FDimensionSampler* ExtentsSampler = ((Params.Dimension == EPCGSplineSamplingDimension::OnSpline) ? &TrivialDimensionSampler : static_cast<FDimensionSampler*>(&VolumeSampler));
 
@@ -690,13 +717,13 @@ namespace PCGSplineSampler
 		}
 	}
  
-	void SampleInteriorData(FPCGContext* Context, const UPCGPolyLineData* LineData, const UPCGSpatialData* SpatialData, const UPCGSpatialData* InBoundingShape, const FPCGSplineSamplerParams& Params, UPCGPointData* OutPointData)
+	/** Samples 2D region bounded by spline. */
+	void SampleInteriorData(FPCGContext* Context, const UPCGPolyLineData* LineData, const UPCGSpatialData* InBoundingShape, const UPCGSpatialData* InProjectionTarget, const FPCGProjectionParams& InProjectionParams, const FPCGSplineSamplerParams& Params, UPCGPointData* OutPointData)
 	{
 		check(Context && LineData && OutPointData);
 
 		const FPCGSplineStruct* Spline = nullptr;
 
-		// TODO: handle projected splines
 		if (const UPCGSplineData* SplineData = Cast<UPCGSplineData>(LineData))
 		{
 			Spline = &SplineData->SplineStruct;
@@ -1006,7 +1033,19 @@ namespace PCGSplineSampler
 							Density = DensityFalloffCurve->Eval(T);
 						}
 
-						InteriorSplinePointData[DispatchIndex].Emplace(TransformLS * LineDataTransform, PointLocationLS, Density);
+						FTransform TransformWS = TransformLS * LineDataTransform;
+						if (InProjectionTarget)
+						{
+							FPCGPoint ProjectedPoint;
+							if (!InProjectionTarget->ProjectPoint(TransformWS, FBox(BoundsMin, BoundsMax), InProjectionParams, ProjectedPoint, nullptr))
+							{
+								continue;
+							}
+
+							TransformWS.SetLocation(ProjectedPoint.Transform.GetLocation());
+						}
+
+						InteriorSplinePointData[DispatchIndex].Emplace(TransformWS, PointLocationLS, Density);
 					}
 				}
 			}
@@ -1120,7 +1159,6 @@ bool FPCGSplineSamplerElement::ExecuteInternal(FPCGContext* Context) const
 	for (FPCGTaggedData& Input : SplineInputs)
 	{
 		const UPCGSpatialData* SpatialData = Cast<const UPCGSpatialData>(Input.Data);
-
 		if (!SpatialData)
 		{
 			continue;
@@ -1128,12 +1166,19 @@ bool FPCGSplineSamplerElement::ExecuteInternal(FPCGContext* Context) const
 
 		// TODO: do something for point data approximations
 		const UPCGPolyLineData* LineData = PCGSplineSampler::GetPolyLineData(SpatialData);
-
 		if (!LineData)
 		{
 			continue;
 		}
 
+		const UPCGSpatialData* ProjectionTarget = nullptr;
+		FPCGProjectionParams ProjectionParams;
+		if (const UPCGSplineProjectionData* SplineProjection = Cast<const UPCGSplineProjectionData>(Input.Data))
+		{
+			ProjectionTarget = SplineProjection->GetSurface();
+			ProjectionParams = SplineProjection->GetProjectionParams();
+		}
+
 		FPCGTaggedData& Output = Outputs.Emplace_GetRef();
 		Output = Input;
 
@@ -1143,11 +1188,11 @@ bool FPCGSplineSamplerElement::ExecuteInternal(FPCGContext* Context) const
 
 		if (SamplerParams.Dimension == EPCGSplineSamplingDimension::OnInterior)
 		{
-			PCGSplineSampler::SampleInteriorData(Context, LineData, SpatialData, BoundingShapeSpatialInput, SamplerParams, SampledPointData);
+			PCGSplineSampler::SampleInteriorData(Context, LineData, BoundingShapeSpatialInput, ProjectionTarget, ProjectionParams, SamplerParams, SampledPointData);
 		}
 		else
 		{
-			PCGSplineSampler::SampleLineData(LineData, SpatialData, BoundingShapeSpatialInput, SamplerParams, SampledPointData);
+			PCGSplineSampler::SampleLineData(LineData, BoundingShapeSpatialInput, ProjectionTarget, ProjectionParams, SamplerParams, SampledPointData);
 		}
 	}
 
diff --git a/Engine/Plugins/Experimental/PCG/Source/PCG/Public/Data/PCGProjectionData.h b/Engine/Plugins/Experimental/PCG/Source/PCG/Public/Data/PCGProjectionData.h
--- a/Engine/Plugins/Experimental/PCG/Source/PCG/Public/Data/PCGProjectionData.h
+++ b/Engine/Plugins/Experimental/PCG/Source/PCG/Public/Data/PCGProjectionData.h
@@ -16,6 +16,8 @@ class PCG_API UPCGProjectionData : public UPCGSpatialDataWithPointCache
 public:
 	void Initialize(const UPCGSpatialData* InSource, const UPCGSpatialData* InTarget, const FPCGProjectionParams& InProjectionParams);
 
+	const FPCGProjectionParams& GetProjectionParams() const { return ProjectionParams; }
+
 	//~Begin UPCGSpatialData interface
 	virtual int GetDimension() const override;
 	virtual FBox GetBounds() const override;
diff --git a/Engine/Plugins/Experimental/PCG/Source/PCG/Public/Elements/PCGSplineSampler.h b/Engine/Plugins/Experimental/PCG/Source/PCG/Public/Elements/PCGSplineSampler.h
--- a/Engine/Plugins/Experimental/PCG/Source/PCG/Public/Elements/PCGSplineSampler.h
+++ b/Engine/Plugins/Experimental/PCG/Source/PCG/Public/Elements/PCGSplineSampler.h
@@ -4,10 +4,12 @@
 
 #include "PCGPin.h"
 #include "PCGSettings.h"
+
 #include "Curves/CurveFloat.h"
 
 #include "PCGSplineSampler.generated.h"
 
+struct FPCGProjectionParams;
 class UPCGPolyLineData;
 class UPCGSpatialData;
 class UPCGPointData;
@@ -125,9 +127,9 @@ namespace PCGSplineSamplerHelpers
 
 namespace PCGSplineSampler
 {
-	void SampleLineData(const UPCGPolyLineData* LineData, const UPCGSpatialData* SpatialData, const UPCGSpatialData* InBoundingShape, const FPCGSplineSamplerParams& Params, UPCGPointData* OutPointData);
+	void SampleLineData(const UPCGPolyLineData* LineData, const UPCGSpatialData* InBoundingShape, const UPCGSpatialData* InProjectionTarget, const FPCGProjectionParams& InProjectionParams, const FPCGSplineSamplerParams& Params, UPCGPointData* OutPointData);
 
-	void SampleInteriorData(const UPCGPolyLineData* LineData, const UPCGSpatialData* SpatialData, const UPCGSpatialData* InBoundingShape, const FPCGSplineSamplerParams& Params, UPCGPointData* OutPointData);
+	void SampleInteriorData(const UPCGPolyLineData* LineData, const UPCGSpatialData* InBoundingShape, const UPCGSpatialData* InProjectionTarget, const FPCGProjectionParams& InProjectionParams, const FPCGSplineSamplerParams& Params, UPCGPointData* OutPointData);
 	const UPCGPolyLineData* GetPolyLineData(const UPCGSpatialData* InSpatialData);
 }
 
