[PCG] Fixed an issue were undo wouldn't recreate pins properly and causing crashes upon node interactions.

diff --git a/Engine/Plugins/Experimental/PCG/Source/PCGEditor/Private/PCGEditorGraphNodeBase.cpp b/Engine/Plugins/Experimental/PCG/Source/PCGEditor/Private/PCGEditorGraphNodeBase.cpp
--- a/Engine/Plugins/Experimental/PCG/Source/PCGEditor/Private/PCGEditorGraphNodeBase.cpp
+++ b/Engine/Plugins/Experimental/PCG/Source/PCGEditor/Private/PCGEditorGraphNodeBase.cpp
@@ -315,6 +315,8 @@ void UPCGEditorGraphNodeBase::ReconstructNode()
 		bDeferredReconstruct = true;
 		return;
 	}
+	
+	Modify();
 
 	// Store copy of old pins
 	TArray<UEdGraphPin*> OldPins = MoveTemp(Pins);
