[PCG] Make elements Attribute Partition, Attribute Operation, Point Sampler all take Point input as they convert to points under the hood.

Part of wider effort to minimize use of Any type.

NOTE: Existing edges won't break - we still allow Any->Point and Spatial->Point, for now. This change is of the same nature as a similar change done to the other nodes, but I missed these cases.

Loads of janitorial work to freshen outdated code.

diff --git a/Engine/Plugins/Experimental/PCG/Source/PCG/Private/Elements/PCGBoundsModifier.cpp b/Engine/Plugins/Experimental/PCG/Source/PCG/Private/Elements/PCGBoundsModifier.cpp
--- a/Engine/Plugins/Experimental/PCG/Source/PCG/Private/Elements/PCGBoundsModifier.cpp
+++ b/Engine/Plugins/Experimental/PCG/Source/PCG/Private/Elements/PCGBoundsModifier.cpp
@@ -14,13 +14,6 @@ FPCGElementPtr UPCGBoundsModifierSettings::CreateElement() const
 	return MakeShared<FPCGBoundsModifier>();
 }
 
-TArray<FPCGPinProperties> UPCGBoundsModifierSettings::InputPinProperties() const
-{
-	TArray<FPCGPinProperties> PinProperties;
-	PinProperties.Emplace(PCGPinConstants::DefaultInputLabel, EPCGDataType::Point);
-	return PinProperties;
-}
-
 #if WITH_EDITOR
 FText UPCGBoundsModifierSettings::GetNodeTooltipText() const
 {
diff --git a/Engine/Plugins/Experimental/PCG/Source/PCG/Private/Elements/PCGDensityNoise.cpp b/Engine/Plugins/Experimental/PCG/Source/PCG/Private/Elements/PCGDensityNoise.cpp
--- a/Engine/Plugins/Experimental/PCG/Source/PCG/Private/Elements/PCGDensityNoise.cpp
+++ b/Engine/Plugins/Experimental/PCG/Source/PCG/Private/Elements/PCGDensityNoise.cpp
@@ -16,15 +16,6 @@ UPCGDensityNoiseSettings::UPCGDensityNoiseSettings()
 	bUseSeed = true;
 }
 
-TArray<FPCGPinProperties> UPCGDensityNoiseSettings::InputPinProperties() const
-{
-	TArray<FPCGPinProperties> PinProperties;
-	// TODO in the future type checking of edges will be stricter and a conversion node will be added to convert from other types
-	PinProperties.Emplace(PCGPinConstants::DefaultInputLabel, EPCGDataType::Point);
-
-	return PinProperties;
-}
-
 FPCGElementPtr UPCGDensityNoiseSettings::CreateElement() const
 {
 	return MakeShared<FPCGDensityNoiseElement>();
diff --git a/Engine/Plugins/Experimental/PCG/Source/PCG/Private/Elements/PCGDensityRemapElement.cpp b/Engine/Plugins/Experimental/PCG/Source/PCG/Private/Elements/PCGDensityRemapElement.cpp
--- a/Engine/Plugins/Experimental/PCG/Source/PCG/Private/Elements/PCGDensityRemapElement.cpp
+++ b/Engine/Plugins/Experimental/PCG/Source/PCG/Private/Elements/PCGDensityRemapElement.cpp
@@ -9,15 +9,6 @@
 
 #include UE_INLINE_GENERATED_CPP_BY_NAME(PCGDensityRemapElement)
 
-TArray<FPCGPinProperties> UPCGDensityRemapSettings::InputPinProperties() const
-{
-	TArray<FPCGPinProperties> PinProperties;
-	// TODO in the future type checking of edges will be stricter and a conversion node will be added to convert from other types
-	PinProperties.Emplace(PCGPinConstants::DefaultInputLabel, EPCGDataType::Point);
-
-	return PinProperties;
-}
-
 FPCGElementPtr UPCGDensityRemapSettings::CreateElement() const
 {
 	return MakeShared<FPCGDensityRemapElement>();
diff --git a/Engine/Plugins/Experimental/PCG/Source/PCG/Private/Elements/PCGPointExtentsModifier.cpp b/Engine/Plugins/Experimental/PCG/Source/PCG/Private/Elements/PCGPointExtentsModifier.cpp
--- a/Engine/Plugins/Experimental/PCG/Source/PCG/Private/Elements/PCGPointExtentsModifier.cpp
+++ b/Engine/Plugins/Experimental/PCG/Source/PCG/Private/Elements/PCGPointExtentsModifier.cpp
@@ -12,13 +12,6 @@ FPCGElementPtr UPCGPointExtentsModifierSettings::CreateElement() const
 	return MakeShared<FPCGPointExtentsModifier>();
 }
 
-TArray<FPCGPinProperties> UPCGPointExtentsModifierSettings::InputPinProperties() const
-{
-	TArray<FPCGPinProperties> PinProperties;
-	PinProperties.Emplace(PCGPinConstants::DefaultInputLabel, EPCGDataType::Point);	
-	return PinProperties;
-}
-
 bool FPCGPointExtentsModifier::ExecuteInternal(FPCGContext* Context) const
 {
 	TRACE_CPUPROFILER_EVENT_SCOPE(FPCGPointExtentsModifier::Execute);
diff --git a/Engine/Plugins/Experimental/PCG/Source/PCG/Private/Elements/PCGTextureSampler.cpp b/Engine/Plugins/Experimental/PCG/Source/PCG/Private/Elements/PCGTextureSampler.cpp
--- a/Engine/Plugins/Experimental/PCG/Source/PCG/Private/Elements/PCGTextureSampler.cpp
+++ b/Engine/Plugins/Experimental/PCG/Source/PCG/Private/Elements/PCGTextureSampler.cpp
@@ -19,11 +19,6 @@ FPCGElementPtr UPCGTextureSamplerSettings::CreateElement() const
 	return MakeShared<FPCGTextureSamplerElement>();
 }
 
-TArray<FPCGPinProperties> UPCGTextureSamplerSettings::InputPinProperties() const
-{
-	return TArray<FPCGPinProperties>();
-}
-
 bool FPCGTextureSamplerElement::ExecuteInternal(FPCGContext* Context) const
 {
 	TRACE_CPUPROFILER_EVENT_SCOPE(FPCGTextureSamplerElement::Execute);
diff --git a/Engine/Plugins/Experimental/PCG/Source/PCG/Public/Elements/Metadata/PCGMetadataElement.h b/Engine/Plugins/Experimental/PCG/Source/PCG/Public/Elements/Metadata/PCGMetadataElement.h
--- a/Engine/Plugins/Experimental/PCG/Source/PCG/Public/Elements/Metadata/PCGMetadataElement.h
+++ b/Engine/Plugins/Experimental/PCG/Source/PCG/Public/Elements/Metadata/PCGMetadataElement.h
@@ -30,13 +30,13 @@ public:
 
 	//~Begin UPCGSettings interface
 #if WITH_EDITOR
-	virtual FName GetDefaultNodeName() const override { return FName(TEXT("AttributeOperation")); }
+	virtual FName GetDefaultNodeName() const override { return FName(TEXT("Attribute Operation")); }
 	virtual EPCGSettingsType GetType() const override { return EPCGSettingsType::Metadata; }
 #endif
 
-	virtual TArray<FPCGPinProperties> OutputPinProperties() const override { return Super::DefaultPointOutputPinProperties(); }
-
 protected:
+	virtual TArray<FPCGPinProperties> InputPinProperties() const override { return Super::DefaultPointInputPinProperties(); }
+	virtual TArray<FPCGPinProperties> OutputPinProperties() const override { return Super::DefaultPointOutputPinProperties(); }
 	virtual FPCGElementPtr CreateElement() const override;
 	//~End UPCGSettings interface
 
diff --git a/Engine/Plugins/Experimental/PCG/Source/PCG/Public/Elements/Metadata/PCGMetadataPartition.h b/Engine/Plugins/Experimental/PCG/Source/PCG/Public/Elements/Metadata/PCGMetadataPartition.h
--- a/Engine/Plugins/Experimental/PCG/Source/PCG/Public/Elements/Metadata/PCGMetadataPartition.h
+++ b/Engine/Plugins/Experimental/PCG/Source/PCG/Public/Elements/Metadata/PCGMetadataPartition.h
@@ -15,13 +15,13 @@ class PCG_API UPCGMetadataPartitionSettings : public UPCGSettings
 public:
 	//~Begin UPCGSettings interface
 #if WITH_EDITOR
-	virtual FName GetDefaultNodeName() const override { return FName(TEXT("AttributePartition")); }
+	virtual FName GetDefaultNodeName() const override { return FName(TEXT("Attribute Partition")); }
 	virtual EPCGSettingsType GetType() const override { return EPCGSettingsType::Metadata; }
 #endif
 
-	virtual TArray<FPCGPinProperties> OutputPinProperties() const override { return Super::DefaultPointOutputPinProperties(); }
-
 protected:
+	virtual TArray<FPCGPinProperties> InputPinProperties() const override { return Super::DefaultPointInputPinProperties(); }
+	virtual TArray<FPCGPinProperties> OutputPinProperties() const override { return Super::DefaultPointOutputPinProperties(); }
 	virtual FPCGElementPtr CreateElement() const override;
 	//~End UPCGSettings interface
 
diff --git a/Engine/Plugins/Experimental/PCG/Source/PCG/Public/Elements/PCGBoundsModifier.h b/Engine/Plugins/Experimental/PCG/Source/PCG/Public/Elements/PCGBoundsModifier.h
--- a/Engine/Plugins/Experimental/PCG/Source/PCG/Public/Elements/PCGBoundsModifier.h
+++ b/Engine/Plugins/Experimental/PCG/Source/PCG/Public/Elements/PCGBoundsModifier.h
@@ -31,15 +31,14 @@ public:
 
 	// ~Begin UPCGSettings interface
 #if WITH_EDITOR
-	virtual FName GetDefaultNodeName() const override { return FName(TEXT("BoundsModifier")); }
+	virtual FName GetDefaultNodeName() const override { return FName(TEXT("Bounds Modifier")); }
 	virtual FText GetNodeTooltipText() const override;
 	virtual EPCGSettingsType GetType() const override { return EPCGSettingsType::Spatial; }
 #endif
 
-	virtual TArray<FPCGPinProperties> InputPinProperties() const override;
-	virtual TArray<FPCGPinProperties> OutputPinProperties() const override { return Super::DefaultPointOutputPinProperties(); }
-
 protected:
+	virtual TArray<FPCGPinProperties> InputPinProperties() const override { return Super::DefaultPointInputPinProperties(); }
+	virtual TArray<FPCGPinProperties> OutputPinProperties() const override { return Super::DefaultPointOutputPinProperties(); }
 	virtual FPCGElementPtr CreateElement() const override;
 	//~End UPCGSettings interface
 
diff --git a/Engine/Plugins/Experimental/PCG/Source/PCG/Public/Elements/PCGDensityNoise.h b/Engine/Plugins/Experimental/PCG/Source/PCG/Public/Elements/PCGDensityNoise.h
--- a/Engine/Plugins/Experimental/PCG/Source/PCG/Public/Elements/PCGDensityNoise.h
+++ b/Engine/Plugins/Experimental/PCG/Source/PCG/Public/Elements/PCGDensityNoise.h
@@ -29,14 +29,13 @@ public:
 
 	// ~Begin UPCGSettings interface
 #if WITH_EDITOR
-	virtual FName GetDefaultNodeName() const override { return FName(TEXT("DensityNoise")); }
+	virtual FName GetDefaultNodeName() const override { return FName(TEXT("Density Noise")); }
 	virtual EPCGSettingsType GetType() const override { return EPCGSettingsType::Density; }
 #endif
 
-	virtual TArray<FPCGPinProperties> InputPinProperties() const override;
-	virtual TArray<FPCGPinProperties> OutputPinProperties() const override { return Super::DefaultPointOutputPinProperties(); }
-
 protected:
+	virtual TArray<FPCGPinProperties> InputPinProperties() const override { return Super::DefaultPointInputPinProperties(); }
+	virtual TArray<FPCGPinProperties> OutputPinProperties() const override { return Super::DefaultPointOutputPinProperties(); }
 	virtual FPCGElementPtr CreateElement() const override;
 	//~End UPCGSettings interface
 
diff --git a/Engine/Plugins/Experimental/PCG/Source/PCG/Public/Elements/PCGDensityRemapElement.h b/Engine/Plugins/Experimental/PCG/Source/PCG/Public/Elements/PCGDensityRemapElement.h
--- a/Engine/Plugins/Experimental/PCG/Source/PCG/Public/Elements/PCGDensityRemapElement.h
+++ b/Engine/Plugins/Experimental/PCG/Source/PCG/Public/Elements/PCGDensityRemapElement.h
@@ -17,14 +17,13 @@ class PCG_API UPCGDensityRemapSettings : public UPCGSettings
 public:
 	// ~Begin UPCGSettings interface
 #if WITH_EDITOR
-	virtual FName GetDefaultNodeName() const override { return FName(TEXT("DensityRemap")); }
+	virtual FName GetDefaultNodeName() const override { return FName(TEXT("Density Remap")); }
 	virtual EPCGSettingsType GetType() const override { return EPCGSettingsType::Density; }
 #endif
 
-	virtual TArray<FPCGPinProperties> InputPinProperties() const override;
-	virtual TArray<FPCGPinProperties> OutputPinProperties() const override { return Super::DefaultPointOutputPinProperties(); }
-
 protected:
+	virtual TArray<FPCGPinProperties> InputPinProperties() const override { return Super::DefaultPointInputPinProperties(); }
+	virtual TArray<FPCGPinProperties> OutputPinProperties() const override { return Super::DefaultPointOutputPinProperties(); }
 	virtual FPCGElementPtr CreateElement() const override;
 	//~End UPCGSettings interface
 
diff --git a/Engine/Plugins/Experimental/PCG/Source/PCG/Public/Elements/PCGPointExtentsModifier.h b/Engine/Plugins/Experimental/PCG/Source/PCG/Public/Elements/PCGPointExtentsModifier.h
--- a/Engine/Plugins/Experimental/PCG/Source/PCG/Public/Elements/PCGPointExtentsModifier.h
+++ b/Engine/Plugins/Experimental/PCG/Source/PCG/Public/Elements/PCGPointExtentsModifier.h
@@ -28,14 +28,13 @@ public:
 
 	// ~Begin UPCGSettings interface
 #if WITH_EDITOR
-	virtual FName GetDefaultNodeName() const override { return FName(TEXT("ExtentsModifier")); }
+	virtual FName GetDefaultNodeName() const override { return FName(TEXT("Extents Modifier")); }
 	virtual EPCGSettingsType GetType() const override { return EPCGSettingsType::Spatial; }
 #endif
 
-	virtual TArray<FPCGPinProperties> InputPinProperties() const override;
-	virtual TArray<FPCGPinProperties> OutputPinProperties() const override { return Super::DefaultPointOutputPinProperties(); }
-
 protected:
+	virtual TArray<FPCGPinProperties> InputPinProperties() const override { return Super::DefaultPointInputPinProperties(); }
+	virtual TArray<FPCGPinProperties> OutputPinProperties() const override { return Super::DefaultPointOutputPinProperties(); }
 	virtual FPCGElementPtr CreateElement() const override;
 	//~End UPCGSettings interface
 
diff --git a/Engine/Plugins/Experimental/PCG/Source/PCG/Public/Elements/PCGPointSampler.h b/Engine/Plugins/Experimental/PCG/Source/PCG/Public/Elements/PCGPointSampler.h
--- a/Engine/Plugins/Experimental/PCG/Source/PCG/Public/Elements/PCGPointSampler.h
+++ b/Engine/Plugins/Experimental/PCG/Source/PCG/Public/Elements/PCGPointSampler.h
@@ -17,13 +17,13 @@ public:
 
 	//~Begin UPCGSettings interface
 #if WITH_EDITOR
-	virtual FName GetDefaultNodeName() const override { return FName(TEXT("PointSampler")); }
+	virtual FName GetDefaultNodeName() const override { return FName(TEXT("Point Sampler")); }
 	virtual EPCGSettingsType GetType() const override { return EPCGSettingsType::Sampler; }
 #endif
 
-	virtual TArray<FPCGPinProperties> OutputPinProperties() const override { return Super::DefaultPointOutputPinProperties(); }
-
 protected:
+	virtual TArray<FPCGPinProperties> InputPinProperties() const override { return Super::DefaultPointInputPinProperties(); }
+	virtual TArray<FPCGPinProperties> OutputPinProperties() const override { return Super::DefaultPointOutputPinProperties(); }
 	virtual FPCGElementPtr CreateElement() const override;
 	//~End UPCGSettings interface
 
diff --git a/Engine/Plugins/Experimental/PCG/Source/PCG/Public/Elements/PCGTextureSampler.h b/Engine/Plugins/Experimental/PCG/Source/PCG/Public/Elements/PCGTextureSampler.h
--- a/Engine/Plugins/Experimental/PCG/Source/PCG/Public/Elements/PCGTextureSampler.h
+++ b/Engine/Plugins/Experimental/PCG/Source/PCG/Public/Elements/PCGTextureSampler.h
@@ -16,14 +16,13 @@ class PCG_API UPCGTextureSamplerSettings : public UPCGSettings
 public:
 	//~Begin UPCGSettings interface
 #if WITH_EDITOR
-	virtual FName GetDefaultNodeName() const override { return FName(TEXT("TextureSampler")); }
+	virtual FName GetDefaultNodeName() const override { return FName(TEXT("Texture Sampler")); }
 	virtual EPCGSettingsType GetType() const override { return EPCGSettingsType::Sampler; }
 #endif
 
-	virtual TArray<FPCGPinProperties> InputPinProperties() const override;
-	virtual TArray<FPCGPinProperties> OutputPinProperties() const override { return Super::DefaultPointOutputPinProperties(); }
-
 protected:
+	virtual TArray<FPCGPinProperties> InputPinProperties() const override { return TArray<FPCGPinProperties>(); }
+	virtual TArray<FPCGPinProperties> OutputPinProperties() const override { return Super::DefaultPointOutputPinProperties(); }
 	virtual FPCGElementPtr CreateElement() const override;
 	//~End UPCGSettings interface
 
