Added initialization on member variables in PCGSplineStruct

#rnx

diff --git a/Engine/Plugins/Experimental/PCG/Source/PCG/Public/Data/PCGSplineStruct.h b/Engine/Plugins/Experimental/PCG/Source/PCG/Public/Data/PCGSplineStruct.h
--- a/Engine/Plugins/Experimental/PCG/Source/PCG/Public/Data/PCGSplineStruct.h
+++ b/Engine/Plugins/Experimental/PCG/Source/PCG/Public/Data/PCGSplineStruct.h
@@ -49,20 +49,20 @@ struct PCG_API FPCGSplineStruct
 
 	// Replaces the component trasnform
 	UPROPERTY()
-	FTransform Transform;
+	FTransform Transform = FTransform::Identity;
 
 	UPROPERTY()
-	FVector DefaultUpVector;
+	FVector DefaultUpVector = FVector::UpVector;
 
 	UPROPERTY()
-	int32 ReparamStepsPerSegment;
+	int32 ReparamStepsPerSegment = 10;
 
 	UPROPERTY()
-	bool bClosedLoop;
+	bool bClosedLoop = false;
 
 	UPROPERTY()
-	FBoxSphereBounds LocalBounds;
+	FBoxSphereBounds LocalBounds = FBoxSphereBounds(EForceInit::ForceInit);
 
 	UPROPERTY()
-	FBoxSphereBounds Bounds;
+	FBoxSphereBounds Bounds = FBoxSphereBounds(EForceInit::ForceInit);
 };
\ No newline at end of file
