[PCG] Fixed a crash when redoing pin-change after GC by saving the pins to the transaction buffer.

diff --git a/Engine/Plugins/Experimental/PCG/Source/PCG/Private/PCGNode.cpp b/Engine/Plugins/Experimental/PCG/Source/PCG/Private/PCGNode.cpp
--- a/Engine/Plugins/Experimental/PCG/Source/PCG/Private/PCGNode.cpp
+++ b/Engine/Plugins/Experimental/PCG/Source/PCG/Private/PCGNode.cpp
@@ -740,6 +740,7 @@ EPCGChangeType UPCGNode::UpdatePins(TFunctionRef<UPCGPin*(UPCGNode*)> PinAllocat
 
 					const int32 InsertIndex = FMath::Min(PinProperties.IndexOfByKey(UnmatchedProperty), Pins.Num());
 					UPCGPin* NewPin = PinAllocator(this);
+					NewPin->Modify();
 					NewPin->Node = this;
 					NewPin->Properties = UnmatchedProperty;
 					Pins.Insert(NewPin, InsertIndex);
