[PCG] Fixed undo crash issue with PCG BP actors after GC

diff --git a/Engine/Plugins/Experimental/PCG/Source/PCG/Private/PCGComponent.cpp b/Engine/Plugins/Experimental/PCG/Source/PCG/Private/PCGComponent.cpp
--- a/Engine/Plugins/Experimental/PCG/Source/PCG/Private/PCGComponent.cpp
+++ b/Engine/Plugins/Experimental/PCG/Source/PCG/Private/PCGComponent.cpp
@@ -2877,4 +2877,12 @@ void FPCGComponentInstanceData::ApplyToComponent(UActorComponent* Component, con
 	}
 }
 
+void FPCGComponentInstanceData::AddReferencedObjects(FReferenceCollector& Collector)
+{
+	Super::AddReferencedObjects(Collector);
+
+	Collector.AddReferencedObject(SourceComponent);
+	Collector.AddReferencedObjects(GeneratedResources);
+}
+
 #undef LOCTEXT_NAMESPACE
diff --git a/Engine/Plugins/Experimental/PCG/Source/PCG/Public/PCGComponent.h b/Engine/Plugins/Experimental/PCG/Source/PCG/Public/PCGComponent.h
--- a/Engine/Plugins/Experimental/PCG/Source/PCG/Public/PCGComponent.h
+++ b/Engine/Plugins/Experimental/PCG/Source/PCG/Public/PCGComponent.h
@@ -469,6 +469,7 @@ public:
 protected:
 	virtual bool ContainsData() const override;
 	virtual void ApplyToComponent(UActorComponent* Component, const ECacheApplyPhase CacheApplyPhase) override;
+	virtual void AddReferencedObjects(FReferenceCollector& Collector) override;
 
 	UPROPERTY()
 	TArray<TObjectPtr<UPCGManagedResource>> GeneratedResources;
