[PCG] Add console command to clear landscape cache

#rnx

diff --git a/Engine/Plugins/Experimental/PCG/Source/PCG/Private/PCGSubsystem.cpp b/Engine/Plugins/Experimental/PCG/Source/PCG/Private/PCGSubsystem.cpp
--- a/Engine/Plugins/Experimental/PCG/Source/PCG/Private/PCGSubsystem.cpp
+++ b/Engine/Plugins/Experimental/PCG/Source/PCG/Private/PCGSubsystem.cpp
@@ -75,6 +75,18 @@ namespace PCGSubsystemConsole
 				}
 			}));
 
+	static FAutoConsoleCommand CommandClearLandscapeCache(
+		TEXT("pcg.ClearLandscapeCache"),
+		TEXT("Clear the landscape cache in the current world."),
+		FConsoleCommandDelegate::CreateLambda([]()
+			{
+				UWorld* World = (GEditor ? (GEditor->PlayWorld ? GEditor->PlayWorld.Get() : GEditor->GetEditorWorldContext().World()) : (GEngine ? GEngine->GetCurrentPlayWorld() : nullptr));
+				if(World && World->GetSubsystem<UPCGSubsystem>())
+				{
+					World->GetSubsystem<UPCGSubsystem>()->ClearLandscapeCache();
+				}
+			}));
+
 	static TAutoConsoleVariable<bool> CVarRebuildLandscapeOnPIE(
 		TEXT("pcg.PIE.RebuildLandscapeOnPIE"),
 		true,
