[PCG] Cleaned up parsed actor data when there is a spline-like component

diff --git a/Engine/Plugins/Experimental/PCG/Source/PCG/Private/PCGComponent.cpp b/Engine/Plugins/Experimental/PCG/Source/PCG/Private/PCGComponent.cpp
--- a/Engine/Plugins/Experimental/PCG/Source/PCG/Private/PCGComponent.cpp
+++ b/Engine/Plugins/Experimental/PCG/Source/PCG/Private/PCGComponent.cpp
@@ -31,8 +31,9 @@
 #include "Algo/Transform.h"
 #include "Components/BillboardComponent.h"
 #include "Components/InstancedStaticMeshComponent.h"
-#include "Components/SplineComponent.h"
 #include "Components/ShapeComponent.h"
+#include "Components/SplineComponent.h"
+#include "Components/SplineMeshComponent.h"
 #include "Engine/Engine.h"
 #include "GameFramework/Volume.h"
 #include "Kismet/GameplayStatics.h"
@@ -2187,7 +2188,8 @@ FPCGDataCollection UPCGComponent::CreateActorPCGDataCollection(AActor* Actor, co
 	}
 	else // Prepare data on a component basis
 	{
-		TInlineComponentArray<UPrimitiveComponent*, 4> Primitives;
+		using PrimitiveComponentArray = TInlineComponentArray<UPrimitiveComponent*, 4>;
+		PrimitiveComponentArray Primitives;
 
 		auto RemoveDuplicatesFromPrimitives = [&Primitives](const auto& InComponents)
 		{
@@ -2208,6 +2210,17 @@ FPCGDataCollection UPCGComponent::CreateActorPCGDataCollection(AActor* Actor, co
 			}
 		};
 
+		auto RemoveSplineMeshComponents = [](PrimitiveComponentArray& InComponents)
+		{
+			for (int32 Index = InComponents.Num() - 1; Index >= 0; --Index)
+			{
+				if (InComponents[Index]->IsA<USplineMeshComponent>())
+				{
+					InComponents.RemoveAtSwap(Index);
+				}
+			}
+		};
+
 		Actor->GetComponents(Primitives);
 		RemovePCGGeneratedEntries(Primitives);
 
@@ -2221,6 +2234,12 @@ FPCGDataCollection UPCGComponent::CreateActorPCGDataCollection(AActor* Actor, co
 		RemovePCGGeneratedEntries(Splines);
 		RemoveDuplicatesFromPrimitives(Splines);
 
+		// If we have a better representation than the spline mesh components, we shouldn't create them
+		if (!LandscapeSplines.IsEmpty() || !Splines.IsEmpty())
+		{
+			RemoveSplineMeshComponents(Primitives);
+		}
+
 		TInlineComponentArray<UShapeComponent*, 4> Shapes;
 		Actor->GetComponents(Shapes);
 		RemovePCGGeneratedEntries(Shapes);
