[PCG] Removed sharing of ISM component resources from different nodes as it can cause issues with ISM reuse and also break ordering of execution

diff --git a/Engine/Plugins/Experimental/PCG/Source/PCG/Private/Elements/PCGDebugElement.cpp b/Engine/Plugins/Experimental/PCG/Source/PCG/Private/Elements/PCGDebugElement.cpp
--- a/Engine/Plugins/Experimental/PCG/Source/PCG/Private/Elements/PCGDebugElement.cpp
+++ b/Engine/Plugins/Experimental/PCG/Source/PCG/Private/Elements/PCGDebugElement.cpp
@@ -177,7 +177,7 @@ namespace PCGDebugElement
 					continue;
 				}
 
-				UInstancedStaticMeshComponent* ISMC = UPCGActorHelpers::GetOrCreateISMC(TargetActor, Context->SourceComponent.Get(), Params[Direction]);
+				UInstancedStaticMeshComponent* ISMC = UPCGActorHelpers::GetOrCreateISMC(TargetActor, Context->SourceComponent.Get(), SettingsInterface->GetSettings()->UID, Params[Direction]);
 				check(ISMC && ISMC->NumCustomDataFloats == NumCustomData);
 
 				ISMC->ComponentTags.AddUnique(PCGHelpers::DefaultPCGDebugTag);
diff --git a/Engine/Plugins/Experimental/PCG/Source/PCG/Private/Elements/PCGSpawnActor.cpp b/Engine/Plugins/Experimental/PCG/Source/PCG/Private/Elements/PCGSpawnActor.cpp
--- a/Engine/Plugins/Experimental/PCG/Source/PCG/Private/Elements/PCGSpawnActor.cpp
+++ b/Engine/Plugins/Experimental/PCG/Source/PCG/Private/Elements/PCGSpawnActor.cpp
@@ -478,7 +478,7 @@ bool FPCGSpawnActorElement::ExecuteInternal(FPCGContext* Context) const
 			{
 				const FPCGISMCBuilderParameters& ISMCParams = ISMCBuilderTransforms.Key;
 
-				UPCGManagedISMComponent* MISMC = UPCGActorHelpers::GetOrCreateManagedISMC(TargetActor, Context->SourceComponent.Get(), ISMCParams);
+				UPCGManagedISMComponent* MISMC = UPCGActorHelpers::GetOrCreateManagedISMC(TargetActor, Context->SourceComponent.Get(), Settings->UID, ISMCParams);
 				if (!MISMC)
 				{
 					continue;
diff --git a/Engine/Plugins/Experimental/PCG/Source/PCG/Private/Elements/PCGStaticMeshSpawner.cpp b/Engine/Plugins/Experimental/PCG/Source/PCG/Private/Elements/PCGStaticMeshSpawner.cpp
--- a/Engine/Plugins/Experimental/PCG/Source/PCG/Private/Elements/PCGStaticMeshSpawner.cpp
+++ b/Engine/Plugins/Experimental/PCG/Source/PCG/Private/Elements/PCGStaticMeshSpawner.cpp
@@ -83,11 +83,11 @@ bool FPCGStaticMeshSpawnerElement::PrepareDataInternal(FPCGContext* InContext) c
 		if (Context->DependenciesCrc.IsValid())
 		{
 			TArray<UPCGManagedISMComponent*> MISMCs;
-			Context->SourceComponent->ForEachManagedResource([&MISMCs, &Context](UPCGManagedResource* InResource)
+			Context->SourceComponent->ForEachManagedResource([&MISMCs, &Context, Settings](UPCGManagedResource* InResource)
 			{
 				if (UPCGManagedISMComponent* Resource = Cast<UPCGManagedISMComponent>(InResource))
 				{
-					if (Resource->GetCrc().IsValid() && Resource->GetCrc() == Context->DependenciesCrc)
+					if (Resource->GetSettingsUID() == Settings->UID && Resource->GetCrc().IsValid() && Resource->GetCrc() == Context->DependenciesCrc)
 					{
 						MISMCs.Add(Resource);
 					}
@@ -330,7 +330,9 @@ void FPCGStaticMeshSpawnerElement::SpawnStaticMeshInstances(FPCGContext* Context
 		Params.Descriptor.Mobility = SceneComponent->Mobility;
 	}
 
-	UPCGManagedISMComponent* MISMC = UPCGActorHelpers::GetOrCreateManagedISMC(TargetActor, Context->SourceComponent.Get(), Params);
+	const UPCGStaticMeshSpawnerSettings* Settings = Context->GetInputSettings<UPCGStaticMeshSpawnerSettings>();
+	check(Settings);
+	UPCGManagedISMComponent* MISMC = UPCGActorHelpers::GetOrCreateManagedISMC(TargetActor, Context->SourceComponent.Get(), Settings->UID, Params);
 	
 	check(MISMC);
 	MISMC->SetCrc(Context->DependenciesCrc);
diff --git a/Engine/Plugins/Experimental/PCG/Source/PCG/Private/Helpers/PCGActorHelpers.cpp b/Engine/Plugins/Experimental/PCG/Source/PCG/Private/Helpers/PCGActorHelpers.cpp
--- a/Engine/Plugins/Experimental/PCG/Source/PCG/Private/Helpers/PCGActorHelpers.cpp
+++ b/Engine/Plugins/Experimental/PCG/Source/PCG/Private/Helpers/PCGActorHelpers.cpp
@@ -19,9 +19,9 @@
 #if WITH_EDITOR
 #endif
 
-UInstancedStaticMeshComponent* UPCGActorHelpers::GetOrCreateISMC(AActor* InTargetActor, UPCGComponent* InSourceComponent, const FPCGISMCBuilderParameters& InParams)
+UInstancedStaticMeshComponent* UPCGActorHelpers::GetOrCreateISMC(AActor* InTargetActor, UPCGComponent* InSourceComponent, uint64 SettingsUID, const FPCGISMCBuilderParameters& InParams)
 {
-	UPCGManagedISMComponent* MISMC = GetOrCreateManagedISMC(InTargetActor, InSourceComponent, InParams);
+	UPCGManagedISMComponent* MISMC = GetOrCreateManagedISMC(InTargetActor, InSourceComponent, SettingsUID, InParams);
 	if (MISMC)
 	{
 		return MISMC->GetComponent();
@@ -32,7 +32,7 @@ UInstancedStaticMeshComponent* UPCGActorHelpers::GetOrCreateISMC(AActor* InTarge
 	}
 }
 
-UPCGManagedISMComponent* UPCGActorHelpers::GetOrCreateManagedISMC(AActor* InTargetActor, UPCGComponent* InSourceComponent, const FPCGISMCBuilderParameters& InParams)
+UPCGManagedISMComponent* UPCGActorHelpers::GetOrCreateManagedISMC(AActor* InTargetActor, UPCGComponent* InSourceComponent, uint64 SettingsUID, const FPCGISMCBuilderParameters& InParams)
 {
 	check(InTargetActor && InSourceComponent);
 
@@ -47,7 +47,7 @@ UPCGManagedISMComponent* UPCGActorHelpers::GetOrCreateManagedISMC(AActor* InTarg
 	{
 		TRACE_CPUPROFILER_EVENT_SCOPE(UPCGActorHelpers::GetOrCreateManagedISMC::FindMatchingMISMC);
 		UPCGManagedISMComponent* MatchingResource = nullptr;
-		InSourceComponent->ForEachManagedResource([&MatchingResource, &InParams](UPCGManagedResource* InResource)
+		InSourceComponent->ForEachManagedResource([&MatchingResource, &InParams, SettingsUID](UPCGManagedResource* InResource)
 		{
 			// Early out if already found a match
 			if (MatchingResource)
@@ -57,6 +57,11 @@ UPCGManagedISMComponent* UPCGActorHelpers::GetOrCreateManagedISMC(AActor* InTarg
 
 			if (UPCGManagedISMComponent* Resource = Cast<UPCGManagedISMComponent>(InResource))
 			{
+				if (Resource->GetSettingsUID() != SettingsUID)
+				{
+					return;
+				}
+
 				if (UInstancedStaticMeshComponent* ISMC = Resource->GetComponent())
 				{
 					if (ISMC->NumCustomDataFloats == InParams.NumCustomDataFloats &&
@@ -116,6 +121,7 @@ UPCGManagedISMComponent* UPCGActorHelpers::GetOrCreateManagedISMC(AActor* InTarg
 	UPCGManagedISMComponent* Resource = NewObject<UPCGManagedISMComponent>(InSourceComponent);
 	Resource->SetComponent(ISMC);
 	Resource->SetDescriptor(InParams.Descriptor);
+	Resource->SetSettingsUID(SettingsUID);
 	InSourceComponent->AddToManagedResources(Resource);
 
 	return Resource;
diff --git a/Engine/Plugins/Experimental/PCG/Source/PCG/Public/Helpers/PCGActorHelpers.h b/Engine/Plugins/Experimental/PCG/Source/PCG/Public/Helpers/PCGActorHelpers.h
--- a/Engine/Plugins/Experimental/PCG/Source/PCG/Public/Helpers/PCGActorHelpers.h
+++ b/Engine/Plugins/Experimental/PCG/Source/PCG/Public/Helpers/PCGActorHelpers.h
@@ -38,8 +38,8 @@ class PCG_API UPCGActorHelpers : public UBlueprintFunctionLibrary
 	GENERATED_BODY()
 
 public:
-	static UInstancedStaticMeshComponent* GetOrCreateISMC(AActor* InTargetActor, UPCGComponent* SourceComponent, const FPCGISMCBuilderParameters& Params);
-	static UPCGManagedISMComponent* GetOrCreateManagedISMC(AActor* InTargetActor, UPCGComponent* SourceComponent, const FPCGISMCBuilderParameters& Params);
+	static UInstancedStaticMeshComponent* GetOrCreateISMC(AActor* InTargetActor, UPCGComponent* SourceComponent, uint64 SettingsUID, const FPCGISMCBuilderParameters& Params);
+	static UPCGManagedISMComponent* GetOrCreateManagedISMC(AActor* InTargetActor, UPCGComponent* SourceComponent, uint64 SettingsUID, const FPCGISMCBuilderParameters& Params);
 	static bool DeleteActors(UWorld* World, const TArray<TSoftObjectPtr<AActor>>& ActorsToDelete);
 
 	/**
diff --git a/Engine/Plugins/Experimental/PCG/Source/PCG/Public/PCGManagedResource.h b/Engine/Plugins/Experimental/PCG/Source/PCG/Public/PCGManagedResource.h
--- a/Engine/Plugins/Experimental/PCG/Source/PCG/Public/PCGManagedResource.h
+++ b/Engine/Plugins/Experimental/PCG/Source/PCG/Public/PCGManagedResource.h
@@ -127,6 +127,9 @@ public:
 	void SetDescriptor(const FISMComponentDescriptor& InDescriptor);
 	const FISMComponentDescriptor& GetDescriptor() const { return Descriptor; }
 
+	uint64 GetSettingsUID() const { return SettingsUID; }
+	void SetSettingsUID(uint64 InSettingsUID) { SettingsUID = InSettingsUID; }
+
 protected:
 	UPROPERTY()
 	bool bHasDescriptor = false;
@@ -134,6 +137,9 @@ protected:
 	UPROPERTY()
 	FISMComponentDescriptor Descriptor;
 
+	UPROPERTY(Transient)
+	uint64 SettingsUID = -1; // purposefully a value that will never happen in data
+
 	// Cached raw pointer to ISM component
 	mutable UInstancedStaticMeshComponent* CachedRawComponentPtr = nullptr;
 };
