[PCG] Added overrides in the world queries, added filter by tag and better controls for landscape

Also added fallback to actor bounds in volume sampler, and 'unbounded' option similar to surface sampler.

diff --git a/Engine/Plugins/Experimental/PCG/Source/PCG/Private/Data/PCGWorldData.cpp b/Engine/Plugins/Experimental/PCG/Source/PCG/Private/Data/PCGWorldData.cpp
--- a/Engine/Plugins/Experimental/PCG/Source/PCG/Private/Data/PCGWorldData.cpp
+++ b/Engine/Plugins/Experimental/PCG/Source/PCG/Private/Data/PCGWorldData.cpp
@@ -2,20 +2,50 @@
 
 #include "Data/PCGWorldData.h"
 
-#include "Data/PCGSurfaceData.h"
 #include "PCGComponent.h"
-#include "Data/PCGVolumeData.h"
 #include "PCGHelpers.h"
+#include "PCGSubsystem.h"
 #include "Data/PCGPointData.h"
+#include "Data/PCGSurfaceData.h"
+#include "Data/PCGVolumeData.h"
 #include "Elements/PCGSurfaceSampler.h"
 #include "Elements/PCGVolumeSampler.h"
+#include "Grid/PCGLandscapeCache.h"
 #include "Helpers/PCGBlueprintHelpers.h"
 
+#include "LandscapeProxy.h"
 #include "Components/BrushComponent.h"
 #include "Engine/World.h"
 
 #include UE_INLINE_GENERATED_CPP_BY_NAME(PCGWorldData)
 
+void FPCGWorldCommonQueryParams::Initialize()
+{
+	if (ActorTagFilter == EPCGWorldQueryFilterByTag::NoTagFilter)
+	{
+		ParsedActorTagsList.Reset();
+	}
+	else
+	{
+		TArray<FString> ParsedList = PCGHelpers::GetStringArrayFromCommaSeparatedString(ActorTagsList);
+		ParsedActorTagsList.Reset();
+		for (const FString& Tag : ParsedList)
+		{
+			ParsedActorTagsList.Add(FName(Tag));
+		}
+	}
+}
+
+void FPCGWorldVolumetricQueryParams::Initialize()
+{
+	FPCGWorldCommonQueryParams::Initialize();
+}
+
+void FPCGWorldRayHitQueryParams::Initialize()
+{
+	FPCGWorldCommonQueryParams::Initialize();
+}
+
 void UPCGWorldVolumetricData::Initialize(UWorld* InWorld, const FBox& InBounds)
 {
 	Super::Initialize(InBounds, nullptr);
@@ -71,7 +101,36 @@ bool UPCGWorldVolumetricData::SamplePoint(const FTransform& InTransform, const F
 			continue;
 		}
 
-		// TODO: additional filtering?
+		// Additional filter as provided in the QueryParams base class
+		if (QueryParams.ActorTagFilter != EPCGWorldQueryFilterByTag::NoTagFilter)
+		{
+			if (AActor* Actor = OverlappedComponent->GetOwner())
+			{
+				bool bFoundMatch = false;
+				for (const FName& Tag : Actor->Tags)
+				{
+					if (QueryParams.ParsedActorTagsList.Contains(Tag))
+					{
+						bFoundMatch = true;
+						break;
+					}
+				}
+
+				if (bFoundMatch != (QueryParams.ActorTagFilter == EPCGWorldQueryFilterByTag::IncludeTagged))
+				{
+					continue;
+				}
+			}
+			else if (QueryParams.ActorTagFilter == EPCGWorldQueryFilterByTag::IncludeTagged)
+			{
+				continue;
+			}
+		}
+
+		if (QueryParams.bIgnoreLandscapeHits && OverlappedComponent->GetOwner() && OverlappedComponent->GetOwner()->IsA<ALandscapeProxy>())
+		{
+			continue;
+		}
 
 		if (QueryParams.bSearchForOverlap)
 		{
@@ -150,6 +209,7 @@ UPCGSpatialData* UPCGWorldVolumetricData::CopyInternal() const
 	NewVolumetricData->World = World;
 	NewVolumetricData->OriginatingComponent = OriginatingComponent;
 	NewVolumetricData->QueryParams = QueryParams;
+	NewVolumetricData->QueryParams.Initialize();
 
 	return NewVolumetricData;
 }
@@ -164,7 +224,6 @@ void UPCGWorldRayHitData::Initialize(UWorld* InWorld, const FBox& InBounds)
 bool UPCGWorldRayHitData::SamplePoint(const FTransform& InTransform, const FBox& InBounds, FPCGPoint& OutPoint, UPCGMetadata* OutMetadata) const
 {
 	// TODO: This seems to be a projection - along a direction. I suspect that UPCGWorldVolumetricData is the SamplePoint(), and this is the ProjectPoint() (in a direction)?
-
 	check(World.IsValid());
 
 	// Todo: consider prebuilding this
@@ -213,7 +272,44 @@ bool UPCGWorldRayHitData::SamplePoint(const FTransform& InTransform, const FBox&
 			continue;
 		}
 
-		// TODO: additional filtering?
+		// Additional filter as provided in the QueryParams base class
+		if (QueryParams.ActorTagFilter != EPCGWorldQueryFilterByTag::NoTagFilter)
+		{
+			AActor* Actor = HitComponent->GetOwner();
+			
+			if (Actor)
+			{
+				bool bFoundMatch = false;
+				for (const FName& Tag : Actor->Tags)
+				{
+					if (QueryParams.ParsedActorTagsList.Contains(Tag))
+					{
+						bFoundMatch = true;
+						break;
+					}
+				}
+
+				if (bFoundMatch != (QueryParams.ActorTagFilter == EPCGWorldQueryFilterByTag::IncludeTagged))
+				{
+					continue;
+				}
+			}
+			else if (QueryParams.ActorTagFilter == EPCGWorldQueryFilterByTag::IncludeTagged)
+			{
+				continue;
+			}
+		}
+
+		bool bHitOnLandscape = false;
+		if (QueryParams.bIgnoreLandscapeHits || QueryParams.bApplyMetadataFromLandscape)
+		{
+			bHitOnLandscape = HitComponent->GetOwner() && HitComponent->GetOwner()->IsA<ALandscapeProxy>();
+		}
+
+		if (QueryParams.bIgnoreLandscapeHits && bHitOnLandscape)
+		{
+			continue;
+		}
 
 		// Finally, fill in OutPoint - we're done
 		// Implementation note: this uses the same orthonormalization process as the landscape cache
@@ -225,9 +321,15 @@ bool UPCGWorldRayHitData::SamplePoint(const FTransform& InTransform, const FBox&
 		OutPoint = FPCGPoint(FTransform(XAxis, YAxis, Hit.ImpactNormal, Hit.ImpactPoint), 1.0f, 0);
 		UPCGBlueprintHelpers::SetSeedFromPosition(OutPoint);
 
-		// TODO: apply metadata
-		// e.g. if hit on landscape -> metadata on landscape at that position
-		// maybe material, etc.
+		// TODO: generalize for other sources of metadata?
+		if (QueryParams.bApplyMetadataFromLandscape && bHitOnLandscape && OutMetadata && World->GetSubsystem<UPCGSubsystem>())
+		{
+			UPCGLandscapeCache* LandscapeCache = World->GetSubsystem<UPCGSubsystem>()->GetLandscapeCache();
+			if (LandscapeCache)
+			{
+				LandscapeCache->SampleMetadataOnPoint(Cast<ALandscapeProxy>(HitComponent->GetOwner()), OutPoint, OutMetadata);
+			}
+		}
 
 		return true;
 	}
@@ -284,6 +386,7 @@ UPCGSpatialData* UPCGWorldRayHitData::CopyInternal() const
 	NewData->OriginatingComponent = OriginatingComponent;
 	NewData->Bounds = Bounds;
 	NewData->QueryParams = QueryParams;
+	NewData->QueryParams.Initialize();
 
 	return NewData;
 }
diff --git a/Engine/Plugins/Experimental/PCG/Source/PCG/Private/Elements/PCGVolumeSampler.cpp b/Engine/Plugins/Experimental/PCG/Source/PCG/Private/Elements/PCGVolumeSampler.cpp
--- a/Engine/Plugins/Experimental/PCG/Source/PCG/Private/Elements/PCGVolumeSampler.cpp
+++ b/Engine/Plugins/Experimental/PCG/Source/PCG/Private/Elements/PCGVolumeSampler.cpp
@@ -2,6 +2,7 @@
 
 #include "Elements/PCGVolumeSampler.h"
 
+#include "PCGComponent.h"
 #include "PCGContext.h"
 #include "PCGCustomVersion.h"
 #include "PCGHelpers.h"
@@ -179,16 +180,23 @@ bool FPCGVolumeSamplerElement::ExecuteInternal(FPCGContext* Context) const
 	// Grab the Bounding Shape input if there is one.
 	TArray<FPCGTaggedData> BoundingShapeInputs = Context->InputData.GetInputsByPin(PCGVolumeSamplerConstants::BoundingShapeLabel);
 	const UPCGSpatialData* BoundingShapeSpatialInput = nullptr;
-	if (BoundingShapeInputs.Num() > 0)
-	{
-		ensure(BoundingShapeInputs.Num() == 1);
-		BoundingShapeSpatialInput = Cast<UPCGSpatialData>(BoundingShapeInputs[0].Data);
 
-		if (!BoundingShapeSpatialInput)
+	if (!Settings->bUnbounded)
+	{
+		if (BoundingShapeInputs.Num() > 0)
+		{
+			ensure(BoundingShapeInputs.Num() == 1);
+			BoundingShapeSpatialInput = Cast<UPCGSpatialData>(BoundingShapeInputs[0].Data);
+		}
+		else if (Context->SourceComponent.IsValid())
 		{
-			PCGE_LOG(Warning, "Bounding Shape input is missing or of unsupported type and will not be used.");
+			BoundingShapeSpatialInput = Cast<UPCGSpatialData>(Context->SourceComponent->GetActorPCGData());
 		}
 	}
+	else if (BoundingShapeInputs.Num() > 0)
+	{
+		PCGE_LOG_C(Verbose, Context, "The bounds of the Bounding Shape input pin will be ignored because the Unbounded option is enabled.");
+	}	
 
 	// Compute bounds of bounding shape input
 	FBox BoundingShapeBounds(EForceInit::ForceInit);
@@ -212,8 +220,16 @@ bool FPCGVolumeSamplerElement::ExecuteInternal(FPCGContext* Context) const
 	if (GeneratingShapes.Num() == 0 && BoundingShapeSpatialInput)
 	{
 		GeneratingShapes.Add(BoundingShapeSpatialInput);
-		check(BoundingShapeInputs.Num() > 0);
-		Outputs.Add(BoundingShapeInputs[0]);
+
+		// If there was a bounding shape input, use it as the starting point to get the tags
+		if (BoundingShapeInputs.Num() > 0)
+		{
+			Outputs.Add(BoundingShapeInputs[0]);
+		}
+		else
+		{
+			Outputs.Emplace();
+		}
 	}
 
 	// Warn if something is connected but no spatial data could be obtained for sampling
diff --git a/Engine/Plugins/Experimental/PCG/Source/PCG/Private/Elements/PCGWorldQuery.cpp b/Engine/Plugins/Experimental/PCG/Source/PCG/Private/Elements/PCGWorldQuery.cpp
--- a/Engine/Plugins/Experimental/PCG/Source/PCG/Private/Elements/PCGWorldQuery.cpp
+++ b/Engine/Plugins/Experimental/PCG/Source/PCG/Private/Elements/PCGWorldQuery.cpp
@@ -49,6 +49,7 @@ bool FPCGWorldVolumetricQueryElement::ExecuteInternal(FPCGContext* Context) cons
 	UPCGWorldVolumetricData* Data = NewObject<UPCGWorldVolumetricData>();
 	Data->Initialize(World);
 	Data->QueryParams = QueryParams;
+	Data->QueryParams.Initialize();
 	Data->OriginatingComponent = Context->SourceComponent;
 	Data->TargetActor = Context->SourceComponent->GetOwner();
 
@@ -126,6 +127,7 @@ bool FPCGWorldRayHitQueryElement::ExecuteInternal(FPCGContext* Context) const
 	UPCGWorldRayHitData* Data = NewObject<UPCGWorldRayHitData>();
 	Data->Initialize(World);
 	Data->QueryParams = QueryParams;
+	Data->QueryParams.Initialize();
 	Data->OriginatingComponent = Context->SourceComponent;
 	Data->TargetActor = Context->SourceComponent->GetOwner();
 
diff --git a/Engine/Plugins/Experimental/PCG/Source/PCG/Private/Grid/PCGLandscapeCache.cpp b/Engine/Plugins/Experimental/PCG/Source/PCG/Private/Grid/PCGLandscapeCache.cpp
--- a/Engine/Plugins/Experimental/PCG/Source/PCG/Private/Grid/PCGLandscapeCache.cpp
+++ b/Engine/Plugins/Experimental/PCG/Source/PCG/Private/Grid/PCGLandscapeCache.cpp
@@ -6,13 +6,13 @@
 #include "Helpers/PCGBlueprintHelpers.h"
 #include "Metadata/PCGMetadata.h"
 
-#include "Async/ParallelFor.h"
-#include "Engine/World.h"
-#include "Kismet/GameplayStatics.h"
 #include "LandscapeComponent.h"
+#include "LandscapeDataAccess.h"
 #include "LandscapeInfoMap.h"
 #include "LandscapeProxy.h"
-#include "LandscapeDataAccess.h"
+#include "Async/ParallelFor.h"
+#include "Engine/World.h"
+#include "Kismet/GameplayStatics.h"
 #include "Serialization/BufferWriter.h"
 
 #include UE_INLINE_GENERATED_CPP_BY_NAME(PCGLandscapeCache)
@@ -210,42 +210,54 @@ void FPCGLandscapeCacheEntry::GetInterpolatedLayerWeights(const FVector2D& Local
 	}
 }
 
-void FPCGLandscapeCacheEntry::GetInterpolatedPoint(const FVector2D& LocalPoint, FPCGPoint& OutPoint, UPCGMetadata* OutMetadata) const
+void FPCGLandscapeCacheEntry::GetInterpolatedPointInternal(const PCGLandscapeCache::FSafeIndices& Indices, FPCGPoint& OutPoint, bool bHeightOnly) const
 {
-	check (bDataLoaded);
-	
-	const PCGLandscapeCache::FSafeIndices Indices = PCGLandscapeCache::CalcSafeIndices(LocalPoint, Stride);
-
+	check(bDataLoaded);
 	check(2 * Indices.X1Y1 < PositionsAndNormals.Num());
 
 	const FVector& PositionX0Y0 = PositionsAndNormals[2 * Indices.X0Y0];
-	const FVector& NormalX0Y0 = PositionsAndNormals[2 * Indices.X0Y0 + 1];
 	const FVector& PositionX1Y0 = PositionsAndNormals[2 * Indices.X1Y0];
-	const FVector& NormalX1Y0 = PositionsAndNormals[2 * Indices.X1Y0 + 1];
 	const FVector& PositionX0Y1 = PositionsAndNormals[2 * Indices.X0Y1];
-	const FVector& NormalX0Y1 = PositionsAndNormals[2 * Indices.X0Y1 + 1];
 	const FVector& PositionX1Y1 = PositionsAndNormals[2 * Indices.X1Y1];
-	const FVector& NormalX1Y1 = PositionsAndNormals[2 * Indices.X1Y1 + 1];
 
 	const FVector LerpPositionY0 = FMath::Lerp(PositionX0Y0, PositionX1Y0, Indices.XFraction);
 	const FVector LerpPositionY1 = FMath::Lerp(PositionX0Y1, PositionX1Y1, Indices.XFraction);
 	const FVector Position = FMath::Lerp(LerpPositionY0, LerpPositionY1, Indices.YFraction);
 
-	const FVector LerpNormalY0 = FMath::Lerp(NormalX0Y0.GetSafeNormal(), NormalX1Y0.GetSafeNormal(), Indices.XFraction).GetSafeNormal();
-	const FVector LerpNormalY1 = FMath::Lerp(NormalX0Y1.GetSafeNormal(), NormalX1Y1.GetSafeNormal(), Indices.XFraction).GetSafeNormal();
-	const FVector Normal = FMath::Lerp(LerpNormalY0, LerpNormalY1, Indices.YFraction);
-
-	FVector TangentX;
-	FVector TangentY;
-	TangentX = FVector(Normal.Z, 0.f, -Normal.X);
-	TangentY = Normal ^ TangentX;
-
 	const int32 Seed = UPCGBlueprintHelpers::ComputeSeedFromPosition(Position);
 	const float Density = 1;
 
-	new(&OutPoint) FPCGPoint(FTransform(TangentX.GetSafeNormal(), TangentY.GetSafeNormal(), Normal.GetSafeNormal(), Position), Density, Seed);
+	if (bHeightOnly)
+	{
+		new(&OutPoint) FPCGPoint(FTransform(Position), Density, Seed);
+	}
+	else
+	{
+		const FVector& NormalX0Y0 = PositionsAndNormals[2 * Indices.X0Y0 + 1];
+		const FVector& NormalX1Y0 = PositionsAndNormals[2 * Indices.X1Y0 + 1];
+		const FVector& NormalX0Y1 = PositionsAndNormals[2 * Indices.X0Y1 + 1];
+		const FVector& NormalX1Y1 = PositionsAndNormals[2 * Indices.X1Y1 + 1];
+
+		const FVector LerpNormalY0 = FMath::Lerp(NormalX0Y0.GetSafeNormal(), NormalX1Y0.GetSafeNormal(), Indices.XFraction).GetSafeNormal();
+		const FVector LerpNormalY1 = FMath::Lerp(NormalX0Y1.GetSafeNormal(), NormalX1Y1.GetSafeNormal(), Indices.XFraction).GetSafeNormal();
+		const FVector Normal = FMath::Lerp(LerpNormalY0, LerpNormalY1, Indices.YFraction);
+
+		FVector TangentX;
+		FVector TangentY;
+		TangentX = FVector(Normal.Z, 0.f, -Normal.X);
+		TangentY = Normal ^ TangentX;
+
+		new(&OutPoint) FPCGPoint(FTransform(TangentX.GetSafeNormal(), TangentY.GetSafeNormal(), Normal.GetSafeNormal(), Position), Density, Seed);
+	}
+
 	OutPoint.BoundsMin = -PointHalfSize;
 	OutPoint.BoundsMax = PointHalfSize;
+}
+
+void FPCGLandscapeCacheEntry::GetInterpolatedPointMetadataInternal(const PCGLandscapeCache::FSafeIndices& Indices, FPCGPoint& OutPoint, UPCGMetadata* OutMetadata) const
+{
+	check(bDataLoaded);
+	check(2 * Indices.X1Y1 < PositionsAndNormals.Num());
 
 	if (OutMetadata && !LayerData.IsEmpty())
 	{
@@ -253,7 +265,7 @@ void FPCGLandscapeCacheEntry::GetInterpolatedPoint(const FVector2D& LocalPoint,
 
 		check(LayerData.Num() == LayerDataNames.Num());
 
-		for(int32 LayerIndex = 0; LayerIndex < LayerData.Num(); ++LayerIndex)
+		for (int32 LayerIndex = 0; LayerIndex < LayerData.Num(); ++LayerIndex)
 		{
 			const FName& CurrentLayerName = LayerDataNames[LayerIndex];
 			const TArray<uint8>& CurrentLayerData = LayerData[LayerIndex];
@@ -261,39 +273,36 @@ void FPCGLandscapeCacheEntry::GetInterpolatedPoint(const FVector2D& LocalPoint,
 			if (FPCGMetadataAttributeBase* Attribute = OutMetadata->GetMutableAttribute(CurrentLayerName))
 			{
 				check(Attribute->GetTypeId() == PCG::Private::MetadataTypes<float>::Id);
-				float Y0Data = FMath::Lerp((float)CurrentLayerData[Indices.X0Y0] / 255.0f, (float)CurrentLayerData[Indices.X1Y0] / 255.0f, Indices.XFraction);
-				float Y1Data = FMath::Lerp((float)CurrentLayerData[Indices.X0Y1] / 255.0f, (float)CurrentLayerData[Indices.X1Y1] / 255.0f, Indices.XFraction);
-				float Data = FMath::Lerp(Y0Data, Y1Data, Indices.YFraction);
+				const float Y0Data = FMath::Lerp((float)CurrentLayerData[Indices.X0Y0] / 255.0f, (float)CurrentLayerData[Indices.X1Y0] / 255.0f, Indices.XFraction);
+				const float Y1Data = FMath::Lerp((float)CurrentLayerData[Indices.X0Y1] / 255.0f, (float)CurrentLayerData[Indices.X1Y1] / 255.0f, Indices.XFraction);
+				const float Data = FMath::Lerp(Y0Data, Y1Data, Indices.YFraction);
 
 				static_cast<FPCGMetadataAttribute<float>*>(Attribute)->SetValue(OutPoint.MetadataEntry, Data);
-			} 
+			}
 		}
 	}
 }
 
-void FPCGLandscapeCacheEntry::GetInterpolatedPointHeightOnly(const FVector2D& LocalPoint, FPCGPoint& OutPoint) const
+void FPCGLandscapeCacheEntry::GetInterpolatedPoint(const FVector2D& LocalPoint, FPCGPoint& OutPoint, UPCGMetadata* OutMetadata) const
 {
-	check(bDataLoaded);
-
+	check (bDataLoaded);
 	const PCGLandscapeCache::FSafeIndices Indices = PCGLandscapeCache::CalcSafeIndices(LocalPoint, Stride);
+	GetInterpolatedPointInternal(Indices, OutPoint);
+	GetInterpolatedPointMetadataInternal(Indices, OutPoint, OutMetadata);
+}
 
-	check(2 * Indices.X1Y1 < PositionsAndNormals.Num());
-
-	const FVector& PositionX0Y0 = PositionsAndNormals[2 * Indices.X0Y0];
-	const FVector& PositionX1Y0 = PositionsAndNormals[2 * Indices.X1Y0];
-	const FVector& PositionX0Y1 = PositionsAndNormals[2 * Indices.X0Y1];
-	const FVector& PositionX1Y1 = PositionsAndNormals[2 * Indices.X1Y1];
-
-	const FVector LerpPositionY0 = FMath::Lerp(PositionX0Y0, PositionX1Y0, Indices.XFraction);
-	const FVector LerpPositionY1 = FMath::Lerp(PositionX0Y1, PositionX1Y1, Indices.XFraction);
-	const FVector Position = FMath::Lerp(LerpPositionY0, LerpPositionY1, Indices.YFraction);
-
-	const int32 Seed = UPCGBlueprintHelpers::ComputeSeedFromPosition(Position);
-	const float Density = 1;
+void FPCGLandscapeCacheEntry::GetInterpolatedPointMetadataOnly(const FVector2D& LocalPoint, FPCGPoint& OutPoint, UPCGMetadata* OutMetadata) const
+{
+	check(bDataLoaded);
+	const PCGLandscapeCache::FSafeIndices Indices = PCGLandscapeCache::CalcSafeIndices(LocalPoint, Stride);
+	GetInterpolatedPointMetadataInternal(Indices, OutPoint, OutMetadata);
+}
 
-	new(&OutPoint) FPCGPoint(FTransform(Position), Density, Seed);
-	OutPoint.BoundsMin = -PointHalfSize;
-	OutPoint.BoundsMax = PointHalfSize;
+void FPCGLandscapeCacheEntry::GetInterpolatedPointHeightOnly(const FVector2D& LocalPoint, FPCGPoint& OutPoint) const
+{
+	check(bDataLoaded);
+	const PCGLandscapeCache::FSafeIndices Indices = PCGLandscapeCache::CalcSafeIndices(LocalPoint, Stride);
+	GetInterpolatedPointInternal(Indices, OutPoint, /*bHeightOnly=*/true);
 }
 
 bool FPCGLandscapeCacheEntry::TouchAndLoad(int32 InTouch) const
@@ -686,6 +695,39 @@ TArray<FName> UPCGLandscapeCache::GetLayerNames(ALandscapeProxy* Landscape)
 	return CachedLayerNames.Array();
 }
 
+void UPCGLandscapeCache::SampleMetadataOnPoint(ALandscapeProxy* Landscape, FPCGPoint& InOutPoint, UPCGMetadata* OutMetadata)
+{
+	if (!Landscape || !Landscape->GetWorld() || !OutMetadata || CachedLayerNames.IsEmpty())
+	{
+		return;
+	}
+
+	ULandscapeInfo* LandscapeInfo = Landscape->GetLandscapeInfo();
+	if (!LandscapeInfo)
+	{
+		UE_LOG(LogPCG, Warning, TEXT("Unable to get landscape layer weights because the landscape info is not available (landscape not registered yet?"));
+		return;
+	}
+
+	const FVector LocalPoint = Landscape->GetTransform().InverseTransformPosition(InOutPoint.Transform.GetLocation());
+	const FIntPoint ComponentMapKey(FMath::FloorToInt(LocalPoint.X / LandscapeInfo->ComponentSizeQuads), FMath::FloorToInt(LocalPoint.Y / LandscapeInfo->ComponentSizeQuads));
+
+#if WITH_EDITOR
+	ULandscapeComponent* LandscapeComponent = LandscapeInfo->XYtoComponentMap.FindRef(ComponentMapKey);
+	const FPCGLandscapeCacheEntry* CacheEntry = GetCacheEntry(LandscapeComponent, ComponentMapKey);
+#else
+	const FPCGLandscapeCacheEntry* CacheEntry = GetCacheEntry(Landscape->GetLandscapeGuid(), ComponentMapKey);
+#endif
+
+	if (!CacheEntry || CacheEntry->LayerData.IsEmpty())
+	{
+		return;
+	}
+
+	const FVector2D ComponentLocalPoint(LocalPoint.X - ComponentMapKey.X * LandscapeInfo->ComponentSizeQuads, LocalPoint.Y - ComponentMapKey.Y * LandscapeInfo->ComponentSizeQuads);	
+	CacheEntry->GetInterpolatedPointMetadataOnly(ComponentLocalPoint, InOutPoint, OutMetadata);
+}
+
 #if WITH_EDITOR
 void UPCGLandscapeCache::SetupLandscapeCallbacks()
 {
diff --git a/Engine/Plugins/Experimental/PCG/Source/PCG/Public/Data/PCGWorldData.h b/Engine/Plugins/Experimental/PCG/Source/PCG/Public/Data/PCGWorldData.h
--- a/Engine/Plugins/Experimental/PCG/Source/PCG/Public/Data/PCGWorldData.h
+++ b/Engine/Plugins/Experimental/PCG/Source/PCG/Public/Data/PCGWorldData.h
@@ -15,26 +15,61 @@ class UWorld;
 class UPCGMetadata;
 class UPCGComponent;
 
+UENUM()
+enum class EPCGWorldQueryFilterByTag
+{
+	NoTagFilter,
+	IncludeTagged,
+	ExcludeTagged
+};
+
 USTRUCT(BlueprintType)
-struct FPCGWorldVolumetricQueryParams
+struct FPCGWorldCommonQueryParams
 {
 	GENERATED_BODY()
 
-	UPROPERTY(BlueprintReadWrite, EditAnywhere, Category = Data)
-	bool bSearchForOverlap = true;
-
-	UPROPERTY(BlueprintReadWrite, EditAnywhere, Category = Data)
+	/** If true, will ignore hits/overlaps on content created from PCG. */
+	UPROPERTY(BlueprintReadWrite, EditAnywhere, Category = Data, meta = (PCG_Overridable))
 	bool bIgnorePCGHits = false;
 
-	UPROPERTY(BlueprintReadWrite, EditAnywhere, Category = Data)
+	UPROPERTY(BlueprintReadWrite, EditAnywhere, Category = Data, meta = (PCG_Overridable))
 	bool bIgnoreSelfHits = true;
 
-	UPROPERTY(BlueprintReadWrite, EditAnywhere, Category = "Data|Advanced")
+	UPROPERTY(BlueprintReadWrite, EditAnywhere, Category = "Data|Advanced", meta = (PCG_Overridable))
 	TEnumAsByte<ECollisionChannel> CollisionChannel = ECC_WorldStatic;
 
 	/** Queries against complex collision if enabled, performance warning */
-	UPROPERTY(BlueprintReadWrite, EditAnywhere, Category = "Data|Advanced")
+	UPROPERTY(BlueprintReadWrite, EditAnywhere, Category = "Data|Advanced", meta = (PCG_Overridable))
 	bool bTraceComplex = false;
+
+	UPROPERTY(BlueprintReadWrite, EditAnywhere, Category = "Data|Filtering", meta = (PCG_Overridable))
+	EPCGWorldQueryFilterByTag ActorTagFilter = EPCGWorldQueryFilterByTag::NoTagFilter;
+
+	UPROPERTY(BlueprintReadWrite, EditAnywhere, Category = "Data|Filtering", meta = (PCG_Overridable, EditCondition = "TagFilter != EPCGWorldQueryFilterByTag::NoTagFilter"))
+	FString ActorTagsList;
+
+	UPROPERTY(BlueprintReadWrite, EditAnywhere, Category = "Data|Filtering", meta = (PCG_Overridable))
+	bool bIgnoreLandscapeHits = false;
+
+	// Not exposed, will be filled in when initializing this
+	UPROPERTY()
+	TSet<FName> ParsedActorTagsList;
+
+protected:
+	/** Sets up the data we need to efficiently perform the queries */
+	void Initialize();
+};
+
+USTRUCT(BlueprintType)
+struct FPCGWorldVolumetricQueryParams : public FPCGWorldCommonQueryParams
+{
+	GENERATED_BODY()
+
+	void Initialize();
+
+	/** Controls whether we are trying to find an overlap with physical objects (true) or to find empty spaces that do not contain anything (false) */
+	UPROPERTY(BlueprintReadWrite, EditAnywhere, Category = Data, meta = (PCG_Overridable))
+	bool bSearchForOverlap = true;
 };
 
 /** Queries volume for presence of world collision or not. Can be used to voxelize environment. */
@@ -73,40 +108,29 @@ public:
 };
 
 USTRUCT(BlueprintType)
-struct FPCGWorldRayHitQueryParams
+struct FPCGWorldRayHitQueryParams : public FPCGWorldCommonQueryParams
 {
 	GENERATED_BODY()
 
-	UPROPERTY(BlueprintReadWrite, EditAnywhere, Category = Data)
+	void Initialize();
+
+	UPROPERTY(BlueprintReadWrite, EditAnywhere, Category = Data, meta = (PCG_Overridable))
 	bool bOverrideDefaultParams = false;
 
-	UPROPERTY(BlueprintReadWrite, EditAnywhere, Category = Data, meta = (EditCondition = "bOverrideDefaultParams"))
+	UPROPERTY(BlueprintReadWrite, EditAnywhere, Category = Data, meta = (PCG_Overridable, EditCondition = "bOverrideDefaultParams"))
 	FVector RayOrigin = FVector::ZeroVector;
 
-	UPROPERTY(BlueprintReadWrite, EditAnywhere, Category = Data, meta = (EditCondition = "bOverrideDefaultParams"))
+	UPROPERTY(BlueprintReadWrite, EditAnywhere, Category = Data, meta = (PCG_Overridable, EditCondition = "bOverrideDefaultParams"))
 	FVector RayDirection = FVector(0.0, 0.0, -1.0);
 
-	UPROPERTY(BlueprintReadWrite, EditAnywhere, Category = Data, meta = (EditCondition = "bOverrideDefaultParams"))
+	UPROPERTY(BlueprintReadWrite, EditAnywhere, Category = Data, meta = (PCG_Overridable, EditCondition = "bOverrideDefaultParams"))
 	double RayLength = 1.0e+5; // 100m
 
-	UPROPERTY(BlueprintReadWrite, EditAnywhere, Category = Data)
-	bool bIgnorePCGHits = false;
-
-	UPROPERTY(BlueprintReadWrite, EditAnywhere, Category = Data)
-	bool bIgnoreSelfHits = true;
-
-	UPROPERTY(BlueprintReadWrite, EditAnywhere, Category = "Data|Advanced")
-	TEnumAsByte<ECollisionChannel> CollisionChannel = ECC_WorldStatic;
-
-	/** Queries against complex collision if enabled, performance warning */
-	UPROPERTY(BlueprintReadWrite, EditAnywhere, Category = "Data|Advanced")
-	bool bTraceComplex = false;
-
 	// TODO: see in FCollisionQueryParams if there are some flags we want to expose
 	// examples: bReturnFaceIndex, bReturnPhysicalMaterial, some ignore patterns
 
-	//TODO UPROPERTY()
-	//bool bUseMetadataFromLandscape = true;
+	UPROPERTY(BlueprintReadWrite, EditAnywhere, Category = Data, meta = (PCG_Overridable))
+	bool bApplyMetadataFromLandscape = false;
 };
 
 /** Executes collision queries against world collision. */
diff --git a/Engine/Plugins/Experimental/PCG/Source/PCG/Public/Elements/PCGVolumeSampler.h b/Engine/Plugins/Experimental/PCG/Source/PCG/Public/Elements/PCGVolumeSampler.h
--- a/Engine/Plugins/Experimental/PCG/Source/PCG/Public/Elements/PCGVolumeSampler.h
+++ b/Engine/Plugins/Experimental/PCG/Source/PCG/Public/Elements/PCGVolumeSampler.h
@@ -33,6 +33,13 @@ public:
 	UPROPERTY(BlueprintReadWrite, EditAnywhere, Category = "Data", meta = (PCG_Overridable))
 	FVector VoxelSize = FVector(100.0, 100.0, 100.0);
 
+	/** If no Bounding Shape input is provided the actor bounds are used to limit the sample generation area.
+	* This option allows ignoring the actor bounds and generating over the entire volume provided. Use with caution as this
+	* may generate a lot of points.
+	*/
+	UPROPERTY(BlueprintReadWrite, EditAnywhere, Category = Settings, meta = (PCG_Overridable))
+	bool bUnbounded = false;
+
 	//~Begin UPCGSettings interface
 #if WITH_EDITOR
 	virtual FName GetDefaultNodeName() const override { return FName(TEXT("Volume Sampler")); }
diff --git a/Engine/Plugins/Experimental/PCG/Source/PCG/Public/Grid/PCGLandscapeCache.h b/Engine/Plugins/Experimental/PCG/Source/PCG/Public/Grid/PCGLandscapeCache.h
--- a/Engine/Plugins/Experimental/PCG/Source/PCG/Public/Grid/PCGLandscapeCache.h
+++ b/Engine/Plugins/Experimental/PCG/Source/PCG/Public/Grid/PCGLandscapeCache.h
@@ -51,10 +51,15 @@ public:
 	void GetPoint(int32 PointIndex, FPCGPoint& OutPoint, UPCGMetadata* OutMetadata) const;
 	void GetPointHeightOnly(int32 PointIndex, FPCGPoint& OutPoint) const;
 	void GetInterpolatedPoint(const FVector2D& LocalPoint, FPCGPoint& OutPoint, UPCGMetadata* OutMetadata) const;
+	void GetInterpolatedPointMetadataOnly(const FVector2D& LocalPoint, FPCGPoint& OutPoint, UPCGMetadata* OutMetadata) const;
 	void GetInterpolatedPointHeightOnly(const FVector2D& LocalPoint, FPCGPoint& OutPoint) const;
 	void GetInterpolatedLayerWeights(const FVector2D& LocalPoint, TArray<FPCGLandscapeLayerWeight>& OutLayerWeights) const;
 
 private:
+	// Private API to remove boilerplate
+	void GetInterpolatedPointInternal(const PCGLandscapeCache::FSafeIndices& Indices, FPCGPoint& OutPoint, bool bHeightOnly = false) const;
+	void GetInterpolatedPointMetadataInternal(const PCGLandscapeCache::FSafeIndices& Indices, FPCGPoint& OutPoint, UPCGMetadata* OutMetadata) const;
+
 	// Private API for UPCGLandscapeCache usage
 	bool TouchAndLoad(int32 Touch) const;
 	void Unload();
@@ -117,6 +122,9 @@ public:
 
 	TArray<FName> GetLayerNames(ALandscapeProxy* Landscape);
 
+	/** Convenience method to get metadata from the landscape for a given pair of landscape and position */
+	void SampleMetadataOnPoint(ALandscapeProxy* Landscape, FPCGPoint& OutPoint, UPCGMetadata* OutMetadata);
+
 private:
 #if WITH_EDITOR
 	void SetupLandscapeCallbacks();
