Fixed merge typo from CL 24714692
#rnx

diff --git a/Engine/Plugins/Experimental/PCG/Source/PCG/Private/Elements/PCGSpawnActor.cpp b/Engine/Plugins/Experimental/PCG/Source/PCG/Private/Elements/PCGSpawnActor.cpp
--- a/Engine/Plugins/Experimental/PCG/Source/PCG/Private/Elements/PCGSpawnActor.cpp
+++ b/Engine/Plugins/Experimental/PCG/Source/PCG/Private/Elements/PCGSpawnActor.cpp
@@ -441,7 +441,7 @@ bool FPCGSpawnActorElement::ExecuteInternal(FPCGContext* Context) const
 				Params.Descriptor.InitFrom(StaticMeshComponent);
 				// TODO: No custom data float support?
 
-				TArray<FTransform>& Transforms = MeshDescriptorTransforms.FindOrAdd(Params)
+				TArray<FTransform>& Transforms = MeshDescriptorTransforms.FindOrAdd(Params);
 
 				if (const UInstancedStaticMeshComponent* InstancedStaticMeshComponent = Cast<UInstancedStaticMeshComponent>(StaticMeshComponent))
 				{
