[PCG] Get Primitive Data node

diff --git a/Engine/Plugins/Experimental/PCG/Source/PCG/Private/Elements/PCGTypedGetter.cpp b/Engine/Plugins/Experimental/PCG/Source/PCG/Private/Elements/PCGTypedGetter.cpp
--- a/Engine/Plugins/Experimental/PCG/Source/PCG/Private/Elements/PCGTypedGetter.cpp
+++ b/Engine/Plugins/Experimental/PCG/Source/PCG/Private/Elements/PCGTypedGetter.cpp
@@ -69,4 +69,25 @@ FText UPCGGetVolumeSettings::GetNodeTooltipText() const
 }
 #endif
 
+UPCGGetPrimitiveSettings::UPCGGetPrimitiveSettings()
+{
+	bDisplayModeSettings = false;
+	Mode = EPCGGetDataFromActorMode::ParseActorComponents;
+}
+
+TArray<FPCGPinProperties> UPCGGetPrimitiveSettings::OutputPinProperties() const
+{
+	TArray<FPCGPinProperties> PinProperties;
+	PinProperties.Emplace(PCGPinConstants::DefaultOutputLabel, EPCGDataType::Primitive);
+
+	return PinProperties;
+}
+
+#if WITH_EDITOR
+FText UPCGGetPrimitiveSettings::GetNodeTooltipText() const
+{
+	return LOCTEXT("GetPrimitiveTooltip", "Builds a collection of primitive data from primitive components on the selected actors.");
+}
+#endif
+
 #undef LOCTEXT_NAMESPACE
diff --git a/Engine/Plugins/Experimental/PCG/Source/PCG/Public/Elements/PCGTypedGetter.h b/Engine/Plugins/Experimental/PCG/Source/PCG/Public/Elements/PCGTypedGetter.h
--- a/Engine/Plugins/Experimental/PCG/Source/PCG/Public/Elements/PCGTypedGetter.h
+++ b/Engine/Plugins/Experimental/PCG/Source/PCG/Public/Elements/PCGTypedGetter.h
@@ -80,3 +80,28 @@ public:
 	virtual bool DataFilter(EPCGDataType InDataType) const { return !!(InDataType & EPCGDataType::Volume); }
 	//~End UPCGDataFromActorSettings
 };
+
+/** Builds a collection of primitive data from primitive components on the selected actors. */
+UCLASS(BlueprintType, ClassGroup = (Procedural))
+class PCG_API UPCGGetPrimitiveSettings : public UPCGDataFromActorSettings
+{
+	GENERATED_BODY()
+
+public:
+	UPCGGetPrimitiveSettings();
+
+	//~Begin UPCGSettings interface
+#if WITH_EDITOR
+	virtual FName GetDefaultNodeName() const override { return FName(TEXT("Get Primitive Data")); }
+	virtual FText GetNodeTooltipText() const override;
+#endif
+
+protected:
+	virtual TArray<FPCGPinProperties> OutputPinProperties() const override;
+	//~End UPCGSettings
+
+public:
+	//~Begin UPCGDataFromActorSettings interface
+	virtual bool DataFilter(EPCGDataType InDataType) const { return !!(InDataType & EPCGDataType::Primitive); }
+	//~End UPCGDataFromActorSettings
+};
