[PCG] Graph editor errors and warnings

These are dynamic execution warnings. They are raised or not raised during execution. I therefore don't cache a node that has issues. It should keep running until the execution code can clear the error/warning, and then the data is good to cache.

diff --git a/Engine/Plugins/Experimental/PCG/Source/PCG/Private/Elements/Metadata/PCGMetadataElement.cpp b/Engine/Plugins/Experimental/PCG/Source/PCG/Private/Elements/Metadata/PCGMetadataElement.cpp
--- a/Engine/Plugins/Experimental/PCG/Source/PCG/Private/Elements/Metadata/PCGMetadataElement.cpp
+++ b/Engine/Plugins/Experimental/PCG/Source/PCG/Private/Elements/Metadata/PCGMetadataElement.cpp
@@ -10,6 +10,8 @@
 
 #include UE_INLINE_GENERATED_CPP_BY_NAME(PCGMetadataElement)
 
+#define LOCTEXT_NAMESPACE "PCGMetadataElement"
+
 UPCGMetadataOperationSettings::UPCGMetadataOperationSettings()
 {
 	// Previous default object was: None for source and target attribute, density for point property
@@ -82,7 +84,7 @@ bool FPCGMetadataOperationElement::ExecuteInternal(FPCGContext* Context) const
 
 		if (!SpatialInput)
 		{
-			PCGE_LOG(Error, "Invalid input data");
+			PCGE_LOG(Error, GraphAndLog, LOCTEXT("InvalidInputData", "Invalid input data type, must be of type Spatial"));
 			continue;
 		}
 
@@ -90,13 +92,13 @@ bool FPCGMetadataOperationElement::ExecuteInternal(FPCGContext* Context) const
 
 		if (!OriginalData)
 		{
-			PCGE_LOG(Error, "Unable to get point data from input");
+			PCGE_LOG(Error, GraphAndLog, LOCTEXT("UnableToGetPointData", "Unable to get point data from input"));
 			continue;
 		}
 
 		if (!OriginalData->Metadata)
 		{
-			PCGE_LOG(Warning, "Input has no metadata");
+			PCGE_LOG(Warning, GraphAndLog, LOCTEXT("MissingMetadata", "Input has no metadata"));
 			continue;
 		}
 
@@ -114,7 +116,7 @@ bool FPCGMetadataOperationElement::ExecuteInternal(FPCGContext* Context) const
 
 		if (Settings->InputSource.Selection == EPCGAttributePropertySelection::Attribute && !OriginalData->Metadata->HasAttribute(LocalSourceAttribute))
 		{
-			PCGE_LOG(Warning, "Input does not have the %s attribute", *LocalSourceAttribute.ToString());
+			PCGE_LOG(Warning, GraphAndLog, FText::Format(LOCTEXT("InputMissingAttribute", "Input does not have the '{0}' attribute"), FText::FromName(LocalSourceAttribute)));
 			continue;
 		}
 
@@ -141,7 +143,7 @@ bool FPCGMetadataOperationElement::ExecuteInternal(FPCGContext* Context) const
 
 			if (!SampledData->Metadata->CopyExistingAttribute(LocalSourceAttribute, DestinationAttribute))
 			{
-				PCGE_LOG(Warning, "Failed to copy to new attribute %s", *DestinationAttribute.ToString());
+				PCGE_LOG(Warning, GraphAndLog, FText::Format(LOCTEXT("FailedCopyToNewAttribute", "Failed to copy to new attribute {0}"), FText::FromName(DestinationAttribute)));
 			}
 
 			continue;
@@ -152,7 +154,7 @@ bool FPCGMetadataOperationElement::ExecuteInternal(FPCGContext* Context) const
 
 		if (!InputAccessor.IsValid() || !InputKeys.IsValid())
 		{
-			PCGE_LOG(Warning, "Failed to create input accessor or iterator");
+			PCGE_LOG(Warning, GraphAndLog, LOCTEXT("FailedToCreateInputAccessor", "Failed to create input accessor or iterator"));
 			continue;
 		}
 
@@ -168,7 +170,7 @@ bool FPCGMetadataOperationElement::ExecuteInternal(FPCGContext* Context) const
 			
 			if (!PCGMetadataAttribute::CallbackWithRightType(InputAccessor->GetUnderlyingType(), CreateAttribute))
 			{
-				PCGE_LOG(Warning, "Failed to create new attribute %s", *LocalDestinationAttribute.ToString());
+				PCGE_LOG(Warning, GraphAndLog, FText::Format(LOCTEXT("FailedToCreateNewAttribute", "Failed to create new attribute {0}"), FText::FromName(LocalDestinationAttribute)));
 				continue;
 			}
 		}
@@ -178,7 +180,7 @@ bool FPCGMetadataOperationElement::ExecuteInternal(FPCGContext* Context) const
 
 		if (!OutputAccessor.IsValid() || !OutputKeys.IsValid())
 		{
-			PCGE_LOG(Warning, "Failed to create output accessor or iterator");
+			PCGE_LOG(Warning, GraphAndLog, LOCTEXT("FailedToCreateOutputAccessor", "Failed to create output accessor or iterator"));
 			continue;
 		}
 
@@ -188,7 +190,7 @@ bool FPCGMetadataOperationElement::ExecuteInternal(FPCGContext* Context) const
 		// Final verification, if we can put the value of input into output
 		if (!PCG::Private::IsBroadcastable(InputAccessor->GetUnderlyingType(), OutputAccessor->GetUnderlyingType()))
 		{
-			PCGE_LOG(Error, "Can't broadcast input type into output type");
+			PCGE_LOG(Error, GraphAndLog, LOCTEXT("CannotBroadcastTypes", "Cannot broadcast input type into output type"));
 			continue;
 		}
 
@@ -223,10 +225,12 @@ bool FPCGMetadataOperationElement::ExecuteInternal(FPCGContext* Context) const
 
 		if (!PCGMetadataAttribute::CallbackWithRightType(OutputAccessor->GetUnderlyingType(), Operation))
 		{
-			PCGE_LOG(Warning, "Error while getting/setting values");
+			PCGE_LOG(Warning, GraphAndLog, LOCTEXT("ErrorGettingSettingValues", "Error while getting/setting values"));
 			continue;
 		}
 	}
 
 	return true;
 }
+
+#undef LOCTEXT_NAMESPACE
diff --git a/Engine/Plugins/Experimental/PCG/Source/PCG/Private/Elements/Metadata/PCGMetadataOpElementBase.cpp b/Engine/Plugins/Experimental/PCG/Source/PCG/Private/Elements/Metadata/PCGMetadataOpElementBase.cpp
--- a/Engine/Plugins/Experimental/PCG/Source/PCG/Private/Elements/Metadata/PCGMetadataOpElementBase.cpp
+++ b/Engine/Plugins/Experimental/PCG/Source/PCG/Private/Elements/Metadata/PCGMetadataOpElementBase.cpp
@@ -12,6 +12,8 @@
 
 #include UE_INLINE_GENERATED_CPP_BY_NAME(PCGMetadataOpElementBase)
 
+#define LOCTEXT_NAMESPACE "PCGMetadataElementBaseElement"
+
 void UPCGMetadataSettingsBase::PostLoad()
 {
 	Super::PostLoad();
@@ -85,9 +87,15 @@ bool FPCGMetadataElementBase::ExecuteInternal(FPCGContext* Context) const
 	for (uint32 i = 0; i < NumberOfInputs; ++i)
 	{
 		TArray<FPCGTaggedData> InputData = Context->InputData.GetInputsByPin(Settings->GetInputPinLabel(i));
-		if (InputData.Num() != 1)
+		if (InputData.IsEmpty())
+		{
+			// Absence of data not worth broadcasting to user as visual warning
+			PCGE_LOG(Warning, LogOnly, FText::Format(LOCTEXT("MissingInputDataForPin", "No data provided on pin {0}"), i));
+			return true;
+		}
+		else if (InputData.Num() > 1)
 		{
-			PCGE_LOG(Error, "Invalid inputs for pin %d", i);
+			PCGE_LOG(Error, GraphAndLog, FText::Format(LOCTEXT("TooMuchDataForPin", "Too many data items ({0}) provided on pin {1}"), InputData.Num(), i));
 			return true;
 		}
 
@@ -105,7 +113,7 @@ bool FPCGMetadataElementBase::ExecuteInternal(FPCGContext* Context) const
 		}
 		else
 		{
-			PCGE_LOG(Error, "Invalid inputs for pin %d", i);
+			PCGE_LOG(Warning, LogOnly, FText::Format(LOCTEXT("InvalidInputDataTypeForPin", "Invalid data provided on pin {0}, must be of type Spatial or Param"), i));
 			return true;
 		}
 	}
@@ -132,7 +140,7 @@ bool FPCGMetadataElementBase::ExecuteInternal(FPCGContext* Context) const
 			InputTaggedData[0].Data->GetClass() != InputTaggedData[i].Data->GetClass() &&
 			!InputTaggedData[i].Data->IsA<UPCGParamData>())
 		{
-			PCGE_LOG(Error, "Input %d is not of the same type than input 0 and is not a param data. Not supported.", i);
+			PCGE_LOG(Error, GraphAndLog, FText::Format(LOCTEXT("InputTypeMismatch", "Input {0} is not of the same type than input 0 and is not a Param Data. This is not supported."), i));
 			return true;
 		}
 
@@ -153,7 +161,7 @@ bool FPCGMetadataElementBase::ExecuteInternal(FPCGContext* Context) const
 
 		if (!OperationData.InputAccessors[i].IsValid() || !OperationData.InputKeys[i].IsValid())
 		{
-			PCGE_LOG(Error, "Attribute/Property %s does not exist for input %d", *InputSource.GetName().ToString(), i);
+			PCGE_LOG(Error, GraphAndLog, FText::Format(LOCTEXT("AttributeDoesNotExist", "Attribute/Property '{0}' does not exist for input {1}"), FText::FromName(InputSource.GetName()), i));
 			return true;
 		}
 
@@ -163,7 +171,7 @@ bool FPCGMetadataElementBase::ExecuteInternal(FPCGContext* Context) const
 		bool bHasSpecialRequirement = false;
 		if (!Settings->IsSupportedInputType(AttributeTypeId, i, bHasSpecialRequirement))
 		{
-			PCGE_LOG(Error, "Attribute/Property %s is not a supported type for input %d", *InputSource.GetName().ToString(), i);
+			PCGE_LOG(Error, GraphAndLog, FText::Format(LOCTEXT("UnsupportedAttributeType", "Attribute/Property '{0}' is not a supported type for input {1}"), FText::FromName(InputSource.GetName()), i));
 			return true;
 		}
 
@@ -176,7 +184,7 @@ bool FPCGMetadataElementBase::ExecuteInternal(FPCGContext* Context) const
 			}
 			else if (OperationData.MostComplexInputType != AttributeTypeId && !PCG::Private::IsBroadcastable(AttributeTypeId, OperationData.MostComplexInputType))
 			{
-				PCGE_LOG(Error, "Attribute %s cannot be broadcasted to match types for input %d", *InputSource.GetName().ToString(), i);
+				PCGE_LOG(Error, GraphAndLog, FText::Format(LOCTEXT("AttributeCannotBeBroadcasted", "Attribute '{0}' cannot be broadcasted to match types for input {1}"), FText::FromName(InputSource.GetName()), i));
 				return true;
 			}
 		}
@@ -192,7 +200,7 @@ bool FPCGMetadataElementBase::ExecuteInternal(FPCGContext* Context) const
 		// Therefore mark that we have nothing to do and early out.
 		if (NumberOfElements[i] == 0)
 		{
-			PCGE_LOG(Verbose, "No elements in input %d.", i);
+			PCGE_LOG(Verbose, LogOnly, FText::Format(LOCTEXT("NoElementsInInput", "No elements in input {0}"), i));
 			bNoOperationNeeded = true;
 			break;
 		}
@@ -200,7 +208,7 @@ bool FPCGMetadataElementBase::ExecuteInternal(FPCGContext* Context) const
 		// Verify that the number of elements makes sense
 		if (OperationData.NumberOfElementsToProcess % NumberOfElements[i] != 0)
 		{
-			PCGE_LOG(Error, "Mismatch between the number of elements in input 0 (%d) and in input %d (%d).", OperationData.NumberOfElementsToProcess, i, NumberOfElements[i]);
+			PCGE_LOG(Error, GraphAndLog, FText::Format(LOCTEXT("MismatchInNumberOfElements", "Mismatch between the number of elements in input 0 ({0}) and in input {1} ({2})"), OperationData.NumberOfElementsToProcess, i, NumberOfElements[i]));
 			return true;
 		}
 
@@ -286,7 +294,7 @@ bool FPCGMetadataElementBase::ExecuteInternal(FPCGContext* Context) const
 				// We matched an attribute/property, check if the output type is valid.
 				if (!PCG::Private::IsBroadcastable(PCG::Private::MetadataTypes<AttributeType>::Id, OperationData.OutputAccessors[OutputIndex]->GetUnderlyingType()))
 				{
-					PCGE_LOG(Error, "Attribute/Property %s cannot be broadcasted to match types for input", *OutputName.ToString());
+					PCGE_LOG(Error, GraphAndLog, FText::Format(LOCTEXT("AttributeTypeBroadcastFailed", "Attribute/Property '{0}' cannot be broadcasted to match types for input"), FText::FromName(OutputName)));
 					return false;
 				}
 
@@ -345,7 +353,7 @@ bool FPCGMetadataElementBase::ExecuteInternal(FPCGContext* Context) const
 
 	if (!bCreateAttributeSucceeded)
 	{
-		PCGE_LOG(Error, "Error while creating output attributes");
+		PCGE_LOG(Error, GraphAndLog, LOCTEXT("ErrorCreatingOutputAttributes", "Error while creating output attributes"));
 		Outputs.Empty();
 		return true;
 	}
@@ -354,10 +362,11 @@ bool FPCGMetadataElementBase::ExecuteInternal(FPCGContext* Context) const
 
 	if (!DoOperation(OperationData))
 	{
-		PCGE_LOG(Error, "Error while performing the metadata operation, check logs for more information");
+		PCGE_LOG(Error, GraphAndLog, LOCTEXT("ErrorOccurred", "Error while performing the metadata operation, check logs for more information"));
 		Outputs.Empty();
 	}
 
 	return true;
 }
 
+#undef LOCTEXT_NAMESPACE
diff --git a/Engine/Plugins/Experimental/PCG/Source/PCG/Private/Elements/Metadata/PCGMetadataPartition.cpp b/Engine/Plugins/Experimental/PCG/Source/PCG/Private/Elements/Metadata/PCGMetadataPartition.cpp
--- a/Engine/Plugins/Experimental/PCG/Source/PCG/Private/Elements/Metadata/PCGMetadataPartition.cpp
+++ b/Engine/Plugins/Experimental/PCG/Source/PCG/Private/Elements/Metadata/PCGMetadataPartition.cpp
@@ -9,6 +9,8 @@
 
 #include UE_INLINE_GENERATED_CPP_BY_NAME(PCGMetadataPartition)
 
+#define LOCTEXT_NAMESPACE "PCGMetadataPartitionElement"
+
 FPCGElementPtr UPCGMetadataPartitionSettings::CreateElement() const
 {
 	return MakeShared<FPCGMetadataPartitionElement>();
@@ -34,7 +36,7 @@ bool FPCGMetadataPartitionElement::ExecuteInternal(FPCGContext* Context) const
 
 		if (!SpatialInput->ConstMetadata())
 		{
-			PCGE_LOG(Error, "Input does not have metadata");
+			PCGE_LOG(Error, GraphAndLog, LOCTEXT("InputMissingMetadata", "Input does not have metadata"));
 			continue;
 		}
 
@@ -42,7 +44,7 @@ bool FPCGMetadataPartitionElement::ExecuteInternal(FPCGContext* Context) const
 
 		if (!SpatialInput->ConstMetadata()->HasAttribute(LocalPartitionAttribute))
 		{
-			PCGE_LOG(Error, "Input does not have the %s attribute", *LocalPartitionAttribute.ToString());
+			PCGE_LOG(Error, GraphAndLog, FText::Format(LOCTEXT("InputMissingAttribute", "Input does not have the '{0}' attribute"), FText::FromName(LocalPartitionAttribute)));
 			continue;
 		}
 
@@ -51,7 +53,7 @@ bool FPCGMetadataPartitionElement::ExecuteInternal(FPCGContext* Context) const
 
 		if (!InputData)
 		{
-			PCGE_LOG(Error, "Unable to get points out of spatial data");
+			PCGE_LOG(Error, GraphAndLog, LOCTEXT("CouldNotObtainPoints", "Unable to get points out of spatial data"));
 			continue;
 		}
 
@@ -135,3 +137,5 @@ bool FPCGMetadataPartitionElement::ExecuteInternal(FPCGContext* Context) const
 
 	return true;
 }
+
+#undef LOCTEXT_NAMESPACE
diff --git a/Engine/Plugins/Experimental/PCG/Source/PCG/Private/Elements/PCGAttributeFilterElement.cpp b/Engine/Plugins/Experimental/PCG/Source/PCG/Private/Elements/PCGAttributeFilterElement.cpp
--- a/Engine/Plugins/Experimental/PCG/Source/PCG/Private/Elements/PCGAttributeFilterElement.cpp
+++ b/Engine/Plugins/Experimental/PCG/Source/PCG/Private/Elements/PCGAttributeFilterElement.cpp
@@ -10,6 +10,8 @@
 
 #include UE_INLINE_GENERATED_CPP_BY_NAME(PCGAttributeFilterElement)
 
+#define LOCTEXT_NAMESPACE "PCGAttributeFilterElement"
+
 namespace PCGAttributeFilterConstants
 {
 	const FName NodeName = TEXT("FilterAttribute");
@@ -140,7 +142,7 @@ bool FPCGAttributeFilterElement::ExecuteInternal(FPCGContext* Context) const
 		}
 		else
 		{
-			PCGE_LOG(Error, "Invalid data as input. Only support spatial and params");
+			PCGE_LOG(Error, GraphAndLog, LOCTEXT("InvalidInputData", "Invalid data as input. Only Spatial and Params data are supported."));
 			continue;
 		}
 
@@ -152,3 +154,4 @@ bool FPCGAttributeFilterElement::ExecuteInternal(FPCGContext* Context) const
 	return true;
 }
 
+#undef LOCTEXT_NAMESPACE
diff --git a/Engine/Plugins/Experimental/PCG/Source/PCG/Private/Elements/PCGAttributeGetFromPointIndexElement.cpp b/Engine/Plugins/Experimental/PCG/Source/PCG/Private/Elements/PCGAttributeGetFromPointIndexElement.cpp
--- a/Engine/Plugins/Experimental/PCG/Source/PCG/Private/Elements/PCGAttributeGetFromPointIndexElement.cpp
+++ b/Engine/Plugins/Experimental/PCG/Source/PCG/Private/Elements/PCGAttributeGetFromPointIndexElement.cpp
@@ -12,6 +12,8 @@
 
 #include UE_INLINE_GENERATED_CPP_BY_NAME(PCGAttributeGetFromPointIndexElement)
 
+#define LOCTEXT_NAMESPACE "PCGAttributeGetFromPointIndexElement"
+
 #if WITH_EDITOR
 FName UPCGAttributeGetFromPointIndexSettings::GetDefaultNodeName() const
 {
@@ -67,7 +69,7 @@ bool FPCGAttributeGetFromPointIndexElement::ExecuteInternal(FPCGContext* Context
 
 	if (Inputs.Num() != 1)
 	{
-		PCGE_LOG(Error, "Input pin doesn't have the right number of inputs.");
+		PCGE_LOG(Warning, LogOnly, FText::Format(LOCTEXT("WrongNumberOfInputs", "Input pin expected to have one input data element, encountered {0}"), Inputs.Num()));
 		return true;
 	}
 
@@ -75,7 +77,7 @@ bool FPCGAttributeGetFromPointIndexElement::ExecuteInternal(FPCGContext* Context
 
 	if (!PointData)
 	{
-		PCGE_LOG(Error, "Input is not a point data.");
+		PCGE_LOG(Error, GraphAndLog, LOCTEXT("InputNotPointData", "Input is not a point data"));
 		return true;
 	}
 
@@ -83,7 +85,7 @@ bool FPCGAttributeGetFromPointIndexElement::ExecuteInternal(FPCGContext* Context
 
 	if (Index < 0 || Index >= PointData->GetPoints().Num())
 	{
-		PCGE_LOG(Error, "Index is out of bounds. Index: %d; Number of Points: %d", Index, PointData->GetPoints().Num());
+		PCGE_LOG(Error, GraphAndLog, FText::Format(LOCTEXT("IndexOutOfBounds", "Index is out of bounds. Index: {0}; Number of Points: {1}"), Index, PointData->GetPoints().Num()));
 		return true;
 	}
 
@@ -136,7 +138,7 @@ bool FPCGAttributeGetFromPointIndexElement::ExecuteInternal(FPCGContext* Context
 
 			if (!NewAttribute)
 			{
-				PCGE_LOG(Error, "Error while creating target attribute %s", *OutputAttributeName.ToString());
+				PCGE_LOG(Error, GraphAndLog, FText::Format(LOCTEXT("ErrorCreatingTargetAttribute", "Error while creating target attribute '{0}'"), FText::FromName(OutputAttributeName)));
 				return false;
 			}
 
@@ -154,8 +156,10 @@ bool FPCGAttributeGetFromPointIndexElement::ExecuteInternal(FPCGContext* Context
 	}
 	else
 	{
-		PCGE_LOG(Warning, "Can't find attribute/property %s in input", *InputSource.GetName().ToString());
+		PCGE_LOG(Warning, GraphAndLog, FText::Format(LOCTEXT("AttributeNotFound", "Cannot find attribute/property '{0}' in input"), FText::FromName(InputSource.GetName())));
 	}
 
 	return true;
 }
+
+#undef LOCTEXT_NAMESPACE
diff --git a/Engine/Plugins/Experimental/PCG/Source/PCG/Private/Elements/PCGAttributeReduceElement.cpp b/Engine/Plugins/Experimental/PCG/Source/PCG/Private/Elements/PCGAttributeReduceElement.cpp
--- a/Engine/Plugins/Experimental/PCG/Source/PCG/Private/Elements/PCGAttributeReduceElement.cpp
+++ b/Engine/Plugins/Experimental/PCG/Source/PCG/Private/Elements/PCGAttributeReduceElement.cpp
@@ -12,6 +12,8 @@
 
 #include UE_INLINE_GENERATED_CPP_BY_NAME(PCGAttributeReduceElement)
 
+#define LOCTEXT_NAMESPACE "PCGAttributeReduceElement"
+
 namespace PCGAttributeReduceElement
 {
 	template <typename T>
@@ -155,7 +157,7 @@ bool FPCGAttributeReduceElement::ExecuteInternal(FPCGContext* Context) const
 
 	if (Inputs.Num() != 1)
 	{
-		PCGE_LOG(Error, "Input pin doesn't have the right number of inputs.");
+		PCGE_LOG(Warning, LogOnly, FText::Format(LOCTEXT("WrongNumberOfInputs", "Input pin expected to have one input data element, encountered {0}"), Inputs.Num()));
 		return true;
 	}
 
@@ -163,7 +165,7 @@ bool FPCGAttributeReduceElement::ExecuteInternal(FPCGContext* Context) const
 
 	if (!SpatialData)
 	{
-		PCGE_LOG(Error, "Input is not a spatial data.");
+		PCGE_LOG(Error, GraphAndLog, LOCTEXT("InputNotSpatialData", "Input is not a spatial data"));
 		return true;
 	}
 
@@ -183,7 +185,7 @@ bool FPCGAttributeReduceElement::ExecuteInternal(FPCGContext* Context) const
 
 	if (!Accessor.IsValid() || !Keys.IsValid())
 	{
-		PCGE_LOG(Error, "Input attribute/property doesn't exists");
+		PCGE_LOG(Error, GraphAndLog, LOCTEXT("AttributeDoesNotExist", "Input attribute/property does not exist"));
 		return true;
 	}
 
@@ -229,7 +231,7 @@ bool FPCGAttributeReduceElement::ExecuteInternal(FPCGContext* Context) const
 
 	if (!PCGMetadataAttribute::CallbackWithRightType(Accessor->GetUnderlyingType(), DoOperation))
 	{
-		PCGE_LOG(Error, "Operation was not compatible with the attribute type or couldn't create attribute %s.", *OutputAttributeName.ToString());
+		PCGE_LOG(Error, GraphAndLog, FText::Format(LOCTEXT("AttributeOperationFailed", "Operation was not compatible with the attribute type or could not create attribute '{0}'"), FText::FromName(OutputAttributeName)));
 		return true;
 	}
 
@@ -239,3 +241,5 @@ bool FPCGAttributeReduceElement::ExecuteInternal(FPCGContext* Context) const
 
 	return true;
 }
+
+#undef LOCTEXT_NAMESPACE
diff --git a/Engine/Plugins/Experimental/PCG/Source/PCG/Private/Elements/PCGAttributeSelectElement.cpp b/Engine/Plugins/Experimental/PCG/Source/PCG/Private/Elements/PCGAttributeSelectElement.cpp
--- a/Engine/Plugins/Experimental/PCG/Source/PCG/Private/Elements/PCGAttributeSelectElement.cpp
+++ b/Engine/Plugins/Experimental/PCG/Source/PCG/Private/Elements/PCGAttributeSelectElement.cpp
@@ -12,6 +12,8 @@
 
 #include UE_INLINE_GENERATED_CPP_BY_NAME(PCGAttributeSelectElement)
 
+#define LOCTEXT_NAMESPACE "PCGAttributeSelectElement"
+
 namespace PCGAttributeSelectElement
 {
 	// Of course, Dot doesn't exist for FVector4...
@@ -258,7 +260,7 @@ bool FPCGAttributeSelectElement::ExecuteInternal(FPCGContext* Context) const
 
 	if (Inputs.Num() != 1)
 	{
-		PCGE_LOG(Error, "Input pin doesn't have the right number of inputs.");
+		PCGE_LOG(Error, LogOnly, FText::Format(LOCTEXT("WrongNumberOfInputs", "Input pin expected to have one input data element, encountered {0}"), Inputs.Num()));
 		return true;
 	}
 
@@ -266,19 +268,19 @@ bool FPCGAttributeSelectElement::ExecuteInternal(FPCGContext* Context) const
 
 	if (!SpatialData)
 	{
-		PCGE_LOG(Error, "Input is not a spatial data.");
+		PCGE_LOG(Error, GraphAndLog, LOCTEXT("InputNotSpatialData", "Input is not a spatial data"));
 		return true;
 	}
 
 	const UPCGPointData* PointData = Cast<UPCGPointData>(SpatialData);
 	if (!PointData && Context->Node && Context->Node->IsOutputPinConnected(PCGAttributeSelectConstants::OutputPointLabel))
 	{
-		PCGE_LOG(Warning, "Not a point data as input, will output nothing in the %s output pin", *PCGAttributeSelectConstants::OutputPointLabel.ToString());
+		PCGE_LOG(Warning, GraphAndLog, FText::Format(LOCTEXT("InputMissingPointData", "No point data in input, will output nothing in the '{0}' output pin"), FText::FromName(PCGAttributeSelectConstants::OutputPointLabel)));
 	}
 
 	if (!SpatialData->Metadata)
 	{
-		PCGE_LOG(Error, "Input data doesn't have metadata");
+		PCGE_LOG(Error, GraphAndLog, LOCTEXT("InputMissingMetadata", "Input data doesn't have metadata"));
 		return true;
 	}
 
@@ -296,7 +298,7 @@ bool FPCGAttributeSelectElement::ExecuteInternal(FPCGContext* Context) const
 
 	if (!Accessor.IsValid() || !Keys.IsValid())
 	{
-		PCGE_LOG(Error, "Input attribute/property doesn't exists");
+		PCGE_LOG(Error, GraphAndLog, LOCTEXT("AttributeMissing", "Input attribute/property does not exist"));
 		return true;
 	}
 
@@ -306,7 +308,7 @@ bool FPCGAttributeSelectElement::ExecuteInternal(FPCGContext* Context) const
 
 		if constexpr (!PCG::Private::IsOfTypes<AttributeType, int32, int64, float, double, FVector2D, FVector, FVector4>())
 		{
-			PCGE_LOG(Error, "Attribute type is not a Vector nor a scalar");
+			PCGE_LOG(Error, GraphAndLog, LOCTEXT("AttributeNotVectorScalar", "Attribute type is not a Vector nor a scalar"));
 			return -1;
 		}
 		else
@@ -324,6 +326,8 @@ bool FPCGAttributeSelectElement::ExecuteInternal(FPCGContext* Context) const
 				return -1;
 			}
 
+			const FText InvalidErrorAxisMessage = LOCTEXT("InvalidAxis", "Invalid axis for attribute type");
+
 			// First we need to verify if the axis we want to project on is valid for dimension of our vector type.
 			// If it is a scalar, we won't project anything.
 			if constexpr (PCG::Private::IsOfTypes<AttributeType, FVector2D, FVector, FVector4>())
@@ -352,7 +356,7 @@ bool FPCGAttributeSelectElement::ExecuteInternal(FPCGContext* Context) const
 
 				if (!bIsValid)
 				{
-					PCGE_LOG(Error, "Invalid axis for attribute type.");
+					PCGE_LOG(Error, GraphAndLog, InvalidErrorAxisMessage);
 					return -1;
 				}
 			}
@@ -388,7 +392,7 @@ bool FPCGAttributeSelectElement::ExecuteInternal(FPCGContext* Context) const
 			}
 			else
 			{
-				PCGE_LOG(Error, "Invalid axis for attribute type.");
+				PCGE_LOG(Error, GraphAndLog, InvalidErrorAxisMessage);
 				OutputIndex = -1;
 			}
 
@@ -424,3 +428,5 @@ bool FPCGAttributeSelectElement::ExecuteInternal(FPCGContext* Context) const
 
 	return true;
 }
+
+#undef LOCTEXT_NAMESPACE
diff --git a/Engine/Plugins/Experimental/PCG/Source/PCG/Private/Elements/PCGAttributeTransferElement.cpp b/Engine/Plugins/Experimental/PCG/Source/PCG/Private/Elements/PCGAttributeTransferElement.cpp
--- a/Engine/Plugins/Experimental/PCG/Source/PCG/Private/Elements/PCGAttributeTransferElement.cpp
+++ b/Engine/Plugins/Experimental/PCG/Source/PCG/Private/Elements/PCGAttributeTransferElement.cpp
@@ -9,6 +9,8 @@
 
 #include UE_INLINE_GENERATED_CPP_BY_NAME(PCGAttributeTransferElement)
 
+#define LOCTEXT_NAMESPACE "PCGAttributeTransferElement"
+
 #if WITH_EDITOR
 FName UPCGAttributeTransferSettings::GetDefaultNodeName() const
 {
@@ -56,7 +58,7 @@ bool FPCGAttributeTransferElement::ExecuteInternal(FPCGContext* Context) const
 
 	if (SourceInputs.Num() != 1 || TargetInputs.Num() != 1)
 	{
-		PCGE_LOG(Error, "Wrong number of inputs");
+		PCGE_LOG(Warning, LogOnly, LOCTEXT("WrongNumberOfInputs", "Source input contains {0} data elements, Target inputs contain {1} data elements, but both should contain precisely 1 data element"));
 		return true;
 	}
 
@@ -65,19 +67,19 @@ bool FPCGAttributeTransferElement::ExecuteInternal(FPCGContext* Context) const
 
 	if (!SourceData || !TargetData || SourceData->IsA<UPCGPointData>() != TargetData->IsA<UPCGPointData>())
 	{
-		PCGE_LOG(Error, "Only support spatial to spatial data or point to point data");
+		PCGE_LOG(Error, GraphAndLog, LOCTEXT("UnsupportedTypes", "Only support Spatial to Spatial data or Point to Point data"));
 		return true;
 	}
 
 	if (!SourceData->Metadata)
 	{
-		PCGE_LOG(Error, "Source doesn't have metadata");
+		PCGE_LOG(Error, GraphAndLog, LOCTEXT("SourceMissingMetadata", "Source does not have metadata"));
 		return true;
 	}
 
 	if (!SourceData->Metadata->HasAttribute(Settings->SourceAttributeName))
 	{
-		PCGE_LOG(Error, "Source doesn't have the attribute \"%s\"", *Settings->SourceAttributeName.ToString());
+		PCGE_LOG(Error, GraphAndLog, FText::Format(LOCTEXT("SourceMissingAttribute", "Source does not have the attribute '{0}'"), FText::FromName(Settings->SourceAttributeName)));
 		return true;
 	}
 
@@ -114,7 +116,7 @@ UPCGSpatialData* FPCGAttributeTransferElement::TransferSpatialToSpatial(FPCGCont
 	// They need to match their number of entries
 	if (SourceData->Metadata->GetItemCountForChild() != TargetData->Metadata->GetItemCountForChild())
 	{
-		PCGE_LOG(Error, "Source and target doesn't have the same number of metadata entries");
+		PCGE_LOG(Error, GraphAndLog, LOCTEXT("MismatchingEntryCounts", "Source and target do not have the same number of metadata entries"));
 		return nullptr;
 	}
 
@@ -132,7 +134,7 @@ UPCGSpatialData* FPCGAttributeTransferElement::TransferSpatialToSpatial(FPCGCont
 
 	if (!NewSpatialData->Metadata->HasAttribute(TargetAttributeName))
 	{
-		PCGE_LOG(Error, "Error while creating target attribute %s", *TargetAttributeName.ToString());
+		PCGE_LOG(Error, GraphAndLog, FText::Format(LOCTEXT("ErrorCreatingAttribute", "Error while creating target attribute '{0}'"), FText::FromName(TargetAttributeName)));
 		return nullptr;
 	}
 
@@ -150,7 +152,7 @@ UPCGPointData* FPCGAttributeTransferElement::TransferPointToPoint(FPCGContext* C
 	// They need to match their number of points
 	if (SourceData->GetPoints().Num() != TargetData->GetPoints().Num())
 	{
-		PCGE_LOG(Error, "Source and target doesn't have the same number of points");
+		PCGE_LOG(Error, GraphAndLog, LOCTEXT("MismatchingPointCounts", "Source and target do not have the same number of points"));
 		return nullptr;
 	}
 
@@ -172,7 +174,7 @@ UPCGPointData* FPCGAttributeTransferElement::TransferPointToPoint(FPCGContext* C
 
 	if (!TargetAttribute)
 	{
-		PCGE_LOG(Error, "Error while creating target attribute %s", *TargetAttributeName.ToString());
+		PCGE_LOG(Error, GraphAndLog, FText::Format(LOCTEXT("ErrorCreatingTargetAttribute", "Error while creating target attribute '{0}'"), FText::FromName(TargetAttributeName)));
 		return nullptr;
 	}
 
@@ -188,3 +190,5 @@ UPCGPointData* FPCGAttributeTransferElement::TransferPointToPoint(FPCGContext* C
 
 	return NewPointData;
 }
+
+#undef LOCTEXT_NAMESPACE
diff --git a/Engine/Plugins/Experimental/PCG/Source/PCG/Private/Elements/PCGCopyPoints.cpp b/Engine/Plugins/Experimental/PCG/Source/PCG/Private/Elements/PCGCopyPoints.cpp
--- a/Engine/Plugins/Experimental/PCG/Source/PCG/Private/Elements/PCGCopyPoints.cpp
+++ b/Engine/Plugins/Experimental/PCG/Source/PCG/Private/Elements/PCGCopyPoints.cpp
@@ -54,7 +54,7 @@ bool FPCGCopyPointsElement::ExecuteInternal(FPCGContext* Context) const
 
 	if (Sources.Num() != 1 || Targets.Num() != 1)
 	{
-		PCGE_LOG(Error, "Invalid number of inputs - Expected one source, got %d; Expected one target, got %d.", Sources.Num(), Targets.Num());
+		PCGE_LOG(Warning, LogOnly, FText::Format(LOCTEXT("InvalidNumberOfInputs", "Invalid number of inputs. Expected one source, got {0}; Expected one target, got {1}."), Sources.Num(), Targets.Num()));
 		return true;
 	}
 	
@@ -65,7 +65,7 @@ bool FPCGCopyPointsElement::ExecuteInternal(FPCGContext* Context) const
 
 	if (!Source.Data || !Target.Data) 
 	{
-		PCGE_LOG(Error, "Invalid input data");
+		PCGE_LOG(Error, GraphAndLog, LOCTEXT("InvalidInputData", "Invalid input data"));
 		return true;
 	}
 
@@ -74,7 +74,7 @@ bool FPCGCopyPointsElement::ExecuteInternal(FPCGContext* Context) const
 
 	if (!SourceSpatialData || !TargetSpatialData)
 	{
-		PCGE_LOG(Error, "Unable to get SpatialData from input");
+		PCGE_LOG(Error, GraphAndLog, LOCTEXT("CouldNotObtainSpatialData", "Unable to get Spatial Data from input"));
 		return true;
 	}
 
@@ -83,7 +83,7 @@ bool FPCGCopyPointsElement::ExecuteInternal(FPCGContext* Context) const
 
 	if (!SourcePointData || !TargetPointData)
 	{
-		PCGE_LOG(Error, "Unable to get PointData from input");
+		PCGE_LOG(Error, GraphAndLog, LOCTEXT("CouldNotGetPointData", "Unable to get Point Data from input"));
 		return true;
 	}
 
diff --git a/Engine/Plugins/Experimental/PCG/Source/PCG/Private/Elements/PCGCreateAttribute.cpp b/Engine/Plugins/Experimental/PCG/Source/PCG/Private/Elements/PCGCreateAttribute.cpp
--- a/Engine/Plugins/Experimental/PCG/Source/PCG/Private/Elements/PCGCreateAttribute.cpp
+++ b/Engine/Plugins/Experimental/PCG/Source/PCG/Private/Elements/PCGCreateAttribute.cpp
@@ -10,6 +10,8 @@
 
 #include UE_INLINE_GENERATED_CPP_BY_NAME(PCGCreateAttribute)
 
+#define LOCTEXT_NAMESPACE "PCGCreateAttributeElement"
+
 namespace PCGCreateAttributeConstants
 {
 	const FName NodeName = TEXT("CreateAttribute");
@@ -131,7 +133,7 @@ bool FPCGCreateAttributeElement::ExecuteInternal(FPCGContext* Context) const
 	{
 		if (SourceParams.IsEmpty())
 		{
-			PCGE_LOG(Error, "Source param was not provided.");
+			PCGE_LOG(Error, GraphAndLog, LOCTEXT("ParamNotProvided", "Source param was not provided"));
 			return true;
 		}
 
@@ -139,7 +141,7 @@ bool FPCGCreateAttributeElement::ExecuteInternal(FPCGContext* Context) const
 
 		if (!SourceParamData->Metadata)
 		{
-			PCGE_LOG(Error, "Source param data doesn't have metadata");
+			PCGE_LOG(Error, GraphAndLog, LOCTEXT("ParamMissingMetadata", "Source param data does not have metadata"));
 			return true;
 		}
 
@@ -147,7 +149,7 @@ bool FPCGCreateAttributeElement::ExecuteInternal(FPCGContext* Context) const
 
 		if (!SourceParamData->Metadata->HasAttribute(SourceParamAttributeName))
 		{
-			PCGE_LOG(Error, "Source param data doesn't have an attribute \"%s\"", *SourceParamAttributeName.ToString());
+			PCGE_LOG(Error, GraphAndLog, FText::Format(LOCTEXT("ParamMissingAttribute", "Source param data does not have an attribute '{0}'"), FText::FromName(SourceParamAttributeName)));
 			return true;
 		}
 	}
@@ -202,7 +204,7 @@ bool FPCGCreateAttributeElement::ExecuteInternal(FPCGContext* Context) const
 		}
 		else
 		{
-			PCGE_LOG(Error, "Invalid data as input. Only support spatial and params");
+			PCGE_LOG(Error, GraphAndLog, LOCTEXT("InvalidInputData", "Invalid data as input. Only Spatial and Param data supported."));
 			continue;
 		}
 
@@ -222,7 +224,7 @@ bool FPCGCreateAttributeElement::ExecuteInternal(FPCGContext* Context) const
 
 		if (!Attribute)
 		{
-			PCGE_LOG(Error, "Error while creating attribute %s", *OutputAttributeName.ToString());
+			PCGE_LOG(Error, GraphAndLog, FText::Format(LOCTEXT("ErrorCreatingAttribute", "Error while creating attribute '{0}'"), FText::FromName(OutputAttributeName)));
 			continue;
 		}
 
@@ -275,3 +277,4 @@ PCGMetadataEntryKey FPCGCreateAttributeElement::SetAttribute(const UPCGCreateAtt
 	return Settings->AttributeTypes.Dispatcher(SetAttribute);
 }
 
+#undef LOCTEXT_NAMESPACE
diff --git a/Engine/Plugins/Experimental/PCG/Source/PCG/Private/Elements/PCGCreateSpline.cpp b/Engine/Plugins/Experimental/PCG/Source/PCG/Private/Elements/PCGCreateSpline.cpp
--- a/Engine/Plugins/Experimental/PCG/Source/PCG/Private/Elements/PCGCreateSpline.cpp
+++ b/Engine/Plugins/Experimental/PCG/Source/PCG/Private/Elements/PCGCreateSpline.cpp
@@ -79,7 +79,7 @@ bool FPCGCreateSplineElement::ExecuteInternal(FPCGContext* Context) const
 
 		if (!SpatialData)
 		{
-			PCGE_LOG(Error, "Invalid input data");
+			PCGE_LOG(Error, GraphAndLog, LOCTEXT("InvalidInputData", "Invalid input data"));
 			continue;
 		}
 
@@ -87,14 +87,14 @@ bool FPCGCreateSplineElement::ExecuteInternal(FPCGContext* Context) const
 
 		if (!PointData)
 		{
-			PCGE_LOG(Error, "Unable to get PointData from input");
+			PCGE_LOG(Error, GraphAndLog, LOCTEXT("UnableToGetPointData", "Unable to get point data from input"));
 			continue;
 		}
 
 		AActor* TargetActor = PointData->TargetActor.Get();
 		if (!TargetActor)
 		{
-			PCGE_LOG(Error, "Invalid target actor");
+			PCGE_LOG(Error, GraphAndLog, LOCTEXT("InvalidTargetActor", "Invalid target actor"));
 			continue;
 		}
 
@@ -109,17 +109,19 @@ bool FPCGCreateSplineElement::ExecuteInternal(FPCGContext* Context) const
 			FName LocalArriveTangentName = ((Settings->ArriveTangentAttribute == NAME_None) ? PointMetadata->GetLatestAttributeNameOrNone() : Settings->ArriveTangentAttribute);
 			FName LocalLeaveTangentName = ((Settings->LeaveTangentAttribute == NAME_None) ? PointMetadata->GetLatestAttributeNameOrNone() : Settings->LeaveTangentAttribute);
 
+			const FText AttributeMissingOrNotVector = LOCTEXT("AttributeMissingOrNotVector", "Attribute '{0}' does not exist or is not a vector");
+
 			const FPCGMetadataAttributeBase* ArriveTangentBaseAttribute = PointMetadata->GetConstAttribute(LocalArriveTangentName);
 			if (!ArriveTangentBaseAttribute || ArriveTangentBaseAttribute->GetTypeId() != PCG::Private::MetadataTypes<FVector>::Id)
 			{
-				PCGE_LOG(Error, "Attribute %s does not exist or is not a vector", *LocalArriveTangentName.ToString());
+				PCGE_LOG(Error, GraphAndLog, FText::Format(AttributeMissingOrNotVector, FText::FromString(LocalArriveTangentName.ToString())));
 				continue;
 			}
 
 			const FPCGMetadataAttributeBase* LeaveTangentBaseAttribute = PointMetadata->GetConstAttribute(LocalLeaveTangentName);
 			if (!LeaveTangentBaseAttribute || LeaveTangentBaseAttribute->GetTypeId() != PCG::Private::MetadataTypes<FVector>::Id)
 			{
-				PCGE_LOG(Error, "Attribute %s does not exist or is not a vector", *LocalLeaveTangentName.ToString());
+				PCGE_LOG(Error, GraphAndLog, FText::Format(AttributeMissingOrNotVector, FText::FromString(LocalLeaveTangentName.ToString())));
 				continue;
 			}
 
@@ -149,7 +151,7 @@ bool FPCGCreateSplineElement::ExecuteInternal(FPCGContext* Context) const
 
 			if (!SplineActor)
 			{
-				PCGE_LOG(Error, "Failed to create actor to hold the spline");
+				PCGE_LOG(Error, GraphAndLog, LOCTEXT("FailedToCreateActor", "Failed to create actor to hold the spline"));
 				continue;
 			}
 
diff --git a/Engine/Plugins/Experimental/PCG/Source/PCG/Private/Elements/PCGDataFromActor.cpp b/Engine/Plugins/Experimental/PCG/Source/PCG/Private/Elements/PCGDataFromActor.cpp
--- a/Engine/Plugins/Experimental/PCG/Source/PCG/Private/Elements/PCGDataFromActor.cpp
+++ b/Engine/Plugins/Experimental/PCG/Source/PCG/Private/Elements/PCGDataFromActor.cpp
@@ -78,7 +78,7 @@ bool FPCGDataFromActorElement::ExecuteInternal(FPCGContext* InContext) const
 
 		if (Context->FoundActors.IsEmpty())
 		{
-			PCGE_LOG(Warning, "No matching actor was found.");
+			PCGE_LOG(Warning, GraphAndLog, LOCTEXT("NoActorFound", "No matching actor was found"));
 			return true;
 		}
 
@@ -110,7 +110,7 @@ bool FPCGDataFromActorElement::ExecuteInternal(FPCGContext* InContext) const
 				}
 				else
 				{
-					PCGE_LOG(Error, "Was unable to wait for end of generation tasks");
+					PCGE_LOG(Error, GraphAndLog, LOCTEXT("UnableToWaitForGenerationTasks", "Unable to wait for end of generation tasks"));
 				}
 			}
 		}
@@ -188,18 +188,18 @@ void FPCGDataFromActorElement::ProcessActor(FPCGContext* Context, const UPCGData
 	{
 		if (bCanGetDataFromComponent)
 		{
-			PCGE_LOG(Warning, "Actor (%s) does not have any previously generated data.", *FoundActor->GetFName().ToString());
+			PCGE_LOG(Warning, GraphAndLog, FText::Format(LOCTEXT("ActorHasNoGeneratedData", "Actor '{0}' does not have any previously generated data"), FText::FromName(FoundActor->GetFName())));
 		}
 		else
 		{
-			PCGE_LOG(Error, "Actor (%s) cannot get its own generated data during generation", *FoundActor->GetFName().ToString());
+			PCGE_LOG(Error, GraphAndLog, FText::Format(LOCTEXT("ActorCannotGetOwnData", "Actor '{0}' cannot get its own generated data during generation"), FText::FromName(FoundActor->GetFName())));
 		}
 		
 		return;
 	}
 	else if (Settings->Mode == EPCGGetDataFromActorMode::GetDataFromProperty && !FoundProperty)
 	{
-		PCGE_LOG(Warning, "Actor (%s) does not have a property name (%s)", *FoundActor->GetFName().ToString(), *Settings->PropertyName.ToString());
+		PCGE_LOG(Warning, GraphAndLog, FText::Format(LOCTEXT("ActorHasNoProperty", "Actor '{0}' does not have a property name '{1}'"), FText::FromName(FoundActor->GetFName()), FText::FromName(Settings->PropertyName)));
 		return;
 	}
 
@@ -244,7 +244,7 @@ void FPCGDataFromActorElement::ProcessActor(FPCGContext* Context, const UPCGData
 
 		if (!bAbleToGetProperty)
 		{
-			PCGE_LOG(Warning, "Actor (%s) property (%s) does not have a supported type", *FoundActor->GetFName().ToString(), *Settings->PropertyName.ToString());
+			PCGE_LOG(Warning, GraphAndLog, FText::Format(LOCTEXT("PropertyTypeUnsupported", "Actor '{0}' property '{1}' does not have a supported type"), FText::FromName(FoundActor->GetFName()), FText::FromName(Settings->PropertyName)));
 		}
 	}
 	else
diff --git a/Engine/Plugins/Experimental/PCG/Source/PCG/Private/Elements/PCGDataTableRowToParamData.cpp b/Engine/Plugins/Experimental/PCG/Source/PCG/Private/Elements/PCGDataTableRowToParamData.cpp
--- a/Engine/Plugins/Experimental/PCG/Source/PCG/Private/Elements/PCGDataTableRowToParamData.cpp
+++ b/Engine/Plugins/Experimental/PCG/Source/PCG/Private/Elements/PCGDataTableRowToParamData.cpp
@@ -12,6 +12,8 @@
 
 #include UE_INLINE_GENERATED_CPP_BY_NAME(PCGDataTableRowToParamData)
 
+#define LOCTEXT_NAMESPACE "PCGDataTableRowToParamDataElement"
+
 FPCGElementPtr UPCGDataTableRowToParamDataSettings::CreateElement() const
 {
 	return MakeShared<FPCGDataTableRowToParamData>();
@@ -55,7 +57,7 @@ bool FPCGDataTableRowToParamData::ExecuteInternal(FPCGContext* Context) const
 	const UDataTable* DataTable = DataTablePtr.LoadSynchronous();
 	if (!DataTable)
 	{
-		PCGE_LOG(Error, TEXT("FPCGDataTableRowToParamData: Data table is invalid"));
+		PCGE_LOG(Warning, GraphAndLog, LOCTEXT("DataTableInvalid", "FPCGDataTableRowToParamData: Data table is invalid"));
 		return true;
 	}
 
@@ -63,7 +65,7 @@ bool FPCGDataTableRowToParamData::ExecuteInternal(FPCGContext* Context) const
 	const uint8* RowDataPtr = RowDataItr ? *RowDataItr : nullptr;
 	if (!RowDataPtr)
 	{
-		PCGE_LOG(Error, TEXT("FPCGDataTableRowToParamData: Data table '%s' has no row named '%s'"), *DataTable->GetPathName(), *RowName.ToString());
+		PCGE_LOG(Error, GraphAndLog, FText::Format(LOCTEXT("NoRowFound", "FPCGDataTableRowToParamData: Data table '{0}' has no row named '{1}'"), FText::FromString(DataTable->GetPathName()), FText::FromName(RowName)));
 		return true;
 	}
 
@@ -77,10 +79,10 @@ bool FPCGDataTableRowToParamData::ExecuteInternal(FPCGContext* Context) const
 	{
 		const FString FieldName = DataTable->GetRowStruct()->GetAuthoredNameForField(*FieldIt);
 		const FName AttributeName(FieldName);
-
+		
 		if (!Metadata->SetAttributeFromDataProperty(AttributeName, EntryKey, RowDataPtr, *FieldIt, /*bCreate=*/ true))
 		{
-			PCGE_LOG(Warning, "Error while creating a attribute '%s'. Either the property type is not supported by PCG or attribute creation failed.", *AttributeName.ToString());
+			PCGE_LOG(Warning, GraphAndLog, FText::Format(LOCTEXT("AttributeCreationFailed", "Error while creating new attribute '{0}'. Either the property type is not supported by PCG or attribute creation failed."), FText::FromName(AttributeName)));
 		}
 	}
 	
@@ -90,3 +92,5 @@ bool FPCGDataTableRowToParamData::ExecuteInternal(FPCGContext* Context) const
 
 	return true;
 }
+
+#undef LOCTEXT_NAMESPACE
diff --git a/Engine/Plugins/Experimental/PCG/Source/PCG/Private/Elements/PCGDensityFilter.cpp b/Engine/Plugins/Experimental/PCG/Source/PCG/Private/Elements/PCGDensityFilter.cpp
--- a/Engine/Plugins/Experimental/PCG/Source/PCG/Private/Elements/PCGDensityFilter.cpp
+++ b/Engine/Plugins/Experimental/PCG/Source/PCG/Private/Elements/PCGDensityFilter.cpp
@@ -10,6 +10,8 @@
 
 #include UE_INLINE_GENERATED_CPP_BY_NAME(PCGDensityFilter)
 
+#define LOCTEXT_NAMESPACE "PCGDensityFilterElement"
+
 FPCGElementPtr UPCGDensityFilterSettings::CreateElement() const
 {
 	return MakeShared<FPCGDensityFilterElement>();
@@ -42,7 +44,7 @@ bool FPCGDensityFilterElement::ExecuteInternal(FPCGContext* Context) const
 
 	if (bNoResults && !bKeepZeroDensityPoints)
 	{
-		PCGE_LOG(Verbose, "Skipped - all inputs rejected");
+		PCGE_LOG(Verbose, LogOnly, LOCTEXT("AllInputsRejected", "Skipped - all inputs rejected"));
 		return true;
 	}
 
@@ -53,14 +55,14 @@ bool FPCGDensityFilterElement::ExecuteInternal(FPCGContext* Context) const
 
 		if (!Input.Data || Cast<UPCGSpatialData>(Input.Data) == nullptr)
 		{
-			PCGE_LOG(Error, "Invalid input data");
+			PCGE_LOG(Error, GraphAndLog, LOCTEXT("InvalidInputData", "Invalid input data"));
 			continue;
 		}
 
 		// Skip processing if the transformation is trivial
 		if (bTrivialFilter)
 		{
-			PCGE_LOG(Verbose, "Skipped - trivial filter");
+			PCGE_LOG(Verbose, LogOnly, LOCTEXT("TrivialFilter", "Skipped - trivial filter"));
 			continue;
 		}
 
@@ -68,7 +70,7 @@ bool FPCGDensityFilterElement::ExecuteInternal(FPCGContext* Context) const
 
 		if (!OriginalData)
 		{
-			PCGE_LOG(Error, "Unable to get point data from input");
+			PCGE_LOG(Error, GraphAndLog, LOCTEXT("NoPointDataInInput", "Unable to get point data from input"));
 			continue;
 		}
 
@@ -104,8 +106,10 @@ bool FPCGDensityFilterElement::ExecuteInternal(FPCGContext* Context) const
 			}
 		});
 
-		PCGE_LOG(Verbose, "Generated %d points out of %d source points", FilteredPoints.Num(), Points.Num());
+		PCGE_LOG(Verbose, LogOnly, FText::Format(LOCTEXT("GenerationInfo", "Generated {0} points out of {1} source points"), FilteredPoints.Num(), Points.Num()));
 	}
 
 	return true;
-}
\ No newline at end of file
+}
+
+#undef LOCTEXT_NAMESPACE
diff --git a/Engine/Plugins/Experimental/PCG/Source/PCG/Private/Elements/PCGDensityRemap.cpp b/Engine/Plugins/Experimental/PCG/Source/PCG/Private/Elements/PCGDensityRemap.cpp
--- a/Engine/Plugins/Experimental/PCG/Source/PCG/Private/Elements/PCGDensityRemap.cpp
+++ b/Engine/Plugins/Experimental/PCG/Source/PCG/Private/Elements/PCGDensityRemap.cpp
@@ -13,6 +13,8 @@
 
 #include UE_INLINE_GENERATED_CPP_BY_NAME(PCGDensityRemap)
 
+#define LOCTEXT_NAMESPACE "PCGDensityRemapElement"
+
 UPCGLinearDensityRemapSettings::UPCGLinearDensityRemapSettings()
 {
 	bUseSeed = true;
@@ -50,7 +52,7 @@ bool FPCGLinearDensityRemapElement::ExecuteInternal(FPCGContext* Context) const
 
 		if (!Input.Data || Cast<UPCGSpatialData>(Input.Data) == nullptr)
 		{
-			PCGE_LOG(Error, "Invalid input data");
+			PCGE_LOG(Error, GraphAndLog, LOCTEXT("InvalidInputData", "Input data missing or not of type Spatial"));
 			continue;
 		}
 
@@ -58,7 +60,7 @@ bool FPCGLinearDensityRemapElement::ExecuteInternal(FPCGContext* Context) const
 
 		if (!OriginalData)
 		{
-			PCGE_LOG(Error, "Unable to get points from input");
+			PCGE_LOG(Error, GraphAndLog, LOCTEXT("NoPointsInInput", "Unable to obtain points from input data"));
 			continue;
 		}
 
@@ -94,3 +96,5 @@ bool FPCGLinearDensityRemapElement::ExecuteInternal(FPCGContext* Context) const
 
 	return true;
 }
+
+#undef LOCTEXT_NAMESPACE
diff --git a/Engine/Plugins/Experimental/PCG/Source/PCG/Private/Elements/PCGDistance.cpp b/Engine/Plugins/Experimental/PCG/Source/PCG/Private/Elements/PCGDistance.cpp
--- a/Engine/Plugins/Experimental/PCG/Source/PCG/Private/Elements/PCGDistance.cpp
+++ b/Engine/Plugins/Experimental/PCG/Source/PCG/Private/Elements/PCGDistance.cpp
@@ -7,6 +7,8 @@
 #include "PCGContext.h"
 #include "PCGPin.h"
 
+#define LOCTEXT_NAMESPACE "PCGDistanceElement"
+
 namespace PCGDistance
 {
 	const FName SourceLabel = TEXT("Source");
@@ -92,14 +94,14 @@ bool FPCGDistanceElement::ExecuteInternal(FPCGContext* Context) const
 
 		if (!TargetData)
 		{
-			PCGE_LOG(Error, TEXT("Target must be spatial data, found '%s'"), *Target.Data->GetClass()->GetName());
+			PCGE_LOG(Error, GraphAndLog, FText::Format(LOCTEXT("TargetMustBeSpatial", "Target must be Spatial data, found '{0}'"), FText::FromString(Target.Data->GetClass()->GetName())));
 			continue;
 		}
 
 		const UPCGPointData* TargetPointData = TargetData->ToPointData(Context);
 		if (!TargetPointData)
 		{
-			PCGE_LOG(Error, TEXT("Cannot convert target '%s' into point data"), *Target.Data->GetClass()->GetName());
+			PCGE_LOG(Error, GraphAndLog, FText::Format(LOCTEXT("CannotConvertToPoint", "Cannot convert target '{0}' into Point data"), FText::FromString(Target.Data->GetClass()->GetName())));
 			continue;			
 		}
 
@@ -115,14 +117,14 @@ bool FPCGDistanceElement::ExecuteInternal(FPCGContext* Context) const
 
 		if (!SourceData)
 		{
-			PCGE_LOG(Error, TEXT("Invalid input data"));
+			PCGE_LOG(Error, GraphAndLog, LOCTEXT("InvalidInputData", "Invalid input data"));
 			continue;
 		}
 
 		const UPCGPointData* SourcePointData = SourceData->ToPointData(Context);
 		if (!SourcePointData)
 		{
-			PCGE_LOG(Error, TEXT("Cannot convert input spatial data to point data"));
+			PCGE_LOG(Error, GraphAndLog, LOCTEXT("CannotConvertToPointData", "Cannot convert input Spatial data to Point data"));
 			continue;			
 		}
 
@@ -217,3 +219,5 @@ bool FPCGDistanceElement::ExecuteInternal(FPCGContext* Context) const
 
 	return true;
 }
+
+#undef LOCTEXT_NAMESPACE
diff --git a/Engine/Plugins/Experimental/PCG/Source/PCG/Private/Elements/PCGExecuteBlueprint.cpp b/Engine/Plugins/Experimental/PCG/Source/PCG/Private/Elements/PCGExecuteBlueprint.cpp
--- a/Engine/Plugins/Experimental/PCG/Source/PCG/Private/Elements/PCGExecuteBlueprint.cpp
+++ b/Engine/Plugins/Experimental/PCG/Source/PCG/Private/Elements/PCGExecuteBlueprint.cpp
@@ -15,6 +15,8 @@
 
 #include UE_INLINE_GENERATED_CPP_BY_NAME(PCGExecuteBlueprint)
 
+#define LOCTEXT_NAMESPACE "PCGBlueprintElement"
+
 #if WITH_EDITOR
 namespace PCGBlueprintHelper
 {
@@ -539,7 +541,7 @@ bool FPCGExecuteBlueprintElement::ExecuteInternal(FPCGContext* InContext) const
 			const FPCGTaggedData& Input = Context->InputData.TaggedData[InputIndex];
 			if (const UPCGPointData* PointData = Cast<UPCGPointData>(Input.Data))
 			{
-				PCGE_LOG(Verbose, "Input %d has %d points", InputIndex, PointData->GetPoints().Num());
+				PCGE_LOG(Verbose, LogOnly, FText::Format(LOCTEXT("InputPointInfo", "Input {0} has {1} points"), InputIndex, PointData->GetPoints().Num()));
 			}
 		}
 
@@ -552,7 +554,7 @@ bool FPCGExecuteBlueprintElement::ExecuteInternal(FPCGContext* InContext) const
 			const FPCGTaggedData& Output = Context->OutputData.TaggedData[OutputIndex];
 			if (const UPCGPointData* PointData = Cast<UPCGPointData>(Output.Data))
 			{
-				PCGE_LOG(Verbose, "Output %d has %d points", OutputIndex, PointData->GetPoints().Num());
+				PCGE_LOG(Verbose, LogOnly, FText::Format(LOCTEXT("OutputPointInfo", "Output {0} has {1} points"), OutputIndex, PointData->GetPoints().Num()));
 			}
 
 			// Important implementation note:
@@ -616,7 +618,7 @@ void UPCGBlueprintElement::LoopOnPoints(FPCGContext& InContext, const UPCGPointD
 {
 	if (!InData)
 	{
-		PCGE_LOG_C(Error, &InContext, "Invalid input data in LoopOnPoints");
+		PCGE_LOG_C(Error, GraphAndLog, &InContext, LOCTEXT("InvalidInputDataLoopOnPoints", "Invalid input data in LoopOnPoints"));
 		return;
 	}
 
@@ -643,7 +645,7 @@ void UPCGBlueprintElement::MultiLoopOnPoints(FPCGContext& InContext, const UPCGP
 {
 	if (!InData)
 	{
-		PCGE_LOG_C(Error, &InContext, "Invalid input data in MultiLoopOnPoints");
+		PCGE_LOG_C(Error, GraphAndLog, &InContext, LOCTEXT("InvalidInputDataMultiLoopOnPoints", "Invalid input data in MultiLoopOnPoints"));
 		return;
 	}
 
@@ -670,7 +672,7 @@ void UPCGBlueprintElement::LoopOnPointPairs(FPCGContext& InContext, const UPCGPo
 {
 	if (!InA || !InB)
 	{
-		PCGE_LOG_C(Error, &InContext, "Invalid input data in LoopOnPointPairs");
+		PCGE_LOG_C(Error, GraphAndLog, &InContext, LOCTEXT("InvalidInputDataLoopOnPointPairs", "Invalid input data in LoopOnPointPairs"));
 		return;
 	}
 
@@ -789,3 +791,5 @@ bool FPCGExecuteBlueprintElement::CanExecuteOnlyOnMainThread(FPCGContext* Contex
 		return false;
 	}
 }
+
+#undef LOCTEXT_NAMESPACE
diff --git a/Engine/Plugins/Experimental/PCG/Source/PCG/Private/Elements/PCGMergeElement.cpp b/Engine/Plugins/Experimental/PCG/Source/PCG/Private/Elements/PCGMergeElement.cpp
--- a/Engine/Plugins/Experimental/PCG/Source/PCG/Private/Elements/PCGMergeElement.cpp
+++ b/Engine/Plugins/Experimental/PCG/Source/PCG/Private/Elements/PCGMergeElement.cpp
@@ -67,7 +67,7 @@ bool FPCGMergeElement::ExecuteInternal(FPCGContext* Context) const
 
 		if (!SourcePointData)
 		{
-			PCGE_LOG(Error, "Unsupported data type in merge");
+			PCGE_LOG(Error, GraphAndLog, LOCTEXT("UnsupportedDataType", "Unsupported data type in merge"));
 			continue;
 		}
 
diff --git a/Engine/Plugins/Experimental/PCG/Source/PCG/Private/Elements/PCGPointFilter.cpp b/Engine/Plugins/Experimental/PCG/Source/PCG/Private/Elements/PCGPointFilter.cpp
--- a/Engine/Plugins/Experimental/PCG/Source/PCG/Private/Elements/PCGPointFilter.cpp
+++ b/Engine/Plugins/Experimental/PCG/Source/PCG/Private/Elements/PCGPointFilter.cpp
@@ -14,6 +14,8 @@
 
 #include UE_INLINE_GENERATED_CPP_BY_NAME(PCGPointFilter)
 
+#define LOCTEXT_NAMESPACE "PCGPointFilterElement"
+
 namespace PCGPointFilterConstants
 {
 	const FName DataToFilterLabel = TEXT("In");
@@ -251,7 +253,7 @@ bool FPCGPointFilterElement::ExecuteInternal(FPCGContext* Context) const
 					ThresholdSpatialData = TentativeThresholdSpatialData->ToPointData();
 					if (!ThresholdSpatialData)
 					{
-						PCGE_LOG(Error, "Unable to get point data from filter input");
+						PCGE_LOG(Error, GraphAndLog, LOCTEXT("CannotGetPointData", "Unable to get Point data from filter input"));
 						return true;
 					}
 
@@ -281,14 +283,14 @@ bool FPCGPointFilterElement::ExecuteInternal(FPCGContext* Context) const
 		const UPCGSpatialData* SpatialInput = Cast<const UPCGSpatialData>(Input.Data);
 		if (!SpatialInput)
 		{
-			PCGE_LOG(Error, "Please move non-spatial data to other pins");
+			PCGE_LOG(Error, GraphAndLog, LOCTEXT("MoveNonSpatialData", "Please move non-Spatial data to other pins"));
 			continue;
 		}
 
 		const UPCGPointData* OriginalData = SpatialInput->ToPointData(Context);
 		if (!OriginalData)
 		{
-			PCGE_LOG(Error, "Unable to get point data from input");
+			PCGE_LOG(Error, GraphAndLog, LOCTEXT("NoPointDataInInput", "Unable to get point data from input"));
 			continue;
 		}
 
@@ -314,7 +316,7 @@ bool FPCGPointFilterElement::ExecuteInternal(FPCGContext* Context) const
 
 		if (!TargetAccessor.IsValid() || !TargetKeys.IsValid())
 		{
-			PCGE_LOG(Warning, "TargetData doesn't have %s target attribute/property", *Settings->TargetAttribute.GetName().ToString());
+			PCGE_LOG(Warning, GraphAndLog, FText::Format(LOCTEXT("TargetMissingAttribute", "TargetData doesn't have target attribute/property '{0}'"), FText::FromName(Settings->TargetAttribute.GetName())));
 			ForwardInputToOutFilterPin();
 			continue;
 		}
@@ -341,7 +343,7 @@ bool FPCGPointFilterElement::ExecuteInternal(FPCGContext* Context) const
 
 		if (!ThresholdAccessor.IsValid() || !ThresholdKeys.IsValid())
 		{
-			PCGE_LOG(Warning, "DataToFilter doesn't have %s threshold attribute/property", *Settings->ThresholdAttribute.GetName().ToString());
+			PCGE_LOG(Warning, GraphAndLog, FText::Format(LOCTEXT("AttributeMissingForFilter", "DataToFilter does not have '{0}' threshold attribute/property"), FText::FromName(Settings->ThresholdAttribute.GetName())));
 			ForwardInputToInFilterPin();
 			continue;
 		}
@@ -351,7 +353,7 @@ bool FPCGPointFilterElement::ExecuteInternal(FPCGContext* Context) const
 		// For example: if target is double but threshold is int32, we can broadcast int32 to double, to compare a double with a double.
 		if (!PCG::Private::IsBroadcastable(ThresholdAccessor->GetUnderlyingType(), TargetAccessor->GetUnderlyingType()))
 		{
-			PCGE_LOG(Warning, "Cannot broadcast threshold type to target type");
+			PCGE_LOG(Warning, GraphAndLog, LOCTEXT("TypeConversionFailed", "Cannot broadcast threshold type to target type"));
 			ForwardInputToInFilterPin();
 			continue;
 		}
@@ -366,7 +368,7 @@ bool FPCGPointFilterElement::ExecuteInternal(FPCGContext* Context) const
 
 			if (!bCanCompare)
 			{
-				PCGE_LOG(Warning, "Cannot compare target type.");
+				PCGE_LOG(Warning, GraphAndLog, LOCTEXT("TypeComparisonFailed", "Cannot compare target type"));
 				ForwardInputToOutFilterPin();
 				continue;
 			}
@@ -471,3 +473,5 @@ bool FPCGPointFilterElement::ExecuteInternal(FPCGContext* Context) const
 
 	return true;
 }
+
+#undef LOCTEXT_NAMESPACE
diff --git a/Engine/Plugins/Experimental/PCG/Source/PCG/Private/Elements/PCGPointMatchAndSet.cpp b/Engine/Plugins/Experimental/PCG/Source/PCG/Private/Elements/PCGPointMatchAndSet.cpp
--- a/Engine/Plugins/Experimental/PCG/Source/PCG/Private/Elements/PCGPointMatchAndSet.cpp
+++ b/Engine/Plugins/Experimental/PCG/Source/PCG/Private/Elements/PCGPointMatchAndSet.cpp
@@ -165,7 +165,7 @@ bool FPCGPointMatchAndSetElement::ExecuteInternal(FPCGContext* Context) const
 
 	if (!MatchAndSet)
 	{
-		PCGE_LOG(Error, "Invalid MatchAndSet instance");
+		PCGE_LOG(Error, GraphAndLog, LOCTEXT("InvalidMatchAndSetInstance", "Invalid MatchAndSet instance, try recreating this node from the node palette"));
 		return true;
 	}
 
@@ -179,7 +179,7 @@ bool FPCGPointMatchAndSetElement::ExecuteInternal(FPCGContext* Context) const
 		const UPCGPointData* InPointData = Cast<UPCGPointData>(Input.Data);
 		if (!InPointData)
 		{
-			PCGE_LOG(Error, "Invalid input data type");
+			PCGE_LOG(Error, GraphAndLog, LOCTEXT("InvalidInputDataType", "Input data must be of type Point"));
 			continue;
 		}
 
@@ -187,7 +187,7 @@ bool FPCGPointMatchAndSetElement::ExecuteInternal(FPCGContext* Context) const
 		// conforms to what is expected (esp. on the metadata side)
 		if (!MatchAndSet->ValidatePreconditions(InPointData))
 		{
-			PCGE_LOG(Error, "MatchAndSet failed to validate preconditions on input data");
+			PCGE_LOG(Error, GraphAndLog, LOCTEXT("PreconditionsFailed", "MatchAndSet failed to validate preconditions on input data"));
 			continue;
 		}
 
diff --git a/Engine/Plugins/Experimental/PCG/Source/PCG/Private/Elements/PCGPointProcessingElementBase.cpp b/Engine/Plugins/Experimental/PCG/Source/PCG/Private/Elements/PCGPointProcessingElementBase.cpp
--- a/Engine/Plugins/Experimental/PCG/Source/PCG/Private/Elements/PCGPointProcessingElementBase.cpp
+++ b/Engine/Plugins/Experimental/PCG/Source/PCG/Private/Elements/PCGPointProcessingElementBase.cpp
@@ -7,6 +7,8 @@
 #include "Helpers/PCGAsync.h"
 #include "PCGContext.h"
 
+#define LOCTEXT_NAMESPACE "PCGPointProcessingElementBase"
+
 void FPCGPointProcessingElementBase::ProcessPoints(FPCGContext* Context, const TArray<FPCGTaggedData>& Inputs, TArray<FPCGTaggedData>& Outputs, const TFunction<bool(const FPCGPoint&, FPCGPoint&)>& PointFunc) const
 {
 	// TODO: embarassingly parallel loop
@@ -19,7 +21,7 @@ void FPCGPointProcessingElementBase::ProcessPoints(FPCGContext* Context, const T
 
 		if (!SpatialData)
 		{
-			PCGE_LOG(Error, "Unable to get SpatialData from input");
+			PCGE_LOG(Error, GraphAndLog, LOCTEXT("InputMissingSpatialData", "Unable to get Spatial data from input"));
 			continue;
 		}
 
@@ -27,7 +29,7 @@ void FPCGPointProcessingElementBase::ProcessPoints(FPCGContext* Context, const T
 
 		if (!PointData)
 		{
-			PCGE_LOG(Error, "Unable to get PointData from input");
+			PCGE_LOG(Error, GraphAndLog, LOCTEXT("InputMissingPointData", "Unable to get Point data from input"));
 			continue;
 		}
 
@@ -41,3 +43,5 @@ void FPCGPointProcessingElementBase::ProcessPoints(FPCGContext* Context, const T
 		FPCGAsync::AsyncPointProcessing(Context, Points, OutputPoints, PointFunc);
 	}
 }
+
+#undef LOCTEXT_NAMESPACE
diff --git a/Engine/Plugins/Experimental/PCG/Source/PCG/Private/Elements/PCGPointSampler.cpp b/Engine/Plugins/Experimental/PCG/Source/PCG/Private/Elements/PCGPointSampler.cpp
--- a/Engine/Plugins/Experimental/PCG/Source/PCG/Private/Elements/PCGPointSampler.cpp
+++ b/Engine/Plugins/Experimental/PCG/Source/PCG/Private/Elements/PCGPointSampler.cpp
@@ -10,6 +10,8 @@
 
 #include UE_INLINE_GENERATED_CPP_BY_NAME(PCGPointSampler)
 
+#define LOCTEXT_NAMESPACE "PCGPointSamplerElement"
+
 UPCGPointSamplerSettings::UPCGPointSamplerSettings()
 {
 	bUseSeed = true;
@@ -46,7 +48,7 @@ bool FPCGPointSamplerElement::ExecuteInternal(FPCGContext* Context) const
 	// Early exit when nothing will be generated out of this sampler
 	if (bNoSampling && !bKeepZeroDensityPoints)
 	{
-		PCGE_LOG(Verbose, "Skipped - all inputs rejected");
+		PCGE_LOG(Verbose, LogOnly, LOCTEXT("AllInputsRejected", "Skipped - all inputs rejected"));
 		return true;
 	}
 
@@ -57,14 +59,14 @@ bool FPCGPointSamplerElement::ExecuteInternal(FPCGContext* Context) const
 
 		if (!Input.Data || Cast<UPCGSpatialData>(Input.Data) == nullptr)
 		{
-			PCGE_LOG(Error, "Invalid input data");
+			PCGE_LOG(Error, GraphAndLog, LOCTEXT("InvalidInputData", "Invalid input data"));
 			continue;
 		}
 
 		// Skip processing if the transformation would be trivial
 		if (bTrivialSampling)
 		{
-			PCGE_LOG(Verbose, "Skipped - trivial sampling");
+			PCGE_LOG(Verbose, LogOnly, LOCTEXT("SkippedTrivialSampling", "Skipped - trivial sampling"));
 			continue;
 		}
 
@@ -72,7 +74,7 @@ bool FPCGPointSamplerElement::ExecuteInternal(FPCGContext* Context) const
 
 		if (!OriginalData)
 		{
-			PCGE_LOG(Error, "Unable to get point data from input");
+			PCGE_LOG(Error, GraphAndLog, LOCTEXT("NoPointDataInInput", "Unable to get point data from input"));
 			continue;
 		}
 
@@ -95,7 +97,7 @@ bool FPCGPointSamplerElement::ExecuteInternal(FPCGContext* Context) const
 		// Early out
 		if (TargetNumPoints == 0)
 		{
-			PCGE_LOG(Verbose, "Skipped - all points rejected");
+			PCGE_LOG(Verbose, LogOnly, LOCTEXT("SkippedAllPointsRejected", "Skipped - all points rejected"));
 			continue;
 		}
 		else
@@ -129,10 +131,11 @@ bool FPCGPointSamplerElement::ExecuteInternal(FPCGContext* Context) const
 				}
 			});
 
-			PCGE_LOG(Verbose, "Generated %d points from %d source points", SampledPoints.Num(), OriginalPointCount);
+			PCGE_LOG(Verbose, LogOnly, FText::Format(LOCTEXT("GenerationInfo", "Generated {0} points from {1} source points"), SampledPoints.Num(), OriginalPointCount));
 		}
 	}
 
 	return true;
 }
 
+#undef LOCTEXT_NAMESPACE
diff --git a/Engine/Plugins/Experimental/PCG/Source/PCG/Private/Elements/PCGProjectionElement.cpp b/Engine/Plugins/Experimental/PCG/Source/PCG/Private/Elements/PCGProjectionElement.cpp
--- a/Engine/Plugins/Experimental/PCG/Source/PCG/Private/Elements/PCGProjectionElement.cpp
+++ b/Engine/Plugins/Experimental/PCG/Source/PCG/Private/Elements/PCGProjectionElement.cpp
@@ -74,7 +74,7 @@ bool FPCGProjectionElement::ExecuteInternal(FPCGContext* Context) const
 
 		if (!ProjectionSource)
 		{
-			PCGE_LOG(Error, "Invalid projection source data input found (non-spatial data). Input will be ignored.");
+			PCGE_LOG(Error, GraphAndLog, LOCTEXT("InvalidProjectionSource", "Invalid projection source data input found (non-Spatial data). Input will be ignored."));
 			continue;
 		}
 
diff --git a/Engine/Plugins/Experimental/PCG/Source/PCG/Private/Elements/PCGPropertyToParamData.cpp b/Engine/Plugins/Experimental/PCG/Source/PCG/Private/Elements/PCGPropertyToParamData.cpp
--- a/Engine/Plugins/Experimental/PCG/Source/PCG/Private/Elements/PCGPropertyToParamData.cpp
+++ b/Engine/Plugins/Experimental/PCG/Source/PCG/Private/Elements/PCGPropertyToParamData.cpp
@@ -10,9 +10,10 @@
 #include "Helpers/PCGSettingsHelpers.h"
 #include "Metadata/Accessors/PCGAttributeAccessorHelpers.h"
 
-
 #include UE_INLINE_GENERATED_CPP_BY_NAME(PCGPropertyToParamData)
 
+#define LOCTEXT_NAMESPACE "PCGPropertyToParamDataElement"
+
 #if WITH_EDITOR
 void UPCGPropertyToParamDataSettings::GetTrackedActorTags(FPCGTagToSettingsMap& OutTagToSettings, TArray<TObjectPtr<const UPCGGraph>>& OutVisitedGraphs) const
 {
@@ -78,7 +79,7 @@ bool FPCGPropertyToParamDataElement::ExecuteInternal(FPCGContext* Context) const
 	// Early out if arguments are not specified
 	if (Settings->PropertyName == NAME_None || (Settings->bSelectComponent && !Settings->ComponentClass))
 	{
-		PCGE_LOG(Error, "Some parameters are missing, abort.");
+		PCGE_LOG(Warning, GraphAndLog, LOCTEXT("ParametersMissing", "Some parameters are missing, aborting"));
 		return true;
 	}
 
@@ -87,7 +88,7 @@ bool FPCGPropertyToParamDataElement::ExecuteInternal(FPCGContext* Context) const
 	// Optimization possibly only in non-editor builds, otherwise we could poison the input-driven cache
 	if (!Context->Node || !Context->Node->IsOutputPinConnected(PCGPinConstants::DefaultOutputLabel))
 	{
-		PCGE_LOG(Verbose, "Node is not connected, nothing to do");
+		PCGE_LOG(Verbose, LogOnly, LOCTEXT("UnconnectedNode", "Node is not connected, nothing to do"));
 		return true;
 	}
 #endif
@@ -98,7 +99,7 @@ bool FPCGPropertyToParamDataElement::ExecuteInternal(FPCGContext* Context) const
 
 	if (!FoundActor)
 	{
-		PCGE_LOG(Error, "No matching actor was found.");
+		PCGE_LOG(Error, GraphAndLog, LOCTEXT("NoActorFound", "No matching actor was found"));
 		return true;
 	}
 
@@ -109,7 +110,7 @@ bool FPCGPropertyToParamDataElement::ExecuteInternal(FPCGContext* Context) const
 		ObjectToInspect = FoundActor->GetComponentByClass(Settings->ComponentClass);
 		if (!ObjectToInspect)
 		{
-			PCGE_LOG(Error, "Component doesn't exist in the found actor.");
+			PCGE_LOG(Error, GraphAndLog, LOCTEXT("ComponentDoesNotExist", "Component does not exist in the found actor"));
 			return true;
 		}
 	}
@@ -118,7 +119,7 @@ bool FPCGPropertyToParamDataElement::ExecuteInternal(FPCGContext* Context) const
 	FProperty* Property = FindFProperty<FProperty>(ObjectToInspect->GetClass(), Settings->PropertyName);
 	if (!Property)
 	{
-		PCGE_LOG(Error, "Property doesn't exist in the found actor.");
+		PCGE_LOG(Error, GraphAndLog, LOCTEXT("PropertyDoesNotExist", "Property does not exist in the found actor"));
 		return true;
 	}
 
@@ -176,7 +177,7 @@ bool FPCGPropertyToParamDataElement::ExecuteInternal(FPCGContext* Context) const
 
 	if (ExtractableProperties.IsEmpty())
 	{
-		PCGE_LOG(Error, "Found no properties to be extractable.");
+		PCGE_LOG(Error, GraphAndLog, LOCTEXT("NoPropertiesFound", "No properties found to extract"));
 		return true;
 	}
 
@@ -195,7 +196,7 @@ bool FPCGPropertyToParamDataElement::ExecuteInternal(FPCGContext* Context) const
 
 		if (!Metadata->SetAttributeFromDataProperty(AttributeName, EntryKey, ContainerPtr, FinalProperty, /*bCreate=*/ true))
 		{
-			PCGE_LOG(Error, "Error while creating an attribute for property %s. Either the property type is not supported by PCG or attribute creation failed.", *FinalProperty->GetName());
+			PCGE_LOG(Error, GraphAndLog, FText::Format(LOCTEXT("ErrorCreatingAttribute", "Error while creating an attribute for property '{0}'. Either the property type is not supported by PCG or attribute creation failed."), FText::FromString(FinalProperty->GetName())));
 			continue;
 		}
 
@@ -211,3 +212,5 @@ bool FPCGPropertyToParamDataElement::ExecuteInternal(FPCGContext* Context) const
 
 	return true;
 }
+
+#undef LOCTEXT_NAMESPACE
diff --git a/Engine/Plugins/Experimental/PCG/Source/PCG/Private/Elements/PCGSelfPruning.cpp b/Engine/Plugins/Experimental/PCG/Source/PCG/Private/Elements/PCGSelfPruning.cpp
--- a/Engine/Plugins/Experimental/PCG/Source/PCG/Private/Elements/PCGSelfPruning.cpp
+++ b/Engine/Plugins/Experimental/PCG/Source/PCG/Private/Elements/PCGSelfPruning.cpp
@@ -7,9 +7,10 @@
 #include "Data/PCGPointData.h"
 #include "Data/PCGSpatialData.h"
 
-
 #include UE_INLINE_GENERATED_CPP_BY_NAME(PCGSelfPruning)
 
+#define LOCTEXT_NAMESPACE "PCGSelfPruningElement"
+
 namespace PCGSelfPruningAlgorithms
 {
 	bool RandomSort(const FPCGPointRef& A, const FPCGPointRef& B)
@@ -129,7 +130,7 @@ namespace PCGSelfPruningElement
 		if (PruningType == EPCGSelfPruningType::None)
 		{
 			Context->OutputData = Context->InputData;
-			PCGE_LOG_C(Verbose, Context, "Skipped - Type is none");
+			PCGE_LOG_C(Verbose, LogOnly, Context, LOCTEXT("TypeNotSpecified", "Skipped - Type is None"));
 			return;
 		}
 
@@ -146,7 +147,7 @@ namespace PCGSelfPruningElement
 
 			if (!SpatialInput)
 			{
-				PCGE_LOG_C(Error, Context, "Invalid input data");
+				PCGE_LOG_C(Error, GraphAndLog, Context, LOCTEXT("InvalidInputData", "Invalid input data"));
 				continue;
 			}
 
@@ -230,11 +231,11 @@ namespace PCGSelfPruningElement
 
 			if (bIsDuplicateTest)
 			{
-				PCGE_LOG_C(Verbose, Context, "Removed %d duplicate points from %d source points", Points.Num() - OutputPoints.Num(), Points.Num());
+				PCGE_LOG_C(Verbose, LogOnly, Context, FText::Format(LOCTEXT("GenerationInfoDuplicate", "Removed {0} duplicate points from {1} source points"), Points.Num() - OutputPoints.Num(), Points.Num()));
 			}
 			else
 			{
-				PCGE_LOG_C(Verbose, Context, "Generated %d points from %d source points", OutputPoints.Num(), Points.Num());
+				PCGE_LOG_C(Verbose, LogOnly, Context, FText::Format(LOCTEXT("GenerationInfo", "Generated {0} points from {1} source points"), OutputPoints.Num(), Points.Num()));
 			}
 		}
 	}
@@ -260,3 +261,5 @@ bool FPCGSelfPruningElement::ExecuteInternal(FPCGContext* Context) const
 
 	return true;
 }
+
+#undef LOCTEXT_NAMESPACE
diff --git a/Engine/Plugins/Experimental/PCG/Source/PCG/Private/Elements/PCGSpawnActor.cpp b/Engine/Plugins/Experimental/PCG/Source/PCG/Private/Elements/PCGSpawnActor.cpp
--- a/Engine/Plugins/Experimental/PCG/Source/PCG/Private/Elements/PCGSpawnActor.cpp
+++ b/Engine/Plugins/Experimental/PCG/Source/PCG/Private/Elements/PCGSpawnActor.cpp
@@ -24,6 +24,8 @@
 
 #include UE_INLINE_GENERATED_CPP_BY_NAME(PCGSpawnActor)
 
+#define LOCTEXT_NAMESPACE "PCGSpawnActorElement"
+
 static TAutoConsoleVariable<bool> CVarAllowActorReuse(
 	TEXT("pcg.Actor.AllowReuse"),
 	true,
@@ -309,7 +311,8 @@ bool FPCGSpawnActorElement::ExecuteInternal(FPCGContext* Context) const
 	// Early out
 	if(!Settings->TemplateActorClass || Settings->TemplateActorClass->HasAnyClassFlags(CLASS_Abstract))
 	{
-		PCGE_LOG(Error, "Invalid template actor class (%s)", Settings->TemplateActorClass ? *Settings->TemplateActorClass->GetFName().ToString() : TEXT("None"));
+		const FText ClassName = Settings->TemplateActorClass ? FText::FromString(Settings->TemplateActorClass->GetFName().ToString()) : FText::FromName(NAME_None);
+		PCGE_LOG(Error, GraphAndLog, FText::Format(LOCTEXT("InvalidTemplateActorClass", "Invalid template actor class '{0}'"), ClassName));
 		return true;
 	}
 
@@ -371,7 +374,7 @@ bool FPCGSpawnActorElement::ExecuteInternal(FPCGContext* Context) const
 		{
 			if (PostSpawnFunction->NumParms != 0)
 			{
-				PCGE_LOG(Warning, "PostSpawnFunction \"%s\" requires parameters. We only support parameter-less functions. Will skip the call.", *PostSpawnFunctionName.ToString());
+				PCGE_LOG(Warning, GraphAndLog, FText::Format(LOCTEXT("ParametersMissing", "PostSpawnFunction '{0}' requires parameters. We only support parameter-less functions. Call will be skipped."), FText::FromName(PostSpawnFunctionName)));
 			}
 			else
 			{
@@ -380,7 +383,7 @@ bool FPCGSpawnActorElement::ExecuteInternal(FPCGContext* Context) const
 		}
 		else
 		{
-			PCGE_LOG(Warning, "PostSpawnFunction \"%s\" was not found in class \"%s\".", *PostSpawnFunctionName.ToString(), *Settings->TemplateActorClass->GetFName().ToString());
+			PCGE_LOG(Warning, GraphAndLog, FText::Format(LOCTEXT("FunctionNotFound", "PostSpawnFunction '{0}' was not found in class '{1}'"), FText::FromName(PostSpawnFunctionName), FText::FromName(Settings->TemplateActorClass->GetFName())));
 		}
 	}
 
@@ -396,7 +399,7 @@ bool FPCGSpawnActorElement::ExecuteInternal(FPCGContext* Context) const
 
 		if (!SpatialData)
 		{
-			PCGE_LOG(Error, "Invalid input data");
+			PCGE_LOG(Error, GraphAndLog, LOCTEXT("InvalidInputData", "Invalid input data"));
 			continue;
 		}
 
@@ -404,7 +407,7 @@ bool FPCGSpawnActorElement::ExecuteInternal(FPCGContext* Context) const
 
 		if (!TargetActor)
 		{
-			PCGE_LOG(Error, "Invalid target actor");
+			PCGE_LOG(Error, GraphAndLog, LOCTEXT("InvalidTargetActor", "Invalid target actor"));
 			continue;
 		}
 
@@ -416,7 +419,7 @@ bool FPCGSpawnActorElement::ExecuteInternal(FPCGContext* Context) const
 
 		if (!PointData)
 		{
-			PCGE_LOG(Error, "Unable to get point data from input");
+			PCGE_LOG(Error, GraphAndLog, LOCTEXT("NoPointDataInInput", "Unable to get point data from input"));
 			continue;
 		}
 
@@ -424,7 +427,7 @@ bool FPCGSpawnActorElement::ExecuteInternal(FPCGContext* Context) const
 
 		if (Points.IsEmpty())
 		{
-			PCGE_LOG(Verbose, "Skipped - no points");
+			PCGE_LOG(Verbose, LogOnly, LOCTEXT("SkippedNoPoints", "Skipped - no points"));
 			continue;
 		}
 
@@ -498,7 +501,8 @@ bool FPCGSpawnActorElement::ExecuteInternal(FPCGContext* Context) const
 				ISMC->AddInstances(Transforms, false, true);
 				ISMC->UpdateBounds();
 				
-				PCGE_LOG(Verbose, "Added %d instances of %s to %s in %s", Transforms.Num(), *ISMC->GetStaticMesh().GetName(), *ISMC->GetName(), *TargetActor->GetActorNameOrLabel());
+				PCGE_LOG(Verbose, LogOnly, FText::Format(LOCTEXT("InstanceCreationInfo", "Added {0} instances of mesh '{1}' to ISMC '{2}' on actor '{3}'"),
+					Transforms.Num(), FText::FromString(ISMC->GetStaticMesh().GetName()), FText::FromString(ISMC->GetName()), FText::FromString(TargetActor->GetActorNameOrLabel())));
 			}
 		}
 		else if (!bFullySkippedDueToReuse && Settings->Option != EPCGSpawnActorOption::CollapseActors && (bHasAuthority || !bSpawnedActorsRequireAuthority))
@@ -588,7 +592,7 @@ bool FPCGSpawnActorElement::ExecuteInternal(FPCGContext* Context) const
 
 					if (!GeneratedActor)
 					{
-						PCGE_LOG(Error, "Failed to spawn actor on point with index %d", i);
+						PCGE_LOG(Error, GraphAndLog, FText::Format(LOCTEXT("ActorSpawnFailed", "Failed to spawn actor on point with index {0}"), i));
 						continue;
 					}
 
@@ -612,7 +616,7 @@ bool FPCGSpawnActorElement::ExecuteInternal(FPCGContext* Context) const
 
 				Context->SourceComponent->AddToManagedResources(ManagedActors);
 
-				PCGE_LOG(Verbose, "Generated %d actors", Points.Num());
+				PCGE_LOG(Verbose, LogOnly, FText::Format(LOCTEXT("GenerationInfo", "Generated {0} actors"), Points.Num()));
 			}
 
 			// Setup & Generate on PCG components if needed
@@ -696,3 +700,5 @@ EPCGSettingsType UPCGSpawnActorSettings::GetType() const
 	return GetSubgraph() ? EPCGSettingsType::Subgraph : EPCGSettingsType::Spawner;
 }
 #endif // WITH_EDITOR
+
+#undef LOCTEXT_NAMESPACE
diff --git a/Engine/Plugins/Experimental/PCG/Source/PCG/Private/Elements/PCGSplineSampler.cpp b/Engine/Plugins/Experimental/PCG/Source/PCG/Private/Elements/PCGSplineSampler.cpp
--- a/Engine/Plugins/Experimental/PCG/Source/PCG/Private/Elements/PCGSplineSampler.cpp
+++ b/Engine/Plugins/Experimental/PCG/Source/PCG/Private/Elements/PCGSplineSampler.cpp
@@ -1148,7 +1148,7 @@ bool FPCGSplineSamplerElement::ExecuteInternal(FPCGContext* Context) const
 
 		if (!BoundingShapeSpatialInput)
 		{
-			PCGE_LOG(Warning, "Bounding Shape input is missing or of unsupported type and will not be used.");
+			PCGE_LOG(Warning, GraphAndLog, LOCTEXT("BoundingShapeMissing", "Bounding Shape input is missing or of unsupported type and will not be used."));
 		}
 	}
 
diff --git a/Engine/Plugins/Experimental/PCG/Source/PCG/Private/Elements/PCGStaticMeshSpawner.cpp b/Engine/Plugins/Experimental/PCG/Source/PCG/Private/Elements/PCGStaticMeshSpawner.cpp
--- a/Engine/Plugins/Experimental/PCG/Source/PCG/Private/Elements/PCGStaticMeshSpawner.cpp
+++ b/Engine/Plugins/Experimental/PCG/Source/PCG/Private/Elements/PCGStaticMeshSpawner.cpp
@@ -19,6 +19,8 @@
 
 #include UE_INLINE_GENERATED_CPP_BY_NAME(PCGStaticMeshSpawner)
 
+#define LOCTEXT_NAMESPACE "PCGStaticMeshSpawnerElement"
+
 static TAutoConsoleVariable<bool> CVarAllowISMReuse(
 	TEXT("pcg.ISM.AllowReuse"),
 	true,
@@ -53,7 +55,7 @@ bool FPCGStaticMeshSpawnerElement::PrepareDataInternal(FPCGContext* InContext) c
 
 	if (!Settings->MeshSelectorInstance)
 	{
-		PCGE_LOG(Error, "Invalid MeshSelectorInstance");
+		PCGE_LOG(Error, GraphAndLog, LOCTEXT("InvalidMeshSelectorInstance", "Invalid MatchAndSet instance, try recreating this node from the node palette"));
 		return true;
 	}
 
@@ -127,7 +129,7 @@ bool FPCGStaticMeshSpawnerElement::PrepareDataInternal(FPCGContext* InContext) c
 
 			if (!SpatialData)
 			{
-				PCGE_LOG(Error, "Invalid input data");
+				PCGE_LOG(Error, GraphAndLog, LOCTEXT("InvalidInputData", "Invalid input data"));
 				++Context->CurrentInputIndex;
 				continue;
 			}
@@ -135,7 +137,7 @@ bool FPCGStaticMeshSpawnerElement::PrepareDataInternal(FPCGContext* InContext) c
 			const UPCGPointData* PointData = SpatialData->ToPointData(Context);
 			if (!PointData)
 			{
-				PCGE_LOG(Error, "Unable to get point data from input");
+				PCGE_LOG(Error, GraphAndLog, LOCTEXT("NoPointDataInInput", "Unable to get point data from input"));
 				++Context->CurrentInputIndex;
 				continue;
 			}
@@ -143,7 +145,7 @@ bool FPCGStaticMeshSpawnerElement::PrepareDataInternal(FPCGContext* InContext) c
 			AActor* TargetActor = Context->GetTargetActor(PointData);
 			if (!TargetActor)
 			{
-				PCGE_LOG(Error, "Invalid target actor");
+				PCGE_LOG(Error, GraphAndLog, LOCTEXT("InvalidTargetActor", "Invalid target actor"));
 				++Context->CurrentInputIndex;
 				continue;
 			}
@@ -158,7 +160,7 @@ bool FPCGStaticMeshSpawnerElement::PrepareDataInternal(FPCGContext* InContext) c
 				if (OutputPointData->Metadata->HasAttribute(Settings->OutAttributeName))
 				{
 					OutputPointData->Metadata->DeleteAttribute(Settings->OutAttributeName);
-					PCGE_LOG(Verbose, "Metadata attribute %s is being overwritten in the output data", *Settings->OutAttributeName.ToString());
+					PCGE_LOG(Verbose, LogOnly, FText::Format(LOCTEXT("AttributeOverwritten", "Metadata attribute '{0}' is being overwritten in the output data"), FText::FromName(Settings->OutAttributeName)));
 				}
 
 				OutputPointData->Metadata->CreateStringAttribute(Settings->OutAttributeName, FName(NAME_None).ToString(), /*bAllowsInterpolation=*/false);
@@ -304,7 +306,7 @@ void FPCGStaticMeshSpawnerElement::SpawnStaticMeshInstances(FPCGContext* Context
 		// Either we have no mesh (so nothing to do) or the mesh couldn't be loaded
 		if (InstanceList.Descriptor.StaticMesh.IsValid())
 		{
-			PCGE_LOG(Error, "Unable to load mesh %s", *InstanceList.Descriptor.StaticMesh.ToString());
+			PCGE_LOG(Error, GraphAndLog, FText::Format(LOCTEXT("MeshLoadFailed", "Unable to load mesh '{0}'"), FText::FromString(InstanceList.Descriptor.StaticMesh.ToString())));
 		}
 
 		return;
@@ -315,7 +317,7 @@ void FPCGStaticMeshSpawnerElement::SpawnStaticMeshInstances(FPCGContext* Context
 	{
 		if (OverrideMaterial.IsValid() && !OverrideMaterial.LoadSynchronous())
 		{
-			PCGE_LOG(Error, "Unable to load override material %s", *OverrideMaterial.ToString());
+			PCGE_LOG(Error, GraphAndLog, FText::Format(LOCTEXT("OverrideMaterialLoadFailed", "Unable to load override material '{0}'"), FText::FromString(OverrideMaterial.ToString())));
 			return;
 		}
 	}
@@ -365,7 +367,10 @@ void FPCGStaticMeshSpawnerElement::SpawnStaticMeshInstances(FPCGContext* Context
 
 	ISMC->UpdateBounds();
 
-	PCGE_LOG(Verbose, "Added %d instances of %s on actor %s", InstanceList.Instances.Num(), *InstanceList.Descriptor.StaticMesh->GetFName().ToString(), *TargetActor->GetFName().ToString());
+	{
+		PCGE_LOG(Verbose, LogOnly, FText::Format(LOCTEXT("GenerationInfo", "Added {0} instances of '{1}' on actor '{2}'"),
+			InstanceList.Instances.Num(), FText::FromString(InstanceList.Descriptor.StaticMesh->GetFName().ToString()), FText::FromString(TargetActor->GetFName().ToString())));
+	}
 }
 
 void UPCGStaticMeshSpawnerSettings::PostLoad()
@@ -502,3 +507,5 @@ void UPCGStaticMeshSpawnerSettings::RefreshInstancePacker()
 
 FPCGStaticMeshSpawnerContext::FPackedInstanceListData::FPackedInstanceListData() = default;
 FPCGStaticMeshSpawnerContext::FPackedInstanceListData::~FPackedInstanceListData() = default;
+
+#undef LOCTEXT_NAMESPACE
diff --git a/Engine/Plugins/Experimental/PCG/Source/PCG/Private/Elements/PCGSurfaceSampler.cpp b/Engine/Plugins/Experimental/PCG/Source/PCG/Private/Elements/PCGSurfaceSampler.cpp
--- a/Engine/Plugins/Experimental/PCG/Source/PCG/Private/Elements/PCGSurfaceSampler.cpp
+++ b/Engine/Plugins/Experimental/PCG/Source/PCG/Private/Elements/PCGSurfaceSampler.cpp
@@ -70,7 +70,7 @@ namespace PCGSurfaceSampler
 			{
 				if (Context)
 				{
-					PCGE_LOG_C(Verbose, Context, "Skipped - invalid cell bounds (%lld x %lld)", CellCountX, CellCountY);
+					PCGE_LOG_C(Verbose, LogOnly, Context, FText::Format(LOCTEXT("InvalidCellBounds", "Skipped - invalid cell bounds({0} x {1})"), CellCountX, CellCountY));
 				}
 
 				return false;
@@ -81,7 +81,7 @@ namespace PCGSurfaceSampler
 			{
 				if (Context)
 				{
-					PCGE_LOG_C(Error, Context, "Skipped - tried to generate too many points (%lld)", CellCount64);
+					PCGE_LOG_C(Error, GraphAndLog, Context, FText::Format(LOCTEXT("TooManyPoints", "Skipped - tried to generate too many points {0}"), CellCount64));
 				}
 
 				return false;
@@ -99,7 +99,7 @@ namespace PCGSurfaceSampler
 		{
 			if (Context)
 			{
-				PCGE_LOG_C(Verbose, Context, "Skipped - density yields no points");
+				PCGE_LOG_C(Verbose, LogOnly, Context, LOCTEXT("NoPointsFromDensity", "Skipped - density yields no points"));
 			}
 			
 			return false;
@@ -130,7 +130,7 @@ namespace PCGSurfaceSampler
 		// We don't support time slicing here
 		if (LoopData.bEnableTimeSlicing)
 		{
-			PCGE_LOG_C(Error, Context, "An output point data must be provided to support sampling with time slicing");
+			PCGE_LOG_C(Error, LogOnly, Context, LOCTEXT("MissingOutputPointData", "An output point data must be provided to support sampling with time slicing"));
 			return nullptr;
 		}
 
@@ -227,7 +227,7 @@ namespace PCGSurfaceSampler
 
 		if (Context && bAsyncDone)
 		{
-			PCGE_LOG_C(Verbose, Context, "Generated %d points in %d cells", SampledPoints.Num(), LoopData.CellCount);
+			PCGE_LOG_C(Verbose, LogOnly, Context, FText::Format(LOCTEXT("GenerationInfo", "Generated {0} points in {1} cells"), SampledPoints.Num(), LoopData.CellCount));
 		}
 
 		return bAsyncDone;
@@ -337,7 +337,7 @@ bool FPCGSurfaceSamplerElement::AddGeneratingShapesToContext(FPCGSurfaceSamplerC
 	const FVector& PointExtents = Settings->PointExtents;
 	if (PointExtents.X <= 0 || PointExtents.Y <= 0)
 	{
-		PCGE_LOG_C(Warning, InContext, "Skipped - Invalid point extents");
+		PCGE_LOG_C(Warning, GraphAndLog, InContext, LOCTEXT("SkippedInvalidPointExtents", "Skipped - Invalid point extents"));
 		return false;
 	}
 
@@ -359,7 +359,7 @@ bool FPCGSurfaceSamplerElement::AddGeneratingShapesToContext(FPCGSurfaceSamplerC
 	}
 	else if (BoundingShapeInputs.Num() > 0)
 	{
-		PCGE_LOG_C(Verbose, InContext, "The bounds of the Bounding Shape input pin will be ignored because the Unbounded option is enabled.");
+		PCGE_LOG_C(Verbose, LogOnly, InContext, LOCTEXT("BoundsIgnored", "The bounds of the Bounding Shape input pin will be ignored because the Unbounded option is enabled."));
 	}
 
 	if (InContext->BoundingShapeSpatialInput)
@@ -413,7 +413,7 @@ bool FPCGSurfaceSamplerElement::AddGeneratingShapesToContext(FPCGSurfaceSamplerC
 	// Warn if something is connected but no shape could be obtained for sampling
 	if (InContext->GeneratingShapes.Num() == 0 && (BoundingShapeInputs.Num() > 0 || SurfaceInputs.Num() > 0))
 	{
-		PCGE_LOG_C(Warning, InContext, "No Surface input was provided, and no surface could be found in the Bounding Shape input for sampling. Connect the surface to be sampled to the Surface input.");
+		PCGE_LOG_C(Warning, GraphAndLog, InContext, LOCTEXT("NoGenerator", "No Surface input was provided, and no surface could be found in the Bounding Shape input for sampling. Connect the surface to be sampled to the Surface input."));
 		return false;
 	}
 
@@ -460,7 +460,7 @@ void FPCGSurfaceSamplerElement::AllocateOutputs(FPCGSurfaceSamplerContext* InCon
 		{
 			if (!InputBounds.IsValid)
 			{
-				PCGE_LOG_C(Verbose, InContext, "Input data has invalid bounds");
+				PCGE_LOG_C(Warning, GraphAndLog, InContext, LOCTEXT("InvalidInputBounds", "Input data has invalid bounds"));
 			}
 
 			Outputs.RemoveAt(GeneratingShapeIndex);
diff --git a/Engine/Plugins/Experimental/PCG/Source/PCG/Private/Elements/PCGTransformPoints.cpp b/Engine/Plugins/Experimental/PCG/Source/PCG/Private/Elements/PCGTransformPoints.cpp
--- a/Engine/Plugins/Experimental/PCG/Source/PCG/Private/Elements/PCGTransformPoints.cpp
+++ b/Engine/Plugins/Experimental/PCG/Source/PCG/Private/Elements/PCGTransformPoints.cpp
@@ -12,6 +12,8 @@
 
 #include UE_INLINE_GENERATED_CPP_BY_NAME(PCGTransformPoints)
 
+#define LOCTEXT_NAMESPACE "PCGTransformPointsElement"
+
 UPCGTransformPointsSettings::UPCGTransformPointsSettings()
 {
 	bUseSeed = true;
@@ -59,7 +61,7 @@ bool FPCGTransformPointsElement::ExecuteInternal(FPCGContext* Context) const
 
 		if (!SpatialData)
 		{
-			PCGE_LOG(Error, "Unable to get SpatialData from input");
+			PCGE_LOG(Error, GraphAndLog, LOCTEXT("InputMissingSpatialData", "Unable to get Spatial data from input"));
 			continue;
 		}
 
@@ -67,7 +69,7 @@ bool FPCGTransformPointsElement::ExecuteInternal(FPCGContext* Context) const
 
 		if (!PointData)
 		{
-			PCGE_LOG(Error, "Unable to get PointData from input");
+			PCGE_LOG(Error, GraphAndLog, LOCTEXT("InputMissingPointData", "Unable to get Point data from input"));
 			continue;
 		}
 
@@ -89,7 +91,7 @@ bool FPCGTransformPointsElement::ExecuteInternal(FPCGContext* Context) const
 
 			if (!FoundAttribute || FoundAttribute->GetTypeId() != PCG::Private::MetadataTypes<FTransform>::Id)
 			{
-				PCGE_LOG(Error, "Attribute %s does not exist or is not a transform", *LocalAttributeName.ToString());
+				PCGE_LOG(Error, GraphAndLog, FText::Format(LOCTEXT("AttributeMissing", "Attribute '{0}' does not exist or is not a transform"), FText::FromName(LocalAttributeName)));
 				continue;
 			}
 
@@ -212,3 +214,4 @@ bool FPCGTransformPointsElement::ExecuteInternal(FPCGContext* Context) const
 	return true;
 }
 
+#undef LOCTEXT_NAMESPACE
diff --git a/Engine/Plugins/Experimental/PCG/Source/PCG/Private/Elements/PCGVolumeSampler.cpp b/Engine/Plugins/Experimental/PCG/Source/PCG/Private/Elements/PCGVolumeSampler.cpp
--- a/Engine/Plugins/Experimental/PCG/Source/PCG/Private/Elements/PCGVolumeSampler.cpp
+++ b/Engine/Plugins/Experimental/PCG/Source/PCG/Private/Elements/PCGVolumeSampler.cpp
@@ -72,7 +72,7 @@ namespace PCGVolumeSampler
 			{
 				if (InContext)
 				{
-					PCGE_LOG_C(Verbose, InContext, "Skipped - invalid cell bounds (%lld x %lld x %lld)", NumX, NumY, NumZ);
+					PCGE_LOG_C(Verbose, LogOnly, InContext, FText::Format(FText::FromString(TEXT("Skipped - invalid cell bounds ({0} x {1} x {2})")), NumX, NumY, NumZ));
 				}
 
 				return;
@@ -86,7 +86,7 @@ namespace PCGVolumeSampler
 			{
 				if (InContext)
 				{
-					PCGE_LOG_C(Error, InContext, "Skipped - tried to generate too many points (%lld x %lld x %lld)", NumX, NumY, NumZ);
+					PCGE_LOG_C(Error, GraphAndLog, InContext, FText::Format(FText::FromString(TEXT("Skipped - tried to generate too many points ({0} x {1} x {2})")), NumX, NumY, NumZ));
 				}
 
 				return;
@@ -168,7 +168,7 @@ bool FPCGVolumeSamplerElement::ExecuteInternal(FPCGContext* Context) const
 	const FVector& VoxelSize = Settings->VoxelSize;
 	if (VoxelSize.X <= 0 || VoxelSize.Y <= 0 || VoxelSize.Z <= 0)
 	{
-		PCGE_LOG(Warning, "Skipped - Invalid voxel size");
+		PCGE_LOG(Warning, GraphAndLog, LOCTEXT("InvalidVoxelSize", "Skipped - Invalid voxel size"));
 		return true;
 	}
 
@@ -195,7 +195,7 @@ bool FPCGVolumeSamplerElement::ExecuteInternal(FPCGContext* Context) const
 	}
 	else if (BoundingShapeInputs.Num() > 0)
 	{
-		PCGE_LOG_C(Verbose, Context, "The bounds of the Bounding Shape input pin will be ignored because the Unbounded option is enabled.");
+		PCGE_LOG_C(Verbose, LogOnly, Context, LOCTEXT("BoundsIgnored", "The bounds of the Bounding Shape input pin will be ignored because the Unbounded option is enabled"));
 	}	
 
 	// Compute bounds of bounding shape input
@@ -235,7 +235,7 @@ bool FPCGVolumeSamplerElement::ExecuteInternal(FPCGContext* Context) const
 	// Warn if something is connected but no spatial data could be obtained for sampling
 	if (GeneratingShapes.Num() == 0 && (BoundingShapeInputs.Num() > 0 || VolumeInputs.Num() > 0))
 	{
-		PCGE_LOG(Warning, "No spatial data shape was provided for sampling, no points will be produced.");
+		PCGE_LOG(Warning, GraphAndLog, LOCTEXT("NoShapeToSample", "No Spatial data shape was provided for sampling, no points will be produced"));
 	}
 
 	// TODO: embarassingly parallel loop
@@ -262,7 +262,7 @@ bool FPCGVolumeSamplerElement::ExecuteInternal(FPCGContext* Context) const
 
 		if (!InputBounds.IsValid)
 		{
-			PCGE_LOG(Verbose, "Input data has invalid bounds");
+			PCGE_LOG(Verbose, LogOnly, LOCTEXT("InvalidBounds", "Input data has invalid bounds"));
 
 			Outputs.RemoveAt(GenerationIndex);
 			GeneratingShapes.RemoveAt(GenerationIndex);
@@ -276,7 +276,7 @@ bool FPCGVolumeSamplerElement::ExecuteInternal(FPCGContext* Context) const
 
 		if (SampledData)
 		{
-			PCGE_LOG(Verbose, "Generated %d points in volume", SampledData->GetPoints().Num());
+			PCGE_LOG(Verbose, LogOnly, FText::Format(LOCTEXT("GenerationInfo", "Generated {0} points in volume"), SampledData->GetPoints().Num()));
 		}
 	}
 
diff --git a/Engine/Plugins/Experimental/PCG/Source/PCG/Private/Graph/PCGGraphExecutor.cpp b/Engine/Plugins/Experimental/PCG/Source/PCG/Private/Graph/PCGGraphExecutor.cpp
--- a/Engine/Plugins/Experimental/PCG/Source/PCG/Private/Graph/PCGGraphExecutor.cpp
+++ b/Engine/Plugins/Experimental/PCG/Source/PCG/Private/Graph/PCGGraphExecutor.cpp
@@ -12,6 +12,7 @@
 #include "PCGModule.h"
 #include "PCGParamData.h"
 #include "PCGPin.h"
+#include "PCGSubsystem.h"
 #include "Graph/PCGGraphCache.h"
 #include "Metadata/PCGMetadata.h"
 
@@ -98,6 +99,17 @@ FPCGTaskId FPCGGraphExecutor::Schedule(UPCGGraph* Graph, UPCGComponent* SourceCo
 		UE_LOG(LogPCG, Log, TEXT("[%s] --- SCHEDULE GRAPH ---"), *SourceComponent->GetOwner()->GetName());
 	}
 
+#if WITH_EDITOR
+	if (UPCGSubsystem* Subsystem = UPCGSubsystem::GetInstance(SourceComponent->GetWorld()))
+	{
+		for (const UPCGNode* Node : Graph->GetNodes())
+		{
+			// Always clear warnings/errors before compile regardless of connectivity
+			Subsystem->GetNodeVisualLogsMutable().ClearLogs(Node, SourceComponent);
+		}
+	}
+#endif
+
 	FPCGTaskId ScheduledId = InvalidPCGTaskId;
 
 	// Get compiled tasks from compiler
@@ -608,6 +620,7 @@ void FPCGGraphExecutor::Execute()
 						SourceComponent->StoreInspectionData(Task.Node, CachedOutput);
 					}
 #endif
+
 					// Fast-forward cached result to stored results
 					FPCGTaskId SkippedTaskId = Task.NodeId;
 					StoreResults(SkippedTaskId, CachedOutput);
@@ -732,7 +745,15 @@ void FPCGGraphExecutor::Execute()
 
 				// Store result in cache as needed - done here because it needs to be done on the main thread
 				const UPCGSettingsInterface* ActiveTaskSettingsInterface = ActiveTask.Context->GetInputSettingsInterface();
-				if (ActiveTaskSettingsInterface && ActiveTask.Element->IsCacheableInstance(ActiveTaskSettingsInterface))
+
+				// Don't store if errors or warnings present
+#if WITH_EDITOR
+				const bool bHasErrorOrWarning = ActiveTask.Context->Node && (ActiveTask.Context->HasVisualLogs());
+#else
+				const bool bHasErrorOrWarning = false;
+#endif
+
+				if (ActiveTaskSettingsInterface && !bHasErrorOrWarning && ActiveTask.Element->IsCacheableInstance(ActiveTaskSettingsInterface))
 				{
 					const UPCGSettings* ActiveTaskSettings = ActiveTaskSettingsInterface ? ActiveTaskSettingsInterface->GetSettings() : nullptr;
 					GraphCache.StoreInCache(ActiveTask.Element.Get(), ActiveTask.Context->DependenciesCrc, ActiveTask.Context->InputData, ActiveTaskSettings, ActiveTask.Context->SourceComponent.Get(), ActiveTask.Context->OutputData);
diff --git a/Engine/Plugins/Experimental/PCG/Source/PCG/Private/InstancePackers/PCGInstancePackerByAttribute.cpp b/Engine/Plugins/Experimental/PCG/Source/PCG/Private/InstancePackers/PCGInstancePackerByAttribute.cpp
--- a/Engine/Plugins/Experimental/PCG/Source/PCG/Private/InstancePackers/PCGInstancePackerByAttribute.cpp
+++ b/Engine/Plugins/Experimental/PCG/Source/PCG/Private/InstancePackers/PCGInstancePackerByAttribute.cpp
@@ -8,11 +8,13 @@
 
 #include UE_INLINE_GENERATED_CPP_BY_NAME(PCGInstancePackerByAttribute)
 
+#define LOCTEXT_NAMESPACE "PCGInstancePackerByAttribute"
+
 void UPCGInstancePackerByAttribute::PackInstances_Implementation(FPCGContext& Context, const UPCGSpatialData* InSpatialData, const FPCGMeshInstanceList& InstanceList, FPCGPackedCustomData& OutPackedCustomData) const
 {
 	if (!InSpatialData || !InSpatialData->Metadata)
 	{
-		PCGE_LOG_C(Error, &Context, "Invalid input data");
+		PCGE_LOG_C(Error, GraphAndLog, &Context, LOCTEXT("InvalidInputData", "Invalid input data"));
 		return;
 	}
 
@@ -23,7 +25,7 @@ void UPCGInstancePackerByAttribute::PackInstances_Implementation(FPCGContext& Co
 	{
 		if (!InSpatialData->Metadata->HasAttribute(AttributeName)) 
 		{
-			PCGE_LOG_C(Warning, &Context, "Attribute %s is not in the metadata", *AttributeName.ToString());
+			PCGE_LOG_C(Warning, GraphAndLog, &Context, FText::Format(LOCTEXT("AttributeNotInMetadata", "Attribute '{0}' is not in the metadata"), FText::FromName(AttributeName)));
 			continue;
 		}
 
@@ -32,7 +34,7 @@ void UPCGInstancePackerByAttribute::PackInstances_Implementation(FPCGContext& Co
 
 		if (!AddTypeToPacking(AttributeBase->GetTypeId(), OutPackedCustomData))
 		{
-			PCGE_LOG_C(Warning, &Context, "Attribute name %s is not a valid type", *AttributeName.ToString());
+			PCGE_LOG_C(Warning, GraphAndLog, &Context, FText::Format(LOCTEXT("AttributeInvalidType", "Attribute name '{0}' is not a valid type"), FText::FromName(AttributeName)));
 			continue;
 		}
 
@@ -42,3 +44,4 @@ void UPCGInstancePackerByAttribute::PackInstances_Implementation(FPCGContext& Co
 	PackCustomDataFromAttributes(InstanceList, SelectedAttributes, OutPackedCustomData);
  }
 
+#undef LOCTEXT_NAMESPACE
diff --git a/Engine/Plugins/Experimental/PCG/Source/PCG/Private/InstancePackers/PCGInstancePackerByRegex.cpp b/Engine/Plugins/Experimental/PCG/Source/PCG/Private/InstancePackers/PCGInstancePackerByRegex.cpp
--- a/Engine/Plugins/Experimental/PCG/Source/PCG/Private/InstancePackers/PCGInstancePackerByRegex.cpp
+++ b/Engine/Plugins/Experimental/PCG/Source/PCG/Private/InstancePackers/PCGInstancePackerByRegex.cpp
@@ -9,11 +9,13 @@
 
 #include UE_INLINE_GENERATED_CPP_BY_NAME(PCGInstancePackerByRegex)
 
+#define LOCTEXT_NAMESPACE "PCGInstancePackerByRegex"
+
 void UPCGInstancePackerByRegex::PackInstances_Implementation(FPCGContext& Context, const UPCGSpatialData* InSpatialData, const FPCGMeshInstanceList& InstanceList, FPCGPackedCustomData& OutPackedCustomData) const
 {
 	if (!InSpatialData || !InSpatialData->Metadata)
 	{
-		PCGE_LOG_C(Error, &Context, "Invalid input data");
+		PCGE_LOG_C(Error, GraphAndLog, &Context, LOCTEXT("InvalidInputData", "Invalid input data"));
 		return;
 	}
 	
@@ -52,7 +54,7 @@ void UPCGInstancePackerByRegex::PackInstances_Implementation(FPCGContext& Contex
 	
 			if (!AddTypeToPacking(AttributeBase->GetTypeId(), OutPackedCustomData))
 			{
-				PCGE_LOG_C(Warning, &Context, "Attribute name %s is not a valid type", *AttributeName.ToString());
+				PCGE_LOG_C(Warning, GraphAndLog, &Context, FText::Format(LOCTEXT("InvalidAttributeType", "Attribute name '{0}' is not of a valid type"), FText::FromName(AttributeName)));
 				continue;
 			}
 	
@@ -63,3 +65,4 @@ void UPCGInstancePackerByRegex::PackInstances_Implementation(FPCGContext& Contex
 	PackCustomDataFromAttributes(InstanceList, SelectedAttributes, OutPackedCustomData);
 }
 
+#undef LOCTEXT_NAMESPACE
diff --git a/Engine/Plugins/Experimental/PCG/Source/PCG/Private/MatchAndSet/PCGMatchAndSetBase.cpp b/Engine/Plugins/Experimental/PCG/Source/PCG/Private/MatchAndSet/PCGMatchAndSetBase.cpp
--- a/Engine/Plugins/Experimental/PCG/Source/PCG/Private/MatchAndSet/PCGMatchAndSetBase.cpp
+++ b/Engine/Plugins/Experimental/PCG/Source/PCG/Private/MatchAndSet/PCGMatchAndSetBase.cpp
@@ -9,6 +9,8 @@
 
 #include UE_INLINE_GENERATED_CPP_BY_NAME(PCGMatchAndSetBase)
 
+#define LOCTEXT_NAMESPACE "PCGMatchAndSetBase"
+
 void UPCGMatchAndSetBase::SetType(EPCGMetadataTypes InType, EPCGMetadataTypesConstantStructStringMode InStringMode)
 {
 	Type = InType;
@@ -39,7 +41,7 @@ bool UPCGMatchAndSetBase::CreateAttributeIfNeeded(FPCGContext& Context, const FP
 
 			if (!ConstantValue.Dispatcher(CreateAttribute))
 			{
-				PCGE_LOG_C(Error, &Context, "Unable to create attribute %s on point data", *DestinationAttribute.ToString());
+				PCGE_LOG_C(Error, GraphAndLog, &Context, FText::Format(LOCTEXT("AttributeCreationFailed", "Unable to create attribute '{0}' on Point data"), FText::FromName(DestinationAttribute)));
 				return false;
 			}
 		}
@@ -47,3 +49,5 @@ bool UPCGMatchAndSetBase::CreateAttributeIfNeeded(FPCGContext& Context, const FP
 
 	return true;
 }
+
+#undef LOCTEXT_NAMESPACE
diff --git a/Engine/Plugins/Experimental/PCG/Source/PCG/Private/MatchAndSet/PCGMatchAndSetByAttribute.cpp b/Engine/Plugins/Experimental/PCG/Source/PCG/Private/MatchAndSet/PCGMatchAndSetByAttribute.cpp
--- a/Engine/Plugins/Experimental/PCG/Source/PCG/Private/MatchAndSet/PCGMatchAndSetByAttribute.cpp
+++ b/Engine/Plugins/Experimental/PCG/Source/PCG/Private/MatchAndSet/PCGMatchAndSetByAttribute.cpp
@@ -11,6 +11,8 @@
 
 #include UE_INLINE_GENERATED_CPP_BY_NAME(PCGMatchAndSetByAttribute)
 
+#define LOCTEXT_NAMESPACE "PCGMatchAndSetByAttributeElement"
+
 FPCGMatchAndSetByAttributeEntry::FPCGMatchAndSetByAttributeEntry()
 {
 	ValueToMatch.bAllowsTypeChange = false;
@@ -119,7 +121,7 @@ void UPCGMatchAndSetByAttribute::MatchAndSet_Implementation(
 
 	if (!InputAccessor.IsValid() || !InputKeys.IsValid())
 	{
-		PCGE_LOG_C(Warning, &Context, "Failed to create input accessor or iterator in MatchAndSet");
+		PCGE_LOG_C(Warning, GraphAndLog, &Context, LOCTEXT("InputAccessorCreationFailed", "Failed to create input accessor or iterator in MatchAndSet"));
 		return;
 	}
 
@@ -184,7 +186,7 @@ void UPCGMatchAndSetByAttribute::MatchAndSet_Implementation(
 
 	if (!PCGMetadataAttribute::CallbackWithRightType(InputAccessor->GetUnderlyingType(), MatchAndSetOperation))
 	{
-		PCGE_LOG_C(Warning, &Context, "Error while getting/setting values in the MatchAndSet");
+		PCGE_LOG_C(Warning, GraphAndLog, &Context, LOCTEXT("ErrorGettingSettingValues", "Error while getting/setting values in the MatchAndSet"));
 	}
 }
 
@@ -207,3 +209,5 @@ bool UPCGMatchAndSetByAttribute::ValidatePreconditions_Implementation(const UPCG
 	const FPCGMetadataAttributeBase* AttributeToMatch = InPointData->Metadata->GetConstAttribute(AttributeName);
 	return AttributeToMatch && AttributeToMatch->GetTypeId() == static_cast<uint16>(MatchSourceType);
 }
+
+#undef LOCTEXT_NAMESPACE
diff --git a/Engine/Plugins/Experimental/PCG/Source/PCG/Private/MatchAndSet/PCGMatchAndSetWeighted.cpp b/Engine/Plugins/Experimental/PCG/Source/PCG/Private/MatchAndSet/PCGMatchAndSetWeighted.cpp
--- a/Engine/Plugins/Experimental/PCG/Source/PCG/Private/MatchAndSet/PCGMatchAndSetWeighted.cpp
+++ b/Engine/Plugins/Experimental/PCG/Source/PCG/Private/MatchAndSet/PCGMatchAndSetWeighted.cpp
@@ -13,6 +13,8 @@
 
 #include UE_INLINE_GENERATED_CPP_BY_NAME(PCGMatchAndSetWeighted)
 
+#define LOCTEXT_NAMESPACE "PCGMatchAndSetWeightedEntry"
+
 FPCGMatchAndSetWeightedEntry::FPCGMatchAndSetWeightedEntry()
 {
 	Value.bAllowsTypeChange = false;
@@ -146,6 +148,8 @@ void UPCGMatchAndSetWeighted::MatchAndSet_Implementation(
 
 	if (!PCGMetadataAttribute::CallbackWithRightType(SetTargetAccessor->GetUnderlyingType(), MatchAndSetOperation))
 	{
-		PCGE_LOG_C(Warning, &Context, "Error while getting/setting values in the MatchAndSet");
+		PCGE_LOG_C(Warning, GraphAndLog, &Context, LOCTEXT("ErrorGettingSettingValues", "Error while getting/setting values in the MatchAndSet"));
 	}
 }
+
+#undef LOCTEXT_NAMESPACE
diff --git a/Engine/Plugins/Experimental/PCG/Source/PCG/Private/MatchAndSet/PCGMatchAndSetWeightedByCategory.cpp b/Engine/Plugins/Experimental/PCG/Source/PCG/Private/MatchAndSet/PCGMatchAndSetWeightedByCategory.cpp
--- a/Engine/Plugins/Experimental/PCG/Source/PCG/Private/MatchAndSet/PCGMatchAndSetWeightedByCategory.cpp
+++ b/Engine/Plugins/Experimental/PCG/Source/PCG/Private/MatchAndSet/PCGMatchAndSetWeightedByCategory.cpp
@@ -14,6 +14,8 @@
 
 #include UE_INLINE_GENERATED_CPP_BY_NAME(PCGMatchAndSetWeightedByCategory)
 
+#define LOCTEXT_NAMESPACE "PCGMatchAndSetWeightedByCategoryElement"
+
 FPCGMatchAndSetWeightedByCategoryEntryList::FPCGMatchAndSetWeightedByCategoryEntryList()
 {
 	CategoryValue.bAllowsTypeChange = false;
@@ -174,7 +176,7 @@ void UPCGMatchAndSetWeightedByCategory::MatchAndSet_Implementation(
 
 	if (!InputAccessor.IsValid() || !InputKeys.IsValid())
 	{
-		PCGE_LOG_C(Warning, &Context, "Failed to create input accessor or iterator in MatchAndSet");
+		PCGE_LOG_C(Warning, GraphAndLog, &Context, LOCTEXT("InputAccessorFailed", "Failed to create input accessor or iterator in MatchAndSet"));
 		return;
 	}
 
@@ -272,7 +274,7 @@ void UPCGMatchAndSetWeightedByCategory::MatchAndSet_Implementation(
 
 	if (!PCGMetadataAttribute::CallbackWithRightType(InputAccessor->GetUnderlyingType(), MatchAndSetOperation))
 	{
-		PCGE_LOG_C(Warning, &Context, "Error while getting/setting values in the MatchAndSet");
+		PCGE_LOG_C(Warning, GraphAndLog, &Context, LOCTEXT("ErrorGettingSettingValues", "Error while getting/setting values in the MatchAndSet"));
 	}
 }
 
@@ -294,3 +296,5 @@ bool UPCGMatchAndSetWeightedByCategory::ValidatePreconditions_Implementation(con
 	const FPCGMetadataAttributeBase* AttributeToMatch = InPointData->Metadata->GetConstAttribute(AttributeName);
 	return AttributeToMatch && AttributeToMatch->GetTypeId() == static_cast<uint16>(CategoryType);
 }
+
+#undef LOCTEXT_NAMESPACE
diff --git a/Engine/Plugins/Experimental/PCG/Source/PCG/Private/MeshSelectors/PCGMeshSelectorBase.cpp b/Engine/Plugins/Experimental/PCG/Source/PCG/Private/MeshSelectors/PCGMeshSelectorBase.cpp
--- a/Engine/Plugins/Experimental/PCG/Source/PCG/Private/MeshSelectors/PCGMeshSelectorBase.cpp
+++ b/Engine/Plugins/Experimental/PCG/Source/PCG/Private/MeshSelectors/PCGMeshSelectorBase.cpp
@@ -7,6 +7,8 @@
 
 #include UE_INLINE_GENERATED_CPP_BY_NAME(PCGMeshSelectorBase)
 
+#define LOCTEXT_NAMESPACE "PCGMeshSelectorBase"
+
 void FPCGMeshMaterialOverrideHelper::Initialize(
 	FPCGContext& InContext,
 	bool bInUseMaterialOverrideAttributes,
@@ -52,7 +54,7 @@ void FPCGMeshMaterialOverrideHelper::Initialize(FPCGContext& InContext)
 	{
 		if (!Metadata)
 		{
-			PCGE_LOG_C(Error, &InContext, "Data has no metadata");
+			PCGE_LOG_C(Error, GraphAndLog, &InContext, LOCTEXT("MetadataMissing", "Data has no metadata"));
 			return;
 		}
 
@@ -60,7 +62,7 @@ void FPCGMeshMaterialOverrideHelper::Initialize(FPCGContext& InContext)
 		{
 			if (!Metadata->HasAttribute(MaterialOverrideAttributeName))
 			{
-				PCGE_LOG_C(Error, &InContext, "Attribute %s for material overrides is not present in the metadata", *MaterialOverrideAttributeName.ToString());
+				PCGE_LOG_C(Error, GraphAndLog, &InContext, FText::Format(LOCTEXT("AttributeMissing", "Attribute '{0}' for material overrides is not present in the metadata"), FText::FromName(MaterialOverrideAttributeName)));
 				return;
 			}
 
@@ -69,7 +71,7 @@ void FPCGMeshMaterialOverrideHelper::Initialize(FPCGContext& InContext)
 
 			if (MaterialAttributeBase->GetTypeId() != PCG::Private::MetadataTypes<FString>::Id)
 			{
-				PCGE_LOG_C(Error, &InContext, "Material override attribute is not of valid type");
+				PCGE_LOG_C(Error, GraphAndLog, &InContext, LOCTEXT("AttributeInvalidType", "Material override attribute is not of valid type"));
 				return;
 			}
 
@@ -131,3 +133,5 @@ const TArray<TSoftObjectPtr<UMaterialInterface>>& FPCGMeshMaterialOverrideHelper
 		return StaticMaterialOverrides;
 	}
 }
+
+#undef LOCTEXT_NAMESPACE
diff --git a/Engine/Plugins/Experimental/PCG/Source/PCG/Private/MeshSelectors/PCGMeshSelectorByAttribute.cpp b/Engine/Plugins/Experimental/PCG/Source/PCG/Private/MeshSelectors/PCGMeshSelectorByAttribute.cpp
--- a/Engine/Plugins/Experimental/PCG/Source/PCG/Private/MeshSelectors/PCGMeshSelectorByAttribute.cpp
+++ b/Engine/Plugins/Experimental/PCG/Source/PCG/Private/MeshSelectors/PCGMeshSelectorByAttribute.cpp
@@ -10,6 +10,8 @@
 
 #include UE_INLINE_GENERATED_CPP_BY_NAME(PCGMeshSelectorByAttribute)
 
+#define LOCTEXT_NAMESPACE "PCGMeshSelectorByAttribute"
+
 namespace PCGMeshSelectorAttribute
 {
 	// Returns variation based on mesh, material overrides and reverse culling
@@ -101,19 +103,19 @@ bool UPCGMeshSelectorByAttribute::SelectInstances(
 
 	if (!InPointData)
 	{
-		PCGE_LOG_C(Error, &Context, "Missing input data");
+		PCGE_LOG_C(Error, GraphAndLog, &Context, LOCTEXT("InputMissingData", "Missing input data"));
 		return true;
 	}
 
 	if (!InPointData->Metadata)
 	{
-		PCGE_LOG_C(Error, &Context, "Unable to get metadata from input");
+		PCGE_LOG_C(Error, GraphAndLog, &Context, LOCTEXT("InputMissingMetadata", "Unable to get metadata from input"));
 		return true;
 	}
 
 	if (!InPointData->Metadata->HasAttribute(AttributeName))
 	{
-		PCGE_LOG_C(Error, &Context, "Attribute %s is not in the metadata", *AttributeName.ToString());
+		PCGE_LOG_C(Error, GraphAndLog, &Context, FText::Format(LOCTEXT("AttributeNotInMetadata", "Attribute '{0}' is not in the metadata"), FText::FromName(AttributeName)));
 		return true;
 	}
 
@@ -122,7 +124,7 @@ bool UPCGMeshSelectorByAttribute::SelectInstances(
 
 	if (AttributeBase->GetTypeId() != PCG::Private::MetadataTypes<FString>::Id)
 	{
-		PCGE_LOG_C(Error, &Context, "Attribute is not of valid type FString");
+		PCGE_LOG_C(Error, GraphAndLog, &Context, FText::Format(LOCTEXT("AttributeInvalidType", "Attribute '{0}' is not of valid type FString"), FText::FromName(AttributeName)));
 		return true;
 	}
 
@@ -182,12 +184,12 @@ bool UPCGMeshSelectorByAttribute::SelectInstances(
 
 				if (Mesh.IsNull())
 				{
-					PCGE_LOG_C(Error, &Context, "Invalid mesh path: %s.", *MeshSoftObjectPath);
+					PCGE_LOG_C(Error, GraphAndLog, &Context, FText::Format(LOCTEXT("InvalidMeshPath", "Invalid mesh path: '{0}'."), FText::FromString(MeshSoftObjectPath)));
 				}
 			}
 			else
 			{
-				PCGE_LOG_C(Warning, &Context, "Trivially invalid mesh path used: %s", *MeshSoftObjectPath);
+				PCGE_LOG_C(Warning, LogOnly, &Context, FText::Format(LOCTEXT("TrivialInvalidMeshPath", "Trivially invalid mesh path used: '{0}'"), FText::FromString(MeshSoftObjectPath)));
 			}
 
 			ValueKeyToMesh.Add(ValueKey, Mesh);
@@ -225,3 +227,5 @@ bool UPCGMeshSelectorByAttribute::SelectInstances(
 
 	return (CurrentPointIndex == Points.Num());
 }
+
+#undef LOCTEXT_NAMESPACE
diff --git a/Engine/Plugins/Experimental/PCG/Source/PCG/Private/MeshSelectors/PCGMeshSelectorWeighted.cpp b/Engine/Plugins/Experimental/PCG/Source/PCG/Private/MeshSelectors/PCGMeshSelectorWeighted.cpp
--- a/Engine/Plugins/Experimental/PCG/Source/PCG/Private/MeshSelectors/PCGMeshSelectorWeighted.cpp
+++ b/Engine/Plugins/Experimental/PCG/Source/PCG/Private/MeshSelectors/PCGMeshSelectorWeighted.cpp
@@ -13,6 +13,8 @@
 
 #include UE_INLINE_GENERATED_CPP_BY_NAME(PCGMeshSelectorWeighted)
 
+#define LOCTEXT_NAMESPACE "PCGMeshSelectorWeighted"
+
 namespace PCGMeshSelectorWeighted
 {
 	FPCGMeshInstanceList& GetInstanceList(
@@ -121,7 +123,7 @@ bool UPCGMeshSelectorWeighted::SelectInstances(
 {
 	if (!InPointData)
 	{
-		PCGE_LOG_C(Error, &Context, "Missing input data");
+		PCGE_LOG_C(Error, GraphAndLog, &Context, LOCTEXT("MissingInputData", "Missing input data"));
 		return true;
 	}
 
@@ -138,7 +140,7 @@ bool UPCGMeshSelectorWeighted::SelectInstances(
 		{
 			if (Entry.Weight <= 0)
 			{
-				PCGE_LOG_C(Verbose, &Context, "Entry found with weight <= 0");
+				PCGE_LOG_C(Verbose, LogOnly, &Context, LOCTEXT("EntryWithNegativeWeight", "Entry found with weight <= 0"));
 				continue;
 			}
 
@@ -175,7 +177,7 @@ bool UPCGMeshSelectorWeighted::SelectInstances(
 
 		if (!OutPointData->Metadata->HasAttribute(Settings->OutAttributeName)) 
 		{
-			PCGE_LOG_C(Error, &Context, "Out attribute %s is not in the metadata", *Settings->OutAttributeName.ToString());
+			PCGE_LOG_C(Error, GraphAndLog, &Context, FText::Format(LOCTEXT("AttributeNotInMetadata", "Out attribute '{0}' is not in the metadata"), FText::FromName(Settings->OutAttributeName)));
 		}
 
 		FPCGMetadataAttributeBase* OutAttributeBase = OutPointData->Metadata->GetMutableAttribute(Settings->OutAttributeName);
@@ -189,7 +191,7 @@ bool UPCGMeshSelectorWeighted::SelectInstances(
 			}
 			else
 			{
-				PCGE_LOG_C(Error, &Context, "Out attribute is not of valid type FString");
+				PCGE_LOG_C(Error, GraphAndLog, &Context, LOCTEXT("AttributeNotFString", "Out attribute is not of valid type FString"));
 			}
 		}
 	}
@@ -285,3 +287,4 @@ bool UPCGMeshSelectorWeighted::SelectInstances(
 	}
 }
 
+#undef LOCTEXT_NAMESPACE
diff --git a/Engine/Plugins/Experimental/PCG/Source/PCG/Private/MeshSelectors/PCGMeshSelectorWeightedByCategory.cpp b/Engine/Plugins/Experimental/PCG/Source/PCG/Private/MeshSelectors/PCGMeshSelectorWeightedByCategory.cpp
--- a/Engine/Plugins/Experimental/PCG/Source/PCG/Private/MeshSelectors/PCGMeshSelectorWeightedByCategory.cpp
+++ b/Engine/Plugins/Experimental/PCG/Source/PCG/Private/MeshSelectors/PCGMeshSelectorWeightedByCategory.cpp
@@ -13,6 +13,8 @@
 
 #include UE_INLINE_GENERATED_CPP_BY_NAME(PCGMeshSelectorWeightedByCategory)
 
+#define LOCTEXT_NAMESPACE "PCGMeshSelectorWeightedByCategory"
+
 #if WITH_EDITOR
 void FPCGWeightedByCategoryEntryList::ApplyDeprecation()
 {
@@ -44,19 +46,19 @@ bool UPCGMeshSelectorWeightedByCategory::SelectInstances(
 {
 	if (!InPointData)
 	{
-		PCGE_LOG_C(Error, &Context, "Missing input data");
+		PCGE_LOG_C(Error, GraphAndLog, &Context, LOCTEXT("InputMissingData", "Missing input data"));
 		return true;
 	}
 
 	if (!InPointData->Metadata)
 	{
-		PCGE_LOG_C(Error, &Context, "Unable to get metadata from input");
+		PCGE_LOG_C(Error, GraphAndLog, &Context, LOCTEXT("InputMissingMetadata", "Unable to get metadata from input"));
 		return true;
 	}
 
 	if (!InPointData->Metadata->HasAttribute(CategoryAttribute))
 	{
-		PCGE_LOG_C(Error, &Context, "Attribute %s is not in the metadata", *CategoryAttribute.ToString());
+		PCGE_LOG_C(Error, GraphAndLog, &Context, FText::Format(LOCTEXT("InputMissingAttribute", "Attribute '{0}' is not in the metadata"), FText::FromName(CategoryAttribute)));
 		return true;
 	}
 
@@ -66,7 +68,7 @@ bool UPCGMeshSelectorWeightedByCategory::SelectInstances(
 	// TODO: support enum type as well
 	if (AttributeBase->GetTypeId() != PCG::Private::MetadataTypes<FString>::Id)
 	{
-		PCGE_LOG_C(Error, &Context, "Attribute is not of valid type FString");
+		PCGE_LOG_C(Error, GraphAndLog, &Context, FText::Format(LOCTEXT("AttributeInvalidType", "Attribute '{0}' is not of valid type FString"), FText::FromName(CategoryAttribute)));
 		return true;
 	}
 
@@ -85,7 +87,7 @@ bool UPCGMeshSelectorWeightedByCategory::SelectInstances(
 		{
 			if (Entry.WeightedMeshEntries.Num() == 0)
 			{
-				PCGE_LOG_C(Verbose, &Context, "Empty entry found in category %s", *Entry.CategoryEntry);
+				PCGE_LOG_C(Verbose, LogOnly, &Context, FText::Format(LOCTEXT("EmptyEntryInCategory", "Empty entry found in category '{0}'"), FText::FromString(Entry.CategoryEntry)));
 				continue;
 			}
 
@@ -93,7 +95,7 @@ bool UPCGMeshSelectorWeightedByCategory::SelectInstances(
 
 			if (ValueKey == PCGDefaultValueKey)
 			{
-				PCGE_LOG_C(Verbose, &Context, "Invalid category %s", *Entry.CategoryEntry);
+				PCGE_LOG_C(Verbose, LogOnly, &Context, FText::Format(LOCTEXT("InvalidCategory", "Invalid category '{0}'"), FText::FromString(Entry.CategoryEntry)));
 				continue;
 			}
 
@@ -101,7 +103,7 @@ bool UPCGMeshSelectorWeightedByCategory::SelectInstances(
 
 			if (InstancesAndWeights)
 			{
-				PCGE_LOG_C(Warning, &Context, "Duplicate entry found in category %s. Subsequent entries are ignored.", *Entry.CategoryEntry);
+				PCGE_LOG_C(Warning, GraphAndLog, &Context, FText::Format(LOCTEXT("DuplicateEntry", "Duplicate entry found in category '{0}'. Subsequent entries are ignored."), FText::FromString(Entry.CategoryEntry)));
 				continue;
 			}
 
@@ -113,7 +115,7 @@ bool UPCGMeshSelectorWeightedByCategory::SelectInstances(
 				}
 				else
 				{
-					PCGE_LOG_C(Warning, &Context, "Duplicate default entry found. Subsequent default entries are ignored.");
+					PCGE_LOG_C(Warning, GraphAndLog, &Context, LOCTEXT("DuplicateDefaultEntry", "Duplicate default entry found. Subsequent default entries are ignored."));
 				}
 			}
 
@@ -124,7 +126,7 @@ bool UPCGMeshSelectorWeightedByCategory::SelectInstances(
 			{
 				if (WeightedEntry.Weight <= 0)
 				{
-					PCGE_LOG_C(Verbose, &Context, "Entry found with weight <= 0 in category %s", *Entry.CategoryEntry);
+					PCGE_LOG_C(Verbose, LogOnly, &Context, FText::Format(LOCTEXT("ZeroOrNegativeWeight", "Entry found with weight <= 0 in category '{0}'"), FText::FromString(Entry.CategoryEntry)));
 					continue;
 				}
 
@@ -164,7 +166,7 @@ bool UPCGMeshSelectorWeightedByCategory::SelectInstances(
 
 		if (!OutPointData->Metadata->HasAttribute(Settings->OutAttributeName)) 
 		{
-			PCGE_LOG_C(Error, &Context, "Out attribute %s is not in the metadata", *Settings->OutAttributeName.ToString());
+			PCGE_LOG_C(Error, GraphAndLog, &Context, FText::Format(LOCTEXT("AttributeNotInMetadata", "Out attribute '{0}' is not in the metadata"), FText::FromName(Settings->OutAttributeName)));
 		}
 
 		FPCGMetadataAttributeBase* OutAttributeBase = OutPointData->Metadata->GetMutableAttribute(Settings->OutAttributeName);
@@ -178,7 +180,7 @@ bool UPCGMeshSelectorWeightedByCategory::SelectInstances(
 			}
 			else
 			{
-				PCGE_LOG_C(Error, &Context, "Out attribute is not of valid type FString");
+				PCGE_LOG_C(Error, GraphAndLog, &Context, LOCTEXT("TypeNotFString", "Out attribute is not of valid type FString"));
 			}
 		}
 	}
@@ -292,3 +294,5 @@ bool UPCGMeshSelectorWeightedByCategory::SelectInstances(
 		return false;
 	}
 }
+
+#undef LOCTEXT_NAMESPACE
diff --git a/Engine/Plugins/Experimental/PCG/Source/PCG/Private/PCGContext.cpp b/Engine/Plugins/Experimental/PCG/Source/PCG/Private/PCGContext.cpp
--- a/Engine/Plugins/Experimental/PCG/Source/PCG/Private/PCGContext.cpp
+++ b/Engine/Plugins/Experimental/PCG/Source/PCG/Private/PCGContext.cpp
@@ -5,6 +5,7 @@
 #include "PCGHelpers.h"
 #include "PCGParamData.h"
 #include "PCGPin.h"
+#include "PCGSubsystem.h"
 #include "Data/PCGSpatialData.h"
 #include "Metadata/PCGAttributePropertySelector.h"
 #include "Metadata/PCGMetadata.h"
@@ -19,6 +20,8 @@
 
 #include UE_INLINE_GENERATED_CPP_BY_NAME(PCGContext)
 
+#define LOCTEXT_NAMESPACE "PCGContext"
+
 FPCGContext::~FPCGContext()
 {
 	if (SettingsWithOverride)
@@ -198,7 +201,7 @@ void FPCGContext::OverrideSettings()
 			PropertyType Value{};
 			if (!AttributeAccessor->Get<PropertyType>(Value, FirstEntry, EPCGAttributeAccessorFlags::AllowBroadcast | EPCGAttributeAccessorFlags::AllowConstructible))
 			{
-				PCGE_LOG_C(Warning, this, "%s param can't be converted from %s attribute", *Param.Label.ToString(), *AttributeName.ToString());
+				PCGE_LOG_C(Warning, GraphAndLog, this, FText::Format(LOCTEXT("ConversionFailed", "Param '{0}' cannot be converted from attribute '{1}'"), FText::FromName(Param.Label), FText::FromName(AttributeName)));
 				return false;
 			}
 			// TODO: Perhaps factorise this code in another property accessor.
@@ -216,3 +219,25 @@ void FPCGContext::OverrideSettings()
 		});
 	}
 }
+
+#if WITH_EDITOR
+void FPCGContext::LogVisual(ELogVerbosity::Type InVerbosity, const FText& InMessage) const
+{
+	if (UPCGSubsystem* Subsystem = UPCGSubsystem::GetInstance(SourceComponent->GetWorld()))
+	{
+		Subsystem->GetNodeVisualLogsMutable().Log(Node, SourceComponent, InVerbosity, InMessage);
+	}
+}
+
+bool FPCGContext::HasVisualLogs() const
+{
+	if (UPCGSubsystem* Subsystem = UPCGSubsystem::GetInstance(SourceComponent->GetWorld()))
+	{
+		return Subsystem->GetNodeVisualLogs().HasLogs(Node, SourceComponent.Get());
+	}
+
+	return false;
+}
+#endif // WITH_EDITOR
+
+#undef LOCTEXT_NAMESPACE
diff --git a/Engine/Plugins/Experimental/PCG/Source/PCG/Private/PCGElement.cpp b/Engine/Plugins/Experimental/PCG/Source/PCG/Private/PCGElement.cpp
--- a/Engine/Plugins/Experimental/PCG/Source/PCG/Private/PCGElement.cpp
+++ b/Engine/Plugins/Experimental/PCG/Source/PCG/Private/PCGElement.cpp
@@ -12,6 +12,8 @@
 #include "HAL/IConsoleManager.h"
 #include "Utils/PCGExtraCapture.h"
 
+#define LOCTEXT_NAMESPACE "PCGElement"
+
 static TAutoConsoleVariable<bool> CVarPCGValidatePointMetadata(
 	TEXT("pcg.debug.ValidatePointMetadata"),
 	true,
@@ -204,7 +206,7 @@ void IPCGElement::PostExecute(FPCGContext* Context) const
 		Context->InputData = ElementOutputs;
 		Context->OutputData = FPCGDataCollection();
 
-		PCGE_LOG(Verbose, "Performing remove duplicate points test (perf warning)");
+		PCGE_LOG(Verbose, LogOnly, LOCTEXT("PerformingDuplicatePointTest", "Performing remove duplicate points test (perf warning)"));
 		PCGSelfPruningElement::Execute(Context, EPCGSelfPruningType::RemoveDuplicates, 0.0f, false);
 
 		Context->InputData = ElementInputs;
@@ -354,12 +356,12 @@ void IPCGElement::CleanupAndValidateOutput(FPCGContext* Context) const
 				const int32 MatchIndex = OutputPinProperties.IndexOfByPredicate([&TaggedData](const FPCGPinProperties& InProp) { return TaggedData.Pin == InProp.Label; });
 				if (MatchIndex == INDEX_NONE)
 				{
-					PCGE_LOG(Warning, "Output generated for pin %s but cannot be routed", *TaggedData.Pin.ToString());
+					PCGE_LOG(Warning, LogOnly, FText::Format(LOCTEXT("OutputCannotBeRouted", "Output generated for pin '{0}' but cannot be routed"), FText::FromName(TaggedData.Pin)));
 				}
 				// TODO: Temporary fix for Settings directly from InputData (ie. from elements with code and not PCG nodes)
 				else if(TaggedData.Data && !(OutputPinProperties[MatchIndex].AllowedTypes & TaggedData.Data->GetDataType()) && TaggedData.Data->GetDataType() != EPCGDataType::Settings)
 				{
-					PCGE_LOG(Warning, "Output generated for pin %s does not have a compatible type: %s", *TaggedData.Pin.ToString(), *UEnum::GetValueAsString(TaggedData.Data->GetDataType()));
+					PCGE_LOG(Warning, LogOnly, FText::Format(LOCTEXT("OutputIncompatibleType", "Output generated for pin '{0}' does not have a compatible type: '{1}'"), FText::FromName(TaggedData.Pin), FText::FromString(UEnum::GetValueAsString(TaggedData.Data->GetDataType()))));
 				}
 
 				if (CVarPCGValidatePointMetadata.GetValueOnAnyThread())
@@ -378,7 +380,7 @@ void IPCGElement::CleanupAndValidateOutput(FPCGContext* Context) const
 
 						if (bHasError)
 						{
-							PCGE_LOG(Warning, "Output generated for pin %s does not have valid point metadata", *TaggedData.Pin.ToString());
+							PCGE_LOG(Warning, GraphAndLog, FText::Format(LOCTEXT("OutputMissingPointMetadata", "Output generated for pin '{0}' does not have valid point metadata"), FText::FromName(TaggedData.Pin)));
 						}
 					}
 				}
@@ -438,3 +440,5 @@ FPCGContext* FSimplePCGElement::Initialize(const FPCGDataCollection& InputData,
 
 	return Context;
 }
+
+#undef LOCTEXT_NAMESPACE
diff --git a/Engine/Plugins/Experimental/PCG/Source/PCG/Private/PCGSubsystem.cpp b/Engine/Plugins/Experimental/PCG/Source/PCG/Private/PCGSubsystem.cpp
--- a/Engine/Plugins/Experimental/PCG/Source/PCG/Private/PCGSubsystem.cpp
+++ b/Engine/Plugins/Experimental/PCG/Source/PCG/Private/PCGSubsystem.cpp
@@ -167,6 +167,18 @@ UPCGSubsystem* UPCGSubsystem::GetInstance(UWorld* World)
 	}
 }
 
+#if WITH_EDITOR
+UPCGSubsystem* UPCGSubsystem::GetActiveEditorInstance()
+{
+	if (GEditor)
+	{
+		return GEditor->PlayWorld ? UPCGSubsystem::GetInstance(GEditor->PlayWorld.Get()) : UPCGSubsystem::GetInstance(GEditor->GetEditorWorldContext().World());
+	}
+
+	return nullptr;
+}
+#endif
+
 void UPCGSubsystem::Tick(float DeltaSeconds)
 {
 	Super::Tick(DeltaSeconds);
diff --git a/Engine/Plugins/Experimental/PCG/Source/PCG/Private/Utils/PCGNodeVisualLogs.cpp b/Engine/Plugins/Experimental/PCG/Source/PCG/Private/Utils/PCGNodeVisualLogs.cpp
new file mode 100644
--- /dev/null
+++ b/Engine/Plugins/Experimental/PCG/Source/PCG/Private/Utils/PCGNodeVisualLogs.cpp
@@ -0,0 +1,198 @@
+// Copyright Epic Games, Inc. All Rights Reserved.
+
+#include "Utils/PCGNodeVisualLogs.h"
+
+#include "PCGComponent.h"
+
+#include "Algo/Find.h"
+#include "GameFramework/Actor.h"
+
+#define LOCTEXT_NAMESPACE "PCGNode"
+
+#if WITH_EDITOR
+
+void FPCGNodeVisualLogs::Log(TWeakObjectPtr<const UPCGNode> InNode, TWeakObjectPtr<UPCGComponent> InComponent, ELogVerbosity::Type InVerbosity, const FText& InMessage)
+{
+	if (!ensure(InNode.Get()) || !ensure(InComponent.Get()))
+	{
+		return;
+	}
+
+	bool bAdded = false;
+
+	{
+		FWriteScopeLock ScopedWriteLock(LogsLock);
+
+		FPCGPerNodeVisualLogs* NodeLogs = NodeToLogs.Find(InNode);
+		if (!NodeLogs)
+		{
+			NodeLogs = &NodeToLogs.Add(InNode);
+		}
+
+		constexpr int32 MaxLogged = 1024;
+		if (NodeToLogs.Num() < MaxLogged)
+		{
+			NodeLogs->Emplace(InMessage, InVerbosity, InComponent);
+
+			bAdded = true;
+		}
+	}
+
+	// Broadcast outside of write scope lock
+	if (bAdded)
+	{
+		InNode->OnNodeChangedDelegate.Broadcast(const_cast<UPCGNode*>(InNode.Get()), EPCGChangeType::Cosmetic);
+	}
+}
+
+bool FPCGNodeVisualLogs::HasLogs(TWeakObjectPtr<const UPCGNode> InNode, const UPCGComponent* InComponent) const
+{
+	FReadScopeLock ScopedReadLock(LogsLock);
+
+	if (!ensure(InNode.Get()))
+	{
+		return false;
+	}
+
+	const FPCGPerNodeVisualLogs* NodeLogs = NodeToLogs.Find(InNode);
+	if (!NodeLogs)
+	{
+		return false;
+	}
+
+	if (!InComponent)
+	{
+		return !NodeLogs->IsEmpty();
+	}
+	else
+	{
+		return !!Algo::FindByPredicate(*NodeLogs, [InComponent](const FPCGNodeLogEntry& Log)
+		{
+			return Log.Component == InComponent;
+		});
+	}
+}
+
+bool FPCGNodeVisualLogs::HasLogs(TWeakObjectPtr<const UPCGNode> InNode, const UPCGComponent* InComponent, ELogVerbosity::Type InVerbosity) const
+{
+	FReadScopeLock ScopedReadLock(LogsLock);
+
+	if (!ensure(InNode.Get()))
+	{
+		return false;
+	}
+
+	const FPCGPerNodeVisualLogs* NodeLogs = NodeToLogs.Find(InNode);
+	if (!NodeLogs)
+	{
+		return false;
+	}
+
+	return !!Algo::FindByPredicate(*NodeLogs, [InVerbosity, InComponent](const FPCGNodeLogEntry& Log)
+	{
+		return Log.Verbosity == InVerbosity && (!InComponent || Log.Component == InComponent);
+	});
+}
+
+FText FPCGNodeVisualLogs::GetLogsSummaryText(TWeakObjectPtr<const UPCGNode> InNode, const UPCGComponent* InComponent) const
+{
+	FReadScopeLock ScopedReadLock(LogsLock);
+
+	if (!ensure(InNode.Get()))
+	{
+		return FText::GetEmpty();
+	}
+
+	const FPCGPerNodeVisualLogs* NodeLogsPtr = NodeToLogs.Find(InNode);
+	if (!NodeLogsPtr)
+	{
+		return FText::GetEmpty();
+	}
+	const FPCGPerNodeVisualLogs& NodeLogs = *NodeLogsPtr;
+
+	// Compose user friendly summary of first N errors
+	FText ResultText = FText::GetEmpty();
+	int32 LogCounter = 0;
+
+	for (int32 i = 0; i < NodeLogs.Num(); ++i)
+	{
+		if (InComponent && NodeLogs[i].Component != InComponent)
+		{
+			continue;
+		}
+
+		++LogCounter;
+
+		if (LogCounter > 1)
+		{
+			// New lines after each log
+			ResultText = FText::Format(FText::FromString(TEXT("{0}\n")), ResultText);
+		}
+
+		const FText MessageVerbosity = NodeLogs[i].Verbosity == ELogVerbosity::Warning ? FText::FromString(TEXT("Warning")) : FText::FromString(TEXT("Error"));
+
+		if (!InComponent)
+		{
+			FText ActorName = ensure(NodeLogs[i].Component.Get() && NodeLogs[i].Component->GetOwner()) ? FText::FromName(NodeLogs[i].Component->GetOwner()->GetFName()) : FText::FromString(TEXT("MissingComponent"));
+			ResultText = FText::Format(LOCTEXT("NodeTooltipLogWithActor", "{0}{1}/{2}: [{3}] {4}: {5}"), ResultText, i + 1, NodeLogs.Num(), ActorName, MessageVerbosity, NodeLogs[i].Message);
+		}
+		else
+		{
+			ResultText = FText::Format(LOCTEXT("NodeTooltipLog", "{0}{1}/{2} {3}: {4}"), ResultText, i + 1, NodeLogs.Num(), MessageVerbosity, NodeLogs[i].Message);
+		}
+
+		constexpr int32 MaxLogs = 8;
+		if (LogCounter >= MaxLogs)
+		{
+			ResultText = FText::Format(FText::FromString(TEXT("{0}\n...")), ResultText);
+			break;
+		}
+	}
+
+	return ResultText;
+}
+
+void FPCGNodeVisualLogs::ClearLogs(TWeakObjectPtr<const UPCGNode> InNode, const UPCGComponent* InComponent)
+{
+	if (!ensure(InNode.Get()) || !ensure(InComponent))
+	{
+		return;
+	}
+
+	FPCGPerNodeVisualLogs* NodeLogs = nullptr;
+	{
+		FReadScopeLock ScopedReadLock(LogsLock);
+		NodeLogs = NodeToLogs.Find(InNode);
+
+		if (!NodeLogs || NodeLogs->IsEmpty())
+		{
+			// No logs, nothing to do
+			return;
+		}
+	}
+
+	bool bAnyRemoved = false;
+	{
+		FWriteScopeLock ScopedWriteLock(LogsLock);
+
+		check(NodeLogs);
+		for (int32 i = NodeLogs->Num() - 1; i >= 0; --i)
+		{
+			if ((*NodeLogs)[i].Component == InComponent)
+			{
+				NodeLogs->RemoveAtSwap(i);
+				bAnyRemoved = true;
+			}
+		}
+	}
+
+	// Broadcast change notification outside of write scope lock
+	if (bAnyRemoved)
+	{
+		InNode->OnNodeChangedDelegate.Broadcast(const_cast<UPCGNode*>(InNode.Get()), EPCGChangeType::Cosmetic);
+	}
+}
+
+#endif // WITH_EDITOR
+
+#undef LOCTEXT_NAMESPACE
diff --git a/Engine/Plugins/Experimental/PCG/Source/PCG/Public/PCGContext.h b/Engine/Plugins/Experimental/PCG/Source/PCG/Public/PCGContext.h
--- a/Engine/Plugins/Experimental/PCG/Source/PCG/Public/PCGContext.h
+++ b/Engine/Plugins/Experimental/PCG/Source/PCG/Public/PCGContext.h
@@ -90,6 +90,14 @@ struct PCG_API FPCGContext
 
 	AActor* GetTargetActor(const UPCGSpatialData* InSpatialData) const;
 
+#if WITH_EDITOR
+	/** Log warnings and errors to be displayed on node in graph editor. */
+	void LogVisual(ELogVerbosity::Type InVerbosity, const FText& InMessage) const;
+
+	/** True if any issues were logged during last execution. */
+	bool HasVisualLogs() const;
+#endif // WITH_EDITOR
+
 protected:
 	virtual UObject* GetExternalContainerForOverridableParam(const FPCGSettingsOverridableParam& InParam) { return nullptr; }
 
diff --git a/Engine/Plugins/Experimental/PCG/Source/PCG/Public/PCGElement.h b/Engine/Plugins/Experimental/PCG/Source/PCG/Public/PCGElement.h
--- a/Engine/Plugins/Experimental/PCG/Source/PCG/Public/PCGElement.h
+++ b/Engine/Plugins/Experimental/PCG/Source/PCG/Public/PCGElement.h
@@ -2,6 +2,7 @@
 
 #pragma once
 
+#include "Logging/LogVerbosity.h"
 #include "Templates/SharedPointer.h"
 #include "UObject/WeakObjectPtrTemplates.h"
 
@@ -17,18 +18,36 @@ class UPCGNode;
 
 typedef TSharedPtr<IPCGElement, ESPMode::ThreadSafe> FPCGElementPtr;
 
-#define PCGE_LOG_C(Verbosity, CustomContext, Format, ...) \
+enum class EPCGElementLogMode : uint8
+{
+	/** Output to log. */
+	LogOnly,
+	/** Display errors/warnings on graph as well as writing to log. */
+	GraphAndLog,
+};
+
+#define PCGE_LOG_BASE(Verbosity, CustomContext, Message) \
 	UE_LOG(LogPCG, \
 		Verbosity, \
-		TEXT("[%s - %s]: " Format), \
+		TEXT("[%s - %s]: %s"), \
 		*((CustomContext)->GetComponentName()), \
 		*((CustomContext)->GetTaskName()), \
-		##__VA_ARGS__)
+		*Message.ToString())
 
 #if WITH_EDITOR
-#define PCGE_LOG(Verbosity, Format, ...) do{ if(ShouldLog()) { PCGE_LOG_C(Verbosity, Context, Format, ##__VA_ARGS__); } }while(0)
+// Output to PCG log and optionally also display wanrings/errors on graph.
+#define PCGE_LOG(Verbosity, LogMode, Message) do { \
+		if ((EPCGElementLogMode::LogMode) == EPCGElementLogMode::GraphAndLog && (Context)) { (Context)->LogVisual(ELogVerbosity::Verbosity, Message); } \
+		if (ShouldLog()) { PCGE_LOG_BASE(Verbosity, Context, Message); } \
+	} while (0)
+// Output to PCG log and optionally also display wanrings/errors on graph. Takes context as argument.
+#define PCGE_LOG_C(Verbosity, LogMode, CustomContext, Message) do { \
+		if ((EPCGElementLogMode::LogMode) == EPCGElementLogMode::GraphAndLog && (CustomContext)) { (CustomContext)->LogVisual(ELogVerbosity::Verbosity, Message); } \
+		PCGE_LOG_BASE(Verbosity, CustomContext, Message); \
+	} while (0)
 #else
-#define PCGE_LOG(Verbosity, Format, ...) PCGE_LOG_C(Verbosity, Context, Format, ##__VA_ARGS__)
+#define PCGE_LOG(Verbosity, LogMode, Message) PCGE_LOG_BASE(Verbosity, Context, Message)
+#define PCGE_LOG_C(Verbosity, LogMode, CustomContext, Message) PCGE_LOG_BASE(Verbosity, CustomContext, Message)
 #endif
 
 /**
diff --git a/Engine/Plugins/Experimental/PCG/Source/PCG/Public/PCGSubsystem.h b/Engine/Plugins/Experimental/PCG/Source/PCG/Public/PCGSubsystem.h
--- a/Engine/Plugins/Experimental/PCG/Source/PCG/Public/PCGSubsystem.h
+++ b/Engine/Plugins/Experimental/PCG/Source/PCG/Public/PCGSubsystem.h
@@ -7,6 +7,7 @@
 #include "PCGCommon.h"
 #include "PCGActorAndComponentMapping.h"
 #include "Grid/PCGComponentOctree.h"
+#include "Utils/PCGNodeVisualLogs.h"
 
 #include "PCGSubsystem.generated.h"
 
@@ -62,7 +63,12 @@ public:
 	/** Will return the subsystem from the World if it exists and if it is initialized */
 	static UPCGSubsystem* GetInstance(UWorld* World);
 
-	/** Subsytem must not be used without this condition being true. */
+#if WITH_EDITOR
+	/** Returns PIE world if it is active, otherwise returns editor world. */
+	static UPCGSubsystem* GetActiveEditorInstance();
+#endif
+
+	/** Subsystem must not be used without this condition being true. */
 	bool IsInitialized() const { return GraphExecutor != nullptr; }
 
 	APCGWorldActor* GetPCGWorldActor();
@@ -181,6 +187,10 @@ public:
 	/** Returns how many times InElement is present in the cache. */
 	uint32 GetGraphCacheEntryCount(IPCGElement* InElement) const;
 
+	/** Get graph warnings and errors for all nodes. */
+	const FPCGNodeVisualLogs& GetNodeVisualLogs() const { return NodeVisualLogs; }
+	FPCGNodeVisualLogs& GetNodeVisualLogsMutable() { return NodeVisualLogs; }
+
 private:
 	enum class EOperation : uint32
 	{
@@ -191,6 +201,8 @@ private:
 
 	FPCGTaskId ProcessGraph(UPCGComponent* Component, const FBox& InPreviousBounds, const FBox& InNewBounds, EOperation InOperation, bool bSave);
 	void CreatePartitionActorsWithinBounds(const FBox& InBounds);
+
+	FPCGNodeVisualLogs NodeVisualLogs;
 #endif // WITH_EDITOR
 	
 private:
diff --git a/Engine/Plugins/Experimental/PCG/Source/PCG/Public/Utils/PCGNodeVisualLogs.h b/Engine/Plugins/Experimental/PCG/Source/PCG/Public/Utils/PCGNodeVisualLogs.h
new file mode 100644
--- /dev/null
+++ b/Engine/Plugins/Experimental/PCG/Source/PCG/Public/Utils/PCGNodeVisualLogs.h
@@ -0,0 +1,57 @@
+// Copyright Epic Games, Inc. All Rights Reserved.
+
+#pragma once
+
+#include "Logging/LogVerbosity.h"
+
+#include "Containers/Map.h"
+#include "Internationalization/Text.h"
+#include "UObject/WeakObjectPtr.h"
+
+#if WITH_EDITOR
+
+class UPCGComponent;
+class UPCGNode;
+
+struct FPCGNodeLogEntry
+{
+	explicit FPCGNodeLogEntry(const FText& InMessage, ELogVerbosity::Type InVerbosity, TWeakObjectPtr<UPCGComponent> InComponent)
+		: Message(InMessage)
+		, Verbosity(InVerbosity)
+		, Component(InComponent)
+	{
+	}
+
+	FText Message;
+	ELogVerbosity::Type Verbosity;
+	TWeakObjectPtr<UPCGComponent> Component;
+};
+
+typedef TArray<FPCGNodeLogEntry, TInlineAllocator<16>> FPCGPerNodeVisualLogs;
+
+/** Collections per-node graph execution warnings and errors. */
+class PCG_API FPCGNodeVisualLogs
+{
+public:
+	/** Log warnings and errors to be displayed on node in graph editor. */
+	void Log(TWeakObjectPtr<const UPCGNode> InNode, TWeakObjectPtr<UPCGComponent> InComponent, ELogVerbosity::Type InVerbosity, const FText& InMessage);
+
+	/** True if any issues were logged during last execution. */
+	bool HasLogs(TWeakObjectPtr<const UPCGNode> InNode, const UPCGComponent* InComponent) const;
+
+	/** True if an issue with given severity was logged during last execution. */
+	bool HasLogs(TWeakObjectPtr<const UPCGNode> InNode, const UPCGComponent* InComponent, ELogVerbosity::Type InVerbosity) const;
+
+	/** Summary text of visual logs from recent execution, appropriate for display in graph editor tooltip. */
+	FText GetLogsSummaryText(TWeakObjectPtr<const UPCGNode> InNode, const UPCGComponent* InComponent = nullptr) const;
+
+	/** Clear all errors and warnings that occurred while executing the given component. */
+	void ClearLogs(TWeakObjectPtr<const UPCGNode> InNode, const UPCGComponent* InComponent);
+
+private:
+	TMap<TWeakObjectPtr<const UPCGNode>, FPCGPerNodeVisualLogs> NodeToLogs;
+
+	mutable FRWLock LogsLock;
+};
+
+#endif // WITH_EDITOR
diff --git a/Engine/Plugins/Experimental/PCG/Source/PCGEditor/Private/PCGEditor.cpp b/Engine/Plugins/Experimental/PCG/Source/PCGEditor/Private/PCGEditor.cpp
--- a/Engine/Plugins/Experimental/PCG/Source/PCGEditor/Private/PCGEditor.cpp
+++ b/Engine/Plugins/Experimental/PCG/Source/PCGEditor/Private/PCGEditor.cpp
@@ -82,6 +82,7 @@ void FPCGEditor::Initialize(const EToolkitMode::Type InMode, const TSharedPtr<cl
 	PCGEditorGraph = NewObject<UPCGEditorGraph>(PCGGraphBeingEdited, UPCGEditorGraph::StaticClass(), NAME_None, RF_Transactional | RF_Transient);
 	PCGEditorGraph->Schema = UPCGEditorGraphSchema::StaticClass();
 	PCGEditorGraph->InitFromNodeGraph(InPCGGraph);
+	PCGEditorGraph->SetEditor(SharedThis(this));
 
 	FPropertyEditorModule& PropertyEditorModule = FModuleManager::GetModuleChecked<FPropertyEditorModule>("PropertyEditor");
 
@@ -171,6 +172,15 @@ void FPCGEditor::SetPCGComponentBeingDebugged(UPCGComponent* InPCGComponent)
 			// We need to force generation so that we create debug data
 			PCGComponentBeingDebugged->GenerateLocal(/*bForce=*/true);
 		}
+
+		for (UEdGraphNode* Node : PCGEditorGraph->Nodes)
+		{
+			if (UPCGEditorGraphNodeBase* PCGNode = Cast<UPCGEditorGraphNodeBase>(Node))
+			{
+				// Update now that component has changed. Will fire OnNodeChanged if necessary.
+				PCGNode->UpdateErrorsAndWarnings();
+			}
+		}
 	}
 }
 
diff --git a/Engine/Plugins/Experimental/PCG/Source/PCGEditor/Private/PCGEditorGraph.h b/Engine/Plugins/Experimental/PCG/Source/PCGEditor/Private/PCGEditorGraph.h
--- a/Engine/Plugins/Experimental/PCG/Source/PCGEditor/Private/PCGEditorGraph.h
+++ b/Engine/Plugins/Experimental/PCG/Source/PCGEditor/Private/PCGEditorGraph.h
@@ -6,6 +6,7 @@
 
 #include "PCGEditorGraph.generated.h"
 
+class FPCGEditor;
 class UPCGNode;
 class UPCGGraph;
 class UPCGEditorGraphNodeBase;
@@ -24,10 +25,15 @@ public:
 
 	UPCGGraph* GetPCGGraph() { return PCGGraph; }
 
+	void SetEditor(TWeakPtr<const FPCGEditor> InEditor) { PCGEditor = InEditor; }
+	TWeakPtr<const FPCGEditor> GetEditor() const { return PCGEditor; }
+
 protected:
 	void CreateLinks(UPCGEditorGraphNodeBase* InGraphNode, bool bCreateInbound, bool bCreateOutbound, const TMap<UPCGNode*, UPCGEditorGraphNodeBase*>& InGraphNodeToPCGNodeMap);
 
 private:
 	UPROPERTY()
 	TObjectPtr<UPCGGraph> PCGGraph = nullptr;
+
+	TWeakPtr<const FPCGEditor> PCGEditor = nullptr;
 };
diff --git a/Engine/Plugins/Experimental/PCG/Source/PCGEditor/Private/PCGEditorGraphNodeBase.cpp b/Engine/Plugins/Experimental/PCG/Source/PCGEditor/Private/PCGEditorGraphNodeBase.cpp
--- a/Engine/Plugins/Experimental/PCG/Source/PCGEditor/Private/PCGEditorGraphNodeBase.cpp
+++ b/Engine/Plugins/Experimental/PCG/Source/PCGEditor/Private/PCGEditorGraphNodeBase.cpp
@@ -2,6 +2,7 @@
 
 #include "PCGEditorGraphNodeBase.h"
 
+#include "PCGEditor.h"
 #include "PCGEditorCommands.h"
 #include "PCGEditorCommon.h"
 #include "PCGEditorGraph.h"
@@ -9,10 +10,12 @@
 #include "PCGEditorSettings.h"
 #include "PCGGraph.h"
 #include "PCGPin.h"
+#include "PCGSubsystem.h"
 
 #include "GraphEditorActions.h"
 #include "ToolMenu.h"
 #include "ToolMenuSection.h"
+#include "Logging/TokenizedMessage.h"
 #include "Misc/TransactionObjectEvent.h"
 #include "Widgets/Colors/SColorPicker.h"
 
@@ -275,6 +278,8 @@ void UPCGEditorGraphNodeBase::OnNodeChanged(UPCGNode* InNode, EPCGChangeType Cha
 			}
 		}
 
+		UpdateErrorsAndWarnings();
+
 		if (!!(ChangeType & (EPCGChangeType::Structural | EPCGChangeType::Node | EPCGChangeType::Edge | EPCGChangeType::Cosmetic)))
 		{
 			ReconstructNode();
@@ -282,6 +287,48 @@ void UPCGEditorGraphNodeBase::OnNodeChanged(UPCGNode* InNode, EPCGChangeType Cha
 	}
 }
 
+void UPCGEditorGraphNodeBase::UpdateErrorsAndWarnings()
+{
+	bool bNodeChanged = false;
+
+	// Pull current errors/warnings state from PCG subsystem.
+	if (const UPCGSubsystem* Subsystem = UPCGSubsystem::GetActiveEditorInstance())
+	{
+		const UPCGComponent* ComponentBeingDebugged = nullptr;
+		{
+			const UPCGEditorGraph* EditorGraph = CastChecked<UPCGEditorGraph>(GetGraph());
+			const FPCGEditor* Editor = (EditorGraph && EditorGraph->GetEditor().IsValid()) ? EditorGraph->GetEditor().Pin().Get() : nullptr;
+			ComponentBeingDebugged = Editor ? Editor->GetPCGComponentBeingDebugged() : nullptr;
+		}
+
+		const bool bOldHasCompilerMessage = bHasCompilerMessage;
+		const int32 OldErrorType = ErrorType;
+		const FString OldErrorMsg = ErrorMsg;
+
+		bHasCompilerMessage = Subsystem->GetNodeVisualLogs().HasLogs(PCGNode, ComponentBeingDebugged);
+
+		if (bHasCompilerMessage)
+		{
+			const bool bHasErrors = Subsystem->GetNodeVisualLogs().HasLogs(PCGNode, ComponentBeingDebugged, ELogVerbosity::Error);
+			ErrorType = bHasErrors ? EMessageSeverity::Error : EMessageSeverity::Warning;
+
+			ErrorMsg = Subsystem->GetNodeVisualLogs().GetLogsSummaryText(PCGNode, ComponentBeingDebugged).ToString();
+		}
+		else
+		{
+			ErrorMsg.Empty();
+			ErrorType = 0;
+		}
+
+		bNodeChanged = (bHasCompilerMessage != bOldHasCompilerMessage) || (ErrorType != OldErrorType) || (ErrorMsg != OldErrorMsg);
+	}
+
+	if (bNodeChanged)
+	{
+		OnNodeChangedDelegate.ExecuteIfBound();
+	}
+}
+
 void UPCGEditorGraphNodeBase::OnPickColor()
 {
 	FColorPickerArgs PickerArgs;
diff --git a/Engine/Plugins/Experimental/PCG/Source/PCGEditor/Private/PCGEditorGraphNodeBase.h b/Engine/Plugins/Experimental/PCG/Source/PCGEditor/Private/PCGEditorGraphNodeBase.h
--- a/Engine/Plugins/Experimental/PCG/Source/PCGEditor/Private/PCGEditorGraphNodeBase.h
+++ b/Engine/Plugins/Experimental/PCG/Source/PCGEditor/Private/PCGEditorGraphNodeBase.h
@@ -9,6 +9,7 @@
 
 enum class EPCGChangeType : uint8;
 
+class UPCGComponent;
 class UPCGNode;
 class UPCGPin;
 class UToolMenu;
@@ -54,6 +55,9 @@ public:
 	/** Decrease deferred reconstruct counter, ReconstructNode will be called if counter hits zero and the node is flagged for reconstruction  */
 	void DisableDeferredReconstruct();
 
+	/** Pulls current errors/warnings state from PCG subsystem. */
+	void UpdateErrorsAndWarnings();
+
 	DECLARE_DELEGATE(FOnPCGEditorGraphNodeChanged);
 	FOnPCGEditorGraphNodeChanged OnNodeChangedDelegate;
 
diff --git a/Engine/Plugins/Experimental/PCGGeometryScriptInterop/Source/PCGGeometryScriptInterop/Private/Elements/PCGMeshSampler.cpp b/Engine/Plugins/Experimental/PCGGeometryScriptInterop/Source/PCGGeometryScriptInterop/Private/Elements/PCGMeshSampler.cpp
--- a/Engine/Plugins/Experimental/PCGGeometryScriptInterop/Source/PCGGeometryScriptInterop/Private/Elements/PCGMeshSampler.cpp
+++ b/Engine/Plugins/Experimental/PCGGeometryScriptInterop/Source/PCGGeometryScriptInterop/Private/Elements/PCGMeshSampler.cpp
@@ -95,7 +95,7 @@ bool FPCGMeshSamplerElement::PrepareDataInternal(FPCGContext* InContext) const
 
 	if (!StaticMesh)
 	{
-		PCGE_LOG(Error, "Provided static mesh does not exist: %s", *Settings->StaticMeshPath.ToString());
+		PCGE_LOG(Error, GraphAndLog, FText::Format(LOCTEXT("MeshDoesNotExist", "Provided static mesh does not exist: '{0}'"), FText::FromString(Settings->StaticMeshPath.ToString())));
 		return true;
 	}
 
@@ -110,7 +110,7 @@ bool FPCGMeshSamplerElement::PrepareDataInternal(FPCGContext* InContext) const
 
 	if (Outcome == EGeometryScriptOutcomePins::Failure)
 	{
-		PCGE_LOG(Error, "Static to Dynamic mesh failed");
+		PCGE_LOG(Error, GraphAndLog, LOCTEXT("StaticToDynamicMeshFailed", "Static to Dynamic mesh failed"));
 		return true;
 	}
 	
@@ -339,4 +339,4 @@ bool FPCGMeshSamplerElement::ExecuteInternal(FPCGContext* InContext) const
 	return bIsDone;
 }
 
-#undef LOCTEXT_NAMESPACE
\ No newline at end of file
+#undef LOCTEXT_NAMESPACE
