[PCG] Fix static analysis warnings

Added couple of null checks.

More substantial change solves following problem - macro expands to things like:

if (EPCGElementLogMode::GraphAndLog == EPCGElementLogMode::GraphAndLog)

Which it complains is suspicious-looking. I could not simply add an ignore comment into the log macro but I instead found this pattern in AsyncLoading2.cpp to make the comparison appear less-trivial.

diff --git a/Engine/Plugins/Experimental/PCG/Source/PCG/Private/Graph/PCGGraphExecutor.cpp b/Engine/Plugins/Experimental/PCG/Source/PCG/Private/Graph/PCGGraphExecutor.cpp
--- a/Engine/Plugins/Experimental/PCG/Source/PCG/Private/Graph/PCGGraphExecutor.cpp
+++ b/Engine/Plugins/Experimental/PCG/Source/PCG/Private/Graph/PCGGraphExecutor.cpp
@@ -100,7 +100,7 @@ FPCGTaskId FPCGGraphExecutor::Schedule(UPCGGraph* Graph, UPCGComponent* SourceCo
 	}
 
 #if WITH_EDITOR
-	if (UPCGSubsystem* Subsystem = UPCGSubsystem::GetInstance(SourceComponent->GetWorld()))
+	if (UPCGSubsystem* Subsystem = SourceComponent ? UPCGSubsystem::GetInstance(SourceComponent->GetWorld()) : nullptr)
 	{
 		for (const UPCGNode* Node : Graph->GetNodes())
 		{
diff --git a/Engine/Plugins/Experimental/PCG/Source/PCG/Public/PCGElement.h b/Engine/Plugins/Experimental/PCG/Source/PCG/Public/PCGElement.h
--- a/Engine/Plugins/Experimental/PCG/Source/PCG/Public/PCGElement.h
+++ b/Engine/Plugins/Experimental/PCG/Source/PCG/Public/PCGElement.h
@@ -18,31 +18,42 @@ class UPCGNode;
 
 typedef TSharedPtr<IPCGElement, ESPMode::ThreadSafe> FPCGElementPtr;
 
-enum class EPCGElementLogMode : uint8
+namespace EPCGElementLogMode
 {
-	/** Output to log. */
-	LogOnly,
-	/** Display errors/warnings on graph as well as writing to log. */
-	GraphAndLog,
+	enum Type : uint8
+	{
+		/** Output to log. */
+		LogOnly = 0,
+		/** Display errors/warnings on graph as well as writing to log. */
+		GraphAndLog,
+
+		NumLogModes,
+
+		// Used below in log macros to silence PVS by making the log mode comparison 'look' non-trivial by adding a trivial mask op (an identical
+		// mechanism as the one employed in the macro UE_ASYNC_PACKAGE_LOG in AsyncLoading2.cpp).
+		// Warning V501: There are identical sub-expressions 'EPCGElementLogMode::GraphAndLog' to the left and to the right of the '==' operator.
+		// The warning disable comment can can't be used in a macro: //-V501 
+		LogModeMask = 0xff
+	};
 };
 
 #define PCGE_LOG_BASE(Verbosity, CustomContext, Message) \
 	UE_LOG(LogPCG, \
 		Verbosity, \
 		TEXT("[%s - %s]: %s"), \
-		*((CustomContext)->GetComponentName()), \
-		*((CustomContext)->GetTaskName()), \
+		(CustomContext) ? *((CustomContext)->GetComponentName()) : TEXT("UnknownComponent"), \
+		(CustomContext) ? *((CustomContext)->GetTaskName()) : TEXT("UnknownTask"), \
 		*Message.ToString())
 
 #if WITH_EDITOR
-// Output to PCG log and optionally also display wanrings/errors on graph.
+// Output to PCG log and optionally also display warnings/errors on graph.
 #define PCGE_LOG(Verbosity, LogMode, Message) do { \
-		if ((EPCGElementLogMode::LogMode) == EPCGElementLogMode::GraphAndLog && (Context)) { (Context)->LogVisual(ELogVerbosity::Verbosity, Message); } \
+		if (((EPCGElementLogMode::LogMode) & EPCGElementLogMode::LogModeMask) == EPCGElementLogMode::GraphAndLog && (Context)) { (Context)->LogVisual(ELogVerbosity::Verbosity, Message); } \
 		if (ShouldLog()) { PCGE_LOG_BASE(Verbosity, Context, Message); } \
 	} while (0)
-// Output to PCG log and optionally also display wanrings/errors on graph. Takes context as argument.
+// Output to PCG log and optionally also display warnings/errors on graph. Takes context as argument.
 #define PCGE_LOG_C(Verbosity, LogMode, CustomContext, Message) do { \
-		if ((EPCGElementLogMode::LogMode) == EPCGElementLogMode::GraphAndLog && (CustomContext)) { (CustomContext)->LogVisual(ELogVerbosity::Verbosity, Message); } \
+		if (((EPCGElementLogMode::LogMode) & EPCGElementLogMode::LogModeMask) == EPCGElementLogMode::GraphAndLog && (CustomContext)) { (CustomContext)->LogVisual(ELogVerbosity::Verbosity, Message); } \
 		PCGE_LOG_BASE(Verbosity, CustomContext, Message); \
 	} while (0)
 #else
