[PCG] Fix task input CRC - the wrong pin label was being CRC'd, which meant one could move an incident edge across input pins and not get a CRC change.

diff --git a/Engine/Plugins/Experimental/PCG/Source/PCG/Private/Graph/PCGGraphExecutor.cpp b/Engine/Plugins/Experimental/PCG/Source/PCG/Private/Graph/PCGGraphExecutor.cpp
--- a/Engine/Plugins/Experimental/PCG/Source/PCG/Private/Graph/PCGGraphExecutor.cpp
+++ b/Engine/Plugins/Experimental/PCG/Source/PCG/Private/Graph/PCGGraphExecutor.cpp
@@ -1031,8 +1031,8 @@ void FPCGGraphExecutor::BuildTaskInput(const FPCGGraphTask& Task, FPCGDataCollec
 		{
 			TaskInput.TaggedData.Append(InputCollection.GetInputsByPin(Input.InPin->Properties.Label));
 
-			// Write pin name Crc to uniquely identify inputs per-pin.
-			Crc.Combine(GetTypeHash(Input.InPin->Properties.Label));
+			// Write input pin name Crc to uniquely identify inputs per-pin.
+			Crc.Combine(GetTypeHash(Input.OutPin->Properties.Label));
 		}
 		else
 		{
diff --git a/Engine/Plugins/Experimental/PCG/Source/PCG/Private/Graph/PCGGraphExecutor.h b/Engine/Plugins/Experimental/PCG/Source/PCG/Private/Graph/PCGGraphExecutor.h
--- a/Engine/Plugins/Experimental/PCG/Source/PCG/Private/Graph/PCGGraphExecutor.h
+++ b/Engine/Plugins/Experimental/PCG/Source/PCG/Private/Graph/PCGGraphExecutor.h
@@ -29,8 +29,14 @@ struct FPCGGraphTaskInput
 	}
 
 	FPCGTaskId TaskId;
+
+	/** The upstream output pin from which the input data comes. */
 	const UPCGPin* InPin;
+
+	/** The input pin on the task element. */
 	const UPCGPin* OutPin;
+
+	/** Whether the input provides any data. For the post execute task, only the output node will provide data. */
 	bool bProvideData;
 };
 
