[PCG] Fixed a bug where the debug object dropdown would always reset the selection on open

diff --git a/Engine/Plugins/Experimental/PCG/Source/PCGEditor/Private/SPCGEditorGraphDebugObjectWidget.cpp b/Engine/Plugins/Experimental/PCG/Source/PCGEditor/Private/SPCGEditorGraphDebugObjectWidget.cpp
--- a/Engine/Plugins/Experimental/PCG/Source/PCGEditor/Private/SPCGEditorGraphDebugObjectWidget.cpp
+++ b/Engine/Plugins/Experimental/PCG/Source/PCGEditor/Private/SPCGEditorGraphDebugObjectWidget.cpp
@@ -66,7 +66,8 @@ void SPCGEditorGraphDebugObjectWidget::Construct(const FArguments& InArgs, TShar
 void SPCGEditorGraphDebugObjectWidget::OnComboBoxOpening()
 {
 	DebugObjects.Empty();
-
+	DebugObjectsComboBox->RefreshOptions();
+	
 	if (!PCGEditorPtr.IsValid())
 	{
 		return;
@@ -84,9 +85,15 @@ void SPCGEditorGraphDebugObjectWidget::OnComboBoxOpening()
 		return;
 	}
 
+	TSharedPtr<FPCGEditorGraphDebugObjectInstance> SelectedItem = DebugObjectsComboBox->GetSelectedItem();
+
 	DebugObjects.Add(MakeShared<FPCGEditorGraphDebugObjectInstance>(PCGEditorGraphDebugObjectWidget::NoObjectString));
-	DebugObjectsComboBox->SetSelectedItem(DebugObjects[0]);
 
+	if (!SelectedItem.IsValid() || SelectedItem->PCGComponent == nullptr)
+	{
+		DebugObjectsComboBox->SetSelectedItem(DebugObjects[0]);
+	}
+	
 	TArray<UObject*> PCGComponents;
 	GetObjectsOfClass(UPCGComponent::StaticClass(), PCGComponents, true);
 	for (UObject* PCGComponentObject : PCGComponents)
@@ -129,7 +136,13 @@ void SPCGEditorGraphDebugObjectWidget::OnComboBoxOpening()
 			}
 
 			const FString Label = ActorNameOrLabel + PCGEditorGraphDebugObjectWidget::SeparatorString + ComponentName;
-			DebugObjects.Add(MakeShared<FPCGEditorGraphDebugObjectInstance>(PCGComponent, Label));
+			TSharedPtr<FPCGEditorGraphDebugObjectInstance> DebugInstance = MakeShared<FPCGEditorGraphDebugObjectInstance>(PCGComponent, Label);
+			DebugObjects.Add(DebugInstance);
+
+			if (SelectedItem.IsValid() && SelectedItem->PCGComponent == PCGComponent)
+			{
+				DebugObjectsComboBox->SetSelectedItem(DebugInstance);
+			}
 		}
 	}
 }
