[PCG] Fix random crash when tracking actors dependencies

#rnx

diff --git a/Engine/Plugins/Experimental/PCG/Source/PCG/Private/PCGHelpers.cpp b/Engine/Plugins/Experimental/PCG/Source/PCG/Private/PCGHelpers.cpp
--- a/Engine/Plugins/Experimental/PCG/Source/PCG/Private/PCGHelpers.cpp
+++ b/Engine/Plugins/Experimental/PCG/Source/PCG/Private/PCGHelpers.cpp
@@ -362,11 +362,11 @@ namespace PCGHelpers
 			{
 				if (Helper.IsValidIndex(DynamicIndex))
 				{
-					const void* KeyPtr = Helper.GetKeyPtr(DynamicIndex);
-					GatherDependencies(MapProperty->KeyProp, KeyPtr, OutDependencies, MaxDepth);
-
-					const void* ValuePtr = Helper.GetValuePtr(DynamicIndex);
-					GatherDependencies(MapProperty->ValueProp, ValuePtr, OutDependencies, MaxDepth);
+					// Key and Value are stored next to each other in memory.
+					// ValueProp has an offset, so we should use the same starting address for both.
+					const void* PairKeyValuePtr = Helper.GetKeyPtr(DynamicIndex);
+					GatherDependencies(MapProperty->KeyProp, PairKeyValuePtr, OutDependencies, MaxDepth);
+					GatherDependencies(MapProperty->ValueProp, PairKeyValuePtr, OutDependencies, MaxDepth);
 
 					--Num;
 				}
