[PCG] Clean out graph visual logs for invalid components

diff --git a/Engine/Plugins/Experimental/PCG/Source/PCG/Private/Utils/PCGNodeVisualLogs.cpp b/Engine/Plugins/Experimental/PCG/Source/PCG/Private/Utils/PCGNodeVisualLogs.cpp
--- a/Engine/Plugins/Experimental/PCG/Source/PCG/Private/Utils/PCGNodeVisualLogs.cpp
+++ b/Engine/Plugins/Experimental/PCG/Source/PCG/Private/Utils/PCGNodeVisualLogs.cpp
@@ -178,7 +178,8 @@ void FPCGNodeVisualLogs::ClearLogs(TWeakObjectPtr<const UPCGNode> InNode, const
 		check(NodeLogs);
 		for (int32 i = NodeLogs->Num() - 1; i >= 0; --i)
 		{
-			if ((*NodeLogs)[i].Component == InComponent)
+			// Remove entry if it matches the given component, or if the component is no longer valid
+			if (!(*NodeLogs)[i].Component.IsValid() || (*NodeLogs)[i].Component == InComponent)
 			{
 				NodeLogs->RemoveAtSwap(i);
 				bAnyRemoved = true;
