[PCG] Static analysis - guard against pin pointers being null

diff --git a/Engine/Plugins/Experimental/PCG/Source/PCG/Private/Graph/PCGGraphExecutor.cpp b/Engine/Plugins/Experimental/PCG/Source/PCG/Private/Graph/PCGGraphExecutor.cpp
--- a/Engine/Plugins/Experimental/PCG/Source/PCG/Private/Graph/PCGGraphExecutor.cpp
+++ b/Engine/Plugins/Experimental/PCG/Source/PCG/Private/Graph/PCGGraphExecutor.cpp
@@ -1032,7 +1032,7 @@ void FPCGGraphExecutor::BuildTaskInput(const FPCGGraphTask& Task, FPCGDataCollec
 			TaskInput.TaggedData.Append(InputCollection.GetInputsByPin(Input.InPin->Properties.Label));
 
 			// Write input pin name Crc to uniquely identify inputs per-pin.
-			Crc.Combine(GetTypeHash(Input.OutPin->Properties.Label));
+			Crc.Combine(GetTypeHash(Input.OutPin ? Input.OutPin->Properties.Label : FName(TEXT("MissingLabel"))));
 		}
 		else
 		{
