[PCG] Renamed current "Create Attribute" node to "Add Attribute" node, made a new "Create Attribute Set" node, rename user facing name of "Param Data" to "Attribute Set"

Both of these renames done at surface level to avoid difficulties of changing C++ types (we could tackle this later if we want).

Also renamed 'Params' pins to 'Overrides'.

BPs that have a custom Params or Param pin will get them removed and the code in UpdatePins will make sure the edges are preserved.

diff --git a/Engine/Plugins/Experimental/PCG/Source/PCG/Private/Elements/PCGCreateAttribute.cpp b/Engine/Plugins/Experimental/PCG/Source/PCG/Private/Elements/PCGCreateAttribute.cpp
--- a/Engine/Plugins/Experimental/PCG/Source/PCG/Private/Elements/PCGCreateAttribute.cpp
+++ b/Engine/Plugins/Experimental/PCG/Source/PCG/Private/Elements/PCGCreateAttribute.cpp
@@ -14,8 +14,9 @@
 
 namespace PCGCreateAttributeConstants
 {
-	const FName NodeName = TEXT("CreateAttribute");
-	const FName SourceLabel = TEXT("Source");
+	const FName NodeNameAddAttribute = TEXT("Add Attribute");
+	const FName NodeNameCreateAttribute = TEXT("Create Attribute");
+	const FName AttributesLabel = TEXT("Attributes");
 }
 
 void UPCGCreateAttributeSettings::PostLoad()
@@ -48,30 +49,13 @@ void UPCGCreateAttributeSettings::PostLoad()
 
 FName UPCGCreateAttributeSettings::AdditionalTaskName() const
 {
-	if (bFromSourceParam)
-	{
-		const FName NodeName = PCGCreateAttributeConstants::NodeName;
-
-		if ((OutputAttributeName == NAME_None) && (SourceParamAttributeName == NAME_None))
-		{
-			return NodeName;
-		}
-		else
-		{
-			const FString AttributeName = ((OutputAttributeName == NAME_None) ? SourceParamAttributeName : OutputAttributeName).ToString();
-			return FName(FString::Printf(TEXT("%s %s"), *NodeName.ToString(), *AttributeName));
-		}
-	}
-	else
-	{
-		return FName(FString::Printf(TEXT("%s: %s"), *OutputAttributeName.ToString(), *AttributeTypes.ToString()));
-	}
+	return AdditionalTaskNameInternal(PCGCreateAttributeConstants::NodeNameAddAttribute);
 }
 
 #if WITH_EDITOR
 FName UPCGCreateAttributeSettings::GetDefaultNodeName() const
 {
-	return PCGCreateAttributeConstants::NodeName;
+	return PCGCreateAttributeConstants::NodeNameAddAttribute;
 }
 #endif // WITH_EDITOR
 
@@ -90,7 +74,7 @@ TArray<FPCGPinProperties> UPCGCreateAttributeSettings::InputPinProperties() cons
 
 	if (bFromSourceParam)
 	{
-		PinProperties.Emplace(PCGCreateAttributeConstants::SourceLabel, EPCGDataType::Param);
+		PinProperties.Emplace(PCGCreateAttributeConstants::AttributesLabel, EPCGDataType::Param);
 	}
 
 	return PinProperties;
@@ -116,6 +100,53 @@ FPCGElementPtr UPCGCreateAttributeSettings::CreateElement() const
 	return MakeShared<FPCGCreateAttributeElement>();
 }
 
+FName UPCGCreateAttributeSettings::AdditionalTaskNameInternal(FName NodeName) const
+{
+	if (bFromSourceParam)
+	{
+		if ((OutputAttributeName == NAME_None) && (SourceParamAttributeName == NAME_None))
+		{
+			return NodeName;
+		}
+		else
+		{
+			const FString AttributeName = ((OutputAttributeName == NAME_None) ? SourceParamAttributeName : OutputAttributeName).ToString();
+			return FName(FString::Printf(TEXT("%s %s"), *NodeName.ToString(), *AttributeName));
+		}
+	}
+	else
+	{
+		return FName(FString::Printf(TEXT("%s: %s"), *OutputAttributeName.ToString(), *AttributeTypes.ToString()));
+	}
+}
+
+UPCGCreateAttributeSetSettings::UPCGCreateAttributeSetSettings()
+{
+	// No input pin to grab source param from
+	bDisplayFromSourceParamSetting = false;
+}
+
+#if WITH_EDITOR
+FName UPCGCreateAttributeSetSettings::GetDefaultNodeName() const
+{
+	return PCGCreateAttributeConstants::NodeNameCreateAttribute;
+}
+#endif // WITH_EDITOR
+
+TArray<FPCGPinProperties> UPCGCreateAttributeSetSettings::OutputPinProperties() const
+{
+	FPCGPinProperties PinProperties;
+	PinProperties.Label = PCGPinConstants::DefaultOutputLabel;
+	PinProperties.AllowedTypes = EPCGDataType::Param;
+
+	return { PinProperties };
+}
+
+FName UPCGCreateAttributeSetSettings::AdditionalTaskName() const
+{
+	return AdditionalTaskNameInternal(PCGCreateAttributeConstants::NodeNameCreateAttribute);
+}
+
 bool FPCGCreateAttributeElement::ExecuteInternal(FPCGContext* Context) const
 {
 	TRACE_CPUPROFILER_EVENT_SCOPE(FPCGCreateAttributeElement::Execute);
@@ -125,7 +156,7 @@ bool FPCGCreateAttributeElement::ExecuteInternal(FPCGContext* Context) const
 	const UPCGCreateAttributeSettings* Settings = Context->GetInputSettings<UPCGCreateAttributeSettings>();
 	check(Settings);
 
-	TArray<FPCGTaggedData> SourceParams = Context->InputData.GetInputsByPin(PCGCreateAttributeConstants::SourceLabel);
+	TArray<FPCGTaggedData> SourceParams = Context->InputData.GetInputsByPin(PCGCreateAttributeConstants::AttributesLabel);
 	UPCGParamData* SourceParamData = nullptr;
 	FName SourceParamAttributeName = NAME_None;
 
@@ -276,5 +307,4 @@ PCGMetadataEntryKey FPCGCreateAttributeElement::SetAttribute(const UPCGCreateAtt
 
 	return Settings->AttributeTypes.Dispatcher(SetAttribute);
 }
-
 #undef LOCTEXT_NAMESPACE
diff --git a/Engine/Plugins/Experimental/PCG/Source/PCG/Private/Elements/PCGExecuteBlueprint.cpp b/Engine/Plugins/Experimental/PCG/Source/PCG/Private/Elements/PCGExecuteBlueprint.cpp
--- a/Engine/Plugins/Experimental/PCG/Source/PCG/Private/Elements/PCGExecuteBlueprint.cpp
+++ b/Engine/Plugins/Experimental/PCG/Source/PCG/Private/Elements/PCGExecuteBlueprint.cpp
@@ -66,6 +66,18 @@ void UPCGBlueprintElement::PostLoad()
 
 		OutputPinLabels_DEPRECATED.Reset();
 	}
+
+	// Go through the user-defined custom input pins and remove any Param pins labelled 'Params' or 'Param'. Such pins should not be
+	// added manually, the params pin is created dynamically from code based on presence of overrides.
+	for (int32 i = CustomInputPins.Num() - 1; i >= 0; --i)
+	{
+		FPCGPinProperties& Properties = CustomInputPins[i];
+		if (Properties.AllowedTypes == EPCGDataType::Param && (Properties.Label == FName(TEXT("Params")) || Properties.Label == FName(TEXT("Param"))))
+		{
+			// Non-swap version to preserve order
+			CustomInputPins.RemoveAt(i);
+		}
+	}
 #endif
 }
 
@@ -411,6 +423,22 @@ UObject* UPCGBlueprintSettings::GetJumpTargetForDoubleClick() const
 	
 	return Super::GetJumpTargetForDoubleClick();
 }
+
+void UPCGBlueprintSettings::ApplyDeprecationBeforeUpdatePins(UPCGNode* InOutNode, TArray<TObjectPtr<UPCGPin>>& InputPins, TArray<TObjectPtr<UPCGPin>>& OutputPins)
+{
+	Super::ApplyDeprecationBeforeUpdatePins(InOutNode, InputPins, OutputPins);
+
+	// Rename first found 'Param' or 'Params' pin to 'Overrides' which helps to ensure legacy params pins will retain incident edges.
+	for (TObjectPtr<UPCGPin>& InputPin : InputPins)
+	{
+		if (InputPin && InputPin->Properties.AllowedTypes == EPCGDataType::Param && 
+			(InputPin->Properties.Label == FName(TEXT("Params")) || InputPin->Properties.Label == FName(TEXT("Param"))))
+		{
+			InputPin->Properties.Label = PCGPinConstants::DefaultParamsLabel;
+			break;
+		}
+	}
+}
 #endif // WITH_EDITOR
 
 FName UPCGBlueprintSettings::AdditionalTaskName() const
diff --git a/Engine/Plugins/Experimental/PCG/Source/PCG/Public/Elements/PCGCreateAttribute.h b/Engine/Plugins/Experimental/PCG/Source/PCG/Public/Elements/PCGCreateAttribute.h
--- a/Engine/Plugins/Experimental/PCG/Source/PCG/Public/Elements/PCGCreateAttribute.h
+++ b/Engine/Plugins/Experimental/PCG/Source/PCG/Public/Elements/PCGCreateAttribute.h
@@ -10,12 +10,10 @@
 
 class FPCGMetadataAttributeBase;
 class UPCGMetadata;
-class UPCGParamData;
 
 /**
-* Node to create an attribute from a Spatial data or a ParamData.
-* If no input is provided, the node behave as a "CreateParamData" and will output a ParamData
-* with the given attribute set.
+* Adds an attribute to Spatial data or to an Attribute Set, or creates a new Attribute Set if no
+* input is provided.
 * 
 * Note: This need to be updated if we ever add new types.
 */
@@ -35,18 +33,19 @@ public:
 	virtual bool HasDynamicPins() const override { return true; }
 #endif // WITH_EDITOR
 
-	virtual TArray<FPCGPinProperties> InputPinProperties() const override;
-	virtual TArray<FPCGPinProperties> OutputPinProperties() const override;
-
 	virtual FName AdditionalTaskName() const override;
 	//~End UPCGSettings interface
 
 	UPROPERTY(BlueprintReadWrite, EditAnywhere, Category = Settings)
 	FName OutputAttributeName = NAME_None;
 
-	UPROPERTY(BlueprintReadWrite, EditAnywhere, Category = Settings)
+	UPROPERTY(BlueprintReadWrite, EditAnywhere, Category = Settings, meta = (EditCondition = bDisplayFromSourceParamSetting, EditConditionHides, HideEditConditionToggle))
 	bool bFromSourceParam = false;
 
+	// This can be set false by inheriting nodes to hide the 'From Source Param' property.
+	UPROPERTY(Transient, meta = (EditCondition = false, EditConditionHides))
+	bool bDisplayFromSourceParamSetting = true;
+
 	UPROPERTY(BlueprintReadWrite, EditAnywhere, Category = Settings, meta = (EditCondition = "bFromSourceParam", EditConditionHides))
 	FName SourceParamAttributeName = NAME_None;
 
@@ -101,9 +100,35 @@ public:
 #endif // WITH_EDITORONLY_DATA
 
 protected:
+	virtual TArray<FPCGPinProperties> InputPinProperties() const override;
+	virtual TArray<FPCGPinProperties> OutputPinProperties() const override;
+
+	FName AdditionalTaskNameInternal(FName NodeName) const;
+
 	virtual FPCGElementPtr CreateElement() const override;
 };
 
+/* Creates a new Attribute Set. */
+UCLASS(BlueprintType, ClassGroup = (Procedural))
+class PCG_API UPCGCreateAttributeSetSettings : public UPCGCreateAttributeSettings
+{
+	GENERATED_BODY()
+
+public:
+	UPCGCreateAttributeSetSettings();
+
+	//~Begin UPCGSettings interface
+#if WITH_EDITOR
+	virtual FName GetDefaultNodeName() const override;
+#endif
+	virtual FName AdditionalTaskName() const override;
+
+protected:
+	virtual TArray<FPCGPinProperties> InputPinProperties() const override { return TArray<FPCGPinProperties>(); }
+	virtual TArray<FPCGPinProperties> OutputPinProperties() const override;
+	//~End UPCGSettings interface
+};
+
 class FPCGCreateAttributeElement : public FSimplePCGElement
 {
 protected:
diff --git a/Engine/Plugins/Experimental/PCG/Source/PCG/Public/Elements/PCGExecuteBlueprint.h b/Engine/Plugins/Experimental/PCG/Source/PCG/Public/Elements/PCGExecuteBlueprint.h
--- a/Engine/Plugins/Experimental/PCG/Source/PCG/Public/Elements/PCGExecuteBlueprint.h
+++ b/Engine/Plugins/Experimental/PCG/Source/PCG/Public/Elements/PCGExecuteBlueprint.h
@@ -187,6 +187,7 @@ public:
 	virtual EPCGSettingsType GetType() const override;
 	virtual void GetTrackedActorTags(FPCGTagToSettingsMap& OutTagToSettings, TArray<TObjectPtr<const UPCGGraph>>& OutVisitedGraphs) const override;
 	virtual UObject* GetJumpTargetForDoubleClick() const override;
+	virtual void ApplyDeprecationBeforeUpdatePins(UPCGNode* InOutNode, TArray<TObjectPtr<UPCGPin>>& InputPins, TArray<TObjectPtr<UPCGPin>>& OutputPins) override;
 #endif
 
 	virtual FName AdditionalTaskName() const override;
diff --git a/Engine/Plugins/Experimental/PCG/Source/PCG/Public/PCGCommon.h b/Engine/Plugins/Experimental/PCG/Source/PCG/Public/PCGCommon.h
--- a/Engine/Plugins/Experimental/PCG/Source/PCG/Public/PCGCommon.h
+++ b/Engine/Plugins/Experimental/PCG/Source/PCG/Public/PCGCommon.h
@@ -59,7 +59,7 @@ namespace PCGPinConstants
 {
 	const FName DefaultInputLabel = TEXT("In");
 	const FName DefaultOutputLabel = TEXT("Out");
-	const FName DefaultParamsLabel = TEXT("Params");
+	const FName DefaultParamsLabel = TEXT("Overrides");
 }
 
 // Metadata used by PCG
diff --git a/Engine/Plugins/Experimental/PCG/Source/PCGEditor/Private/PCGEditorCommon.h b/Engine/Plugins/Experimental/PCG/Source/PCGEditor/Private/PCGEditorCommon.h
--- a/Engine/Plugins/Experimental/PCG/Source/PCGEditor/Private/PCGEditorCommon.h
+++ b/Engine/Plugins/Experimental/PCG/Source/PCGEditor/Private/PCGEditorCommon.h
@@ -28,7 +28,7 @@ namespace FPCGEditorCommon
 	const FString ContextIdentifier = TEXT("PCGEditorContext");
 
 	const FName SpatialDataType = FName(TEXT("Spatial Data"));
-	const FName ParamDataType = FName(TEXT("Param Data"));
+	const FName ParamDataType = FName(TEXT("Attribute Set"));
 	const FName SettingsDataType = FName(TEXT("Settings Data"));
 	const FName OtherDataType = FName(TEXT("Other Data"));
 
