[PCG] Narrowing of Gather input/output type - narrow to union of all incident types

diff --git a/Engine/Plugins/Experimental/PCG/Source/PCG/Private/Elements/PCGGather.cpp b/Engine/Plugins/Experimental/PCG/Source/PCG/Private/Elements/PCGGather.cpp
--- a/Engine/Plugins/Experimental/PCG/Source/PCG/Private/Elements/PCGGather.cpp
+++ b/Engine/Plugins/Experimental/PCG/Source/PCG/Private/Elements/PCGGather.cpp
@@ -1,21 +1,34 @@
 // Copyright Epic Games, Inc. All Rights Reserved.
 
 #include "Elements/PCGGather.h"
+
 #include "PCGContext.h"
 #include "PCGPin.h"
 	
 TArray<FPCGPinProperties> UPCGGatherSettings::InputPinProperties() const
 {
+	EPCGDataType InputTypeUnion = GetTypeUnionOfIncidentEdges(PCGPinConstants::DefaultInputLabel);
+	if (InputTypeUnion == EPCGDataType::None)
+	{
+		InputTypeUnion = EPCGDataType::Any;
+	}
+
 	TArray<FPCGPinProperties> PinProperties;
-	PinProperties.Emplace(PCGPinConstants::DefaultInputLabel, EPCGDataType::Any);
+	PinProperties.Emplace(PCGPinConstants::DefaultInputLabel, InputTypeUnion);
 
 	return PinProperties;
 }
 
 TArray<FPCGPinProperties> UPCGGatherSettings::OutputPinProperties() const
 {
+	EPCGDataType InputTypeUnion = GetTypeUnionOfIncidentEdges(PCGPinConstants::DefaultInputLabel);
+	if (InputTypeUnion == EPCGDataType::None)
+	{
+		InputTypeUnion = EPCGDataType::Any;
+	}
+
 	TArray<FPCGPinProperties> PinProperties;
-	PinProperties.Emplace(PCGPinConstants::DefaultOutputLabel, EPCGDataType::Any);
+	PinProperties.Emplace(PCGPinConstants::DefaultOutputLabel, InputTypeUnion);
 
 	return PinProperties;
 }
@@ -32,5 +45,4 @@ bool FPCGGatherElement::ExecuteInternal(FPCGContext* Context) const
 	Context->OutputData = Context->InputData;
 
 	return true;
-
 }
diff --git a/Engine/Plugins/Experimental/PCG/Source/PCG/Public/Elements/PCGGather.h b/Engine/Plugins/Experimental/PCG/Source/PCG/Public/Elements/PCGGather.h
--- a/Engine/Plugins/Experimental/PCG/Source/PCG/Public/Elements/PCGGather.h
+++ b/Engine/Plugins/Experimental/PCG/Source/PCG/Public/Elements/PCGGather.h
@@ -18,12 +18,12 @@ public:
 #if WITH_EDITOR
 	virtual FName GetDefaultNodeName() const override { return FName(TEXT("Gather")); }
 	virtual EPCGSettingsType GetType() const override { return EPCGSettingsType::Spatial; }
+	virtual bool HasDynamicPins() const override { return true; }
 #endif
 
+protected:
 	virtual TArray<FPCGPinProperties> InputPinProperties() const override;
 	virtual TArray<FPCGPinProperties> OutputPinProperties() const override;
-
-protected:
 	virtual FPCGElementPtr CreateElement() const override;
 	//~End UPCGSettings interface
 };
