Used proper garbage collection flags while in editor

#rnx

diff --git a/Engine/Plugins/Experimental/PCG/Source/PCG/Private/Graph/PCGGraphExecutor.cpp b/Engine/Plugins/Experimental/PCG/Source/PCG/Private/Graph/PCGGraphExecutor.cpp
--- a/Engine/Plugins/Experimental/PCG/Source/PCG/Private/Graph/PCGGraphExecutor.cpp
+++ b/Engine/Plugins/Experimental/PCG/Source/PCG/Private/Graph/PCGGraphExecutor.cpp
@@ -1237,7 +1237,7 @@ void FPCGGraphExecutor::ReleaseUnusedActors()
 		if (ReleaseActorsCountUntilGC <= 0)
 		{
 			ReleaseActorsCountUntilGC = 30;
-			CollectGarbage(RF_NoFlags, true);
+			CollectGarbage(GARBAGE_COLLECTION_KEEPFLAGS, true);
 		}
 	}
 #endif
diff --git a/Engine/Plugins/Experimental/PCG/Source/PCG/Private/PCGSubsystem.cpp b/Engine/Plugins/Experimental/PCG/Source/PCG/Private/PCGSubsystem.cpp
--- a/Engine/Plugins/Experimental/PCG/Source/PCG/Private/PCGSubsystem.cpp
+++ b/Engine/Plugins/Experimental/PCG/Source/PCG/Private/PCGSubsystem.cpp
@@ -1228,7 +1228,7 @@ void UPCGSubsystem::FlushCache()
 	// Garbage collection is very seldom run in the editor, but we currently can consume a lot of memory in the cache.
 	if (!PCGHelpers::IsRuntimeOrPIE())
 	{
-		CollectGarbage(RF_NoFlags, true);
-	}	
+		CollectGarbage(GARBAGE_COLLECTION_KEEPFLAGS, true);
+	}
 #endif
 }
