[PCG] Set defaults for Get Landscape Data so it Just Works when you drop it into a graph (replicates old behaviour of Landscape input node pin)

Massaged actor selection/filtering to find all overlapping world actors.

Must Overlap Self - does the overlap test. I could have added a new Actor Filter dropdown item but this seemed cleaner. The option is hidden if not using "All World Actors".
Disable Filter - This option used to be hidden when using "All World Actors" but i think we want it here, so I un-hid it

Also refined some spurious warnings that came up in testing so that the verbosity is correct for different situations.

diff --git a/Engine/Plugins/Experimental/PCG/Source/PCG/Private/Data/PCGWorldData.cpp b/Engine/Plugins/Experimental/PCG/Source/PCG/Private/Data/PCGWorldData.cpp
--- a/Engine/Plugins/Experimental/PCG/Source/PCG/Private/Data/PCGWorldData.cpp
+++ b/Engine/Plugins/Experimental/PCG/Source/PCG/Private/Data/PCGWorldData.cpp
@@ -185,7 +185,7 @@ const UPCGPointData* UPCGWorldVolumetricData::CreatePointData(FPCGContext* Conte
 	{
 		if (!Bounds.IsValid && !InBounds.IsValid)
 		{
-			UE_LOG(LogPCG, Error, TEXT("PCG World Volumetric Data cannot generate if there are no framing bounds"));
+			UE_LOG(LogPCG, Error, TEXT("PCG World Volumetric Data cannot generate without sampling bounds. Consider using a Volume Sampler with the Unbounded option disabled."));
 		}
 		
 		return Data;
@@ -361,7 +361,7 @@ const UPCGPointData* UPCGWorldRayHitData::CreatePointData(FPCGContext* Context,
 	{
 		if (!Bounds.IsValid && !InBounds.IsValid)
 		{
-			UE_LOG(LogPCG, Error, TEXT("PCG World Ray Hit Data cannot generate if there are no framing bounds"));
+			UE_LOG(LogPCG, Error, TEXT("PCG World Ray Hit Data cannot generate without sampling bounds. Consider using a Surface Sampler with the Unbounded option disabled."));
 		}
 
 		return Data;
diff --git a/Engine/Plugins/Experimental/PCG/Source/PCG/Private/Elements/PCGActorSelector.cpp b/Engine/Plugins/Experimental/PCG/Source/PCG/Private/Elements/PCGActorSelector.cpp
--- a/Engine/Plugins/Experimental/PCG/Source/PCG/Private/Elements/PCGActorSelector.cpp
+++ b/Engine/Plugins/Experimental/PCG/Source/PCG/Private/Elements/PCGActorSelector.cpp
@@ -12,16 +12,23 @@
 
 namespace PCGActorSelector
 {
+	// Filter is required if it is not disabled and if we are gathering all world actors or gathering all children.
+	bool FilterRequired(const FPCGActorSelectorSettings& InSettings)
+	{
+		return (InSettings.ActorFilter == EPCGActorFilter::AllWorldActors || InSettings.bIncludeChildren) && !InSettings.bDisableFilter;
+	}
+
 	// Need to pass a pointer of pointer to the found actor. The lambda will capture this pointer and modify its value when an actor is found.
-	TFunction<bool(AActor*)> GetFilteringFunction(const FPCGActorSelectorSettings& InSettings, TArray<AActor*>& InFoundActors)
+	TFunction<bool(AActor*)> GetFilteringFunction(const FPCGActorSelectorSettings& InSettings, const TFunction<bool(const AActor*)>& BoundsCheck, TArray<AActor*>& InFoundActors)
 	{
-		// If we're not filtering all world actors, and if the Include Children is enabled, then we allow
-		// user to disable the filter which ignores all filtering options, which is convenient for workflow simplification.
-		if (InSettings.ActorFilter != EPCGActorFilter::AllWorldActors && (!InSettings.bIncludeChildren || InSettings.bDisableFilter))
+		if (!FilterRequired(InSettings))
 		{
-			return [&InFoundActors](AActor* Actor) -> bool
+			return [&InFoundActors, &BoundsCheck](AActor* Actor) -> bool
 			{
-				InFoundActors.Add(Actor);
+				if (BoundsCheck(Actor))
+				{
+					InFoundActors.Add(Actor);
+				}
 				return true;
 			};
 		}
@@ -31,9 +38,9 @@ namespace PCGActorSelector
 		switch (InSettings.ActorSelection)
 		{
 		case EPCGActorSelection::ByTag:
-			return[ActorSelectionTag = InSettings.ActorSelectionTag, &InFoundActors, bMultiSelect](AActor* Actor) -> bool
+			return[ActorSelectionTag = InSettings.ActorSelectionTag, &InFoundActors, bMultiSelect, &BoundsCheck](AActor* Actor) -> bool
 			{
-				if (Actor->ActorHasTag(ActorSelectionTag))
+				if (Actor->ActorHasTag(ActorSelectionTag) && BoundsCheck(Actor))
 				{
 					InFoundActors.Add(Actor);
 					return bMultiSelect;
@@ -43,9 +50,9 @@ namespace PCGActorSelector
 			};
 
 		case EPCGActorSelection::ByName:
-			return[ActorSelectionName = InSettings.ActorSelectionName, &InFoundActors, bMultiSelect](AActor* Actor) -> bool
+			return[ActorSelectionName = InSettings.ActorSelectionName, &InFoundActors, bMultiSelect, &BoundsCheck](AActor* Actor) -> bool
 			{
-				if (Actor->GetFName().IsEqual(ActorSelectionName, ENameCase::IgnoreCase, /*bCompareNumber=*/ false))
+				if (Actor->GetFName().IsEqual(ActorSelectionName, ENameCase::IgnoreCase, /*bCompareNumber=*/ false) && BoundsCheck(Actor))
 				{
 					InFoundActors.Add(Actor);
 					return bMultiSelect;
@@ -55,9 +62,9 @@ namespace PCGActorSelector
 			};
 
 		case EPCGActorSelection::ByClass:
-			return[ActorSelectionClass = InSettings.ActorSelectionClass, &InFoundActors, bMultiSelect](AActor* Actor) -> bool
+			return[ActorSelectionClass = InSettings.ActorSelectionClass, &InFoundActors, bMultiSelect, &BoundsCheck](AActor* Actor) -> bool
 			{
-				if (Actor->IsA(ActorSelectionClass))
+				if (Actor->IsA(ActorSelectionClass) && BoundsCheck(Actor))
 				{
 					InFoundActors.Add(Actor);
 					return bMultiSelect;
@@ -73,7 +80,7 @@ namespace PCGActorSelector
 		return [](AActor* Actor) -> bool { return false; };
 	}
 
-	TArray<AActor*> FindActors(const FPCGActorSelectorSettings& Settings, const UPCGComponent* InComponent)
+	TArray<AActor*> FindActors(const FPCGActorSelectorSettings& Settings, const UPCGComponent* InComponent, const TFunction<bool(const AActor*)>& BoundsCheck)
 	{
 		TRACE_CPUPROFILER_EVENT_SCOPE(PCGActorSelector::FindActor);
 
@@ -87,10 +94,8 @@ namespace PCGActorSelector
 			return FoundActors;
 		}
 
-		// Early out if we have not the information necessary. A filter must be set up if we are taking all world actors,
-		// or including children and not disabling the filter.
-		const bool bFilterRequired = (Settings.ActorFilter == EPCGActorFilter::AllWorldActors) || (Settings.bIncludeChildren && !Settings.bDisableFilter);
-		if (bFilterRequired &&
+		// Early out if we have not the information necessary.
+		if (FilterRequired(Settings) &&
 			((Settings.ActorSelection == EPCGActorSelection::ByTag && Settings.ActorSelectionTag == NAME_None) ||
 			(Settings.ActorSelection == EPCGActorSelection::ByName && Settings.ActorSelectionName == NAME_None) ||
 			(Settings.ActorSelection == EPCGActorSelection::ByClass && !Settings.ActorSelectionClass)))
@@ -100,11 +105,12 @@ namespace PCGActorSelector
 
 		// We pass FoundActor ref, that will be captured by the FilteringFunction
 		// It will modify the FoundActor pointer to the found actor, if found.
-		TFunction<bool(AActor*)> FilteringFunction = PCGActorSelector::GetFilteringFunction(Settings, FoundActors);
+		TFunction<bool(AActor*)> FilteringFunction = PCGActorSelector::GetFilteringFunction(Settings, BoundsCheck, FoundActors);
 
 		// In case of iterating over all actors in the world, call our filtering function and get out.
 		if (Settings.ActorFilter == EPCGActorFilter::AllWorldActors)
 		{
+			// A potential optimization if we know the sought actors are collide-able could be to obtain overlaps via a collision query.
 			UPCGActorHelpers::ForEachActorInWorld<AActor>(World, FilteringFunction);
 
 			// FoundActor is set by the FilteringFunction (captured)
@@ -194,13 +200,13 @@ namespace PCGActorSelector
 		return FoundActors;
 	}
 
-	AActor* FindActor(const FPCGActorSelectorSettings& InSettings, UPCGComponent* InComponent)
+	AActor* FindActor(const FPCGActorSelectorSettings& InSettings, UPCGComponent* InComponent, const TFunction<bool(const AActor*)>& BoundsCheck)
 	{
 		// In order to make sure we don't try to select multiple, we'll do a copy of the settings here.
 		FPCGActorSelectorSettings Settings = InSettings;
 		Settings.bSelectMultiple = false;
 
-		TArray<AActor*> Actors = FindActors(Settings, InComponent);
+		TArray<AActor*> Actors = FindActors(Settings, InComponent, BoundsCheck);
 		return Actors.IsEmpty() ? nullptr : Actors[0];
 	}
 }
diff --git a/Engine/Plugins/Experimental/PCG/Source/PCG/Private/Elements/PCGDataFromActor.cpp b/Engine/Plugins/Experimental/PCG/Source/PCG/Private/Elements/PCGDataFromActor.cpp
--- a/Engine/Plugins/Experimental/PCG/Source/PCG/Private/Elements/PCGDataFromActor.cpp
+++ b/Engine/Plugins/Experimental/PCG/Source/PCG/Private/Elements/PCGDataFromActor.cpp
@@ -3,6 +3,7 @@
 #include "Elements/PCGDataFromActor.h"
 
 #include "PCGComponent.h"
+#include "PCGHelpers.h"
 #include "PCGSubsystem.h"
 #include "Data/PCGSpatialData.h"
 
@@ -69,11 +70,23 @@ bool FPCGDataFromActorElement::ExecuteInternal(FPCGContext* InContext) const
 	const UPCGDataFromActorSettings* Settings = Context->GetInputSettings<UPCGDataFromActorSettings>();
 	check(Settings);
 
-	UWorld* World = Context->SourceComponent.IsValid() ? Context->SourceComponent->GetWorld() : nullptr;
-
 	if (!Context->bPerformedQuery)
 	{
-		Context->FoundActors = PCGActorSelector::FindActors(Settings->ActorSelector, Context->SourceComponent.Get());
+		TFunction<bool(const AActor*)> BoundsCheck = [](const AActor*) -> bool { return true; };
+		const UPCGComponent* PCGComponent = Context->SourceComponent.IsValid() ? Context->SourceComponent.Get() : nullptr;
+		const AActor* Self = PCGComponent ? PCGComponent->GetOwner() : nullptr;
+		if (Self && Settings->ActorSelector.bMustOverlapSelf)
+		{
+			// Capture ActorBounds by value because it goes out of scope
+			const FBox ActorBounds = PCGHelpers::GetActorBounds(Self);
+			BoundsCheck = [Settings, ActorBounds, PCGComponent](const AActor* OtherActor) -> bool
+			{
+				const FBox OtherActorBounds = OtherActor ? PCGHelpers::GetGridBounds(OtherActor, PCGComponent) : FBox(EForceInit::ForceInit);
+				return ActorBounds.Intersect(OtherActorBounds);
+			};
+		}
+
+		Context->FoundActors = PCGActorSelector::FindActors(Settings->ActorSelector, Context->SourceComponent.Get(), BoundsCheck);
 		Context->bPerformedQuery = true;
 
 		if (Context->FoundActors.IsEmpty())
diff --git a/Engine/Plugins/Experimental/PCG/Source/PCG/Private/Elements/PCGPropertyToParamData.cpp b/Engine/Plugins/Experimental/PCG/Source/PCG/Private/Elements/PCGPropertyToParamData.cpp
--- a/Engine/Plugins/Experimental/PCG/Source/PCG/Private/Elements/PCGPropertyToParamData.cpp
+++ b/Engine/Plugins/Experimental/PCG/Source/PCG/Private/Elements/PCGPropertyToParamData.cpp
@@ -95,7 +95,8 @@ bool FPCGPropertyToParamDataElement::ExecuteInternal(FPCGContext* Context) const
 
 	// First find the actor depending on the selection
 	UPCGComponent* OriginalComponent = UPCGBlueprintHelpers::GetOriginalComponent(*Context);
-	AActor* FoundActor = PCGActorSelector::FindActor(Settings->ActorSelector, OriginalComponent);
+	auto NoBoundsCheck = [](const AActor*) -> bool { return true; };
+	AActor* FoundActor = PCGActorSelector::FindActor(Settings->ActorSelector, OriginalComponent, NoBoundsCheck);
 
 	if (!FoundActor)
 	{
diff --git a/Engine/Plugins/Experimental/PCG/Source/PCG/Private/Elements/PCGSurfaceSampler.cpp b/Engine/Plugins/Experimental/PCG/Source/PCG/Private/Elements/PCGSurfaceSampler.cpp
--- a/Engine/Plugins/Experimental/PCG/Source/PCG/Private/Elements/PCGSurfaceSampler.cpp
+++ b/Engine/Plugins/Experimental/PCG/Source/PCG/Private/Elements/PCGSurfaceSampler.cpp
@@ -458,9 +458,14 @@ void FPCGSurfaceSamplerElement::AllocateOutputs(FPCGSurfaceSamplerContext* InCon
 		TentativeLoopData.bEnableTimeSlicing = true;
 		if (!InputBounds.IsValid || !TentativeLoopData.Initialize(Settings, InContext, InputBounds))
 		{
-			if (!InputBounds.IsValid)
+			if (!GeneratingShape->IsBounded())
 			{
-				PCGE_LOG_C(Warning, GraphAndLog, InContext, LOCTEXT("InvalidInputBounds", "Input data has invalid bounds"));
+				// Some inputs are unable to provide bounds, like the WorldRayHit, in which case the user must provide bounds.
+				PCGE_LOG_C(Warning, GraphAndLog, InContext, LOCTEXT("CouldNotObtainInputBounds", "Input data is not bounded, so bounds must be provided for sampling. Consider providing a Bounding Shape input."));
+			}
+			else if(!InputBounds.IsValid)
+			{
+				PCGE_LOG_C(Verbose, LogOnly, InContext, LOCTEXT("InvalidSamplingBounds", "Final sampling bounds is invalid/zero-sized."));
 			}
 
 			Outputs.RemoveAt(GeneratingShapeIndex);
diff --git a/Engine/Plugins/Experimental/PCG/Source/PCG/Private/Elements/PCGTypedGetter.cpp b/Engine/Plugins/Experimental/PCG/Source/PCG/Private/Elements/PCGTypedGetter.cpp
--- a/Engine/Plugins/Experimental/PCG/Source/PCG/Private/Elements/PCGTypedGetter.cpp
+++ b/Engine/Plugins/Experimental/PCG/Source/PCG/Private/Elements/PCGTypedGetter.cpp
@@ -2,6 +2,8 @@
 
 #include "Elements/PCGTypedGetter.h"
 
+#include "Landscape.h"
+
 #include UE_INLINE_GENERATED_CPP_BY_NAME(PCGTypedGetter)
 
 #define LOCTEXT_NAMESPACE "PCGTypedGetterElements"
@@ -10,6 +12,18 @@ UPCGGetLandscapeSettings::UPCGGetLandscapeSettings()
 {
 	bDisplayModeSettings = false;
 	Mode = EPCGGetDataFromActorMode::ParseActorComponents;
+
+	// We want to apply different defaults to newly placed nodes. We detect new object if they are not a default object/archetype
+	// and/or they do not need load. Followed similar pattern to UPCGComponent::PostInitProperties().
+	if (!HasAnyFlags(RF_ClassDefaultObject | RF_NeedLoad | RF_NeedPostLoad))
+	{
+		// This setup replicates what was implemented on the Landscape input node pin
+		ActorSelector.ActorFilter = EPCGActorFilter::AllWorldActors;
+		ActorSelector.bMustOverlapSelf = true;
+		ActorSelector.bSelectMultiple = true;
+		ActorSelector.ActorSelection = EPCGActorSelection::ByClass;
+		ActorSelector.ActorSelectionClass = ALandscapeProxy::StaticClass();
+	}
 }
 
 TArray<FPCGPinProperties> UPCGGetLandscapeSettings::OutputPinProperties() const
diff --git a/Engine/Plugins/Experimental/PCG/Source/PCG/Private/Elements/PCGVolumeSampler.cpp b/Engine/Plugins/Experimental/PCG/Source/PCG/Private/Elements/PCGVolumeSampler.cpp
--- a/Engine/Plugins/Experimental/PCG/Source/PCG/Private/Elements/PCGVolumeSampler.cpp
+++ b/Engine/Plugins/Experimental/PCG/Source/PCG/Private/Elements/PCGVolumeSampler.cpp
@@ -262,7 +262,15 @@ bool FPCGVolumeSamplerElement::ExecuteInternal(FPCGContext* Context) const
 
 		if (!InputBounds.IsValid)
 		{
-			PCGE_LOG(Verbose, LogOnly, LOCTEXT("InvalidBounds", "Input data has invalid bounds"));
+			if (!GeneratingShape->IsBounded())
+			{
+				// Some inputs are unable to provide bounds, like the WorldVolumetricQuery, in which case the user must provide bounds.
+				PCGE_LOG(Warning, GraphAndLog, LOCTEXT("CouldNotObtainInputBounds", "Input data is not bounded, so bounds must be provided for sampling. Consider providing a Bounding Shape input."));
+			}
+			else
+			{
+				PCGE_LOG(Verbose, LogOnly, LOCTEXT("InvalidSamplingBounds", "Final sampling bounds is invalid/zero-sized."));
+			}
 
 			Outputs.RemoveAt(GenerationIndex);
 			GeneratingShapes.RemoveAt(GenerationIndex);
diff --git a/Engine/Plugins/Experimental/PCG/Source/PCG/Public/Elements/PCGActorSelector.h b/Engine/Plugins/Experimental/PCG/Source/PCG/Public/Elements/PCGActorSelector.h
--- a/Engine/Plugins/Experimental/PCG/Source/PCG/Public/Elements/PCGActorSelector.h
+++ b/Engine/Plugins/Experimental/PCG/Source/PCG/Public/Elements/PCGActorSelector.h
@@ -42,6 +42,9 @@ struct FPCGActorSelectorSettings
 	UPROPERTY(BlueprintReadWrite, EditAnywhere, Category = Settings)
 	EPCGActorFilter ActorFilter = EPCGActorFilter::Self;
 
+	UPROPERTY(BlueprintReadWrite, EditAnywhere, Category = Settings, meta = (EditCondition = "ActorFilter==EPCGActorFilter::AllWorldActors", EditConditionHides))
+	bool bMustOverlapSelf = false;
+
 	/** Whether to consider child actors. */
 	UPROPERTY(BlueprintReadWrite, EditAnywhere, Category = Settings, meta = (EditCondition = "ActorFilter!=EPCGActorFilter::AllWorldActors", EditConditionHides))
 	bool bIncludeChildren = false;
@@ -59,7 +62,7 @@ struct FPCGActorSelectorSettings
 	UPROPERTY(BlueprintReadWrite, EditAnywhere, Category = Settings, meta = (EditCondition = "(ActorFilter==EPCGActorFilter::AllWorldActors || (bIncludeChildren && !bDisableFilter)) && ActorSelection==EPCGActorSelection::ByName", EditConditionHides))
 	FName ActorSelectionName;
 
-	UPROPERTY(BlueprintReadWrite, EditAnywhere, Category = Settings, meta = (EditCondition = "(ActorFilter==EPCGActorFilter::AllWorldActors || (bIncludeChildren && !bDisableFilter)) && ActorSelection==EPCGActorSelection::ByClass", EditConditionHides))
+	UPROPERTY(BlueprintReadWrite, EditAnywhere, Category = Settings, meta = (EditCondition = "(ActorFilter==EPCGActorFilter::AllWorldActors || (bIncludeChildren && !bDisableFilter)) && ActorSelection==EPCGActorSelection::ByClass", EditConditionHides, AllowAbstract = "true"))
 	TSubclassOf<AActor> ActorSelectionClass;
 
 	/** If true processes all matching actors, otherwise returns data from first match. */
@@ -69,6 +72,6 @@ struct FPCGActorSelectorSettings
 
 namespace PCGActorSelector
 {
-	TArray<AActor*> FindActors(const FPCGActorSelectorSettings& Settings, const UPCGComponent* InComponent);
-	AActor* FindActor(const FPCGActorSelectorSettings& InSettings, UPCGComponent* InComponent);
+	TArray<AActor*> FindActors(const FPCGActorSelectorSettings& Settings, const UPCGComponent* InComponent, const TFunction<bool(const AActor*)>& BoundsCheck);
+	AActor* FindActor(const FPCGActorSelectorSettings& InSettings, UPCGComponent* InComponent, const TFunction<bool(const AActor*)>& BoundsCheck);
 }
diff --git a/Engine/Plugins/Experimental/PCG/Source/PCG/Public/Elements/PCGTypedGetter.h b/Engine/Plugins/Experimental/PCG/Source/PCG/Public/Elements/PCGTypedGetter.h
--- a/Engine/Plugins/Experimental/PCG/Source/PCG/Public/Elements/PCGTypedGetter.h
+++ b/Engine/Plugins/Experimental/PCG/Source/PCG/Public/Elements/PCGTypedGetter.h
@@ -27,7 +27,7 @@ protected:
 
 public:
 	//~Begin UPCGDataFromActorSettings interface
-	virtual bool DataFilter(EPCGDataType InDataType) const { return !!(InDataType & EPCGDataType::Landscape); }
+	virtual bool DataFilter(EPCGDataType InDataType) const override { return !!(InDataType & EPCGDataType::Landscape); }
 	//~End UPCGDataFromActorSettings
 };
 
@@ -52,7 +52,7 @@ protected:
 
 public:
 	//~Begin UPCGDataFromActorSettings interface
-	virtual bool DataFilter(EPCGDataType InDataType) const { return !!(InDataType & EPCGDataType::PolyLine); }
+	virtual bool DataFilter(EPCGDataType InDataType) const override { return !!(InDataType & EPCGDataType::PolyLine); }
 	//~End UPCGDataFromActorSettings
 };
 
@@ -77,7 +77,7 @@ protected:
 
 public:
 	//~Begin UPCGDataFromActorSettings interface
-	virtual bool DataFilter(EPCGDataType InDataType) const { return !!(InDataType & EPCGDataType::Volume); }
+	virtual bool DataFilter(EPCGDataType InDataType) const override { return !!(InDataType & EPCGDataType::Volume); }
 	//~End UPCGDataFromActorSettings
 };
 
@@ -102,6 +102,6 @@ protected:
 
 public:
 	//~Begin UPCGDataFromActorSettings interface
-	virtual bool DataFilter(EPCGDataType InDataType) const { return !!(InDataType & EPCGDataType::Primitive); }
+	virtual bool DataFilter(EPCGDataType InDataType) const override { return !!(InDataType & EPCGDataType::Primitive); }
 	//~End UPCGDataFromActorSettings
 };
