[PCG] Change default key for Inspect to A (for Attributes window) to keep shortcuts clustered on left hand by default.

diff --git a/Engine/Plugins/Experimental/PCG/Source/PCGEditor/Private/PCGEditorCommands.cpp b/Engine/Plugins/Experimental/PCG/Source/PCGEditor/Private/PCGEditorCommands.cpp
--- a/Engine/Plugins/Experimental/PCG/Source/PCGEditor/Private/PCGEditorCommands.cpp
+++ b/Engine/Plugins/Experimental/PCG/Source/PCGEditor/Private/PCGEditorCommands.cpp
@@ -31,7 +31,7 @@ void FPCGEditorCommands::RegisterCommands()
 	UI_COMMAND(ToggleDebug, "Toggle Debug", "Toggle node debug state for selected nodes", EUserInterfaceActionType::ToggleButton, FInputChord(EKeys::D));
 	UI_COMMAND(DebugOnlySelected, "Debug Only Selected", "Enable node debug state for selected nodes and disable debug state for the others", EUserInterfaceActionType::Button, FInputChord(EModifierKey::Control | EModifierKey::Alt, EKeys::D));
 	UI_COMMAND(DisableDebugOnAllNodes, "Disable Debug on all nodes", "Disable debug state for all nodes", EUserInterfaceActionType::Button, FInputChord(EModifierKey::Alt, EKeys::D));
-	UI_COMMAND(ToggleInspect, "Toggle Inspection", "Toggle node inspection for selected node", EUserInterfaceActionType::ToggleButton, FInputChord(EKeys::I));
+	UI_COMMAND(ToggleInspect, "Toggle Inspection", "Toggle node inspection for selected node", EUserInterfaceActionType::ToggleButton, FInputChord(EKeys::A));
 }
 
 #undef LOCTEXT_NAMESPACE
