Fixed issue with settings override being deleted by calling into a PCG blueprint

diff --git a/Engine/Plugins/Experimental/PCG/Source/PCG/Private/Elements/PCGExecuteBlueprint.cpp b/Engine/Plugins/Experimental/PCG/Source/PCG/Private/Elements/PCGExecuteBlueprint.cpp
--- a/Engine/Plugins/Experimental/PCG/Source/PCG/Private/Elements/PCGExecuteBlueprint.cpp
+++ b/Engine/Plugins/Experimental/PCG/Source/PCG/Private/Elements/PCGExecuteBlueprint.cpp
@@ -573,9 +573,19 @@ bool FPCGExecuteBlueprintElement::ExecuteInternal(FPCGContext* InContext) const
 			}
 		}
 
+		// Since the ExecuteWithContext method is a blueprint call, it will perform a copy of the context
+		// in a local variable. To prevent the overridden settings from being unrooted when the call finishes,
+		// we'll mark it temporarily off here. Note that for the same reason, the context is actually sliced
+		// so there should never be any members in the BP element context that are visible/accessible from blueprint
+		const bool bShouldUnrootSettingsOnDelete = Context->bShouldUnrootSettingsOnDelete;
+		Context->bShouldUnrootSettingsOnDelete = false;
+
 		/** Finally, execute the actual blueprint */
 		Context->BlueprintElementInstance->ExecuteWithContext(*Context, Context->InputData, Context->OutputData);
 
+		// Put back the proper unroot flag
+		Context->bShouldUnrootSettingsOnDelete = bShouldUnrootSettingsOnDelete;
+
 		// Log info on outputs
 		for (int32 OutputIndex = 0; OutputIndex < Context->OutputData.TaggedData.Num(); ++OutputIndex)
 		{
diff --git a/Engine/Plugins/Experimental/PCG/Source/PCG/Private/PCGContext.cpp b/Engine/Plugins/Experimental/PCG/Source/PCG/Private/PCGContext.cpp
--- a/Engine/Plugins/Experimental/PCG/Source/PCG/Private/PCGContext.cpp
+++ b/Engine/Plugins/Experimental/PCG/Source/PCG/Private/PCGContext.cpp
@@ -24,7 +24,7 @@
 
 FPCGContext::~FPCGContext()
 {
-	if (SettingsWithOverride)
+	if (SettingsWithOverride && bShouldUnrootSettingsOnDelete)
 	{
 		SettingsWithOverride->RemoveFromRoot();
 		SettingsWithOverride->MarkAsGarbage();
diff --git a/Engine/Plugins/Experimental/PCG/Source/PCG/Public/PCGContext.h b/Engine/Plugins/Experimental/PCG/Source/PCG/Public/PCGContext.h
--- a/Engine/Plugins/Experimental/PCG/Source/PCG/Public/PCGContext.h
+++ b/Engine/Plugins/Experimental/PCG/Source/PCG/Public/PCGContext.h
@@ -60,6 +60,9 @@ struct PCG_API FPCGContext
 	FPCGTaskId CompiledTaskId = InvalidPCGTaskId;
 	bool bIsPaused = false;
 
+	// Used to preven settings override being deleted, needs to be false when going through blueprint calls with a context
+	bool bShouldUnrootSettingsOnDelete = true;
+
 	EPCGExecutionPhase CurrentPhase = EPCGExecutionPhase::NotExecuted;
 	int32 BypassedOutputCount = 0;
 
