Follow-up on the GC flag issues we had for the cache cleanup mechanism

#rnx

diff --git a/Engine/Plugins/Experimental/PCG/Source/PCG/Private/Graph/PCGGraphExecutor.cpp b/Engine/Plugins/Experimental/PCG/Source/PCG/Private/Graph/PCGGraphExecutor.cpp
--- a/Engine/Plugins/Experimental/PCG/Source/PCG/Private/Graph/PCGGraphExecutor.cpp
+++ b/Engine/Plugins/Experimental/PCG/Source/PCG/Private/Graph/PCGGraphExecutor.cpp
@@ -921,7 +921,7 @@ void FPCGGraphExecutor::Execute()
 			if (TidyCacheCountUntilGC <= 0)
 			{
 				TidyCacheCountUntilGC = 100;
-				CollectGarbage(RF_NoFlags, true);
+				CollectGarbage(GARBAGE_COLLECTION_KEEPFLAGS, true);
 			}
 		}
 #endif
