[PCG] Element log house-keeping

Moved some final legacy UE_LOG's over to the element logs in order to get graph warnings/errors.

Moved to 'Parameter Set' terminology (rather than Param).

Additional minor QoL changes.

diff --git a/Engine/Plugins/Experimental/PCG/Source/PCG/Private/Elements/Metadata/PCGMetadataOpElementBase.cpp b/Engine/Plugins/Experimental/PCG/Source/PCG/Private/Elements/Metadata/PCGMetadataOpElementBase.cpp
--- a/Engine/Plugins/Experimental/PCG/Source/PCG/Private/Elements/Metadata/PCGMetadataOpElementBase.cpp
+++ b/Engine/Plugins/Experimental/PCG/Source/PCG/Private/Elements/Metadata/PCGMetadataOpElementBase.cpp
@@ -113,7 +113,7 @@ bool FPCGMetadataElementBase::ExecuteInternal(FPCGContext* Context) const
 		}
 		else
 		{
-			PCGE_LOG(Warning, LogOnly, FText::Format(LOCTEXT("InvalidInputDataTypeForPin", "Invalid data provided on pin {0}, must be of type Spatial or Param"), i));
+			PCGE_LOG(Warning, LogOnly, FText::Format(LOCTEXT("InvalidInputDataTypeForPin", "Invalid data provided on pin {0}, must be of type Spatial or Attribute Set"), i));
 			return true;
 		}
 	}
@@ -140,7 +140,7 @@ bool FPCGMetadataElementBase::ExecuteInternal(FPCGContext* Context) const
 			InputTaggedData[0].Data->GetClass() != InputTaggedData[i].Data->GetClass() &&
 			!InputTaggedData[i].Data->IsA<UPCGParamData>())
 		{
-			PCGE_LOG(Error, GraphAndLog, FText::Format(LOCTEXT("InputTypeMismatch", "Input {0} is not of the same type than input 0 and is not a Param Data. This is not supported."), i));
+			PCGE_LOG(Error, GraphAndLog, FText::Format(LOCTEXT("InputTypeMismatch", "Input {0} is not of the same type than input 0 and is not an Attribute Set. This is not supported."), i));
 			return true;
 		}
 
diff --git a/Engine/Plugins/Experimental/PCG/Source/PCG/Private/Elements/Metadata/PCGMetadataRenameElement.cpp b/Engine/Plugins/Experimental/PCG/Source/PCG/Private/Elements/Metadata/PCGMetadataRenameElement.cpp
--- a/Engine/Plugins/Experimental/PCG/Source/PCG/Private/Elements/Metadata/PCGMetadataRenameElement.cpp
+++ b/Engine/Plugins/Experimental/PCG/Source/PCG/Private/Elements/Metadata/PCGMetadataRenameElement.cpp
@@ -10,6 +10,7 @@
 
 #include UE_INLINE_GENERATED_CPP_BY_NAME(PCGMetadataRenameElement)
 
+#define LOCTEXT_NAMESPACE "PCGMetadataRenameElement"
 
 TArray<FPCGPinProperties> UPCGMetadataRenameSettings::OutputPinProperties() const
 {
@@ -39,7 +40,7 @@ bool FPCGMetadataRenameElement::ExecuteInternal(FPCGContext* Context) const
 
 	if (NewAttributeName == NAME_None)
 	{
-		UE_LOG(LogPCG, Warning, TEXT("Metadata rename operation cannot rename from %s to %s"), *AttributeToRename.ToString(), *NewAttributeName.ToString());
+		PCGE_LOG(Warning, GraphAndLog, FText::Format(LOCTEXT("InvalidTargetAttributeName", "Invalid target attribute name '{0}'"), FText::FromName(NewAttributeName)));
 		// Bypass
 		Context->OutputData = Context->InputData;
 		return true;
@@ -64,7 +65,7 @@ bool FPCGMetadataRenameElement::ExecuteInternal(FPCGContext* Context) const
 
 		if (!Metadata)
 		{
-			UE_LOG(LogPCG, Warning, TEXT("Input is not supported, only supports spatial and params"));
+			PCGE_LOG(Warning, GraphAndLog, LOCTEXT("InvalidInputType", "Input data type is not supported, only supports Spatial and Attribute Set"));
 			Output.Data = nullptr;
 			continue;
 		}
@@ -81,9 +82,12 @@ bool FPCGMetadataRenameElement::ExecuteInternal(FPCGContext* Context) const
 
 		if (!NewMetadata || !NewMetadata->RenameAttribute(LocalAttributeToRename, NewAttributeName))
 		{
-			UE_LOG(LogPCG, Warning, TEXT("Failed to rename attribute from %s to %s"), *LocalAttributeToRename.ToString(), *NewAttributeName.ToString());
+			PCGE_LOG(Warning, GraphAndLog, FText::Format(LOCTEXT("AttributeRenamedFailed", "Failed to rename attribute from '{0}' to '{1}'"),
+				FText::FromName(LocalAttributeToRename), FText::FromName(NewAttributeName)));
 		}
 	}
 
 	return true;
 }
+
+#undef LOCTEXT_NAMESPACE
diff --git a/Engine/Plugins/Experimental/PCG/Source/PCG/Private/Elements/PCGAttributeFilterElement.cpp b/Engine/Plugins/Experimental/PCG/Source/PCG/Private/Elements/PCGAttributeFilterElement.cpp
--- a/Engine/Plugins/Experimental/PCG/Source/PCG/Private/Elements/PCGAttributeFilterElement.cpp
+++ b/Engine/Plugins/Experimental/PCG/Source/PCG/Private/Elements/PCGAttributeFilterElement.cpp
@@ -142,7 +142,7 @@ bool FPCGAttributeFilterElement::ExecuteInternal(FPCGContext* Context) const
 		}
 		else
 		{
-			PCGE_LOG(Error, GraphAndLog, LOCTEXT("InvalidInputData", "Invalid data as input. Only Spatial and Params data are supported."));
+			PCGE_LOG(Error, GraphAndLog, LOCTEXT("InvalidInputData", "Invalid data as input. Only Spatial and Attribute Set data are supported."));
 			continue;
 		}
 
diff --git a/Engine/Plugins/Experimental/PCG/Source/PCG/Private/Elements/PCGCreateAttribute.cpp b/Engine/Plugins/Experimental/PCG/Source/PCG/Private/Elements/PCGCreateAttribute.cpp
--- a/Engine/Plugins/Experimental/PCG/Source/PCG/Private/Elements/PCGCreateAttribute.cpp
+++ b/Engine/Plugins/Experimental/PCG/Source/PCG/Private/Elements/PCGCreateAttribute.cpp
@@ -164,7 +164,7 @@ bool FPCGCreateAttributeElement::ExecuteInternal(FPCGContext* Context) const
 	{
 		if (SourceParams.IsEmpty())
 		{
-			PCGE_LOG(Error, GraphAndLog, LOCTEXT("ParamNotProvided", "Source param was not provided"));
+			PCGE_LOG(Error, GraphAndLog, LOCTEXT("ParamNotProvided", "Source attribute was not provided"));
 			return true;
 		}
 
@@ -172,7 +172,7 @@ bool FPCGCreateAttributeElement::ExecuteInternal(FPCGContext* Context) const
 
 		if (!SourceParamData->Metadata)
 		{
-			PCGE_LOG(Error, GraphAndLog, LOCTEXT("ParamMissingMetadata", "Source param data does not have metadata"));
+			PCGE_LOG(Error, GraphAndLog, LOCTEXT("ParamMissingMetadata", "Source Attribute Set data does not have metadata"));
 			return true;
 		}
 
@@ -180,7 +180,7 @@ bool FPCGCreateAttributeElement::ExecuteInternal(FPCGContext* Context) const
 
 		if (!SourceParamData->Metadata->HasAttribute(SourceParamAttributeName))
 		{
-			PCGE_LOG(Error, GraphAndLog, FText::Format(LOCTEXT("ParamMissingAttribute", "Source param data does not have an attribute '{0}'"), FText::FromName(SourceParamAttributeName)));
+			PCGE_LOG(Error, GraphAndLog, FText::Format(LOCTEXT("ParamMissingAttribute", "Source Attribute Set data does not have an attribute '{0}'"), FText::FromName(SourceParamAttributeName)));
 			return true;
 		}
 	}
@@ -235,7 +235,7 @@ bool FPCGCreateAttributeElement::ExecuteInternal(FPCGContext* Context) const
 		}
 		else
 		{
-			PCGE_LOG(Error, GraphAndLog, LOCTEXT("InvalidInputData", "Invalid data as input. Only Spatial and Param data supported."));
+			PCGE_LOG(Error, GraphAndLog, LOCTEXT("InvalidInputData", "Invalid data as input. Only Spatial and Attribute Set data supported."));
 			continue;
 		}
 
diff --git a/Engine/Plugins/Experimental/PCG/Source/PCG/Private/Elements/PCGDebugElement.cpp b/Engine/Plugins/Experimental/PCG/Source/PCG/Private/Elements/PCGDebugElement.cpp
--- a/Engine/Plugins/Experimental/PCG/Source/PCG/Private/Elements/PCGDebugElement.cpp
+++ b/Engine/Plugins/Experimental/PCG/Source/PCG/Private/Elements/PCGDebugElement.cpp
@@ -18,6 +18,8 @@
 
 #include UE_INLINE_GENERATED_CPP_BY_NAME(PCGDebugElement)
 
+#define LOCTEXT_NAMESPACE "PCGDebugElement"
+
 namespace PCGDebugElement
 {
 	void ExecuteDebugDisplay(FPCGContext* Context)
@@ -43,7 +45,8 @@ namespace PCGDebugElement
 
 		if (!Mesh)
 		{
-			UE_LOG(LogPCG, Error, TEXT("Debug display was unable to load mesh %s"), *DebugSettings.PointMesh.ToString());
+			PCGE_LOG_C(Error, GraphAndLog, Context, FText::Format(LOCTEXT("UnableToLoadMesh", "Debug display was unable to load mesh '{0}'"),
+				FText::FromString(DebugSettings.PointMesh.ToString())));
 			return;
 		}
 
@@ -95,7 +98,7 @@ namespace PCGDebugElement
 			if (!TargetActor)
 			{
 				// No target actor
-				UE_LOG(LogPCG, Error, TEXT("Debug display cannot show data that have no target actor"));
+				PCGE_LOG_C(Error, GraphAndLog, Context, LOCTEXT("NoTargetActor", "Debug display cannot show data that have no target actor"));
 				continue;
 			}
 
@@ -230,10 +233,7 @@ FPCGElementPtr UPCGDebugSettings::CreateElement() const
 
 TArray<FPCGPinProperties> UPCGDebugSettings::OutputPinProperties() const
 {
-	TArray<FPCGPinProperties> PinProperties;
-	PinProperties.Emplace(PCGPinConstants::DefaultOutputLabel, EPCGDataType::Any);
-
-	return PinProperties;
+	return TArray<FPCGPinProperties>();
 }
 
 bool FPCGDebugElement::ExecuteInternal(FPCGContext* Context) const
@@ -241,6 +241,7 @@ bool FPCGDebugElement::ExecuteInternal(FPCGContext* Context) const
 	TRACE_CPUPROFILER_EVENT_SCOPE(FPCGDebugElement::Execute);
 	PCGDebugElement::ExecuteDebugDisplay(Context);
 	
-	Context->OutputData = Context->InputData;
 	return true;
 }
+
+#undef LOCTEXT_NAMESPACE
diff --git a/Engine/Plugins/Experimental/PCG/Source/PCG/Private/Elements/PCGExecuteBlueprint.cpp b/Engine/Plugins/Experimental/PCG/Source/PCG/Private/Elements/PCGExecuteBlueprint.cpp
--- a/Engine/Plugins/Experimental/PCG/Source/PCG/Private/Elements/PCGExecuteBlueprint.cpp
+++ b/Engine/Plugins/Experimental/PCG/Source/PCG/Private/Elements/PCGExecuteBlueprint.cpp
@@ -557,7 +557,8 @@ bool FPCGExecuteBlueprintElement::ExecuteInternal(FPCGContext* InContext) const
 		{
 			if (Blueprint->Status == BS_Error)
 			{
-				UE_LOG(LogPCG, Error, TEXT("PCG blueprint element cannot be executed since %s is not properly compiled"), *Blueprint->GetFName().ToString());
+				PCGE_LOG(Error, GraphAndLog, FText::Format(LOCTEXT("BPNotCompiled", "Blueprint cannot be executed since '{0}' is not properly compiled"),
+					FText::FromName(Blueprint->GetFName())));
 				return true;
 			}
 		}
@@ -739,7 +740,7 @@ void UPCGBlueprintElement::LoopNTimes(FPCGContext& InContext, int64 NumIteration
 {
 	if (NumIterations < 0)
 	{
-		UE_LOG(LogPCG, Error, TEXT("Invalid number of iterations in PCG blueprint element"));
+		PCGE_LOG_C(Error, GraphAndLog, &InContext, FText::Format(LOCTEXT("InvalidIterationCount", "Invalid number of iterations ({0})"), NumIterations));
 		return;
 	}
 
diff --git a/Engine/Plugins/Experimental/PCG/Source/PCG/Private/Elements/PCGSplineSampler.cpp b/Engine/Plugins/Experimental/PCG/Source/PCG/Private/Elements/PCGSplineSampler.cpp
--- a/Engine/Plugins/Experimental/PCG/Source/PCG/Private/Elements/PCGSplineSampler.cpp
+++ b/Engine/Plugins/Experimental/PCG/Source/PCG/Private/Elements/PCGSplineSampler.cpp
@@ -730,12 +730,12 @@ namespace PCGSplineSampler
 		}
 		else if (const UPCGLandscapeSplineData* LandscapeSplineData = Cast<UPCGLandscapeSplineData>(LineData))
 		{
-			UE_LOG(LogPCG, Error, TEXT("LandscapeSplines are not supported for interior sampling"));
+			PCGE_LOG_C(Error, GraphAndLog, Context, LOCTEXT("LandscapeSplinesNotSupported", "Input data of type Landscape Spline are not supported for interior sampling"));
 			return;
 		}
 		else
 		{
-			UE_LOG(LogPCG, Error, TEXT("Could not create UPCGSplineData from LineData"));
+			PCGE_LOG_C(Error, GraphAndLog, Context, LOCTEXT("CouldNotCreateSplineData", "Could not create UPCGSplineData from LineData"));
 			return;
 		}
 
@@ -743,7 +743,7 @@ namespace PCGSplineSampler
 
 		if (!Spline->IsClosedLoop())
 		{
-			UE_LOG(LogPCG, Error, TEXT("Interior sampling only generates for closed shapes"));
+			PCGE_LOG_C(Error, GraphAndLog, Context, LOCTEXT("ShapeNotClosed", "Interior sampling only generates for closed shapes, enable the 'Closed Loop' setting on the spline"));
 			return;
 		}
 
@@ -868,6 +868,11 @@ namespace PCGSplineSampler
 					MedialAxisEdges.Add(Edge);
 				}
 			}
+
+			if (MedialAxisEdges.IsEmpty())
+			{
+				PCGE_LOG_C(Warning, GraphAndLog, Context, LOCTEXT("MedialAxisFailed", "Failed to compute medial axis in interior region, density fall-off will not be applied. This functionality requires a closed spline with at least 4 spline points."));
+			}
 		}
 
 		TArray<FPCGPoint>& OutPoints = OutPointData->GetMutablePoints();
@@ -910,7 +915,7 @@ namespace PCGSplineSampler
 
 				if (Intersections.Num() % 2 != 0)
 				{
-					UE_LOG(LogPCG, Error, TEXT("Intersection test failed, skipping samples for this row"));
+					PCGE_LOG_C(Error, GraphAndLog, Context, LOCTEXT("IntersectionTestFailed", "Intersection test failed, skipping samples for this row"));
 					continue;
 				}
 
diff --git a/Engine/Plugins/Experimental/PCG/Source/PCG/Public/Elements/PCGSplineSampler.h b/Engine/Plugins/Experimental/PCG/Source/PCG/Public/Elements/PCGSplineSampler.h
--- a/Engine/Plugins/Experimental/PCG/Source/PCG/Public/Elements/PCGSplineSampler.h
+++ b/Engine/Plugins/Experimental/PCG/Source/PCG/Public/Elements/PCGSplineSampler.h
@@ -106,7 +106,7 @@ struct PCG_API FPCGSplineSamplerParams
 	bool bProjectOntoSurface = false;
 
 	// TODO: DirtyCache for OnDependencyChanged when this float curve is an external asset
-	/** Defines the density for each sample based on its distance from the spline */
+	/** Defines the density for each sample based on its distance from the spline. X axis is normalized distance to boundary (0-1), Y axis is density value. */
 	UPROPERTY(EditAnywhere, Category = Settings, meta = (EditCondition = "Dimension==EPCGSplineSamplingDimension::OnInterior"))
 	FRuntimeFloatCurve InteriorDensityFalloffCurve;
 
