[PCG] Improve element output data validation to help guide type usage

diff --git a/Engine/Plugins/Experimental/PCG/Source/PCG/Private/PCGElement.cpp b/Engine/Plugins/Experimental/PCG/Source/PCG/Private/PCGElement.cpp
--- a/Engine/Plugins/Experimental/PCG/Source/PCG/Private/PCGElement.cpp
+++ b/Engine/Plugins/Experimental/PCG/Source/PCG/Private/PCGElement.cpp
@@ -356,12 +356,17 @@ void IPCGElement::CleanupAndValidateOutput(FPCGContext* Context) const
 				const int32 MatchIndex = OutputPinProperties.IndexOfByPredicate([&TaggedData](const FPCGPinProperties& InProp) { return TaggedData.Pin == InProp.Label; });
 				if (MatchIndex == INDEX_NONE)
 				{
-					PCGE_LOG(Warning, LogOnly, FText::Format(LOCTEXT("OutputCannotBeRouted", "Output generated for pin '{0}' but cannot be routed"), FText::FromName(TaggedData.Pin)));
+					PCGE_LOG(Error, GraphAndLog, FText::Format(LOCTEXT("OutputCannotBeRouted", "Output data generated for non-existent output pin '{0}'"), FText::FromName(TaggedData.Pin)));
 				}
-				// TODO: Temporary fix for Settings directly from InputData (ie. from elements with code and not PCG nodes)
-				else if(TaggedData.Data && !(OutputPinProperties[MatchIndex].AllowedTypes & TaggedData.Data->GetDataType()) && TaggedData.Data->GetDataType() != EPCGDataType::Settings)
+				else if (TaggedData.Data)
 				{
-					PCGE_LOG(Warning, LogOnly, FText::Format(LOCTEXT("OutputIncompatibleType", "Output generated for pin '{0}' does not have a compatible type: '{1}'"), FText::FromName(TaggedData.Pin), FText::FromString(UEnum::GetValueAsString(TaggedData.Data->GetDataType()))));
+					const bool bTypesOverlap = !!(OutputPinProperties[MatchIndex].AllowedTypes & TaggedData.Data->GetDataType());
+					const bool bTypeIsSubset = !(~OutputPinProperties[MatchIndex].AllowedTypes & TaggedData.Data->GetDataType());
+					// TODO: Temporary fix for Settings directly from InputData (ie. from elements with code and not PCG nodes)
+					if ((!bTypesOverlap || !bTypeIsSubset) && TaggedData.Data->GetDataType() != EPCGDataType::Settings)
+					{
+						PCGE_LOG(Warning, GraphAndLog, FText::Format(LOCTEXT("OutputIncompatibleType", "Output data generated for pin '{0}' does not have a compatible type: '{1}'. Consider using more specific/narrower input pin types, or more general/wider output pin types."), FText::FromName(TaggedData.Pin), FText::FromString(UEnum::GetValueAsString(TaggedData.Data->GetDataType()))));
+					}
 				}
 
 				if (CVarPCGValidatePointMetadata.GetValueOnAnyThread())
