Fixed warning on load/on cook related to actors with a PCG component holding a reference to a non-spatially loaded actor - the landscape

diff --git a/Engine/Plugins/Experimental/PCG/Source/PCG/Private/PCGComponent.cpp b/Engine/Plugins/Experimental/PCG/Source/PCG/Private/PCGComponent.cpp
--- a/Engine/Plugins/Experimental/PCG/Source/PCG/Private/PCGComponent.cpp
+++ b/Engine/Plugins/Experimental/PCG/Source/PCG/Private/PCGComponent.cpp
@@ -1601,12 +1601,12 @@ void UPCGComponent::UpdateTrackedLandscape(bool bBoundsCheck)
 				UPCGData* ActorData = GetActorPCGData();
 				if (const UPCGSpatialData* ActorSpatialData = Cast<const UPCGSpatialData>(ActorData))
 				{
-					TrackedLandscapes = PCGHelpers::GetLandscapeProxies(World, ActorSpatialData->GetBounds());
+					Algo::Transform(PCGHelpers::GetLandscapeProxies(World, ActorSpatialData->GetBounds()), TrackedLandscapes, [](TWeakObjectPtr<ALandscapeProxy> Landscape) { return TSoftObjectPtr<ALandscapeProxy>(Landscape.Get()); });
 				}
 			}
 			else
 			{
-				TrackedLandscapes = PCGHelpers::GetAllLandscapeProxies(World);
+				Algo::Transform(PCGHelpers::GetAllLandscapeProxies(World), TrackedLandscapes, [](TWeakObjectPtr<ALandscapeProxy> Landscape) { return TSoftObjectPtr<ALandscapeProxy>(Landscape.Get()); });
 			}
 		}
 	}
@@ -1616,7 +1616,7 @@ void UPCGComponent::UpdateTrackedLandscape(bool bBoundsCheck)
 
 void UPCGComponent::SetupLandscapeTracking()
 {
-	for (TWeakObjectPtr<ALandscapeProxy> LandscapeProxy : TrackedLandscapes)
+	for (TSoftObjectPtr<ALandscapeProxy> LandscapeProxy : TrackedLandscapes)
 	{
 		if (LandscapeProxy.IsValid())
 		{
@@ -1627,7 +1627,7 @@ void UPCGComponent::SetupLandscapeTracking()
 
 void UPCGComponent::TeardownLandscapeTracking()
 {
-	for (TWeakObjectPtr<ALandscapeProxy> LandscapeProxy : TrackedLandscapes)
+	for (TSoftObjectPtr<ALandscapeProxy> LandscapeProxy : TrackedLandscapes)
 	{
 		if (LandscapeProxy.IsValid())
 		{
diff --git a/Engine/Plugins/Experimental/PCG/Source/PCG/Public/PCGComponent.h b/Engine/Plugins/Experimental/PCG/Source/PCG/Public/PCGComponent.h
--- a/Engine/Plugins/Experimental/PCG/Source/PCG/Public/PCGComponent.h
+++ b/Engine/Plugins/Experimental/PCG/Source/PCG/Public/PCGComponent.h
@@ -428,7 +428,7 @@ private:
 
 #if WITH_EDITORONLY_DATA
 	UPROPERTY()
-	TArray<TWeakObjectPtr<ALandscapeProxy>> TrackedLandscapes;
+	TArray<TSoftObjectPtr<ALandscapeProxy>> TrackedLandscapes;
 #endif
 
 #if WITH_EDITORONLY_DATA
