[PCG] Fix type warning on 'Get Texture Data' node

Type validation is catching error - this node outputs TextureData through a Point pin which is incorrect. Fix by broadening output pin type to Texture.

Expose errors on graph as these are happen frequently with this node.

diff --git a/Engine/Plugins/Experimental/PCG/Source/PCG/Private/Data/PCGTextureData.cpp b/Engine/Plugins/Experimental/PCG/Source/PCG/Private/Data/PCGTextureData.cpp
--- a/Engine/Plugins/Experimental/PCG/Source/PCG/Private/Data/PCGTextureData.cpp
+++ b/Engine/Plugins/Experimental/PCG/Source/PCG/Private/Data/PCGTextureData.cpp
@@ -275,15 +275,13 @@ void UPCGTextureData::Initialize(UTexture2D* InTexture, const FTransform& InTran
 
 	if (InTexture)
 	{
-		FTexturePlatformData* PlatformData = Texture->GetPlatformData();
-		if (PlatformData && PlatformData->Mips.Num() > 0)
+		if (IsSupported(InTexture))
 		{
-			TRACE_CPUPROFILER_EVENT_SCOPE(UPCGTextureData::Initialize::ReadData);
-
-			if (PlatformData->PixelFormat == PF_B8G8R8A8 ||
-				PlatformData->PixelFormat == PF_R8G8B8A8 ||
-				PlatformData->PixelFormat == PF_G8)
+			FTexturePlatformData* PlatformData = Texture->GetPlatformData();
+			if (PlatformData)
 			{
+				TRACE_CPUPROFILER_EVENT_SCOPE(UPCGTextureData::Initialize::ReadData);
+
 				if (const uint8_t* BulkData = reinterpret_cast<const uint8_t*>(PlatformData->Mips[0].BulkData.LockReadOnly()))
 				{
 					Width = Texture->GetSizeX();
@@ -322,13 +320,13 @@ void UPCGTextureData::Initialize(UTexture2D* InTexture, const FTransform& InTran
 				{
 					UE_LOG(LogPCG, Error, TEXT("PCGTextureData unable to get bulk data from %s"), *Texture->GetFName().ToString());
 				}
-				
+
 				PlatformData->Mips[0].BulkData.Unlock();
 			}
-			else
-			{
-				UE_LOG(LogPCG, Error, TEXT("PCGTextureData does not support the format of %s"), *Texture->GetFName().ToString());
-			}
+		}
+		else
+		{
+			UE_LOG(LogPCG, Error, TEXT("PCGTextureData does not support the format of %s"), *Texture->GetFName().ToString());
 		}
 	}
 
@@ -338,6 +336,14 @@ void UPCGTextureData::Initialize(UTexture2D* InTexture, const FTransform& InTran
 	Bounds = Bounds.TransformBy(Transform);
 }
 
+bool UPCGTextureData::IsSupported(UTexture2D* InTexture)
+{
+	const FTexturePlatformData* PlatformData = InTexture ? InTexture->GetPlatformData() : nullptr;
+
+	return PlatformData && PlatformData->Mips.Num() > 0 &&
+		(PlatformData->PixelFormat == PF_B8G8R8A8 || PlatformData->PixelFormat == PF_R8G8B8A8 || PlatformData->PixelFormat == PF_G8);
+}
+
 UPCGSpatialData* UPCGTextureData::CopyInternal() const
 {
 	UPCGTextureData* NewTextureData = NewObject<UPCGTextureData>();
diff --git a/Engine/Plugins/Experimental/PCG/Source/PCG/Private/Elements/PCGTextureSampler.cpp b/Engine/Plugins/Experimental/PCG/Source/PCG/Private/Elements/PCGTextureSampler.cpp
--- a/Engine/Plugins/Experimental/PCG/Source/PCG/Private/Elements/PCGTextureSampler.cpp
+++ b/Engine/Plugins/Experimental/PCG/Source/PCG/Private/Elements/PCGTextureSampler.cpp
@@ -14,6 +14,16 @@
 
 #include UE_INLINE_GENERATED_CPP_BY_NAME(PCGTextureSampler)
 
+#define LOCTEXT_NAMESPACE "PCGTextureSamplerElement"
+
+TArray<FPCGPinProperties> UPCGTextureSamplerSettings::OutputPinProperties() const
+{
+	TArray<FPCGPinProperties> Properties;
+	Properties.Emplace(PCGPinConstants::DefaultOutputLabel, EPCGDataType::Texture);
+
+	return Properties;
+}
+
 FPCGElementPtr UPCGTextureSamplerSettings::CreateElement() const
 {
 	return MakeShared<FPCGTextureSamplerElement>();
@@ -26,9 +36,21 @@ bool FPCGTextureSamplerElement::ExecuteInternal(FPCGContext* Context) const
 	const UPCGTextureSamplerSettings* Settings = Context->GetInputSettings<UPCGTextureSamplerSettings>();
 	check(Settings);
 
+	UTexture2D* Texture = Settings->Texture.LoadSynchronous();
+	if (!Texture)
+	{
+		PCGE_LOG(Error, GraphAndLog, FText::Format(LOCTEXT("CouldNotResolveTexture", "Texture at path '{0}' could not be loaded"), FText::FromString(Settings->Texture.ToString())));
+		return true;
+	}
+
+	if (!UPCGTextureData::IsSupported(Texture))
+	{
+		PCGE_LOG(Error, GraphAndLog, FText::Format(LOCTEXT("UnsupportedTextureFormat", "Texture '{0}' has unsupported texture format, currently supported formats are B8G8R8A8, PF_R8G8B8A8 and PF_G8"), FText::FromName(Texture->GetFName())));
+		return true;
+	}
+
 	const FTransform& Transform = Settings->Transform;
 	const bool bUseAbsoluteTransform = Settings->bUseAbsoluteTransform;
-	TObjectPtr<UTexture2D> Texture = Settings->Texture;
 	const EPCGTextureDensityFunction DensityFunction = Settings->DensityFunction;
 	const EPCGTextureColorChannel ColorChannel = Settings->ColorChannel;
 	const float TexelSize = Settings->TexelSize;
@@ -73,6 +95,11 @@ bool FPCGTextureSamplerElement::ExecuteInternal(FPCGContext* Context) const
 	TextureData->bUseTileBounds = bUseTileBounds;
 	TextureData->TileBounds = FBox2D(TileBoundsMin, TileBoundsMax);
 
+	if (!TextureData->IsValid())
+	{
+		PCGE_LOG(Error, GraphAndLog, LOCTEXT("UnsupportedTextureFormat", "Texture data failed to initialize, check log for more information"));
+	}
+
 	return true;
 }
 
@@ -97,3 +124,5 @@ void FPCGTextureSamplerElement::GetDependenciesCrc(const FPCGDataCollection& InI
 
 	OutCrc = Crc;
 }
+
+#undef LOCTEXT_NAMESPACE
diff --git a/Engine/Plugins/Experimental/PCG/Source/PCG/Public/Data/PCGTextureData.h b/Engine/Plugins/Experimental/PCG/Source/PCG/Public/Data/PCGTextureData.h
--- a/Engine/Plugins/Experimental/PCG/Source/PCG/Public/Data/PCGTextureData.h
+++ b/Engine/Plugins/Experimental/PCG/Source/PCG/Public/Data/PCGTextureData.h
@@ -104,15 +104,18 @@ public:
 	virtual EPCGDataType GetDataType() const override { return EPCGDataType::Texture; }
 	// ~End UPCGData interface
 
+	UFUNCTION(BlueprintCallable, Category = Texture)
+	void Initialize(UTexture2D* InTexture, const FTransform& InTransform);
+
+	/** Returns true if the format of InTexture is compatible and can be loaded. Will load texture if not already loaded. */
+	static bool IsSupported(UTexture2D* InTexture);
+
 	//~Begin UPCGSpatialData interface
 protected:
 	virtual UPCGSpatialData* CopyInternal() const override;
 	//~End UPCGSpatialData interface
 
 public:
-	UFUNCTION(BlueprintCallable, Category = Texture)
-	void Initialize(UTexture2D* InTexture, const FTransform& InTransform);
-
 	UPROPERTY(BlueprintReadOnly, VisibleAnywhere, Category = Properties)
 	TWeakObjectPtr<UTexture2D> Texture = nullptr;
 };
diff --git a/Engine/Plugins/Experimental/PCG/Source/PCG/Public/Elements/PCGTextureSampler.h b/Engine/Plugins/Experimental/PCG/Source/PCG/Public/Elements/PCGTextureSampler.h
--- a/Engine/Plugins/Experimental/PCG/Source/PCG/Public/Elements/PCGTextureSampler.h
+++ b/Engine/Plugins/Experimental/PCG/Source/PCG/Public/Elements/PCGTextureSampler.h
@@ -16,13 +16,13 @@ class PCG_API UPCGTextureSamplerSettings : public UPCGSettings
 public:
 	//~Begin UPCGSettings interface
 #if WITH_EDITOR
-	virtual FName GetDefaultNodeName() const override { return FName(TEXT("Texture Sampler")); }
-	virtual EPCGSettingsType GetType() const override { return EPCGSettingsType::Sampler; }
+	virtual FName GetDefaultNodeName() const override { return FName(TEXT("Get Texture Data")); }
+	virtual EPCGSettingsType GetType() const override { return EPCGSettingsType::Spatial; }
 #endif
 
 protected:
 	virtual TArray<FPCGPinProperties> InputPinProperties() const override { return TArray<FPCGPinProperties>(); }
-	virtual TArray<FPCGPinProperties> OutputPinProperties() const override { return Super::DefaultPointOutputPinProperties(); }
+	virtual TArray<FPCGPinProperties> OutputPinProperties() const override;
 	virtual FPCGElementPtr CreateElement() const override;
 	//~End UPCGSettings interface
 
@@ -36,7 +36,7 @@ public:
 
 	// Texture specific parameters
 	UPROPERTY(BlueprintReadWrite, EditAnywhere, Category = Settings)
-	TObjectPtr<UTexture2D> Texture = nullptr;
+	TSoftObjectPtr<UTexture2D> Texture = nullptr;
 
 	// Common members in BaseTextureData
 	UPROPERTY(BlueprintReadWrite, VisibleAnywhere, Category = SpatialData, meta = (PCG_Overridable))
