[PCG] Bounding Shape pin of samplers take union of all spatial inputs.

Matches established standard for all nodes and fixes error that people are running into.

diff --git a/Engine/Plugins/Experimental/PCG/Source/PCG/Private/Elements/PCGSplineSampler.cpp b/Engine/Plugins/Experimental/PCG/Source/PCG/Private/Elements/PCGSplineSampler.cpp
--- a/Engine/Plugins/Experimental/PCG/Source/PCG/Private/Elements/PCGSplineSampler.cpp
+++ b/Engine/Plugins/Experimental/PCG/Source/PCG/Private/Elements/PCGSplineSampler.cpp
@@ -1142,19 +1142,15 @@ bool FPCGSplineSamplerElement::ExecuteInternal(FPCGContext* Context) const
 	check(Settings);
 
 	TArray<FPCGTaggedData> SplineInputs = Context->InputData.GetInputsByPin(PCGSplineSamplerConstants::SplineLabel);
+	TArray<FPCGTaggedData> BoundingShapeInputs = Context->InputData.GetInputsByPin(PCGSplineSamplerConstants::BoundingShapeLabel);
 
 	// Grab the Bounding Shape input if there is one.
-	TArray<FPCGTaggedData> BoundingShapeInputs = Context->InputData.GetInputsByPin(PCGSplineSamplerConstants::BoundingShapeLabel);
-	const UPCGSpatialData* BoundingShapeSpatialInput = nullptr;
-	if (BoundingShapeInputs.Num() > 0)
+	// TODO: Once we support time-slicing, put this in the context and root (see FPCGSurfaceSamplerContext)
+	bool bUnionCreated = false;
+	const UPCGSpatialData* BoundingShape = Context->InputData.GetSpatialUnionOfInputsByPin(PCGSplineSamplerConstants::BoundingShapeLabel, bUnionCreated);
+	if (!BoundingShape && BoundingShapeInputs.Num() > 0)
 	{
-		ensure(BoundingShapeInputs.Num() == 1);
-		BoundingShapeSpatialInput = Cast<UPCGSpatialData>(BoundingShapeInputs[0].Data);
-
-		if (!BoundingShapeSpatialInput)
-		{
-			PCGE_LOG(Warning, GraphAndLog, LOCTEXT("BoundingShapeMissing", "Bounding Shape input is missing or of unsupported type and will not be used."));
-		}
+		PCGE_LOG(Warning, GraphAndLog, LOCTEXT("BoundingShapeMissing", "Bounding Shape input is missing or of unsupported type and will not be used"));
 	}
 
 	const FPCGSplineSamplerParams& SamplerParams = Settings->SamplerParams;
@@ -1193,11 +1189,11 @@ bool FPCGSplineSamplerElement::ExecuteInternal(FPCGContext* Context) const
 
 		if (SamplerParams.Dimension == EPCGSplineSamplingDimension::OnInterior)
 		{
-			PCGSplineSampler::SampleInteriorData(Context, LineData, BoundingShapeSpatialInput, ProjectionTarget, ProjectionParams, SamplerParams, SampledPointData);
+			PCGSplineSampler::SampleInteriorData(Context, LineData, BoundingShape, ProjectionTarget, ProjectionParams, SamplerParams, SampledPointData);
 		}
 		else
 		{
-			PCGSplineSampler::SampleLineData(LineData, BoundingShapeSpatialInput, ProjectionTarget, ProjectionParams, SamplerParams, SampledPointData);
+			PCGSplineSampler::SampleLineData(LineData, BoundingShape, ProjectionTarget, ProjectionParams, SamplerParams, SampledPointData);
 		}
 	}
 
diff --git a/Engine/Plugins/Experimental/PCG/Source/PCG/Private/Elements/PCGSurfaceSampler.cpp b/Engine/Plugins/Experimental/PCG/Source/PCG/Private/Elements/PCGSurfaceSampler.cpp
--- a/Engine/Plugins/Experimental/PCG/Source/PCG/Private/Elements/PCGSurfaceSampler.cpp
+++ b/Engine/Plugins/Experimental/PCG/Source/PCG/Private/Elements/PCGSurfaceSampler.cpp
@@ -307,6 +307,19 @@ bool UPCGSurfaceSamplerSettings::IsPinUsedByNodeExecution(const UPCGPin* InPin)
 }
 #endif
 
+FPCGSurfaceSamplerContext::~FPCGSurfaceSamplerContext()
+{
+	if (bUnionDataCreated && BoundingShapeSpatialInput)
+	{
+		// We created the data so can cast away constness
+		UPCGSpatialData* BoundingShapeMutable = const_cast<UPCGSpatialData*>(BoundingShapeSpatialInput);
+		BoundingShapeMutable->RemoveFromRoot();
+		BoundingShapeMutable->MarkAsGarbage();
+
+		BoundingShapeSpatialInput = nullptr;
+	}
+}
+
 FPCGElementPtr UPCGSurfaceSamplerSettings::CreateElement() const
 {
 	return MakeShared<FPCGSurfaceSamplerElement>();
@@ -346,12 +359,15 @@ bool FPCGSurfaceSamplerElement::AddGeneratingShapesToContext(FPCGSurfaceSamplerC
 
 	if (!Settings->bUnbounded)
 	{
-		if (BoundingShapeInputs.Num() > 0)
+		InContext->BoundingShapeSpatialInput = InContext->InputData.GetSpatialUnionOfInputsByPin(PCGSurfaceSamplerConstants::BoundingShapeLabel, InContext->bUnionDataCreated);
+
+		// Root dynamically created data to keep it safe from GC
+		if (InContext->bUnionDataCreated && InContext->BoundingShapeSpatialInput)
 		{
-			ensure(BoundingShapeInputs.Num() == 1);
-			InContext->BoundingShapeSpatialInput = Cast<UPCGSpatialData>(BoundingShapeInputs[0].Data);
+			const_cast<UPCGSpatialData*>(InContext->BoundingShapeSpatialInput)->AddToRoot();
 		}
-		else if (InContext->SourceComponent.IsValid())
+
+		if (!InContext->BoundingShapeSpatialInput && InContext->SourceComponent.IsValid())
 		{
 			// Fallback to getting bounds from actor
 			InContext->BoundingShapeSpatialInput = Cast<UPCGSpatialData>(InContext->SourceComponent->GetActorPCGData());
diff --git a/Engine/Plugins/Experimental/PCG/Source/PCG/Private/Elements/PCGVolumeSampler.cpp b/Engine/Plugins/Experimental/PCG/Source/PCG/Private/Elements/PCGVolumeSampler.cpp
--- a/Engine/Plugins/Experimental/PCG/Source/PCG/Private/Elements/PCGVolumeSampler.cpp
+++ b/Engine/Plugins/Experimental/PCG/Source/PCG/Private/Elements/PCGVolumeSampler.cpp
@@ -179,18 +179,16 @@ bool FPCGVolumeSamplerElement::ExecuteInternal(FPCGContext* Context) const
 
 	// Grab the Bounding Shape input if there is one.
 	TArray<FPCGTaggedData> BoundingShapeInputs = Context->InputData.GetInputsByPin(PCGVolumeSamplerConstants::BoundingShapeLabel);
-	const UPCGSpatialData* BoundingShapeSpatialInput = nullptr;
+	const UPCGSpatialData* BoundingShape = nullptr;
 
 	if (!Settings->bUnbounded)
 	{
-		if (BoundingShapeInputs.Num() > 0)
-		{
-			ensure(BoundingShapeInputs.Num() == 1);
-			BoundingShapeSpatialInput = Cast<UPCGSpatialData>(BoundingShapeInputs[0].Data);
-		}
-		else if (Context->SourceComponent.IsValid())
+		// TODO: Once we support time-slicing, put this in the context and root (see FPCGSurfaceSamplerContext)
+		bool bUnionCreated = false;
+		BoundingShape = Context->InputData.GetSpatialUnionOfInputsByPin(PCGVolumeSamplerConstants::BoundingShapeLabel, bUnionCreated);
+		if (!BoundingShape && Context->SourceComponent.IsValid())
 		{
-			BoundingShapeSpatialInput = Cast<UPCGSpatialData>(Context->SourceComponent->GetActorPCGData());
+			BoundingShape = Cast<UPCGSpatialData>(Context->SourceComponent->GetActorPCGData());
 		}
 	}
 	else if (BoundingShapeInputs.Num() > 0)
@@ -200,9 +198,9 @@ bool FPCGVolumeSamplerElement::ExecuteInternal(FPCGContext* Context) const
 
 	// Compute bounds of bounding shape input
 	FBox BoundingShapeBounds(EForceInit::ForceInit);
-	if (BoundingShapeSpatialInput)
+	if (BoundingShape)
 	{
-		BoundingShapeBounds = BoundingShapeSpatialInput->GetBounds();
+		BoundingShapeBounds = BoundingShape->GetBounds();
 	}
 
 	// Construct a list of shapes to generate samples from. Prefer to get these directly from the first input pin.
@@ -217,9 +215,9 @@ bool FPCGVolumeSamplerElement::ExecuteInternal(FPCGContext* Context) const
 	}
 
 	// If no shapes were obtained from the first input pin, try to find a shape to sample from nodes connected to the second pin.
-	if (GeneratingShapes.Num() == 0 && BoundingShapeSpatialInput)
+	if (GeneratingShapes.Num() == 0 && BoundingShape)
 	{
-		GeneratingShapes.Add(BoundingShapeSpatialInput);
+		GeneratingShapes.Add(BoundingShape);
 
 		// If there was a bounding shape input, use it as the starting point to get the tags
 		if (BoundingShapeInputs.Num() > 0)
@@ -279,7 +277,7 @@ bool FPCGVolumeSamplerElement::ExecuteInternal(FPCGContext* Context) const
 		}
 
 		// Sample volume
-		const UPCGPointData* SampledData = PCGVolumeSampler::SampleVolume(Context, GeneratingShape, BoundingShapeSpatialInput, InputBounds, SamplerSettings);
+		const UPCGPointData* SampledData = PCGVolumeSampler::SampleVolume(Context, GeneratingShape, BoundingShape, InputBounds, SamplerSettings);
 		Outputs[GenerationIndex].Data = SampledData;
 
 		if (SampledData)
diff --git a/Engine/Plugins/Experimental/PCG/Source/PCG/Private/PCGData.cpp b/Engine/Plugins/Experimental/PCG/Source/PCG/Private/PCGData.cpp
--- a/Engine/Plugins/Experimental/PCG/Source/PCG/Private/PCGData.cpp
+++ b/Engine/Plugins/Experimental/PCG/Source/PCG/Private/PCGData.cpp
@@ -8,6 +8,7 @@
 #include "PCGSettings.h"
 #include "Data/PCGPointData.h"
 #include "Data/PCGSpatialData.h"
+#include "Data/PCGUnionData.h"
 
 #include "HAL/IConsoleManager.h"
 #include "Serialization/ArchiveCrc32.h"
@@ -152,6 +153,52 @@ TArray<FPCGTaggedData> FPCGDataCollection::GetInputsByPin(const FName& InPinLabe
 		});
 }
 
+const UPCGSpatialData* FPCGDataCollection::GetSpatialUnionOfInputsByPin(const FName& InPinLabel, bool& bOutUnionDataCreated) const
+{
+	TArray<FPCGTaggedData> SpatialDataOnPin = TaggedData.FilterByPredicate([&InPinLabel](const FPCGTaggedData& Data) {
+		return Data.Pin == InPinLabel && Data.Data && Data.Data->IsA<UPCGSpatialData>();
+	});
+
+	const UPCGSpatialData* Result = nullptr;
+	UPCGUnionData* Union = nullptr;
+	for (const FPCGTaggedData& Data : SpatialDataOnPin)
+	{
+		const UPCGSpatialData* SpatialInput = Data.Data ? Cast<const UPCGSpatialData>(Data.Data) : nullptr;
+		if (!ensure(SpatialInput))
+		{
+			continue;
+		}
+
+		if (!Result)
+		{
+			// First valid data
+			Result = SpatialInput;
+		}
+		else
+		{
+			if (!Union)
+			{
+				// Second valid data - set up union
+				Union = NewObject<UPCGUnionData>();
+				Union->Initialize(Result, SpatialInput);
+
+				// Make result union
+				Result = Union;
+			}
+			else
+			{
+				// Nth valid data
+				Union->AddData(SpatialInput);
+			}
+		}
+	}
+
+	// Indicate whether we created a union data
+	bOutUnionDataCreated = !!Union;
+
+	return Result;
+}
+
 TArray<FPCGTaggedData> FPCGDataCollection::GetTaggedInputs(const FString& InTag) const
 {
 	return TaggedData.FilterByPredicate([&InTag](const FPCGTaggedData& Data) {
@@ -212,7 +259,7 @@ UPCGParamData* FPCGDataCollection::GetParamsWithDeprecation(const UPCGNode* Node
 			{
 				if (UPCGParamData* Params = Cast<UPCGParamData>(TaggedDatum.Data))
 				{
-					UE_LOG(LogPCG, Warning, TEXT("[%s] Found a param data on an input pin that should not accept params. Make sure to re-wire it to the param pin if it is used for overrides"), *Node->GetNodeTitle().ToString());
+					UE_LOG(LogPCG, Warning, TEXT("[%s] Found an Attribute Set data on an input pin that should not accept attributes. Make sure to re-wire it to the Overrides pin if it is used for overrides."), *Node->GetNodeTitle().ToString());
 					return Params;
 				}
 			}
diff --git a/Engine/Plugins/Experimental/PCG/Source/PCG/Public/Elements/PCGSurfaceSampler.h b/Engine/Plugins/Experimental/PCG/Source/PCG/Public/Elements/PCGSurfaceSampler.h
--- a/Engine/Plugins/Experimental/PCG/Source/PCG/Public/Elements/PCGSurfaceSampler.h
+++ b/Engine/Plugins/Experimental/PCG/Source/PCG/Public/Elements/PCGSurfaceSampler.h
@@ -125,6 +125,8 @@ namespace PCGSurfaceSampler
 
 struct FPCGSurfaceSamplerContext : public FPCGContext
 {
+	virtual ~FPCGSurfaceSamplerContext();
+
 	// Keep all data that won't be recomputed at every execution
 	const UPCGSpatialData* BoundingShapeSpatialInput = nullptr;
 	TArray<const UPCGSpatialData*, TInlineAllocator<16>> GeneratingShapes;
@@ -132,6 +134,9 @@ struct FPCGSurfaceSamplerContext : public FPCGContext
 	FBox BoundingShapeBounds = FBox(EForceInit::ForceInit);
 	int32 CurrentGeneratingShape = 0;
 	bool bDataPrepared = false;
+
+	/** True if BoundingShapeSpatialInput is data that was createdand rooted during data preparation */
+	bool bUnionDataCreated = false;
 };
 
 class FPCGSurfaceSamplerElement : public IPCGElement
diff --git a/Engine/Plugins/Experimental/PCG/Source/PCG/Public/PCGData.h b/Engine/Plugins/Experimental/PCG/Source/PCG/Public/PCGData.h
--- a/Engine/Plugins/Experimental/PCG/Source/PCG/Public/PCGData.h
+++ b/Engine/Plugins/Experimental/PCG/Source/PCG/Public/PCGData.h
@@ -11,9 +11,10 @@
 
 class FArchiveCrc32;
 class UPCGNode;
+class UPCGParamData;
 class UPCGSettings;
 class UPCGSettingsInterface;
-class UPCGParamData;
+class UPCGSpatialData;
 
 /**
 * Base class for any "data" class in the PCG framework.
@@ -97,6 +98,8 @@ struct PCG_API FPCGDataCollection
 	TArray<FPCGTaggedData> GetInputs() const;
 	/** Returns all data on a given pin */
 	TArray<FPCGTaggedData> GetInputsByPin(const FName& InPinLabel) const;
+	/** Returns spatial union of all data on a given pin, returns null if no such data exists. bOutUnionDataCreated indicates if new data created that may need rooting. */
+	const UPCGSpatialData* GetSpatialUnionOfInputsByPin(const FName& InPinLabel, bool& bOutUnionDataCreated) const;
 	/** Returns all spatial data in the collection with the given tag */
 	TArray<FPCGTaggedData> GetTaggedInputs(const FString& InTag) const;
 	/** Returns all settings in the collection */
