[PCG] Add user parameters to Graphs and overrides to Graph Instances.

* Added a new UI node, to mimic GetVar nodes in Blueprints
* Added dependency to ScriptUtils plugin
* Added possibility to add graph parameters
* Add the possibility to override those parameters in graph instances
* Handle the syncing between parameters and overrides
* Add new actions in the editor palette to create GetVar nodes, dynamic with the user parameters list.
* Added a new PCG Node that take a parameter and create a attribute from it.
* Add UI button in the PCG Graph Instance to save an instance
* Support for Subgraph and override parameters with auto added pins.

#rnx

diff --git a/Engine/Plugins/Experimental/PCG/PCG.uplugin b/Engine/Plugins/Experimental/PCG/PCG.uplugin
--- a/Engine/Plugins/Experimental/PCG/PCG.uplugin
+++ b/Engine/Plugins/Experimental/PCG/PCG.uplugin
@@ -35,6 +35,10 @@
 		{
 			"Name": "EditorScriptingUtilities",
 			"Enabled": true
+		},
+		{
+			"Name": "StructUtils",
+			"Enabled": true
 		}
 	]
 }
diff --git a/Engine/Plugins/Experimental/PCG/Source/PCG/PCG.Build.cs b/Engine/Plugins/Experimental/PCG/Source/PCG/PCG.Build.cs
--- a/Engine/Plugins/Experimental/PCG/Source/PCG/PCG.Build.cs
+++ b/Engine/Plugins/Experimental/PCG/Source/PCG/PCG.Build.cs
@@ -20,6 +20,7 @@ namespace UnrealBuildTool.Rules
 					"Projects",
 					"RenderCore",
 					"RHI",
+					"StructUtils"
 				});
 
 			PrivateDependencyModuleNames.AddRange(
diff --git a/Engine/Plugins/Experimental/PCG/Source/PCG/Private/Elements/PCGUserParameterGet.cpp b/Engine/Plugins/Experimental/PCG/Source/PCG/Private/Elements/PCGUserParameterGet.cpp
new file mode 100644
--- /dev/null
+++ b/Engine/Plugins/Experimental/PCG/Source/PCG/Private/Elements/PCGUserParameterGet.cpp
@@ -0,0 +1,198 @@
+// Copyright Epic Games, Inc. All Rights Reserved.
+
+#include "Elements/PCGUserParameterGet.h"
+
+#include "PCGComponent.h"
+#include "PCGContext.h"
+#include "PCGEdge.h"
+#include "PCGGraph.h"
+#include "PCGNode.h"
+#include "PCGParamData.h"
+#include "PCGPin.h"
+#include "PCGSubgraph.h"
+#include "Data/PCGUserParametersData.h"
+#include "Metadata/PCGMetadata.h"
+#include "Metadata/PCGMetadataAttributeTpl.h"
+#include "Metadata/Accessors/IPCGAttributeAccessor.h"
+#include "Metadata/Accessors/PCGAttributeAccessorHelpers.h"
+#include "Metadata/Accessors/PCGAttributeAccessorKeys.h"
+
+#include "PropertyBag.h"
+#include "StructView.h"
+
+#define LOCTEXT_NAMESPACE "PCGUserParameterGetElement"
+
+namespace PCGUserParameterGetSettings
+{
+	/**
+	* Utility function to get the first valid instanced property bag
+	* We define valid as if the ParameterOverrides from a GraphInstance and UserParameters from the graph owner of the node have the same property bag.
+	* By construction, it should be always the case, but we want to prevent cases where graph instances depend on other graph instances, that has changed their graph
+	* but didn't propagate the changes.
+	* If the property bags aren't the same, we are traversing the graph instance hierarchy to find the first graph/graph instance that matches.
+	*/
+	FConstStructView GetFirstValidLayout(FPCGContext& InContext)
+	{
+		FConstStructView GraphParameters{};
+		const UPCGGraphInterface* GraphInterface = nullptr;
+
+		// First we will read from the input if we find an input for this graph, set by the subgraph element.
+		TArray<FPCGTaggedData> UserParameterData = InContext.InputData.GetTaggedTypedInputs<UPCGUserParametersData>(PCGBaseSubgraphConstants::UserParameterTagData);
+		if (!UserParameterData.IsEmpty())
+		{
+			if (const UPCGUserParametersData* OverrideParametersData = CastChecked<UPCGUserParametersData>(UserParameterData[0].Data))
+			{
+				GraphInterface = OverrideParametersData->OriginalGraph;
+				GraphParameters = FConstStructView{ OverrideParametersData->UserParameters };
+			}
+		}
+
+		// Then if we don't have any input, we will use the graph instance from the component.
+		if (!GraphInterface || GraphInterface->IsA<UPCGGraph>())
+		{
+			const UPCGComponent* SourceComponent = InContext.SourceComponent.Get();
+			GraphInterface = SourceComponent ? SourceComponent->GetGraphInstance() : nullptr;
+		}
+
+		if (!GraphInterface)
+		{
+			return FConstStructView{};
+		}
+
+		// We gather the outer graph of this node, to make sure it matches our interface.
+		const UPCGGraph* GraphFromNode = Cast<UPCGGraph>(InContext.Node->GetOuter());
+		check(GraphFromNode);
+		FConstStructView GraphFromNodeParameters = GraphFromNode->GetUserParametersStruct() ? GraphFromNode->GetUserParametersStruct()->GetValue() : FConstStructView{};
+		if (!GraphFromNodeParameters.IsValid())
+		{
+			return FConstStructView{};
+		}
+
+		// If we don't have a graph instance, we just use the user parameters from the node graph owner.
+		const UPCGGraphInstance* GraphInstance = Cast<UPCGGraphInstance>(GraphInterface);
+		if (!GraphInstance)
+		{
+			return GraphFromNodeParameters;
+		}
+
+		// Making sure the graph matches.
+		if (GraphFromNode != GraphInstance->GetGraph())
+		{
+			return FConstStructView{};
+		}
+
+		// If we don't have overridden parameters from the input, get the parameters of the graph instance.
+		if (!GraphParameters.IsValid())
+		{
+			GraphParameters = GraphInstance->GetUserParametersStruct() ? GraphInstance->GetUserParametersStruct()->GetValue() : FConstStructView{};
+		}
+
+		// If there is no user parameters, just use the one from the graph.
+		if (!GraphParameters.IsValid())
+		{
+			return GraphFromNodeParameters;
+		}
+
+		const UScriptStruct* WantedScriptStruct = GraphFromNodeParameters.GetScriptStruct();
+
+		// Finally, we go down the graph instance chain to find the first valid layout.
+		// In most cases, we should never enter this loop, as layouts should match. But it is a safeguard.
+		while (GraphInstance && GraphParameters.GetScriptStruct() != WantedScriptStruct)
+		{
+			GraphInstance = Cast<UPCGGraphInstance>(GraphInstance->Graph);
+			if (!GraphInstance)
+			{
+				// We reached the end, use the user parameters default values. 
+				GraphParameters = GraphFromNodeParameters;
+				break;
+			}
+
+			GraphParameters = GraphInstance->GetUserParametersStruct()->GetValue();
+		}
+
+
+		return GraphParameters;
+	}
+}
+
+UPCGUserParameterGetSettings::UPCGUserParameterGetSettings(const FObjectInitializer& ObjectInitializer)
+{
+#if WITH_EDITORONLY_DATA
+	bExposeToLibrary = false;
+#endif
+}
+
+TArray<FPCGPinProperties> UPCGUserParameterGetSettings::OutputPinProperties() const
+{
+	TArray<FPCGPinProperties> PinProperties;
+	PinProperties.Emplace(PropertyName, EPCGDataType::Param);
+
+	return PinProperties;
+}
+
+void UPCGUserParameterGetSettings::UpdatePropertyName(FName InNewName)
+{
+	if (PropertyName != InNewName)
+	{
+		Modify();
+		PropertyName = InNewName;
+	}
+}
+
+FPCGElementPtr UPCGUserParameterGetSettings::CreateElement() const
+{
+	return MakeShared<FPCGUserParameterGetElement>();
+}
+
+bool FPCGUserParameterGetElement::ExecuteInternal(FPCGContext* Context) const
+{
+	check(Context);
+
+	const UPCGUserParameterGetSettings* Settings = Context->GetInputSettings<UPCGUserParameterGetSettings>();
+	check(Settings);
+
+	const FName PropertyName = Settings->PropertyName;
+
+	FConstStructView Parameters = PCGUserParameterGetSettings::GetFirstValidLayout(*Context);
+	const UScriptStruct* PropertyBag = Parameters.GetScriptStruct();
+	const FProperty* Property = PropertyBag ? PropertyBag->FindPropertyByName(PropertyName) : nullptr;
+
+	if (!Property)
+	{
+		PCGE_LOG(Error, GraphAndLog, FText::Format(LOCTEXT("InvalidProperty", "Could not find the property '{0}' in the user parameters"), FText::FromName(PropertyName)));
+		return true;
+	}
+
+	TUniquePtr<IPCGAttributeAccessor> PropertyAccessor = PCGAttributeAccessorHelpers::CreatePropertyAccessor(Property);
+
+	if (!PropertyAccessor)
+	{
+		// TODO: Should not happen when we have a working filter on the Property types.
+		PCGE_LOG(Error, GraphAndLog, FText::Format(LOCTEXT("InvalidType", "Parameter '{0}' type is unsupported by metadata. Available types are the ones you can find in the Create Attribute node"), FText::FromName(PropertyName)));
+		return true;
+	}
+
+	const FPCGAttributeAccessorKeysSingleObjectPtr PropertyAccessorKey(Parameters.GetMemory());
+
+	UPCGParamData* NewParamData = NewObject<UPCGParamData>();
+	check(NewParamData && NewParamData->Metadata);
+	PCGMetadataEntryKey NewEntry = NewParamData->Metadata->AddEntry();
+
+	auto CreateAndSet = [&PropertyAccessor, &PropertyAccessorKey, PropertyName, NewEntry, Metadata = NewParamData->Metadata](auto&& Dummy)
+	{
+		using AttributeType = std::decay_t<decltype(Dummy)>;
+
+		AttributeType Value{};
+		PropertyAccessor->Get<AttributeType>(Value, PropertyAccessorKey);
+
+		FPCGMetadataAttribute<AttributeType>* Attribute = Metadata->CreateAttribute<AttributeType>(PropertyName, Value, /*bAllowInterpolation=*/true, /*bOverrideParent=*/false);
+		Attribute->SetValue(NewEntry, Value);
+	};
+
+	PCGMetadataAttribute::CallbackWithRightType(PropertyAccessor->GetUnderlyingType(), CreateAndSet);
+
+	Context->OutputData.TaggedData.Emplace_GetRef().Data = NewParamData;
+	return true;
+}
+
+#undef LOCTEXT_NAMESPACE
\ No newline at end of file
diff --git a/Engine/Plugins/Experimental/PCG/Source/PCG/Private/Graph/PCGGraphCompiler.cpp b/Engine/Plugins/Experimental/PCG/Source/PCG/Private/Graph/PCGGraphCompiler.cpp
--- a/Engine/Plugins/Experimental/PCG/Source/PCG/Private/Graph/PCGGraphCompiler.cpp
+++ b/Engine/Plugins/Experimental/PCG/Source/PCG/Private/Graph/PCGGraphCompiler.cpp
@@ -110,6 +110,22 @@ TArray<FPCGGraphTask> FPCGGraphCompiler::CompileGraph(UPCGGraph* InGraph, FPCGTa
 				InputNodeTask->Inputs.Emplace(PreId, nullptr, nullptr);
 			}
 
+			// Hook nodes to the PreTask if they require so.
+			// Only do it for nodes that are directly under the subgraph, not in subsequent subgraphs.
+			for (FPCGGraphTask& Subtask : Subtasks)
+			{
+				if (!Subtask.Node || Subtask.Node->GetOuter() != Subgraph)
+				{
+					continue;
+				}
+
+				const UPCGSettings* Settings = Subtask.Node->GetSettings();
+				if (Settings && Settings->ShouldHookToPreTask())
+				{
+					Subtask.Inputs.Emplace(PreId, nullptr, nullptr);
+				}
+			}
+
 			// Merge subgraph tasks into current tasks.
 			CompiledTasks.Append(Subtasks);
 
diff --git a/Engine/Plugins/Experimental/PCG/Source/PCG/Private/PCGComponent.cpp b/Engine/Plugins/Experimental/PCG/Source/PCG/Private/PCGComponent.cpp
--- a/Engine/Plugins/Experimental/PCG/Source/PCG/Private/PCGComponent.cpp
+++ b/Engine/Plugins/Experimental/PCG/Source/PCG/Private/PCGComponent.cpp
@@ -103,7 +103,7 @@ void UPCGComponent::SetIsPartitioned(bool bIsNowPartitioned)
 	}
 }
 
-void UPCGComponent::SetGraph_Implementation(UPCGGraph* InGraph)
+void UPCGComponent::SetGraph_Implementation(UPCGGraphInterface* InGraph)
 {
 	SetGraphInterfaceLocal(InGraph);
 }
@@ -113,7 +113,7 @@ UPCGGraph* UPCGComponent::GetGraph() const
 	return (GraphInstance ? GraphInstance->GetGraph() : nullptr);
 }
 
-void UPCGComponent::SetGraphLocal(UPCGGraph* InGraph)
+void UPCGComponent::SetGraphLocal(UPCGGraphInterface* InGraph)
 {
 	SetGraphInterfaceLocal(InGraph);
 }
@@ -194,6 +194,8 @@ void UPCGComponent::SetPropertiesFromOriginal(const UPCGComponent* Original)
 		}
 	}
 
+	const bool bGraphInstanceIsDifferent = !GraphInstance->IsEquivalent(Original->GraphInstance);
+
 #if WITH_EDITOR
 	const bool bHasDirtyInput = InputType != NewInputType;
 	const bool bHasDirtyExclusions = !(ExcludedTags.Num() == Original->ExcludedTags.Num() && ExcludedTags.Includes(Original->ExcludedTags));
@@ -204,15 +206,24 @@ void UPCGComponent::SetPropertiesFromOriginal(const UPCGComponent* Original)
 	Original->CachedTrackedTagsToSettings.GetKeys(OriginalTrackedTags);
 	const bool bHasDirtyTracking = !(TrackedTags.Num() == OriginalTrackedTags.Num() && TrackedTags.Includes(OriginalTrackedTags));
 
-	const bool bGraphInstanceIsDifferent = !GraphInstance->IsEquivalent(Original->GraphInstance);
 	const bool bIsDirty = bHasDirtyInput || bHasDirtyExclusions || bHasDirtyTracking || bGraphInstanceIsDifferent;
 #endif // WITH_EDITOR
 
 	InputType = NewInputType;
 	Seed = Original->Seed;
-	GraphInstance->SetGraph(Original->GraphInstance ? Original->GraphInstance->GetGraph() : nullptr);
 	GenerationTrigger = Original->GenerationTrigger;
 
+	UPCGGraph* OriginalGraph = Original->GraphInstance ? Original->GraphInstance->GetGraph() : nullptr;
+	if (OriginalGraph != GraphInstance->GetGraph())
+	{
+		GraphInstance->SetGraph(OriginalGraph);
+	}
+
+	if (bGraphInstanceIsDifferent && OriginalGraph)
+	{
+		GraphInstance->CopyParameterOverrides(Original->GraphInstance);
+	}
+
 #if WITH_EDITOR
 	if (bHasDirtyExclusions || bHasDirtyTracking)
 	{
diff --git a/Engine/Plugins/Experimental/PCG/Source/PCG/Private/PCGContext.cpp b/Engine/Plugins/Experimental/PCG/Source/PCG/Private/PCGContext.cpp
--- a/Engine/Plugins/Experimental/PCG/Source/PCG/Private/PCGContext.cpp
+++ b/Engine/Plugins/Experimental/PCG/Source/PCG/Private/PCGContext.cpp
@@ -160,12 +160,15 @@ void FPCGContext::OverrideSettings()
 		if (!SettingsWithOverride->GetClass()->IsChildOf(Param.PropertyClass))
 		{
 			UObject* ObjectPtr = GetExternalContainerForOverridableParam(Param);
-			if (!ObjectPtr || !ObjectPtr->IsA(Param.Properties[0]->GetOwnerClass()))
+			if (ObjectPtr && ObjectPtr->IsA(Param.Properties[0]->GetOwnerClass()))
 			{
-				continue;
+				Container = ObjectPtr;
+			}
+			else if (!ObjectPtr)
+			{
+				// It's the responsability of the callee to make sure we have a valid memory space to read from.
+				Container = GetUnsafeExternalContainerForOverridableParam(Param);
 			}
-
-			Container = ObjectPtr;
 		}
 		else
 		{
@@ -177,7 +180,10 @@ void FPCGContext::OverrideSettings()
 			Container = SettingsWithOverride.Get();
 		}
 
-		check(Container);
+		if (!Container)
+		{
+			continue;
+		}
 
 		FName AttributeName = NAME_None;
 		TUniquePtr<const IPCGAttributeAccessor> AttributeAccessor = PCGAttributeAccessorHelpers::CreateConstAccessorForOverrideParam(InputData, Param, &AttributeName);
diff --git a/Engine/Plugins/Experimental/PCG/Source/PCG/Private/PCGData.cpp b/Engine/Plugins/Experimental/PCG/Source/PCG/Private/PCGData.cpp
--- a/Engine/Plugins/Experimental/PCG/Source/PCG/Private/PCGData.cpp
+++ b/Engine/Plugins/Experimental/PCG/Source/PCG/Private/PCGData.cpp
@@ -201,9 +201,7 @@ const UPCGSpatialData* FPCGDataCollection::GetSpatialUnionOfInputsByPin(const FN
 
 TArray<FPCGTaggedData> FPCGDataCollection::GetTaggedInputs(const FString& InTag) const
 {
-	return TaggedData.FilterByPredicate([&InTag](const FPCGTaggedData& Data) {
-		return Data.Tags.Contains(InTag) && Cast<UPCGSpatialData>(Data.Data);
-		});
+	return GetTaggedTypedInputs<UPCGSpatialData>(InTag);
 }
 
 TArray<FPCGTaggedData> FPCGDataCollection::GetAllSettings() const
diff --git a/Engine/Plugins/Experimental/PCG/Source/PCG/Private/PCGElement.cpp b/Engine/Plugins/Experimental/PCG/Source/PCG/Private/PCGElement.cpp
--- a/Engine/Plugins/Experimental/PCG/Source/PCG/Private/PCGElement.cpp
+++ b/Engine/Plugins/Experimental/PCG/Source/PCG/Private/PCGElement.cpp
@@ -356,7 +356,11 @@ void IPCGElement::CleanupAndValidateOutput(FPCGContext* Context) const
 				const int32 MatchIndex = OutputPinProperties.IndexOfByPredicate([&TaggedData](const FPCGPinProperties& InProp) { return TaggedData.Pin == InProp.Label; });
 				if (MatchIndex == INDEX_NONE)
 				{
-					PCGE_LOG(Error, GraphAndLog, FText::Format(LOCTEXT("OutputCannotBeRouted", "Output data generated for non-existent output pin '{0}'"), FText::FromName(TaggedData.Pin)));
+					// Only display an error if we expected this data to have a pin.
+					if (!TaggedData.bPinlessData)
+					{
+						PCGE_LOG(Error, GraphAndLog, FText::Format(LOCTEXT("OutputCannotBeRouted", "Output data generated for non-existent output pin '{0}'"), FText::FromName(TaggedData.Pin)));
+					}
 				}
 				else if (TaggedData.Data)
 				{
diff --git a/Engine/Plugins/Experimental/PCG/Source/PCG/Private/PCGGraph.cpp b/Engine/Plugins/Experimental/PCG/Source/PCG/Private/PCGGraph.cpp
--- a/Engine/Plugins/Experimental/PCG/Source/PCG/Private/PCGGraph.cpp
+++ b/Engine/Plugins/Experimental/PCG/Source/PCG/Private/PCGGraph.cpp
@@ -5,9 +5,11 @@
 #include "PCGComponent.h"
 #include "PCGInputOutputSettings.h"
 #include "PCGModule.h"
+#include "PCGNode.h"
 #include "PCGPin.h"
 #include "PCGSubsystem.h"
 #include "PCGSubgraph.h"
+#include "Elements/PCGUserParameterGet.h"
 
 #include UE_INLINE_GENERATED_CPP_BY_NAME(PCGGraph)
 
@@ -17,6 +19,54 @@
 #include "UObject/Package.h"
 #endif
 
+namespace PCGGraphUtils
+{
+	/** Returns true if the two descriptors are valid and compatible */
+	bool ArePropertiesCompatible(const FPropertyBagPropertyDesc* InSourcePropertyDesc, const FPropertyBagPropertyDesc* InTargetPropertyDesc)
+	{
+		return InSourcePropertyDesc && InTargetPropertyDesc && InSourcePropertyDesc->CompatibleType(*InTargetPropertyDesc);
+	}
+
+	/** Checks if the value for a source property in a source struct has the same value that the target property in the target struct. */
+	bool ArePropertiesIdentical(const FPropertyBagPropertyDesc* InSourcePropertyDesc, const FInstancedPropertyBag& InSourceInstance, const FPropertyBagPropertyDesc* InTargetPropertyDesc, const FInstancedPropertyBag& InTargetInstance)
+	{
+		if (!InSourceInstance.IsValid() || !InTargetInstance.IsValid() || !InSourcePropertyDesc || !InSourcePropertyDesc->CachedProperty || !InTargetPropertyDesc || !InTargetPropertyDesc->CachedProperty)
+		{
+			return false;
+		}
+
+		if (!InSourcePropertyDesc->CompatibleType(*InTargetPropertyDesc))
+		{
+			return false;
+		}
+
+		const uint8* SourceValueAddress = InSourceInstance.GetValue().GetMemory() + InSourcePropertyDesc->CachedProperty->GetOffset_ForInternal();
+		const uint8* TargetValueAddress = InTargetInstance.GetValue().GetMemory() + InTargetPropertyDesc->CachedProperty->GetOffset_ForInternal();
+
+		return InSourcePropertyDesc->CachedProperty->Identical(SourceValueAddress, TargetValueAddress);
+	}
+
+	/** Copy the value for a source property in a source struct to the target property in the target struct. */
+	void CopyPropertyValue(const FPropertyBagPropertyDesc* InSourcePropertyDesc, const FInstancedPropertyBag& InSourceInstance, const FPropertyBagPropertyDesc* InTargetPropertyDesc, FInstancedPropertyBag& InTargetInstance)
+	{
+		if (!InSourceInstance.IsValid() || !InTargetInstance.IsValid() || !InSourcePropertyDesc || !InSourcePropertyDesc->CachedProperty || !InTargetPropertyDesc || !InTargetPropertyDesc->CachedProperty)
+		{
+			return;
+		}
+
+		// Can't copy if they are not compatible.
+		if (!InSourcePropertyDesc->CompatibleType(*InTargetPropertyDesc))
+		{
+			return;
+		}
+
+		const uint8* SourceValueAddress = InSourceInstance.GetValue().GetMemory() + InSourcePropertyDesc->CachedProperty->GetOffset_ForInternal();
+		uint8* TargetValueAddress = InTargetInstance.GetMutableValue().GetMutableMemory() + InTargetPropertyDesc->CachedProperty->GetOffset_ForInternal();
+
+		InSourcePropertyDesc->CachedProperty->CopyCompleteValue(TargetValueAddress, SourceValueAddress);
+	}
+}
+
 /****************************
 * UPCGGraphInterface
 ****************************/
@@ -34,8 +84,58 @@ bool UPCGGraphInterface::IsEquivalent(const UPCGGraphInterface* Other) const
 	}
 
 	const UPCGGraph* OtherGraph = Other ? Other->GetGraph() : nullptr;
+	const UPCGGraph* ThisGraph = GetGraph();
 
-	return GetGraph() == OtherGraph;
+	if (ThisGraph != OtherGraph)
+	{
+		return false;
+	}
+	else if (!ThisGraph && !OtherGraph)
+	{
+		return true;
+	}
+
+	const FInstancedPropertyBag* OtherParameters = Other->GetUserParametersStruct();
+	const FInstancedPropertyBag* ThisParameters = this->GetUserParametersStruct();
+	check(OtherParameters && ThisParameters);
+
+	if (ThisParameters->GetNumPropertiesInBag() != OtherParameters->GetNumPropertiesInBag())
+	{
+		return false;
+	}
+
+	const UPropertyBag* OtherPropertyBag = OtherParameters->GetPropertyBagStruct();
+	const UPropertyBag* ThisPropertyBag = ThisParameters->GetPropertyBagStruct();
+
+	if (!ThisPropertyBag || !OtherPropertyBag)
+	{
+		return ThisPropertyBag == OtherPropertyBag;
+	}
+
+	// TODO: Be more resitant to different layout.
+	// For now we are only comparing structs that must have the same layout.
+	TConstArrayView<FPropertyBagPropertyDesc> OtherParametersDescs = OtherPropertyBag->GetPropertyDescs();
+	TConstArrayView<FPropertyBagPropertyDesc> ThisParametersDescs = ThisPropertyBag->GetPropertyDescs();
+	check(OtherParametersDescs.Num() == ThisParametersDescs.Num());
+
+	// TODO: Hashing might be more efficient.
+	for (int32 i = 0; i < ThisParametersDescs.Num(); ++i)
+	{
+		const FPropertyBagPropertyDesc& ThisParametersDesc = ThisParametersDescs[i];
+		const FPropertyBagPropertyDesc& OtherParametersDesc = OtherParametersDescs[i];
+
+		if (!PCGGraphUtils::ArePropertiesCompatible(&ThisParametersDesc, &OtherParametersDesc))
+		{
+			return false;
+		}
+
+		if (!PCGGraphUtils::ArePropertiesIdentical(&ThisParametersDesc, *ThisParameters, &OtherParametersDesc, *OtherParameters))
+		{
+			return false;
+		}
+	}
+
+	return true;
 }
 
 /****************************
@@ -687,6 +787,20 @@ void UPCGGraph::NotifyGraphChanged(EPCGChangeType ChangeType)
 	bIsNotifying = false;
 }
 
+void UPCGGraph::NotifyGraphParametersChanged(EPCGGraphParameterEvent InChangeType, FName InChangedPropertyName)
+{
+	if (bIsNotifying)
+	{
+		return;
+	}
+
+	bIsNotifying = true;
+	OnGraphParametersChangedDelegate.Broadcast(this, InChangeType, InChangedPropertyName);
+	bIsNotifying = false;
+
+	NotifyGraphChanged(EPCGChangeType::Settings);
+}
+
 void UPCGGraph::OnNodeChanged(UPCGNode* InNode, EPCGChangeType ChangeType)
 {
 	if((ChangeType & ~EPCGChangeType::Cosmetic) != EPCGChangeType::None)
@@ -695,15 +809,139 @@ void UPCGGraph::OnNodeChanged(UPCGNode* InNode, EPCGChangeType ChangeType)
 	}
 }
 
+void UPCGGraph::PreEditChange(FProperty* InProperty)
+{
+	Super::PreEditChange(InProperty);
+
+	if (!InProperty)
+	{
+		return;
+	}
+
+	if (InProperty->GetFName() == GET_MEMBER_NAME_CHECKED(UPCGGraph, UserParameters))
+	{
+		// We need to keep track of the number of properties, to detect if a property was added/removed/modified
+		NumberOfUserParametersPreEdit = UserParameters.GetNumPropertiesInBag();
+	}
+	else if (InProperty->GetOwnerStruct() == UserParameters.GetPropertyBagStruct())
+	{
+		// This is a bit unconventional, but we have to store the property name that has changed, because how Pre/Post Edit change is called.
+		// First PreEdit change is called twice. Once for the property changed (the one we want to track), and the second time on our UserParameters struct.
+		// Then PostEdit change is also called twice. But at the time of the first call (on the property we want to track), the value is not yet changed in memory of
+		// our UserParameters struct, it is still the old value. We need to wait the second call, PostEdit on our UserParameters struct, to have the memory updated with the new
+		// value. But then, we lost the name of our property.
+		// That's why we store the name here, and will use it in Post to know which property changed its value.
+		// TODO: It might be something that would need to change in StructUtils.
+		UserParameterModifiedName = InProperty->GetFName();
+	}
+}
+
 void UPCGGraph::PostEditChangeProperty(FPropertyChangedEvent& PropertyChangedEvent)
 {
 	Super::PostEditChangeProperty(PropertyChangedEvent);
 
-	if (PropertyChangedEvent.GetPropertyName() == GET_MEMBER_NAME_CHECKED(UPCGGraph, bLandscapeUsesMetadata))
+	if (!PropertyChangedEvent.Property)
+	{
+		return;
+	}
+
+	const FName PropertyName = PropertyChangedEvent.GetPropertyName();
+	const FName MemberPropertyName = PropertyChangedEvent.GetMemberPropertyName();
+
+	if (PropertyName == GET_MEMBER_NAME_CHECKED(UPCGGraph, bLandscapeUsesMetadata))
 	{
 		NotifyGraphChanged(EPCGChangeType::Input);
 	}
+	else if (PropertyName == GET_MEMBER_NAME_CHECKED(UPCGGraph, UserParameters))
+	{
+		EPCGGraphParameterEvent ChangeType;
+		int32 NumberOfUserParametersPostEdit = UserParameters.GetNumPropertiesInBag();
+
+		if (NumberOfUserParametersPostEdit > NumberOfUserParametersPreEdit)
+		{
+			ChangeType = EPCGGraphParameterEvent::Added;
+		}
+		else if (NumberOfUserParametersPostEdit < NumberOfUserParametersPreEdit)
+		{
+			ChangeType = EPCGGraphParameterEvent::Removed;
+		}
+		else //NumberOfUserParametersPostEdit == NumberOfUserParametersPreEdit
+		{
+			ChangeType = EPCGGraphParameterEvent::PropertyModified;
+		}
+
+		OnGraphParametersChanged(ChangeType, NAME_None);
+	}
+	else if (UserParameterModifiedName != NAME_None && MemberPropertyName == GET_MEMBER_NAME_CHECKED(UPCGGraph, UserParameters))
+	{
+		// cf. PreEditChnage comment to understand why we need UserParameterModifiedName
+		OnGraphParametersChanged(EPCGGraphParameterEvent::ValueModifiedLocally, UserParameterModifiedName);
+		UserParameterModifiedName = NAME_None;
+	}
+
+	NumberOfUserParametersPreEdit = 0;
+}
+
+void UPCGGraph::OnGraphParametersChanged(EPCGGraphParameterEvent InChangeType, FName InChangedPropertyName)
+{
+	bool bWasModified = false;
+
+	if (InChangeType == EPCGGraphParameterEvent::Removed || InChangeType == EPCGGraphParameterEvent::PropertyModified)
+	{
+		// Look for all the Get Parameter nodes and make sure to delete all nodes that doesn't exist anymore
+		TArray<UPCGNode*> NodesToRemove;
+
+		for (UPCGNode* Node : Nodes)
+		{
+			if (!Node)
+			{
+				continue;
+			}
+
+			if (UPCGUserParameterGetSettings* Settings = Cast<UPCGUserParameterGetSettings>(Node->GetSettings()))
+			{
+				const FPropertyBagPropertyDesc* PropertyDesc = UserParameters.FindPropertyDescByID(Settings->PropertyGuid);
+				if (!PropertyDesc)
+				{
+					NodesToRemove.Add(Node);
+				}
+				else if (Settings->PropertyName != PropertyDesc->Name)
+				{
+					const FName OldName = Settings->PropertyName;
+					Settings->UpdatePropertyName(PropertyDesc->Name);
+					// We make sure to keep the edges connected, by renaming the pin label
+					Node->RenameOutputPin(OldName, PropertyDesc->Name);
+				}
+			}
+		}
+
+		if (!NodesToRemove.IsEmpty())
+		{
+			bWasModified = true;
+			Modify();
+
+			for (UPCGNode* Node : NodesToRemove)
+			{
+				RemoveNode(Node);
+			}
+		}
+	}
+
+	// Also if we have no more properties, just reset
+	if (UserParameters.GetNumPropertiesInBag() == 0)
+	{
+		if (!bWasModified)
+		{
+			bWasModified = true;
+			Modify();
+		}
+
+		UserParameters.Reset();
+	}
+
+	NotifyGraphParametersChanged(InChangeType, InChangedPropertyName);
 }
+
 #endif // WITH_EDITOR
 
 /****************************
@@ -714,11 +952,18 @@ void UPCGGraphInstance::PostLoad()
 {
 	Super::PostLoad();
 
-#if WITH_EDITOR
 	if (Graph)
 	{
 		Graph->ConditionalPostLoad();
+	}
+
+	RefreshParameters(EPCGGraphParameterEvent::GraphPostLoad);
+
+#if WITH_EDITOR
+	if (Graph)
+	{
 		Graph->OnGraphChangedDelegate.AddUObject(this, &UPCGGraphInstance::OnGraphChanged);
+		Graph->OnGraphParametersChangedDelegate.AddUObject(this, &UPCGGraphInstance::OnGraphParametersChanged);
 	}
 #endif // WITH_EDITOR
 }
@@ -729,6 +974,7 @@ void UPCGGraphInstance::BeginDestroy()
 	if (Graph)
 	{
 		Graph->OnGraphChangedDelegate.RemoveAll(this);
+		Graph->OnGraphParametersChangedDelegate.RemoveAll(this);
 	}
 #endif // WITH_EDITOR
 
@@ -740,9 +986,19 @@ void UPCGGraphInstance::PreEditChange(FProperty* InProperty)
 {
 	Super::PreEditChange(InProperty);
 
-	if (InProperty && InProperty->GetFName() == GET_MEMBER_NAME_CHECKED(UPCGGraphInstance, Graph) && Graph)
+	if (!InProperty)
+	{
+		return;
+	}
+
+	if (InProperty->GetOwnerStruct() == ParametersOverrides.Parameters.GetPropertyBagStruct())
+	{
+		UserParameterModifiedName = InProperty->GetFName();
+	}
+	else if (InProperty->GetFName() == GET_MEMBER_NAME_CHECKED(UPCGGraphInstance, Graph) && Graph)
 	{
 		Graph->OnGraphChangedDelegate.RemoveAll(this);
+		Graph->OnGraphParametersChangedDelegate.RemoveAll(this);
 	}
 }
 
@@ -750,9 +1006,22 @@ void UPCGGraphInstance::PostEditChangeProperty(FPropertyChangedEvent& PropertyCh
 {
 	Super::PostEditChangeProperty(PropertyChangedEvent);
 
-	if (PropertyChangedEvent.Property && PropertyChangedEvent.Property->GetFName() == GET_MEMBER_NAME_CHECKED(UPCGGraphInstance, Graph) && Graph)
+	if (!PropertyChangedEvent.Property)
+	{
+		return;
+	}
+
+	if (PropertyChangedEvent.GetPropertyName() == GET_MEMBER_NAME_CHECKED(UPCGGraphInstance, Graph) && Graph)
 	{
 		Graph->OnGraphChangedDelegate.AddUObject(this, &UPCGGraphInstance::OnGraphChanged);
+		Graph->OnGraphParametersChangedDelegate.AddUObject(this, &UPCGGraphInstance::OnGraphParametersChanged);
+
+		RefreshParameters(EPCGGraphParameterEvent::GraphChanged);
+	}
+	else if (UserParameterModifiedName != NAME_None && PropertyChangedEvent.GetMemberPropertyName() == GET_MEMBER_NAME_CHECKED(UPCGGraphInstance, ParametersOverrides))
+	{
+		OnGraphParametersChanged(this, EPCGGraphParameterEvent::ValueModifiedLocally, UserParameterModifiedName);
+		UserParameterModifiedName = NAME_None;
 	}
 }
 
@@ -763,6 +1032,7 @@ void UPCGGraphInstance::PreEditUndo()
 	if (Graph)
 	{
 		Graph->OnGraphChangedDelegate.RemoveAll(this);
+		Graph->OnGraphParametersChangedDelegate.RemoveAll(this);
 	}
 
 	UndoRedoGraphCache = Graph;
@@ -775,6 +1045,7 @@ void UPCGGraphInstance::PostEditUndo()
 	if (Graph)
 	{
 		Graph->OnGraphChangedDelegate.AddUObject(this, &UPCGGraphInstance::OnGraphChanged);
+		Graph->OnGraphParametersChangedDelegate.AddUObject(this, &UPCGGraphInstance::OnGraphParametersChanged);
 	}
 }
 
@@ -819,6 +1090,7 @@ void UPCGGraphInstance::SetGraph(UPCGGraphInterface* InGraph)
 	if (Graph)
 	{
 		Graph->OnGraphChangedDelegate.RemoveAll(this);
+		Graph->OnGraphParametersChangedDelegate.RemoveAll(this);
 	}
 #endif // WITH_EDITOR
 
@@ -828,8 +1100,18 @@ void UPCGGraphInstance::SetGraph(UPCGGraphInterface* InGraph)
 	if (Graph)
 	{
 		Graph->OnGraphChangedDelegate.AddUObject(this, &UPCGGraphInstance::OnGraphChanged);
+		Graph->OnGraphParametersChangedDelegate.AddUObject(this, &UPCGGraphInstance::OnGraphParametersChanged);
 	}
 #endif // WITH_EDITOR
+
+#if WITH_EDITOR
+	OnGraphParametersChanged(Graph, EPCGGraphParameterEvent::GraphChanged, NAME_None);
+#else
+	// TODO: We need to revisit this, because it won't update any child graph that has this instance as their graph.
+	// Making the hotswap of graph instances within graph instances not working as intended in non-editor builds.
+	// Perhaps that should not be possible? At least it is mitigated in the GetUserParameter node, that will take the first valid layout.
+	RefreshParameters(EPCGGraphParameterEvent::GraphChanged, NAME_None);
+#endif // WITH_EDITOR
 }
 
 TObjectPtr<UPCGGraphInterface> UPCGGraphInstance::CreateInstance(UObject* InOwner, UPCGGraphInterface* InGraph)
@@ -844,3 +1126,305 @@ TObjectPtr<UPCGGraphInterface> UPCGGraphInstance::CreateInstance(UObject* InOwne
 
 	return GraphInstance;
 }
+
+#if WITH_EDITOR
+void UPCGGraphInstance::NotifyGraphParametersChanged(EPCGGraphParameterEvent InChangeType, FName InChangedPropertyName)
+{
+	OnGraphParametersChangedDelegate.Broadcast(this, InChangeType, InChangedPropertyName);
+
+	// Also propagates the changes
+	OnGraphChanged(Graph, EPCGChangeType::Settings);
+}
+
+void UPCGGraphInstance::OnGraphParametersChanged(UPCGGraphInterface* InGraph, EPCGGraphParameterEvent InChangeType, FName InChangedPropertyName)
+{
+	if (InGraph != Graph && InGraph != this)
+	{
+		return;
+	}
+
+	EPCGGraphParameterEvent ChangeType = InChangeType;
+	if (InGraph == Graph && InChangeType == EPCGGraphParameterEvent::ValueModifiedLocally)
+	{
+		// If we receive a "ValueModifiedLocally" and it was on our Graph, we transform it to "ValueModifiedByParent"
+		ChangeType = EPCGGraphParameterEvent::ValueModifiedByParent;
+	}
+
+	RefreshParameters(ChangeType, InChangedPropertyName);
+	NotifyGraphParametersChanged(ChangeType, InChangedPropertyName);
+}
+#endif // WITH_EDITOR
+
+void UPCGGraphInstance::RefreshParameters(EPCGGraphParameterEvent InChangeType, FName InChangedPropertyName)
+{
+	if (!Graph)
+	{
+		if (ParametersOverrides.IsValid())
+		{
+			Modify();
+		}
+
+		ParametersOverrides.Reset();
+	}
+	else
+	{		
+		const FInstancedPropertyBag* ParentUserParameters = Graph->GetUserParametersStruct();
+
+		// Refresh can modify nothing, but we still need to keep a snapshot of this object state, if it ever change.
+		// Don't mark it dirty by default, only if something changed.
+		Modify(/*bAlwaysMarkDirty=*/false);
+
+		if (ParametersOverrides.RefreshParameters(ParentUserParameters, InChangeType, InChangedPropertyName))
+		{
+			MarkPackageDirty();
+		}
+	}
+}
+
+void UPCGGraphInstance::UpdatePropertyOverride(const FProperty* InProperty, bool bMarkAsOverridden)
+{
+	if (!Graph || !InProperty)
+	{
+		return;
+	}
+
+	Modify();
+
+	const FInstancedPropertyBag* ParentUserParameters = Graph->GetUserParametersStruct();
+	if (ParametersOverrides.UpdatePropertyOverride(InProperty, bMarkAsOverridden, ParentUserParameters))
+	{
+#if WITH_EDITOR
+		// If it is true, it means that the value has changed, so propagate the changes, in Editor
+		NotifyGraphParametersChanged(EPCGGraphParameterEvent::ValueModifiedLocally, InProperty->GetFName());
+		OnGraphChangedDelegate.Broadcast(this, EPCGChangeType::Settings);
+#endif // WITH_EDITOR
+	}
+}
+
+void UPCGGraphInstance::CopyParameterOverrides(UPCGGraphInterface* InGraph)
+{
+	if (!InGraph)
+	{
+		return;
+	}
+
+	const UPCGGraph* ThisGraph = GetGraph();
+	const UPCGGraph* OtherGraph = InGraph->GetGraph();
+
+	// Can't copy if they have not the same base graph
+	if (ThisGraph != OtherGraph)
+	{
+		return;
+	}
+
+	ParametersOverrides.Parameters.CopyMatchingValuesByID(*InGraph->GetUserParametersStruct());
+}
+
+bool FPCGOverrideInstancedPropertyBag::RefreshParameters(const FInstancedPropertyBag* ParentUserParameters, EPCGGraphParameterEvent InChangeType, FName InChangedPropertyName)
+{
+	check(ParentUserParameters);
+
+	bool bWasModified = false;
+
+	if (!ParentUserParameters->IsValid())
+	{
+		Reset();
+		return true;
+	}
+
+	switch (InChangeType)
+	{
+	case EPCGGraphParameterEvent::GraphChanged:
+	{
+		// If the parameters property bags match, nothing to do.
+		if (ParentUserParameters->GetPropertyBagStruct() != Parameters.GetPropertyBagStruct())
+		{
+			bWasModified = true;
+			// Copy the parent parameters and reset overriddes
+			Parameters = *ParentUserParameters;
+			PropertiesIDsOverridden.Reset();
+		}
+		break;
+	}
+	case EPCGGraphParameterEvent::Added:
+	case EPCGGraphParameterEvent::Removed:
+	case EPCGGraphParameterEvent::PropertyModified:
+	{
+		bWasModified = true;
+		const FPropertyBagPropertyDesc* ThisPropertyDesc = Parameters.FindPropertyDescByName(InChangedPropertyName);
+
+		if (ThisPropertyDesc)
+		{
+			UpdatePropertyOverride(ThisPropertyDesc->CachedProperty, false, ParentUserParameters);
+		}
+
+		MigrateToNewBagInstance(*ParentUserParameters);
+		break;
+	}
+	case EPCGGraphParameterEvent::ValueModifiedByParent:
+	{
+		const FPropertyBagPropertyDesc* OriginalPropertyDesc = ParentUserParameters->FindPropertyDescByName(InChangedPropertyName);
+		const FPropertyBagPropertyDesc* ThisPropertyDesc = Parameters.FindPropertyDescByName(InChangedPropertyName);
+
+		check(InChangedPropertyName != NAME_None);
+		check(OriginalPropertyDesc);
+
+		if (!PCGGraphUtils::ArePropertiesCompatible(OriginalPropertyDesc, ThisPropertyDesc))
+		{
+			bWasModified = true;
+			MigrateToNewBagInstance(*ParentUserParameters);
+		}
+		else if (!IsPropertyOverridden(ThisPropertyDesc->CachedProperty))
+		{
+			// Only update the value if the property is not overriden.
+			bWasModified = true;
+			PCGGraphUtils::CopyPropertyValue(OriginalPropertyDesc, *ParentUserParameters, ThisPropertyDesc, Parameters);
+		}
+		break;
+	}
+	case EPCGGraphParameterEvent::ValueModifiedLocally:
+	{
+		const FPropertyBagPropertyDesc* OriginalPropertyDesc = ParentUserParameters->FindPropertyDescByName(InChangedPropertyName);
+		const FPropertyBagPropertyDesc* ThisPropertyDesc = Parameters.FindPropertyDescByName(InChangedPropertyName);
+
+		check(InChangedPropertyName != NAME_None);
+		check(OriginalPropertyDesc);
+
+		if (!PCGGraphUtils::ArePropertiesCompatible(OriginalPropertyDesc, ThisPropertyDesc))
+		{
+			bWasModified = true;
+			MigrateToNewBagInstance(*ParentUserParameters);
+		}
+		else
+		{
+			// Force the value to be overridden, if it is not equal to the value and it was changed from the outside
+			if (!PCGGraphUtils::ArePropertiesIdentical(OriginalPropertyDesc, *ParentUserParameters, ThisPropertyDesc, Parameters))
+			{
+				bWasModified = true;
+				UpdatePropertyOverride(ThisPropertyDesc->CachedProperty, true, ParentUserParameters);
+			}
+		}
+		break;
+	}
+	case EPCGGraphParameterEvent::GraphPostLoad:
+	{
+		// Check if the property struct mismatch. If so, do the migration
+		if (Parameters.GetPropertyBagStruct() != ParentUserParameters->GetPropertyBagStruct())
+		{
+			bWasModified = true;
+			MigrateToNewBagInstance(*ParentUserParameters);
+		}
+
+		if (Parameters.GetPropertyBagStruct() == nullptr)
+		{
+			return bWasModified;
+		}
+
+		// And then overwrite all non-overridden values
+		for (const FPropertyBagPropertyDesc& ThisPropertyDesc : Parameters.GetPropertyBagStruct()->GetPropertyDescs())
+		{
+			if (!IsPropertyOverridden(ThisPropertyDesc.CachedProperty))
+			{
+				const FPropertyBagPropertyDesc* OriginalPropertyDesc = ParentUserParameters->FindPropertyDescByID(ThisPropertyDesc.ID);
+
+				if (!PCGGraphUtils::ArePropertiesIdentical(OriginalPropertyDesc, *ParentUserParameters, &ThisPropertyDesc, Parameters))
+				{
+					bWasModified = true;
+					PCGGraphUtils::CopyPropertyValue(OriginalPropertyDesc, *ParentUserParameters, &ThisPropertyDesc, Parameters);
+				}
+			}
+		}
+		break;
+	}
+	}
+
+	return bWasModified;
+}
+
+bool FPCGOverrideInstancedPropertyBag::UpdatePropertyOverride(const FProperty* InProperty, bool bMarkAsOverridden, const FInstancedPropertyBag* ParentUserParameters)
+{
+	if (!InProperty)
+	{
+		return false;
+	}
+
+	if (const FPropertyBagPropertyDesc* PropertyDesc = Parameters.FindPropertyDescByName(InProperty->GetFName()))
+	{
+		if (bMarkAsOverridden)
+		{
+			PropertiesIDsOverridden.Add(PropertyDesc->ID);
+		}
+		else
+		{
+			PropertiesIDsOverridden.Remove(PropertyDesc->ID);
+		}
+	}
+
+	// Reset the value if it is not marked overridden anymore.
+	if (!bMarkAsOverridden)
+	{
+		check(ParentUserParameters);
+
+		const FPropertyBagPropertyDesc* OriginalPropertyDesc = ParentUserParameters->FindPropertyDescByName(InProperty->GetFName());
+		const FPropertyBagPropertyDesc* ThisPropertyDesc = Parameters.FindPropertyDescByName(InProperty->GetFName());
+
+		check(OriginalPropertyDesc && ThisPropertyDesc);
+		PCGGraphUtils::CopyPropertyValue(OriginalPropertyDesc, *ParentUserParameters, ThisPropertyDesc, Parameters);
+		return true;
+	}
+
+	return false;
+}
+
+bool FPCGOverrideInstancedPropertyBag::IsPropertyOverridden(const FProperty* InProperty) const
+{
+	if (!InProperty)
+	{
+		return false;
+	}
+
+	const FPropertyBagPropertyDesc* PropertyDesc = Parameters.FindPropertyDescByName(InProperty->GetFName());
+	return PropertyDesc && PropertiesIDsOverridden.Contains(PropertyDesc->ID);
+}
+
+void FPCGOverrideInstancedPropertyBag::Reset()
+{
+	Parameters.Reset();
+	PropertiesIDsOverridden.Reset();
+}
+
+void FPCGOverrideInstancedPropertyBag::MigrateToNewBagInstance(const FInstancedPropertyBag& NewBagInstance)
+{
+	// Keeping a map between id and types. We will remove override for property that changed types.
+	TMap<FGuid, FPropertyBagPropertyDesc> IdToDescMap;
+	if (Parameters.GetPropertyBagStruct())
+	{
+		for (const FPropertyBagPropertyDesc& PropertyDesc : Parameters.GetPropertyBagStruct()->GetPropertyDescs())
+		{
+			IdToDescMap.Emplace(PropertyDesc.ID, PropertyDesc);
+		}
+	}
+
+	Parameters.MigrateToNewBagInstance(NewBagInstance);
+
+	if (NewBagInstance.GetPropertyBagStruct() == nullptr)
+	{
+		return;
+	}
+
+	// Remove overridden parameters that are not in the bag anymore, or have changed type
+	TArray<FGuid> OverriddenParametersCopy = PropertiesIDsOverridden.Array();
+	for (const FGuid PropertyId: OverriddenParametersCopy)
+	{
+		const FPropertyBagPropertyDesc* NewPropertyDesc = NewBagInstance.FindPropertyDescByID(PropertyId);
+		const FPropertyBagPropertyDesc* OldPropertyDesc = IdToDescMap.Find(PropertyId);
+
+		const bool bTypeHasChanged = NewPropertyDesc && OldPropertyDesc && (NewPropertyDesc->ValueType != OldPropertyDesc->ValueType || NewPropertyDesc->ValueTypeObject != OldPropertyDesc->ValueTypeObject);
+
+		if (!NewPropertyDesc || bTypeHasChanged)
+		{
+			PropertiesIDsOverridden.Remove(PropertyId);
+			continue;
+		}
+	}
+}
diff --git a/Engine/Plugins/Experimental/PCG/Source/PCG/Private/PCGNode.cpp b/Engine/Plugins/Experimental/PCG/Source/PCG/Private/PCGNode.cpp
--- a/Engine/Plugins/Experimental/PCG/Source/PCG/Private/PCGNode.cpp
+++ b/Engine/Plugins/Experimental/PCG/Source/PCG/Private/PCGNode.cpp
@@ -368,6 +368,32 @@ bool UPCGNode::IsOutputPinConnected(const FName& Label) const
 	}
 }
 
+void UPCGNode::RenameInputPin(const FName& OldLabel, const FName& NewLabel)
+{
+	if (UPCGPin* Pin = GetInputPin(OldLabel))
+	{
+		Pin->Modify();
+		Pin->Properties.Label = NewLabel;
+
+#if WITH_EDITOR
+		OnNodeChangedDelegate.Broadcast(this, EPCGChangeType::Node);
+#endif // WITH_EDITOR
+	}
+}
+
+void UPCGNode::RenameOutputPin(const FName& OldLabel, const FName& NewLabel)
+{
+	if (UPCGPin* Pin = GetOutputPin(OldLabel))
+	{
+		Pin->Modify();
+		Pin->Properties.Label = NewLabel;
+
+#if WITH_EDITOR
+		OnNodeChangedDelegate.Broadcast(this, EPCGChangeType::Node);
+#endif // WITH_EDITOR
+	}
+}
+
 bool UPCGNode::HasInboundEdges() const
 {
 	for (const UPCGPin* InputPin : InputPins)
diff --git a/Engine/Plugins/Experimental/PCG/Source/PCG/Private/PCGSubgraph.cpp b/Engine/Plugins/Experimental/PCG/Source/PCG/Private/PCGSubgraph.cpp
--- a/Engine/Plugins/Experimental/PCG/Source/PCG/Private/PCGSubgraph.cpp
+++ b/Engine/Plugins/Experimental/PCG/Source/PCG/Private/PCGSubgraph.cpp
@@ -5,6 +5,8 @@
 #include "PCGComponent.h"
 #include "PCGPin.h"
 #include "PCGSubsystem.h"
+#include "Data/PCGUserParametersData.h"
+#include "Helpers/PCGSettingsHelpers.h"
 
 #include UE_INLINE_GENERATED_CPP_BY_NAME(PCGSubgraph)
 
@@ -167,6 +169,54 @@ TArray<FPCGPinProperties> UPCGBaseSubgraphSettings::OutputPinProperties() const
 	}
 }
 
+void UPCGBaseSubgraphSettings::FixingOverridableParamPropertyClass(FPCGSettingsOverridableParam& Param) const
+{
+	bool bFound = false;
+	const UPCGGraph* PCGGraph = GetSubgraph();
+
+	if (PCGGraph && !Param.PropertiesNames.IsEmpty())
+	{
+		if (const FInstancedPropertyBag* UserParameterStruct = PCGGraph->GetUserParametersStruct())
+		{
+			const UScriptStruct* ScriptStruct = UserParameterStruct->GetPropertyBagStruct();
+			if (ScriptStruct && ScriptStruct->FindPropertyByName(Param.PropertiesNames[0]))
+			{
+				Param.PropertyClass = ScriptStruct;
+				bFound = true;
+			}
+		}
+	}
+
+	if (!bFound)
+	{
+		Super::FixingOverridableParamPropertyClass(Param);
+	}
+}
+
+#if WITH_EDITOR
+TArray<FPCGSettingsOverridableParam> UPCGBaseSubgraphSettings::GatherOverridableParams() const
+{
+	TArray<FPCGSettingsOverridableParam> OverridableParams = Super::GatherOverridableParams();
+
+	const UPCGGraph* PCGGraph = GetSubgraph();
+	if (PCGGraph)
+	{
+		if (const FInstancedPropertyBag* UserParameterStruct = PCGGraph->GetUserParametersStruct())
+		{
+			if (const UScriptStruct* ScriptStruct = UserParameterStruct->GetPropertyBagStruct())
+			{
+				PCGSettingsHelpers::FPCGGetAllOverridableParamsConfig Config;
+				Config.bExcludeSuperProperties = true;
+				Config.ExcludePropertyFlags = CPF_DisableEditOnInstance;
+				OverridableParams.Append(PCGSettingsHelpers::GetAllOverridableParams(ScriptStruct, Config));
+			}
+		}
+	}
+
+	return OverridableParams;
+}
+#endif // WITH_EDITOR
+
 UPCGSubgraphSettings::UPCGSubgraphSettings(const FObjectInitializer& InObjectInitializer)
 	: Super(InObjectInitializer)
 {
@@ -240,6 +290,23 @@ TObjectPtr<UPCGGraphInterface> UPCGSubgraphNode::GetSubgraphInterface() const
 	return Settings ? Settings->GetSubgraphInterface() : nullptr;
 }
 
+void* FPCGSubgraphContext::GetUnsafeExternalContainerForOverridableParam(const FPCGSettingsOverridableParam& InParam)
+{
+	const UPCGSubgraphSettings* Settings = GetInputSettings<UPCGSubgraphSettings>();
+	check(Settings);
+	const UPCGGraphInterface* Graph = Settings->GetSubgraphInterface();
+	const FInstancedPropertyBag* UserParameters = Graph ? Graph->GetUserParametersStruct() : nullptr;
+
+	if (UserParameters && !InParam.PropertiesNames.IsEmpty() && UserParameters->FindPropertyDescByName(InParam.PropertiesNames[0]) && GraphInstanceParametersOverride.IsValid())
+	{
+		return GraphInstanceParametersOverride.GetMutableMemory();
+	}
+	else
+	{
+		return nullptr;
+	}
+}
+
 FPCGContext* FPCGSubgraphElement::Initialize(const FPCGDataCollection& InputData, TWeakObjectPtr<UPCGComponent> SourceComponent, const UPCGNode* Node)
 {
 	FPCGSubgraphContext* Context = new FPCGSubgraphContext();
@@ -247,6 +314,31 @@ FPCGContext* FPCGSubgraphElement::Initialize(const FPCGDataCollection& InputData
 	Context->SourceComponent = SourceComponent;
 	Context->Node = Node;
 
+	// Only duplicate the UserParameters if we have overriable params and we have at least one param pin connected.
+	const UPCGSubgraphSettings* Settings = Context->GetInputSettings<UPCGSubgraphSettings>();
+	check(Settings);
+	const TArray<FPCGSettingsOverridableParam>& OverridableParams = Settings->OverridableParams();
+
+	const UPCGGraphInterface* Graph = Settings->GetSubgraphInterface();
+	const FInstancedPropertyBag* UserParameters = Graph ? Graph->GetUserParametersStruct() : nullptr;
+	FConstStructView UserParametersView = UserParameters ? UserParameters->GetValue() : FConstStructView{};
+
+	if (!OverridableParams.IsEmpty() && UserParametersView.IsValid())
+	{
+		bool bHasParamConnected = InputData.GetParamsWithDeprecation(Node) != nullptr;
+
+		int32 Index = 0;
+		while (!bHasParamConnected && Index < OverridableParams.Num())
+		{
+			bHasParamConnected |= !InputData.GetParamsByPin(OverridableParams[Index++].Label).IsEmpty();
+		}
+
+		if (bHasParamConnected)
+		{
+			Context->GraphInstanceParametersOverride = FInstancedStruct(UserParametersView);
+		}
+	}
+
 	return Context;
 }
 
@@ -254,6 +346,9 @@ bool FPCGSubgraphElement::ExecuteInternal(FPCGContext* InContext) const
 {
 	FPCGSubgraphContext* Context = static_cast<FPCGSubgraphContext*>(InContext);
 
+	const UPCGSubgraphSettings* Settings = Context->GetInputSettings<UPCGSubgraphSettings>();
+	check(Settings);
+
 	const UPCGSubgraphNode* SubgraphNode = Cast<const UPCGSubgraphNode>(Context->Node);
 	if (SubgraphNode && SubgraphNode->bDynamicGraph)
 	{
@@ -261,10 +356,7 @@ bool FPCGSubgraphElement::ExecuteInternal(FPCGContext* InContext) const
 
 		if (!Context->bScheduledSubgraph)
 		{
-			const UPCGSubgraphSettings* Settings = Context->GetInputSettings<UPCGSubgraphSettings>();
-			check(Settings);
 			UPCGGraph* Subgraph = Settings->GetSubgraph();
-
 			UPCGSubsystem* Subsystem = Context->SourceComponent.IsValid() ? Context->SourceComponent->GetSubsystem() : nullptr;
 
 			if (Subsystem && Subgraph)
@@ -319,6 +411,25 @@ bool FPCGSubgraphElement::ExecuteInternal(FPCGContext* InContext) const
 	else
 	{
 		Context->OutputData = Context->InputData;
+
+		// Also create a new data containing information about the original subgraph and the parameters override
+		// It is used mainly by the UseParameterGetElement to access the correct value.
+		if (const UPCGGraphInterface* SubgraphInterface = Settings->GetSubgraphInterface())
+		{
+			UPCGUserParametersData* UserParamData = NewObject<UPCGUserParametersData>();
+			UserParamData->OriginalGraph = SubgraphInterface;
+			if (Context->GraphInstanceParametersOverride.IsValid())
+			{
+				UserParamData->UserParameters = std::move(Context->GraphInstanceParametersOverride);
+			}
+
+			FPCGTaggedData& TaggedData = Context->OutputData.TaggedData.Emplace_GetRef();
+			TaggedData.Data = UserParamData;
+			TaggedData.Tags.Add(PCGBaseSubgraphConstants::UserParameterTagData);
+			// Mark this data pinless, since it is internal data, not meant to be shown in the graph editor.
+			TaggedData.bPinlessData = true;
+		}
+
 		return true;
 	}
 }
diff --git a/Engine/Plugins/Experimental/PCG/Source/PCG/Private/Tests/PCGGraphParametersTest.cpp b/Engine/Plugins/Experimental/PCG/Source/PCG/Private/Tests/PCGGraphParametersTest.cpp
new file mode 100644
--- /dev/null
+++ b/Engine/Plugins/Experimental/PCG/Source/PCG/Private/Tests/PCGGraphParametersTest.cpp
@@ -0,0 +1,490 @@
+// Copyright Epic Games, Inc. All Rights Reserved.
+
+#if WITH_EDITOR
+
+#include "Tests/PCGTestsCommon.h"
+
+#include "PCGGraph.h"
+#include "PropertyBag.h"
+#include "StructView.h"
+
+IMPLEMENT_CUSTOM_SIMPLE_AUTOMATION_TEST(FPCGGraphNoUserParameters, FPCGTestBaseClass, "pcg.Graph.NoUserParameters", PCGTestsCommon::TestFlags)
+IMPLEMENT_CUSTOM_SIMPLE_AUTOMATION_TEST(FPCGGraphAddSingleUserParameters, FPCGTestBaseClass, "pcg.Graph.AddSingleUserParameter", PCGTestsCommon::TestFlags)
+IMPLEMENT_CUSTOM_SIMPLE_AUTOMATION_TEST(FPCGGraphAddSingleUserParametersPropagates, FPCGTestBaseClass, "pcg.Graph.AddSingleUserParameterPropagates", PCGTestsCommon::TestFlags)
+IMPLEMENT_CUSTOM_SIMPLE_AUTOMATION_TEST(FPCGGraphRemoveSingleUserParameters, FPCGTestBaseClass, "pcg.Graph.RemoveSingleUserParameter", PCGTestsCommon::TestFlags)
+IMPLEMENT_CUSTOM_SIMPLE_AUTOMATION_TEST(FPCGGraphRemoveSingleUserParametersPropagates, FPCGTestBaseClass, "pcg.Graph.RemoveSingleUserParameterPropagates", PCGTestsCommon::TestFlags)
+
+// We test the graph add/remove parameters just as a redundancy, but it is already tested by StructUtils. So we won't test more.
+// What is important is the propagation between graphs and graphs instances
+IMPLEMENT_CUSTOM_SIMPLE_AUTOMATION_TEST(FPCGGraphRenameUserParametersPropagates, FPCGTestBaseClass, "pcg.Graph.RenameUserParameterPropagates", PCGTestsCommon::TestFlags)
+IMPLEMENT_CUSTOM_SIMPLE_AUTOMATION_TEST(FPCGGraphChangeTypeUserParametersPropagates, FPCGTestBaseClass, "pcg.Graph.ChangeTypeUserParameterPropagates", PCGTestsCommon::TestFlags)
+IMPLEMENT_CUSTOM_SIMPLE_AUTOMATION_TEST(FPCGGraphRemovingGraphCleansUserParameters, FPCGTestBaseClass, "pcg.Graph.RemovingGraphCleansUserParameters", PCGTestsCommon::TestFlags)
+IMPLEMENT_CUSTOM_SIMPLE_AUTOMATION_TEST(FPCGGraphSetGraphUpdatesUserParameters, FPCGTestBaseClass, "pcg.Graph.SetGraphCleansUserParameters", PCGTestsCommon::TestFlags)
+IMPLEMENT_CUSTOM_SIMPLE_AUTOMATION_TEST(FPCGGraphSetValueUserParametersPropagates, FPCGTestBaseClass, "pcg.Graph.SetValueUserParametersPropagates", PCGTestsCommon::TestFlags)
+
+#define PCG_FMT(Text, ...) FString::Printf(TEXT(Text), __VA_ARGS__)
+
+namespace PCGTests
+{
+	template <typename T>
+	struct TempObject
+	{
+		template <typename... Args>
+		TempObject(Args&&... InArgs)
+		{
+			Ptr = NewObject<T>(std::forward<Args>(InArgs)...);
+		}
+
+		~TempObject()
+		{
+			if (Ptr)
+			{
+				Ptr->MarkAsGarbage();
+			}
+		}
+
+		operator T* () { return Ptr; }
+		operator const T* () const { return Ptr; }
+		T* operator->() { return Ptr; }
+		const T* operator->() const { return Ptr; }
+
+		T* Ptr = nullptr;
+	};
+
+	constexpr const TCHAR* UserParametersName = TEXT("UserParameters");
+
+	template <typename Func>
+	void EmulateModifyingUserParameters(UPCGGraph* Graph, Func Callback)
+	{
+		FProperty* Property = Graph->GetClass()->FindPropertyByName(UserParametersName);
+		check(Property);
+		Graph->PreEditChange(Property);
+
+		Callback();
+
+		FPropertyChangedEvent PropertyChangedEvent{Property};
+		Graph->PostEditChangeProperty(PropertyChangedEvent);
+	}
+
+	template <typename Func>
+	void EmulateModifyingUserParametersValue(UPCGGraph* Graph, FName PropertyName, Func Callback)
+	{
+		const FInstancedPropertyBag* UserParameters = Graph->GetUserParametersStruct();
+		check(UserParameters && UserParameters->GetPropertyBagStruct());
+
+		FProperty* ValueProperty = UserParameters->GetPropertyBagStruct()->FindPropertyByName(PropertyName);
+		FProperty* UserParametersProperty = Graph->GetClass()->FindPropertyByName(UserParametersName);
+		check(ValueProperty && UserParametersProperty);
+		Graph->PreEditChange(ValueProperty);
+
+		Callback();
+
+		FPropertyChangedEvent PropertyChangedEvent{ ValueProperty };
+		PropertyChangedEvent.MemberProperty = UserParametersProperty;
+		Graph->PostEditChangeProperty(PropertyChangedEvent);
+	}
+}
+
+bool FPCGGraphNoUserParameters::RunTest(const FString& Parameters)
+{
+	PCGTests::TempObject<UPCGGraph> Graph{};
+	UTEST_NOT_NULL("Graph UserParameters is not null", Graph->GetUserParametersStruct());
+
+	UTEST_FALSE("Graph UserParameters property bag is valid", Graph->GetUserParametersStruct()->IsValid());
+	UTEST_EQUAL("Graph UserParameters has no property", Graph->GetUserParametersStruct()->GetNumPropertiesInBag(), 0);
+
+	return true;
+}
+
+bool FPCGGraphAddSingleUserParameters::RunTest(const FString& Parameters)
+{
+	PCGTests::TempObject<UPCGGraph> Graph{};
+
+	const FName MyPropertyName = TEXT("MyProperty");
+
+	// Creating a new property MyProperty, which is a double
+	PCGTests::EmulateModifyingUserParameters(Graph, [Graph, MyPropertyName]()
+	{
+		FInstancedPropertyBag* UserParameters = const_cast<FInstancedPropertyBag*>(Graph->GetUserParametersStruct());
+		UserParameters->AddProperty(MyPropertyName, EPropertyBagPropertyType::Double, nullptr);
+	});
+
+	UTEST_NOT_NULL("Graph UserParameters is not null", Graph->GetUserParametersStruct());
+
+	UTEST_TRUE("Graph UserParameters property bag is valid", Graph->GetUserParametersStruct()->IsValid());
+	UTEST_EQUAL("Graph UserParameters has 1 property", Graph->GetUserParametersStruct()->GetNumPropertiesInBag(), 1);
+
+	const FPropertyBagPropertyDesc* PropertyDesc = Graph->GetUserParametersStruct()->FindPropertyDescByName(MyPropertyName);
+	UTEST_NOT_NULL("Property exists", PropertyDesc);
+
+	UTEST_TRUE("Property is a double", PropertyDesc->IsNumericFloatType());
+
+	return true;
+}
+
+bool FPCGGraphAddSingleUserParametersPropagates::RunTest(const FString& Parameters)
+{
+	PCGTests::TempObject<UPCGGraph> Graph{};
+	PCGTests::TempObject<UPCGGraphInstance> GraphInstance{};
+	PCGTests::TempObject<UPCGGraphInstance> GraphInstanceChild{};
+
+	GraphInstance->SetGraph(Graph);
+	GraphInstanceChild->SetGraph(GraphInstance);
+
+	UTEST_EQUAL("Graph and GraphInstance have the same graph", GraphInstance->GetGraph(), (UPCGGraph*)Graph);
+	UTEST_EQUAL("GraphInstance and GraphInstanceChild have the same graph", GraphInstanceChild->GetGraph(), GraphInstance->GetGraph());
+
+	const FName MyPropertyName = TEXT("MyProperty");
+
+	// Creating a new property MyProperty, which is a double
+	PCGTests::EmulateModifyingUserParameters(Graph, [Graph, MyPropertyName]()
+	{
+		FInstancedPropertyBag* UserParameters = const_cast<FInstancedPropertyBag*>(Graph->GetUserParametersStruct());
+		UserParameters->AddProperty(MyPropertyName, EPropertyBagPropertyType::Double, nullptr);
+	});
+
+	auto Verification = [this, MyPropertyName](UPCGGraphInstance* GraphInstance, const TCHAR* Name) -> bool
+	{
+		UTEST_NOT_NULL(PCG_FMT("%s UserParameters is not null", Name), GraphInstance->GetUserParametersStruct());
+
+		UTEST_TRUE(PCG_FMT("%s UserParameters property bag is valid", Name), GraphInstance->GetUserParametersStruct()->IsValid());
+		UTEST_EQUAL(PCG_FMT("%s UserParameters has 1 property", Name), GraphInstance->GetUserParametersStruct()->GetNumPropertiesInBag(), 1);
+
+		const FPropertyBagPropertyDesc* PropertyDesc = GraphInstance->GetUserParametersStruct()->FindPropertyDescByName(MyPropertyName);
+		UTEST_NOT_NULL(PCG_FMT("%s: Property exists", Name), PropertyDesc);
+
+		UTEST_TRUE(PCG_FMT("%s: Property is a double", Name), PropertyDesc->IsNumericFloatType());
+
+		return true;
+	};
+
+	if (!Verification(GraphInstance, TEXT("GraphInstance"))) { return false; }
+	if (!Verification(GraphInstanceChild, TEXT("GraphInstanceChild"))) { return false; }
+
+	return true;
+}
+
+bool FPCGGraphRemoveSingleUserParameters::RunTest(const FString& Parameters)
+{
+	PCGTests::TempObject<UPCGGraph> Graph{};
+
+	const FName MyPropertyName = TEXT("MyProperty");
+
+	// Creating a new property MyProperty, which is a double
+	PCGTests::EmulateModifyingUserParameters(Graph, [Graph, MyPropertyName]()
+	{
+		FInstancedPropertyBag* UserParameters = const_cast<FInstancedPropertyBag*>(Graph->GetUserParametersStruct());
+		UserParameters->AddProperty(MyPropertyName, EPropertyBagPropertyType::Double, nullptr);
+	});
+
+	// Remove that same property
+	PCGTests::EmulateModifyingUserParameters(Graph, [Graph, MyPropertyName]()
+	{
+		FInstancedPropertyBag* UserParameters = const_cast<FInstancedPropertyBag*>(Graph->GetUserParametersStruct());
+		UserParameters->RemovePropertyByName(MyPropertyName);
+	});
+
+	UTEST_NOT_NULL("Graph UserParameters is not null", Graph->GetUserParametersStruct());
+
+	UTEST_TRUE("Graph UserParameters property bag is valid", Graph->GetUserParametersStruct()->IsValid());
+	UTEST_EQUAL("Graph UserParameters has no property", Graph->GetUserParametersStruct()->GetNumPropertiesInBag(), 0);
+
+	return true;
+}
+
+bool FPCGGraphRemoveSingleUserParametersPropagates::RunTest(const FString& Parameters)
+{
+	PCGTests::TempObject<UPCGGraph> Graph{};
+	PCGTests::TempObject<UPCGGraphInstance> GraphInstance{};
+	PCGTests::TempObject<UPCGGraphInstance> GraphInstanceChild{};
+
+	GraphInstance->SetGraph(Graph);
+	GraphInstanceChild->SetGraph(GraphInstance);
+
+	UTEST_EQUAL("Graph and GraphInstance have the same graph", GraphInstance->GetGraph(), (UPCGGraph*)Graph);
+	UTEST_EQUAL("GraphInstance and GraphInstanceChild have the same graph", GraphInstanceChild->GetGraph(), GraphInstance->GetGraph());
+
+	const FName MyPropertyName = TEXT("MyProperty");
+
+	// Creating a new property MyProperty, which is a double
+	PCGTests::EmulateModifyingUserParameters(Graph, [Graph, MyPropertyName]()
+	{
+		FInstancedPropertyBag* UserParameters = const_cast<FInstancedPropertyBag*>(Graph->GetUserParametersStruct());
+		UserParameters->AddProperty(MyPropertyName, EPropertyBagPropertyType::Double, nullptr);
+	});
+
+	// Remove that same property
+	PCGTests::EmulateModifyingUserParameters(Graph, [Graph, MyPropertyName]()
+	{
+		FInstancedPropertyBag* UserParameters = const_cast<FInstancedPropertyBag*>(Graph->GetUserParametersStruct());
+		UserParameters->RemovePropertyByName(MyPropertyName);
+	});
+
+	auto Verification = [this, MyPropertyName](UPCGGraphInstance* GraphInstance, const TCHAR* Name) -> bool
+	{
+		UTEST_NOT_NULL(PCG_FMT("%s UserParameters is not null", Name), GraphInstance->GetUserParametersStruct());
+
+		UTEST_TRUE(PCG_FMT("%s UserParameters property bag is valid", Name), GraphInstance->GetUserParametersStruct()->IsValid());
+		UTEST_EQUAL(PCG_FMT("%s UserParameters has 0 property", Name), GraphInstance->GetUserParametersStruct()->GetNumPropertiesInBag(), 0);
+
+		return true;
+	};
+
+	if (!Verification(GraphInstance, TEXT("GraphInstance"))) { return false; }
+	if (!Verification(GraphInstanceChild, TEXT("GraphInstanceChild"))) { return false; }
+
+	return true;
+}
+
+bool FPCGGraphRenameUserParametersPropagates::RunTest(const FString& Parameters)
+{
+	PCGTests::TempObject<UPCGGraph> Graph{};
+	PCGTests::TempObject<UPCGGraphInstance> GraphInstance{};
+	PCGTests::TempObject<UPCGGraphInstance> GraphInstanceChild{};
+
+	GraphInstance->SetGraph(Graph);
+	GraphInstanceChild->SetGraph(GraphInstance);
+
+	UTEST_EQUAL("Graph and GraphInstance have the same graph", GraphInstance->GetGraph(), (UPCGGraph*)Graph);
+	UTEST_EQUAL("GraphInstance and GraphInstanceChild have the same graph", GraphInstanceChild->GetGraph(), GraphInstance->GetGraph());
+
+	const FName MyPropertyName = TEXT("MyProperty");
+	const FName MyNewPropertyName = TEXT("MyNewProperty");
+
+	// Creating a new property MyProperty, which is a double
+	PCGTests::EmulateModifyingUserParameters(Graph, [Graph, MyPropertyName]()
+	{
+		FInstancedPropertyBag* UserParameters = const_cast<FInstancedPropertyBag*>(Graph->GetUserParametersStruct());
+		UserParameters->AddProperty(MyPropertyName, EPropertyBagPropertyType::Double, nullptr);
+	});
+
+	// Rename it to MyNewProperty, while keeping the type Double
+	PCGTests::EmulateModifyingUserParameters(Graph, [Graph, MyNewPropertyName]()
+	{
+		FInstancedPropertyBag* UserParameters = const_cast<FInstancedPropertyBag*>(Graph->GetUserParametersStruct());
+
+		const UPropertyBag* NewPropertyBag = UPropertyBag::GetOrCreateFromDescs({ FPropertyBagPropertyDesc{MyNewPropertyName, EPropertyBagPropertyType::Double, nullptr} });
+		UserParameters->MigrateToNewBagStruct(NewPropertyBag);
+	});
+
+	auto Verification = [this, MyPropertyName, MyNewPropertyName](UPCGGraphInstance* GraphInstance, const TCHAR* Name) -> bool
+	{
+		UTEST_NOT_NULL(PCG_FMT("%s UserParameters is not null", Name), GraphInstance->GetUserParametersStruct());
+
+		UTEST_TRUE(PCG_FMT("%s UserParameters property bag is valid", Name), GraphInstance->GetUserParametersStruct()->IsValid());
+		UTEST_EQUAL(PCG_FMT("%s UserParameters has 1 property", Name), GraphInstance->GetUserParametersStruct()->GetNumPropertiesInBag(), 1);
+
+		const FPropertyBagPropertyDesc* PropertyDesc = GraphInstance->GetUserParametersStruct()->FindPropertyDescByName(MyPropertyName);
+		UTEST_NULL(PCG_FMT("%s: Old property exists", Name), PropertyDesc);
+
+		PropertyDesc = GraphInstance->GetUserParametersStruct()->FindPropertyDescByName(MyNewPropertyName);
+		UTEST_NOT_NULL(PCG_FMT("%s: New property exists", Name), PropertyDesc);
+
+		UTEST_TRUE(PCG_FMT("%s: Property is a double", Name), PropertyDesc->IsNumericFloatType());
+
+		return true;
+	};
+
+	if (!Verification(GraphInstance, TEXT("GraphInstance"))) { return false; }
+	if (!Verification(GraphInstanceChild, TEXT("GraphInstanceChild"))) { return false; }
+
+	return true;
+}
+
+bool FPCGGraphChangeTypeUserParametersPropagates::RunTest(const FString& Parameters)
+{
+	PCGTests::TempObject<UPCGGraph> Graph{};
+	PCGTests::TempObject<UPCGGraphInstance> GraphInstance{};
+	PCGTests::TempObject<UPCGGraphInstance> GraphInstanceChild{};
+
+	GraphInstance->SetGraph(Graph);
+	GraphInstanceChild->SetGraph(GraphInstance);
+
+	UTEST_EQUAL("Graph and GraphInstance have the same graph", GraphInstance->GetGraph(), (UPCGGraph*)Graph);
+	UTEST_EQUAL("GraphInstance and GraphInstanceChild have the same graph", GraphInstanceChild->GetGraph(), GraphInstance->GetGraph());
+
+	const FName MyPropertyName = TEXT("MyProperty");
+
+	// Creating a new property MyProperty, which is a double
+	PCGTests::EmulateModifyingUserParameters(Graph, [Graph, MyPropertyName]()
+	{
+		FInstancedPropertyBag* UserParameters = const_cast<FInstancedPropertyBag*>(Graph->GetUserParametersStruct());
+		UserParameters->AddProperty(MyPropertyName, EPropertyBagPropertyType::Double, nullptr);
+	});
+
+	// Keeping the name MyProperty but changing the type to Int64
+	PCGTests::EmulateModifyingUserParameters(Graph, [Graph, MyPropertyName]()
+	{
+		FInstancedPropertyBag* UserParameters = const_cast<FInstancedPropertyBag*>(Graph->GetUserParametersStruct());
+
+		const UPropertyBag* NewPropertyBag = UPropertyBag::GetOrCreateFromDescs({ FPropertyBagPropertyDesc{MyPropertyName, EPropertyBagPropertyType::Int64, nullptr} });
+		UserParameters->MigrateToNewBagStruct(NewPropertyBag);
+	});
+
+	auto Verification = [this, MyPropertyName](UPCGGraphInstance* GraphInstance, const TCHAR* Name) -> bool
+	{
+		UTEST_NOT_NULL(PCG_FMT("%s UserParameters is not null", Name), GraphInstance->GetUserParametersStruct());
+
+		UTEST_TRUE(PCG_FMT("%s UserParameters property bag is valid", Name), GraphInstance->GetUserParametersStruct()->IsValid());
+		UTEST_EQUAL(PCG_FMT("%s UserParameters has 1 property", Name), GraphInstance->GetUserParametersStruct()->GetNumPropertiesInBag(), 1);
+
+		const FPropertyBagPropertyDesc* PropertyDesc = GraphInstance->GetUserParametersStruct()->FindPropertyDescByName(MyPropertyName);
+		UTEST_NOT_NULL(PCG_FMT("%s: Old property still exists", Name), PropertyDesc);
+
+		UTEST_TRUE(PCG_FMT("%s: Property is a int", Name), PropertyDesc->IsNumericType() && !PropertyDesc->IsNumericFloatType());
+
+		return true;
+	};
+
+	if (!Verification(GraphInstance, TEXT("GraphInstance"))) { return false; }
+	if (!Verification(GraphInstanceChild, TEXT("GraphInstanceChild"))) { return false; }
+
+	return true;
+}
+
+bool FPCGGraphRemovingGraphCleansUserParameters::RunTest(const FString& Parameters)
+{
+	PCGTests::TempObject<UPCGGraph> Graph{};
+	PCGTests::TempObject<UPCGGraphInstance> GraphInstance{};
+
+	GraphInstance->SetGraph(Graph);
+
+	UTEST_EQUAL("Graph and GraphInstance have the same graph", GraphInstance->GetGraph(), (UPCGGraph*)Graph);
+
+	const FName MyPropertyName = TEXT("MyProperty");
+	const FName MyPropertyName2 = TEXT("MyProperty2");
+
+	// Creating a new property MyProperty, which is a double
+	PCGTests::EmulateModifyingUserParameters(Graph, [Graph, MyPropertyName]()
+	{
+		FInstancedPropertyBag* UserParameters = const_cast<FInstancedPropertyBag*>(Graph->GetUserParametersStruct());
+		UserParameters->AddProperty(MyPropertyName, EPropertyBagPropertyType::Double, nullptr);
+	});
+
+	// Creating a new property MyProperty2, which is a int64
+	PCGTests::EmulateModifyingUserParameters(Graph, [Graph, MyPropertyName2]()
+	{
+		FInstancedPropertyBag* UserParameters = const_cast<FInstancedPropertyBag*>(Graph->GetUserParametersStruct());
+		UserParameters->AddProperty(MyPropertyName2, EPropertyBagPropertyType::Int64, nullptr);
+	});
+
+	GraphInstance->SetGraph(nullptr);
+
+	UTEST_NOT_NULL("UserParameters is not null", GraphInstance->GetUserParametersStruct());
+
+	UTEST_FALSE("UserParameters property bag is valid", GraphInstance->GetUserParametersStruct()->IsValid());
+	UTEST_EQUAL("UserParameters has no property", GraphInstance->GetUserParametersStruct()->GetNumPropertiesInBag(), 0);
+
+	return true;
+}
+
+bool FPCGGraphSetGraphUpdatesUserParameters::RunTest(const FString& Parameters)
+{
+	PCGTests::TempObject<UPCGGraph> Graph{};
+	PCGTests::TempObject<UPCGGraphInstance> GraphInstance{};
+
+	const FName MyPropertyName = TEXT("MyProperty");
+	const FName MyPropertyName2 = TEXT("MyProperty2");
+
+	// Creating a new property MyProperty, which is a double
+	PCGTests::EmulateModifyingUserParameters(Graph, [Graph, MyPropertyName]()
+	{
+		FInstancedPropertyBag* UserParameters = const_cast<FInstancedPropertyBag*>(Graph->GetUserParametersStruct());
+		UserParameters->AddProperty(MyPropertyName, EPropertyBagPropertyType::Double, nullptr);
+	});
+
+	// Creating a new property MyProperty2, which is a int64
+	PCGTests::EmulateModifyingUserParameters(Graph, [Graph, MyPropertyName2]()
+	{
+		FInstancedPropertyBag* UserParameters = const_cast<FInstancedPropertyBag*>(Graph->GetUserParametersStruct());
+		UserParameters->AddProperty(MyPropertyName2, EPropertyBagPropertyType::Int64, nullptr);
+	});
+
+	GraphInstance->SetGraph(Graph);
+
+	UTEST_EQUAL("Graph and GraphInstance have the same graph", GraphInstance->GetGraph(), (UPCGGraph*)Graph);
+
+	UTEST_NOT_NULL("UserParameters is not null", GraphInstance->GetUserParametersStruct());
+
+	UTEST_TRUE("UserParameters property bag is valid", GraphInstance->GetUserParametersStruct()->IsValid());
+	UTEST_EQUAL("UserParameters has 2 properties", GraphInstance->GetUserParametersStruct()->GetNumPropertiesInBag(), 2);
+
+	const FPropertyBagPropertyDesc* PropertyDesc = Graph->GetUserParametersStruct()->FindPropertyDescByName(MyPropertyName);
+	UTEST_NOT_NULL("First property exists", PropertyDesc);
+
+	UTEST_TRUE("First property is a double", PropertyDesc->IsNumericFloatType());
+
+	PropertyDesc = Graph->GetUserParametersStruct()->FindPropertyDescByName(MyPropertyName2);
+	UTEST_NOT_NULL("Second property exists", PropertyDesc);
+
+	UTEST_TRUE("First property is an int", PropertyDesc->IsNumericType() && !PropertyDesc->IsNumericFloatType());
+
+	return true;
+}
+
+bool FPCGGraphSetValueUserParametersPropagates::RunTest(const FString& Parameters)
+{
+	PCGTests::TempObject<UPCGGraph> Graph{};
+	PCGTests::TempObject<UPCGGraphInstance> GraphInstance{};
+	PCGTests::TempObject<UPCGGraphInstance> GraphInstanceChild{};
+
+	GraphInstance->SetGraph(Graph);
+	GraphInstanceChild->SetGraph(GraphInstance);
+
+	UTEST_EQUAL("Graph and GraphInstance have the same graph", GraphInstance->GetGraph(), (UPCGGraph*)Graph);
+	UTEST_EQUAL("GraphInstance and GraphInstanceChild have the same graph", GraphInstanceChild->GetGraph(), GraphInstance->GetGraph());
+
+	const FName MyPropertyName = TEXT("MyProperty");
+	const FName MyPropertyName2 = TEXT("MyProperty2");
+
+	// Creating a new property MyProperty, which is a double
+	PCGTests::EmulateModifyingUserParameters(Graph, [Graph, MyPropertyName]()
+	{
+		FInstancedPropertyBag* UserParameters = const_cast<FInstancedPropertyBag*>(Graph->GetUserParametersStruct());
+		UserParameters->AddProperty(MyPropertyName, EPropertyBagPropertyType::Double, nullptr);
+	});
+
+	// Creating another property MyProperty2, which is a double
+	PCGTests::EmulateModifyingUserParameters(Graph, [Graph, MyPropertyName2]()
+	{
+		FInstancedPropertyBag* UserParameters = const_cast<FInstancedPropertyBag*>(Graph->GetUserParametersStruct());
+		UserParameters->AddProperty(MyPropertyName2, EPropertyBagPropertyType::Int64, nullptr);
+	});
+
+	// Setting MyProperty to 3.0
+	constexpr double NewValue = 3.0;
+	PCGTests::EmulateModifyingUserParametersValue(Graph, MyPropertyName, [Graph, MyPropertyName, NewValue]()
+	{
+		const FPropertyBagPropertyDesc* PropertyDesc = Graph->GetUserParametersStruct()->FindPropertyDescByName(MyPropertyName);
+		FInstancedPropertyBag* UserParameters = const_cast<FInstancedPropertyBag*>(Graph->GetUserParametersStruct());
+
+		PropertyDesc->CachedProperty->SetValue_InContainer(UserParameters->GetMutableValue().GetMutableMemory(), &NewValue);
+	});
+
+	auto Verification = [this, MyPropertyName, NewValue](UPCGGraphInstance* GraphInstance, const TCHAR* Name) -> bool
+	{
+		UTEST_NOT_NULL(PCG_FMT("%s UserParameters is not null", Name), GraphInstance->GetUserParametersStruct());
+
+		UTEST_TRUE(PCG_FMT("%s UserParameters property bag is valid", Name), GraphInstance->GetUserParametersStruct()->IsValid());
+		UTEST_EQUAL(PCG_FMT("%s UserParameters has 2 properties", Name), GraphInstance->GetUserParametersStruct()->GetNumPropertiesInBag(), 2);
+
+		const FPropertyBagPropertyDesc* PropertyDesc = GraphInstance->GetUserParametersStruct()->FindPropertyDescByName(MyPropertyName);
+		UTEST_NOT_NULL(PCG_FMT("%s: First property exists", Name), PropertyDesc);
+
+		double Value = 0.0;
+		PropertyDesc->CachedProperty->GetValue_InContainer(GraphInstance->GetUserParametersStruct()->GetValue().GetMemory(), &Value);
+
+		UTEST_EQUAL(PCG_FMT("%s: Property has the right value", Name), Value, NewValue);
+
+		return true;
+	};
+
+	if (!Verification(GraphInstance, TEXT("GraphInstance"))) { return false; }
+	if (!Verification(GraphInstanceChild, TEXT("GraphInstanceChild"))) { return false; }
+
+	return true;
+}
+
+#undef PCG_FMT
+
+#endif // WITH_EDITOR
\ No newline at end of file
diff --git a/Engine/Plugins/Experimental/PCG/Source/PCG/Public/Data/PCGUserParametersData.h b/Engine/Plugins/Experimental/PCG/Source/PCG/Public/Data/PCGUserParametersData.h
new file mode 100644
--- /dev/null
+++ b/Engine/Plugins/Experimental/PCG/Source/PCG/Public/Data/PCGUserParametersData.h
@@ -0,0 +1,36 @@
+// Copyright Epic Games, Inc. All Rights Reserved.
+
+#pragma once
+
+#include "PCGData.h"
+
+#include "InstancedStruct.h"
+#include "UObject/ObjectPtr.h"
+
+#include "PCGUserParametersData.generated.h"
+
+class UPCGGraphInterface;
+
+/**
+* PCG Data meant only to be used internally.
+* It contains a copy of UserParameters for a given graph instance, with overrides in it.
+* The idea is to have a structure to hold our overrides, provided by the override pins on the Subgraph
+* and use this as input to PCGUserParametersGetElement. By doing so, we are able to provide the right
+* parameters to the getter node.
+*/
+UCLASS(BlueprintType, HideDropdown, ClassGroup = (Procedural))
+class UPCGUserParametersData : public UPCGData
+{
+	GENERATED_BODY()
+
+public:
+	// ~Begin UPCGData interface
+	virtual EPCGDataType GetDataType() const override { return EPCGDataType::Other; }
+	// ~End UPCGData interface
+
+	UPROPERTY()
+	TObjectPtr<const UPCGGraphInterface> OriginalGraph = nullptr;
+
+	UPROPERTY()
+	FInstancedStruct UserParameters;
+};
\ No newline at end of file
diff --git a/Engine/Plugins/Experimental/PCG/Source/PCG/Public/Elements/PCGUserParameterGet.h b/Engine/Plugins/Experimental/PCG/Source/PCG/Public/Elements/PCGUserParameterGet.h
new file mode 100644
--- /dev/null
+++ b/Engine/Plugins/Experimental/PCG/Source/PCG/Public/Elements/PCGUserParameterGet.h
@@ -0,0 +1,53 @@
+// Copyright Epic Games, Inc. All Rights Reserved.
+
+#pragma once
+
+#include "CoreMinimal.h"
+
+#include "PCGElement.h"
+#include "PCGPin.h"
+#include "PCGSettings.h"
+
+#include "PCGUserParameterGet.generated.h"
+
+/**
+* Getter for user parameters defined in PCGGraph, by the user.
+* Will pick up the value from the graph instance.
+*/
+UCLASS(BlueprintType, ClassGroup = (Procedural))
+class PCG_API UPCGUserParameterGetSettings : public UPCGSettings
+{
+	GENERATED_BODY()
+
+public:
+	UPCGUserParameterGetSettings(const FObjectInitializer& ObjectInitializer);
+
+	UPROPERTY()
+	FGuid PropertyGuid;
+
+	UPROPERTY()
+	FName PropertyName;
+
+	void UpdatePropertyName(FName InNewName);
+
+	//~Begin UPCGSettings interface
+	virtual bool ShouldHookToPreTask() const override { return true; }
+protected:
+	virtual TArray<FPCGPinProperties> InputPinProperties() const override { return {}; }
+	virtual TArray<FPCGPinProperties> OutputPinProperties() const override;
+	virtual FPCGElementPtr CreateElement() const override;
+
+#if WITH_EDITOR
+	virtual FName GetDefaultNodeName() const override { return FName("Get User Parameter"); }
+#endif
+	//~End UPCGSettings interface
+};
+
+class PCG_API FPCGUserParameterGetElement : public FSimplePCGElement
+{
+public:
+	virtual bool IsCacheable(const UPCGSettings* InSettings) const override { return false; }
+
+protected:
+	virtual bool ExecuteInternal(FPCGContext* Context) const override;
+};
diff --git a/Engine/Plugins/Experimental/PCG/Source/PCG/Public/PCGComponent.h b/Engine/Plugins/Experimental/PCG/Source/PCG/Public/PCGComponent.h
--- a/Engine/Plugins/Experimental/PCG/Source/PCG/Public/PCGComponent.h
+++ b/Engine/Plugins/Experimental/PCG/Source/PCG/Public/PCGComponent.h
@@ -108,10 +108,10 @@ public:
 	UPCGGraph* GetGraph() const;
 	UPCGGraphInstance* GetGraphInstance() const { return GraphInstance; }
 
-	void SetGraphLocal(UPCGGraph* InGraph);
+	void SetGraphLocal(UPCGGraphInterface* InGraph);
 
 	UFUNCTION(BlueprintCallable, NetMulticast, Reliable, Category = PCG)
-	void SetGraph(UPCGGraph* InGraph);
+	void SetGraph(UPCGGraphInterface* InGraph);
 
 	/** Registers some managed resource to the current component */
 	UFUNCTION(BlueprintCallable, Category = PCG)
@@ -274,7 +274,7 @@ public:
 	static UPCGData* CreateActorPCGData(AActor* Actor, const UPCGComponent* Component, bool bParseActor = true);
 
 protected:
-	UPROPERTY(BlueprintReadOnly, VisibleAnywhere, Category = PCG, Instanced, meta = (NoResetToDefault, ShowOnlyInnerProperties))
+	UPROPERTY(BlueprintReadOnly, VisibleAnywhere, Category = PCG, Instanced, meta = (NoResetToDefault))
 	TObjectPtr<UPCGGraphInstance> GraphInstance;
 
 #if WITH_EDITORONLY_DATA
diff --git a/Engine/Plugins/Experimental/PCG/Source/PCG/Public/PCGContext.h b/Engine/Plugins/Experimental/PCG/Source/PCG/Public/PCGContext.h
--- a/Engine/Plugins/Experimental/PCG/Source/PCG/Public/PCGContext.h
+++ b/Engine/Plugins/Experimental/PCG/Source/PCG/Public/PCGContext.h
@@ -9,6 +9,7 @@
 #include "PCGContext.generated.h"
 
 class UPCGComponent;
+class UPCGGraphInterface;
 class UPCGSettingsInterface;
 class UPCGSpatialData;
 struct FPCGGraphCache;
@@ -103,6 +104,7 @@ struct PCG_API FPCGContext
 
 protected:
 	virtual UObject* GetExternalContainerForOverridableParam(const FPCGSettingsOverridableParam& InParam) { return nullptr; }
+	virtual void* GetUnsafeExternalContainerForOverridableParam(const FPCGSettingsOverridableParam& InParam) { return nullptr; }
 
 private:
 	template<typename SettingsType>
diff --git a/Engine/Plugins/Experimental/PCG/Source/PCG/Public/PCGData.h b/Engine/Plugins/Experimental/PCG/Source/PCG/Public/PCGData.h
--- a/Engine/Plugins/Experimental/PCG/Source/PCG/Public/PCGData.h
+++ b/Engine/Plugins/Experimental/PCG/Source/PCG/Public/PCGData.h
@@ -72,6 +72,10 @@ struct PCG_API FPCGTaggedData
 	UPROPERTY(BlueprintReadWrite, EditAnywhere, Category = Data)
 	FName Pin = NAME_None;
 
+	// Special flag for data that are forwarded to other nodes, but without a pin. Useful for internal data.
+	UPROPERTY()
+	bool bPinlessData = false;
+
 	bool operator==(const FPCGTaggedData& Other) const;
 	bool operator!=(const FPCGTaggedData& Other) const;
 };
@@ -111,6 +115,10 @@ struct PCG_API FPCGDataCollection
 	/** Returns all params on a given pin */
 	TArray<FPCGTaggedData> GetParamsByPin(const FName& InPinLabel) const;
 
+	/** Returns all data in the collection with the given tag and given type */
+	template <typename PCGDataType>
+	TArray<FPCGTaggedData> GetTaggedTypedInputs(const FString& InTag) const;
+
 	UE_DEPRECATED(5.2, "GetParams is deprecated, please use GetParamsByPin or GetFirstParamsOnParamsPin.")
 	/** Returns the first params found in the collection */
 	UPCGParamData* GetParams() const;
@@ -167,6 +175,14 @@ inline const SettingsType* FPCGDataCollection::GetSettings() const
 	return MatchingData ? Cast<const SettingsType>(MatchingData->Data) : nullptr;
 }
 
+template <typename PCGDataType>
+inline TArray<FPCGTaggedData> FPCGDataCollection::GetTaggedTypedInputs(const FString& InTag) const
+{
+	return TaggedData.FilterByPredicate([&InTag](const FPCGTaggedData& Data) {
+		return Data.Tags.Contains(InTag) && Cast<PCGDataType>(Data.Data);
+	});
+}
+
 UCLASS()
 class PCG_API UPCGDataFunctionLibrary : public UBlueprintFunctionLibrary
 {
diff --git a/Engine/Plugins/Experimental/PCG/Source/PCG/Public/PCGGraph.h b/Engine/Plugins/Experimental/PCG/Source/PCG/Public/PCGGraph.h
--- a/Engine/Plugins/Experimental/PCG/Source/PCG/Public/PCGGraph.h
+++ b/Engine/Plugins/Experimental/PCG/Source/PCG/Public/PCGGraph.h
@@ -5,14 +5,63 @@
 #include "PCGNode.h"
 #include "PCGSettings.h"
 
+#include "PropertyBag.h"
 #include "UObject/ObjectPtr.h"
 
 #include "PCGGraph.generated.h"
 
+enum class EPCGGraphParameterEvent
+{
+	GraphChanged,
+	GraphPostLoad,
+	Added,
+	Removed,
+	PropertyModified,
+	ValueModifiedLocally,
+	ValueModifiedByParent
+};
+
 #if WITH_EDITOR
 	DECLARE_MULTICAST_DELEGATE_TwoParams(FOnPCGGraphChanged, UPCGGraphInterface* /*Graph*/, EPCGChangeType /*ChangeType*/);
+	DECLARE_MULTICAST_DELEGATE_ThreeParams(FOnPCGGraphParametersChanged, UPCGGraphInterface* /*Graph*/, EPCGGraphParameterEvent /*ChangeType*/, FName /*ChangedPropertyName*/);
 #endif // WITH_EDITOR
 
+/**
+* Extended version of FInstancedPropertyBag, to support overrides and have a custom UI for it
+* Must only be used with PCGGraphInstances.
+* TODO: Should be made generic and moved to ScriptUtils.
+*/
+USTRUCT()
+struct PCG_API FPCGOverrideInstancedPropertyBag
+{
+	GENERATED_BODY()
+
+public:
+	/** Add/Remove given property from overrides, and reset its value if it is removed. Returns true if the value was changed. */
+	bool UpdatePropertyOverride(const FProperty* InProperty, bool bMarkAsOverridden, const FInstancedPropertyBag* ParentUserParameters);
+
+	/** Return if the property is currently marked overridden. */
+	bool IsPropertyOverridden(const FProperty* InProperty) const;
+
+	/** Reset the struct. */
+	void Reset();
+
+	/** Return if the parameters are valid */
+	bool IsValid() { return Parameters.IsValid(); }
+
+	/** Will migrate to a new property bag instanced, and will remove porperties that doesn't exists anymore or have changed types. */
+	void MigrateToNewBagInstance(const FInstancedPropertyBag& NewBagInstance);
+
+	/** Handle the sync between the parent parameters and the overrides. Will return true if something changed. */
+	bool RefreshParameters(const FInstancedPropertyBag* ParentUserParameters, EPCGGraphParameterEvent InChangeType, FName InChangedPropertyName);
+
+	UPROPERTY(EditAnywhere, Category = "")
+	FInstancedPropertyBag Parameters;
+
+	UPROPERTY(VisibleAnywhere, Category = "", meta = (Hidden))
+	TSet<FGuid> PropertiesIDsOverridden;
+};
+
 UCLASS(Abstract)
 class PCG_API UPCGGraphInterface : public UObject
 {
@@ -30,6 +79,8 @@ public:
 	virtual UPCGGraph* GetGraph() PURE_VIRTUAL(UPCGGraphInterface::GetGraph, return nullptr;)
 	virtual const UPCGGraph* GetGraph() const PURE_VIRTUAL(UPCGGraphInterface::GetGraph, return nullptr;)
 
+	virtual const FInstancedPropertyBag* GetUserParametersStruct() const PURE_VIRTUAL(UPCGGraphInterface::GetUserParametersStruct, return nullptr;)
+
 	bool IsInstance() const;
 
 	/** A graph interface is equivalent to another graph interface if they are the same (same ptr), or if they have the same graph. Will be overriden when graph instance supports overrides. */
@@ -37,6 +88,7 @@ public:
 
 #if WITH_EDITOR
 	FOnPCGGraphChanged OnGraphChangedDelegate;
+	FOnPCGGraphParametersChanged OnGraphParametersChangedDelegate;
 #endif // WITH_EDITOR
 };
 
@@ -60,6 +112,7 @@ public:
 	virtual void BeginDestroy() override;
 
 #if WITH_EDITOR
+	virtual void PreEditChange(FProperty* InProperty) override;
 	virtual void PostEditChangeProperty(FPropertyChangedEvent& PropertyChangedEvent) override;
 #endif
 	/** ~End UObject interface */
@@ -177,16 +230,28 @@ protected:
 	TArray<TObjectPtr<UObject>> ExtraEditorNodes;
 #endif // WITH_EDITORONLY_DATA
 
+	// Parameters
+	UPROPERTY(EditAnywhere, Category = Instance)
+	FInstancedPropertyBag UserParameters;
+
+public:
+	virtual const FInstancedPropertyBag* GetUserParametersStruct() const override { return &UserParameters; }
+
 #if WITH_EDITOR
 private:
 	void NotifyGraphChanged(EPCGChangeType ChangeType);
 	void OnNodeChanged(UPCGNode* InNode, EPCGChangeType ChangeType);
 
+	void NotifyGraphParametersChanged(EPCGGraphParameterEvent InChangeType, FName InChangedPropertyName);
+	void OnGraphParametersChanged(EPCGGraphParameterEvent InChangeType, FName InChangedPropertyName);
+
 	int32 GraphChangeNotificationsDisableCounter = 0;
 	bool bDelayedChangeNotification = false;
 	EPCGChangeType DelayedChangeType = EPCGChangeType::None;
 	bool bIsNotifying = false;
 	bool bUserPausedNotificationsInGraphEditor = false;
+	int32 NumberOfUserParametersPreEdit = 0;
+	FName UserParameterModifiedName = NAME_None;
 #endif // WITH_EDITOR
 };
 
@@ -216,20 +281,33 @@ public:
 protected:
 #if WITH_EDITOR
 	void OnGraphChanged(UPCGGraphInterface* InGraph, EPCGChangeType ChangeType);
+	void NotifyGraphParametersChanged(EPCGGraphParameterEvent InChangeType, FName InChangedPropertyName);
+	void OnGraphParametersChanged(UPCGGraphInterface* InGraph, EPCGGraphParameterEvent InChangeType, FName InChangedPropertyName);
 #endif
+	void RefreshParameters(EPCGGraphParameterEvent InChangeType, FName InChangedPropertyName = NAME_None);
 
 public:
 	static TObjectPtr<UPCGGraphInterface> CreateInstance(UObject* InOwner, UPCGGraphInterface* InGraph);
 
 	void SetGraph(UPCGGraphInterface* InGraph);
+	void CopyParameterOverrides(UPCGGraphInterface* InGraph);
+	void UpdatePropertyOverride(const FProperty* InProperty, bool bMarkAsOverridden);
+	bool IsPropertyOverridden(const FProperty* InProperty) { return ParametersOverrides.IsPropertyOverridden(InProperty); }
 
 	UPROPERTY(EditAnywhere, BlueprintReadOnly, Category = Instance)
 	TObjectPtr<UPCGGraphInterface> Graph;
 
-#if WITH_EDITORONLY_DATA
+	UPROPERTY(EditAnywhere, Category = Instance)
+	FPCGOverrideInstancedPropertyBag ParametersOverrides;
+
+	virtual const FInstancedPropertyBag* GetUserParametersStruct() const override { return &ParametersOverrides.Parameters; }
+
 private:
+#if WITH_EDITORONLY_DATA
 	// Transient, to keep track the undo/redo changed the graph.
 	UPCGGraphInterface* UndoRedoGraphCache = nullptr;
+
+	FName UserParameterModifiedName = NAME_None;
 #endif // WITH_EDITORONLY_DATA
 };
 
diff --git a/Engine/Plugins/Experimental/PCG/Source/PCG/Public/PCGNode.h b/Engine/Plugins/Experimental/PCG/Source/PCG/Public/PCGNode.h
--- a/Engine/Plugins/Experimental/PCG/Source/PCG/Public/PCGNode.h
+++ b/Engine/Plugins/Experimental/PCG/Source/PCG/Public/PCGNode.h
@@ -102,6 +102,10 @@ public:
 	const UPCGPin* GetOutputPin(const FName& Label) const;
 	bool HasInboundEdges() const;
 
+	/** Allow to change the name of a pin, to keep edges connected. You need to make sure that the underlying settings are also updated, otherwise, it will be overwritten next time the settings are updated */
+	void RenameInputPin(const FName& OldLabel, const FName& NewLabel);
+	void RenameOutputPin(const FName& OldLabel, const FName& NewLabel);
+
 	/** Pin from which data is passed through when this node is disabled. */
 	virtual const UPCGPin* GetPassThroughInputPin() const;
 
diff --git a/Engine/Plugins/Experimental/PCG/Source/PCG/Public/PCGSettings.h b/Engine/Plugins/Experimental/PCG/Source/PCG/Public/PCGSettings.h
--- a/Engine/Plugins/Experimental/PCG/Source/PCG/Public/PCGSettings.h
+++ b/Engine/Plugins/Experimental/PCG/Source/PCG/Public/PCGSettings.h
@@ -268,8 +268,12 @@ private:
 public:
 	/** List of all the overridable params available for this settings. */
 	virtual const TArray<FPCGSettingsOverridableParam>& OverridableParams() const { return CachedOverridableParams; }
+
 	/** Check if we have some override. Can be overriden to force params pin for example */
 	virtual bool HasOverridableParams() const { return !CachedOverridableParams.IsEmpty(); }
+
+	/** Check if we need to hook the output of the pre-task to this. One use is to compute overrides in the subgraph element and pass the overrides as data, to all nodes that needs it. */
+	virtual bool ShouldHookToPreTask() const { return false; }
 protected:
 	/** Iterate over OverridableParams to automatically add param pins to the list. */
 	void FillOverridableParamsPins(TArray<FPCGPinProperties>& OutPins) const;
diff --git a/Engine/Plugins/Experimental/PCG/Source/PCG/Public/PCGSubgraph.h b/Engine/Plugins/Experimental/PCG/Source/PCG/Public/PCGSubgraph.h
--- a/Engine/Plugins/Experimental/PCG/Source/PCG/Public/PCGSubgraph.h
+++ b/Engine/Plugins/Experimental/PCG/Source/PCG/Public/PCGSubgraph.h
@@ -10,6 +10,11 @@
 
 #include "PCGSubgraph.generated.h"
 
+namespace PCGBaseSubgraphConstants
+{
+	static const FString UserParameterTagData = TEXT("PCGUserParametersTagData");
+}
+
 UCLASS(Abstract)
 class PCG_API UPCGBaseSubgraphSettings : public UPCGSettings
 {
@@ -45,6 +50,13 @@ protected:
 #if WITH_EDITOR
 	void OnSubgraphChanged(UPCGGraphInterface* InGraph, EPCGChangeType ChangeType);
 #endif
+
+protected:
+	// Overrides
+	virtual void FixingOverridableParamPropertyClass(FPCGSettingsOverridableParam& Param) const;
+#if WITH_EDITOR
+	virtual TArray<FPCGSettingsOverridableParam> GatherOverridableParams() const override;
+#endif // WITH_EDITOR
 };
 
 UCLASS(BlueprintType, ClassGroup=(Procedural))
@@ -87,7 +99,7 @@ protected:
 	//~End UPCGBaseSubgraphSettings interface
 
 public:
-	UPROPERTY(BlueprintReadOnly, VisibleAnywhere, Category = Properties, Instanced, meta = (NoResetToDefault, ShowOnlyInnerProperties))
+	UPROPERTY(BlueprintReadOnly, VisibleAnywhere, Category = Properties, Instanced, meta = (NoResetToDefault))
 	TObjectPtr<UPCGGraphInstance> SubgraphInstance;
 
 #if WITH_EDITORONLY_DATA
@@ -124,6 +136,10 @@ struct PCG_API FPCGSubgraphContext : public FPCGContext
 {
 	FPCGTaskId SubgraphTaskId = InvalidPCGTaskId;
 	bool bScheduledSubgraph = false;
+	FInstancedStruct GraphInstanceParametersOverride;
+
+protected:
+	virtual void* GetUnsafeExternalContainerForOverridableParam(const FPCGSettingsOverridableParam& InParam) override;
 };
 
 class PCG_API FPCGSubgraphElement : public IPCGElement
diff --git a/Engine/Plugins/Experimental/PCG/Source/PCGEditor/PCGEditor.Build.cs b/Engine/Plugins/Experimental/PCG/Source/PCGEditor/PCGEditor.Build.cs
--- a/Engine/Plugins/Experimental/PCG/Source/PCGEditor/PCGEditor.Build.cs
+++ b/Engine/Plugins/Experimental/PCG/Source/PCGEditor/PCGEditor.Build.cs
@@ -40,6 +40,8 @@ namespace UnrealBuildTool.Rules
 					"KismetWidgets",
 					"PCG",
 					"PropertyEditor",
+					"StructUtils",
+					"StructUtilsEditor",
 					"Slate",
 					"SlateCore",
 					"SourceControl",
diff --git a/Engine/Plugins/Experimental/PCG/Source/PCGEditor/Private/Details/PCGGraphInstanceDetails.cpp b/Engine/Plugins/Experimental/PCG/Source/PCGEditor/Private/Details/PCGGraphInstanceDetails.cpp
--- a/Engine/Plugins/Experimental/PCG/Source/PCGEditor/Private/Details/PCGGraphInstanceDetails.cpp
+++ b/Engine/Plugins/Experimental/PCG/Source/PCGEditor/Private/Details/PCGGraphInstanceDetails.cpp
@@ -1,13 +1,21 @@
 // Copyright Epic Games, Inc. All Rights Reserved.
 
 #include "Details/PCGGraphInstanceDetails.h"
+
+#include "PCGEditorModule.h"
+#include "PCGEditorUtils.h"
 #include "PCGGraph.h"
+#include "PCGGraphFactory.h"
 
-#include "PropertyHandle.h"
-#include "DetailLayoutBuilder.h"
+#include "AssetToolsModule.h"
 #include "DetailCategoryBuilder.h"
+#include "DetailLayoutBuilder.h"
+#include "DetailWidgetRow.h"
+#include "EditorAssetLibrary.h"
+#include "Widgets/Text/STextBlock.h"
+#include "Widgets/Input/SButton.h"
 
-#define LOCTEXT_NAMESPACE "PCGGraphDetails"
+#define LOCTEXT_NAMESPACE "PCGGraphInstanceDetails"
 
 TSharedRef<IDetailCustomization> FPCGGraphInstanceDetails::MakeInstance()
 {
@@ -16,13 +24,15 @@ TSharedRef<IDetailCustomization> FPCGGraphInstanceDetails::MakeInstance()
 
 void FPCGGraphInstanceDetails::CustomizeDetails(IDetailLayoutBuilder& DetailBuilder)
 {
-	const FName PCGCategoryName("PCG");
-	IDetailCategoryBuilder& PCGCategory = DetailBuilder.EditCategory(PCGCategoryName);
+	const FName InstanceCategoryName("Instance");
+	IDetailCategoryBuilder& InstanceCategory = DetailBuilder.EditCategory(InstanceCategoryName);
+
+	FDetailWidgetRow& NewRow = InstanceCategory.AddCustomRow(FText::GetEmpty());
 
 	TArray<TWeakObjectPtr<UObject>> ObjectsBeingCustomized;
 	DetailBuilder.GetObjectsBeingCustomized(ObjectsBeingCustomized);
 
-	for (TWeakObjectPtr<UObject>& Object : ObjectsBeingCustomized)
+	for (const TWeakObjectPtr<UObject>& Object : ObjectsBeingCustomized)
 	{
 		UPCGGraphInstance* GraphInstance = Cast<UPCGGraphInstance>(Object.Get());
 		if (ensure(GraphInstance))
@@ -31,16 +41,66 @@ void FPCGGraphInstanceDetails::CustomizeDetails(IDetailLayoutBuilder& DetailBuil
 		}
 	}
 
-	TArray<TSharedRef<IPropertyHandle>> AllProperties;
-	bool bSimpleProperties = true;
-	bool bAdvancedProperties = false;
-	// Add all properties in the category in order
-	PCGCategory.GetDefaultProperties(AllProperties, bSimpleProperties, bAdvancedProperties);
+	if (SelectedGraphInstances.Num() == 1)
+	{
+		NewRow.ValueContent()
+		.MaxDesiredWidth(120.f)
+		[
+			SNew(SHorizontalBox)
+			+ SHorizontalBox::Slot()
+			.AutoWidth()
+			.Padding(4.0f)
+			[
+				SNew(SButton)
+				.OnClicked(this, &FPCGGraphInstanceDetails::OnSaveInstanceClicked)
+				.IsEnabled_Raw(this, &FPCGGraphInstanceDetails::SaveInstanceButtonEnabled)
+				[
+					SNew(STextBlock)
+					.Font(IDetailLayoutBuilder::GetDetailFont())
+					.Text(LOCTEXT("SaveInstanceButton", "Save Instance"))
+				]
+			]
+		];
+	}
+}
 
-	for (auto& Property : AllProperties)
+bool FPCGGraphInstanceDetails::SaveInstanceButtonEnabled() const
+{
+	const UPCGGraphInstance* GraphInstance = !SelectedGraphInstances.IsEmpty() ? SelectedGraphInstances[0].Get() : nullptr;
+	return (GraphInstance && GraphInstance->Graph);
+}
+
+FReply FPCGGraphInstanceDetails::OnSaveInstanceClicked()
+{
+	const UPCGGraphInstance* GraphInstance = !SelectedGraphInstances.IsEmpty() ? SelectedGraphInstances[0].Get() : nullptr;
+
+	if (!GraphInstance || !GraphInstance->Graph)
+	{
+		return FReply::Handled();
+	}
+
+	IAssetTools& AssetTools = FModuleManager::Get().LoadModuleChecked<FAssetToolsModule>("AssetTools").Get();
+	UPCGGraphInstanceFactory* Factory = NewObject<UPCGGraphInstanceFactory>();
+
+	FString NewPackageName;
+	FString NewAssetName;
+	PCGEditorUtils::GetParentPackagePathAndUniqueName(GraphInstance, LOCTEXT("NewPCGGraphInstanceAsset", "NewPCGGraphInstance").ToString(), NewPackageName, NewAssetName);
+
+	UPCGGraphInstance* NewPCGGraphInstance = Cast<UPCGGraphInstance>(AssetTools.CreateAssetWithDialog(NewAssetName, NewPackageName, GraphInstance->GetClass(), Factory, "PCGEditor_SaveGraphInstance"));
+
+	if (NewPCGGraphInstance == nullptr)
 	{
-		PCGCategory.AddProperty(Property);
+		UE_LOG(LogPCGEditor, Warning, TEXT("Graph instance asset creation was aborted or failed, abort."));
+		return FReply::Handled();
 	}
+
+	NewPCGGraphInstance->SetGraph(GraphInstance->Graph);
+	NewPCGGraphInstance->ParametersOverrides = GraphInstance->ParametersOverrides;
+
+	// Save the new asset
+	UEditorAssetLibrary::SaveLoadedAsset(NewPCGGraphInstance);
+
+	return FReply::Handled();
 }
 
 #undef LOCTEXT_NAMESPACE
diff --git a/Engine/Plugins/Experimental/PCG/Source/PCGEditor/Private/Details/PCGInstancedPropertyBagOverrideDetails.cpp b/Engine/Plugins/Experimental/PCG/Source/PCGEditor/Private/Details/PCGInstancedPropertyBagOverrideDetails.cpp
new file mode 100644
--- /dev/null
+++ b/Engine/Plugins/Experimental/PCG/Source/PCGEditor/Private/Details/PCGInstancedPropertyBagOverrideDetails.cpp
@@ -0,0 +1,171 @@
+// Copyright Epic Games, Inc. All Rights Reserved.
+
+#include "Details/PCGInstancedPropertyBagOverrideDetails.h"
+#include "PCGGraph.h"
+
+#include "DetailLayoutBuilder.h"
+#include "IDetailChildrenBuilder.h"
+#include "PropertyCustomizationHelpers.h"
+#include "ScopedTransaction.h"
+#include "Widgets/SBoxPanel.h"
+#include "Widgets/Layout/SBox.h"
+#include "Widgets/Images/SImage.h"
+#include "Widgets/Input/SCheckBox.h"
+
+#define LOCTEXT_NAMESPACE "PCGOverrideInstancedPropertyBagDetails"
+
+namespace PCGOverrideInstancedPropertyBagDataDetails
+{
+	/** @return true if the property is one of the known missing types. */
+	bool HasMissingType(const TSharedPtr<IPropertyHandle>& PropertyHandle)
+	{
+		if (!PropertyHandle)
+		{
+			return false;
+		}
+
+		// Handles Struct
+		if (FStructProperty* StructProperty = CastField<FStructProperty>(PropertyHandle->GetProperty()))
+		{
+			return StructProperty->Struct == FPropertyBagMissingStruct::StaticStruct();
+		}
+		// Handles Object, SoftObject, Class, SoftClass.
+		if (FObjectPropertyBase* ObjectProperty = CastField<FObjectPropertyBase>(PropertyHandle->GetProperty()))
+		{
+			return ObjectProperty->PropertyClass == UPropertyBagMissingObject::StaticClass();
+		}
+		// Handles Enum
+		if (FEnumProperty* EnumProperty = CastField<FEnumProperty>(PropertyHandle->GetProperty()))
+		{
+			return EnumProperty->GetEnum() == StaticEnum<EPropertyBagMissingEnum>();
+		}
+
+		return false;
+	}
+}
+
+TSharedRef<IPropertyTypeCustomization> FPCGOverrideInstancedPropertyBagDetails::MakeInstance()
+{
+	return MakeShareable(new FPCGOverrideInstancedPropertyBagDetails);
+}
+
+void FPCGOverrideInstancedPropertyBagDetails::CustomizeHeader(TSharedRef<IPropertyHandle> InPropertyHandle, FDetailWidgetRow& InHeaderRow, IPropertyTypeCustomizationUtils& InCustomizationUtils)
+{
+	PropertyHandle = InPropertyHandle;
+
+	InHeaderRow
+		.NameContent()
+		[
+			InPropertyHandle->CreatePropertyNameWidget()
+		];
+}
+
+void FPCGOverrideInstancedPropertyBagDetails::CustomizeChildren(TSharedRef<IPropertyHandle> InPropertyHandle, IDetailChildrenBuilder& InChildrenBuilder, IPropertyTypeCustomizationUtils& InCustomizationUtils)
+{
+	TSharedPtr<IPropertyHandle> InstancedPropertyBagHandle = PropertyHandle.IsValid() ? PropertyHandle->GetChildHandle(TEXT("Parameters")) : nullptr;
+	const TSharedRef<FPCGOverrideInstancedPropertyBagDataDetails> InstanceDetails = MakeShareable(new FPCGOverrideInstancedPropertyBagDataDetails(InstancedPropertyBagHandle, InCustomizationUtils.GetPropertyUtilities()));
+	InChildrenBuilder.AddCustomBuilder(InstanceDetails);
+}
+
+FPCGOverrideInstancedPropertyBagDataDetails::FPCGOverrideInstancedPropertyBagDataDetails(TSharedPtr<IPropertyHandle> InStructProperty, const TSharedPtr<IPropertyUtilities>& InPropUtils)
+	: FInstancedStructDataDetails(InStructProperty.IsValid() ? InStructProperty->GetChildHandle(TEXT("Value")) : nullptr)
+{
+	if (InStructProperty.IsValid())
+	{
+		// InStructProperty correspond to GraphInstance->ParameterOverrides->Parameters
+		// We need GraphInstance->ParameterOverrides->PropertiesIDsOverridden. So look at the parent then the child.
+		TSharedPtr<IPropertyHandle> ParentProperty = InStructProperty->GetParentHandle();
+		if (ParentProperty.IsValid())
+		{
+			PropertiesIDsOverriddenHandle = ParentProperty->GetChildHandle(GET_MEMBER_NAME_CHECKED(FPCGOverrideInstancedPropertyBag, PropertiesIDsOverridden));
+		}
+
+		TArray<UObject*> OuterObjects;
+		InStructProperty->GetOuterObjects(OuterObjects);
+
+		if (!OuterObjects.IsEmpty())
+		{
+			Owner = Cast<UPCGGraphInstance>(OuterObjects[0]);
+		}
+	}
+}
+
+void FPCGOverrideInstancedPropertyBagDataDetails::OnChildRowAdded(IDetailPropertyRow& ChildRow)
+{
+	if (!Owner || !PropertiesIDsOverriddenHandle.IsValid())
+	{
+		return;
+	}
+
+	TSharedPtr<SWidget> NameWidget;
+	TSharedPtr<SWidget> ValueWidget;
+	FDetailWidgetRow Row;
+	ChildRow.GetDefaultWidgets(NameWidget, ValueWidget, Row);
+
+	TSharedPtr<IPropertyHandle> ChildPropertyHandle = ChildRow.GetPropertyHandle();
+
+	FGuid PropertyID;
+
+	const bool bIsOverridden = Owner->IsPropertyOverridden(ChildPropertyHandle->GetProperty());
+
+	ChildRow
+	.CustomWidget(/*bShowChildren*/true)
+	.NameContent()
+	[
+		SNew(SHorizontalBox)
+		// Error icon
+		+ SHorizontalBox::Slot()
+		.AutoWidth()
+		.HAlign(HAlign_Right)
+		.VAlign(VAlign_Center)
+		.Padding(0, 0, 2, 0)
+		[
+			SNew(SBox)
+			.WidthOverride(12)
+			.HeightOverride(12)
+			[
+				SNew(SImage)
+				.ToolTipText(LOCTEXT("MissingType", "The property is missing type. The Struct, Enum, or Object may have been removed."))
+				.Visibility_Lambda([ChildPropertyHandle]()
+				{
+					return PCGOverrideInstancedPropertyBagDataDetails::HasMissingType(ChildPropertyHandle) ? EVisibility::Visible : EVisibility::Collapsed;
+				})
+				.Image(FAppStyle::GetBrush("Icons.Error"))
+			]
+		]
+		// Override 
+		+ SHorizontalBox::Slot()
+		.AutoWidth()
+		[
+			SNew(SCheckBox)
+			.IsChecked_Lambda([this, ChildPropertyHandle]()
+			{
+				return Owner->IsPropertyOverridden(ChildPropertyHandle->GetProperty()) ? ECheckBoxState::Checked : ECheckBoxState::Unchecked;
+			})
+			.OnCheckStateChanged_Lambda([this, ValueWidget, ChildPropertyHandle](ECheckBoxState NewState)
+			{
+				PropertiesIDsOverriddenHandle->NotifyPreChange();
+
+				FScopedTransaction Transaction(FText::Format(LOCTEXT("OnCheckStateChanged", "Change Override for {0}"), FText::FromName(ChildPropertyHandle->GetProperty()->GetFName())));
+				const bool bIsOverridden = NewState == ECheckBoxState::Checked;
+				Owner->UpdatePropertyOverride(ChildPropertyHandle->GetProperty(), bIsOverridden);
+
+				PropertiesIDsOverriddenHandle->NotifyPostChange(EPropertyChangeType::ValueSet);
+			})
+		]
+		// Name
+		+ SHorizontalBox::Slot()
+		.AutoWidth()
+		[
+			NameWidget.ToSharedRef()
+		]
+	]
+	.ValueContent()
+	[
+		ValueWidget.ToSharedRef()
+	];
+
+	ValueWidget->SetEnabled(TAttribute<bool>::CreateLambda([this, ChildPropertyHandle]() -> bool { return Owner->IsPropertyOverridden(ChildPropertyHandle->GetProperty()); }));
+}
+
+#undef LOCTEXT_NAMESPACE
\ No newline at end of file
diff --git a/Engine/Plugins/Experimental/PCG/Source/PCGEditor/Private/Details/PCGInstancedPropertyBagOverrideDetails.h b/Engine/Plugins/Experimental/PCG/Source/PCGEditor/Private/Details/PCGInstancedPropertyBagOverrideDetails.h
new file mode 100644
--- /dev/null
+++ b/Engine/Plugins/Experimental/PCG/Source/PCGEditor/Private/Details/PCGInstancedPropertyBagOverrideDetails.h
@@ -0,0 +1,36 @@
+// Copyright Epic Games, Inc. All Rights Reserved.
+
+#pragma once
+
+#include "IDetailCustomization.h"
+#include "IPropertyTypeCustomization.h"
+#include "InstancedStructDetails.h"
+#include "Templates/SharedPointer.h"
+
+class UPCGGraphInstance;
+
+class FPCGOverrideInstancedPropertyBagDetails : public IPropertyTypeCustomization
+{
+public:
+	/** Makes a new instance of this detail layout class for a specific detail view requesting it */
+	static TSharedRef<IPropertyTypeCustomization> MakeInstance();
+
+	/** ~Begin IPropertyTypeCustomization interface */
+	virtual void CustomizeHeader(TSharedRef<IPropertyHandle> InPropertyHandle, FDetailWidgetRow& InHeaderRow, IPropertyTypeCustomizationUtils& InCustomizationUtils) override;
+	virtual void CustomizeChildren(TSharedRef<IPropertyHandle> InPropertyHandle, IDetailChildrenBuilder& InChildrenBuilder, IPropertyTypeCustomizationUtils& InCustomizationUtils) override;
+	/** ~End IPropertyTypeCustomization interface */
+private:
+	TSharedPtr<IPropertyHandle> PropertyHandle;
+};
+
+class FPCGOverrideInstancedPropertyBagDataDetails : public FInstancedStructDataDetails
+{
+public:
+	FPCGOverrideInstancedPropertyBagDataDetails(TSharedPtr<IPropertyHandle> InStructProperty, const TSharedPtr<IPropertyUtilities>& InPropUtils);
+
+	virtual void OnChildRowAdded(IDetailPropertyRow& ChildRow) override;
+
+private:
+	UPCGGraphInstance* Owner = nullptr;
+	TSharedPtr<IPropertyHandle> PropertiesIDsOverriddenHandle;
+};
\ No newline at end of file
diff --git a/Engine/Plugins/Experimental/PCG/Source/PCGEditor/Private/PCGEditor.cpp b/Engine/Plugins/Experimental/PCG/Source/PCGEditor/Private/PCGEditor.cpp
--- a/Engine/Plugins/Experimental/PCG/Source/PCGEditor/Private/PCGEditor.cpp
+++ b/Engine/Plugins/Experimental/PCG/Source/PCGEditor/Private/PCGEditor.cpp
@@ -2012,6 +2012,11 @@ void FPCGEditor::SaveAsset_Execute()
 
 void FPCGEditor::OnClose()
 {
+	if (PCGEditorGraph)
+	{
+		PCGEditorGraph->OnClose();
+	}
+
 	// Extra nodes are replicated on editor close, to be saved in the underlying PCGGraph
 	ReplicateExtraNodes();
 
diff --git a/Engine/Plugins/Experimental/PCG/Source/PCGEditor/Private/PCGEditorGraph.cpp b/Engine/Plugins/Experimental/PCG/Source/PCGEditor/Private/PCGEditorGraph.cpp
--- a/Engine/Plugins/Experimental/PCG/Source/PCGEditor/Private/PCGEditorGraph.cpp
+++ b/Engine/Plugins/Experimental/PCG/Source/PCGEditor/Private/PCGEditorGraph.cpp
@@ -11,6 +11,7 @@
 #include "PCGGraph.h"
 #include "PCGPin.h"
 #include "Elements/PCGReroute.h"
+#include "Elements/PCGUserParameterGet.h"
 
 #include "EdGraph/EdGraphPin.h"
 
@@ -19,6 +20,8 @@ void UPCGEditorGraph::InitFromNodeGraph(UPCGGraph* InPCGGraph)
 	check(InPCGGraph && !PCGGraph);
 	PCGGraph = InPCGGraph;
 
+	PCGGraph->OnGraphParametersChangedDelegate.AddUObject(this, &UPCGEditorGraph::OnGraphUserParametersChanged);
+
 	TMap<UPCGNode*, UPCGEditorGraphNodeBase*> NodeLookup;
 	const bool bSelectNewNode = false;
 
@@ -79,6 +82,21 @@ void UPCGEditorGraph::InitFromNodeGraph(UPCGGraph* InPCGGraph)
 	}
 }
 
+void UPCGEditorGraph::BeginDestroy()
+{
+	Super::BeginDestroy();
+
+	OnClose();
+}
+
+void UPCGEditorGraph::OnClose()
+{
+	if (PCGGraph)
+	{
+		PCGGraph->OnGraphParametersChangedDelegate.RemoveAll(this);
+	}
+}
+
 void UPCGEditorGraph::CreateLinks(UPCGEditorGraphNodeBase* GraphNode, bool bCreateInbound, bool bCreateOutbound)
 {
 	check(GraphNode);
@@ -185,3 +203,42 @@ void UPCGEditorGraph::CreateLinks(UPCGEditorGraphNodeBase* GraphNode, bool bCrea
 		}
 	}
 }
+
+void UPCGEditorGraph::OnGraphUserParametersChanged(UPCGGraphInterface* InGraph, EPCGGraphParameterEvent ChangeType, FName ChangedPropertyName)
+{
+	if (ChangeType != EPCGGraphParameterEvent::Removed || InGraph != PCGGraph)
+	{
+		return;
+	}
+
+	// If a parameter was removed, just look for getter nodes that do exists in the editor graph, but not in the PCG graph.
+	TArray<UPCGEditorGraphNodeBase*> NodesToRemove;
+	for (UEdGraphNode* EditorNode : Nodes)
+	{
+		if (UPCGEditorGraphNodeBase* PCGEditorNode = Cast<UPCGEditorGraphNodeBase>(EditorNode))
+		{
+			if (UPCGNode* PCGNode = PCGEditorNode->GetPCGNode())
+			{
+				if (UPCGUserParameterGetSettings* Settings = Cast<UPCGUserParameterGetSettings>(PCGNode->GetSettings()))
+				{
+					if (!PCGGraph->Contains(PCGNode))
+					{
+						NodesToRemove.Add(PCGEditorNode);
+					}
+				}
+			}
+		}
+	}
+
+	if (NodesToRemove.IsEmpty())
+	{
+		return;
+	}
+
+	Modify();
+
+	for (UPCGEditorGraphNodeBase* NodeToRemove : NodesToRemove)
+	{
+		NodeToRemove->DestroyNode();
+	}
+}
\ No newline at end of file
diff --git a/Engine/Plugins/Experimental/PCG/Source/PCGEditor/Private/PCGEditorGraph.h b/Engine/Plugins/Experimental/PCG/Source/PCGEditor/Private/PCGEditorGraph.h
--- a/Engine/Plugins/Experimental/PCG/Source/PCGEditor/Private/PCGEditorGraph.h
+++ b/Engine/Plugins/Experimental/PCG/Source/PCGEditor/Private/PCGEditorGraph.h
@@ -4,11 +4,12 @@
 
 #include "EdGraph/EdGraph.h"
 
+#include "PCGGraph.h"
+
 #include "PCGEditorGraph.generated.h"
 
 class FPCGEditor;
 class UPCGNode;
-class UPCGGraph;
 class UPCGEditorGraphNodeBase;
 
 UCLASS()
@@ -17,9 +18,17 @@ class UPCGEditorGraph : public UEdGraph
 	GENERATED_BODY()
 
 public:
+
+	// ~Begin UObject interface
+	virtual void BeginDestroy() override;
+	// ~End UObject interface
+
 	/** Initialize the editor graph from a PCGGraph */
 	void InitFromNodeGraph(UPCGGraph* InPCGGraph);
 
+	/** When the editor is closing */
+	void OnClose();
+
 	/** Creates the links for a given node */
 	void CreateLinks(UPCGEditorGraphNodeBase* InGraphNode, bool bCreateInbound, bool bCreateOutbound);
 
@@ -31,6 +40,8 @@ public:
 protected:
 	void CreateLinks(UPCGEditorGraphNodeBase* InGraphNode, bool bCreateInbound, bool bCreateOutbound, const TMap<UPCGNode*, UPCGEditorGraphNodeBase*>& InGraphNodeToPCGNodeMap);
 
+	void OnGraphUserParametersChanged(UPCGGraphInterface* InGraph, EPCGGraphParameterEvent ChangeType, FName ChangedPropertyName);
+
 private:
 	UPROPERTY()
 	TObjectPtr<UPCGGraph> PCGGraph = nullptr;
diff --git a/Engine/Plugins/Experimental/PCG/Source/PCGEditor/Private/PCGEditorGraphNodeFactory.cpp b/Engine/Plugins/Experimental/PCG/Source/PCGEditor/Private/PCGEditorGraphNodeFactory.cpp
--- a/Engine/Plugins/Experimental/PCG/Source/PCGEditor/Private/PCGEditorGraphNodeFactory.cpp
+++ b/Engine/Plugins/Experimental/PCG/Source/PCGEditor/Private/PCGEditorGraphNodeFactory.cpp
@@ -6,7 +6,9 @@
 #include "PCGEditorGraphNodeReroute.h"
 #include "PCGNode.h"
 #include "SPCGEditorGraphNode.h"
+#include "SPCGEditorGraphVarNode.h"
 #include "Elements/PCGReroute.h"
+#include "Elements/PCGUserParameterGet.h"
 
 #include "SGraphNodeKnot.h"
 
@@ -41,6 +43,10 @@ TSharedPtr<SGraphNode> FPCGEditorGraphNodeFactory::CreateNode(UEdGraphNode* InNo
 		{
 			SAssignNew(VisualNode, SPCGEditorGraphNodeKnot, GraphNode);
 		}
+		else if (PCGNode && Cast<UPCGUserParameterGetSettings>(PCGNode->GetSettings()))
+		{
+			SAssignNew(VisualNode, SPCGEditorGraphVarNode, GraphNode);
+		}
 		else
 		{
 			SAssignNew(VisualNode, SPCGEditorGraphNode, GraphNode);
diff --git a/Engine/Plugins/Experimental/PCG/Source/PCGEditor/Private/PCGEditorGraphSchema.cpp b/Engine/Plugins/Experimental/PCG/Source/PCGEditor/Private/PCGEditorGraphSchema.cpp
--- a/Engine/Plugins/Experimental/PCG/Source/PCGEditor/Private/PCGEditorGraphSchema.cpp
+++ b/Engine/Plugins/Experimental/PCG/Source/PCGEditor/Private/PCGEditorGraphSchema.cpp
@@ -6,6 +6,7 @@
 #include "PCGEdge.h"
 #include "PCGGraph.h"
 #include "PCGPin.h"
+#include "Elements/PCGUserParameterGet.h"
 
 #include "PCGEditorCommon.h"
 #include "PCGEditorGraph.h"
@@ -53,7 +54,7 @@ void UPCGEditorGraphSchema::GetGraphContextActions(FGraphContextMenuBuilder& Con
 {
 	Super::GetGraphContextActions(ContextMenuBuilder);
 
-	GetNativeElementActions(ContextMenuBuilder);
+	GetNativeElementActions(ContextMenuBuilder, ContextMenuBuilder.CurrentGraph);
 	GetSubgraphElementActions(ContextMenuBuilder);
 	GetBlueprintElementActions(ContextMenuBuilder);
 	GetSettingsElementActions(ContextMenuBuilder, /*bIsContextual=*/true);
@@ -200,7 +201,7 @@ void UPCGEditorGraphSchema::BreakSinglePinLink(UEdGraphPin* SourcePin, UEdGraphP
 	PCGGraph->RemoveEdge(SourcePCGNode, SourcePin->PinName, TargetPCGNode, TargetPin->PinName);
 }
 
-void UPCGEditorGraphSchema::GetNativeElementActions(FGraphActionMenuBuilder& ActionMenuBuilder) const
+void UPCGEditorGraphSchema::GetNativeElementActions(FGraphActionMenuBuilder& ActionMenuBuilder, const UEdGraph* CurrentGraph) const
 {
 	TArray<UClass*> SettingsClasses;
 	for (TObjectIterator<UClass> It; It; ++It)
@@ -230,6 +231,32 @@ void UPCGEditorGraphSchema::GetNativeElementActions(FGraphActionMenuBuilder& Act
 			}
 		}
 	}
+
+	if (const UPCGEditorGraph* Graph = Cast<UPCGEditorGraph>(CurrentGraph))
+	{
+		if (const UPCGGraph* PCGGraph = const_cast<UPCGEditorGraph*>(Graph)->GetPCGGraph())
+		{
+			if (const FInstancedPropertyBag* UserParameters = PCGGraph->GetUserParametersStruct())
+			{
+				if (const UPropertyBag* BagStruct = UserParameters->GetPropertyBagStruct())
+				{
+					const FText Category = LOCTEXT("UserParametersCategoryName", "User Parameters");
+
+					for (const FPropertyBagPropertyDesc& PropertyDesc : BagStruct->GetPropertyDescs())
+					{
+						const FText MenuDesc = FText::Format(FText::FromString(TEXT("Get {0}")), FText::FromName(PropertyDesc.Name));
+						const FText Description = FText::Format(LOCTEXT("NodeTooltip", "Get the value from '{0}' parameter, can be overridden by the graph instance."), FText::FromName(PropertyDesc.Name));
+
+						TSharedPtr<FPCGEditorGraphSchemaAction_NewGetParameterElement> NewAction(new FPCGEditorGraphSchemaAction_NewGetParameterElement(Category, MenuDesc, Description, 0));
+						NewAction->SettingsClass = UPCGUserParameterGetSettings::StaticClass();
+						NewAction->PropertyName = PropertyDesc.Name;
+						NewAction->PropertyGuid = PropertyDesc.ID;
+						ActionMenuBuilder.AddAction(NewAction);
+					}
+				}
+			}
+		}
+	}
 }
 
 void UPCGEditorGraphSchema::GetBlueprintElementActions(FGraphActionMenuBuilder& ActionMenuBuilder) const
diff --git a/Engine/Plugins/Experimental/PCG/Source/PCGEditor/Private/PCGEditorGraphSchema.h b/Engine/Plugins/Experimental/PCG/Source/PCGEditor/Private/PCGEditorGraphSchema.h
--- a/Engine/Plugins/Experimental/PCG/Source/PCGEditor/Private/PCGEditorGraphSchema.h
+++ b/Engine/Plugins/Experimental/PCG/Source/PCGEditor/Private/PCGEditorGraphSchema.h
@@ -33,7 +33,7 @@ public:
 	//~ End EdGraphSchema Interface
 
 private:
-	void GetNativeElementActions(FGraphActionMenuBuilder& ActionMenuBuilder) const;
+	void GetNativeElementActions(FGraphActionMenuBuilder& ActionMenuBuilder, const UEdGraph* CurrentGraph = nullptr) const;
 	void GetBlueprintElementActions(FGraphActionMenuBuilder& ActionMenuBuilder) const;
 	void GetSubgraphElementActions(FGraphActionMenuBuilder& ActionMenuBuilder) const;
 	void GetSettingsElementActions(FGraphActionMenuBuilder& ActionMenuBuilder, bool bIsContextual) const;
diff --git a/Engine/Plugins/Experimental/PCG/Source/PCGEditor/Private/PCGEditorGraphSchemaActions.cpp b/Engine/Plugins/Experimental/PCG/Source/PCGEditor/Private/PCGEditorGraphSchemaActions.cpp
--- a/Engine/Plugins/Experimental/PCG/Source/PCGEditor/Private/PCGEditorGraphSchemaActions.cpp
+++ b/Engine/Plugins/Experimental/PCG/Source/PCGEditor/Private/PCGEditorGraphSchemaActions.cpp
@@ -2,15 +2,17 @@
 
 #include "PCGEditorGraphSchemaActions.h"
 
-#include "Elements/PCGExecuteBlueprint.h"
-#include "Elements/PCGReroute.h"
 #include "PCGEditorCommon.h"
 #include "PCGEditorGraph.h"
 #include "PCGEditorGraphNode.h"
 #include "PCGEditorGraphNodeReroute.h"
 #include "PCGEditorModule.h"
 #include "PCGGraph.h"
+#include "PCGNode.h"
 #include "PCGSubgraph.h"
+#include "Elements/PCGExecuteBlueprint.h"
+#include "Elements/PCGReroute.h"
+#include "Elements/PCGUserParameterGet.h"
 
 #include "EdGraphNode_Comment.h"
 #include "GraphEditor.h"
@@ -43,6 +45,14 @@ UEdGraphNode* FPCGEditorGraphSchemaAction_NewNativeElement::PerformAction(UEdGra
 	UPCGSettings* DefaultNodeSettings = nullptr;
 	UPCGNode* NewPCGNode = PCGGraph->AddNodeOfType(SettingsClass, DefaultNodeSettings);
 
+	if (!NewPCGNode)
+	{
+		UE_LOG(LogPCGEditor, Error, TEXT("Failed to add a node of type %s"), *SettingsClass->GetName());
+		return nullptr;
+	}
+
+	PostCreation(NewPCGNode);
+
 	FGraphNodeCreator<UPCGEditorGraphNode> NodeCreator(*EditorGraph);
 	UPCGEditorGraphNode* NewNode = NodeCreator.CreateUserInvokedNode(bSelectNewNode);
 	NewNode->Construct(NewPCGNode);
@@ -61,6 +71,18 @@ UEdGraphNode* FPCGEditorGraphSchemaAction_NewNativeElement::PerformAction(UEdGra
 	return NewNode;
 }
 
+void FPCGEditorGraphSchemaAction_NewGetParameterElement::PostCreation(UPCGNode* NewNode)
+{
+	check(NewNode);
+	UPCGUserParameterGetSettings* Settings = CastChecked<UPCGUserParameterGetSettings>(NewNode->GetSettings());
+
+	Settings->PropertyGuid = PropertyGuid;
+	Settings->PropertyName = PropertyName;
+
+	// We need to set the settings to update the pins.
+	NewNode->SetSettingsInterface(Settings);
+}
+
 UEdGraphNode* FPCGEditorGraphSchemaAction_NewSettingsElement::PerformAction(UEdGraph* ParentGraph, UEdGraphPin* FromPin, const FVector2D Location, bool bSelectNewNode)
 {
 	UPCGEditorGraph* EditorGraph = Cast<UPCGEditorGraph>(ParentGraph);
diff --git a/Engine/Plugins/Experimental/PCG/Source/PCGEditor/Private/PCGEditorGraphSchemaActions.h b/Engine/Plugins/Experimental/PCG/Source/PCGEditor/Private/PCGEditorGraphSchemaActions.h
--- a/Engine/Plugins/Experimental/PCG/Source/PCGEditor/Private/PCGEditorGraphSchemaActions.h
+++ b/Engine/Plugins/Experimental/PCG/Source/PCGEditor/Private/PCGEditorGraphSchemaActions.h
@@ -10,6 +10,7 @@
 class UEdGraph;
 class UEdGraphPin;
 class UPCGBlueprintElement;
+class UPCGNode;
 class UPCGSettings;
 class UPCGEditorGraph;
 
@@ -43,6 +44,27 @@ struct FPCGEditorGraphSchemaAction_NewNativeElement : public FEdGraphSchemaActio
 	virtual FName GetTypeId() const override { return StaticGetTypeId(); }
 	virtual UEdGraphNode* PerformAction(UEdGraph* ParentGraph, UEdGraphPin* FromPin, const FVector2D Location, bool bSelectNewNode = true) override;
 	//~ End FEdGraphSchemaAction Interface
+
+protected:
+	virtual void PostCreation(UPCGNode* NewNode) {}
+};
+
+USTRUCT()
+struct FPCGEditorGraphSchemaAction_NewGetParameterElement : public FPCGEditorGraphSchemaAction_NewNativeElement
+{
+	GENERATED_BODY()
+
+	// Inherit the base class's constructors
+	using FPCGEditorGraphSchemaAction_NewNativeElement::FPCGEditorGraphSchemaAction_NewNativeElement;
+
+	UPROPERTY()
+	FGuid PropertyGuid;
+
+	UPROPERTY()
+	FName PropertyName;
+
+protected:
+	virtual void PostCreation(UPCGNode* NewNode) override;
 };
 
 USTRUCT()
diff --git a/Engine/Plugins/Experimental/PCG/Source/PCGEditor/Private/PCGEditorModule.cpp b/Engine/Plugins/Experimental/PCG/Source/PCGEditor/Private/PCGEditorModule.cpp
--- a/Engine/Plugins/Experimental/PCG/Source/PCGEditor/Private/PCGEditorModule.cpp
+++ b/Engine/Plugins/Experimental/PCG/Source/PCGEditor/Private/PCGEditorModule.cpp
@@ -21,6 +21,7 @@
 #include "Details/PCGBlueprintSettingsDetails.h"
 #include "Details/PCGGraphDetails.h"
 #include "Details/PCGGraphInstanceDetails.h"
+#include "Details/PCGInstancedPropertyBagOverrideDetails.h"
 #include "Details/PCGVolumeDetails.h"
 #include "Framework/MultiBox/MultiBoxBuilder.h"
 
@@ -80,6 +81,7 @@ void FPCGEditorModule::RegisterDetailsCustomizations()
 	PropertyEditor.RegisterCustomClassLayout("PCGVolume", FOnGetDetailCustomizationInstance::CreateStatic(&FPCGVolumeDetails::MakeInstance));
 
 	PropertyEditor.RegisterCustomPropertyTypeLayout("PCGAttributePropertySelector", FOnGetPropertyTypeCustomizationInstance::CreateStatic(&FPCGAttributePropertySelectorDetails::MakeInstance));
+	PropertyEditor.RegisterCustomPropertyTypeLayout("PCGOverrideInstancedPropertyBag", FOnGetPropertyTypeCustomizationInstance::CreateStatic(&FPCGOverrideInstancedPropertyBagDetails::MakeInstance));
 
 	PropertyEditor.NotifyCustomizationModuleChanged();
 }
@@ -95,6 +97,7 @@ void FPCGEditorModule::UnregisterDetailsCustomizations()
 		PropertyModule->UnregisterCustomClassLayout("PCGVolume");
 
 		PropertyModule->UnregisterCustomPropertyTypeLayout("PCGAttributePropertySelector");
+		PropertyModule->UnregisterCustomPropertyTypeLayout("PCGOverrideInstancedPropertyBag");
 
 		PropertyModule->NotifyCustomizationModuleChanged();
 	}
diff --git a/Engine/Plugins/Experimental/PCG/Source/PCGEditor/Private/PCGGraphFactory.cpp b/Engine/Plugins/Experimental/PCG/Source/PCGEditor/Private/PCGGraphFactory.cpp
--- a/Engine/Plugins/Experimental/PCG/Source/PCGEditor/Private/PCGGraphFactory.cpp
+++ b/Engine/Plugins/Experimental/PCG/Source/PCGEditor/Private/PCGGraphFactory.cpp
@@ -3,6 +3,10 @@
 #include "PCGGraphFactory.h"
 #include "PCGGraph.h"
 
+/////////////////////////
+// PCGGraph
+/////////////////////////
+
 UPCGGraphFactory::UPCGGraphFactory(const FObjectInitializer& ObjectInitializer)
 	: Super(ObjectInitializer)
 {
@@ -19,4 +23,26 @@ UObject* UPCGGraphFactory::FactoryCreateNew(UClass* InClass, UObject* InParent,
 bool UPCGGraphFactory::ShouldShowInNewMenu() const
 {
 	return true;
+}
+
+/////////////////////////
+// PCGGraphInstance
+/////////////////////////
+
+UPCGGraphInstanceFactory::UPCGGraphInstanceFactory(const FObjectInitializer& ObjectInitializer)
+	: Super(ObjectInitializer)
+{
+	SupportedClass = UPCGGraphInstance::StaticClass();
+	bCreateNew = true;
+	bEditAfterNew = true;
+}
+
+UObject* UPCGGraphInstanceFactory::FactoryCreateNew(UClass* InClass, UObject* InParent, FName InName, EObjectFlags Flags, UObject* Context, FFeedbackContext* Warn)
+{
+	return NewObject<UPCGGraphInstance>(InParent, InClass, InName, Flags);
+}
+
+bool UPCGGraphInstanceFactory::ShouldShowInNewMenu() const
+{
+	return false;
 }
\ No newline at end of file
diff --git a/Engine/Plugins/Experimental/PCG/Source/PCGEditor/Private/PCGGraphFactory.h b/Engine/Plugins/Experimental/PCG/Source/PCGEditor/Private/PCGGraphFactory.h
--- a/Engine/Plugins/Experimental/PCG/Source/PCGEditor/Private/PCGGraphFactory.h
+++ b/Engine/Plugins/Experimental/PCG/Source/PCGEditor/Private/PCGGraphFactory.h
@@ -18,3 +18,16 @@ public:
 	virtual UObject* FactoryCreateNew(UClass* InClass, UObject* InParent, FName InName, EObjectFlags Flags, UObject* Context, FFeedbackContext* Warn) override;
 	virtual bool ShouldShowInNewMenu() const override;
 };
+
+UCLASS(hidecategories = Object)
+class UPCGGraphInstanceFactory : public UFactory
+{
+	GENERATED_BODY()
+
+public:
+	UPCGGraphInstanceFactory(const FObjectInitializer& ObjectInitializer);
+
+	//~UFactory interface
+	virtual UObject* FactoryCreateNew(UClass* InClass, UObject* InParent, FName InName, EObjectFlags Flags, UObject* Context, FFeedbackContext* Warn) override;
+	virtual bool ShouldShowInNewMenu() const override;
+};
diff --git a/Engine/Plugins/Experimental/PCG/Source/PCGEditor/Private/PCGVolumeFactory.cpp b/Engine/Plugins/Experimental/PCG/Source/PCGEditor/Private/PCGVolumeFactory.cpp
--- a/Engine/Plugins/Experimental/PCG/Source/PCGEditor/Private/PCGVolumeFactory.cpp
+++ b/Engine/Plugins/Experimental/PCG/Source/PCGEditor/Private/PCGVolumeFactory.cpp
@@ -24,7 +24,7 @@ bool UPCGVolumeFactory::CanCreateActorFrom(const FAssetData& AssetData, FText& O
 		return true;
 	}
 
-	if ( AssetData.IsValid() && !AssetData.IsInstanceOf(UPCGGraph::StaticClass()))
+	if ( AssetData.IsValid() && !AssetData.IsInstanceOf<UPCGGraphInterface>())
 	{
 		OutErrorMsg = LOCTEXT("NoPCGGraph", "A valid PCG graph asset must be specified.");
 		return false;
@@ -43,7 +43,7 @@ void UPCGVolumeFactory::PostSpawnActor(UObject* Asset, AActor* NewActor)
 		return;
 	}
 
-	if (UPCGGraph* PCGGraph = Cast<UPCGGraph>(Asset))
+	if (UPCGGraphInterface* PCGGraph = Cast<UPCGGraphInterface>(Asset))
 	{
 		const UPCGEngineSettings* Settings = GetDefault<UPCGEngineSettings>();
 		
diff --git a/Engine/Plugins/Experimental/PCG/Source/PCGEditor/Private/SPCGEditorGraphVarNode.cpp b/Engine/Plugins/Experimental/PCG/Source/PCGEditor/Private/SPCGEditorGraphVarNode.cpp
new file mode 100644
--- /dev/null
+++ b/Engine/Plugins/Experimental/PCG/Source/PCGEditor/Private/SPCGEditorGraphVarNode.cpp
@@ -0,0 +1,174 @@
+// Copyright Epic Games, Inc. All Rights Reserved.
+
+/**
+* Adapted from SGraphNodeK2Var.cpp
+*/
+
+#include "SPCGEditorGraphVarNode.h"
+
+#include "GraphEditorSettings.h"
+#include "SCommentBubble.h"
+#include "SGraphNode.h"
+#include "SlotBase.h"
+#include "EdGraph/EdGraphNode.h"
+#include "TutorialMetaData.h"
+#include "Widgets/Images/SImage.h"
+#include "Widgets/Layout/SSpacer.h"
+#include "Widgets/SBoxPanel.h"
+
+class SWidget;
+struct FSlateBrush;
+
+void SPCGEditorGraphVarNode::Construct(const FArguments& InArgs, UPCGEditorGraphNodeBase* InNode)
+{
+	SPCGEditorGraphNode::Construct(SPCGEditorGraphNode::FArguments{}, InNode);
+}
+
+FSlateColor SPCGEditorGraphVarNode::GetVariableColor() const
+{
+	return FColor::Yellow;
+}
+
+void SPCGEditorGraphVarNode::UpdateGraphNode()
+{
+	InputPins.Empty();
+	OutputPins.Empty();
+
+	// Reset variables that are going to be exposed, in case we are refreshing an already setup node.
+	RightNodeBox.Reset();
+	LeftNodeBox.Reset();
+
+	FMargin ContentAreaMargin = FMargin(0.0f, 4.0f);
+
+	SetupErrorReporting();
+
+	// Setup a meta tag for this node
+	FGraphNodeMetaData TagMeta(TEXT("Graphnode"));
+	PopulateMetaTag(&TagMeta);
+
+	//             ________________
+	//            | (>) L |  R (>) |
+	//            | (>) E |  I (>) |
+	//            | (>) F |  G (>) |
+	//            | (>) T |  H (>) |
+	//            |       |  T (>) |
+	//            |_______|________|
+	//
+	this->ContentScale.Bind( this, &SGraphNode::GetContentScale );
+	this->GetOrAddSlot( ENodeZone::Center )
+	.HAlign(HAlign_Center)
+	.VAlign(VAlign_Center)
+	[
+		SNew(SVerticalBox)
+		+SVerticalBox::Slot()
+		[
+			SNew(SOverlay)
+			.AddMetaData<FGraphNodeMetaData>(TagMeta)
+			+ SOverlay::Slot()
+			[
+				SNew(SImage)
+				.Image( FAppStyle::GetBrush("Graph.VarNode.Body") )
+			]
+			+ SOverlay::Slot()
+			.VAlign(VAlign_Top)
+			[
+				SNew(SImage)
+				.Image( FAppStyle::GetBrush("Graph.VarNode.ColorSpill") )
+				.ColorAndOpacity( this, &SPCGEditorGraphVarNode::GetVariableColor )
+			]
+			+ SOverlay::Slot()
+			[
+				SNew(SImage)
+				.Image( FAppStyle::GetBrush("Graph.VarNode.Gloss") )
+			]
+			+ SOverlay::Slot()
+			.Padding( ContentAreaMargin )
+			[
+				// NODE CONTENT AREA
+				SNew(SHorizontalBox)
+				+SHorizontalBox::Slot()
+				.HAlign(HAlign_Left)
+				.FillWidth(1.0f)
+				.Padding( FMargin(2,0) )
+				[
+					// LEFT
+					SAssignNew(LeftNodeBox, SVerticalBox)
+				]
+				+ SHorizontalBox::Slot()
+				.AutoWidth()
+				.HAlign(HAlign_Right)
+				.Padding( FMargin(2,0) )
+				[
+					// RIGHT
+					SAssignNew(RightNodeBox, SVerticalBox)
+				]
+			]
+		]
+		+SVerticalBox::Slot()
+		.VAlign(VAlign_Top)
+		.AutoHeight() 
+		.Padding( FMargin(5.0f, 1.0f) )
+		[
+			ErrorReporting->AsWidget()
+		]
+	];
+
+	float VerticalPaddingAmount = 0.0f;
+
+	if (VerticalPaddingAmount > 0.0f)
+	{
+		LeftNodeBox->AddSlot()
+		.AutoHeight()
+		.HAlign(HAlign_Left)
+		.VAlign(VAlign_Center)
+		[
+			SNew(SSpacer).Size(FVector2D(0.0f, VerticalPaddingAmount))
+		];
+
+		RightNodeBox->AddSlot()
+		.AutoHeight()
+		.HAlign(HAlign_Left)
+		.VAlign(VAlign_Center)
+		[
+			SNew(SSpacer).Size(FVector2D(0.0f, VerticalPaddingAmount))
+		];
+	}
+	// Create comment bubble
+	TSharedPtr<SCommentBubble> CommentBubble;
+	const FSlateColor CommentColor = GetDefault<UGraphEditorSettings>()->DefaultCommentNodeTitleColor;
+
+	SAssignNew( CommentBubble, SCommentBubble )
+	.GraphNode( GraphNode )
+	.Text( this, &SGraphNode::GetNodeComment )
+	.OnTextCommitted( this, &SGraphNode::OnCommentTextCommitted )
+	.ColorAndOpacity( CommentColor )
+	.AllowPinning( true )
+	.EnableTitleBarBubble( true )
+	.EnableBubbleCtrls( true )
+	.GraphLOD( this, &SGraphNode::GetCurrentLOD )
+	.IsGraphNodeHovered( this, &SGraphNode::IsHovered );
+
+	GetOrAddSlot( ENodeZone::TopCenter )
+	.SlotOffset( TAttribute<FVector2D>( CommentBubble.Get(), &SCommentBubble::GetOffset ))
+	.SlotSize( TAttribute<FVector2D>( CommentBubble.Get(), &SCommentBubble::GetSize ))
+	.AllowScaling( TAttribute<bool>( CommentBubble.Get(), &SCommentBubble::IsScalingAllowed ))
+	.VAlign( VAlign_Top )
+	[
+		CommentBubble.ToSharedRef()
+	];
+
+	// Create widgets for each of the real pins
+	CreatePinWidgets();
+}
+
+const FSlateBrush* SPCGEditorGraphVarNode::GetShadowBrush(bool bSelected) const
+{
+	return bSelected ? FAppStyle::GetBrush(TEXT("Graph.VarNode.ShadowSelected")) : FAppStyle::GetBrush(TEXT("Graph.VarNode.Shadow"));
+}
+
+void SPCGEditorGraphVarNode::GetDiffHighlightBrushes(const FSlateBrush*& BackgroundOut, const FSlateBrush*& ForegroundOut) const
+{
+	BackgroundOut = FAppStyle::GetBrush(TEXT("Graph.VarNode.DiffHighlight"));
+	ForegroundOut = FAppStyle::GetBrush(TEXT("Graph.VarNode.DiffHighlightShading"));
+}
+
diff --git a/Engine/Plugins/Experimental/PCG/Source/PCGEditor/Private/SPCGEditorGraphVarNode.h b/Engine/Plugins/Experimental/PCG/Source/PCGEditor/Private/SPCGEditorGraphVarNode.h
new file mode 100644
--- /dev/null
+++ b/Engine/Plugins/Experimental/PCG/Source/PCGEditor/Private/SPCGEditorGraphVarNode.h
@@ -0,0 +1,36 @@
+// Copyright Epic Games, Inc. All Rights Reserved.
+
+#pragma once
+
+#include "CoreMinimal.h"
+#include "SPCGEditorGraphNode.h"
+#include "Styling/SlateColor.h"
+#include "Templates/SharedPointer.h"
+#include "Types/SlateEnums.h"
+#include "Widgets/DeclarativeSyntaxSupport.h"
+
+class FText;
+class SWidget;
+class UK2Node;
+struct FMargin;
+struct FSlateBrush;
+
+class SPCGEditorGraphVarNode : public SPCGEditorGraphNode
+{
+public:
+	SLATE_BEGIN_ARGS(SPCGEditorGraphVarNode){}
+	SLATE_END_ARGS()
+
+	void Construct(const FArguments& InArgs, UPCGEditorGraphNodeBase* InNode);
+
+	virtual const FSlateBrush* GetShadowBrush(bool bSelected) const override;
+
+	virtual void GetDiffHighlightBrushes(const FSlateBrush*& BackgroundOut, const FSlateBrush*& ForegroundOut) const override;
+
+	// SGraphNode interface
+	virtual void UpdateGraphNode() override;
+	// End of SGraphNode interface
+
+protected:
+	FSlateColor GetVariableColor() const;
+};
diff --git a/Engine/Plugins/Experimental/PCG/Source/PCGEditor/Public/Details/PCGGraphInstanceDetails.h b/Engine/Plugins/Experimental/PCG/Source/PCGEditor/Public/Details/PCGGraphInstanceDetails.h
--- a/Engine/Plugins/Experimental/PCG/Source/PCGEditor/Public/Details/PCGGraphInstanceDetails.h
+++ b/Engine/Plugins/Experimental/PCG/Source/PCGEditor/Public/Details/PCGGraphInstanceDetails.h
@@ -3,6 +3,9 @@
 #pragma once
 
 #include "IDetailCustomization.h"
+
+#include "Containers/Array.h"
+#include "Input/Reply.h" // IWYU pragma: keep
 #include "UObject/WeakObjectPtrTemplates.h"
 
 class UPCGGraphInstance;
@@ -15,10 +18,8 @@ public:
 	virtual void CustomizeDetails(IDetailLayoutBuilder& DetailBuilder) override;
 
 private:
+	bool SaveInstanceButtonEnabled() const;
+	FReply OnSaveInstanceClicked();
+
 	TArray<TWeakObjectPtr<UPCGGraphInstance>> SelectedGraphInstances;
 };
-
-#if UE_ENABLE_INCLUDE_ORDER_DEPRECATED_IN_5_2
-#include "Input/Reply.h"
-#include "UObject/WeakObjectPtr.h"
-#endif
