[PCG] Fix StructUtils API deprecated methods

#rnx

diff --git a/Engine/Plugins/Experimental/PCG/Source/PCG/Private/PCGGraph.cpp b/Engine/Plugins/Experimental/PCG/Source/PCG/Private/PCGGraph.cpp
--- a/Engine/Plugins/Experimental/PCG/Source/PCG/Private/PCGGraph.cpp
+++ b/Engine/Plugins/Experimental/PCG/Source/PCG/Private/PCGGraph.cpp
@@ -61,7 +61,7 @@ namespace PCGGraphUtils
 		}
 
 		const uint8* SourceValueAddress = InSourceInstance.GetValue().GetMemory() + InSourcePropertyDesc->CachedProperty->GetOffset_ForInternal();
-		uint8* TargetValueAddress = InTargetInstance.GetMutableValue().GetMutableMemory() + InTargetPropertyDesc->CachedProperty->GetOffset_ForInternal();
+		uint8* TargetValueAddress = InTargetInstance.GetMutableValue().GetMemory() + InTargetPropertyDesc->CachedProperty->GetOffset_ForInternal();
 
 		InSourcePropertyDesc->CachedProperty->CopyCompleteValue(TargetValueAddress, SourceValueAddress);
 	}
diff --git a/Engine/Plugins/Experimental/PCG/Source/PCG/Private/Tests/PCGGraphParametersTest.cpp b/Engine/Plugins/Experimental/PCG/Source/PCG/Private/Tests/PCGGraphParametersTest.cpp
--- a/Engine/Plugins/Experimental/PCG/Source/PCG/Private/Tests/PCGGraphParametersTest.cpp
+++ b/Engine/Plugins/Experimental/PCG/Source/PCG/Private/Tests/PCGGraphParametersTest.cpp
@@ -458,7 +458,7 @@ bool FPCGGraphSetValueUserParametersPropagates::RunTest(const FString& Parameter
 		const FPropertyBagPropertyDesc* PropertyDesc = Graph->GetUserParametersStruct()->FindPropertyDescByName(MyPropertyName);
 		FInstancedPropertyBag* UserParameters = const_cast<FInstancedPropertyBag*>(Graph->GetUserParametersStruct());
 
-		PropertyDesc->CachedProperty->SetValue_InContainer(UserParameters->GetMutableValue().GetMutableMemory(), &NewValue);
+		PropertyDesc->CachedProperty->SetValue_InContainer(UserParameters->GetMutableValue().GetMemory(), &NewValue);
 	});
 
 	auto Verification = [this, MyPropertyName, NewValue](UPCGGraphInstance* GraphInstance, const TCHAR* Name) -> bool
