[PCG] Fix duplicated localisation key

diff --git a/Engine/Plugins/Experimental/PCG/Source/PCG/Private/Elements/PCGTextureSampler.cpp b/Engine/Plugins/Experimental/PCG/Source/PCG/Private/Elements/PCGTextureSampler.cpp
--- a/Engine/Plugins/Experimental/PCG/Source/PCG/Private/Elements/PCGTextureSampler.cpp
+++ b/Engine/Plugins/Experimental/PCG/Source/PCG/Private/Elements/PCGTextureSampler.cpp
@@ -97,7 +97,7 @@ bool FPCGTextureSamplerElement::ExecuteInternal(FPCGContext* Context) const
 
 	if (!TextureData->IsValid())
 	{
-		PCGE_LOG(Error, GraphAndLog, LOCTEXT("UnsupportedTextureFormat", "Texture data failed to initialize, check log for more information"));
+		PCGE_LOG(Error, GraphAndLog, LOCTEXT("TextureDataInitFailed", "Texture data failed to initialize, check log for more information"));
 	}
 
 	return true;
