[PCG] Remove selecting actors by name which never worked, because name is different than *label*.

Actor labels are not available in shipping builds, so removing.

diff --git a/Engine/Plugins/Experimental/PCG/Source/PCG/Private/Elements/PCGActorSelector.cpp b/Engine/Plugins/Experimental/PCG/Source/PCG/Private/Elements/PCGActorSelector.cpp
--- a/Engine/Plugins/Experimental/PCG/Source/PCG/Private/Elements/PCGActorSelector.cpp
+++ b/Engine/Plugins/Experimental/PCG/Source/PCG/Private/Elements/PCGActorSelector.cpp
@@ -3,6 +3,7 @@
 #include "Elements/PCGActorSelector.h"
 
 #include "PCGComponent.h"
+#include "PCGModule.h"
 #include "Grid/PCGPartitionActor.h"
 #include "Helpers/PCGActorHelpers.h"
 
@@ -49,18 +50,6 @@ namespace PCGActorSelector
 				return true;
 			};
 
-		case EPCGActorSelection::ByName:
-			return[ActorSelectionName = InSettings.ActorSelectionName, &InFoundActors, bMultiSelect, &BoundsCheck](AActor* Actor) -> bool
-			{
-				if (Actor->GetFName().IsEqual(ActorSelectionName, ENameCase::IgnoreCase, /*bCompareNumber=*/ false) && BoundsCheck(Actor))
-				{
-					InFoundActors.Add(Actor);
-					return bMultiSelect;
-				}
-
-				return true;
-			};
-
 		case EPCGActorSelection::ByClass:
 			return[ActorSelectionClass = InSettings.ActorSelectionClass, &InFoundActors, bMultiSelect, &BoundsCheck](AActor* Actor) -> bool
 			{
@@ -73,6 +62,10 @@ namespace PCGActorSelector
 				return true;
 			};
 
+		case EPCGActorSelection::ByName:
+			UE_LOG(LogPCG, Error, TEXT("PCGActorSelector::GetFilteringFunction: Unsupported value for EPCGActorSelection - selection by name is no longer supported."));
+			break;
+
 		default:
 			break;
 		}
@@ -97,7 +90,6 @@ namespace PCGActorSelector
 		// Early out if we have not the information necessary.
 		if (FilterRequired(Settings) &&
 			((Settings.ActorSelection == EPCGActorSelection::ByTag && Settings.ActorSelectionTag == NAME_None) ||
-			(Settings.ActorSelection == EPCGActorSelection::ByName && Settings.ActorSelectionName == NAME_None) ||
 			(Settings.ActorSelection == EPCGActorSelection::ByClass && !Settings.ActorSelectionClass)))
 		{
 			return FoundActors;
diff --git a/Engine/Plugins/Experimental/PCG/Source/PCG/Private/Elements/PCGPropertyToParamData.cpp b/Engine/Plugins/Experimental/PCG/Source/PCG/Private/Elements/PCGPropertyToParamData.cpp
--- a/Engine/Plugins/Experimental/PCG/Source/PCG/Private/Elements/PCGPropertyToParamData.cpp
+++ b/Engine/Plugins/Experimental/PCG/Source/PCG/Private/Elements/PCGPropertyToParamData.cpp
@@ -40,7 +40,6 @@ void UPCGPropertyToParamDataSettings::PostLoad()
 	{
 		ActorSelector.ActorSelection = ActorSelection_DEPRECATED;
 		ActorSelector.ActorSelectionTag = ActorSelectionTag_DEPRECATED;
-		ActorSelector.ActorSelectionName = ActorSelectionName_DEPRECATED;
 		ActorSelector.ActorSelectionClass = ActorSelectionClass_DEPRECATED;
 		ActorSelector.ActorFilter = ActorFilter_DEPRECATED;
 		ActorSelector.bIncludeChildren = bIncludeChildren_DEPRECATED;
diff --git a/Engine/Plugins/Experimental/PCG/Source/PCG/Private/Tests/Elements/PCGPropertyToParamDataElementTest.cpp b/Engine/Plugins/Experimental/PCG/Source/PCG/Private/Tests/Elements/PCGPropertyToParamDataElementTest.cpp
--- a/Engine/Plugins/Experimental/PCG/Source/PCG/Private/Tests/Elements/PCGPropertyToParamDataElementTest.cpp
+++ b/Engine/Plugins/Experimental/PCG/Source/PCG/Private/Tests/Elements/PCGPropertyToParamDataElementTest.cpp
@@ -210,7 +210,6 @@ bool FPCGPropertyToParamDataActorFindTest::RunTest(const FString& Parameters)
 	UPCGPropertyToParamDataSettings* Settings = NewObject<UPCGPropertyToParamDataSettings>();
 	Settings->ActorSelector.ActorSelectionClass = APCGUnitTestDummyActor::StaticClass();
 	Settings->ActorSelector.ActorSelectionTag = Tag;
-	Settings->ActorSelector.ActorSelectionName = APCGUnitTestDummyActor::StaticClass()->GetFName();
 	Settings->ComponentClass = UPCGUnitTestDummyComponent::StaticClass();
 
 	// Self by class
@@ -236,17 +235,6 @@ bool FPCGPropertyToParamDataActorFindTest::RunTest(const FString& Parameters)
 		bSuccess &= VerifyAttributeValueValid(this, TestData, GET_MEMBER_NAME_CHECKED(APCGUnitTestDummyActor, IntProperty), 42ll, "PropertyToParamDataActorFindTest_Self_Tag");
 	}
 
-	// Self by name
-	{
-		Settings->ActorSelector.ActorSelection = EPCGActorSelection::ByName;
-		Settings->ActorSelector.ActorFilter = EPCGActorFilter::Self;
-
-		PCGTestsCommon::FTestData TestData(Seed, Settings, APCGUnitTestDummyActor::StaticClass());
-		Cast<APCGUnitTestDummyActor>(TestData.TestActor)->IntProperty = 42;
-
-		bSuccess &= VerifyAttributeValueValid(this, TestData, GET_MEMBER_NAME_CHECKED(APCGUnitTestDummyActor, IntProperty), 42ll, "PropertyToParamDataActorFindTest_Self_Name");
-	}
-
 	// TODO: Need a good way to spawn actors with parenting relation between them
 	//// Parent by class
 	//{
@@ -329,17 +317,6 @@ bool FPCGPropertyToParamDataActorFindTest::RunTest(const FString& Parameters)
 		bSuccess &= VerifyAttributeValueValid(this, TestData, GET_MEMBER_NAME_CHECKED(APCGUnitTestDummyActor, IntProperty), 42ll, "PropertyToParamDataActorFindTest_World_Tag");
 	}
 
-	// World by name
-	{
-		Settings->ActorSelector.ActorSelection = EPCGActorSelection::ByName;
-		Settings->ActorSelector.ActorFilter = EPCGActorFilter::AllWorldActors;
-
-		PCGTestsCommon::FTestData TestData(Seed, Settings, APCGUnitTestDummyActor::StaticClass());
-		Cast<APCGUnitTestDummyActor>(TestData.TestActor)->IntProperty = 42;
-
-		bSuccess &= VerifyAttributeValueValid(this, TestData, GET_MEMBER_NAME_CHECKED(APCGUnitTestDummyActor, IntProperty), 42ll, "PropertyToParamDataActorFindTest_World_Name");
-	}
-
 	return bSuccess;
 }
 
diff --git a/Engine/Plugins/Experimental/PCG/Source/PCG/Public/Elements/PCGActorSelector.h b/Engine/Plugins/Experimental/PCG/Source/PCG/Public/Elements/PCGActorSelector.h
--- a/Engine/Plugins/Experimental/PCG/Source/PCG/Public/Elements/PCGActorSelector.h
+++ b/Engine/Plugins/Experimental/PCG/Source/PCG/Public/Elements/PCGActorSelector.h
@@ -14,7 +14,8 @@ UENUM()
 enum class EPCGActorSelection : uint8
 {
 	ByTag,
-	ByName,
+	// Deprecated - actor labels are unavailable in shipping builds
+	ByName UMETA(Hidden),
 	ByClass
 };
 
@@ -59,9 +60,6 @@ struct FPCGActorSelectorSettings
 	UPROPERTY(BlueprintReadWrite, EditAnywhere, Category = Settings, meta = (EditCondition = "(ActorFilter==EPCGActorFilter::AllWorldActors || (bIncludeChildren && !bDisableFilter)) && ActorSelection==EPCGActorSelection::ByTag", EditConditionHides))
 	FName ActorSelectionTag;
 
-	UPROPERTY(BlueprintReadWrite, EditAnywhere, Category = Settings, meta = (EditCondition = "(ActorFilter==EPCGActorFilter::AllWorldActors || (bIncludeChildren && !bDisableFilter)) && ActorSelection==EPCGActorSelection::ByName", EditConditionHides))
-	FName ActorSelectionName;
-
 	UPROPERTY(BlueprintReadWrite, EditAnywhere, Category = Settings, meta = (EditCondition = "(ActorFilter==EPCGActorFilter::AllWorldActors || (bIncludeChildren && !bDisableFilter)) && ActorSelection==EPCGActorSelection::ByClass", EditConditionHides, AllowAbstract = "true"))
 	TSubclassOf<AActor> ActorSelectionClass;
 
