[PCG] Workaround the limitation in StructUtils with BP Reconstruction script and PCG components

* StructUtils UI detail view is using a temporary struct to make changes on InstancedStruct
* When PostEditChangeProperty is called in the PCG Component, the Override Parameters struct is not yet modified (only the temp struct is).
* It triggers a reconstruction script, copying the parameters without the value modified
* When the modification is done on the Override Parameters struct, it is on the old component, marked trash
* This workaround will not call for a reconstruction script in the case of the value changed.
* Reconstruction will happen on the second call to PostEditChangeProperty, when the actual struct was modified and contain the right value.

#rnx

diff --git a/Engine/Plugins/Experimental/PCG/Source/PCG/Private/PCGComponent.cpp b/Engine/Plugins/Experimental/PCG/Source/PCG/Private/PCGComponent.cpp
--- a/Engine/Plugins/Experimental/PCG/Source/PCG/Private/PCGComponent.cpp
+++ b/Engine/Plugins/Experimental/PCG/Source/PCG/Private/PCGComponent.cpp
@@ -1248,6 +1248,91 @@ void UPCGComponent::PreEditChange(FProperty* PropertyAboutToChange)
 	Super::PreEditChange(PropertyAboutToChange);
 }
 
+/**
+* Temporary workaround, while we are waiting for UE-182059
+* On the UI side of StructUtils, property widgets are referencing a temporary structure, not the one that is owned by the Graph/GraphInstance.
+* After the value is modified and the callback is propagated, the temporary structure is modified, but the one owned by the Graph is not.
+* If we keep it like this, PostEditChangeChainProperty will call for a reconstruction of the component if it is on a Blueprint
+* and we will go through the process of copying and transfering data between the 2 components (old and new).
+* Because of that, we copy the old structure, not modified. And when the structure is actually modified, the callback is called on the old component,
+* already marked trash, and we lose the change.
+* The workaround is to detect when a value from our parameters is changed, and just discard the call to the super method, avoiding a reconstruction script.
+* Note that it should be OK to do so in the current situation, because there is another callback fired just after the copy of the temp 
+* structure to the real one, and on this callback we are calling the super method and go through construction script.
+* We do this workaround for pre and post edit, since we can run into trouble if a pre edit was processed, but post was not.
+*/
+bool UPCGComponent::IsChangingGraphInstanceParameterValue(FEditPropertyChain& InEditPropertyChain) const
+{
+	if (!GraphInstance)
+	{
+		return false;
+	}
+
+	FEditPropertyChain::TDoubleLinkedListNode* PropertyNode = InEditPropertyChain.GetActiveNode();
+
+	if (!PropertyNode)
+	{
+		return false;
+	}
+
+	constexpr int32 PropertyPathSize = 4;
+	static const FName PropertyPath[PropertyPathSize] = {
+		GET_MEMBER_NAME_CHECKED(UPCGComponent, GraphInstance),
+		GET_MEMBER_NAME_CHECKED(UPCGGraphInstance, ParametersOverrides),
+		GET_MEMBER_NAME_CHECKED(FPCGOverrideInstancedPropertyBag, Parameters),
+		FName(TEXT("Value"))
+	};
+
+	int32 Index = 0;
+
+	while (Index < PropertyPathSize && PropertyNode)
+	{
+		FProperty* Property = PropertyNode->GetValue();
+		if (Property && Property->GetFName() == PropertyPath[Index])
+		{
+			PropertyNode = PropertyNode->GetNextNode();
+			++Index;
+		}
+		else
+		{
+			PropertyNode = nullptr;
+		}
+	}
+
+	if (Index == PropertyPathSize && PropertyNode)
+	{
+		FProperty* Property = PropertyNode->GetValue();
+		if (Property && Property->GetOwnerStruct() == GraphInstance->ParametersOverrides.Parameters.GetPropertyBagStruct())
+		{
+			return true;
+		}
+	}
+
+	return false;
+}
+
+void UPCGComponent::PreEditChange(FEditPropertyChain& PropertyAboutToChange)
+{
+	if (IsChangingGraphInstanceParameterValue(PropertyAboutToChange))
+	{
+		// Skip pre edit, cf IsChangingGraphInstanceParameterValue comment
+		return;
+	}
+
+	UObject::PreEditChange(PropertyAboutToChange);
+}
+
+void UPCGComponent::PostEditChangeChainProperty(FPropertyChangedChainEvent& PropertyChangedEvent)
+{
+	if (IsChangingGraphInstanceParameterValue(PropertyChangedEvent.PropertyChain))
+	{
+		// Skip post edit, cf IsChangingGraphInstanceParameterValue comment
+		return;
+	}
+
+	Super::PostEditChangeChainProperty(PropertyChangedEvent);
+}
+
 void UPCGComponent::PostEditChangeProperty(FPropertyChangedEvent& PropertyChangedEvent)
 {
 	Super::PostEditChangeProperty(PropertyChangedEvent);
diff --git a/Engine/Plugins/Experimental/PCG/Source/PCG/Public/PCGComponent.h b/Engine/Plugins/Experimental/PCG/Source/PCG/Public/PCGComponent.h
--- a/Engine/Plugins/Experimental/PCG/Source/PCG/Public/PCGComponent.h
+++ b/Engine/Plugins/Experimental/PCG/Source/PCG/Public/PCGComponent.h
@@ -327,10 +327,14 @@ private:
 
 #if WITH_EDITOR
 	virtual void PreEditChange(FProperty* PropertyAboutToChange) override;
+	virtual void PreEditChange(FEditPropertyChain& PropertyAboutToChange) override;
+	virtual void PostEditChangeChainProperty(FPropertyChangedChainEvent& PropertyChangedEvent) override;
 	virtual void PostEditChangeProperty(FPropertyChangedEvent& PropertyChangedEvent) override;
 	virtual void PreEditUndo() override;
 	virtual void PostEditUndo() override;
 
+	bool IsChangingGraphInstanceParameterValue(FEditPropertyChain& InEditPropertyChain) const;
+
 	/** Sets up actor, tracking, landscape and graph callbacks */
 	void SetupCallbacksOnCreation();
 
