[PCG] Housekeeping BP reconstruction for PCG Components

* Reconnect/Teardown callbacks added to PCG Graph
* bDirtyGenerated flag is transient and needs to be copied to the new component
* SetGraph will not dirty the component if the graph is the same.

#rnx

diff --git a/Engine/Plugins/Experimental/PCG/Source/PCG/Private/PCGComponent.cpp b/Engine/Plugins/Experimental/PCG/Source/PCG/Private/PCGComponent.cpp
--- a/Engine/Plugins/Experimental/PCG/Source/PCG/Private/PCGComponent.cpp
+++ b/Engine/Plugins/Experimental/PCG/Source/PCG/Private/PCGComponent.cpp
@@ -2946,6 +2946,11 @@ void FPCGComponentInstanceData::ApplyToComponent(UActorComponent* Component, con
 				}
 			}
 
+#if WITH_EDITOR
+			// bDirtyGenerated is transient.
+			PCGComponent->bDirtyGenerated = SourceComponent->bDirtyGenerated; 
+#endif // WITH_EDITOR
+
 			// Non-critical but should be done: transient data, excluded & tracked actors cache, landscape tracking
 			// TODO Validate usefulness + move accordingly
 		}
@@ -2967,8 +2972,10 @@ void FPCGComponentInstanceData::ApplyToComponent(UActorComponent* Component, con
 			PCGComponent->SetManagedResources(DuplicatedResources);
 		}
 
-		// Re-set the graph to reconnect callbacks correctly
-		PCGComponent->GraphInstance->SetGraph(PCGComponent->GraphInstance->Graph);
+#if WITH_EDITOR
+		// Reconnect callbacks
+		PCGComponent->GraphInstance->FixCallbacks();
+#endif // WITH_EDITOR
 
 		bool bDoActorMapping = PCGComponent->bGenerated || PCGHelpers::IsRuntimeOrPIE();
 
@@ -2980,6 +2987,12 @@ void FPCGComponentInstanceData::ApplyToComponent(UActorComponent* Component, con
 		}
 
 #if WITH_EDITOR
+		// Disconnect callbacks on source.
+		if (SourceComponent && SourceComponent->GraphInstance)
+		{
+			SourceComponent->GraphInstance->TeardownCallbacks();
+		}
+
 		// Finally, start a delayed refresh task (if there is not one already), in editor only
 		// It is important to be delayed, because we cannot spawn Partition Actors within this scope,
 		// because we are in a construction script.
diff --git a/Engine/Plugins/Experimental/PCG/Source/PCG/Private/PCGGraph.cpp b/Engine/Plugins/Experimental/PCG/Source/PCG/Private/PCGGraph.cpp
--- a/Engine/Plugins/Experimental/PCG/Source/PCG/Private/PCGGraph.cpp
+++ b/Engine/Plugins/Experimental/PCG/Source/PCG/Private/PCGGraph.cpp
@@ -1076,6 +1076,27 @@ bool UPCGGraphInstance::CanEditChange(const FProperty* InProperty) const
 
 	return true;
 }
+
+void UPCGGraphInstance::TeardownCallbacks()
+{
+	if (Graph)
+	{
+		Graph->OnGraphChangedDelegate.RemoveAll(this);
+		Graph->OnGraphParametersChangedDelegate.RemoveAll(this);
+	}
+}
+
+void UPCGGraphInstance::FixCallbacks()
+{
+	// Start from a clean state.
+	TeardownCallbacks();
+
+	if (Graph)
+	{
+		Graph->OnGraphChangedDelegate.AddUObject(this, &UPCGGraphInstance::OnGraphChanged);
+		Graph->OnGraphParametersChangedDelegate.AddUObject(this, &UPCGGraphInstance::OnGraphParametersChanged);
+	}
+}
 #endif
 
 void UPCGGraphInstance::SetGraph(UPCGGraphInterface* InGraph)
@@ -1086,6 +1107,12 @@ void UPCGGraphInstance::SetGraph(UPCGGraphInterface* InGraph)
 		return;
 	}
 
+	if (InGraph == Graph)
+	{
+		// Nothing to do
+		return;
+	}
+
 #if WITH_EDITOR
 	if (Graph)
 	{
diff --git a/Engine/Plugins/Experimental/PCG/Source/PCG/Public/PCGGraph.h b/Engine/Plugins/Experimental/PCG/Source/PCG/Public/PCGGraph.h
--- a/Engine/Plugins/Experimental/PCG/Source/PCG/Public/PCGGraph.h
+++ b/Engine/Plugins/Experimental/PCG/Source/PCG/Public/PCGGraph.h
@@ -275,9 +275,14 @@ public:
 	virtual bool CanEditChange(const FProperty* InProperty) const override;
 	virtual void PreEditUndo() override;
 	virtual void PostEditUndo() override;
-#endif
 	// ~End UObject interface
 
+	// Reconstruction script specific, to fix callbacks on reconstructed component
+	// and teardown callback on trashed component.
+	void FixCallbacks();
+	void TeardownCallbacks();
+#endif
+
 protected:
 #if WITH_EDITOR
 	void OnGraphChanged(UPCGGraphInterface* InGraph, EPCGChangeType ChangeType);
