Added localized node titles

diff --git a/Engine/Plugins/Experimental/PCG/Source/PCG/Private/Elements/Metadata/PCGMetadataBitwiseOpElement.cpp b/Engine/Plugins/Experimental/PCG/Source/PCG/Private/Elements/Metadata/PCGMetadataBitwiseOpElement.cpp
--- a/Engine/Plugins/Experimental/PCG/Source/PCG/Private/Elements/Metadata/PCGMetadataBitwiseOpElement.cpp
+++ b/Engine/Plugins/Experimental/PCG/Source/PCG/Private/Elements/Metadata/PCGMetadataBitwiseOpElement.cpp
@@ -100,7 +100,12 @@ FName UPCGMetadataBitwiseSettings::AdditionalTaskName() const
 #if WITH_EDITOR
 FName UPCGMetadataBitwiseSettings::GetDefaultNodeName() const
 {
-	return TEXT("Attribute Bitwise Op");
+	return TEXT("AttributeBitwiseOp");
+}
+
+FText UPCGMetadataBitwiseSettings::GetDefaultNodeTitle() const
+{
+	return NSLOCTEXT("PCGMetadataBitwiseSettings", "NodeTitle", "Attribute Bitwise Op");
 }
 #endif // WITH_EDITOR
 
diff --git a/Engine/Plugins/Experimental/PCG/Source/PCG/Private/Elements/Metadata/PCGMetadataBooleanOpElement.cpp b/Engine/Plugins/Experimental/PCG/Source/PCG/Private/Elements/Metadata/PCGMetadataBooleanOpElement.cpp
--- a/Engine/Plugins/Experimental/PCG/Source/PCG/Private/Elements/Metadata/PCGMetadataBooleanOpElement.cpp
+++ b/Engine/Plugins/Experimental/PCG/Source/PCG/Private/Elements/Metadata/PCGMetadataBooleanOpElement.cpp
@@ -105,7 +105,12 @@ FName UPCGMetadataBooleanSettings::AdditionalTaskName() const
 #if WITH_EDITOR
 FName UPCGMetadataBooleanSettings::GetDefaultNodeName() const
 {
-	return TEXT("Attribute Boolean Op");
+	return TEXT("AttributeBooleanOp");
+}
+
+FText UPCGMetadataBooleanSettings::GetDefaultNodeTitle() const
+{
+	return NSLOCTEXT("PCGMetadataBooleanSettings", "NodeTitle", "Attribute Boolean Op");
 }
 #endif // WITH_EDITOR
 
diff --git a/Engine/Plugins/Experimental/PCG/Source/PCG/Private/Elements/Metadata/PCGMetadataBreakTransform.cpp b/Engine/Plugins/Experimental/PCG/Source/PCG/Private/Elements/Metadata/PCGMetadataBreakTransform.cpp
--- a/Engine/Plugins/Experimental/PCG/Source/PCG/Private/Elements/Metadata/PCGMetadataBreakTransform.cpp
+++ b/Engine/Plugins/Experimental/PCG/Source/PCG/Private/Elements/Metadata/PCGMetadataBreakTransform.cpp
@@ -93,7 +93,12 @@ FName UPCGMetadataBreakTransformSettings::GetOutputAttributeName(FName BaseName,
 #if WITH_EDITOR
 FName UPCGMetadataBreakTransformSettings::GetDefaultNodeName() const
 {
-	return TEXT("Break Transform Attribute");
+	return TEXT("BreakTransformAttribute");
+}
+
+FText UPCGMetadataBreakTransformSettings::GetDefaultNodeTitle() const
+{
+	return NSLOCTEXT("PCGMetadataBreakTransformSettings", "NodeTitle", "Break Transform Attribute");
 }
 #endif // WITH_EDITOR
 
diff --git a/Engine/Plugins/Experimental/PCG/Source/PCG/Private/Elements/Metadata/PCGMetadataBreakVector.cpp b/Engine/Plugins/Experimental/PCG/Source/PCG/Private/Elements/Metadata/PCGMetadataBreakVector.cpp
--- a/Engine/Plugins/Experimental/PCG/Source/PCG/Private/Elements/Metadata/PCGMetadataBreakVector.cpp
+++ b/Engine/Plugins/Experimental/PCG/Source/PCG/Private/Elements/Metadata/PCGMetadataBreakVector.cpp
@@ -146,7 +146,12 @@ FName UPCGMetadataBreakVectorSettings::GetOutputAttributeName(FName BaseName, ui
 #if WITH_EDITOR
 FName UPCGMetadataBreakVectorSettings::GetDefaultNodeName() const
 {
-	return TEXT("Break Vector Attribute");
+	return TEXT("BreakVectorAttribute");
+}
+
+FText UPCGMetadataBreakVectorSettings::GetDefaultNodeTitle() const
+{
+	return NSLOCTEXT("PCGMetadataBreakVectorSettings", "NodeTitle", "Break Vector Attribute");
 }
 #endif // WITH_EDITOR
 
diff --git a/Engine/Plugins/Experimental/PCG/Source/PCG/Private/Elements/Metadata/PCGMetadataCompareOpElement.cpp b/Engine/Plugins/Experimental/PCG/Source/PCG/Private/Elements/Metadata/PCGMetadataCompareOpElement.cpp
--- a/Engine/Plugins/Experimental/PCG/Source/PCG/Private/Elements/Metadata/PCGMetadataCompareOpElement.cpp
+++ b/Engine/Plugins/Experimental/PCG/Source/PCG/Private/Elements/Metadata/PCGMetadataCompareOpElement.cpp
@@ -131,7 +131,12 @@ FName UPCGMetadataCompareSettings::AdditionalTaskName() const
 #if WITH_EDITOR
 FName UPCGMetadataCompareSettings::GetDefaultNodeName() const
 {
-	return TEXT("Attribute Compare Op");
+	return TEXT("AttributeCompareOp");
+}
+
+FText UPCGMetadataCompareSettings::GetDefaultNodeTitle() const
+{
+	return NSLOCTEXT("PCGMetadataCompareSettings", "NodeTitle", "Attribute Compare Op");
 }
 #endif // WITH_EDITOR
 
diff --git a/Engine/Plugins/Experimental/PCG/Source/PCG/Private/Elements/Metadata/PCGMetadataMakeTransform.cpp b/Engine/Plugins/Experimental/PCG/Source/PCG/Private/Elements/Metadata/PCGMetadataMakeTransform.cpp
--- a/Engine/Plugins/Experimental/PCG/Source/PCG/Private/Elements/Metadata/PCGMetadataMakeTransform.cpp
+++ b/Engine/Plugins/Experimental/PCG/Source/PCG/Private/Elements/Metadata/PCGMetadataMakeTransform.cpp
@@ -109,7 +109,12 @@ uint16 UPCGMetadataMakeTransformSettings::GetOutputType(uint16 InputTypeId) cons
 #if WITH_EDITOR
 FName UPCGMetadataMakeTransformSettings::GetDefaultNodeName() const
 {
-	return TEXT("Make Transform Attribute");
+	return TEXT("MakeTransformAttribute");
+}
+
+FText UPCGMetadataMakeTransformSettings::GetDefaultNodeTitle() const
+{
+	return NSLOCTEXT("PCGMetadataMakeTransformSettings", "NodeTitle", "Make Transform Attribute");
 }
 #endif // WITH_EDITOR
 
diff --git a/Engine/Plugins/Experimental/PCG/Source/PCG/Private/Elements/Metadata/PCGMetadataMakeVector.cpp b/Engine/Plugins/Experimental/PCG/Source/PCG/Private/Elements/Metadata/PCGMetadataMakeVector.cpp
--- a/Engine/Plugins/Experimental/PCG/Source/PCG/Private/Elements/Metadata/PCGMetadataMakeVector.cpp
+++ b/Engine/Plugins/Experimental/PCG/Source/PCG/Private/Elements/Metadata/PCGMetadataMakeVector.cpp
@@ -184,7 +184,12 @@ uint16 UPCGMetadataMakeVectorSettings::GetOutputType(uint16 InputTypeId) const
 #if WITH_EDITOR
 FName UPCGMetadataMakeVectorSettings::GetDefaultNodeName() const
 {
-	return TEXT("Make Vector Attribute");
+	return TEXT("MakeVectorAttribute");
+}
+
+FText UPCGMetadataMakeVectorSettings::GetDefaultNodeTitle() const
+{
+	return NSLOCTEXT("PCGMetadataMakeVectorSettings", "NodeTitle", "Make Vector Attribute");
 }
 #endif // WITH_EDITOR
 
diff --git a/Engine/Plugins/Experimental/PCG/Source/PCG/Private/Elements/Metadata/PCGMetadataMathsOpElement.cpp b/Engine/Plugins/Experimental/PCG/Source/PCG/Private/Elements/Metadata/PCGMetadataMathsOpElement.cpp
--- a/Engine/Plugins/Experimental/PCG/Source/PCG/Private/Elements/Metadata/PCGMetadataMathsOpElement.cpp
+++ b/Engine/Plugins/Experimental/PCG/Source/PCG/Private/Elements/Metadata/PCGMetadataMathsOpElement.cpp
@@ -247,7 +247,12 @@ FName UPCGMetadataMathsSettings::AdditionalTaskName() const
 #if WITH_EDITOR
 FName UPCGMetadataMathsSettings::GetDefaultNodeName() const
 {
-	return TEXT("Attribute Maths Op");
+	return TEXT("AttributeMathsOp");
+}
+
+FText UPCGMetadataMathsSettings::GetDefaultNodeTitle() const
+{
+	return NSLOCTEXT("PCGMetadataMathsSettings", "NodeTitle", "Attribute Maths Op");
 }
 #endif // WITH_EDITOR
 
diff --git a/Engine/Plugins/Experimental/PCG/Source/PCG/Private/Elements/Metadata/PCGMetadataRotatorOpElement.cpp b/Engine/Plugins/Experimental/PCG/Source/PCG/Private/Elements/Metadata/PCGMetadataRotatorOpElement.cpp
--- a/Engine/Plugins/Experimental/PCG/Source/PCG/Private/Elements/Metadata/PCGMetadataRotatorOpElement.cpp
+++ b/Engine/Plugins/Experimental/PCG/Source/PCG/Private/Elements/Metadata/PCGMetadataRotatorOpElement.cpp
@@ -203,7 +203,12 @@ FName UPCGMetadataRotatorSettings::AdditionalTaskName() const
 #if WITH_EDITOR
 FName UPCGMetadataRotatorSettings::GetDefaultNodeName() const
 {
-	return TEXT("Attribute Rotator Op");
+	return TEXT("AttributeRotatorOp");
+}
+
+FText UPCGMetadataRotatorSettings::GetDefaultNodeTitle() const
+{
+	return NSLOCTEXT("PCGMetadataRotatorSettings", "NodeTitle", "Attribute Rotator Op");
 }
 #endif // WITH_EDITOR
 
diff --git a/Engine/Plugins/Experimental/PCG/Source/PCG/Private/Elements/Metadata/PCGMetadataTransformOpElement.cpp b/Engine/Plugins/Experimental/PCG/Source/PCG/Private/Elements/Metadata/PCGMetadataTransformOpElement.cpp
--- a/Engine/Plugins/Experimental/PCG/Source/PCG/Private/Elements/Metadata/PCGMetadataTransformOpElement.cpp
+++ b/Engine/Plugins/Experimental/PCG/Source/PCG/Private/Elements/Metadata/PCGMetadataTransformOpElement.cpp
@@ -171,7 +171,12 @@ FName UPCGMetadataTransformSettings::AdditionalTaskName() const
 #if WITH_EDITOR
 FName UPCGMetadataTransformSettings::GetDefaultNodeName() const
 {
-	return TEXT("Attribute Transform Op");
+	return TEXT("AttributeTransformOp");
+}
+
+FText UPCGMetadataTransformSettings::GetDefaultNodeTitle() const
+{
+	return NSLOCTEXT("PCGMetadataTransformSettings", "NodeTitle", "Attribute Transform Op");
 }
 #endif // WITH_EDITOR
 
diff --git a/Engine/Plugins/Experimental/PCG/Source/PCG/Private/Elements/Metadata/PCGMetadataTrigOpElement.cpp b/Engine/Plugins/Experimental/PCG/Source/PCG/Private/Elements/Metadata/PCGMetadataTrigOpElement.cpp
--- a/Engine/Plugins/Experimental/PCG/Source/PCG/Private/Elements/Metadata/PCGMetadataTrigOpElement.cpp
+++ b/Engine/Plugins/Experimental/PCG/Source/PCG/Private/Elements/Metadata/PCGMetadataTrigOpElement.cpp
@@ -117,7 +117,12 @@ FName UPCGMetadataTrigSettings::AdditionalTaskName() const
 #if WITH_EDITOR
 FName UPCGMetadataTrigSettings::GetDefaultNodeName() const
 {
-	return TEXT("Attribute Trig Op");
+	return TEXT("AttributeTrigOp");
+}
+
+FText UPCGMetadataTrigSettings::GetDefaultNodeTitle() const
+{
+	return NSLOCTEXT("PCGMetadataTrigSettings", "NodeTitle", "Attribute Trig Op");
 }
 #endif // WITH_EDITOR
 
diff --git a/Engine/Plugins/Experimental/PCG/Source/PCG/Private/Elements/Metadata/PCGMetadataVectorOpElement.cpp b/Engine/Plugins/Experimental/PCG/Source/PCG/Private/Elements/Metadata/PCGMetadataVectorOpElement.cpp
--- a/Engine/Plugins/Experimental/PCG/Source/PCG/Private/Elements/Metadata/PCGMetadataVectorOpElement.cpp
+++ b/Engine/Plugins/Experimental/PCG/Source/PCG/Private/Elements/Metadata/PCGMetadataVectorOpElement.cpp
@@ -333,7 +333,12 @@ FName UPCGMetadataVectorSettings::AdditionalTaskName() const
 #if WITH_EDITOR
 FName UPCGMetadataVectorSettings::GetDefaultNodeName() const
 {
-	return TEXT("Attribute Vector Op");
+	return TEXT("AttributeVectorOp");
+}
+
+FText UPCGMetadataVectorSettings::GetDefaultNodeTitle() const
+{
+	return NSLOCTEXT("PCGMetadataVectorSettings", "NodeTitle", "Attribute Vector Op");
 }
 #endif // WITH_EDITOR
 
diff --git a/Engine/Plugins/Experimental/PCG/Source/PCG/Private/Elements/PCGAttributeFilterElement.cpp b/Engine/Plugins/Experimental/PCG/Source/PCG/Private/Elements/PCGAttributeFilterElement.cpp
--- a/Engine/Plugins/Experimental/PCG/Source/PCG/Private/Elements/PCGAttributeFilterElement.cpp
+++ b/Engine/Plugins/Experimental/PCG/Source/PCG/Private/Elements/PCGAttributeFilterElement.cpp
@@ -15,6 +15,7 @@
 namespace PCGAttributeFilterConstants
 {
 	const FName NodeName = TEXT("FilterAttribute");
+	const FText NodeTitle = LOCTEXT("NodeTitle", "Filter Attribute");
 }
 
 void UPCGAttributeFilterSettings::PostLoad()
@@ -47,6 +48,11 @@ FName UPCGAttributeFilterSettings::GetDefaultNodeName() const
 {
 	return PCGAttributeFilterConstants::NodeName;
 }
+
+FText UPCGAttributeFilterSettings::GetDefaultNodeTitle() const
+{
+	return PCGAttributeFilterConstants::NodeTitle;
+}
 #endif
 
 FName UPCGAttributeFilterSettings::AdditionalTaskName() const
diff --git a/Engine/Plugins/Experimental/PCG/Source/PCG/Private/Elements/PCGAttributeGetFromPointIndexElement.cpp b/Engine/Plugins/Experimental/PCG/Source/PCG/Private/Elements/PCGAttributeGetFromPointIndexElement.cpp
--- a/Engine/Plugins/Experimental/PCG/Source/PCG/Private/Elements/PCGAttributeGetFromPointIndexElement.cpp
+++ b/Engine/Plugins/Experimental/PCG/Source/PCG/Private/Elements/PCGAttributeGetFromPointIndexElement.cpp
@@ -19,6 +19,11 @@ FName UPCGAttributeGetFromPointIndexSettings::GetDefaultNodeName() const
 {
 	return TEXT("GetAttributeFromPointIndex");
 }
+
+FText UPCGAttributeGetFromPointIndexSettings::GetDefaultNodeTitle() const
+{
+	return LOCTEXT("NodeTitle", "Get Attribute From Point Index");
+}
 #endif
 
 void UPCGAttributeGetFromPointIndexSettings::PostLoad()
diff --git a/Engine/Plugins/Experimental/PCG/Source/PCG/Private/Elements/PCGAttributeReduceElement.cpp b/Engine/Plugins/Experimental/PCG/Source/PCG/Private/Elements/PCGAttributeReduceElement.cpp
--- a/Engine/Plugins/Experimental/PCG/Source/PCG/Private/Elements/PCGAttributeReduceElement.cpp
+++ b/Engine/Plugins/Experimental/PCG/Source/PCG/Private/Elements/PCGAttributeReduceElement.cpp
@@ -81,6 +81,11 @@ FName UPCGAttributeReduceSettings::GetDefaultNodeName() const
 {
 	return TEXT("AttributeReduce");
 }
+
+FText UPCGAttributeReduceSettings::GetDefaultNodeTitle() const
+{
+	return LOCTEXT("NodeTitle", "Attribute Reduce");
+}
 #endif
 
 void UPCGAttributeReduceSettings::PostLoad()
diff --git a/Engine/Plugins/Experimental/PCG/Source/PCG/Private/Elements/PCGAttributeSelectElement.cpp b/Engine/Plugins/Experimental/PCG/Source/PCG/Private/Elements/PCGAttributeSelectElement.cpp
--- a/Engine/Plugins/Experimental/PCG/Source/PCG/Private/Elements/PCGAttributeSelectElement.cpp
+++ b/Engine/Plugins/Experimental/PCG/Source/PCG/Private/Elements/PCGAttributeSelectElement.cpp
@@ -186,6 +186,11 @@ FName UPCGAttributeSelectSettings::GetDefaultNodeName() const
 {
 	return TEXT("AttributeSelect");
 }
+
+FText UPCGAttributeSelectSettings::GetDefaultNodeTitle() const
+{
+	return LOCTEXT("NodeTitle", "Attribute Select");
+}
 #endif
 
 FName UPCGAttributeSelectSettings::AdditionalTaskName() const
diff --git a/Engine/Plugins/Experimental/PCG/Source/PCG/Private/Elements/PCGAttributeTransferElement.cpp b/Engine/Plugins/Experimental/PCG/Source/PCG/Private/Elements/PCGAttributeTransferElement.cpp
--- a/Engine/Plugins/Experimental/PCG/Source/PCG/Private/Elements/PCGAttributeTransferElement.cpp
+++ b/Engine/Plugins/Experimental/PCG/Source/PCG/Private/Elements/PCGAttributeTransferElement.cpp
@@ -16,6 +16,11 @@ FName UPCGAttributeTransferSettings::GetDefaultNodeName() const
 {
 	return PCGAttributeTransferConstants::NodeName;
 }
+
+FText UPCGAttributeTransferSettings::GetDefaultNodeTitle() const
+{
+	return PCGAttributeTransferConstants::NodeTitle;
+}
 #endif
 
 FName UPCGAttributeTransferSettings::AdditionalTaskName() const
diff --git a/Engine/Plugins/Experimental/PCG/Source/PCG/Private/Elements/PCGCreateAttribute.cpp b/Engine/Plugins/Experimental/PCG/Source/PCG/Private/Elements/PCGCreateAttribute.cpp
--- a/Engine/Plugins/Experimental/PCG/Source/PCG/Private/Elements/PCGCreateAttribute.cpp
+++ b/Engine/Plugins/Experimental/PCG/Source/PCG/Private/Elements/PCGCreateAttribute.cpp
@@ -14,8 +14,10 @@
 
 namespace PCGCreateAttributeConstants
 {
-	const FName NodeNameAddAttribute = TEXT("Add Attribute");
-	const FName NodeNameCreateAttribute = TEXT("Create Attribute");
+	const FName NodeNameAddAttribute = TEXT("AddAttribute");
+	const FText NodeTitleAddAttribute = LOCTEXT("NodeTitleAddAttribute", "Add Attribute");
+	const FName NodeNameCreateAttribute = TEXT("CreateAttribute");
+	const FText NodeTitleCreateAttribute = LOCTEXT("NodeTitleCreateAttribute", "Create Attribute");
 	const FName AttributesLabel = TEXT("Attributes");
 }
 
@@ -57,6 +59,11 @@ FName UPCGCreateAttributeSettings::GetDefaultNodeName() const
 {
 	return PCGCreateAttributeConstants::NodeNameAddAttribute;
 }
+
+FText UPCGCreateAttributeSettings::GetDefaultNodeTitle() const
+{
+	return PCGCreateAttributeConstants::NodeTitleAddAttribute;
+}
 #endif // WITH_EDITOR
 
 TArray<FPCGPinProperties> UPCGCreateAttributeSettings::InputPinProperties() const
@@ -131,6 +138,11 @@ FName UPCGCreateAttributeSetSettings::GetDefaultNodeName() const
 {
 	return PCGCreateAttributeConstants::NodeNameCreateAttribute;
 }
+
+FText UPCGCreateAttributeSetSettings::GetDefaultNodeTitle() const
+{
+	return PCGCreateAttributeConstants::NodeTitleCreateAttribute;
+}
 #endif // WITH_EDITOR
 
 TArray<FPCGPinProperties> UPCGCreateAttributeSetSettings::OutputPinProperties() const
diff --git a/Engine/Plugins/Experimental/PCG/Source/PCG/Private/Elements/PCGFilterByTag.cpp b/Engine/Plugins/Experimental/PCG/Source/PCG/Private/Elements/PCGFilterByTag.cpp
--- a/Engine/Plugins/Experimental/PCG/Source/PCG/Private/Elements/PCGFilterByTag.cpp
+++ b/Engine/Plugins/Experimental/PCG/Source/PCG/Private/Elements/PCGFilterByTag.cpp
@@ -11,7 +11,8 @@
 
 namespace PCGFilterByTagConstants
 {
-	const FName NodeName = FName(TEXT("Filter By Tag"));
+	const FName NodeName = FName(TEXT("FilterByTag"));
+	const FText NodeTitle = LOCTEXT("NodeTitle", "Filter By Tag");
 }
 
 #if WITH_EDITOR
@@ -20,6 +21,11 @@ FName UPCGFilterByTagSettings::GetDefaultNodeName() const
 	return PCGFilterByTagConstants::NodeName;
 }
 
+FText UPCGFilterByTagSettings::GetDefaultNodeTitle() const
+{
+	return PCGFilterByTagConstants::NodeTitle;
+}
+
 FText UPCGFilterByTagSettings::GetNodeTooltipText() const
 {
 	return LOCTEXT("FilterByTagNodeTooltip", "Filters data in the collection according to whether they have, or don't have, some tags");
diff --git a/Engine/Plugins/Experimental/PCG/Source/PCG/Private/Graph/PCGGraphCache.cpp b/Engine/Plugins/Experimental/PCG/Source/PCG/Private/Graph/PCGGraphCache.cpp
--- a/Engine/Plugins/Experimental/PCG/Source/PCG/Private/Graph/PCGGraphCache.cpp
+++ b/Engine/Plugins/Experimental/PCG/Source/PCG/Private/Graph/PCGGraphCache.cpp
@@ -170,7 +170,7 @@ void FPCGGraphCache::CleanFromCache(const IPCGElement* InElement, const UPCGSett
 
 	if (IsDebuggingEnabled())
 	{
-		UE_LOG(LogPCG, Warning, TEXT("[] \t\tCACHE: PURGED [%s]"), InSettings ? *InSettings->GetDefaultNodeName().ToString() : TEXT("AnonymousElement"));
+		UE_LOG(LogPCG, Warning, TEXT("[] \t\tCACHE: PURGED [%s]"), InSettings ? *InSettings->GetDefaultNodeTitle().ToString() : TEXT("AnonymousElement"));
 	}
 
 	{
diff --git a/Engine/Plugins/Experimental/PCG/Source/PCG/Private/PCGNode.cpp b/Engine/Plugins/Experimental/PCG/Source/PCG/Private/PCGNode.cpp
--- a/Engine/Plugins/Experimental/PCG/Source/PCG/Private/PCGNode.cpp
+++ b/Engine/Plugins/Experimental/PCG/Source/PCG/Private/PCGNode.cpp
@@ -236,27 +236,27 @@ bool UPCGNode::RemoveEdgeTo(FName FromPinLabel, UPCGNode* To, FName ToPinLabel)
 	}
 }
 
-FName UPCGNode::GetNodeTitle() const
+FText UPCGNode::GetNodeTitle() const
 {
 	if (NodeTitle != NAME_None)
 	{
-		return NodeTitle;
+		return FText::FromName(NodeTitle);
 	}
 	else if (UPCGSettings* Settings = GetSettings())
 	{
 		if (Settings->AdditionalTaskName() != NAME_None)
 		{
-			return Settings->AdditionalTaskName();
+			return FText::FromName(Settings->AdditionalTaskName());
 		}
 #if WITH_EDITOR
 		else
 		{
-			return Settings->GetDefaultNodeName();
+			return Settings->GetDefaultNodeTitle();
 		}
 #endif
 	}
 
-	return TEXT("Unnamed node");
+	return NSLOCTEXT("PCGNode", "NodeTitle", "Unnamed node");
 }
 
 #if WITH_EDITOR
diff --git a/Engine/Plugins/Experimental/PCG/Source/PCG/Private/Tests/Determinism/PCGDeterminismTestBlueprintBase.cpp b/Engine/Plugins/Experimental/PCG/Source/PCG/Private/Tests/Determinism/PCGDeterminismTestBlueprintBase.cpp
--- a/Engine/Plugins/Experimental/PCG/Source/PCG/Private/Tests/Determinism/PCGDeterminismTestBlueprintBase.cpp
+++ b/Engine/Plugins/Experimental/PCG/Source/PCG/Private/Tests/Determinism/PCGDeterminismTestBlueprintBase.cpp
@@ -9,7 +9,7 @@
 void UPCGDeterminismTestBlueprintBase::ExecuteTest_Implementation(const UPCGNode* InPCGNode, FDeterminismTestResult& InOutTestResult)
 {
 	// The default case if the user tries to run the blueprint test unimplemented
-	InOutTestResult.TestResultTitle = InPCGNode->GetNodeTitle();
+	InOutTestResult.TestResultTitle = FName(*InPCGNode->GetNodeTitle().ToString());
 	InOutTestResult.TestResultName = InPCGNode->GetName();
 	InOutTestResult.Seed = InPCGNode->GetSettings()->GetSeed();
 	PCGDeterminismTests::RunDeterminismTest(InPCGNode, InOutTestResult, PCGDeterminismTests::Defaults::DeterminismBasicTestInfo);
diff --git a/Engine/Plugins/Experimental/PCG/Source/PCG/Public/Elements/Metadata/PCGMetadataBitwiseOpElement.h b/Engine/Plugins/Experimental/PCG/Source/PCG/Public/Elements/Metadata/PCGMetadataBitwiseOpElement.h
--- a/Engine/Plugins/Experimental/PCG/Source/PCG/Public/Elements/Metadata/PCGMetadataBitwiseOpElement.h
+++ b/Engine/Plugins/Experimental/PCG/Source/PCG/Public/Elements/Metadata/PCGMetadataBitwiseOpElement.h
@@ -28,6 +28,7 @@ public:
 	//~Begin UPCGSettings interface
 #if WITH_EDITOR
 	virtual FName GetDefaultNodeName() const override;
+	virtual FText GetDefaultNodeTitle() const override;
 #endif
 	virtual FName AdditionalTaskName() const override;
 	//~End UPCGSettings interface
diff --git a/Engine/Plugins/Experimental/PCG/Source/PCG/Public/Elements/Metadata/PCGMetadataBooleanOpElement.h b/Engine/Plugins/Experimental/PCG/Source/PCG/Public/Elements/Metadata/PCGMetadataBooleanOpElement.h
--- a/Engine/Plugins/Experimental/PCG/Source/PCG/Public/Elements/Metadata/PCGMetadataBooleanOpElement.h
+++ b/Engine/Plugins/Experimental/PCG/Source/PCG/Public/Elements/Metadata/PCGMetadataBooleanOpElement.h
@@ -28,6 +28,7 @@ public:
 	//~Begin UPCGSettings interface
 #if WITH_EDITOR
 	virtual FName GetDefaultNodeName() const override;
+	virtual FText GetDefaultNodeTitle() const override;
 #endif
 	virtual FName AdditionalTaskName() const override;
 	//~End UPCGSettings interface
diff --git a/Engine/Plugins/Experimental/PCG/Source/PCG/Public/Elements/Metadata/PCGMetadataBreakTransform.h b/Engine/Plugins/Experimental/PCG/Source/PCG/Public/Elements/Metadata/PCGMetadataBreakTransform.h
--- a/Engine/Plugins/Experimental/PCG/Source/PCG/Public/Elements/Metadata/PCGMetadataBreakTransform.h
+++ b/Engine/Plugins/Experimental/PCG/Source/PCG/Public/Elements/Metadata/PCGMetadataBreakTransform.h
@@ -20,6 +20,7 @@ public:
 	//~Begin UPCGSettings interface
 #if WITH_EDITOR
 	virtual FName GetDefaultNodeName() const override;
+	virtual FText GetDefaultNodeTitle() const override;
 #endif
 
 	FPCGAttributePropertySelector GetInputSource(uint32 Index) const override;
diff --git a/Engine/Plugins/Experimental/PCG/Source/PCG/Public/Elements/Metadata/PCGMetadataBreakVector.h b/Engine/Plugins/Experimental/PCG/Source/PCG/Public/Elements/Metadata/PCGMetadataBreakVector.h
--- a/Engine/Plugins/Experimental/PCG/Source/PCG/Public/Elements/Metadata/PCGMetadataBreakVector.h
+++ b/Engine/Plugins/Experimental/PCG/Source/PCG/Public/Elements/Metadata/PCGMetadataBreakVector.h
@@ -28,6 +28,7 @@ public:
 	//~Begin UPCGSettings interface
 #if WITH_EDITOR
 	virtual FName GetDefaultNodeName() const override;
+	virtual FText GetDefaultNodeTitle() const override;
 #endif
 
 	FPCGAttributePropertySelector GetInputSource(uint32 Index) const override;
diff --git a/Engine/Plugins/Experimental/PCG/Source/PCG/Public/Elements/Metadata/PCGMetadataCompareOpElement.h b/Engine/Plugins/Experimental/PCG/Source/PCG/Public/Elements/Metadata/PCGMetadataCompareOpElement.h
--- a/Engine/Plugins/Experimental/PCG/Source/PCG/Public/Elements/Metadata/PCGMetadataCompareOpElement.h
+++ b/Engine/Plugins/Experimental/PCG/Source/PCG/Public/Elements/Metadata/PCGMetadataCompareOpElement.h
@@ -30,6 +30,7 @@ public:
 	//~Begin UPCGSettings interface
 #if WITH_EDITOR
 	virtual FName GetDefaultNodeName() const override;
+	virtual FText GetDefaultNodeTitle() const override;
 #endif
 	virtual FName AdditionalTaskName() const override;
 	//~End UPCGSettings interface
diff --git a/Engine/Plugins/Experimental/PCG/Source/PCG/Public/Elements/Metadata/PCGMetadataElement.h b/Engine/Plugins/Experimental/PCG/Source/PCG/Public/Elements/Metadata/PCGMetadataElement.h
--- a/Engine/Plugins/Experimental/PCG/Source/PCG/Public/Elements/Metadata/PCGMetadataElement.h
+++ b/Engine/Plugins/Experimental/PCG/Source/PCG/Public/Elements/Metadata/PCGMetadataElement.h
@@ -30,7 +30,8 @@ public:
 
 	//~Begin UPCGSettings interface
 #if WITH_EDITOR
-	virtual FName GetDefaultNodeName() const override { return FName(TEXT("Attribute Operation")); }
+	virtual FName GetDefaultNodeName() const override { return FName(TEXT("AttributeOperation")); }
+	virtual FText GetDefaultNodeTitle() const override { return NSLOCTEXT("PCGMetadataOperationSettings", "NodeTitle", "Attribute Operation"); }
 	virtual EPCGSettingsType GetType() const override { return EPCGSettingsType::Metadata; }
 #endif
 
diff --git a/Engine/Plugins/Experimental/PCG/Source/PCG/Public/Elements/Metadata/PCGMetadataMakeTransform.h b/Engine/Plugins/Experimental/PCG/Source/PCG/Public/Elements/Metadata/PCGMetadataMakeTransform.h
--- a/Engine/Plugins/Experimental/PCG/Source/PCG/Public/Elements/Metadata/PCGMetadataMakeTransform.h
+++ b/Engine/Plugins/Experimental/PCG/Source/PCG/Public/Elements/Metadata/PCGMetadataMakeTransform.h
@@ -27,6 +27,7 @@ public:
 #if WITH_EDITOR
 	//~Begin UPCGSettings interface
 	virtual FName GetDefaultNodeName() const override;
+	virtual FText GetDefaultNodeTitle() const override;
 	//~End UPCGSettings interface
 #endif
 
diff --git a/Engine/Plugins/Experimental/PCG/Source/PCG/Public/Elements/Metadata/PCGMetadataMakeVector.h b/Engine/Plugins/Experimental/PCG/Source/PCG/Public/Elements/Metadata/PCGMetadataMakeVector.h
--- a/Engine/Plugins/Experimental/PCG/Source/PCG/Public/Elements/Metadata/PCGMetadataMakeVector.h
+++ b/Engine/Plugins/Experimental/PCG/Source/PCG/Public/Elements/Metadata/PCGMetadataMakeVector.h
@@ -49,6 +49,7 @@ public:
 #if WITH_EDITOR
 	//~Begin UPCGSettings interface
 	virtual FName GetDefaultNodeName() const override;
+	virtual FText GetDefaultNodeTitle() const override;
 	//~End UPCGSettings interface
 #endif
 
diff --git a/Engine/Plugins/Experimental/PCG/Source/PCG/Public/Elements/Metadata/PCGMetadataMathsOpElement.h b/Engine/Plugins/Experimental/PCG/Source/PCG/Public/Elements/Metadata/PCGMetadataMathsOpElement.h
--- a/Engine/Plugins/Experimental/PCG/Source/PCG/Public/Elements/Metadata/PCGMetadataMathsOpElement.h
+++ b/Engine/Plugins/Experimental/PCG/Source/PCG/Public/Elements/Metadata/PCGMetadataMathsOpElement.h
@@ -53,6 +53,7 @@ public:
 	//~Begin UPCGSettings interface
 #if WITH_EDITOR
 	virtual FName GetDefaultNodeName() const override;
+	virtual FText GetDefaultNodeTitle() const override;
 #endif
 	virtual FName AdditionalTaskName() const override;
 	//~End UPCGSettings interface
diff --git a/Engine/Plugins/Experimental/PCG/Source/PCG/Public/Elements/Metadata/PCGMetadataPartition.h b/Engine/Plugins/Experimental/PCG/Source/PCG/Public/Elements/Metadata/PCGMetadataPartition.h
--- a/Engine/Plugins/Experimental/PCG/Source/PCG/Public/Elements/Metadata/PCGMetadataPartition.h
+++ b/Engine/Plugins/Experimental/PCG/Source/PCG/Public/Elements/Metadata/PCGMetadataPartition.h
@@ -15,7 +15,8 @@ class PCG_API UPCGMetadataPartitionSettings : public UPCGSettings
 public:
 	//~Begin UPCGSettings interface
 #if WITH_EDITOR
-	virtual FName GetDefaultNodeName() const override { return FName(TEXT("Attribute Partition")); }
+	virtual FName GetDefaultNodeName() const override { return FName(TEXT("AttributePartition")); }
+	virtual FText GetDefaultNodeTitle() const override { return NSLOCTEXT("PCGMetadataPartitionSettings", "NodeTitle", "Attribute Partition"); }
 	virtual EPCGSettingsType GetType() const override { return EPCGSettingsType::Metadata; }
 #endif
 
diff --git a/Engine/Plugins/Experimental/PCG/Source/PCG/Public/Elements/Metadata/PCGMetadataRenameElement.h b/Engine/Plugins/Experimental/PCG/Source/PCG/Public/Elements/Metadata/PCGMetadataRenameElement.h
--- a/Engine/Plugins/Experimental/PCG/Source/PCG/Public/Elements/Metadata/PCGMetadataRenameElement.h
+++ b/Engine/Plugins/Experimental/PCG/Source/PCG/Public/Elements/Metadata/PCGMetadataRenameElement.h
@@ -14,7 +14,8 @@ class PCG_API UPCGMetadataRenameSettings : public UPCGSettings
 public:
 	//~Begin UPCGSettings interface
 #if WITH_EDITOR
-	virtual FName GetDefaultNodeName() const override { return FName(TEXT("Attribute Rename")); }
+	virtual FName GetDefaultNodeName() const override { return FName(TEXT("AttributeRename")); }
+	virtual FText GetDefaultNodeTitle() const override { return NSLOCTEXT("PCGMetadataRenameSettings", "NodeTitle", "Attribute Rename"); }
 	virtual EPCGSettingsType GetType() const override { return EPCGSettingsType::Metadata; }
 #endif
 
diff --git a/Engine/Plugins/Experimental/PCG/Source/PCG/Public/Elements/Metadata/PCGMetadataRotatorOpElement.h b/Engine/Plugins/Experimental/PCG/Source/PCG/Public/Elements/Metadata/PCGMetadataRotatorOpElement.h
--- a/Engine/Plugins/Experimental/PCG/Source/PCG/Public/Elements/Metadata/PCGMetadataRotatorOpElement.h
+++ b/Engine/Plugins/Experimental/PCG/Source/PCG/Public/Elements/Metadata/PCGMetadataRotatorOpElement.h
@@ -38,6 +38,7 @@ public:
 	//~Begin UPCGSettings interface
 #if WITH_EDITOR
 	virtual FName GetDefaultNodeName() const override;
+	virtual FText GetDefaultNodeTitle() const override;
 #endif
 
 	virtual FName AdditionalTaskName() const override;
diff --git a/Engine/Plugins/Experimental/PCG/Source/PCG/Public/Elements/Metadata/PCGMetadataTransformOpElement.h b/Engine/Plugins/Experimental/PCG/Source/PCG/Public/Elements/Metadata/PCGMetadataTransformOpElement.h
--- a/Engine/Plugins/Experimental/PCG/Source/PCG/Public/Elements/Metadata/PCGMetadataTransformOpElement.h
+++ b/Engine/Plugins/Experimental/PCG/Source/PCG/Public/Elements/Metadata/PCGMetadataTransformOpElement.h
@@ -41,6 +41,7 @@ public:
 	//~Begin UPCGSettings interface
 #if WITH_EDITOR
 	virtual FName GetDefaultNodeName() const override;
+	virtual FText GetDefaultNodeTitle() const override;
 #endif
 	virtual FName AdditionalTaskName() const override;
 	//~End UPCGSettings interface
diff --git a/Engine/Plugins/Experimental/PCG/Source/PCG/Public/Elements/Metadata/PCGMetadataTrigOpElement.h b/Engine/Plugins/Experimental/PCG/Source/PCG/Public/Elements/Metadata/PCGMetadataTrigOpElement.h
--- a/Engine/Plugins/Experimental/PCG/Source/PCG/Public/Elements/Metadata/PCGMetadataTrigOpElement.h
+++ b/Engine/Plugins/Experimental/PCG/Source/PCG/Public/Elements/Metadata/PCGMetadataTrigOpElement.h
@@ -33,6 +33,7 @@ public:
 	//~Begin UPCGSettings interface
 #if WITH_EDITOR
 	virtual FName GetDefaultNodeName() const override;
+	virtual FText GetDefaultNodeTitle() const override;
 #endif
 	virtual FName AdditionalTaskName() const override;
 	//~End UPCGSettings interface
diff --git a/Engine/Plugins/Experimental/PCG/Source/PCG/Public/Elements/Metadata/PCGMetadataVectorOpElement.h b/Engine/Plugins/Experimental/PCG/Source/PCG/Public/Elements/Metadata/PCGMetadataVectorOpElement.h
--- a/Engine/Plugins/Experimental/PCG/Source/PCG/Public/Elements/Metadata/PCGMetadataVectorOpElement.h
+++ b/Engine/Plugins/Experimental/PCG/Source/PCG/Public/Elements/Metadata/PCGMetadataVectorOpElement.h
@@ -37,6 +37,7 @@ public:
 	//~Begin UPCGSettings interface
 #if WITH_EDITOR
 	virtual FName GetDefaultNodeName() const override;
+	virtual FText GetDefaultNodeTitle() const override;
 #endif
 	virtual FName AdditionalTaskName() const override;
 	//~End UPCGSettings interface
diff --git a/Engine/Plugins/Experimental/PCG/Source/PCG/Public/Elements/PCGAttributeFilterElement.h b/Engine/Plugins/Experimental/PCG/Source/PCG/Public/Elements/PCGAttributeFilterElement.h
--- a/Engine/Plugins/Experimental/PCG/Source/PCG/Public/Elements/PCGAttributeFilterElement.h
+++ b/Engine/Plugins/Experimental/PCG/Source/PCG/Public/Elements/PCGAttributeFilterElement.h
@@ -33,6 +33,7 @@ public:
 	//~Begin UPCGSettings interface
 #if WITH_EDITOR
 	virtual FName GetDefaultNodeName() const override;
+	virtual FText GetDefaultNodeTitle() const override;
 	virtual EPCGSettingsType GetType() const override { return EPCGSettingsType::Metadata; }
 #endif
 
diff --git a/Engine/Plugins/Experimental/PCG/Source/PCG/Public/Elements/PCGAttributeGetFromPointIndexElement.h b/Engine/Plugins/Experimental/PCG/Source/PCG/Public/Elements/PCGAttributeGetFromPointIndexElement.h
--- a/Engine/Plugins/Experimental/PCG/Source/PCG/Public/Elements/PCGAttributeGetFromPointIndexElement.h
+++ b/Engine/Plugins/Experimental/PCG/Source/PCG/Public/Elements/PCGAttributeGetFromPointIndexElement.h
@@ -33,6 +33,7 @@ public:
 	//~Begin UPCGSettings interface
 #if WITH_EDITOR
 	virtual FName GetDefaultNodeName() const override;
+	virtual FText GetDefaultNodeTitle() const override;
 	virtual EPCGSettingsType GetType() const override { return EPCGSettingsType::Metadata; }
 #endif
 
diff --git a/Engine/Plugins/Experimental/PCG/Source/PCG/Public/Elements/PCGAttributeReduceElement.h b/Engine/Plugins/Experimental/PCG/Source/PCG/Public/Elements/PCGAttributeReduceElement.h
--- a/Engine/Plugins/Experimental/PCG/Source/PCG/Public/Elements/PCGAttributeReduceElement.h
+++ b/Engine/Plugins/Experimental/PCG/Source/PCG/Public/Elements/PCGAttributeReduceElement.h
@@ -34,6 +34,7 @@ public:
 	//~Begin UPCGSettings interface
 #if WITH_EDITOR
 	virtual FName GetDefaultNodeName() const override;
+	virtual FText GetDefaultNodeTitle() const override;
 	virtual EPCGSettingsType GetType() const override { return EPCGSettingsType::Metadata; }
 #endif
 
diff --git a/Engine/Plugins/Experimental/PCG/Source/PCG/Public/Elements/PCGAttributeSelectElement.h b/Engine/Plugins/Experimental/PCG/Source/PCG/Public/Elements/PCGAttributeSelectElement.h
--- a/Engine/Plugins/Experimental/PCG/Source/PCG/Public/Elements/PCGAttributeSelectElement.h
+++ b/Engine/Plugins/Experimental/PCG/Source/PCG/Public/Elements/PCGAttributeSelectElement.h
@@ -57,6 +57,7 @@ public:
 	//~Begin UPCGSettings interface
 #if WITH_EDITOR
 	virtual FName GetDefaultNodeName() const override;
+	virtual FText GetDefaultNodeTitle() const override;
 	virtual EPCGSettingsType GetType() const override { return EPCGSettingsType::Metadata; }
 #endif
 
diff --git a/Engine/Plugins/Experimental/PCG/Source/PCG/Public/Elements/PCGAttributeTransferElement.h b/Engine/Plugins/Experimental/PCG/Source/PCG/Public/Elements/PCGAttributeTransferElement.h
--- a/Engine/Plugins/Experimental/PCG/Source/PCG/Public/Elements/PCGAttributeTransferElement.h
+++ b/Engine/Plugins/Experimental/PCG/Source/PCG/Public/Elements/PCGAttributeTransferElement.h
@@ -12,6 +12,7 @@ class UPCGSpatialData;
 namespace PCGAttributeTransferConstants
 {
 	const FName NodeName = TEXT("TransferAttribute");
+	const FText NodeTitle = NSLOCTEXT("PCGAttributeTransferElement", "NodeTitle", "Transfer Attribute");
 	const FName SourceLabel = TEXT("Source");
 	const FName TargetLabel = TEXT("Target");
 }
@@ -36,6 +37,7 @@ public:
 	//~Begin UPCGSettings interface
 #if WITH_EDITOR
 	virtual FName GetDefaultNodeName() const override;
+	virtual FText GetDefaultNodeTitle() const override;
 	virtual EPCGSettingsType GetType() const override { return EPCGSettingsType::Metadata; }
 #endif
 	virtual FName AdditionalTaskName() const override;
diff --git a/Engine/Plugins/Experimental/PCG/Source/PCG/Public/Elements/PCGBoundsModifier.h b/Engine/Plugins/Experimental/PCG/Source/PCG/Public/Elements/PCGBoundsModifier.h
--- a/Engine/Plugins/Experimental/PCG/Source/PCG/Public/Elements/PCGBoundsModifier.h
+++ b/Engine/Plugins/Experimental/PCG/Source/PCG/Public/Elements/PCGBoundsModifier.h
@@ -31,7 +31,8 @@ public:
 
 	// ~Begin UPCGSettings interface
 #if WITH_EDITOR
-	virtual FName GetDefaultNodeName() const override { return FName(TEXT("Bounds Modifier")); }
+	virtual FName GetDefaultNodeName() const override { return FName(TEXT("BoundsModifier")); }
+	virtual FText GetDefaultNodeTitle() const override { return NSLOCTEXT("PCGBoundsModifierSettings", "NodeTitle", "Bounds Modifier"); }
 	virtual FText GetNodeTooltipText() const override;
 	virtual EPCGSettingsType GetType() const override { return EPCGSettingsType::Spatial; }
 #endif
diff --git a/Engine/Plugins/Experimental/PCG/Source/PCG/Public/Elements/PCGCollapseElement.h b/Engine/Plugins/Experimental/PCG/Source/PCG/Public/Elements/PCGCollapseElement.h
--- a/Engine/Plugins/Experimental/PCG/Source/PCG/Public/Elements/PCGCollapseElement.h
+++ b/Engine/Plugins/Experimental/PCG/Source/PCG/Public/Elements/PCGCollapseElement.h
@@ -15,7 +15,8 @@ class PCG_API UPCGCollapseSettings : public UPCGSettings
 public:
 	//~Begin UPCGSettings interface
 #if WITH_EDITOR
-	virtual FName GetDefaultNodeName() const override { return FName(TEXT("To Point")); }
+	virtual FName GetDefaultNodeName() const override { return FName(TEXT("ToPoint")); }
+	virtual FText GetDefaultNodeTitle() const override { return NSLOCTEXT("PCGCollapseSettings", "NodeTitle", "To Point"); }
 	virtual EPCGSettingsType GetType() const override { return EPCGSettingsType::Spatial; }
 #endif
 
diff --git a/Engine/Plugins/Experimental/PCG/Source/PCG/Public/Elements/PCGCopyPoints.h b/Engine/Plugins/Experimental/PCG/Source/PCG/Public/Elements/PCGCopyPoints.h
--- a/Engine/Plugins/Experimental/PCG/Source/PCG/Public/Elements/PCGCopyPoints.h
+++ b/Engine/Plugins/Experimental/PCG/Source/PCG/Public/Elements/PCGCopyPoints.h
@@ -39,7 +39,8 @@ class PCG_API UPCGCopyPointsSettings : public UPCGSettings
 public:
 	//~Begin UPCGSettings interface
 #if WITH_EDITOR
-	virtual FName GetDefaultNodeName() const override { return FName(TEXT("Copy Points")); }
+	virtual FName GetDefaultNodeName() const override { return FName(TEXT("CopyPoints")); }
+	virtual FText GetDefaultNodeTitle() const override { return NSLOCTEXT("PCGCopyPointSettings", "NodeTitle", "Copy Points"); }
 	virtual FText GetNodeTooltipText() const override;
 	virtual EPCGSettingsType GetType() const override { return EPCGSettingsType::Sampler; }
 #endif
diff --git a/Engine/Plugins/Experimental/PCG/Source/PCG/Public/Elements/PCGCreateAttribute.h b/Engine/Plugins/Experimental/PCG/Source/PCG/Public/Elements/PCGCreateAttribute.h
--- a/Engine/Plugins/Experimental/PCG/Source/PCG/Public/Elements/PCGCreateAttribute.h
+++ b/Engine/Plugins/Experimental/PCG/Source/PCG/Public/Elements/PCGCreateAttribute.h
@@ -29,6 +29,7 @@ public:
 	//~Begin UPCGSettings interface
 #if WITH_EDITOR
 	virtual FName GetDefaultNodeName() const override;
+	virtual FText GetDefaultNodeTitle() const override;
 	virtual EPCGSettingsType GetType() const override { return EPCGSettingsType::Metadata; }
 	virtual bool HasDynamicPins() const override { return true; }
 #endif // WITH_EDITOR
@@ -120,6 +121,7 @@ public:
 	//~Begin UPCGSettings interface
 #if WITH_EDITOR
 	virtual FName GetDefaultNodeName() const override;
+	virtual FText GetDefaultNodeTitle() const override;
 #endif
 	virtual FName AdditionalTaskName() const override;
 
diff --git a/Engine/Plugins/Experimental/PCG/Source/PCG/Public/Elements/PCGCreateSpline.h b/Engine/Plugins/Experimental/PCG/Source/PCG/Public/Elements/PCGCreateSpline.h
--- a/Engine/Plugins/Experimental/PCG/Source/PCG/Public/Elements/PCGCreateSpline.h
+++ b/Engine/Plugins/Experimental/PCG/Source/PCG/Public/Elements/PCGCreateSpline.h
@@ -27,7 +27,8 @@ class PCG_API UPCGCreateSplineSettings : public UPCGSettings
 public:
 	//~Begin UPCGSettings interface
 #if WITH_EDITOR
-	virtual FName GetDefaultNodeName() const override { return FName(TEXT("Create Spline")); }
+	virtual FName GetDefaultNodeName() const override { return FName(TEXT("CreateSpline")); }
+	virtual FText GetDefaultNodeTitle() const override { return NSLOCTEXT("PCGCreateSplineSettings", "NodeTitle", "Create Spline"); }
 	virtual FText GetNodeTooltipText() const override;
 	virtual EPCGSettingsType GetType() const override { return EPCGSettingsType::Spatial; }
 #endif
diff --git a/Engine/Plugins/Experimental/PCG/Source/PCG/Public/Elements/PCGDataFromActor.h b/Engine/Plugins/Experimental/PCG/Source/PCG/Public/Elements/PCGDataFromActor.h
--- a/Engine/Plugins/Experimental/PCG/Source/PCG/Public/Elements/PCGDataFromActor.h
+++ b/Engine/Plugins/Experimental/PCG/Source/PCG/Public/Elements/PCGDataFromActor.h
@@ -28,7 +28,8 @@ class PCG_API UPCGDataFromActorSettings : public UPCGSettings
 public:
 	//~Begin UPCGSettings interface
 #if WITH_EDITOR
-	virtual FName GetDefaultNodeName() const override { return FName(TEXT("Get Actor Data")); }
+	virtual FName GetDefaultNodeName() const override { return FName(TEXT("GetActorData")); }
+	virtual FText GetDefaultNodeTitle() const override { return NSLOCTEXT("PCGDataFromActorSettings", "NodeTitle", "Get Actor Data"); }
 	virtual FText GetNodeTooltipText() const override;
 	virtual EPCGSettingsType GetType() const override { return EPCGSettingsType::Spatial; }
 	virtual void GetTrackedActorTags(FPCGTagToSettingsMap& OutTagToSettings, TArray<TObjectPtr<const UPCGGraph>>& OutVisitedGraphs) const override;
diff --git a/Engine/Plugins/Experimental/PCG/Source/PCG/Public/Elements/PCGDataTableRowToParamData.h b/Engine/Plugins/Experimental/PCG/Source/PCG/Public/Elements/PCGDataTableRowToParamData.h
--- a/Engine/Plugins/Experimental/PCG/Source/PCG/Public/Elements/PCGDataTableRowToParamData.h
+++ b/Engine/Plugins/Experimental/PCG/Source/PCG/Public/Elements/PCGDataTableRowToParamData.h
@@ -16,7 +16,8 @@ class UPCGDataTableRowToParamDataSettings : public UPCGSettings
 public:
 #if WITH_EDITOR
 	//~Begin UPCGSettings interface
-	virtual FName GetDefaultNodeName() const override { return FName(TEXT("DataTableRowToParamData")); }
+	virtual FName GetDefaultNodeName() const override { return FName(TEXT("DataTableRowToAttributeSet")); }
+	virtual FText GetDefaultNodeTitle() const override { return NSLOCTEXT("PCGDataTableRowToParamDataSettings", "NodeTitle", "Data Table Row To Attribute Set"); }
 	virtual EPCGSettingsType GetType() const override { return EPCGSettingsType::Metadata; }
 #endif
 
diff --git a/Engine/Plugins/Experimental/PCG/Source/PCG/Public/Elements/PCGDebugElement.h b/Engine/Plugins/Experimental/PCG/Source/PCG/Public/Elements/PCGDebugElement.h
--- a/Engine/Plugins/Experimental/PCG/Source/PCG/Public/Elements/PCGDebugElement.h
+++ b/Engine/Plugins/Experimental/PCG/Source/PCG/Public/Elements/PCGDebugElement.h
@@ -20,6 +20,7 @@ public:
 #if WITH_EDITOR
 	// ~Begin UPCGSettings interface
 	virtual FName GetDefaultNodeName() const override { return FName(TEXT("Debug")); }
+	virtual FText GetDefaultNodeTitle() const override { return NSLOCTEXT("PCGDebugSettings", "NodeTitle", "Debug"); }
 	virtual EPCGSettingsType GetType() const override { return EPCGSettingsType::Debug; }
 #endif
 
diff --git a/Engine/Plugins/Experimental/PCG/Source/PCG/Public/Elements/PCGDensityFilter.h b/Engine/Plugins/Experimental/PCG/Source/PCG/Public/Elements/PCGDensityFilter.h
--- a/Engine/Plugins/Experimental/PCG/Source/PCG/Public/Elements/PCGDensityFilter.h
+++ b/Engine/Plugins/Experimental/PCG/Source/PCG/Public/Elements/PCGDensityFilter.h
@@ -16,7 +16,8 @@ class PCG_API UPCGDensityFilterSettings : public UPCGSettings
 public:
 	//~Begin UPCGSettings interface
 #if WITH_EDITOR
-	virtual FName GetDefaultNodeName() const override { return FName(TEXT("Density Filter")); }
+	virtual FName GetDefaultNodeName() const override { return FName(TEXT("DensityFilter")); }
+	virtual FText GetDefaultNodeTitle() const override { return NSLOCTEXT("PCGDensityFilterSettings", "NodeTitle", "Density Filter"); }
 	virtual EPCGSettingsType GetType() const override { return EPCGSettingsType::Filter; }
 #endif
 
diff --git a/Engine/Plugins/Experimental/PCG/Source/PCG/Public/Elements/PCGDensityNoise.h b/Engine/Plugins/Experimental/PCG/Source/PCG/Public/Elements/PCGDensityNoise.h
--- a/Engine/Plugins/Experimental/PCG/Source/PCG/Public/Elements/PCGDensityNoise.h
+++ b/Engine/Plugins/Experimental/PCG/Source/PCG/Public/Elements/PCGDensityNoise.h
@@ -29,7 +29,8 @@ public:
 
 	// ~Begin UPCGSettings interface
 #if WITH_EDITOR
-	virtual FName GetDefaultNodeName() const override { return FName(TEXT("Density Noise")); }
+	virtual FName GetDefaultNodeName() const override { return FName(TEXT("DensityNoise")); }
+	virtual FText GetDefaultNodeTitle() const override { return NSLOCTEXT("PCGDensityNodeSettings", "NodeTitle", "Density Noise"); }
 	virtual EPCGSettingsType GetType() const override { return EPCGSettingsType::Density; }
 #endif
 
diff --git a/Engine/Plugins/Experimental/PCG/Source/PCG/Public/Elements/PCGDensityRemap.h b/Engine/Plugins/Experimental/PCG/Source/PCG/Public/Elements/PCGDensityRemap.h
--- a/Engine/Plugins/Experimental/PCG/Source/PCG/Public/Elements/PCGDensityRemap.h
+++ b/Engine/Plugins/Experimental/PCG/Source/PCG/Public/Elements/PCGDensityRemap.h
@@ -17,7 +17,8 @@ public:
 
 	// ~Begin UPCGSettings interface
 #if WITH_EDITOR
-	virtual FName GetDefaultNodeName() const override { return FName(TEXT("Linear Density Remap")); }
+	virtual FName GetDefaultNodeName() const override { return FName(TEXT("LinearDensityRemap")); }
+	virtual FText GetDefaultNodeTitle() const override { return NSLOCTEXT("PCGLinearDensityRemapSettings", "NodeTitle", "Linear Density Remap"); }
 	virtual EPCGSettingsType GetType() const override { return EPCGSettingsType::Density; }
 #endif
 
diff --git a/Engine/Plugins/Experimental/PCG/Source/PCG/Public/Elements/PCGDensityRemapElement.h b/Engine/Plugins/Experimental/PCG/Source/PCG/Public/Elements/PCGDensityRemapElement.h
--- a/Engine/Plugins/Experimental/PCG/Source/PCG/Public/Elements/PCGDensityRemapElement.h
+++ b/Engine/Plugins/Experimental/PCG/Source/PCG/Public/Elements/PCGDensityRemapElement.h
@@ -17,7 +17,8 @@ class PCG_API UPCGDensityRemapSettings : public UPCGSettings
 public:
 	// ~Begin UPCGSettings interface
 #if WITH_EDITOR
-	virtual FName GetDefaultNodeName() const override { return FName(TEXT("Density Remap")); }
+	virtual FName GetDefaultNodeName() const override { return FName(TEXT("DensityRemap")); }
+	virtual FText GetDefaultNodeTitle() const override { return NSLOCTEXT("PCGDensityRemapSettings", "NodeTitle", "Density Remap"); }
 	virtual EPCGSettingsType GetType() const override { return EPCGSettingsType::Density; }
 #endif
 
diff --git a/Engine/Plugins/Experimental/PCG/Source/PCG/Public/Elements/PCGDifferenceElement.h b/Engine/Plugins/Experimental/PCG/Source/PCG/Public/Elements/PCGDifferenceElement.h
--- a/Engine/Plugins/Experimental/PCG/Source/PCG/Public/Elements/PCGDifferenceElement.h
+++ b/Engine/Plugins/Experimental/PCG/Source/PCG/Public/Elements/PCGDifferenceElement.h
@@ -22,6 +22,7 @@ public:
 	//~Begin UPCGSettings interface
 #if WITH_EDITOR
 	virtual FName GetDefaultNodeName() const override { return FName(TEXT("Difference")); }
+	virtual FText GetDefaultNodeTitle() const override { return NSLOCTEXT("PCGDifferenceSettings", "NodeTitle", "Difference"); }
 	virtual EPCGSettingsType GetType() const override { return EPCGSettingsType::Spatial; }
 #endif
 
diff --git a/Engine/Plugins/Experimental/PCG/Source/PCG/Public/Elements/PCGDistance.h b/Engine/Plugins/Experimental/PCG/Source/PCG/Public/Elements/PCGDistance.h
--- a/Engine/Plugins/Experimental/PCG/Source/PCG/Public/Elements/PCGDistance.h
+++ b/Engine/Plugins/Experimental/PCG/Source/PCG/Public/Elements/PCGDistance.h
@@ -32,6 +32,7 @@ public:
 	//~Begin UPCGSettings interface
 #if WITH_EDITOR
 	virtual FName GetDefaultNodeName() const override { return FName(TEXT("Distance")); }
+	virtual FText GetDefaultNodeTitle() const override { return NSLOCTEXT("PCGDistanceSettings", "NodeTitle", "Distance"); }
 	virtual EPCGSettingsType GetType() const override { return EPCGSettingsType::Spatial; }
 #endif
 
diff --git a/Engine/Plugins/Experimental/PCG/Source/PCG/Public/Elements/PCGExecuteBlueprint.h b/Engine/Plugins/Experimental/PCG/Source/PCG/Public/Elements/PCGExecuteBlueprint.h
--- a/Engine/Plugins/Experimental/PCG/Source/PCG/Public/Elements/PCGExecuteBlueprint.h
+++ b/Engine/Plugins/Experimental/PCG/Source/PCG/Public/Elements/PCGExecuteBlueprint.h
@@ -183,6 +183,7 @@ public:
 	// ~Begin UPCGSettings interface
 #if WITH_EDITOR
 	virtual FName GetDefaultNodeName() const override { return FName(TEXT("ExecuteBlueprint")); }
+	virtual FText GetDefaultNodeTitle() const override { return NSLOCTEXT("PCGBlueprintSettings", "NodeTitle", "Execute Blueprint"); }
 	virtual FLinearColor GetNodeTitleColor() const override;
 	virtual EPCGSettingsType GetType() const override;
 	virtual void GetTrackedActorTags(FPCGTagToSettingsMap& OutTagToSettings, TArray<TObjectPtr<const UPCGGraph>>& OutVisitedGraphs) const override;
diff --git a/Engine/Plugins/Experimental/PCG/Source/PCG/Public/Elements/PCGFilterByTag.h b/Engine/Plugins/Experimental/PCG/Source/PCG/Public/Elements/PCGFilterByTag.h
--- a/Engine/Plugins/Experimental/PCG/Source/PCG/Public/Elements/PCGFilterByTag.h
+++ b/Engine/Plugins/Experimental/PCG/Source/PCG/Public/Elements/PCGFilterByTag.h
@@ -23,6 +23,7 @@ public:
 	//~Begin UPCGSettings interface
 #if WITH_EDITOR
 	virtual FName GetDefaultNodeName() const override;
+	virtual FText GetDefaultNodeTitle() const override;
 	virtual FText GetNodeTooltipText() const override;
 	virtual EPCGSettingsType GetType() const override { return EPCGSettingsType::Filter; }
 #endif
diff --git a/Engine/Plugins/Experimental/PCG/Source/PCG/Public/Elements/PCGGather.h b/Engine/Plugins/Experimental/PCG/Source/PCG/Public/Elements/PCGGather.h
--- a/Engine/Plugins/Experimental/PCG/Source/PCG/Public/Elements/PCGGather.h
+++ b/Engine/Plugins/Experimental/PCG/Source/PCG/Public/Elements/PCGGather.h
@@ -17,6 +17,7 @@ public:
 	//~Begin UPCGSettings interface
 #if WITH_EDITOR
 	virtual FName GetDefaultNodeName() const override { return FName(TEXT("Gather")); }
+	virtual FText GetDefaultNodeTitle() const override { return NSLOCTEXT("PCGGatherSettings", "NodeTitle", "Gather"); }
 	virtual EPCGSettingsType GetType() const override { return EPCGSettingsType::Spatial; }
 	virtual bool HasDynamicPins() const override { return true; }
 #endif
diff --git a/Engine/Plugins/Experimental/PCG/Source/PCG/Public/Elements/PCGIntersectionElement.h b/Engine/Plugins/Experimental/PCG/Source/PCG/Public/Elements/PCGIntersectionElement.h
--- a/Engine/Plugins/Experimental/PCG/Source/PCG/Public/Elements/PCGIntersectionElement.h
+++ b/Engine/Plugins/Experimental/PCG/Source/PCG/Public/Elements/PCGIntersectionElement.h
@@ -16,6 +16,7 @@ public:
 	//~Begin UPCGSettings interface
 #if WITH_EDITOR
 	virtual FName GetDefaultNodeName() const override { return FName(TEXT("Intersection")); }
+	virtual FText GetDefaultNodeTitle() const override { return NSLOCTEXT("PCGIntersectionSettings", "NodeTitle", "Intersection"); }
 	virtual EPCGSettingsType GetType() const override { return EPCGSettingsType::Spatial; }
 #endif
 
diff --git a/Engine/Plugins/Experimental/PCG/Source/PCG/Public/Elements/PCGMergeElement.h b/Engine/Plugins/Experimental/PCG/Source/PCG/Public/Elements/PCGMergeElement.h
--- a/Engine/Plugins/Experimental/PCG/Source/PCG/Public/Elements/PCGMergeElement.h
+++ b/Engine/Plugins/Experimental/PCG/Source/PCG/Public/Elements/PCGMergeElement.h
@@ -16,6 +16,7 @@ public:
 	//~Begin UPCGSettings interface
 #if WITH_EDITOR
 	virtual FName GetDefaultNodeName() const override { return FName(TEXT("Merge")); }
+	virtual FText GetDefaultNodeTitle() const override { return NSLOCTEXT("PCGMergeSettings", "NodeTitle", "Merge"); }
 	virtual FText GetNodeTooltipText() const override;
 	virtual EPCGSettingsType GetType() const override { return EPCGSettingsType::Spatial; }
 #endif
diff --git a/Engine/Plugins/Experimental/PCG/Source/PCG/Public/Elements/PCGNormalToDensity.h b/Engine/Plugins/Experimental/PCG/Source/PCG/Public/Elements/PCGNormalToDensity.h
--- a/Engine/Plugins/Experimental/PCG/Source/PCG/Public/Elements/PCGNormalToDensity.h
+++ b/Engine/Plugins/Experimental/PCG/Source/PCG/Public/Elements/PCGNormalToDensity.h
@@ -30,7 +30,8 @@ class UPCGNormalToDensitySettings : public UPCGSettings
 public:
 	//~Begin UPCGSettings interface
 #if WITH_EDITOR
-	virtual FName GetDefaultNodeName() const override { return FName(TEXT("Normal To Density")); }
+	virtual FName GetDefaultNodeName() const override { return FName(TEXT("NormalToDensity")); }
+	virtual FText GetDefaultNodeTitle() const override { return NSLOCTEXT("PCGNormalToDensitySettings", "NodeTitle", "Normal To Density"); }
 	virtual EPCGSettingsType GetType() const override { return EPCGSettingsType::Spatial; }
 #endif
 
diff --git a/Engine/Plugins/Experimental/PCG/Source/PCG/Public/Elements/PCGPointExtentsModifier.h b/Engine/Plugins/Experimental/PCG/Source/PCG/Public/Elements/PCGPointExtentsModifier.h
--- a/Engine/Plugins/Experimental/PCG/Source/PCG/Public/Elements/PCGPointExtentsModifier.h
+++ b/Engine/Plugins/Experimental/PCG/Source/PCG/Public/Elements/PCGPointExtentsModifier.h
@@ -28,7 +28,8 @@ public:
 
 	// ~Begin UPCGSettings interface
 #if WITH_EDITOR
-	virtual FName GetDefaultNodeName() const override { return FName(TEXT("Extents Modifier")); }
+	virtual FName GetDefaultNodeName() const override { return FName(TEXT("ExtentsModifier")); }
+	virtual FText GetDefaultNodeTitle() const override { return NSLOCTEXT("PCGPointExtentsModifierSettings", "NodeTitle", "Extents Modifier"); }
 	virtual EPCGSettingsType GetType() const override { return EPCGSettingsType::Spatial; }
 #endif
 
diff --git a/Engine/Plugins/Experimental/PCG/Source/PCG/Public/Elements/PCGPointFilter.h b/Engine/Plugins/Experimental/PCG/Source/PCG/Public/Elements/PCGPointFilter.h
--- a/Engine/Plugins/Experimental/PCG/Source/PCG/Public/Elements/PCGPointFilter.h
+++ b/Engine/Plugins/Experimental/PCG/Source/PCG/Public/Elements/PCGPointFilter.h
@@ -73,6 +73,7 @@ public:
 	//~Begin UPCGSettings interface
 #if WITH_EDITOR
 	virtual FName GetDefaultNodeName() const override { return FName(TEXT("PointFilter")); }
+	virtual FText GetDefaultNodeTitle() const override { return NSLOCTEXT("PCGPointFilterSettings", "NodeTitle", "Point Filter"); }
 	virtual EPCGSettingsType GetType() const override { return EPCGSettingsType::Filter; }
 #endif
 
diff --git a/Engine/Plugins/Experimental/PCG/Source/PCG/Public/Elements/PCGPointMatchAndSet.h b/Engine/Plugins/Experimental/PCG/Source/PCG/Public/Elements/PCGPointMatchAndSet.h
--- a/Engine/Plugins/Experimental/PCG/Source/PCG/Public/Elements/PCGPointMatchAndSet.h
+++ b/Engine/Plugins/Experimental/PCG/Source/PCG/Public/Elements/PCGPointMatchAndSet.h
@@ -25,6 +25,7 @@ public:
 	// ~Begin UPCGSettings interface
 #if WITH_EDITOR
 	virtual FName GetDefaultNodeName() const override { return FName(TEXT("PointMatchAndSet")); }
+	virtual FText GetDefaultNodeTitle() const override { return NSLOCTEXT("PCGPointMatchAndSetSettings", "NodeTitle", "Point Match And Set"); }
 	virtual FText GetNodeTooltipText() const override;
 	virtual EPCGSettingsType GetType() const override { return EPCGSettingsType::Metadata; }
 #endif // WITH_EDITOR
diff --git a/Engine/Plugins/Experimental/PCG/Source/PCG/Public/Elements/PCGPointSampler.h b/Engine/Plugins/Experimental/PCG/Source/PCG/Public/Elements/PCGPointSampler.h
--- a/Engine/Plugins/Experimental/PCG/Source/PCG/Public/Elements/PCGPointSampler.h
+++ b/Engine/Plugins/Experimental/PCG/Source/PCG/Public/Elements/PCGPointSampler.h
@@ -17,7 +17,8 @@ public:
 
 	//~Begin UPCGSettings interface
 #if WITH_EDITOR
-	virtual FName GetDefaultNodeName() const override { return FName(TEXT("Point Sampler")); }
+	virtual FName GetDefaultNodeName() const override { return FName(TEXT("PointSampler")); }
+	virtual FText GetDefaultNodeTitle() const override { return NSLOCTEXT("PCGPointSamplerSettings", "NodeTitle", "Point Sampler"); }
 	virtual EPCGSettingsType GetType() const override { return EPCGSettingsType::Sampler; }
 #endif
 
diff --git a/Engine/Plugins/Experimental/PCG/Source/PCG/Public/Elements/PCGProjectionElement.h b/Engine/Plugins/Experimental/PCG/Source/PCG/Public/Elements/PCGProjectionElement.h
--- a/Engine/Plugins/Experimental/PCG/Source/PCG/Public/Elements/PCGProjectionElement.h
+++ b/Engine/Plugins/Experimental/PCG/Source/PCG/Public/Elements/PCGProjectionElement.h
@@ -22,6 +22,7 @@ public:
 	//~Begin UPCGSettings interface
 #if WITH_EDITOR
 	virtual FName GetDefaultNodeName() const override { return FName(TEXT("Projection")); }
+	virtual FText GetDefaultNodeTitle() const override { return NSLOCTEXT("PCGProjectionSettings", "NodeTitle", "Projection"); }
 	virtual FText GetNodeTooltipText() const override;
 	virtual EPCGSettingsType GetType() const override { return EPCGSettingsType::Spatial; }
 #endif
diff --git a/Engine/Plugins/Experimental/PCG/Source/PCG/Public/Elements/PCGPropertyToParamData.h b/Engine/Plugins/Experimental/PCG/Source/PCG/Public/Elements/PCGPropertyToParamData.h
--- a/Engine/Plugins/Experimental/PCG/Source/PCG/Public/Elements/PCGPropertyToParamData.h
+++ b/Engine/Plugins/Experimental/PCG/Source/PCG/Public/Elements/PCGPropertyToParamData.h
@@ -25,7 +25,8 @@ public:
 
 	//~Begin UPCGSettings interface
 #if WITH_EDITOR
-	virtual FName GetDefaultNodeName() const override { return FName(TEXT("PropertyToParamData")); }
+	virtual FName GetDefaultNodeName() const override { return FName(TEXT("GetActorProperty")); }
+	virtual FText GetDefaultNodeTitle() const override { return NSLOCTEXT("PCGPropertyToParamDataSettings", "NodeTitle", "Get Actor Property"); }
 	virtual EPCGSettingsType GetType() const override { return EPCGSettingsType::Param; }
 	virtual void GetTrackedActorTags(FPCGTagToSettingsMap& OutTagToSettings, TArray<TObjectPtr<const UPCGGraph>>& OutVisitedGraphs) const override;
 #endif
diff --git a/Engine/Plugins/Experimental/PCG/Source/PCG/Public/Elements/PCGReroute.h b/Engine/Plugins/Experimental/PCG/Source/PCG/Public/Elements/PCGReroute.h
--- a/Engine/Plugins/Experimental/PCG/Source/PCG/Public/Elements/PCGReroute.h
+++ b/Engine/Plugins/Experimental/PCG/Source/PCG/Public/Elements/PCGReroute.h
@@ -26,6 +26,7 @@ protected:
 	
 #if WITH_EDITOR
 	virtual FName GetDefaultNodeName() const override { return FName("Reroute"); }
+	virtual FText GetDefaultNodeTitle() const override { return NSLOCTEXT("PCGRerouteSettings", "NodeTitle", "Reroute"); }
 #endif
 };
 
diff --git a/Engine/Plugins/Experimental/PCG/Source/PCG/Public/Elements/PCGSelfPruning.h b/Engine/Plugins/Experimental/PCG/Source/PCG/Public/Elements/PCGSelfPruning.h
--- a/Engine/Plugins/Experimental/PCG/Source/PCG/Public/Elements/PCGSelfPruning.h
+++ b/Engine/Plugins/Experimental/PCG/Source/PCG/Public/Elements/PCGSelfPruning.h
@@ -39,7 +39,8 @@ public:
 
 #if WITH_EDITOR
 	//~Begin UPCGSettings interface
-	virtual FName GetDefaultNodeName() const override { return FName(TEXT("Self Pruning")); }
+	virtual FName GetDefaultNodeName() const override { return FName(TEXT("SelfPruning")); }
+	virtual FText GetDefaultNodeTitle() const override { return NSLOCTEXT("PCGSelfPruningSettings", "NodeTitle", "Self Pruning"); }
 	virtual EPCGSettingsType GetType() const override { return EPCGSettingsType::Filter; }
 #endif
 
diff --git a/Engine/Plugins/Experimental/PCG/Source/PCG/Public/Elements/PCGSpawnActor.h b/Engine/Plugins/Experimental/PCG/Source/PCG/Public/Elements/PCGSpawnActor.h
--- a/Engine/Plugins/Experimental/PCG/Source/PCG/Public/Elements/PCGSpawnActor.h
+++ b/Engine/Plugins/Experimental/PCG/Source/PCG/Public/Elements/PCGSpawnActor.h
@@ -90,7 +90,8 @@ public:
 	virtual UPCGNode* CreateNode() const override;
 
 #if WITH_EDITOR	
-	virtual FName GetDefaultNodeName() const override { return FName(TEXT("Spawn Actor")); }
+	virtual FName GetDefaultNodeName() const override { return FName(TEXT("SpawnActor")); }
+	virtual FText GetDefaultNodeTitle() const override { return NSLOCTEXT("PCGSpawnActorSettings", "NodeTitle", "Spawn Actor"); }
 	virtual EPCGSettingsType GetType() const override;
 #endif
 
diff --git a/Engine/Plugins/Experimental/PCG/Source/PCG/Public/Elements/PCGSplineSampler.h b/Engine/Plugins/Experimental/PCG/Source/PCG/Public/Elements/PCGSplineSampler.h
--- a/Engine/Plugins/Experimental/PCG/Source/PCG/Public/Elements/PCGSplineSampler.h
+++ b/Engine/Plugins/Experimental/PCG/Source/PCG/Public/Elements/PCGSplineSampler.h
@@ -141,7 +141,8 @@ class PCG_API UPCGSplineSamplerSettings : public UPCGSettings
 public:
 	// ~Begin UPCGSettings interface
 #if WITH_EDITOR
-	virtual FName GetDefaultNodeName() const override { return FName(TEXT("Spline Sampler")); }
+	virtual FName GetDefaultNodeName() const override { return FName(TEXT("SplineSampler")); }
+	virtual FText GetDefaultNodeTitle() const override { return NSLOCTEXT("PCGSplineSamplerSettings", "NodeTitle", "Spline Sampler"); }
 	virtual FText GetNodeTooltipText() const override;
 	virtual EPCGSettingsType GetType() const override { return EPCGSettingsType::Sampler; }
 #endif
diff --git a/Engine/Plugins/Experimental/PCG/Source/PCG/Public/Elements/PCGStaticMeshSpawner.h b/Engine/Plugins/Experimental/PCG/Source/PCG/Public/Elements/PCGStaticMeshSpawner.h
--- a/Engine/Plugins/Experimental/PCG/Source/PCG/Public/Elements/PCGStaticMeshSpawner.h
+++ b/Engine/Plugins/Experimental/PCG/Source/PCG/Public/Elements/PCGStaticMeshSpawner.h
@@ -47,6 +47,7 @@ public:
 #if WITH_EDITOR
 	// ~Begin UPCGSettings interface
 	virtual FName GetDefaultNodeName() const override { return FName(TEXT("StaticMeshSpawner")); }
+	virtual FText GetDefaultNodeTitle() const override { return NSLOCTEXT("PCGStaticMeshSpawnerSettings", "NodeTitle", "Static Mesh Spawner"); }
 	virtual EPCGSettingsType GetType() const override { return EPCGSettingsType::Spawner; }
 #endif
 
diff --git a/Engine/Plugins/Experimental/PCG/Source/PCG/Public/Elements/PCGSurfaceSampler.h b/Engine/Plugins/Experimental/PCG/Source/PCG/Public/Elements/PCGSurfaceSampler.h
--- a/Engine/Plugins/Experimental/PCG/Source/PCG/Public/Elements/PCGSurfaceSampler.h
+++ b/Engine/Plugins/Experimental/PCG/Source/PCG/Public/Elements/PCGSurfaceSampler.h
@@ -28,7 +28,8 @@ public:
 
 	//~Begin UPCGSettings interface
 #if WITH_EDITOR
-	virtual FName GetDefaultNodeName() const override { return FName(TEXT("Surface Sampler")); }
+	virtual FName GetDefaultNodeName() const override { return FName(TEXT("SurfaceSampler")); }
+	virtual FText GetDefaultNodeTitle() const override { return NSLOCTEXT("PCGSurfaceSamplerSettings", "NodeTitle", "Surface Sampler"); }
 	virtual FText GetNodeTooltipText() const override;
 	virtual EPCGSettingsType GetType() const override { return EPCGSettingsType::Sampler; }
 	virtual bool IsPinUsedByNodeExecution(const UPCGPin* InPin) const override;
diff --git a/Engine/Plugins/Experimental/PCG/Source/PCG/Public/Elements/PCGTextureSampler.h b/Engine/Plugins/Experimental/PCG/Source/PCG/Public/Elements/PCGTextureSampler.h
--- a/Engine/Plugins/Experimental/PCG/Source/PCG/Public/Elements/PCGTextureSampler.h
+++ b/Engine/Plugins/Experimental/PCG/Source/PCG/Public/Elements/PCGTextureSampler.h
@@ -16,7 +16,8 @@ class PCG_API UPCGTextureSamplerSettings : public UPCGSettings
 public:
 	//~Begin UPCGSettings interface
 #if WITH_EDITOR
-	virtual FName GetDefaultNodeName() const override { return FName(TEXT("Get Texture Data")); }
+	virtual FName GetDefaultNodeName() const override { return FName(TEXT("GetTextureData")); }
+	virtual FText GetDefaultNodeTitle() const override { return NSLOCTEXT("PCGTextureSamplerSettings", "NodeTitle", "Get Texture Data"); }
 	virtual EPCGSettingsType GetType() const override { return EPCGSettingsType::Spatial; }
 #endif
 
diff --git a/Engine/Plugins/Experimental/PCG/Source/PCG/Public/Elements/PCGTransformPoints.h b/Engine/Plugins/Experimental/PCG/Source/PCG/Public/Elements/PCGTransformPoints.h
--- a/Engine/Plugins/Experimental/PCG/Source/PCG/Public/Elements/PCGTransformPoints.h
+++ b/Engine/Plugins/Experimental/PCG/Source/PCG/Public/Elements/PCGTransformPoints.h
@@ -17,7 +17,8 @@ public:
 
 	//~Begin UPCGSettings interface
 #if WITH_EDITOR
-	virtual FName GetDefaultNodeName() const override { return FName(TEXT("Transform Points")); }
+	virtual FName GetDefaultNodeName() const override { return FName(TEXT("TransformPoints")); }
+	virtual FText GetDefaultNodeTitle() const override { return NSLOCTEXT("PCGTransformPointsSettings", "NodeTitle", "Transform Points"); }
 	virtual EPCGSettingsType GetType() const override { return EPCGSettingsType::Spatial; }
 #endif
 
diff --git a/Engine/Plugins/Experimental/PCG/Source/PCG/Public/Elements/PCGTypedGetter.h b/Engine/Plugins/Experimental/PCG/Source/PCG/Public/Elements/PCGTypedGetter.h
--- a/Engine/Plugins/Experimental/PCG/Source/PCG/Public/Elements/PCGTypedGetter.h
+++ b/Engine/Plugins/Experimental/PCG/Source/PCG/Public/Elements/PCGTypedGetter.h
@@ -17,7 +17,8 @@ public:
 
 	//~Begin UPCGSettings interface
 #if WITH_EDITOR
-	virtual FName GetDefaultNodeName() const override { return FName(TEXT("Get Landscape Data")); }
+	virtual FName GetDefaultNodeName() const override { return FName(TEXT("GetLandscapeData")); }
+	virtual FText GetDefaultNodeTitle() const override { return NSLOCTEXT("PCGGetLandscapeSettings", "NodeTitle", "Get Landscape Data"); }
 	virtual FText GetNodeTooltipText() const override;
 #endif
 
@@ -42,7 +43,8 @@ public:
 
 	//~Begin UPCGSettings interface
 #if WITH_EDITOR
-	virtual FName GetDefaultNodeName() const override { return FName(TEXT("Get Spline Data")); }
+	virtual FName GetDefaultNodeName() const override { return FName(TEXT("GetSplineData")); }
+	virtual FText GetDefaultNodeTitle() const override { return NSLOCTEXT("PCGGetSplineSettings", "NodeTitle", "Get Spline Data"); }
 	virtual FText GetNodeTooltipText() const override;
 #endif
 
@@ -92,7 +94,8 @@ public:
 
 	//~Begin UPCGSettings interface
 #if WITH_EDITOR
-	virtual FName GetDefaultNodeName() const override { return FName(TEXT("Get Primitive Data")); }
+	virtual FName GetDefaultNodeName() const override { return FName(TEXT("GetPrimitiveData")); }
+	virtual FText GetDefaultNodeTitle() const override { return NSLOCTEXT("PCGGetPrimitiveSettings", "NodeTitle", "Get Primitive Data"); }
 	virtual FText GetNodeTooltipText() const override;
 #endif
 
diff --git a/Engine/Plugins/Experimental/PCG/Source/PCG/Public/Elements/PCGUnionElement.h b/Engine/Plugins/Experimental/PCG/Source/PCG/Public/Elements/PCGUnionElement.h
--- a/Engine/Plugins/Experimental/PCG/Source/PCG/Public/Elements/PCGUnionElement.h
+++ b/Engine/Plugins/Experimental/PCG/Source/PCG/Public/Elements/PCGUnionElement.h
@@ -16,6 +16,7 @@ public:
 	//~Begin UPCGSettings interface
 #if WITH_EDITOR
 	virtual FName GetDefaultNodeName() const override { return FName(TEXT("Union")); }
+	virtual FText GetDefaultNodeTitle() const override { return NSLOCTEXT("PCGUnionSettings", "NodeTitle", "Union"); }
 	virtual EPCGSettingsType GetType() const override { return EPCGSettingsType::Spatial; }
 #endif
 
diff --git a/Engine/Plugins/Experimental/PCG/Source/PCG/Public/Elements/PCGUserParameterGet.h b/Engine/Plugins/Experimental/PCG/Source/PCG/Public/Elements/PCGUserParameterGet.h
--- a/Engine/Plugins/Experimental/PCG/Source/PCG/Public/Elements/PCGUserParameterGet.h
+++ b/Engine/Plugins/Experimental/PCG/Source/PCG/Public/Elements/PCGUserParameterGet.h
@@ -38,7 +38,8 @@ protected:
 	virtual FPCGElementPtr CreateElement() const override;
 
 #if WITH_EDITOR
-	virtual FName GetDefaultNodeName() const override { return FName("Get User Parameter"); }
+	virtual FName GetDefaultNodeName() const override { return FName("GetUserParameter"); }
+	virtual FText GetDefaultNodeTitle() const override { return NSLOCTEXT("PCGUserParameterGetSettings", "NodeTitle", "Get User Parameter"); }
 #endif
 	//~End UPCGSettings interface
 };
diff --git a/Engine/Plugins/Experimental/PCG/Source/PCG/Public/Elements/PCGVolumeSampler.h b/Engine/Plugins/Experimental/PCG/Source/PCG/Public/Elements/PCGVolumeSampler.h
--- a/Engine/Plugins/Experimental/PCG/Source/PCG/Public/Elements/PCGVolumeSampler.h
+++ b/Engine/Plugins/Experimental/PCG/Source/PCG/Public/Elements/PCGVolumeSampler.h
@@ -42,7 +42,8 @@ public:
 
 	//~Begin UPCGSettings interface
 #if WITH_EDITOR
-	virtual FName GetDefaultNodeName() const override { return FName(TEXT("Volume Sampler")); }
+	virtual FName GetDefaultNodeName() const override { return FName(TEXT("VolumeSampler")); }
+	virtual FText GetDefaultNodeTitle() const override { return NSLOCTEXT("PCGVolumeSamplerSettings", "NodeTitle", "Volume Sampler"); }
 	virtual FText GetNodeTooltipText() const override;
 	virtual EPCGSettingsType GetType() const override { return EPCGSettingsType::Sampler; }
 #endif
diff --git a/Engine/Plugins/Experimental/PCG/Source/PCG/Public/Elements/PCGWorldQuery.h b/Engine/Plugins/Experimental/PCG/Source/PCG/Public/Elements/PCGWorldQuery.h
--- a/Engine/Plugins/Experimental/PCG/Source/PCG/Public/Elements/PCGWorldQuery.h
+++ b/Engine/Plugins/Experimental/PCG/Source/PCG/Public/Elements/PCGWorldQuery.h
@@ -16,7 +16,8 @@ class PCG_API UPCGWorldQuerySettings : public UPCGSettings
 public:
 	//~Begin UPCGSettings interface
 #if WITH_EDITOR
-	virtual FName GetDefaultNodeName() const override { return FName(TEXT("World Volumetric Query")); }
+	virtual FName GetDefaultNodeName() const override { return FName(TEXT("WorldVolumetricQuery")); }
+	virtual FText GetDefaultNodeTitle() const override { return NSLOCTEXT("PCGWorldQuerySettings", "NodeTitle", "World Volumetric Query"); }
 	virtual FText GetNodeTooltipText() const override;
 	virtual EPCGSettingsType GetType() const override { return EPCGSettingsType::Spatial; }
 #endif
@@ -49,7 +50,8 @@ class PCG_API UPCGWorldRayHitSettings : public UPCGSettings
 public:
 	//~Begin UPCGSettings interface
 #if WITH_EDITOR
-	virtual FName GetDefaultNodeName() const override { return FName(TEXT("World Ray Hit Query")); }
+	virtual FName GetDefaultNodeName() const override { return FName(TEXT("WorldRayHitQuery")); }
+	virtual FText GetDefaultNodeTitle() const override { return NSLOCTEXT("PCGWorldRayHitSettings", "NodeTitle", "World Ray Hit Query"); }
 	virtual FText GetNodeTooltipText() const override;
 	virtual EPCGSettingsType GetType() const override { return EPCGSettingsType::Spatial; }
 #endif
diff --git a/Engine/Plugins/Experimental/PCG/Source/PCG/Public/PCGInputOutputSettings.h b/Engine/Plugins/Experimental/PCG/Source/PCG/Public/PCGInputOutputSettings.h
--- a/Engine/Plugins/Experimental/PCG/Source/PCG/Public/PCGInputOutputSettings.h
+++ b/Engine/Plugins/Experimental/PCG/Source/PCG/Public/PCGInputOutputSettings.h
@@ -39,6 +39,7 @@ public:
 	// ~Begin UPCGSettings interface
 #if WITH_EDITOR
 	virtual FName GetDefaultNodeName() const override { return bIsInput ? FName(TEXT("InputNode")) : FName(TEXT("OutputNode")); }
+	virtual FText GetDefaultNodeTitle() const override { return bIsInput ? NSLOCTEXT("PCGGraphInputOutputSettings", "InputNodeTitle", "Input Node") : NSLOCTEXT("PCGGraphInputOutputSettings", "OutputNodeTitle", "Output Node"); }
 	virtual EPCGSettingsType GetType() const override { return EPCGSettingsType::InputOutput; }
 	virtual void PostEditChangeProperty(struct FPropertyChangedEvent& PropertyChangedEvent);
 #endif
diff --git a/Engine/Plugins/Experimental/PCG/Source/PCG/Public/PCGNode.h b/Engine/Plugins/Experimental/PCG/Source/PCG/Public/PCGNode.h
--- a/Engine/Plugins/Experimental/PCG/Source/PCG/Public/PCGNode.h
+++ b/Engine/Plugins/Experimental/PCG/Source/PCG/Public/PCGNode.h
@@ -61,7 +61,7 @@ public:
 	bool RemoveEdgeTo(FName FromPinLable, UPCGNode* To, FName ToPinLabel);
 
 	/** Returns the node title, based either on the current node label, or defaulted to its settings */
-	FName GetNodeTitle() const;
+	FText GetNodeTitle() const;
 
 #if WITH_EDITOR
 	/** Tooltip that describes node functionality and other information. */
diff --git a/Engine/Plugins/Experimental/PCG/Source/PCG/Public/PCGSettings.h b/Engine/Plugins/Experimental/PCG/Source/PCG/Public/PCGSettings.h
--- a/Engine/Plugins/Experimental/PCG/Source/PCG/Public/PCGSettings.h
+++ b/Engine/Plugins/Experimental/PCG/Source/PCG/Public/PCGSettings.h
@@ -168,6 +168,7 @@ public:
 	virtual void ApplyDeprecation(UPCGNode* InOutNode);
 
 	virtual FName GetDefaultNodeName() const { return NAME_None; }
+	virtual FText GetDefaultNodeTitle() const { return FText::FromName(GetDefaultNodeName()); }
 	virtual FText GetNodeTooltipText() const { return FText::GetEmpty(); }
 	virtual FLinearColor GetNodeTitleColor() const { return FLinearColor::White; }
 	virtual EPCGSettingsType GetType() const { return EPCGSettingsType::Generic; }
diff --git a/Engine/Plugins/Experimental/PCG/Source/PCG/Public/PCGSubgraph.h b/Engine/Plugins/Experimental/PCG/Source/PCG/Public/PCGSubgraph.h
--- a/Engine/Plugins/Experimental/PCG/Source/PCG/Public/PCGSubgraph.h
+++ b/Engine/Plugins/Experimental/PCG/Source/PCG/Public/PCGSubgraph.h
@@ -78,6 +78,7 @@ public:
 
 #if WITH_EDITOR
 	virtual FName GetDefaultNodeName() const override { return FName(TEXT("Subgraph")); }
+	virtual FText GetDefaultNodeTitle() const override { return NSLOCTEXT("PCGSubgraphSettings", "NodeTitle", "Subgraph"); }
 	virtual EPCGSettingsType GetType() const override { return EPCGSettingsType::Subgraph; }
 	virtual UObject* GetJumpTargetForDoubleClick() const override;
 #endif
diff --git a/Engine/Plugins/Experimental/PCG/Source/PCGEditor/Private/PCGEditor.cpp b/Engine/Plugins/Experimental/PCG/Source/PCGEditor/Private/PCGEditor.cpp
--- a/Engine/Plugins/Experimental/PCG/Source/PCGEditor/Private/PCGEditor.cpp
+++ b/Engine/Plugins/Experimental/PCG/Source/PCGEditor/Private/PCGEditor.cpp
@@ -568,7 +568,7 @@ void FPCGEditor::OnDeterminismNodeTest()
 				check(PCGNode && PCGNode->GetSettings());
 
 				TSharedPtr<FDeterminismTestResult> NodeResult = MakeShared<FDeterminismTestResult>();
-				NodeResult->TestResultTitle = PCGNode->GetNodeTitle();
+				NodeResult->TestResultTitle = FName(*PCGNode->GetNodeTitle().ToString());
 				NodeResult->TestResultName = PCGNode->GetName();
 				NodeResult->Seed = PCGNode->GetSettings()->GetSeed();
 
diff --git a/Engine/Plugins/Experimental/PCG/Source/PCGEditor/Private/PCGEditorGraphNode.cpp b/Engine/Plugins/Experimental/PCG/Source/PCGEditor/Private/PCGEditorGraphNode.cpp
--- a/Engine/Plugins/Experimental/PCG/Source/PCGEditor/Private/PCGEditorGraphNode.cpp
+++ b/Engine/Plugins/Experimental/PCG/Source/PCGEditor/Private/PCGEditorGraphNode.cpp
@@ -106,12 +106,10 @@ void UPCGEditorGraphNode::OnRenameNode(const FString& NewName)
 		return;
 	}
 
-	const FName TentativeName(*NewName);
-
-	if (PCGNode->GetNodeTitle() != TentativeName)
+	if(PCGNode->GetNodeTitle().ToString() != NewName)
 	{
 		PCGNode->Modify();
-		PCGNode->NodeTitle = TentativeName;
+		PCGNode->NodeTitle = FName(*NewName);
 	}
 }
 
diff --git a/Engine/Plugins/Experimental/PCG/Source/PCGEditor/Private/PCGEditorGraphNodeInput.cpp b/Engine/Plugins/Experimental/PCG/Source/PCGEditor/Private/PCGEditorGraphNodeInput.cpp
--- a/Engine/Plugins/Experimental/PCG/Source/PCGEditor/Private/PCGEditorGraphNodeInput.cpp
+++ b/Engine/Plugins/Experimental/PCG/Source/PCGEditor/Private/PCGEditorGraphNodeInput.cpp
@@ -9,7 +9,7 @@
 
 FText UPCGEditorGraphNodeInput::GetNodeTitle(ENodeTitleType::Type TitleType) const
 {	
-	return FText::FromName(TEXT("Input"));
+	return LOCTEXT("NodeTitle", "Input");
 }
 
 void UPCGEditorGraphNodeInput::AllocateDefaultPins()
diff --git a/Engine/Plugins/Experimental/PCG/Source/PCGEditor/Private/PCGEditorGraphNodeOutput.cpp b/Engine/Plugins/Experimental/PCG/Source/PCGEditor/Private/PCGEditorGraphNodeOutput.cpp
--- a/Engine/Plugins/Experimental/PCG/Source/PCGEditor/Private/PCGEditorGraphNodeOutput.cpp
+++ b/Engine/Plugins/Experimental/PCG/Source/PCGEditor/Private/PCGEditorGraphNodeOutput.cpp
@@ -9,7 +9,7 @@
 
 FText UPCGEditorGraphNodeOutput::GetNodeTitle(ENodeTitleType::Type TitleType) const
 {
-	return FText::FromName(TEXT("Output"));
+	return LOCTEXT("NodeTitle", "Output");
 }
 
 void UPCGEditorGraphNodeOutput::AllocateDefaultPins()
diff --git a/Engine/Plugins/Experimental/PCG/Source/PCGEditor/Private/PCGEditorGraphNodeReroute.cpp b/Engine/Plugins/Experimental/PCG/Source/PCGEditor/Private/PCGEditorGraphNodeReroute.cpp
--- a/Engine/Plugins/Experimental/PCG/Source/PCGEditor/Private/PCGEditorGraphNodeReroute.cpp
+++ b/Engine/Plugins/Experimental/PCG/Source/PCGEditor/Private/PCGEditorGraphNodeReroute.cpp
@@ -6,7 +6,7 @@
 
 FText UPCGEditorGraphNodeReroute::GetNodeTitle(ENodeTitleType::Type TitleType) const
 {
-	return FText::FromName(TEXT("Reroute"));
+	return NSLOCTEXT("PCGEditorGraphNodeReroute", "NodeTitle", "Reroute");
 }
 
 bool UPCGEditorGraphNodeReroute::ShouldOverridePinNames() const
diff --git a/Engine/Plugins/Experimental/PCG/Source/PCGEditor/Private/PCGEditorGraphSchema.cpp b/Engine/Plugins/Experimental/PCG/Source/PCGEditor/Private/PCGEditorGraphSchema.cpp
--- a/Engine/Plugins/Experimental/PCG/Source/PCGEditor/Private/PCGEditorGraphSchema.cpp
+++ b/Engine/Plugins/Experimental/PCG/Source/PCGEditor/Private/PCGEditorGraphSchema.cpp
@@ -221,7 +221,7 @@ void UPCGEditorGraphSchema::GetNativeElementActions(FGraphActionMenuBuilder& Act
 		{
 			if (PCGSettings->bExposeToLibrary)
 			{
-				const FText MenuDesc = FText::FromName(PCGSettings->GetDefaultNodeName());
+				const FText MenuDesc = PCGSettings->GetDefaultNodeTitle();
 				const FText Category = StaticEnum<EPCGSettingsType>()->GetDisplayNameTextByValue(static_cast<__underlying_type(EPCGSettingsType)>(PCGSettings->GetType()));
 				const FText Description = PCGSettings->GetNodeTooltipText();
 
diff --git a/Engine/Plugins/Experimental/PCG/Source/PCGEditor/Private/SPCGEditorGraphLogView.cpp b/Engine/Plugins/Experimental/PCG/Source/PCGEditor/Private/SPCGEditorGraphLogView.cpp
--- a/Engine/Plugins/Experimental/PCG/Source/PCGEditor/Private/SPCGEditorGraphLogView.cpp
+++ b/Engine/Plugins/Experimental/PCG/Source/PCGEditor/Private/SPCGEditorGraphLogView.cpp
@@ -307,7 +307,7 @@ FReply SPCGEditorGraphLogView::Refresh()
 		const UPCGNode* PCGNode = PCGEditorNode ? PCGEditorNode->GetPCGNode() : nullptr;
 		if (PCGNode)
 		{
-			CreateAndAddItem(PCGEditorNode, PCGNode, PCGNode->GetNodeTitle());
+			CreateAndAddItem(PCGEditorNode, PCGNode, *PCGNode->GetNodeTitle().ToString());
 		}
 	}
 
diff --git a/Engine/Plugins/Experimental/PCGGeometryScriptInterop/Source/PCGGeometryScriptInterop/Public/Elements/PCGMeshSampler.h b/Engine/Plugins/Experimental/PCGGeometryScriptInterop/Source/PCGGeometryScriptInterop/Public/Elements/PCGMeshSampler.h
--- a/Engine/Plugins/Experimental/PCGGeometryScriptInterop/Source/PCGGeometryScriptInterop/Public/Elements/PCGMeshSampler.h
+++ b/Engine/Plugins/Experimental/PCGGeometryScriptInterop/Source/PCGGeometryScriptInterop/Public/Elements/PCGMeshSampler.h
@@ -49,7 +49,8 @@ public:
 
 	//~Begin UPCGSettings interface
 #if WITH_EDITOR
-	virtual FName GetDefaultNodeName() const override { return FName(TEXT("Mesh Sampler")); }
+	virtual FName GetDefaultNodeName() const override { return FName(TEXT("MeshSampler")); }
+	virtual FText GetDefaultNodeTitle() const override { return NSLOCTEXT("PCGMeshSamplerSettings", "NodeTitle", "Mesh Sampler"); }
 	virtual FText GetNodeTooltipText() const override;
 	virtual EPCGSettingsType GetType() const override { return EPCGSettingsType::Sampler; }
 #endif
