[PCG] Clean up a bunch of user-facing references to Param, which was rebranded to Attribute Set and will not make sense to users

diff --git a/Engine/Plugins/Experimental/PCG/Source/PCG/Private/PCGContext.cpp b/Engine/Plugins/Experimental/PCG/Source/PCG/Private/PCGContext.cpp
--- a/Engine/Plugins/Experimental/PCG/Source/PCG/Private/PCGContext.cpp
+++ b/Engine/Plugins/Experimental/PCG/Source/PCG/Private/PCGContext.cpp
@@ -207,7 +207,7 @@ void FPCGContext::OverrideSettings()
 			PropertyType Value{};
 			if (!AttributeAccessor->Get<PropertyType>(Value, FirstEntry, EPCGAttributeAccessorFlags::AllowBroadcast | EPCGAttributeAccessorFlags::AllowConstructible))
 			{
-				PCGE_LOG_C(Warning, GraphAndLog, this, FText::Format(LOCTEXT("ConversionFailed", "Param '{0}' cannot be converted from attribute '{1}'"), FText::FromName(Param.Label), FText::FromName(AttributeName)));
+				PCGE_LOG_C(Warning, GraphAndLog, this, FText::Format(LOCTEXT("ConversionFailed", "Parameter '{0}' cannot be converted from attribute '{1}'"), FText::FromName(Param.Label), FText::FromName(AttributeName)));
 				return false;
 			}
 			// TODO: Perhaps factorise this code in another property accessor.
diff --git a/Engine/Plugins/Experimental/PCG/Source/PCG/Private/PCGSettings.cpp b/Engine/Plugins/Experimental/PCG/Source/PCG/Private/PCGSettings.cpp
--- a/Engine/Plugins/Experimental/PCG/Source/PCG/Private/PCGSettings.cpp
+++ b/Engine/Plugins/Experimental/PCG/Source/PCG/Private/PCGSettings.cpp
@@ -269,8 +269,9 @@ void UPCGSettings::FillOverridableParamsPins(TArray<FPCGPinProperties>& OutPins)
 		if (*PinType != EPCGDataType::Param)
 		{
 			const FString ParamsName = PCGPinConstants::DefaultParamsLabel.ToString();
-			UE_LOG(LogPCG, Error, TEXT("[%s-%s] While adding %s pin, we found another %s pin with not the same allowed type (Param). "
-				"Please rename this pin if you want to take advantage of automatic override. Until then it will probably break your graph."), *GraphName, *NodeName, *ParamsName, *ParamsName);
+			UE_LOG(LogPCG, Error, TEXT("[%s-%s] While adding pin '%s', another pin '%s' was found with invalid type (type must be Attribute Set). "
+				"Rename or remove this pin to allow an override pin to be added automatically."),
+				*GraphName, *NodeName, *ParamsName, *ParamsName);
 		}
 	}
 	else
@@ -279,7 +280,7 @@ void UPCGSettings::FillOverridableParamsPins(TArray<FPCGPinProperties>& OutPins)
 		ParamPin.bAdvancedPin = true;
 
 #if WITH_EDITOR
-		ParamPin.Tooltip = LOCTEXT("GlobalParamPinTooltip", "Can bundle multiple param data to override multiple parameters at the same time. Names need to match perfectly.");
+		ParamPin.Tooltip = LOCTEXT("GlobalParamPinTooltip", "Atribute Set containing multiple parameters to override. Names must match perfectly.");
 #endif // WITH_EDITOR
 	}
 
@@ -289,9 +290,10 @@ void UPCGSettings::FillOverridableParamsPins(TArray<FPCGPinProperties>& OutPins)
 	{
 		if (InputPinsLabelsAndTypes.Contains(OverridableParam.Label))
 		{
-			//const FString ParamsName = OverridableParam.Label.ToString();
-			//UE_LOG(LogPCG, Warning, TEXT("[%s-%s] While automatically adding overriable param pins, we found a %s pin. "
-			//	"Please rename this pin if you want to take advantage of automatic override. Until then, we will not add a %s pin."), *GraphName, *NodeName, *ParamsName, *ParamsName);
+			const FString ParamsName = OverridableParam.Label.ToString();
+			UE_LOG(LogPCG, Warning, TEXT("[%s-%s] While automatically adding override pins, an existing pin was found with conflicting name '%s'. "
+				"Rename or remove this pin to allow the automatic override pin to be added. Automatic override pin '%s' skipped."),
+				*GraphName, *NodeName, *ParamsName, *ParamsName);
 			continue;
 		}
 
@@ -309,7 +311,7 @@ void UPCGSettings::FillOverridableParamsPins(TArray<FPCGPinProperties>& OutPins)
 			Tooltip = *TooltipPtr + TEXT("\n");
 		}
 
-		ParamPin.Tooltip = FText::Format(LOCTEXT("OverridableParamPinTooltip", "{0}Param type is \"{1}\" and its exact name is \"{2}\""),
+		ParamPin.Tooltip = FText::Format(LOCTEXT("OverridableParamPinTooltip", "{0}Attribute type is \"{1}\" and its exact name is \"{2}\""),
 			FText::FromString(Tooltip),
 			FText::FromString(Property->GetCPPType()),
 			FText::FromName(Property->GetFName()));
diff --git a/Engine/Plugins/Experimental/PCG/Source/PCG/Public/Helpers/PCGSettingsHelpers.h b/Engine/Plugins/Experimental/PCG/Source/PCG/Public/Helpers/PCGSettingsHelpers.h
--- a/Engine/Plugins/Experimental/PCG/Source/PCG/Public/Helpers/PCGSettingsHelpers.h
+++ b/Engine/Plugins/Experimental/PCG/Source/PCG/Public/Helpers/PCGSettingsHelpers.h
@@ -194,7 +194,7 @@ PRAGMA_ENABLE_DEPRECATION_WARNINGS
 
 			if (!AttributeAccessor->Get<T>(OutValue, FirstEntry, EPCGAttributeAccessorFlags::AllowBroadcast | EPCGAttributeAccessorFlags::AllowConstructible))
 			{
-				UE_LOG(LogPCG, Warning, TEXT("[PCGSettingsHelpers::GetOverrideValue] %s param can't be converted from %s attribute, incompatible types."), *Param->Label.ToString(), *AttributeName.ToString());
+				UE_LOG(LogPCG, Warning, TEXT("[PCGSettingsHelpers::GetOverrideValue] '%s' parameter cannot be converted from '%s' attribute, incompatible types."), *Param->Label.ToString(), *AttributeName.ToString());
 				return false;
 			}
 
diff --git a/Engine/Plugins/Experimental/PCG/Source/PCG/Public/PCGCommon.h b/Engine/Plugins/Experimental/PCG/Source/PCG/Public/PCGCommon.h
--- a/Engine/Plugins/Experimental/PCG/Source/PCG/Public/PCGCommon.h
+++ b/Engine/Plugins/Experimental/PCG/Source/PCG/Public/PCGCommon.h
@@ -48,7 +48,7 @@ enum class EPCGDataType : uint32
 	/** Combinations of concrete data and/or boolean operations. */
 	Spatial = Composite | Concrete,
 
-	Param = 1 << 27,
+	Param = 1 << 27 UMETA(DisplayName = "Attribute Set"),
 	Settings = 1 << 28 UMETA(Hidden),
 	Other = 1 << 29,
 	Any = (1 << 30) - 1
diff --git a/Engine/Plugins/Experimental/PCG/Source/PCG/Public/PCGParamData.h b/Engine/Plugins/Experimental/PCG/Source/PCG/Public/PCGParamData.h
--- a/Engine/Plugins/Experimental/PCG/Source/PCG/Public/PCGParamData.h
+++ b/Engine/Plugins/Experimental/PCG/Source/PCG/Public/PCGParamData.h
@@ -10,7 +10,7 @@
 /**
 * Class to hold execution parameters that will be consumed in nodes of the graph
 */
-UCLASS(BlueprintType, ClassGroup = (Procedural))
+UCLASS(BlueprintType, ClassGroup = (Procedural), DisplayName = "Attribute Set")
 class PCG_API UPCGParamData : public UPCGData
 {
 	GENERATED_BODY()
diff --git a/Engine/Plugins/Experimental/PCG/Source/PCGEditor/Private/SPCGEditorGraphAttributeListView.cpp b/Engine/Plugins/Experimental/PCG/Source/PCGEditor/Private/SPCGEditorGraphAttributeListView.cpp
--- a/Engine/Plugins/Experimental/PCG/Source/PCGEditor/Private/SPCGEditorGraphAttributeListView.cpp
+++ b/Engine/Plugins/Experimental/PCG/Source/PCGEditor/Private/SPCGEditorGraphAttributeListView.cpp
@@ -758,7 +758,8 @@ void SPCGEditorGraphAttributeListView::RefreshDataComboBox()
 	for(int32 TaggedDataIndex = 0; TaggedDataIndex < InspectionData->TaggedData.Num(); ++TaggedDataIndex)
 	{
 		const FPCGTaggedData& TaggedData = InspectionData->TaggedData[TaggedDataIndex];
-		FString ItemName = FString::Format(TEXT("[{0}] {1} - {2}"), { FText::AsNumber(TaggedDataIndex).ToString(), TaggedData.Pin.ToString(), (TaggedData.Data ? TaggedData.Data->GetClass()->GetFName().ToString(): TEXT("No Data")) });
+		FString ItemName = FString::Format(TEXT("[{0}] {1} - {2}"),
+			{ FText::AsNumber(TaggedDataIndex).ToString(), TaggedData.Pin.ToString(), (TaggedData.Data ? TaggedData.Data->GetClass()->GetDisplayNameText().ToString(): TEXT("No Data")) });
 
 		if (!TaggedData.Tags.IsEmpty())
 		{
diff --git a/Engine/Plugins/Experimental/PCG/Source/PCGEditor/Public/PCGEditorSettings.h b/Engine/Plugins/Experimental/PCG/Source/PCGEditor/Public/PCGEditorSettings.h
--- a/Engine/Plugins/Experimental/PCG/Source/PCGEditor/Public/PCGEditorSettings.h
+++ b/Engine/Plugins/Experimental/PCG/Source/PCGEditor/Public/PCGEditorSettings.h
@@ -65,8 +65,8 @@ public:
 	UPROPERTY(EditAnywhere, config, Category = Node, meta = (HideAlphaChannel))
 	FLinearColor SubgraphNodeColor;
 
-	/** Color used for param data operations */
-	UPROPERTY(EditAnywhere, config, Category = Node, meta = (HideAlphaChannel))
+	/** Color used for Attribute Set operations */
+	UPROPERTY(EditAnywhere, config, Category = Node, meta = (HideAlphaChannel, DisplayName = "Attribute Set Node Color"))
 	FLinearColor ParamDataNodeColor;
 
 	/** Color used for debug operations */
@@ -117,8 +117,8 @@ public:
 	UPROPERTY(EditAnywhere, config, Category = Node, meta = (HideAlphaChannel))
 	FLinearColor PrimitiveDataPinColor;
 
-	/** Color used for param data pins */
-	UPROPERTY(EditAnywhere, config, Category = Node, meta = (HideAlphaChannel))
+	/** Color used for data pins of type Attribute Set */
+	UPROPERTY(EditAnywhere, config, Category = Node, meta = (HideAlphaChannel, DisplayName = "Attribute Set Pin Color"))
 	FLinearColor ParamDataPinColor;
 
 	/** Color used for other/unknown data types */
