[PCG] Remove exclusion data as it's giving type warnings and the plan is to phase it out anyway.

diff --git a/Engine/Plugins/Experimental/PCG/Source/PCG/Private/Graph/PCGGraphExecutor.cpp b/Engine/Plugins/Experimental/PCG/Source/PCG/Private/Graph/PCGGraphExecutor.cpp
--- a/Engine/Plugins/Experimental/PCG/Source/PCG/Private/Graph/PCGGraphExecutor.cpp
+++ b/Engine/Plugins/Experimental/PCG/Source/PCG/Private/Graph/PCGGraphExecutor.cpp
@@ -1353,17 +1353,6 @@ bool FPCGFetchInputElement::ExecuteInternal(FPCGContext* Context) const
 		}
 	}
 
-	if (Context->Node->IsOutputPinConnected(PCGInputOutputConstants::DefaultExcludedActorsLabel))
-	{
-		TArray<UPCGData*> ExclusionsPCGData = Component->GetPCGExclusionData();
-		for (UPCGData* ExclusionPCGData : ExclusionsPCGData)
-		{
-			FPCGTaggedData& TaggedData = Context->OutputData.TaggedData.Emplace_GetRef();
-			TaggedData.Data = ExclusionPCGData;
-			TaggedData.Pin = PCGInputOutputConstants::DefaultExcludedActorsLabel;
-		}
-	}
-
 	return true;
 }
 
diff --git a/Engine/Plugins/Experimental/PCG/Source/PCG/Private/Helpers/PCGBlueprintHelpers.cpp b/Engine/Plugins/Experimental/PCG/Source/PCG/Private/Helpers/PCGBlueprintHelpers.cpp
--- a/Engine/Plugins/Experimental/PCG/Source/PCG/Private/Helpers/PCGBlueprintHelpers.cpp
+++ b/Engine/Plugins/Experimental/PCG/Source/PCG/Private/Helpers/PCGBlueprintHelpers.cpp
@@ -61,11 +61,6 @@ UPCGData* UPCGBlueprintHelpers::GetInputData(FPCGContext& Context)
 	return Context.SourceComponent.IsValid() ? Context.SourceComponent->GetInputPCGData() : nullptr;
 }
 
-TArray<UPCGData*> UPCGBlueprintHelpers::GetExclusionData(FPCGContext& Context)
-{
-	return Context.SourceComponent.IsValid() ? Context.SourceComponent->GetPCGExclusionData() : TArray<UPCGData*>();
-}
-
 UPCGComponent* UPCGBlueprintHelpers::GetComponent(FPCGContext& Context)
 {
 	return Context.SourceComponent.Get();
diff --git a/Engine/Plugins/Experimental/PCG/Source/PCG/Private/PCGComponent.cpp b/Engine/Plugins/Experimental/PCG/Source/PCG/Private/PCGComponent.cpp
--- a/Engine/Plugins/Experimental/PCG/Source/PCG/Private/PCGComponent.cpp
+++ b/Engine/Plugins/Experimental/PCG/Source/PCG/Private/PCGComponent.cpp
@@ -198,7 +198,6 @@ void UPCGComponent::SetPropertiesFromOriginal(const UPCGComponent* Original)
 
 #if WITH_EDITOR
 	const bool bHasDirtyInput = InputType != NewInputType;
-	const bool bHasDirtyExclusions = !(ExcludedTags.Num() == Original->ExcludedTags.Num() && ExcludedTags.Includes(Original->ExcludedTags));
 
 	TSet<FName> TrackedTags;
 	TSet<FName> OriginalTrackedTags;
@@ -206,7 +205,7 @@ void UPCGComponent::SetPropertiesFromOriginal(const UPCGComponent* Original)
 	Original->CachedTrackedTagsToSettings.GetKeys(OriginalTrackedTags);
 	const bool bHasDirtyTracking = !(TrackedTags.Num() == OriginalTrackedTags.Num() && TrackedTags.Includes(OriginalTrackedTags));
 
-	const bool bIsDirty = bHasDirtyInput || bHasDirtyExclusions || bHasDirtyTracking || bGraphInstanceIsDifferent;
+	const bool bIsDirty = bHasDirtyInput || bHasDirtyTracking || bGraphInstanceIsDifferent;
 #endif // WITH_EDITOR
 
 	InputType = NewInputType;
@@ -225,10 +224,9 @@ void UPCGComponent::SetPropertiesFromOriginal(const UPCGComponent* Original)
 	}
 
 #if WITH_EDITOR
-	if (bHasDirtyExclusions || bHasDirtyTracking)
+	if (bHasDirtyTracking)
 	{
 		TeardownTrackingCallbacks();
-		ExcludedTags = Original->ExcludedTags;
 		SetupTrackingCallbacks();
 		RefreshTrackingData();
 	}
@@ -237,10 +235,8 @@ void UPCGComponent::SetPropertiesFromOriginal(const UPCGComponent* Original)
 	if (bIsDirty)
 	{
 		Modify();
-		DirtyGenerated((bHasDirtyInput ? EPCGComponentDirtyFlag::Input : EPCGComponentDirtyFlag::None) | (bHasDirtyExclusions ? EPCGComponentDirtyFlag::Exclusions : EPCGComponentDirtyFlag::None));
+		DirtyGenerated(bHasDirtyInput ? EPCGComponentDirtyFlag::Input : EPCGComponentDirtyFlag::None);
 	}
-#else
-	ExcludedTags = Original->ExcludedTags;
 #endif
 }
 
@@ -323,45 +319,6 @@ FPCGTaskId UPCGComponent::CreateGenerateTask(bool bForce, const TArray<FPCGTaskI
 	return GetSubsystem()->ScheduleGraph(this, *AllDependencies);
 }
 
-bool UPCGComponent::GetActorsFromTags(const TSet<FName>& InTags, TSet<TWeakObjectPtr<AActor>>& OutActors, bool bCullAgainstLocalBounds)
-{
-	TRACE_CPUPROFILER_EVENT_SCOPE(UPCGComponent::GetActorsFromTags::Excluded);
-	UWorld* World = GetWorld();
-
-	if (!World)
-	{
-		return false;
-	}
-
-	FBox LocalBounds = bCullAgainstLocalBounds ? GetGridBounds() : FBox(EForceInit::ForceInit);
-
-	TArray<AActor*> PerTagActors;
-
-	OutActors.Reset();
-
-	bool bHasValidTag = false;
-	for (const FName& Tag : InTags)
-	{
-		if (Tag != NAME_None)
-		{
-			bHasValidTag = true;
-			UGameplayStatics::GetAllActorsWithTag(World, Tag, PerTagActors);
-
-			for (AActor* Actor : PerTagActors)
-			{
-				if (!bCullAgainstLocalBounds || LocalBounds.Intersect(GetGridBounds(Actor)))
-				{
-					OutActors.Emplace(Actor);
-				}
-			}
-
-			PerTagActors.Reset();
-		}
-	}
-
-	return bHasValidTag;
-}
-
 bool UPCGComponent::GetActorsFromTags(const TMap<FName, bool>& InTagsAndCulling, TSet<TWeakObjectPtr<AActor>>& OutActors)
 {
 	TRACE_CPUPROFILER_EVENT_SCOPE(UPCGComponent::GetActorsFromTags::Tracked);
@@ -1010,12 +967,6 @@ void UPCGComponent::PostLoad()
 	// Force dirty to be false on load. We should never refresh on load.
 	bDirtyGenerated = false;
 
-	if (!ExclusionTags_DEPRECATED.IsEmpty() && ExcludedTags.IsEmpty())
-	{
-		ExcludedTags.Append(ExclusionTags_DEPRECATED);
-		ExclusionTags_DEPRECATED.Reset();
-	}
-
 	// If we have both default value (bIsComponentPartitioned = false and bIsPartitioned = true)
 	// we will follow the value of bIsPartitioned.
 	// bIsPartitioned will be set to false to new objects
@@ -1233,21 +1184,6 @@ void UPCGComponent::RefreshAfterGraphChanged(UPCGGraphInterface* InGraph, bool b
 }
 
 #if WITH_EDITOR
-void UPCGComponent::PreEditChange(FProperty* PropertyAboutToChange)
-{
-	if (PropertyAboutToChange)
-	{
-		const FName PropName = PropertyAboutToChange->GetFName();
-
-		if (PropName == GET_MEMBER_NAME_CHECKED(UPCGComponent, ExcludedTags))
-		{
-			TeardownTrackingCallbacks();
-		}
-	}
-
-	Super::PreEditChange(PropertyAboutToChange);
-}
-
 /**
 * Temporary workaround, while we are waiting for UE-182059
 * On the UI side of StructUtils, property widgets are referencing a temporary structure, not the one that is owned by the Graph/GraphInstance.
@@ -1383,20 +1319,6 @@ void UPCGComponent::PostEditChangeProperty(FPropertyChangedEvent& PropertyChange
 		Refresh();
 	}
 	// General properties that don't affect behavior
-	else if (PropName == GET_MEMBER_NAME_CHECKED(UPCGComponent, ExcludedTags))
-	{
-		SetupTrackingCallbacks();
-		RefreshTrackingData();
-
-		const bool bHadExclusionData = !CachedExclusionData.IsEmpty();
-		const bool bHasExcludedActors = !CachedExcludedActors.IsEmpty();
-
-		if(bHadExclusionData || bHasExcludedActors)
-		{
-			DirtyGenerated(EPCGComponentDirtyFlag::Exclusions);
-			Refresh();
-		}
-	}
 	else
 	{
 		Refresh();
@@ -1497,7 +1419,7 @@ void UPCGComponent::SetupTrackingCallbacks()
 		}
 	}
 
-	if(!ExcludedTags.IsEmpty() || !CachedTrackedTagsToSettings.IsEmpty())
+	if(!CachedTrackedTagsToSettings.IsEmpty())
 	{
 		GEngine->OnLevelActorAdded().AddUObject(this, &UPCGComponent::OnActorAdded);
 		GEngine->OnLevelActorDeleted().AddUObject(this, &UPCGComponent::OnActorDeleted);
@@ -1514,7 +1436,6 @@ void UPCGComponent::RefreshTrackingData()
 		return;
 	}
 
-	GetActorsFromTags(ExcludedTags, CachedExcludedActors, /*bCullAgainstLocalBounds=*/true);
 	GetActorsFromTags(CachedTrackedTagsToCulling, CachedTrackedActors);
 	PopulateTrackedActorToTagsMap(/*bForce=*/true);
 }
@@ -1525,53 +1446,6 @@ void UPCGComponent::TeardownTrackingCallbacks()
 	GEngine->OnLevelActorDeleted().RemoveAll(this);
 }
 
-bool UPCGComponent::ActorHasExcludedTag(AActor* InActor) const
-{
-	if (!InActor)
-	{
-		return false;
-	}
-
-	bool bHasExcludedTag = false;
-
-	for (const FName& Tag : InActor->Tags)
-	{
-		if (ExcludedTags.Contains(Tag))
-		{
-			bHasExcludedTag = true;
-			break;
-		}
-	}
-
-	return bHasExcludedTag;
-}
-
-bool UPCGComponent::UpdateExcludedActor(AActor* InActor)
-{
-	TRACE_CPUPROFILER_EVENT_SCOPE(UPCGComponent::UpdateExcludedActor);
-	// Dirty data in all cases - the tag or positional changes will be picked up in the test later
-	if (CachedExcludedActors.Contains(InActor))
-	{
-		if (TObjectPtr<UPCGData>* ExclusionData = CachedExclusionData.Find(InActor))
-		{
-			*ExclusionData = nullptr;
-		}
-
-		CachedPCGData = nullptr;
-		return true;
-	}
-	// Dirty only if the impact actor is inside the bounds
-	else if (ActorHasExcludedTag(InActor) && GetGridBounds().Intersect(GetGridBounds(InActor)))
-	{
-		CachedPCGData = nullptr;
-		return true;
-	}
-	else
-	{
-		return false;
-	}
-}
-
 bool UPCGComponent::ActorIsTracked(AActor* InActor) const
 {
 	if (!InActor || !GetGraph())
@@ -1600,12 +1474,10 @@ void UPCGComponent::OnActorAdded(AActor* InActor)
 		return;
 	}
 
-	const bool bIsExcluded = UpdateExcludedActor(InActor);
 	const bool bIsTracked = AddTrackedActor(InActor);
-
-	if (bIsExcluded || bIsTracked)
+	if (bIsTracked)
 	{
-		DirtyGenerated(bIsExcluded ? EPCGComponentDirtyFlag::Exclusions : EPCGComponentDirtyFlag::None, /*bDispatchToLocalComponents=*/ false);
+		DirtyGenerated(EPCGComponentDirtyFlag::None, /*bDispatchToLocalComponents=*/ false);
 		Refresh();
 	}
 }
@@ -1618,12 +1490,10 @@ void UPCGComponent::OnActorDeleted(AActor* InActor)
 		return;
 	}
 
-	const bool bWasExcluded = UpdateExcludedActor(InActor);
 	const bool bWasTracked = RemoveTrackedActor(InActor);
-
-	if (bWasExcluded || bWasTracked)
+	if (bWasTracked)
 	{
-		DirtyGenerated(bWasExcluded ? EPCGComponentDirtyFlag::Exclusions : EPCGComponentDirtyFlag::None, /*bDispatchToLocalComponents=*/ false);
+		DirtyGenerated(EPCGComponentDirtyFlag::None, /*bDispatchToLocalComponents=*/ false);
 		Refresh();
 	}
 }
@@ -1651,23 +1521,9 @@ void UPCGComponent::OnActorMoved(AActor* InActor)
 	}
 	else
 	{
-		bool bDirtyAndRefresh = false;
-		bool bDirtyExclusions = false;
-
-		if (UpdateExcludedActor(InActor))
-		{
-			bDirtyAndRefresh = true;
-			bDirtyExclusions = true;
-		}
-
 		if (DirtyTrackedActor(InActor))
 		{
-			bDirtyAndRefresh = true;
-		}
-
-		if (bDirtyAndRefresh)
-		{
-			DirtyGenerated(bDirtyExclusions ? EPCGComponentDirtyFlag::Exclusions : EPCGComponentDirtyFlag::None, /*bDispatchToLocalComponents=*/false);
+			DirtyGenerated(EPCGComponentDirtyFlag::None, /*bDispatchToLocalComponents=*/false);
 			Refresh();
 		}
 	}
@@ -1835,19 +1691,7 @@ void UPCGComponent::OnActorChanged(AActor* Actor, UObject* InObject, bool bActor
 	}
 	else if(Actor && !Actor->bIsEditorPreviewActor)
 	{
-		bool bDirtyAndRefresh = false;
-
-		if (UpdateExcludedActor(Actor))
-		{
-			bDirtyAndRefresh = true;
-		}
-
 		if ((bActorTagChange && Actor == InObject && UpdateTrackedActor(Actor)) || DirtyTrackedActor(Actor))
-		{
-			bDirtyAndRefresh = true;
-		}
-
-		if (bDirtyAndRefresh)
 		{
 			DirtyGenerated(EPCGComponentDirtyFlag::None, /*bDispatchToLocalComponents=*/ false);
 			Refresh();
@@ -1896,12 +1740,6 @@ void UPCGComponent::DirtyGenerated(EPCGComponentDirtyFlag DirtyFlag, const bool
 		CachedPCGData = nullptr;
 	}
 
-	if (!!(DirtyFlag & EPCGComponentDirtyFlag::Exclusions))
-	{
-		CachedExclusionData.Reset();
-		CachedPCGData = nullptr;
-	}
-
 	if (!!(DirtyFlag & EPCGComponentDirtyFlag::Data))
 	{
 		CachedPCGData = nullptr;
@@ -2116,72 +1954,6 @@ UPCGData* UPCGComponent::GetOriginalActorPCGData()
 	return nullptr;
 }
 
-TArray<UPCGData*> UPCGComponent::GetPCGExclusionData()
-{
-	// TODO: replace with a boolean, unify.
-	UpdatePCGExclusionData();
-
-	TArray<typename decltype(CachedExclusionData)::ValueType> ExclusionData;
-	CachedExclusionData.GenerateValueArray(ExclusionData);
-	return ToRawPtrTArrayUnsafe(ExclusionData);
-}
-
-void UPCGComponent::UpdatePCGExclusionData()
-{
-	TRACE_CPUPROFILER_EVENT_SCOPE(UPCGComponent::UpdatePCGExclusionData);
-	const UPCGData* InputData = GetInputPCGData();
-	const UPCGSpatialData* InputSpatialData = Cast<const UPCGSpatialData>(InputData);
-
-	// Update the list of cached excluded actors here, since we might not have picked up everything on map load (due to WP)
-	GetActorsFromTags(ExcludedTags, CachedExcludedActors, /*bCullAgainstLocalBounds=*/true);
-
-	// Build exclusion data based on the CachedExcludedActors
-	decltype(CachedExclusionData) ExclusionData;
-
-	for(TWeakObjectPtr<AActor> ExcludedActorWeakPtr : CachedExcludedActors)
-	{
-		if(!ExcludedActorWeakPtr.IsValid())
-		{
-			continue;
-		}
-
-		AActor* ExcludedActor = ExcludedActorWeakPtr.Get();
-
-		TObjectPtr<UPCGData>* PreviousExclusionData = CachedExclusionData.Find(ExcludedActor);
-
-		if (PreviousExclusionData && *PreviousExclusionData)
-		{
-			ExclusionData.Add(ExcludedActor, *PreviousExclusionData);
-		}
-		else
-		{
-			// Create the new exclusion data
-			UPCGData* ActorData = CreateActorPCGData(ExcludedActor);
-			UPCGSpatialData* ActorSpatialData = Cast<UPCGSpatialData>(ActorData);
-
-			if (InputSpatialData && ActorSpatialData)
-			{
-				// Change the target actor to this - otherwise we could push changes on another actor
-				ActorSpatialData->TargetActor = GetOwner();
-
-				// Create intersection or projection depending on the dimension
-				// TODO: there's an ambiguity here when it's the same dimension.
-				// For volumes, we'd expect an intersection, for surfaces we'd expect a projection
-				if (ActorSpatialData->GetDimension() > InputSpatialData->GetDimension())
-				{
-					ExclusionData.Add(ExcludedActor, ActorSpatialData->IntersectWith(InputSpatialData));
-				}
-				else
-				{
-					ExclusionData.Add(ExcludedActor, ActorSpatialData->ProjectOn(InputSpatialData));
-				}
-			}
-		}
-	}
-
-	CachedExclusionData = ExclusionData;
-}
-
 UPCGData* UPCGComponent::CreateActorPCGData()
 {
 	return CreateActorPCGData(GetOwner(), bParseActorComponents);
@@ -2428,34 +2200,7 @@ FPCGDataCollection UPCGComponent::CreateActorPCGDataCollection(AActor* Actor, co
 UPCGData* UPCGComponent::CreatePCGData()
 {
 	TRACE_CPUPROFILER_EVENT_SCOPE(UPCGComponent::CreatePCGData);
-	UPCGData* InputData = GetInputPCGData();
-	UPCGSpatialData* SpatialInput = Cast<UPCGSpatialData>(InputData);
-	
-	// Early out: incompatible data
-	if (!SpatialInput)
-	{
-		return InputData;
-	}
-
-	UPCGDifferenceData* Difference = nullptr;
-	TArray<UPCGData*> ExclusionData = GetPCGExclusionData();
-
-	for (UPCGData* Exclusion : ExclusionData)
-	{
-		if (UPCGSpatialData* SpatialExclusion = Cast<UPCGSpatialData>(Exclusion))
-		{
-			if (!Difference)
-			{
-				Difference = SpatialInput->Subtract(SpatialExclusion);
-			}
-			else
-			{
-				Difference->AddDifference(SpatialExclusion);
-			}
-		}
-	}
-
-	return Difference ? Difference : InputData;
+	return GetInputPCGData();
 }
 
 UPCGData* UPCGComponent::CreateLandscapePCGData(bool bHeightOnly)
@@ -2951,7 +2696,7 @@ void FPCGComponentInstanceData::ApplyToComponent(UActorComponent* Component, con
 			PCGComponent->bDirtyGenerated = SourceComponent->bDirtyGenerated; 
 #endif // WITH_EDITOR
 
-			// Non-critical but should be done: transient data, excluded & tracked actors cache, landscape tracking
+			// Non-critical but should be done: transient data, tracked actors cache, landscape tracking
 			// TODO Validate usefulness + move accordingly
 		}
 		
diff --git a/Engine/Plugins/Experimental/PCG/Source/PCG/Private/PCGInputOutputSettings.cpp b/Engine/Plugins/Experimental/PCG/Source/PCG/Private/PCGInputOutputSettings.cpp
--- a/Engine/Plugins/Experimental/PCG/Source/PCG/Private/PCGInputOutputSettings.cpp
+++ b/Engine/Plugins/Experimental/PCG/Source/PCG/Private/PCGInputOutputSettings.cpp
@@ -20,7 +20,7 @@ UPCGGraphInputOutputSettings::UPCGGraphInputOutputSettings(const FObjectInitiali
 	: Super(ObjectInitializer)
 {
 	StaticInLabels.Emplace(PCGPinConstants::DefaultInputLabel, LOCTEXT("InputOutputInPinTooltip",
-		"Provides the result of the Input pin, but with 'Excluded Tags' actors removed (configured on the PCG Component)."
+		"Provides the same result as the Input pin."
 	));
 	StaticAdvancedInLabels.Emplace(PCGInputOutputConstants::DefaultInputLabel, LOCTEXT("InputOutputInputPinTooltip",
 		"Takes the output of the Actor pin and if the 'Input Type' setting on the PCG Component is set to Landscape, combines it with the result of the Landscape pin. "
@@ -42,9 +42,6 @@ UPCGGraphInputOutputSettings::UPCGGraphInputOutputSettings(const FObjectInitiali
 	StaticAdvancedInLabels.Emplace(PCGInputOutputConstants::DefaultLandscapeHeightLabel, LOCTEXT("InputOutputLandscapeHeightPinTooltip",
 		"Similar to Landscape pin, but only provides height data and not other layers."
 	));
-	StaticAdvancedInLabels.Emplace(PCGInputOutputConstants::DefaultExcludedActorsLabel, LOCTEXT("InputOutputExcludedPinTooltip",
-		"Returns union of data read from actors have any tag specified in the 'Excluded Tags' setting on the PCG Component."
-	));
 	
 	StaticOutLabels.Emplace(PCGPinConstants::DefaultOutputLabel);
 }
diff --git a/Engine/Plugins/Experimental/PCG/Source/PCG/Public/Helpers/PCGBlueprintHelpers.h b/Engine/Plugins/Experimental/PCG/Source/PCG/Public/Helpers/PCGBlueprintHelpers.h
--- a/Engine/Plugins/Experimental/PCG/Source/PCG/Public/Helpers/PCGBlueprintHelpers.h
+++ b/Engine/Plugins/Experimental/PCG/Source/PCG/Public/Helpers/PCGBlueprintHelpers.h
@@ -41,9 +41,6 @@ public:
 	UFUNCTION(BlueprintCallable, Category = "PCG|Temporary", meta = (ScriptMethod))
 	static UPCGData* GetInputData(UPARAM(ref) FPCGContext& Context);
 
-	UFUNCTION(BlueprintCallable, Category = "PCG|Temporary", meta = (ScriptMethod))
-	static TArray<UPCGData*> GetExclusionData(UPARAM(ref) FPCGContext& Context);
-
 	UFUNCTION(BlueprintCallable, Category = "PCG|Temporary", meta = (ScriptMethod))
 	static UPCGComponent* GetComponent(UPARAM(ref) FPCGContext& Context);
 
diff --git a/Engine/Plugins/Experimental/PCG/Source/PCG/Public/PCGComponent.h b/Engine/Plugins/Experimental/PCG/Source/PCG/Public/PCGComponent.h
--- a/Engine/Plugins/Experimental/PCG/Source/PCG/Public/PCGComponent.h
+++ b/Engine/Plugins/Experimental/PCG/Source/PCG/Public/PCGComponent.h
@@ -54,9 +54,8 @@ enum class EPCGComponentDirtyFlag : uint8
 	Actor = 1 << 0,
 	Landscape = 1 << 1,
 	Input = 1 << 2,
-	Exclusions = 1 << 3,
-	Data = 1 << 4,
-	All = Actor | Landscape | Input | Exclusions | Data
+	Data = 1 << 3,
+	All = Actor | Landscape | Input | Data
 };
 ENUM_CLASS_FLAGS(EPCGComponentDirtyFlag);
 
@@ -101,7 +100,6 @@ public:
 	UPCGData* GetLandscapePCGData();
 	UPCGData* GetLandscapeHeightPCGData();
 	UPCGData* GetOriginalActorPCGData();
-	TArray<UPCGData*> GetPCGExclusionData();
 
 	bool CanPartition() const;
 
@@ -166,14 +164,6 @@ public:
 	UPROPERTY(BlueprintReadWrite, EditAnywhere, Category = Settings)
 	int Seed = 42;
 
-	UPROPERTY(BlueprintReadWrite, EditAnywhere, Category = Settings)
-	TSet<FName> ExcludedTags;
-
-#if WITH_EDITORONLY_DATA
-	UPROPERTY()
-	TArray<FName> ExclusionTags_DEPRECATED;
-#endif
-
 	UPROPERTY(BlueprintReadWrite, EditAnywhere, Category = Properties)
 	bool bActivated = true;
 
@@ -293,7 +283,6 @@ private:
 	UPCGData* CreateActorPCGData();
 	UPCGData* CreateActorPCGData(AActor* Actor, bool bParseActor = true);
 	UPCGData* CreateLandscapePCGData(bool bHeightOnly);
-	void UpdatePCGExclusionData();
 	bool IsLandscapeCachedDataDirty(const UPCGData* Data) const;
 
 	bool ShouldGenerate(bool bForce, EPCGComponentGenerationTrigger RequestedGenerationTrigger) const;
@@ -319,14 +308,14 @@ private:
 	void GetManagedResources(TArray<TObjectPtr<UPCGManagedResource>>& Resources) const;
 	void SetManagedResources(const TArray<TObjectPtr<UPCGManagedResource>>& Resources);
 
-	bool GetActorsFromTags(const TSet<FName>& InTags, TSet<TWeakObjectPtr<AActor>>& OutActors, bool bCullAgainstLocalBounds);
 	bool GetActorsFromTags(const TMap<FName, bool>& InTagsAndCulling, TSet<TWeakObjectPtr<AActor>>& OutActors);
 
 	void RefreshAfterGraphChanged(UPCGGraphInterface* InGraph, bool bIsStructural, bool bDirtyInputs);
 	void OnGraphChanged(UPCGGraphInterface* InGraph, EPCGChangeType ChangeType);
 
 #if WITH_EDITOR
-	virtual void PreEditChange(FProperty* PropertyAboutToChange) override;
+	// Stub for the other PreEditChange prototype to prevent compile issues from name hiding
+	virtual void PreEditChange(FProperty* PropertyAboutToChange) override { Super::PreEditChange(PropertyAboutToChange); }
 	virtual void PreEditChange(FEditPropertyChain& PropertyAboutToChange) override;
 	virtual void PostEditChangeChainProperty(FPropertyChangedChainEvent& PropertyChangedEvent) override;
 	virtual void PostEditChangeProperty(FPropertyChangedEvent& PropertyChangedEvent) override;
@@ -348,11 +337,8 @@ private:
 	void OnActorDeleted(AActor* InActor);
 	void OnActorMoved(AActor* InActor);
 	void OnObjectPropertyChanged(UObject* InObject, FPropertyChangedEvent& InEvent);
-	bool ActorHasExcludedTag(AActor* InActor) const;
 	bool ActorIsTracked(AActor* InActor) const;
 
-	bool UpdateExcludedActor(AActor* InActor);
-
 	void OnActorChanged(AActor* InActor, UObject* InSourceObject, bool bActorTagChange);
 
 	bool PopulateTrackedActorToTagsMap(bool bForce = false);
@@ -383,13 +369,6 @@ private:
 	UPROPERTY(Transient, NonPIEDuplicateTransient)
 	TObjectPtr<UPCGData> CachedLandscapeHeightData = nullptr;
 
-	UPROPERTY(Transient, NonPIEDuplicateTransient)
-	TMap<TObjectPtr<AActor>, TObjectPtr<UPCGData>> CachedExclusionData;
-
-	// Cached excluded actors list is serialized because we can't get it at postload time
-	UPROPERTY()
-	TSet<TWeakObjectPtr<AActor>> CachedExcludedActors; 
-
 #if WITH_EDITORONLY_DATA
 	UPROPERTY()
 	TSet<TSoftObjectPtr<AActor>> GeneratedActors_DEPRECATED;
diff --git a/Engine/Plugins/Experimental/PCG/Source/PCG/Public/PCGInputOutputSettings.h b/Engine/Plugins/Experimental/PCG/Source/PCG/Public/PCGInputOutputSettings.h
--- a/Engine/Plugins/Experimental/PCG/Source/PCG/Public/PCGInputOutputSettings.h
+++ b/Engine/Plugins/Experimental/PCG/Source/PCG/Public/PCGInputOutputSettings.h
@@ -13,7 +13,6 @@ namespace PCGInputOutputConstants
 	const FName DefaultOriginalActorLabel = TEXT("OriginalActor");
 	const FName DefaultLandscapeLabel = TEXT("Landscape");
 	const FName DefaultLandscapeHeightLabel = TEXT("Landscape Height");
-	const FName DefaultExcludedActorsLabel = TEXT("ExcludedActors");
 	const FName DefaultNewCustomPinName = TEXT("NewPin");
 }
 
