[PCG] Check for InputNode/OutputNode in PCG Graph destroy method, crash detected when using LiveCode.

#rnx

diff --git a/Engine/Plugins/Experimental/PCG/Source/PCG/Private/PCGGraph.cpp b/Engine/Plugins/Experimental/PCG/Source/PCG/Private/PCGGraph.cpp
--- a/Engine/Plugins/Experimental/PCG/Source/PCG/Private/PCGGraph.cpp
+++ b/Engine/Plugins/Experimental/PCG/Source/PCG/Private/PCGGraph.cpp
@@ -285,8 +285,15 @@ void UPCGGraph::BeginDestroy()
 		OnNodeRemoved(Node);
 	}
 
-	OutputNode->OnNodeChangedDelegate.RemoveAll(this);
-	InputNode->OnNodeChangedDelegate.RemoveAll(this);
+	if (OutputNode)
+	{
+		OutputNode->OnNodeChangedDelegate.RemoveAll(this);
+	}
+
+	if (InputNode)
+	{
+		InputNode->OnNodeChangedDelegate.RemoveAll(this);
+	}
 
 	// Notify the compiler to remove this graph from its cache
 	if (GEditor)
