Verify that when the PCG plugin is enabled the PCG Volume is exposed to Creation menus
Steps:
- From the Level Editor toolbar, select Create
- Search for PCGVolume
- PCGVolume is exposed within the Create menu

diff --git a/Engine/Plugins/Experimental/PCG/Source/PCGEditor/PCGEditor.Build.cs b/Engine/Plugins/Experimental/PCG/Source/PCGEditor/PCGEditor.Build.cs
--- a/Engine/Plugins/Experimental/PCG/Source/PCGEditor/PCGEditor.Build.cs
+++ b/Engine/Plugins/Experimental/PCG/Source/PCGEditor/PCGEditor.Build.cs
@@ -9,12 +9,18 @@ namespace UnrealBuildTool.Rules
 			PCHUsage = PCHUsageMode.UseExplicitOrSharedPCHs;
 
 			PublicDependencyModuleNames.AddRange(
-				new string[] {
-					"Core",
-					"Projects",
-					"Engine",
-					"CoreUObject",
-				});
+			new string[] {
+			"Core",
+			"Projects",
+			"Engine",
+			"CoreUObject"
+        });
+
+		if (Target.Configuration != UnrealTargetConfiguration.Shipping)
+		{
+			PublicDependencyModuleNames.Add("PlacementMode");
+		}
+
 
 			PrivateDependencyModuleNames.AddRange(
 				new string[]
@@ -54,5 +60,6 @@ namespace UnrealBuildTool.Rules
 				new string[] {
 				});
 		}
+
 	}
 }
diff --git a/Engine/Plugins/Experimental/PCG/Source/PCGEditor/Private/Tests/Display/PCGVolumeExposedInCreationMenu.cpp b/Engine/Plugins/Experimental/PCG/Source/PCGEditor/Private/Tests/Display/PCGVolumeExposedInCreationMenu.cpp
new file mode 100644
--- /dev/null
+++ b/Engine/Plugins/Experimental/PCG/Source/PCGEditor/Private/Tests/Display/PCGVolumeExposedInCreationMenu.cpp
@@ -0,0 +1,29 @@
+// Copyright Epic Games, Inc. All Rights Reserved.
+
+#include "Tests/PCGTestsCommon.h"
+#include "IPlacementModeModule.h"
+
+#if WITH_AUTOMATION_TESTS
+
+IMPLEMENT_CUSTOM_SIMPLE_AUTOMATION_TEST(FPCGVolumeExposedInCreationMenu, FPCGTestBaseClass, "Editor.Plugins.Tools.PCG.PCGVolumeExposedInCreationMenu", PCGTestsCommon::TestFlags)
+
+bool FPCGVolumeExposedInCreationMenu::RunTest(const FString& Parameters)
+{
+	// create menu queries place mode categories to dislplay 
+	const IPlacementModeModule& PlacementModeModule = IPlacementModeModule::Get();
+	// PCG volume is located in Volumes category
+	const FName PlacementModeCategoryHandle = TEXT("Volumes");
+	FString PCGDisplayVolumeName = TEXT("PCGVolume");
+	TArray<TSharedPtr<FPlaceableItem>>OutItems;
+
+
+	PlacementModeModule.GetItemsForCategory(PlacementModeCategoryHandle, OutItems);
+	bool VolumeCategoryFound = OutItems.ContainsByPredicate([&PCGDisplayVolumeName](const TSharedPtr<FPlaceableItem>& Item) {
+		return Item->DisplayName.ToString().Contains(PCGDisplayVolumeName);
+		});
+
+	TestTrue(TEXT("PCGVolume is not found in Create Menu"), VolumeCategoryFound);
+	return true;
+}
+
+#endif 
