Verify that when the PCG plugin is enabled the PCG Section is exposed within Content Browser creation menus

Steps:
- Within the Content Browser, select Add  ->
A section called PCG is exposed within creation menus

-Within the Content Browser, right-click within the empty space -> The creation menu opens
Within the create menu, a section called PCG is exposed

diff --git a/Engine/Plugins/Experimental/PCG/Source/PCGEditor/Private/Tests/Display/PCGExposedInContentMenu.cpp b/Engine/Plugins/Experimental/PCG/Source/PCGEditor/Private/Tests/Display/PCGExposedInContentMenu.cpp
new file mode 100644
--- /dev/null
+++ b/Engine/Plugins/Experimental/PCG/Source/PCGEditor/Private/Tests/Display/PCGExposedInContentMenu.cpp
@@ -0,0 +1,29 @@
+// Copyright Epic Games, Inc. All Rights Reserved.
+
+#include "Tests/PCGTestsCommon.h"
+#include "AssetToolsModule.h"
+
+#if WITH_AUTOMATION_TESTS
+
+IMPLEMENT_CUSTOM_SIMPLE_AUTOMATION_TEST(FPCGExposedInContentMenu, FPCGTestBaseClass, "Editor.Plugins.Tools.PCG.ExposedInContentMenu", PCGTestsCommon::TestFlags)
+
+bool FPCGExposedInContentMenu::RunTest(const FString& Parameters)
+{
+	const FName NAME_AssetTools = "AssetTools";
+	const FString PluginName = "PCG";
+
+	FAssetToolsModule& AssetToolsModule = FModuleManager::LoadModuleChecked<FAssetToolsModule>(NAME_AssetTools);
+	TArray<FAdvancedAssetCategory> AdvancedAssetCategories;
+
+	// Content Browser queries AdvancedAssetCategory to display in the menu
+	AssetToolsModule.Get().GetAllAdvancedAssetCategories(/*out*/ AdvancedAssetCategories);
+
+	FAdvancedAssetCategory* PCGAssetCategory = AdvancedAssetCategories.FindByPredicate([&PluginName](const FAdvancedAssetCategory& AdvancedAssetCategory) {
+		return AdvancedAssetCategory.CategoryName.ToString().Contains(PluginName);
+		});
+
+	TestNotNull(TEXT("PCG Category not found"), PCGAssetCategory);
+	return true;
+}
+
+#endif
