Verify that when the PCG plugin is enabled the PCG Volume is exposed to Creation menus
Steps:
- From the Level Editor toolbar, select Create
- Search for PCGVolume
- PCGVolume is exposed within the Create menu

diff --git a/Engine/Plugins/Experimental/PCG/Source/PCGEditor/PCGEditor.Build.cs b/Engine/Plugins/Experimental/PCG/Source/PCGEditor/PCGEditor.Build.cs
--- a/Engine/Plugins/Experimental/PCG/Source/PCGEditor/PCGEditor.Build.cs
+++ b/Engine/Plugins/Experimental/PCG/Source/PCGEditor/PCGEditor.Build.cs
@@ -13,14 +13,9 @@ namespace UnrealBuildTool.Rules
 			"Core",
 			"Projects",
 			"Engine",
-			"CoreUObject"
-        });
-
-		if (Target.Configuration != UnrealTargetConfiguration.Shipping)
-		{
-			PublicDependencyModuleNames.Add("PlacementMode");
-		}
-
+			"CoreUObject",
+			"PlacementMode"
+		});
 
 			PrivateDependencyModuleNames.AddRange(
 				new string[]
