[PCG] Add node title override for 'Get Volume Data' which was missed in the title rework.

Without this node appears as a duped 'Get Actor Data'.

diff --git a/Engine/Plugins/Experimental/PCG/Source/PCG/Public/Elements/PCGTypedGetter.h b/Engine/Plugins/Experimental/PCG/Source/PCG/Public/Elements/PCGTypedGetter.h
--- a/Engine/Plugins/Experimental/PCG/Source/PCG/Public/Elements/PCGTypedGetter.h
+++ b/Engine/Plugins/Experimental/PCG/Source/PCG/Public/Elements/PCGTypedGetter.h
@@ -69,7 +69,8 @@ public:
 
 	//~Begin UPCGSettings interface
 #if WITH_EDITOR
-	virtual FName GetDefaultNodeName() const override { return FName(TEXT("Get Volume Data")); }
+	virtual FName GetDefaultNodeName() const override { return FName(TEXT("GetVolumeData")); }
+	virtual FText GetDefaultNodeTitle() const override { return NSLOCTEXT("PCGGetVolumeSettings", "NodeTitle", "Get Volume Data"); }
 	virtual FText GetNodeTooltipText() const override;
 #endif
 
