Made the GetGridBounds call a bit more robust to invalid inputs since it's user facing

#rnx

diff --git a/Engine/Plugins/Experimental/PCG/Source/PCG/Private/PCGHelpers.cpp b/Engine/Plugins/Experimental/PCG/Source/PCG/Private/PCGHelpers.cpp
--- a/Engine/Plugins/Experimental/PCG/Source/PCG/Private/PCGHelpers.cpp
+++ b/Engine/Plugins/Experimental/PCG/Source/PCG/Private/PCGHelpers.cpp
@@ -60,8 +60,6 @@ namespace PCGHelpers
 
 	FBox GetGridBounds(const AActor* Actor, const UPCGComponent* Component)
 	{
-		check(Actor);
-
 		FBox Bounds(EForceInit::ForceInit);
 
 		if (const APCGPartitionActor* PartitionActor = Cast<const APCGPartitionActor>(Actor))
@@ -83,10 +81,14 @@ namespace PCGHelpers
 		{
 			Bounds = GetLandscapeBounds(LandscapeActor);
 		}
-		else
+		else if (Actor)
 		{
 			Bounds = GetActorBounds(Actor);
 		}
+		else
+		{
+			UE_LOG(LogPCG, Error, TEXT("Actor is invalid in GetGridBounds"));
+		}
 
 		return Bounds;
 	}
