Rename PCG blueprint methods to better align to what they do.

#rnx

diff --git a/Engine/Plugins/Experimental/PCG/Config/BasePCG.ini b/Engine/Plugins/Experimental/PCG/Config/BasePCG.ini
--- a/Engine/Plugins/Experimental/PCG/Config/BasePCG.ini
+++ b/Engine/Plugins/Experimental/PCG/Config/BasePCG.ini
@@ -1,3 +1,17 @@
 [CoreRedirects]
++FunctionRedirects=(OldName="PCGBlueprintElement.MultiPointLoopBody",NewName="PCGBlueprintElement.VariableLoopBody")
++FunctionRedirects=(OldName="PCGBlueprintElement.PointPairLoopBody",NewName="PCGBlueprintElement.NestedLoopBody")
++FunctionRedirects=(OldName="PCGBlueprintElement.LoopOnPoints",NewName="PCGBlueprintElement.PointLoop")
++FunctionRedirects=(OldName="PCGBlueprintElement.MultiLoopOnPoints",NewName="PCGBlueprintElement.VariableLoop")
++FunctionRedirects=(OldName="PCGBlueprintElement.LoopOnPointPairs",NewName="PCGBlueprintElement.NestedLoop")
++FunctionRedirects=(OldName="PCGBlueprintElement.LoopNTimes",NewName="PCGBlueprintElement.IterationLoop")
 +PropertyRedirects=(OldName="PCGSplineSamplerSettings.Params",NewName="PCGSplineSamplerSettings.SamplerParams")
 +PropertyRedirects=(OldName="PCGProjectionSettings.Params",NewName="PCGProjectionSettings.ProjectionParams")
++PropertyRedirects=(OldName="PCGBlueprintElement.IterationLoop.InA",NewName="PCGBlueprintElement.IterationLoop.OptionalA")
++PropertyRedirects=(OldName="PCGBlueprintElement.IterationLoop.InB",NewName="PCGBlueprintElement.IterationLoop.OptionalB")
++PropertyRedirects=(OldName="PCGBlueprintElement.NestedLoop.InA",NewName="PCGBlueprintElement.NestedLoop.InOuterData")
++PropertyRedirects=(OldName="PCGBlueprintElement.NestedLoop.InB",NewName="PCGBlueprintElement.NestedLoop.InInnerData")
++PropertyRedirects=(OldName="PCGBlueprintElement.NestedLoopBody.InA",NewName="PCGBlueprintElement.NestedLoopBody.InOuterData")
++PropertyRedirects=(OldName="PCGBlueprintElement.NestedLoopBody.InB",NewName="PCGBlueprintElement.NestedLoopBody.InInnerData")
++PropertyRedirects=(OldName="PCGBlueprintElement.NestedLoopBody.InPointA",NewName="PCGBlueprintElement.NestedLoopBody.InOuterPoint")
++PropertyRedirects=(OldName="PCGBlueprintElement.NestedLoopBody.InPointB",NewName="PCGBlueprintElement.NestedLoopBody.InInnerPoint")
\ No newline at end of file
diff --git a/Engine/Plugins/Experimental/PCG/Source/PCG/Private/Elements/PCGExecuteBlueprint.cpp b/Engine/Plugins/Experimental/PCG/Source/PCG/Private/Elements/PCGExecuteBlueprint.cpp
--- a/Engine/Plugins/Experimental/PCG/Source/PCG/Private/Elements/PCGExecuteBlueprint.cpp
+++ b/Engine/Plugins/Experimental/PCG/Source/PCG/Private/Elements/PCGExecuteBlueprint.cpp
@@ -653,7 +653,7 @@ bool FPCGExecuteBlueprintElement::ExecuteInternal(FPCGContext* InContext) const
 	return true;
 }
 
-void UPCGBlueprintElement::LoopOnPoints(FPCGContext& InContext, const UPCGPointData* InData, UPCGPointData*& OutData, UPCGPointData* OptionalOutData) const
+void UPCGBlueprintElement::PointLoop(FPCGContext& InContext, const UPCGPointData* InData, UPCGPointData*& OutData, UPCGPointData* OptionalOutData) const
 {
 	if (!InData)
 	{
@@ -680,7 +680,7 @@ void UPCGBlueprintElement::LoopOnPoints(FPCGContext& InContext, const UPCGPointD
 	});
 }
 
-void UPCGBlueprintElement::MultiLoopOnPoints(FPCGContext& InContext, const UPCGPointData* InData, UPCGPointData*& OutData, UPCGPointData* OptionalOutData) const
+void UPCGBlueprintElement::VariableLoop(FPCGContext& InContext, const UPCGPointData* InData, UPCGPointData*& OutData, UPCGPointData* OptionalOutData) const
 {
 	if (!InData)
 	{
@@ -703,13 +703,13 @@ void UPCGBlueprintElement::MultiLoopOnPoints(FPCGContext& InContext, const UPCGP
 
 	FPCGAsync::AsyncMultiPointProcessing(&InContext, InPoints.Num(), OutPoints, [this, &InContext, InData, OutData, &InPoints](int32 Index)
 	{
-		return MultiPointLoopBody(InContext, InData, InPoints[Index], OutData->Metadata);
+		return VariableLoopBody(InContext, InData, InPoints[Index], OutData->Metadata);
 	});
 }
 
-void UPCGBlueprintElement::LoopOnPointPairs(FPCGContext& InContext, const UPCGPointData* InA, const UPCGPointData* InB, UPCGPointData*& OutData, UPCGPointData* OptionalOutData) const
+void UPCGBlueprintElement::NestedLoop(FPCGContext& InContext, const UPCGPointData* InOuterData, const UPCGPointData* InInnerData, UPCGPointData*& OutData, UPCGPointData* OptionalOutData) const
 {
-	if (!InA || !InB)
+	if (!InOuterData || !InInnerData)
 	{
 		PCGE_LOG_C(Error, GraphAndLog, &InContext, LOCTEXT("InvalidInputDataLoopOnPointPairs", "Invalid input data in LoopOnPointPairs"));
 		return;
@@ -722,21 +722,21 @@ void UPCGBlueprintElement::LoopOnPointPairs(FPCGContext& InContext, const UPCGPo
 	else
 	{
 		OutData = NewObject<UPCGPointData>();
-		OutData->InitializeFromData(InA);
-		OutData->Metadata->AddAttributes(InB->Metadata);
+		OutData->InitializeFromData(InOuterData);
+		OutData->Metadata->AddAttributes(InInnerData->Metadata);
 	}
 
-	const TArray<FPCGPoint>& InPointsA = InA->GetPoints();
-	const TArray<FPCGPoint>& InPointsB = InB->GetPoints();
+	const TArray<FPCGPoint>& InOuterPoints = InOuterData->GetPoints();
+	const TArray<FPCGPoint>& InInnerPoints = InInnerData->GetPoints();
 	TArray<FPCGPoint>& OutPoints = OutData->GetMutablePoints();
 
-	FPCGAsync::AsyncPointProcessing(&InContext, InPointsA.Num() * InPointsB.Num(), OutPoints, [this, &InContext, InA, InB, OutData, &InPointsA, &InPointsB](int32 Index, FPCGPoint& OutPoint)
+	FPCGAsync::AsyncPointProcessing(&InContext, InOuterPoints.Num() * InInnerPoints.Num(), OutPoints, [this, &InContext, InOuterData, InInnerData, OutData, &InOuterPoints, &InInnerPoints](int32 Index, FPCGPoint& OutPoint)
 	{
-		return PointPairLoopBody(InContext, InA, InB, InPointsA[Index / InPointsB.Num()], InPointsB[Index % InPointsB.Num()], OutPoint, OutData->Metadata);
+		return NestedLoopBody(InContext, InOuterData, InInnerData, InOuterPoints[Index / InInnerPoints.Num()], InInnerPoints[Index % InInnerPoints.Num()], OutPoint, OutData->Metadata);
 	});
 }
 
-void UPCGBlueprintElement::LoopNTimes(FPCGContext& InContext, int64 NumIterations, UPCGPointData*& OutData, const UPCGSpatialData* InA, const UPCGSpatialData* InB, UPCGPointData* OptionalOutData) const
+void UPCGBlueprintElement::IterationLoop(FPCGContext& InContext, int64 NumIterations, UPCGPointData*& OutData, const UPCGSpatialData* InA, const UPCGSpatialData* InB, UPCGPointData* OptionalOutData) const
 {
 	if (NumIterations < 0)
 	{
diff --git a/Engine/Plugins/Experimental/PCG/Source/PCG/Public/Elements/PCGExecuteBlueprint.h b/Engine/Plugins/Experimental/PCG/Source/PCG/Public/Elements/PCGExecuteBlueprint.h
--- a/Engine/Plugins/Experimental/PCG/Source/PCG/Public/Elements/PCGExecuteBlueprint.h
+++ b/Engine/Plugins/Experimental/PCG/Source/PCG/Public/Elements/PCGExecuteBlueprint.h
@@ -37,51 +37,54 @@ public:
 	virtual void BeginDestroy() override;
 	// ~End UObject interface
 
-	UFUNCTION(BlueprintNativeEvent, Category = Execution)
+	UFUNCTION(BlueprintNativeEvent, Category = "PCG|Execution")
 	void ExecuteWithContext(UPARAM(ref)FPCGContext& InContext, const FPCGDataCollection& Input, FPCGDataCollection& Output);
 
-	UFUNCTION(BlueprintImplementableEvent, BlueprintCallable, Category = Execution)
+	UFUNCTION(BlueprintImplementableEvent, BlueprintCallable, Category = "PCG|Execution")
 	void Execute(const FPCGDataCollection& Input, FPCGDataCollection& Output);
 
-	UFUNCTION(BlueprintImplementableEvent, Category = Execution)
+	UFUNCTION(BlueprintImplementableEvent, Category = "PCG|Flow Control")
 	bool PointLoopBody(const FPCGContext& InContext, const UPCGPointData* InData, const FPCGPoint& InPoint, FPCGPoint& OutPoint, UPCGMetadata* OutMetadata) const;
 
-	UFUNCTION(BlueprintImplementableEvent, Category = Execution)
-	TArray<FPCGPoint> MultiPointLoopBody(const FPCGContext& InContext, const UPCGPointData* InData, const FPCGPoint& InPoint, UPCGMetadata* OutMetadata) const;
+	UFUNCTION(BlueprintImplementableEvent, Category = "PCG|Flow Control")
+	TArray<FPCGPoint> VariableLoopBody(const FPCGContext& InContext, const UPCGPointData* InData, const FPCGPoint& InPoint, UPCGMetadata* OutMetadata) const;
 
-	UFUNCTION(BlueprintImplementableEvent, Category = Execution)
-	bool PointPairLoopBody(const FPCGContext& InContext, const UPCGPointData* InA, const UPCGPointData* InB, const FPCGPoint& InPointA, const FPCGPoint& InPointB, FPCGPoint& OutPoint, UPCGMetadata* OutMetadata) const;
+	UFUNCTION(BlueprintImplementableEvent, Category = "PCG|Flow Control")
+	bool NestedLoopBody(const FPCGContext& InContext, const UPCGPointData* InOuterData, const UPCGPointData* InInnerData, const FPCGPoint& InOuterPoint, const FPCGPoint& InInnerPoint, FPCGPoint& OutPoint, UPCGMetadata* OutMetadata) const;
 
-	UFUNCTION(BlueprintImplementableEvent, Category = Execution)
+	UFUNCTION(BlueprintImplementableEvent, Category = "PCG|Flow Control")
 	bool IterationLoopBody(const FPCGContext& InContext, int64 Iteration, const UPCGSpatialData* InA, const UPCGSpatialData* InB, FPCGPoint& OutPoint, UPCGMetadata* OutMetadata) const;
 
-	/** Calls the LoopBody function on all points */
-	UFUNCTION(BlueprintCallable, BlueprintPure = false, Category = Execution, meta = (HideSelfPin = "true"))
-	void LoopOnPoints(UPARAM(ref) FPCGContext& InContext, const UPCGPointData* InData, UPCGPointData*& OutData, UPCGPointData* OptionalOutData = nullptr) const;
+	/** Calls the PointLoopBody function on all points */
+	UFUNCTION(BlueprintCallable, BlueprintPure = false, Category = "PCG|Flow Control", meta = (HideSelfPin = "true"))
+	void PointLoop(UPARAM(ref) FPCGContext& InContext, const UPCGPointData* InData, UPCGPointData*& OutData, UPCGPointData* OptionalOutData = nullptr) const;
 
-	UFUNCTION(BlueprintCallable, BlueprintPure = false, Category = Execution, meta = (HideSelfPin = "true"))
-	void MultiLoopOnPoints(UPARAM(ref) FPCGContext& InContext, const UPCGPointData* InData, UPCGPointData*& OutData, UPCGPointData* OptionalOutData = nullptr) const;
+	/** Calls the VariableLoopBody function on all points, each call can return a variable number of points */
+	UFUNCTION(BlueprintCallable, BlueprintPure = false, Category = "PCG|Flow Control", meta = (HideSelfPin = "true"))
+	void VariableLoop(UPARAM(ref) FPCGContext& InContext, const UPCGPointData* InData, UPCGPointData*& OutData, UPCGPointData* OptionalOutData = nullptr) const;
 
-	UFUNCTION(BlueprintCallable, BlueprintPure = false, Category = Execution, meta = (HideSelfPin = "true"))
-	void LoopOnPointPairs(UPARAM(ref) FPCGContext& InContext, const UPCGPointData* InA, const UPCGPointData* InB, UPCGPointData*& OutData, UPCGPointData* OptionalOutData = nullptr) const;
-
-	UFUNCTION(BlueprintCallable, BlueprintPure = false, Category = Execution, meta = (HideSelfPin = "true"))
-	void LoopNTimes(UPARAM(ref) FPCGContext& InContext, int64 NumIterations, UPCGPointData*& OutData, const UPCGSpatialData* InA = nullptr, const UPCGSpatialData* InB = nullptr, UPCGPointData* OptionalOutData = nullptr) const;
+	/** Calls the NestedLoopBody function on all nested loop pairs (e.g. (o, i) for all o in Outer, i in Inner) */
+	UFUNCTION(BlueprintCallable, BlueprintPure = false, Category = "PCG|FLow Control", meta = (HideSelfPin = "true"))
+	void NestedLoop(UPARAM(ref) FPCGContext& InContext, const UPCGPointData* InOuterData, const UPCGPointData* InInnerData, UPCGPointData*& OutData, UPCGPointData* OptionalOutData = nullptr) const;
+	
+	/** Calls the IterationLoopBody a fixed number of times, optional parameters are used to potentially initialized the Out Data, but otherwise are used to remove the need to have variables */
+	UFUNCTION(BlueprintCallable, BlueprintPure = false, Category = "PCG|Flow Control", meta = (HideSelfPin = "true"))
+	void IterationLoop(UPARAM(ref) FPCGContext& InContext, int64 NumIterations, UPCGPointData*& OutData, const UPCGSpatialData* OptionalA = nullptr, const UPCGSpatialData* OptionalB = nullptr, UPCGPointData* OptionalOutData = nullptr) const;
 
 	/** Override for the default node name */
-	UFUNCTION(BlueprintNativeEvent, Category = Graph)
+	UFUNCTION(BlueprintNativeEvent, Category = "PCG|Node Customization")
 	FName NodeTitleOverride() const;
 
-	UFUNCTION(BlueprintNativeEvent, Category = Graph)
+	UFUNCTION(BlueprintNativeEvent, Category = "PCG|Node Customization")
 	FLinearColor NodeColorOverride() const;
 
-	UFUNCTION(BlueprintNativeEvent, Category = Graph)
+	UFUNCTION(BlueprintNativeEvent, Category = "PCG|Node Customization")
 	EPCGSettingsType NodeTypeOverride() const;
 
-	UFUNCTION(BlueprintCallable, Category = "Input & Output")
+	UFUNCTION(BlueprintCallable, Category = "PCG|Input & Output")
 	TSet<FName> InputLabels() const;
 
-	UFUNCTION(BlueprintCallable, Category = "Input & Output")
+	UFUNCTION(BlueprintCallable, Category = "PCG|Input & Output")
 	TSet<FName> OutputLabels() const;
 
 	/** Gets the seed from the associated settings & source component */
