[PCG] UX improvements to Graph Parameters

* Renamed UserParameters to Parameters
* Add "ResetToDefault" widget
* Remove the possibility to reset the Parameters Override struct.
* Fix error logged when overriding parameters in subgraphs
* Fix override pin not updated when changing graph parameters on the subgraph
* Rename "Save Instance" to "Export Instance"
* Disable override widget on subgraph, if the override pin is connected.
* Use WeakPtrs in the override details view, to be more resilient.

#rnx

diff --git a/Engine/Plugins/Experimental/PCG/Source/PCG/Private/PCGGraph.cpp b/Engine/Plugins/Experimental/PCG/Source/PCG/Private/PCGGraph.cpp
--- a/Engine/Plugins/Experimental/PCG/Source/PCG/Private/PCGGraph.cpp
+++ b/Engine/Plugins/Experimental/PCG/Source/PCG/Private/PCGGraph.cpp
@@ -1254,6 +1254,28 @@ void UPCGGraphInstance::CopyParameterOverrides(UPCGGraphInterface* InGraph)
 	ParametersOverrides.Parameters.CopyMatchingValuesByID(*InGraph->GetUserParametersStruct());
 }
 
+void UPCGGraphInstance::ResetPropertyToDefault(const FProperty* InProperty)
+{
+	if (!IsPropertyOverridden(InProperty))
+	{
+		return;
+	}
+
+	Modify();
+
+	ParametersOverrides.ResetPropertyToDefault(InProperty, Graph->GetUserParametersStruct());
+
+#if WITH_EDITOR
+	NotifyGraphParametersChanged(EPCGGraphParameterEvent::ValueModifiedLocally, InProperty->GetFName());
+	OnGraphChangedDelegate.Broadcast(this, EPCGChangeType::Settings);
+#endif // WITH_EDITOR
+}
+
+bool UPCGGraphInstance::IsPropertyOverriddenAndNotDefault(const FProperty* InProperty) const
+{
+	return ParametersOverrides.IsPropertyOverriddenAndNotDefault(InProperty, Graph->GetUserParametersStruct());
+}
+
 bool FPCGOverrideInstancedPropertyBag::RefreshParameters(const FInstancedPropertyBag* ParentUserParameters, EPCGGraphParameterEvent InChangeType, FName InChangedPropertyName)
 {
 	check(ParentUserParameters);
@@ -1397,19 +1419,26 @@ bool FPCGOverrideInstancedPropertyBag::UpdatePropertyOverride(const FProperty* I
 	// Reset the value if it is not marked overridden anymore.
 	if (!bMarkAsOverridden)
 	{
-		check(ParentUserParameters);
-
-		const FPropertyBagPropertyDesc* OriginalPropertyDesc = ParentUserParameters->FindPropertyDescByName(InProperty->GetFName());
-		const FPropertyBagPropertyDesc* ThisPropertyDesc = Parameters.FindPropertyDescByName(InProperty->GetFName());
-
-		check(OriginalPropertyDesc && ThisPropertyDesc);
-		PCGGraphUtils::CopyPropertyValue(OriginalPropertyDesc, *ParentUserParameters, ThisPropertyDesc, Parameters);
+		ResetPropertyToDefault(InProperty, ParentUserParameters);
 		return true;
 	}
 
 	return false;
 }
 
+void FPCGOverrideInstancedPropertyBag::ResetPropertyToDefault(const FProperty* InProperty, const FInstancedPropertyBag* ParentUserParameters)
+{
+	check(ParentUserParameters);
+
+	const FPropertyBagPropertyDesc* OriginalPropertyDesc = ParentUserParameters->FindPropertyDescByName(InProperty->GetFName());
+	const FPropertyBagPropertyDesc* ThisPropertyDesc = Parameters.FindPropertyDescByName(InProperty->GetFName());
+
+	if (OriginalPropertyDesc && ThisPropertyDesc)
+	{
+		PCGGraphUtils::CopyPropertyValue(OriginalPropertyDesc, *ParentUserParameters, ThisPropertyDesc, Parameters);
+	}
+}
+
 bool FPCGOverrideInstancedPropertyBag::IsPropertyOverridden(const FProperty* InProperty) const
 {
 	if (!InProperty)
@@ -1421,6 +1450,23 @@ bool FPCGOverrideInstancedPropertyBag::IsPropertyOverridden(const FProperty* InP
 	return PropertyDesc && PropertiesIDsOverridden.Contains(PropertyDesc->ID);
 }
 
+bool FPCGOverrideInstancedPropertyBag::IsPropertyOverriddenAndNotDefault(const FProperty* InProperty, const FInstancedPropertyBag* ParentUserParameters) const
+{
+	check(ParentUserParameters);
+
+	const FPropertyBagPropertyDesc* OriginalPropertyDesc = ParentUserParameters->FindPropertyDescByName(InProperty->GetFName());
+	const FPropertyBagPropertyDesc* ThisPropertyDesc = Parameters.FindPropertyDescByName(InProperty->GetFName());
+
+	if (OriginalPropertyDesc && ThisPropertyDesc && PropertiesIDsOverridden.Contains(ThisPropertyDesc->ID))
+	{
+		return !PCGGraphUtils::ArePropertiesIdentical(OriginalPropertyDesc, *ParentUserParameters, ThisPropertyDesc, Parameters);
+	}
+	else
+	{
+		return false;
+	}
+}
+
 void FPCGOverrideInstancedPropertyBag::Reset()
 {
 	Parameters.Reset();
diff --git a/Engine/Plugins/Experimental/PCG/Source/PCG/Private/PCGSettings.cpp b/Engine/Plugins/Experimental/PCG/Source/PCG/Private/PCGSettings.cpp
--- a/Engine/Plugins/Experimental/PCG/Source/PCG/Private/PCGSettings.cpp
+++ b/Engine/Plugins/Experimental/PCG/Source/PCG/Private/PCGSettings.cpp
@@ -444,6 +444,11 @@ bool UPCGSettings::CanEditChange(const FProperty* InProperty) const
 		return true;
 	}
 
+	return !IsPropertyOverriddenByPin(InProperty);
+}
+
+bool UPCGSettings::IsPropertyOverriddenByPin(const FProperty* InProperty) const
+{
 	if (const UPCGNode* Node = Cast<UPCGNode>(GetOuter()))
 	{
 		const FPCGSettingsOverridableParam* Param = OverridableParams().FindByPredicate([InProperty](const FPCGSettingsOverridableParam& ParamToCheck)
@@ -457,12 +462,12 @@ bool UPCGSettings::CanEditChange(const FProperty* InProperty) const
 		{
 			if (const UPCGPin* Pin = Node->GetInputPin(Param->Label))
 			{
-				return !Pin->IsConnected();
+				return Pin->IsConnected();
 			}
 		}
 	}
 
-	return true;
+	return false;
 }
 
 void UPCGSettings::ApplyDeprecationBeforeUpdatePins(UPCGNode* InOutNode, TArray<TObjectPtr<UPCGPin>>& InputPins, TArray<TObjectPtr<UPCGPin>>& OutputPins)
diff --git a/Engine/Plugins/Experimental/PCG/Source/PCG/Private/PCGSubgraph.cpp b/Engine/Plugins/Experimental/PCG/Source/PCG/Private/PCGSubgraph.cpp
--- a/Engine/Plugins/Experimental/PCG/Source/PCG/Private/PCGSubgraph.cpp
+++ b/Engine/Plugins/Experimental/PCG/Source/PCG/Private/PCGSubgraph.cpp
@@ -8,6 +8,8 @@
 #include "Data/PCGUserParametersData.h"
 #include "Helpers/PCGSettingsHelpers.h"
 
+#include "Algo/Find.h"
+
 #include UE_INLINE_GENERATED_CPP_BY_NAME(PCGSubgraph)
 
 namespace PCGSubgraphSettings
@@ -137,6 +139,9 @@ void UPCGBaseSubgraphSettings::OnSubgraphChanged(UPCGGraphInterface* InGraph, EP
 	if (InGraph == GetSubgraphInterface())
 	{
 		OnSettingsChangedDelegate.Broadcast(this, (ChangeType | EPCGChangeType::Settings));
+
+		// Also rebuild the overrides
+		InitializeCachedOverridableParams(/*bReset=*/true);
 	}
 }
 #endif // WITH_EDITOR
@@ -410,7 +415,30 @@ bool FPCGSubgraphElement::ExecuteInternal(FPCGContext* InContext) const
 	}
 	else
 	{
-		Context->OutputData = Context->InputData;
+		// Don't forward overrides
+		if (Settings->HasOverridableParams())
+		{
+			Context->OutputData.TaggedData.Reserve(Context->InputData.TaggedData.Num());
+			const TArray<FPCGPinProperties> InputPins = Settings->DefaultInputPinProperties();
+
+			for (const FPCGTaggedData& InputData : Context->InputData.TaggedData)
+			{
+				if (!InputData.Data)
+				{
+					continue;
+				}
+
+				// Discard params that don't have a pin on the subgraph input node
+				if (!InputData.Data->IsA<UPCGParamData>() || Algo::FindByPredicate(InputPins, [Pin = InputData.Pin](const FPCGPinProperties& PinProperty) { return PinProperty.Label == Pin;}))
+				{
+					Context->OutputData.TaggedData.Add(InputData);
+				}
+			}
+		}
+		else
+		{
+			Context->OutputData = Context->InputData;
+		}
 
 		// Also create a new data containing information about the original subgraph and the parameters override
 		// It is used mainly by the UseParameterGetElement to access the correct value.
diff --git a/Engine/Plugins/Experimental/PCG/Source/PCG/Public/Elements/PCGUserParameterGet.h b/Engine/Plugins/Experimental/PCG/Source/PCG/Public/Elements/PCGUserParameterGet.h
--- a/Engine/Plugins/Experimental/PCG/Source/PCG/Public/Elements/PCGUserParameterGet.h
+++ b/Engine/Plugins/Experimental/PCG/Source/PCG/Public/Elements/PCGUserParameterGet.h
@@ -38,8 +38,8 @@ protected:
 	virtual FPCGElementPtr CreateElement() const override;
 
 #if WITH_EDITOR
-	virtual FName GetDefaultNodeName() const override { return FName("GetUserParameter"); }
-	virtual FText GetDefaultNodeTitle() const override { return NSLOCTEXT("PCGUserParameterGetSettings", "NodeTitle", "Get User Parameter"); }
+	virtual FName GetDefaultNodeName() const override { return FName("GetGraphParameter"); }
+	virtual FText GetDefaultNodeTitle() const override { return NSLOCTEXT("PCGUserParameterGetSettings", "NodeTitle", "Get Graph Parameter"); }
 #endif
 	//~End UPCGSettings interface
 };
diff --git a/Engine/Plugins/Experimental/PCG/Source/PCG/Public/PCGGraph.h b/Engine/Plugins/Experimental/PCG/Source/PCG/Public/PCGGraph.h
--- a/Engine/Plugins/Experimental/PCG/Source/PCG/Public/PCGGraph.h
+++ b/Engine/Plugins/Experimental/PCG/Source/PCG/Public/PCGGraph.h
@@ -40,9 +40,15 @@ public:
 	/** Add/Remove given property from overrides, and reset its value if it is removed. Returns true if the value was changed. */
 	bool UpdatePropertyOverride(const FProperty* InProperty, bool bMarkAsOverridden, const FInstancedPropertyBag* ParentUserParameters);
 
+	/** Reset overridden property to its parent value. */
+	void ResetPropertyToDefault(const FProperty* InProperty, const FInstancedPropertyBag* ParentUserParameters);
+
 	/** Return if the property is currently marked overridden. */
 	bool IsPropertyOverridden(const FProperty* InProperty) const;
 
+	/** Return if the property is currently marked overridden and has a different value than its default value. */
+	bool IsPropertyOverriddenAndNotDefault(const FProperty* InProperty, const FInstancedPropertyBag* ParentUserParameters) const;
+
 	/** Reset the struct. */
 	void Reset();
 
@@ -231,7 +237,7 @@ protected:
 #endif // WITH_EDITORONLY_DATA
 
 	// Parameters
-	UPROPERTY(EditAnywhere, Category = Instance)
+	UPROPERTY(EditAnywhere, Category = Instance, meta = (DisplayName = "Parameters"))
 	FInstancedPropertyBag UserParameters;
 
 public:
@@ -297,12 +303,14 @@ public:
 	void SetGraph(UPCGGraphInterface* InGraph);
 	void CopyParameterOverrides(UPCGGraphInterface* InGraph);
 	void UpdatePropertyOverride(const FProperty* InProperty, bool bMarkAsOverridden);
-	bool IsPropertyOverridden(const FProperty* InProperty) { return ParametersOverrides.IsPropertyOverridden(InProperty); }
+	void ResetPropertyToDefault(const FProperty* InProperty);
+	bool IsPropertyOverridden(const FProperty* InProperty) const { return ParametersOverrides.IsPropertyOverridden(InProperty); }
+	bool IsPropertyOverriddenAndNotDefault(const FProperty* InProperty) const;
 
 	UPROPERTY(EditAnywhere, BlueprintReadOnly, Category = Instance)
 	TObjectPtr<UPCGGraphInterface> Graph;
 
-	UPROPERTY(EditAnywhere, Category = Instance)
+	UPROPERTY(EditAnywhere, Category = Instance, meta = (NoResetToDefault))
 	FPCGOverrideInstancedPropertyBag ParametersOverrides;
 
 	virtual const FInstancedPropertyBag* GetUserParametersStruct() const override { return &ParametersOverrides.Parameters; }
diff --git a/Engine/Plugins/Experimental/PCG/Source/PCG/Public/PCGSettings.h b/Engine/Plugins/Experimental/PCG/Source/PCG/Public/PCGSettings.h
--- a/Engine/Plugins/Experimental/PCG/Source/PCG/Public/PCGSettings.h
+++ b/Engine/Plugins/Experimental/PCG/Source/PCG/Public/PCGSettings.h
@@ -179,6 +179,7 @@ public:
 	 *  By default, it will return the underlying class, to try to jump to its header in code
      */
 	virtual UObject* GetJumpTargetForDoubleClick() const;
+	virtual bool IsPropertyOverriddenByPin(const FProperty* InProperty) const;
 #endif
 
 	/** Derived classes can implement this to expose additional name information in the logs */
diff --git a/Engine/Plugins/Experimental/PCG/Source/PCGEditor/Private/Details/PCGGraphInstanceDetails.cpp b/Engine/Plugins/Experimental/PCG/Source/PCGEditor/Private/Details/PCGGraphInstanceDetails.cpp
--- a/Engine/Plugins/Experimental/PCG/Source/PCGEditor/Private/Details/PCGGraphInstanceDetails.cpp
+++ b/Engine/Plugins/Experimental/PCG/Source/PCGEditor/Private/Details/PCGGraphInstanceDetails.cpp
@@ -57,7 +57,7 @@ void FPCGGraphInstanceDetails::CustomizeDetails(IDetailLayoutBuilder& DetailBuil
 				[
 					SNew(STextBlock)
 					.Font(IDetailLayoutBuilder::GetDetailFont())
-					.Text(LOCTEXT("SaveInstanceButton", "Save Instance"))
+					.Text(LOCTEXT("SaveInstanceButton", "Export Instance"))
 				]
 			]
 		];
diff --git a/Engine/Plugins/Experimental/PCG/Source/PCGEditor/Private/Details/PCGInstancedPropertyBagOverrideDetails.cpp b/Engine/Plugins/Experimental/PCG/Source/PCGEditor/Private/Details/PCGInstancedPropertyBagOverrideDetails.cpp
--- a/Engine/Plugins/Experimental/PCG/Source/PCGEditor/Private/Details/PCGInstancedPropertyBagOverrideDetails.cpp
+++ b/Engine/Plugins/Experimental/PCG/Source/PCGEditor/Private/Details/PCGInstancedPropertyBagOverrideDetails.cpp
@@ -2,6 +2,7 @@
 
 #include "Details/PCGInstancedPropertyBagOverrideDetails.h"
 #include "PCGGraph.h"
+#include "PCGSettings.h"
 
 #include "DetailLayoutBuilder.h"
 #include "IDetailChildrenBuilder.h"
@@ -10,6 +11,7 @@
 #include "Widgets/SBoxPanel.h"
 #include "Widgets/Layout/SBox.h"
 #include "Widgets/Images/SImage.h"
+#include "Widgets/Input/SButton.h"
 #include "Widgets/Input/SCheckBox.h"
 
 #define LOCTEXT_NAMESPACE "PCGOverrideInstancedPropertyBagDetails"
@@ -86,13 +88,37 @@ FPCGOverrideInstancedPropertyBagDataDetails::FPCGOverrideInstancedPropertyBagDat
 		if (!OuterObjects.IsEmpty())
 		{
 			Owner = Cast<UPCGGraphInstance>(OuterObjects[0]);
+			if (Owner.IsValid())
+			{
+				SettingsOwner = Cast<UPCGSettings>(Owner->GetOuter());
+			}
 		}
 	}
 }
 
+EVisibility FPCGOverrideInstancedPropertyBagDataDetails::IsResetVisible(TSharedPtr<IPropertyHandle> InPropertyHandle) const
+{
+	return (Owner.IsValid() ? Owner->IsPropertyOverriddenAndNotDefault(InPropertyHandle->GetProperty()) : false) ? EVisibility::Visible : EVisibility::Hidden;
+}
+
+FReply FPCGOverrideInstancedPropertyBagDataDetails::OnResetToDefaultValue(TSharedPtr<IPropertyHandle> InPropertyHandle) const
+{
+	if (Owner.IsValid())
+	{
+		Owner->ResetPropertyToDefault(InPropertyHandle->GetProperty());
+	}
+
+	return FReply::Handled();
+}
+
+bool FPCGOverrideInstancedPropertyBagDataDetails::IsPropertyOverriddenByPin(TSharedPtr<IPropertyHandle> InPropertyHandle) const
+{
+	return InPropertyHandle.IsValid() && SettingsOwner.IsValid() && SettingsOwner->IsPropertyOverriddenByPin(InPropertyHandle->GetProperty());
+}
+
 void FPCGOverrideInstancedPropertyBagDataDetails::OnChildRowAdded(IDetailPropertyRow& ChildRow)
 {
-	if (!Owner || !PropertiesIDsOverriddenHandle.IsValid())
+	if (!Owner.IsValid() || !PropertiesIDsOverriddenHandle.IsValid())
 	{
 		return;
 	}
@@ -140,7 +166,7 @@ void FPCGOverrideInstancedPropertyBagDataDetails::OnChildRowAdded(IDetailPropert
 			SNew(SCheckBox)
 			.IsChecked_Lambda([this, ChildPropertyHandle]()
 			{
-				return Owner->IsPropertyOverridden(ChildPropertyHandle->GetProperty()) ? ECheckBoxState::Checked : ECheckBoxState::Unchecked;
+				return (Owner.IsValid() && Owner->IsPropertyOverridden(ChildPropertyHandle->GetProperty())) ? ECheckBoxState::Checked : ECheckBoxState::Unchecked;
 			})
 			.OnCheckStateChanged_Lambda([this, ValueWidget, ChildPropertyHandle](ECheckBoxState NewState)
 			{
@@ -148,10 +174,14 @@ void FPCGOverrideInstancedPropertyBagDataDetails::OnChildRowAdded(IDetailPropert
 
 				FScopedTransaction Transaction(FText::Format(LOCTEXT("OnCheckStateChanged", "Change Override for {0}"), FText::FromName(ChildPropertyHandle->GetProperty()->GetFName())));
 				const bool bIsOverridden = NewState == ECheckBoxState::Checked;
-				Owner->UpdatePropertyOverride(ChildPropertyHandle->GetProperty(), bIsOverridden);
+				if (Owner.IsValid())
+				{
+					Owner->UpdatePropertyOverride(ChildPropertyHandle->GetProperty(), bIsOverridden);
+				}
 
 				PropertiesIDsOverriddenHandle->NotifyPostChange(EPropertyChangeType::ValueSet);
 			})
+			.IsEnabled_Lambda([this, ChildPropertyHandle]() -> bool { return !IsPropertyOverriddenByPin(ChildPropertyHandle); })
 		]
 		// Name
 		+ SHorizontalBox::Slot()
@@ -163,9 +193,26 @@ void FPCGOverrideInstancedPropertyBagDataDetails::OnChildRowAdded(IDetailPropert
 	.ValueContent()
 	[
 		ValueWidget.ToSharedRef()
+	]
+	.ExtensionContent()
+	[
+		SNew(SButton)
+		.IsFocusable(false)
+		.ToolTipText(LOCTEXT("ResetButtonTooltip", "Reset property value to its default value, defined by the parent."))
+		.ButtonStyle(FAppStyle::Get(), "SimpleButton")
+		.ContentPadding(0)
+		.Visibility_Lambda([this, ChildPropertyHandle]() { return IsResetVisible(ChildPropertyHandle); })
+		.OnClicked_Lambda([this, ChildPropertyHandle]() { return OnResetToDefaultValue(ChildPropertyHandle); })
+		.Content()
+		[
+			SNew(SImage)
+			.Image(FAppStyle::GetBrush("PropertyWindow.DiffersFromDefault"))
+			.ColorAndOpacity(FSlateColor::UseForeground())
+		]
 	];
 
-	ValueWidget->SetEnabled(TAttribute<bool>::CreateLambda([this, ChildPropertyHandle]() -> bool { return Owner->IsPropertyOverridden(ChildPropertyHandle->GetProperty()); }));
+	ValueWidget->SetEnabled(TAttribute<bool>::CreateLambda([this, ChildPropertyHandle]() -> bool { return Owner.IsValid() && Owner->IsPropertyOverridden(ChildPropertyHandle->GetProperty()) && !IsPropertyOverriddenByPin(ChildPropertyHandle); }));
+	NameWidget->SetEnabled(TAttribute<bool>::CreateLambda([this, ChildPropertyHandle]() -> bool { return !IsPropertyOverriddenByPin(ChildPropertyHandle); }));
 }
 
 #undef LOCTEXT_NAMESPACE
\ No newline at end of file
diff --git a/Engine/Plugins/Experimental/PCG/Source/PCGEditor/Private/Details/PCGInstancedPropertyBagOverrideDetails.h b/Engine/Plugins/Experimental/PCG/Source/PCGEditor/Private/Details/PCGInstancedPropertyBagOverrideDetails.h
--- a/Engine/Plugins/Experimental/PCG/Source/PCGEditor/Private/Details/PCGInstancedPropertyBagOverrideDetails.h
+++ b/Engine/Plugins/Experimental/PCG/Source/PCGEditor/Private/Details/PCGInstancedPropertyBagOverrideDetails.h
@@ -5,9 +5,13 @@
 #include "IDetailCustomization.h"
 #include "IPropertyTypeCustomization.h"
 #include "InstancedStructDetails.h"
+#include "Input/Reply.h"
+#include "Layout/Visibility.h"
 #include "Templates/SharedPointer.h"
+#include "UObject/WeakObjectPtrTemplates.h"
 
 class UPCGGraphInstance;
+class UPCGSettings;
 
 class FPCGOverrideInstancedPropertyBagDetails : public IPropertyTypeCustomization
 {
@@ -31,6 +35,11 @@ public:
 	virtual void OnChildRowAdded(IDetailPropertyRow& ChildRow) override;
 
 private:
-	UPCGGraphInstance* Owner = nullptr;
+	FReply OnResetToDefaultValue(TSharedPtr<IPropertyHandle> InPropertyHandle) const;
+	EVisibility IsResetVisible(TSharedPtr<IPropertyHandle> InPropertyHandle) const;
+	bool IsPropertyOverriddenByPin(TSharedPtr<IPropertyHandle> InPropertyHandle) const;
+
+	TWeakObjectPtr<UPCGGraphInstance> Owner;
+	TWeakObjectPtr<UPCGSettings> SettingsOwner;
 	TSharedPtr<IPropertyHandle> PropertiesIDsOverriddenHandle;
 };
\ No newline at end of file
diff --git a/Engine/Plugins/Experimental/PCG/Source/PCGEditor/Private/PCGEditorGraphSchema.cpp b/Engine/Plugins/Experimental/PCG/Source/PCGEditor/Private/PCGEditorGraphSchema.cpp
--- a/Engine/Plugins/Experimental/PCG/Source/PCGEditor/Private/PCGEditorGraphSchema.cpp
+++ b/Engine/Plugins/Experimental/PCG/Source/PCGEditor/Private/PCGEditorGraphSchema.cpp
@@ -240,7 +240,7 @@ void UPCGEditorGraphSchema::GetNativeElementActions(FGraphActionMenuBuilder& Act
 			{
 				if (const UPropertyBag* BagStruct = UserParameters->GetPropertyBagStruct())
 				{
-					const FText Category = LOCTEXT("UserParametersCategoryName", "User Parameters");
+					const FText Category = LOCTEXT("UserParametersCategoryName", "Graph Parameters");
 
 					for (const FPropertyBagPropertyDesc& PropertyDesc : BagStruct->GetPropertyDescs())
 					{
