Moved PCGHelpers to a public folder and into the Helpers folder.
Made the methods exposed in the DLL too.

#rnx

diff --git a/Engine/Plugins/Experimental/PCG/Source/PCG/Private/Data/PCGIntersectionData.cpp b/Engine/Plugins/Experimental/PCG/Source/PCG/Private/Data/PCGIntersectionData.cpp
--- a/Engine/Plugins/Experimental/PCG/Source/PCG/Private/Data/PCGIntersectionData.cpp
+++ b/Engine/Plugins/Experimental/PCG/Source/PCG/Private/Data/PCGIntersectionData.cpp
@@ -4,7 +4,7 @@
 #include "Data/PCGPointData.h"
 #include "Data/PCGSpatialData.h"
 #include "Helpers/PCGAsync.h"
-#include "PCGHelpers.h"
+#include "Helpers/PCGHelpers.h"
 
 #include "Serialization/ArchiveCrc32.h"
 
diff --git a/Engine/Plugins/Experimental/PCG/Source/PCG/Private/Data/PCGLandscapeData.cpp b/Engine/Plugins/Experimental/PCG/Source/PCG/Private/Data/PCGLandscapeData.cpp
--- a/Engine/Plugins/Experimental/PCG/Source/PCG/Private/Data/PCGLandscapeData.cpp
+++ b/Engine/Plugins/Experimental/PCG/Source/PCG/Private/Data/PCGLandscapeData.cpp
@@ -2,15 +2,15 @@
 
 #include "Data/PCGLandscapeData.h"
 
-#include "Data/PCGSurfaceData.h"
-#include "PCGHelpers.h"
-#include "Engine/World.h"
 #include "PCGSubsystem.h"
 #include "Data/PCGPointData.h"
+#include "Data/PCGSurfaceData.h"
 #include "Grid/PCGLandscapeCache.h"
+#include "Helpers/PCGHelpers.h"
 
 #include "LandscapeInfo.h"
 #include "LandscapeProxy.h"
+#include "Engine/World.h"
 
 #include UE_INLINE_GENERATED_CPP_BY_NAME(PCGLandscapeData)
 
diff --git a/Engine/Plugins/Experimental/PCG/Source/PCG/Private/Data/PCGLandscapeSplineData.cpp b/Engine/Plugins/Experimental/PCG/Source/PCG/Private/Data/PCGLandscapeSplineData.cpp
--- a/Engine/Plugins/Experimental/PCG/Source/PCG/Private/Data/PCGLandscapeSplineData.cpp
+++ b/Engine/Plugins/Experimental/PCG/Source/PCG/Private/Data/PCGLandscapeSplineData.cpp
@@ -3,9 +3,9 @@
 #include "Data/PCGLandscapeSplineData.h"
 
 #include "Data/PCGPolyLineData.h"
-#include "PCGHelpers.h"
 #include "Data/PCGPointData.h"
 #include "Elements/PCGSplineSampler.h"
+#include "Helpers/PCGHelpers.h"
 
 #include "LandscapeSplineSegment.h"
 
diff --git a/Engine/Plugins/Experimental/PCG/Source/PCG/Private/Data/PCGPointData.cpp b/Engine/Plugins/Experimental/PCG/Source/PCG/Private/Data/PCGPointData.cpp
--- a/Engine/Plugins/Experimental/PCG/Source/PCG/Private/Data/PCGPointData.cpp
+++ b/Engine/Plugins/Experimental/PCG/Source/PCG/Private/Data/PCGPointData.cpp
@@ -2,7 +2,7 @@
 
 #include "Data/PCGPointData.h"
 
-#include "PCGHelpers.h"
+#include "Helpers/PCGHelpers.h"
 #include "Metadata/PCGMetadataAccessor.h"
 #include "Metadata/Accessors/PCGAttributeAccessorHelpers.h"
 #include "Metadata/Accessors/PCGAttributeAccessorKeys.h"
diff --git a/Engine/Plugins/Experimental/PCG/Source/PCG/Private/Data/PCGSplineData.cpp b/Engine/Plugins/Experimental/PCG/Source/PCG/Private/Data/PCGSplineData.cpp
--- a/Engine/Plugins/Experimental/PCG/Source/PCG/Private/Data/PCGSplineData.cpp
+++ b/Engine/Plugins/Experimental/PCG/Source/PCG/Private/Data/PCGSplineData.cpp
@@ -2,12 +2,12 @@
 
 #include "Data/PCGSplineData.h"
 
-#include "PCGHelpers.h"
 #include "Data/PCGPointData.h"
 #include "Data/PCGPolyLineData.h"
 #include "Data/PCGProjectionData.h"
 #include "Data/PCGSpatialData.h"
 #include "Elements/PCGSplineSampler.h"
+#include "Helpers/PCGHelpers.h"
 
 #include "Components/SplineComponent.h"
 
diff --git a/Engine/Plugins/Experimental/PCG/Source/PCG/Private/Data/PCGTextureData.cpp b/Engine/Plugins/Experimental/PCG/Source/PCG/Private/Data/PCGTextureData.cpp
--- a/Engine/Plugins/Experimental/PCG/Source/PCG/Private/Data/PCGTextureData.cpp
+++ b/Engine/Plugins/Experimental/PCG/Source/PCG/Private/Data/PCGTextureData.cpp
@@ -2,11 +2,12 @@
 
 #include "Data/PCGTextureData.h"
 
-#include "Engine/Texture2D.h"
-#include "PCGHelpers.h"
 #include "Data/PCGPointData.h"
 #include "Helpers/PCGAsync.h"
+#include "Helpers/PCGHelpers.h"
+
 #include "TextureResource.h"
+#include "Engine/Texture2D.h"
 
 #include UE_INLINE_GENERATED_CPP_BY_NAME(PCGTextureData)
 
diff --git a/Engine/Plugins/Experimental/PCG/Source/PCG/Private/Data/PCGUnionData.cpp b/Engine/Plugins/Experimental/PCG/Source/PCG/Private/Data/PCGUnionData.cpp
--- a/Engine/Plugins/Experimental/PCG/Source/PCG/Private/Data/PCGUnionData.cpp
+++ b/Engine/Plugins/Experimental/PCG/Source/PCG/Private/Data/PCGUnionData.cpp
@@ -4,8 +4,8 @@
 #include "Data/PCGPointData.h"
 #include "Data/PCGSpatialData.h"
 #include "Helpers/PCGAsync.h"
+#include "Helpers/PCGHelpers.h"
 #include "Metadata/PCGMetadataAccessor.h"
-#include "PCGHelpers.h"
 
 #include "Serialization/ArchiveCrc32.h"
 
diff --git a/Engine/Plugins/Experimental/PCG/Source/PCG/Private/Data/PCGVolumeData.cpp b/Engine/Plugins/Experimental/PCG/Source/PCG/Private/Data/PCGVolumeData.cpp
--- a/Engine/Plugins/Experimental/PCG/Source/PCG/Private/Data/PCGVolumeData.cpp
+++ b/Engine/Plugins/Experimental/PCG/Source/PCG/Private/Data/PCGVolumeData.cpp
@@ -3,8 +3,8 @@
 #include "Data/PCGVolumeData.h"
 #include "Data/PCGPointData.h"
 #include "Data/PCGSpatialData.h"
-#include "PCGHelpers.h"
 #include "Elements/PCGVolumeSampler.h"
+#include "Helpers/PCGHelpers.h"
 
 #include "GameFramework/Volume.h"
 
diff --git a/Engine/Plugins/Experimental/PCG/Source/PCG/Private/Data/PCGWorldData.cpp b/Engine/Plugins/Experimental/PCG/Source/PCG/Private/Data/PCGWorldData.cpp
--- a/Engine/Plugins/Experimental/PCG/Source/PCG/Private/Data/PCGWorldData.cpp
+++ b/Engine/Plugins/Experimental/PCG/Source/PCG/Private/Data/PCGWorldData.cpp
@@ -3,7 +3,6 @@
 #include "Data/PCGWorldData.h"
 
 #include "PCGComponent.h"
-#include "PCGHelpers.h"
 #include "PCGSubsystem.h"
 #include "Data/PCGPointData.h"
 #include "Data/PCGSurfaceData.h"
@@ -12,6 +11,7 @@
 #include "Elements/PCGVolumeSampler.h"
 #include "Grid/PCGLandscapeCache.h"
 #include "Helpers/PCGBlueprintHelpers.h"
+#include "Helpers/PCGHelpers.h"
 
 #include "LandscapeProxy.h"
 #include "Components/BrushComponent.h"
diff --git a/Engine/Plugins/Experimental/PCG/Source/PCG/Private/Elements/PCGAttributeFilterElement.cpp b/Engine/Plugins/Experimental/PCG/Source/PCG/Private/Elements/PCGAttributeFilterElement.cpp
--- a/Engine/Plugins/Experimental/PCG/Source/PCG/Private/Elements/PCGAttributeFilterElement.cpp
+++ b/Engine/Plugins/Experimental/PCG/Source/PCG/Private/Elements/PCGAttributeFilterElement.cpp
@@ -2,11 +2,11 @@
 
 #include "Elements/PCGAttributeFilterElement.h"
 
-#include "Data/PCGSpatialData.h"
 #include "PCGContext.h"
-#include "PCGHelpers.h"
 #include "PCGParamData.h"
 #include "PCGPin.h"
+#include "Data/PCGSpatialData.h"
+#include "Helpers/PCGHelpers.h"
 
 #include UE_INLINE_GENERATED_CPP_BY_NAME(PCGAttributeFilterElement)
 
diff --git a/Engine/Plugins/Experimental/PCG/Source/PCG/Private/Elements/PCGCopyPoints.cpp b/Engine/Plugins/Experimental/PCG/Source/PCG/Private/Elements/PCGCopyPoints.cpp
--- a/Engine/Plugins/Experimental/PCG/Source/PCG/Private/Elements/PCGCopyPoints.cpp
+++ b/Engine/Plugins/Experimental/PCG/Source/PCG/Private/Elements/PCGCopyPoints.cpp
@@ -3,10 +3,10 @@
 #include "Elements/PCGCopyPoints.h"
 
 #include "PCGContext.h"
-#include "PCGHelpers.h"
 #include "Data/PCGPointData.h"
 #include "Data/PCGSpatialData.h"
 #include "Helpers/PCGAsync.h"
+#include "Helpers/PCGHelpers.h"
 
 #include "Async/ParallelFor.h"
 
diff --git a/Engine/Plugins/Experimental/PCG/Source/PCG/Private/Elements/PCGCreateSpline.cpp b/Engine/Plugins/Experimental/PCG/Source/PCG/Private/Elements/PCGCreateSpline.cpp
--- a/Engine/Plugins/Experimental/PCG/Source/PCG/Private/Elements/PCGCreateSpline.cpp
+++ b/Engine/Plugins/Experimental/PCG/Source/PCG/Private/Elements/PCGCreateSpline.cpp
@@ -5,11 +5,11 @@
 #include "PCGCommon.h"
 #include "PCGComponent.h"
 #include "PCGContext.h"
-#include "PCGHelpers.h"
 #include "PCGManagedResource.h"
 #include "PCGPin.h"
 #include "Data/PCGPointData.h"
 #include "Data/PCGSplineData.h"
+#include "Helpers/PCGHelpers.h"
 #include "Metadata/PCGMetadata.h"
 #include "Metadata/PCGMetadataAttributeTpl.h"
 
diff --git a/Engine/Plugins/Experimental/PCG/Source/PCG/Private/Elements/PCGDataFromActor.cpp b/Engine/Plugins/Experimental/PCG/Source/PCG/Private/Elements/PCGDataFromActor.cpp
--- a/Engine/Plugins/Experimental/PCG/Source/PCG/Private/Elements/PCGDataFromActor.cpp
+++ b/Engine/Plugins/Experimental/PCG/Source/PCG/Private/Elements/PCGDataFromActor.cpp
@@ -3,9 +3,9 @@
 #include "Elements/PCGDataFromActor.h"
 
 #include "PCGComponent.h"
-#include "PCGHelpers.h"
 #include "PCGSubsystem.h"
 #include "Data/PCGSpatialData.h"
+#include "Helpers/PCGHelpers.h"
 
 #include "GameFramework/Actor.h"
 #include "UObject/Package.h"
diff --git a/Engine/Plugins/Experimental/PCG/Source/PCG/Private/Elements/PCGDebugElement.cpp b/Engine/Plugins/Experimental/PCG/Source/PCG/Private/Elements/PCGDebugElement.cpp
--- a/Engine/Plugins/Experimental/PCG/Source/PCG/Private/Elements/PCGDebugElement.cpp
+++ b/Engine/Plugins/Experimental/PCG/Source/PCG/Private/Elements/PCGDebugElement.cpp
@@ -4,11 +4,11 @@
 
 #include "PCGComponent.h"
 #include "PCGContext.h"
-#include "PCGHelpers.h"
 #include "PCGPin.h"
 #include "Data/PCGSpatialData.h"
 #include "Data/PCGPointData.h"
 #include "Helpers/PCGActorHelpers.h"
+#include "Helpers/PCGHelpers.h"
 
 #include "Components/InstancedStaticMeshComponent.h"
 #include "Engine/CollisionProfile.h"
diff --git a/Engine/Plugins/Experimental/PCG/Source/PCG/Private/Elements/PCGDensityNoise.cpp b/Engine/Plugins/Experimental/PCG/Source/PCG/Private/Elements/PCGDensityNoise.cpp
--- a/Engine/Plugins/Experimental/PCG/Source/PCG/Private/Elements/PCGDensityNoise.cpp
+++ b/Engine/Plugins/Experimental/PCG/Source/PCG/Private/Elements/PCGDensityNoise.cpp
@@ -4,7 +4,7 @@
 
 #include "PCGContext.h"
 #include "PCGCustomVersion.h"
-#include "PCGHelpers.h"
+#include "Helpers/PCGHelpers.h"
 
 #include "Math/RandomStream.h"
 #include "PCGPoint.h"
diff --git a/Engine/Plugins/Experimental/PCG/Source/PCG/Private/Elements/PCGDensityRemap.cpp b/Engine/Plugins/Experimental/PCG/Source/PCG/Private/Elements/PCGDensityRemap.cpp
--- a/Engine/Plugins/Experimental/PCG/Source/PCG/Private/Elements/PCGDensityRemap.cpp
+++ b/Engine/Plugins/Experimental/PCG/Source/PCG/Private/Elements/PCGDensityRemap.cpp
@@ -4,10 +4,10 @@
 
 #include "PCGContext.h"
 #include "PCGCustomVersion.h"
-#include "PCGHelpers.h"
 #include "Data/PCGPointData.h"
 #include "Data/PCGSpatialData.h"
 #include "Helpers/PCGAsync.h"
+#include "Helpers/PCGHelpers.h"
 
 #include "Math/RandomStream.h"
 
diff --git a/Engine/Plugins/Experimental/PCG/Source/PCG/Private/Elements/PCGExecuteBlueprint.cpp b/Engine/Plugins/Experimental/PCG/Source/PCG/Private/Elements/PCGExecuteBlueprint.cpp
--- a/Engine/Plugins/Experimental/PCG/Source/PCG/Private/Elements/PCGExecuteBlueprint.cpp
+++ b/Engine/Plugins/Experimental/PCG/Source/PCG/Private/Elements/PCGExecuteBlueprint.cpp
@@ -4,9 +4,9 @@
 #include "Engine/Blueprint.h"
 #include "Math/RandomStream.h"
 #include "PCGComponent.h"
-#include "PCGHelpers.h"
 #include "Data/PCGPointData.h"
 #include "Helpers/PCGAsync.h"
+#include "Helpers/PCGHelpers.h"
 #include "Helpers/PCGSettingsHelpers.h"
 
 #include "Engine/World.h"
diff --git a/Engine/Plugins/Experimental/PCG/Source/PCG/Private/Elements/PCGFilterByTag.cpp b/Engine/Plugins/Experimental/PCG/Source/PCG/Private/Elements/PCGFilterByTag.cpp
--- a/Engine/Plugins/Experimental/PCG/Source/PCG/Private/Elements/PCGFilterByTag.cpp
+++ b/Engine/Plugins/Experimental/PCG/Source/PCG/Private/Elements/PCGFilterByTag.cpp
@@ -2,8 +2,8 @@
 
 #include "Elements/PCGFilterByTag.h"
 #include "PCGContext.h"
-#include "PCGHelpers.h"
 #include "PCGPin.h"
+#include "Helpers/PCGHelpers.h"
 
 #include UE_INLINE_GENERATED_CPP_BY_NAME(PCGFilterByTag)
 
diff --git a/Engine/Plugins/Experimental/PCG/Source/PCG/Private/Elements/PCGPointMatchAndSet.cpp b/Engine/Plugins/Experimental/PCG/Source/PCG/Private/Elements/PCGPointMatchAndSet.cpp
--- a/Engine/Plugins/Experimental/PCG/Source/PCG/Private/Elements/PCGPointMatchAndSet.cpp
+++ b/Engine/Plugins/Experimental/PCG/Source/PCG/Private/Elements/PCGPointMatchAndSet.cpp
@@ -3,9 +3,9 @@
 #include "Elements/PCGPointMatchAndSet.h"
 
 #include "PCGContext.h"
-#include "PCGHelpers.h"
 #include "PCGPin.h"
 #include "Data/PCGPointData.h"
+#include "Helpers/PCGHelpers.h"
 #include "MatchAndSet/PCGMatchAndSetWeighted.h"
 #include "Metadata/Accessors/IPCGAttributeAccessor.h"
 #include "Metadata/Accessors/PCGAttributeAccessorHelpers.h"
diff --git a/Engine/Plugins/Experimental/PCG/Source/PCG/Private/Elements/PCGPointSampler.cpp b/Engine/Plugins/Experimental/PCG/Source/PCG/Private/Elements/PCGPointSampler.cpp
--- a/Engine/Plugins/Experimental/PCG/Source/PCG/Private/Elements/PCGPointSampler.cpp
+++ b/Engine/Plugins/Experimental/PCG/Source/PCG/Private/Elements/PCGPointSampler.cpp
@@ -1,12 +1,14 @@
 // Copyright Epic Games, Inc. All Rights Reserved.
 
 #include "Elements/PCGPointSampler.h"
+
+#include "PCGContext.h"
 #include "Data/PCGSpatialData.h"
-#include "PCGHelpers.h"
 #include "Data/PCGPointData.h"
 #include "Helpers/PCGAsync.h"
+#include "Helpers/PCGHelpers.h"
+
 #include "Math/RandomStream.h"
-#include "PCGContext.h"
 
 #include UE_INLINE_GENERATED_CPP_BY_NAME(PCGPointSampler)
 
diff --git a/Engine/Plugins/Experimental/PCG/Source/PCG/Private/Elements/PCGSpawnActor.cpp b/Engine/Plugins/Experimental/PCG/Source/PCG/Private/Elements/PCGSpawnActor.cpp
--- a/Engine/Plugins/Experimental/PCG/Source/PCG/Private/Elements/PCGSpawnActor.cpp
+++ b/Engine/Plugins/Experimental/PCG/Source/PCG/Private/Elements/PCGSpawnActor.cpp
@@ -4,13 +4,13 @@
 
 #include "PCGComponent.h"
 #include "PCGGraph.h"
-#include "PCGHelpers.h"
 #include "PCGManagedResource.h"
 #include "PCGSubsystem.h"
 #include "Data/PCGPointData.h"
 #include "Data/PCGSpatialData.h"
 #include "Grid/PCGPartitionActor.h"
 #include "Helpers/PCGActorHelpers.h"
+#include "Helpers/PCGHelpers.h"
 #include "Metadata/Accessors/IPCGAttributeAccessor.h"
 #include "Metadata/Accessors/PCGAttributeAccessorKeys.h"
 #include "Metadata/Accessors/PCGAttributeAccessorHelpers.h"
diff --git a/Engine/Plugins/Experimental/PCG/Source/PCG/Private/Elements/PCGSplineSampler.cpp b/Engine/Plugins/Experimental/PCG/Source/PCG/Private/Elements/PCGSplineSampler.cpp
--- a/Engine/Plugins/Experimental/PCG/Source/PCG/Private/Elements/PCGSplineSampler.cpp
+++ b/Engine/Plugins/Experimental/PCG/Source/PCG/Private/Elements/PCGSplineSampler.cpp
@@ -4,7 +4,6 @@
 
 #include "PCGContext.h"
 #include "PCGCustomVersion.h"
-#include "PCGHelpers.h"
 #include "Data/PCGIntersectionData.h"
 #include "Data/PCGLandscapeSplineData.h"
 #include "Data/PCGPointData.h"
@@ -12,6 +11,7 @@
 #include "Data/PCGSplineData.h"
 #include "Data/PCGSpatialData.h"
 #include "Helpers/PCGBlueprintHelpers.h"
+#include "Helpers/PCGHelpers.h"
 
 #include "Async/ParallelFor.h"
 #include "Components/SplineComponent.h"
diff --git a/Engine/Plugins/Experimental/PCG/Source/PCG/Private/Elements/PCGSurfaceSampler.cpp b/Engine/Plugins/Experimental/PCG/Source/PCG/Private/Elements/PCGSurfaceSampler.cpp
--- a/Engine/Plugins/Experimental/PCG/Source/PCG/Private/Elements/PCGSurfaceSampler.cpp
+++ b/Engine/Plugins/Experimental/PCG/Source/PCG/Private/Elements/PCGSurfaceSampler.cpp
@@ -6,11 +6,11 @@
 #include "PCGCustomVersion.h"
 #include "PCGEdge.h"
 #include "PCGGraph.h"
-#include "PCGHelpers.h"
 #include "PCGPin.h"
 #include "Data/PCGPointData.h"
 #include "Data/PCGSpatialData.h"
 #include "Helpers/PCGAsync.h"
+#include "Helpers/PCGHelpers.h"
 #include "Helpers/PCGSettingsHelpers.h"
 
 #include "Math/RandomStream.h"
diff --git a/Engine/Plugins/Experimental/PCG/Source/PCG/Private/Elements/PCGTextureSampler.cpp b/Engine/Plugins/Experimental/PCG/Source/PCG/Private/Elements/PCGTextureSampler.cpp
--- a/Engine/Plugins/Experimental/PCG/Source/PCG/Private/Elements/PCGTextureSampler.cpp
+++ b/Engine/Plugins/Experimental/PCG/Source/PCG/Private/Elements/PCGTextureSampler.cpp
@@ -5,8 +5,8 @@
 #include "PCGComponent.h"
 #include "PCGContext.h"
 #include "PCGCrc.h"
-#include "PCGHelpers.h"
 #include "Helpers/PCGBlueprintHelpers.h"
+#include "Helpers/PCGHelpers.h"
 #include "Helpers/PCGSettingsHelpers.h"
 
 #include "Engine/Texture2D.h"
diff --git a/Engine/Plugins/Experimental/PCG/Source/PCG/Private/Elements/PCGTransformPoints.cpp b/Engine/Plugins/Experimental/PCG/Source/PCG/Private/Elements/PCGTransformPoints.cpp
--- a/Engine/Plugins/Experimental/PCG/Source/PCG/Private/Elements/PCGTransformPoints.cpp
+++ b/Engine/Plugins/Experimental/PCG/Source/PCG/Private/Elements/PCGTransformPoints.cpp
@@ -2,13 +2,14 @@
 
 #include "Elements/PCGTransformPoints.h"
 
-#include "Data/PCGSpatialData.h"
-#include "Math/RandomStream.h"
-#include "Helpers/PCGAsync.h"
-#include "PCGHelpers.h"
-#include "Data/PCGPointData.h"
 #include "PCGContext.h"
 #include "PCGPin.h"
+#include "Data/PCGSpatialData.h"
+#include "Data/PCGPointData.h"
+#include "Helpers/PCGAsync.h"
+#include "Helpers/PCGHelpers.h"
+
+#include "Math/RandomStream.h"
 
 #include UE_INLINE_GENERATED_CPP_BY_NAME(PCGTransformPoints)
 
diff --git a/Engine/Plugins/Experimental/PCG/Source/PCG/Private/Elements/PCGVolumeSampler.cpp b/Engine/Plugins/Experimental/PCG/Source/PCG/Private/Elements/PCGVolumeSampler.cpp
--- a/Engine/Plugins/Experimental/PCG/Source/PCG/Private/Elements/PCGVolumeSampler.cpp
+++ b/Engine/Plugins/Experimental/PCG/Source/PCG/Private/Elements/PCGVolumeSampler.cpp
@@ -5,10 +5,10 @@
 #include "PCGComponent.h"
 #include "PCGContext.h"
 #include "PCGCustomVersion.h"
-#include "PCGHelpers.h"
 #include "Data/PCGPointData.h"
 #include "Data/PCGSpatialData.h"
 #include "Helpers/PCGAsync.h"
+#include "Helpers/PCGHelpers.h"
 
 #include UE_INLINE_GENERATED_CPP_BY_NAME(PCGVolumeSampler)
 
diff --git a/Engine/Plugins/Experimental/PCG/Source/PCG/Private/Elements/PCGWorldQuery.cpp b/Engine/Plugins/Experimental/PCG/Source/PCG/Private/Elements/PCGWorldQuery.cpp
--- a/Engine/Plugins/Experimental/PCG/Source/PCG/Private/Elements/PCGWorldQuery.cpp
+++ b/Engine/Plugins/Experimental/PCG/Source/PCG/Private/Elements/PCGWorldQuery.cpp
@@ -4,7 +4,7 @@
 
 #include "PCGComponent.h"
 #include "PCGContext.h"
-#include "PCGHelpers.h"
+#include "Helpers/PCGHelpers.h"
 
 #include "GameFramework/Actor.h"
 
diff --git a/Engine/Plugins/Experimental/PCG/Source/PCG/Private/Graph/PCGGraphExecutor.cpp b/Engine/Plugins/Experimental/PCG/Source/PCG/Private/Graph/PCGGraphExecutor.cpp
--- a/Engine/Plugins/Experimental/PCG/Source/PCG/Private/Graph/PCGGraphExecutor.cpp
+++ b/Engine/Plugins/Experimental/PCG/Source/PCG/Private/Graph/PCGGraphExecutor.cpp
@@ -7,13 +7,13 @@
 #include "PCGCrc.h"
 #include "PCGGraph.h"
 #include "PCGGraphCompiler.h"
-#include "PCGHelpers.h"
 #include "PCGInputOutputSettings.h"
 #include "PCGModule.h"
 #include "PCGParamData.h"
 #include "PCGPin.h"
 #include "PCGSubsystem.h"
 #include "Graph/PCGGraphCache.h"
+#include "Helpers/PCGHelpers.h"
 #include "Metadata/PCGMetadata.h"
 
 #include "Algo/AnyOf.h"
diff --git a/Engine/Plugins/Experimental/PCG/Source/PCG/Private/Grid/PCGPartitionActor.cpp b/Engine/Plugins/Experimental/PCG/Source/PCG/Private/Grid/PCGPartitionActor.cpp
--- a/Engine/Plugins/Experimental/PCG/Source/PCG/Private/Grid/PCGPartitionActor.cpp
+++ b/Engine/Plugins/Experimental/PCG/Source/PCG/Private/Grid/PCGPartitionActor.cpp
@@ -3,11 +3,11 @@
 #include "PCGPartitionActor.h"
 #include "Engine/World.h"
 #include "PCGComponent.h"
-#include "PCGHelpers.h"
 #include "PCGModule.h"
 #include "PCGSubsystem.h"
 #include "PCGWorldActor.h"
 #include "Helpers/PCGActorHelpers.h"
+#include "Helpers/PCGHelpers.h"
 
 #include "Components/BoxComponent.h"
 
diff --git a/Engine/Plugins/Experimental/PCG/Source/PCG/Private/Helpers/PCGActorHelpers.cpp b/Engine/Plugins/Experimental/PCG/Source/PCG/Private/Helpers/PCGActorHelpers.cpp
--- a/Engine/Plugins/Experimental/PCG/Source/PCG/Private/Helpers/PCGActorHelpers.cpp
+++ b/Engine/Plugins/Experimental/PCG/Source/PCG/Private/Helpers/PCGActorHelpers.cpp
@@ -2,16 +2,16 @@
 
 #include "Helpers/PCGActorHelpers.h"
 
-#include "Materials/MaterialInterface.h"
-#include "PCGHelpers.h"
-
 #include "PCGComponent.h"
 #include "PCGManagedResource.h"
-#include "Engine/Blueprint.h"
+#include "PCGModule.h"
+#include "Helpers/PCGHelpers.h"
+
 #include "EngineUtils.h"
 #include "Components/HierarchicalInstancedStaticMeshComponent.h"
+#include "Engine/Blueprint.h"
 #include "Engine/StaticMesh.h"
-#include "PCGModule.h"
+#include "Materials/MaterialInterface.h"
 
 #include UE_INLINE_GENERATED_CPP_BY_NAME(PCGActorHelpers)
 
diff --git a/Engine/Plugins/Experimental/PCG/Source/PCG/Private/Helpers/PCGBlueprintHelpers.cpp b/Engine/Plugins/Experimental/PCG/Source/PCG/Private/Helpers/PCGBlueprintHelpers.cpp
--- a/Engine/Plugins/Experimental/PCG/Source/PCG/Private/Helpers/PCGBlueprintHelpers.cpp
+++ b/Engine/Plugins/Experimental/PCG/Source/PCG/Private/Helpers/PCGBlueprintHelpers.cpp
@@ -1,14 +1,16 @@
 // Copyright Epic Games, Inc. All Rights Reserved.
 
 #include "Helpers/PCGBlueprintHelpers.h"
+
 #include "PCGComponent.h"
 #include "PCGContext.h"
-#include "PCGHelpers.h"
 #include "PCGModule.h"
 #include "PCGPoint.h"
 #include "PCGSubsystem.h"
 #include "Grid/PCGPartitionActor.h"
 #include "Grid/PCGLandscapeCache.h"
+#include "Helpers/PCGHelpers.h"
+
 #include "Engine/World.h"
 #include "LandscapeProxy.h"
 #include "LandscapeInfo.h"
diff --git a/Engine/Plugins/Experimental/PCG/Source/PCG/Private/Helpers/PCGHelpers.cpp b/Engine/Plugins/Experimental/PCG/Source/PCG/Private/Helpers/PCGHelpers.cpp
new file mode 100644
--- /dev/null
+++ b/Engine/Plugins/Experimental/PCG/Source/PCG/Private/Helpers/PCGHelpers.cpp
@@ -0,0 +1,394 @@
+// Copyright Epic Games, Inc. All Rights Reserved.
+
+#include "Helpers/PCGHelpers.h"
+
+#include "PCGComponent.h"
+#include "PCGModule.h"
+#include "PCGSubsystem.h"
+#include "Grid/PCGPartitionActor.h"
+
+#include "Landscape.h"
+#include "UObject/UObjectIterator.h"
+
+#if WITH_EDITOR
+#include "Editor.h"
+#else
+#include "Engine/World.h"
+#endif
+
+namespace PCGHelpers
+{
+	int ComputeSeed(int A)
+	{
+		return (A * 196314165U) + 907633515U;
+	}
+
+	int ComputeSeed(int A, int B)
+	{
+		return ((A * 196314165U) + 907633515U) ^ ((B * 73148459U) + 453816763U);
+	}
+
+	int ComputeSeed(int A, int B, int C)
+	{
+		return ((A * 196314165U) + 907633515U) ^ ((B * 73148459U) + 453816763U) ^ ((C * 34731343U) + 453816743U);
+	}
+
+	bool IsInsideBounds(const FBox& InBox, const FVector& InPosition)
+	{
+		return (InPosition.X >= InBox.Min.X) && (InPosition.X < InBox.Max.X) &&
+			(InPosition.Y >= InBox.Min.Y) && (InPosition.Y < InBox.Max.Y) &&
+			(InPosition.Z >= InBox.Min.Z) && (InPosition.Z < InBox.Max.Z);
+	}
+
+	bool IsInsideBoundsXY(const FBox& InBox, const FVector& InPosition)
+	{
+		return (InPosition.X >= InBox.Min.X) && (InPosition.X < InBox.Max.X) &&
+			(InPosition.Y >= InBox.Min.Y) && (InPosition.Y < InBox.Max.Y);
+	}
+
+	FBox OverlapBounds(const FBox& InA, const FBox& InB)
+	{
+		if (!InA.IsValid || !InB.IsValid)
+		{
+			return FBox(EForceInit::ForceInit);
+		}
+		else
+		{
+			return InA.Overlap(InB);
+		}
+	}
+
+	FBox GetGridBounds(const AActor* Actor, const UPCGComponent* Component)
+	{
+		FBox Bounds(EForceInit::ForceInit);
+
+		if (const APCGPartitionActor* PartitionActor = Cast<const APCGPartitionActor>(Actor))
+		{
+			// First, get the bounds from the partition actor
+			Bounds = PartitionActor->GetFixedBounds();
+
+			const UPCGComponent* OriginalComponent = Component ? PartitionActor->GetOriginalComponent(Component) : nullptr;
+			if (OriginalComponent)
+			{
+				if (OriginalComponent->GetOwner() != PartitionActor)
+				{
+					Bounds = Bounds.Overlap(OriginalComponent->GetGridBounds());
+				}
+			}
+		}
+		// TODO: verify this works as expected in non-editor builds
+		else if (const ALandscapeProxy* LandscapeActor = Cast<const ALandscape>(Actor))
+		{
+			Bounds = GetLandscapeBounds(LandscapeActor);
+		}
+		else if (Actor)
+		{
+			Bounds = GetActorBounds(Actor);
+		}
+		else
+		{
+			UE_LOG(LogPCG, Error, TEXT("Actor is invalid in GetGridBounds"));
+		}
+
+		return Bounds;
+	}
+
+	FBox GetActorBounds(const AActor* InActor, bool bIgnorePCGCreatedComponents)
+	{
+		// Specialized version of GetComponentsBoundingBox that skips over PCG generated components
+		// This is to ensure stable bounds and no timing issues (cleared ISMs, etc.)
+		FBox Box(EForceInit::ForceInit);
+
+		const bool bNonColliding = true;
+		const bool bIncludeFromChildActors = true;
+
+		if (InActor)
+		{
+			InActor->ForEachComponent<UPrimitiveComponent>(bIncludeFromChildActors, [bNonColliding, bIgnorePCGCreatedComponents, &Box](const UPrimitiveComponent* InPrimComp)
+			{
+				// Note: we omit the IsRegistered check here (e.g. InPrimComp->IsRegistered() )
+				// since this can be called in a scope where the components are temporarily unregistered
+				if ((bNonColliding || InPrimComp->IsCollisionEnabled()) &&
+					(!bIgnorePCGCreatedComponents || !InPrimComp->ComponentTags.Contains(DefaultPCGTag)))
+				{
+					Box += InPrimComp->Bounds.GetBox();
+				}
+			});
+		}
+		else
+		{
+			UE_LOG(LogPCG, Error, TEXT("Actor is invalid in GetActorBounds"));
+		}
+
+		return Box;
+	}
+
+	FBox GetActorLocalBounds(const AActor* InActor, bool bIgnorePCGCreatedComponents)
+	{
+		// Specialized version of CalculateComponentsBoundingBoxInLocalScape that skips over PCG generated components
+		// This is to ensure stable bounds and no timing issues (cleared ISMs, etc.)
+		FBox Box(EForceInit::ForceInit);
+
+		const bool bNonColliding = true;
+		const bool bIncludeFromChildActors = true;
+
+		if (InActor)
+		{
+			const FTransform& ActorToWorld = InActor->GetTransform();
+			const FTransform WorldToActor = ActorToWorld.Inverse();
+
+			InActor->ForEachComponent<UPrimitiveComponent>(bIncludeFromChildActors, [bNonColliding, bIgnorePCGCreatedComponents, &WorldToActor, &Box](const UPrimitiveComponent* InPrimComp)
+			{
+				if ((bNonColliding || InPrimComp->IsCollisionEnabled()) &&
+					(!bIgnorePCGCreatedComponents || !InPrimComp->ComponentTags.Contains(DefaultPCGTag)))
+				{
+					const FTransform ComponentToActor = InPrimComp->GetComponentTransform() * WorldToActor;
+					Box += InPrimComp->CalcBounds(ComponentToActor).GetBox();
+				}
+			});
+		}
+		else
+		{
+			UE_LOG(LogPCG, Error, TEXT("Actor is invalid in GetActorLocalBounds"));
+		}
+
+		return Box;
+	}
+
+	bool IsRuntimeOrPIE()
+	{
+#if WITH_EDITOR
+		return (GEditor && GEditor->PlayWorld) || GIsPlayInEditorWorld || IsRunningGame();
+#else
+		return true;
+#endif // WITH_EDITOR
+	}
+
+	FBox GetLandscapeBounds(const ALandscapeProxy* InLandscape)
+	{
+		check(InLandscape);
+
+		if (const ALandscape* Landscape = Cast<const ALandscape>(InLandscape))
+		{
+			// If the landscape isn't done being loaded, we're very unlikely to want to interact with it
+			if (Landscape->GetLandscapeInfo() == nullptr)
+			{
+				return FBox(EForceInit::ForceInit);
+			}
+
+#if WITH_EDITOR
+			if (!IsRuntimeOrPIE())
+			{
+				return Landscape->GetCompleteBounds();
+			}
+			else
+#endif
+			{
+				return Landscape->GetLoadedBounds();
+			}
+		}
+		else
+		{
+			return GetActorBounds(InLandscape);
+		}
+	}
+
+	TArray<TWeakObjectPtr<ALandscapeProxy>> GetAllLandscapeProxies(UWorld* InWorld)
+	{
+		TArray<TWeakObjectPtr<ALandscapeProxy>> LandscapeProxies;
+
+		for (TObjectIterator<ALandscapeProxy> It; It; ++It)
+		{
+			if (It->GetWorld() == InWorld)
+			{
+				LandscapeProxies.Add(*It);
+			}
+		}
+
+		return LandscapeProxies;
+	}
+
+	ALandscape* GetLandscape(UWorld* InWorld, const FBox& InBounds)
+	{
+		ALandscape* Landscape = nullptr;
+
+		if (!InBounds.IsValid)
+		{
+			return Landscape;
+		}
+
+		for (TObjectIterator<ALandscape> It; It; ++It)
+		{
+			if (It->GetWorld() == InWorld)
+			{
+				const FBox LandscapeBounds = GetLandscapeBounds(*It);
+				if (LandscapeBounds.IsValid && LandscapeBounds.Intersect(InBounds))
+				{
+					Landscape = (*It);
+					break;
+				}
+			}
+		}
+
+		return Landscape;
+	}
+
+	TArray<TWeakObjectPtr<ALandscapeProxy>> GetLandscapeProxies(UWorld* InWorld, const FBox& InBounds)
+	{
+		TArray<TWeakObjectPtr<ALandscapeProxy>> LandscapeProxies;
+
+		if (!InBounds.IsValid)
+		{
+			return LandscapeProxies;
+		}
+
+		for (TObjectIterator<ALandscapeProxy> It; It; ++It)
+		{
+			if (It->GetWorld() == InWorld)
+			{
+				const FBox LandscapeBounds = GetLandscapeBounds(*It);
+				if (LandscapeBounds.IsValid && LandscapeBounds.Intersect(InBounds))
+				{
+					LandscapeProxies.Add(*It);
+				}
+			}
+		}
+
+		return LandscapeProxies;
+	}
+
+	APCGWorldActor* GetPCGWorldActor(UWorld* InWorld)
+	{
+		return (InWorld && InWorld->GetSubsystem<UPCGSubsystem>()) ? InWorld->GetSubsystem<UPCGSubsystem>()->GetPCGWorldActor() : nullptr;
+	}
+
+	TArray<FString> GetStringArrayFromCommaSeparatedString(const FString& InCommaSeparatedString)
+	{
+		TArray<FString> Result;
+		InCommaSeparatedString.ParseIntoArrayWS(Result, TEXT(","));
+		return Result;
+	}
+
+#if WITH_EDITOR
+	bool CanBeExpanded(UClass* ObjectClass)
+	{
+		// There shouldn't be any need to dig through Niagara assets + there are some issues (most likely related to loading) with parsing all their dependencies
+		if (!ObjectClass ||
+			ObjectClass->GetFName() == TEXT("NiagaraSystem") ||
+			ObjectClass->GetFName() == TEXT("NiagaraComponent"))
+		{
+			return false;
+		}
+
+		return true;
+	}
+
+	void GatherDependencies(UObject* Object, TSet<TObjectPtr<UObject>>& OutDependencies, int32 MaxDepth)
+	{
+		UClass* ObjectClass = Object ? Object->GetClass() : nullptr;
+
+		if(!CanBeExpanded(ObjectClass))
+		{
+			return;
+		}
+		checkSlow(ObjectClass); // CanBeExpanded checks this but static analyzer doesn't note that
+
+		for (FProperty* Property = ObjectClass->PropertyLink; Property != nullptr; Property = Property->PropertyLinkNext)
+		{
+			// Skip any kind of internal property and the ones that are susceptible to be instable
+			if (!Property || Property->HasAnyPropertyFlags(CPF_Transient | CPF_DuplicateTransient | CPF_Deprecated))
+			{
+				continue;
+			}
+
+			GatherDependencies(Property, Object, OutDependencies, MaxDepth);
+		}
+	}
+
+	// Inspired by IteratePropertiesRecursive in ObjectPropertyTrace.cpp
+	void GatherDependencies(FProperty* Property, const void* InContainer, TSet<TObjectPtr<UObject>>& OutDependencies, int32 MaxDepth)
+	{
+		auto AddToDependenciesAndGatherRecursively = [&OutDependencies, MaxDepth](UObject* Object) {
+			if (Object && !OutDependencies.Contains(Object))
+			{
+				OutDependencies.Add(Object);
+				if (MaxDepth != 0)
+				{
+					GatherDependencies(Object, OutDependencies, MaxDepth - 1);
+				}
+			}
+		};
+
+		if (!InContainer || !Property)
+		{
+			return;
+		}
+
+		if (FObjectProperty* ObjectProperty = CastField<FObjectProperty>(Property))
+		{
+			UObject* Object = ObjectProperty->GetPropertyValue_InContainer(InContainer);
+			AddToDependenciesAndGatherRecursively(Object);
+		}
+		else if (FWeakObjectProperty* WeakObjectProperty = CastField<FWeakObjectProperty>(Property))
+		{
+			FWeakObjectPtr WeakObject = WeakObjectProperty->GetPropertyValue_InContainer(InContainer);
+			AddToDependenciesAndGatherRecursively(WeakObject.Get());
+		}
+		else if (FSoftObjectProperty* SoftObjectProperty = CastField<FSoftObjectProperty>(Property))
+		{
+			FSoftObjectPtr SoftObject = SoftObjectProperty->GetPropertyValue_InContainer(InContainer);
+			AddToDependenciesAndGatherRecursively(SoftObject.Get());
+		}
+		else if (FStructProperty* StructProperty = CastField<FStructProperty>(Property))
+		{
+			const void* StructContainer = StructProperty->ContainerPtrToValuePtr<const void>(InContainer);
+			for (TFieldIterator<FProperty> It(StructProperty->Struct); It; ++It)
+			{
+				GatherDependencies(*It, StructContainer, OutDependencies, MaxDepth);
+			}
+		}
+		else if (FArrayProperty* ArrayProperty = CastField<FArrayProperty>(Property))
+		{
+			FScriptArrayHelper_InContainer Helper(ArrayProperty, InContainer);
+			for (int32 DynamicIndex = 0; DynamicIndex < Helper.Num(); ++DynamicIndex)
+			{
+				const void* ValuePtr = Helper.GetRawPtr(DynamicIndex);
+				GatherDependencies(ArrayProperty->Inner, ValuePtr, OutDependencies, MaxDepth);
+			}
+		}
+		else if (FMapProperty* MapProperty = CastField<FMapProperty>(Property))
+		{
+			FScriptMapHelper_InContainer Helper(MapProperty, InContainer);
+			int32 Num = Helper.Num();
+			for (int32 DynamicIndex = 0; Num; ++DynamicIndex)
+			{
+				if (Helper.IsValidIndex(DynamicIndex))
+				{
+					// Key and Value are stored next to each other in memory.
+					// ValueProp has an offset, so we should use the same starting address for both.
+					const void* PairKeyValuePtr = Helper.GetKeyPtr(DynamicIndex);
+					GatherDependencies(MapProperty->KeyProp, PairKeyValuePtr, OutDependencies, MaxDepth);
+					GatherDependencies(MapProperty->ValueProp, PairKeyValuePtr, OutDependencies, MaxDepth);
+
+					--Num;
+				}
+			}
+		}
+		else if (FSetProperty* SetProperty = CastField<FSetProperty>(Property))
+		{
+			FScriptSetHelper_InContainer Helper(SetProperty, InContainer);
+			int32 Num = Helper.Num();
+			for (int32 DynamicIndex = 0; Num; ++DynamicIndex)
+			{
+				if (Helper.IsValidIndex(DynamicIndex))
+				{
+					const void* ValuePtr = Helper.GetElementPtr(DynamicIndex);
+					GatherDependencies(SetProperty->ElementProp, ValuePtr, OutDependencies, MaxDepth);
+
+					--Num;
+				}
+			}
+		}
+	}
+#endif
+}
diff --git a/Engine/Plugins/Experimental/PCG/Source/PCG/Private/Helpers/PCGSettingsHelpers.cpp b/Engine/Plugins/Experimental/PCG/Source/PCG/Private/Helpers/PCGSettingsHelpers.cpp
--- a/Engine/Plugins/Experimental/PCG/Source/PCG/Private/Helpers/PCGSettingsHelpers.cpp
+++ b/Engine/Plugins/Experimental/PCG/Source/PCG/Private/Helpers/PCGSettingsHelpers.cpp
@@ -4,8 +4,8 @@
 
 #include "PCGComponent.h"
 #include "PCGEdge.h"
-#include "PCGHelpers.h"
 #include "PCGPin.h"
+#include "Helpers/PCGHelpers.h"
 
 #include "UObject/EnumProperty.h"
 
diff --git a/Engine/Plugins/Experimental/PCG/Source/PCG/Private/PCGActorAndComponentMapping.cpp b/Engine/Plugins/Experimental/PCG/Source/PCG/Private/PCGActorAndComponentMapping.cpp
--- a/Engine/Plugins/Experimental/PCG/Source/PCG/Private/PCGActorAndComponentMapping.cpp
+++ b/Engine/Plugins/Experimental/PCG/Source/PCG/Private/PCGActorAndComponentMapping.cpp
@@ -3,12 +3,12 @@
 #include "PCGActorAndComponentMapping.h"
 
 #include "PCGComponent.h"
-#include "PCGHelpers.h"
 #include "PCGModule.h"
 #include "PCGSubsystem.h"
 #include "PCGWorldActor.h"
 #include "Grid/PCGPartitionActor.h"
 #include "Helpers/PCGActorHelpers.h"
+#include "Helpers/PCGHelpers.h"
 
 #include "Engine/Engine.h"
 
diff --git a/Engine/Plugins/Experimental/PCG/Source/PCG/Private/PCGComponent.cpp b/Engine/Plugins/Experimental/PCG/Source/PCG/Private/PCGComponent.cpp
--- a/Engine/Plugins/Experimental/PCG/Source/PCG/Private/PCGComponent.cpp
+++ b/Engine/Plugins/Experimental/PCG/Source/PCG/Private/PCGComponent.cpp
@@ -5,7 +5,6 @@
 #include "PCGContext.h"
 #include "PCGEngineSettings.h"
 #include "PCGGraph.h"
-#include "PCGHelpers.h"
 #include "PCGInputOutputSettings.h"
 #include "PCGManagedResource.h"
 #include "PCGParamData.h"
@@ -24,6 +23,7 @@
 #include "Data/PCGVolumeData.h"
 #include "Grid/PCGPartitionActor.h"
 #include "Helpers/PCGActorHelpers.h"
+#include "Helpers/PCGHelpers.h"
 #include "Utils/PCGGeneratedResourcesLogging.h"
 
 #include "LandscapeComponent.h"
diff --git a/Engine/Plugins/Experimental/PCG/Source/PCG/Private/PCGContext.cpp b/Engine/Plugins/Experimental/PCG/Source/PCG/Private/PCGContext.cpp
--- a/Engine/Plugins/Experimental/PCG/Source/PCG/Private/PCGContext.cpp
+++ b/Engine/Plugins/Experimental/PCG/Source/PCG/Private/PCGContext.cpp
@@ -2,11 +2,11 @@
 
 #include "PCGContext.h"
 #include "PCGComponent.h"
-#include "PCGHelpers.h"
 #include "PCGParamData.h"
 #include "PCGPin.h"
 #include "PCGSubsystem.h"
 #include "Data/PCGSpatialData.h"
+#include "Helpers/PCGHelpers.h"
 #include "Metadata/PCGAttributePropertySelector.h"
 #include "Metadata/PCGMetadata.h"
 #include "Metadata/PCGMetadataAttribute.h"
diff --git a/Engine/Plugins/Experimental/PCG/Source/PCG/Private/PCGHelpers.cpp b/Engine/Plugins/Experimental/PCG/Source/PCG/Private/PCGHelpers.cpp
deleted file mode 100644
diff --git a/Engine/Plugins/Experimental/PCG/Source/PCG/Private/PCGHelpers.h b/Engine/Plugins/Experimental/PCG/Source/PCG/Private/PCGHelpers.h
deleted file mode 100644
diff --git a/Engine/Plugins/Experimental/PCG/Source/PCG/Private/PCGManagedResource.cpp b/Engine/Plugins/Experimental/PCG/Source/PCG/Private/PCGManagedResource.cpp
--- a/Engine/Plugins/Experimental/PCG/Source/PCG/Private/PCGManagedResource.cpp
+++ b/Engine/Plugins/Experimental/PCG/Source/PCG/Private/PCGManagedResource.cpp
@@ -3,8 +3,9 @@
 #include "PCGManagedResource.h"
 
 #include "PCGComponent.h"
-#include "PCGHelpers.h"
 #include "PCGModule.h"
+#include "Helpers/PCGHelpers.h"
+
 #include "Components/InstancedStaticMeshComponent.h"
 #include "Utils/PCGGeneratedResourcesLogging.h"
 
diff --git a/Engine/Plugins/Experimental/PCG/Source/PCG/Private/PCGSettings.cpp b/Engine/Plugins/Experimental/PCG/Source/PCG/Private/PCGSettings.cpp
--- a/Engine/Plugins/Experimental/PCG/Source/PCG/Private/PCGSettings.cpp
+++ b/Engine/Plugins/Experimental/PCG/Source/PCG/Private/PCGSettings.cpp
@@ -7,10 +7,10 @@
 #include "PCGCustomVersion.h"
 #include "PCGEdge.h"
 #include "PCGGraph.h"
-#include "PCGHelpers.h"
 #include "PCGModule.h"
 #include "PCGPin.h"
 #include "PCGSubsystem.h"
+#include "Helpers/PCGHelpers.h"
 #include "Helpers/PCGSettingsHelpers.h"
 #include "Metadata/Accessors/PCGAttributeAccessorHelpers.h"
 #include "Metadata/Accessors/IPCGAttributeAccessor.h"
diff --git a/Engine/Plugins/Experimental/PCG/Source/PCG/Private/PCGSubsystem.cpp b/Engine/Plugins/Experimental/PCG/Source/PCG/Private/PCGSubsystem.cpp
--- a/Engine/Plugins/Experimental/PCG/Source/PCG/Private/PCGSubsystem.cpp
+++ b/Engine/Plugins/Experimental/PCG/Source/PCG/Private/PCGSubsystem.cpp
@@ -1,17 +1,18 @@
 // Copyright Epic Games, Inc. All Rights Reserved.
 
 #include "PCGSubsystem.h"
-#include "Engine/Level.h"
+
 #include "PCGComponent.h"
-#include "Graph/PCGGraphCache.h"
 #include "PCGModule.h"
-#include "PCGHelpers.h"
-#include "Grid/PCGLandscapeCache.h"
 #include "PCGWorldActor.h"
+#include "Graph/PCGGraphCache.h"
 #include "Graph/PCGGraphExecutor.h"
+#include "Grid/PCGLandscapeCache.h"
 #include "Grid/PCGPartitionActor.h"
 #include "Helpers/PCGActorHelpers.h"
+#include "Helpers/PCGHelpers.h"
 
+#include "Engine/Level.h"
 
 #include UE_INLINE_GENERATED_CPP_BY_NAME(PCGSubsystem)
 
diff --git a/Engine/Plugins/Experimental/PCG/Source/PCG/Private/Tests/Determinism/PCGDeterminismTestsCommon.cpp b/Engine/Plugins/Experimental/PCG/Source/PCG/Private/Tests/Determinism/PCGDeterminismTestsCommon.cpp
--- a/Engine/Plugins/Experimental/PCG/Source/PCG/Private/Tests/Determinism/PCGDeterminismTestsCommon.cpp
+++ b/Engine/Plugins/Experimental/PCG/Source/PCG/Private/Tests/Determinism/PCGDeterminismTestsCommon.cpp
@@ -6,7 +6,6 @@
 #include "Data/PCGSpatialData.h"
 #include "PCGGraph.h"
 #include "PCGContext.h"
-#include "PCGHelpers.h"
 #include "PCGParamData.h"
 #include "PCGPin.h"
 #include "Data/PCGPointData.h"
@@ -16,6 +15,7 @@
 #include "Data/PCGVolumeData.h"
 #include "Graph/PCGGraphCache.h"
 #include "Graph/PCGGraphExecutor.h"
+#include "Helpers/PCGHelpers.h"
 
 #include "Components/SplineComponent.h"
 #include "Components/StaticMeshComponent.h"
diff --git a/Engine/Plugins/Experimental/PCG/Source/PCG/Private/Tests/Elements/PCGCopyPointsTest.cpp b/Engine/Plugins/Experimental/PCG/Source/PCG/Private/Tests/Elements/PCGCopyPointsTest.cpp
--- a/Engine/Plugins/Experimental/PCG/Source/PCG/Private/Tests/Elements/PCGCopyPointsTest.cpp
+++ b/Engine/Plugins/Experimental/PCG/Source/PCG/Private/Tests/Elements/PCGCopyPointsTest.cpp
@@ -2,14 +2,13 @@
 
 #include "Tests/Determinism/PCGDeterminismTestsCommon.h"
 
-#include "Data/PCGSpatialData.h"
 #include "PCGComponent.h"
-#include "Metadata/PCGMetadataAccessor.h"
-#include "PCGHelpers.h"
-
+#include "PCGContext.h"
 #include "Data/PCGPointData.h"
+#include "Data/PCGSpatialData.h"
 #include "Elements/PCGCopyPoints.h"
-#include "PCGContext.h"
+#include "Helpers/PCGHelpers.h"
+#include "Metadata/PCGMetadataAccessor.h"
 
 #if WITH_EDITOR
 
diff --git a/Engine/Plugins/Experimental/PCG/Source/PCG/Private/Tests/Elements/PCGTransformPointsTest.cpp b/Engine/Plugins/Experimental/PCG/Source/PCG/Private/Tests/Elements/PCGTransformPointsTest.cpp
--- a/Engine/Plugins/Experimental/PCG/Source/PCG/Private/Tests/Elements/PCGTransformPointsTest.cpp
+++ b/Engine/Plugins/Experimental/PCG/Source/PCG/Private/Tests/Elements/PCGTransformPointsTest.cpp
@@ -5,10 +5,9 @@
 
 #include "PCGComponent.h"
 #include "PCGContext.h"
-#include "PCGHelpers.h"
-
 #include "Data/PCGPointData.h"
 #include "Elements/PCGTransformPoints.h"
+#include "Helpers/PCGHelpers.h"
 
 #if WITH_EDITOR
 
diff --git a/Engine/Plugins/Experimental/PCG/Source/PCG/Public/Helpers/PCGHelpers.h b/Engine/Plugins/Experimental/PCG/Source/PCG/Public/Helpers/PCGHelpers.h
new file mode 100644
--- /dev/null
+++ b/Engine/Plugins/Experimental/PCG/Source/PCG/Public/Helpers/PCGHelpers.h
@@ -0,0 +1,52 @@
+// Copyright Epic Games, Inc. All Rights Reserved.
+
+#pragma once
+
+#include "Math/Box.h"
+
+class AActor;
+class APCGWorldActor;
+class ALandscape;
+class ALandscapeProxy;
+class UPCGComponent;
+class UWorld;
+
+namespace PCGHelpers
+{
+	/** Tag that will be added on every component generated through the PCG system */
+	const FName DefaultPCGTag = TEXT("PCG Generated Component");
+	const FName DefaultPCGDebugTag = TEXT("PCG Generated Debug Component");
+	const FName DefaultPCGActorTag = TEXT("PCG Generated Actor");
+	const FName MarkedForCleanupPCGTag = TEXT("PCG Marked For Cleanup");
+
+	PCG_API int ComputeSeed(int A);
+	PCG_API int ComputeSeed(int A, int B);
+	PCG_API int ComputeSeed(int A, int B, int C);
+
+	PCG_API bool IsInsideBounds(const FBox& InBox, const FVector& InPosition);
+	PCG_API bool IsInsideBoundsXY(const FBox& InBox, const FVector& InPosition);
+
+	PCG_API FBox OverlapBounds(const FBox& InBoxA, const FBox& InBoxB);
+
+	/** Returns the bounds of InActor, intersected with the component if InActor is a partition actor */
+	PCG_API FBox GetGridBounds(const AActor* InActor, const UPCGComponent* InComponent);
+
+	PCG_API FBox GetActorBounds(const AActor* InActor, bool bIgnorePCGCreatedComponents = true);
+	PCG_API FBox GetActorLocalBounds(const AActor* InActor, bool bIgnorePCGCreatedComponents = true);
+	PCG_API FBox GetLandscapeBounds(const ALandscapeProxy* InLandscape);
+
+	PCG_API ALandscape* GetLandscape(UWorld* InWorld, const FBox& InActorBounds);
+	PCG_API TArray<TWeakObjectPtr<ALandscapeProxy>> GetLandscapeProxies(UWorld* InWorld, const FBox& InActorBounds);
+	PCG_API TArray<TWeakObjectPtr<ALandscapeProxy>> GetAllLandscapeProxies(UWorld* InWorld);
+
+	PCG_API bool IsRuntimeOrPIE();
+
+	PCG_API APCGWorldActor* GetPCGWorldActor(UWorld* InWorld);
+
+	PCG_API TArray<FString> GetStringArrayFromCommaSeparatedString(const FString& InCommaSeparatedString);
+
+#if WITH_EDITOR
+	PCG_API void GatherDependencies(UObject* Object, TSet<TObjectPtr<UObject>>& OutDependencies, int32 MaxDepth = -1);
+	PCG_API void GatherDependencies(FProperty* Property, const void* InContainer, TSet<TObjectPtr<UObject>>& OutDependencies, int32 MaxDepth);
+#endif
+};
diff --git a/Engine/Plugins/Experimental/PCGGeometryScriptInterop/Source/PCGGeometryScriptInterop/Private/Elements/PCGMeshSampler.cpp b/Engine/Plugins/Experimental/PCGGeometryScriptInterop/Source/PCGGeometryScriptInterop/Private/Elements/PCGMeshSampler.cpp
--- a/Engine/Plugins/Experimental/PCGGeometryScriptInterop/Source/PCGGeometryScriptInterop/Private/Elements/PCGMeshSampler.cpp
+++ b/Engine/Plugins/Experimental/PCGGeometryScriptInterop/Source/PCGGeometryScriptInterop/Private/Elements/PCGMeshSampler.cpp
@@ -8,6 +8,7 @@
 #include "Data/PCGPointData.h"
 #include "Helpers/PCGAsync.h"
 #include "Helpers/PCGBlueprintHelpers.h"
+#include "Helpers/PCGHelpers.h"
 
 #include "UDynamicMesh.h"
 #include "Engine/AssetManager.h"
@@ -264,16 +265,10 @@ bool FPCGMeshSamplerElement::ExecuteInternal(FPCGContext* InContext) const
 			{
 				TRACE_CPUPROFILER_EVENT_SCOPE(FPCGMeshSamplerElement::Execute::PoissonSampling);
 
-				// Taken from PCGHelpers, which is private.
-				auto ComputeSeed = [](int A, int B) -> int
-				{
-					return ((A * 196314165U) + 907633515U) ^ ((B * 73148459U) + 453816763U);
-				};
-
 				UE::Geometry::FMeshSurfacePointSampling PointSampling;
 				PointSampling.SampleRadius = Settings->SamplingOptions.SamplingRadius;
 				PointSampling.MaxSamples = Settings->SamplingOptions.MaxNumSamples;
-				PointSampling.RandomSeed = ComputeSeed(Seed, Settings->SamplingOptions.RandomSeed);
+				PointSampling.RandomSeed = PCGHelpers::ComputeSeed(Seed, Settings->SamplingOptions.RandomSeed);
 				PointSampling.SubSampleDensity = Settings->SamplingOptions.SubSampleDensity;
 
 				if (Settings->NonUniformSamplingOptions.MaxSamplingRadius > PointSampling.SampleRadius)
