[PCG] Fix dynamic pins not updating on downstream nodes when you break an edge via an output pin

diff --git a/Engine/Plugins/Experimental/PCG/Source/PCG/Private/PCGGraph.cpp b/Engine/Plugins/Experimental/PCG/Source/PCG/Private/PCGGraph.cpp
--- a/Engine/Plugins/Experimental/PCG/Source/PCG/Private/PCGGraph.cpp
+++ b/Engine/Plugins/Experimental/PCG/Source/PCG/Private/PCGGraph.cpp
@@ -3,6 +3,7 @@
 #include "PCGGraph.h"
 
 #include "PCGComponent.h"
+#include "PCGEdge.h"
 #include "PCGInputOutputSettings.h"
 #include "PCGModule.h"
 #include "PCGNode.h"
@@ -571,47 +572,6 @@ void UPCGGraph::ForEachNode(const TFunction<void(UPCGNode*)>& Action)
 	}
 }
 
-bool UPCGGraph::RemoveAllInboundEdges(UPCGNode* InNode)
-{
-	check(InNode);
-	bool bChanged = false;
-
-	for (UPCGPin* InputPin : InNode->InputPins)
-	{
-		bChanged |= InputPin->BreakAllEdges();
-	}
-
-#if WITH_EDITOR
-	if (bChanged)
-	{
-		InNode->UpdateDynamicPins();
-		NotifyGraphChanged(EPCGChangeType::Structural);
-	}
-#endif
-
-	return bChanged;
-}
-
-bool UPCGGraph::RemoveAllOutboundEdges(UPCGNode* InNode)
-{
-	check(InNode);
-	bool bChanged = false;
-	for (UPCGPin* OutputPin : InNode->OutputPins)
-	{
-		bChanged |= OutputPin->BreakAllEdges();
-	}
-
-#if WITH_EDITOR
-	if (bChanged)
-	{
-		InNode->UpdateDynamicPins();
-		NotifyGraphChanged(EPCGChangeType::Structural);
-	}
-#endif
-
-	return bChanged;
-}
-
 bool UPCGGraph::RemoveInboundEdges(UPCGNode* InNode, const FName& InboundLabel)
 {
 	check(InNode);
@@ -625,8 +585,8 @@ bool UPCGGraph::RemoveInboundEdges(UPCGNode* InNode, const FName& InboundLabel)
 #if WITH_EDITOR
 	if (bChanged)
 	{
-		InNode->UpdateDynamicPins();
-		NotifyGraphChanged(EPCGChangeType::Structural);
+		const EPCGChangeType ChangeType = InNode->UpdateDynamicPins() | EPCGChangeType::Structural;
+		NotifyGraphChanged(ChangeType);
 	}
 #endif
 
@@ -637,17 +597,36 @@ bool UPCGGraph::RemoveOutboundEdges(UPCGNode* InNode, const FName& OutboundLabel
 {
 	check(InNode);
 	bool bChanged = false;
+	// Make a list of downstream nodes which may need pin updates when the edges change
+	TSet<UPCGNode*> ConnectedNodes;
 
 	if (UPCGPin* OutputPin = InNode->GetOutputPin(OutboundLabel))
 	{
+		for (const UPCGEdge* Edge : OutputPin->Edges)
+		{
+			if (Edge && Edge->OutputPin && Edge->OutputPin->Node)
+			{
+				ConnectedNodes.Add(Edge->OutputPin->Node);
+			}
+		}
+
 		bChanged = OutputPin->BreakAllEdges();
 	}
 
 #if WITH_EDITOR
 	if (bChanged)
 	{
-		InNode->UpdateDynamicPins();
-		NotifyGraphChanged(EPCGChangeType::Structural);
+		EPCGChangeType ChangeType = EPCGChangeType::Structural;
+
+		// Pins may need updating on downstream nodes
+		for (UPCGNode* ConnectedNode : ConnectedNodes)
+		{
+			ChangeType |= ConnectedNode->UpdateDynamicPins();
+		}
+
+		ChangeType |= InNode->UpdateDynamicPins();
+
+		NotifyGraphChanged(ChangeType);
 	}
 #endif
 
diff --git a/Engine/Plugins/Experimental/PCG/Source/PCG/Public/PCGGraph.h b/Engine/Plugins/Experimental/PCG/Source/PCG/Public/PCGGraph.h
--- a/Engine/Plugins/Experimental/PCG/Source/PCG/Public/PCGGraph.h
+++ b/Engine/Plugins/Experimental/PCG/Source/PCG/Public/PCGGraph.h
@@ -190,8 +190,6 @@ public:
 	/** Calls the lambda on every node in graph. */
 	void ForEachNode(const TFunction<void(UPCGNode*)>& Action);
 
-	bool RemoveAllInboundEdges(UPCGNode* InNode);
-	bool RemoveAllOutboundEdges(UPCGNode* InNode);
 	bool RemoveInboundEdges(UPCGNode* InNode, const FName& InboundLabel);
 	bool RemoveOutboundEdges(UPCGNode* InNode, const FName& OutboundLabel);
 
