[PCG] Roll out dynamic pin types to last nodes on my list, to minimize use of blanket 'Any' type

diff --git a/Engine/Plugins/Experimental/PCG/Source/PCG/Private/Elements/Metadata/PCGMetadataOpElementBase.cpp b/Engine/Plugins/Experimental/PCG/Source/PCG/Private/Elements/Metadata/PCGMetadataOpElementBase.cpp
--- a/Engine/Plugins/Experimental/PCG/Source/PCG/Private/Elements/Metadata/PCGMetadataOpElementBase.cpp
+++ b/Engine/Plugins/Experimental/PCG/Source/PCG/Private/Elements/Metadata/PCGMetadataOpElementBase.cpp
@@ -32,12 +32,19 @@ TArray<FPCGPinProperties> UPCGMetadataSettingsBase::InputPinProperties() const
 {
 	TArray<FPCGPinProperties> PinProperties;
 
+	EPCGDataType FirstPinInputTypeUnion = EPCGDataType::Any;
+	if (GetInputPinNum() > 0)
+	{
+		const EPCGDataType IncidentTypeUnion = GetTypeUnionOfIncidentEdges(GetInputPinLabel(0));
+		FirstPinInputTypeUnion = (IncidentTypeUnion != EPCGDataType::None) ? IncidentTypeUnion : FirstPinInputTypeUnion;
+	}
+
 	for (uint32 i = 0; i < GetInputPinNum(); ++i)
 	{
 		const FName PinLabel = GetInputPinLabel(i);
 		if (PinLabel != NAME_None)
 		{
-			PinProperties.Emplace(PinLabel, EPCGDataType::Any, /*bAllowMultipleConnections=*/false);
+			PinProperties.Emplace(PinLabel, (i == 0) ? FirstPinInputTypeUnion : EPCGDataType::Any, /*bAllowMultipleConnections=*/false);
 		}
 	}
 
@@ -48,12 +55,19 @@ TArray<FPCGPinProperties> UPCGMetadataSettingsBase::OutputPinProperties() const
 {
 	TArray<FPCGPinProperties> PinProperties;
 	
+	EPCGDataType OutputType = EPCGDataType::Any;
+	if (GetInputPinNum() > 0)
+	{
+		const EPCGDataType IncidentTypeUnion = GetTypeUnionOfIncidentEdges(GetInputPinLabel(0));
+		OutputType = (IncidentTypeUnion != EPCGDataType::None) ? IncidentTypeUnion : OutputType;
+	}
+
 	for (uint32 i = 0; i < GetOutputPinNum(); ++i)
 	{
 		const FName PinLabel = GetOutputPinLabel(i);
 		if (PinLabel != NAME_None)
 		{
-			PinProperties.Emplace(PinLabel, EPCGDataType::Any);
+			PinProperties.Emplace(PinLabel, OutputType);
 		}
 	}
 
diff --git a/Engine/Plugins/Experimental/PCG/Source/PCG/Private/Elements/Metadata/PCGMetadataRenameElement.cpp b/Engine/Plugins/Experimental/PCG/Source/PCG/Private/Elements/Metadata/PCGMetadataRenameElement.cpp
--- a/Engine/Plugins/Experimental/PCG/Source/PCG/Private/Elements/Metadata/PCGMetadataRenameElement.cpp
+++ b/Engine/Plugins/Experimental/PCG/Source/PCG/Private/Elements/Metadata/PCGMetadataRenameElement.cpp
@@ -12,10 +12,22 @@
 
 #define LOCTEXT_NAMESPACE "PCGMetadataRenameElement"
 
+TArray<FPCGPinProperties> UPCGMetadataRenameSettings::InputPinProperties() const
+{
+	EPCGDataType InputTypeUnion = GetTypeUnionOfIncidentEdges(PCGPinConstants::DefaultInputLabel);
+
+	TArray<FPCGPinProperties> PinProperties;
+	PinProperties.Emplace(PCGPinConstants::DefaultInputLabel, (InputTypeUnion != EPCGDataType::None) ? InputTypeUnion : EPCGDataType::Any);
+
+	return PinProperties;
+}
+
 TArray<FPCGPinProperties> UPCGMetadataRenameSettings::OutputPinProperties() const
 {
+	const EPCGDataType InputTypeUnion = GetTypeUnionOfIncidentEdges(PCGPinConstants::DefaultInputLabel);
+
 	TArray<FPCGPinProperties> PinProperties;
-	PinProperties.Emplace(PCGPinConstants::DefaultOutputLabel, EPCGDataType::Any);
+	PinProperties.Emplace(PCGPinConstants::DefaultOutputLabel, (InputTypeUnion != EPCGDataType::None) ? InputTypeUnion : EPCGDataType::Any);
 
 	return PinProperties;
 }
diff --git a/Engine/Plugins/Experimental/PCG/Source/PCG/Private/Elements/PCGAttributeFilterElement.cpp b/Engine/Plugins/Experimental/PCG/Source/PCG/Private/Elements/PCGAttributeFilterElement.cpp
--- a/Engine/Plugins/Experimental/PCG/Source/PCG/Private/Elements/PCGAttributeFilterElement.cpp
+++ b/Engine/Plugins/Experimental/PCG/Source/PCG/Private/Elements/PCGAttributeFilterElement.cpp
@@ -82,10 +82,22 @@ FName UPCGAttributeFilterSettings::AdditionalTaskName() const
 	}
 }
 
+TArray<FPCGPinProperties> UPCGAttributeFilterSettings::InputPinProperties() const
+{
+	const EPCGDataType InputTypeUnion = GetTypeUnionOfIncidentEdges(PCGPinConstants::DefaultInputLabel);
+
+	TArray<FPCGPinProperties> PinProperties;
+	PinProperties.Emplace(PCGPinConstants::DefaultInputLabel, (InputTypeUnion != EPCGDataType::None) ? InputTypeUnion : EPCGDataType::Any);
+
+	return PinProperties;
+}
+
 TArray<FPCGPinProperties> UPCGAttributeFilterSettings::OutputPinProperties() const
 {
+	const EPCGDataType InputTypeUnion = GetTypeUnionOfIncidentEdges(PCGPinConstants::DefaultInputLabel);
+
 	TArray<FPCGPinProperties> PinProperties;
-	PinProperties.Emplace(PCGPinConstants::DefaultOutputLabel, EPCGDataType::Any);
+	PinProperties.Emplace(PCGPinConstants::DefaultOutputLabel, (InputTypeUnion != EPCGDataType::None) ? InputTypeUnion : EPCGDataType::Any);
 
 	return PinProperties;
 }
diff --git a/Engine/Plugins/Experimental/PCG/Source/PCG/Private/Elements/PCGAttributeTransferElement.cpp b/Engine/Plugins/Experimental/PCG/Source/PCG/Private/Elements/PCGAttributeTransferElement.cpp
--- a/Engine/Plugins/Experimental/PCG/Source/PCG/Private/Elements/PCGAttributeTransferElement.cpp
+++ b/Engine/Plugins/Experimental/PCG/Source/PCG/Private/Elements/PCGAttributeTransferElement.cpp
@@ -21,6 +21,16 @@ FText UPCGAttributeTransferSettings::GetDefaultNodeTitle() const
 {
 	return PCGAttributeTransferConstants::NodeTitle;
 }
+
+FText UPCGAttributeTransferSettings::GetNodeTooltipText() const
+{
+	return LOCTEXT("NodeTooltip",
+		"Transfer an attribute from a source metadata to a target data. Only supports Spatial to Spatial and Points to Points, and they need to match.\n\n"
+		" - For Spatial data, number of entries in the metadata should be the same between source and target.\n"
+		" - For Point data, number of points should be the same between source and target.\n\n"
+		"The output will be the target data with the updated metadata. If the TargetAttributeName is None, it will use SourceAttributeName instead."
+	);
+}
 #endif
 
 FName UPCGAttributeTransferSettings::AdditionalTaskName() const
@@ -36,10 +46,20 @@ FName UPCGAttributeTransferSettings::AdditionalTaskName() const
 }
 
 TArray<FPCGPinProperties> UPCGAttributeTransferSettings::InputPinProperties() const
+{
+	const EPCGDataType InputType = GetPinsType();
+
+	TArray<FPCGPinProperties> PinProperties;
+	PinProperties.Emplace(PCGAttributeTransferConstants::TargetLabel, InputType, /*bInAllowMultipleConnections=*/ false);
+	PinProperties.Emplace(PCGAttributeTransferConstants::SourceLabel, InputType, /*bInAllowMultipleConnections=*/ false);
+
+	return PinProperties;
+}
+
+TArray<FPCGPinProperties> UPCGAttributeTransferSettings::OutputPinProperties() const
 {
 	TArray<FPCGPinProperties> PinProperties;
-	PinProperties.Emplace(PCGAttributeTransferConstants::TargetLabel, EPCGDataType::Spatial, /*bInAllowMultipleConnections=*/ false);
-	PinProperties.Emplace(PCGAttributeTransferConstants::SourceLabel, EPCGDataType::Spatial, /*bInAllowMultipleConnections=*/ false);
+	PinProperties.Emplace(PCGPinConstants::DefaultOutputLabel, GetPinsType());
 
 	return PinProperties;
 }
@@ -49,6 +69,15 @@ FPCGElementPtr UPCGAttributeTransferSettings::CreateElement() const
 	return MakeShared<FPCGAttributeTransferElement>();
 }
 
+EPCGDataType UPCGAttributeTransferSettings::GetPinsType() const
+{
+	// If either input is type Point then all pins are Point, otherwise Spatial
+	const bool bAnyArePoint = GetTypeUnionOfIncidentEdges(PCGAttributeTransferConstants::TargetLabel) == EPCGDataType::Point
+		|| GetTypeUnionOfIncidentEdges(PCGAttributeTransferConstants::SourceLabel) == EPCGDataType::Point;
+
+	return bAnyArePoint ? EPCGDataType::Point : EPCGDataType::Spatial;
+}
+
 bool FPCGAttributeTransferElement::ExecuteInternal(FPCGContext* Context) const
 {
 	TRACE_CPUPROFILER_EVENT_SCOPE(FPCGAttributeTransferElement::Execute);
diff --git a/Engine/Plugins/Experimental/PCG/Source/PCG/Private/Elements/PCGFilterByTag.cpp b/Engine/Plugins/Experimental/PCG/Source/PCG/Private/Elements/PCGFilterByTag.cpp
--- a/Engine/Plugins/Experimental/PCG/Source/PCG/Private/Elements/PCGFilterByTag.cpp
+++ b/Engine/Plugins/Experimental/PCG/Source/PCG/Private/Elements/PCGFilterByTag.cpp
@@ -49,10 +49,22 @@ FName UPCGFilterByTagSettings::AdditionalTaskName() const
 	}
 }
 
+TArray<FPCGPinProperties> UPCGFilterByTagSettings::InputPinProperties() const
+{
+	const EPCGDataType InputTypeUnion = GetTypeUnionOfIncidentEdges(PCGPinConstants::DefaultInputLabel);
+
+	TArray<FPCGPinProperties> PinProperties;
+	PinProperties.Emplace(PCGPinConstants::DefaultInputLabel, (InputTypeUnion != EPCGDataType::None) ? InputTypeUnion : EPCGDataType::Any);
+
+	return PinProperties;
+}
+
 TArray<FPCGPinProperties> UPCGFilterByTagSettings::OutputPinProperties() const
 {
+	const EPCGDataType InputTypeUnion = GetTypeUnionOfIncidentEdges(PCGPinConstants::DefaultInputLabel);
+
 	TArray<FPCGPinProperties> PinProperties;
-	PinProperties.Emplace(PCGPinConstants::DefaultOutputLabel, EPCGDataType::Any);
+	PinProperties.Emplace(PCGPinConstants::DefaultOutputLabel, (InputTypeUnion != EPCGDataType::None) ? InputTypeUnion : EPCGDataType::Any);
 
 	return PinProperties;
 }
diff --git a/Engine/Plugins/Experimental/PCG/Source/PCG/Public/Elements/Metadata/PCGMetadataOpElementBase.h b/Engine/Plugins/Experimental/PCG/Source/PCG/Public/Elements/Metadata/PCGMetadataOpElementBase.h
--- a/Engine/Plugins/Experimental/PCG/Source/PCG/Public/Elements/Metadata/PCGMetadataOpElementBase.h
+++ b/Engine/Plugins/Experimental/PCG/Source/PCG/Public/Elements/Metadata/PCGMetadataOpElementBase.h
@@ -78,6 +78,7 @@ public:
 	//~Begin UPCGSettings interface
 #if WITH_EDITOR
 	virtual EPCGSettingsType GetType() const override { return EPCGSettingsType::Metadata; }
+	virtual bool HasDynamicPins() const override { return true; }
 #endif
 
 	virtual TArray<FPCGPinProperties> InputPinProperties() const override;
diff --git a/Engine/Plugins/Experimental/PCG/Source/PCG/Public/Elements/Metadata/PCGMetadataRenameElement.h b/Engine/Plugins/Experimental/PCG/Source/PCG/Public/Elements/Metadata/PCGMetadataRenameElement.h
--- a/Engine/Plugins/Experimental/PCG/Source/PCG/Public/Elements/Metadata/PCGMetadataRenameElement.h
+++ b/Engine/Plugins/Experimental/PCG/Source/PCG/Public/Elements/Metadata/PCGMetadataRenameElement.h
@@ -17,11 +17,12 @@ public:
 	virtual FName GetDefaultNodeName() const override { return FName(TEXT("AttributeRename")); }
 	virtual FText GetDefaultNodeTitle() const override { return NSLOCTEXT("PCGMetadataRenameSettings", "NodeTitle", "Attribute Rename"); }
 	virtual EPCGSettingsType GetType() const override { return EPCGSettingsType::Metadata; }
+	virtual bool HasDynamicPins() const override { return true; }
 #endif
 
-	virtual TArray<FPCGPinProperties> OutputPinProperties() const override;
-
 protected:
+	virtual TArray<FPCGPinProperties> InputPinProperties() const override;
+	virtual TArray<FPCGPinProperties> OutputPinProperties() const override;
 	virtual FPCGElementPtr CreateElement() const override;
 	//~End UPCGSettings interface
 
diff --git a/Engine/Plugins/Experimental/PCG/Source/PCG/Public/Elements/PCGAttributeFilterElement.h b/Engine/Plugins/Experimental/PCG/Source/PCG/Public/Elements/PCGAttributeFilterElement.h
--- a/Engine/Plugins/Experimental/PCG/Source/PCG/Public/Elements/PCGAttributeFilterElement.h
+++ b/Engine/Plugins/Experimental/PCG/Source/PCG/Public/Elements/PCGAttributeFilterElement.h
@@ -35,12 +35,18 @@ public:
 	virtual FName GetDefaultNodeName() const override;
 	virtual FText GetDefaultNodeTitle() const override;
 	virtual EPCGSettingsType GetType() const override { return EPCGSettingsType::Metadata; }
+	virtual bool HasDynamicPins() const override { return true; }
 #endif
 
-	virtual TArray<FPCGPinProperties> OutputPinProperties() const override;
 	virtual FName AdditionalTaskName() const override;
+
+protected:
+	virtual TArray<FPCGPinProperties> InputPinProperties() const override;
+	virtual TArray<FPCGPinProperties> OutputPinProperties() const override;
+	virtual FPCGElementPtr CreateElement() const override;
 	//~End UPCGSettings interface
 
+public:
 	UPROPERTY(BlueprintReadWrite, EditAnywhere, Category = Settings)
 	EPCGAttributeFilterOperation Operation = EPCGAttributeFilterOperation::KeepSelectedAttributes;
 
@@ -51,9 +57,6 @@ public:
 	UPROPERTY()
 	TArray<FName> AttributesToKeep_DEPRECATED;
 #endif // WITH_EDITORONLY_DATA
-
-protected:
-	virtual FPCGElementPtr CreateElement() const override;
 };
 
 
diff --git a/Engine/Plugins/Experimental/PCG/Source/PCG/Public/Elements/PCGAttributeTransferElement.h b/Engine/Plugins/Experimental/PCG/Source/PCG/Public/Elements/PCGAttributeTransferElement.h
--- a/Engine/Plugins/Experimental/PCG/Source/PCG/Public/Elements/PCGAttributeTransferElement.h
+++ b/Engine/Plugins/Experimental/PCG/Source/PCG/Public/Elements/PCGAttributeTransferElement.h
@@ -39,20 +39,26 @@ public:
 	virtual FName GetDefaultNodeName() const override;
 	virtual FText GetDefaultNodeTitle() const override;
 	virtual EPCGSettingsType GetType() const override { return EPCGSettingsType::Metadata; }
+	virtual FText GetNodeTooltipText() const override;
+	virtual bool HasDynamicPins() const override { return true; }
 #endif
 	virtual FName AdditionalTaskName() const override;
 
+protected:
 	virtual TArray<FPCGPinProperties> InputPinProperties() const override;
+	virtual TArray<FPCGPinProperties> OutputPinProperties() const override;
+	virtual FPCGElementPtr CreateElement() const override;
 	//~End UPCGSettings interface
 
+	/** Get allowed types for all pins */
+	EPCGDataType GetPinsType() const;
+
+public:
 	UPROPERTY(BlueprintReadWrite, EditAnywhere, Category = Settings)
 	FName SourceAttributeName;
 
 	UPROPERTY(BlueprintReadWrite, EditAnywhere, Category = Settings)
 	FName TargetAttributeName;
-
-protected:
-	virtual FPCGElementPtr CreateElement() const override;
 };
 
 
diff --git a/Engine/Plugins/Experimental/PCG/Source/PCG/Public/Elements/PCGFilterByTag.h b/Engine/Plugins/Experimental/PCG/Source/PCG/Public/Elements/PCGFilterByTag.h
--- a/Engine/Plugins/Experimental/PCG/Source/PCG/Public/Elements/PCGFilterByTag.h
+++ b/Engine/Plugins/Experimental/PCG/Source/PCG/Public/Elements/PCGFilterByTag.h
@@ -26,11 +26,13 @@ public:
 	virtual FText GetDefaultNodeTitle() const override;
 	virtual FText GetNodeTooltipText() const override;
 	virtual EPCGSettingsType GetType() const override { return EPCGSettingsType::Filter; }
+	virtual bool HasDynamicPins() const override { return true; }
 #endif
 
 	virtual FName AdditionalTaskName() const override;
 
 protected:
+	virtual TArray<FPCGPinProperties> InputPinProperties() const override;
 	virtual TArray<FPCGPinProperties> OutputPinProperties() const override;
 	virtual FPCGElementPtr CreateElement() const override;
 	//~End UPCGSettings interface
@@ -48,4 +50,4 @@ class FPCGFilterByTagElement : public FSimplePCGElement
 {
 protected:
 	virtual bool ExecuteInternal(FPCGContext* Context) const override;
-};
\ No newline at end of file
+};
