[PCG] Enforce 'single data' option on pin - only take first data item for such pins

diff --git a/Engine/Plugins/Experimental/PCG/Source/PCG/Private/Elements/PCGProjectionElement.cpp b/Engine/Plugins/Experimental/PCG/Source/PCG/Private/Elements/PCGProjectionElement.cpp
--- a/Engine/Plugins/Experimental/PCG/Source/PCG/Private/Elements/PCGProjectionElement.cpp
+++ b/Engine/Plugins/Experimental/PCG/Source/PCG/Private/Elements/PCGProjectionElement.cpp
@@ -46,7 +46,7 @@ bool FPCGProjectionElement::ExecuteInternal(FPCGContext* Context) const
 	TArray<FPCGTaggedData> Targets = Context->InputData.GetInputsByPin(PCGProjectionConstants::ProjectionTargetLabel);
 
 	// If there are no sources or no targets, then nothing to do.
-	if (Sources.Num() == 0 || Targets.Num() != 1)
+	if (Sources.Num() == 0 || Targets.Num() == 0)
 	{
 		return true;
 	}
diff --git a/Engine/Plugins/Experimental/PCG/Source/PCG/Private/Graph/PCGGraphExecutor.cpp b/Engine/Plugins/Experimental/PCG/Source/PCG/Private/Graph/PCGGraphExecutor.cpp
--- a/Engine/Plugins/Experimental/PCG/Source/PCG/Private/Graph/PCGGraphExecutor.cpp
+++ b/Engine/Plugins/Experimental/PCG/Source/PCG/Private/Graph/PCGGraphExecutor.cpp
@@ -1008,6 +1008,15 @@ void FPCGGraphExecutor::BuildTaskInput(const FPCGGraphTask& Task, FPCGDataCollec
 {
 	TRACE_CPUPROFILER_EVENT_SCOPE(FPCGGraphExecutor::BuildTaskInput);
 
+	auto LogDiscardedData = [&Task](const UPCGPin* InPin)
+	{
+		// Log only - currently context has not yet been allocated when this is called
+		UE_LOG(LogPCG, Warning, TEXT("[%s] %s - BuildTaskInput - too many data items arriving on single data pin '%s', only first data item will be used"),
+			(Task.SourceComponent.Get() && Task.SourceComponent->GetOwner()) ? *Task.SourceComponent->GetOwner()->GetName() : TEXT("MissingComponent"),
+			Task.Node ? *Task.Node->GetNodeTitle().ToString() : TEXT("MissingNode"),
+			InPin ? *InPin->Properties.Label.ToString() : TEXT("MissingPin"));
+	};
+
 	// Initialize a Crc onto which each input Crc will be combined.
 	FPCGCrc Crc(Task.Inputs.Num());
 
@@ -1021,6 +1030,15 @@ void FPCGGraphExecutor::BuildTaskInput(const FPCGGraphTask& Task, FPCGDataCollec
 			continue;
 		}
 
+		const bool bAllowMultipleData = Input.OutPin ? Input.OutPin->Properties.bAllowMultipleData : true;
+
+		// Enforce single data - if already have input for this pin, don't add more. Early check before other side effects below.
+		if (Input.OutPin && !bAllowMultipleData && TaskInput.GetInputCountByPin(Input.OutPin->Properties.Label) > 0)
+		{
+			LogDiscardedData(Input.OutPin);
+			continue;
+		}
+
 		const FPCGDataCollection& InputCollection = OutputData[Input.TaskId];
 
 		TaskInput.bCancelExecution |= InputCollection.bCancelExecution;
@@ -1029,7 +1047,17 @@ void FPCGGraphExecutor::BuildTaskInput(const FPCGGraphTask& Task, FPCGDataCollec
 		const int32 TaggedDataOffset = TaskInput.TaggedData.Num();
 		if (Input.InPin)
 		{
-			TaskInput.TaggedData.Append(InputCollection.GetInputsByPin(Input.InPin->Properties.Label));
+			// Proceed carefully when adding data items - if pin is single-data, only add first item.
+			TArray<FPCGTaggedData> InputsOnPin = InputCollection.GetInputsByPin(Input.InPin->Properties.Label);
+			if (InputsOnPin.Num() > 1 && !bAllowMultipleData)
+			{
+				LogDiscardedData(Input.OutPin);
+				TaskInput.TaggedData.Add(InputsOnPin[0]);
+			}
+			else
+			{
+				TaskInput.TaggedData.Append(InputsOnPin);
+			}
 
 			// Write input pin name Crc to uniquely identify inputs per-pin.
 			Crc.Combine(GetTypeHash(Input.OutPin ? Input.OutPin->Properties.Label : FName(TEXT("MissingLabel"))));
diff --git a/Engine/Plugins/Experimental/PCG/Source/PCG/Private/PCGData.cpp b/Engine/Plugins/Experimental/PCG/Source/PCG/Private/PCGData.cpp
--- a/Engine/Plugins/Experimental/PCG/Source/PCG/Private/PCGData.cpp
+++ b/Engine/Plugins/Experimental/PCG/Source/PCG/Private/PCGData.cpp
@@ -153,6 +153,19 @@ TArray<FPCGTaggedData> FPCGDataCollection::GetInputsByPin(const FName& InPinLabe
 		});
 }
 
+int32 FPCGDataCollection::GetInputCountByPin(const FName& InPinLabel) const
+{
+	int32 Count = 0;
+	for (const FPCGTaggedData& Data : TaggedData)
+	{
+		if (Data.Pin == InPinLabel)
+		{
+			++Count;
+		}
+	}
+	return Count;
+}
+
 const UPCGSpatialData* FPCGDataCollection::GetSpatialUnionOfInputsByPin(const FName& InPinLabel, bool& bOutUnionDataCreated) const
 {
 	TArray<FPCGTaggedData> SpatialDataOnPin = TaggedData.FilterByPredicate([&InPinLabel](const FPCGTaggedData& Data) {
diff --git a/Engine/Plugins/Experimental/PCG/Source/PCG/Public/PCGData.h b/Engine/Plugins/Experimental/PCG/Source/PCG/Public/PCGData.h
--- a/Engine/Plugins/Experimental/PCG/Source/PCG/Public/PCGData.h
+++ b/Engine/Plugins/Experimental/PCG/Source/PCG/Public/PCGData.h
@@ -102,6 +102,8 @@ struct PCG_API FPCGDataCollection
 	TArray<FPCGTaggedData> GetInputs() const;
 	/** Returns all data on a given pin */
 	TArray<FPCGTaggedData> GetInputsByPin(const FName& InPinLabel) const;
+	/** Gets number of data items on a given pin */
+	int32 GetInputCountByPin(const FName& InPinLabel) const;
 	/** Returns spatial union of all data on a given pin, returns null if no such data exists. bOutUnionDataCreated indicates if new data created that may need rooting. */
 	const UPCGSpatialData* GetSpatialUnionOfInputsByPin(const FName& InPinLabel, bool& bOutUnionDataCreated) const;
 	/** Returns all spatial data in the collection with the given tag */
