[PCG] Rework pin types and pin/edge change broadcast. With this change we're able to discern static pin type (widest type) vs current dynamic pin type, which is needed for detecting if conversion nodes are required.

Pin properties are treated as static (makes logical sense with how they're serialized/edited in BPs by user/etc). Dynamic pin types do not touch data. Dynamic typing removes state that needs to be managed.

This also separates out dynamic vs static logic/concerns in the node implementations, and allows easier sharing (for e.g. Gather node got much simpler, CreateAttribute is also a more compact expression of the behaviour).

Node changes were previously broadcast if pin properties change, which no longer happens when not committing dynamic type to pin properties data. I also found the node change notifications are triggered in an ad hoc way and were not triggered in many cases, for example on connected nodes when edges are changed. This change adds a common, robust mechanism - edge operations return set of affected nodes, which we then broadcast change notifications on.

diff --git a/Engine/Plugins/Experimental/PCG/Source/PCG/Private/Elements/Metadata/PCGMetadataOpElementBase.cpp b/Engine/Plugins/Experimental/PCG/Source/PCG/Private/Elements/Metadata/PCGMetadataOpElementBase.cpp
--- a/Engine/Plugins/Experimental/PCG/Source/PCG/Private/Elements/Metadata/PCGMetadataOpElementBase.cpp
+++ b/Engine/Plugins/Experimental/PCG/Source/PCG/Private/Elements/Metadata/PCGMetadataOpElementBase.cpp
@@ -28,23 +28,38 @@ void UPCGMetadataSettingsBase::PostLoad()
 #endif // WITH_EDITOR
 }
 
-TArray<FPCGPinProperties> UPCGMetadataSettingsBase::InputPinProperties() const
+EPCGDataType UPCGMetadataSettingsBase::GetCurrentPinTypes(const UPCGPin* InPin) const
 {
-	TArray<FPCGPinProperties> PinProperties;
+	if (GetInputPinNum() == 0)
+	{
+		return EPCGDataType::Any;
+	}
 
-	EPCGDataType FirstPinInputTypeUnion = EPCGDataType::Any;
-	if (GetInputPinNum() > 0)
+	check(InPin);
+	const FName FirstPinLabel = GetInputPinLabel(0);
+	if (InPin->Properties.Label == FirstPinLabel || InPin->IsOutputPin())
 	{
-		const EPCGDataType IncidentTypeUnion = GetTypeUnionOfIncidentEdges(GetInputPinLabel(0));
-		FirstPinInputTypeUnion = (IncidentTypeUnion != EPCGDataType::None) ? IncidentTypeUnion : FirstPinInputTypeUnion;
+		// First input pin, and any output pins, narrow to types connected to first pin
+		const EPCGDataType FirstPinTypeUnion = GetTypeUnionOfIncidentEdges(FirstPinLabel);
+		return FirstPinTypeUnion != EPCGDataType::None ? FirstPinTypeUnion : EPCGDataType::Any;
 	}
+	else
+	{
+		// Other input pins stay at Any
+		return EPCGDataType::Any;
+	}
+}
+
+TArray<FPCGPinProperties> UPCGMetadataSettingsBase::InputPinProperties() const
+{
+	TArray<FPCGPinProperties> PinProperties;
 
 	for (uint32 i = 0; i < GetInputPinNum(); ++i)
 	{
 		const FName PinLabel = GetInputPinLabel(i);
 		if (PinLabel != NAME_None)
 		{
-			PinProperties.Emplace(PinLabel, (i == 0) ? FirstPinInputTypeUnion : EPCGDataType::Any, /*bAllowMultipleConnections=*/false);
+			PinProperties.Emplace(PinLabel, EPCGDataType::Any, /*bAllowMultipleConnections=*/false);
 		}
 	}
 
@@ -55,19 +70,12 @@ TArray<FPCGPinProperties> UPCGMetadataSettingsBase::OutputPinProperties() const
 {
 	TArray<FPCGPinProperties> PinProperties;
 	
-	EPCGDataType OutputType = EPCGDataType::Any;
-	if (GetInputPinNum() > 0)
-	{
-		const EPCGDataType IncidentTypeUnion = GetTypeUnionOfIncidentEdges(GetInputPinLabel(0));
-		OutputType = (IncidentTypeUnion != EPCGDataType::None) ? IncidentTypeUnion : OutputType;
-	}
-
 	for (uint32 i = 0; i < GetOutputPinNum(); ++i)
 	{
 		const FName PinLabel = GetOutputPinLabel(i);
 		if (PinLabel != NAME_None)
 		{
-			PinProperties.Emplace(PinLabel, OutputType);
+			PinProperties.Emplace(PinLabel, EPCGDataType::Any);
 		}
 	}
 
diff --git a/Engine/Plugins/Experimental/PCG/Source/PCG/Private/Elements/Metadata/PCGMetadataRenameElement.cpp b/Engine/Plugins/Experimental/PCG/Source/PCG/Private/Elements/Metadata/PCGMetadataRenameElement.cpp
--- a/Engine/Plugins/Experimental/PCG/Source/PCG/Private/Elements/Metadata/PCGMetadataRenameElement.cpp
+++ b/Engine/Plugins/Experimental/PCG/Source/PCG/Private/Elements/Metadata/PCGMetadataRenameElement.cpp
@@ -12,22 +12,25 @@
 
 #define LOCTEXT_NAMESPACE "PCGMetadataRenameElement"
 
-TArray<FPCGPinProperties> UPCGMetadataRenameSettings::InputPinProperties() const
+EPCGDataType UPCGMetadataRenameSettings::GetCurrentPinTypes(const UPCGPin* InPin) const
 {
-	EPCGDataType InputTypeUnion = GetTypeUnionOfIncidentEdges(PCGPinConstants::DefaultInputLabel);
+	// All pins narrow to union of input types
+	const EPCGDataType InputTypeUnion = GetTypeUnionOfIncidentEdges(PCGPinConstants::DefaultInputLabel);
+	return (InputTypeUnion != EPCGDataType::None) ? InputTypeUnion : EPCGDataType::Any;
+}
 
+TArray<FPCGPinProperties> UPCGMetadataRenameSettings::InputPinProperties() const
+{
 	TArray<FPCGPinProperties> PinProperties;
-	PinProperties.Emplace(PCGPinConstants::DefaultInputLabel, (InputTypeUnion != EPCGDataType::None) ? InputTypeUnion : EPCGDataType::Any);
+	PinProperties.Emplace(PCGPinConstants::DefaultInputLabel, EPCGDataType::Any);
 
 	return PinProperties;
 }
 
 TArray<FPCGPinProperties> UPCGMetadataRenameSettings::OutputPinProperties() const
 {
-	const EPCGDataType InputTypeUnion = GetTypeUnionOfIncidentEdges(PCGPinConstants::DefaultInputLabel);
-
 	TArray<FPCGPinProperties> PinProperties;
-	PinProperties.Emplace(PCGPinConstants::DefaultOutputLabel, (InputTypeUnion != EPCGDataType::None) ? InputTypeUnion : EPCGDataType::Any);
+	PinProperties.Emplace(PCGPinConstants::DefaultOutputLabel, EPCGDataType::Any);
 
 	return PinProperties;
 }
diff --git a/Engine/Plugins/Experimental/PCG/Source/PCG/Private/Elements/PCGAttributeFilterElement.cpp b/Engine/Plugins/Experimental/PCG/Source/PCG/Private/Elements/PCGAttributeFilterElement.cpp
--- a/Engine/Plugins/Experimental/PCG/Source/PCG/Private/Elements/PCGAttributeFilterElement.cpp
+++ b/Engine/Plugins/Experimental/PCG/Source/PCG/Private/Elements/PCGAttributeFilterElement.cpp
@@ -55,6 +55,13 @@ FText UPCGAttributeFilterSettings::GetDefaultNodeTitle() const
 }
 #endif
 
+EPCGDataType UPCGAttributeFilterSettings::GetCurrentPinTypes(const UPCGPin* InPin) const
+{
+	// All pins narrow to union of input types
+	const EPCGDataType InputTypeUnion = GetTypeUnionOfIncidentEdges(PCGPinConstants::DefaultInputLabel);
+	return (InputTypeUnion != EPCGDataType::None) ? InputTypeUnion : EPCGDataType::Any;
+}
+
 FName UPCGAttributeFilterSettings::AdditionalTaskName() const
 {
 	TArray<FString> AttributesToKeep = PCGHelpers::GetStringArrayFromCommaSeparatedString(SelectedAttributes);
@@ -84,20 +91,16 @@ FName UPCGAttributeFilterSettings::AdditionalTaskName() const
 
 TArray<FPCGPinProperties> UPCGAttributeFilterSettings::InputPinProperties() const
 {
-	const EPCGDataType InputTypeUnion = GetTypeUnionOfIncidentEdges(PCGPinConstants::DefaultInputLabel);
-
 	TArray<FPCGPinProperties> PinProperties;
-	PinProperties.Emplace(PCGPinConstants::DefaultInputLabel, (InputTypeUnion != EPCGDataType::None) ? InputTypeUnion : EPCGDataType::Any);
+	PinProperties.Emplace(PCGPinConstants::DefaultInputLabel, EPCGDataType::Any);
 
 	return PinProperties;
 }
 
 TArray<FPCGPinProperties> UPCGAttributeFilterSettings::OutputPinProperties() const
 {
-	const EPCGDataType InputTypeUnion = GetTypeUnionOfIncidentEdges(PCGPinConstants::DefaultInputLabel);
-
 	TArray<FPCGPinProperties> PinProperties;
-	PinProperties.Emplace(PCGPinConstants::DefaultOutputLabel, (InputTypeUnion != EPCGDataType::None) ? InputTypeUnion : EPCGDataType::Any);
+	PinProperties.Emplace(PCGPinConstants::DefaultOutputLabel, EPCGDataType::Any);
 
 	return PinProperties;
 }
diff --git a/Engine/Plugins/Experimental/PCG/Source/PCG/Private/Elements/PCGAttributeTransferElement.cpp b/Engine/Plugins/Experimental/PCG/Source/PCG/Private/Elements/PCGAttributeTransferElement.cpp
--- a/Engine/Plugins/Experimental/PCG/Source/PCG/Private/Elements/PCGAttributeTransferElement.cpp
+++ b/Engine/Plugins/Experimental/PCG/Source/PCG/Private/Elements/PCGAttributeTransferElement.cpp
@@ -33,6 +33,15 @@ FText UPCGAttributeTransferSettings::GetNodeTooltipText() const
 }
 #endif
 
+EPCGDataType UPCGAttributeTransferSettings::GetCurrentPinTypes(const UPCGPin* InPin) const
+{
+	// All pins narrow to same type, which is Point if any input is Point, otherwise Spatial
+	const bool bAnyArePoint = (GetTypeUnionOfIncidentEdges(PCGAttributeTransferConstants::TargetLabel) == EPCGDataType::Point)
+		|| (GetTypeUnionOfIncidentEdges(PCGAttributeTransferConstants::SourceLabel) == EPCGDataType::Point);
+
+	return bAnyArePoint ? EPCGDataType::Point : EPCGDataType::Spatial;
+}
+
 FName UPCGAttributeTransferSettings::AdditionalTaskName() const
 {
 	if (SourceAttributeName != TargetAttributeName && TargetAttributeName != NAME_None)
@@ -47,11 +56,9 @@ FName UPCGAttributeTransferSettings::AdditionalTaskName() const
 
 TArray<FPCGPinProperties> UPCGAttributeTransferSettings::InputPinProperties() const
 {
-	const EPCGDataType InputType = GetPinsType();
-
 	TArray<FPCGPinProperties> PinProperties;
-	PinProperties.Emplace(PCGAttributeTransferConstants::TargetLabel, InputType, /*bInAllowMultipleConnections=*/ false);
-	PinProperties.Emplace(PCGAttributeTransferConstants::SourceLabel, InputType, /*bInAllowMultipleConnections=*/ false);
+	PinProperties.Emplace(PCGAttributeTransferConstants::TargetLabel, EPCGDataType::Spatial, /*bInAllowMultipleConnections=*/ false);
+	PinProperties.Emplace(PCGAttributeTransferConstants::SourceLabel, EPCGDataType::Spatial, /*bInAllowMultipleConnections=*/ false);
 
 	return PinProperties;
 }
@@ -59,7 +66,7 @@ TArray<FPCGPinProperties> UPCGAttributeTransferSettings::InputPinProperties() co
 TArray<FPCGPinProperties> UPCGAttributeTransferSettings::OutputPinProperties() const
 {
 	TArray<FPCGPinProperties> PinProperties;
-	PinProperties.Emplace(PCGPinConstants::DefaultOutputLabel, GetPinsType());
+	PinProperties.Emplace(PCGPinConstants::DefaultOutputLabel, EPCGDataType::Spatial);
 
 	return PinProperties;
 }
@@ -69,15 +76,6 @@ FPCGElementPtr UPCGAttributeTransferSettings::CreateElement() const
 	return MakeShared<FPCGAttributeTransferElement>();
 }
 
-EPCGDataType UPCGAttributeTransferSettings::GetPinsType() const
-{
-	// If either input is type Point then all pins are Point, otherwise Spatial
-	const bool bAnyArePoint = GetTypeUnionOfIncidentEdges(PCGAttributeTransferConstants::TargetLabel) == EPCGDataType::Point
-		|| GetTypeUnionOfIncidentEdges(PCGAttributeTransferConstants::SourceLabel) == EPCGDataType::Point;
-
-	return bAnyArePoint ? EPCGDataType::Point : EPCGDataType::Spatial;
-}
-
 bool FPCGAttributeTransferElement::ExecuteInternal(FPCGContext* Context) const
 {
 	TRACE_CPUPROFILER_EVENT_SCOPE(FPCGAttributeTransferElement::Execute);
diff --git a/Engine/Plugins/Experimental/PCG/Source/PCG/Private/Elements/PCGCreateAttribute.cpp b/Engine/Plugins/Experimental/PCG/Source/PCG/Private/Elements/PCGCreateAttribute.cpp
--- a/Engine/Plugins/Experimental/PCG/Source/PCG/Private/Elements/PCGCreateAttribute.cpp
+++ b/Engine/Plugins/Experimental/PCG/Source/PCG/Private/Elements/PCGCreateAttribute.cpp
@@ -49,6 +49,23 @@ void UPCGCreateAttributeSettings::PostLoad()
 #endif // WITH_EDITOR
 }
 
+EPCGDataType UPCGCreateAttributeSettings::GetCurrentPinTypes(const UPCGPin* InPin) const
+{
+	if (!HasDynamicPins())
+	{
+		return Super::GetCurrentPinTypes(InPin);
+	}
+
+	check(InPin);
+	if (InPin->Properties.Label == PCGPinConstants::DefaultInputLabel || InPin->Properties.Label == PCGPinConstants::DefaultOutputLabel)
+	{
+		const EPCGDataType PrimaryInputType = GetTypeUnionOfIncidentEdges(PCGPinConstants::DefaultInputLabel);
+		return (PrimaryInputType != EPCGDataType::None) ? PrimaryInputType : EPCGDataType::Any;
+	}
+
+	return Super::GetCurrentPinTypes(InPin);
+}
+
 FName UPCGCreateAttributeSettings::AdditionalTaskName() const
 {
 	return AdditionalTaskNameInternal(PCGCreateAttributeConstants::NodeNameAddAttribute);
@@ -69,16 +86,7 @@ FText UPCGCreateAttributeSettings::GetDefaultNodeTitle() const
 TArray<FPCGPinProperties> UPCGCreateAttributeSettings::InputPinProperties() const
 {
 	TArray<FPCGPinProperties> PinProperties;
-
-	EPCGDataType PrimaryInputType = GetTypeUnionOfIncidentEdges(PCGPinConstants::DefaultInputLabel);
-	if (PrimaryInputType == EPCGDataType::None)
-	{
-		// Fall back to Any if no edges, or no overlap in types (which should not normally occur).
-		PrimaryInputType = EPCGDataType::Any;
-	}
-
-	PinProperties.Emplace(PCGPinConstants::DefaultInputLabel, PrimaryInputType);
-
+	PinProperties.Emplace(PCGPinConstants::DefaultInputLabel, EPCGDataType::Any);
 	if (bFromSourceParam)
 	{
 		PinProperties.Emplace(PCGCreateAttributeConstants::AttributesLabel, EPCGDataType::Param);
@@ -91,13 +99,7 @@ TArray<FPCGPinProperties> UPCGCreateAttributeSettings::OutputPinProperties() con
 {
 	FPCGPinProperties PinProperties;
 	PinProperties.Label = PCGPinConstants::DefaultOutputLabel;
-
-	PinProperties.AllowedTypes = GetTypeUnionOfIncidentEdges(PCGPinConstants::DefaultInputLabel);
-	if (PinProperties.AllowedTypes == EPCGDataType::None)
-	{
-		// Nothing connected or incompatible types - output Param.
-		PinProperties.AllowedTypes = EPCGDataType::Param;
-	}
+	PinProperties.AllowedTypes = EPCGDataType::Any;
 
 	return { PinProperties };
 }
@@ -145,13 +147,17 @@ FText UPCGCreateAttributeSetSettings::GetDefaultNodeTitle() const
 }
 #endif // WITH_EDITOR
 
+TArray<FPCGPinProperties> UPCGCreateAttributeSetSettings::InputPinProperties() const
+{
+	return TArray<FPCGPinProperties>();
+}
+
 TArray<FPCGPinProperties> UPCGCreateAttributeSetSettings::OutputPinProperties() const
 {
-	FPCGPinProperties PinProperties;
-	PinProperties.Label = PCGPinConstants::DefaultOutputLabel;
-	PinProperties.AllowedTypes = EPCGDataType::Param;
+	TArray<FPCGPinProperties> PinProperties;
+	PinProperties.Emplace(PCGPinConstants::DefaultOutputLabel, EPCGDataType::Param);
 
-	return { PinProperties };
+	return PinProperties;
 }
 
 FName UPCGCreateAttributeSetSettings::AdditionalTaskName() const
diff --git a/Engine/Plugins/Experimental/PCG/Source/PCG/Private/Elements/PCGFilterByTag.cpp b/Engine/Plugins/Experimental/PCG/Source/PCG/Private/Elements/PCGFilterByTag.cpp
--- a/Engine/Plugins/Experimental/PCG/Source/PCG/Private/Elements/PCGFilterByTag.cpp
+++ b/Engine/Plugins/Experimental/PCG/Source/PCG/Private/Elements/PCGFilterByTag.cpp
@@ -32,6 +32,13 @@ FText UPCGFilterByTagSettings::GetNodeTooltipText() const
 }
 #endif
 
+EPCGDataType UPCGFilterByTagSettings::GetCurrentPinTypes(const UPCGPin* InPin) const
+{
+	// All pins narrow to input type
+	const EPCGDataType InputTypeUnion = GetTypeUnionOfIncidentEdges(PCGPinConstants::DefaultInputLabel);
+	return InputTypeUnion != EPCGDataType::None ? InputTypeUnion : EPCGDataType::Any;
+}
+
 FName UPCGFilterByTagSettings::AdditionalTaskName() const
 {
 	const TArray<FString> Tags = PCGHelpers::GetStringArrayFromCommaSeparatedString(SelectedTags);
@@ -51,20 +58,16 @@ FName UPCGFilterByTagSettings::AdditionalTaskName() const
 
 TArray<FPCGPinProperties> UPCGFilterByTagSettings::InputPinProperties() const
 {
-	const EPCGDataType InputTypeUnion = GetTypeUnionOfIncidentEdges(PCGPinConstants::DefaultInputLabel);
-
 	TArray<FPCGPinProperties> PinProperties;
-	PinProperties.Emplace(PCGPinConstants::DefaultInputLabel, (InputTypeUnion != EPCGDataType::None) ? InputTypeUnion : EPCGDataType::Any);
+	PinProperties.Emplace(PCGPinConstants::DefaultInputLabel, EPCGDataType::Any);
 
 	return PinProperties;
 }
 
 TArray<FPCGPinProperties> UPCGFilterByTagSettings::OutputPinProperties() const
 {
-	const EPCGDataType InputTypeUnion = GetTypeUnionOfIncidentEdges(PCGPinConstants::DefaultInputLabel);
-
 	TArray<FPCGPinProperties> PinProperties;
-	PinProperties.Emplace(PCGPinConstants::DefaultOutputLabel, (InputTypeUnion != EPCGDataType::None) ? InputTypeUnion : EPCGDataType::Any);
+	PinProperties.Emplace(PCGPinConstants::DefaultOutputLabel, EPCGDataType::Any);
 
 	return PinProperties;
 }
diff --git a/Engine/Plugins/Experimental/PCG/Source/PCG/Private/Elements/PCGGather.cpp b/Engine/Plugins/Experimental/PCG/Source/PCG/Private/Elements/PCGGather.cpp
--- a/Engine/Plugins/Experimental/PCG/Source/PCG/Private/Elements/PCGGather.cpp
+++ b/Engine/Plugins/Experimental/PCG/Source/PCG/Private/Elements/PCGGather.cpp
@@ -5,30 +5,25 @@
 #include "PCGContext.h"
 #include "PCGPin.h"
 	
-TArray<FPCGPinProperties> UPCGGatherSettings::InputPinProperties() const
+EPCGDataType UPCGGatherSettings::GetCurrentPinTypes(const UPCGPin* InPin) const
 {
+	// All pins narrow to input type
 	EPCGDataType InputTypeUnion = GetTypeUnionOfIncidentEdges(PCGPinConstants::DefaultInputLabel);
-	if (InputTypeUnion == EPCGDataType::None)
-	{
-		InputTypeUnion = EPCGDataType::Any;
-	}
+	return InputTypeUnion != EPCGDataType::None ? InputTypeUnion : EPCGDataType::Any;
+}
 
+TArray<FPCGPinProperties> UPCGGatherSettings::InputPinProperties() const
+{
 	TArray<FPCGPinProperties> PinProperties;
-	PinProperties.Emplace(PCGPinConstants::DefaultInputLabel, InputTypeUnion);
+	PinProperties.Emplace(PCGPinConstants::DefaultInputLabel, EPCGDataType::Any);
 
 	return PinProperties;
 }
 
 TArray<FPCGPinProperties> UPCGGatherSettings::OutputPinProperties() const
 {
-	EPCGDataType InputTypeUnion = GetTypeUnionOfIncidentEdges(PCGPinConstants::DefaultInputLabel);
-	if (InputTypeUnion == EPCGDataType::None)
-	{
-		InputTypeUnion = EPCGDataType::Any;
-	}
-
 	TArray<FPCGPinProperties> PinProperties;
-	PinProperties.Emplace(PCGPinConstants::DefaultOutputLabel, InputTypeUnion);
+	PinProperties.Emplace(PCGPinConstants::DefaultOutputLabel, EPCGDataType::Any);
 
 	return PinProperties;
 }
diff --git a/Engine/Plugins/Experimental/PCG/Source/PCG/Private/Elements/PCGReroute.cpp b/Engine/Plugins/Experimental/PCG/Source/PCG/Private/Elements/PCGReroute.cpp
--- a/Engine/Plugins/Experimental/PCG/Source/PCG/Private/Elements/PCGReroute.cpp
+++ b/Engine/Plugins/Experimental/PCG/Source/PCG/Private/Elements/PCGReroute.cpp
@@ -14,13 +14,10 @@ UPCGRerouteSettings::UPCGRerouteSettings()
 #endif
 }
 
-TArray<FPCGPinProperties> UPCGRerouteSettings::InputPinProperties() const
+EPCGDataType UPCGRerouteSettings::GetCurrentPinTypes(const UPCGPin* InPin) const
 {
-	FPCGPinProperties PinProperties;
-	PinProperties.Label = PCGPinConstants::DefaultInputLabel;
-	PinProperties.bAllowMultipleConnections = false;
-	PinProperties.AllowedTypes = EPCGDataType::Any;
-	
+	// All pins have same type
+	EPCGDataType PinTypes = EPCGDataType::Any;
 	if (UPCGNode* PCGNode = Cast<UPCGNode>(GetOuter()))
 	{
 		if (UPCGPin* InputPin = PCGNode->GetInputPin(PCGPinConstants::DefaultInputLabel))
@@ -31,39 +28,31 @@ TArray<FPCGPinProperties> UPCGRerouteSettings::InputPinProperties() const
 				{
 					if (const UPCGPin* OtherOutputPin = Edge->GetOtherPin(InputPin))
 					{
-						PinProperties.AllowedTypes = OtherOutputPin->Properties.AllowedTypes;
+						PinTypes = OtherOutputPin->GetCurrentTypes();
 					}
 				}
 			}
 		}
 	}
 
-	return {PinProperties };
+	return PinTypes;
+}
+
+TArray<FPCGPinProperties> UPCGRerouteSettings::InputPinProperties() const
+{
+	FPCGPinProperties PinProperties;
+	PinProperties.Label = PCGPinConstants::DefaultInputLabel;
+	PinProperties.bAllowMultipleConnections = false;
+	PinProperties.AllowedTypes = EPCGDataType::Any;
+
+	return { PinProperties };
 }
 
 TArray<FPCGPinProperties> UPCGRerouteSettings::OutputPinProperties() const
 {
 	FPCGPinProperties PinProperties;
 	PinProperties.Label = PCGPinConstants::DefaultOutputLabel;
-	PinProperties.bAllowMultipleConnections = true;
 	PinProperties.AllowedTypes = EPCGDataType::Any;
-	
-	if (UPCGNode* PCGNode = Cast<UPCGNode>(GetOuter()))
-	{
-		if (UPCGPin* InputPin = PCGNode->GetInputPin(PCGPinConstants::DefaultInputLabel))
-		{
-			if (InputPin->EdgeCount() > 0)
-			{
-				if (UPCGEdge* Edge = InputPin->Edges[0])
-				{
-					if (const UPCGPin* OtherOutputPin = Edge->GetOtherPin(InputPin))
-					{
-						PinProperties.AllowedTypes = OtherOutputPin->Properties.AllowedTypes;
-					}
-				}
-			}
-		}
-	}
 
 	return { PinProperties };
 }
diff --git a/Engine/Plugins/Experimental/PCG/Source/PCG/Private/PCGElement.cpp b/Engine/Plugins/Experimental/PCG/Source/PCG/Private/PCGElement.cpp
--- a/Engine/Plugins/Experimental/PCG/Source/PCG/Private/PCGElement.cpp
+++ b/Engine/Plugins/Experimental/PCG/Source/PCG/Private/PCGElement.cpp
@@ -247,7 +247,7 @@ void IPCGElement::DisabledPassThroughData(FPCGContext* Context) const
 	// Grab data from pass-through pin
 	Context->OutputData.TaggedData = Context->InputData.GetInputsByPin(PassThroughPin->Properties.Label);
 
-	const EPCGDataType OutputType = Context->Node->GetOutputPins()[0]->Properties.AllowedTypes;
+	const EPCGDataType OutputType = Context->Node->GetOutputPins()[0]->GetCurrentTypes();
 
 	// Pass through input data if it is not params, and if the output type supports it (e.g. if we have a incoming
 	// surface connected to an input pin of type Any, do not pass the surface through to an output pin of type Point).
@@ -364,8 +364,12 @@ void IPCGElement::CleanupAndValidateOutput(FPCGContext* Context) const
 				}
 				else if (TaggedData.Data)
 				{
-					const bool bTypesOverlap = !!(OutputPinProperties[MatchIndex].AllowedTypes & TaggedData.Data->GetDataType());
-					const bool bTypeIsSubset = !(~OutputPinProperties[MatchIndex].AllowedTypes & TaggedData.Data->GetDataType());
+					// Try to get dynamic current pin types, otherwise settle for static types
+					const UPCGPin* OutputPin = Context->Node ? Context->Node->GetOutputPin(OutputPinProperties[MatchIndex].Label) : nullptr;
+					const EPCGDataType PinTypes = OutputPin ? OutputPin->GetCurrentTypes() : OutputPinProperties[MatchIndex].AllowedTypes;
+
+					const bool bTypesOverlap = !!(PinTypes & TaggedData.Data->GetDataType());
+					const bool bTypeIsSubset = !(~PinTypes & TaggedData.Data->GetDataType());
 					// TODO: Temporary fix for Settings directly from InputData (ie. from elements with code and not PCG nodes)
 					if ((!bTypesOverlap || !bTypeIsSubset) && TaggedData.Data->GetDataType() != EPCGDataType::Settings)
 					{
diff --git a/Engine/Plugins/Experimental/PCG/Source/PCG/Private/PCGGraph.cpp b/Engine/Plugins/Experimental/PCG/Source/PCG/Private/PCGGraph.cpp
--- a/Engine/Plugins/Experimental/PCG/Source/PCG/Private/PCGGraph.cpp
+++ b/Engine/Plugins/Experimental/PCG/Source/PCG/Private/PCGGraph.cpp
@@ -66,6 +66,25 @@ namespace PCGGraphUtils
 
 		InSourcePropertyDesc->CachedProperty->CopyCompleteValue(TargetValueAddress, SourceValueAddress);
 	}
+
+	EPCGChangeType NotifyTouchedNodes(const TSet<UPCGNode*>& InTouchedNodes)
+	{
+		EPCGChangeType ChangeType = EPCGChangeType::None;
+
+		for (UPCGNode* TouchedNode : InTouchedNodes)
+		{
+			if (TouchedNode)
+			{
+				ChangeType |= TouchedNode->PropagateDynamicPinTypes();
+
+#if WITH_EDITOR
+				TouchedNode->OnNodeChangedDelegate.Broadcast(TouchedNode, EPCGChangeType::Node);
+#endif
+			}
+		}
+
+		return ChangeType;
+	}
 }
 
 /****************************
@@ -456,22 +475,23 @@ bool UPCGGraph::AddLabeledEdge(UPCGNode* From, const FName& FromPinLabel, UPCGNo
 		return false;
 	}
 
+	TSet<UPCGNode*> TouchedNodes;
+
 	// Create edge
-	FromPin->AddEdgeTo(ToPin);
+	FromPin->AddEdgeTo(ToPin, &TouchedNodes);
 
 	bool bToPinBrokeOtherEdges = false;
 	
 	// Add an edge to a pin that doesn't allow multiple connections requires to do some cleanup
 	if (!ToPin->AllowMultipleConnections())
 	{
-		bToPinBrokeOtherEdges = ToPin->BreakAllIncompatibleEdges();
+		bToPinBrokeOtherEdges = ToPin->BreakAllIncompatibleEdges(&TouchedNodes);
 	}
 	
-	From->UpdateDynamicPins();
-	To->UpdateDynamicPins();
+	const EPCGChangeType ChangeType = PCGGraphUtils::NotifyTouchedNodes(TouchedNodes);
 
 #if WITH_EDITOR
-	NotifyGraphChanged(EPCGChangeType::Structural);
+	NotifyGraphChanged(ChangeType | EPCGChangeType::Structural);
 #endif
 
 	return bToPinBrokeOtherEdges;
@@ -520,17 +540,24 @@ void UPCGGraph::RemoveNode(UPCGNode* InNode)
 	check(InNode);
 
 	Modify();
+	
+	TSet<UPCGNode*> TouchedNodes;
 
 	for (UPCGPin* InputPin : InNode->InputPins)
 	{
-		InputPin->BreakAllEdges();
+		InputPin->BreakAllEdges(&TouchedNodes);
 	}
 
 	for (UPCGPin* OutputPin : InNode->OutputPins)
 	{
-		OutputPin->BreakAllEdges();
+		OutputPin->BreakAllEdges(&TouchedNodes);
 	}
 
+	// We're about to remove InNode, so don't bother triggering updates
+	TouchedNodes.Remove(InNode);
+
+	PCGGraphUtils::NotifyTouchedNodes(TouchedNodes);
+
 	Nodes.Remove(InNode);
 	OnNodeRemoved(InNode);
 }
@@ -546,19 +573,22 @@ bool UPCGGraph::RemoveEdge(UPCGNode* From, const FName& FromLabel, UPCGNode* To,
 	UPCGPin* OutPin = From->GetOutputPin(FromLabel);
 	UPCGPin* InPin = To->GetInputPin(ToLabel);
 
-	const bool bChanged = OutPin && OutPin->BreakEdgeTo(InPin);
-
-	if (bChanged)
+	TSet<UPCGNode*> TouchedNodes;
+	if (OutPin)
 	{
-		From->UpdateDynamicPins();
-		To->UpdateDynamicPins();
+		OutPin->BreakEdgeTo(InPin, &TouchedNodes);
+	}
+
+	const EPCGChangeType ChangeType = PCGGraphUtils::NotifyTouchedNodes(TouchedNodes);
 
+	if (TouchedNodes.Num() > 0)
+	{
 #if WITH_EDITOR
-		NotifyGraphChanged(EPCGChangeType::Structural);
+		NotifyGraphChanged(ChangeType | EPCGChangeType::Structural);
 #endif
 	}
 
-	return bChanged;
+	return TouchedNodes.Num() > 0;
 }
 
 void UPCGGraph::ForEachNode(const TFunction<void(UPCGNode*)>& Action)
@@ -575,62 +605,46 @@ void UPCGGraph::ForEachNode(const TFunction<void(UPCGNode*)>& Action)
 bool UPCGGraph::RemoveInboundEdges(UPCGNode* InNode, const FName& InboundLabel)
 {
 	check(InNode);
-	bool bChanged = false;
+	TSet<UPCGNode*> TouchedNodes;
 
 	if (UPCGPin* InputPin = InNode->GetInputPin(InboundLabel))
 	{
-		bChanged = InputPin->BreakAllEdges();
+		InputPin->BreakAllEdges(&TouchedNodes);
 	}
 
+	const EPCGChangeType ChangeType = PCGGraphUtils::NotifyTouchedNodes(TouchedNodes);
+
 #if WITH_EDITOR
-	if (bChanged)
+	if (TouchedNodes.Num() > 0)
 	{
-		const EPCGChangeType ChangeType = InNode->UpdateDynamicPins() | EPCGChangeType::Structural;
-		NotifyGraphChanged(ChangeType);
+		NotifyGraphChanged(ChangeType | EPCGChangeType::Structural);
 	}
 #endif
 
-	return bChanged;
+	return TouchedNodes.Num() > 0;
 }
 
 bool UPCGGraph::RemoveOutboundEdges(UPCGNode* InNode, const FName& OutboundLabel)
 {
 	check(InNode);
-	bool bChanged = false;
 	// Make a list of downstream nodes which may need pin updates when the edges change
-	TSet<UPCGNode*> ConnectedNodes;
+	TSet<UPCGNode*> TouchedNodes;
 
 	if (UPCGPin* OutputPin = InNode->GetOutputPin(OutboundLabel))
 	{
-		for (const UPCGEdge* Edge : OutputPin->Edges)
-		{
-			if (Edge && Edge->OutputPin && Edge->OutputPin->Node)
-			{
-				ConnectedNodes.Add(Edge->OutputPin->Node);
-			}
-		}
-
-		bChanged = OutputPin->BreakAllEdges();
+		OutputPin->BreakAllEdges(&TouchedNodes);
 	}
 
+	const EPCGChangeType ChangeType = PCGGraphUtils::NotifyTouchedNodes(TouchedNodes);
+
 #if WITH_EDITOR
-	if (bChanged)
+	if (TouchedNodes.Num() > 0)
 	{
-		EPCGChangeType ChangeType = EPCGChangeType::Structural;
-
-		// Pins may need updating on downstream nodes
-		for (UPCGNode* ConnectedNode : ConnectedNodes)
-		{
-			ChangeType |= ConnectedNode->UpdateDynamicPins();
-		}
-
-		ChangeType |= InNode->UpdateDynamicPins();
-
-		NotifyGraphChanged(ChangeType);
+		NotifyGraphChanged(ChangeType | EPCGChangeType::Structural);
 	}
 #endif
 
-	return bChanged;
+	return TouchedNodes.Num() > 0;
 }
 
 #if WITH_EDITOR
diff --git a/Engine/Plugins/Experimental/PCG/Source/PCG/Private/PCGNode.cpp b/Engine/Plugins/Experimental/PCG/Source/PCG/Private/PCGNode.cpp
--- a/Engine/Plugins/Experimental/PCG/Source/PCG/Private/PCGNode.cpp
+++ b/Engine/Plugins/Experimental/PCG/Source/PCG/Private/PCGNode.cpp
@@ -418,7 +418,7 @@ const UPCGPin* UPCGNode::GetPassThroughInputPin() const
 		return nullptr;
 	}
 
-	const EPCGDataType OutputType = GetOutputPins()[0]->Properties.AllowedTypes;
+	const EPCGDataType OutputType = GetOutputPins()[0]->GetCurrentTypes();
 
 	// We assume a node whose primary output is params is a param processing node.
 	const bool bNodeOutputsParams = (OutputType == EPCGDataType::Param);
@@ -426,7 +426,7 @@ const UPCGPin* UPCGNode::GetPassThroughInputPin() const
 	if (bNodeOutputsParams)
 	{
 		// If this node primarily processes params, then look for a params pin, if we find it then we will pass it through
-		auto FindParams = [](const TObjectPtr<UPCGPin>& InPin) { return InPin->Properties.AllowedTypes == EPCGDataType::Param; };
+		auto FindParams = [](const TObjectPtr<UPCGPin>& InPin) { return InPin->GetCurrentTypes() == EPCGDataType::Param; };
 		const TObjectPtr<UPCGPin>* FirstParamsPinPtr = Algo::FindByPredicate(GetInputPins(), FindParams);
 		return FirstParamsPinPtr ? *FirstParamsPinPtr : nullptr;
 	}
@@ -434,14 +434,14 @@ const UPCGPin* UPCGNode::GetPassThroughInputPin() const
 	{
 		// 'Normal' node. Look for the first pin that is not of type Params to pass through. If the node is disabled, the params are unused.
 		// Params-only pins will be rejected/ignored.
-		auto FindNonParams = [](const TObjectPtr<UPCGPin>& InPin) { return InPin->Properties.AllowedTypes != EPCGDataType::Param; };
+		auto FindNonParams = [](const TObjectPtr<UPCGPin>& InPin) { return InPin->GetCurrentTypes() != EPCGDataType::Param; };
 		const TObjectPtr<UPCGPin>* FirstNonParamsPinPtr = Algo::FindByPredicate(GetInputPins(), FindNonParams);
 
 		if (FirstNonParamsPinPtr)
 		{
 			// Finally in order to be a candidate for passing through, data must be compatible.
 			const UPCGPin* InputPin = *FirstNonParamsPinPtr;
-			const EPCGDataType InputType = InputPin->Properties.AllowedTypes;
+			const EPCGDataType InputType = InputPin->GetCurrentTypes();
 			const bool bTypesOverlap = !!(InputType & OutputType);
 			
 			// Misc note - it would be nice if we could be stricter, like line below this comment. However it would mean it will stop an Any
@@ -632,8 +632,9 @@ void UPCGNode::CreateDefaultPins(TFunctionRef<UPCGPin* (UPCGNode*)> PinAllocator
 	CreatePins(OutputPins, Settings->DefaultOutputPinProperties());
 }
 
-EPCGChangeType UPCGNode::UpdatePins(TFunctionRef<UPCGPin*(UPCGNode*)> PinAllocator, const UPCGNode* FromNode /*= nullptr*/)
+EPCGChangeType UPCGNode::UpdatePins(TFunctionRef<UPCGPin*(UPCGNode*)> PinAllocator)
 {
+	TSet<UPCGNode*> TouchedNodes;
 	const UPCGSettings* Settings = GetSettings();
 	
 	if (!Settings)
@@ -650,7 +651,7 @@ EPCGChangeType UPCGNode::UpdatePins(TFunctionRef<UPCGPin*(UPCGNode*)> PinAllocat
 		{
 			if (Pin)
 			{
-				Pin->BreakAllEdges();
+				Pin->BreakAllEdges(&TouchedNodes);
 			}
 		}
 
@@ -658,7 +659,7 @@ EPCGChangeType UPCGNode::UpdatePins(TFunctionRef<UPCGPin*(UPCGNode*)> PinAllocat
 		{
 			if (Pin)
 			{
-				Pin->BreakAllEdges();
+				Pin->BreakAllEdges(&TouchedNodes);
 			}
 		}
 
@@ -670,7 +671,7 @@ EPCGChangeType UPCGNode::UpdatePins(TFunctionRef<UPCGPin*(UPCGNode*)> PinAllocat
 	const TArray<FPCGPinProperties> InboundPinProperties = Settings->AllInputPinProperties();
 	const TArray<FPCGPinProperties> OutboundPinProperties = Settings->AllOutputPinProperties();
 
-	auto UpdatePins = [this, &PinAllocator](TArray<UPCGPin*>& Pins, const TArray<FPCGPinProperties>& PinProperties)
+	auto UpdatePins = [this, &PinAllocator, &TouchedNodes](TArray<UPCGPin*>& Pins, const TArray<FPCGPinProperties>& PinProperties)
 	{
 		bool bAppliedEdgeChanges = false;
 		bool bChangedPins = false;
@@ -685,7 +686,8 @@ EPCGChangeType UPCGNode::UpdatePins(TFunctionRef<UPCGPin*(UPCGNode*)> PinAllocat
 				{
 					Pin->Modify();
 					Pin->Properties = *MatchingProperties;
-					bAppliedEdgeChanges |= Pin->BreakAllIncompatibleEdges();
+
+					bAppliedEdgeChanges |= Pin->BreakAllIncompatibleEdges(&TouchedNodes);
 					bChangedPins = true;
 				}
 			}
@@ -710,7 +712,8 @@ EPCGChangeType UPCGNode::UpdatePins(TFunctionRef<UPCGPin*(UPCGNode*)> PinAllocat
 		{
 			UnmatchedPins[0]->Modify();
 			UnmatchedPins[0]->Properties = UnmatchedProperties[0];
-			bAppliedEdgeChanges |= UnmatchedPins[0]->BreakAllIncompatibleEdges();
+
+			bAppliedEdgeChanges |= UnmatchedPins[0]->BreakAllIncompatibleEdges(&TouchedNodes);
 			bChangedPins = true;
 		}
 		else
@@ -736,7 +739,7 @@ EPCGChangeType UPCGNode::UpdatePins(TFunctionRef<UPCGPin*(UPCGNode*)> PinAllocat
 				bChangedPins = true;
 			}
 
-			auto RemovePins = [&Pins, &AllPinNames, &bAppliedEdgeChanges](TArray<UPCGPin*>& PinsToRemove, bool bRemoveFromAllNames)
+			auto RemovePins = [&Pins, &AllPinNames, &bAppliedEdgeChanges, &TouchedNodes](TArray<UPCGPin*>& PinsToRemove, bool bRemoveFromAllNames)
 			{
 				for (int32 RemovedPinIndex = PinsToRemove.Num() - 1; RemovedPinIndex >= 0; --RemovedPinIndex)
 				{
@@ -748,7 +751,7 @@ EPCGChangeType UPCGNode::UpdatePins(TFunctionRef<UPCGPin*(UPCGNode*)> PinAllocat
 							AllPinNames.Remove(Pins[PinIndex]->Properties.Label);
 						}
 
-						bAppliedEdgeChanges |= Pins[PinIndex]->BreakAllEdges();
+						bAppliedEdgeChanges |= Pins[PinIndex]->BreakAllEdges(&TouchedNodes);
 						Pins.RemoveAt(PinIndex);
 					}
 				}
@@ -782,13 +785,20 @@ EPCGChangeType UPCGNode::UpdatePins(TFunctionRef<UPCGPin*(UPCGNode*)> PinAllocat
 	ChangeType |= UpdatePins(OutputPins, OutboundPinProperties);
 
 #if WITH_EDITOR
-	OnNodeChangedDelegate.Broadcast(this, ChangeType);
+	for (UPCGNode* Node : TouchedNodes)
+	{
+		if (Node)
+		{
+			// Only this node gets full change type
+			Node->OnNodeChangedDelegate.Broadcast(Node, (Node == this) ? ChangeType : EPCGChangeType::Node);
+		}
+	}
 #endif // WITH_EDITOR
-		
+
 	return ChangeType;
 }
 
-EPCGChangeType UPCGNode::UpdateDynamicPins(const UPCGNode* FromNode /*= nullptr*/)
+EPCGChangeType UPCGNode::PropagateDynamicPinTypes(const UPCGNode* FromNode /*= nullptr*/)
 {
 	EPCGChangeType ChangeType = EPCGChangeType::None;
 	const UPCGSettings* Settings = GetSettings();
@@ -798,7 +808,12 @@ EPCGChangeType UPCGNode::UpdateDynamicPins(const UPCGNode* FromNode /*= nullptr*
 	}
 
 	ChangeType |= UpdatePins();
-	
+
+#if WITH_EDITOR
+	// Reconstruct in case pin types need to change
+	OnNodeChangedDelegate.Broadcast(this, EPCGChangeType::Node);
+#endif
+
 	for (UPCGPin* OutputPin : OutputPins)
 	{
 		for (UPCGEdge* Edge : OutputPin->Edges)
@@ -818,7 +833,7 @@ EPCGChangeType UPCGNode::UpdateDynamicPins(const UPCGNode* FromNode /*= nullptr*
 			const UPCGSettings* OtherSettings = OtherNode->GetSettings();
 			if (Settings && OtherSettings->HasDynamicPins())
 			{
-				ChangeType |= OtherNode->UpdateDynamicPins(this);
+				ChangeType |= OtherNode->PropagateDynamicPinTypes(this);
 			}
 		}
 	}
diff --git a/Engine/Plugins/Experimental/PCG/Source/PCG/Private/PCGPin.cpp b/Engine/Plugins/Experimental/PCG/Source/PCG/Private/PCGPin.cpp
--- a/Engine/Plugins/Experimental/PCG/Source/PCG/Private/PCGPin.cpp
+++ b/Engine/Plugins/Experimental/PCG/Source/PCG/Private/PCGPin.cpp
@@ -3,6 +3,7 @@
 #include "PCGPin.h"
 #include "PCGEdge.h"
 #include "PCGNode.h"
+#include "PCGSettings.h"
 
 #include UE_INLINE_GENERATED_CPP_BY_NAME(PCGPin)
 
@@ -55,7 +56,7 @@ void UPCGPin::SetTooltip(const FText& InTooltip)
 #endif
 }
 
-bool UPCGPin::AddEdgeTo(UPCGPin* OtherPin)
+bool UPCGPin::AddEdgeTo(UPCGPin* OtherPin, TSet<UPCGNode*>* InTouchedNodes/*= nullptr*/)
 {
 	if (!OtherPin)
 	{
@@ -90,10 +91,16 @@ bool UPCGPin::AddEdgeTo(UPCGPin* OtherPin)
 	NewEdge->InputPin = bThisPinIsUpstream ? this : OtherPin;
 	NewEdge->OutputPin = bThisPinIsUpstream ? OtherPin : this;
 
+	if (InTouchedNodes)
+	{
+		InTouchedNodes->Add(Node);
+		InTouchedNodes->Add(OtherPin->Node);
+	}
+
 	return true;
 }
 
-bool UPCGPin::BreakEdgeTo(UPCGPin* OtherPin)
+bool UPCGPin::BreakEdgeTo(UPCGPin* OtherPin, TSet<UPCGNode*>* InTouchedNodes/*= nullptr*/)
 {
 	if (!OtherPin)
 	{
@@ -109,6 +116,13 @@ bool UPCGPin::BreakEdgeTo(UPCGPin* OtherPin)
 
 			ensure(OtherPin->Edges.Remove(Edge));
 			Edges.Remove(Edge);
+
+			if (InTouchedNodes)
+			{
+				InTouchedNodes->Add(Node);
+				InTouchedNodes->Add(OtherPin->Node);
+			}
+
 			return true;
 		}
 	}
@@ -116,12 +130,17 @@ bool UPCGPin::BreakEdgeTo(UPCGPin* OtherPin)
 	return false;
 }
 
-bool UPCGPin::BreakAllEdges()
+bool UPCGPin::BreakAllEdges(TSet<UPCGNode*>* InTouchedNodes/*= nullptr*/)
 {
 	bool bChanged = false;
 
 	if (!Edges.IsEmpty())
 	{
+		if (InTouchedNodes)
+		{
+			InTouchedNodes->Add(Node);
+		}
+
 		Modify();
 	}
 
@@ -132,6 +151,11 @@ bool UPCGPin::BreakAllEdges()
 			OtherPin->Modify();
 			ensure(OtherPin->Edges.Remove(Edge));
 			bChanged = true;
+
+			if (InTouchedNodes)
+			{
+				InTouchedNodes->Add(OtherPin->Node);
+			}
 		}
 	}
 
@@ -140,7 +164,7 @@ bool UPCGPin::BreakAllEdges()
 	return bChanged;
 }
 
-bool UPCGPin::BreakAllIncompatibleEdges()
+bool UPCGPin::BreakAllIncompatibleEdges(TSet<UPCGNode*>* InTouchedNodes/*= nullptr*/)
 {
 	bool bChanged = false;
 	bool bHasAValidEdge = false;
@@ -157,11 +181,18 @@ bool UPCGPin::BreakAllIncompatibleEdges()
 			Modify();
 			Edges.RemoveAtSwap(EdgeIndex);
 
+			InTouchedNodes->Add(Node);
+
 			if (OtherPin)
 			{
 				OtherPin->Modify();
 				ensure(OtherPin->Edges.Remove(Edge));
 				bChanged = true;
+
+				if (InTouchedNodes)
+				{
+					InTouchedNodes->Add(OtherPin->Node);
+				}
 			}
 		}
 		else
@@ -206,6 +237,13 @@ int32 UPCGPin::EdgeCount() const
 	return EdgeNum;
 }
 
+EPCGDataType UPCGPin::GetCurrentTypes() const
+{
+	check(Node);
+	const UPCGSettings* Settings = Node->GetSettings();
+	return Settings ? Settings->GetCurrentPinTypes(this) : Properties.AllowedTypes;
+}
+
 bool UPCGPin::IsCompatible(const UPCGPin* OtherPin) const
 {
 	if (!OtherPin)
@@ -223,18 +261,21 @@ bool UPCGPin::IsCompatible(const UPCGPin* OtherPin) const
 	// Sort pins
 	const UPCGPin* UpstreamPin = bThisPinOutput ? this : OtherPin;
 	const UPCGPin* DownstreamPin = bThisPinOutput ? OtherPin : this;
+	check(UpstreamPin && DownstreamPin);
+	const EPCGDataType UpstreamTypes = UpstreamPin->GetCurrentTypes();
+	const EPCGDataType DownstreamTypes = DownstreamPin->GetCurrentTypes();
 
 	// Concrete can always be used as a composite - allow connections from concrete to composite
-	const bool bUpstreamConcrete = !!(UpstreamPin->Properties.AllowedTypes & EPCGDataType::Concrete);
-	const bool bDownstreamSpatial = !!(DownstreamPin->Properties.AllowedTypes & EPCGDataType::Spatial);
+	const bool bUpstreamConcrete = !!(UpstreamTypes & EPCGDataType::Concrete);
+	const bool bDownstreamSpatial = !!(DownstreamTypes & EPCGDataType::Spatial);
 	if (bUpstreamConcrete && bDownstreamSpatial)
 	{
 		return true;
 	}
 
 	// Catch case when a composite type is being connected to a concrete type - that can require collapse
-	const bool bUpstreamSpatial = !!(UpstreamPin->Properties.AllowedTypes & EPCGDataType::Spatial);
-	const bool bDownstreamConcrete = !!(DownstreamPin->Properties.AllowedTypes & EPCGDataType::Concrete);
+	const bool bUpstreamSpatial = !!(UpstreamTypes & EPCGDataType::Spatial);
+	const bool bDownstreamConcrete = !!(DownstreamTypes & EPCGDataType::Concrete);
 	if (bUpstreamSpatial && bDownstreamConcrete)
 	{
 		// This will trigger a collapse, but let it slide for now.
@@ -242,7 +283,7 @@ bool UPCGPin::IsCompatible(const UPCGPin* OtherPin) const
 		return true;
 	}
 
-	return !!(Properties.AllowedTypes & OtherPin->Properties.AllowedTypes);
+	return !!(UpstreamTypes & DownstreamTypes);
 }
 
 bool UPCGPin::AllowMultipleConnections() const
diff --git a/Engine/Plugins/Experimental/PCG/Source/PCG/Private/PCGSettings.cpp b/Engine/Plugins/Experimental/PCG/Source/PCG/Private/PCGSettings.cpp
--- a/Engine/Plugins/Experimental/PCG/Source/PCG/Private/PCGSettings.cpp
+++ b/Engine/Plugins/Experimental/PCG/Source/PCG/Private/PCGSettings.cpp
@@ -339,7 +339,7 @@ EPCGDataType UPCGSettings::GetTypeUnionOfIncidentEdges(const FName& PinLabel) co
 			const UPCGPin* OtherOutputPin = Edge ? Edge->GetOtherPin(Pin) : nullptr;
 			if (OtherOutputPin)
 			{
-				Result |= OtherOutputPin->Properties.AllowedTypes;
+				Result |= OtherOutputPin->GetCurrentTypes();
 			}
 		}
 	}
@@ -481,6 +481,12 @@ void UPCGSettings::ApplyDeprecationBeforeUpdatePins(UPCGNode* InOutNode, TArray<
 }
 #endif // WITH_EDITOR
 
+EPCGDataType UPCGSettings::GetCurrentPinTypes(const UPCGPin* InPin) const
+{
+	check(InPin);
+	return InPin->Properties.AllowedTypes;
+}
+
 #if WITH_EDITOR
 TArray<FPCGSettingsOverridableParam> UPCGSettings::GatherOverridableParams() const
 {
diff --git a/Engine/Plugins/Experimental/PCG/Source/PCG/Public/Elements/Metadata/PCGMetadataOpElementBase.h b/Engine/Plugins/Experimental/PCG/Source/PCG/Public/Elements/Metadata/PCGMetadataOpElementBase.h
--- a/Engine/Plugins/Experimental/PCG/Source/PCG/Public/Elements/Metadata/PCGMetadataOpElementBase.h
+++ b/Engine/Plugins/Experimental/PCG/Source/PCG/Public/Elements/Metadata/PCGMetadataOpElementBase.h
@@ -80,11 +80,14 @@ public:
 	virtual EPCGSettingsType GetType() const override { return EPCGSettingsType::Metadata; }
 	virtual bool HasDynamicPins() const override { return true; }
 #endif
+	virtual EPCGDataType GetCurrentPinTypes(const UPCGPin* InPin) const override;
 
+protected:
 	virtual TArray<FPCGPinProperties> InputPinProperties() const override;
 	virtual TArray<FPCGPinProperties> OutputPinProperties() const override;
 	//~End UPCGSettings interface
 
+public:
 	virtual FPCGAttributePropertySelector GetInputSource(uint32 Index) const { return FPCGAttributePropertySelector(); };
 
 	virtual FName GetInputPinLabel(uint32 Index) const { return PCGPinConstants::DefaultInputLabel; }
diff --git a/Engine/Plugins/Experimental/PCG/Source/PCG/Public/Elements/Metadata/PCGMetadataRenameElement.h b/Engine/Plugins/Experimental/PCG/Source/PCG/Public/Elements/Metadata/PCGMetadataRenameElement.h
--- a/Engine/Plugins/Experimental/PCG/Source/PCG/Public/Elements/Metadata/PCGMetadataRenameElement.h
+++ b/Engine/Plugins/Experimental/PCG/Source/PCG/Public/Elements/Metadata/PCGMetadataRenameElement.h
@@ -19,6 +19,7 @@ public:
 	virtual EPCGSettingsType GetType() const override { return EPCGSettingsType::Metadata; }
 	virtual bool HasDynamicPins() const override { return true; }
 #endif
+	virtual EPCGDataType GetCurrentPinTypes(const UPCGPin* InPin) const override;
 
 protected:
 	virtual TArray<FPCGPinProperties> InputPinProperties() const override;
diff --git a/Engine/Plugins/Experimental/PCG/Source/PCG/Public/Elements/PCGAttributeFilterElement.h b/Engine/Plugins/Experimental/PCG/Source/PCG/Public/Elements/PCGAttributeFilterElement.h
--- a/Engine/Plugins/Experimental/PCG/Source/PCG/Public/Elements/PCGAttributeFilterElement.h
+++ b/Engine/Plugins/Experimental/PCG/Source/PCG/Public/Elements/PCGAttributeFilterElement.h
@@ -37,7 +37,7 @@ public:
 	virtual EPCGSettingsType GetType() const override { return EPCGSettingsType::Metadata; }
 	virtual bool HasDynamicPins() const override { return true; }
 #endif
-
+	virtual EPCGDataType GetCurrentPinTypes(const UPCGPin* InPin) const override;
 	virtual FName AdditionalTaskName() const override;
 
 protected:
diff --git a/Engine/Plugins/Experimental/PCG/Source/PCG/Public/Elements/PCGAttributeTransferElement.h b/Engine/Plugins/Experimental/PCG/Source/PCG/Public/Elements/PCGAttributeTransferElement.h
--- a/Engine/Plugins/Experimental/PCG/Source/PCG/Public/Elements/PCGAttributeTransferElement.h
+++ b/Engine/Plugins/Experimental/PCG/Source/PCG/Public/Elements/PCGAttributeTransferElement.h
@@ -42,6 +42,7 @@ public:
 	virtual FText GetNodeTooltipText() const override;
 	virtual bool HasDynamicPins() const override { return true; }
 #endif
+	virtual EPCGDataType GetCurrentPinTypes(const UPCGPin* InPin) const override;
 	virtual FName AdditionalTaskName() const override;
 
 protected:
@@ -50,9 +51,6 @@ protected:
 	virtual FPCGElementPtr CreateElement() const override;
 	//~End UPCGSettings interface
 
-	/** Get allowed types for all pins */
-	EPCGDataType GetPinsType() const;
-
 public:
 	UPROPERTY(BlueprintReadWrite, EditAnywhere, Category = Settings)
 	FName SourceAttributeName;
diff --git a/Engine/Plugins/Experimental/PCG/Source/PCG/Public/Elements/PCGCreateAttribute.h b/Engine/Plugins/Experimental/PCG/Source/PCG/Public/Elements/PCGCreateAttribute.h
--- a/Engine/Plugins/Experimental/PCG/Source/PCG/Public/Elements/PCGCreateAttribute.h
+++ b/Engine/Plugins/Experimental/PCG/Source/PCG/Public/Elements/PCGCreateAttribute.h
@@ -33,6 +33,7 @@ public:
 	virtual EPCGSettingsType GetType() const override { return EPCGSettingsType::Metadata; }
 	virtual bool HasDynamicPins() const override { return true; }
 #endif // WITH_EDITOR
+	virtual EPCGDataType GetCurrentPinTypes(const UPCGPin* InPin) const override;
 
 	virtual FName AdditionalTaskName() const override;
 	//~End UPCGSettings interface
@@ -122,11 +123,12 @@ public:
 #if WITH_EDITOR
 	virtual FName GetDefaultNodeName() const override;
 	virtual FText GetDefaultNodeTitle() const override;
+	virtual bool HasDynamicPins() const override { return false; }
 #endif
 	virtual FName AdditionalTaskName() const override;
 
 protected:
-	virtual TArray<FPCGPinProperties> InputPinProperties() const override { return TArray<FPCGPinProperties>(); }
+	virtual TArray<FPCGPinProperties> InputPinProperties() const override;
 	virtual TArray<FPCGPinProperties> OutputPinProperties() const override;
 	//~End UPCGSettings interface
 };
diff --git a/Engine/Plugins/Experimental/PCG/Source/PCG/Public/Elements/PCGFilterByTag.h b/Engine/Plugins/Experimental/PCG/Source/PCG/Public/Elements/PCGFilterByTag.h
--- a/Engine/Plugins/Experimental/PCG/Source/PCG/Public/Elements/PCGFilterByTag.h
+++ b/Engine/Plugins/Experimental/PCG/Source/PCG/Public/Elements/PCGFilterByTag.h
@@ -28,7 +28,7 @@ public:
 	virtual EPCGSettingsType GetType() const override { return EPCGSettingsType::Filter; }
 	virtual bool HasDynamicPins() const override { return true; }
 #endif
-
+	virtual EPCGDataType GetCurrentPinTypes(const UPCGPin* InPin) const override;
 	virtual FName AdditionalTaskName() const override;
 
 protected:
diff --git a/Engine/Plugins/Experimental/PCG/Source/PCG/Public/Elements/PCGGather.h b/Engine/Plugins/Experimental/PCG/Source/PCG/Public/Elements/PCGGather.h
--- a/Engine/Plugins/Experimental/PCG/Source/PCG/Public/Elements/PCGGather.h
+++ b/Engine/Plugins/Experimental/PCG/Source/PCG/Public/Elements/PCGGather.h
@@ -21,6 +21,7 @@ public:
 	virtual EPCGSettingsType GetType() const override { return EPCGSettingsType::Spatial; }
 	virtual bool HasDynamicPins() const override { return true; }
 #endif
+	virtual EPCGDataType GetCurrentPinTypes(const UPCGPin* InPin) const override;
 
 protected:
 	virtual TArray<FPCGPinProperties> InputPinProperties() const override;
diff --git a/Engine/Plugins/Experimental/PCG/Source/PCG/Public/Elements/PCGReroute.h b/Engine/Plugins/Experimental/PCG/Source/PCG/Public/Elements/PCGReroute.h
--- a/Engine/Plugins/Experimental/PCG/Source/PCG/Public/Elements/PCGReroute.h
+++ b/Engine/Plugins/Experimental/PCG/Source/PCG/Public/Elements/PCGReroute.h
@@ -18,10 +18,11 @@ public:
 	UPCGRerouteSettings();
 	
 	virtual bool HasDynamicPins() const override { return true; }
+	virtual EPCGDataType GetCurrentPinTypes(const UPCGPin* InPin) const override;
+
+protected:
 	virtual TArray<FPCGPinProperties> InputPinProperties() const override;
 	virtual TArray<FPCGPinProperties> OutputPinProperties() const override;
-	
-protected:
 	virtual FPCGElementPtr CreateElement() const override;
 	
 #if WITH_EDITOR
diff --git a/Engine/Plugins/Experimental/PCG/Source/PCG/Public/PCGNode.h b/Engine/Plugins/Experimental/PCG/Source/PCG/Public/PCGNode.h
--- a/Engine/Plugins/Experimental/PCG/Source/PCG/Public/PCGNode.h
+++ b/Engine/Plugins/Experimental/PCG/Source/PCG/Public/PCGNode.h
@@ -118,6 +118,9 @@ public:
 	const TArray<TObjectPtr<UPCGPin>>& GetInputPins() const { return InputPins; }
 	const TArray<TObjectPtr<UPCGPin>>& GetOutputPins() const { return OutputPins; }
 
+	/** Recursively follow downstream edges and call UpdatePins on each node that has dynamic pins. */
+	EPCGChangeType PropagateDynamicPinTypes(const UPCGNode* FromNode = nullptr);
+
 #if WITH_EDITOR
 	/** Transfer all editor only properties to the other node */
 	void TransferEditorProperties(UPCGNode* OtherNode) const;
@@ -166,9 +169,10 @@ public:
 #endif // WITH_EDITORONLY_DATA
 
 protected:
+	/** Updates pins based on node settings. Attempts to migrate pins via matching. Broadcasts node change events for affected nodes. */
 	EPCGChangeType UpdatePins();
-	EPCGChangeType UpdatePins(TFunctionRef<UPCGPin* (UPCGNode*)> PinAllocator, const UPCGNode* FromNode = nullptr);
-	EPCGChangeType UpdateDynamicPins(const UPCGNode* FromNode = nullptr);
+	/** Updates pins based on node settings PinAllocator creates new pin objects. Attempts to migrate pins via matching. Broadcasts node change events for affected nodes. */
+	EPCGChangeType UpdatePins(TFunctionRef<UPCGPin* (UPCGNode*)> PinAllocator);
 
 	// When we create a new graph, we initialize the input/output nodes as default, with default pins.
 	// Those default pins are not serialized, therefore if we change the default pins, combined with the use
diff --git a/Engine/Plugins/Experimental/PCG/Source/PCG/Public/PCGPin.h b/Engine/Plugins/Experimental/PCG/Source/PCG/Public/PCGPin.h
--- a/Engine/Plugins/Experimental/PCG/Source/PCG/Public/PCGPin.h
+++ b/Engine/Plugins/Experimental/PCG/Source/PCG/Public/PCGPin.h
@@ -74,13 +74,25 @@ public:
 	bool AllowMultipleConnections() const;
 	bool IsCompatible(const UPCGPin* OtherPin) const;
 	bool CanConnect(const UPCGPin* OtherPin) const;
-	bool AddEdgeTo(UPCGPin* OtherPin);
-	bool BreakEdgeTo(UPCGPin* OtherPin);
-	bool BreakAllEdges();
-	bool BreakAllIncompatibleEdges();
+
+	/** Adds edge to pin if edge does not exist. Optionally returns list of affected nodes (including the pin owner). */
+	bool AddEdgeTo(UPCGPin* OtherPin, TSet<UPCGNode*>* InTouchedNodes = nullptr);
+
+	/** Breaks edge to pin if edge exists. Optionally returns list of affected nodes (including the pin owner). */
+	bool BreakEdgeTo(UPCGPin* OtherPin, TSet<UPCGNode*>* InTouchedNodes = nullptr);
+
+	/** Breaks all connected edges. Optionally returns list of affected nodes (including the pin owner). */
+	bool BreakAllEdges(TSet<UPCGNode*>* InTouchedNodes = nullptr);
+
+	/** Break all edges the connect pins of invalid type. Optionally returns list of affected nodes (including the pin owner). */
+	bool BreakAllIncompatibleEdges(TSet<UPCGNode*>* InTouchedNodes = nullptr);
+
 	bool IsConnected() const;
 	bool IsOutputPin() const;
 	int32 EdgeCount() const;
+
+	/** Returns the current pin types, which can either be the static types from the pin properties, or a dynamic type based on connected edges. */
+	EPCGDataType GetCurrentTypes() const;
 };
 
 #if UE_ENABLE_INCLUDE_ORDER_DEPRECATED_IN_5_2
diff --git a/Engine/Plugins/Experimental/PCG/Source/PCG/Public/PCGSettings.h b/Engine/Plugins/Experimental/PCG/Source/PCG/Public/PCGSettings.h
--- a/Engine/Plugins/Experimental/PCG/Source/PCG/Public/PCGSettings.h
+++ b/Engine/Plugins/Experimental/PCG/Source/PCG/Public/PCGSettings.h
@@ -191,6 +191,9 @@ public:
 	/** Returns true if only the first input edge is used from the primary pin when the node is disabled. */
 	virtual bool OnlyPassThroughOneEdgeWhenDisabled() const { return false; }
 
+	/** Returns the current pin types, which can either be the static types from the pin properties, or a dynamic type based on connected edges. */
+	virtual EPCGDataType GetCurrentPinTypes(const UPCGPin* InPin) const;
+
 	UPROPERTY(EditAnywhere, BlueprintReadWrite, Category = Settings, meta=(EditCondition=bUseSeed, EditConditionHides, PCG_Overridable))
 	int Seed = 0xC35A9631; // random prime number
 
diff --git a/Engine/Plugins/Experimental/PCG/Source/PCGEditor/Private/PCGEditor.cpp b/Engine/Plugins/Experimental/PCG/Source/PCGEditor/Private/PCGEditor.cpp
--- a/Engine/Plugins/Experimental/PCG/Source/PCGEditor/Private/PCGEditor.cpp
+++ b/Engine/Plugins/Experimental/PCG/Source/PCGEditor/Private/PCGEditor.cpp
@@ -1041,19 +1041,20 @@ void FPCGEditor::OnCollapseNodesInSubgraph()
 	}
 
 	// 7. Connect all the edges
+	TSet<UPCGNode*> TouchedNodes;
 	for (EdgePlaceholder& Edge : EdgePlaceholders)
 	{
 		if (Edge.InputPin == nullptr)
 		{
-			SubgraphNode->GetOutputPin(Edge.InputPinLabel)->AddEdgeTo(Edge.OutputPin);
+			SubgraphNode->GetOutputPin(Edge.InputPinLabel)->AddEdgeTo(Edge.OutputPin, &TouchedNodes);
 		}
 		else if (Edge.OutputPin == nullptr)
 		{
-			Edge.InputPin->AddEdgeTo(SubgraphNode->GetInputPin(Edge.OutputPinLabel));
+			Edge.InputPin->AddEdgeTo(SubgraphNode->GetInputPin(Edge.OutputPinLabel), &TouchedNodes);
 		}
 		else
 		{
-			Edge.InputPin->AddEdgeTo(Edge.OutputPin);
+			Edge.InputPin->AddEdgeTo(Edge.OutputPin, &TouchedNodes);
 		}
 	}
 
@@ -1067,6 +1068,10 @@ void FPCGEditor::OnCollapseNodesInSubgraph()
 	GraphEditorWidget->NotifyGraphChanged();
 	PCGGraph->NotifyGraphChanged(EPCGChangeType::Structural);
 	NewPCGGraph->NotifyGraphChanged(EPCGChangeType::Structural);
+	for (UPCGNode* TouchedNode : TouchedNodes)
+	{
+		TouchedNode->OnNodeChangedDelegate.Broadcast(TouchedNode, EPCGChangeType::Node);
+	}
 }
 
 bool FPCGEditor::CanExportNodes() const
diff --git a/Engine/Plugins/Experimental/PCG/Source/PCGEditor/Private/PCGEditorGraphNodeBase.cpp b/Engine/Plugins/Experimental/PCG/Source/PCGEditor/Private/PCGEditorGraphNodeBase.cpp
--- a/Engine/Plugins/Experimental/PCG/Source/PCGEditor/Private/PCGEditorGraphNodeBase.cpp
+++ b/Engine/Plugins/Experimental/PCG/Source/PCGEditor/Private/PCGEditorGraphNodeBase.cpp
@@ -199,11 +199,12 @@ void UPCGEditorGraphNodeBase::PostPaste()
 		PCGNode->PositionX = NodePosX;
 		PCGNode->PositionY = NodePosY;
 
+		// Refresh the node if it has dynamic pins
 		if (const UPCGSettings* Settings = PCGNode->GetSettings())
 		{
 			if (Settings->HasDynamicPins())
 			{
-				PCGNode->UpdateAfterSettingsChangeDuringCreation();
+				PCGNode->OnNodeChangedDelegate.Broadcast(PCGNode, EPCGChangeType::Node);
 			}
 		}
 	}
@@ -460,7 +461,9 @@ FEdGraphPinType UPCGEditorGraphNodeBase::GetPinType(const UPCGPin* InPin)
 	EdPinType.PinSubCategory = NAME_None;
 	EdPinType.ContainerType = EPinContainerType::None;
 
-	const EPCGDataType PinType = InPin->Properties.AllowedTypes;
+	check(InPin);
+	const EPCGDataType PinType = InPin->GetCurrentTypes();
+
 	auto CheckType = [PinType](EPCGDataType AllowedType)
 	{
 		return !!(PinType & AllowedType) && !(PinType & ~AllowedType);
@@ -562,13 +565,13 @@ void UPCGEditorGraphNodeBase::GetPinHoverText(const UEdGraphPin& Pin, FString& H
 		{
 			return FText::FromName(Category);
 		}
-		else if (MatchingPin && MatchingPin->Properties.AllowedTypes == EPCGDataType::Any)
+		else if (MatchingPin && MatchingPin->GetCurrentTypes() == EPCGDataType::Any)
 		{
 			return FText::FromName(FName("Any"));
 		}
 		else if (MatchingPin)
 		{
-			return PCGDataTypeToText(MatchingPin->Properties.AllowedTypes);
+			return PCGDataTypeToText(MatchingPin->GetCurrentTypes());
 		}
 		else
 		{
diff --git a/Engine/Plugins/Experimental/PCG/Source/PCGEditor/Private/SPCGEditorGraphNode.cpp b/Engine/Plugins/Experimental/PCG/Source/PCGEditor/Private/SPCGEditorGraphNode.cpp
--- a/Engine/Plugins/Experimental/PCG/Source/PCGEditor/Private/SPCGEditorGraphNode.cpp
+++ b/Engine/Plugins/Experimental/PCG/Source/PCGEditor/Private/SPCGEditorGraphNode.cpp
@@ -307,14 +307,14 @@ void SPCGEditorGraphNode::AddPin(const TSharedRef<SGraphPin>& PinToAdd)
 			const bool bIsMultiConnections = Pin->AllowMultipleConnections();
 
 			// Check for special types
-			if (Pin->Properties.AllowedTypes == EPCGDataType::Param)
+			if (Pin->GetCurrentTypes() == EPCGDataType::Param)
 			{
 				const FSlateBrush* ConnectedBrush = FPCGEditorStyle::Get().GetBrush(bIsInPin ? PCGEditorStyleConstants::Pin_Param_IN_C : PCGEditorStyleConstants::Pin_Param_OUT_C);
 				const FSlateBrush* DisconnectedBrush = FPCGEditorStyle::Get().GetBrush(bIsInPin ? PCGEditorStyleConstants::Pin_Param_IN_DC : PCGEditorStyleConstants::Pin_Param_OUT_DC);
 
 				PinToAdd->SetCustomPinIcon(ConnectedBrush, DisconnectedBrush);
 			}
-			else if (Pin->Properties.AllowedTypes == EPCGDataType::Spatial)
+			else if (Pin->GetCurrentTypes() == EPCGDataType::Spatial)
 			{
 				const FSlateBrush* ConnectedBrush = FPCGEditorStyle::Get().GetBrush(bIsInPin ? PCGEditorStyleConstants::Pin_Composite_IN_C : PCGEditorStyleConstants::Pin_Composite_OUT_C);
 				const FSlateBrush* DisconnectedBrush = FPCGEditorStyle::Get().GetBrush(bIsInPin ? PCGEditorStyleConstants::Pin_Composite_IN_DC : PCGEditorStyleConstants::Pin_Composite_OUT_DC);
