Minor Fixes:
- Improved naming on some nodes with custom names (were missing localized name & spaces)
- Improved node name display for blueprint nodes in the node palette
- Made the world ray/volumetric query parameters properly overridable

diff --git a/Engine/Plugins/Experimental/PCG/Source/PCG/Private/Elements/PCGAttributeFilterElement.cpp b/Engine/Plugins/Experimental/PCG/Source/PCG/Private/Elements/PCGAttributeFilterElement.cpp
--- a/Engine/Plugins/Experimental/PCG/Source/PCG/Private/Elements/PCGAttributeFilterElement.cpp
+++ b/Engine/Plugins/Experimental/PCG/Source/PCG/Private/Elements/PCGAttributeFilterElement.cpp
@@ -66,7 +66,7 @@ FName UPCGAttributeFilterSettings::AdditionalTaskName() const
 {
 	TArray<FString> AttributesToKeep = PCGHelpers::GetStringArrayFromCommaSeparatedString(SelectedAttributes);
 
-	FString NodeName = PCGAttributeFilterConstants::NodeName.ToString();
+	FString NodeName = PCGAttributeFilterConstants::NodeTitle.ToString();
 
 	switch (Operation)
 	{
diff --git a/Engine/Plugins/Experimental/PCG/Source/PCG/Private/Elements/PCGExecuteBlueprint.cpp b/Engine/Plugins/Experimental/PCG/Source/PCG/Private/Elements/PCGExecuteBlueprint.cpp
--- a/Engine/Plugins/Experimental/PCG/Source/PCG/Private/Elements/PCGExecuteBlueprint.cpp
+++ b/Engine/Plugins/Experimental/PCG/Source/PCG/Private/Elements/PCGExecuteBlueprint.cpp
@@ -449,11 +449,15 @@ FName UPCGBlueprintSettings::AdditionalTaskName() const
 	}
 	else
 	{
+		FName ElementName = NAME_None;
+
 #if WITH_EDITOR
-		return (BlueprintElementType && BlueprintElementType->ClassGeneratedBy) ? BlueprintElementType->ClassGeneratedBy->GetFName() : Super::AdditionalTaskName();
+		ElementName = (BlueprintElementType && BlueprintElementType->ClassGeneratedBy) ? BlueprintElementType->ClassGeneratedBy->GetFName() : Super::AdditionalTaskName();
 #else
-		return BlueprintElementType ? BlueprintElementType->GetFName() : Super::AdditionalTaskName();
+		ElementName = BlueprintElementType ? BlueprintElementType->GetFName() : Super::AdditionalTaskName();
 #endif
+		// Normalize node name only if not explicitly set in the NodeTitleOverride call
+		return FName(FName::NameToDisplayString(ElementName.ToString(), false));
 	}
 }
 
diff --git a/Engine/Plugins/Experimental/PCG/Source/PCG/Private/Elements/PCGFilterByTag.cpp b/Engine/Plugins/Experimental/PCG/Source/PCG/Private/Elements/PCGFilterByTag.cpp
--- a/Engine/Plugins/Experimental/PCG/Source/PCG/Private/Elements/PCGFilterByTag.cpp
+++ b/Engine/Plugins/Experimental/PCG/Source/PCG/Private/Elements/PCGFilterByTag.cpp
@@ -43,7 +43,7 @@ FName UPCGFilterByTagSettings::AdditionalTaskName() const
 {
 	const TArray<FString> Tags = PCGHelpers::GetStringArrayFromCommaSeparatedString(SelectedTags);
 
-	FString NodeName = PCGFilterByTagConstants::NodeName.ToString();
+	FString NodeName = PCGFilterByTagConstants::NodeTitle.ToString();
 	NodeName += (Operation == EPCGFilterByTagOperation::KeepTagged ? TEXT(" (Keep)") : TEXT(" (Remove)"));
 
 	if (Tags.Num() == 1)
diff --git a/Engine/Plugins/Experimental/PCG/Source/PCG/Private/PCGNode.cpp b/Engine/Plugins/Experimental/PCG/Source/PCG/Private/PCGNode.cpp
--- a/Engine/Plugins/Experimental/PCG/Source/PCG/Private/PCGNode.cpp
+++ b/Engine/Plugins/Experimental/PCG/Source/PCG/Private/PCGNode.cpp
@@ -240,7 +240,7 @@ FText UPCGNode::GetNodeTitle() const
 {
 	if (NodeTitle != NAME_None)
 	{
-		return FText::FromName(NodeTitle);
+		return FText::FromString(FName::NameToDisplayString(NodeTitle.ToString(), false));
 	}
 	else if (UPCGSettings* Settings = GetSettings())
 	{
diff --git a/Engine/Plugins/Experimental/PCG/Source/PCG/Public/Elements/PCGWorldQuery.h b/Engine/Plugins/Experimental/PCG/Source/PCG/Public/Elements/PCGWorldQuery.h
--- a/Engine/Plugins/Experimental/PCG/Source/PCG/Public/Elements/PCGWorldQuery.h
+++ b/Engine/Plugins/Experimental/PCG/Source/PCG/Public/Elements/PCGWorldQuery.h
@@ -30,7 +30,7 @@ protected:
 	//~End UPCGSettings interface
 
 public:
-	UPROPERTY(BlueprintReadWrite, EditAnywhere, Category = Data, meta = (ShowOnlyInnerProperties))
+	UPROPERTY(BlueprintReadWrite, EditAnywhere, Category = Data, meta = (PCG_Overridable, ShowOnlyInnerProperties))
 	FPCGWorldVolumetricQueryParams QueryParams;
 };
 
@@ -64,7 +64,7 @@ protected:
 	//~End UPCGSettings interface
 
 public:
-	UPROPERTY(BlueprintReadWrite, EditAnywhere, Category = Data, meta = (ShowOnlyInnerProperties))
+	UPROPERTY(BlueprintReadWrite, EditAnywhere, Category = Data, meta = (PCG_Overridable, ShowOnlyInnerProperties))
 	FPCGWorldRayHitQueryParams QueryParams;
 };
 
diff --git a/Engine/Plugins/Experimental/PCG/Source/PCGEditor/Private/PCGEditorGraphNodeBase.cpp b/Engine/Plugins/Experimental/PCG/Source/PCGEditor/Private/PCGEditorGraphNodeBase.cpp
--- a/Engine/Plugins/Experimental/PCG/Source/PCGEditor/Private/PCGEditorGraphNodeBase.cpp
+++ b/Engine/Plugins/Experimental/PCG/Source/PCGEditor/Private/PCGEditorGraphNodeBase.cpp
@@ -683,6 +683,7 @@ void UPCGEditorGraphNodeBase::CreatePins(const TArray<UPCGPin*>& InInputPins, co
 		}
 
 		UEdGraphPin* Pin = CreatePin(EEdGraphPinDirection::EGPD_Input, GetPinType(InputPin), InputPin->Properties.Label);
+		Pin->PinFriendlyName = FText::FromString(FName::NameToDisplayString(InputPin->Properties.Label.ToString(), false));
 		Pin->bAdvancedView = InputPin->Properties.bAdvancedPin;
 		bHasAdvancedPin |= Pin->bAdvancedView;
 	}
@@ -695,6 +696,7 @@ void UPCGEditorGraphNodeBase::CreatePins(const TArray<UPCGPin*>& InInputPins, co
 		}
 
 		UEdGraphPin* Pin = CreatePin(EEdGraphPinDirection::EGPD_Output, GetPinType(OutputPin), OutputPin->Properties.Label);
+		Pin->PinFriendlyName = FText::FromString(FName::NameToDisplayString(OutputPin->Properties.Label.ToString(), false));
 		Pin->bAdvancedView = OutputPin->Properties.bAdvancedPin;
 		bHasAdvancedPin |= Pin->bAdvancedView;
 	}
diff --git a/Engine/Plugins/Experimental/PCG/Source/PCGEditor/Private/PCGEditorGraphSchema.cpp b/Engine/Plugins/Experimental/PCG/Source/PCGEditor/Private/PCGEditorGraphSchema.cpp
--- a/Engine/Plugins/Experimental/PCG/Source/PCGEditor/Private/PCGEditorGraphSchema.cpp
+++ b/Engine/Plugins/Experimental/PCG/Source/PCGEditor/Private/PCGEditorGraphSchema.cpp
@@ -266,7 +266,7 @@ void UPCGEditorGraphSchema::GetBlueprintElementActions(FGraphActionMenuBuilder&
 		const bool bExposeToLibrary = AssetData.GetTagValueRef<bool>(TEXT("bExposeToLibrary"));
 		if (bExposeToLibrary)
 		{
-			const FText MenuDesc = FText::FromName(AssetData.AssetName);
+			const FText MenuDesc = FText::FromString(FName::NameToDisplayString(AssetData.AssetName.ToString(), false));
 			const FText Category = AssetData.GetTagValueRef<FText>(TEXT("Category"));
 			const FText Description = AssetData.GetTagValueRef<FText>(TEXT("Description"));
 
@@ -288,7 +288,7 @@ void UPCGEditorGraphSchema::GetSettingsElementActions(FGraphActionMenuBuilder& A
 		const bool bExposeToLibrary = AssetData.GetTagValueRef<bool>(TEXT("bExposeToLibrary"));
 		if (bExposeToLibrary)
 		{
-			const FText MenuDesc = FText::FromName(AssetData.AssetName);
+			const FText MenuDesc = FText::FromString(FName::NameToDisplayString(AssetData.AssetName.ToString(), false));
 			const FText Category = AssetData.GetTagValueRef<FText>(TEXT("Category"));
 			const FText Description = AssetData.GetTagValueRef<FText>(TEXT("Description"));
 
@@ -325,7 +325,7 @@ void UPCGEditorGraphSchema::GetSubgraphElementActions(FGraphActionMenuBuilder& A
 		const bool bExposeToLibrary = AssetData.GetTagValueRef<bool>(TEXT("bExposeToLibrary"));
 		if (bExposeToLibrary)
 		{
-			const FText MenuDesc = FText::FromName(AssetData.AssetName);
+			const FText MenuDesc = FText::FromString(FName::NameToDisplayString(AssetData.AssetName.ToString(), false));
 			const FText Category = AssetData.GetTagValueRef<FText>(TEXT("Category"));
 			const FText Description = AssetData.GetTagValueRef<FText>(TEXT("Description"));
 
