[PCG] Conversion nodes - To Point and Filter

diff --git a/Engine/Plugins/Experimental/PCG/Source/PCG/Private/Elements/PCGFilterByType.cpp b/Engine/Plugins/Experimental/PCG/Source/PCG/Private/Elements/PCGFilterByType.cpp
new file mode 100644
--- /dev/null
+++ b/Engine/Plugins/Experimental/PCG/Source/PCG/Private/Elements/PCGFilterByType.cpp
@@ -0,0 +1,74 @@
+// Copyright Epic Games, Inc. All Rights Reserved.
+
+#include "Elements/PCGFilterByType.h"
+
+#include "PCGContext.h"
+#include "PCGPin.h"
+
+#include UE_INLINE_GENERATED_CPP_BY_NAME(PCGFilterByType)
+
+#define LOCTEXT_NAMESPACE "PCGFilterByType"
+
+#if WITH_EDITOR
+FText UPCGFilterByTypeSettings::GetDefaultNodeTitle() const
+{
+	const FText TypeText = FText::FromString(StaticEnum<EPCGDataType>()->GetDisplayNameTextByValue(static_cast<int64>(TargetType)).ToString());
+	return FText::Format(LOCTEXT("NodeTitleFormat", "Filter - {0}"), TypeText);
+}
+
+FText UPCGFilterByTypeSettings::GetNodeTooltipText() const
+{
+	return LOCTEXT("FilterByTypeNodeTooltip", "Filters data in the collection according to data type");
+}
+#endif
+
+EPCGDataType UPCGFilterByTypeSettings::GetCurrentPinTypes(const UPCGPin* InPin) const
+{
+	// Output pin narrows to set type, input pin stays at Any to invite all connections
+	check(InPin);
+	if (InPin->Properties.Label == PCGPinConstants::DefaultOutputLabel)
+	{
+		return TargetType;
+	}
+
+	return Super::GetCurrentPinTypes(InPin);
+}
+
+FPCGElementPtr UPCGFilterByTypeSettings::CreateElement() const
+{
+	return MakeShared<FPCGFilterByTypeElement>();
+}
+
+bool FPCGFilterByTypeElement::ExecuteInternal(FPCGContext* InContext) const
+{
+	TRACE_CPUPROFILER_EVENT_SCOPE(FPCGFilterByTypeElement::Execute);
+	check(InContext);
+
+	const UPCGFilterByTypeSettings* Settings = InContext->GetInputSettings<UPCGFilterByTypeSettings>();
+	check(Settings);
+
+	TArray<FPCGTaggedData> Inputs = InContext->InputData.GetInputs();
+	TArray<FPCGTaggedData>& Outputs = InContext->OutputData.TaggedData;
+
+	for (const FPCGTaggedData& Input : Inputs)
+	{
+		if (!Input.Data)
+		{
+			continue;
+		}
+
+		// No match if data does not overlap in type with the target, or if data is broader than the target
+		const bool bDataTypeOverlapsWithTarget = !!(Input.Data->GetDataType() & Settings->TargetType);
+		const bool bDataTypeBroaderThanTarget = !!(Input.Data->GetDataType() & ~Settings->TargetType);
+		if (!bDataTypeOverlapsWithTarget || bDataTypeBroaderThanTarget)
+		{
+			continue;
+		}
+
+		Outputs.Add(Input);
+	}
+
+	return true;
+}
+
+#undef LOCTEXT_NAMESPACE
diff --git a/Engine/Plugins/Experimental/PCG/Source/PCG/Private/PCGNode.cpp b/Engine/Plugins/Experimental/PCG/Source/PCG/Private/PCGNode.cpp
--- a/Engine/Plugins/Experimental/PCG/Source/PCG/Private/PCGNode.cpp
+++ b/Engine/Plugins/Experimental/PCG/Source/PCG/Private/PCGNode.cpp
@@ -587,8 +587,14 @@ void UPCGNode::OnSettingsChanged(UPCGSettings* InSettings, EPCGChangeType Change
 {
 	if (InSettings == GetSettings())
 	{
-		const EPCGChangeType PinChangeType = UpdatePins();
-		OnNodeChangedDelegate.Broadcast(this, ChangeType | PinChangeType);
+		EPCGChangeType PinChangeType = ChangeType | UpdatePins();
+		if (InSettings->HasDynamicPins())
+		{
+			// Add in case dynamic pin types changed when settings changed.
+			PinChangeType |= EPCGChangeType::Node;
+		}
+
+		OnNodeChangedDelegate.Broadcast(this, PinChangeType);
 	}
 }
 
diff --git a/Engine/Plugins/Experimental/PCG/Source/PCG/Private/PCGPin.cpp b/Engine/Plugins/Experimental/PCG/Source/PCG/Private/PCGPin.cpp
--- a/Engine/Plugins/Experimental/PCG/Source/PCG/Private/PCGPin.cpp
+++ b/Engine/Plugins/Experimental/PCG/Source/PCG/Private/PCGPin.cpp
@@ -1,12 +1,37 @@
 // Copyright Epic Games, Inc. All Rights Reserved.
 
 #include "PCGPin.h"
+
 #include "PCGEdge.h"
 #include "PCGNode.h"
 #include "PCGSettings.h"
 
 #include UE_INLINE_GENERATED_CPP_BY_NAME(PCGPin)
 
+namespace PCGPin
+{
+	bool SortPinsAndRetrieveTypes(const UPCGPin* InPinA, const UPCGPin* InPinB, const UPCGPin*& OutUpstreamPin, const UPCGPin*& OutDownstreamPin, EPCGDataType& OutUpstreamTypes, EPCGDataType& OutDownstreamTypes)
+	{
+		check(InPinA && InPinB);
+		const bool bPinAIsOutput = InPinA->IsOutputPin();
+
+		if (!ensure(bPinAIsOutput != InPinB->IsOutputPin()))
+		{
+			// Cannot connect two pins of same polarity
+			return false;
+		}
+
+		OutUpstreamPin = bPinAIsOutput ? InPinA : InPinB;
+		OutDownstreamPin = bPinAIsOutput ? InPinB : InPinA;
+
+		check(OutUpstreamPin && OutDownstreamPin);
+		OutUpstreamTypes = OutUpstreamPin->GetCurrentTypes();
+		OutDownstreamTypes = OutDownstreamPin->Properties.AllowedTypes;
+
+		return true;
+	}
+}
+
 FPCGPinProperties::FPCGPinProperties(const FName& InLabel, EPCGDataType InAllowedTypes, bool bInAllowMultipleConnections, bool bInAllowMultipleData, const FText& InTooltip)
 	: Label(InLabel), AllowedTypes(InAllowedTypes), bAllowMultipleData(bInAllowMultipleData), bAllowMultipleConnections(bInAllowMultipleConnections)
 #if WITH_EDITORONLY_DATA
@@ -251,38 +276,38 @@ bool UPCGPin::IsCompatible(const UPCGPin* OtherPin) const
 		return false;
 	}
 
-	const bool bThisPinOutput = IsOutputPin();
-	if (!ensure(bThisPinOutput != OtherPin->IsOutputPin()))
+	const UPCGPin* UpstreamPin = nullptr;
+	const UPCGPin* DownstreamPin = nullptr;
+	EPCGDataType UpstreamTypes = EPCGDataType::None, DownstreamTypes = EPCGDataType::None;
+	if (!PCGPin::SortPinsAndRetrieveTypes(this, OtherPin, UpstreamPin, DownstreamPin, UpstreamTypes, DownstreamTypes))
 	{
-		// Cannot connect two pins of same polarity
 		return false;
 	}
-
-	// Sort pins
-	const UPCGPin* UpstreamPin = bThisPinOutput ? this : OtherPin;
-	const UPCGPin* DownstreamPin = bThisPinOutput ? OtherPin : this;
 	check(UpstreamPin && DownstreamPin);
-	const EPCGDataType UpstreamTypes = UpstreamPin->GetCurrentTypes();
-	const EPCGDataType DownstreamTypes = DownstreamPin->GetCurrentTypes();
+
+	// Types missing
+	if (UpstreamTypes == EPCGDataType::None || DownstreamTypes == EPCGDataType::None)
+	{
+		return false;
+	}
 
 	// Concrete can always be used as a composite - allow connections from concrete to composite
-	const bool bUpstreamConcrete = !!(UpstreamTypes & EPCGDataType::Concrete);
-	const bool bDownstreamSpatial = !!(DownstreamTypes & EPCGDataType::Spatial);
-	if (bUpstreamConcrete && bDownstreamSpatial)
+	const bool bUpstreamInConcrete = !(UpstreamTypes & ~EPCGDataType::Concrete);
+	const bool bDownstreamIsSpatial = (DownstreamTypes == EPCGDataType::Spatial);
+	if (bUpstreamInConcrete && bDownstreamIsSpatial)
 	{
 		return true;
 	}
 
-	// Catch case when a composite type is being connected to a concrete type - that can require collapse
-	const bool bUpstreamSpatial = !!(UpstreamTypes & EPCGDataType::Spatial);
-	const bool bDownstreamConcrete = !!(DownstreamTypes & EPCGDataType::Concrete);
-	if (bUpstreamSpatial && bDownstreamConcrete)
+	// Anything spatial can collapse to point
+	const bool bUpstreamInSpatial = !(UpstreamTypes & ~EPCGDataType::Spatial);
+	const bool bDownstreamIsPoint = (DownstreamTypes == EPCGDataType::Point);
+	if (bUpstreamInSpatial && bDownstreamIsPoint)
 	{
-		// This will trigger a collapse, but let it slide for now.
-		// TODO in the future we should inject a conversion node.
 		return true;
 	}
 
+	// Otherwise allow if there is overlap. Don't detect wide -> narrow issues - conversion nodes deal with that
 	return !!(UpstreamTypes & DownstreamTypes);
 }
 
@@ -296,3 +321,83 @@ bool UPCGPin::CanConnect(const UPCGPin* OtherPin) const
 {
 	return OtherPin && (Edges.IsEmpty() || AllowMultipleConnections());
 }
+
+bool UPCGPin::RequiresPointConversionToConnect(const UPCGPin* OtherPin) const
+{
+	if (!OtherPin)
+	{
+		return false;
+	}
+
+	const UPCGPin* UpstreamPin = nullptr;
+	const UPCGPin* DownstreamPin = nullptr;
+	EPCGDataType UpstreamTypes = EPCGDataType::None, DownstreamTypes = EPCGDataType::None;
+	if (!PCGPin::SortPinsAndRetrieveTypes(this, OtherPin, UpstreamPin, DownstreamPin, UpstreamTypes, DownstreamTypes))
+	{
+		return false;
+	}
+	check(UpstreamPin && DownstreamPin);
+
+	// This conversion requires downstream Point
+	if (DownstreamTypes != EPCGDataType::Point)
+	{
+		return false;
+	}
+
+	// Only collapse from Spatial to Point
+	const bool bUpstreamInSpatial = !(UpstreamTypes & ~EPCGDataType::Spatial);
+	if (!bUpstreamInSpatial)
+	{
+		return false;
+	}
+
+	// Requires conversion if upstream pin is broader than point
+	return !!(UpstreamTypes & (~EPCGDataType::Point));
+}
+
+bool UPCGPin::RequiresFilterToConnect(const UPCGPin* OtherPin) const
+{
+	if (!OtherPin)
+	{
+		return false;
+	}
+
+	const UPCGPin* UpstreamPin = nullptr;
+	const UPCGPin* DownstreamPin = nullptr;
+	EPCGDataType UpstreamTypes = EPCGDataType::None, DownstreamTypes = EPCGDataType::None;
+	if (!PCGPin::SortPinsAndRetrieveTypes(this, OtherPin, UpstreamPin, DownstreamPin, UpstreamTypes, DownstreamTypes))
+	{
+		return false;
+	}
+	check(UpstreamPin && DownstreamPin);
+
+	// Same types - no filter
+	if (DownstreamTypes == UpstreamTypes)
+	{
+		return false;
+	}
+
+	// No type - early out
+	if (!DownstreamTypes)
+	{
+		return false;
+	}
+
+	// Spatial -> Point - that uses conversion, not filter
+	const bool bUpstreamInSpatial = !(UpstreamTypes & ~EPCGDataType::Spatial);
+	if (bUpstreamInSpatial && DownstreamTypes == EPCGDataType::Point)
+	{
+		return false;
+	}
+
+	// Upstream type should be Any, Spatial, Concrete
+	if (UpstreamTypes != EPCGDataType::Any
+		&& UpstreamTypes != EPCGDataType::Spatial
+		&& UpstreamTypes != EPCGDataType::Concrete)
+	{
+		return false;
+	}
+
+	// Requires conversion if upstream type is broader than downstream type
+	return !!(UpstreamTypes & (~DownstreamTypes));
+}
diff --git a/Engine/Plugins/Experimental/PCG/Source/PCG/Public/Elements/PCGCollapseElement.h b/Engine/Plugins/Experimental/PCG/Source/PCG/Public/Elements/PCGCollapseElement.h
--- a/Engine/Plugins/Experimental/PCG/Source/PCG/Public/Elements/PCGCollapseElement.h
+++ b/Engine/Plugins/Experimental/PCG/Source/PCG/Public/Elements/PCGCollapseElement.h
@@ -7,6 +7,7 @@
 
 #include "PCGCollapseElement.generated.h"
 
+/** Convert input to point data, performing sampling with default settings if necessary */
 UCLASS(BlueprintType, ClassGroup = (Procedural))
 class PCG_API UPCGCollapseSettings : public UPCGSettings
 {
@@ -18,6 +19,7 @@ public:
 	virtual FName GetDefaultNodeName() const override { return FName(TEXT("ToPoint")); }
 	virtual FText GetDefaultNodeTitle() const override { return NSLOCTEXT("PCGCollapseSettings", "NodeTitle", "To Point"); }
 	virtual EPCGSettingsType GetType() const override { return EPCGSettingsType::Spatial; }
+	virtual bool ShouldDrawNodeCompact() const override { return true; }
 #endif
 
 protected:
diff --git a/Engine/Plugins/Experimental/PCG/Source/PCG/Public/Elements/PCGFilterByType.h b/Engine/Plugins/Experimental/PCG/Source/PCG/Public/Elements/PCGFilterByType.h
new file mode 100644
--- /dev/null
+++ b/Engine/Plugins/Experimental/PCG/Source/PCG/Public/Elements/PCGFilterByType.h
@@ -0,0 +1,40 @@
+// Copyright Epic Games, Inc. All Rights Reserved.
+
+#pragma once
+
+#include "PCGSettings.h"
+
+#include "PCGFilterByType.generated.h"
+
+/** Filters an input collection based on data type. */
+UCLASS(BlueprintType, ClassGroup = (Procedural))
+class PCG_API UPCGFilterByTypeSettings : public UPCGSettings
+{
+	GENERATED_BODY()
+
+public:
+	//~Begin UPCGSettings interface
+#if WITH_EDITOR
+	virtual FName GetDefaultNodeName() const override { return TEXT("FilterByType"); }
+	virtual FText GetDefaultNodeTitle() const override;
+	virtual FText GetNodeTooltipText() const override;
+	virtual EPCGSettingsType GetType() const override { return EPCGSettingsType::Filter; }
+	virtual bool HasDynamicPins() const override { return true; }
+	virtual bool ShouldDrawNodeCompact() const override { return true; }
+#endif
+	virtual EPCGDataType GetCurrentPinTypes(const UPCGPin* InPin) const override;
+
+protected:
+	virtual FPCGElementPtr CreateElement() const override;
+	//~End UPCGSettings interface
+
+public:
+	UPROPERTY(BlueprintReadWrite, EditAnywhere, Category = Settings)
+	EPCGDataType TargetType = EPCGDataType::Any;
+};
+
+class FPCGFilterByTypeElement : public FSimplePCGElement
+{
+protected:
+	virtual bool ExecuteInternal(FPCGContext* InContext) const;
+};
diff --git a/Engine/Plugins/Experimental/PCG/Source/PCG/Public/PCGCommon.h b/Engine/Plugins/Experimental/PCG/Source/PCG/Public/PCGCommon.h
--- a/Engine/Plugins/Experimental/PCG/Source/PCG/Public/PCGCommon.h
+++ b/Engine/Plugins/Experimental/PCG/Source/PCG/Public/PCGCommon.h
@@ -28,7 +28,7 @@ enum class EPCGDataType : uint32
 
 	Spline = 1 << 2,
 	LandscapeSpline = 1 << 3,
-	PolyLine = Spline | LandscapeSpline UMETA(Hidden),
+	PolyLine = Spline | LandscapeSpline UMETA(DisplayName = "Curve"),
 
 	Landscape = 1 << 4,
 	Texture = 1 << 5,
diff --git a/Engine/Plugins/Experimental/PCG/Source/PCG/Public/PCGPin.h b/Engine/Plugins/Experimental/PCG/Source/PCG/Public/PCGPin.h
--- a/Engine/Plugins/Experimental/PCG/Source/PCG/Public/PCGPin.h
+++ b/Engine/Plugins/Experimental/PCG/Source/PCG/Public/PCGPin.h
@@ -93,6 +93,12 @@ public:
 
 	/** Returns the current pin types, which can either be the static types from the pin properties, or a dynamic type based on connected edges. */
 	EPCGDataType GetCurrentTypes() const;
+
+	/** Returns true if a collapse is required to convert from upstream type to downstream type */
+	bool RequiresPointConversionToConnect(const UPCGPin* OtherPin) const;
+
+	/** Returns true if a data filter is required to convert from upstream type to downstream type */
+	bool RequiresFilterToConnect(const UPCGPin* OtherPin) const;
 };
 
 #if UE_ENABLE_INCLUDE_ORDER_DEPRECATED_IN_5_2
diff --git a/Engine/Plugins/Experimental/PCG/Source/PCG/Public/PCGSettings.h b/Engine/Plugins/Experimental/PCG/Source/PCG/Public/PCGSettings.h
--- a/Engine/Plugins/Experimental/PCG/Source/PCG/Public/PCGSettings.h
+++ b/Engine/Plugins/Experimental/PCG/Source/PCG/Public/PCGSettings.h
@@ -134,7 +134,7 @@ public:
 	/*virtual*/ FPCGElementPtr GetElement() const;
 	virtual UPCGNode* CreateNode() const;
 	
-	/** Return the concatenation of InputPinPropertiesand FillOverridableParamsPins */
+	/** Return the concatenation of InputPinProperties and FillOverridableParamsPins */
 	TArray<FPCGPinProperties> AllInputPinProperties() const;
 
 	/** For symmetry reason, do the same with output pins. For now forward just the call to OutputPinProperties */
@@ -162,6 +162,9 @@ public:
 	const FPCGCrc& GetCachedCrc() const { return CachedCrc; }
 
 #if WITH_EDITOR
+	/** Puts node title on node body, reducing overall node size */
+	virtual bool ShouldDrawNodeCompact() const { return false; }
+
 	/** UpdatePins will kick off invalid edges, so this is useful for moving edges around in case of pin changes. */
 	virtual void ApplyDeprecationBeforeUpdatePins(UPCGNode* InOutNode, TArray<TObjectPtr<UPCGPin>>& InputPins, TArray<TObjectPtr<UPCGPin>>& OutputPins);
 	/** Any final migration/recovery that can be done after pins are finalized. This function should also set DataVersion to LatestVersion. */
diff --git a/Engine/Plugins/Experimental/PCG/Source/PCGEditor/Private/PCGEditorGraphNodeBase.cpp b/Engine/Plugins/Experimental/PCG/Source/PCGEditor/Private/PCGEditorGraphNodeBase.cpp
--- a/Engine/Plugins/Experimental/PCG/Source/PCGEditor/Private/PCGEditorGraphNodeBase.cpp
+++ b/Engine/Plugins/Experimental/PCG/Source/PCGEditor/Private/PCGEditorGraphNodeBase.cpp
@@ -330,6 +330,12 @@ void UPCGEditorGraphNodeBase::UpdateErrorsAndWarnings()
 	}
 }
 
+bool UPCGEditorGraphNodeBase::ShouldDrawCompact() const
+{
+	UPCGSettings* Settings = PCGNode ? PCGNode->GetSettings() : nullptr;
+	return Settings && Settings->ShouldDrawNodeCompact();
+}
+
 void UPCGEditorGraphNodeBase::OnPickColor()
 {
 	FColorPickerArgs PickerArgs;
diff --git a/Engine/Plugins/Experimental/PCG/Source/PCGEditor/Private/PCGEditorGraphNodeBase.h b/Engine/Plugins/Experimental/PCG/Source/PCGEditor/Private/PCGEditorGraphNodeBase.h
--- a/Engine/Plugins/Experimental/PCG/Source/PCGEditor/Private/PCGEditorGraphNodeBase.h
+++ b/Engine/Plugins/Experimental/PCG/Source/PCGEditor/Private/PCGEditorGraphNodeBase.h
@@ -58,6 +58,9 @@ public:
 	/** Pulls current errors/warnings state from PCG subsystem. */
 	void UpdateErrorsAndWarnings();
 
+	/** Puts node title on node body, reducing overall node size */
+	bool ShouldDrawCompact() const;
+
 	DECLARE_DELEGATE(FOnPCGEditorGraphNodeChanged);
 	FOnPCGEditorGraphNodeChanged OnNodeChangedDelegate;
 
diff --git a/Engine/Plugins/Experimental/PCG/Source/PCGEditor/Private/PCGEditorGraphSchema.cpp b/Engine/Plugins/Experimental/PCG/Source/PCGEditor/Private/PCGEditorGraphSchema.cpp
--- a/Engine/Plugins/Experimental/PCG/Source/PCGEditor/Private/PCGEditorGraphSchema.cpp
+++ b/Engine/Plugins/Experimental/PCG/Source/PCGEditor/Private/PCGEditorGraphSchema.cpp
@@ -2,10 +2,11 @@
 
 #include "PCGEditorGraphSchema.h"
 
-#include "Blueprint/BlueprintSupport.h"
 #include "PCGEdge.h"
 #include "PCGGraph.h"
 #include "PCGPin.h"
+#include "Elements/PCGCollapseElement.h"
+#include "Elements/PCGFilterByType.h"
 #include "Elements/PCGUserParameterGet.h"
 
 #include "PCGEditorCommon.h"
@@ -17,6 +18,7 @@
 #include "PCGEditorUtils.h"
 
 #include "AssetRegistry/AssetData.h"
+#include "Blueprint/BlueprintSupport.h"
 #include "Engine/Blueprint.h"
 #include "Framework/Application/SlateApplication.h"
 #include "ScopedTransaction.h"
@@ -115,6 +117,16 @@ const FPinConnectionResponse UPCGEditorGraphSchema::CanCreateConnection(const UE
 		return FPinConnectionResponse(CONNECT_RESPONSE_DISALLOW, LOCTEXT("ConnectionTypesIncompatible", "Pins are incompatible"));
 	}
 
+	if (InputPin->RequiresPointConversionToConnect(OutputPin))
+	{
+		return FPinConnectionResponse(CONNECT_RESPONSE_MAKE_WITH_CONVERSION_NODE, LOCTEXT("ConnectionConversionToPoint", "Convert to Point"));
+	}
+
+	if (InputPin->RequiresFilterToConnect(OutputPin))
+	{
+		return FPinConnectionResponse(CONNECT_RESPONSE_MAKE_WITH_CONVERSION_NODE, LOCTEXT("ConnectionUsingFilter", "Filter data to match type"));
+	}
+
 	if (!InputPin->AllowMultipleConnections() && InputPin->EdgeCount() > 0)
 	{
 		return FPinConnectionResponse((A->Direction == EGPD_Output) ? CONNECT_RESPONSE_BREAK_OTHERS_B : CONNECT_RESPONSE_BREAK_OTHERS_A, LOCTEXT("ConnectionBreakExisting", "Break existing connection?"));
@@ -125,34 +137,136 @@ const FPinConnectionResponse UPCGEditorGraphSchema::CanCreateConnection(const UE
 
 bool UPCGEditorGraphSchema::TryCreateConnection(UEdGraphPin* InA, UEdGraphPin* InB) const
 {
-	// TODO: check if we need to verify connectivity first
-	const bool bModified = Super::TryCreateConnection(InA, InB);
+	return TryCreateConnectionInternal(InA, InB, /*bAddConversionNodeIfNeeded=*/true);
+}
 
-	if (bModified)
+bool UPCGEditorGraphSchema::TryCreateConnectionInternal(UEdGraphPin* InA, UEdGraphPin* InB, bool bAddConversionNodeIfNeeded) const
+{
+	check(InA && InB);
+	if (InA->Direction == InB->Direction)
 	{
-		check(InA && InB);
-		const UEdGraphPin* A = (InA->Direction == EGPD_Output) ? InA : InB;
-		const UEdGraphPin* B = (InA->Direction == EGPD_Input) ? InA : InB;
-		
-		check(A->Direction == EGPD_Output && B->Direction == EGPD_Input);
+		// Don't connect same polarity
+		return false;
+	}
+
+	UEdGraphPin* A = (InA->Direction == EGPD_Output) ? InA : InB;
+	UEdGraphPin* B = (InA->Direction == EGPD_Input) ? InA : InB;
+	check(A->Direction == EGPD_Output && B->Direction == EGPD_Input);
+
+	UEdGraphNode* NodeA = A->GetOwningNode();
+	UEdGraphNode* NodeB = B->GetOwningNode();
+	check(NodeA && NodeB);
+
+	UPCGEditorGraphNodeBase* PCGEdGraphNodeA = CastChecked<UPCGEditorGraphNodeBase>(NodeA);
+	UPCGEditorGraphNodeBase* PCGEdGraphNodeB = CastChecked<UPCGEditorGraphNodeBase>(NodeB);
+
+	UPCGNode* PCGNodeA = PCGEdGraphNodeA->GetPCGNode();
+	UPCGNode* PCGNodeB = PCGEdGraphNodeB->GetPCGNode();
+	check(PCGNodeA && PCGNodeB);
+
+	const UPCGPin* PCGPinA = PCGNodeA->GetOutputPin(A->PinName);
+	const UPCGPin* PCGPinB = PCGNodeB->GetInputPin(B->PinName);
+	check(PCGPinA && PCGPinB);
+	if (!PCGPinA->IsCompatible(PCGPinB))
+	{
+		return false;
+	}
+
+	UPCGGraph* PCGGraph = PCGNodeA->GetGraph();
+	check(PCGGraph);
+
+	// Creates a connection via an intermediate conversion node.
+	auto ConnectViaIntermediate = [this, PCGGraph, NodeA, NodeB, A, B](UPCGNode* IntermediateNode)
+	{
+		UEdGraph* Graph = NodeA->GetGraph();
+		check(Graph);
+		Graph->Modify();
+
+		FGraphNodeCreator<UPCGEditorGraphNode> NodeCreator(*Graph);
+		UPCGEditorGraphNode* ConversionNode = NodeCreator.CreateUserInvokedNode(/*bSelectNewNode=*/false);
+		ConversionNode->Construct(IntermediateNode);
+
+		// Put the conversion node between A & B but make it stay within a radius of B to keep things tidy.
+		{
+			// Initially place at mid point
+			ConversionNode->NodePosX = (NodeA->NodePosX + NodeB->NodePosX) / 2;
+			ConversionNode->NodePosY = (NodeA->NodePosY + NodeB->NodePosY) / 2;
+
+			// A hand tweaked distance that keeps it reasonably close.
+			constexpr float MaxDistFromB = 200;
+			const FVector2D OffsetFromB(ConversionNode->NodePosX - NodeB->NodePosX, ConversionNode->NodePosY - NodeB->NodePosY);
+			const float Dist = OffsetFromB.Length();
+			if (Dist > MaxDistFromB)
+			{
+				const float Scale = MaxDistFromB / Dist;
+				ConversionNode->NodePosX = NodeB->NodePosX + Scale * OffsetFromB.X;
+				ConversionNode->NodePosY = NodeB->NodePosY + Scale * OffsetFromB.Y;
+			}
+		}
 
-		UEdGraphNode* NodeA = A->GetOwningNode();
-		UEdGraphNode* NodeB = B->GetOwningNode();
+		NodeCreator.Finalize();
 
-		UPCGEditorGraphNodeBase* PCGGraphNodeA = CastChecked<UPCGEditorGraphNodeBase>(NodeA);
-		UPCGEditorGraphNodeBase* PCGGraphNodeB = CastChecked<UPCGEditorGraphNodeBase>(NodeB);
+		IntermediateNode->PositionX = ConversionNode->NodePosX;
+		IntermediateNode->PositionY = ConversionNode->NodePosY;
 
-		UPCGNode* PCGNodeA = PCGGraphNodeA->GetPCGNode();
-		UPCGNode* PCGNodeB = PCGGraphNodeB->GetPCGNode();
-		check(PCGNodeA && PCGNodeB);
+		bool bModifiedA = false, bModifiedB = false;
 
-		UPCGGraph* PCGGraph = PCGNodeA->GetGraph();
-		check(PCGGraph);
+		UEdGraphPin*const* ConversionInputPin = ConversionNode->GetAllPins().FindByPredicate([](const UEdGraphPin* InPin)
+		{
+			return InPin->Direction == EGPD_Input && InPin->GetFName() == PCGPinConstants::DefaultInputLabel;
+		});
+
+		if (ensure(ConversionInputPin && *ConversionInputPin))
+		{
+			// Last argument: don't allow recursively adding conversion nodes.
+			bModifiedA = TryCreateConnectionInternal(A, *ConversionInputPin, /*bAddConversionNodeIfNeeded=*/false);
+		}
+
+		// Call GetAllPins() a second time. It's important that we wire up the pins one at a time. Wiring a pin can change dynamic pin types
+		// which can refresh the node, so we must re-query the pins after each connection is made.
+		UEdGraphPin*const* ConversionOutputPin = ConversionNode->GetAllPins().FindByPredicate([](const UEdGraphPin* InPin)
+		{
+			return InPin->Direction == EGPD_Output && InPin->GetFName() == PCGPinConstants::DefaultOutputLabel;
+		});
+
+		if (ensure(ConversionOutputPin && *ConversionOutputPin))
+		{
+			// Last argument: don't allow recursively adding conversion nodes.
+			bModifiedB = TryCreateConnectionInternal(*ConversionOutputPin, B, /*bAddConversionNodeIfNeeded=*/false);
+		}
 
-		PCGGraph->AddLabeledEdge(PCGNodeA, A->PinName, PCGNodeB, B->PinName);
+		return bModifiedA || bModifiedB;
+	};
+
+	if (bAddConversionNodeIfNeeded && PCGPinA->RequiresPointConversionToConnect(PCGPinB))
+	{
+		UPCGSettings* NodeSettings = nullptr;
+		UPCGNode* ConversionPCGNode = PCGGraph->AddNodeOfType(UPCGCollapseSettings::StaticClass(), NodeSettings);
+
+		return ConnectViaIntermediate(ConversionPCGNode);
+	}
+	else if (bAddConversionNodeIfNeeded && PCGPinA->RequiresFilterToConnect(PCGPinB))
+	{
+		UPCGSettings* NodeSettings = nullptr;
+		UPCGNode* ConversionPCGNode = PCGGraph->AddNodeOfType(UPCGFilterByTypeSettings::StaticClass(), NodeSettings);
+
+		// Setup the output pin based on the conversion target type, before new node is finalized.
+		UPCGFilterByTypeSettings* Settings = CastChecked<UPCGFilterByTypeSettings>(NodeSettings);
+		Settings->TargetType = PCGPinB->Properties.AllowedTypes;
+		ConversionPCGNode->UpdateAfterSettingsChangeDuringCreation();
+
+		return ConnectViaIntermediate(ConversionPCGNode);
 	}
+	else
+	{
+		const bool bModified = Super::TryCreateConnection(InA, InB);
+		if (bModified)
+		{
+			PCGGraph->AddLabeledEdge(PCGNodeA, A->PinName, PCGNodeB, B->PinName);
+		}
 
-	return bModified;
+		return bModified;
+	}
 }
 
 void UPCGEditorGraphSchema::BreakPinLinks(UEdGraphPin& TargetPin, bool bSendsNodeNotification) const
diff --git a/Engine/Plugins/Experimental/PCG/Source/PCGEditor/Private/PCGEditorGraphSchema.h b/Engine/Plugins/Experimental/PCG/Source/PCGEditor/Private/PCGEditorGraphSchema.h
--- a/Engine/Plugins/Experimental/PCG/Source/PCGEditor/Private/PCGEditorGraphSchema.h
+++ b/Engine/Plugins/Experimental/PCG/Source/PCGEditor/Private/PCGEditorGraphSchema.h
@@ -38,6 +38,8 @@ private:
 	void GetSubgraphElementActions(FGraphActionMenuBuilder& ActionMenuBuilder) const;
 	void GetSettingsElementActions(FGraphActionMenuBuilder& ActionMenuBuilder, bool bIsContextual) const;
 	void GetExtraElementActions(FGraphActionMenuBuilder& ActionMenuBuilder) const;
+
+	virtual bool TryCreateConnectionInternal(UEdGraphPin* A, UEdGraphPin* B, bool bAddConversionNodeIfNeeded) const;
 };
 
 class FPCGEditorConnectionDrawingPolicy : public FConnectionDrawingPolicy
diff --git a/Engine/Plugins/Experimental/PCG/Source/PCGEditor/Private/SPCGEditorGraphNode.cpp b/Engine/Plugins/Experimental/PCG/Source/PCGEditor/Private/SPCGEditorGraphNode.cpp
--- a/Engine/Plugins/Experimental/PCG/Source/PCGEditor/Private/SPCGEditorGraphNode.cpp
+++ b/Engine/Plugins/Experimental/PCG/Source/PCGEditor/Private/SPCGEditorGraphNode.cpp
@@ -8,9 +8,12 @@
 #include "PCGPin.h"
 #include "PCGSettings.h"
 
+#include "GraphEditorSettings.h"
+#include "SCommentBubble.h"
 #include "SGraphPin.h"
 #include "SLevelOfDetailBranchNode.h"
 #include "SPinTypeSelector.h"
+#include "TutorialMetaData.h"
 #include "Widgets/Images/SImage.h"
 #include "Widgets/SToolTip.h"
 #include "Widgets/Input/SButton.h"
@@ -18,6 +21,8 @@
 #include "Widgets/Text/SInlineEditableTextBlock.h"
 #include "Widgets/SBoxPanel.h"
 
+#define LOCTEXT_NAMESPACE "SPCGEditorGraphNode"
+
 /** PCG pin primarily to give more control over pin coloring. */
 class SPCGEditorGraphNodePin : public SGraphPin
 {
@@ -257,6 +262,18 @@ void SPCGEditorGraphNode::Construct(const FArguments& InArgs, UPCGEditorGraphNod
 	UpdateGraphNode();
 }
 
+void SPCGEditorGraphNode::UpdateGraphNode()
+{
+	if (PCGEditorGraphNode && PCGEditorGraphNode->ShouldDrawCompact())
+	{
+		UpdateCompactNode();
+	}
+	else
+	{
+		SGraphNode::UpdateGraphNode();
+	}
+}
+
 const FSlateBrush* SPCGEditorGraphNode::GetNodeBodyBrush() const
 {
 	if (PCGEditorGraphNode && PCGEditorGraphNode->GetPCGNode() && PCGEditorGraphNode->GetPCGNode()->IsInstance())
@@ -395,3 +412,196 @@ void SPCGEditorGraphNode::OnNodeChanged()
 {
 	UpdateGraphNode();
 }
+
+void SPCGEditorGraphNode::UpdateCompactNode()
+{
+	// Based on SGraphNodeK2Base::UpdateCompactNode. Changes:
+	// * Removed creation of tooltip widget, it did not port across trivially and the usage is fairly obvious for the current
+	//   compact nodes, but this could be re-added - TODO
+	// * Changed title style - reduced font size substantially
+	// * Layout differentiation for "pure" vs "impure" K2 nodes removed
+
+	InputPins.Empty();
+	OutputPins.Empty();
+
+	// Error handling set-up
+	SetupErrorReporting();
+
+	// Reset variables that are going to be exposed, in case we are refreshing an already setup node.
+	RightNodeBox.Reset();
+	LeftNodeBox.Reset();
+
+	// Setup a meta tag for this node
+	FGraphNodeMetaData TagMeta(TEXT("Graphnode"));
+	PopulateMetaTag(&TagMeta);
+	
+	TSharedPtr<SNodeTitle> NodeTitle = SNew(SNodeTitle, GraphNode);
+	TSharedRef<SOverlay> NodeOverlay = SNew(SOverlay);
+	
+	// Add optional node specific widget to the overlay:
+	TSharedPtr<SWidget> OverlayWidget = GraphNode->CreateNodeImage();
+	if(OverlayWidget.IsValid())
+	{
+		NodeOverlay->AddSlot()
+		.HAlign(HAlign_Center)
+		.VAlign(VAlign_Center)
+		[
+			SNew( SBox )
+			.WidthOverride( 70.f )
+			.HeightOverride( 70.f )
+			[
+				OverlayWidget.ToSharedRef()
+			]
+		];
+	}
+
+	NodeOverlay->AddSlot()
+		.HAlign(HAlign_Center)
+		.VAlign(VAlign_Center)
+		.Padding(45.f, 0.f, 45.f, 0.f)
+		[
+			// MIDDLE
+			SNew(SVerticalBox)
+			+SVerticalBox::Slot()
+			.HAlign(HAlign_Center)
+			.AutoHeight()
+			[
+				SNew(STextBlock)
+					.TextStyle(FAppStyle::Get(), "Graph.Node.NodeTitle" )
+					.Text( NodeTitle.Get(), &SNodeTitle::GetHeadTitle )
+					.WrapTextAt(128.0f)
+			]
+			+SVerticalBox::Slot()
+			.AutoHeight()
+			[
+				NodeTitle.ToSharedRef()
+			]
+		];
+	
+	NodeOverlay->AddSlot()
+		.HAlign(HAlign_Left)
+		.VAlign(VAlign_Center)
+		.Padding(0.f, 0.f, 55.f, 0.f)
+		[
+			// LEFT
+			SAssignNew(LeftNodeBox, SVerticalBox)
+		];
+
+	NodeOverlay->AddSlot()
+		.HAlign(HAlign_Right)
+		.VAlign(VAlign_Center)
+		.Padding(55.f, 0.f, 0.f, 0.f)
+		[
+			// RIGHT
+			SAssignNew(RightNodeBox, SVerticalBox)
+		];
+
+	//
+	//             ______________________
+	//            | (>) L |      | R (>) |
+	//            | (>) E |      | I (>) |
+	//            | (>) F |   +  | G (>) |
+	//            | (>) T |      | H (>) |
+	//            |       |      | T (>) |
+	//            |_______|______|_______|
+	//
+	this->ContentScale.Bind( this, &SGraphNode::GetContentScale );
+	
+	TSharedRef<SVerticalBox> InnerVerticalBox =
+		SNew(SVerticalBox)
+		+SVerticalBox::Slot()
+		[
+			// NODE CONTENT AREA
+			SNew( SOverlay)
+			+SOverlay::Slot()
+			[
+				SNew(SImage)
+				.Image( FAppStyle::GetBrush("Graph.VarNode.Body") )
+			]
+			+ SOverlay::Slot()
+			[
+				SNew(SImage)
+				.Image( FAppStyle::GetBrush("Graph.VarNode.Gloss") )
+			]
+			+SOverlay::Slot()
+			.Padding( FMargin(0,3) )
+			[
+				NodeOverlay
+			]
+		];
+	
+	TSharedPtr<SWidget> EnabledStateWidget = GetEnabledStateWidget();
+	if (EnabledStateWidget.IsValid())
+	{
+		InnerVerticalBox->AddSlot()
+			.AutoHeight()
+			.HAlign(HAlign_Fill)
+			.VAlign(VAlign_Top)
+			.Padding(FMargin(2, 0))
+			[
+				EnabledStateWidget.ToSharedRef()
+			];
+	}
+
+	InnerVerticalBox->AddSlot()
+		.AutoHeight()
+		.Padding( FMargin(5.0f, 1.0f) )
+		[
+			ErrorReporting->AsWidget()
+		];
+
+	this->GetOrAddSlot( ENodeZone::Center )
+		.HAlign(HAlign_Center)
+		.VAlign(VAlign_Center)
+		[
+			InnerVerticalBox
+		];
+
+	CreatePinWidgets();
+
+	// Hide pin labels
+	for (auto InputPin: this->InputPins)
+	{
+		if (InputPin->GetPinObj()->ParentPin == nullptr)
+		{
+			InputPin->SetShowLabel(false);
+		}
+	}
+
+	for (auto OutputPin : this->OutputPins)
+	{
+		if (OutputPin->GetPinObj()->ParentPin == nullptr)
+		{
+			OutputPin->SetShowLabel(false);
+		}
+	}
+
+	// Create comment bubble
+	TSharedPtr<SCommentBubble> CommentBubble;
+	const FSlateColor CommentColor = GetDefault<UGraphEditorSettings>()->DefaultCommentNodeTitleColor;
+
+	SAssignNew( CommentBubble, SCommentBubble )
+		.GraphNode( GraphNode )
+		.Text( this, &SGraphNode::GetNodeComment )
+		.OnTextCommitted( this, &SGraphNode::OnCommentTextCommitted )
+		.ColorAndOpacity( CommentColor )
+		.AllowPinning( true )
+		.EnableTitleBarBubble( true )
+		.EnableBubbleCtrls( true )
+		.GraphLOD( this, &SGraphNode::GetCurrentLOD )
+		.IsGraphNodeHovered( this, &SGraphNode::IsHovered );
+
+	GetOrAddSlot( ENodeZone::TopCenter )
+		.SlotOffset(TAttribute<FVector2D>(CommentBubble.Get(), &SCommentBubble::GetOffset))
+		.SlotSize(TAttribute<FVector2D>(CommentBubble.Get(), &SCommentBubble::GetSize))
+		.AllowScaling(TAttribute<bool>(CommentBubble.Get(), &SCommentBubble::IsScalingAllowed))
+		.VAlign(VAlign_Top)
+		[
+			CommentBubble.ToSharedRef()
+		];
+
+	CreateInputSideAddButton(LeftNodeBox);
+	CreateOutputSideAddButton(RightNodeBox);
+}
+
+#undef LOCTEXT_NAMESPACE
diff --git a/Engine/Plugins/Experimental/PCG/Source/PCGEditor/Private/SPCGEditorGraphNode.h b/Engine/Plugins/Experimental/PCG/Source/PCGEditor/Private/SPCGEditorGraphNode.h
--- a/Engine/Plugins/Experimental/PCG/Source/PCGEditor/Private/SPCGEditorGraphNode.h
+++ b/Engine/Plugins/Experimental/PCG/Source/PCGEditor/Private/SPCGEditorGraphNode.h
@@ -17,6 +17,7 @@ public:
 	void Construct(const FArguments& InArgs, UPCGEditorGraphNodeBase* InNode);
 
 	//~ Begin SGraphNode Interface
+	virtual void UpdateGraphNode() override;
 	virtual const FSlateBrush* GetNodeBodyBrush() const override;
 	virtual void RequestRenameOnSpawn() override { /* Empty to avoid the default behavior to rename on node spawn */ }
 	virtual void AddPin(const TSharedRef<SGraphPin>& PinToAdd) override;
@@ -31,6 +32,9 @@ public:
 protected:
 	void OnNodeChanged();
 
+	/** Set up node in 'compact' mode */
+	void UpdateCompactNode();
+
 private:
 	UPCGEditorGraphNodeBase* PCGEditorGraphNode = nullptr;
 };
