[PCG] Add support for pre-configured nodes, native and blueprint

* Add pre-configured nodes for all metadata op nodes.

diff --git a/Engine/Plugins/Experimental/PCG/Source/PCG/Private/Elements/Metadata/PCGMetadataBitwiseOpElement.cpp b/Engine/Plugins/Experimental/PCG/Source/PCG/Private/Elements/Metadata/PCGMetadataBitwiseOpElement.cpp
--- a/Engine/Plugins/Experimental/PCG/Source/PCG/Private/Elements/Metadata/PCGMetadataBitwiseOpElement.cpp
+++ b/Engine/Plugins/Experimental/PCG/Source/PCG/Private/Elements/Metadata/PCGMetadataBitwiseOpElement.cpp
@@ -2,6 +2,7 @@
 
 #include "Elements/Metadata/PCGMetadataBitwiseOpElement.h"
 
+#include "Elements/Metadata/PCGMetadataElementCommon.h"
 
 #include UE_INLINE_GENERATED_CPP_BY_NAME(PCGMetadataBitwiseOpElement)
 
@@ -107,8 +108,24 @@ FText UPCGMetadataBitwiseSettings::GetDefaultNodeTitle() const
 {
 	return NSLOCTEXT("PCGMetadataBitwiseSettings", "NodeTitle", "Attribute Bitwise Op");
 }
+
+TArray<FPCGPreConfiguredSettingsInfo> UPCGMetadataBitwiseSettings::GetPreconfiguredInfo() const
+{
+	return PCGMetadataElementCommon::FillPreconfiguredSettingsInfoFromEnum<EPCGMedadataBitwiseOperation>();
+}
 #endif // WITH_EDITOR
 
+void UPCGMetadataBitwiseSettings::ApplyPreconfiguredSettings(const FPCGPreConfiguredSettingsInfo& PreconfiguredInfo)
+{
+	if (const UEnum* EnumPtr = StaticEnum<EPCGMedadataBitwiseOperation>())
+	{
+		if (EnumPtr->IsValidEnumValue(PreconfiguredInfo.PreconfiguredIndex))
+		{
+			Operation = EPCGMedadataBitwiseOperation(PreconfiguredInfo.PreconfiguredIndex);
+		}
+	}
+}
+
 FPCGElementPtr UPCGMetadataBitwiseSettings::CreateElement() const
 {
 	return MakeShared<FPCGMetadataBitwiseElement>();
diff --git a/Engine/Plugins/Experimental/PCG/Source/PCG/Private/Elements/Metadata/PCGMetadataBooleanOpElement.cpp b/Engine/Plugins/Experimental/PCG/Source/PCG/Private/Elements/Metadata/PCGMetadataBooleanOpElement.cpp
--- a/Engine/Plugins/Experimental/PCG/Source/PCG/Private/Elements/Metadata/PCGMetadataBooleanOpElement.cpp
+++ b/Engine/Plugins/Experimental/PCG/Source/PCG/Private/Elements/Metadata/PCGMetadataBooleanOpElement.cpp
@@ -2,6 +2,7 @@
 
 #include "Elements/Metadata/PCGMetadataBooleanOpElement.h"
 
+#include "Elements/Metadata/PCGMetadataElementCommon.h"
 
 #include UE_INLINE_GENERATED_CPP_BY_NAME(PCGMetadataBooleanOpElement)
 
@@ -92,7 +93,7 @@ uint16 UPCGMetadataBooleanSettings::GetOutputType(uint16 InputTypeId) const
 
 FName UPCGMetadataBooleanSettings::AdditionalTaskName() const
 {
-	if (const UEnum* EnumPtr = FindObject<UEnum>(nullptr, TEXT("/Script/PCG.EPCGMedadataBooleanOperation"), true))
+	if (const UEnum* EnumPtr = StaticEnum<EPCGMedadataBooleanOperation>())
 	{
 		return FName(FString("Boolean: ") + EnumPtr->GetNameStringByValue(static_cast<int>(Operation)));
 	}
@@ -112,8 +113,24 @@ FText UPCGMetadataBooleanSettings::GetDefaultNodeTitle() const
 {
 	return NSLOCTEXT("PCGMetadataBooleanSettings", "NodeTitle", "Attribute Boolean Op");
 }
+
+TArray<FPCGPreConfiguredSettingsInfo> UPCGMetadataBooleanSettings::GetPreconfiguredInfo() const
+{
+	return PCGMetadataElementCommon::FillPreconfiguredSettingsInfoFromEnum<EPCGMedadataBooleanOperation>();
+}
 #endif // WITH_EDITOR
 
+void UPCGMetadataBooleanSettings::ApplyPreconfiguredSettings(const FPCGPreConfiguredSettingsInfo& PreconfiguredInfo)
+{
+	if (const UEnum* EnumPtr = StaticEnum<EPCGMedadataBooleanOperation>())
+	{
+		if (EnumPtr->IsValidEnumValue(PreconfiguredInfo.PreconfiguredIndex))
+		{
+			Operation = EPCGMedadataBooleanOperation(PreconfiguredInfo.PreconfiguredIndex);
+		}
+	}
+}
+
 FPCGElementPtr UPCGMetadataBooleanSettings::CreateElement() const
 {
 	return MakeShared<FPCGMetadataBooleanElement>();
diff --git a/Engine/Plugins/Experimental/PCG/Source/PCG/Private/Elements/Metadata/PCGMetadataCompareOpElement.cpp b/Engine/Plugins/Experimental/PCG/Source/PCG/Private/Elements/Metadata/PCGMetadataCompareOpElement.cpp
--- a/Engine/Plugins/Experimental/PCG/Source/PCG/Private/Elements/Metadata/PCGMetadataCompareOpElement.cpp
+++ b/Engine/Plugins/Experimental/PCG/Source/PCG/Private/Elements/Metadata/PCGMetadataCompareOpElement.cpp
@@ -2,6 +2,7 @@
 
 #include "Elements/Metadata/PCGMetadataCompareOpElement.h"
 
+#include "Elements/Metadata/PCGMetadataElementCommon.h"
 #include "Metadata/PCGMetadataAttributeTpl.h"
 
 #include UE_INLINE_GENERATED_CPP_BY_NAME(PCGMetadataCompareOpElement)
@@ -138,8 +139,24 @@ FText UPCGMetadataCompareSettings::GetDefaultNodeTitle() const
 {
 	return NSLOCTEXT("PCGMetadataCompareSettings", "NodeTitle", "Attribute Compare Op");
 }
+
+TArray<FPCGPreConfiguredSettingsInfo> UPCGMetadataCompareSettings::GetPreconfiguredInfo() const
+{
+	return PCGMetadataElementCommon::FillPreconfiguredSettingsInfoFromEnum<EPCGMedadataCompareOperation>();
+}
 #endif // WITH_EDITOR
 
+void UPCGMetadataCompareSettings::ApplyPreconfiguredSettings(const FPCGPreConfiguredSettingsInfo& PreconfiguredInfo)
+{
+	if (const UEnum* EnumPtr = StaticEnum<EPCGMedadataCompareOperation>())
+	{
+		if (EnumPtr->IsValidEnumValue(PreconfiguredInfo.PreconfiguredIndex))
+		{
+			Operation = EPCGMedadataCompareOperation(PreconfiguredInfo.PreconfiguredIndex);
+		}
+	}
+}
+
 FPCGElementPtr UPCGMetadataCompareSettings::CreateElement() const
 {
 	return MakeShared<FPCGMetadataCompareElement>();
diff --git a/Engine/Plugins/Experimental/PCG/Source/PCG/Private/Elements/Metadata/PCGMetadataMathsOpElement.cpp b/Engine/Plugins/Experimental/PCG/Source/PCG/Private/Elements/Metadata/PCGMetadataMathsOpElement.cpp
--- a/Engine/Plugins/Experimental/PCG/Source/PCG/Private/Elements/Metadata/PCGMetadataMathsOpElement.cpp
+++ b/Engine/Plugins/Experimental/PCG/Source/PCG/Private/Elements/Metadata/PCGMetadataMathsOpElement.cpp
@@ -2,6 +2,7 @@
 
 #include "Elements/Metadata/PCGMetadataMathsOpElement.h"
 
+#include "Elements/Metadata/PCGMetadataElementCommon.h"
 #include "Metadata/PCGMetadataAttributeTpl.h"
 
 #include "Elements/Metadata/PCGMetadataMaths.inl"
@@ -254,8 +255,24 @@ FText UPCGMetadataMathsSettings::GetDefaultNodeTitle() const
 {
 	return NSLOCTEXT("PCGMetadataMathsSettings", "NodeTitle", "Attribute Maths Op");
 }
+
+TArray<FPCGPreConfiguredSettingsInfo> UPCGMetadataMathsSettings::GetPreconfiguredInfo() const
+{
+	return PCGMetadataElementCommon::FillPreconfiguredSettingsInfoFromEnum<EPCGMedadataMathsOperation>({ EPCGMedadataMathsOperation::UnaryOp, EPCGMedadataMathsOperation::BinaryOp, EPCGMedadataMathsOperation::TernaryOp });
+}
 #endif // WITH_EDITOR
 
+void UPCGMetadataMathsSettings::ApplyPreconfiguredSettings(const FPCGPreConfiguredSettingsInfo& PreconfiguredInfo)
+{
+	if (const UEnum* EnumPtr = StaticEnum<EPCGMedadataMathsOperation>())
+	{
+		if (EnumPtr->IsValidEnumValue(PreconfiguredInfo.PreconfiguredIndex))
+		{
+			Operation = EPCGMedadataMathsOperation(PreconfiguredInfo.PreconfiguredIndex);
+		}
+	}
+}
+
 FPCGElementPtr UPCGMetadataMathsSettings::CreateElement() const
 {
 	return MakeShared<FPCGMetadataMathsElement>();
diff --git a/Engine/Plugins/Experimental/PCG/Source/PCG/Private/Elements/Metadata/PCGMetadataRotatorOpElement.cpp b/Engine/Plugins/Experimental/PCG/Source/PCG/Private/Elements/Metadata/PCGMetadataRotatorOpElement.cpp
--- a/Engine/Plugins/Experimental/PCG/Source/PCG/Private/Elements/Metadata/PCGMetadataRotatorOpElement.cpp
+++ b/Engine/Plugins/Experimental/PCG/Source/PCG/Private/Elements/Metadata/PCGMetadataRotatorOpElement.cpp
@@ -2,6 +2,7 @@
 
 #include "Elements/Metadata/PCGMetadataRotatorOpElement.h"
 
+#include "Elements/Metadata/PCGMetadataElementCommon.h"
 #include "Metadata/PCGMetadataAttributeTpl.h"
 
 #include UE_INLINE_GENERATED_CPP_BY_NAME(PCGMetadataRotatorOpElement)
@@ -190,7 +191,7 @@ FPCGAttributePropertySelector UPCGMetadataRotatorSettings::GetInputSource(uint32
 
 FName UPCGMetadataRotatorSettings::AdditionalTaskName() const
 {
-	if (const UEnum* EnumPtr = FindObject<UEnum>(nullptr, TEXT("/Script/PCG.EPCGMedadataRotatorOperation"), true))
+	if (const UEnum* EnumPtr = StaticEnum<EPCGMedadataRotatorOperation>())
 	{
 		return FName(FString("Rotator: ") + EnumPtr->GetNameStringByValue(static_cast<int>(Operation)));
 	}
@@ -210,8 +211,24 @@ FText UPCGMetadataRotatorSettings::GetDefaultNodeTitle() const
 {
 	return NSLOCTEXT("PCGMetadataRotatorSettings", "NodeTitle", "Attribute Rotator Op");
 }
+
+TArray<FPCGPreConfiguredSettingsInfo> UPCGMetadataRotatorSettings::GetPreconfiguredInfo() const
+{
+	return PCGMetadataElementCommon::FillPreconfiguredSettingsInfoFromEnum<EPCGMedadataRotatorOperation>({ EPCGMedadataRotatorOperation::RotatorOp, EPCGMedadataRotatorOperation::TransformOp });
+}
 #endif // WITH_EDITOR
 
+void UPCGMetadataRotatorSettings::ApplyPreconfiguredSettings(const FPCGPreConfiguredSettingsInfo& PreconfiguredInfo)
+{
+	if (const UEnum* EnumPtr = StaticEnum<EPCGMedadataRotatorOperation>())
+	{
+		if (EnumPtr->IsValidEnumValue(PreconfiguredInfo.PreconfiguredIndex))
+		{
+			Operation = EPCGMedadataRotatorOperation(PreconfiguredInfo.PreconfiguredIndex);
+		}
+	}
+}
+
 FPCGElementPtr UPCGMetadataRotatorSettings::CreateElement() const
 {
 	return MakeShared<FPCGMetadataRotatorElement>();
diff --git a/Engine/Plugins/Experimental/PCG/Source/PCG/Private/Elements/Metadata/PCGMetadataTransformOpElement.cpp b/Engine/Plugins/Experimental/PCG/Source/PCG/Private/Elements/Metadata/PCGMetadataTransformOpElement.cpp
--- a/Engine/Plugins/Experimental/PCG/Source/PCG/Private/Elements/Metadata/PCGMetadataTransformOpElement.cpp
+++ b/Engine/Plugins/Experimental/PCG/Source/PCG/Private/Elements/Metadata/PCGMetadataTransformOpElement.cpp
@@ -2,6 +2,7 @@
 
 #include "Elements/Metadata/PCGMetadataTransformOpElement.h"
 
+#include "Elements/Metadata/PCGMetadataElementCommon.h"
 #include "Elements/Metadata/PCGMetadataRotatorOpElement.h"
 
 #include "Math/DualQuat.h"
@@ -158,7 +159,7 @@ FPCGAttributePropertySelector UPCGMetadataTransformSettings::GetInputSource(uint
 
 FName UPCGMetadataTransformSettings::AdditionalTaskName() const
 {
-	if (const UEnum* EnumPtr = FindObject<UEnum>(nullptr, TEXT("/Script/PCG.EPCGMedadataTransformOperation"), true))
+	if (const UEnum* EnumPtr = StaticEnum<EPCGMedadataTransformOperation>())
 	{
 		return FName(FString("Transform: ") + EnumPtr->GetNameStringByValue(static_cast<int>(Operation)));
 	}
@@ -178,8 +179,24 @@ FText UPCGMetadataTransformSettings::GetDefaultNodeTitle() const
 {
 	return NSLOCTEXT("PCGMetadataTransformSettings", "NodeTitle", "Attribute Transform Op");
 }
+
+TArray<FPCGPreConfiguredSettingsInfo> UPCGMetadataTransformSettings::GetPreconfiguredInfo() const
+{
+	return PCGMetadataElementCommon::FillPreconfiguredSettingsInfoFromEnum<EPCGMedadataTransformOperation>();
+}
 #endif // WITH_EDITOR
 
+void UPCGMetadataTransformSettings::ApplyPreconfiguredSettings(const FPCGPreConfiguredSettingsInfo& PreconfiguredInfo)
+{
+	if (const UEnum* EnumPtr = StaticEnum<EPCGMedadataTransformOperation>())
+	{
+		if (EnumPtr->IsValidEnumValue(PreconfiguredInfo.PreconfiguredIndex))
+		{
+			Operation = EPCGMedadataTransformOperation(PreconfiguredInfo.PreconfiguredIndex);
+		}
+	}
+}
+
 FPCGElementPtr UPCGMetadataTransformSettings::CreateElement() const
 {
 	return MakeShared<FPCGMetadataTransformElement>();
diff --git a/Engine/Plugins/Experimental/PCG/Source/PCG/Private/Elements/Metadata/PCGMetadataTrigOpElement.cpp b/Engine/Plugins/Experimental/PCG/Source/PCG/Private/Elements/Metadata/PCGMetadataTrigOpElement.cpp
--- a/Engine/Plugins/Experimental/PCG/Source/PCG/Private/Elements/Metadata/PCGMetadataTrigOpElement.cpp
+++ b/Engine/Plugins/Experimental/PCG/Source/PCG/Private/Elements/Metadata/PCGMetadataTrigOpElement.cpp
@@ -2,6 +2,7 @@
 
 #include "Elements/Metadata/PCGMetadataTrigOpElement.h"
 
+#include "Elements/Metadata/PCGMetadataElementCommon.h"
 
 #include UE_INLINE_GENERATED_CPP_BY_NAME(PCGMetadataTrigOpElement)
 
@@ -104,7 +105,7 @@ uint16 UPCGMetadataTrigSettings::GetOutputType(uint16 InputTypeId) const
 
 FName UPCGMetadataTrigSettings::AdditionalTaskName() const
 {
-	if (const UEnum* EnumPtr = FindObject<UEnum>(nullptr, TEXT("/Script/PCG.EPCGMedadataTrigOperation"), true))
+	if (const UEnum* EnumPtr = StaticEnum<EPCGMedadataTrigOperation>())
 	{
 		return FName(FString("Trig: ") + EnumPtr->GetNameStringByValue(static_cast<int>(Operation)));
 	}
@@ -124,8 +125,24 @@ FText UPCGMetadataTrigSettings::GetDefaultNodeTitle() const
 {
 	return NSLOCTEXT("PCGMetadataTrigSettings", "NodeTitle", "Attribute Trig Op");
 }
+
+TArray<FPCGPreConfiguredSettingsInfo> UPCGMetadataTrigSettings::GetPreconfiguredInfo() const
+{
+	return PCGMetadataElementCommon::FillPreconfiguredSettingsInfoFromEnum<EPCGMedadataTrigOperation>();
+}
 #endif // WITH_EDITOR
 
+void UPCGMetadataTrigSettings::ApplyPreconfiguredSettings(const FPCGPreConfiguredSettingsInfo& PreconfiguredInfo)
+{
+	if (const UEnum* EnumPtr = StaticEnum<EPCGMedadataTrigOperation>())
+	{
+		if (EnumPtr->IsValidEnumValue(PreconfiguredInfo.PreconfiguredIndex))
+		{
+			Operation = EPCGMedadataTrigOperation(PreconfiguredInfo.PreconfiguredIndex);
+		}
+	}
+}
+
 FPCGElementPtr UPCGMetadataTrigSettings::CreateElement() const
 {
 	return MakeShared<FPCGMetadataTrigElement>();
diff --git a/Engine/Plugins/Experimental/PCG/Source/PCG/Private/Elements/Metadata/PCGMetadataVectorOpElement.cpp b/Engine/Plugins/Experimental/PCG/Source/PCG/Private/Elements/Metadata/PCGMetadataVectorOpElement.cpp
--- a/Engine/Plugins/Experimental/PCG/Source/PCG/Private/Elements/Metadata/PCGMetadataVectorOpElement.cpp
+++ b/Engine/Plugins/Experimental/PCG/Source/PCG/Private/Elements/Metadata/PCGMetadataVectorOpElement.cpp
@@ -2,6 +2,7 @@
 
 #include "Elements/Metadata/PCGMetadataVectorOpElement.h"
 
+#include "Elements/Metadata/PCGMetadataElementCommon.h"
 #include "Metadata/PCGMetadataAttributeTpl.h"
 
 #include UE_INLINE_GENERATED_CPP_BY_NAME(PCGMetadataVectorOpElement)
@@ -320,7 +321,7 @@ FPCGAttributePropertySelector UPCGMetadataVectorSettings::GetInputSource(uint32
 
 FName UPCGMetadataVectorSettings::AdditionalTaskName() const
 {
-	if (const UEnum* EnumPtr = FindObject<UEnum>(nullptr, TEXT("/Script/PCG.EPCGMedadataVectorOperation"), true))
+	if (const UEnum* EnumPtr = StaticEnum<EPCGMedadataVectorOperation>())
 	{
 		return FName(FString("Vector: ") + EnumPtr->GetNameStringByValue(static_cast<int>(Operation)));
 	}
@@ -340,8 +341,24 @@ FText UPCGMetadataVectorSettings::GetDefaultNodeTitle() const
 {
 	return NSLOCTEXT("PCGMetadataVectorSettings", "NodeTitle", "Attribute Vector Op");
 }
+
+TArray<FPCGPreConfiguredSettingsInfo> UPCGMetadataVectorSettings::GetPreconfiguredInfo() const
+{
+	return PCGMetadataElementCommon::FillPreconfiguredSettingsInfoFromEnum<EPCGMedadataVectorOperation>({ EPCGMedadataVectorOperation::VectorOp, EPCGMedadataVectorOperation::TransformOp });
+}
 #endif // WITH_EDITOR
 
+void UPCGMetadataVectorSettings::ApplyPreconfiguredSettings(const FPCGPreConfiguredSettingsInfo& PreconfiguredInfo)
+{
+	if (const UEnum* EnumPtr = StaticEnum<EPCGMedadataVectorOperation>())
+	{
+		if (EnumPtr->IsValidEnumValue(PreconfiguredInfo.PreconfiguredIndex))
+		{
+			Operation = EPCGMedadataVectorOperation(PreconfiguredInfo.PreconfiguredIndex);
+		}
+	}
+}
+
 FPCGElementPtr UPCGMetadataVectorSettings::CreateElement() const
 {
 	return MakeShared<FPCGMetadataVectorElement>();
diff --git a/Engine/Plugins/Experimental/PCG/Source/PCG/Private/Elements/PCGExecuteBlueprint.cpp b/Engine/Plugins/Experimental/PCG/Source/PCG/Private/Elements/PCGExecuteBlueprint.cpp
--- a/Engine/Plugins/Experimental/PCG/Source/PCG/Private/Elements/PCGExecuteBlueprint.cpp
+++ b/Engine/Plugins/Experimental/PCG/Source/PCG/Private/Elements/PCGExecuteBlueprint.cpp
@@ -439,8 +439,26 @@ void UPCGBlueprintSettings::ApplyDeprecationBeforeUpdatePins(UPCGNode* InOutNode
 		}
 	}
 }
+
+TArray<FPCGPreConfiguredSettingsInfo> UPCGBlueprintSettings::GetPreconfiguredInfo() const
+{
+	return BlueprintElementInstance ? BlueprintElementInstance->PreconfiguredInfo : TArray<FPCGPreConfiguredSettingsInfo>();
+}
+
+bool UPCGBlueprintSettings::OnlyExposePreconfiguredSettings() const
+{
+	return BlueprintElementInstance && BlueprintElementInstance->bOnlyExposePreconfiguredSettings;
+}
 #endif // WITH_EDITOR
 
+void UPCGBlueprintSettings::ApplyPreconfiguredSettings(const FPCGPreConfiguredSettingsInfo& InPreconfiguredInfo)
+{
+	if (BlueprintElementInstance)
+	{
+		BlueprintElementInstance->ApplyPreconfiguredSettings(InPreconfiguredInfo);
+	}
+}
+
 FName UPCGBlueprintSettings::AdditionalTaskName() const
 {
 	if (BlueprintElementInstance && BlueprintElementInstance->NodeTitleOverride() != NAME_None)
diff --git a/Engine/Plugins/Experimental/PCG/Source/PCG/Public/Elements/Metadata/PCGMetadataBitwiseOpElement.h b/Engine/Plugins/Experimental/PCG/Source/PCG/Public/Elements/Metadata/PCGMetadataBitwiseOpElement.h
--- a/Engine/Plugins/Experimental/PCG/Source/PCG/Public/Elements/Metadata/PCGMetadataBitwiseOpElement.h
+++ b/Engine/Plugins/Experimental/PCG/Source/PCG/Public/Elements/Metadata/PCGMetadataBitwiseOpElement.h
@@ -29,8 +29,11 @@ public:
 #if WITH_EDITOR
 	virtual FName GetDefaultNodeName() const override;
 	virtual FText GetDefaultNodeTitle() const override;
+	virtual TArray<FPCGPreConfiguredSettingsInfo> GetPreconfiguredInfo() const override;
+	virtual bool OnlyExposePreconfiguredSettings() const override { return true; }
 #endif
 	virtual FName AdditionalTaskName() const override;
+	virtual void ApplyPreconfiguredSettings(const FPCGPreConfiguredSettingsInfo& PreconfigureInfo) override;
 	//~End UPCGSettings interface
 
 	//~Begin UPCGMetadataSettingsBase interface
diff --git a/Engine/Plugins/Experimental/PCG/Source/PCG/Public/Elements/Metadata/PCGMetadataBooleanOpElement.h b/Engine/Plugins/Experimental/PCG/Source/PCG/Public/Elements/Metadata/PCGMetadataBooleanOpElement.h
--- a/Engine/Plugins/Experimental/PCG/Source/PCG/Public/Elements/Metadata/PCGMetadataBooleanOpElement.h
+++ b/Engine/Plugins/Experimental/PCG/Source/PCG/Public/Elements/Metadata/PCGMetadataBooleanOpElement.h
@@ -29,8 +29,11 @@ public:
 #if WITH_EDITOR
 	virtual FName GetDefaultNodeName() const override;
 	virtual FText GetDefaultNodeTitle() const override;
+	virtual TArray<FPCGPreConfiguredSettingsInfo> GetPreconfiguredInfo() const override;
+	virtual bool OnlyExposePreconfiguredSettings() const override { return true; }
 #endif
 	virtual FName AdditionalTaskName() const override;
+	virtual void ApplyPreconfiguredSettings(const FPCGPreConfiguredSettingsInfo& PreconfigureInfo) override;
 	//~End UPCGSettings interface
 
 	//~Begin UPCGMetadataSettingsBase interface
diff --git a/Engine/Plugins/Experimental/PCG/Source/PCG/Public/Elements/Metadata/PCGMetadataCompareOpElement.h b/Engine/Plugins/Experimental/PCG/Source/PCG/Public/Elements/Metadata/PCGMetadataCompareOpElement.h
--- a/Engine/Plugins/Experimental/PCG/Source/PCG/Public/Elements/Metadata/PCGMetadataCompareOpElement.h
+++ b/Engine/Plugins/Experimental/PCG/Source/PCG/Public/Elements/Metadata/PCGMetadataCompareOpElement.h
@@ -31,8 +31,11 @@ public:
 #if WITH_EDITOR
 	virtual FName GetDefaultNodeName() const override;
 	virtual FText GetDefaultNodeTitle() const override;
+	virtual TArray<FPCGPreConfiguredSettingsInfo> GetPreconfiguredInfo() const override;
+	virtual bool OnlyExposePreconfiguredSettings() const override { return true; }
 #endif
 	virtual FName AdditionalTaskName() const override;
+	virtual void ApplyPreconfiguredSettings(const FPCGPreConfiguredSettingsInfo& PreconfigureInfo) override;
 	//~End UPCGSettings interface
 
 	//~Begin UPCGMetadataSettingsBase interface
diff --git a/Engine/Plugins/Experimental/PCG/Source/PCG/Public/Elements/Metadata/PCGMetadataElementCommon.h b/Engine/Plugins/Experimental/PCG/Source/PCG/Public/Elements/Metadata/PCGMetadataElementCommon.h
--- a/Engine/Plugins/Experimental/PCG/Source/PCG/Public/Elements/Metadata/PCGMetadataElementCommon.h
+++ b/Engine/Plugins/Experimental/PCG/Source/PCG/Public/Elements/Metadata/PCGMetadataElementCommon.h
@@ -2,6 +2,7 @@
 
 #pragma once
 
+#include "PCGSettings.h"
 #include "Metadata/PCGMetadata.h"
 #include "Metadata/Accessors/IPCGAttributeAccessor.h" // IWYU pragma: keep
 #include "Metadata/Accessors/PCGAttributeAccessorKeys.h"
@@ -83,4 +84,30 @@ namespace PCGMetadataElementCommon
 
 		return true;
 	}
+
+	/**
+	* Automatically fill all preconfigured settings depending on the enum operation.
+	* Can also specify explicitly values that should not be included, as metadata is not available in non-editor builds.
+	*/
+	template <typename EnumOperation>
+	TArray<FPCGPreConfiguredSettingsInfo> FillPreconfiguredSettingsInfoFromEnum(const TSet<EnumOperation>& InValuesToSkip = {})
+	{
+		TArray<FPCGPreConfiguredSettingsInfo> PreconfiguredInfo;
+
+		if (const UEnum* EnumPtr = StaticEnum<EnumOperation>())
+		{
+			PreconfiguredInfo.Reserve(EnumPtr->NumEnums());
+			for (int32 i = 0; i < EnumPtr->NumEnums(); ++i)
+			{
+				int64 Value = EnumPtr->GetValueByIndex(i);
+
+				if (Value != EnumPtr->GetMaxEnumValue() && !InValuesToSkip.Contains(EnumOperation(Value)))
+				{
+					PreconfiguredInfo.Emplace(Value, EnumPtr->GetDisplayNameTextByValue(Value));
+				}
+			}
+		}
+
+		return PreconfiguredInfo;
+	}
 }
diff --git a/Engine/Plugins/Experimental/PCG/Source/PCG/Public/Elements/Metadata/PCGMetadataMathsOpElement.h b/Engine/Plugins/Experimental/PCG/Source/PCG/Public/Elements/Metadata/PCGMetadataMathsOpElement.h
--- a/Engine/Plugins/Experimental/PCG/Source/PCG/Public/Elements/Metadata/PCGMetadataMathsOpElement.h
+++ b/Engine/Plugins/Experimental/PCG/Source/PCG/Public/Elements/Metadata/PCGMetadataMathsOpElement.h
@@ -54,8 +54,11 @@ public:
 #if WITH_EDITOR
 	virtual FName GetDefaultNodeName() const override;
 	virtual FText GetDefaultNodeTitle() const override;
+	virtual TArray<FPCGPreConfiguredSettingsInfo> GetPreconfiguredInfo() const override;
+	virtual bool OnlyExposePreconfiguredSettings() const override { return true; }
 #endif
 	virtual FName AdditionalTaskName() const override;
+	virtual void ApplyPreconfiguredSettings(const FPCGPreConfiguredSettingsInfo& PreconfigureInfo) override;
 	//~End UPCGSettings interface
 
 	//~Begin UPCGMetadataSettingsBase interface
diff --git a/Engine/Plugins/Experimental/PCG/Source/PCG/Public/Elements/Metadata/PCGMetadataRotatorOpElement.h b/Engine/Plugins/Experimental/PCG/Source/PCG/Public/Elements/Metadata/PCGMetadataRotatorOpElement.h
--- a/Engine/Plugins/Experimental/PCG/Source/PCG/Public/Elements/Metadata/PCGMetadataRotatorOpElement.h
+++ b/Engine/Plugins/Experimental/PCG/Source/PCG/Public/Elements/Metadata/PCGMetadataRotatorOpElement.h
@@ -39,9 +39,11 @@ public:
 #if WITH_EDITOR
 	virtual FName GetDefaultNodeName() const override;
 	virtual FText GetDefaultNodeTitle() const override;
+	virtual TArray<FPCGPreConfiguredSettingsInfo> GetPreconfiguredInfo() const override;
+	virtual bool OnlyExposePreconfiguredSettings() const override { return true; }
 #endif
-
 	virtual FName AdditionalTaskName() const override;
+	virtual void ApplyPreconfiguredSettings(const FPCGPreConfiguredSettingsInfo& PreconfigureInfo) override;
 	//~End UPCGSettings interface
 
 	//~Begin UPCGMetadataSettingsBase interface
diff --git a/Engine/Plugins/Experimental/PCG/Source/PCG/Public/Elements/Metadata/PCGMetadataTransformOpElement.h b/Engine/Plugins/Experimental/PCG/Source/PCG/Public/Elements/Metadata/PCGMetadataTransformOpElement.h
--- a/Engine/Plugins/Experimental/PCG/Source/PCG/Public/Elements/Metadata/PCGMetadataTransformOpElement.h
+++ b/Engine/Plugins/Experimental/PCG/Source/PCG/Public/Elements/Metadata/PCGMetadataTransformOpElement.h
@@ -42,8 +42,11 @@ public:
 #if WITH_EDITOR
 	virtual FName GetDefaultNodeName() const override;
 	virtual FText GetDefaultNodeTitle() const override;
+	virtual TArray<FPCGPreConfiguredSettingsInfo> GetPreconfiguredInfo() const override;
+	virtual bool OnlyExposePreconfiguredSettings() const override { return true; }
 #endif
 	virtual FName AdditionalTaskName() const override;
+	virtual void ApplyPreconfiguredSettings(const FPCGPreConfiguredSettingsInfo& PreconfigureInfo) override;
 	//~End UPCGSettings interface
 
 	//~Begin UPCGMetadataSettingsBase interface
diff --git a/Engine/Plugins/Experimental/PCG/Source/PCG/Public/Elements/Metadata/PCGMetadataTrigOpElement.h b/Engine/Plugins/Experimental/PCG/Source/PCG/Public/Elements/Metadata/PCGMetadataTrigOpElement.h
--- a/Engine/Plugins/Experimental/PCG/Source/PCG/Public/Elements/Metadata/PCGMetadataTrigOpElement.h
+++ b/Engine/Plugins/Experimental/PCG/Source/PCG/Public/Elements/Metadata/PCGMetadataTrigOpElement.h
@@ -34,8 +34,11 @@ public:
 #if WITH_EDITOR
 	virtual FName GetDefaultNodeName() const override;
 	virtual FText GetDefaultNodeTitle() const override;
+	virtual TArray<FPCGPreConfiguredSettingsInfo> GetPreconfiguredInfo() const override;
+	virtual bool OnlyExposePreconfiguredSettings() const override { return true; }
 #endif
 	virtual FName AdditionalTaskName() const override;
+	virtual void ApplyPreconfiguredSettings(const FPCGPreConfiguredSettingsInfo& PreconfigureInfo) override;
 	//~End UPCGSettings interface
 
 	//~Begin UPCGMetadataSettingsBase interface
diff --git a/Engine/Plugins/Experimental/PCG/Source/PCG/Public/Elements/Metadata/PCGMetadataVectorOpElement.h b/Engine/Plugins/Experimental/PCG/Source/PCG/Public/Elements/Metadata/PCGMetadataVectorOpElement.h
--- a/Engine/Plugins/Experimental/PCG/Source/PCG/Public/Elements/Metadata/PCGMetadataVectorOpElement.h
+++ b/Engine/Plugins/Experimental/PCG/Source/PCG/Public/Elements/Metadata/PCGMetadataVectorOpElement.h
@@ -38,8 +38,11 @@ public:
 #if WITH_EDITOR
 	virtual FName GetDefaultNodeName() const override;
 	virtual FText GetDefaultNodeTitle() const override;
+	virtual TArray<FPCGPreConfiguredSettingsInfo> GetPreconfiguredInfo() const override;
+	virtual bool OnlyExposePreconfiguredSettings() const override { return true; }
 #endif
 	virtual FName AdditionalTaskName() const override;
+	virtual void ApplyPreconfiguredSettings(const FPCGPreConfiguredSettingsInfo& PreconfigureInfo) override;
 	//~End UPCGSettings interface
 
 	//~Begin UPCGMetadataSettingsBase interface
diff --git a/Engine/Plugins/Experimental/PCG/Source/PCG/Public/Elements/PCGExecuteBlueprint.h b/Engine/Plugins/Experimental/PCG/Source/PCG/Public/Elements/PCGExecuteBlueprint.h
--- a/Engine/Plugins/Experimental/PCG/Source/PCG/Public/Elements/PCGExecuteBlueprint.h
+++ b/Engine/Plugins/Experimental/PCG/Source/PCG/Public/Elements/PCGExecuteBlueprint.h
@@ -81,6 +81,9 @@ public:
 	UFUNCTION(BlueprintNativeEvent, Category = "PCG|Node Customization")
 	EPCGSettingsType NodeTypeOverride() const;
 
+	UFUNCTION(BlueprintImplementableEvent, Category = "PCG|Preconfigure Settings", meta = (ForceAsFunction))
+	void ApplyPreconfiguredSettings(UPARAM(ref) const FPCGPreConfiguredSettingsInfo& InPreconfigureInfo);
+
 	UFUNCTION(BlueprintCallable, Category = "PCG|Input & Output")
 	TSet<FName> InputLabels() const;
 
@@ -152,6 +155,15 @@ public:
 	UPROPERTY(BlueprintReadWrite, EditDefaultsOnly, Category = AssetInfo, AssetRegistrySearchable)
 	bool bExposeToLibrary = false;
 
+	UPROPERTY(BlueprintReadWrite, EditDefaultsOnly, Category = "AssetInfo|Preconfigured Settings", AssetRegistrySearchable)
+	bool bEnablePreconfiguredSettings = false;
+
+	UPROPERTY(BlueprintReadWrite, EditDefaultsOnly, Category = "AssetInfo|Preconfigured Settings", AssetRegistrySearchable, meta = (EditCondition = bEnablePreconfiguredSettings, EditConditionHides))
+	bool bOnlyExposePreconfiguredSettings = false;
+
+	UPROPERTY(BlueprintReadWrite, EditDefaultsOnly, Category = "AssetInfo|Preconfigured Settings", meta = (EditCondition = bEnablePreconfiguredSettings, EditConditionHides))
+	TArray<FPCGPreConfiguredSettingsInfo> PreconfiguredInfo;
+
 	UPROPERTY(BlueprintReadWrite, EditDefaultsOnly, Category = AssetInfo, AssetRegistrySearchable)
 	FText Category;
 
@@ -192,8 +204,11 @@ public:
 	virtual void GetTrackedActorTags(FPCGTagToSettingsMap& OutTagToSettings, TArray<TObjectPtr<const UPCGGraph>>& OutVisitedGraphs) const override;
 	virtual UObject* GetJumpTargetForDoubleClick() const override;
 	virtual void ApplyDeprecationBeforeUpdatePins(UPCGNode* InOutNode, TArray<TObjectPtr<UPCGPin>>& InputPins, TArray<TObjectPtr<UPCGPin>>& OutputPins) override;
+	virtual TArray<FPCGPreConfiguredSettingsInfo> GetPreconfiguredInfo() const override;
+	virtual bool OnlyExposePreconfiguredSettings() const override;
 #endif
 
+	virtual void ApplyPreconfiguredSettings(const FPCGPreConfiguredSettingsInfo& InPreconfiguredsInfo) override;
 	virtual FName AdditionalTaskName() const override;
 	virtual TArray<FPCGPinProperties> InputPinProperties() const override;
 	virtual TArray<FPCGPinProperties> OutputPinProperties() const override;
diff --git a/Engine/Plugins/Experimental/PCG/Source/PCG/Public/PCGSettings.h b/Engine/Plugins/Experimental/PCG/Source/PCG/Public/PCGSettings.h
--- a/Engine/Plugins/Experimental/PCG/Source/PCG/Public/PCGSettings.h
+++ b/Engine/Plugins/Experimental/PCG/Source/PCG/Public/PCGSettings.h
@@ -69,6 +69,45 @@ struct FPCGSettingsOverridableParam
 	TArray<const FProperty*> Properties;
 };
 
+/**
+* Pre-configured settings info
+* Will be passed to the settings to pre-configure the settings on creation.
+* Example: Maths operations: Add, Mul, etc...
+*/
+USTRUCT(BlueprintType)
+struct FPCGPreConfiguredSettingsInfo
+{
+	GENERATED_BODY()
+
+	FPCGPreConfiguredSettingsInfo() = default;
+
+	explicit FPCGPreConfiguredSettingsInfo(int32 InIndex, FText InLabel = FText{})
+		: PreconfiguredIndex(InIndex)
+		, Label(std::move(InLabel))
+	{}
+
+#if WITH_EDITOR
+	FPCGPreConfiguredSettingsInfo(int32 InIndex, FText InLabel, FText InTooltip)
+		: PreconfiguredIndex(InIndex)
+		, Label(std::move(InLabel))
+		, Tooltip(std::move(InTooltip))
+	{}
+#endif // WITH_EDITOR
+
+	/* Index used by the settings to know which preconfigured settings it needs to set. */
+	UPROPERTY(BlueprintReadWrite, EditAnywhere, Category = "")
+	int32 PreconfiguredIndex = -1;
+
+	/* Label for the exposed asset. Can also be used instead of the index, if it is easier to deal with strings. */
+	UPROPERTY(BlueprintReadWrite, EditAnywhere, Category = "")
+	FText Label;
+
+#if WITH_EDITORONLY_DATA
+	UPROPERTY(EditAnywhere, Category = "")
+	FText Tooltip;
+#endif // WITH_EDITORONLY_DATA
+};
+
 
 UCLASS(Abstract)
 class PCG_API UPCGSettingsInterface : public UPCGData
@@ -183,8 +222,16 @@ public:
      */
 	virtual UObject* GetJumpTargetForDoubleClick() const;
 	virtual bool IsPropertyOverriddenByPin(const FProperty* InProperty) const;
+
+	/* Return preconfigured info that will be filled in the editor palette action, allowing to create pre-configured settings */
+	virtual TArray<FPCGPreConfiguredSettingsInfo> GetPreconfiguredInfo() const { return {}; }
+
+	/* If there are preconfigured info, we can skip the default settings and only expose pre-configured actions in the editor palette */
+	virtual bool OnlyExposePreconfiguredSettings() const { return false; }
 #endif
 
+	virtual void ApplyPreconfiguredSettings(const FPCGPreConfiguredSettingsInfo& PreconfigureInfo) {}
+
 	/** Derived classes can implement this to expose additional name information in the logs */
 	virtual FName AdditionalTaskName() const { return NAME_None; }
 
diff --git a/Engine/Plugins/Experimental/PCG/Source/PCGEditor/Private/PCGEditorGraphSchema.cpp b/Engine/Plugins/Experimental/PCG/Source/PCGEditor/Private/PCGEditorGraphSchema.cpp
--- a/Engine/Plugins/Experimental/PCG/Source/PCGEditor/Private/PCGEditorGraphSchema.cpp
+++ b/Engine/Plugins/Experimental/PCG/Source/PCGEditor/Private/PCGEditorGraphSchema.cpp
@@ -6,6 +6,7 @@
 #include "PCGGraph.h"
 #include "PCGPin.h"
 #include "Elements/PCGCollapseElement.h"
+#include "Elements/PCGExecuteBlueprint.h"
 #include "Elements/PCGFilterByType.h"
 #include "Elements/PCGUserParameterGet.h"
 
@@ -339,9 +340,25 @@ void UPCGEditorGraphSchema::GetNativeElementActions(FGraphActionMenuBuilder& Act
 				const FText Category = StaticEnum<EPCGSettingsType>()->GetDisplayNameTextByValue(static_cast<__underlying_type(EPCGSettingsType)>(PCGSettings->GetType()));
 				const FText Description = PCGSettings->GetNodeTooltipText();
 
-				TSharedPtr<FPCGEditorGraphSchemaAction_NewNativeElement> NewAction(new FPCGEditorGraphSchemaAction_NewNativeElement(Category, MenuDesc, Description, 0));
-				NewAction->SettingsClass = SettingsClass;
-				ActionMenuBuilder.AddAction(NewAction);
+				TArray<FPCGPreConfiguredSettingsInfo> AllPreconfiguredInfo = PCGSettings->GetPreconfiguredInfo();
+
+				if (AllPreconfiguredInfo.IsEmpty() || !PCGSettings->OnlyExposePreconfiguredSettings())
+				{
+					TSharedPtr<FPCGEditorGraphSchemaAction_NewNativeElement> NewAction(new FPCGEditorGraphSchemaAction_NewNativeElement(Category, MenuDesc, Description, 0));
+					NewAction->SettingsClass = SettingsClass;
+					ActionMenuBuilder.AddAction(NewAction);
+				}
+
+				// Also add preconfigured settings
+				const FText NewCategory = FText::Format(LOCTEXT("PreconfiguredSettingsCategory", "{0}|{1}"), Category, MenuDesc);
+
+				for (FPCGPreConfiguredSettingsInfo PreconfiguredInfo : AllPreconfiguredInfo)
+				{
+					TSharedPtr<FPCGEditorGraphSchemaAction_NewNativeElement> NewPreconfiguredAction(new FPCGEditorGraphSchemaAction_NewNativeElement(NewCategory, PreconfiguredInfo.Label, PreconfiguredInfo.Tooltip.IsEmpty() ? Description : PreconfiguredInfo.Tooltip, 0));
+					NewPreconfiguredAction->SettingsClass = SettingsClass;
+					NewPreconfiguredAction->PreconfiguredInfo = std::move(PreconfiguredInfo);
+					ActionMenuBuilder.AddAction(NewPreconfiguredAction);
+				}
 			}
 		}
 	}
@@ -358,7 +375,7 @@ void UPCGEditorGraphSchema::GetNativeElementActions(FGraphActionMenuBuilder& Act
 
 					for (const FPropertyBagPropertyDesc& PropertyDesc : BagStruct->GetPropertyDescs())
 					{
-						const FText MenuDesc = FText::Format(FText::FromString(TEXT("Get {0}")), FText::FromName(PropertyDesc.Name));
+						const FText MenuDesc = FText::Format(LOCTEXT("GetterNodeName", "Get {0}"), FText::FromName(PropertyDesc.Name));
 						const FText Description = FText::Format(LOCTEXT("NodeTooltip", "Get the value from '{0}' parameter, can be overridden by the graph instance."), FText::FromName(PropertyDesc.Name));
 
 						TSharedPtr<FPCGEditorGraphSchemaAction_NewGetParameterElement> NewAction(new FPCGEditorGraphSchemaAction_NewGetParameterElement(Category, MenuDesc, Description, 0));
@@ -377,18 +394,44 @@ void UPCGEditorGraphSchema::GetBlueprintElementActions(FGraphActionMenuBuilder&
 {
 	PCGEditorUtils::ForEachPCGBlueprintAssetData([&ActionMenuBuilder](const FAssetData& AssetData)
 	{
-		const bool bExposeToLibrary = AssetData.GetTagValueRef<bool>(TEXT("bExposeToLibrary"));
+		const bool bExposeToLibrary = AssetData.GetTagValueRef<bool>(GET_MEMBER_NAME_CHECKED(UPCGBlueprintElement, bExposeToLibrary));
+		const bool bOnlyExposePreconfiguredSettings = AssetData.GetTagValueRef<bool>(GET_MEMBER_NAME_CHECKED(UPCGBlueprintElement, bOnlyExposePreconfiguredSettings));
+
 		if (bExposeToLibrary)
 		{
 			const FText MenuDesc = FText::FromString(FName::NameToDisplayString(AssetData.AssetName.ToString(), false));
-			const FText Category = AssetData.GetTagValueRef<FText>(TEXT("Category"));
-			const FText Description = AssetData.GetTagValueRef<FText>(TEXT("Description"));
+			const FText Category = AssetData.GetTagValueRef<FText>(GET_MEMBER_NAME_CHECKED(UPCGBlueprintElement, Category));
+			const FText Description = AssetData.GetTagValueRef<FText>(GET_MEMBER_NAME_CHECKED(UPCGBlueprintElement, Description));
 
-			const FString GeneratedClass = AssetData.GetTagValueRef<FString>(FBlueprintTags::GeneratedClassPath);
+			const FSoftClassPath GeneratedClass = FSoftClassPath(AssetData.GetTagValueRef<FString>(FBlueprintTags::GeneratedClassPath));
 
-			TSharedPtr<FPCGEditorGraphSchemaAction_NewBlueprintElement> NewBlueprintAction(new FPCGEditorGraphSchemaAction_NewBlueprintElement(Category, MenuDesc, Description, 0));
-			NewBlueprintAction->BlueprintClassPath = FSoftClassPath(GeneratedClass);
-			ActionMenuBuilder.AddAction(NewBlueprintAction);
+			// Only load the class if we have enabled preconfigured settings.
+			TArray<FPCGPreConfiguredSettingsInfo> AllPreconfiguredInfo;
+			if (AssetData.GetTagValueRef<bool>(GET_MEMBER_NAME_CHECKED(UPCGBlueprintElement, bEnablePreconfiguredSettings)))
+			{
+				TSubclassOf<UPCGBlueprintElement> BlueprintClass = GeneratedClass.TryLoadClass<UPCGBlueprintElement>();
+				const UPCGBlueprintElement* BlueprintElement = BlueprintClass ? Cast<UPCGBlueprintElement>(BlueprintClass->GetDefaultObject()) : nullptr;
+
+				AllPreconfiguredInfo = BlueprintElement ? BlueprintElement->PreconfiguredInfo : TArray<FPCGPreConfiguredSettingsInfo>{};
+			}
+
+			if (AllPreconfiguredInfo.IsEmpty() || !bOnlyExposePreconfiguredSettings)
+			{
+				TSharedPtr<FPCGEditorGraphSchemaAction_NewBlueprintElement> NewBlueprintAction(new FPCGEditorGraphSchemaAction_NewBlueprintElement(Category, MenuDesc, Description, 0));
+				NewBlueprintAction->BlueprintClassPath = GeneratedClass;
+				ActionMenuBuilder.AddAction(NewBlueprintAction);
+			}
+
+			// Also add preconfigured settings
+			const FText NewCategory = FText::Format(LOCTEXT("PreconfiguredSettingsCategory", "{0}|{1}"), Category, MenuDesc);
+
+			for (FPCGPreConfiguredSettingsInfo PreconfiguredInfo : AllPreconfiguredInfo)
+			{
+				TSharedPtr<FPCGEditorGraphSchemaAction_NewBlueprintElement> NewPreconfiguredAction(new FPCGEditorGraphSchemaAction_NewBlueprintElement(NewCategory, PreconfiguredInfo.Label, PreconfiguredInfo.Tooltip.IsEmpty() ? Description : PreconfiguredInfo.Tooltip, 0));
+				NewPreconfiguredAction->BlueprintClassPath = GeneratedClass;
+				NewPreconfiguredAction->PreconfiguredInfo = std::move(PreconfiguredInfo);
+				ActionMenuBuilder.AddAction(NewPreconfiguredAction);
+			}
 		}
 
 		return true;
diff --git a/Engine/Plugins/Experimental/PCG/Source/PCGEditor/Private/PCGEditorGraphSchemaActions.cpp b/Engine/Plugins/Experimental/PCG/Source/PCGEditor/Private/PCGEditorGraphSchemaActions.cpp
--- a/Engine/Plugins/Experimental/PCG/Source/PCGEditor/Private/PCGEditorGraphSchemaActions.cpp
+++ b/Engine/Plugins/Experimental/PCG/Source/PCGEditor/Private/PCGEditorGraphSchemaActions.cpp
@@ -51,6 +51,11 @@ UEdGraphNode* FPCGEditorGraphSchemaAction_NewNativeElement::PerformAction(UEdGra
 		return nullptr;
 	}
 
+	if (DefaultNodeSettings)
+	{
+		DefaultNodeSettings->ApplyPreconfiguredSettings(PreconfiguredInfo);
+	}
+
 	PostCreation(NewPCGNode);
 
 	FGraphNodeCreator<UPCGEditorGraphNode> NodeCreator(*EditorGraph);
@@ -279,6 +284,7 @@ UEdGraphNode* FPCGEditorGraphSchemaAction_NewBlueprintElement::PerformAction(UEd
 	UPCGBlueprintElement* ElementInstance = nullptr;
 	TSubclassOf<UPCGBlueprintElement> BlueprintClass = BlueprintClassPath.TryLoadClass<UPCGBlueprintElement>();
 	DefaultBlueprintSettings->SetElementType(BlueprintClass, ElementInstance);
+	DefaultBlueprintSettings->ApplyPreconfiguredSettings(PreconfiguredInfo);
 
 	NewPCGNode->UpdateAfterSettingsChangeDuringCreation();
 
diff --git a/Engine/Plugins/Experimental/PCG/Source/PCGEditor/Private/PCGEditorGraphSchemaActions.h b/Engine/Plugins/Experimental/PCG/Source/PCGEditor/Private/PCGEditorGraphSchemaActions.h
--- a/Engine/Plugins/Experimental/PCG/Source/PCGEditor/Private/PCGEditorGraphSchemaActions.h
+++ b/Engine/Plugins/Experimental/PCG/Source/PCGEditor/Private/PCGEditorGraphSchemaActions.h
@@ -3,6 +3,8 @@
 #pragma once
 
 #include "EdGraph/EdGraphSchema.h"
+
+#include "PCGSettings.h"
 #include "Templates/SubclassOf.h"
 
 #include "PCGEditorGraphSchemaActions.generated.h"
@@ -11,7 +13,6 @@ class UEdGraph;
 class UEdGraphPin;
 class UPCGBlueprintElement;
 class UPCGNode;
-class UPCGSettings;
 class UPCGEditorGraph;
 
 UENUM()
@@ -40,6 +41,9 @@ struct FPCGEditorGraphSchemaAction_NewNativeElement : public FEdGraphSchemaActio
 	UPROPERTY()
 	TSubclassOf<UPCGSettings> SettingsClass;
 
+	UPROPERTY()
+	FPCGPreConfiguredSettingsInfo PreconfiguredInfo;
+
 	//~ Begin FEdGraphSchemaAction Interface
 	virtual FName GetTypeId() const override { return StaticGetTypeId(); }
 	virtual UEdGraphNode* PerformAction(UEdGraph* ParentGraph, UEdGraphPin* FromPin, const FVector2D Location, bool bSelectNewNode = true) override;
@@ -116,6 +120,9 @@ struct FPCGEditorGraphSchemaAction_NewBlueprintElement : public FEdGraphSchemaAc
 	UPROPERTY()
 	FSoftClassPath BlueprintClassPath;
 
+	UPROPERTY()
+	FPCGPreConfiguredSettingsInfo PreconfiguredInfo;
+
 	//~ Begin FEdGraphSchemaAction Interface
 	virtual FName GetTypeId() const override { return StaticGetTypeId(); }
 	virtual UEdGraphNode* PerformAction(UEdGraph* ParentGraph, UEdGraphPin* FromPin, const FVector2D Location, bool bSelectNewNode = true) override;
