[PCG] Add new custom version for deprecation (Params pins renamed to Overrides)

diff --git a/Engine/Plugins/Experimental/PCG/Source/PCG/Private/PCGSettings.cpp b/Engine/Plugins/Experimental/PCG/Source/PCG/Private/PCGSettings.cpp
--- a/Engine/Plugins/Experimental/PCG/Source/PCG/Private/PCGSettings.cpp
+++ b/Engine/Plugins/Experimental/PCG/Source/PCG/Private/PCGSettings.cpp
@@ -478,6 +478,11 @@ void UPCGSettings::ApplyDeprecationBeforeUpdatePins(UPCGNode* InOutNode, TArray<
 	{
 		PCGSettingsHelpers::DeprecationBreakOutParamsToNewPin(InOutNode, InputPins, OutputPins);
 	}
+
+	if (DataVersion < FPCGCustomVersion::RenameDefaultParamsToOverride)
+	{
+		InOutNode->RenameInputPin(PCGPinConstants::Private::OldDefaultParamsLabel, PCGPinConstants::DefaultParamsLabel);
+	}
 }
 #endif // WITH_EDITOR
 
diff --git a/Engine/Plugins/Experimental/PCG/Source/PCG/Public/PCGCommon.h b/Engine/Plugins/Experimental/PCG/Source/PCG/Public/PCGCommon.h
--- a/Engine/Plugins/Experimental/PCG/Source/PCG/Public/PCGCommon.h
+++ b/Engine/Plugins/Experimental/PCG/Source/PCG/Public/PCGCommon.h
@@ -60,6 +60,11 @@ namespace PCGPinConstants
 	const FName DefaultInputLabel = TEXT("In");
 	const FName DefaultOutputLabel = TEXT("Out");
 	const FName DefaultParamsLabel = TEXT("Overrides");
+
+namespace Private
+{
+	const FName OldDefaultParamsLabel = TEXT("Params");
+}
 }
 
 // Metadata used by PCG
diff --git a/Engine/Plugins/Experimental/PCG/Source/PCG/Public/PCGCustomVersion.h b/Engine/Plugins/Experimental/PCG/Source/PCG/Public/PCGCustomVersion.h
--- a/Engine/Plugins/Experimental/PCG/Source/PCG/Public/PCGCustomVersion.h
+++ b/Engine/Plugins/Experimental/PCG/Source/PCG/Public/PCGCustomVersion.h
@@ -35,6 +35,9 @@ struct PCG_API FPCGCustomVersion
 		// Renamed spline input and added bounding shape input.
 		SplineSamplerUpdatedNodeInputs = 8,
 
+		// Renamed params to override.
+		RenameDefaultParamsToOverride = 9,
+
 		// -----<new versions can be added above this line>-------------------------------------------------
 		VersionPlusOne,
 		LatestVersion = VersionPlusOne - 1
