[PCG] Add broadcast from numeric values to vectors in MakeTransform and MakeVector

diff --git a/Engine/Plugins/Experimental/PCG/Source/PCG/Private/Elements/Metadata/PCGMetadataMakeTransform.cpp b/Engine/Plugins/Experimental/PCG/Source/PCG/Private/Elements/Metadata/PCGMetadataMakeTransform.cpp
--- a/Engine/Plugins/Experimental/PCG/Source/PCG/Private/Elements/Metadata/PCGMetadataMakeTransform.cpp
+++ b/Engine/Plugins/Experimental/PCG/Source/PCG/Private/Elements/Metadata/PCGMetadataMakeTransform.cpp
@@ -82,7 +82,7 @@ bool UPCGMetadataMakeTransformSettings::IsSupportedInputType(uint16 TypeId, uint
 	else
 	{
 		bHasSpecialRequirement = false;
-		return PCG::Private::IsOfTypes<FVector2D, FVector, FVector4>(TypeId);
+		return PCG::Private::IsOfTypes<FVector2D, FVector, FVector4, int32, int64, float , double>(TypeId);
 	}
 }
 
@@ -133,13 +133,17 @@ bool FPCGMetadataMakeTransformElement::DoOperation(FOperationData& OperationData
 	{
 		using AttributeType = decltype(DummyValue);
 
-		if constexpr (!PCG::Private::IsOfTypes<AttributeType, FVector2D, FVector, FVector4>())
+		if constexpr (PCG::Private::IsOfTypes<AttributeType, FVector2D, FVector, FVector4>())
 		{
-			return false;
+			return DoTernaryOp<AttributeType, FQuat, AttributeType>(OperationData, PCGMetadataMakeTransformSettings::MakeTransform<AttributeType>);
+		}
+		else if constexpr (PCG::Private::IsOfTypes<AttributeType, int32, int64, float, double>())
+		{
+			return DoTernaryOp<FVector, FQuat, FVector>(OperationData, PCGMetadataMakeTransformSettings::MakeTransform<FVector>);
 		}
 		else
 		{
-			return DoTernaryOp<AttributeType, FQuat, AttributeType>(OperationData, PCGMetadataMakeTransformSettings::MakeTransform<AttributeType>);
+			return false;
 		}
 	};
 
diff --git a/Engine/Plugins/Experimental/PCG/Source/PCG/Private/Elements/Metadata/PCGMetadataMakeVector.cpp b/Engine/Plugins/Experimental/PCG/Source/PCG/Private/Elements/Metadata/PCGMetadataMakeVector.cpp
--- a/Engine/Plugins/Experimental/PCG/Source/PCG/Private/Elements/Metadata/PCGMetadataMakeVector.cpp
+++ b/Engine/Plugins/Experimental/PCG/Source/PCG/Private/Elements/Metadata/PCGMetadataMakeVector.cpp
@@ -147,15 +147,15 @@ bool UPCGMetadataMakeVectorSettings::IsSupportedInputType(uint16 TypeId, uint32
 
 	if (Label == PCGMetadataMakeVectorConstants::XYZLabel)
 	{
-		return PCG::Private::IsOfTypes<FVector>(TypeId);
+		return PCG::Private::IsOfTypes<FVector, FVector2D, float, double, int32, int64>(TypeId);
 	}
 	else if (Label == PCGMetadataMakeVectorConstants::XYLabel || Label == PCGMetadataMakeVectorConstants::ZWLabel)
 	{
-		return PCG::Private::IsOfTypes<FVector2D>(TypeId);
+		return PCG::Private::IsOfTypes<FVector2D, float, double, int32, int64>(TypeId);
 	}
 	else
 	{
-		return PCG::Private::IsOfTypes<float, double>(TypeId);
+		return PCG::Private::IsOfTypes<float, double, int32, int64>(TypeId);
 	}
 }
 
