[PCG] Debugging feature - trigger breakpoint on element execution

diff --git a/Engine/Plugins/Experimental/PCG/Source/PCG/Private/PCGElement.cpp b/Engine/Plugins/Experimental/PCG/Source/PCG/Private/PCGElement.cpp
--- a/Engine/Plugins/Experimental/PCG/Source/PCG/Private/PCGElement.cpp
+++ b/Engine/Plugins/Experimental/PCG/Source/PCG/Private/PCGElement.cpp
@@ -19,11 +19,28 @@ static TAutoConsoleVariable<bool> CVarPCGValidatePointMetadata(
 	true,
 	TEXT("Controls whether we validate that the metadata entry keys on the output point data are consistent"));
 
+#if WITH_EDITOR
+#define PCG_ELEMENT_EXECUTION_BREAKPOINT() \
+	if (Context && Context->GetInputSettingsInterface() && Context->GetInputSettingsInterface()->bBreakDebugger) \
+	{ \
+		if (const UEnum* EnumPtr = FindObject<UEnum>(nullptr, TEXT("/Script/PCG.EPCGExecutionPhase"), true)) \
+		{ \
+			PCGE_LOG(Log, LogOnly, FText::Format(LOCTEXT("BreakExecution", "Execution halted, phase: {0}"), \
+				FText::FromName(EnumPtr->GetNameByValue(static_cast<int>(Context->CurrentPhase))))); \
+		} \
+		PLATFORM_BREAK(); \
+	}
+#else
+#define PCG_ELEMENT_EXECUTION_BREAKPOINT()
+#endif
+
 bool IPCGElement::Execute(FPCGContext* Context) const
 {
 	check(Context && Context->AsyncState.NumAvailableTasks > 0 && Context->CurrentPhase < EPCGExecutionPhase::Done);
 	check(Context->AsyncState.bIsRunningOnMainThread || !CanExecuteOnlyOnMainThread(Context));
 
+	PCG_ELEMENT_EXECUTION_BREAKPOINT();
+
 	while (Context->CurrentPhase != EPCGExecutionPhase::Done)
 	{
 		PCGUtils::FScopedCall ScopedCall(*this, Context);
diff --git a/Engine/Plugins/Experimental/PCG/Source/PCG/Public/PCGSettings.h b/Engine/Plugins/Experimental/PCG/Source/PCG/Public/PCGSettings.h
--- a/Engine/Plugins/Experimental/PCG/Source/PCG/Public/PCGSettings.h
+++ b/Engine/Plugins/Experimental/PCG/Source/PCG/Public/PCGSettings.h
@@ -135,6 +135,10 @@ public:
 #if WITH_EDITORONLY_DATA
 	UPROPERTY(EditAnywhere, BlueprintReadWrite, Category = Debug, meta = (ShowOnlyInnerProperties))
 	FPCGDebugVisualizationSettings DebugSettings;
+
+	/** If a debugger is attached, triggers a breakpoint inside IPCGElement::Execute(). Editor only. */
+	UPROPERTY(Transient, EditAnywhere, BlueprintReadWrite, Category = Debug)
+	bool bBreakDebugger = false;
 #endif
 };
 
