[PCG] Disable auto-widening of pin types. Don't dynamically narrow input pins unless the node execution requires a narrow type.

diff --git a/Engine/Plugins/Experimental/PCG/Source/PCG/Private/Elements/Metadata/PCGMetadataOpElementBase.cpp b/Engine/Plugins/Experimental/PCG/Source/PCG/Private/Elements/Metadata/PCGMetadataOpElementBase.cpp
--- a/Engine/Plugins/Experimental/PCG/Source/PCG/Private/Elements/Metadata/PCGMetadataOpElementBase.cpp
+++ b/Engine/Plugins/Experimental/PCG/Source/PCG/Private/Elements/Metadata/PCGMetadataOpElementBase.cpp
@@ -30,24 +30,16 @@ void UPCGMetadataSettingsBase::PostLoad()
 
 EPCGDataType UPCGMetadataSettingsBase::GetCurrentPinTypes(const UPCGPin* InPin) const
 {
-	if (GetInputPinNum() == 0)
-	{
-		return EPCGDataType::Any;
-	}
-
 	check(InPin);
-	const FName FirstPinLabel = GetInputPinLabel(0);
-	if (InPin->Properties.Label == FirstPinLabel || InPin->IsOutputPin())
+	if (!InPin->IsOutputPin() || GetInputPinNum() == 0)
 	{
-		// First input pin, and any output pins, narrow to types connected to first pin
-		const EPCGDataType FirstPinTypeUnion = GetTypeUnionOfIncidentEdges(FirstPinLabel);
-		return FirstPinTypeUnion != EPCGDataType::None ? FirstPinTypeUnion : EPCGDataType::Any;
-	}
-	else
-	{
-		// Other input pins stay at Any
-		return EPCGDataType::Any;
+		// Fall back to default for input pins, or if no input pins present from which to obtain type
+		return Super::GetCurrentPinTypes(InPin);
 	}
+
+	// Output pin narrows to union of inputs on first pin
+	const EPCGDataType FirstPinTypeUnion = GetTypeUnionOfIncidentEdges(GetInputPinLabel(0));
+	return FirstPinTypeUnion != EPCGDataType::None ? FirstPinTypeUnion : EPCGDataType::Any;
 }
 
 TArray<FPCGPinProperties> UPCGMetadataSettingsBase::InputPinProperties() const
diff --git a/Engine/Plugins/Experimental/PCG/Source/PCG/Private/Elements/Metadata/PCGMetadataRenameElement.cpp b/Engine/Plugins/Experimental/PCG/Source/PCG/Private/Elements/Metadata/PCGMetadataRenameElement.cpp
--- a/Engine/Plugins/Experimental/PCG/Source/PCG/Private/Elements/Metadata/PCGMetadataRenameElement.cpp
+++ b/Engine/Plugins/Experimental/PCG/Source/PCG/Private/Elements/Metadata/PCGMetadataRenameElement.cpp
@@ -14,7 +14,13 @@
 
 EPCGDataType UPCGMetadataRenameSettings::GetCurrentPinTypes(const UPCGPin* InPin) const
 {
-	// All pins narrow to union of input types
+	check(InPin);
+	if (!InPin->IsOutputPin())
+	{
+		return Super::GetCurrentPinTypes(InPin);
+	}
+
+	// Output pin narrows to union of inputs on first pin
 	const EPCGDataType InputTypeUnion = GetTypeUnionOfIncidentEdges(PCGPinConstants::DefaultInputLabel);
 	return (InputTypeUnion != EPCGDataType::None) ? InputTypeUnion : EPCGDataType::Any;
 }
diff --git a/Engine/Plugins/Experimental/PCG/Source/PCG/Private/Elements/PCGAttributeFilterElement.cpp b/Engine/Plugins/Experimental/PCG/Source/PCG/Private/Elements/PCGAttributeFilterElement.cpp
--- a/Engine/Plugins/Experimental/PCG/Source/PCG/Private/Elements/PCGAttributeFilterElement.cpp
+++ b/Engine/Plugins/Experimental/PCG/Source/PCG/Private/Elements/PCGAttributeFilterElement.cpp
@@ -57,7 +57,13 @@ FText UPCGAttributeFilterSettings::GetDefaultNodeTitle() const
 
 EPCGDataType UPCGAttributeFilterSettings::GetCurrentPinTypes(const UPCGPin* InPin) const
 {
-	// All pins narrow to union of input types
+	check(InPin);
+	if (!InPin->IsOutputPin())
+	{
+		return Super::GetCurrentPinTypes(InPin);
+	}
+
+	// Output pin narrows to union of inputs on first pin
 	const EPCGDataType InputTypeUnion = GetTypeUnionOfIncidentEdges(PCGPinConstants::DefaultInputLabel);
 	return (InputTypeUnion != EPCGDataType::None) ? InputTypeUnion : EPCGDataType::Any;
 }
diff --git a/Engine/Plugins/Experimental/PCG/Source/PCG/Private/Elements/PCGCreateAttribute.cpp b/Engine/Plugins/Experimental/PCG/Source/PCG/Private/Elements/PCGCreateAttribute.cpp
--- a/Engine/Plugins/Experimental/PCG/Source/PCG/Private/Elements/PCGCreateAttribute.cpp
+++ b/Engine/Plugins/Experimental/PCG/Source/PCG/Private/Elements/PCGCreateAttribute.cpp
@@ -51,19 +51,15 @@ void UPCGCreateAttributeSettings::PostLoad()
 
 EPCGDataType UPCGCreateAttributeSettings::GetCurrentPinTypes(const UPCGPin* InPin) const
 {
-	if (!HasDynamicPins())
-	{
-		return Super::GetCurrentPinTypes(InPin);
-	}
-
 	check(InPin);
-	if (InPin->Properties.Label == PCGPinConstants::DefaultInputLabel || InPin->Properties.Label == PCGPinConstants::DefaultOutputLabel)
+	if (!HasDynamicPins() || !InPin->IsOutputPin())
 	{
-		const EPCGDataType PrimaryInputType = GetTypeUnionOfIncidentEdges(PCGPinConstants::DefaultInputLabel);
-		return (PrimaryInputType != EPCGDataType::None) ? PrimaryInputType : EPCGDataType::Any;
+		return Super::GetCurrentPinTypes(InPin);
 	}
 
-	return Super::GetCurrentPinTypes(InPin);
+	// Output pin narrows to union of inputs on first pin
+	const EPCGDataType PrimaryInputType = GetTypeUnionOfIncidentEdges(PCGPinConstants::DefaultInputLabel);
+	return (PrimaryInputType != EPCGDataType::None) ? PrimaryInputType : EPCGDataType::Any;
 }
 
 FName UPCGCreateAttributeSettings::AdditionalTaskName() const
diff --git a/Engine/Plugins/Experimental/PCG/Source/PCG/Private/Elements/PCGFilterByTag.cpp b/Engine/Plugins/Experimental/PCG/Source/PCG/Private/Elements/PCGFilterByTag.cpp
--- a/Engine/Plugins/Experimental/PCG/Source/PCG/Private/Elements/PCGFilterByTag.cpp
+++ b/Engine/Plugins/Experimental/PCG/Source/PCG/Private/Elements/PCGFilterByTag.cpp
@@ -34,7 +34,13 @@ FText UPCGFilterByTagSettings::GetNodeTooltipText() const
 
 EPCGDataType UPCGFilterByTagSettings::GetCurrentPinTypes(const UPCGPin* InPin) const
 {
-	// All pins narrow to input type
+	check(InPin);
+	if (!InPin->IsOutputPin())
+	{
+		return Super::GetCurrentPinTypes(InPin);
+	}
+
+	// Output pin narrows to union of inputs on first pin
 	const EPCGDataType InputTypeUnion = GetTypeUnionOfIncidentEdges(PCGPinConstants::DefaultInputLabel);
 	return InputTypeUnion != EPCGDataType::None ? InputTypeUnion : EPCGDataType::Any;
 }
diff --git a/Engine/Plugins/Experimental/PCG/Source/PCG/Private/Elements/PCGFilterByType.cpp b/Engine/Plugins/Experimental/PCG/Source/PCG/Private/Elements/PCGFilterByType.cpp
--- a/Engine/Plugins/Experimental/PCG/Source/PCG/Private/Elements/PCGFilterByType.cpp
+++ b/Engine/Plugins/Experimental/PCG/Source/PCG/Private/Elements/PCGFilterByType.cpp
@@ -24,14 +24,13 @@ FText UPCGFilterByTypeSettings::GetNodeTooltipText() const
 
 EPCGDataType UPCGFilterByTypeSettings::GetCurrentPinTypes(const UPCGPin* InPin) const
 {
-	// Output pin narrows to set type, input pin stays at Any to invite all connections
 	check(InPin);
-	if (InPin->Properties.Label == PCGPinConstants::DefaultOutputLabel)
+	if (!InPin->IsOutputPin())
 	{
-		return TargetType;
+		return Super::GetCurrentPinTypes(InPin);
 	}
 
-	return Super::GetCurrentPinTypes(InPin);
+	return TargetType;
 }
 
 FPCGElementPtr UPCGFilterByTypeSettings::CreateElement() const
diff --git a/Engine/Plugins/Experimental/PCG/Source/PCG/Private/Elements/PCGGather.cpp b/Engine/Plugins/Experimental/PCG/Source/PCG/Private/Elements/PCGGather.cpp
--- a/Engine/Plugins/Experimental/PCG/Source/PCG/Private/Elements/PCGGather.cpp
+++ b/Engine/Plugins/Experimental/PCG/Source/PCG/Private/Elements/PCGGather.cpp
@@ -7,8 +7,14 @@
 	
 EPCGDataType UPCGGatherSettings::GetCurrentPinTypes(const UPCGPin* InPin) const
 {
-	// All pins narrow to input type
-	EPCGDataType InputTypeUnion = GetTypeUnionOfIncidentEdges(PCGPinConstants::DefaultInputLabel);
+	check(InPin);
+	if (!InPin->IsOutputPin())
+	{
+		return Super::GetCurrentPinTypes(InPin);
+	}
+
+	// Output pin narrows to union of inputs on first pin
+	const EPCGDataType InputTypeUnion = GetTypeUnionOfIncidentEdges(PCGPinConstants::DefaultInputLabel);
 	return InputTypeUnion != EPCGDataType::None ? InputTypeUnion : EPCGDataType::Any;
 }
 
