[PCG] Better logging for Metadata operations + fix unit tests

* Unit tests are relying on the error logged, so we need to change it
* Remove ensure on InNode null in log visual, since we don't have any node in unit tests

diff --git a/Engine/Plugins/Experimental/PCG/Source/PCG/Private/Elements/Metadata/PCGMetadataOpElementBase.cpp b/Engine/Plugins/Experimental/PCG/Source/PCG/Private/Elements/Metadata/PCGMetadataOpElementBase.cpp
--- a/Engine/Plugins/Experimental/PCG/Source/PCG/Private/Elements/Metadata/PCGMetadataOpElementBase.cpp
+++ b/Engine/Plugins/Experimental/PCG/Source/PCG/Private/Elements/Metadata/PCGMetadataOpElementBase.cpp
@@ -100,16 +100,17 @@ bool FPCGMetadataElementBase::ExecuteInternal(FPCGContext* Context) const
 
 	for (uint32 i = 0; i < NumberOfInputs; ++i)
 	{
-		TArray<FPCGTaggedData> InputData = Context->InputData.GetInputsByPin(Settings->GetInputPinLabel(i));
+		FName PinLabel = Settings->GetInputPinLabel(i);
+		TArray<FPCGTaggedData> InputData = Context->InputData.GetInputsByPin(PinLabel);
 		if (InputData.IsEmpty())
 		{
 			// Absence of data not worth broadcasting to user as visual warning
-			PCGE_LOG(Warning, LogOnly, FText::Format(LOCTEXT("MissingInputDataForPin", "No data provided on pin {0}"), i));
+			PCGE_LOG(Warning, LogOnly, FText::Format(LOCTEXT("MissingInputDataForPin", "No data provided on pin {0}"), FText::FromName(PinLabel)));
 			return true;
 		}
 		else if (InputData.Num() > 1)
 		{
-			PCGE_LOG(Error, GraphAndLog, FText::Format(LOCTEXT("TooMuchDataForPin", "Too many data items ({0}) provided on pin {1}"), InputData.Num(), i));
+			PCGE_LOG(Error, GraphAndLog, FText::Format(LOCTEXT("TooMuchDataForPin", "Too many data items ({0}) provided on pin {1}"), InputData.Num(), FText::FromName(PinLabel)));
 			return true;
 		}
 
@@ -127,7 +128,7 @@ bool FPCGMetadataElementBase::ExecuteInternal(FPCGContext* Context) const
 		}
 		else
 		{
-			PCGE_LOG(Warning, LogOnly, FText::Format(LOCTEXT("InvalidInputDataTypeForPin", "Invalid data provided on pin {0}, must be of type Spatial or Attribute Set"), i));
+			PCGE_LOG(Warning, LogOnly, FText::Format(LOCTEXT("InvalidInputDataTypeForPin", "Invalid data provided on pin {0}, must be of type Spatial or Attribute Set"), FText::FromName(PinLabel)));
 			return true;
 		}
 	}
@@ -154,7 +155,7 @@ bool FPCGMetadataElementBase::ExecuteInternal(FPCGContext* Context) const
 			InputTaggedData[0].Data->GetClass() != InputTaggedData[i].Data->GetClass() &&
 			!InputTaggedData[i].Data->IsA<UPCGParamData>())
 		{
-			PCGE_LOG(Error, GraphAndLog, FText::Format(LOCTEXT("InputTypeMismatch", "Input {0} is not of the same type than input 0 and is not an Attribute Set. This is not supported."), i));
+			PCGE_LOG(Error, GraphAndLog, FText::Format(LOCTEXT("InputTypeMismatch", "Data on pin {0} is not of the same type than on pin {1} and is not an Attribute Set. This is not supported."), FText::FromName(InputTaggedData[i].Pin), FText::FromName(InputTaggedData[0].Pin)));
 			return true;
 		}
 
@@ -175,7 +176,7 @@ bool FPCGMetadataElementBase::ExecuteInternal(FPCGContext* Context) const
 
 		if (!OperationData.InputAccessors[i].IsValid() || !OperationData.InputKeys[i].IsValid())
 		{
-			PCGE_LOG(Error, GraphAndLog, FText::Format(LOCTEXT("AttributeDoesNotExist", "Attribute/Property '{0}' does not exist for input {1}"), FText::FromName(InputSource.GetName()), i));
+			PCGE_LOG(Error, GraphAndLog, FText::Format(LOCTEXT("AttributeDoesNotExist", "Attribute/Property '{0}' from pin {1} does not exist"), FText::FromName(InputSource.GetName()), FText::FromName(InputTaggedData[i].Pin)));
 			return true;
 		}
 
@@ -185,7 +186,7 @@ bool FPCGMetadataElementBase::ExecuteInternal(FPCGContext* Context) const
 		bool bHasSpecialRequirement = false;
 		if (!Settings->IsSupportedInputType(AttributeTypeId, i, bHasSpecialRequirement))
 		{
-			PCGE_LOG(Error, GraphAndLog, FText::Format(LOCTEXT("UnsupportedAttributeType", "Attribute/Property '{0}' is not a supported type for input {1}"), FText::FromName(InputSource.GetName()), i));
+			PCGE_LOG(Error, GraphAndLog, FText::Format(LOCTEXT("UnsupportedAttributeType", "Attribute/Property '{0}' from pin {1} is not a supported type"), FText::FromName(InputSource.GetName()), FText::FromName(InputTaggedData[i].Pin)));
 			return true;
 		}
 
@@ -198,7 +199,7 @@ bool FPCGMetadataElementBase::ExecuteInternal(FPCGContext* Context) const
 			}
 			else if (OperationData.MostComplexInputType != AttributeTypeId && !PCG::Private::IsBroadcastable(AttributeTypeId, OperationData.MostComplexInputType))
 			{
-				PCGE_LOG(Error, GraphAndLog, FText::Format(LOCTEXT("AttributeCannotBeBroadcasted", "Attribute '{0}' cannot be broadcasted to match types for input {1}"), FText::FromName(InputSource.GetName()), i));
+				PCGE_LOG(Error, GraphAndLog, FText::Format(LOCTEXT("AttributeCannotBeBroadcasted", "Attribute '{0}' from pin {1} cannot be broadcasted to match types"), FText::FromName(InputSource.GetName()), FText::FromName(InputTaggedData[i].Pin)));
 				return true;
 			}
 		}
@@ -214,7 +215,7 @@ bool FPCGMetadataElementBase::ExecuteInternal(FPCGContext* Context) const
 		// Therefore mark that we have nothing to do and early out.
 		if (NumberOfElements[i] == 0)
 		{
-			PCGE_LOG(Verbose, LogOnly, FText::Format(LOCTEXT("NoElementsInInput", "No elements in input {0}"), i));
+			PCGE_LOG(Verbose, LogOnly, FText::Format(LOCTEXT("NoElementsInInput", "No elements in data from pin {0}"), FText::FromName(InputTaggedData[i].Pin)));
 			bNoOperationNeeded = true;
 			break;
 		}
@@ -222,7 +223,7 @@ bool FPCGMetadataElementBase::ExecuteInternal(FPCGContext* Context) const
 		// Verify that the number of elements makes sense
 		if (OperationData.NumberOfElementsToProcess % NumberOfElements[i] != 0)
 		{
-			PCGE_LOG(Error, GraphAndLog, FText::Format(LOCTEXT("MismatchInNumberOfElements", "Mismatch between the number of elements in input 0 ({0}) and in input {1} ({2})"), OperationData.NumberOfElementsToProcess, i, NumberOfElements[i]));
+			PCGE_LOG(Error, GraphAndLog, FText::Format(LOCTEXT("MismatchInNumberOfElements", "Mismatch between the number of elements from pin {0} ({1}) and from pin {2} ({3})"), FText::FromName(InputTaggedData[0].Pin), OperationData.NumberOfElementsToProcess, FText::FromName(InputTaggedData[i].Pin), NumberOfElements[i]));
 			return true;
 		}
 
diff --git a/Engine/Plugins/Experimental/PCG/Source/PCG/Private/Tests/Elements/Metadata/PCGMetadataBooleanOpTest.cpp b/Engine/Plugins/Experimental/PCG/Source/PCG/Private/Tests/Elements/Metadata/PCGMetadataBooleanOpTest.cpp
--- a/Engine/Plugins/Experimental/PCG/Source/PCG/Private/Tests/Elements/Metadata/PCGMetadataBooleanOpTest.cpp
+++ b/Engine/Plugins/Experimental/PCG/Source/PCG/Private/Tests/Elements/Metadata/PCGMetadataBooleanOpTest.cpp
@@ -140,7 +140,7 @@ bool FPCGMetadataBooleanOpTest::RunTest(const FString& Parameters)
 	ParamTaggedData2.Pin = PCGMetadataSettingsBaseConstants::DoubleInputSecondLabel;
 	Settings->Operation = EPCGMedadataBooleanOperation::And;
 
-	AddExpectedError(TEXT("Attribute/Property Invalid is not a supported type for input"), EAutomationExpectedErrorFlags::Contains, 3);
+	AddExpectedError(TEXT("Attribute/Property 'Invalid' from pin"), EAutomationExpectedErrorFlags::Contains, 3);
 
 	{
 		Settings->InputSource1.SetAttributeName(InvalidAttribute);
diff --git a/Engine/Plugins/Experimental/PCG/Source/PCG/Private/Tests/Elements/Metadata/PCGMetadataBreakTransformTest.cpp b/Engine/Plugins/Experimental/PCG/Source/PCG/Private/Tests/Elements/Metadata/PCGMetadataBreakTransformTest.cpp
--- a/Engine/Plugins/Experimental/PCG/Source/PCG/Private/Tests/Elements/Metadata/PCGMetadataBreakTransformTest.cpp
+++ b/Engine/Plugins/Experimental/PCG/Source/PCG/Private/Tests/Elements/Metadata/PCGMetadataBreakTransformTest.cpp
@@ -258,7 +258,7 @@ bool FPCGMetadataBreakTransformTest::RunTest(const FString& Parameters)
 	AllTestData[1] = PairWhatData("Testing with param data as input", &TestDataParams);
 
 	// Setup error catching. This error should happen only twice (invalid type)
-	AddExpectedError(TEXT("Attribute/Property Float is not a supported type for input 0"), EAutomationExpectedErrorFlags::Contains, 2);
+	AddExpectedError(TEXT("Attribute/Property 'Float' from pin"), EAutomationExpectedErrorFlags::Contains, 2);
 
 	for (PairWhatData& PairTestData : AllTestData)
 	{
diff --git a/Engine/Plugins/Experimental/PCG/Source/PCG/Private/Tests/Elements/Metadata/PCGMetadataBreakVectorTest.cpp b/Engine/Plugins/Experimental/PCG/Source/PCG/Private/Tests/Elements/Metadata/PCGMetadataBreakVectorTest.cpp
--- a/Engine/Plugins/Experimental/PCG/Source/PCG/Private/Tests/Elements/Metadata/PCGMetadataBreakVectorTest.cpp
+++ b/Engine/Plugins/Experimental/PCG/Source/PCG/Private/Tests/Elements/Metadata/PCGMetadataBreakVectorTest.cpp
@@ -324,7 +324,7 @@ bool FPCGMetadataBreakVectorTest::RunTest(const FString& Parameters)
 	AllTestData[1] = PairWhatData("Testing with param data as input", &TestDataParams);
 
 	// Setup error catching. This error should happen only twice (invalid type)
-	AddExpectedError(TEXT("Attribute/Property Float is not a supported type for input 0"), EAutomationExpectedErrorFlags::Contains, 2);
+	AddExpectedError(TEXT("Attribute/Property 'Float' from pin"), EAutomationExpectedErrorFlags::Contains, 2);
 
 	for (PairWhatData& PairTestData : AllTestData)
 	{
diff --git a/Engine/Plugins/Experimental/PCG/Source/PCG/Private/Utils/PCGNodeVisualLogs.cpp b/Engine/Plugins/Experimental/PCG/Source/PCG/Private/Utils/PCGNodeVisualLogs.cpp
--- a/Engine/Plugins/Experimental/PCG/Source/PCG/Private/Utils/PCGNodeVisualLogs.cpp
+++ b/Engine/Plugins/Experimental/PCG/Source/PCG/Private/Utils/PCGNodeVisualLogs.cpp
@@ -13,7 +13,8 @@
 
 void FPCGNodeVisualLogs::Log(TWeakObjectPtr<const UPCGNode> InNode, TWeakObjectPtr<UPCGComponent> InComponent, ELogVerbosity::Type InVerbosity, const FText& InMessage)
 {
-	if (!ensure(InNode.Get()) || !ensure(InComponent.Get()))
+	// InNode can be null, in case of unit tests.
+	if (!InNode.Get() || !ensure(InComponent.Get()))
 	{
 		return;
 	}
