[PCG] Attribute Operation revamp

* Rename Attribute Operation to Copy Attribute
* Add an extra pin to be able to set an attribute from an extra attribute set
* Also update Create/AddAttribute to match this behavior
* Attribute pin unused are now blacked out with a tooltip.

diff --git a/Engine/Plugins/Experimental/PCG/Source/PCG/Private/Elements/Metadata/PCGMetadataElement.cpp b/Engine/Plugins/Experimental/PCG/Source/PCG/Private/Elements/Metadata/PCGMetadataElement.cpp
--- a/Engine/Plugins/Experimental/PCG/Source/PCG/Private/Elements/Metadata/PCGMetadataElement.cpp
+++ b/Engine/Plugins/Experimental/PCG/Source/PCG/Private/Elements/Metadata/PCGMetadataElement.cpp
@@ -2,16 +2,23 @@
 
 #include "Elements/Metadata/PCGMetadataElement.h"
 
+#include "PCGContext.h"
+#include "PCGParamData.h"
+#include "Data/PCGPointData.h"
 #include "Data/PCGSpatialData.h"
 #include "Elements/Metadata/PCGMetadataElementCommon.h"
-#include "Data/PCGPointData.h"
 #include "Metadata/Accessors/PCGAttributeAccessorHelpers.h"
-#include "PCGContext.h"
 
 #include UE_INLINE_GENERATED_CPP_BY_NAME(PCGMetadataElement)
 
 #define LOCTEXT_NAMESPACE "PCGMetadataElement"
 
+namespace PCGMetadataOperationSettings
+{
+	const FName AttributeLabel = TEXT("Attribute");
+	const FText AttributeTooltip = LOCTEXT("AttributeTooltip", "Optional Attribute Set to copy the value from. Not used if not connected.");
+}
+
 UPCGMetadataOperationSettings::UPCGMetadataOperationSettings()
 {
 	// Previous default object was: None for source and target attribute, density for point property
@@ -26,6 +33,26 @@ FPCGElementPtr UPCGMetadataOperationSettings::CreateElement() const
 	return MakeShared<FPCGMetadataOperationElement>();
 }
 
+#if WITH_EDITOR
+FText UPCGMetadataOperationSettings::GetNodeTooltipText() const
+{
+	return LOCTEXT("NodeTooltip", "Copy from the Input Source, taken from either the Attribute Set if connected otherwise from the input points, to Output Target attribute.");
+}
+
+bool UPCGMetadataOperationSettings::IsPinUsedByNodeExecution(const UPCGPin* InPin) const
+{
+	return !InPin || (InPin->Properties.Label != PCGMetadataOperationSettings::AttributeLabel) || InPin->IsConnected();
+}
+#endif // WITH_EDITOR
+
+TArray<FPCGPinProperties> UPCGMetadataOperationSettings::InputPinProperties() const
+{
+	TArray<FPCGPinProperties> Properties;
+	Properties.Emplace(PCGPinConstants::DefaultInputLabel, EPCGDataType::Point);
+	Properties.Emplace(PCGMetadataOperationSettings::AttributeLabel, EPCGDataType::Param, /*bInAllowMultipleConnections=*/ false, /*bAllowMultipleData=*/ false, PCGMetadataOperationSettings::AttributeTooltip);
+	return Properties;
+}
+
 void UPCGMetadataOperationSettings::PostLoad()
 {
 	Super::PostLoad();
@@ -72,7 +99,16 @@ bool FPCGMetadataOperationElement::ExecuteInternal(FPCGContext* Context) const
 	const UPCGMetadataOperationSettings* Settings = Context->GetInputSettings<UPCGMetadataOperationSettings>();
 	check(Settings);
 
-	TArray<FPCGTaggedData> Inputs = Context->InputData.GetInputs();
+	TArray<FPCGTaggedData> Inputs = Context->InputData.GetInputsByPin(PCGPinConstants::DefaultInputLabel);
+	TArray<FPCGTaggedData> SourceAttributeInputs = Context->InputData.GetInputsByPin(PCGMetadataOperationSettings::AttributeLabel);
+
+	const UPCGParamData* SourceAttributeSet = (!SourceAttributeInputs.IsEmpty() ? Cast<UPCGParamData>(SourceAttributeInputs[0].Data) : nullptr);
+
+	if (SourceAttributeSet && Settings->InputSource.Selection != EPCGAttributePropertySelection::Attribute)
+	{
+		PCGE_LOG(Error, GraphAndLog, FText::Format(LOCTEXT("InvalidInputSource", "{0} pin is connected but Input Source is not set to get an attribute."), FText::FromName(PCGMetadataOperationSettings::AttributeLabel)));
+		return true;
+	}
 
 	TArray<FPCGTaggedData>& Outputs = Context->OutputData.TaggedData;
 
@@ -96,7 +132,9 @@ bool FPCGMetadataOperationElement::ExecuteInternal(FPCGContext* Context) const
 			continue;
 		}
 
-		if (!OriginalData->Metadata)
+		const UPCGMetadata* SourceMetadata = SourceAttributeSet ? SourceAttributeSet->Metadata : OriginalData->Metadata;
+
+		if (!SourceMetadata)
 		{
 			PCGE_LOG(Warning, GraphAndLog, LOCTEXT("MissingMetadata", "Input has no metadata"));
 			continue;
@@ -107,14 +145,14 @@ bool FPCGMetadataOperationElement::ExecuteInternal(FPCGContext* Context) const
 
 		const FName SourceAttribute = InputSource.GetName();
 		const FName DestinationAttribute = OutputTarget.GetName();
-		const FName LocalSourceAttribute = ((SourceAttribute != NAME_None) ? SourceAttribute : OriginalData->Metadata->GetLatestAttributeNameOrNone());
-		const FName LocalDestinationAttribute = ((DestinationAttribute != NAME_None) ? DestinationAttribute : OriginalData->Metadata->GetLatestAttributeNameOrNone());
+		const FName LocalSourceAttribute = ((SourceAttribute != NAME_None) ? SourceAttribute : SourceMetadata->GetLatestAttributeNameOrNone());
+		const FName LocalDestinationAttribute = ((DestinationAttribute != NAME_None) ? DestinationAttribute : SourceMetadata->GetLatestAttributeNameOrNone());
 
 		// Make sure we use the right attribute name
 		InputSource.AttributeName = LocalSourceAttribute;
 		OutputTarget.AttributeName = LocalDestinationAttribute;
 
-		if (Settings->InputSource.Selection == EPCGAttributePropertySelection::Attribute && !OriginalData->Metadata->HasAttribute(LocalSourceAttribute))
+		if (Settings->InputSource.Selection == EPCGAttributePropertySelection::Attribute && !SourceMetadata->HasAttribute(LocalSourceAttribute))
 		{
 			PCGE_LOG(Warning, GraphAndLog, FText::Format(LOCTEXT("InputMissingAttribute", "Input does not have the '{0}' attribute"), FText::FromName(LocalSourceAttribute)));
 			continue;
@@ -141,16 +179,25 @@ bool FPCGMetadataOperationElement::ExecuteInternal(FPCGContext* Context) const
 				continue;
 			}
 
-			if (!SampledData->Metadata->CopyExistingAttribute(LocalSourceAttribute, DestinationAttribute))
+			if (SampledData->Metadata->HasAttribute(LocalSourceAttribute))
+			{
+				if (!SampledData->Metadata->CopyExistingAttribute(LocalSourceAttribute, DestinationAttribute))
+				{
+					PCGE_LOG(Warning, GraphAndLog, FText::Format(LOCTEXT("FailedCopyToNewAttribute", "Failed to copy to new attribute '{0}'"), FText::FromName(DestinationAttribute)));
+				}
+			}
+			else
 			{
-				PCGE_LOG(Warning, GraphAndLog, FText::Format(LOCTEXT("FailedCopyToNewAttribute", "Failed to copy to new attribute {0}"), FText::FromName(DestinationAttribute)));
+				SampledData->Metadata->CopyAttribute(SourceMetadata, LocalSourceAttribute, DestinationAttribute);
 			}
 
 			continue;
 		}
 
-		TUniquePtr<const IPCGAttributeAccessor> InputAccessor = PCGAttributeAccessorHelpers::CreateConstAccessor(OriginalData, InputSource);
-		TUniquePtr<const IPCGAttributeAccessorKeys> InputKeys = PCGAttributeAccessorHelpers::CreateConstKeys(OriginalData, InputSource);
+		const UPCGData* SourceData = (SourceAttributeSet ? static_cast<const UPCGData*>(SourceAttributeSet) : static_cast<const UPCGData*>(OriginalData));
+
+		TUniquePtr<const IPCGAttributeAccessor> InputAccessor = PCGAttributeAccessorHelpers::CreateConstAccessor(SourceData, InputSource);
+		TUniquePtr<const IPCGAttributeAccessorKeys> InputKeys = PCGAttributeAccessorHelpers::CreateConstKeys(SourceData, InputSource);
 
 		if (!InputAccessor.IsValid() || !InputKeys.IsValid())
 		{
@@ -160,7 +207,7 @@ bool FPCGMetadataOperationElement::ExecuteInternal(FPCGContext* Context) const
 
 		// If the target is an attribute, only create a new one if the attribute doesn't already exist.
 		// If it exist, it will try to write to it.
-		if (OutputTarget.Selection == EPCGAttributePropertySelection::Attribute && !OriginalData->Metadata->HasAttribute(LocalDestinationAttribute))
+		if (OutputTarget.Selection == EPCGAttributePropertySelection::Attribute && !SampledData->Metadata->HasAttribute(LocalDestinationAttribute))
 		{
 			auto CreateAttribute = [SampledData, LocalDestinationAttribute](auto Dummy)
 			{
@@ -170,7 +217,7 @@ bool FPCGMetadataOperationElement::ExecuteInternal(FPCGContext* Context) const
 			
 			if (!PCGMetadataAttribute::CallbackWithRightType(InputAccessor->GetUnderlyingType(), CreateAttribute))
 			{
-				PCGE_LOG(Warning, GraphAndLog, FText::Format(LOCTEXT("FailedToCreateNewAttribute", "Failed to create new attribute {0}"), FText::FromName(LocalDestinationAttribute)));
+				PCGE_LOG(Warning, GraphAndLog, FText::Format(LOCTEXT("FailedToCreateNewAttribute", "Failed to create new attribute '{0}'"), FText::FromName(LocalDestinationAttribute)));
 				continue;
 			}
 		}
@@ -184,9 +231,6 @@ bool FPCGMetadataOperationElement::ExecuteInternal(FPCGContext* Context) const
 			continue;
 		}
 
-		// By construction, they should be the same
-		check(InputKeys->GetNum() == OutputKeys->GetNum());
-
 		// Final verification, if we can put the value of input into output
 		if (!PCG::Private::IsBroadcastable(InputAccessor->GetUnderlyingType(), OutputAccessor->GetUnderlyingType()))
 		{
@@ -202,7 +246,7 @@ bool FPCGMetadataOperationElement::ExecuteInternal(FPCGContext* Context) const
 
 			EPCGAttributeAccessorFlags Flags = EPCGAttributeAccessorFlags::AllowBroadcast;
 
-			const int32 NumberOfElements = InputKeys->GetNum();
+			const int32 NumberOfElements = OutputKeys->GetNum();
 			constexpr int32 ChunkSize = 256;
 
 			TArray<OutputType, TInlineAllocator<ChunkSize>> TempValues;
diff --git a/Engine/Plugins/Experimental/PCG/Source/PCG/Private/Elements/PCGCreateAttribute.cpp b/Engine/Plugins/Experimental/PCG/Source/PCG/Private/Elements/PCGCreateAttribute.cpp
--- a/Engine/Plugins/Experimental/PCG/Source/PCG/Private/Elements/PCGCreateAttribute.cpp
+++ b/Engine/Plugins/Experimental/PCG/Source/PCG/Private/Elements/PCGCreateAttribute.cpp
@@ -19,6 +19,7 @@ namespace PCGCreateAttributeConstants
 	const FName NodeNameCreateAttribute = TEXT("CreateAttribute");
 	const FText NodeTitleCreateAttribute = LOCTEXT("NodeTitleCreateAttribute", "Create Attribute");
 	const FName AttributesLabel = TEXT("Attributes");
+	const FText AttributesTooltip = LOCTEXT("AttributesTooltip", "Optional Attribute Set to create from. Not used if not connected.");
 }
 
 void UPCGCreateAttributeSettings::PostLoad()
@@ -77,16 +78,18 @@ FText UPCGCreateAttributeSettings::GetDefaultNodeTitle() const
 {
 	return PCGCreateAttributeConstants::NodeTitleAddAttribute;
 }
+
+bool UPCGCreateAttributeSettings::IsPinUsedByNodeExecution(const UPCGPin* InPin) const
+{
+	return !InPin || (InPin->Properties.Label != PCGCreateAttributeConstants::AttributesLabel) || InPin->IsConnected();
+}
 #endif // WITH_EDITOR
 
 TArray<FPCGPinProperties> UPCGCreateAttributeSettings::InputPinProperties() const
 {
 	TArray<FPCGPinProperties> PinProperties;
 	PinProperties.Emplace(PCGPinConstants::DefaultInputLabel, EPCGDataType::Any);
-	if (bFromSourceParam)
-	{
-		PinProperties.Emplace(PCGCreateAttributeConstants::AttributesLabel, EPCGDataType::Param);
-	}
+	PinProperties.Emplace(PCGCreateAttributeConstants::AttributesLabel, EPCGDataType::Param, /*bInAllowMultipleConnections=*/true, /*bAllowMultipleData=*/true, PCGCreateAttributeConstants::AttributesTooltip);
 
 	return PinProperties;
 }
@@ -107,7 +110,7 @@ FPCGElementPtr UPCGCreateAttributeSettings::CreateElement() const
 
 FName UPCGCreateAttributeSettings::AdditionalTaskNameInternal(FName NodeName) const
 {
-	if (bFromSourceParam)
+	if (ShouldAddAttributesPin())
 	{
 		if ((OutputAttributeName == NAME_None) && (SourceParamAttributeName == NAME_None))
 		{
@@ -174,14 +177,8 @@ bool FPCGCreateAttributeElement::ExecuteInternal(FPCGContext* Context) const
 	UPCGParamData* SourceParamData = nullptr;
 	FName SourceParamAttributeName = NAME_None;
 
-	if (Settings->bFromSourceParam)
+	if (!SourceParams.IsEmpty())
 	{
-		if (SourceParams.IsEmpty())
-		{
-			PCGE_LOG(Error, GraphAndLog, LOCTEXT("ParamNotProvided", "Source attribute was not provided"));
-			return true;
-		}
-
 		SourceParamData = CastChecked<UPCGParamData>(SourceParams[0].Data);
 
 		if (!SourceParamData->Metadata)
@@ -253,11 +250,11 @@ bool FPCGCreateAttributeElement::ExecuteInternal(FPCGContext* Context) const
 			continue;
 		}
 
-		const FName OutputAttributeName = (Settings->bFromSourceParam && Settings->OutputAttributeName == NAME_None) ? SourceParamAttributeName : Settings->OutputAttributeName;
+		const FName OutputAttributeName = (SourceParamData && Settings->OutputAttributeName == NAME_None) ? SourceParamAttributeName : Settings->OutputAttributeName;
 
 		FPCGMetadataAttributeBase* Attribute = nullptr;
 
-		if (Settings->bFromSourceParam)
+		if (SourceParamData)
 		{
 			const FPCGMetadataAttributeBase* SourceAttribute = SourceParamData->Metadata->GetConstAttribute(SourceParamAttributeName);
 			Attribute = Metadata->CopyAttribute(SourceAttribute, OutputAttributeName, /*bKeepParent=*/false, /*bCopyEntries=*/bShouldAddNewEntry, /*bCopyValues=*/bShouldAddNewEntry);
@@ -278,7 +275,7 @@ bool FPCGCreateAttributeElement::ExecuteInternal(FPCGContext* Context) const
 		Output.Data = OutputData;
 
 		// Add a new entry if it is a param data and not from source (because entries are already copied)
-		if (bShouldAddNewEntry && !Settings->bFromSourceParam)
+		if (bShouldAddNewEntry && !SourceParamData)
 		{
 			// If the metadata is empty, we need to add a new entry, so set it to PCGInvalidEntryKey.
 			// Otherwise, use the entry key 0.
diff --git a/Engine/Plugins/Experimental/PCG/Source/PCG/Public/Elements/Metadata/PCGMetadataElement.h b/Engine/Plugins/Experimental/PCG/Source/PCG/Public/Elements/Metadata/PCGMetadataElement.h
--- a/Engine/Plugins/Experimental/PCG/Source/PCG/Public/Elements/Metadata/PCGMetadataElement.h
+++ b/Engine/Plugins/Experimental/PCG/Source/PCG/Public/Elements/Metadata/PCGMetadataElement.h
@@ -30,13 +30,15 @@ public:
 
 	//~Begin UPCGSettings interface
 #if WITH_EDITOR
-	virtual FName GetDefaultNodeName() const override { return FName(TEXT("AttributeOperation")); }
-	virtual FText GetDefaultNodeTitle() const override { return NSLOCTEXT("PCGMetadataOperationSettings", "NodeTitle", "Attribute Operation"); }
+	virtual FName GetDefaultNodeName() const override { return FName(TEXT("CopyAttribute")); }
+	virtual FText GetDefaultNodeTitle() const override { return NSLOCTEXT("PCGMetadataOperationSettings", "NodeTitle", "Copy Attribute"); }
+	virtual FText GetNodeTooltipText() const;
 	virtual EPCGSettingsType GetType() const override { return EPCGSettingsType::Metadata; }
+	virtual bool IsPinUsedByNodeExecution(const UPCGPin* InPin) const override;
 #endif
 
 protected:
-	virtual TArray<FPCGPinProperties> InputPinProperties() const override { return Super::DefaultPointInputPinProperties(); }
+	virtual TArray<FPCGPinProperties> InputPinProperties() const override;
 	virtual TArray<FPCGPinProperties> OutputPinProperties() const override { return Super::DefaultPointOutputPinProperties(); }
 	virtual FPCGElementPtr CreateElement() const override;
 	//~End UPCGSettings interface
diff --git a/Engine/Plugins/Experimental/PCG/Source/PCG/Public/Elements/PCGCreateAttribute.h b/Engine/Plugins/Experimental/PCG/Source/PCG/Public/Elements/PCGCreateAttribute.h
--- a/Engine/Plugins/Experimental/PCG/Source/PCG/Public/Elements/PCGCreateAttribute.h
+++ b/Engine/Plugins/Experimental/PCG/Source/PCG/Public/Elements/PCGCreateAttribute.h
@@ -32,6 +32,7 @@ public:
 	virtual FText GetDefaultNodeTitle() const override;
 	virtual EPCGSettingsType GetType() const override { return EPCGSettingsType::Metadata; }
 	virtual bool HasDynamicPins() const override { return true; }
+	virtual bool IsPinUsedByNodeExecution(const UPCGPin* InPin) const override;
 #endif // WITH_EDITOR
 	virtual EPCGDataType GetCurrentPinTypes(const UPCGPin* InPin) const override;
 
@@ -41,17 +42,14 @@ public:
 	UPROPERTY(BlueprintReadWrite, EditAnywhere, Category = Settings)
 	FName OutputAttributeName = NAME_None;
 
-	UPROPERTY(BlueprintReadWrite, EditAnywhere, Category = Settings, meta = (EditCondition = bDisplayFromSourceParamSetting, EditConditionHides, HideEditConditionToggle))
-	bool bFromSourceParam = false;
-
 	// This can be set false by inheriting nodes to hide the 'From Source Param' property.
 	UPROPERTY(Transient, meta = (EditCondition = false, EditConditionHides))
 	bool bDisplayFromSourceParamSetting = true;
 
-	UPROPERTY(BlueprintReadWrite, EditAnywhere, Category = Settings, meta = (EditCondition = "bFromSourceParam", EditConditionHides))
+	UPROPERTY(BlueprintReadWrite, EditAnywhere, Category = Settings, meta = (EditCondition = "bDisplayFromSourceParamSetting", EditConditionHides, HideEditConditionToggle))
 	FName SourceParamAttributeName = NAME_None;
 
-	UPROPERTY(BlueprintReadWrite, EditAnywhere, Category = Settings, meta = (EditCondition = "!bFromSourceParam", EditConditionHides, ShowOnlyInnerProperties, DisplayAfter = "bFromSourceParam"))
+	UPROPERTY(BlueprintReadWrite, EditAnywhere, Category = Settings, meta = (ShowOnlyInnerProperties))
 	FPCGMetadataTypesConstantStruct AttributeTypes;
 
 #if WITH_EDITORONLY_DATA
@@ -108,6 +106,7 @@ protected:
 	FName AdditionalTaskNameInternal(FName NodeName) const;
 
 	virtual FPCGElementPtr CreateElement() const override;
+	virtual bool ShouldAddAttributesPin() const { return true; }
 };
 
 /* Creates a new Attribute Set. */
@@ -131,6 +130,8 @@ protected:
 	virtual TArray<FPCGPinProperties> InputPinProperties() const override;
 	virtual TArray<FPCGPinProperties> OutputPinProperties() const override;
 	//~End UPCGSettings interface
+
+	virtual bool ShouldAddAttributesPin() const override { return false; }
 };
 
 class FPCGCreateAttributeElement : public FSimplePCGElement
diff --git a/Engine/Plugins/Experimental/PCG/Source/PCGEditor/Private/SPCGEditorGraphNode.cpp b/Engine/Plugins/Experimental/PCG/Source/PCGEditor/Private/SPCGEditorGraphNode.cpp
--- a/Engine/Plugins/Experimental/PCG/Source/PCGEditor/Private/SPCGEditorGraphNode.cpp
+++ b/Engine/Plugins/Experimental/PCG/Source/PCGEditor/Private/SPCGEditorGraphNode.cpp
@@ -40,6 +40,10 @@ public:
 	void Construct(const FArguments& InArgs, UEdGraphPin* InPin);
 
 	virtual FSlateColor GetPinColor() const override;
+	virtual FSlateColor GetPinTextColor() const override;
+
+private:
+	void ApplyUnusedPinStyle(FSlateColor& InOutColor) const;
 };
 
 void SPCGEditorGraphNodePin::Construct(const FArguments& InArgs, UEdGraphPin* InPin)
@@ -227,11 +231,8 @@ void SPCGEditorGraphNodePin::Construct(const FArguments& InArgs, UEdGraphPin* In
 
 }
 
-// Adapted from SGraphPin::GetPinColor
-FSlateColor SPCGEditorGraphNodePin::GetPinColor() const
+void SPCGEditorGraphNodePin::ApplyUnusedPinStyle(FSlateColor& InOutColor) const
 {
-	FSlateColor Color = SGraphPin::GetPinColor();
-
 	UEdGraphPin* GraphPin = GetPinObj();
 	if (GraphPin && !GraphPin->IsPendingKill())
 	{
@@ -239,12 +240,31 @@ FSlateColor SPCGEditorGraphNodePin::GetPinColor() const
 		const UPCGNode* PCGNode = EditorNode ? EditorNode->GetPCGNode() : nullptr;
 		const UPCGPin* PCGPin = PCGNode ? PCGNode->GetInputPin(GraphPin->GetFName()) : nullptr;
 
-		// Desaturate if pin is unused - intended to happen whether disabled or not
+		// Halve opacity if pin is unused - intended to happen whether disabled or not
 		if (PCGPin && !PCGNode->IsPinUsedByNodeExecution(PCGPin))
 		{
-			Color = Color.GetSpecifiedColor().Desaturate(0.7f);
+			FLinearColor Color = InOutColor.GetSpecifiedColor();
+			Color.A *= 0.5;
+			InOutColor = Color;
 		}
 	}
+}
+
+// Adapted from SGraphPin::GetPinColor
+FSlateColor SPCGEditorGraphNodePin::GetPinColor() const
+{
+	FSlateColor Color = SGraphPin::GetPinColor();
+
+	ApplyUnusedPinStyle(Color);
+
+	return Color;
+}
+
+FSlateColor SPCGEditorGraphNodePin::GetPinTextColor() const
+{
+	FSlateColor Color = SGraphPin::GetPinTextColor();
+
+	ApplyUnusedPinStyle(Color);
 
 	return Color;
 }
