Verify that when the PCG plugin is enabled PCG Graph is exposed within Content Browser creation menus
Steps:
- Within the Content Browser, select Add > PCG
- PCG Graph is exposed within the creation menus
- Within the Content Browser, right-click within the empty space
-Within the creation menu, mouse over PCG
-PCG Graph is exposed within the creation menus

diff --git a/Engine/Plugins/Experimental/PCG/Source/PCGEditor/Private/Tests/Display/PCGGrapthAssetIsExposedInContetBrowser.cpp b/Engine/Plugins/Experimental/PCG/Source/PCGEditor/Private/Tests/Display/PCGGrapthAssetIsExposedInContetBrowser.cpp
new file mode 100644
--- /dev/null
+++ b/Engine/Plugins/Experimental/PCG/Source/PCGEditor/Private/Tests/Display/PCGGrapthAssetIsExposedInContetBrowser.cpp
@@ -0,0 +1,47 @@
+#include "Tests/PCGTestsCommon.h"
+#include "IAssetTools.h"
+#include "AssetToolsModule.h"
+#include  "Factories/Factory.h"
+
+#if WITH_AUTOMATION_TESTS
+
+IMPLEMENT_CUSTOM_SIMPLE_AUTOMATION_TEST(FPCGGraphExposedInContentMenuTest, FPCGTestBaseClass, "Editor.Plugins.Tools.PCG.PCGGraphExposedInContentMenuTest", PCGTestsCommon::TestFlags)
+
+bool FPCGGraphExposedInContentMenuTest::RunTest(const FString& Parameters)
+{
+	const FName NAME_AssetTools = "AssetTools";
+	const FString PluginName = "PCG";
+	uint32 FactoryCategories = 0;
+	bool pcgGraphFound = false;
+
+	FAssetToolsModule& AssetToolsModule = FModuleManager::LoadModuleChecked<FAssetToolsModule>(NAME_AssetTools);
+
+	TArray<FAdvancedAssetCategory> AdvancedAssetCategories;
+	AssetToolsModule.Get().GetAllAdvancedAssetCategories(/*out*/ AdvancedAssetCategories);
+
+	FAdvancedAssetCategory* PCGAssetCategory = AdvancedAssetCategories.FindByPredicate([&PluginName](const FAdvancedAssetCategory& AdvancedAssetCategory) {
+		return AdvancedAssetCategory.CategoryName.ToString().Contains(PluginName);
+		});
+
+	UTEST_NOT_NULL(TEXT("PCG category not found"), PCGAssetCategory)
+
+	// We check if the PCG category factory's display name contains "PCG Graph"
+	TArray<UFactory*> Factories = AssetToolsModule.Get().GetNewAssetFactories();
+	for (UFactory* Factory : Factories)
+	{
+		FactoryCategories = Factory->GetMenuCategories();
+		if (FactoryCategories & PCGAssetCategory->CategoryType)
+		{
+			if (Factory->GetDisplayName().ToString().Contains("PCG Graph")) {
+				pcgGraphFound = true;
+				break;
+			}
+		}
+	}
+
+	TestTrue(TEXT("PCG Graph not found in content browser"), pcgGraphFound);
+
+	return true;
+}
+
+#endif
\ No newline at end of file
