Verify that the PCG Component is added by default to new PCG Volumes
Precondition: Open QAGame in Unreal Editor
Steps:
- From the Level Editor toolbar, search for and add a PCGVolume to the world
-Within the Details Panel Component section, a PCG Component is present

diff --git a/Engine/Plugins/Experimental/PCG/Source/PCGEditor/Private/Tests/Display/PCGComponentInNewPCGVolumes.cpp b/Engine/Plugins/Experimental/PCG/Source/PCGEditor/Private/Tests/Display/PCGComponentInNewPCGVolumes.cpp
new file mode 100644
--- /dev/null
+++ b/Engine/Plugins/Experimental/PCG/Source/PCGEditor/Private/Tests/Display/PCGComponentInNewPCGVolumes.cpp
@@ -0,0 +1,51 @@
+// Copyright Epic Games, Inc. All Rights Reserved.
+
+#include "EngineUtils.h"
+#include "Engine/World.h"
+#include "Misc/AutomationTest.h"
+#include "Editor.h"
+#include "Editor/UnrealEd/Public/PackageTools.h"
+#include "AssetRegistry/AssetRegistryModule.h"
+#include "IPlacementModeModule.h"
+#include "PCGVolume.h"
+#include "PCGComponent.h"
+#include "Tests/PCGTestsCommon.h"
+
+
+#if WITH_AUTOMATION_TESTS
+
+IMPLEMENT_CUSTOM_SIMPLE_AUTOMATION_TEST(FPCGComponentInNewPCGVolumes, FPCGTestBaseClass, "Editor.Plugins.Tools.PCG.PCGComponentInNewPCGVolumes", PCGTestsCommon::TestFlags)
+bool FPCGComponentInNewPCGVolumes::RunTest(const FString& Parameters)
+{
+	// Get current world
+	UWorld* World = GEditor->GetEditorWorldContext().World();
+	UTEST_NOT_NULL(TEXT("Failed to get editor world context!"), World);
+
+	// We need to get the PCGVolume actor factory
+	TSubclassOf<AActor> PCGVolumeClass = APCGVolume::StaticClass();
+	UActorFactory* PCGVolumeFactory = GEditor->FindActorFactoryForActorClass(PCGVolumeClass);
+	UTEST_NOT_NULL("Failed to find PCGVolume actor factory.", PCGVolumeFactory);
+
+	// Get the asset data for the PCGVolume class
+	FAssetData PCGVolumeAssetData = FAssetData(PCGVolumeClass);
+
+	// Add a new PCGVolume actor to the level
+	AActor* VolumeActor = nullptr;
+	// To clean created assets, there appears window save content
+	if (GCurrentLevelEditingViewportClient)
+	{
+		FTransform ActorTransform;
+		VolumeActor = GEditor->UseActorFactory(PCGVolumeFactory, PCGVolumeAssetData, &ActorTransform);
+		UPackage* Package = VolumeActor->GetOutermost();
+		Package->SetFlags(RF_Transient); // Mark package as transient
+	}
+	UTEST_NOT_NULL(TEXT("Failed to add PCGVolume actor."), VolumeActor);
+
+	UPCGComponent* PCGComponent = VolumeActor->FindComponentByClass<UPCGComponent>();
+	TestNotNull("PCGVolume actor does not contain a PCGComponent!", PCGComponent);
+
+	// Destroy the PCGVolume actor after test
+	World->EditorDestroyActor(VolumeActor, false);
+	return true;
+}
+#endif
\ No newline at end of file
