Verify that when the PCG plugin is enabled PCG Graph is exposed within Content Browser creation menus
Steps:
- Within the Content Browser, select Add > PCG
- PCG Graph is exposed within the creation menus
- Within the Content Browser, right-click within the empty space
-Within the creation menu, mouse over PCG
-PCG Graph is exposed within the creation menus

diff --git a/Engine/Plugins/Experimental/PCG/Source/PCGEditor/Private/Tests/Display/PCGGrapthAssetIsExposedInContetBrowser.cpp b/Engine/Plugins/Experimental/PCG/Source/PCGEditor/Private/Tests/Display/PCGGrapthAssetIsExposedInContetBrowser.cpp
--- a/Engine/Plugins/Experimental/PCG/Source/PCGEditor/Private/Tests/Display/PCGGrapthAssetIsExposedInContetBrowser.cpp
+++ b/Engine/Plugins/Experimental/PCG/Source/PCGEditor/Private/Tests/Display/PCGGrapthAssetIsExposedInContetBrowser.cpp
@@ -1,3 +1,5 @@
+// Copyright Epic Games, Inc. All Rights Reserved.
+
 #include "Tests/PCGTestsCommon.h"
 #include "IAssetTools.h"
 #include "AssetToolsModule.h"
