[PCG] FName metadata should not be compressed

diff --git a/Engine/Plugins/Experimental/PCG/Source/PCG/Public/Metadata/PCGMetadataAttributeTraits.h b/Engine/Plugins/Experimental/PCG/Source/PCG/Public/Metadata/PCGMetadataAttributeTraits.h
--- a/Engine/Plugins/Experimental/PCG/Source/PCG/Public/Metadata/PCGMetadataAttributeTraits.h
+++ b/Engine/Plugins/Experimental/PCG/Source/PCG/Public/Metadata/PCGMetadataAttributeTraits.h
@@ -742,7 +742,7 @@ namespace PCG
 		template<>
 		struct MetadataTraits<FName>
 		{
-			enum { CompressData = true };
+			enum { CompressData = false };
 			enum { CanMinMax = false };
 			enum { CanSubAdd = false };
 			enum { CanMulDiv = false };
