[PCG] Minor improvement to execution breakpoints after getting more mileage and feedback

diff --git a/Engine/Plugins/Experimental/PCG/Source/PCG/Private/PCGElement.cpp b/Engine/Plugins/Experimental/PCG/Source/PCG/Private/PCGElement.cpp
--- a/Engine/Plugins/Experimental/PCG/Source/PCG/Private/PCGElement.cpp
+++ b/Engine/Plugins/Experimental/PCG/Source/PCG/Private/PCGElement.cpp
@@ -23,11 +23,6 @@ static TAutoConsoleVariable<bool> CVarPCGValidatePointMetadata(
 #define PCG_ELEMENT_EXECUTION_BREAKPOINT() \
 	if (Context && Context->GetInputSettingsInterface() && Context->GetInputSettingsInterface()->bBreakDebugger) \
 	{ \
-		if (const UEnum* EnumPtr = FindObject<UEnum>(nullptr, TEXT("/Script/PCG.EPCGExecutionPhase"), true)) \
-		{ \
-			PCGE_LOG(Log, LogOnly, FText::Format(LOCTEXT("BreakExecution", "Execution halted, phase: {0}"), \
-				FText::FromName(EnumPtr->GetNameByValue(static_cast<int>(Context->CurrentPhase))))); \
-		} \
 		PLATFORM_BREAK(); \
 	}
 #else
@@ -39,8 +34,6 @@ bool IPCGElement::Execute(FPCGContext* Context) const
 	check(Context && Context->AsyncState.NumAvailableTasks > 0 && Context->CurrentPhase < EPCGExecutionPhase::Done);
 	check(Context->AsyncState.bIsRunningOnMainThread || !CanExecuteOnlyOnMainThread(Context));
 
-	PCG_ELEMENT_EXECUTION_BREAKPOINT();
-
 	while (Context->CurrentPhase != EPCGExecutionPhase::Done)
 	{
 		PCGUtils::FScopedCall ScopedCall(*this, Context);
@@ -51,6 +44,8 @@ bool IPCGElement::Execute(FPCGContext* Context) const
 			case EPCGExecutionPhase::NotExecuted:
 			{
 				TRACE_CPUPROFILER_EVENT_SCOPE(EPCGExecutionPhase::NotExecuted);
+				PCG_ELEMENT_EXECUTION_BREAKPOINT();
+
 				PreExecute(Context);
 
 				// Will override the settings if there is any override.
@@ -62,6 +57,7 @@ bool IPCGElement::Execute(FPCGContext* Context) const
 			case EPCGExecutionPhase::PrepareData:
 			{
 				TRACE_CPUPROFILER_EVENT_SCOPE(EPCGExecutionPhase::PrepareData);
+				PCG_ELEMENT_EXECUTION_BREAKPOINT();
 
 				if (PrepareDataInternal(Context))
 				{
@@ -77,6 +73,8 @@ bool IPCGElement::Execute(FPCGContext* Context) const
 			case EPCGExecutionPhase::Execute:
 			{
 				TRACE_CPUPROFILER_EVENT_SCOPE(EPCGExecutionPhase::Execute);
+				PCG_ELEMENT_EXECUTION_BREAKPOINT();
+
 				if (ExecuteInternal(Context))
 				{
 					Context->CurrentPhase = EPCGExecutionPhase::PostExecute;
@@ -91,6 +89,8 @@ bool IPCGElement::Execute(FPCGContext* Context) const
 			case EPCGExecutionPhase::PostExecute:
 			{
 				TRACE_CPUPROFILER_EVENT_SCOPE(EPCGExecutionPhase::PostExecute);
+				PCG_ELEMENT_EXECUTION_BREAKPOINT();
+
 				PostExecute(Context);
 				break;
 			}
diff --git a/Engine/Plugins/Experimental/PCG/Source/PCG/Public/PCGSettings.h b/Engine/Plugins/Experimental/PCG/Source/PCG/Public/PCGSettings.h
--- a/Engine/Plugins/Experimental/PCG/Source/PCG/Public/PCGSettings.h
+++ b/Engine/Plugins/Experimental/PCG/Source/PCG/Public/PCGSettings.h
@@ -136,8 +136,8 @@ public:
 	UPROPERTY(EditAnywhere, BlueprintReadWrite, Category = Debug, meta = (ShowOnlyInnerProperties))
 	FPCGDebugVisualizationSettings DebugSettings;
 
-	/** If a debugger is attached, triggers a breakpoint inside IPCGElement::Execute(). Editor only. */
-	UPROPERTY(Transient, EditAnywhere, BlueprintReadWrite, Category = Debug)
+	/** If a debugger is attached, triggers a breakpoint inside IPCGElement::Execute(). Editor only. Transient. */
+	UPROPERTY(Transient, DuplicateTransient, EditAnywhere, BlueprintReadWrite, Category = Debug)
 	bool bBreakDebugger = false;
 #endif
 };
