Verify that the PCG Component is added by default to new PCG Volumes
Precondition: Open QAGame in Unreal Editor
Steps:
From the Level Editor toolbar, search for and add a PCGVolume to the world
-Within the Details Panel Component section, a PCG Component is present

diff --git a/Engine/Plugins/Experimental/PCG/Source/PCGEditor/Private/Tests/Display/PCGComponentInNewPCGVolumes.cpp b/Engine/Plugins/Experimental/PCG/Source/PCGEditor/Private/Tests/Display/PCGComponentInNewPCGVolumes.cpp
--- a/Engine/Plugins/Experimental/PCG/Source/PCGEditor/Private/Tests/Display/PCGComponentInNewPCGVolumes.cpp
+++ b/Engine/Plugins/Experimental/PCG/Source/PCGEditor/Private/Tests/Display/PCGComponentInNewPCGVolumes.cpp
@@ -11,7 +11,6 @@
 #include "PCGComponent.h"
 #include "Tests/PCGTestsCommon.h"
 
-
 #if WITH_AUTOMATION_TESTS
 
 IMPLEMENT_CUSTOM_SIMPLE_AUTOMATION_TEST(FPCGComponentInNewPCGVolumes, FPCGTestBaseClass, "Editor.Plugins.Tools.PCG.PCGComponentInNewPCGVolumes", PCGTestsCommon::TestFlags)
@@ -24,7 +23,7 @@ bool FPCGComponentInNewPCGVolumes::RunTest(const FString& Parameters)
 	// We need to get the PCGVolume actor factory
 	TSubclassOf<AActor> PCGVolumeClass = APCGVolume::StaticClass();
 	UActorFactory* PCGVolumeFactory = GEditor->FindActorFactoryForActorClass(PCGVolumeClass);
-	UTEST_NOT_NULL("Failed to find PCGVolume actor factory.", PCGVolumeFactory);
+	UTEST_NOT_NULL(TEXT("Failed to find PCGVolume actor factory."), PCGVolumeFactory);
 
 	// Get the asset data for the PCGVolume class
 	FAssetData PCGVolumeAssetData = FAssetData(PCGVolumeClass);
@@ -36,16 +35,21 @@ bool FPCGComponentInNewPCGVolumes::RunTest(const FString& Parameters)
 	{
 		FTransform ActorTransform;
 		VolumeActor = GEditor->UseActorFactory(PCGVolumeFactory, PCGVolumeAssetData, &ActorTransform);
+	}
+	if (!VolumeActor) return false;
+	
+	if (VolumeActor)
+	{
 		UPackage* Package = VolumeActor->GetOutermost();
 		Package->SetFlags(RF_Transient); // Mark package as transient
-	}
-	UTEST_NOT_NULL(TEXT("Failed to add PCGVolume actor."), VolumeActor);
 
-	UPCGComponent* PCGComponent = VolumeActor->FindComponentByClass<UPCGComponent>();
-	TestNotNull("PCGVolume actor does not contain a PCGComponent!", PCGComponent);
+		UPCGComponent* PCGComponent = VolumeActor->FindComponentByClass<UPCGComponent>();
+		TestNotNull("PCGVolume actor does not contain a PCGComponent!", PCGComponent);
+
+		// Destroy the PCGVolume actor after test
+		World->EditorDestroyActor(VolumeActor, false);
+	}
 
-	// Destroy the PCGVolume actor after test
-	World->EditorDestroyActor(VolumeActor, false);
 	return true;
 }
-#endif
\ No newline at end of file
+#endif
