[PCG] Refactored attribute list view to use point data / metadata accessors

diff --git a/Engine/Plugins/Experimental/PCG/Source/PCG/Public/Metadata/Accessors/IPCGAttributeAccessor.h b/Engine/Plugins/Experimental/PCG/Source/PCG/Public/Metadata/Accessors/IPCGAttributeAccessor.h
--- a/Engine/Plugins/Experimental/PCG/Source/PCG/Public/Metadata/Accessors/IPCGAttributeAccessor.h
+++ b/Engine/Plugins/Experimental/PCG/Source/PCG/Public/Metadata/Accessors/IPCGAttributeAccessor.h
@@ -37,7 +37,7 @@ ENUM_CLASS_FLAGS(EPCGAttributeAccessorFlags);
 * For Get/Set, you need a key, that will represent how you can access the value wanted. cf PCGAttributeAccessorKeys.
 * NOTE: This is not threadsafe and not intended to be used unprotected.
 */
-class IPCGAttributeAccessor
+class PCG_API IPCGAttributeAccessor
 {
 public:
 	virtual ~IPCGAttributeAccessor() = default;
@@ -153,8 +153,8 @@ protected:
 
 // Specialization of IPCGAttributeAccessor::Get and IPCGAttributeAccessor::Set, for all supported types.
 #define IACCESSOR_DECL(T) \
-template <> bool IPCGAttributeAccessor::GetRange<T>(TArrayView<T> OutValues, int32 Index, const IPCGAttributeAccessorKeys& Keys, EPCGAttributeAccessorFlags Flags) const; \
-template <> bool IPCGAttributeAccessor::SetRange<T>(TArrayView<const T> InValues, int32 Index, IPCGAttributeAccessorKeys& Keys, EPCGAttributeAccessorFlags Flags);
+template <> bool PCG_API IPCGAttributeAccessor::GetRange<T>(TArrayView<T> OutValues, int32 Index, const IPCGAttributeAccessorKeys& Keys, EPCGAttributeAccessorFlags Flags) const; \
+template <> bool PCG_API IPCGAttributeAccessor::SetRange<T>(TArrayView<const T> InValues, int32 Index, IPCGAttributeAccessorKeys& Keys, EPCGAttributeAccessorFlags Flags);
 PCG_FOREACH_SUPPORTEDTYPES(IACCESSOR_DECL);
 #undef IACCESSOR_DECL
 
diff --git a/Engine/Plugins/Experimental/PCG/Source/PCG/Public/Metadata/Accessors/PCGAttributeAccessorHelpers.h b/Engine/Plugins/Experimental/PCG/Source/PCG/Public/Metadata/Accessors/PCGAttributeAccessorHelpers.h
--- a/Engine/Plugins/Experimental/PCG/Source/PCG/Public/Metadata/Accessors/PCGAttributeAccessorHelpers.h
+++ b/Engine/Plugins/Experimental/PCG/Source/PCG/Public/Metadata/Accessors/PCGAttributeAccessorHelpers.h
@@ -18,35 +18,35 @@ struct FPCGSettingsOverridableParam;
 
 namespace PCGAttributeAccessorHelpers
 {
-	bool IsPropertyAccessorSupported(const FProperty* InProperty);
-	bool IsPropertyAccessorSupported(const FName InPropertyName, const UClass* InClass);
-	bool IsPropertyAccessorSupported(const FName InPropertyName, const UStruct* InStruct);
+	PCG_API bool IsPropertyAccessorSupported(const FProperty* InProperty);
+	PCG_API bool IsPropertyAccessorSupported(const FName InPropertyName, const UClass* InClass);
+	PCG_API bool IsPropertyAccessorSupported(const FName InPropertyName, const UStruct* InStruct);
 
-	TUniquePtr<IPCGAttributeAccessor> CreatePropertyAccessor(const FProperty* InProperty);
-	TUniquePtr<IPCGAttributeAccessor> CreatePropertyAccessor(const FName InPropertyName, const UClass* InClass);
-	TUniquePtr<IPCGAttributeAccessor> CreatePropertyAccessor(const FName InPropertyName, const UStruct* InStruct);
+	PCG_API TUniquePtr<IPCGAttributeAccessor> CreatePropertyAccessor(const FProperty* InProperty);
+	PCG_API TUniquePtr<IPCGAttributeAccessor> CreatePropertyAccessor(const FName InPropertyName, const UClass* InClass);
+	PCG_API TUniquePtr<IPCGAttributeAccessor> CreatePropertyAccessor(const FName InPropertyName, const UStruct* InStruct);
 
-	TUniquePtr<IPCGAttributeAccessor> CreateChainAccessor(TUniquePtr<IPCGAttributeAccessor> InAccessor, FName Name, bool& bOutSuccess);
+	PCG_API TUniquePtr<IPCGAttributeAccessor> CreateChainAccessor(TUniquePtr<IPCGAttributeAccessor> InAccessor, FName Name, bool& bOutSuccess);
 
 	/**
 	* Create a const accessor depending on an overridable param
 	*/
-	TUniquePtr<const IPCGAttributeAccessor> CreateConstAccessorForOverrideParam(const FPCGDataCollection& InInputData, const FPCGSettingsOverridableParam& InParam, FName* OutAttributeName = nullptr);
+	PCG_API TUniquePtr<const IPCGAttributeAccessor> CreateConstAccessorForOverrideParam(const FPCGDataCollection& InInputData, const FPCGSettingsOverridableParam& InParam, FName* OutAttributeName = nullptr);
 
 	/** 
 	* Creates a const accessor to the property or attribute pointed at by the InSelector.
 	* Note that InData must not be null if the selector points to an attribute,
 	* but in the case of properties, it either has to be the appropriate type or null.
 	*/
-	TUniquePtr<const IPCGAttributeAccessor> CreateConstAccessor(const UPCGData* InData, const FPCGAttributePropertySelector& InSelector);
+	PCG_API TUniquePtr<const IPCGAttributeAccessor> CreateConstAccessor(const UPCGData* InData, const FPCGAttributePropertySelector& InSelector);
 	
 	/** 
 	* Creates a non-const accessor to the property or attribute pointed at by the InSelector.
 	* Note that InData must not be null if the selector points to an attribute,
 	* but in the case of properties, it either has to be the appropriate type or null.
 	*/
-	TUniquePtr<IPCGAttributeAccessor> CreateAccessor(UPCGData* InData, const FPCGAttributePropertySelector& InSelector);
+	PCG_API TUniquePtr<IPCGAttributeAccessor> CreateAccessor(UPCGData* InData, const FPCGAttributePropertySelector& InSelector);
 
-	TUniquePtr<const IPCGAttributeAccessorKeys> CreateConstKeys(const UPCGData* InData, const FPCGAttributePropertySelector& InSelector);
-	TUniquePtr<IPCGAttributeAccessorKeys> CreateKeys(UPCGData* InData, const FPCGAttributePropertySelector& InSelector);
+	PCG_API TUniquePtr<const IPCGAttributeAccessorKeys> CreateConstKeys(const UPCGData* InData, const FPCGAttributePropertySelector& InSelector);
+	PCG_API TUniquePtr<IPCGAttributeAccessorKeys> CreateKeys(UPCGData* InData, const FPCGAttributePropertySelector& InSelector);
 }
\ No newline at end of file
diff --git a/Engine/Plugins/Experimental/PCG/Source/PCG/Public/Metadata/Accessors/PCGAttributeAccessorKeys.h b/Engine/Plugins/Experimental/PCG/Source/PCG/Public/Metadata/Accessors/PCGAttributeAccessorKeys.h
--- a/Engine/Plugins/Experimental/PCG/Source/PCG/Public/Metadata/Accessors/PCGAttributeAccessorKeys.h
+++ b/Engine/Plugins/Experimental/PCG/Source/PCG/Public/Metadata/Accessors/PCGAttributeAccessorKeys.h
@@ -45,7 +45,7 @@ namespace PCGAttributeAccessorKeys
 /**
 * Base class to identify keys to use with an accessor.
 */
-class IPCGAttributeAccessorKeys
+class PCG_API IPCGAttributeAccessorKeys
 {
 public:
 	explicit IPCGAttributeAccessorKeys(bool bInReadOnly)
@@ -134,7 +134,7 @@ protected:
 /**
 * Key around a metadata entry key
 */
-class FPCGAttributeAccessorKeysEntries : public IPCGAttributeAccessorKeys
+class PCG_API FPCGAttributeAccessorKeysEntries : public IPCGAttributeAccessorKeys
 {
 public:
 	explicit FPCGAttributeAccessorKeysEntries(const FPCGMetadataAttributeBase* Attribute);
@@ -154,7 +154,7 @@ protected:
 /**
 * Key around points
 */
-class FPCGAttributeAccessorKeysPoints : public IPCGAttributeAccessorKeys
+class PCG_API FPCGAttributeAccessorKeysPoints : public IPCGAttributeAccessorKeys
 {
 public:
 	FPCGAttributeAccessorKeysPoints(const TArrayView<FPCGPoint>& InPoints);
@@ -296,7 +296,7 @@ protected:
 * We can't use FPCGAttributeAccessorKeysGeneric since it has a constructor taking a reference on a object,
 * and you can't have void&.
 */
-class FPCGAttributeAccessorKeysGenericPtrs : public IPCGAttributeAccessorKeys
+class PCG_API FPCGAttributeAccessorKeysGenericPtrs : public IPCGAttributeAccessorKeys
 {
 public:
 	FPCGAttributeAccessorKeysGenericPtrs(const TArrayView<void*>& InPtrs)
diff --git a/Engine/Plugins/Experimental/PCG/Source/PCG/Public/PCGPoint.h b/Engine/Plugins/Experimental/PCG/Source/PCG/Public/PCGPoint.h
--- a/Engine/Plugins/Experimental/PCG/Source/PCG/Public/PCGPoint.h
+++ b/Engine/Plugins/Experimental/PCG/Source/PCG/Public/PCGPoint.h
@@ -21,7 +21,8 @@ enum class EPCGPointProperties : uint8
 	Scale,
 	Transform,
 	Steepness,
-	LocalCenter
+	LocalCenter,
+	Seed
 };
 
 USTRUCT(BlueprintType)
diff --git a/Engine/Plugins/Experimental/PCG/Source/PCGEditor/Private/SPCGEditorGraphAttributeListView.cpp b/Engine/Plugins/Experimental/PCG/Source/PCGEditor/Private/SPCGEditorGraphAttributeListView.cpp
--- a/Engine/Plugins/Experimental/PCG/Source/PCGEditor/Private/SPCGEditorGraphAttributeListView.cpp
+++ b/Engine/Plugins/Experimental/PCG/Source/PCGEditor/Private/SPCGEditorGraphAttributeListView.cpp
@@ -3,12 +3,19 @@
 #include "SPCGEditorGraphAttributeListView.h"
 
 #include "PCGComponent.h"
+#include "PCGData.h"
 #include "PCGEditor.h"
 #include "PCGEditorGraphNodeBase.h"
 #include "PCGParamData.h"
 #include "PCGSubsystem.h"
 #include "Data/PCGPointData.h"
 #include "Data/PCGSpatialData.h"
+#include "Metadata/PCGAttributePropertySelector.h"
+#include "Metadata/Accessors/IPCGAttributeAccessor.h"
+#include "Metadata/Accessors/PCGAttributeAccessor.h"
+#include "Metadata/Accessors/PCGAttributeAccessorHelpers.h"
+#include "Metadata/Accessors/PCGAttributeAccessorKeys.h"
+#include "Metadata/Accessors/PCGCustomAccessor.h"
 
 #include "Fonts/FontMeasure.h"
 #include "Framework/MultiBox/MultiBoxBuilder.h"
@@ -25,58 +32,59 @@ static TAutoConsoleVariable<bool> CVarShowAdvancedAttributesFields(
 
 namespace PCGEditorGraphAttributeListView
 {
+	const FString LastMetadataId = TEXT("@Last");
 	const FText NoDataAvailableText = LOCTEXT("NoDataAvailableText", "No data available");
 	const FText NoNodeInspectedText = LOCTEXT("NoNodeInspectedText", "No node being inspected");
 	const FText NoNodeInspectedToolTip = LOCTEXT("NoNodeInspectedToolTip", "Inspect a node using the right click menu");
 	
 	/** Names of the columns in the attribute list */
 	const FName NAME_IndexColumn = FName(TEXT("IndexColumn"));
-	const FName NAME_PointPositionX = FName(TEXT("PointPositionX"));
-	const FName NAME_PointPositionY = FName(TEXT("PointPositionY"));
-	const FName NAME_PointPositionZ = FName(TEXT("PointPositionZ"));
-	const FName NAME_PointRotationX = FName(TEXT("PointRotationX"));
-	const FName NAME_PointRotationY = FName(TEXT("PointRotationY"));
-	const FName NAME_PointRotationZ = FName(TEXT("PointRotationZ"));
-	const FName NAME_PointScaleX = FName(TEXT("PointScaleX"));
-	const FName NAME_PointScaleY = FName(TEXT("PointScaleY"));
-	const FName NAME_PointScaleZ = FName(TEXT("PointScaleZ"));
-	const FName NAME_PointBoundsMinX = FName(TEXT("PointBoundsMinX"));
-	const FName NAME_PointBoundsMinY = FName(TEXT("PointBoundsMinY"));
-	const FName NAME_PointBoundsMinZ = FName(TEXT("PointBoundsMinZ"));
-	const FName NAME_PointBoundsMaxX = FName(TEXT("PointBoundsMaxX"));
-	const FName NAME_PointBoundsMaxY = FName(TEXT("PointBoundsMaxY"));
-	const FName NAME_PointBoundsMaxZ = FName(TEXT("PointBoundsMaxZ"));
-	const FName NAME_PointColorR = FName(TEXT("PointColorR"));
-	const FName NAME_PointColorG = FName(TEXT("PointColorG"));
-	const FName NAME_PointColorB = FName(TEXT("PointColorB"));
-	const FName NAME_PointColorA = FName(TEXT("PointColorA"));
-	const FName NAME_PointDensity = FName(TEXT("PointDensity"));
-	const FName NAME_PointSteepness = FName(TEXT("PointSteepness"));
-	const FName NAME_PointSeed = FName(TEXT("PointSeed"));
-	const FName NAME_PointMetadataEntry = FName(TEXT("PointMetadataEntry"));
+	const FName NAME_PointPositionX = FName(TEXT("$Transform.Location.X"));
+	const FName NAME_PointPositionY = FName(TEXT("$Transform.Location.Y"));
+	const FName NAME_PointPositionZ = FName(TEXT("$Transform.Location.Z"));
+	const FName NAME_PointRotationX = FName(TEXT("$Transform.Rotation.X"));
+	const FName NAME_PointRotationY = FName(TEXT("$Transform.Rotation.Y"));
+	const FName NAME_PointRotationZ = FName(TEXT("$Transform.Rotation.Z"));
+	const FName NAME_PointScaleX = FName(TEXT("$Transform.Scale.X"));
+	const FName NAME_PointScaleY = FName(TEXT("$Transform.Scale.Y"));
+	const FName NAME_PointScaleZ = FName(TEXT("$Transform.Scale.Z"));
+	const FName NAME_PointBoundsMinX = FName(TEXT("$BoundsMin.X"));
+	const FName NAME_PointBoundsMinY = FName(TEXT("$BoundsMin.Y"));
+	const FName NAME_PointBoundsMinZ = FName(TEXT("$BoundsMin.Z"));
+	const FName NAME_PointBoundsMaxX = FName(TEXT("$BoundsMax.X"));
+	const FName NAME_PointBoundsMaxY = FName(TEXT("$BoundsMax.Y"));
+	const FName NAME_PointBoundsMaxZ = FName(TEXT("$BoundsMax.Z"));
+	const FName NAME_PointColorR = FName(TEXT("$Color.R"));
+	const FName NAME_PointColorG = FName(TEXT("$Color.G"));
+	const FName NAME_PointColorB = FName(TEXT("$Color.B"));
+	const FName NAME_PointColorA = FName(TEXT("$Color.A"));
+	const FName NAME_PointDensity = FName(TEXT("$Density"));
+	const FName NAME_PointSteepness = FName(TEXT("$Steepness"));
+	const FName NAME_PointSeed = FName(TEXT("$Seed"));
+	const FName NAME_PointMetadataEntry = FName(TEXT("MetadataEntry"));
 	const FName NAME_PointMetadataEntryParent = FName(TEXT("PointMetadataEntryParent"));
 
 	/** Labels of the columns */
 	const FText TEXT_IndexLabel = LOCTEXT("IndexLabel", "Index");
-	const FText TEXT_PointPositionLabelX = LOCTEXT("PointPositionLabelX", "PositionX");
-	const FText TEXT_PointPositionLabelY = LOCTEXT("PointPositionLabelY", "PositionY");
-	const FText TEXT_PointPositionLabelZ = LOCTEXT("PointPositionLabelZ", "PositionZ");
-	const FText TEXT_PointRotationLabelX = LOCTEXT("PointRotationLabelX", "RotationX");
-	const FText TEXT_PointRotationLabelY = LOCTEXT("PointRotationLabelY", "RotationY");
-	const FText TEXT_PointRotationLabelZ = LOCTEXT("PointRotationLabelZ", "RotationZ");
-	const FText TEXT_PointScaleLabelX = LOCTEXT("PointScaleLabelX", "ScaleX");
-	const FText TEXT_PointScaleLabelY = LOCTEXT("PointScaleLabelY", "ScaleY");
-	const FText TEXT_PointScaleLabelZ = LOCTEXT("PointScaleLabelZ", "ScaleZ");
-	const FText TEXT_PointBoundsLabelMinX = LOCTEXT("PointBoundsMinX", "BoundsMinX");
-	const FText TEXT_PointBoundsLabelMinY = LOCTEXT("PointBoundsMinY", "BoundsMinY");
-	const FText TEXT_PointBoundsLabelMinZ = LOCTEXT("PointBoundsMinZ", "BoundsMinZ");
-	const FText TEXT_PointBoundsLabelMaxX = LOCTEXT("PointBoundsMaxX", "BoundsMaxX");
-	const FText TEXT_PointBoundsLabelMaxY = LOCTEXT("PointBoundsMaxY", "BoundsMaxY");
-	const FText TEXT_PointBoundsLabelMaxZ = LOCTEXT("PointBoundsMaxZ", "BoundsMaxZ");
-	const FText TEXT_PointColorLabelR = LOCTEXT("PointColorR", "ColorR");
-	const FText TEXT_PointColorLabelG = LOCTEXT("PointColorG", "ColorG");
-	const FText TEXT_PointColorLabelB = LOCTEXT("PointColorB", "ColorB");
-	const FText TEXT_PointColorLabelA = LOCTEXT("PointColorA", "ColorA");
+	const FText TEXT_PointPositionLabelX = LOCTEXT("PointPositionLabelX", "Position.X");
+	const FText TEXT_PointPositionLabelY = LOCTEXT("PointPositionLabelY", "Position.Y");
+	const FText TEXT_PointPositionLabelZ = LOCTEXT("PointPositionLabelZ", "Position.Z");
+	const FText TEXT_PointRotationLabelX = LOCTEXT("PointRotationLabelX", "Rotation.X");
+	const FText TEXT_PointRotationLabelY = LOCTEXT("PointRotationLabelY", "Rotation.Y");
+	const FText TEXT_PointRotationLabelZ = LOCTEXT("PointRotationLabelZ", "Rotation.Z");
+	const FText TEXT_PointScaleLabelX = LOCTEXT("PointScaleLabelX", "Scale.X");
+	const FText TEXT_PointScaleLabelY = LOCTEXT("PointScaleLabelY", "Scale.Y");
+	const FText TEXT_PointScaleLabelZ = LOCTEXT("PointScaleLabelZ", "Scale.Z");
+	const FText TEXT_PointBoundsLabelMinX = LOCTEXT("PointBoundsMinX", "BoundsMin.X");
+	const FText TEXT_PointBoundsLabelMinY = LOCTEXT("PointBoundsMinY", "BoundsMin.Y");
+	const FText TEXT_PointBoundsLabelMinZ = LOCTEXT("PointBoundsMinZ", "BoundsMin.Z");
+	const FText TEXT_PointBoundsLabelMaxX = LOCTEXT("PointBoundsMaxX", "BoundsMax.X");
+	const FText TEXT_PointBoundsLabelMaxY = LOCTEXT("PointBoundsMaxY", "BoundsMax.Y");
+	const FText TEXT_PointBoundsLabelMaxZ = LOCTEXT("PointBoundsMaxZ", "BoundsMax.Z");
+	const FText TEXT_PointColorLabelR = LOCTEXT("PointColorR", "Color.R");
+	const FText TEXT_PointColorLabelG = LOCTEXT("PointColorG", "Color.G");
+	const FText TEXT_PointColorLabelB = LOCTEXT("PointColorB", "Color.B");
+	const FText TEXT_PointColorLabelA = LOCTEXT("PointColorA", "Color.A");
 	const FText TEXT_PointDensityLabel = LOCTEXT("PointDensityLabel", "Density");
 	const FText TEXT_PointSteepnessLabel = LOCTEXT("PointSteepnessLabel", "Steepness");
 	const FText TEXT_PointSeedLabel = LOCTEXT("PointSeedLabel", "Seed");
@@ -99,15 +107,17 @@ namespace PCGEditorGraphAttributeListView
 		const FSlateFontInfo FontInfo = FAppStyle::GetFontStyle(TEXT("NormalText"));
 		
 		const float TextWidth = FontMeasure->Measure(InText, FontInfo).X;
-		constexpr float ColumnPadding = 12.0f; // TODO: Grab padding from header style
+		constexpr float ColumnPadding = 22.0f; // TODO: Grab padding from header style
 		const float ColumnWidth = TextWidth + ColumnPadding;
 		return FMath::Min(ColumnWidth, MaxColumnWidth);
 	}
 }
 
-void SPCGListViewItemRow::Construct(const FArguments& InArgs, const TSharedRef<STableViewBase>& InOwnerTableView, const PCGListviewItemPtr& Item)
+void SPCGListViewItemRow::Construct(const FArguments& InArgs, const TSharedRef<STableViewBase>& InOwnerTableView)
 {
-	InternalItem = Item;
+	check(InArgs._ListViewItem && InArgs._AttributeListView);
+	InternalItem = InArgs._ListViewItem;
+	AttributeListView = InArgs._AttributeListView;
 
 	SMultiColumnTableRow<PCGListviewItemPtr>::Construct(
 		SMultiColumnTableRow::FArguments()
@@ -117,303 +127,60 @@ void SPCGListViewItemRow::Construct(const FArguments& InArgs, const TSharedRef<S
 
 TSharedRef<SWidget> SPCGListViewItemRow::GenerateWidgetForColumn(const FName& ColumnId)
 {
-	FText ColumnData = LOCTEXT("ColumnError", "Unrecognized Column");
+	FText RowText = LOCTEXT("ColumnError", "Unrecognized Column");
 
 	if (ColumnId == PCGEditorGraphAttributeListView::NAME_IndexColumn)
 	{
-		ColumnData = FText::FromString(FString::FromInt(InternalItem->Index));
-		return SNew(STextBlock).Text(ColumnData);
+		RowText = FText::FromString(FString::FromInt(InternalItem->Index));
+		return SNew(STextBlock).Text(RowText);
 	}
 
-	if (const FPCGPoint* PCGPoint = InternalItem->PCGPoint)
-	{
-		ColumnData = ConvertPointDataToText(PCGPoint, ColumnId);
-	}
-	else if (const UPCGMetadata* PCGMetadata = InternalItem->PCGMetadata)
-	{
-		if (const FPCGMetadataInfo* MetadataInfo = InternalItem->MetadataInfos->Find(ColumnId))
-		{
-			if (const FPCGMetadataAttributeBase* AttributeBase = PCGMetadata->GetConstAttribute((*MetadataInfo).MetadataId))
-			{
-				ColumnData = ConvertMetadataAttributeToText(AttributeBase, MetadataInfo, InternalItem->MetaDataItemKey);
-			}
-		}
-	}
-
-	return SNew(STextBlock)
-		.Text(ColumnData)
-		.OverflowPolicy(ETextOverflowPolicy::Ellipsis)
-		.Margin(FMargin(2.0f, 0.0f));
-}
-
-FText SPCGListViewItemRow::ConvertPointDataToText(const FPCGPoint* PCGPoint, const FName& ColumnId) const
-{
-	check(PCGPoint);
-
-	const FTransform& Transform = PCGPoint->Transform;
-	if (ColumnId == PCGEditorGraphAttributeListView::NAME_PointPositionX)
-	{
-		const FVector& Position = Transform.GetLocation();
-		return FText::AsNumber(Position.X);
-	}
-	else if (ColumnId == PCGEditorGraphAttributeListView::NAME_PointPositionY)
-	{
-		const FVector& Position = Transform.GetLocation();
-		return FText::AsNumber(Position.Y);
-	}
-	else if (ColumnId == PCGEditorGraphAttributeListView::NAME_PointPositionZ)
-	{
-		const FVector& Position = Transform.GetLocation();
-		return FText::AsNumber(Position.Z);
-	}
-	else if (ColumnId == PCGEditorGraphAttributeListView::NAME_PointRotationX)
-	{
-		const FRotator& Rotation = Transform.Rotator();
-		return FText::AsNumber(Rotation.Roll);
-	}
-	else if (ColumnId == PCGEditorGraphAttributeListView::NAME_PointRotationY)
-	{
-		const FRotator& Rotation = Transform.Rotator();
-		return FText::AsNumber(Rotation.Pitch);
-	}
-	else if (ColumnId == PCGEditorGraphAttributeListView::NAME_PointRotationZ)
-	{
-		const FRotator& Rotation = Transform.Rotator();
-		return FText::AsNumber(Rotation.Yaw);
-	}
-	else if (ColumnId == PCGEditorGraphAttributeListView::NAME_PointScaleX)
-	{
-		const FVector& Scale = Transform.GetScale3D();
-		return FText::AsNumber(Scale.X);
-	}
-	else if (ColumnId == PCGEditorGraphAttributeListView::NAME_PointScaleY)
-	{
-		const FVector& Scale = Transform.GetScale3D();
-		return FText::AsNumber(Scale.Y);
-	}
-	else if (ColumnId == PCGEditorGraphAttributeListView::NAME_PointScaleZ)
-	{
-		const FVector& Scale = Transform.GetScale3D();
-		return FText::AsNumber(Scale.Z);
-	}
-	else if (ColumnId == PCGEditorGraphAttributeListView::NAME_PointBoundsMinX)
-	{
-		return FText::AsNumber(PCGPoint->BoundsMin.X);
-	}
-	else if (ColumnId == PCGEditorGraphAttributeListView::NAME_PointBoundsMinY)
-	{
-		return FText::AsNumber(PCGPoint->BoundsMin.Y);
-	}
-	else if (ColumnId == PCGEditorGraphAttributeListView::NAME_PointBoundsMinZ)
-	{
-		return FText::AsNumber(PCGPoint->BoundsMin.Z);
-	}
-	else if (ColumnId == PCGEditorGraphAttributeListView::NAME_PointBoundsMaxX)
-	{
-		return FText::AsNumber(PCGPoint->BoundsMax.X);
-	}
-	else if (ColumnId == PCGEditorGraphAttributeListView::NAME_PointBoundsMaxY)
-	{
-		return FText::AsNumber(PCGPoint->BoundsMax.Y);
-	}
-	else if (ColumnId == PCGEditorGraphAttributeListView::NAME_PointBoundsMaxZ)
-	{
-		return FText::AsNumber(PCGPoint->BoundsMax.Z);
-	}
-	else if (ColumnId == PCGEditorGraphAttributeListView::NAME_PointColorR)
-	{
-		return FText::AsNumber(PCGPoint->Color.X);
-	}
-	else if (ColumnId == PCGEditorGraphAttributeListView::NAME_PointColorG)
-	{
-		return FText::AsNumber(PCGPoint->Color.Y);
-	}
-	else if (ColumnId == PCGEditorGraphAttributeListView::NAME_PointColorB)
-	{
-		return FText::AsNumber(PCGPoint->Color.Z);
-	}
-	else if (ColumnId == PCGEditorGraphAttributeListView::NAME_PointColorA)
-	{
-		return FText::AsNumber(PCGPoint->Color.W);
-	}
-	else if (ColumnId == PCGEditorGraphAttributeListView::NAME_PointDensity)
-	{
-		const float Density = PCGPoint->Density;
-		return FText::AsNumber(Density);
-	}
-	else if (ColumnId == PCGEditorGraphAttributeListView::NAME_PointSteepness)
-	{
-		const float Steepness = PCGPoint->Steepness;
-		return FText::AsNumber(Steepness);
-	}
-	else if (ColumnId == PCGEditorGraphAttributeListView::NAME_PointSeed)
-	{
-		const int32 Seed = PCGPoint->Seed;
-		return FText::AsNumber(Seed);
-	}
-	else if (ColumnId == PCGEditorGraphAttributeListView::NAME_PointMetadataEntry)
-	{
-		const int64 MetadataEntryKey = PCGPoint->MetadataEntry;
-		return FText::AsNumber(MetadataEntryKey);
-	}
-	else if (ColumnId == PCGEditorGraphAttributeListView::NAME_PointMetadataEntryParent)
-	{
-		if(InternalItem->PCGMetadata)
-		{
-			return FText::AsNumber(InternalItem->PCGMetadata->GetParentKey(PCGPoint->MetadataEntry));
-		}
-		else
-		{
-			return FText::AsNumber(PCGInvalidEntryKey);
-		}
-	}
-	// None of the default point columns were caught, see if its metadata
-	else if (const UPCGMetadata* PCGMetadata = InternalItem->PCGMetadata)
-	{
-		if (const FPCGMetadataInfo* MetadataInfo = InternalItem->MetadataInfos->Find(ColumnId))
-		{
-			if (const FPCGMetadataAttributeBase* AttributeBase = PCGMetadata->GetConstAttribute((*MetadataInfo).MetadataId))
-			{
-				return ConvertMetadataAttributeToText(AttributeBase, MetadataInfo, InternalItem->MetaDataItemKey);
-			}
-		}
-	}
-
-	return PCGEditorGraphAttributeListView::NoDataAvailableText;
-}
+	const TSharedPtr<SPCGEditorGraphAttributeListView> SharedAttributeListView = AttributeListView.Pin();
+	check(SharedAttributeListView.IsValid());
 
-FText SPCGListViewItemRow::ConvertMetadataAttributeToText(const FPCGMetadataAttributeBase* AttributeBase, const FPCGMetadataInfo* MetadataInfo, int64 ItemKey) const
-{
-	switch (AttributeBase->GetTypeId())
-	{
-	case PCG::Private::MetadataTypes<float>::Id:
-		{
-			const float MetaFloat = static_cast<const FPCGMetadataAttribute<float>*>(AttributeBase)->GetValueFromItemKey(ItemKey);
-			return FText::AsNumber(MetaFloat);
-		}
-	case PCG::Private::MetadataTypes<double>::Id:
-		{
-			const double MetaDouble = static_cast<const FPCGMetadataAttribute<double>*>(AttributeBase)->GetValueFromItemKey(ItemKey);
-			return FText::AsNumber(MetaDouble);
-		}
-	case PCG::Private::MetadataTypes<bool>::Id:
-		{
-			const bool bMetaBool = static_cast<const FPCGMetadataAttribute<bool>*>(AttributeBase)->GetValueFromItemKey(ItemKey);
-			return bMetaBool ? FText::FromString(TEXT("true")) : FText::FromString(TEXT("false"));
-		}
-	case PCG::Private::MetadataTypes<FVector2D>::Id:
+	if (FPCGColumnData* PCGColumnData = SharedAttributeListView->PCGColumnData.Find(ColumnId))
 	{
-		const FVector2D MetaVector = static_cast<const FPCGMetadataAttribute<FVector2D>*>(AttributeBase)->GetValueFromItemKey(ItemKey);
-		return FText::AsNumber(MetaVector[(*MetadataInfo).Index]);
-	}
-	case PCG::Private::MetadataTypes<FVector>::Id:
-		{
-			const FVector MetaVector = static_cast<const FPCGMetadataAttribute<FVector>*>(AttributeBase)->GetValueFromItemKey(ItemKey);
-			return FText::AsNumber(MetaVector[(*MetadataInfo).Index]);
-		}
-	case PCG::Private::MetadataTypes<FVector4>::Id:
-		{
-			const FVector4 MetaVector4 = static_cast<const FPCGMetadataAttribute<FVector4>*>(AttributeBase)->GetValueFromItemKey(ItemKey);
-			return FText::AsNumber(MetaVector4[(*MetadataInfo).Index]);
-		}
-	case PCG::Private::MetadataTypes<int32>::Id:
+		if (PCGColumnData->DataAccessor.IsValid() && PCGColumnData->DataKeys.IsValid())
 		{
-			const int32 MetaInt32 = static_cast<const FPCGMetadataAttribute<int32>*>(AttributeBase)->GetValueFromItemKey(ItemKey);
-			return FText::AsNumber(MetaInt32);
-		}
-	case PCG::Private::MetadataTypes<int64>::Id:
-		{
-			const int64 MetaInt64 = static_cast<const FPCGMetadataAttribute<int64>*>(AttributeBase)->GetValueFromItemKey(ItemKey);
-			return FText::AsNumber(MetaInt64);
-		}
-	case PCG::Private::MetadataTypes<FString>::Id:
-		{
-			const FString MetaString = static_cast<const FPCGMetadataAttribute<FString>*>(AttributeBase)->GetValueFromItemKey(ItemKey);
-			return FText::FromString(MetaString);
-		}
-	case PCG::Private::MetadataTypes<FName>::Id:
-		{
-			const FName MetaName = static_cast<const FPCGMetadataAttribute<FName>*>(AttributeBase)->GetValueFromItemKey(ItemKey);
-			return FText::FromName(MetaName);
-		}
-	case PCG::Private::MetadataTypes<FQuat>::Id:
-		{
-			const FQuat MetaQuat = static_cast<const FPCGMetadataAttribute<FQuat>*>(AttributeBase)->GetValueFromItemKey(ItemKey);
-			float QuatValue = 0.0f;
-			if ((*MetadataInfo).Index == 0)
-			{
-				QuatValue = MetaQuat.X;
-			}
-			else if ((*MetadataInfo).Index == 1)
-			{
-				QuatValue = MetaQuat.Y;
-			}
-			else if ((*MetadataInfo).Index == 2)
+			int32 Index = InternalItem->Index;
+			auto Callback = [&PCGColumnData, Index, &RowText] (auto Dummy)
 			{
-				QuatValue = MetaQuat.Z;
-			}
-			else if ((*MetadataInfo).Index == 3)
-			{
-				QuatValue = MetaQuat.W;
-			}
-			return FText::AsNumber(QuatValue);
-		}
-	case PCG::Private::MetadataTypes<FRotator>::Id:
-		{
-			const FRotator MetaRotator = static_cast<const FPCGMetadataAttribute<FRotator>*>(AttributeBase)->GetValueFromItemKey(ItemKey);
-			float RotatorValue = 0.0f;
-			if ((*MetadataInfo).Index == 0)
-			{
-				RotatorValue = MetaRotator.Roll;
-			}
-			else if ((*MetadataInfo).Index == 1)
-			{
-				RotatorValue = MetaRotator.Pitch;
-			}
-			else if ((*MetadataInfo).Index == 2)
-			{
-				RotatorValue = MetaRotator.Yaw;
-			}
-			return FText::AsNumber(RotatorValue);
-		}
-	case PCG::Private::MetadataTypes<FTransform>::Id:
-		{
-			const FTransform& MetaTransform = static_cast<const FPCGMetadataAttribute<FTransform>*>(AttributeBase)->GetValueFromItemKey(ItemKey);
-			const int8 ComponentIndex = (*MetadataInfo).Index / 3;
-			const int8 ValueIndex = (*MetadataInfo).Index % 3;
-
-			if (ComponentIndex == 0)
-			{
-				const FVector& MetaVector = MetaTransform.GetLocation();
-				return FText::AsNumber(MetaVector[ValueIndex]);
-			}
-			else if (ComponentIndex == 1)
-			{
-				const FRotator& MetaRotator = MetaTransform.Rotator();
-				float RotatorValue = 0.0f;
-				if (ValueIndex == 0)
+				using ValueType = decltype(Dummy);
+				ValueType Value{};
+				if (PCGColumnData->DataAccessor->Get<ValueType>(Value, Index, *PCGColumnData->DataKeys))
 				{
-					RotatorValue = MetaRotator.Roll;
+					if constexpr (PCG::Private::IsOfTypes<ValueType, bool>())
+					{
+						RowText = FText::FromString(LexToString(Value));
+					}
+					else if constexpr (PCG::Private::IsOfTypes<ValueType, FString>())
+					{
+						RowText = FText::FromString(Value);
+					}
+					else if constexpr (PCG::Private::IsOfTypes<ValueType, FName>())
+					{
+						RowText = FText::FromName(Value);
+					}
+					else if constexpr (FTextAsNumberIsValid<ValueType>::value)
+					{
+						RowText = FText::AsNumber(Value);
+					}
+					else
+					{
+						ensureMsgf(false, TEXT("Unsupported Data Type"));
+						RowText = LOCTEXT("UnsupportedDataTypeError", "Unsupported Data Type");
+					}
 				}
-				else if (ValueIndex == 1)
-				{
-					RotatorValue = MetaRotator.Pitch;
-				}
-				else if (ValueIndex == 2)
-				{
-					RotatorValue = MetaRotator.Yaw;
-				}
-				return FText::AsNumber(RotatorValue);
-			}
-			else if (ComponentIndex == 2)
-			{
-				const FVector& MetaVector = MetaTransform.GetScale3D();
-				return FText::AsNumber(MetaVector[ValueIndex]);
-			}
+			};
+			
+			PCGMetadataAttribute::CallbackWithRightType(PCGColumnData->DataAccessor->GetUnderlyingType(), Callback);
 		}
-	default:
-		return PCGEditorGraphAttributeListView::NoDataAvailableText;
 	}
+	
+	return SNew(STextBlock)
+		.Text(RowText)
+		.OverflowPolicy(ETextOverflowPolicy::Ellipsis)
+		.Margin(FMargin(2.0f, 0.0f));
 }
 
 SPCGEditorGraphAttributeListView::~SPCGEditorGraphAttributeListView()
@@ -492,11 +259,13 @@ void SPCGEditorGraphAttributeListView::Construct(const FArguments& InArgs, TShar
 			SNew(SHorizontalBox)
 			+SHorizontalBox::Slot()
 			.AutoWidth()
+			.Padding(1.0f, 0.0f)
 			[
 				FilterButton->AsShared()
 			]
 			+SHorizontalBox::Slot()
 			.AutoWidth()
+			.Padding(1.0f, 0.0f)
 			[
 				DataComboBox->AsShared()
 			]
@@ -622,7 +391,6 @@ void SPCGEditorGraphAttributeListView::RefreshAttributeList()
 	
 	ListViewItems.Empty();
 	ListViewHeader->ClearColumns();
-	MetadataInfos.Empty();
 	InfoTextBlock->SetText(FText::GetEmpty());
 	
 	if (!PCGComponent.IsValid())
@@ -661,7 +429,7 @@ void SPCGEditorGraphAttributeListView::RefreshAttributeList()
 		if (const UPCGMetadata* PCGMetadata = PCGParamData->ConstMetadata())
 		{
 			AddIndexColumn();
-			GenerateColumnsFromMetadata(PCGMetadata);
+			GenerateColumnsFromMetadata(PCGData, PCGMetadata);
 
 			PCGMetadataEntryKey ItemKeyLowerBound = PCGMetadata->GetItemKeyCountForParent();
 			PCGMetadataEntryKey ItemKeyUpperBound = PCGMetadata->GetItemCountForChild();
@@ -669,10 +437,6 @@ void SPCGEditorGraphAttributeListView::RefreshAttributeList()
 			{
 				PCGListviewItemPtr ListViewItem = MakeShared<FPCGListViewItem>();
 				ListViewItem->Index = MetadataItemKey - ItemKeyLowerBound;
-				ListViewItem->PCGParamData = PCGParamData;
-				ListViewItem->PCGMetadata = PCGMetadata;
-				ListViewItem->MetaDataItemKey = MetadataItemKey;
-				ListViewItem->MetadataInfos = &MetadataInfos;
 				ListViewItems.Add(ListViewItem);
 			}
 
@@ -693,8 +457,8 @@ void SPCGEditorGraphAttributeListView::RefreshAttributeList()
 		{
 			const UPCGMetadata* PCGMetadata = PCGPointData->ConstMetadata();
 
-			AddPointDataColumns();
-			GenerateColumnsFromMetadata(PCGMetadata);
+			AddPointDataColumns(PCGPointData);
+			GenerateColumnsFromMetadata(PCGData, PCGMetadata);
 
 			const TArray<FPCGPoint>& PCGPoints = PCGPointData->GetPoints();
 			const int32 NumPoints = PCGPoints.Num();
@@ -706,9 +470,6 @@ void SPCGEditorGraphAttributeListView::RefreshAttributeList()
 				PCGListviewItemPtr ListViewItem = MakeShared<FPCGListViewItem>();
 				ListViewItem->Index = PointIndex;
 				ListViewItem->PCGPoint = &PCGPoint;
-				ListViewItem->PCGMetadata = PCGMetadata;
-				ListViewItem->MetaDataItemKey = PCGPoint.MetadataEntry;
-				ListViewItem->MetadataInfos = &MetadataInfos;
 				ListViewItems.Add(ListViewItem);
 			}
 			
@@ -864,7 +625,7 @@ int32 SPCGEditorGraphAttributeListView::GetSelectedDataIndex() const
 	return Index;
 }
 
-void SPCGEditorGraphAttributeListView::GenerateColumnsFromMetadata(const UPCGMetadata* PCGMetadata)
+void SPCGEditorGraphAttributeListView::GenerateColumnsFromMetadata(const UPCGData* InPCGData, const UPCGMetadata* PCGMetadata)
 {
 	TArray<FName> AttributeNames;
 	TArray<EPCGMetadataTypes> AttributeTypes;
@@ -875,13 +636,7 @@ void SPCGEditorGraphAttributeListView::GenerateColumnsFromMetadata(const UPCGMet
 		const FName& AttributeName = AttributeNames[I];
 		const EPCGMetadataTypes AttributeType = AttributeTypes[I];
 		FName ColumnName = AttributeName;
-
-		if (ColumnName == NAME_None)
-		{
-			const FString TypeString = UEnum::GetDisplayValueAsText(AttributeType).ToString();
-			ColumnName = *TypeString;
-		}
-
+		
 		switch (AttributeType)
 		{
 		case EPCGMetadataTypes::Float:
@@ -892,42 +647,49 @@ void SPCGEditorGraphAttributeListView::GenerateColumnsFromMetadata(const UPCGMet
 		case EPCGMetadataTypes::String:
 		case EPCGMetadataTypes::Name:
 			{
-				AddMetadataColumn(ColumnName, AttributeName, AttributeType);
+				AddMetadataColumn(InPCGData, ColumnName, AttributeType);
 				break;
 			}
 		case EPCGMetadataTypes::Vector2:
 		{
-			AddMetadataColumn(ColumnName, AttributeName, AttributeType, 0, TEXT("_X"));
-			AddMetadataColumn(ColumnName, AttributeName, AttributeType, 1, TEXT("_Y"));
+			AddMetadataColumn(InPCGData, ColumnName, AttributeType, TEXT(".X"));
+			AddMetadataColumn(InPCGData, ColumnName, AttributeType, TEXT(".Y"));
 			break;
 		}
 		case EPCGMetadataTypes::Vector:
 			{
-				AddMetadataColumn(ColumnName, AttributeName, AttributeType, 0, TEXT("_X"));
-				AddMetadataColumn(ColumnName, AttributeName, AttributeType, 1, TEXT("_Y"));
-				AddMetadataColumn(ColumnName, AttributeName, AttributeType, 2, TEXT("_Z"));
+				AddMetadataColumn(InPCGData, ColumnName, AttributeType, TEXT(".X"));
+				AddMetadataColumn(InPCGData, ColumnName, AttributeType, TEXT(".Y"));
+				AddMetadataColumn(InPCGData, ColumnName, AttributeType, TEXT(".Z"));
 				break;
 			}
 		case EPCGMetadataTypes::Vector4:
 		case EPCGMetadataTypes::Quaternion:
 			{
-				AddMetadataColumn(ColumnName, AttributeName, AttributeType, 0, TEXT("_X"));
-				AddMetadataColumn(ColumnName, AttributeName, AttributeType, 1, TEXT("_Y"));
-				AddMetadataColumn(ColumnName, AttributeName, AttributeType, 2, TEXT("_Z"));
-				AddMetadataColumn(ColumnName, AttributeName, AttributeType, 3, TEXT("_W"));
+				AddMetadataColumn(InPCGData, ColumnName, AttributeType, TEXT(".X"));
+				AddMetadataColumn(InPCGData, ColumnName, AttributeType, TEXT(".Y"));
+				AddMetadataColumn(InPCGData, ColumnName, AttributeType, TEXT(".Z"));
+				AddMetadataColumn(InPCGData, ColumnName, AttributeType, TEXT(".W"));
+				break;
+			}
+		case EPCGMetadataTypes::Rotator:
+			{
+				AddMetadataColumn(InPCGData, ColumnName, AttributeType, TEXT(".Roll"));
+				AddMetadataColumn(InPCGData, ColumnName, AttributeType, TEXT(".Pitch"));
+				AddMetadataColumn(InPCGData, ColumnName, AttributeType, TEXT(".Yaw"));
 				break;
 			}
 		case EPCGMetadataTypes::Transform:
 			{
-				AddMetadataColumn(ColumnName, AttributeName, AttributeType, 0, TEXT("_tX"));
-				AddMetadataColumn(ColumnName, AttributeName, AttributeType, 1, TEXT("_tY"));
-				AddMetadataColumn(ColumnName, AttributeName, AttributeType, 2, TEXT("_tZ"));
-				AddMetadataColumn(ColumnName, AttributeName, AttributeType, 3, TEXT("_rX"));
-				AddMetadataColumn(ColumnName, AttributeName, AttributeType, 4, TEXT("_rY"));
-				AddMetadataColumn(ColumnName, AttributeName, AttributeType, 5, TEXT("_rZ"));
-				AddMetadataColumn(ColumnName, AttributeName, AttributeType, 6, TEXT("_sX"));
-				AddMetadataColumn(ColumnName, AttributeName, AttributeType, 7, TEXT("_sY"));
-				AddMetadataColumn(ColumnName, AttributeName, AttributeType, 8, TEXT("_sZ"));
+				AddMetadataColumn(InPCGData, ColumnName, AttributeType, TEXT(".Position.X"));
+				AddMetadataColumn(InPCGData, ColumnName, AttributeType, TEXT(".Position.Y"));
+				AddMetadataColumn(InPCGData, ColumnName, AttributeType, TEXT(".Position.Z"));
+				AddMetadataColumn(InPCGData, ColumnName, AttributeType, TEXT(".Rotation.X"));
+				AddMetadataColumn(InPCGData, ColumnName, AttributeType, TEXT(".Rotation.Y"));
+				AddMetadataColumn(InPCGData, ColumnName, AttributeType, TEXT(".Rotation.Z"));
+				AddMetadataColumn(InPCGData, ColumnName, AttributeType, TEXT(".Scale.X"));
+				AddMetadataColumn(InPCGData, ColumnName, AttributeType, TEXT(".Scale.Y"));
+				AddMetadataColumn(InPCGData, ColumnName, AttributeType, TEXT(".Scale.Z"));
 				break;
 			}
 		default:
@@ -999,9 +761,11 @@ bool SPCGEditorGraphAttributeListView::IsAttributeEnabled(FName InAttributeName)
 	return ListViewHeader->IsColumnVisible(InAttributeName);
 }
 
-TSharedRef<ITableRow> SPCGEditorGraphAttributeListView::OnGenerateRow(PCGListviewItemPtr Item, const TSharedRef<STableViewBase>& OwnerTable) const
+TSharedRef<ITableRow> SPCGEditorGraphAttributeListView::OnGenerateRow(PCGListviewItemPtr Item, const TSharedRef<STableViewBase>& OwnerTable)
 {
-	return SNew(SPCGListViewItemRow, OwnerTable, Item);
+	return SNew(SPCGListViewItemRow, OwnerTable)
+		.AttributeListView(SharedThis(this))
+		.ListViewItem(Item);
 }
 
 void SPCGEditorGraphAttributeListView::OnItemDoubleClicked(PCGListviewItemPtr Item) const
@@ -1015,10 +779,44 @@ void SPCGEditorGraphAttributeListView::OnItemDoubleClicked(PCGListviewItemPtr It
 	}
 }
 
-void SPCGEditorGraphAttributeListView::AddColumn(const FName& InColumnId, const FText& ColumnLabel, EHorizontalAlignment HeaderHAlign, EHorizontalAlignment CellHAlign)
+void SPCGEditorGraphAttributeListView::AddColumn(const UPCGPointData* InPCGPointData, const FName& InColumnId, const FText& ColumnLabel, EHorizontalAlignment HeaderHAlign, EHorizontalAlignment CellHAlign)
 {
 	const float ColumnWidth = PCGEditorGraphAttributeListView::CalculateColumnWidth(ColumnLabel);
 
+	if (InPCGPointData)
+	{
+		const FString ColumnIdString = InColumnId.ToString();
+		
+		FPCGAttributePropertySelector TargetSelector;
+		TargetSelector.Update(ColumnIdString);
+	
+		FPCGColumnData& ColumnData = PCGColumnData.Add(InColumnId);
+		
+		if (InColumnId == PCGEditorGraphAttributeListView::NAME_PointMetadataEntry)
+		{
+			ColumnData.DataAccessor = PCGAttributeAccessorHelpers::CreatePropertyAccessor(GET_MEMBER_NAME_CHECKED(FPCGPoint, MetadataEntry), FPCGPoint::StaticStruct());
+			ColumnData.DataKeys = MakeUnique<FPCGAttributeAccessorKeysPoints>(InPCGPointData->GetPoints());
+		}
+		else if (InColumnId == PCGEditorGraphAttributeListView::NAME_PointMetadataEntryParent)
+		{
+			ColumnData.DataAccessor = MakeUnique<FPCGCustomPointAccessor<int64>>([InPCGPointData](const FPCGPoint& Point, void* OutValue)
+			{
+				if (const UPCGMetadata* Metadata = InPCGPointData->Metadata)
+				{
+					*reinterpret_cast<int64*>(OutValue) = Metadata->GetParentKey(Point.MetadataEntry);
+					return true;
+				}
+				return false;
+			}, nullptr);
+			ColumnData.DataKeys = MakeUnique<FPCGAttributeAccessorKeysPoints>(InPCGPointData->GetPoints());
+		}
+		else
+		{	
+			ColumnData.DataAccessor = PCGAttributeAccessorHelpers::CreateConstAccessor(InPCGPointData, TargetSelector);
+			ColumnData.DataKeys = PCGAttributeAccessorHelpers::CreateConstKeys(InPCGPointData, TargetSelector);
+		}
+	}
+	
 	SHeaderRow::FColumn::FArguments Arguments;
 	Arguments.ColumnId(InColumnId);
 	Arguments.DefaultLabel(ColumnLabel);
@@ -1031,101 +829,77 @@ void SPCGEditorGraphAttributeListView::AddColumn(const FName& InColumnId, const
 	ListViewHeader->AddColumn(*NewColumn);
 }
 
-void SPCGEditorGraphAttributeListView::RemoveColumn(const FName& InColumnId)
-{
-	ListViewHeader->RemoveColumn(InColumnId);
-}
-
 void SPCGEditorGraphAttributeListView::AddIndexColumn()
 {
-	AddColumn(PCGEditorGraphAttributeListView::NAME_IndexColumn, PCGEditorGraphAttributeListView::TEXT_IndexLabel);
+	AddColumn(nullptr, PCGEditorGraphAttributeListView::NAME_IndexColumn, PCGEditorGraphAttributeListView::TEXT_IndexLabel);
 }
 
-void SPCGEditorGraphAttributeListView::RemoveIndexColumn()
-{
-	RemoveColumn(PCGEditorGraphAttributeListView::NAME_IndexColumn);
-}
-
-void SPCGEditorGraphAttributeListView::AddPointDataColumns()
+void SPCGEditorGraphAttributeListView::AddPointDataColumns(const UPCGPointData* InPCGPointData)
 {
 	AddIndexColumn();
-	AddColumn(PCGEditorGraphAttributeListView::NAME_PointPositionX, PCGEditorGraphAttributeListView::TEXT_PointPositionLabelX);
-	AddColumn(PCGEditorGraphAttributeListView::NAME_PointPositionY, PCGEditorGraphAttributeListView::TEXT_PointPositionLabelY);
-	AddColumn(PCGEditorGraphAttributeListView::NAME_PointPositionZ, PCGEditorGraphAttributeListView::TEXT_PointPositionLabelZ);
-	AddColumn(PCGEditorGraphAttributeListView::NAME_PointRotationX, PCGEditorGraphAttributeListView::TEXT_PointRotationLabelX);
-	AddColumn(PCGEditorGraphAttributeListView::NAME_PointRotationY, PCGEditorGraphAttributeListView::TEXT_PointRotationLabelY);
-	AddColumn(PCGEditorGraphAttributeListView::NAME_PointRotationZ, PCGEditorGraphAttributeListView::TEXT_PointRotationLabelZ);
-	AddColumn(PCGEditorGraphAttributeListView::NAME_PointScaleX, PCGEditorGraphAttributeListView::TEXT_PointScaleLabelX);
-	AddColumn(PCGEditorGraphAttributeListView::NAME_PointScaleY, PCGEditorGraphAttributeListView::TEXT_PointScaleLabelY);
-	AddColumn(PCGEditorGraphAttributeListView::NAME_PointScaleZ, PCGEditorGraphAttributeListView::TEXT_PointScaleLabelZ);
-	AddColumn(PCGEditorGraphAttributeListView::NAME_PointBoundsMinX, PCGEditorGraphAttributeListView::TEXT_PointBoundsLabelMinX);
-	AddColumn(PCGEditorGraphAttributeListView::NAME_PointBoundsMinY, PCGEditorGraphAttributeListView::TEXT_PointBoundsLabelMinY);
-	AddColumn(PCGEditorGraphAttributeListView::NAME_PointBoundsMinZ, PCGEditorGraphAttributeListView::TEXT_PointBoundsLabelMinZ);
-	AddColumn(PCGEditorGraphAttributeListView::NAME_PointBoundsMaxX, PCGEditorGraphAttributeListView::TEXT_PointBoundsLabelMaxX);
-	AddColumn(PCGEditorGraphAttributeListView::NAME_PointBoundsMaxY, PCGEditorGraphAttributeListView::TEXT_PointBoundsLabelMaxY);
-	AddColumn(PCGEditorGraphAttributeListView::NAME_PointBoundsMaxZ, PCGEditorGraphAttributeListView::TEXT_PointBoundsLabelMaxZ);
-	AddColumn(PCGEditorGraphAttributeListView::NAME_PointColorR, PCGEditorGraphAttributeListView::TEXT_PointColorLabelR);
-	AddColumn(PCGEditorGraphAttributeListView::NAME_PointColorG, PCGEditorGraphAttributeListView::TEXT_PointColorLabelG);
-	AddColumn(PCGEditorGraphAttributeListView::NAME_PointColorB, PCGEditorGraphAttributeListView::TEXT_PointColorLabelB);
-	AddColumn(PCGEditorGraphAttributeListView::NAME_PointColorA, PCGEditorGraphAttributeListView::TEXT_PointColorLabelA);
-	AddColumn(PCGEditorGraphAttributeListView::NAME_PointDensity, PCGEditorGraphAttributeListView::TEXT_PointDensityLabel);
-	AddColumn(PCGEditorGraphAttributeListView::NAME_PointSteepness, PCGEditorGraphAttributeListView::TEXT_PointSteepnessLabel);
-	AddColumn(PCGEditorGraphAttributeListView::NAME_PointSeed, PCGEditorGraphAttributeListView::TEXT_PointSeedLabel);
+	AddColumn(InPCGPointData, PCGEditorGraphAttributeListView::NAME_PointPositionX, PCGEditorGraphAttributeListView::TEXT_PointPositionLabelX);
+	AddColumn(InPCGPointData, PCGEditorGraphAttributeListView::NAME_PointPositionY, PCGEditorGraphAttributeListView::TEXT_PointPositionLabelY);
+	AddColumn(InPCGPointData, PCGEditorGraphAttributeListView::NAME_PointPositionZ, PCGEditorGraphAttributeListView::TEXT_PointPositionLabelZ);
+	AddColumn(InPCGPointData, PCGEditorGraphAttributeListView::NAME_PointRotationX, PCGEditorGraphAttributeListView::TEXT_PointRotationLabelX);
+	AddColumn(InPCGPointData, PCGEditorGraphAttributeListView::NAME_PointRotationY, PCGEditorGraphAttributeListView::TEXT_PointRotationLabelY);
+	AddColumn(InPCGPointData, PCGEditorGraphAttributeListView::NAME_PointRotationZ, PCGEditorGraphAttributeListView::TEXT_PointRotationLabelZ);
+	AddColumn(InPCGPointData, PCGEditorGraphAttributeListView::NAME_PointScaleX, PCGEditorGraphAttributeListView::TEXT_PointScaleLabelX);
+	AddColumn(InPCGPointData, PCGEditorGraphAttributeListView::NAME_PointScaleY, PCGEditorGraphAttributeListView::TEXT_PointScaleLabelY);
+	AddColumn(InPCGPointData, PCGEditorGraphAttributeListView::NAME_PointScaleZ, PCGEditorGraphAttributeListView::TEXT_PointScaleLabelZ);
+	AddColumn(InPCGPointData, PCGEditorGraphAttributeListView::NAME_PointBoundsMinX, PCGEditorGraphAttributeListView::TEXT_PointBoundsLabelMinX);
+	AddColumn(InPCGPointData, PCGEditorGraphAttributeListView::NAME_PointBoundsMinY, PCGEditorGraphAttributeListView::TEXT_PointBoundsLabelMinY);
+	AddColumn(InPCGPointData, PCGEditorGraphAttributeListView::NAME_PointBoundsMinZ, PCGEditorGraphAttributeListView::TEXT_PointBoundsLabelMinZ);
+	AddColumn(InPCGPointData, PCGEditorGraphAttributeListView::NAME_PointBoundsMaxX, PCGEditorGraphAttributeListView::TEXT_PointBoundsLabelMaxX);
+	AddColumn(InPCGPointData, PCGEditorGraphAttributeListView::NAME_PointBoundsMaxY, PCGEditorGraphAttributeListView::TEXT_PointBoundsLabelMaxY);
+	AddColumn(InPCGPointData, PCGEditorGraphAttributeListView::NAME_PointBoundsMaxZ, PCGEditorGraphAttributeListView::TEXT_PointBoundsLabelMaxZ);
+	AddColumn(InPCGPointData, PCGEditorGraphAttributeListView::NAME_PointColorR, PCGEditorGraphAttributeListView::TEXT_PointColorLabelR);
+	AddColumn(InPCGPointData, PCGEditorGraphAttributeListView::NAME_PointColorG, PCGEditorGraphAttributeListView::TEXT_PointColorLabelG);
+	AddColumn(InPCGPointData, PCGEditorGraphAttributeListView::NAME_PointColorB, PCGEditorGraphAttributeListView::TEXT_PointColorLabelB);
+	AddColumn(InPCGPointData, PCGEditorGraphAttributeListView::NAME_PointColorA, PCGEditorGraphAttributeListView::TEXT_PointColorLabelA);
+	AddColumn(InPCGPointData, PCGEditorGraphAttributeListView::NAME_PointDensity, PCGEditorGraphAttributeListView::TEXT_PointDensityLabel);
+	AddColumn(InPCGPointData, PCGEditorGraphAttributeListView::NAME_PointSteepness, PCGEditorGraphAttributeListView::TEXT_PointSteepnessLabel);
+	AddColumn(InPCGPointData, PCGEditorGraphAttributeListView::NAME_PointSeed, PCGEditorGraphAttributeListView::TEXT_PointSeedLabel);
 
 	if (CVarShowAdvancedAttributesFields.GetValueOnAnyThread())
 	{
-		AddColumn(PCGEditorGraphAttributeListView::NAME_PointMetadataEntry, PCGEditorGraphAttributeListView::TEXT_PointMetadataEntryLabel);
-		AddColumn(PCGEditorGraphAttributeListView::NAME_PointMetadataEntryParent, PCGEditorGraphAttributeListView::TEXT_PointMetadataEntryParentLabel);
+		AddColumn(InPCGPointData, PCGEditorGraphAttributeListView::NAME_PointMetadataEntry, PCGEditorGraphAttributeListView::TEXT_PointMetadataEntryLabel);
+		AddColumn(InPCGPointData, PCGEditorGraphAttributeListView::NAME_PointMetadataEntryParent, PCGEditorGraphAttributeListView::TEXT_PointMetadataEntryParentLabel);
 	}	
 }
 
-void SPCGEditorGraphAttributeListView::RemovePointDataColumns()
-{
-	RemoveIndexColumn();
-	RemoveColumn(PCGEditorGraphAttributeListView::NAME_PointPositionX);
-	RemoveColumn(PCGEditorGraphAttributeListView::NAME_PointPositionY);
-	RemoveColumn(PCGEditorGraphAttributeListView::NAME_PointPositionZ);
-	RemoveColumn(PCGEditorGraphAttributeListView::NAME_PointRotationX);
-	RemoveColumn(PCGEditorGraphAttributeListView::NAME_PointRotationY);
-	RemoveColumn(PCGEditorGraphAttributeListView::NAME_PointRotationZ);
-	RemoveColumn(PCGEditorGraphAttributeListView::NAME_PointScaleX);
-	RemoveColumn(PCGEditorGraphAttributeListView::NAME_PointScaleY);
-	RemoveColumn(PCGEditorGraphAttributeListView::NAME_PointScaleZ);
-	RemoveColumn(PCGEditorGraphAttributeListView::NAME_PointBoundsMinX);
-	RemoveColumn(PCGEditorGraphAttributeListView::NAME_PointBoundsMinY);
-	RemoveColumn(PCGEditorGraphAttributeListView::NAME_PointBoundsMinZ);
-	RemoveColumn(PCGEditorGraphAttributeListView::NAME_PointBoundsMaxX);
-	RemoveColumn(PCGEditorGraphAttributeListView::NAME_PointBoundsMaxY);
-	RemoveColumn(PCGEditorGraphAttributeListView::NAME_PointBoundsMaxZ);
-	RemoveColumn(PCGEditorGraphAttributeListView::NAME_PointColorR);
-	RemoveColumn(PCGEditorGraphAttributeListView::NAME_PointColorG);
-	RemoveColumn(PCGEditorGraphAttributeListView::NAME_PointColorB);
-	RemoveColumn(PCGEditorGraphAttributeListView::NAME_PointColorA);
-	RemoveColumn(PCGEditorGraphAttributeListView::NAME_PointDensity);
-	RemoveColumn(PCGEditorGraphAttributeListView::NAME_PointSteepness);
-	RemoveColumn(PCGEditorGraphAttributeListView::NAME_PointSeed);
-
-	RemoveColumn(PCGEditorGraphAttributeListView::NAME_PointMetadataEntry);
-	RemoveColumn(PCGEditorGraphAttributeListView::NAME_PointMetadataEntryParent);
-}
-
-void SPCGEditorGraphAttributeListView::AddMetadataColumn(const FName& InColumnId, const FName& InMetadataId, EPCGMetadataTypes InMetadataType, const int8 InValueIndex, const TCHAR* PostFix)
-{
-	FString ColumnIdString = InColumnId.ToString();
-
+void SPCGEditorGraphAttributeListView::AddMetadataColumn(const UPCGData* InPCGData, const FName& InColumnId, EPCGMetadataTypes InMetadataType, const TCHAR* PostFix)
+{	
+	FString OriginalColumnIdString = InColumnId.ToString();
 	if (PostFix)
 	{
-		ColumnIdString.Append(PostFix);
+		OriginalColumnIdString.Append(PostFix);
 	}
 
+	FText ColumnLabel;
+	FString ColumnIdString(OriginalColumnIdString);
+	if (InColumnId == NAME_None)
+	{
+		ColumnIdString.ReplaceInline(LexToString(NAME_None), *PCGEditorGraphAttributeListView::LastMetadataId);
+		ColumnLabel = FText::Format(LOCTEXT("LastMetadataLabelFormat", "{0}{1}"), UEnum::GetDisplayValueAsText(InMetadataType), FText::FromString(PostFix));
+	}
+	
 	const FName ColumnId(ColumnIdString);
 
-	FPCGMetadataInfo MetadataInfo;
-	MetadataInfo.MetadataId = InMetadataId;
-	MetadataInfo.Index = InValueIndex;
-	MetadataInfos.Add(ColumnId, MetadataInfo);
+	if (InPCGData)
+	{
+		FPCGAttributePropertySelector TargetSelector;
+		TargetSelector.Update(OriginalColumnIdString);
+	
+		FPCGColumnData& ColumnData = PCGColumnData.Add(ColumnId);
+		ColumnData.DataAccessor = PCGAttributeAccessorHelpers::CreateConstAccessor(InPCGData, TargetSelector);
+		ColumnData.DataKeys = PCGAttributeAccessorHelpers::CreateConstKeys(InPCGData, TargetSelector);
+	}
 
-	const FText ColumnLabel = FText::FromName(ColumnId);
+	if (ColumnLabel.IsEmpty())
+	{
+		ColumnLabel = FText::FromName(ColumnId);	
+	}
+	
 	float ColumnWidth = 0.0f;
 
 	EHorizontalAlignment CellAlignment = EHorizontalAlignment::HAlign_Right;
@@ -1152,12 +926,4 @@ void SPCGEditorGraphAttributeListView::AddMetadataColumn(const FName& InColumnId
 	ListViewHeader->AddColumn(*NewColumn);
 }
 
-void SPCGEditorGraphAttributeListView::RemoveMetadataColumns()
-{
-	for (const TTuple<FName, FPCGMetadataInfo>& MetadataInfo : MetadataInfos)
-	{
-		ListViewHeader->RemoveColumn(MetadataInfo.Key);
-	}
-}
-
 #undef LOCTEXT_NAMESPACE
diff --git a/Engine/Plugins/Experimental/PCG/Source/PCGEditor/Private/SPCGEditorGraphAttributeListView.h b/Engine/Plugins/Experimental/PCG/Source/PCGEditor/Private/SPCGEditorGraphAttributeListView.h
--- a/Engine/Plugins/Experimental/PCG/Source/PCGEditor/Private/SPCGEditorGraphAttributeListView.h
+++ b/Engine/Plugins/Experimental/PCG/Source/PCGEditor/Private/SPCGEditorGraphAttributeListView.h
@@ -2,7 +2,10 @@
 
 #pragma once
 
-#include "Widgets/Views/ITableRow.h"
+#include "Metadata/Accessors/IPCGAttributeAccessor.h"
+#include "Metadata/Accessors/PCGAttributeAccessorKeys.h"
+
+#include "Widgets/Views/SListView.h"
 #include "Widgets/Views/STableRow.h"
 
 class STableViewBase;
@@ -11,7 +14,10 @@ template <typename OptionType> class SComboBox;
 
 class FPCGEditor;
 class FPCGMetadataAttributeBase;
+class SPCGEditorGraphAttributeListView;
 class UPCGComponent;
+class UPCGPointData;
+class UPCGData;
 class UPCGEditorGraphNodeBase;
 class UPCGMetadata;
 class UPCGParamData;
@@ -22,45 +28,50 @@ class SComboButton;
 class SHeaderRow;
 struct FSlateBrush;
 
-struct FPCGMetadataInfo
-{
-	FName MetadataId = NAME_None;
-	int8 Index = 0;
-};
-
 struct FPCGListViewItem
 {
 	int32 Index = INDEX_NONE;
 	const FPCGPoint* PCGPoint = nullptr;
-	const UPCGMetadata* PCGMetadata = nullptr;
-	const UPCGParamData* PCGParamData = nullptr;
-	int64 MetaDataItemKey = -1;
-	const TMap<FName, FPCGMetadataInfo>* MetadataInfos = nullptr;
+};
+
+struct FPCGColumnData
+{
+	TUniquePtr<const IPCGAttributeAccessor> DataAccessor;
+	TUniquePtr<const IPCGAttributeAccessorKeys> DataKeys;
 };
 
 typedef TSharedPtr<FPCGListViewItem> PCGListviewItemPtr;
 
+template <typename T, typename = void>
+struct FTextAsNumberIsValid : std::false_type {};
+
+/** Utility to see if a value type is supported by FText::AsNumber */
+template <typename T>
+struct FTextAsNumberIsValid<T, std::void_t<decltype(FText::AsNumber(std::declval<T>()))>> : std::true_type {};
+
 class SPCGListViewItemRow : public SMultiColumnTableRow<PCGListviewItemPtr>
 {
 public:
 	SLATE_BEGIN_ARGS(SPCGListViewItemRow) {}
+	SLATE_ARGUMENT(TSharedPtr<SPCGEditorGraphAttributeListView>, AttributeListView)
+	SLATE_ARGUMENT(PCGListviewItemPtr, ListViewItem)
 	SLATE_END_ARGS()
 
-	void Construct(const FArguments& InArgs, const TSharedRef<STableViewBase>& InOwnerTableView, const PCGListviewItemPtr& Item);
+	void Construct(const FArguments& InArgs, const TSharedRef<STableViewBase>& InOwnerTableView);
 
 	virtual TSharedRef<SWidget> GenerateWidgetForColumn(const FName& ColumnId) override;
+	
 private:
-	FText ConvertPointDataToText(const FPCGPoint* PCGPoint, const FName& ColumnId) const;
-	FText ConvertMetadataAttributeToText(const FPCGMetadataAttributeBase* AttributeBase, const FPCGMetadataInfo* MetadataInfo, int64 ItemKey = 0) const;
-
-protected:
+	TWeakPtr<SPCGEditorGraphAttributeListView> AttributeListView;
 	PCGListviewItemPtr InternalItem;
 };
 
 class SPCGEditorGraphAttributeListView : public SCompoundWidget
 {
+	friend SPCGListViewItemRow;
+	
 public:
-	SLATE_BEGIN_ARGS(SPCGEditorGraphAttributeListView) { }
+	SLATE_BEGIN_ARGS(SPCGEditorGraphAttributeListView) {}
 	SLATE_END_ARGS()
 
 	virtual ~SPCGEditorGraphAttributeListView() override;
@@ -84,27 +95,20 @@ private:
 	void OnSelectionChanged(TSharedPtr<FName> Item, ESelectInfo::Type SelectInfo);
 	FText OnGenerateSelectedDataText() const;
 	int32 GetSelectedDataIndex() const;
-	void GenerateColumnsFromMetadata(const UPCGMetadata* PCGMetadata);
+	void GenerateColumnsFromMetadata(const UPCGData* InPCGData, const UPCGMetadata* PCGMetadata);
 
 	void ToggleAllAttributes();
 	void ToggleAttribute(FName InAttributeName);
 	ECheckBoxState GetAnyAttributeEnabledState() const;
 	bool IsAttributeEnabled(FName InAttributeName) const;
 
-	TSharedRef<ITableRow> OnGenerateRow(PCGListviewItemPtr Item, const TSharedRef<STableViewBase>& OwnerTable) const;
+	TSharedRef<ITableRow> OnGenerateRow(PCGListviewItemPtr Item, const TSharedRef<STableViewBase>& OwnerTable);
 	void OnItemDoubleClicked(PCGListviewItemPtr Item) const;
 
-	void AddColumn(const FName& InColumnId, const FText& ColumnLabel, EHorizontalAlignment HeaderHAlign = HAlign_Center, EHorizontalAlignment CellHAlign = HAlign_Right);
-	void RemoveColumn(const FName& InColumnId);
-
+	void AddColumn(const UPCGPointData* InPCGPointData, const FName& InColumnId, const FText& ColumnLabel, EHorizontalAlignment HeaderHAlign = HAlign_Center, EHorizontalAlignment CellHAlign = HAlign_Right);
 	void AddIndexColumn();
-	void RemoveIndexColumn();
-
-	void AddPointDataColumns();
-	void RemovePointDataColumns();
-
-	void AddMetadataColumn(const FName& InColumnId, const FName& InMetadataId, EPCGMetadataTypes InMetadataType, const int8 InValueIndex = INDEX_NONE, const TCHAR* PostFix = nullptr);
-	void RemoveMetadataColumns();
+	void AddPointDataColumns(const UPCGPointData* InPCGPointData);
+	void AddMetadataColumn(const UPCGData* InPCGData, const FName& InColumnId, EPCGMetadataTypes InMetadataType, const TCHAR* PostFix = nullptr);
 
 	/** Pointer back to the PCG editor that owns us */
 	TWeakPtr<FPCGEditor> PCGEditorPtr;
@@ -126,6 +130,7 @@ private:
 	TSharedPtr<STextBlock> InfoTextBlock;
 	TSharedPtr<SComboButton> FilterButton;
 
-	TMap<FName, FPCGMetadataInfo> MetadataInfos;
-	TArray<FName> HiddenAttributes;	
+	TArray<FName> HiddenAttributes;
+
+	TMap<FName, FPCGColumnData> PCGColumnData;
 };
