[PCG] AttributeRotatorOp hasn't got the right EditCondition for InputSource

#rnx

diff --git a/Engine/Plugins/Experimental/PCG/Source/PCG/Public/Elements/Metadata/PCGMetadataRotatorOpElement.h b/Engine/Plugins/Experimental/PCG/Source/PCG/Public/Elements/Metadata/PCGMetadataRotatorOpElement.h
--- a/Engine/Plugins/Experimental/PCG/Source/PCG/Public/Elements/Metadata/PCGMetadataRotatorOpElement.h
+++ b/Engine/Plugins/Experimental/PCG/Source/PCG/Public/Elements/Metadata/PCGMetadataRotatorOpElement.h
@@ -67,10 +67,10 @@ public:
 	UPROPERTY(BlueprintReadWrite, EditAnywhere, Category = Input)
 	FPCGAttributePropertySelector InputSource1;
 
-	UPROPERTY(BlueprintReadWrite, EditAnywhere, Category = Input, meta = (EditCondition = "Operation != EPCGMedadataTrigOperation::Invert", EditConditionHides))
+	UPROPERTY(BlueprintReadWrite, EditAnywhere, Category = Input, meta = (EditCondition = "Operation != EPCGMedadataRotatorOperation::Invert", EditConditionHides))
 	FPCGAttributePropertySelector InputSource2;
 
-	UPROPERTY(BlueprintReadWrite, EditAnywhere, Category = Input, meta = (EditCondition = "Operation == EPCGMedadataTrigOperation::Lerp", EditConditionHides))
+	UPROPERTY(BlueprintReadWrite, EditAnywhere, Category = Input, meta = (EditCondition = "Operation == EPCGMedadataRotatorOperation::Lerp", EditConditionHides))
 	FPCGAttributePropertySelector InputSource3;
 
 #if WITH_EDITORONLY_DATA
