[PCG] Fix callbacks when doing a reconstruction script with PCG Component

* Should fix refresh problem when a BP has a PCG Component without a graph, get instanciated and Graph is then set on the instance

diff --git a/Engine/Plugins/Experimental/PCG/Source/PCG/Private/PCGComponent.cpp b/Engine/Plugins/Experimental/PCG/Source/PCG/Private/PCGComponent.cpp
--- a/Engine/Plugins/Experimental/PCG/Source/PCG/Private/PCGComponent.cpp
+++ b/Engine/Plugins/Experimental/PCG/Source/PCG/Private/PCGComponent.cpp
@@ -2719,7 +2719,12 @@ void FPCGComponentInstanceData::ApplyToComponent(UActorComponent* Component, con
 
 #if WITH_EDITOR
 		// Reconnect callbacks
-		PCGComponent->GraphInstance->FixCallbacks();
+		if (PCGComponent->GraphInstance)
+		{
+			PCGComponent->GraphInstance->FixCallbacks();
+			PCGComponent->GraphInstance->OnGraphChangedDelegate.RemoveAll(PCGComponent);
+			PCGComponent->GraphInstance->OnGraphChangedDelegate.AddUObject(PCGComponent, &UPCGComponent::OnGraphChanged);
+		}
 #endif // WITH_EDITOR
 
 		bool bDoActorMapping = PCGComponent->bGenerated || PCGHelpers::IsRuntimeOrPIE();
