[PCG] Restore the name logic on AddAttribute when the attribute pin is not connected

* Before it was using bFromSourceParam, but we removed it
* Now check if the pin is connected and if so, behave like if bFromSourceParam was true

#rnx

diff --git a/Engine/Plugins/Experimental/PCG/Source/PCG/Private/Elements/PCGCreateAttribute.cpp b/Engine/Plugins/Experimental/PCG/Source/PCG/Private/Elements/PCGCreateAttribute.cpp
--- a/Engine/Plugins/Experimental/PCG/Source/PCG/Private/Elements/PCGCreateAttribute.cpp
+++ b/Engine/Plugins/Experimental/PCG/Source/PCG/Private/Elements/PCGCreateAttribute.cpp
@@ -110,7 +110,15 @@ FPCGElementPtr UPCGCreateAttributeSettings::CreateElement() const
 
 FName UPCGCreateAttributeSettings::AdditionalTaskNameInternal(FName NodeName) const
 {
-	if (ShouldAddAttributesPin())
+	if (HasAnyFlags(RF_ClassDefaultObject))
+	{
+		return NodeName;
+	}
+
+	const UPCGNode* Node = Cast<UPCGNode>(GetOuter());
+	const bool AttributesPinIsConnected = Node ? Node->IsInputPinConnected(PCGCreateAttributeConstants::AttributesLabel) : false;
+
+	if (ShouldAddAttributesPin() && AttributesPinIsConnected)
 	{
 		if ((OutputAttributeName == NAME_None) && (SourceParamAttributeName == NAME_None))
 		{
