[PCG] Fix badly-correlated results when using a Point Sampler to Density Noise combination, by varying random behaviour per node type.

Change node name to "Select Points" to clarify operation

diff --git a/Engine/Plugins/Experimental/PCG/Config/BasePCG.ini b/Engine/Plugins/Experimental/PCG/Config/BasePCG.ini
--- a/Engine/Plugins/Experimental/PCG/Config/BasePCG.ini
+++ b/Engine/Plugins/Experimental/PCG/Config/BasePCG.ini
@@ -14,4 +14,6 @@
 +PropertyRedirects=(OldName="PCGBlueprintElement.NestedLoopBody.InA",NewName="PCGBlueprintElement.NestedLoopBody.InOuterData")
 +PropertyRedirects=(OldName="PCGBlueprintElement.NestedLoopBody.InB",NewName="PCGBlueprintElement.NestedLoopBody.InInnerData")
 +PropertyRedirects=(OldName="PCGBlueprintElement.NestedLoopBody.InPointA",NewName="PCGBlueprintElement.NestedLoopBody.InOuterPoint")
-+PropertyRedirects=(OldName="PCGBlueprintElement.NestedLoopBody.InPointB",NewName="PCGBlueprintElement.NestedLoopBody.InInnerPoint")
\ No newline at end of file
++PropertyRedirects=(OldName="PCGBlueprintElement.NestedLoopBody.InPointB",NewName="PCGBlueprintElement.NestedLoopBody.InInnerPoint")
+
++ClassRedirects=(OldName="PCGPointSamplerSettings",NewName="/Script/PCG.PCGSelectPointsSettings")
diff --git a/Engine/Plugins/Experimental/PCG/Source/PCG/Private/Elements/PCGPointSampler.cpp b/Engine/Plugins/Experimental/PCG/Source/PCG/Private/Elements/PCGPointSampler.cpp
deleted file mode 100644
diff --git a/Engine/Plugins/Experimental/PCG/Source/PCG/Private/Elements/PCGSelectPoints.cpp b/Engine/Plugins/Experimental/PCG/Source/PCG/Private/Elements/PCGSelectPoints.cpp
new file mode 100644
--- /dev/null
+++ b/Engine/Plugins/Experimental/PCG/Source/PCG/Private/Elements/PCGSelectPoints.cpp
@@ -0,0 +1,150 @@
+// Copyright Epic Games, Inc. All Rights Reserved.
+
+#include "Elements/PCGSelectPoints.h"
+
+#include "PCGContext.h"
+#include "Data/PCGSpatialData.h"
+#include "Data/PCGPointData.h"
+#include "Helpers/PCGAsync.h"
+#include "Helpers/PCGHelpers.h"
+
+#include "Math/RandomStream.h"
+
+#include UE_INLINE_GENERATED_CPP_BY_NAME(PCGSelectPoints)
+
+#define LOCTEXT_NAMESPACE "PCGSelectPointsElement"
+
+UPCGSelectPointsSettings::UPCGSelectPointsSettings()
+{
+	bUseSeed = true;
+}
+
+#if WITH_EDITOR
+FText UPCGSelectPointsSettings::GetNodeTooltipText() const
+{
+	return LOCTEXT("NodeTooltip", "Selects a stable random subset of the input points.");
+}
+#endif
+
+FPCGElementPtr UPCGSelectPointsSettings::CreateElement() const
+{
+	return MakeShared<FPCGSelectPointsElement>();
+}
+
+bool FPCGSelectPointsElement::ExecuteInternal(FPCGContext* Context) const
+{
+	TRACE_CPUPROFILER_EVENT_SCOPE(FPCGSelectPointsElement::Execute);
+	// TODO: make time-sliced implementation
+	const UPCGSelectPointsSettings* Settings = Context->GetInputSettings<UPCGSelectPointsSettings>();
+	check(Settings);
+
+	TArray<FPCGTaggedData> Inputs = Context->InputData.GetInputs();
+
+	TArray<FPCGTaggedData>& Outputs = Context->OutputData.TaggedData;
+
+	const float Ratio = Settings->Ratio;
+#if WITH_EDITOR
+	const bool bKeepZeroDensityPoints = Settings->bKeepZeroDensityPoints;
+#else
+	const bool bKeepZeroDensityPoints = false;
+#endif
+
+	const int Seed = Context->GetSeed();
+
+	const bool bNoSampling = (Ratio <= 0.0f);
+	const bool bTrivialSampling = (Ratio >= 1.0f);
+
+	// Early exit when nothing will be generated out of this sampler
+	if (bNoSampling && !bKeepZeroDensityPoints)
+	{
+		PCGE_LOG(Verbose, LogOnly, LOCTEXT("AllInputsRejected", "Skipped - all inputs rejected"));
+		return true;
+	}
+
+	// TODO: embarassingly parallel loop
+	for (const FPCGTaggedData& Input : Inputs)
+	{
+		FPCGTaggedData& Output = Outputs.Add_GetRef(Input);
+
+		if (!Input.Data || Cast<UPCGSpatialData>(Input.Data) == nullptr)
+		{
+			PCGE_LOG(Error, GraphAndLog, LOCTEXT("InvalidInputData", "Invalid input data"));
+			continue;
+		}
+
+		// Skip processing if the transformation would be trivial
+		if (bTrivialSampling)
+		{
+			PCGE_LOG(Verbose, LogOnly, LOCTEXT("SkippedTrivialSampling", "Skipped - trivial sampling"));
+			continue;
+		}
+
+		const UPCGPointData* OriginalData = Cast<UPCGSpatialData>(Input.Data)->ToPointData(Context);
+
+		if (!OriginalData)
+		{
+			PCGE_LOG(Error, GraphAndLog, LOCTEXT("NoPointDataInInput", "Unable to get point data from input"));
+			continue;
+		}
+
+		const TArray<FPCGPoint>& Points = OriginalData->GetPoints();
+		const int OriginalPointCount = Points.Num();
+
+		UPCGPointData* SampledData = NewObject<UPCGPointData>();
+		SampledData->InitializeFromData(OriginalData);
+		TArray<FPCGPoint>& SampledPoints = SampledData->GetMutablePoints();
+
+		Output.Data = SampledData;
+
+		// TODO: randomize on the fractional number of points
+#if WITH_EDITOR
+		int TargetNumPoints = (bKeepZeroDensityPoints ? OriginalPointCount : OriginalPointCount * Ratio);
+#else
+		int TargetNumPoints = OriginalPointCount * Ratio;
+#endif
+
+		// Early out
+		if (TargetNumPoints == 0)
+		{
+			PCGE_LOG(Verbose, LogOnly, LOCTEXT("SkippedAllPointsRejected", "Skipped - all points rejected"));
+			continue;
+		}
+		else
+		{
+			TRACE_CPUPROFILER_EVENT_SCOPE(FPCGSelectPointsElement::Execute::SelectPoints);
+
+			FPCGAsync::AsyncPointProcessing(Context, OriginalPointCount, SampledPoints, [&Points, Seed, Ratio, bKeepZeroDensityPoints](int32 Index, FPCGPoint& OutPoint)
+			{
+				const FPCGPoint& Point = Points[Index];
+
+				// Apply a high-pass filter based on selected ratio
+				FRandomStream RandomSource(PCGHelpers::ComputeSeed(Seed, Point.Seed));
+				float Chance = RandomSource.FRand();
+
+				if (Chance < Ratio)
+				{
+					OutPoint = Point;
+					return true;
+				}
+#if WITH_EDITOR
+				else if (bKeepZeroDensityPoints)
+				{
+					OutPoint = Point;
+					OutPoint.Density = 0;
+					return true;
+				}
+#endif
+				else
+				{
+					return false;
+				}
+			});
+
+			PCGE_LOG(Verbose, LogOnly, FText::Format(LOCTEXT("GenerationInfo", "Generated {0} points from {1} source points"), SampledPoints.Num(), OriginalPointCount));
+		}
+	}
+
+	return true;
+}
+
+#undef LOCTEXT_NAMESPACE
diff --git a/Engine/Plugins/Experimental/PCG/Source/PCG/Private/PCGSettings.cpp b/Engine/Plugins/Experimental/PCG/Source/PCG/Private/PCGSettings.cpp
--- a/Engine/Plugins/Experimental/PCG/Source/PCG/Private/PCGSettings.cpp
+++ b/Engine/Plugins/Experimental/PCG/Source/PCG/Private/PCGSettings.cpp
@@ -80,6 +80,17 @@ void UPCGSettingsInterface::SetEnabled(bool bInEnabled)
 	}
 }
 
+UPCGSettings::UPCGSettings()
+{
+	TypeNameHash = GetTypeHash(GetClass()->GetFName());
+}
+
+UPCGSettings::UPCGSettings(const FObjectInitializer& ObjectInitializer)
+	: Super(ObjectInitializer)
+{
+	TypeNameHash = GetTypeHash(GetClass()->GetFName());
+}
+
 bool UPCGSettings::operator==(const UPCGSettings& Other) const
 {
 	if (this == &Other)
@@ -386,7 +397,11 @@ UPCGNode* UPCGSettings::CreateNode() const
 
 int UPCGSettings::GetSeed(const UPCGComponent* InSourceComponent) const
 {
-	return !bUseSeed ? 42 : (InSourceComponent ? PCGHelpers::ComputeSeed(Seed, InSourceComponent->Seed) : Seed);
+	// Mix in a type hash so we don't get multiple nodes doing randomness in an identical way (we had a case where
+	// a Point Subset followed by a Density Noise produced uniform densities because the random seeds/streams were correllated).
+	// It is still possible to chain nodes of the same type (Point Subset -> Point Subset) and see correllated behaviour, in which
+	// case the random seed on one of them should be changed in the node settings.
+	return !bUseSeed ? TypeNameHash : (InSourceComponent ? PCGHelpers::ComputeSeed(Seed, InSourceComponent->Seed, TypeNameHash) : Seed);
 }
 
 #if WITH_EDITOR
diff --git a/Engine/Plugins/Experimental/PCG/Source/PCG/Public/Elements/PCGPointSampler.h b/Engine/Plugins/Experimental/PCG/Source/PCG/Public/Elements/PCGPointSampler.h
deleted file mode 100644
diff --git a/Engine/Plugins/Experimental/PCG/Source/PCG/Public/Elements/PCGSelectPoints.h b/Engine/Plugins/Experimental/PCG/Source/PCG/Public/Elements/PCGSelectPoints.h
new file mode 100644
--- /dev/null
+++ b/Engine/Plugins/Experimental/PCG/Source/PCG/Public/Elements/PCGSelectPoints.h
@@ -0,0 +1,51 @@
+// Copyright Epic Games, Inc. All Rights Reserved.
+
+#pragma once
+
+#include "PCGPin.h"
+#include "PCGSettings.h"
+
+#include "PCGSelectPoints.generated.h"
+
+UCLASS(BlueprintType, ClassGroup = (Procedural))
+class PCG_API UPCGSelectPointsSettings : public UPCGSettings
+{
+	GENERATED_BODY()
+
+public:
+	UPCGSelectPointsSettings();
+
+	//~Begin UPCGSettings interface
+#if WITH_EDITOR
+	virtual FName GetDefaultNodeName() const override { return FName(TEXT("SelectPoints")); }
+	virtual FText GetDefaultNodeTitle() const override { return NSLOCTEXT("PCGSelectPointsElement", "NodeTitle", "Select Points"); }
+	virtual FText GetNodeTooltipText() const override;
+	virtual EPCGSettingsType GetType() const override { return EPCGSettingsType::Sampler; }
+#endif
+
+protected:
+	virtual TArray<FPCGPinProperties> InputPinProperties() const override { return Super::DefaultPointInputPinProperties(); }
+	virtual TArray<FPCGPinProperties> OutputPinProperties() const override { return Super::DefaultPointOutputPinProperties(); }
+	virtual FPCGElementPtr CreateElement() const override;
+	//~End UPCGSettings interface
+
+public:
+	UPROPERTY(BlueprintReadWrite, EditAnywhere, Category = Settings, meta=(ClampMin="0", ClampMax="1", PCG_Overridable))
+	float Ratio = 0.1f;
+
+#if WITH_EDITORONLY_DATA
+	UPROPERTY(Transient, BlueprintReadWrite, EditAnywhere, Category = "Settings|Debug", meta = (PCG_Overridable))
+	bool bKeepZeroDensityPoints = false;
+#endif
+};
+
+class FPCGSelectPointsElement : public FSimplePCGElement
+{
+protected:
+	virtual bool ExecuteInternal(FPCGContext* Context) const override;
+};
+
+#if UE_ENABLE_INCLUDE_ORDER_DEPRECATED_IN_5_2
+#include "CoreMinimal.h"
+#include "PCGNode.h"
+#endif
diff --git a/Engine/Plugins/Experimental/PCG/Source/PCG/Public/PCGSettings.h b/Engine/Plugins/Experimental/PCG/Source/PCG/Public/PCGSettings.h
--- a/Engine/Plugins/Experimental/PCG/Source/PCG/Public/PCGSettings.h
+++ b/Engine/Plugins/Experimental/PCG/Source/PCG/Public/PCGSettings.h
@@ -154,6 +154,9 @@ class PCG_API UPCGSettings : public UPCGSettingsInterface
 	friend class UPCGSettingsInterface;
 
 public:
+	UPCGSettings();
+	UPCGSettings(const FObjectInitializer& ObjectInitializer);
+
 	// ~Begin UPCGData interface
 	virtual EPCGDataType GetDataType() const override { return EPCGDataType::Settings; }
 	// ~End UPCGData interface
@@ -365,6 +368,10 @@ private:
 
 	/** The cached Crc for these settings. */
 	FPCGCrc CachedCrc;
+
+	/** Hash of the derived type name which is of the form 'PCGExecuteBlueprint'. */
+	UPROPERTY(Transient)
+	uint32 TypeNameHash = 0;
 };
 
 UCLASS(BlueprintType, ClassGroup = (Procedural))
