[PCG] Tentative fix for static analysis in editor test

#rnx

diff --git a/Engine/Plugins/Experimental/PCG/Source/PCGEditor/Private/Tests/Display/PCGComponentInNewPCGVolumes.cpp b/Engine/Plugins/Experimental/PCG/Source/PCGEditor/Private/Tests/Display/PCGComponentInNewPCGVolumes.cpp
--- a/Engine/Plugins/Experimental/PCG/Source/PCGEditor/Private/Tests/Display/PCGComponentInNewPCGVolumes.cpp
+++ b/Engine/Plugins/Experimental/PCG/Source/PCGEditor/Private/Tests/Display/PCGComponentInNewPCGVolumes.cpp
@@ -20,6 +20,12 @@ bool FPCGComponentInNewPCGVolumes::RunTest(const FString& Parameters)
 	UWorld* World = GEditor->GetEditorWorldContext().World();
 	UTEST_NOT_NULL(TEXT("Failed to get editor world context!"), World);
 
+	// Condition for static analysis.
+	if (!World)
+	{
+		return false;
+	}
+
 	// We need to get the PCGVolume actor factory
 	TSubclassOf<AActor> PCGVolumeClass = APCGVolume::StaticClass();
 	UActorFactory* PCGVolumeFactory = GEditor->FindActorFactoryForActorClass(PCGVolumeClass);
@@ -36,13 +42,12 @@ bool FPCGComponentInNewPCGVolumes::RunTest(const FString& Parameters)
 		FTransform ActorTransform;
 		VolumeActor = GEditor->UseActorFactory(PCGVolumeFactory, PCGVolumeAssetData, &ActorTransform);
 	}
-	if (!VolumeActor) return false;
-	
+
+	UTEST_NOT_NULL(TEXT("Failed to add PCGVolume actor."), VolumeActor);
+
+	// Condition for static analysis.
 	if (VolumeActor)
 	{
-		UPackage* Package = VolumeActor->GetOutermost();
-		Package->SetFlags(RF_Transient); // Mark package as transient
-
 		UPCGComponent* PCGComponent = VolumeActor->FindComponentByClass<UPCGComponent>();
 		TestNotNull("PCGVolume actor does not contain a PCGComponent!", PCGComponent);
 
