[PCG] Teardown callbacks in OnComponentDestroyed to avoid stacking callbacks on trash components

#rnx

diff --git a/Engine/Plugins/Experimental/PCG/Source/PCG/Private/PCGComponent.cpp b/Engine/Plugins/Experimental/PCG/Source/PCG/Private/PCGComponent.cpp
--- a/Engine/Plugins/Experimental/PCG/Source/PCG/Private/PCGComponent.cpp
+++ b/Engine/Plugins/Experimental/PCG/Source/PCG/Private/PCGComponent.cpp
@@ -947,6 +947,15 @@ void UPCGComponent::OnComponentDestroyed(bool bDestroyingHierarchy)
 			Subsystem->UnregisterPCGComponent(this);
 		}
 	}
+
+	// BeginDestroy is not called immediately when a component is destroyed. Therefore callbacks are not cleaned
+	// until GC is ran, and can stack up with BP reconstruction scripts. Force the removal of callbacks here. If the component
+	// is dead, we don't want to react to callbacks anyway.
+	if (GraphInstance)
+	{
+		GraphInstance->OnGraphChangedDelegate.RemoveAll(this);
+		GraphInstance->TeardownCallbacks();
+	}
 #endif
 
 	// Bookkeeping local components that might be deleted by the user.
