[PCG] New attribute operations + few fixes

* Add OneMinus math operation
* Add MakeRotator attribute, with preconfigured settings, matching BP MakeRot
* Fix issue with preconfigred settings where pins were not updated accordingly
* Re-order pins

diff --git a/Engine/Plugins/Experimental/PCG/Source/PCG/Private/Elements/Metadata/PCGMetadataMakeRotator.cpp b/Engine/Plugins/Experimental/PCG/Source/PCG/Private/Elements/Metadata/PCGMetadataMakeRotator.cpp
new file mode 100644
--- /dev/null
+++ b/Engine/Plugins/Experimental/PCG/Source/PCG/Private/Elements/Metadata/PCGMetadataMakeRotator.cpp
@@ -0,0 +1,153 @@
+// Copyright Epic Games, Inc. All Rights Reserved.
+
+#include "Elements/Metadata/PCGMetadataMakeRotator.h"
+
+#include "Elements/Metadata/PCGMetadataElementCommon.h"
+
+#define LOCTEXT_NAMESPACE "PCGMetadataMakeRotatorSettings"
+
+FName UPCGMetadataMakeRotatorSettings::GetInputPinLabel(uint32 Index) const
+{
+	switch (Operation)
+	{
+	case EPCGMetadataMakeRotatorOp::MakeRotFromX:
+		return PCGMetadataMakeRotatorConstants::XLabel;
+	case EPCGMetadataMakeRotatorOp::MakeRotFromY:
+		return PCGMetadataMakeRotatorConstants::YLabel;
+	case EPCGMetadataMakeRotatorOp::MakeRotFromZ:
+		return PCGMetadataMakeRotatorConstants::ZLabel;
+	case EPCGMetadataMakeRotatorOp::MakeRotFromXY:
+		return (Index == 0 ? PCGMetadataMakeRotatorConstants::XLabel : PCGMetadataMakeRotatorConstants::YLabel);
+	case EPCGMetadataMakeRotatorOp::MakeRotFromYX:
+		return (Index == 1 ? PCGMetadataMakeRotatorConstants::XLabel : PCGMetadataMakeRotatorConstants::YLabel);
+	case EPCGMetadataMakeRotatorOp::MakeRotFromXZ:
+		return (Index == 0 ? PCGMetadataMakeRotatorConstants::XLabel : PCGMetadataMakeRotatorConstants::ZLabel);
+	case EPCGMetadataMakeRotatorOp::MakeRotFromZX:
+		return (Index == 1 ? PCGMetadataMakeRotatorConstants::XLabel : PCGMetadataMakeRotatorConstants::ZLabel);
+	case EPCGMetadataMakeRotatorOp::MakeRotFromYZ:
+		return (Index == 0 ? PCGMetadataMakeRotatorConstants::YLabel : PCGMetadataMakeRotatorConstants::ZLabel);
+	case EPCGMetadataMakeRotatorOp::MakeRotFromZY:
+		return (Index == 1 ? PCGMetadataMakeRotatorConstants::YLabel : PCGMetadataMakeRotatorConstants::ZLabel);
+	case EPCGMetadataMakeRotatorOp::MakeRotFromAxis:
+	default:
+		return (Index == 0 ? PCGMetadataMakeRotatorConstants::ForwardLabel : (Index == 1 ? PCGMetadataMakeRotatorConstants::RightLabel : PCGMetadataMakeRotatorConstants::UpLabel));
+	}
+}
+
+uint32 UPCGMetadataMakeRotatorSettings::GetInputPinNum() const
+{
+	switch (Operation)
+	{
+	case EPCGMetadataMakeRotatorOp::MakeRotFromX: // fall-through
+	case EPCGMetadataMakeRotatorOp::MakeRotFromY: // fall-through
+	case EPCGMetadataMakeRotatorOp::MakeRotFromZ:
+		return 1;
+	case EPCGMetadataMakeRotatorOp::MakeRotFromXY: // fall-through
+	case EPCGMetadataMakeRotatorOp::MakeRotFromYX: // fall-through
+	case EPCGMetadataMakeRotatorOp::MakeRotFromXZ: // fall-through
+	case EPCGMetadataMakeRotatorOp::MakeRotFromZX: // fall-through
+	case EPCGMetadataMakeRotatorOp::MakeRotFromYZ: // fall-through
+	case EPCGMetadataMakeRotatorOp::MakeRotFromZY:
+		return 2;
+	case EPCGMetadataMakeRotatorOp::MakeRotFromAxis:
+	default:
+		return 3;
+	}
+}
+
+bool UPCGMetadataMakeRotatorSettings::IsSupportedInputType(uint16 TypeId, uint32 InputIndex, bool& bHasSpecialRequirement) const
+{
+	bHasSpecialRequirement = false;
+	return PCG::Private::IsOfTypes<FVector, FVector2D, float, double, int32, int64>(TypeId);
+}
+
+FPCGAttributePropertySelector UPCGMetadataMakeRotatorSettings::GetInputSource(uint32 Index) const
+{
+	switch (Index)
+	{
+	case 0:
+		return InputSource1;
+	case 1:
+		return InputSource2;
+	case 2:
+		return InputSource3;
+	default:
+		return FPCGAttributePropertySelector();
+	}
+}
+
+uint16 UPCGMetadataMakeRotatorSettings::GetOutputType(uint16 InputTypeId) const
+{
+	return (uint16)EPCGMetadataTypes::Rotator;
+}
+
+#if WITH_EDITOR
+FName UPCGMetadataMakeRotatorSettings::GetDefaultNodeName() const
+{
+	return TEXT("MakeRotatorAttribute");
+}
+
+FText UPCGMetadataMakeRotatorSettings::GetDefaultNodeTitle() const
+{
+	return LOCTEXT("NodeTitle", "Make Rotator Attribute");
+}
+
+TArray<FPCGPreConfiguredSettingsInfo> UPCGMetadataMakeRotatorSettings::GetPreconfiguredInfo() const
+{
+	return PCGMetadataElementCommon::FillPreconfiguredSettingsInfoFromEnum<EPCGMetadataMakeRotatorOp>();
+}
+#endif // WITH_EDITOR
+
+void UPCGMetadataMakeRotatorSettings::ApplyPreconfiguredSettings(const FPCGPreConfiguredSettingsInfo& PreconfiguredInfo)
+{
+	if (const UEnum* EnumPtr = StaticEnum<EPCGMetadataMakeRotatorOp>())
+	{
+		if (EnumPtr->IsValidEnumValue(PreconfiguredInfo.PreconfiguredIndex))
+		{
+			Operation = EPCGMetadataMakeRotatorOp(PreconfiguredInfo.PreconfiguredIndex);
+		}
+	}
+}
+
+FPCGElementPtr UPCGMetadataMakeRotatorSettings::CreateElement() const
+{
+	return MakeShared<FPCGMetadataMakeRotatorElement>();
+}
+
+bool FPCGMetadataMakeRotatorElement::DoOperation(FOperationData& OperationData) const
+{
+	TRACE_CPUPROFILER_EVENT_SCOPE(FPCGMetadataMakeRotatorElement::Execute);
+
+	const UPCGMetadataMakeRotatorSettings* Settings = CastChecked<UPCGMetadataMakeRotatorSettings>(OperationData.Settings);
+
+	switch (Settings->Operation)
+	{
+	case EPCGMetadataMakeRotatorOp::MakeRotFromX:
+		return DoUnaryOp<FVector>(OperationData, [](const FVector& X) -> FRotator { return FRotationMatrix::MakeFromX(X).Rotator(); });
+	case EPCGMetadataMakeRotatorOp::MakeRotFromY:
+		return DoUnaryOp<FVector>(OperationData, [](const FVector& Y) -> FRotator { return FRotationMatrix::MakeFromY(Y).Rotator(); });
+	case EPCGMetadataMakeRotatorOp::MakeRotFromZ:
+		return DoUnaryOp<FVector>(OperationData, [](const FVector& Z) -> FRotator { return FRotationMatrix::MakeFromZ(Z).Rotator(); });
+	case EPCGMetadataMakeRotatorOp::MakeRotFromXY:
+		return DoBinaryOp<FVector, FVector>(OperationData, [](const FVector& X, const FVector& Y) -> FRotator { return FRotationMatrix::MakeFromXY(X, Y).Rotator(); });
+	case EPCGMetadataMakeRotatorOp::MakeRotFromYX:
+		return DoBinaryOp<FVector, FVector>(OperationData, [](const FVector& Y, const FVector& X) -> FRotator { return FRotationMatrix::MakeFromYX(Y, X).Rotator(); });
+	case EPCGMetadataMakeRotatorOp::MakeRotFromXZ:
+		return DoBinaryOp<FVector, FVector>(OperationData, [](const FVector& X, const FVector& Z) -> FRotator { return FRotationMatrix::MakeFromXZ(X, Z).Rotator(); });
+	case EPCGMetadataMakeRotatorOp::MakeRotFromZX:
+		return DoBinaryOp<FVector, FVector>(OperationData, [](const FVector& Z, const FVector& X) -> FRotator { return FRotationMatrix::MakeFromZX(Z, X).Rotator(); });
+	case EPCGMetadataMakeRotatorOp::MakeRotFromYZ:
+		return DoBinaryOp<FVector, FVector>(OperationData, [](const FVector& Y, const FVector& Z) -> FRotator { return FRotationMatrix::MakeFromYZ(Y, Z).Rotator(); });
+	case EPCGMetadataMakeRotatorOp::MakeRotFromZY:
+		return DoBinaryOp<FVector, FVector>(OperationData, [](const FVector& Z, const FVector& Y) -> FRotator { return FRotationMatrix::MakeFromZY(Z, Y).Rotator(); });
+	case EPCGMetadataMakeRotatorOp::MakeRotFromAxis:
+		return DoTernaryOp<FVector, FVector, FVector>(OperationData, [](const FVector& X, const FVector& Y, const FVector& Z) -> FRotator
+		{
+			return FMatrix(X.GetSafeNormal(), Y.GetSafeNormal(), Z.GetSafeNormal(), FVector::ZeroVector).Rotator();
+		});
+	}
+
+	return true;
+}
+
+#undef LOCTEXT_NAMESPACE
diff --git a/Engine/Plugins/Experimental/PCG/Source/PCG/Private/Elements/Metadata/PCGMetadataMathsOpElement.cpp b/Engine/Plugins/Experimental/PCG/Source/PCG/Private/Elements/Metadata/PCGMetadataMathsOpElement.cpp
--- a/Engine/Plugins/Experimental/PCG/Source/PCG/Private/Elements/Metadata/PCGMetadataMathsOpElement.cpp
+++ b/Engine/Plugins/Experimental/PCG/Source/PCG/Private/Elements/Metadata/PCGMetadataMathsOpElement.cpp
@@ -101,6 +101,8 @@ namespace PCGMetadataMathsSettings
 			return PCGMetadataMaths::Floor(Value);
 		case EPCGMedadataMathsOperation::Ceil:
 			return PCGMetadataMaths::Ceil(Value);
+		case EPCGMedadataMathsOperation::OneMinus:
+			return T{ 1 } - Value;
 		default:
 			return T{};
 		}
diff --git a/Engine/Plugins/Experimental/PCG/Source/PCG/Private/PCGNode.cpp b/Engine/Plugins/Experimental/PCG/Source/PCG/Private/PCGNode.cpp
--- a/Engine/Plugins/Experimental/PCG/Source/PCG/Private/PCGNode.cpp
+++ b/Engine/Plugins/Experimental/PCG/Source/PCG/Private/PCGNode.cpp
@@ -713,6 +713,7 @@ EPCGChangeType UPCGNode::UpdatePins(TFunctionRef<UPCGPin*(UPCGNode*)> PinAllocat
 			}
 		}
 
+		bool bWasModified = false;
 		const bool bUpdateUnmatchedPin = UnmatchedPins.Num() == 1 && UnmatchedProperties.Num() == 1;
 		if (bUpdateUnmatchedPin)
 		{
@@ -741,6 +742,7 @@ EPCGChangeType UPCGNode::UpdatePins(TFunctionRef<UPCGPin*(UPCGNode*)> PinAllocat
 
 			if(!UnmatchedPins.IsEmpty() || !UnmatchedProperties.IsEmpty() || !DuplicatedNamePins.IsEmpty())
 			{
+				bWasModified = true;
 				Modify();
 				bChangedPins = true;
 			}
@@ -783,6 +785,26 @@ EPCGChangeType UPCGNode::UpdatePins(TFunctionRef<UPCGPin*(UPCGNode*)> PinAllocat
 			}
 		}
 
+		// Final pass, to check the order. We re-order if the order is not the same in PinProperties and Pins, without breaking the edges.
+		// Also, at this point, we should have the same number of items in the pins and in the pin properties
+		check(Pins.Num() == PinProperties.Num());
+		for (int32 i = 0; i < PinProperties.Num(); ++i)
+		{
+			const FPCGPinProperties& CurrentPinProperties = PinProperties[i];
+			int32 AssociatedPinPindex = Pins.IndexOfByPredicate([&CurrentPinProperties](const UPCGPin* Pin) -> bool { return Pin->Properties.Label == CurrentPinProperties.Label; });
+			if (i != AssociatedPinPindex)
+			{
+				if (!bWasModified)
+				{
+					bWasModified = true;
+					Modify();
+				}
+
+				bChangedPins = true;
+				Pins.Swap(i, AssociatedPinPindex);
+			}
+		}
+
 		return (bAppliedEdgeChanges ? EPCGChangeType::Edge : EPCGChangeType::None) | (bChangedPins ? EPCGChangeType::Node : EPCGChangeType::None);
 	};
 
diff --git a/Engine/Plugins/Experimental/PCG/Source/PCG/Public/Elements/Metadata/PCGMetadataMakeRotator.h b/Engine/Plugins/Experimental/PCG/Source/PCG/Public/Elements/Metadata/PCGMetadataMakeRotator.h
new file mode 100644
--- /dev/null
+++ b/Engine/Plugins/Experimental/PCG/Source/PCG/Public/Elements/Metadata/PCGMetadataMakeRotator.h
@@ -0,0 +1,82 @@
+// Copyright Epic Games, Inc. All Rights Reserved.
+
+#pragma once
+
+#include "Elements/Metadata/PCGMetadataOpElementBase.h"
+
+#include "PCGMetadataMakeRotator.generated.h"
+
+namespace PCGMetadataMakeRotatorConstants
+{
+	const FName XLabel = TEXT("X");
+	const FName YLabel = TEXT("Y");
+	const FName ZLabel = TEXT("Z");
+	const FName ForwardLabel = TEXT("Forward");
+	const FName RightLabel = TEXT("Right");
+	const FName UpLabel = TEXT("Up");
+}
+
+UENUM()
+enum class EPCGMetadataMakeRotatorOp : uint8
+{
+	MakeRotFromX,
+	MakeRotFromY,
+	MakeRotFromZ,
+	MakeRotFromXY,
+	MakeRotFromYX,
+	MakeRotFromXZ,
+	MakeRotFromZX,
+	MakeRotFromYZ,
+	MakeRotFromZY,
+	MakeRotFromAxis
+};
+
+/* Create a Rotator from 1, 2 or 3 axis. */
+UCLASS(BlueprintType, ClassGroup = (Procedural))
+class PCG_API UPCGMetadataMakeRotatorSettings : public UPCGMetadataSettingsBase
+{
+	GENERATED_BODY()
+
+public:
+#if WITH_EDITOR
+	//~Begin UPCGSettings interface
+	virtual FName GetDefaultNodeName() const override;
+	virtual FText GetDefaultNodeTitle() const override;
+	virtual TArray<FPCGPreConfiguredSettingsInfo> GetPreconfiguredInfo() const override;
+	virtual bool OnlyExposePreconfiguredSettings() const override { return true; }
+#endif
+	virtual void ApplyPreconfiguredSettings(const FPCGPreConfiguredSettingsInfo& PreconfiguredInfo) override;
+	//~End UPCGSettings interface
+
+	FPCGAttributePropertySelector GetInputSource(uint32 Index) const override;
+
+	virtual FName GetInputPinLabel(uint32 Index) const override;
+	virtual uint32 GetInputPinNum() const override;
+
+	virtual bool IsSupportedInputType(uint16 TypeId, uint32 InputIndex, bool& bHasSpecialRequirement) const override;
+	virtual uint16 GetOutputType(uint16 InputTypeId) const override;
+
+protected:
+	//~Begin UPCGSettings interface
+	virtual FPCGElementPtr CreateElement() const override;
+	//~End UPCGSettings interface
+
+public:
+	UPROPERTY(BlueprintReadWrite, EditAnywhere, Category = Input)
+	FPCGAttributePropertySelector InputSource1;
+
+	UPROPERTY(BlueprintReadWrite, EditAnywhere, Category = Input, meta = (EditCondition = "Operation != EPCGMetadataMakeRotatorOp::MakeRotFromX && Operation != EPCGMetadataMakeRotatorOp::MakeRotFromY && Operation != EPCGMetadataMakeRotatorOp::MakeRotFromZ", EditConditionHides))
+	FPCGAttributePropertySelector InputSource2;
+
+	UPROPERTY(BlueprintReadWrite, EditAnywhere, Category = Input, meta = (EditCondition = "Operation == EPCGMetadataMakeRotatorOp::MakeRotFromAxis", EditConditionHides))
+	FPCGAttributePropertySelector InputSource3;
+
+	UPROPERTY(BlueprintReadWrite, EditAnywhere, Category = Settings)
+	EPCGMetadataMakeRotatorOp Operation = EPCGMetadataMakeRotatorOp::MakeRotFromAxis;
+};
+
+class FPCGMetadataMakeRotatorElement : public FPCGMetadataElementBase
+{
+protected:
+	virtual bool DoOperation(FOperationData& OperationData) const override;
+};
diff --git a/Engine/Plugins/Experimental/PCG/Source/PCG/Public/Elements/Metadata/PCGMetadataMathsOpElement.h b/Engine/Plugins/Experimental/PCG/Source/PCG/Public/Elements/Metadata/PCGMetadataMathsOpElement.h
--- a/Engine/Plugins/Experimental/PCG/Source/PCG/Public/Elements/Metadata/PCGMetadataMathsOpElement.h
+++ b/Engine/Plugins/Experimental/PCG/Source/PCG/Public/Elements/Metadata/PCGMetadataMathsOpElement.h
@@ -19,6 +19,7 @@ enum class EPCGMedadataMathsOperation : uint16
 	Abs,
 	Floor,
 	Ceil,
+	OneMinus UMETA(Tooltip = "1 - X operation"),
 
 	// Binary op
 	BinaryOp = 1 << 11 UMETA(Hidden),
diff --git a/Engine/Plugins/Experimental/PCG/Source/PCGEditor/Private/PCGEditorGraphSchemaActions.cpp b/Engine/Plugins/Experimental/PCG/Source/PCGEditor/Private/PCGEditorGraphSchemaActions.cpp
--- a/Engine/Plugins/Experimental/PCG/Source/PCGEditor/Private/PCGEditorGraphSchemaActions.cpp
+++ b/Engine/Plugins/Experimental/PCG/Source/PCGEditor/Private/PCGEditorGraphSchemaActions.cpp
@@ -54,6 +54,7 @@ UEdGraphNode* FPCGEditorGraphSchemaAction_NewNativeElement::PerformAction(UEdGra
 	if (DefaultNodeSettings)
 	{
 		DefaultNodeSettings->ApplyPreconfiguredSettings(PreconfiguredInfo);
+		NewPCGNode->UpdateAfterSettingsChangeDuringCreation();
 	}
 
 	PostCreation(NewPCGNode);
