[PCG] Fix possible crash with nullptr dereferencement

#rnx

diff --git a/Engine/Plugins/Experimental/PCG/Source/PCG/Private/PCGPin.cpp b/Engine/Plugins/Experimental/PCG/Source/PCG/Private/PCGPin.cpp
--- a/Engine/Plugins/Experimental/PCG/Source/PCG/Private/PCGPin.cpp
+++ b/Engine/Plugins/Experimental/PCG/Source/PCG/Private/PCGPin.cpp
@@ -206,7 +206,10 @@ bool UPCGPin::BreakAllIncompatibleEdges(TSet<UPCGNode*>* InTouchedNodes/*= nullp
 			Modify();
 			Edges.RemoveAtSwap(EdgeIndex);
 
-			InTouchedNodes->Add(Node);
+			if (InTouchedNodes)
+			{
+				InTouchedNodes->Add(Node);
+			}
 
 			if (OtherPin)
 			{
