[PCG] Fix pin tooltips not updating on change

diff --git a/Engine/Plugins/Experimental/PCG/Source/PCG/Private/PCGNode.cpp b/Engine/Plugins/Experimental/PCG/Source/PCG/Private/PCGNode.cpp
--- a/Engine/Plugins/Experimental/PCG/Source/PCG/Private/PCGNode.cpp
+++ b/Engine/Plugins/Experimental/PCG/Source/PCG/Private/PCGNode.cpp
@@ -681,6 +681,7 @@ EPCGChangeType UPCGNode::UpdatePins(TFunctionRef<UPCGPin*(UPCGNode*)> PinAllocat
 	{
 		bool bAppliedEdgeChanges = false;
 		bool bChangedPins = false;
+		bool bChangedTooltips = false;
 
 		// Find unmatched pins vs. properties on a name basis
 		TArray<UPCGPin*> UnmatchedPins;
@@ -696,6 +697,13 @@ EPCGChangeType UPCGNode::UpdatePins(TFunctionRef<UPCGPin*(UPCGNode*)> PinAllocat
 					bAppliedEdgeChanges |= Pin->BreakAllIncompatibleEdges(&TouchedNodes);
 					bChangedPins = true;
 				}
+#if WITH_EDITOR
+				else if (Pin->Properties.Tooltip.CompareTo(MatchingProperties->Tooltip))
+				{
+					Pin->Properties.Tooltip = MatchingProperties->Tooltip;
+					bChangedTooltips = true;
+				}
+#endif // WITH_EDITOR
 			}
 			else
 			{
@@ -805,7 +813,9 @@ EPCGChangeType UPCGNode::UpdatePins(TFunctionRef<UPCGPin*(UPCGNode*)> PinAllocat
 			}
 		}
 
-		return (bAppliedEdgeChanges ? EPCGChangeType::Edge : EPCGChangeType::None) | (bChangedPins ? EPCGChangeType::Node : EPCGChangeType::None);
+		return (bAppliedEdgeChanges ? EPCGChangeType::Edge : EPCGChangeType::None)
+			| (bChangedPins ? EPCGChangeType::Node : EPCGChangeType::None)
+			| (bChangedTooltips ? EPCGChangeType::Cosmetic : EPCGChangeType::None);
 	};
 
 	EPCGChangeType ChangeType = EPCGChangeType::None;
