[PCG] Fix CopyAttribute flow Attribute -> Attribute

#rnx

diff --git a/Engine/Plugins/Experimental/PCG/Source/PCG/Private/Elements/Metadata/PCGMetadataElement.cpp b/Engine/Plugins/Experimental/PCG/Source/PCG/Private/Elements/Metadata/PCGMetadataElement.cpp
--- a/Engine/Plugins/Experimental/PCG/Source/PCG/Private/Elements/Metadata/PCGMetadataElement.cpp
+++ b/Engine/Plugins/Experimental/PCG/Source/PCG/Private/Elements/Metadata/PCGMetadataElement.cpp
@@ -178,17 +178,16 @@ bool FPCGMetadataOperationElement::ExecuteInternal(FPCGContext* Context) const
 				// Nothing to do...
 				continue;
 			}
+			
+			const FPCGMetadataAttributeBase* OriginalAttribute = SourceMetadata->GetConstAttribute(LocalSourceAttribute);
+			check(OriginalAttribute);
 
-			if (SampledData->Metadata->HasAttribute(LocalSourceAttribute))
-			{
-				if (!SampledData->Metadata->CopyExistingAttribute(LocalSourceAttribute, DestinationAttribute))
-				{
-					PCGE_LOG(Warning, GraphAndLog, FText::Format(LOCTEXT("FailedCopyToNewAttribute", "Failed to copy to new attribute '{0}'"), FText::FromName(DestinationAttribute)));
-				}
-			}
-			else
+			// We only copy entries/values if we copy from the input spatial metadata
+			// If it is from the source attribute set, we don't copy (and all points will have the same default value, value from the attribute set)
+			const bool bCopyEntriesAndValues = (SourceAttributeSet == nullptr);
+			if (!SampledData->Metadata->CopyAttribute(OriginalAttribute, DestinationAttribute, /*bKeepParent=*/ false, /*bCopyEntries=*/ bCopyEntriesAndValues, /*bCopyValues=*/ bCopyEntriesAndValues))
 			{
-				SampledData->Metadata->CopyAttribute(SourceMetadata, LocalSourceAttribute, DestinationAttribute);
+				PCGE_LOG(Warning, GraphAndLog, FText::Format(LOCTEXT("FailedCopyToNewAttribute", "Failed to copy to new attribute '{0}'"), FText::FromName(DestinationAttribute)));
 			}
 
 			continue;
@@ -231,6 +230,12 @@ bool FPCGMetadataOperationElement::ExecuteInternal(FPCGContext* Context) const
 			continue;
 		}
 
+		if (OutputAccessor->IsReadOnly())
+		{
+			PCGE_LOG(Warning, GraphAndLog, FText::Format(LOCTEXT("OutputAccessorIsReadOnly", "Attribute/Property '{0}' is read only."), OutputTarget.GetDisplayText()));
+			continue;
+		}
+
 		// Final verification, if we can put the value of input into output
 		if (!PCG::Private::IsBroadcastable(InputAccessor->GetUnderlyingType(), OutputAccessor->GetUnderlyingType()))
 		{
