[PCG] Fix unwanted change that went through

#rnx

diff --git a/Engine/Plugins/Experimental/PCG/Source/PCG/Private/Elements/Metadata/PCGMetadataElement.cpp b/Engine/Plugins/Experimental/PCG/Source/PCG/Private/Elements/Metadata/PCGMetadataElement.cpp
--- a/Engine/Plugins/Experimental/PCG/Source/PCG/Private/Elements/Metadata/PCGMetadataElement.cpp
+++ b/Engine/Plugins/Experimental/PCG/Source/PCG/Private/Elements/Metadata/PCGMetadataElement.cpp
@@ -230,12 +230,6 @@ bool FPCGMetadataOperationElement::ExecuteInternal(FPCGContext* Context) const
 			continue;
 		}
 
-		if (OutputAccessor->IsReadOnly())
-		{
-			PCGE_LOG(Warning, GraphAndLog, FText::Format(LOCTEXT("OutputAccessorIsReadOnly", "Attribute/Property '{0}' is read only."), OutputTarget.GetDisplayText()));
-			continue;
-		}
-
 		// Final verification, if we can put the value of input into output
 		if (!PCG::Private::IsBroadcastable(InputAccessor->GetUnderlyingType(), OutputAccessor->GetUnderlyingType()))
 		{
