Fixed pin type on the merge node.

#rnx

diff --git a/Engine/Plugins/Experimental/PCG/Source/PCG/Private/Elements/PCGMergeElement.cpp b/Engine/Plugins/Experimental/PCG/Source/PCG/Private/Elements/PCGMergeElement.cpp
--- a/Engine/Plugins/Experimental/PCG/Source/PCG/Private/Elements/PCGMergeElement.cpp
+++ b/Engine/Plugins/Experimental/PCG/Source/PCG/Private/Elements/PCGMergeElement.cpp
@@ -28,7 +28,7 @@ TArray<FPCGPinProperties> UPCGMergeSettings::InputPinProperties() const
 TArray<FPCGPinProperties> UPCGMergeSettings::OutputPinProperties() const
 {
 	TArray<FPCGPinProperties> PinProperties;
-	PinProperties.Emplace(PCGPinConstants::DefaultOutputLabel, EPCGDataType::Point, /*bAllowMultipleConnections=*/false);
+	PinProperties.Emplace(PCGPinConstants::DefaultOutputLabel, EPCGDataType::Point, /*bAllowMultipleConnections=*/false, /*bAllowMultipleData=*/false);
 
 	return PinProperties;
 }
