[PCG] Added FName support for create / set / get metadata functions

diff --git a/Engine/Plugins/Experimental/PCG/Source/PCG/Private/Metadata/PCGMetadata.cpp b/Engine/Plugins/Experimental/PCG/Source/PCG/Private/Metadata/PCGMetadata.cpp
--- a/Engine/Plugins/Experimental/PCG/Source/PCG/Private/Metadata/PCGMetadata.cpp
+++ b/Engine/Plugins/Experimental/PCG/Source/PCG/Private/Metadata/PCGMetadata.cpp
@@ -503,6 +503,11 @@ void UPCGMetadata::CreateStringAttribute(FName AttributeName, FString DefaultVal
 	CreateAttribute<FString>(AttributeName, DefaultValue, bAllowsInterpolation, bOverrideParent);
 }
 
+void UPCGMetadata::CreateNameAttribute(FName AttributeName, FName DefaultValue, bool bAllowsInterpolation, bool bOverrideParent)
+{
+	CreateAttribute<FName>(AttributeName, DefaultValue, bAllowsInterpolation, bOverrideParent);
+}
+
 void UPCGMetadata::CreateBoolAttribute(FName AttributeName, bool DefaultValue, bool bAllowsInterpolation, bool bOverrideParent)
 {
 	CreateAttribute<bool>(AttributeName, DefaultValue, bAllowsInterpolation, bOverrideParent);
diff --git a/Engine/Plugins/Experimental/PCG/Source/PCG/Private/Metadata/PCGMetadataAccessor.cpp b/Engine/Plugins/Experimental/PCG/Source/PCG/Private/Metadata/PCGMetadataAccessor.cpp
--- a/Engine/Plugins/Experimental/PCG/Source/PCG/Private/Metadata/PCGMetadataAccessor.cpp
+++ b/Engine/Plugins/Experimental/PCG/Source/PCG/Private/Metadata/PCGMetadataAccessor.cpp
@@ -385,6 +385,16 @@ void UPCGMetadataAccessorHelpers::SetStringAttribute(FPCGPoint& Point, UPCGMetad
 	SetAttribute(Point.MetadataEntry, Metadata, AttributeName, Value);
 }
 
+FName UPCGMetadataAccessorHelpers::GetNameAttribute(const FPCGPoint& Point, const UPCGMetadata* Metadata, FName AttributeName)
+{
+	return GetAttribute<FName>(Point.MetadataEntry, Metadata, AttributeName);
+}
+
+void UPCGMetadataAccessorHelpers::SetNameAttribute(FPCGPoint& Point, UPCGMetadata* Metadata, FName AttributeName, const FName& Value)
+{
+	SetAttribute(Point.MetadataEntry, Metadata, AttributeName, Value);
+}
+
 bool UPCGMetadataAccessorHelpers::GetBoolAttribute(const FPCGPoint& Point, const UPCGMetadata* Metadata, FName AttributeName)
 {
 	return GetAttribute<bool>(Point.MetadataEntry, Metadata, AttributeName);
diff --git a/Engine/Plugins/Experimental/PCG/Source/PCG/Public/Metadata/PCGMetadata.h b/Engine/Plugins/Experimental/PCG/Source/PCG/Public/Metadata/PCGMetadata.h
--- a/Engine/Plugins/Experimental/PCG/Source/PCG/Public/Metadata/PCGMetadata.h
+++ b/Engine/Plugins/Experimental/PCG/Source/PCG/Public/Metadata/PCGMetadata.h
@@ -118,6 +118,9 @@ public:
 	UFUNCTION(BlueprintCallable, Category = "PCG|Metadata")
 	void CreateStringAttribute(FName AttributeName, FString DefaultValue, bool bAllowsInterpolation, bool bOverrideParent = true);
 
+	UFUNCTION(BlueprintCallable, Category = "PCG|Metadata")
+	void CreateNameAttribute(FName AttributeName, FName DefaultValue, bool bAllowsInterpolation, bool bOverrideParent = true);
+
 	UFUNCTION(BlueprintCallable, Category = "PCG|Metadata")
 	void CreateBoolAttribute(FName AttributeName, bool DefaultValue, bool bAllowsInterpolation, bool bOverrideParent = true);
 
diff --git a/Engine/Plugins/Experimental/PCG/Source/PCG/Public/Metadata/PCGMetadataAccessor.h b/Engine/Plugins/Experimental/PCG/Source/PCG/Public/Metadata/PCGMetadataAccessor.h
--- a/Engine/Plugins/Experimental/PCG/Source/PCG/Public/Metadata/PCGMetadataAccessor.h
+++ b/Engine/Plugins/Experimental/PCG/Source/PCG/Public/Metadata/PCGMetadataAccessor.h
@@ -173,6 +173,12 @@ public:
 
 	UFUNCTION(BlueprintCallable, Category = "PCG|Metadata", meta = (ScriptMethod))
 	static void SetStringAttribute(UPARAM(ref) FPCGPoint& Point, UPCGMetadata* Metadata, FName AttributeName, const FString& Value);
+	
+	UFUNCTION(BlueprintCallable, Category = "PCG|Metadata", meta = (ScriptMethod))
+	static FName GetNameAttribute(const FPCGPoint& Point, const UPCGMetadata* Metadata, FName AttributeName);
+	
+	UFUNCTION(BlueprintCallable, Category = "PCG|Metadata", meta = (ScriptMethod))
+	static void SetNameAttribute(UPARAM(ref) FPCGPoint& Point, UPCGMetadata* Metadata, FName AttributeName, const FName& Value);
 
 	UFUNCTION(BlueprintCallable, Category = "PCG|Metadata", meta = (ScriptMethod))
 	static bool GetBoolAttribute(const FPCGPoint& Point, const UPCGMetadata* Metadata, FName AttributeName);
