[PCG] Read-only property extractor on output target are now throwing an error

diff --git a/Engine/Plugins/Experimental/PCG/Source/PCG/Private/Elements/Metadata/PCGMetadataElement.cpp b/Engine/Plugins/Experimental/PCG/Source/PCG/Private/Elements/Metadata/PCGMetadataElement.cpp
--- a/Engine/Plugins/Experimental/PCG/Source/PCG/Private/Elements/Metadata/PCGMetadataElement.cpp
+++ b/Engine/Plugins/Experimental/PCG/Source/PCG/Private/Elements/Metadata/PCGMetadataElement.cpp
@@ -230,6 +230,12 @@ bool FPCGMetadataOperationElement::ExecuteInternal(FPCGContext* Context) const
 			continue;
 		}
 
+		if (OutputAccessor->IsReadOnly())
+		{
+			PCGE_LOG(Warning, GraphAndLog, FText::Format(LOCTEXT("OutputAccessorIsReadOnly", "Attribute/Property '{0}' is read only."), OutputTarget.GetDisplayText()));
+			continue;
+		}
+
 		// Final verification, if we can put the value of input into output
 		if (!PCG::Private::IsBroadcastable(InputAccessor->GetUnderlyingType(), OutputAccessor->GetUnderlyingType()))
 		{
diff --git a/Engine/Plugins/Experimental/PCG/Source/PCG/Private/Elements/Metadata/PCGMetadataOpElementBase.cpp b/Engine/Plugins/Experimental/PCG/Source/PCG/Private/Elements/Metadata/PCGMetadataOpElementBase.cpp
--- a/Engine/Plugins/Experimental/PCG/Source/PCG/Private/Elements/Metadata/PCGMetadataOpElementBase.cpp
+++ b/Engine/Plugins/Experimental/PCG/Source/PCG/Private/Elements/Metadata/PCGMetadataOpElementBase.cpp
@@ -325,6 +325,12 @@ bool FPCGMetadataElementBase::ExecuteInternal(FPCGContext* Context) const
 			return false;
 		}
 
+		if (OperationData.OutputAccessors[OutputIndex]->IsReadOnly())
+		{
+			PCGE_LOG(Error, GraphAndLog, FText::Format(LOCTEXT("OutputAccessorIsReadOnly", "Attribute/Property '{0}' is read only."), OutputTarget.GetDisplayText()));
+			return false;
+		}
+
 		OperationData.OutputKeys[OutputIndex] = PCGAttributeAccessorHelpers::CreateKeys(Cast<UPCGData>(OutputData.Data), OutputTarget);
 
 		return OperationData.OutputKeys[OutputIndex].IsValid();
diff --git a/Engine/Plugins/Experimental/PCG/Source/PCG/Private/MatchAndSet/PCGMatchAndSetByAttribute.cpp b/Engine/Plugins/Experimental/PCG/Source/PCG/Private/MatchAndSet/PCGMatchAndSetByAttribute.cpp
--- a/Engine/Plugins/Experimental/PCG/Source/PCG/Private/MatchAndSet/PCGMatchAndSetByAttribute.cpp
+++ b/Engine/Plugins/Experimental/PCG/Source/PCG/Private/MatchAndSet/PCGMatchAndSetByAttribute.cpp
@@ -129,6 +129,18 @@ void UPCGMatchAndSetByAttribute::MatchAndSet_Implementation(
 	TUniquePtr<IPCGAttributeAccessor> SetTargetAccessor = PCGAttributeAccessorHelpers::CreateAccessor(OutPointData, SetTarget);
 	TUniquePtr<IPCGAttributeAccessorKeys> SetTargetKeys = PCGAttributeAccessorHelpers::CreateKeys(OutPointData, SetTarget);
 
+	if (!SetTargetAccessor.IsValid() || !SetTargetKeys.IsValid())
+	{
+		PCGE_LOG_C(Warning, GraphAndLog, &Context, LOCTEXT("FailedToCreateTargetAccessor", "Failed to create output accessor or iterator"));
+		return;
+	}
+
+	if (SetTargetAccessor->IsReadOnly())
+	{
+		PCGE_LOG_C(Warning, GraphAndLog, &Context, FText::Format(LOCTEXT("SetTargetAccessorIsReadOnly", "Attribute/Property '{0}' is read only."), SetTarget.GetDisplayText()));
+		return;
+	}
+
 	auto MatchAndSetOperation = [&InputAccessor, &InputKeys, &EntryMatchSourceAccessors, &EntrySetValueAccessors, &EntryKeys, &SetTargetAccessor, &SetTargetKeys](auto MatchDummy)
 	{
 		using MatchType = decltype(MatchDummy);
diff --git a/Engine/Plugins/Experimental/PCG/Source/PCG/Private/MatchAndSet/PCGMatchAndSetWeighted.cpp b/Engine/Plugins/Experimental/PCG/Source/PCG/Private/MatchAndSet/PCGMatchAndSetWeighted.cpp
--- a/Engine/Plugins/Experimental/PCG/Source/PCG/Private/MatchAndSet/PCGMatchAndSetWeighted.cpp
+++ b/Engine/Plugins/Experimental/PCG/Source/PCG/Private/MatchAndSet/PCGMatchAndSetWeighted.cpp
@@ -105,6 +105,18 @@ void UPCGMatchAndSetWeighted::MatchAndSet_Implementation(
 	TUniquePtr<IPCGAttributeAccessor> SetTargetAccessor = PCGAttributeAccessorHelpers::CreateAccessor(OutPointData, SetTarget);
 	TUniquePtr<IPCGAttributeAccessorKeys> SetTargetKeys = PCGAttributeAccessorHelpers::CreateKeys(OutPointData, SetTarget);
 
+	if (!SetTargetAccessor.IsValid() || !SetTargetKeys.IsValid())
+	{
+		PCGE_LOG_C(Warning, GraphAndLog, &Context, LOCTEXT("FailedToCreateTargetAccessor", "Failed to create output accessor or iterator"));
+		return;
+	}
+
+	if (SetTargetAccessor->IsReadOnly())
+	{
+		PCGE_LOG_C(Warning, GraphAndLog, &Context, FText::Format(LOCTEXT("SetTargetAccessorIsReadOnly", "Attribute/Property '{0}' is read only."), SetTarget.GetDisplayText()));
+		return;
+	}
+
 	auto GetValueIndexFromElementIndex = [&InPoints, InSettings, SourceComponent, &CumulativeWeights, TotalWeight](int32 Index) -> int32
 	{
 		const FPCGPoint& Point = InPoints[Index];
diff --git a/Engine/Plugins/Experimental/PCG/Source/PCG/Private/MatchAndSet/PCGMatchAndSetWeightedByCategory.cpp b/Engine/Plugins/Experimental/PCG/Source/PCG/Private/MatchAndSet/PCGMatchAndSetWeightedByCategory.cpp
--- a/Engine/Plugins/Experimental/PCG/Source/PCG/Private/MatchAndSet/PCGMatchAndSetWeightedByCategory.cpp
+++ b/Engine/Plugins/Experimental/PCG/Source/PCG/Private/MatchAndSet/PCGMatchAndSetWeightedByCategory.cpp
@@ -183,6 +183,18 @@ void UPCGMatchAndSetWeightedByCategory::MatchAndSet_Implementation(
 	TUniquePtr<IPCGAttributeAccessor> SetTargetAccessor = PCGAttributeAccessorHelpers::CreateAccessor(OutPointData, SetTarget);
 	TUniquePtr<IPCGAttributeAccessorKeys> SetTargetKeys = PCGAttributeAccessorHelpers::CreateKeys(OutPointData, SetTarget);
 
+	if (!SetTargetAccessor.IsValid() || !SetTargetKeys.IsValid())
+	{
+		PCGE_LOG_C(Warning, GraphAndLog, &Context, LOCTEXT("FailedToCreateTargetAccessor", "Failed to create output accessor or iterator"));
+		return;
+	}
+
+	if (SetTargetAccessor->IsReadOnly())
+	{
+		PCGE_LOG_C(Warning, GraphAndLog, &Context, FText::Format(LOCTEXT("SetTargetAccessorIsReadOnly", "Attribute/Property '{0}' is read only."), SetTarget.GetDisplayText()));
+		return;
+	}
+
 	auto GetValueIndex = [this, InSettings, SourceComponent, &InPoints](int32 CategoryIndex, int32 PointIndex) -> int32
 	{
 		const FPCGMatchAndSetWeightedByCategoryEntryList& Category = Categories[CategoryIndex];
diff --git a/Engine/Plugins/Experimental/PCG/Source/PCG/Private/Metadata/Accessors/PCGAttributeAccessorHelpers.cpp b/Engine/Plugins/Experimental/PCG/Source/PCG/Private/Metadata/Accessors/PCGAttributeAccessorHelpers.cpp
--- a/Engine/Plugins/Experimental/PCG/Source/PCG/Private/Metadata/Accessors/PCGAttributeAccessorHelpers.cpp
+++ b/Engine/Plugins/Experimental/PCG/Source/PCG/Private/Metadata/Accessors/PCGAttributeAccessorHelpers.cpp
@@ -412,6 +412,13 @@ TUniquePtr<IPCGAttributeAccessor> PCGAttributeAccessorHelpers::CreateAccessor(UP
 		}
 	}
 
+	// Can't create mutable accessor that are read only.
+	if (Accessor && Accessor->IsReadOnly())
+	{
+		UE_LOG(LogPCG, Error, TEXT("[PCGAttributeAccessorHelpers::CreateAccessor] Accessor is read-only, therefore not mutable."));
+		return TUniquePtr<IPCGAttributeAccessor>();
+	}
+
 	return Accessor;
 }
 
