[PCG] Make the seed dependent by default on the node type name for new nodes, but keep the serialized one for already created nodes

* With this, we keep already setup seeds for existing graph, to avoid breaking their results.
* New nodes will have a new seed set, to avoid correlation between different nodes.
* Factorize the logic of `IsNewObjectAndNotDefault` in helpers and replace it where it was duplicated.

diff --git a/Engine/Plugins/Experimental/PCG/Source/PCG/Private/Elements/PCGTypedGetter.cpp b/Engine/Plugins/Experimental/PCG/Source/PCG/Private/Elements/PCGTypedGetter.cpp
--- a/Engine/Plugins/Experimental/PCG/Source/PCG/Private/Elements/PCGTypedGetter.cpp
+++ b/Engine/Plugins/Experimental/PCG/Source/PCG/Private/Elements/PCGTypedGetter.cpp
@@ -2,6 +2,8 @@
 
 #include "Elements/PCGTypedGetter.h"
 
+#include "Helpers/PCGHelpers.h"
+
 #include "Landscape.h"
 
 #include UE_INLINE_GENERATED_CPP_BY_NAME(PCGTypedGetter)
@@ -14,8 +16,8 @@ UPCGGetLandscapeSettings::UPCGGetLandscapeSettings()
 	Mode = EPCGGetDataFromActorMode::ParseActorComponents;
 
 	// We want to apply different defaults to newly placed nodes. We detect new object if they are not a default object/archetype
-	// and/or they do not need load. Followed similar pattern to UPCGComponent::PostInitProperties().
-	if (!HasAnyFlags(RF_ClassDefaultObject | RF_NeedLoad | RF_NeedPostLoad))
+	// and/or they do not need load.
+	if (PCGHelpers::IsNewObjectAndNotDefault(this))
 	{
 		// This setup replicates what was implemented on the Landscape input node pin
 		ActorSelector.ActorFilter = EPCGActorFilter::AllWorldActors;
diff --git a/Engine/Plugins/Experimental/PCG/Source/PCG/Private/Helpers/PCGHelpers.cpp b/Engine/Plugins/Experimental/PCG/Source/PCG/Private/Helpers/PCGHelpers.cpp
--- a/Engine/Plugins/Experimental/PCG/Source/PCG/Private/Helpers/PCGHelpers.cpp
+++ b/Engine/Plugins/Experimental/PCG/Source/PCG/Private/Helpers/PCGHelpers.cpp
@@ -391,4 +391,18 @@ namespace PCGHelpers
 		}
 	}
 #endif
+
+	bool IsNewObjectAndNotDefault(const UObject* InObject, bool bCheckHierarchy)
+	{
+		const UObject* CurrentInspectedObject = InObject;
+		while (bCheckHierarchy && CurrentInspectedObject && CurrentInspectedObject->HasAnyFlags(RF_DefaultSubObject))
+		{
+			CurrentInspectedObject = CurrentInspectedObject->GetOuter();
+		}
+
+		// We detect new objects if they are not a default object/archetype and/or they do not need load.
+		// In some cases, where the component is a default sub object (like APCGVolume), it has no loading flags
+		// even if it is loading, so we use the outer found above.
+		return CurrentInspectedObject && !CurrentInspectedObject->HasAnyFlags(RF_ClassDefaultObject | RF_NeedLoad | RF_NeedPostLoad);
+	}
 }
diff --git a/Engine/Plugins/Experimental/PCG/Source/PCG/Private/PCGComponent.cpp b/Engine/Plugins/Experimental/PCG/Source/PCG/Private/PCGComponent.cpp
--- a/Engine/Plugins/Experimental/PCG/Source/PCG/Private/PCGComponent.cpp
+++ b/Engine/Plugins/Experimental/PCG/Source/PCG/Private/PCGComponent.cpp
@@ -1076,21 +1076,9 @@ void UPCGComponent::PostInitProperties()
 	GraphInstance->OnGraphChangedDelegate.AddUObject(this, &UPCGComponent::OnGraphChanged);
 #endif // WITH_EDITOR
 
-	// Note that if the component is a default sub object, we find the first outer that is not a sub object.
-	UObject* CurrentInspectedObject = this;
-	while (CurrentInspectedObject && CurrentInspectedObject->HasAnyFlags(RF_DefaultSubObject))
-	{
-		CurrentInspectedObject = CurrentInspectedObject->GetOuter();
-	}
-
-	// We detect new object if they are not a default object/archetype and/or they do not need load.
-	// In some cases, were the component is a default sub object (like APCGVolume), it has no loading flags
-	// even if it is loading, so we use the outer found above.
-	const bool bIsNewObject = CurrentInspectedObject && !CurrentInspectedObject->HasAnyFlags(RF_ClassDefaultObject | RF_NeedLoad | RF_NeedPostLoad);
-
 #if WITH_EDITOR
 	// Force bIsPartitioned at false for new objects
-	if (bIsNewObject)
+	if (PCGHelpers::IsNewObjectAndNotDefault(this, /*bCheckHierarchy=*/true))
 	{
 		bIsPartitioned = false;
 	}
diff --git a/Engine/Plugins/Experimental/PCG/Source/PCG/Private/PCGSettings.cpp b/Engine/Plugins/Experimental/PCG/Source/PCG/Private/PCGSettings.cpp
--- a/Engine/Plugins/Experimental/PCG/Source/PCG/Private/PCGSettings.cpp
+++ b/Engine/Plugins/Experimental/PCG/Source/PCG/Private/PCGSettings.cpp
@@ -80,15 +80,9 @@ void UPCGSettingsInterface::SetEnabled(bool bInEnabled)
 	}
 }
 
-UPCGSettings::UPCGSettings()
+uint32 UPCGSettings::GetTypeNameHash() const
 {
-	TypeNameHash = GetTypeHash(GetClass()->GetFName());
-}
-
-UPCGSettings::UPCGSettings(const FObjectInitializer& ObjectInitializer)
-	: Super(ObjectInitializer)
-{
-	TypeNameHash = GetTypeHash(GetClass()->GetFName());
+	return GetTypeHash(GetClass()->GetFName());
 }
 
 bool UPCGSettings::operator==(const UPCGSettings& Other) const
@@ -154,6 +148,15 @@ void UPCGSettings::PostInitProperties()
 	InitializeCachedOverridableParams();
 #endif //WITH_EDITOR
 
+	// For new objects, use type hash as a seed so we don't get multiple nodes doing randomness in an identical way (we had a case where
+	// a Point Subset followed by a Density Noise produced uniform densities because the random seeds/streams were correllated).
+	// It is still possible to chain nodes of the same type (Point Subset -> Point Subset) and see correllated behaviour, in which
+	// case the random seed on one of them should be changed in the node settings.
+	if (PCGHelpers::IsNewObjectAndNotDefault(this, /*bCheckHierarchy=*/true))
+	{
+		Seed = GetTypeNameHash();
+	}
+
 	CacheCrc();
 }
 
@@ -397,11 +400,7 @@ UPCGNode* UPCGSettings::CreateNode() const
 
 int UPCGSettings::GetSeed(const UPCGComponent* InSourceComponent) const
 {
-	// Mix in a type hash so we don't get multiple nodes doing randomness in an identical way (we had a case where
-	// a Point Subset followed by a Density Noise produced uniform densities because the random seeds/streams were correllated).
-	// It is still possible to chain nodes of the same type (Point Subset -> Point Subset) and see correllated behaviour, in which
-	// case the random seed on one of them should be changed in the node settings.
-	return !bUseSeed ? TypeNameHash : (InSourceComponent ? PCGHelpers::ComputeSeed(Seed, InSourceComponent->Seed, TypeNameHash) : Seed);
+	return !bUseSeed ? 42 : (InSourceComponent ? PCGHelpers::ComputeSeed(Seed, InSourceComponent->Seed) : Seed);
 }
 
 #if WITH_EDITOR
diff --git a/Engine/Plugins/Experimental/PCG/Source/PCG/Private/PCGVolume.cpp b/Engine/Plugins/Experimental/PCG/Source/PCG/Private/PCGVolume.cpp
--- a/Engine/Plugins/Experimental/PCG/Source/PCG/Private/PCGVolume.cpp
+++ b/Engine/Plugins/Experimental/PCG/Source/PCG/Private/PCGVolume.cpp
@@ -3,6 +3,7 @@
 #include "PCGVolume.h"
 #include "PCGComponent.h"
 #include "PCGGraph.h"
+#include "Helpers/PCGHelpers.h"
 
 #include "Components/BrushComponent.h"
 
@@ -13,13 +14,7 @@ APCGVolume::APCGVolume(const FObjectInitializer& ObjectInitializer)
 {
 	PCGComponent = ObjectInitializer.CreateDefaultSubobject<UPCGComponent>(this, TEXT("PCG Component"));
 
-	UObject* ThisObject = this;
-	while (ThisObject && ThisObject->HasAnyFlags(RF_DefaultSubObject))
-	{
-		ThisObject = ThisObject->GetOuter();
-	}
-
-	if (ThisObject && !ThisObject->HasAnyFlags(RF_ClassDefaultObject | RF_NeedLoad | RF_NeedPostLoad))
+	if (PCGHelpers::IsNewObjectAndNotDefault(this, /*bCheckHierarchy=*/true))
 	{
 		if (UBrushComponent* MyBrushComponent = GetBrushComponent())
 		{
diff --git a/Engine/Plugins/Experimental/PCG/Source/PCG/Public/Helpers/PCGHelpers.h b/Engine/Plugins/Experimental/PCG/Source/PCG/Public/Helpers/PCGHelpers.h
--- a/Engine/Plugins/Experimental/PCG/Source/PCG/Public/Helpers/PCGHelpers.h
+++ b/Engine/Plugins/Experimental/PCG/Source/PCG/Public/Helpers/PCGHelpers.h
@@ -49,4 +49,12 @@ namespace PCGHelpers
 	PCG_API void GatherDependencies(UObject* Object, TSet<TObjectPtr<UObject>>& OutDependencies, int32 MaxDepth = -1);
 	PCG_API void GatherDependencies(FProperty* Property, const void* InContainer, TSet<TObjectPtr<UObject>>& OutDependencies, int32 MaxDepth);
 #endif
+
+	/** 
+	* Check if an object is a new object and not the CDO.
+	*
+	* Some objects might not have the appropriate flags if they are embedded inside of other objects. 
+	* Use the bCheckHierarchy flag to true to go up the object hierarchy if you want to check for this situation.
+	*/
+	PCG_API bool IsNewObjectAndNotDefault(const UObject* InObject, bool bCheckHierarchy = false);
 };
diff --git a/Engine/Plugins/Experimental/PCG/Source/PCG/Public/PCGSettings.h b/Engine/Plugins/Experimental/PCG/Source/PCG/Public/PCGSettings.h
--- a/Engine/Plugins/Experimental/PCG/Source/PCG/Public/PCGSettings.h
+++ b/Engine/Plugins/Experimental/PCG/Source/PCG/Public/PCGSettings.h
@@ -154,9 +154,6 @@ class PCG_API UPCGSettings : public UPCGSettingsInterface
 	friend class UPCGSettingsInterface;
 
 public:
-	UPCGSettings();
-	UPCGSettings(const FObjectInitializer& ObjectInitializer);
-
 	// ~Begin UPCGData interface
 	virtual EPCGDataType GetDataType() const override { return EPCGDataType::Settings; }
 	// ~End UPCGData interface
@@ -252,7 +249,7 @@ public:
 	virtual EPCGDataType GetCurrentPinTypes(const UPCGPin* InPin) const;
 
 	UPROPERTY(EditAnywhere, BlueprintReadWrite, Category = Settings, meta=(EditCondition=bUseSeed, EditConditionHides, PCG_Overridable))
-	int Seed = 0xC35A9631; // random prime number
+	int Seed = 0xC35A9631; // Default seed is a random prime number, but will be overriden for new settings based on the class type name hash, making each settings class have a different default seed.
 
 	/** Warning - this is deprecated and will be removed soon since we have a Filter By Tag node for this specific purpose */
 	UPROPERTY(EditAnywhere, BlueprintReadWrite, Category = "Tags")
@@ -295,6 +292,9 @@ protected:
 	/** An additional custom version number that external system users can use to track versions. This version will be serialized into the asset and will be provided by UserDataVersion after load. */
 	virtual FGuid GetUserCustomVersionGuid() { return FGuid(); }
 
+	/** Can be overriden by child class if they ever got renamed to avoid changing the default seed for this one. Otherwise default is hash of the class name. */
+	virtual uint32 GetTypeNameHash() const;
+
 #if WITH_EDITOR
 	virtual void PostEditChangeProperty(FPropertyChangedEvent& PropertyChangedEvent) override;
 	virtual bool IsStructuralProperty(const FName& InPropertyName) const { return false; }
@@ -368,10 +368,6 @@ private:
 
 	/** The cached Crc for these settings. */
 	FPCGCrc CachedCrc;
-
-	/** Hash of the derived type name which is of the form 'PCGExecuteBlueprint'. */
-	UPROPERTY(Transient)
-	uint32 TypeNameHash = 0;
 };
 
 UCLASS(BlueprintType, ClassGroup = (Procedural))
