[PCG] Tidy up Get Texture Data warning

diff --git a/Engine/Plugins/Experimental/PCG/Source/PCG/Private/Elements/PCGTextureSampler.cpp b/Engine/Plugins/Experimental/PCG/Source/PCG/Private/Elements/PCGTextureSampler.cpp
--- a/Engine/Plugins/Experimental/PCG/Source/PCG/Private/Elements/PCGTextureSampler.cpp
+++ b/Engine/Plugins/Experimental/PCG/Source/PCG/Private/Elements/PCGTextureSampler.cpp
@@ -45,7 +45,7 @@ bool FPCGTextureSamplerElement::ExecuteInternal(FPCGContext* Context) const
 
 	if (!UPCGTextureData::IsSupported(Texture))
 	{
-		PCGE_LOG(Error, GraphAndLog, FText::Format(LOCTEXT("UnsupportedTextureFormat", "Texture '{0}' has unsupported texture format, currently supported formats are B8G8R8A8, PF_R8G8B8A8 and PF_G8"), FText::FromName(Texture->GetFName())));
+		PCGE_LOG(Error, GraphAndLog, FText::Format(LOCTEXT("UnsupportedTextureFormat", "Texture '{0}' has unsupported texture format, currently supported formats are B8G8R8A8, R8G8B8A8 and G8"), FText::FromName(Texture->GetFName())));
 		return true;
 	}
 
