Moving a method to the class' cpp file so forward declarations can work.

diff --git a/Engine/Plugins/Experimental/PCG/Source/PCG/Private/Elements/PCGDebugElement.cpp b/Engine/Plugins/Experimental/PCG/Source/PCG/Private/Elements/PCGDebugElement.cpp
--- a/Engine/Plugins/Experimental/PCG/Source/PCG/Private/Elements/PCGDebugElement.cpp
+++ b/Engine/Plugins/Experimental/PCG/Source/PCG/Private/Elements/PCGDebugElement.cpp
@@ -231,6 +231,11 @@ FPCGElementPtr UPCGDebugSettings::CreateElement() const
 	return MakeShared<FPCGDebugElement>();
 }
 
+TArray<FPCGPinProperties> UPCGDebugSettings::InputPinProperties() const
+{
+	return Super::DefaultPointInputPinProperties();
+}
+
 TArray<FPCGPinProperties> UPCGDebugSettings::OutputPinProperties() const
 {
 	return TArray<FPCGPinProperties>();
diff --git a/Engine/Plugins/Experimental/PCG/Source/PCG/Public/Elements/PCGDebugElement.h b/Engine/Plugins/Experimental/PCG/Source/PCG/Public/Elements/PCGDebugElement.h
--- a/Engine/Plugins/Experimental/PCG/Source/PCG/Public/Elements/PCGDebugElement.h
+++ b/Engine/Plugins/Experimental/PCG/Source/PCG/Public/Elements/PCGDebugElement.h
@@ -25,7 +25,7 @@ public:
 #endif
 
 protected:
-	virtual TArray<FPCGPinProperties> InputPinProperties() const override { return Super::DefaultPointInputPinProperties(); }
+	virtual TArray<FPCGPinProperties> InputPinProperties() const override;
 	virtual TArray<FPCGPinProperties> OutputPinProperties() const override;
 
 	virtual FPCGElementPtr CreateElement() const override;
