[PCG] Remove duplicated property definitions which will fail in the name matching

diff --git a/Engine/Plugins/Experimental/PCG/Source/PCG/Private/PCGNode.cpp b/Engine/Plugins/Experimental/PCG/Source/PCG/Private/PCGNode.cpp
--- a/Engine/Plugins/Experimental/PCG/Source/PCG/Private/PCGNode.cpp
+++ b/Engine/Plugins/Experimental/PCG/Source/PCG/Private/PCGNode.cpp
@@ -110,14 +110,14 @@ void UPCGNode::ApplyDeprecation()
 
 		if (!OutputPin)
 		{
-			UE_LOG(LogPCG, Error, TEXT("Unable to apply deprecation on outbound edge on node %s - can't find output pin %s"), *GetFName().ToString(), *OutboundEdge->InboundLabel_DEPRECATED.ToString());
+			UE_LOG(LogPCG, Error, TEXT("Unable to apply deprecation on outbound edge on node '%s' - can't find output pin '%s'"), *GetFName().ToString(), *OutboundEdge->InboundLabel_DEPRECATED.ToString());
 			continue;
 		}
 
 		UPCGNode* OtherNode = OutboundEdge->OutboundNode_DEPRECATED;
 		if (!OtherNode)
 		{
-			UE_LOG(LogPCG, Error, TEXT("Unable to apply deprecation on outbound edge on node %s - can't find other node"), *GetFName().ToString());
+			UE_LOG(LogPCG, Error, TEXT("Unable to apply deprecation on outbound edge on node '%s' - can't find other node"), *GetFName().ToString());
 			continue;
 		}
 
@@ -674,8 +674,29 @@ EPCGChangeType UPCGNode::UpdatePins(TFunctionRef<UPCGPin*(UPCGNode*)> PinAllocat
 		return EPCGChangeType::Edge | EPCGChangeType::Node;
 	}
 	
-	const TArray<FPCGPinProperties> InboundPinProperties = Settings->AllInputPinProperties();
-	const TArray<FPCGPinProperties> OutboundPinProperties = Settings->AllOutputPinProperties();
+	TArray<FPCGPinProperties> InboundPinProperties = Settings->AllInputPinProperties();
+	TArray<FPCGPinProperties> OutboundPinProperties = Settings->AllOutputPinProperties();
+
+	auto RemoveDuplicates = [this](TArray<FPCGPinProperties>& Properties)
+	{
+		for (int32 i = Properties.Num() - 2; i >= 0; --i)
+		{
+			for (int32 j = i + 1; j < Properties.Num(); ++j)
+			{
+				if (Properties[i].Label == Properties[j].Label)
+				{
+					UE_LOG(LogPCG, Warning, TEXT("UpdatePins: Pin properties from the settings on node '%s' contained a duplicate pin '%s', removing this pin properties."),
+						*GetFName().ToString(), *Properties[i].Label.ToString());
+
+					// Remove but preserve order
+					Properties.RemoveAt(j);
+					break;
+				}
+			}
+		}
+	};
+	RemoveDuplicates(InboundPinProperties);
+	RemoveDuplicates(OutboundPinProperties);
 
 	auto UpdatePins = [this, &PinAllocator, &TouchedNodes](TArray<UPCGPin*>& Pins, const TArray<FPCGPinProperties>& PinProperties)
 	{
