[PCG] Set Get Actor Data output pin type based on selected Mode

diff --git a/Engine/Plugins/Experimental/PCG/Source/PCG/Private/Elements/PCGDataFromActor.cpp b/Engine/Plugins/Experimental/PCG/Source/PCG/Private/Elements/PCGDataFromActor.cpp
--- a/Engine/Plugins/Experimental/PCG/Source/PCG/Private/Elements/PCGDataFromActor.cpp
+++ b/Engine/Plugins/Experimental/PCG/Source/PCG/Private/Elements/PCGDataFromActor.cpp
@@ -35,6 +35,25 @@ FPCGElementPtr UPCGDataFromActorSettings::CreateElement() const
 	return MakeShared<FPCGDataFromActorElement>();
 }
 
+EPCGDataType UPCGDataFromActorSettings::GetCurrentPinTypes(const UPCGPin* InPin) const
+{
+	check(InPin);
+
+	if (InPin->IsOutputPin())
+	{
+		if (Mode == EPCGGetDataFromActorMode::GetSinglePoint)
+		{
+			return EPCGDataType::Point;
+		}
+		else if (Mode == EPCGGetDataFromActorMode::GetDataFromProperty)
+		{
+			return EPCGDataType::Param;
+		}
+	}
+
+	return Super::GetCurrentPinTypes(InPin);
+}
+
 TArray<FPCGPinProperties> UPCGDataFromActorSettings::OutputPinProperties() const
 {
 	TArray<FPCGPinProperties> Pins = Super::OutputPinProperties();
diff --git a/Engine/Plugins/Experimental/PCG/Source/PCG/Public/Elements/PCGDataFromActor.h b/Engine/Plugins/Experimental/PCG/Source/PCG/Public/Elements/PCGDataFromActor.h
--- a/Engine/Plugins/Experimental/PCG/Source/PCG/Public/Elements/PCGDataFromActor.h
+++ b/Engine/Plugins/Experimental/PCG/Source/PCG/Public/Elements/PCGDataFromActor.h
@@ -33,7 +33,9 @@ public:
 	virtual FText GetNodeTooltipText() const override;
 	virtual EPCGSettingsType GetType() const override { return EPCGSettingsType::Spatial; }
 	virtual void GetTrackedActorTags(FPCGTagToSettingsMap& OutTagToSettings, TArray<TObjectPtr<const UPCGGraph>>& OutVisitedGraphs) const override;
+	virtual bool HasDynamicPins() const override { return true; }
 #endif
+	virtual EPCGDataType GetCurrentPinTypes(const UPCGPin* InPin) const override;
 
 	virtual TArray<FPCGPinProperties> InputPinProperties() const override { return TArray<FPCGPinProperties>(); }
 	virtual TArray<FPCGPinProperties> OutputPinProperties() const override;
@@ -82,6 +84,7 @@ public:
 	virtual FPCGContext* Initialize(const FPCGDataCollection& InputData, TWeakObjectPtr<UPCGComponent> SourceComponent, const UPCGNode* Node) override;
 	virtual bool CanExecuteOnlyOnMainThread(FPCGContext* Context) const override { return true; }
 	virtual bool IsCacheable(const UPCGSettings* InSettings) const override { return false; }
+
 protected:
 	virtual bool ExecuteInternal(FPCGContext* Context) const;
 	void GatherWaitTasks(AActor* FoundActor, FPCGContext* Context, TArray<FPCGTaskId>& OutWaitTasks) const;
