[PCG] Remove workarounds for GraphParameters with the latest version of StructUtils plugin

diff --git a/Engine/Plugins/Experimental/PCG/Source/PCG/Private/PCGComponent.cpp b/Engine/Plugins/Experimental/PCG/Source/PCG/Private/PCGComponent.cpp
--- a/Engine/Plugins/Experimental/PCG/Source/PCG/Private/PCGComponent.cpp
+++ b/Engine/Plugins/Experimental/PCG/Source/PCG/Private/PCGComponent.cpp
@@ -1181,91 +1181,6 @@ void UPCGComponent::RefreshAfterGraphChanged(UPCGGraphInterface* InGraph, bool b
 }
 
 #if WITH_EDITOR
-/**
-* Temporary workaround, while we are waiting for UE-182059
-* On the UI side of StructUtils, property widgets are referencing a temporary structure, not the one that is owned by the Graph/GraphInstance.
-* After the value is modified and the callback is propagated, the temporary structure is modified, but the one owned by the Graph is not.
-* If we keep it like this, PostEditChangeChainProperty will call for a reconstruction of the component if it is on a Blueprint
-* and we will go through the process of copying and transfering data between the 2 components (old and new).
-* Because of that, we copy the old structure, not modified. And when the structure is actually modified, the callback is called on the old component,
-* already marked trash, and we lose the change.
-* The workaround is to detect when a value from our parameters is changed, and just discard the call to the super method, avoiding a reconstruction script.
-* Note that it should be OK to do so in the current situation, because there is another callback fired just after the copy of the temp 
-* structure to the real one, and on this callback we are calling the super method and go through construction script.
-* We do this workaround for pre and post edit, since we can run into trouble if a pre edit was processed, but post was not.
-*/
-bool UPCGComponent::IsChangingGraphInstanceParameterValue(FEditPropertyChain& InEditPropertyChain) const
-{
-	if (!GraphInstance)
-	{
-		return false;
-	}
-
-	FEditPropertyChain::TDoubleLinkedListNode* PropertyNode = InEditPropertyChain.GetActiveNode();
-
-	if (!PropertyNode)
-	{
-		return false;
-	}
-
-	constexpr int32 PropertyPathSize = 4;
-	static const FName PropertyPath[PropertyPathSize] = {
-		GET_MEMBER_NAME_CHECKED(UPCGComponent, GraphInstance),
-		GET_MEMBER_NAME_CHECKED(UPCGGraphInstance, ParametersOverrides),
-		GET_MEMBER_NAME_CHECKED(FPCGOverrideInstancedPropertyBag, Parameters),
-		FName(TEXT("Value"))
-	};
-
-	int32 Index = 0;
-
-	while (Index < PropertyPathSize && PropertyNode)
-	{
-		FProperty* Property = PropertyNode->GetValue();
-		if (Property && Property->GetFName() == PropertyPath[Index])
-		{
-			PropertyNode = PropertyNode->GetNextNode();
-			++Index;
-		}
-		else
-		{
-			PropertyNode = nullptr;
-		}
-	}
-
-	if (Index == PropertyPathSize && PropertyNode)
-	{
-		FProperty* Property = PropertyNode->GetValue();
-		if (Property && Property->GetOwnerStruct() == GraphInstance->ParametersOverrides.Parameters.GetPropertyBagStruct())
-		{
-			return true;
-		}
-	}
-
-	return false;
-}
-
-void UPCGComponent::PreEditChange(FEditPropertyChain& PropertyAboutToChange)
-{
-	if (IsChangingGraphInstanceParameterValue(PropertyAboutToChange))
-	{
-		// Skip pre edit, cf IsChangingGraphInstanceParameterValue comment
-		return;
-	}
-
-	UObject::PreEditChange(PropertyAboutToChange);
-}
-
-void UPCGComponent::PostEditChangeChainProperty(FPropertyChangedChainEvent& PropertyChangedEvent)
-{
-	if (IsChangingGraphInstanceParameterValue(PropertyChangedEvent.PropertyChain))
-	{
-		// Skip post edit, cf IsChangingGraphInstanceParameterValue comment
-		return;
-	}
-
-	Super::PostEditChangeChainProperty(PropertyChangedEvent);
-}
-
 void UPCGComponent::PostEditChangeProperty(FPropertyChangedEvent& PropertyChangedEvent)
 {
 	Super::PostEditChangeProperty(PropertyChangedEvent);
diff --git a/Engine/Plugins/Experimental/PCG/Source/PCG/Private/PCGGraph.cpp b/Engine/Plugins/Experimental/PCG/Source/PCG/Private/PCGGraph.cpp
--- a/Engine/Plugins/Experimental/PCG/Source/PCG/Private/PCGGraph.cpp
+++ b/Engine/Plugins/Experimental/PCG/Source/PCG/Private/PCGGraph.cpp
@@ -823,17 +823,6 @@ void UPCGGraph::PreEditChange(FProperty* InProperty)
 		// We need to keep track of the number of properties, to detect if a property was added/removed/modified
 		NumberOfUserParametersPreEdit = UserParameters.GetNumPropertiesInBag();
 	}
-	else if (InProperty->GetOwnerStruct() == UserParameters.GetPropertyBagStruct())
-	{
-		// This is a bit unconventional, but we have to store the property name that has changed, because how Pre/Post Edit change is called.
-		// First PreEdit change is called twice. Once for the property changed (the one we want to track), and the second time on our UserParameters struct.
-		// Then PostEdit change is also called twice. But at the time of the first call (on the property we want to track), the value is not yet changed in memory of
-		// our UserParameters struct, it is still the old value. We need to wait the second call, PostEdit on our UserParameters struct, to have the memory updated with the new
-		// value. But then, we lost the name of our property.
-		// That's why we store the name here, and will use it in Post to know which property changed its value.
-		// TODO: It might be something that would need to change in StructUtils.
-		UserParameterModifiedName = InProperty->GetFName();
-	}
 }
 
 void UPCGGraph::PostEditChangeProperty(FPropertyChangedEvent& PropertyChangedEvent)
@@ -872,11 +861,9 @@ void UPCGGraph::PostEditChangeProperty(FPropertyChangedEvent& PropertyChangedEve
 
 		OnGraphParametersChanged(ChangeType, NAME_None);
 	}
-	else if (UserParameterModifiedName != NAME_None && MemberPropertyName == GET_MEMBER_NAME_CHECKED(UPCGGraph, UserParameters))
+	else if (PropertyChangedEvent.MemberProperty && PropertyChangedEvent.MemberProperty->GetOwnerStruct() == UserParameters.GetPropertyBagStruct())
 	{
-		// cf. PreEditChnage comment to understand why we need UserParameterModifiedName
-		OnGraphParametersChanged(EPCGGraphParameterEvent::ValueModifiedLocally, UserParameterModifiedName);
-		UserParameterModifiedName = NAME_None;
+		OnGraphParametersChanged(EPCGGraphParameterEvent::ValueModifiedLocally, MemberPropertyName);
 	}
 
 	NumberOfUserParametersPreEdit = 0;
@@ -991,11 +978,7 @@ void UPCGGraphInstance::PreEditChange(FProperty* InProperty)
 		return;
 	}
 
-	if (InProperty->GetOwnerStruct() == ParametersOverrides.Parameters.GetPropertyBagStruct())
-	{
-		UserParameterModifiedName = InProperty->GetFName();
-	}
-	else if (InProperty->GetFName() == GET_MEMBER_NAME_CHECKED(UPCGGraphInstance, Graph) && Graph)
+	if (InProperty->GetFName() == GET_MEMBER_NAME_CHECKED(UPCGGraphInstance, Graph) && Graph)
 	{
 		Graph->OnGraphChangedDelegate.RemoveAll(this);
 		Graph->OnGraphParametersChangedDelegate.RemoveAll(this);
@@ -1018,10 +1001,9 @@ void UPCGGraphInstance::PostEditChangeProperty(FPropertyChangedEvent& PropertyCh
 
 		RefreshParameters(EPCGGraphParameterEvent::GraphChanged);
 	}
-	else if (UserParameterModifiedName != NAME_None && PropertyChangedEvent.GetMemberPropertyName() == GET_MEMBER_NAME_CHECKED(UPCGGraphInstance, ParametersOverrides))
+	else if (PropertyChangedEvent.MemberProperty && PropertyChangedEvent.MemberProperty->GetOwnerStruct() == ParametersOverrides.Parameters.GetPropertyBagStruct())
 	{
-		OnGraphParametersChanged(this, EPCGGraphParameterEvent::ValueModifiedLocally, UserParameterModifiedName);
-		UserParameterModifiedName = NAME_None;
+		OnGraphParametersChanged(this, EPCGGraphParameterEvent::ValueModifiedLocally, PropertyChangedEvent.GetMemberPropertyName());
 	}
 }
 
diff --git a/Engine/Plugins/Experimental/PCG/Source/PCG/Public/PCGComponent.h b/Engine/Plugins/Experimental/PCG/Source/PCG/Public/PCGComponent.h
--- a/Engine/Plugins/Experimental/PCG/Source/PCG/Public/PCGComponent.h
+++ b/Engine/Plugins/Experimental/PCG/Source/PCG/Public/PCGComponent.h
@@ -314,16 +314,10 @@ private:
 	void OnGraphChanged(UPCGGraphInterface* InGraph, EPCGChangeType ChangeType);
 
 #if WITH_EDITOR
-	// Stub for the other PreEditChange prototype to prevent compile issues from name hiding
-	virtual void PreEditChange(FProperty* PropertyAboutToChange) override { Super::PreEditChange(PropertyAboutToChange); }
-	virtual void PreEditChange(FEditPropertyChain& PropertyAboutToChange) override;
-	virtual void PostEditChangeChainProperty(FPropertyChangedChainEvent& PropertyChangedEvent) override;
 	virtual void PostEditChangeProperty(FPropertyChangedEvent& PropertyChangedEvent) override;
 	virtual void PreEditUndo() override;
 	virtual void PostEditUndo() override;
 
-	bool IsChangingGraphInstanceParameterValue(FEditPropertyChain& InEditPropertyChain) const;
-
 	/** Sets up actor, tracking, landscape and graph callbacks */
 	void SetupCallbacksOnCreation();
 
diff --git a/Engine/Plugins/Experimental/PCG/Source/PCG/Public/PCGGraph.h b/Engine/Plugins/Experimental/PCG/Source/PCG/Public/PCGGraph.h
--- a/Engine/Plugins/Experimental/PCG/Source/PCG/Public/PCGGraph.h
+++ b/Engine/Plugins/Experimental/PCG/Source/PCG/Public/PCGGraph.h
@@ -255,7 +255,6 @@ private:
 	bool bIsNotifying = false;
 	bool bUserPausedNotificationsInGraphEditor = false;
 	int32 NumberOfUserParametersPreEdit = 0;
-	FName UserParameterModifiedName = NAME_None;
 #endif // WITH_EDITOR
 };
 
@@ -317,8 +316,6 @@ private:
 #if WITH_EDITORONLY_DATA
 	// Transient, to keep track the undo/redo changed the graph.
 	UPCGGraphInterface* UndoRedoGraphCache = nullptr;
-
-	FName UserParameterModifiedName = NAME_None;
 #endif // WITH_EDITORONLY_DATA
 };
 
