No longer return null data entries from FPCGDataCollection::GetInputsByPin

diff --git a/Engine/Plugins/Experimental/PCG/Source/PCG/Private/PCGData.cpp b/Engine/Plugins/Experimental/PCG/Source/PCG/Private/PCGData.cpp
--- a/Engine/Plugins/Experimental/PCG/Source/PCG/Private/PCGData.cpp
+++ b/Engine/Plugins/Experimental/PCG/Source/PCG/Private/PCGData.cpp
@@ -149,6 +149,11 @@ TArray<FPCGTaggedData> FPCGDataCollection::GetInputs() const
 TArray<FPCGTaggedData> FPCGDataCollection::GetInputsByPin(const FName& InPinLabel) const
 {
 	return TaggedData.FilterByPredicate([&InPinLabel](const FPCGTaggedData& Data) {
+		if (!ensure(Data.Data))
+		{
+			return false;
+		}
+
 		return Data.Pin == InPinLabel;
 		});
 }
diff --git a/Engine/Plugins/Experimental/PCG/Source/PCG/Private/PCGElement.cpp b/Engine/Plugins/Experimental/PCG/Source/PCG/Private/PCGElement.cpp
--- a/Engine/Plugins/Experimental/PCG/Source/PCG/Private/PCGElement.cpp
+++ b/Engine/Plugins/Experimental/PCG/Source/PCG/Private/PCGElement.cpp
@@ -368,6 +368,19 @@ void IPCGElement::CleanupAndValidateOutput(FPCGContext* Context) const
 #if WITH_EDITOR
 		if (SettingsInterface->bEnabled)
 		{
+			// remove null outputs
+			Context->OutputData.TaggedData.RemoveAll([this, Context](const FPCGTaggedData& TaggedData){
+
+				if (TaggedData.Data == nullptr)
+				{
+					PCGE_LOG(Error, GraphAndLog, FText::Format(LOCTEXT("NullPinOutputData", "Invalid output(s) generated for pin '{0}'"), FText::FromName(TaggedData.Pin)));
+					return true;
+				}
+
+				return false;
+			});
+
+
 			for (FPCGTaggedData& TaggedData : Context->OutputData.TaggedData)
 			{
 				const int32 MatchIndex = OutputPinProperties.IndexOfByPredicate([&TaggedData](const FPCGPinProperties& InProp) { return TaggedData.Pin == InProp.Label; });
@@ -379,7 +392,7 @@ void IPCGElement::CleanupAndValidateOutput(FPCGContext* Context) const
 						PCGE_LOG(Error, GraphAndLog, FText::Format(LOCTEXT("OutputCannotBeRouted", "Output data generated for non-existent output pin '{0}'"), FText::FromName(TaggedData.Pin)));
 					}
 				}
-				else if (TaggedData.Data)
+				else if (ensure(TaggedData.Data))
 				{
 					// Try to get dynamic current pin types, otherwise settle for static types
 					const UPCGPin* OutputPin = Context->Node ? Context->Node->GetOutputPin(OutputPinProperties[MatchIndex].Label) : nullptr;
diff --git a/Engine/Plugins/Experimental/PCG/Source/PCG/Private/Tests/Elements/PCGElementTest.cpp b/Engine/Plugins/Experimental/PCG/Source/PCG/Private/Tests/Elements/PCGElementTest.cpp
new file mode 100644
--- /dev/null
+++ b/Engine/Plugins/Experimental/PCG/Source/PCG/Private/Tests/Elements/PCGElementTest.cpp
@@ -0,0 +1,56 @@
+// Copyright Epic Games, Inc. All Rights Reserved.
+
+#include "Tests/Elements/PCGElementTest.h"
+
+#include "Tests/Determinism/PCGDeterminismTestsCommon.h"
+#include "PCGData.h"
+#include "PCGContext.h"
+
+UPCGBadOutputsNodeSettings::UPCGBadOutputsNodeSettings()
+{
+#if WITH_EDITORONLY_DATA
+	bExposeToLibrary = false;
+#endif
+}
+
+FPCGElementPtr UPCGBadOutputsNodeSettings::CreateElement() const
+{	
+	return MakeShared<FPCGBadOutputNodeElement>();
+}
+
+bool FPCGBadOutputNodeElement::ExecuteInternal(FPCGContext* Context) const
+{
+	// add some null outputs
+	FPCGTaggedData& BadOutput = Context->OutputData.TaggedData.Add_GetRef({});
+	BadOutput.Pin = PCGPinConstants::DefaultOutputLabel;
+
+	return true;
+}
+
+#if WITH_EDITOR
+
+IMPLEMENT_CUSTOM_SIMPLE_AUTOMATION_TEST(FPCGElementTest_CleanupAndValidateOutput, FPCGTestBaseClass, "pcg.tests.PCGElement.CleanupAndValidateOutput", PCGTestsCommon::TestFlags)
+
+bool FPCGElementTest_CleanupAndValidateOutput::RunTest(const FString& Parameters)
+{
+	PCGTestsCommon::FTestData TestData(PCGDeterminismTests::Defaults::Seed);
+	PCGDeterminismTests::GenerateSettings<UPCGBadOutputsNodeSettings>(TestData);
+
+	TUniquePtr<FPCGContext> Context = TestData.InitializeTestContext();
+
+	FPCGElementPtr TestElement = TestData.Settings->GetElement();
+
+	AddExpectedError(TEXT("Invalid output\\(s\\) generated for pin"));
+
+	while (!TestElement->Execute(Context.Get())) {}
+
+	// check we have no null outputs
+
+	TArray<FPCGTaggedData> Outputs = Context->OutputData.GetInputsByPin(PCGPinConstants::DefaultOutputLabel);
+
+	UTEST_EQUAL("Null entries should have been removed", Outputs.Num(), 0);
+
+	return true;
+}
+
+#endif // WITH_EDITOR
diff --git a/Engine/Plugins/Experimental/PCG/Source/PCG/Private/Tests/Elements/PCGElementTest.h b/Engine/Plugins/Experimental/PCG/Source/PCG/Private/Tests/Elements/PCGElementTest.h
new file mode 100644
--- /dev/null
+++ b/Engine/Plugins/Experimental/PCG/Source/PCG/Private/Tests/Elements/PCGElementTest.h
@@ -0,0 +1,34 @@
+// Copyright Epic Games, Inc. All Rights Reserved.
+
+#pragma once
+
+#include "PCGSettings.h"
+#include "PCGElement.h"
+#include "PCGElementTest.generated.h"
+
+
+/**
+ * Test node to write bad outputs
+ */
+UCLASS(HideDropdown, NotPlaceable, MinimalAPI)
+class UPCGBadOutputsNodeSettings : public UPCGSettings
+{
+	GENERATED_BODY()
+
+public:
+	UPCGBadOutputsNodeSettings();
+
+#if WITH_EDITOR
+	virtual FName GetDefaultNodeName() const override { return FName(TEXT("BadOutputs")); }
+	virtual EPCGSettingsType GetType() const override { return EPCGSettingsType::Spatial; }
+#endif
+
+	virtual FPCGElementPtr CreateElement() const override;
+
+};
+
+class FPCGBadOutputNodeElement : public FSimplePCGElement
+{
+public:
+	virtual bool ExecuteInternal(FPCGContext* Context) const override;
+};
