[PCG] Adding column sorting for attribute list view

diff --git a/Engine/Plugins/Experimental/PCG/Source/PCGEditor/Private/SPCGEditorGraphAttributeListView.cpp b/Engine/Plugins/Experimental/PCG/Source/PCGEditor/Private/SPCGEditorGraphAttributeListView.cpp
--- a/Engine/Plugins/Experimental/PCG/Source/PCGEditor/Private/SPCGEditorGraphAttributeListView.cpp
+++ b/Engine/Plugins/Experimental/PCG/Source/PCGEditor/Private/SPCGEditorGraphAttributeListView.cpp
@@ -195,6 +195,7 @@ SPCGEditorGraphAttributeListView::~SPCGEditorGraphAttributeListView()
 void SPCGEditorGraphAttributeListView::Construct(const FArguments& InArgs, TSharedPtr<FPCGEditor> InPCGEditor)
 {
 	PCGEditorPtr = InPCGEditor;
+	SortingColumn = PCGEditorGraphAttributeListView::NAME_IndexColumn;
 
 	PCGEditorPtr.Pin()->OnDebugObjectChangedDelegate.AddSP(this, &SPCGEditorGraphAttributeListView::OnDebugObjectChanged);
 	PCGEditorPtr.Pin()->OnInspectedNodeChangedDelegate.AddSP(this, &SPCGEditorGraphAttributeListView::OnInspectedNodeChanged);
@@ -388,7 +389,8 @@ void SPCGEditorGraphAttributeListView::RefreshAttributeList()
 	// Swapping to an empty item list to force a widget clear, otherwise the widgets will try to update during add column and access invalid data
 	static const TArray<PCGListviewItemPtr> EmptyList;
 	ListView->SetItemsSource(&EmptyList);
-	
+
+	PCGColumnData.Empty();
 	ListViewItems.Empty();
 	ListViewHeader->ClearColumns();
 	InfoTextBlock->SetText(FText::GetEmpty());
@@ -486,6 +488,7 @@ void SPCGEditorGraphAttributeListView::RefreshAttributeList()
 	}
 
 	ListView->SetItemsSource(&ListViewItems);
+	RefreshSorting();
 }
 
 void SPCGEditorGraphAttributeListView::RefreshDataComboBox()
@@ -779,6 +782,111 @@ void SPCGEditorGraphAttributeListView::OnItemDoubleClicked(PCGListviewItemPtr It
 	}
 }
 
+void SPCGEditorGraphAttributeListView::OnColumnSortModeChanged(const EColumnSortPriority::Type InSortPriority, const FName& InColumnId, const EColumnSortMode::Type InNewSortMode)
+{
+	if (SortingColumn == InColumnId)
+	{
+		// Cycling
+		SortMode = static_cast<EColumnSortMode::Type>((SortMode + 1) % 3);
+	}
+	else
+	{
+		SortingColumn = InColumnId;
+		SortMode = InNewSortMode;
+	}
+
+	RefreshSorting();
+}
+
+EColumnSortMode::Type SPCGEditorGraphAttributeListView::GetColumnSortMode(const FName InColumnId) const
+{
+	if (SortingColumn != InColumnId)
+	{
+		return EColumnSortMode::None;
+	}
+
+	return SortMode;
+}
+
+void SPCGEditorGraphAttributeListView::RefreshSorting()
+{
+	TRACE_CPUPROFILER_EVENT_SCOPE(SPCGEditorGraphAttributeListView::ColumnSort);
+	
+	if (SortingColumn == PCGEditorGraphAttributeListView::NAME_IndexColumn || SortMode == EColumnSortMode::None)
+	{
+		if (SortMode == EColumnSortMode::Ascending || SortMode == EColumnSortMode::None)
+		{
+			ListViewItems.Sort([](const PCGListviewItemPtr& LHS, const PCGListviewItemPtr& RHS)
+			{
+				return  LHS->Index < RHS->Index;
+			});
+		}
+		else
+		{
+			ListViewItems.Sort([](const PCGListviewItemPtr& LHS, const PCGListviewItemPtr& RHS)
+			{
+				return  LHS->Index > RHS->Index;
+			});
+		}
+	}
+	else if (FPCGColumnData* ColumnData = PCGColumnData.Find(SortingColumn))
+	{
+		if (ColumnData->DataAccessor.IsValid() && ColumnData->DataKeys.IsValid())
+		{
+			auto Callback = [this, &ColumnData](auto Dummy)
+			{
+				using ValueType = decltype(Dummy);
+
+				if constexpr (PCG::Private::MetadataTraits<ValueType>::CanCompare)
+				{
+					TArray<ValueType> CachedValues;
+					CachedValues.SetNumUninitialized(ListViewItems.Num());
+					ColumnData->DataAccessor->GetRange(TArrayView<ValueType>(CachedValues), 0, *ColumnData->DataKeys);
+
+					auto SortAscending = [&CachedValues] (const PCGListviewItemPtr& LHS, const PCGListviewItemPtr& RHS)
+					{
+						const ValueType& LHSValue = CachedValues[LHS->Index];
+						const ValueType& RHSValue = CachedValues[RHS->Index];
+
+						if (PCG::Private::MetadataTraits<ValueType>::Equal(LHSValue, RHSValue))
+						{
+							return LHS->Index < RHS->Index;
+						}
+
+						return PCG::Private::MetadataTraits<ValueType>::Less(LHSValue, RHSValue);
+					};
+
+					auto SortDescending = [&CachedValues] (const PCGListviewItemPtr& LHS, const PCGListviewItemPtr& RHS)
+					{
+						const ValueType& LHSValue = CachedValues[LHS->Index];
+						const ValueType& RHSValue = CachedValues[RHS->Index];
+
+						if (PCG::Private::MetadataTraits<ValueType>::Equal(LHSValue, RHSValue))
+						{
+							return LHS->Index > RHS->Index;
+						}	
+
+						return PCG::Private::MetadataTraits<ValueType>::Greater(LHSValue, RHSValue);
+					};
+
+					if (SortMode == EColumnSortMode::Ascending)
+					{
+						ListViewItems.Sort(SortAscending);
+					}
+					else if (SortMode == EColumnSortMode::Descending)
+					{
+						ListViewItems.Sort(SortDescending);
+					}
+				}
+			};
+
+			PCGMetadataAttribute::CallbackWithRightType(ColumnData->DataAccessor->GetUnderlyingType(), Callback);
+		}
+	}
+
+	ListView->RequestListRefresh();
+}
+
 void SPCGEditorGraphAttributeListView::AddColumn(const UPCGPointData* InPCGPointData, const FName& InColumnId, const FText& ColumnLabel, EHorizontalAlignment HeaderHAlign, EHorizontalAlignment CellHAlign)
 {
 	const float ColumnWidth = PCGEditorGraphAttributeListView::CalculateColumnWidth(ColumnLabel);
@@ -823,6 +931,8 @@ void SPCGEditorGraphAttributeListView::AddColumn(const UPCGPointData* InPCGPoint
 	Arguments.ManualWidth(ColumnWidth);
 	Arguments.HAlignHeader(HeaderHAlign);
 	Arguments.HAlignCell(CellHAlign);
+	Arguments.SortMode(this, &SPCGEditorGraphAttributeListView::GetColumnSortMode, InColumnId);
+	Arguments.OnSort(this, &SPCGEditorGraphAttributeListView::OnColumnSortModeChanged);
 
 	SHeaderRow::FColumn* NewColumn = new SHeaderRow::FColumn(Arguments);
 	NewColumn->bIsVisible = !HiddenAttributes.Contains(InColumnId); 
@@ -920,9 +1030,11 @@ void SPCGEditorGraphAttributeListView::AddMetadataColumn(const UPCGData* InPCGDa
 	ColumnArguments.HAlignHeader(EHorizontalAlignment::HAlign_Center);
 	ColumnArguments.HAlignCell(CellAlignment);
 	ColumnArguments.ManualWidth(ColumnWidth);
+	ColumnArguments.SortMode(this, &SPCGEditorGraphAttributeListView::GetColumnSortMode, ColumnId);
+	ColumnArguments.OnSort(this, &SPCGEditorGraphAttributeListView::OnColumnSortModeChanged);
 
 	SHeaderRow::FColumn* NewColumn = new SHeaderRow::FColumn(ColumnArguments);
-	NewColumn->bIsVisible = !HiddenAttributes.Contains(InColumnId);
+	NewColumn->bIsVisible = !HiddenAttributes.Contains(ColumnId);
 	ListViewHeader->AddColumn(*NewColumn);
 }
 
diff --git a/Engine/Plugins/Experimental/PCG/Source/PCGEditor/Private/SPCGEditorGraphAttributeListView.h b/Engine/Plugins/Experimental/PCG/Source/PCGEditor/Private/SPCGEditorGraphAttributeListView.h
--- a/Engine/Plugins/Experimental/PCG/Source/PCGEditor/Private/SPCGEditorGraphAttributeListView.h
+++ b/Engine/Plugins/Experimental/PCG/Source/PCGEditor/Private/SPCGEditorGraphAttributeListView.h
@@ -105,6 +105,10 @@ private:
 	TSharedRef<ITableRow> OnGenerateRow(PCGListviewItemPtr Item, const TSharedRef<STableViewBase>& OwnerTable);
 	void OnItemDoubleClicked(PCGListviewItemPtr Item) const;
 
+	void OnColumnSortModeChanged(const EColumnSortPriority::Type SortPriority, const FName& ColumnId, const EColumnSortMode::Type InSortMode);
+	EColumnSortMode::Type GetColumnSortMode(const FName InColumnId) const;
+	void RefreshSorting();
+
 	void AddColumn(const UPCGPointData* InPCGPointData, const FName& InColumnId, const FText& ColumnLabel, EHorizontalAlignment HeaderHAlign = HAlign_Center, EHorizontalAlignment CellHAlign = HAlign_Right);
 	void AddIndexColumn();
 	void AddPointDataColumns(const UPCGPointData* InPCGPointData);
@@ -133,4 +137,7 @@ private:
 	TArray<FName> HiddenAttributes;
 
 	TMap<FName, FPCGColumnData> PCGColumnData;
+
+	FName SortingColumn = NAME_None;
+	EColumnSortMode::Type SortMode = EColumnSortMode::Type::Ascending;
 };
