[PCG] Using "single collection" icons on output pins, because nodes always output a single collection

diff --git a/Engine/Plugins/Experimental/PCG/Source/PCGEditor/Private/SPCGEditorGraphNode.cpp b/Engine/Plugins/Experimental/PCG/Source/PCGEditor/Private/SPCGEditorGraphNode.cpp
--- a/Engine/Plugins/Experimental/PCG/Source/PCGEditor/Private/SPCGEditorGraphNode.cpp
+++ b/Engine/Plugins/Experimental/PCG/Source/PCGEditor/Private/SPCGEditorGraphNode.cpp
@@ -360,6 +360,7 @@ void SPCGEditorGraphNode::AddPin(const TSharedRef<SGraphPin>& PinToAdd)
 			}
 			else
 			{
+				// Node outputs are always single collection (SC).
 				static const FName* PinBrushes[] =
 				{
 					&PCGEditorStyleConstants::Pin_SD_SC_IN_C,
@@ -372,12 +373,12 @@ void SPCGEditorGraphNode::AddPin(const TSharedRef<SGraphPin>& PinToAdd)
 					&PCGEditorStyleConstants::Pin_MD_MC_IN_DC,
 					&PCGEditorStyleConstants::Pin_SD_SC_OUT_C,
 					&PCGEditorStyleConstants::Pin_SD_SC_OUT_DC,
-					&PCGEditorStyleConstants::Pin_SD_MC_OUT_C,
-					&PCGEditorStyleConstants::Pin_SD_MC_OUT_DC,
+					&PCGEditorStyleConstants::Pin_SD_SC_OUT_C,
+					&PCGEditorStyleConstants::Pin_SD_SC_OUT_DC,
 					&PCGEditorStyleConstants::Pin_MD_SC_OUT_C,
 					&PCGEditorStyleConstants::Pin_MD_SC_OUT_DC,
-					&PCGEditorStyleConstants::Pin_MD_MC_OUT_C,
-					&PCGEditorStyleConstants::Pin_MD_MC_OUT_DC
+					&PCGEditorStyleConstants::Pin_MD_SC_OUT_C,
+					&PCGEditorStyleConstants::Pin_MD_SC_OUT_DC
 				};
 
 				const int32 ConnectedIndex = (bIsInPin ? 0 : 8) + (bIsMultiData ? 4 : 0) + (bIsMultiConnections ? 2 : 0);
