[PCG] Fix default object change in InputNode that was breaking old graphs

* The In pin will still be part of the default object
* But it will also be advanced by default

#rnx

diff --git a/Engine/Plugins/Experimental/PCG/Source/PCG/Private/PCGInputOutputSettings.cpp b/Engine/Plugins/Experimental/PCG/Source/PCG/Private/PCGInputOutputSettings.cpp
--- a/Engine/Plugins/Experimental/PCG/Source/PCG/Private/PCGInputOutputSettings.cpp
+++ b/Engine/Plugins/Experimental/PCG/Source/PCG/Private/PCGInputOutputSettings.cpp
@@ -19,7 +19,8 @@ bool FPCGInputOutputElement::ExecuteInternal(FPCGContext* Context) const
 UPCGGraphInputOutputSettings::UPCGGraphInputOutputSettings(const FObjectInitializer& ObjectInitializer)
 	: Super(ObjectInitializer)
 {
-	StaticAdvancedInLabels.Emplace(PCGPinConstants::DefaultInputLabel, LOCTEXT("InputOutputInPinTooltip",
+	// This label is kept into "StaticInLabels" to avoid changing the default object. But it will be marked as advanced too, since we want to move away from it.
+	StaticInLabels.Emplace(PCGPinConstants::DefaultInputLabel, LOCTEXT("InputOutputInPinTooltip",
 		"DEPRECATED: Please use custom pins to pass inputs into the graph, or use the 'getter' nodes (e.g. Get Volume Data, Get Actor Data, ...) to import data from the level.\n\n"
 		"Provides the same result as the Input pin."
 	));
@@ -73,7 +74,9 @@ TArray<FPCGPinProperties> UPCGGraphInputOutputSettings::GetPinProperties() const
 	const bool bIsInputPin = bIsInput;
 	const EPCGDataType DefaultPinDataType = bIsInput ? EPCGDataType::Spatial : EPCGDataType::Any;
 	Algo::Transform(StaticLabels(), PinProperties, [bIsInputPin, DefaultPinDataType](const FLabelAndTooltip& InLabelAndTooltip) {
-		return FPCGPinProperties(InLabelAndTooltip.Label, DefaultPinDataType, /*bMultiConnections=*/true, /*bMultiData=*/true, InLabelAndTooltip.Tooltip);
+		FPCGPinProperties Res = FPCGPinProperties(InLabelAndTooltip.Label, DefaultPinDataType, /*bMultiConnections=*/true, /*bMultiData=*/true, InLabelAndTooltip.Tooltip);
+		Res.bAdvancedPin = true;
+		return Res;
 	});
 	
 
@@ -107,7 +110,9 @@ TArray<FPCGPinProperties> UPCGGraphInputOutputSettings::DefaultInputPinPropertie
 	const bool bIsInputPin = bIsInput;
 	const EPCGDataType DefaultPinDataType = bIsInput ? EPCGDataType::Spatial : EPCGDataType::Any;
 	Algo::Transform(StaticLabels(), PinProperties, [bIsInputPin, DefaultPinDataType](const FLabelAndTooltip& InLabelAndTooltip) {
-		return FPCGPinProperties(InLabelAndTooltip.Label, DefaultPinDataType, /*bMultiConnections=*/true, /*bMultiData=*/true, InLabelAndTooltip.Tooltip);
+		FPCGPinProperties Res = FPCGPinProperties(InLabelAndTooltip.Label, DefaultPinDataType, /*bMultiConnections=*/true, /*bMultiData=*/true, InLabelAndTooltip.Tooltip);
+		Res.bAdvancedPin = true;
+		return Res;
 	});
 
 	return PinProperties;
