[PCG] Fix Create Spline with custom tangents, that were ignored.

* Also make it clear that custom tangents can't be used if the spline is linear

diff --git a/Engine/Plugins/Experimental/PCG/Source/PCG/Private/Elements/PCGCreateSpline.cpp b/Engine/Plugins/Experimental/PCG/Source/PCG/Private/Elements/PCGCreateSpline.cpp
--- a/Engine/Plugins/Experimental/PCG/Source/PCG/Private/Elements/PCGCreateSpline.cpp
+++ b/Engine/Plugins/Experimental/PCG/Source/PCG/Private/Elements/PCGCreateSpline.cpp
@@ -174,7 +174,16 @@ bool FPCGCreateSplineElement::ExecuteInternal(FPCGContext* Context) const
 		SplinePoints.Reserve(Points.Num());
 
 		const FTransform SplineActorTransform = SplineActor->GetTransform();
-		ESplinePointType::Type PointType = (Settings->bLinear ? ESplinePointType::Linear : ESplinePointType::Curve);
+		
+		ESplinePointType::Type PointType = ESplinePointType::Curve;
+		if (Settings->bLinear)
+		{
+			PointType = ESplinePointType::Linear;
+		}
+		else if (Settings->bApplyCustomTangents)
+		{
+			PointType = ESplinePointType::CurveCustomTangent;
+		}
 
 		for(int32 PointIndex = 0; PointIndex < Points.Num(); ++PointIndex)
 		{
diff --git a/Engine/Plugins/Experimental/PCG/Source/PCG/Public/Elements/PCGCreateSpline.h b/Engine/Plugins/Experimental/PCG/Source/PCG/Public/Elements/PCGCreateSpline.h
--- a/Engine/Plugins/Experimental/PCG/Source/PCG/Public/Elements/PCGCreateSpline.h
+++ b/Engine/Plugins/Experimental/PCG/Source/PCG/Public/Elements/PCGCreateSpline.h
@@ -3,6 +3,7 @@
 #pragma once
 
 #include "PCGElement.h"
+#include "PCGPin.h"
 #include "PCGSettings.h"
 #include "Data/PCGSplineData.h"
 
@@ -51,13 +52,14 @@ public:
 	UPROPERTY(BlueprintReadWrite, EditAnywhere, Category = Settings)
 	bool bLinear = false;
 
-	UPROPERTY(BlueprintReadWrite, EditAnywhere, Category = Settings)
+	/** Allow to specify custom tangents for each point, as an attribute. Can't be set if the spline is linear. */
+	UPROPERTY(BlueprintReadWrite, EditAnywhere, Category = Settings, meta = (EditCondition = "!bLinear"))
 	bool bApplyCustomTangents = false;
 
-	UPROPERTY(BlueprintReadWrite, EditAnywhere, Category = Settings, meta = (EditCondition = "bApplyCustomTangents"))
+	UPROPERTY(BlueprintReadWrite, EditAnywhere, Category = Settings, meta = (EditCondition = "!bLinear && bApplyCustomTangents"))
 	FName ArriveTangentAttribute;
 
-	UPROPERTY(BlueprintReadWrite, EditAnywhere, Category = Settings, meta = (EditCondition = "bApplyCustomTangents"))
+	UPROPERTY(BlueprintReadWrite, EditAnywhere, Category = Settings, meta = (EditCondition = "!bLinear && bApplyCustomTangents"))
 	FName LeaveTangentAttribute;
 };
 
