Verify that the PCG Component is added by default to new PCG Volumes
Precondition: Open QAGame in Unreal Editor
Steps:
- From the Level Editor toolbar, search for and add a PCGVolume to the world
-Within the Details Panel Component section, a PCG Component is present

diff --git a/Engine/Plugins/Experimental/PCG/Source/PCGEditor/PCGEditor.Build.cs b/Engine/Plugins/Experimental/PCG/Source/PCGEditor/PCGEditor.Build.cs
--- a/Engine/Plugins/Experimental/PCG/Source/PCGEditor/PCGEditor.Build.cs
+++ b/Engine/Plugins/Experimental/PCG/Source/PCGEditor/PCGEditor.Build.cs
@@ -14,12 +14,13 @@ namespace UnrealBuildTool.Rules
 			"Projects",
 			"Engine",
 			"CoreUObject",
-			"PlacementMode"
+			"PlacementMode",
 		});
 
 			PrivateDependencyModuleNames.AddRange(
 				new string[]
-				{					
+				{
+				
 				}
 			);
 
@@ -28,6 +29,7 @@ namespace UnrealBuildTool.Rules
 					"AppFramework",
 					"ApplicationCore",
 					"AssetTools",
+					"AssetRegistry",
 					"BlueprintGraph",
 					"DesktopWidgets",
 					"DetailCustomizations",
diff --git a/Engine/Plugins/Experimental/PCG/Source/PCGEditor/Private/Tests/Display/PCGComponentInNewPCGVolumes.cpp b/Engine/Plugins/Experimental/PCG/Source/PCGEditor/Private/Tests/Display/PCGComponentInNewPCGVolumes.cpp
--- a/Engine/Plugins/Experimental/PCG/Source/PCGEditor/Private/Tests/Display/PCGComponentInNewPCGVolumes.cpp
+++ b/Engine/Plugins/Experimental/PCG/Source/PCGEditor/Private/Tests/Display/PCGComponentInNewPCGVolumes.cpp
@@ -2,18 +2,22 @@
 
 #include "EngineUtils.h"
 #include "Engine/World.h"
-#include "Misc/AutomationTest.h"
+
 #include "Editor.h"
 #include "Editor/UnrealEd/Public/PackageTools.h"
+#include "ObjectTools.h"
 #include "AssetRegistry/AssetRegistryModule.h"
-#include "IPlacementModeModule.h"
+
 #include "PCGVolume.h"
 #include "PCGComponent.h"
 #include "Tests/PCGTestsCommon.h"
+#include "Tests/AutomationEditorCommon.h"
+
 
 #if WITH_AUTOMATION_TESTS
 
 IMPLEMENT_CUSTOM_SIMPLE_AUTOMATION_TEST(FPCGComponentInNewPCGVolumes, FPCGTestBaseClass, "Editor.Plugins.Tools.PCG.PCGComponentInNewPCGVolumes", PCGTestsCommon::TestFlags)
+
 bool FPCGComponentInNewPCGVolumes::RunTest(const FString& Parameters)
 {
 	// Get current world
@@ -26,35 +30,50 @@ bool FPCGComponentInNewPCGVolumes::RunTest(const FString& Parameters)
 		return false;
 	}
 
-	// We need to get the PCGVolume actor factory
-	TSubclassOf<AActor> PCGVolumeClass = APCGVolume::StaticClass();
+	TSubclassOf<APCGVolume> PCGVolumeClass = APCGVolume::StaticClass();
 	UActorFactory* PCGVolumeFactory = GEditor->FindActorFactoryForActorClass(PCGVolumeClass);
+
 	UTEST_NOT_NULL(TEXT("Failed to find PCGVolume actor factory."), PCGVolumeFactory);
 
-	// Get the asset data for the PCGVolume class
-	FAssetData PCGVolumeAssetData = FAssetData(PCGVolumeClass);
+	APCGVolume* VolumeActor = nullptr;
 
-	// Add a new PCGVolume actor to the level
-	AActor* VolumeActor = nullptr;
-	// To clean created assets, there appears window save content
 	if (GCurrentLevelEditingViewportClient)
 	{
 		FTransform ActorTransform;
-		VolumeActor = GEditor->UseActorFactory(PCGVolumeFactory, PCGVolumeAssetData, &ActorTransform);
+		VolumeActor = Cast<APCGVolume>(GEditor->UseActorFactory(PCGVolumeFactory, FAssetData(PCGVolumeClass), &ActorTransform));
 	}
-
 	UTEST_NOT_NULL(TEXT("Failed to add PCGVolume actor."), VolumeActor);
-
 	// Condition for static analysis.
 	if (VolumeActor)
 	{
-		UPCGComponent* PCGComponent = VolumeActor->FindComponentByClass<UPCGComponent>();
-		TestNotNull("PCGVolume actor does not contain a PCGComponent!", PCGComponent);
+		UPackage* Package = VolumeActor->GetOutermost();
+		Package->SetFlags(RF_Transient);
+		TWeakObjectPtr<UPCGComponent> PCGComponent = VolumeActor->FindComponentByClass<UPCGComponent>();
+		TestNotNull(TEXT("PCGVolume actor does not contain a PCGComponent!"), PCGComponent.Get());
 
-		// Destroy the PCGVolume actor after test
-		World->EditorDestroyActor(VolumeActor, false);
-	}
+		if (PCGComponent.IsValid())
+		{
+			PCGComponent->ClearPCGLink(APCGVolume::StaticClass());
+			PCGComponent = nullptr;
+
+		}
 
+		// Cleanup the VolumeActor
+		FAssetRegistryModule& AssetRegistryModule = FModuleManager::LoadModuleChecked<FAssetRegistryModule>("AssetRegistry");
+		IAssetRegistry& AssetRegistry = AssetRegistryModule.Get();
+		FAssetData PCGVolumeAssetData(VolumeActor);
+		AssetRegistry.AssetDeleted(PCGVolumeAssetData.GetAsset());
+
+		bool bSuccessful = ObjectTools::DeleteSingleObject(VolumeActor, true);
+		if (!bSuccessful)
+		{
+			//Clear references to the object so we can delete it
+			FAutomationEditorCommonUtils::NullReferencesToObject(VolumeActor);
+			bSuccessful = ObjectTools::DeleteSingleObject(VolumeActor, false);
+			TestTrue(TEXT("References to PCGVolume component could not be cleaned"), bSuccessful);
+		}
+	}
 	return true;
-}
-#endif
+};
+
+#endif
\ No newline at end of file
