[PCG] Added flag in attribute selector to force get the "None" attribute instead getting the last attribute

diff --git a/Engine/Plugins/Experimental/PCG/Source/PCG/Private/Metadata/Accessors/PCGAttributeAccessorHelpers.cpp b/Engine/Plugins/Experimental/PCG/Source/PCG/Private/Metadata/Accessors/PCGAttributeAccessorHelpers.cpp
--- a/Engine/Plugins/Experimental/PCG/Source/PCG/Private/Metadata/Accessors/PCGAttributeAccessorHelpers.cpp
+++ b/Engine/Plugins/Experimental/PCG/Source/PCG/Private/Metadata/Accessors/PCGAttributeAccessorHelpers.cpp
@@ -17,7 +17,7 @@
 
 namespace PCGAttributeAccessorHelpers
 {
-	void ExtractMetadataAtribute(UPCGData* InData, FName Name, UPCGMetadata*& OutMetadata, FPCGMetadataAttributeBase*& OutAttribute)
+	void ExtractMetadataAtribute(UPCGData* InData, FName Name, UPCGMetadata*& OutMetadata, FPCGMetadataAttributeBase*& OutAttribute, bool bForceGetNone = false)
 	{
 		OutMetadata = nullptr;
 		OutAttribute = nullptr;
@@ -33,8 +33,8 @@ namespace PCGAttributeAccessorHelpers
 
 		if (OutMetadata)
 		{
-			// If Name is None, try ot get the latest attribute
-			if (Name == NAME_None)
+			// If Name is None, try to get the latest attribute
+			if (!bForceGetNone && Name == NAME_None)
 			{
 				Name = OutMetadata->GetLatestAttributeNameOrNone();
 			}
@@ -43,11 +43,11 @@ namespace PCGAttributeAccessorHelpers
 		}
 	}
 
-	void ExtractMetadataAtribute(const UPCGData* InData, FName Name, UPCGMetadata const*& OutMetadata, FPCGMetadataAttributeBase const*& OutAttribute)
+	void ExtractMetadataAtribute(const UPCGData* InData, FName Name, UPCGMetadata const*& OutMetadata, FPCGMetadataAttributeBase const*& OutAttribute, bool bForceGetNone = false)
 	{
 		UPCGMetadata* Metadata = nullptr;
 		FPCGMetadataAttributeBase* Attribute = nullptr;
-		ExtractMetadataAtribute(const_cast<UPCGData*>(InData), Name, Metadata, Attribute);
+		ExtractMetadataAtribute(const_cast<UPCGData*>(InData), Name, Metadata, Attribute, bForceGetNone);
 		OutMetadata = Metadata;
 		OutAttribute = Attribute;
 	}
@@ -306,7 +306,7 @@ TUniquePtr<const IPCGAttributeAccessor> PCGAttributeAccessorHelpers::CreateConst
 		const UPCGMetadata* Metadata = nullptr;
 		const FPCGMetadataAttributeBase* Attribute = nullptr;
 
-		ExtractMetadataAtribute(InData, Name, Metadata, Attribute);
+		ExtractMetadataAtribute(InData, Name, Metadata, Attribute, InSelector.bForceGetNone);
 
 		auto CreateTypedAccessor = [Attribute, Metadata](auto Dummy) -> TUniquePtr<IPCGAttributeAccessor>
 		{
diff --git a/Engine/Plugins/Experimental/PCG/Source/PCG/Public/Metadata/PCGAttributePropertySelector.h b/Engine/Plugins/Experimental/PCG/Source/PCG/Public/Metadata/PCGAttributePropertySelector.h
--- a/Engine/Plugins/Experimental/PCG/Source/PCG/Public/Metadata/PCGAttributePropertySelector.h
+++ b/Engine/Plugins/Experimental/PCG/Source/PCG/Public/Metadata/PCGAttributePropertySelector.h
@@ -38,6 +38,10 @@ public:
 	UPROPERTY()
 	TArray<FString> ExtraNames;
 
+	// Force get the "None" attribute instead getting the last attribute
+	UPROPERTY()
+	bool bForceGetNone = false;
+
 	// Setters, retrurn true if something changed.
 	bool SetPointProperty(EPCGPointProperties InPointProperty, bool bResetExtraNames = true);
 	bool SetAttributeName(FName InAttributeName, bool bResetExtraNames = true);
diff --git a/Engine/Plugins/Experimental/PCG/Source/PCGEditor/Private/SPCGEditorGraphAttributeListView.cpp b/Engine/Plugins/Experimental/PCG/Source/PCGEditor/Private/SPCGEditorGraphAttributeListView.cpp
--- a/Engine/Plugins/Experimental/PCG/Source/PCGEditor/Private/SPCGEditorGraphAttributeListView.cpp
+++ b/Engine/Plugins/Experimental/PCG/Source/PCGEditor/Private/SPCGEditorGraphAttributeListView.cpp
@@ -896,6 +896,7 @@ void SPCGEditorGraphAttributeListView::AddColumn(const UPCGPointData* InPCGPoint
 		const FString ColumnIdString = InColumnId.ToString();
 		
 		FPCGAttributePropertySelector TargetSelector;
+		TargetSelector.bForceGetNone = true;
 		TargetSelector.Update(ColumnIdString);
 	
 		FPCGColumnData& ColumnData = PCGColumnData.Add(InColumnId);
@@ -998,6 +999,7 @@ void SPCGEditorGraphAttributeListView::AddMetadataColumn(const UPCGData* InPCGDa
 	if (InPCGData)
 	{
 		FPCGAttributePropertySelector TargetSelector;
+		TargetSelector.bForceGetNone = true;
 		TargetSelector.Update(OriginalColumnIdString);
 	
 		FPCGColumnData& ColumnData = PCGColumnData.Add(ColumnId);
