Fixed resource reuse issue when using multiple inputs and no merging
Fixed typo in name of parameter that was not a bool anymore

diff --git a/Engine/Plugins/Experimental/PCG/Source/PCG/Private/Elements/PCGSpawnActor.cpp b/Engine/Plugins/Experimental/PCG/Source/PCG/Private/Elements/PCGSpawnActor.cpp
--- a/Engine/Plugins/Experimental/PCG/Source/PCG/Private/Elements/PCGSpawnActor.cpp
+++ b/Engine/Plugins/Experimental/PCG/Source/PCG/Private/Elements/PCGSpawnActor.cpp
@@ -201,6 +201,13 @@ void UPCGSpawnActorSettings::PostLoad()
 	Super::PostLoad();
 
 #if WITH_EDITOR
+	// Apply data deprecation - deprecated in UE 5.3
+	if (bGenerationTrigger_DEPRECATED != EPCGSpawnActorGenerationTrigger::Default)
+	{
+		GenerationTrigger = bGenerationTrigger_DEPRECATED;
+		bGenerationTrigger_DEPRECATED = EPCGSpawnActorGenerationTrigger::Default;
+	}
+
 	SetupBlueprintEvent();
 
 	if (TemplateActorClass)
@@ -522,26 +529,33 @@ bool FPCGSpawnActorElement::ExecuteInternal(FPCGContext* Context) const
 				SpawnParams.ObjectFlags |= RF_Transient;
 			}
 
-			const bool bForceCallGenerate = Settings->bGenerationTrigger == EPCGSpawnActorGenerationTrigger::ForceGenerate;
+			const bool bForceCallGenerate = (Settings->GenerationTrigger == EPCGSpawnActorGenerationTrigger::ForceGenerate);
 #if WITH_EDITOR
-			const bool bOnLoadCallGenerate = Settings->bGenerationTrigger == EPCGSpawnActorGenerationTrigger::Default;
+			const bool bOnLoadCallGenerate = (Settings->GenerationTrigger == EPCGSpawnActorGenerationTrigger::Default);
 #else
-			const bool bOnLoadCallGenerate = (Settings->bGenerationTrigger == EPCGSpawnActorGenerationTrigger::Default ||
-				Settings->bGenerationTrigger == EPCGSpawnActorGenerationTrigger::DoNotGenerateInEditor);
+			const bool bOnLoadCallGenerate = (Settings->GenerationTrigger == EPCGSpawnActorGenerationTrigger::Default ||
+				Settings->GenerationTrigger == EPCGSpawnActorGenerationTrigger::DoNotGenerateInEditor);
 #endif
 			UPCGSubsystem* Subsystem = (Context->SourceComponent.Get() ? Context->SourceComponent->GetSubsystem() : nullptr);
 
 			// Try to reuse actors if the are preexisting
 			TArray<UPCGManagedActors*> ReusedManagedActorsResources;
+			FPCGCrc InputDependenciesCrc;
 			if (CVarAllowActorReuse.GetValueOnAnyThread())
 			{
-				if (Context->DependenciesCrc.IsValid())
+				FPCGDataCollection SingleInputCollection;
+				SingleInputCollection.TaggedData.Add(Input);
+				SingleInputCollection.Crc = SingleInputCollection.ComputeCrc(/*bFullDataCrc=*/false);
+
+				GetDependenciesCrc(SingleInputCollection, Settings, Context->SourceComponent.Get(), InputDependenciesCrc);
+
+				if(InputDependenciesCrc.IsValid())
 				{
-					Context->SourceComponent->ForEachManagedResource([&ReusedManagedActorsResources, &Context](UPCGManagedResource* InResource)
+					Context->SourceComponent->ForEachManagedResource([&ReusedManagedActorsResources, &InputDependenciesCrc , &Context](UPCGManagedResource* InResource)
 					{
 						if (UPCGManagedActors* Resource = Cast<UPCGManagedActors>(InResource))
 						{
-							if (Resource->GetCrc().IsValid() && Resource->GetCrc() == Context->DependenciesCrc)
+							if (Resource->GetCrc().IsValid() && Resource->GetCrc() == InputDependenciesCrc)
 							{
 								ReusedManagedActorsResources.Add(Resource);
 							}
@@ -580,7 +594,7 @@ bool FPCGSpawnActorElement::ExecuteInternal(FPCGContext* Context) const
 
 				// Create managed resource for actor tracking
 				UPCGManagedActors* ManagedActors = NewObject<UPCGManagedActors>(Context->SourceComponent.Get());
-				ManagedActors->SetCrc(Context->DependenciesCrc);
+				ManagedActors->SetCrc(InputDependenciesCrc);
 
 				for (int32 i = 0; i < Points.Num(); ++i)
 				{
diff --git a/Engine/Plugins/Experimental/PCG/Source/PCG/Public/Elements/PCGSpawnActor.h b/Engine/Plugins/Experimental/PCG/Source/PCG/Public/Elements/PCGSpawnActor.h
--- a/Engine/Plugins/Experimental/PCG/Source/PCG/Public/Elements/PCGSpawnActor.h
+++ b/Engine/Plugins/Experimental/PCG/Source/PCG/Public/Elements/PCGSpawnActor.h
@@ -60,7 +60,7 @@ public:
 	bool bForceDisableActorParsing = true;
 
 	UPROPERTY(BlueprintReadWrite, EditAnywhere, Category = Settings, meta = (EditCondition = "Option==EPCGSpawnActorOption::NoMerging", EditConditionHides))
-	EPCGSpawnActorGenerationTrigger bGenerationTrigger = EPCGSpawnActorGenerationTrigger::Default;
+	EPCGSpawnActorGenerationTrigger GenerationTrigger = EPCGSpawnActorGenerationTrigger::Default;
 
 	/** Warning: inheriting parent actor tags work only in non-collapsed actor hierarchies */
 	UPROPERTY(BlueprintReadWrite, EditAnywhere, Category = Settings, meta = (EditCondition = "Option!=EPCGSpawnActorOption::CollapseActors", EditConditionHides))
@@ -75,6 +75,13 @@ public:
 	UPROPERTY(EditAnywhere, BlueprintReadWrite, Category = Settings, meta = (EditCondition = "Option != EPCGSpawnActorOption::CollapseActors", EditConditionHides))
 	TArray<FPCGActorPropertyOverride> ActorOverrides;
 
+protected:
+#if WITH_EDITORONLY_DATA
+	UPROPERTY()
+	EPCGSpawnActorGenerationTrigger bGenerationTrigger_DEPRECATED = EPCGSpawnActorGenerationTrigger::Default;
+#endif
+
+public:
 	virtual void PostLoad() override;
 	virtual void BeginDestroy() override;
 
