Fixed stamping naming issue, could happen when using the break PCG link feature

#rnx

diff --git a/Engine/Plugins/Experimental/PCG/Source/PCG/Private/Helpers/PCGActorHelpers.cpp b/Engine/Plugins/Experimental/PCG/Source/PCG/Private/Helpers/PCGActorHelpers.cpp
--- a/Engine/Plugins/Experimental/PCG/Source/PCG/Private/Helpers/PCGActorHelpers.cpp
+++ b/Engine/Plugins/Experimental/PCG/Source/PCG/Private/Helpers/PCGActorHelpers.cpp
@@ -271,7 +271,7 @@ AActor* UPCGActorHelpers::SpawnDefaultActor(UWorld* World, TSubclassOf<AActor> A
 	}
 
 	FActorSpawnParameters SpawnParams;
-	SpawnParams.Name = MakeUniqueObjectName(World, ActorClass, BaseName);
+	SpawnParams.Name = MakeUniqueObjectName(World->GetCurrentLevel(), ActorClass, BaseName);
 	SpawnParams.Owner = Parent;
 	AActor* NewActor = World->SpawnActor(*ActorClass, &Transform, SpawnParams);
 	
